[PCG] Move PCG Editor windows into the Widgets directory for organizational purposes

Note: Preliminary shelf for Graph Param separate window

#rnx

diff --git a/Engine/Plugins/PCG/Source/PCGEditor/Private/DataVisualizations/PCGDataVisualizationHelpers.cpp b/Engine/Plugins/PCG/Source/PCGEditor/Private/DataVisualizations/PCGDataVisualizationHelpers.cpp
--- a/Engine/Plugins/PCG/Source/PCGEditor/Private/DataVisualizations/PCGDataVisualizationHelpers.cpp
+++ b/Engine/Plugins/PCG/Source/PCGEditor/Private/DataVisualizations/PCGDataVisualizationHelpers.cpp
@@ -2,7 +2,7 @@
 
 #include "DataVisualizations/PCGDataVisualizationHelpers.h"
 
-#include "SPCGEditorGraphAttributeListView.h"
+#include "Widgets/SPCGEditorGraphAttributeListView.h"
 #include "Metadata/PCGAttributePropertySelector.h"
 #include "Metadata/PCGMetadata.h"
 #include "Metadata/Accessors/PCGAttributeAccessorHelpers.h"
diff --git a/Engine/Plugins/PCG/Source/PCGEditor/Private/PCGEditor.cpp b/Engine/Plugins/PCG/Source/PCGEditor/Private/PCGEditor.cpp
--- a/Engine/Plugins/PCG/Source/PCGEditor/Private/PCGEditor.cpp
+++ b/Engine/Plugins/PCG/Source/PCGEditor/Private/PCGEditor.cpp
@@ -24,14 +24,14 @@
 #include "PCGEditorSettings.h"
 #include "PCGEditorStyle.h"
 #include "PCGEditorUtils.h"
-#include "SPCGEditorGraphAttributeListView.h"
-#include "SPCGEditorGraphDebugObjectTree.h"
-#include "SPCGEditorGraphDetailsView.h"
-#include "SPCGEditorGraphDeterminism.h"
-#include "SPCGEditorGraphFind.h"
-#include "SPCGEditorGraphLogView.h"
-#include "SPCGEditorGraphNodePalette.h"
-#include "SPCGEditorGraphProfilingView.h"
+#include "Widgets/SPCGEditorGraphAttributeListView.h"
+#include "Widgets/SPCGEditorGraphDebugObjectTree.h"
+#include "Widgets/SPCGEditorGraphDetailsView.h"
+#include "Widgets/SPCGEditorGraphDeterminism.h"
+#include "Widgets/SPCGEditorGraphFind.h"
+#include "Widgets/SPCGEditorGraphLogView.h"
+#include "Widgets/SPCGEditorGraphNodePalette.h"
+#include "Widgets/SPCGEditorGraphProfilingView.h"
 #include "Nodes/PCGEditorGraphNode.h"
 #include "Nodes/PCGEditorGraphNodeInput.h"
 #include "Nodes/PCGEditorGraphNodeOutput.h"
diff --git a/Engine/Plugins/PCG/Source/PCGEditor/Private/SPCGEditorGraphAttributeListView.cpp b/Engine/Plugins/PCG/Source/PCGEditor/Private/SPCGEditorGraphAttributeListView.cpp
deleted file mode 100644
diff --git a/Engine/Plugins/PCG/Source/PCGEditor/Private/SPCGEditorGraphAttributeListView.h b/Engine/Plugins/PCG/Source/PCGEditor/Private/SPCGEditorGraphAttributeListView.h
deleted file mode 100644
diff --git a/Engine/Plugins/PCG/Source/PCGEditor/Private/SPCGEditorGraphDebugObjectTree.cpp b/Engine/Plugins/PCG/Source/PCGEditor/Private/SPCGEditorGraphDebugObjectTree.cpp
deleted file mode 100644
diff --git a/Engine/Plugins/PCG/Source/PCGEditor/Private/SPCGEditorGraphDebugObjectTree.h b/Engine/Plugins/PCG/Source/PCGEditor/Private/SPCGEditorGraphDebugObjectTree.h
deleted file mode 100644
diff --git a/Engine/Plugins/PCG/Source/PCGEditor/Private/SPCGEditorGraphDetailsView.cpp b/Engine/Plugins/PCG/Source/PCGEditor/Private/SPCGEditorGraphDetailsView.cpp
deleted file mode 100644
diff --git a/Engine/Plugins/PCG/Source/PCGEditor/Private/SPCGEditorGraphDetailsView.h b/Engine/Plugins/PCG/Source/PCGEditor/Private/SPCGEditorGraphDetailsView.h
deleted file mode 100644
diff --git a/Engine/Plugins/PCG/Source/PCGEditor/Private/SPCGEditorGraphDeterminism.cpp b/Engine/Plugins/PCG/Source/PCGEditor/Private/SPCGEditorGraphDeterminism.cpp
deleted file mode 100644
diff --git a/Engine/Plugins/PCG/Source/PCGEditor/Private/SPCGEditorGraphDeterminism.h b/Engine/Plugins/PCG/Source/PCGEditor/Private/SPCGEditorGraphDeterminism.h
deleted file mode 100644
diff --git a/Engine/Plugins/PCG/Source/PCGEditor/Private/SPCGEditorGraphFind.cpp b/Engine/Plugins/PCG/Source/PCGEditor/Private/SPCGEditorGraphFind.cpp
deleted file mode 100644
diff --git a/Engine/Plugins/PCG/Source/PCGEditor/Private/SPCGEditorGraphFind.h b/Engine/Plugins/PCG/Source/PCGEditor/Private/SPCGEditorGraphFind.h
deleted file mode 100644
diff --git a/Engine/Plugins/PCG/Source/PCGEditor/Private/SPCGEditorGraphLogView.cpp b/Engine/Plugins/PCG/Source/PCGEditor/Private/SPCGEditorGraphLogView.cpp
deleted file mode 100644
diff --git a/Engine/Plugins/PCG/Source/PCGEditor/Private/SPCGEditorGraphLogView.h b/Engine/Plugins/PCG/Source/PCGEditor/Private/SPCGEditorGraphLogView.h
deleted file mode 100644
diff --git a/Engine/Plugins/PCG/Source/PCGEditor/Private/SPCGEditorGraphNodePalette.cpp b/Engine/Plugins/PCG/Source/PCGEditor/Private/SPCGEditorGraphNodePalette.cpp
deleted file mode 100644
diff --git a/Engine/Plugins/PCG/Source/PCGEditor/Private/SPCGEditorGraphNodePalette.h b/Engine/Plugins/PCG/Source/PCGEditor/Private/SPCGEditorGraphNodePalette.h
deleted file mode 100644
diff --git a/Engine/Plugins/PCG/Source/PCGEditor/Private/SPCGEditorGraphProfilingView.cpp b/Engine/Plugins/PCG/Source/PCGEditor/Private/SPCGEditorGraphProfilingView.cpp
deleted file mode 100644
diff --git a/Engine/Plugins/PCG/Source/PCGEditor/Private/SPCGEditorGraphProfilingView.h b/Engine/Plugins/PCG/Source/PCGEditor/Private/SPCGEditorGraphProfilingView.h
deleted file mode 100644
diff --git a/Engine/Plugins/PCG/Source/PCGEditor/Private/Widgets/SPCGEditorGraphAttributeListView.cpp b/Engine/Plugins/PCG/Source/PCGEditor/Private/Widgets/SPCGEditorGraphAttributeListView.cpp
new file mode 100644
--- /dev/null
+++ b/Engine/Plugins/PCG/Source/PCGEditor/Private/Widgets/SPCGEditorGraphAttributeListView.cpp
@@ -0,0 +1,1566 @@
+// Copyright Epic Games, Inc. All Rights Reserved.
+
+#include "SPCGEditorGraphAttributeListView.h"
+
+#include "PCGAssetExporterUtils.h"
+#include "PCGComponent.h"
+#include "PCGData.h"
+#include "PCGDataVisualization.h"
+#include "PCGGraph.h"
+#include "PCGNode.h"
+#include "PCGParamData.h"
+#include "PCGPin.h"
+#include "PCGSubsystem.h"
+#include "Compute/Data/PCGProxyForGPUData.h"
+#include "Data/PCGSpatialData.h"
+#include "Elements/IO/PCGSaveAssetElement.h"
+#include "Metadata/PCGAttributePropertySelector.h"
+#include "Metadata/Accessors/IPCGAttributeAccessor.h"
+#include "Metadata/Accessors/PCGAttributeAccessor.h"
+#include "Metadata/Accessors/PCGAttributeAccessorHelpers.h"
+#include "Metadata/Accessors/PCGAttributeAccessorKeys.h"
+#include "Metadata/Accessors/PCGCustomAccessor.h"
+
+#include "PCGEditor.h"
+#include "PCGEditorGraph.h"
+#include "PCGEditorStyle.h"
+#include "Nodes/PCGEditorGraphNodeBase.h"
+
+#include "Fonts/FontMeasure.h"
+#include "Framework/Commands/GenericCommands.h"
+#include "Framework/MultiBox/MultiBoxBuilder.h"
+#include "HAL/PlatformApplicationMisc.h"
+#include "Widgets/Images/SLayeredImage.h"
+#include "Widgets/Images/SThrobber.h"
+#include "Widgets/Input/SButton.h"
+#include "Widgets/Input/SComboBox.h"
+#include "Widgets/Input/SSearchBox.h"
+#include "Widgets/Layout/SScrollBox.h"
+
+#define LOCTEXT_NAMESPACE "SPCGEditorGraphAttributeListView"
+
+namespace PCGEditorGraphAttributeListView
+{
+	const FText NoPinAvailableText = LOCTEXT("NoPinAvailableText", "No pins");
+	const FText CrcLabelFormat = LOCTEXT("InfoTextBlockWithCrcFmt", "{0} | CRC: {1}");
+	const FText LastLabelFormat = LOCTEXT("InfoTextBlockWithLastAttributeFmt", "{0} | Last attribute: {1}");
+	const FText NoDataAvailableText = LOCTEXT("NoDataAvailableText", "No data available");
+	const FText NoNodeInspectedText = LOCTEXT("NoNodeInspectedText", "No node being inspected");
+	const FText NoNodeInspectedToolTip = LOCTEXT("NoNodeInspectedToolTip", "Inspect a node using the right click menu");
+	const FText TEXT_IndexLabel = LOCTEXT("IndexLabel", "Index");
+
+	bool IsGraphCacheDebuggingEnabled()
+	{
+		UWorld* World = GEditor ? (GEditor->PlayWorld ? GEditor->PlayWorld.Get() : GEditor->GetEditorWorldContext().World()) : nullptr;
+		UPCGSubsystem* Subsystem = UPCGSubsystem::GetInstance(World);
+		return Subsystem && Subsystem->IsGraphCacheDebuggingEnabled();
+	}
+
+	float CalculateColumnWidth(const FText& InText, bool bClampToMaxColumnWidth = true)
+	{
+		check(FSlateApplication::Get().GetRenderer());
+		const TSharedRef<FSlateFontMeasure> FontMeasure = FSlateApplication::Get().GetRenderer()->GetFontMeasureService();
+		const FSlateFontInfo FontInfo = FAppStyle::GetFontStyle(TEXT("NormalText"));
+
+		const float TextWidth = FontMeasure->Measure(InText, FontInfo).X;
+		constexpr float ColumnPadding = 22.0f; // TODO: Grab padding from header style
+		const float ColumnWidth = TextWidth + ColumnPadding;
+		return bClampToMaxColumnWidth ? FMath::Min(ColumnWidth, MaxColumnWidth) : ColumnWidth;
+	}
+}
+
+bool FPCGListViewUpdater::IsCompleted() const
+{
+	return UpdateTask.IsCompleted();
+}
+
+void FPCGListViewUpdater::Launch()
+{
+	// Passing a shared pointer to this in order for the task to keep the object alive even if we discard it in the attribute list view
+	UpdateTask = UE::Tasks::Launch(UE_SOURCE_LOCATION, [SharedContext = SharedThis(this)]()
+	{
+		SharedContext->AsyncFilter();
+		SharedContext->AsyncSort();
+	});
+}
+
+void FPCGListViewUpdater::AsyncSort()
+{
+	if (const FPCGColumnData* Data = ColumnData.Find(SortingColumn))
+	{
+		if (Data->DataAccessor.IsValid() && Data->DataKeys.IsValid() && Data->DataKeys->GetNum() == ListViewItems.Num())
+		{
+			//lambda used here to get the index value of an item in the array for sorting
+			PCGAttributeAccessorHelpers::SortByAttribute(*Data->DataAccessor, *Data->DataKeys, ListViewItems, !(SortMode & EColumnSortMode::Descending), [this](int Index) { return ListViewItems[Index]->Index; });
+		}
+	}
+}
+
+void FPCGListViewUpdater::AsyncFilter()
+{
+	TArray<PCGListViewItemPtr> FilteredListViewItems;
+	FilteredListViewItems.Reserve(ListViewItems.Num());
+
+	for (const PCGListViewItemPtr& ListViewItem : ListViewItems)
+	{
+		const FPCGPointFilterExpressionContext PointFilterContext(ListViewItem.Get(), &ColumnData);
+		if (TextFilter->TestTextFilter(PointFilterContext))
+		{
+			FilteredListViewItems.Add(ListViewItem);
+		}
+	}
+
+	ListViewItems = MoveTemp(FilteredListViewItems);
+}
+
+void SPCGListViewItemRow::Construct(const FArguments& InArgs, const TSharedRef<STableViewBase>& InOwnerTableView)
+{
+	check(InArgs._ListViewItem && InArgs._AttributeListView);
+	InternalItem = InArgs._ListViewItem;
+	AttributeListView = InArgs._AttributeListView;
+
+	SMultiColumnTableRow<PCGListViewItemPtr>::Construct(
+		SMultiColumnTableRow::FArguments()
+		.Style(FAppStyle::Get(), "DataTableEditor.CellListViewRow"),
+		InOwnerTableView);
+}
+
+TSharedRef<SWidget> SPCGListViewItemRow::GenerateWidgetForColumn(const FName& ColumnId)
+{
+	FText RowText = LOCTEXT("ColumnError", "Unrecognized Column");
+
+	const TSharedPtr<SPCGEditorGraphAttributeListView> SharedAttributeListView = AttributeListView.Pin();
+	check(SharedAttributeListView.IsValid());
+
+	// We will only compute cell width for types are that not numbers.
+	bool bShouldComputeCellWidth = false;
+
+	if (FPCGColumnData* PCGColumnData = SharedAttributeListView->PCGColumnData.Find(ColumnId))
+	{
+		if (PCGColumnData->DataAccessor.IsValid() && PCGColumnData->DataKeys.IsValid())
+		{
+			int32 Index = InternalItem->Index;
+			auto Callback = [&PCGColumnData, Index, &RowText, &bShouldComputeCellWidth] (auto Dummy)
+			{
+				using ValueType = decltype(Dummy);
+				ValueType Value{};
+				if (PCGColumnData->DataAccessor->Get<ValueType>(Value, Index, *PCGColumnData->DataKeys))
+				{
+					if constexpr (PCG::Private::IsOfTypes<ValueType, bool>())
+					{
+						RowText = FText::FromString(LexToString(Value));
+					}
+					else if constexpr (PCG::Private::IsOfTypes<ValueType, FString>())
+					{
+						RowText = FText::FromString(Value);
+						bShouldComputeCellWidth = true;
+					}
+					else if constexpr (PCG::Private::IsOfTypes<ValueType, FName>())
+					{
+						RowText = FText::FromName(Value);
+						bShouldComputeCellWidth = true;
+					}
+					else if constexpr (FTextAsNumberIsValid<ValueType>::value)
+					{
+						RowText = FText::AsNumber(Value);
+					}
+					else if constexpr (PCG::Private::IsOfTypes<ValueType, FSoftObjectPath>())
+					{
+						RowText = FText::FromString(Value.ToString());
+						bShouldComputeCellWidth = true;
+					}
+					else if constexpr (PCG::Private::IsOfTypes<ValueType, FSoftClassPath>())
+					{
+						RowText = FText::FromString(Value.ToString());
+						bShouldComputeCellWidth = true;
+					}
+					else
+					{
+						ensureMsgf(false, TEXT("Unsupported Data Type"));
+						RowText = LOCTEXT("UnsupportedDataTypeError", "Unsupported Data Type");
+					}
+				}
+			};
+
+			PCGMetadataAttribute::CallbackWithRightType(PCGColumnData->DataAccessor->GetUnderlyingType(), Callback);
+		}
+	}
+
+	if (bShouldComputeCellWidth)
+	{
+		float& CurrentColumnWidth = SharedAttributeListView->ColumnsMaxWidthMapping[ColumnId];
+		const float TextWidth = PCGEditorGraphAttributeListView::CalculateColumnWidth(RowText, /*bClampToMaxColumnWidth=*/false);
+		if (TextWidth > CurrentColumnWidth)
+		{
+			CurrentColumnWidth = TextWidth;
+		}
+	}
+
+	return SNew(STextBlock)
+		.Text(RowText)
+		.OverflowPolicy(ETextOverflowPolicy::Ellipsis)
+		.Margin(FMargin(2.0f, 0.0f));
+}
+
+FPCGPointFilterExpressionContext::FPCGPointFilterExpressionContext(const FPCGListViewItem* InRowItem, const TMap<FName, FPCGColumnData>* InPCGColumnData)
+	: RowItem(InRowItem)
+	, PCGColumnData(InPCGColumnData)
+{
+	check(InRowItem);
+}
+
+bool FPCGPointFilterExpressionContext::TestBasicStringExpression(const FTextFilterString& InValue, const ETextFilterTextComparisonMode InTextComparisonMode) const
+{
+	// Basic string search is disabled as it would require us to search the entire attribute table at once and it's not very useful.
+	return false;
+}
+
+bool FPCGPointFilterExpressionContext::TestComplexExpression(const FName& InKey, const FTextFilterString& InValue, const ETextFilterComparisonOperation InComparisonOperation, const ETextFilterTextComparisonMode InTextComparisonMode) const
+{
+	const int32 Index = RowItem->Index;
+
+	if (PCGEditorGraphAttributeListView::TEXT_IndexLabel.EqualToCaseIgnored(FText::FromName(InKey)))
+	{
+		const FTextFilterString PointValue(FString::FromInt(Index));
+		return TextFilterUtils::TestComplexExpression(PointValue, InValue, InComparisonOperation, InTextComparisonMode);
+	}
+	else if(const FPCGColumnData* PCGColumnInfo = PCGColumnData->Find(InKey))
+	{
+		if (PCGColumnInfo->DataAccessor.IsValid() && PCGColumnInfo->DataKeys.IsValid())
+		{
+			auto Callback = [&PCGColumnInfo, Index, &InValue, &InComparisonOperation, &InTextComparisonMode] (auto Dummy)
+			{
+				using ValueType = decltype(Dummy);
+				ValueType Value{};
+				if (PCGColumnInfo->DataAccessor->Get<ValueType>(Value, Index, *PCGColumnInfo->DataKeys))
+				{
+					FText TextValue;
+					bool bInvalid = false;
+					if constexpr (PCG::Private::IsOfTypes<ValueType, bool>())
+					{
+						TextValue = FText::FromString(LexToString(Value));
+					}
+					else if constexpr (PCG::Private::IsOfTypes<ValueType, FString>())
+					{
+						TextValue = FText::FromString(Value);
+					}
+					else if constexpr (PCG::Private::IsOfTypes<ValueType, FName>())
+					{
+						TextValue = FText::FromName(Value);
+					}
+					else if constexpr (FTextAsNumberIsValid<ValueType>::value)
+					{
+						TextValue = FText::AsNumber(Value, &FNumberFormattingOptions::DefaultNoGrouping());
+					}
+					else
+					{
+						ensureMsgf(false, TEXT("Unsupported Data Type"));
+						bInvalid = true;
+					}
+
+					if (!bInvalid)
+					{
+						const FTextFilterString PointValue(TextValue.ToString());
+						return TextFilterUtils::TestComplexExpression(PointValue, InValue, InComparisonOperation, InTextComparisonMode);
+					}
+				}
+
+				return false;
+			};
+
+			return PCGMetadataAttribute::CallbackWithRightType(PCGColumnInfo->DataAccessor->GetUnderlyingType(), Callback);
+		}
+	}
+
+	return true;
+}
+
+SPCGEditorGraphAttributeListView::~SPCGEditorGraphAttributeListView()
+{
+	if (PCGEditorPtr.IsValid())
+	{
+		PCGEditorPtr.Pin()->OnInspectedStackChangedDelegate.RemoveAll(this);
+	}
+
+	DataPtr.Reset();
+}
+
+void SPCGEditorGraphAttributeListView::Construct(const FArguments& InArgs, TSharedPtr<FPCGEditor> InPCGEditor)
+{
+	PCGEditorPtr = InPCGEditor;
+	PCGEditorPtr.Pin()->OnInspectedStackChangedDelegate.AddSP(this, &SPCGEditorGraphAttributeListView::OnInspectedStackChanged);
+
+	TextFilter = MakeShareable(new FTextFilterExpressionEvaluator(ETextFilterExpressionEvaluatorMode::Complex));
+
+	ColumnWidthVisibilityCache.Empty(PCGEditorGraphAttributeListView::MaxNodeColumnWidthCachedItems);
+
+	ListViewHeader = CreateHeaderRowWidget();
+
+	ListViewCommands = MakeShareable(new FUICommandList);
+	ListViewCommands->MapAction(FGenericCommands::Get().Copy,
+		FExecuteAction::CreateSP(this, &SPCGEditorGraphAttributeListView::CopySelectionToClipboard),
+		FCanExecuteAction::CreateSP(this, &SPCGEditorGraphAttributeListView::CanCopySelectionToClipboard));
+
+	const TSharedRef<SScrollBar> HorizontalScrollBar = SNew(SScrollBar)
+		.Orientation(Orient_Horizontal)
+		.Thickness(FVector2D(12.0f, 12.0f));
+
+	const TSharedRef<SScrollBar> VerticalScrollBar = SNew(SScrollBar)
+		.Orientation(Orient_Vertical)
+		.Thickness(FVector2D(12.0f, 12.0f));
+
+	auto VisibilityTest = [this]()
+	{
+		return (ListViewItems.IsEmpty() && ListViewHeader->GetColumns().IsEmpty()) ? EVisibility::Hidden : EVisibility::Visible;
+	};
+
+	SAssignNew(ListView, SListView<PCGListViewItemPtr>)
+		.ListItemsSource(&ListViewItems)
+		.HeaderRow(ListViewHeader)
+		.OnGenerateRow(this, &SPCGEditorGraphAttributeListView::OnGenerateRow)
+		.OnMouseButtonDoubleClick(this, &SPCGEditorGraphAttributeListView::OnItemDoubleClicked)
+		.OnContextMenuOpening(this, &SPCGEditorGraphAttributeListView::OnItemsContextMenu)
+		.AllowOverscroll(EAllowOverscroll::No)
+		.ExternalScrollbar(VerticalScrollBar)
+		.Visibility_Lambda(VisibilityTest)
+		.OnKeyDownHandler(this, &SPCGEditorGraphAttributeListView::OnListViewKeyDown)
+		.ConsumeMouseWheel(EConsumeMouseWheel::Always);
+
+	SAssignNew(PinComboBox, SComboBox<TSharedPtr<FPinComboBoxItem>>)
+		.OptionsSource(&PinComboBoxItems)
+		.OnGenerateWidget(this, &SPCGEditorGraphAttributeListView::OnGeneratePinWidget)
+		.OnSelectionChanged(this, &SPCGEditorGraphAttributeListView::OnSelectionChangedPin)
+		[
+			SNew(STextBlock)
+			.Text(this, &SPCGEditorGraphAttributeListView::OnGenerateSelectedPinText)
+		];
+
+	SAssignNew(DataComboBox, SComboBox<TSharedPtr<FString>>)
+		.OptionsSource(&DataComboBoxItems)
+		.OnGenerateWidget(this, &SPCGEditorGraphAttributeListView::OnGenerateDataWidget)
+		.OnSelectionChanged(this, &SPCGEditorGraphAttributeListView::OnSelectionChanged)
+		[
+			SNew(STextBlock)
+			.Text(this, &SPCGEditorGraphAttributeListView::OnGenerateSelectedDataText)
+		];
+
+	TSharedPtr<SLayeredImage> FilterImage = SNew(SLayeredImage)
+		.Image(FAppStyle::Get().GetBrush("Icons.Filter"))
+		.ColorAndOpacity(FSlateColor::UseForeground());
+
+	FilterImage->AddLayer(TAttribute<const FSlateBrush*>(this, &SPCGEditorGraphAttributeListView::GetFilterBadgeIcon));
+
+	TSharedPtr<SButton> LockButton = SNew(SButton)
+		.ButtonStyle(FAppStyle::Get(), "SimpleButton")
+		.OnClicked(this, &SPCGEditorGraphAttributeListView::OnLockClick)
+		.ContentPadding(FMargin(4, 2))
+		.HAlign(HAlign_Center)
+		.VAlign(VAlign_Center)
+		.ToolTipText(LOCTEXT("LockSelectionButton_ToolTip", "Locks the current attribute list view to this selection."))
+		[
+			SNew(SImage)
+				.ColorAndOpacity(FSlateColor::UseForeground())
+				.Image(this, &SPCGEditorGraphAttributeListView::OnGetLockButtonImageResource)
+		];
+
+	TSharedPtr<SButton> FrameDataButton = SNew(SButton)
+		.ButtonStyle(FAppStyle::Get(), "SimpleButton")
+		.OnClicked(this, &SPCGEditorGraphAttributeListView::OnFocusOnDataClicked)
+		.IsEnabled(this, &SPCGEditorGraphAttributeListView::IsFocusOnDataEnabled)
+		.ContentPadding(FMargin(4, 2))
+		.HAlign(HAlign_Center)
+		.VAlign(VAlign_Center)
+		.ToolTipText(LOCTEXT("FocusOnDataButton_ToolTip", "Zoom to selected data."))
+		[
+			SNew(SImage)
+				.ColorAndOpacity(FSlateColor::UseForeground())
+				.Image(FPCGEditorStyle::Get().GetBrush("PCG.Editor.ZoomToSelection"))
+		];
+		
+	TSharedPtr<SComboButton> FilterButton = SNew(SComboButton)
+		.ForegroundColor(FSlateColor::UseStyle())
+		.HasDownArrow(false)
+		.OnGetMenuContent(this, &SPCGEditorGraphAttributeListView::OnGenerateFilterMenu)
+		.ContentPadding(1)
+		.ButtonContent()
+		[
+			FilterImage.ToSharedRef()
+		];
+
+	TSharedPtr<SComboButton> AdditionalOperationsButton = SNew(SComboButton)
+		.ForegroundColor(FSlateColor::UseStyle())
+		.HasDownArrow(false)
+		.OnGetMenuContent(this, &SPCGEditorGraphAttributeListView::OnGenerateAdditionalOperationsMenu)
+		.ContentPadding(1)
+		.ButtonContent()
+		[
+			SNew(SImage)
+			.Image(FAppStyle::GetBrush("EditorViewportToolBar.OptionsDropdown"))
+			.ColorAndOpacity(FSlateColor::UseForeground())
+		];
+
+	SAssignNew(SearchBoxWidget, SSearchBox)
+		.MinDesiredWidth(100.0f)
+		.InitialText(ActiveFilterText)
+		.OnTextChanged(this, &SPCGEditorGraphAttributeListView::OnFilterTextChanged)
+		.OnTextCommitted(this, &SPCGEditorGraphAttributeListView::OnFilterTextCommitted)
+		.DelayChangeNotificationsWhileTyping(true)
+		.DelayChangeNotificationsWhileTypingSeconds(0.5f);
+
+	this->ChildSlot
+	[
+		SNew(SVerticalBox)
+		+SVerticalBox::Slot()
+		.AutoHeight()
+		.Padding(1.0f)
+		[
+			SNew(SHorizontalBox)
+			+SHorizontalBox::Slot()
+			.AutoWidth()
+			.Padding(1.0f, 0.0f)
+			[
+				LockButton->AsShared()
+			]
+			+SHorizontalBox::Slot()
+			.AutoWidth()
+			.Padding(1.0f, 0.0f)
+			[
+				FrameDataButton->AsShared()
+			]
+			+SHorizontalBox::Slot()
+			.AutoWidth()
+			.Padding(1.0f, 0.0f)
+			[
+				FilterButton->AsShared()
+			]
+			+SHorizontalBox::Slot()
+			.AutoWidth()
+			.Padding(1.0f, 0.0f)
+			[
+				PinComboBox->AsShared()
+			]
+			+SHorizontalBox::Slot()
+			.AutoWidth()
+			.Padding(1.0f, 0.0f)
+			[
+				DataComboBox->AsShared()
+			]
+			+SHorizontalBox::Slot()
+			.MinWidth(100.0)
+			.MaxWidth(300.0)
+			.Padding(1.0f, 0.0f)
+			[
+				SearchBoxWidget->AsShared()
+			]
+			+SHorizontalBox::Slot()
+			.AutoWidth()
+			.VAlign(EVerticalAlignment::VAlign_Center)
+			.Padding(4.0f, 0.0f)
+			[
+				SNew(SButton)
+				.ButtonStyle(FAppStyle::Get(), "SimpleButton")
+				.OnClicked(this, &SPCGEditorGraphAttributeListView::OnNodeNameClicked)
+				[
+					SAssignNew(NodeNameTextBlock, STextBlock)
+					.Text(PCGEditorGraphAttributeListView::NoNodeInspectedText)
+					.ToolTipText(PCGEditorGraphAttributeListView::NoNodeInspectedToolTip)
+				]
+			]
+			+SHorizontalBox::Slot()
+			.AutoWidth()
+			.VAlign(EVerticalAlignment::VAlign_Center)
+			.Padding(1.0f, 0.0f)
+			[
+				AdditionalOperationsButton->AsShared()
+			]
+			+SHorizontalBox::Slot()
+			.AutoWidth()
+			.HAlign(EHorizontalAlignment::HAlign_Left)
+			.VAlign(EVerticalAlignment::VAlign_Center)
+			.Padding(4.0f, 0.0f)
+			[
+				SAssignNew(InfoTextBlock, STextBlock)
+			]
+		]
+		+SVerticalBox::Slot()
+		.FillHeight(1.0f)
+		[
+			SNew(SHorizontalBox)
+			+SHorizontalBox::Slot()
+			.FillWidth(1.0f)
+			[
+				SNew(SOverlay)
+				+SOverlay::Slot()
+				[
+					SNew(SScrollBox)
+					.Orientation(Orient_Horizontal)
+					.ExternalScrollbar(HorizontalScrollBar)
+					+SScrollBox::Slot()
+					[
+						ListView->AsShared()
+					]
+				]
+				+SOverlay::Slot()
+				.HAlign(HAlign_Center)
+				.VAlign(VAlign_Center)
+				[
+					SNew(SCircularThrobber)
+					.Radius(12.0f)
+					.Visibility_Lambda([this](){return CurrentUpdateTask.IsValid() && !CurrentUpdateTask->IsCompleted() ? EVisibility::Visible : EVisibility::Hidden; })
+				]
+			]
+			+SHorizontalBox::Slot()
+			.AutoWidth()
+			[
+				VerticalScrollBar
+			]
+		]
+		+SVerticalBox::Slot()
+		.AutoHeight()
+		[
+			SNew(SHorizontalBox)
+			+SHorizontalBox::Slot()
+			[
+				HorizontalScrollBar
+			]
+		]
+	];
+}
+
+void SPCGEditorGraphAttributeListView::Tick(const FGeometry& AllottedGeometry, const double InCurrentTime, const float InDeltaTime)
+{
+	SCompoundWidget::Tick(AllottedGeometry, InCurrentTime, InDeltaTime);
+
+	if (bNeedsRefresh)
+	{
+		bNeedsRefresh = false;
+
+		RefreshPinComboBox();
+		RefreshDataComboBox();
+		RefreshAttributeList();
+	}
+
+	if (CurrentUpdateTask.IsValid() && CurrentUpdateTask->IsCompleted())
+	{
+		FilteredListViewItems = MoveTemp(CurrentUpdateTask->ListViewItems);
+		if (ListView.IsValid())
+		{
+			ListView->SetItemsSource(&FilteredListViewItems);
+			ListView->RequestListRefresh();
+
+			InfoTextBlock->SetText(FText::Format(LOCTEXT("InfoTextBlockFmt", "Showing {0}/{1} entries"), FilteredListViewItems.Num(), ListViewItems.Num()));
+
+			if (PCGEditorGraphAttributeListView::IsGraphCacheDebuggingEnabled())
+			{
+				// If cache debugging enabled, write CRC to help diagnose missed-dependency issues
+				const FPCGDataCollection* InspectionData = GetInspectionData();
+				const int32 DataIndex = GetSelectedDataIndex();
+				const FPCGCrc Crc = (InspectionData && InspectionData->DataCrcs.IsValidIndex(DataIndex)) ? InspectionData->DataCrcs[DataIndex] : FPCGCrc(0);
+				InfoTextBlock->SetText(FText::Format(PCGEditorGraphAttributeListView::CrcLabelFormat, InfoTextBlock->GetText(), Crc.GetValue()));
+			}
+
+			if (DataPtr.IsValid() && DataPtr->HasCachedLastSelector())
+			{
+				const FText LastSelector = DataPtr->GetCachedLastSelector().GetDisplayText();
+				InfoTextBlock->SetText(FText::Format(PCGEditorGraphAttributeListView::LastLabelFormat, InfoTextBlock->GetText(), LastSelector));
+			}
+		}
+
+		CurrentUpdateTask.Reset();
+	}
+}
+
+TSharedRef<SHeaderRow> SPCGEditorGraphAttributeListView::CreateHeaderRowWidget() const
+{
+	return SNew(SHeaderRow);
+}
+
+TWeakObjectPtr<UPCGComponent> SPCGEditorGraphAttributeListView::GetPCGComponent() const
+{
+	return PCGEditorPtr.IsValid() ? PCGEditorPtr.Pin()->GetPCGComponentBeingInspected() : nullptr;
+}
+
+void SPCGEditorGraphAttributeListView::OnInspectedStackChanged(const FPCGStack& InPCGStack)
+{
+	RequestRefresh();
+}
+
+UPCGEditorGraphNodeBase* SPCGEditorGraphAttributeListView::GetNodeBeingInspected() const
+{
+	return PCGEditorGraphNode.Get();
+}
+
+void SPCGEditorGraphAttributeListView::SetNodeBeingInspected(UPCGEditorGraphNodeBase* InPCGEditorGraphNode)
+{
+	if (PCGEditorGraphNode == InPCGEditorGraphNode)
+	{
+		return;
+	}
+
+	CacheColumnWidthVisibility();
+
+	PCGEditorGraphNode = InPCGEditorGraphNode;
+
+	if (PCGEditorGraphNode.IsValid())
+	{
+		NodeNameTextBlock->SetText(PCGEditorGraphNode->GetNodeTitle(ENodeTitleType::ListView));
+		NodeNameTextBlock->SetToolTipText(PCGEditorGraphNode->GetTooltipText());
+	}
+	else
+	{
+		NodeNameTextBlock->SetText(PCGEditorGraphAttributeListView::NoNodeInspectedText);
+		NodeNameTextBlock->SetToolTipText(PCGEditorGraphAttributeListView::NoNodeInspectedToolTip);
+	}
+
+	// Always unlock when changing the node, to make sure we unlock when removing the inspected node
+	bIsLocked = false;
+
+	RequestRefresh();
+}
+
+void SPCGEditorGraphAttributeListView::OnGenerateUpdated(UPCGComponent* /*InPCGComponent*/)
+{
+	RequestRefresh();
+}
+
+const FPCGDataCollection* SPCGEditorGraphAttributeListView::GetInspectionData(const TSharedPtr<FPinComboBoxItem>& EditorPin) const
+{
+	TWeakObjectPtr<UPCGComponent> PCGComponent = GetPCGComponent();
+	if (!PCGComponent.IsValid())
+	{
+		return nullptr;
+	}
+
+	const UPCGNode* PCGNode = PCGEditorGraphNode.IsValid() ? PCGEditorGraphNode->GetPCGNode() : nullptr;
+	if (!PCGNode)
+	{
+		return nullptr;
+	}
+
+	const UPCGPin* Pin = nullptr;
+	if (EditorPin)
+	{
+		const TArray<TObjectPtr<UPCGPin>>& Pins = EditorPin->bIsOutputPin ? PCGNode->GetOutputPins() : PCGNode->GetInputPins();
+		if (Pins.IsValidIndex(EditorPin->PinIndex))
+		{
+			Pin = Pins[EditorPin->PinIndex];
+		}
+	}
+
+	if (!Pin)
+	{
+		return nullptr;
+	}
+
+	PCGEditorGraphUtils::GetInspectablePin(PCGNode, Pin, PCGNode, Pin);
+
+	if (!PCGNode || !Pin)
+	{
+		return nullptr;
+	}
+
+	const TSharedPtr<FPCGEditor> PCGEditor = PCGEditorPtr.Pin();
+	const FPCGStack* PCGStack = PCGEditor->GetStackBeingInspected();
+	if (!PCGStack)
+	{
+		return nullptr;
+	}
+
+	// Create a temporary stack with Node+Pin to query the exact DataCollection we are inspecting
+	FPCGStack Stack = *PCGStack;
+	TArray<FPCGStackFrame>& StackFrames = Stack.GetStackFramesMutable();
+	StackFrames.Reserve(StackFrames.Num() + 2);
+	StackFrames.Emplace(PCGNode);
+	StackFrames.Emplace(Pin);
+
+	return PCGComponent->GetInspectionData(Stack);
+}
+
+const FPCGDataCollection* SPCGEditorGraphAttributeListView::GetInspectionData() const
+{
+	return GetInspectionData(PinComboBox->GetSelectedItem());
+}
+
+void SPCGEditorGraphAttributeListView::RefreshAttributeList()
+{
+	CacheColumnWidthVisibility();
+
+	HiddenAttributes = ListViewHeader->GetHiddenColumnIds();
+
+	// Swapping to an empty item list to force a widget clear, otherwise the widgets will try to update during add column and access invalid data
+	static const TArray<PCGListViewItemPtr> EmptyList;
+	ListView->SetItemsSource(&EmptyList);
+
+	PCGColumnData.Empty();
+	ColumnsMaxWidthMapping.Empty();
+	ListViewItems.Empty();
+	ListViewHeader->ClearColumns();
+	InfoTextBlock->SetText(FText::GetEmpty());
+	DataPtr.Reset();
+
+	const FPCGDataCollection* InspectionData = GetInspectionData();
+	if (!InspectionData)
+	{
+		return;
+	}
+
+	const int32 DataIndex = GetSelectedDataIndex();
+	if (!InspectionData->TaggedData.IsValidIndex(DataIndex))
+	{
+		return;
+	}
+
+	const FPCGTaggedData& TaggedData = InspectionData->TaggedData[DataIndex];
+	const UPCGData* PCGData = TaggedData.Data;
+	const FPCGCrc Crc = InspectionData->DataCrcs.IsValidIndex(DataIndex) ? InspectionData->DataCrcs[DataIndex] : FPCGCrc(0);
+
+	if (!PCGData)
+	{
+		return;
+	}
+
+	// If we have a proxy for GPU data, read back CPU data for inspection.
+	const UPCGProxyForGPUData* Proxy = Cast<UPCGProxyForGPUData>(PCGData);
+	if (Proxy)
+	{
+		UPCGProxyForGPUData::FReadbackResult Result = Proxy->GetCPUData(/*InContext=*/nullptr);
+
+		if (!Result.bComplete)
+		{
+			// Poll next tick.
+			bNeedsRefresh = true;
+			return;
+		}
+
+		if (!Result.TaggedData.Data)
+		{
+			return;
+		}
+
+		PCGData = Result.TaggedData.Data;
+	}
+
+	const FPCGDataVisualizationRegistry& DataVisRegistry = FPCGModule::GetConstPCGDataVisualizationRegistry();
+	const UPCGData* DataToVisualize = PCGData;
+
+	if (const IPCGDataVisualization* DataVisualization = DataVisRegistry.GetDataVisualization(PCGData->GetClass()))
+	{
+		const FPCGTableVisualizerInfo TableVisualizerInfo = DataVisualization->GetTableVisualizerInfo(PCGData);
+		DataToVisualize = TableVisualizerInfo.Data;
+
+		for (const FPCGTableVisualizerColumnInfo& ColumnInfo : TableVisualizerInfo.ColumnInfos)
+		{
+			AddColumn(DataToVisualize, ColumnInfo, TableVisualizerInfo.AccessorKeys);
+		}
+
+		SortingColumn = TableVisualizerInfo.SortingColumn;
+		SortMode = (EColumnSortMode::Type)TableVisualizerInfo.SortingMode;
+		FocusOnDataCallback = TableVisualizerInfo.FocusOnDataCallback;
+
+		const uint32 NumEntries = TableVisualizerInfo.AccessorKeys.IsValid() ? TableVisualizerInfo.AccessorKeys->GetNum() : 0;
+		ListViewItems.Reserve(NumEntries);
+
+		for (uint32 Index = 0; Index < NumEntries; ++Index)
+		{
+			PCGListViewItemPtr ListViewItem = MakeShared<FPCGListViewItem>();
+			ListViewItem->Index = Index;
+			ListViewItems.Add(ListViewItem);
+		}
+	}
+	else // No visualization, just default back the values
+	{
+		SortingColumn = NAME_None;
+		SortMode = EColumnSortMode::Type::Ascending;
+		FocusOnDataCallback.Reset();
+	}
+
+	DataPtr = DataToVisualize;
+
+	ListView->SetItemsSource(&ListViewItems);
+	ListView->RequestListRefresh();
+
+	RestoreColumnWidthVisibility();
+
+	LaunchUpdateTask();
+}
+
+void SPCGEditorGraphAttributeListView::RefreshPinComboBox()
+{
+	PinComboBoxItems.Empty();
+	PinComboBox->ClearSelection();
+	PinComboBox->RefreshOptions();
+
+	const UPCGNode* PCGNode = PCGEditorGraphNode.IsValid() ? PCGEditorGraphNode->GetPCGNode() : nullptr;
+	if (!PCGNode)
+	{
+		return;
+	}
+
+	// Add output and then input pins to list. Optionally output the first connected item - useful for initializing
+	// the selected item to the first connected output pin.
+	auto PopulatePins = [](
+		const TArray<TObjectPtr<UPCGPin>>& InPins,
+		const FString& InFormatText,
+		TArray<TSharedPtr<FPinComboBoxItem>>& InOutItems,
+		int32* OutFirstConnectedItemIndex)
+	{
+		for (int32 PinIndex = 0; PinIndex < InPins.Num(); ++PinIndex)
+		{
+			const UPCGPin* PCGPin = InPins[PinIndex];
+			const bool bIsOutputPin = PCGPin->IsOutputPin();
+			// Pin is included in list if it is connected, or if it is an output pin.
+			if (PCGPin->IsConnected() || bIsOutputPin)
+			{
+				FString ItemName = FString::Format(*InFormatText, { PCGPin->Properties.Label.ToString() });
+				InOutItems.Add(MakeShared<FPinComboBoxItem>(FName(ItemName), PinIndex, bIsOutputPin));
+
+				// Look for first connected, null pointer once found so only first is taken.
+				if (OutFirstConnectedItemIndex && PCGPin->IsConnected())
+				{
+					*OutFirstConnectedItemIndex = InOutItems.Num() - 1;
+					OutFirstConnectedItemIndex = nullptr;
+				}
+			}
+		}
+	};
+
+	// Pick first connected output pin by default if there is one, otherwise default to first output pin.
+	int32 FirstConnectedItemIndex = 0;
+	PopulatePins(PCGNode->GetOutputPins(), TEXT("Output: {0}"), PinComboBoxItems, &FirstConnectedItemIndex);
+	PopulatePins(PCGNode->GetInputPins(), TEXT("Input: {0}"), PinComboBoxItems, nullptr);
+
+	if (PinComboBoxItems.Num() > 0 && ensure(PinComboBoxItems.IsValidIndex(FirstConnectedItemIndex)))
+	{
+		PinComboBox->SetSelectedItem(PinComboBoxItems[FirstConnectedItemIndex]);
+	}
+}
+
+void SPCGEditorGraphAttributeListView::RefreshDataComboBox()
+{
+	DataComboBoxItems.Empty();
+	DataComboBox->ClearSelection();
+	DataComboBox->RefreshOptions();
+
+	const FPCGDataCollection* InspectionData = GetInspectionData();
+	if (!InspectionData)
+	{
+		return;
+	}
+
+	for(int32 TaggedDataIndex = 0; TaggedDataIndex < InspectionData->TaggedData.Num(); ++TaggedDataIndex)
+	{
+		const FPCGTaggedData& TaggedData = InspectionData->TaggedData[TaggedDataIndex];
+
+		FString ItemName = FString::Format(
+			TEXT("[{0}] {1}"),
+			{ FText::AsNumber(TaggedDataIndex).ToString(), (TaggedData.Data ? TaggedData.Data->GetClass()->GetDisplayNameText().ToString(): TEXT("No Data")) });
+
+		if (!TaggedData.Tags.IsEmpty())
+		{
+			ItemName.Append(FString::Format(TEXT(": ({0})"), { FString::Join(TaggedData.Tags, TEXT(", ")) }));
+		}
+
+		DataComboBoxItems.Add(MakeShared<FString>(ItemName));
+	}
+
+	if (DataComboBoxItems.Num() > 0)
+	{
+		DataComboBox->SetSelectedItem(DataComboBoxItems[0]);
+	}
+}
+
+void SPCGEditorGraphAttributeListView::LaunchUpdateTask()
+{
+	// Discarding any currently running updater, task will still run and keep the old object alive but we wont care about the result.
+	// This is done because we cant afford to wait for task completion before starting a new task.
+	CurrentUpdateTask.Reset();
+	CurrentUpdateTask = MakeShared<FPCGListViewUpdater>(ListViewItems,
+		PCGColumnData,
+		SortMode,
+		SortingColumn,
+		TextFilter);
+	CurrentUpdateTask->Launch();
+}
+
+const FSlateBrush* SPCGEditorGraphAttributeListView::GetFilterBadgeIcon() const
+{
+	for (const SHeaderRow::FColumn& Column : ListViewHeader->GetColumns())
+	{
+		if (!Column.bIsVisible)
+		{
+			return FAppStyle::Get().GetBrush("Icons.BadgeModified");
+		}
+	}
+
+	return nullptr;
+}
+
+TSharedRef<SWidget> SPCGEditorGraphAttributeListView::OnGenerateFilterMenu()
+{
+	FMenuBuilder MenuBuilder(false, nullptr);
+
+	MenuBuilder.AddMenuEntry(
+		LOCTEXT("ToggleAllAttributes", "Toggle All"),
+		LOCTEXT("ToggleAllAttributesTooltip", "Toggle visibility for all attributes"),
+		FSlateIcon(),
+		FUIAction(FExecuteAction::CreateSP(this, &SPCGEditorGraphAttributeListView::ToggleAllAttributes),
+			FCanExecuteAction(),
+			FGetActionCheckState::CreateSP(this, &SPCGEditorGraphAttributeListView::GetAnyAttributeEnabledState)),
+		NAME_None,
+		EUserInterfaceActionType::ToggleButton);
+
+	MenuBuilder.AddSeparator();
+
+	const TIndirectArray<SHeaderRow::FColumn>& Columns = ListViewHeader->GetColumns();
+	TArray<FName> HiddenColumns = ListViewHeader->GetHiddenColumnIds();
+
+	for (const SHeaderRow::FColumn& Column : Columns)
+	{
+		MenuBuilder.AddMenuEntry(
+			Column.DefaultText,
+			Column.DefaultTooltip,
+			FSlateIcon(),
+			FUIAction(FExecuteAction::CreateSP(this, &SPCGEditorGraphAttributeListView::ToggleAttribute, Column.ColumnId),
+				FCanExecuteAction(),
+				FIsActionChecked::CreateSP(this, &SPCGEditorGraphAttributeListView::IsAttributeEnabled, Column.ColumnId)),
+			NAME_None,
+			EUserInterfaceActionType::ToggleButton);
+	}
+
+	return MenuBuilder.MakeWidget();
+}
+
+TSharedRef<SWidget> SPCGEditorGraphAttributeListView::OnGenerateAdditionalOperationsMenu()
+{
+	FMenuBuilder MenuBuilder(false, nullptr);
+
+	MenuBuilder.AddMenuEntry(
+		LOCTEXT("SaveThisData", "Save this data"),
+		LOCTEXT("SaveThisDataTooltip", "Saves this data to a PCG Data Asset."),
+		FSlateIcon(FAppStyle::GetAppStyleSetName(), "Icons.Download"),
+		FUIAction(FExecuteAction::CreateSP(this, &SPCGEditorGraphAttributeListView::SaveData, true, true),
+			FCanExecuteAction::CreateSP(this, &SPCGEditorGraphAttributeListView::CanSaveData, true, true)),
+		NAME_None,
+		EUserInterfaceActionType::Button);
+
+	MenuBuilder.AddMenuEntry(
+		LOCTEXT("SaveThisPinData", "Save pin data"),
+		LOCTEXT("SaveThisPinDataTooltip", "Saves all data from the selected pin to a PCG Data Asset."),
+		FSlateIcon(FAppStyle::GetAppStyleSetName(), "Icons.Download"),
+		FUIAction(FExecuteAction::CreateSP(this, &SPCGEditorGraphAttributeListView::SaveData, true, false),
+			FCanExecuteAction::CreateSP(this, &SPCGEditorGraphAttributeListView::CanSaveData, true, false)),
+		NAME_None,
+		EUserInterfaceActionType::Button);
+
+	MenuBuilder.AddMenuEntry(
+		LOCTEXT("SaveAllData", "Save all"),
+		LOCTEXT("SaveAllDataTooltip", "Saves all the input or output data to a PCG Data Asset."),
+		FSlateIcon(FAppStyle::GetAppStyleSetName(), "Icons.Download"),
+		FUIAction(FExecuteAction::CreateSP(this, &SPCGEditorGraphAttributeListView::SaveData, false, false),
+			FCanExecuteAction::CreateSP(this, &SPCGEditorGraphAttributeListView::CanSaveData, false, false)),
+		NAME_None,
+		EUserInterfaceActionType::Button);
+
+	MenuBuilder.AddSeparator();
+
+	MenuBuilder.AddMenuEntry(
+		LOCTEXT("ResetDefaultWidth", "Reset columns to default width"),
+		LOCTEXT("ResetDefaultWidthTooltip", "Resets all the columns width to their default."),
+		FSlateIcon(),
+		FUIAction(FExecuteAction::CreateSP(this, &SPCGEditorGraphAttributeListView::ResetColumnsWidthToDefault)),
+		NAME_None,
+		EUserInterfaceActionType::Button);
+	
+	MenuBuilder.AddMenuEntry(
+		LOCTEXT("ExpendMaxWidth", "Expend columns to max width"),
+		LOCTEXT("ExpendMaxWidthTooltip", "Expends all the columns width to their maximum to see their content."),
+		FSlateIcon(),
+		FUIAction(FExecuteAction::CreateSP(this, &SPCGEditorGraphAttributeListView::ExpendAllColumnToMaxWidth)),
+		NAME_None,
+		EUserInterfaceActionType::Button);
+
+	return MenuBuilder.MakeWidget();
+}
+
+FPCGDataCollection SPCGEditorGraphAttributeListView::BuildDataCollectionForSave(bool bUsePinComboIndex, bool bUseDataComboIndex) const
+{
+	// Can't save N'th data on all pins, as it makes no sense
+	check(!bUseDataComboIndex || bUsePinComboIndex);
+
+	// Build a data collection from the indices we're given.
+	FPCGDataCollection Collection;
+
+	if (!bUsePinComboIndex)
+	{
+		for (const auto& PinComboBoxItem : PinComboBoxItems)
+		{
+			const FPCGDataCollection* PinInspectionData = GetInspectionData(PinComboBoxItem);
+			if (PinInspectionData)
+			{
+				Collection.TaggedData.Append(PinInspectionData->TaggedData);
+			}
+		}
+	}
+	else
+	{
+		if (const FPCGDataCollection* PinInspectionData = GetInspectionData())
+		{
+			if (!bUseDataComboIndex)
+			{
+				Collection = *PinInspectionData;
+			}
+			else
+			{
+
+				Collection.TaggedData.Add(PinInspectionData->TaggedData[GetSelectedDataIndex()]);
+			}
+		}
+	}
+
+	return Collection;
+}
+
+void SPCGEditorGraphAttributeListView::SaveData(bool bUsePinComboIndex, bool bUseDataComboIndex)
+{
+	FPCGDataCollection Collection = BuildDataCollectionForSave(bUsePinComboIndex, bUseDataComboIndex);
+
+	if (Collection.TaggedData.IsEmpty())
+	{
+		return;
+	}
+
+	UPCGDataCollectionExporter* Exporter = NewObject<UPCGDataCollectionExporter>();
+	Exporter->Data = Collection;
+
+	FPCGAssetExporterParameters Parameters;
+	UPCGAssetExporterUtils::CreateAsset(Exporter, Parameters);
+}
+
+bool SPCGEditorGraphAttributeListView::CanSaveData(bool bUsePinComboIndex, bool bUseDataComboIndex) const
+{
+	FPCGDataCollection Collection = BuildDataCollectionForSave(bUsePinComboIndex, bUseDataComboIndex);
+	return !Collection.TaggedData.IsEmpty();
+}
+
+FText SPCGEditorGraphAttributeListView::OnGenerateSelectedPinText() const
+{
+	if (const TSharedPtr<FPinComboBoxItem> SelectedPin = PinComboBox->GetSelectedItem())
+	{
+		return FText::FromName(SelectedPin->Name);
+	}
+	else
+	{
+		return PCGEditorGraphAttributeListView::NoPinAvailableText;
+	}
+}
+
+void SPCGEditorGraphAttributeListView::OnSelectionChangedPin(TSharedPtr<FPinComboBoxItem> InItem, ESelectInfo::Type InSelectInfo)
+{
+	RefreshDataComboBox();
+
+	if (InSelectInfo != ESelectInfo::Direct)
+	{
+		RefreshAttributeList();
+	}
+}
+
+TSharedRef<SWidget> SPCGEditorGraphAttributeListView::OnGeneratePinWidget(TSharedPtr<FPinComboBoxItem> InItem) const
+{
+	return SNew(STextBlock).Text(FText::FromName(InItem.IsValid() ? InItem->Name : NAME_None));
+}
+
+TSharedRef<SWidget> SPCGEditorGraphAttributeListView::OnGenerateDataWidget(TSharedPtr<FString> InItem) const
+{
+	return SNew(STextBlock).Text(InItem.IsValid() ? FText::FromString(*InItem) : FText());
+}
+
+void SPCGEditorGraphAttributeListView::OnSelectionChanged(TSharedPtr<FString> Item, ESelectInfo::Type SelectInfo)
+{
+	if (SelectInfo != ESelectInfo::Direct)
+	{
+		RefreshAttributeList();
+	}
+}
+
+FText SPCGEditorGraphAttributeListView::OnGenerateSelectedDataText() const
+{
+	if (const TSharedPtr<FString> SelectedDataName = DataComboBox->GetSelectedItem())
+	{
+		return FText::FromString(*SelectedDataName);
+	}
+	else
+	{
+		return PCGEditorGraphAttributeListView::NoDataAvailableText;
+	}
+}
+
+int32 SPCGEditorGraphAttributeListView::GetSelectedDataIndex() const
+{
+	int32 Index = INDEX_NONE;
+	if (const TSharedPtr<FString> SelectedItem = DataComboBox->GetSelectedItem())
+	{
+		DataComboBoxItems.Find(SelectedItem, Index);
+	}
+
+	return Index;
+}
+
+void SPCGEditorGraphAttributeListView::ToggleAllAttributes()
+{
+	const TArray<FName> HiddenColumns = ListViewHeader->GetHiddenColumnIds();
+	if (HiddenColumns.Num() > 0)
+	{
+		for (const FName& HiddenColumn : HiddenColumns)
+		{
+			ListViewHeader->SetShowGeneratedColumn(HiddenColumn, /*InShow=*/true);
+		}
+	}
+	else
+	{
+		const TIndirectArray<SHeaderRow::FColumn>& Columns = ListViewHeader->GetColumns();
+		for (const SHeaderRow::FColumn& Column : Columns)
+		{
+			ListViewHeader->SetShowGeneratedColumn(Column.ColumnId, /*InShow=*/false);
+		}
+	}
+}
+
+void SPCGEditorGraphAttributeListView::ToggleAttribute(FName InAttributeName)
+{
+	ListViewHeader->SetShowGeneratedColumn(InAttributeName, !ListViewHeader->IsColumnVisible(InAttributeName));
+}
+
+ECheckBoxState SPCGEditorGraphAttributeListView::GetAnyAttributeEnabledState() const
+{
+	bool bAllEnabled = true;
+	bool bAnyEnabled = false;
+
+	for (const SHeaderRow::FColumn& Column : ListViewHeader->GetColumns())
+	{
+		bAllEnabled &= Column.bIsVisible;
+		bAnyEnabled |= Column.bIsVisible;
+	}
+
+	if (bAllEnabled)
+	{
+		return ECheckBoxState::Checked;
+	}
+	else if (bAnyEnabled)
+	{
+		return ECheckBoxState::Undetermined;
+	}
+	else
+	{
+		return ECheckBoxState::Unchecked;
+	}
+}
+
+bool SPCGEditorGraphAttributeListView::IsAttributeEnabled(FName InAttributeName) const
+{
+	return ListViewHeader->IsColumnVisible(InAttributeName);
+}
+
+TSharedRef<ITableRow> SPCGEditorGraphAttributeListView::OnGenerateRow(PCGListViewItemPtr Item, const TSharedRef<STableViewBase>& OwnerTable)
+{
+	return SNew(SPCGListViewItemRow, OwnerTable)
+		.AttributeListView(SharedThis(this))
+		.ListViewItem(Item);
+}
+
+void SPCGEditorGraphAttributeListView::OnItemDoubleClicked(PCGListViewItemPtr Item) const
+{
+	check(Item);
+	if (FocusOnDataCallback && DataPtr.Get())
+	{
+		FocusOnDataCallback(DataPtr.Get(), { Item->Index });
+	}
+}
+
+TSharedPtr<SWidget> SPCGEditorGraphAttributeListView::OnItemsContextMenu()
+{
+	FMenuBuilder MenuBuilder(/*bInShouldCloseWindowAfterMenuSelection=*/true, /*InCommandList=*/nullptr);
+	MenuBuilder.AddMenuEntry(
+		LOCTEXT("ShowInViewport", "Zoom to selection"),
+		LOCTEXT("ShowInViewport_Tooltip", "Frames the viewport so that all selected entries are visible."),
+		FSlateIcon(FPCGEditorStyle::Get().GetStyleSetName(), "PCG.Editor.ZoomToSelection"),
+		FUIAction(FExecuteAction::CreateSP(this, &SPCGEditorGraphAttributeListView::FocusOnSelection),
+			FCanExecuteAction::CreateSP(this, &SPCGEditorGraphAttributeListView::CanFocusOnSelection)),
+		NAME_None,
+		EUserInterfaceActionType::Button);
+
+	MenuBuilder.AddMenuEntry(
+		LOCTEXT("CopyToClipboard", "Copy to clipboard"),
+		LOCTEXT("CopyToClipboard_Tooltip", "Copies the contents of the entries to the clipboard."),
+		FSlateIcon(FPCGEditorStyle::Get().GetStyleSetName(), "PCG.Editor.CopyToClipboard"),
+		FUIAction(FExecuteAction::CreateSP(this, &SPCGEditorGraphAttributeListView::CopySelectionToClipboard),
+			FCanExecuteAction::CreateSP(this, &SPCGEditorGraphAttributeListView::CanCopySelectionToClipboard)),
+		NAME_None,
+		EUserInterfaceActionType::Button);
+
+	return MenuBuilder.MakeWidget();
+}
+
+void SPCGEditorGraphAttributeListView::OnColumnSortModeChanged(const EColumnSortPriority::Type InSortPriority, const FName& InColumnId, const EColumnSortMode::Type InNewSortMode)
+{
+	if (SortingColumn == InColumnId)
+	{
+		// Cycling
+		SortMode = static_cast<EColumnSortMode::Type>((SortMode + 1) % 3);
+	}
+	else
+	{
+		SortingColumn = InColumnId;
+		SortMode = InNewSortMode;
+	}
+
+	LaunchUpdateTask();
+}
+
+EColumnSortMode::Type SPCGEditorGraphAttributeListView::GetColumnSortMode(const FName InColumnId) const
+{
+	if (SortingColumn != InColumnId)
+	{
+		return EColumnSortMode::None;
+	}
+
+	return SortMode;
+}
+
+void SPCGEditorGraphAttributeListView::OnFilterTextChanged(const FText& InFilterText)
+{
+	ActiveFilterText = InFilterText;
+	TextFilter->SetFilterText(InFilterText);
+
+	const FText ErrorText = TextFilter->GetFilterErrorText();
+	if(ErrorText.IsEmpty())
+	{
+		LaunchUpdateTask();
+	}
+
+	SearchBoxWidget->SetError(ErrorText);
+}
+
+void SPCGEditorGraphAttributeListView::OnFilterTextCommitted(const FText& NewText, ETextCommit::Type CommitInfo)
+{
+	if (CommitInfo == ETextCommit::OnCleared)
+	{
+		SearchBoxWidget->SetText(FText::GetEmpty());
+		OnFilterTextChanged(FText::GetEmpty());
+	}
+}
+
+FReply SPCGEditorGraphAttributeListView::OnListViewKeyDown(const FGeometry& /*InGeometry*/, const FKeyEvent& InKeyEvent) const
+{
+	if (ListViewCommands->ProcessCommandBindings(InKeyEvent))
+	{
+		return FReply::Handled();
+	}
+
+	return FReply::Unhandled();
+}
+
+void SPCGEditorGraphAttributeListView::AddColumn(const UPCGData* InPCGData, const FPCGTableVisualizerColumnInfo& InColumnInfo, TSharedPtr<const IPCGAttributeAccessorKeys> InAccessorKeys)
+{
+	FPCGColumnData& ColumnData = PCGColumnData.Add(InColumnInfo.Id);
+	ColumnData.DataAccessor = InColumnInfo.Accessor;
+	ColumnData.DataKeys = InAccessorKeys;
+
+	const float ColumnWidth = InColumnInfo.Width < 0.0f ? PCGEditorGraphAttributeListView::CalculateColumnWidth(InColumnInfo.Label) : InColumnInfo.Width;
+	ColumnsMaxWidthMapping.FindOrAdd(InColumnInfo.Id) = ColumnWidth;
+
+	SHeaderRow::FColumn::FArguments Arguments;
+	Arguments.ColumnId(InColumnInfo.Id);
+	Arguments.DefaultLabel(InColumnInfo.Label);
+	Arguments.DefaultTooltip(InColumnInfo.Tooltip);
+	Arguments.ManualWidth(ColumnWidth);
+	Arguments.HAlignHeader(HAlign_Center);
+	Arguments.HAlignCell((EHorizontalAlignment)InColumnInfo.CellAlignment);
+	Arguments.SortMode(this, &SPCGEditorGraphAttributeListView::GetColumnSortMode, InColumnInfo.Id);
+	Arguments.OnSort(this, &SPCGEditorGraphAttributeListView::OnColumnSortModeChanged);
+	Arguments.OverflowPolicy(ETextOverflowPolicy::Ellipsis);
+	Arguments.HeaderComboVisibility(EHeaderComboVisibility::Never);
+	Arguments.MenuContent()
+	[
+		GenerateColumnMenu(InColumnInfo.Id)
+	];
+	Arguments.OnColumnSplitterDoubleClick_Lambda([WeakThis = AsWeak(), ColumnId = InColumnInfo.Id](const FGeometry&, const FPointerEvent&) -> FReply
+	{
+		if (SPCGEditorGraphAttributeListView* This = static_cast<SPCGEditorGraphAttributeListView*>(WeakThis.Pin().Get()))
+		{
+			This->ResizeColumnToMaxWidth(ColumnId);
+			return FReply::Handled();
+		}
+		
+		return FReply::Unhandled();
+	});
+
+	SHeaderRow::FColumn* NewColumn = new SHeaderRow::FColumn(Arguments);
+	NewColumn->bIsVisible = !HiddenAttributes.Contains(InColumnInfo.Id);
+	ListViewHeader->AddColumn(*NewColumn);
+}
+
+void SPCGEditorGraphAttributeListView::ResizeColumnToMaxWidth(const FName ColumnId)
+{
+	ListViewHeader->SetColumnWidth(ColumnId, ColumnsMaxWidthMapping[ColumnId]);
+}
+
+void SPCGEditorGraphAttributeListView::ResetColumnsWidthToDefault()
+{
+	if (!ListViewHeader.IsValid())
+	{
+		return;
+	}
+	
+	for (const SHeaderRow::FColumn& Column : ListViewHeader->GetColumns())
+	{
+		ListViewHeader->SetColumnWidth(Column.ColumnId, PCGEditorGraphAttributeListView::CalculateColumnWidth(Column.DefaultText.Get()));
+	}
+}
+
+void SPCGEditorGraphAttributeListView::ExpendAllColumnToMaxWidth()
+{
+	if (!ListViewHeader.IsValid())
+	{
+		return;
+	}
+	
+	for (const SHeaderRow::FColumn& Column : ListViewHeader->GetColumns())
+	{
+		ResizeColumnToMaxWidth(Column.ColumnId);
+	}
+}
+
+void SPCGEditorGraphAttributeListView::CacheColumnWidthVisibility()
+{
+	if (!PCGEditorGraphNode.IsValid() || !ListViewHeader.IsValid() || ListViewHeader->GetColumns().IsEmpty())
+	{
+		return;
+	}
+	
+	const TObjectKey<UPCGEditorGraphNodeBase> NodeKey(PCGEditorGraphNode.Get());
+	
+	FNodeKeyToColumnWidthVisibilityMap* It = Algo::FindBy(ColumnWidthVisibilityCache, NodeKey, [](const FNodeKeyToColumnWidthVisibilityMap& It)
+	{
+		return It.Get<0>();
+	});
+
+	if (!It)
+	{
+		FNodeKeyToColumnWidthVisibilityMap KeyAndCache;
+		KeyAndCache.Get<0>() = NodeKey;
+		It = &ColumnWidthVisibilityCache.Add_GetRef(std::move(KeyAndCache));
+	}
+
+	check(It);
+	for (const SHeaderRow::FColumn& Column : ListViewHeader->GetColumns())
+	{
+		TTuple<float, bool>& WidthVisiblity = It->Get<1>().FindOrAdd(Column.ColumnId);
+		WidthVisiblity.Get<0>() = Column.GetWidth();
+		WidthVisiblity.Get<1>() = Column.bIsVisible;
+	}
+}
+
+void SPCGEditorGraphAttributeListView::RestoreColumnWidthVisibility()
+{
+	if (!PCGEditorGraphNode.IsValid() || !ListViewHeader.IsValid() || ListViewHeader->GetColumns().IsEmpty())
+	{
+		return;
+	}
+	
+	const TObjectKey<UPCGEditorGraphNodeBase> NodeKey(PCGEditorGraphNode.Get());
+	
+	const FNodeKeyToColumnWidthVisibilityMap* It = Algo::FindBy(ColumnWidthVisibilityCache, NodeKey, [](const FNodeKeyToColumnWidthVisibilityMap& It)
+	{
+		return It.Get<0>();
+	});
+
+	if (!It)
+	{
+		return;
+	}
+
+	for (const SHeaderRow::FColumn& Column : ListViewHeader->GetColumns())
+	{
+		if (const TTuple<float, bool>* WidthVisibility = It->Get<1>().Find(Column.ColumnId))
+		{
+			ListViewHeader->SetColumnWidth(Column.ColumnId, WidthVisibility->Get<0>());
+			ListViewHeader->SetShowGeneratedColumn(Column.ColumnId, WidthVisibility->Get<1>());
+		}
+	}
+}
+
+TSharedRef<SWidget> SPCGEditorGraphAttributeListView::GenerateColumnMenu(FName ColumnId)
+{
+	FMenuBuilder MenuBuilder(/*bShouldCloseAfterMenuSelection=*/true, nullptr);
+	MenuBuilder.AddMenuEntry(
+		LOCTEXT("CopyAttributeNameToClipboard", "Copy attribute name"),
+		LOCTEXT("CopyAttributeNameToClipboardTooltip", "Copies the attribute name to the clipboard."),
+		FSlateIcon(),
+		FUIAction(FExecuteAction::CreateLambda([ColumnId]()
+		{
+			FString ColumnIdString = ColumnId.ToString();
+
+			// TODO[UE-221219]: until we support @None as an actual valid token, we need to replace it with None.
+			ColumnIdString = ColumnIdString.Replace(TEXT("@None"), TEXT("None"));
+
+			FPlatformApplicationMisc::ClipboardCopy(*ColumnIdString);
+		})),
+		NAME_None,
+		EUserInterfaceActionType::Button);
+
+	MenuBuilder.AddMenuEntry(
+		LOCTEXT("ResizeColumnToMaxWidth", "Resize Column width to match content"),
+		LOCTEXT("ResizeColumnToMaxWidthTooltip", "Resizes column width to match the content width."),
+		FSlateIcon(),
+		FUIAction(FExecuteAction::CreateLambda([WeakThis = AsWeak(), ColumnId]()
+		{
+			if (SPCGEditorGraphAttributeListView* This = static_cast<SPCGEditorGraphAttributeListView*>(WeakThis.Pin().Get()))
+			{
+				This->ResizeColumnToMaxWidth(ColumnId);
+			}
+		})),
+		NAME_None,
+		EUserInterfaceActionType::Button);
+
+	return MenuBuilder.MakeWidget();
+}
+
+void SPCGEditorGraphAttributeListView::CopySelectionToClipboard() const
+{
+	constexpr TCHAR Delimiter = TEXT(',');
+	constexpr TCHAR LineEnd = TEXT('\n');
+
+	const TArray<FName> HiddenColumnIds = ListViewHeader->GetHiddenColumnIds();
+
+	TArray<TPair<FName, FPCGColumnData>> FilteredPCGColumnData;
+	for (const TPair<FName, FPCGColumnData>& Data : PCGColumnData)
+	{
+		if (!HiddenColumnIds.Contains(Data.Key))
+		{
+			FilteredPCGColumnData.Add(Data);
+		}
+	}
+
+	TStringBuilder<2048> CSVExport;
+
+	// Write column header row
+	for (int ColumnIndex = 0; ColumnIndex < FilteredPCGColumnData.Num(); ColumnIndex++)
+	{
+		const TPair<FName, FPCGColumnData>& Data = FilteredPCGColumnData[ColumnIndex];
+
+		if (ColumnIndex > 0)
+		{
+			CSVExport += Delimiter;
+		}
+		CSVExport += Data.Key.ToString();
+	}
+
+	// Gather selected rows and sort them to match the displayed order instead of selection order
+	TArray<PCGListViewItemPtr> SelectedListViewItems = ListView->GetSelectedItems();
+	if (const FPCGColumnData* ColumnData = PCGColumnData.Find(SortingColumn))
+	{
+		if (ColumnData->DataAccessor.IsValid() && ColumnData->DataKeys.IsValid())
+		{
+			// Lambda used here to get the index value of an item in the array for sorting
+			PCGAttributeAccessorHelpers::SortByAttribute(*ColumnData->DataAccessor, *ColumnData->DataKeys, SelectedListViewItems, !(SortMode & EColumnSortMode::Descending), [&SelectedListViewItems](int Index) { return SelectedListViewItems[Index]->Index; });
+		}
+	}
+
+	// Write each row
+	for (const PCGListViewItemPtr&  ListViewItem : SelectedListViewItems)
+	{
+		CSVExport += LineEnd;
+
+		for (int ColumnIndex = 0; ColumnIndex < FilteredPCGColumnData.Num(); ColumnIndex++)
+		{
+			if (ColumnIndex > 0)
+			{
+				CSVExport += Delimiter;
+			}
+
+			const TPair<FName, FPCGColumnData>& Data = FilteredPCGColumnData[ColumnIndex];
+			const FPCGColumnData& ColumnData = Data.Value;
+			if (ColumnData.DataAccessor.IsValid() && ColumnData.DataKeys.IsValid())
+			{
+				FString RowString;
+				if (ColumnData.DataAccessor->Get<FString>(RowString, ListViewItem->Index, *ColumnData.DataKeys, EPCGAttributeAccessorFlags::AllowBroadcast))
+				{
+					CSVExport += RowString;
+				}
+			}
+		}
+	}
+
+	FPlatformApplicationMisc::ClipboardCopy(*CSVExport);
+}
+
+bool SPCGEditorGraphAttributeListView::CanCopySelectionToClipboard() const
+{
+	return ListView->GetNumItemsSelected() > 0;
+}
+
+const FSlateBrush* SPCGEditorGraphAttributeListView::OnGetLockButtonImageResource() const
+{
+	return FAppStyle::GetBrush(bIsLocked ? TEXT("PropertyWindow.Locked") : TEXT("PropertyWindow.Unlocked"));
+}
+
+FReply SPCGEditorGraphAttributeListView::OnLockClick()
+{
+	bIsLocked = !bIsLocked;
+	return FReply::Handled();
+}
+
+FReply SPCGEditorGraphAttributeListView::OnNodeNameClicked()
+{
+	if (PCGEditorGraphNode.Get() && PCGEditorPtr.IsValid())
+	{
+		PCGEditorPtr.Pin()->JumpToNode(PCGEditorGraphNode.Get());
+	}
+
+	return FReply::Handled();
+}
+
+FReply SPCGEditorGraphAttributeListView::OnFocusOnDataClicked() const
+{
+	if (IsFocusOnDataEnabled() && DataPtr.Get())
+	{
+		FocusOnDataCallback(DataPtr.Get(), {});
+	}
+
+	return FReply::Handled();
+}
+
+bool SPCGEditorGraphAttributeListView::IsFocusOnDataEnabled() const
+{
+	return (!ListViewItems.IsEmpty() && FocusOnDataCallback);
+}
+
+void SPCGEditorGraphAttributeListView::FocusOnSelection() const
+{
+	if (!IsFocusOnDataEnabled() || !DataPtr.Get())
+	{
+		return;
+	}
+
+	// Note: this implementation assumes it's the same callback for all entries, which is currently true
+	const TArray<PCGListViewItemPtr> SelectedItems = ListView->GetSelectedItems();
+	TArray<int> Indices;
+	Indices.Reserve(SelectedItems.Num());
+
+	for (PCGListViewItemPtr SelectedItem : SelectedItems)
+	{
+		if (SelectedItem)
+		{
+			Indices.Add(SelectedItem->Index);
+		}
+	}
+
+	FocusOnDataCallback(DataPtr.Get(), Indices);
+}
+
+bool SPCGEditorGraphAttributeListView::CanFocusOnSelection() const
+{
+	return ListView->GetNumItemsSelected() > 0;
+}
+
+#undef LOCTEXT_NAMESPACE
diff --git a/Engine/Plugins/PCG/Source/PCGEditor/Private/Widgets/SPCGEditorGraphAttributeListView.h b/Engine/Plugins/PCG/Source/PCGEditor/Private/Widgets/SPCGEditorGraphAttributeListView.h
new file mode 100644
--- /dev/null
+++ b/Engine/Plugins/PCG/Source/PCGEditor/Private/Widgets/SPCGEditorGraphAttributeListView.h
@@ -0,0 +1,300 @@
+// Copyright Epic Games, Inc. All Rights Reserved.
+
+#pragma once
+
+#include "Data/PCGPointData.h"
+#include "Graph/PCGStackContext.h"
+#include "Metadata/Accessors/IPCGAttributeAccessor.h"
+#include "Metadata/Accessors/PCGAttributeAccessorKeys.h"
+
+#include "Containers/RingBuffer.h"
+#include "Misc/TextFilterExpressionEvaluator.h"
+#include "Tasks/Task.h"
+#include "UObject/ObjectKey.h"
+#include "Widgets/Views/SListView.h"
+#include "Widgets/Views/STableRow.h"
+
+class STableViewBase;
+class STextBlock;
+template <typename OptionType> class SComboBox;
+
+class FPCGEditor;
+class FPCGMetadataAttributeBase;
+class FUICommandList;
+class SPCGEditorGraphAttributeListView;
+class UPCGComponent;
+class UPCGData;
+class UPCGEditorGraphNodeBase;
+class UPCGMetadata;
+class UPCGParamData;
+class UPCGPointData;
+struct FPCGTableVisualizerColumnInfo;
+struct FPCGDataCollection;
+struct FPCGPoint;
+enum class EPCGMetadataTypes : uint8;
+
+class SComboButton;
+class SHeaderRow;
+class SSearchBox;
+struct FSlateBrush;
+
+namespace PCGEditorGraphAttributeListView
+{
+	constexpr float MaxColumnWidth = 200.0f;
+	constexpr int32 MaxNodeColumnWidthCachedItems = 256;
+}
+
+struct FPCGListViewItem
+{
+	int32 Index = INDEX_NONE;
+};
+
+struct FPCGColumnData
+{
+	TSharedPtr<const IPCGAttributeAccessor> DataAccessor;
+	TSharedPtr<const IPCGAttributeAccessorKeys> DataKeys;
+};
+
+typedef TSharedPtr<FPCGListViewItem> PCGListViewItemPtr;
+
+template <typename T, typename = void>
+struct FTextAsNumberIsValid : std::false_type {};
+
+/** Utility to see if a value type is supported by FText::AsNumber */
+template <typename T>
+struct FTextAsNumberIsValid<T, std::void_t<decltype(FText::AsNumber(std::declval<T>()))>> : std::true_type {};
+
+/** Class used for threaded filtering and sorting of list view items */
+class FPCGListViewUpdater : public TSharedFromThis<FPCGListViewUpdater>
+{
+public:
+	FPCGListViewUpdater(
+		const TArray<PCGListViewItemPtr>& InListViewItems,
+		const TMap<FName, FPCGColumnData>& InColumnData,
+		const EColumnSortMode::Type InSortMode,
+		const FName InSortingColumn,
+		const TSharedPtr<FTextFilterExpressionEvaluator>& InTextFilter)
+	: ListViewItems(InListViewItems)
+	, ColumnData(InColumnData)
+	, SortMode(InSortMode)
+	, SortingColumn(InSortingColumn)
+	, TextFilter(InTextFilter)
+	{}
+
+	bool IsCompleted() const;
+	void Launch();
+
+	TArray<PCGListViewItemPtr> ListViewItems;
+
+private:
+	void AsyncSort();
+	void AsyncFilter();
+
+	TMap<FName, FPCGColumnData> ColumnData;
+
+	EColumnSortMode::Type SortMode = EColumnSortMode::Type::Ascending;
+	FName SortingColumn = NAME_None;
+
+	TSharedPtr<FTextFilterExpressionEvaluator> TextFilter;
+
+	UE::Tasks::FTask UpdateTask;
+};
+
+class SPCGListViewItemRow : public SMultiColumnTableRow<PCGListViewItemPtr>
+{
+public:
+	SLATE_BEGIN_ARGS(SPCGListViewItemRow) {}
+	SLATE_ARGUMENT(TSharedPtr<SPCGEditorGraphAttributeListView>, AttributeListView)
+	SLATE_ARGUMENT(PCGListViewItemPtr, ListViewItem)
+	SLATE_END_ARGS()
+
+	void Construct(const FArguments& InArgs, const TSharedRef<STableViewBase>& InOwnerTableView);
+
+	virtual TSharedRef<SWidget> GenerateWidgetForColumn(const FName& ColumnId) override;
+
+private:
+	TWeakPtr<SPCGEditorGraphAttributeListView> AttributeListView;
+	PCGListViewItemPtr InternalItem;
+};
+
+class FPCGPointFilterExpressionContext : public ITextFilterExpressionContext
+{
+public:
+	explicit FPCGPointFilterExpressionContext(const FPCGListViewItem* InRowItem, const TMap<FName, FPCGColumnData>* InPCGColumnData);
+
+	// ~Begin ITextFilterExpressionContext interface
+	virtual bool TestBasicStringExpression(const FTextFilterString& InValue, const ETextFilterTextComparisonMode InTextComparisonMode) const override;
+	virtual bool TestComplexExpression(const FName& InKey, const FTextFilterString& InValue, const ETextFilterComparisonOperation InComparisonOperation, const ETextFilterTextComparisonMode InTextComparisonMode) const override;
+	// ~End ITextFilterExpressionContext interface
+
+private:
+	const FPCGListViewItem* RowItem = nullptr;
+	const TMap<FName, FPCGColumnData>* PCGColumnData = nullptr;
+};
+
+class SPCGEditorGraphAttributeListView : public SCompoundWidget
+{
+	friend SPCGListViewItemRow;
+
+public:
+	SLATE_BEGIN_ARGS(SPCGEditorGraphAttributeListView) {}
+	SLATE_END_ARGS()
+
+	virtual ~SPCGEditorGraphAttributeListView();
+
+	void Construct(const FArguments& InArgs, TSharedPtr<FPCGEditor> InPCGEditor);
+
+	virtual void Tick(const FGeometry& AllottedGeometry, const double InCurrentTime, const float InDeltaTime) override;
+
+	void RequestRefresh() { bNeedsRefresh = true; }
+
+	UPCGEditorGraphNodeBase* GetNodeBeingInspected() const;
+	void SetNodeBeingInspected(UPCGEditorGraphNodeBase* InNode);
+
+	bool IsLocked() const { return bIsLocked; }
+	void SetIsLocked(bool bInIsLocked) { bIsLocked = bInIsLocked; }
+
+private:
+	TSharedRef<SHeaderRow> CreateHeaderRowWidget() const;
+
+	void OnInspectedStackChanged(const FPCGStack& InPCGStack);
+
+	void OnGenerateUpdated(UPCGComponent* InPCGComponent);
+
+	void ResizeColumnToMaxWidth(FName InColumnId);
+	void ResetColumnsWidthToDefault();
+	void ExpendAllColumnToMaxWidth();
+	void CacheColumnWidthVisibility();
+	void RestoreColumnWidthVisibility();
+
+	void RefreshAttributeList();
+	void RefreshPinComboBox();
+	void RefreshDataComboBox();
+
+	void LaunchUpdateTask();
+
+	/** Only connected input pins are added to combo box, so keep track of the node pin index for each item. */
+	struct FPinComboBoxItem
+	{
+		explicit FPinComboBoxItem(FName InName, int32 InPinIndex, bool bInIsOutputPin)
+			: Name(InName)
+			, PinIndex(InPinIndex)
+			, bIsOutputPin(bInIsOutputPin)
+		{}
+
+		FName Name;
+		int32 PinIndex = INDEX_NONE;
+		bool bIsOutputPin = true;
+	};
+
+	const FPCGDataCollection* GetInspectionData(const TSharedPtr<FPinComboBoxItem>& EditorPin) const;
+	const FPCGDataCollection* GetInspectionData() const;
+
+	FText OnGenerateSelectedPinText() const;
+	void OnSelectionChangedPin(TSharedPtr<FPinComboBoxItem> InItem, ESelectInfo::Type InSelectInfo);
+	TSharedRef<SWidget> OnGeneratePinWidget(TSharedPtr<FPinComboBoxItem> InItem) const;
+
+	const FSlateBrush* GetFilterBadgeIcon() const;
+	TSharedRef<SWidget> OnGenerateFilterMenu();
+	TSharedRef<SWidget> OnGenerateAdditionalOperationsMenu();
+	TSharedRef<SWidget> OnGenerateDataWidget(TSharedPtr<FString> InItem) const;
+	void OnSelectionChanged(TSharedPtr<FString> Item, ESelectInfo::Type SelectInfo);
+	FText OnGenerateSelectedDataText() const;
+	int32 GetSelectedDataIndex() const;
+
+	void ToggleAllAttributes();
+	void ToggleAttribute(FName InAttributeName);
+	ECheckBoxState GetAnyAttributeEnabledState() const;
+	bool IsAttributeEnabled(FName InAttributeName) const;
+
+	FPCGDataCollection BuildDataCollectionForSave(bool bUsePinComboIndex, bool bUseDataComboIndex) const;
+	void SaveData(bool bUsePinIndex, bool bUseDataIndex);
+	bool CanSaveData(bool bUsePinIndex, bool bUseDataIndex) const;
+
+	TSharedRef<ITableRow> OnGenerateRow(PCGListViewItemPtr Item, const TSharedRef<STableViewBase>& OwnerTable);
+	void OnItemDoubleClicked(PCGListViewItemPtr Item) const;
+	TSharedPtr<SWidget> OnItemsContextMenu();
+
+	void OnColumnSortModeChanged(const EColumnSortPriority::Type SortPriority, const FName& ColumnId, const EColumnSortMode::Type InSortMode);
+	EColumnSortMode::Type GetColumnSortMode(const FName InColumnId) const;
+	TSharedRef<SWidget> GenerateColumnMenu(FName ColumnId);
+
+	void OnFilterTextChanged(const FText& InFilterText);
+	void OnFilterTextCommitted(const FText& NewText, ETextCommit::Type CommitInfo);
+
+	FReply OnListViewKeyDown(const FGeometry& InGeometry, const FKeyEvent& InKeyEvent) const;
+
+	void AddColumn(const UPCGData* InPCGData, const FPCGTableVisualizerColumnInfo& InColumnInfo, TSharedPtr<const IPCGAttributeAccessorKeys> InAccessorKeys);
+	void AddPointDataColumns(const UPCGPointData* InPCGPointData);
+	void AddMetadataColumn(const UPCGData* InPCGData, const FName& InColumnId, EPCGMetadataTypes InMetadataType, const TCHAR* PostFix = nullptr);
+
+	void CopySelectionToClipboard() const;
+	bool CanCopySelectionToClipboard() const;
+
+	/** @return the Slate brush to use for the lock image */
+	const FSlateBrush* OnGetLockButtonImageResource() const;
+
+	FReply OnLockClick();
+	FReply OnNodeNameClicked();
+
+	FReply OnFocusOnDataClicked() const;
+	bool IsFocusOnDataEnabled() const;
+
+	void FocusOnSelection() const;
+	bool CanFocusOnSelection() const;
+
+	/** Returns the current PCG component */
+	TWeakObjectPtr<UPCGComponent> GetPCGComponent() const;
+
+	/** Pointer back to the PCG editor that owns us */
+	TWeakPtr<FPCGEditor> PCGEditorPtr;
+
+	/** Cached PCGGraphNode being viewed */
+	TWeakObjectPtr<UPCGEditorGraphNodeBase> PCGEditorGraphNode;
+
+	TSharedPtr<FUICommandList> ListViewCommands;
+
+	TSharedPtr<FTextFilterExpressionEvaluator> TextFilter;
+
+	TSharedPtr<SSearchBox> SearchBoxWidget;
+	TSharedPtr<SHeaderRow> ListViewHeader;
+	TSharedPtr<SListView<PCGListViewItemPtr>> ListView;
+	TArray<PCGListViewItemPtr> ListViewItems;
+	TArray<PCGListViewItemPtr> FilteredListViewItems;
+
+	TSharedPtr<SComboBox<TSharedPtr<FPinComboBoxItem>>> PinComboBox;
+	TArray<TSharedPtr<FPinComboBoxItem>> PinComboBoxItems;
+
+	TSharedPtr<SComboBox<TSharedPtr<FString>>> DataComboBox;
+	TArray<TSharedPtr<FString>> DataComboBoxItems;
+
+	TSharedPtr<STextBlock> NodeNameTextBlock;
+	TSharedPtr<STextBlock> InfoTextBlock;
+
+	TArray<FName> HiddenAttributes;
+
+	TMap<FName, FPCGColumnData> PCGColumnData;
+	
+	/** Contains info of the max width to display the full content of the column. Computed only for string types. */
+	TMap<FName, float> ColumnsMaxWidthMapping;
+
+	/** Ring buffer to keep the latest inspected nodes column width and visibility. */
+	using FNodeKeyToColumnWidthVisibilityMap = TTuple<TObjectKey<UPCGEditorGraphNodeBase>, TMap<FName, TTuple<float, bool>>>;
+	TRingBuffer<FNodeKeyToColumnWidthVisibilityMap> ColumnWidthVisibilityCache;
+
+	FText ActiveFilterText;
+
+	FName SortingColumn = NAME_None;
+	EColumnSortMode::Type SortMode = EColumnSortMode::Type::Ascending;
+	TFunction<void(const UPCGData*, TArrayView<const int>)> FocusOnDataCallback;
+
+	bool bNeedsRefresh : 1 = false;
+
+	/** True if this property view is currently locked (i.e. The objects being observed are not changed automatically due to user selection)*/
+	bool bIsLocked : 1 = false;
+
+	TSharedPtr<FPCGListViewUpdater> CurrentUpdateTask = nullptr;
+
+	/** Used to ensure data collapsed for inspection is kept alive. */
+	TWeakObjectPtr<const UPCGData> DataPtr = nullptr;
+};
diff --git a/Engine/Plugins/PCG/Source/PCGEditor/Private/Widgets/SPCGEditorGraphDebugObjectTree.cpp b/Engine/Plugins/PCG/Source/PCGEditor/Private/Widgets/SPCGEditorGraphDebugObjectTree.cpp
new file mode 100644
--- /dev/null
+++ b/Engine/Plugins/PCG/Source/PCGEditor/Private/Widgets/SPCGEditorGraphDebugObjectTree.cpp
@@ -0,0 +1,1647 @@
+// Copyright Epic Games, Inc. All Rights Reserved.
+
+#include "SPCGEditorGraphDebugObjectTree.h"
+
+#include "PCGComponent.h"
+#include "PCGGraph.h"
+#include "PCGSubsystem.h"
+#include "Grid/PCGPartitionActor.h"
+#include "Helpers/PCGHelpers.h"
+
+#include "PCGEditor.h"
+#include "PCGEditorGraph.h"
+#include "PCGEditorStyle.h"
+#include "AssetDefinitions/AssetDefinition_PCGGraphInterface.h"
+
+#include "PropertyCustomizationHelpers.h"
+#include "Selection.h"
+#include "UnrealEdGlobals.h"
+#include "Algo/AnyOf.h"
+#include "Editor/UnrealEdEngine.h"
+#include "Styling/SlateIconFinder.h"
+#include "Widgets/Images/SImage.h"
+#include "Widgets/Layout/SScrollBar.h"
+#include "Widgets/Layout/SScrollBox.h"
+
+#define LOCTEXT_NAMESPACE "SPCGEditorGraphDebugObjectTree"
+
+namespace PCGEditorGraphDebugObjectTree
+{
+	void GetRowIconState(FPCGEditorGraphDebugObjectItemPtr Item, const FSlateBrush*& OutBrush, FLinearColor& OutColorAndOpacity)
+	{
+		OutBrush = Item->GetIcon();
+
+		const FLinearColor DefaultGraphColor = CastChecked<UAssetDefinition_PCGGraphInterface>(UAssetDefinition_PCGGraphInterface::StaticClass()->ClassDefaultObject)->GetAssetColor();
+
+		OutColorAndOpacity = Item->IsDebuggable() ? DefaultGraphColor : FLinearColor(1.0f, 1.0f, 1.0f, 0.5f);
+	}
+
+	void GetErrorInfo(FPCGEditorGraphDebugObjectItemPtr Item, EVisibility& OutIconVisibility, FText& OutIconTooltipText, FLinearColor& OutIconColorAndOpacity, FLinearColor& OutRowTextColorAndOpacity)
+	{
+		OutIconVisibility = EVisibility::Hidden;
+		OutIconTooltipText = FText();
+		OutIconColorAndOpacity = FLinearColor::White;
+		OutRowTextColorAndOpacity = FLinearColor::White;
+
+		if (!ensure(Item->GetPCGStack()))
+		{
+			return;
+		}
+
+		ELogVerbosity::Type MinLocalVerbosity = ELogVerbosity::All;
+		ELogVerbosity::Type MinVerbosity = ELogVerbosity::All;
+
+		if (const UPCGSubsystem* Subsystem = UPCGSubsystem::GetActiveEditorInstance())
+		{
+			FPCGPerNodeVisualLogs Logs;
+			Subsystem->GetNodeVisualLogs().ForAllMatchingLogs(*Item->GetPCGStack(), [&Item, &Logs, &MinLocalVerbosity](const FPCGStack& InStack, const FPCGPerNodeVisualLogs& InLogs)
+			{
+				// Need to verify this, and clarify this when it's a loop item.
+				// i.e. should get the graph + loop index and compare
+				if (Item->GetPCGStack()->GetGraphForCurrentFrame() == InStack.GetGraphForCurrentFrame())
+				{
+					for (const FPCGNodeLogEntry& LogEntry : InLogs)
+					{
+						MinLocalVerbosity = FMath::Min(MinLocalVerbosity, LogEntry.Verbosity);
+					}
+				}
+
+				Logs.Append(InLogs);
+				return true;
+			});
+
+			OutIconTooltipText = Subsystem->GetNodeVisualLogs().GetSummaryText(Logs, nullptr, &MinVerbosity);
+		}
+
+		OutIconVisibility = OutIconTooltipText.IsEmpty() ? EVisibility::Hidden : EVisibility::Visible;
+
+		constexpr FLinearColor WarningColor(1.0f, 0.75f, 0.0f, 0.9f);
+		constexpr FLinearColor ErrorColor(1.0f, 0.0f, 0.0f, 0.9f);
+
+		if (MinLocalVerbosity < ELogVerbosity::All)
+		{
+			OutRowTextColorAndOpacity = MinLocalVerbosity <= ELogVerbosity::Error ? ErrorColor : WarningColor;
+		}
+
+		if (MinVerbosity < ELogVerbosity::All)
+		{
+			OutIconColorAndOpacity = MinVerbosity <= ELogVerbosity::Error ? ErrorColor : WarningColor;
+		}
+	}
+
+	FString GetStringFromName(FName InName, bool bForSorting)
+	{
+		if (bForSorting)
+		{
+			return InName.GetPlainNameString() + FString::Printf(TEXT("%06d"), InName.GetNumber());
+		}
+		else
+		{
+			return InName.ToString();
+		}
+	}
+
+	FSlateFontInfo GetRowFont(FPCGEditorGraphDebugObjectItemPtr Item)
+	{
+		const bool bUseBold = Item->IsSelected();
+		const bool bUseItalic = Item->IsDynamic();
+
+		if (!bUseBold && !bUseItalic)
+		{
+			return FAppStyle::GetFontStyle("NormalFont");
+		}
+		else if (bUseBold && !bUseItalic)
+		{
+			return FAppStyle::GetFontStyle("NormalFontBold");
+		}
+		else if (!bUseBold && bUseItalic)
+		{
+			return FAppStyle::GetFontStyle("NormalFontItalic");
+		}
+		else
+		{
+			return FAppStyle::GetFontStyle("NormalFontBoldItalic");
+		}
+	}
+}
+
+void FPCGEditorGraphDebugObjectItem::AddChild(TSharedRef<FPCGEditorGraphDebugObjectItem> InChild)
+{
+	check(!Children.Contains(InChild));
+	InChild->Parent = AsShared();
+	Children.Add(MoveTemp(InChild));
+}
+
+const FSlateBrush* FPCGEditorGraphDebugObjectItem::GetIcon() const
+{
+	return FSlateIconFinder::FindIconBrushForClass(UPCGGraph::StaticClass());
+}
+
+const TSet<TSharedPtr<FPCGEditorGraphDebugObjectItem>>& FPCGEditorGraphDebugObjectItem::GetChildren() const
+{
+	return Children;
+}
+
+FPCGEditorGraphDebugObjectItemPtr FPCGEditorGraphDebugObjectItem::GetParent() const
+{
+	return Parent.Pin();
+}
+
+void FPCGEditorGraphDebugObjectItem::SortChildren(bool bIsAscending, bool bIsRecursive)
+{
+	Children.Sort([bIsAscending](const FPCGEditorGraphDebugObjectItemPtr& InLHS, const FPCGEditorGraphDebugObjectItemPtr& InRHS)
+	{
+		// If both items have an explicit sort priority like a loop index, this is the primary sort key.
+		const int32 IndexLHS = InLHS->GetSortPriority();
+		const int32 IndexRHS = InRHS->GetSortPriority();
+		if (IndexLHS != INDEX_NONE && IndexRHS != INDEX_NONE)
+		{
+			return (IndexLHS < IndexRHS) == bIsAscending;
+		}
+
+		// Next sort priority is presence or not of children. Items without children are shown first to reduce the possibility that
+		// a child item ends up displayed far away from its parent item when the tree is expanded.
+		const int HasChildrenLHS = InLHS->Children.Num() ? 1 : 0;
+		const int HasChildrenRHS = InRHS->Children.Num() ? 1 : 0;
+		if (HasChildrenLHS != HasChildrenRHS)
+		{
+			return (HasChildrenLHS < HasChildrenRHS) == bIsAscending;
+		}
+
+		// Otherwise fall back to alphanumeric order.
+		return (InLHS->GetLabel(/*bForSorting=*/true) < InRHS->GetLabel(/*bForSorting=*/true)) == bIsAscending;
+	});
+
+	if (bIsRecursive)
+	{
+		for (const FPCGEditorGraphDebugObjectItemPtr& Child : Children)
+		{
+			Child->SortChildren(bIsAscending, bIsRecursive);
+		}
+	}
+}
+
+void FPCGEditorGraphDebugObjectItem::OnObjectsReplaced(const TMap<UObject*, UObject*>& ReplacementMap)
+{
+	if (FPCGStack* Stack = GetMutablePCGStack())
+	{
+		Stack->ReplaceRoot(ReplacementMap);
+	}
+}
+
+FString FPCGEditorGraphDebugObjectItem_Actor::GetLabel(bool bForSorting) const
+{
+	return Actor.IsValid() ? Actor->GetActorNameOrLabel() : FString();
+}
+
+const FSlateBrush* FPCGEditorGraphDebugObjectItem_Actor::GetIcon() const
+{
+	const UObject* Object = GetObject();
+	return FSlateIconFinder::FindIconBrushForClass(Object ? Object->GetClass() : AActor::StaticClass());
+}
+
+FString FPCGEditorGraphDebugObjectItem_PCGComponent::GetLabel(bool bForSorting) const
+{
+	if (PCGComponent.IsValid() && PCGGraph.IsValid())
+	{
+		return PCGEditorGraphDebugObjectTree::GetStringFromName(PCGComponent->GetFName(), bForSorting) + FString(TEXT(" - ")) + PCGEditorGraphDebugObjectTree::GetStringFromName(PCGGraph->GetFName(), bForSorting);
+	}
+
+	return FString();
+}
+
+void FPCGEditorGraphDebugObjectItem_PCGComponent::OnObjectsReplaced(const TMap<UObject*, UObject*>& ReplacementMap)
+{
+	FPCGEditorGraphDebugObjectItem::OnObjectsReplaced(ReplacementMap);
+
+	if (UObject* NewComponent = ReplacementMap.FindRef(PCGComponent.GetEvenIfUnreachable()))
+	{
+		PCGComponent = Cast<UPCGComponent>(NewComponent);
+	}
+}
+
+FString FPCGEditorGraphDebugObjectItem_PCGSubgraph::GetLabel(bool bForSorting) const
+{
+	if (PCGNode.IsValid() && PCGGraph.IsValid())
+	{
+		if (PCGNode->HasAuthoredTitle())
+		{
+			return PCGEditorGraphDebugObjectTree::GetStringFromName(PCGGraph->GetFName(), bForSorting) + FString(TEXT(" - ")) + PCGNode->GetAuthoredTitleLine().ToString();
+		}
+		else
+		{
+			return PCGEditorGraphDebugObjectTree::GetStringFromName(PCGGraph->GetFName(), bForSorting) + FString(TEXT(" - ")) + PCGEditorGraphDebugObjectTree::GetStringFromName(PCGNode->GetFName(), bForSorting);
+		}
+	}
+
+	return FString();
+}
+
+const FSlateBrush* FPCGEditorGraphDebugObjectItem_PCGSubgraph::GetIcon() const
+{
+	if (SubgraphType == EPCGEditorSubgraphNodeType::LoopSubgraph)
+	{
+		return FAppStyle::Get().GetBrush("GraphEditor.Macro.Loop_16x");
+	}
+	else
+	{
+		return FPCGEditorGraphDebugObjectItem::GetIcon();
+	}
+}
+
+FString FPCGEditorGraphDebugObjectItem_PCGLoopIndex::GetLabel(bool bForSorting) const
+{
+	return FString::Format(TEXT("{0}"), { LoopIndex });
+}
+
+void SPCGEditorGraphDebugObjectItemRow::Construct(const FArguments& InArgs, const TSharedRef<STableViewBase>& InOwnerTableView, FPCGEditorGraphDebugObjectItemPtr InItem)
+{
+	Item = InItem;
+
+	// This function should auto-expand the row and select the deepest entry as the debug object if it is unambiguous (the only entry at its level in the tree).
+	OnDoubleClick = InArgs._OnDoubleClick;
+	OnJumpTo = InArgs._OnJumpTo;
+	OnFocus = InArgs._OnFocus;
+	CanJumpTo = InArgs._CanJumpTo;
+	CanFocus = InArgs._CanFocus;
+
+	// Computed once during construction as the tree is refreshed on relevant events.
+	const FSlateBrush* RowIcon = nullptr;
+	FLinearColor RowIconColorAndOpacity;
+	PCGEditorGraphDebugObjectTree::GetRowIconState(Item, RowIcon, RowIconColorAndOpacity);
+
+	// Icon indicating warnings and errors. Tree refreshes after execution so computing once here is sufficient.
+	FText ErrorIconTooltipText;
+	EVisibility ErrorIconVisibility;
+	FLinearColor ErrorIconColorAndOpacity;
+	FLinearColor RowTextColorAndOpacity;
+	PCGEditorGraphDebugObjectTree::GetErrorInfo(InItem, ErrorIconVisibility, ErrorIconTooltipText, ErrorIconColorAndOpacity, RowTextColorAndOpacity);
+
+	// Transfer warning/error color to icon if there is a local error, but preserve original opacity (otherwise it will look weird for non-local graphs).
+	if (ErrorIconVisibility != EVisibility::Hidden && RowTextColorAndOpacity != FLinearColor::White)
+	{
+		const float OriginalOpacity = RowIconColorAndOpacity.A;
+		RowIconColorAndOpacity = RowTextColorAndOpacity;
+		RowIconColorAndOpacity.A = OriginalOpacity;
+	}
+	
+	ChildSlot
+	[
+		SNew(SHorizontalBox)
+		+SHorizontalBox::Slot()
+		.AutoWidth()
+		.Padding(6.0f, 0.0f)
+		[
+			SNew(SImage)
+			.Visibility(EVisibility::HitTestInvisible)
+			.ColorAndOpacity(RowIconColorAndOpacity)
+			.Image(RowIcon)
+		]
+		+SHorizontalBox::Slot()
+		[
+			SNew(STextBlock)
+			.Text(FText::FromString(Item->GetLabel()))
+			.ToolTipText(FText::FromString(Item->GetLabel()))
+			.Font_Lambda([this]() { return PCGEditorGraphDebugObjectTree::GetRowFont(Item); })
+			// Highlight based on data available for currently inspected node. Computed dynamically in lambda to respond to inspection changes.
+			.ColorAndOpacity_Lambda([this, RowTextColorAndOpacity](){ return Item->IsGrayedOut() ? FColor(75, 75, 75) : RowTextColorAndOpacity; })
+			.OverflowPolicy(ETextOverflowPolicy::Ellipsis)
+		]
+		+SHorizontalBox::Slot()
+		.AutoWidth()
+		.Padding(2.0f, 0.0f)
+		[
+			SNew(SImage)
+			.Visibility(ErrorIconVisibility)
+			.ToolTipText(ErrorIconTooltipText)
+			.Image(FAppStyle::Get().GetBrush("Icons.Error"))
+			.ColorAndOpacity(ErrorIconColorAndOpacity)
+		]
+		+SHorizontalBox::Slot()
+		.AutoWidth()
+		.Padding(2.0f, 0.0f)
+		[
+			SNew(SButton)
+			.ButtonStyle(FAppStyle::Get(), "SimpleButton")
+			.Visibility_Lambda([this]() { return IsHovered() ? EVisibility::Visible : EVisibility::Hidden; })
+			.OnClicked(this, &SPCGEditorGraphDebugObjectItemRow::FocusClicked)
+			.IsEnabled(this, &SPCGEditorGraphDebugObjectItemRow::IsFocusEnabled)
+			.ToolTipText(LOCTEXT("FocusOnNode", "Show the calling node in the current graph."))
+			.ContentPadding(0.0f)
+			[
+				SNew(SImage)
+				.Image(FPCGEditorStyle::Get().GetBrush("PCG.Editor.ZoomToSelection"))
+				.ColorAndOpacity(FSlateColor::UseForeground())
+			]
+		]
+		+SHorizontalBox::Slot()
+		.AutoWidth()
+		.Padding(2.0f, 0.0f)
+		[
+			SNew(SButton)
+			.ButtonStyle(FAppStyle::Get(), "SimpleButton")
+			.Visibility_Lambda([this]() { return IsHovered() ? EVisibility::Visible : EVisibility::Hidden; })
+			.OnClicked(this, &SPCGEditorGraphDebugObjectItemRow::JumpToClicked)
+			.IsEnabled(this, &SPCGEditorGraphDebugObjectItemRow::IsJumpToEnabled)
+			.ToolTipText(LOCTEXT("JumpToGraph", "Jump to graph, with this debug object context."))
+			.ContentPadding(0.0f)
+			[
+				SNew(SImage)
+				.Image(FPCGEditorStyle::Get().GetBrush("PCG.Editor.JumpTo"))
+				.ColorAndOpacity(FSlateColor::UseForeground())
+			]
+		]
+	];
+}
+
+FReply SPCGEditorGraphDebugObjectItemRow::OnMouseButtonDoubleClick(const FGeometry& InMyGeometry, const FPointerEvent& InMouseEvent)
+{
+	if (Item.IsValid() && OnDoubleClick.IsBound())
+	{
+		OnDoubleClick.Execute(Item);
+	}
+
+	return FReply::Handled();
+}
+
+FReply SPCGEditorGraphDebugObjectItemRow::FocusClicked() const
+{
+	if (Item.IsValid() && OnFocus.IsBound())
+	{
+		OnFocus.Execute(Item);
+	}
+
+	return FReply::Handled();
+}
+
+FReply SPCGEditorGraphDebugObjectItemRow::JumpToClicked() const
+{
+	if (Item.IsValid() && OnJumpTo.IsBound())
+	{
+		OnJumpTo.Execute(Item);
+	}
+
+	return FReply::Handled();
+}
+
+bool SPCGEditorGraphDebugObjectItemRow::IsJumpToEnabled() const
+{
+	if (Item.IsValid() && CanJumpTo.IsBound())
+	{
+		return CanJumpTo.Execute(Item);
+	}
+	else
+	{
+		return true;
+	}
+}
+
+bool SPCGEditorGraphDebugObjectItemRow::IsFocusEnabled() const
+{
+	if (Item.IsValid() && CanFocus.IsBound())
+	{
+		return CanFocus.Execute(Item);
+	}
+	else
+	{
+		return true;
+	}
+}
+
+SPCGEditorGraphDebugObjectTree::~SPCGEditorGraphDebugObjectTree()
+{
+	FCoreUObjectDelegates::OnPreObjectPropertyChanged.RemoveAll(this);
+	FCoreUObjectDelegates::OnObjectPropertyChanged.RemoveAll(this);
+	FCoreUObjectDelegates::OnObjectConstructed.RemoveAll(this);
+	USelection::SelectionChangedEvent.RemoveAll(this);
+}
+
+void SPCGEditorGraphDebugObjectTree::Construct(const FArguments& InArgs, TSharedPtr<FPCGEditor> InPCGEditor)
+{
+	check(InPCGEditor);
+	PCGEditor = InPCGEditor;
+
+	UPCGGraph* PCGGraph = GetPCGGraph();
+	check(PCGGraph);
+
+	FCoreUObjectDelegates::OnPreObjectPropertyChanged.AddSP(this, &SPCGEditorGraphDebugObjectTree::OnPreObjectPropertyChanged);
+	FCoreUObjectDelegates::OnObjectPropertyChanged.AddSP(this, &SPCGEditorGraphDebugObjectTree::OnObjectPropertyChanged);
+	FCoreUObjectDelegates::OnObjectConstructed.AddSP(this, &SPCGEditorGraphDebugObjectTree::OnObjectConstructed);
+	USelection::SelectionChangedEvent.AddSP(this, &SPCGEditorGraphDebugObjectTree::OnEditorSelectionChanged);
+
+	const TSharedRef<SScrollBar> VerticalScrollBar = SNew(SScrollBar)
+		.Orientation(Orient_Vertical)
+		.Thickness(FVector2D(12.0f, 12.0f));
+
+	DebugObjectTreeView = SNew(STreeView<FPCGEditorGraphDebugObjectItemPtr>)
+		.TreeItemsSource(&RootItems)
+		.OnGenerateRow(this, &SPCGEditorGraphDebugObjectTree::MakeTreeRowWidget)
+		.OnGetChildren(this, &SPCGEditorGraphDebugObjectTree::OnGetChildren)
+		.OnSelectionChanged(this, &SPCGEditorGraphDebugObjectTree::OnSelectionChanged)
+		.OnExpansionChanged(this, &SPCGEditorGraphDebugObjectTree::OnExpansionChanged)
+		.OnSetExpansionRecursive(this, &SPCGEditorGraphDebugObjectTree::OnSetExpansionRecursive)
+		.SelectionMode(ESelectionMode::SingleToggle)
+		.AllowOverscroll(EAllowOverscroll::No)
+		.ExternalScrollbar(VerticalScrollBar)
+		.ConsumeMouseWheel(EConsumeMouseWheel::Always)
+		.OnContextMenuOpening(FOnContextMenuOpening::CreateSP(this, &SPCGEditorGraphDebugObjectTree::OpenContextMenu));
+
+	const TSharedRef<SWidget> SetButton = PropertyCustomizationHelpers::MakeUseSelectedButton(
+		FSimpleDelegate::CreateSP(this, &SPCGEditorGraphDebugObjectTree::SetDebugObjectFromSelection_OnClicked),
+		LOCTEXT("SetDebugObject", "Set debug object from Level Editor selection."),
+		TAttribute<bool>::Create(TAttribute<bool>::FGetter::CreateSP(this, &SPCGEditorGraphDebugObjectTree::IsSetDebugObjectFromSelectionButtonEnabled)));
+
+	ChildSlot
+	[
+		SNew(SVerticalBox)
+		+SVerticalBox::Slot()
+		.AutoHeight()
+		.HAlign(HAlign_Right)
+		[
+			SNew(SHorizontalBox)
+			+SHorizontalBox::Slot()
+			.AutoWidth()
+			.VAlign(VAlign_Center)
+			.Padding(4.0f)
+			[
+				SetButton
+			]
+			+SHorizontalBox::Slot()
+			.AutoWidth()
+			.VAlign(VAlign_Center)
+			.Padding(4.0f)
+			[
+				SNew(SButton)
+				.ButtonStyle(FAppStyle::Get(), "SimpleButton")
+				.OnClicked(this, &SPCGEditorGraphDebugObjectTree::FocusOnDebugObject_OnClicked)
+				.IsEnabled(this, &SPCGEditorGraphDebugObjectTree::IsFocusOnDebugObjectButtonEnabled)
+				.ToolTipText(LOCTEXT("DebugSelectActor", "Select and frame the debug actor in the Level Editor."))
+				.ContentPadding(0.0f)
+				[
+					SNew(SImage)
+					.Image(FPCGEditorStyle::Get().GetBrush("PCG.Editor.ZoomToSelection"))
+					.ColorAndOpacity(FSlateColor::UseForeground())
+				]
+			]
+		]
+		+SVerticalBox::Slot()
+		.FillHeight(1.0f)
+		[
+			SNew(SHorizontalBox)
+			+SHorizontalBox::Slot()
+			.FillWidth(1.0f)
+			[
+				DebugObjectTreeView.ToSharedRef()
+			]
+			+SHorizontalBox::Slot()
+			.AutoWidth()
+			[
+				VerticalScrollBar
+			]
+		]
+	];
+
+	RequestRefresh();
+}
+
+void SPCGEditorGraphDebugObjectTree::Tick(const FGeometry& AllottedGeometry, const double InCurrentTime, const float InDeltaTime)
+{
+	SCompoundWidget::Tick(AllottedGeometry, InCurrentTime, InDeltaTime);
+
+	if (bNeedsRefresh && InCurrentTime >= NextRefreshTime)
+	{
+		// Updating the tree while the inspected component is generating can be bad as the selection can
+		// be lost. Don't change the tree if we're inspecting something that is generating.
+		const UPCGComponent* InspectedComponent = SelectedStack.GetRootComponent();
+		if(InspectedComponent == nullptr || !InspectedComponent->IsGenerating())
+		{
+			bNeedsRefresh = false;
+			RefreshTree();
+
+			constexpr double RefreshCooldownTime = 0.25f;
+			NextRefreshTime = InCurrentTime + RefreshCooldownTime;
+		}
+	}
+
+	if (!IsSetDebugObjectFromSelectionEnabled.IsSet())
+	{
+		UpdateIsSetDebugObjectFromSelectionEnabled();
+	}
+}
+
+void SPCGEditorGraphDebugObjectTree::RequestRefresh()
+{
+	bNeedsRefresh = true;
+}
+
+void SPCGEditorGraphDebugObjectTree::OnObjectsReplaced(const TMap<UObject*, UObject*>& ReplacementMap)
+{
+	for (FPCGEditorGraphDebugObjectItemPtr& GraphItem : AllGraphItems)
+	{
+		if (GraphItem)
+		{
+			GraphItem->OnObjectsReplaced(ReplacementMap);
+		}
+	}
+
+	// We need to replace the selected component, selected stacks and expanded stacks as they are used for selection matching.
+	if (UObject* NewComponent = ReplacementMap.FindRef(SelectedOriginalComponent.GetEvenIfUnreachable()))
+	{
+		SelectedOriginalComponent = Cast<UPCGComponent>(NewComponent);
+	}
+
+	SelectedStack.ReplaceRoot(ReplacementMap);
+
+	// For the expanded stacks, since it is a set, if we ever modify the root, we should probably re-hash it.
+	// We can't modify it in place so we'll just rebuild it.
+	TSet<FPCGStack> UpdatedExpandedStacks;
+	UpdatedExpandedStacks.Reserve(ExpandedStacks.Num());
+	for (FPCGStack& Stack : ExpandedStacks)
+	{
+		Stack.ReplaceRoot(ReplacementMap);
+		UpdatedExpandedStacks.Add(std::move(Stack));
+	}
+	
+	ExpandedStacks = std::move(UpdatedExpandedStacks);
+}
+
+void SPCGEditorGraphDebugObjectTree::SetNodeBeingInspected(const UPCGNode* InPCGNode)
+{
+	PCGNodeBeingInspected = IsValid(InPCGNode) ? InPCGNode : nullptr;
+
+	RequestRefresh();
+}
+
+FReply SPCGEditorGraphDebugObjectTree::FocusOnDebugObject_OnClicked() const
+{
+	if (UPCGComponent* PCGComponent = PCGEditor.Pin()->GetPCGComponentBeingInspected())
+	{
+		AActor* Actor = PCGComponent->GetOwner();
+		if (Actor && GEditor && GUnrealEd && GEditor->CanSelectActor(Actor, /*bInSelected=*/true))
+		{
+			GEditor->SelectNone(/*bNoteSelectionChange=*/false, /*bDeselectBSPSurfs=*/true, /*WarnAboutManyActors=*/false);
+			GEditor->SelectActor(Actor, /*bInSelected=*/true, /*bNotify=*/true, /*bSelectEvenIfHidden=*/true);
+			GUnrealEd->Exec(Actor->GetWorld(), TEXT("CAMERA ALIGN ACTIVEVIEWPORTONLY"));
+			GEditor->SelectComponent(PCGComponent, /*bInSelected=*/true, /*bNotify=*/true, /*bSelectEvenIfHidden=*/true);
+		}
+	}
+
+	return FReply::Handled();
+}
+
+bool SPCGEditorGraphDebugObjectTree::IsFocusOnDebugObjectButtonEnabled() const
+{
+	return PCGEditor.IsValid() && (PCGEditor.Pin()->GetPCGComponentBeingInspected() != nullptr);
+}
+
+void SPCGEditorGraphDebugObjectTree::SetDebugObjectFromSelection_OnClicked()
+{
+	const UPCGGraph* PCGGraph = GetPCGGraph();
+	if (!PCGGraph)
+	{
+		return;
+	}
+
+	check(GEditor);
+	USelection* SelectedActors = GEditor->GetSelectedActors();
+	if (!IsValid(SelectedActors))
+	{
+		return;
+	}
+
+	for (FSelectionIterator It(*SelectedActors); It; ++It)
+	{
+		const AActor* SelectedActor = Cast<AActor>(*It);
+		if (!IsValid(SelectedActor))
+		{
+			continue;
+		}
+
+		TArray<const UPCGComponent*> PCGComponents;
+		SelectedActor->GetComponents<const UPCGComponent>(PCGComponents, /*bIncludeFromChildActors=*/true);
+
+		for (const UPCGComponent* PCGComponent : PCGComponents)
+		{
+			if (!IsValid(PCGComponent))
+			{
+				continue;
+			}
+
+			FPCGStackContext StackContext;
+			if (!PCGComponent->GetStackContext(StackContext))
+			{
+				continue;
+			}
+
+			bool bSelectedDebugObject = false;
+
+			for (const FPCGStack& Stack : StackContext.GetStacks())
+			{
+				if (FPCGEditorGraphDebugObjectItemPtr Item = GetItemFromStack(Stack))
+				{
+					ExpandAndSelectFirstLeafDebugObject(Item);
+					bSelectedDebugObject = true;
+					break;
+				}
+			}
+
+			if (bSelectedDebugObject)
+			{
+				break;
+			}
+		}
+	}
+}
+
+bool SPCGEditorGraphDebugObjectTree::SetDebugObjectFromStackFromAnotherEditor(const FPCGStack& InStack)
+{
+	RefreshTree();
+
+	const UPCGGraph* CurrentGraph = PCGEditor.Pin()->GetPCGGraph();
+	// If the given stack is already valid for this editor, try to select it as-is.
+	if (InStack.GetGraphForCurrentFrame() == CurrentGraph)
+	{
+		if (FPCGEditorGraphDebugObjectItemPtr Item = GetItemFromStack(InStack))
+		{
+			ExpandAndSelectDebugObject(Item);
+			return true;
+		}
+	}
+
+	// Find first instance that starts with the given stack but has the current graph at its end.
+	for (const FPCGEditorGraphDebugObjectItemPtr& Item : AllGraphItems)
+	{
+		const FPCGStack* ItemStack = Item->GetPCGStack();
+
+		if (ItemStack && ItemStack->BeginsWith(InStack) && ItemStack->GetGraphForCurrentFrame() == CurrentGraph)
+		{
+			ExpandAndSelectDebugObject(Item);
+			return true;
+		}
+	}
+
+	return false;
+}
+
+void SPCGEditorGraphDebugObjectTree::ExpandAndSelectDebugObject(const FPCGEditorGraphDebugObjectItemPtr& InItem)
+{
+	if (!InItem)
+	{
+		return;
+	}
+
+	FPCGEditorGraphDebugObjectItemPtr Parent = InItem->GetParent();
+	while (Parent)
+	{
+		DebugObjectTreeView->SetItemExpansion(Parent, true);
+		Parent = Parent->GetParent();
+	}
+
+	DebugObjectTreeView->SetSelection(InItem);
+	DebugObjectTreeView->RequestScrollIntoView(InItem);
+}
+
+FPCGEditorGraphDebugObjectItemPtr SPCGEditorGraphDebugObjectTree::GetItemFromStack(const FPCGStack& InStack) const
+{
+	for (const FPCGEditorGraphDebugObjectItemPtr& Item : AllGraphItems)
+	{
+		if (Item->GetPCGStack() && *(Item->GetPCGStack()) == InStack)
+		{
+			return Item;
+		}
+	}
+
+	return nullptr;
+}
+
+bool SPCGEditorGraphDebugObjectTree::GetFirstStackFromSelection(const UPCGNode* InNode, const UPCGGraph* InGraph, FPCGStack& OutStack) const
+{
+	if (SelectedStack.GetStackFrames().IsEmpty())
+	{
+		return false;
+	}
+
+	for (const FPCGEditorGraphDebugObjectItemPtr& Item : AllGraphItems)
+	{
+		const FPCGStack* ItemStack = Item->GetPCGStack();
+
+		// Early validation - common stack
+		if (!ItemStack || !ItemStack->BeginsWith(SelectedStack))
+		{
+			continue;
+		}
+
+		// Then let's find the occurrence of the node
+		const TArray<FPCGStackFrame>& ItemStackFrames = ItemStack->GetStackFrames();
+		
+		const UPCGNode* ItemNode = nullptr;
+		int NodeIndex = SelectedStack.GetStackFrames().Num();
+		while (NodeIndex < ItemStackFrames.Num())
+		{
+			ItemNode = Cast<const UPCGNode>(ItemStackFrames[NodeIndex].Object.Get());
+			if (ItemNode)
+			{
+				break;
+			}
+			else
+			{
+				++NodeIndex;
+			}
+		}
+
+		if (ItemNode != InNode)
+		{
+			continue;
+		}
+
+		// And make sure the graph that's somewhere after it in the stack is present
+		// Implementation note: the graph is either the next (static) or the second next in the stack.
+		const UPCGGraph* ItemGraph = nullptr;
+
+		if (ItemStackFrames.IsValidIndex(NodeIndex + 1))
+		{
+			ItemGraph = Cast<const UPCGGraph>(ItemStackFrames[NodeIndex + 1].Object.Get());
+		}
+
+		if (!ItemGraph && ItemStackFrames.IsValidIndex(NodeIndex + 2))
+		{
+			ItemGraph = Cast<const UPCGGraph>(ItemStackFrames[NodeIndex + 2].Object.Get());
+		}
+
+		if (ItemGraph && (!InGraph || ItemGraph == InGraph))
+		{
+			OutStack = *Item->GetPCGStack();
+			return true;
+		}
+	}
+
+	return false;
+}
+
+void SPCGEditorGraphDebugObjectTree::UpdateIsSetDebugObjectFromSelectionEnabled()
+{
+	IsSetDebugObjectFromSelectionEnabled = false;
+
+	const UPCGGraph* PCGGraph = GetPCGGraph();
+	if (!PCGGraph)
+	{
+		return;
+	}
+
+	USelection* SelectedActors = GEditor->GetSelectedActors();
+	if (!IsValid(SelectedActors))
+	{
+		return;
+	}
+
+	UPCGSubsystem* Subsystem = PCGEditor.Pin() ? PCGEditor.Pin()->GetSubsystem() : nullptr;
+
+	for (FSelectionIterator It(*SelectedActors); It; ++It)
+	{
+		const AActor* SelectedActor = Cast<AActor>(*It);
+		if (!IsValid(SelectedActor))
+		{
+			continue;
+		}
+
+		TArray<const UPCGComponent*> PCGComponents;
+		SelectedActor->GetComponents<const UPCGComponent>(PCGComponents, /*bIncludeFromChildActors=*/true);
+
+		for (const UPCGComponent* PCGComponent : PCGComponents)
+		{
+			if (!IsValid(PCGComponent))
+			{
+				continue;
+			}
+
+			// Look for graph in static stacks.
+			FPCGStackContext StackContext;
+			if (PCGComponent->GetStackContext(StackContext))
+			{
+				if (Algo::AnyOf(StackContext.GetStacks(), [&PCGGraph](const FPCGStack& InStack) { return InStack.GetGraphForCurrentFrame() == PCGGraph; }))
+				{
+					IsSetDebugObjectFromSelectionEnabled = true;
+					return;
+				}
+			}
+
+			// Look for graph in dynamic stacks.
+			if (Subsystem)
+			{
+				if (Algo::AnyOf(Subsystem->GetExecutedStacks(PCGComponent, PCGGraph), [&PCGGraph](const FPCGStack& InStack) { return InStack.GetGraphForCurrentFrame() == PCGGraph; }))
+				{
+					IsSetDebugObjectFromSelectionEnabled = true;
+					return;
+				}
+			}
+		}
+	}
+}
+
+void SPCGEditorGraphDebugObjectTree::AddStacksToTree(const TArray<FPCGStack>& Stacks,
+	TMap<AActor*, TSharedPtr<FPCGEditorGraphDebugObjectItem_Actor>>& InOutActorItems,
+	TMap<const FPCGStack, FPCGEditorGraphDebugObjectItemPtr>& InOutStackToItem)
+{
+	//TRACE_CPUPROFILER_EVENT_SCOPE(SPCGEditorGraphDebugObjectTree::AddStacksToTree);
+	const UPCGGraph* GraphBeingEdited = GetPCGGraph();
+	UPCGSubsystem* Subsystem = PCGEditor.Pin()->GetSubsystem();
+	if (!GraphBeingEdited || !Subsystem)
+	{
+		return;
+	}
+
+	for (const FPCGStack& Stack : Stacks)
+	{
+		//TRACE_CPUPROFILER_EVENT_SCOPE(SPCGEditorGraphDebugObjectTree::AddStacksToTree::PerStackLoop);
+		// If the current graph is not in the stack at all then skip.
+		if (!Stack.HasObject(GraphBeingEdited))
+		{
+			continue;
+		}
+
+		UPCGComponent* PCGComponent = const_cast<UPCGComponent*>(Stack.GetRootComponent());
+		if (!PCGComponent)
+		{
+			continue;
+		}
+
+		// Prevent duplicate entries from the editor world while in PIE.
+		if (PCGHelpers::IsRuntimeOrPIE())
+		{
+			if (UWorld* World = PCGComponent->GetWorld())
+			{
+				if (!World->IsGameWorld())
+				{
+					continue;
+				}
+			}
+		}
+
+		int32 TopGraphIndex = INDEX_NONE;
+		UPCGGraph* TopGraph = const_cast<UPCGGraph*>(Stack.GetRootGraph(&TopGraphIndex));
+		if (!TopGraph)
+		{
+			continue;
+		}
+
+		// If we're inspecting a node which has not logged inspection data in a previous execution, display grayed out.
+		bool bDisplayGrayedOut = false;
+		if (PCGNodeBeingInspected.IsValid())
+		{
+			const UPCGSettings* Settings = PCGNodeBeingInspected->GetSettings();
+			const bool bGPUNode = Settings && Settings->bEnabled && Settings->ShouldExecuteOnGPU();
+
+			// Display grayed out if no inspection data has been stored for this node, and node does not run on GPU,
+			// because we don't opportunistically store inspection data for GPU nodes.
+			bDisplayGrayedOut = !bGPUNode && !PCGComponent->HasNodeProducedData(PCGNodeBeingInspected.Get(), Stack);
+		}
+
+		AActor* Actor = PCGComponent->GetOwner();
+		if (!Actor)
+		{
+			continue;
+		}
+
+		// Add actor item if not already added.
+		FPCGEditorGraphDebugObjectItemPtr ActorItem;
+		if (TSharedPtr<FPCGEditorGraphDebugObjectItem_Actor>* FoundActorItem = InOutActorItems.Find(Actor))
+		{
+			ActorItem = *FoundActorItem;
+		}
+		else
+		{
+			ActorItem = InOutActorItems.Emplace(Actor, MakeShared<FPCGEditorGraphDebugObjectItem_Actor>(Actor, bDisplayGrayedOut));
+			AllGraphItems.Add(ActorItem);
+		}
+
+		const TArray<FPCGStackFrame>& StackFrames = Stack.GetStackFrames();
+
+		// Generic function to attack an item to the parent at the correct place
+		auto AddToGraphItemsAndAttachItemToParent = [&InOutStackToItem, this](FPCGStack& GraphStack, FPCGEditorGraphDebugObjectItemPtr GraphItem)
+		{
+			//TRACE_CPUPROFILER_EVENT_SCOPE(SPCGEditorGraphDebugObjectTree::AddStacksToTree::PerStackLoop::PerFrame::AddToGraphItemsAndAttachItemToParent);
+			AllGraphItems.Add(GraphItem);
+
+			TArray<FPCGStackFrame>& LoopGraphStackFrames = GraphStack.GetStackFramesMutable();
+			while (LoopGraphStackFrames.Num() > 0)
+			{
+				LoopGraphStackFrames.SetNum(LoopGraphStackFrames.Num() - 1, EAllowShrinking::No);
+
+				if (FPCGEditorGraphDebugObjectItemPtr* ParentItem = InOutStackToItem.Find(GraphStack))
+				{
+					(*ParentItem)->AddChild(GraphItem.ToSharedRef());
+					break;
+				}
+			}
+		};
+
+		// Generic function to create the subgraph/loop item and hook it up to the parent properly
+		auto AddSubgraphOrLoopItemToStack = [&Stack, &InOutStackToItem, &AddToGraphItemsAndAttachItemToParent](const UPCGGraph* InGraph, const UPCGNode* InSubgraphNode, int FrameCutoff, bool bInIsDebuggable, bool bInDisplayGrayedOut, EPCGEditorSubgraphNodeType InSubgraphType)
+		{
+			//TRACE_CPUPROFILER_EVENT_SCOPE(SPCGEditorGraphDebugObjectTree::AddStacksToTree::PerStackLoop::PerFrame::AddSubgraphOrLoopItemToStack);
+			FPCGStack GraphStack = Stack;
+			GraphStack.GetStackFramesMutable().SetNum(FrameCutoff, EAllowShrinking::No);
+
+			if (!InOutStackToItem.Contains(GraphStack))
+			{
+				FPCGEditorGraphDebugObjectItemPtr GraphItem = InOutStackToItem.Emplace(
+					GraphStack,
+					MakeShared<FPCGEditorGraphDebugObjectItem_PCGSubgraph>(InSubgraphNode, InGraph, GraphStack, bInIsDebuggable, bInDisplayGrayedOut, InSubgraphType));
+
+				AddToGraphItemsAndAttachItemToParent(GraphStack, GraphItem);
+			}
+		};
+
+		// Example stack:
+		//     Component/TopGraph/SubgraphNode/Subgraph/LoopSubgraphNode/LoopIndex/LoopSubgraph
+		//                                       ^ static subgraph
+		//     Component/TopGraph/SubgraphNode/INDEX_NONE/Subgraph/...
+		//                                       ^ dynamic subgraph
+		// The loop below adds tree items for component & top graph, and then whenever a graph is encountered
+		// we look a previous frames to determine whether to add a subgraph item or loop subgraph item.
+		int FrameIndex = StackFrames.Num() - 1;
+
+		// Rollback from the last frames to check if there are matching stacks that would correspond to the current stack, so we don't do duplicate work.
+		{
+			//TRACE_CPUPROFILER_EVENT_SCOPE(SPCGEditorGraphDebugObjectTree::AddStacksToTree::PerStackLoop::RollbackLoop);
+			FPCGStack UnrollStack = Stack;
+			TArray<FPCGStackFrame>& UnrollFrames = UnrollStack.GetStackFramesMutable();
+
+			while (FrameIndex > 1)
+			{
+				if (InOutStackToItem.Contains(UnrollStack))
+				{
+					break;
+				}
+				else
+				{
+					--FrameIndex;
+					UnrollFrames.Pop(EAllowShrinking::No);
+				}
+			}
+		}
+
+		for(; FrameIndex < StackFrames.Num(); ++FrameIndex)
+		{
+			//TRACE_CPUPROFILER_EVENT_SCOPE(SPCGEditorGraphDebugObjectTree::AddStacksToTree::PerStackLoop::PerFrame);
+			const FPCGStackFrame& StackFrame = StackFrames[FrameIndex];
+			const FPCGStackFrame& PreviousStackFrame = StackFrames[FrameIndex - 1];
+			FPCGEditorGraphDebugObjectItemPtr CurrentItem;
+
+			// When we encounter a graph, we look at the frame index and/or preceding frames to determine the graph type.
+			if (const UPCGGraph* StackGraph = Cast<const UPCGGraph>(StackFrame.Object.Get()))
+			{
+				const bool bIsDebuggable = (GraphBeingEdited == StackGraph);
+
+				// Top graph.
+				if (FrameIndex == TopGraphIndex && StackGraph == TopGraph)
+				{
+					FPCGStack GraphStack = Stack;
+					GraphStack.GetStackFramesMutable().SetNum(FrameIndex + 1);
+
+					if (!InOutStackToItem.Contains(GraphStack))
+					{
+						FPCGEditorGraphDebugObjectItemPtr TopGraphItem = InOutStackToItem.Emplace(
+							GraphStack,
+							MakeShared<FPCGEditorGraphDebugObjectItem_PCGComponent>(PCGComponent, StackGraph, GraphStack, bIsDebuggable, bDisplayGrayedOut));
+
+						AllGraphItems.Add(TopGraphItem);
+
+						ActorItem->AddChild(TopGraphItem.ToSharedRef());
+					}
+				}
+				// Previous stack was node, therefore static subgraph.
+				else if (const UPCGNode* SubgraphNode = Cast<const UPCGNode>(PreviousStackFrame.Object.Get()))
+				{
+					AddSubgraphOrLoopItemToStack(StackGraph, SubgraphNode, FrameIndex + 1, bIsDebuggable, bDisplayGrayedOut, EPCGEditorSubgraphNodeType::StaticSubgraph);
+				}
+				// Previous stack was loop index, therefore loop subgraph.
+				else if (FrameIndex >= 2 && (PreviousStackFrame.LoopIndex != INDEX_NONE))
+				{
+					const UPCGNode* LoopSubgraphNode = Cast<const UPCGNode>(StackFrames[FrameIndex - 2].Object.Get());
+					if (ensure(LoopSubgraphNode))
+					{
+						// Take the stack up to the looped subgraph node, add a item for the node + graph
+						AddSubgraphOrLoopItemToStack(StackGraph, LoopSubgraphNode, FrameIndex - 1, /*bIsDebuggable=*/false, /*bIsGrayedOut=*/bDisplayGrayedOut, EPCGEditorSubgraphNodeType::LoopSubgraph);
+
+						// Take full stack up until this point which will be the unique stack for the loop iteration.
+						FPCGStack LoopIterationStack = Stack;
+						LoopIterationStack.GetStackFramesMutable().SetNum(FrameIndex + 1);
+
+						if (!InOutStackToItem.Contains(LoopIterationStack))
+						{
+							FPCGEditorGraphDebugObjectItemPtr LoopIterationItem = InOutStackToItem.Emplace(
+								LoopIterationStack,
+								MakeShared<FPCGEditorGraphDebugObjectItem_PCGLoopIndex>(PreviousStackFrame.LoopIndex, StackGraph, LoopIterationStack, bIsDebuggable, bDisplayGrayedOut));
+
+							AddToGraphItemsAndAttachItemToParent(LoopIterationStack, LoopIterationItem);
+						}
+					}
+				}
+				// Previous stack was invalid node / node with a INDEX_NONE loop, therefore most likely a dynamic subgraph
+				else if (FrameIndex >= 2 && !PreviousStackFrame.IsValid())
+				{
+					const UPCGNode* DynamicSubgraphNode = Cast<const UPCGNode>(StackFrames[FrameIndex - 2].Object.Get());
+					if (DynamicSubgraphNode)
+					{
+						AddSubgraphOrLoopItemToStack(StackGraph, DynamicSubgraphNode, FrameIndex + 1, bIsDebuggable, bDisplayGrayedOut, EPCGEditorSubgraphNodeType::DynamicSubgraph);
+					}
+				}
+			}
+		}
+	}
+}
+
+void SPCGEditorGraphDebugObjectTree::RefreshTree()
+{
+	//TRACE_CPUPROFILER_EVENT_SCOPE(SPCGEditorGraphDebugObjectTree::RefreshTree);
+	RootItems.Empty();
+	AllGraphItems.Empty();
+	DebugObjectTreeView->RequestTreeRefresh();
+
+	UPCGGraph* PCGGraph = GetPCGGraph();
+	if (!PCGGraph)
+	{
+		return;
+	}
+
+	UPCGSubsystem* Subsystem = PCGEditor.Pin()->GetSubsystem();
+	if (!Subsystem)
+	{
+		return;
+	}
+
+	TArray<UObject*> PCGComponents;
+	GetObjectsOfClass(UPCGComponent::StaticClass(), PCGComponents, /*bIncludeDerivedClasses=*/ true);
+
+	if (!PCGComponents.IsEmpty())
+	{
+		TMap<AActor*, TSharedPtr<FPCGEditorGraphDebugObjectItem_Actor>> ActorItems;
+		TMap<const FPCGStack, FPCGEditorGraphDebugObjectItemPtr> StackToItem;
+
+		TArray<FPCGStack> GraphStacks = Subsystem->GetExecutedStacks(nullptr, PCGGraph, /*bOnlyWithSubgraphAsCurrentFrame=*/false);
+		TMap<const UPCGComponent*, TArray<FPCGStack>> ComponentsGraphStacks;
+		for (FPCGStack& GraphStack : GraphStacks)
+		{
+			ComponentsGraphStacks.FindOrAdd(GraphStack.GetRootComponent()).Add(MoveTemp(GraphStack));
+		}
+
+		for (UObject* PCGComponentObject : PCGComponents)
+		{
+			if (!IsValid(PCGComponentObject))
+			{
+				continue;
+			}
+
+			UPCGComponent* PCGComponent = Cast<UPCGComponent>(PCGComponentObject);
+			if (!PCGComponent || !PCGComponent->IsRegistered())
+			{
+				continue;
+			}
+
+			// Process static stacks that can be read from the compiled graph.
+			FPCGStackContext StackContext;
+			if (!PCGComponent->GetStackContext(StackContext))
+			{
+				continue;
+			}
+			AddStacksToTree(StackContext.GetStacks(), ActorItems, StackToItem);
+
+			// Process stacks encountered during execution so far, which will include dynamic subgraphs & loop subgraphs.
+			// There will be overlaps with the static stacks but only unique entries will be added to the tree.
+			if (TArray<FPCGStack>* ComponentGraphStacks = ComponentsGraphStacks.Find(PCGComponent))
+			{
+				AddStacksToTree(*ComponentGraphStacks, ActorItems, StackToItem);
+			}
+		}
+
+		for (TPair<AActor*, TSharedPtr<FPCGEditorGraphDebugObjectItem_Actor>>& ActorItem : ActorItems)
+		{
+			RootItems.Add(MoveTemp(ActorItem.Value));
+		}
+	}
+
+	SortTreeItems();
+	RestoreTreeState();
+}
+
+void SPCGEditorGraphDebugObjectTree::SortTreeItems(bool bIsAscending, bool bIsRecursive)
+{
+	RootItems.Sort([bIsAscending](const FPCGEditorGraphDebugObjectItemPtr& InLHS, const FPCGEditorGraphDebugObjectItemPtr& InRHS)
+	{
+		return (InLHS->GetLabel(/*bForSorting=*/true) < InRHS->GetLabel(/*bForSorting=*/true)) == bIsAscending;
+	});
+
+	if (bIsRecursive)
+	{
+		for (const FPCGEditorGraphDebugObjectItemPtr& Item : RootItems)
+		{
+			Item->SortChildren(bIsAscending, bIsRecursive);
+		}
+	}
+}
+
+void SPCGEditorGraphDebugObjectTree::RestoreTreeState()
+{
+	// Try to restore user item expansion. We can't modify ExpandedStack, so we'll rebuild it.
+	TSet<FPCGStack> ExpandedStacksBefore = std::move(ExpandedStacks);
+	ExpandedStacks.Empty(ExpandedStacksBefore.Num());
+	for (const FPCGStack& ExpandedStack : ExpandedStacksBefore)
+	{
+		if (FPCGEditorGraphDebugObjectItemPtr Item = GetItemFromStack(ExpandedStack))
+		{
+			DebugObjectTreeView->SetItemExpansion(Item, true);
+		}
+	}
+
+	bool bFoundMatchingStack = false;
+
+	// Try to restore user item selection by exact matching.
+	if (FPCGEditorGraphDebugObjectItemPtr ExactMatchItem = GetItemFromStack(SelectedStack))
+	{
+		if (!DebugObjectTreeView->IsItemSelected(ExactMatchItem))
+		{
+			DebugObjectTreeView->SetItemSelection(ExactMatchItem, true);
+		}
+
+		bFoundMatchingStack = true;
+	}
+
+	// Try to restore user item selection by fuzzy matching (e.g. share the same owner) if no exactly matching stack was found.
+	if (!bFoundMatchingStack)
+	{
+		for (FPCGEditorGraphDebugObjectItemPtr& Item : AllGraphItems)
+		{
+			const FPCGStack* ItemStack = Item->GetPCGStack();
+
+			bool bFuzzyMatch = false;
+
+			if (ItemStack && SelectedGraph.Get() == ItemStack->GetRootGraph())
+			{
+				const UPCGComponent* RootComponent = ItemStack->GetRootComponent();
+				const AActor* RootOwner = RootComponent ? RootComponent->GetOwner() : nullptr;
+				const APCGPartitionActor* RootPartitionActor = Cast<APCGPartitionActor>(RootOwner);
+
+				if (RootComponent && RootPartitionActor)
+				{
+					// For local components, we can fuzzy match as long as the GridSize, GridCoord, OriginalComponent, and ExecutionDomain are the same.
+					// This is equivalent to saying they are on the same partition actor and come from the same original component.
+					bFuzzyMatch = SelectedGridSize == RootComponent->GetGenerationGridSize()
+						&& SelectedGridCoord == RootPartitionActor->GetGridCoord()
+						&& SelectedOriginalComponent.Get() == RootPartitionActor->GetOriginalComponent(RootComponent)
+						&& (SelectedOriginalComponent.IsValid() && SelectedOriginalComponent->IsManagedByRuntimeGenSystem() == RootComponent->IsManagedByRuntimeGenSystem());
+				}
+				else
+				{
+					// For original components, we can fuzzy match as long as the owning actor is the same.
+					// Note: This fails for multiple original components with the same graph on the same actor, since there is no way to know which one to pick.
+					if (SelectedOwner.Get() == RootOwner && Item->GetParent() && Item->GetParent()->GetChildren().Num() == 1)
+					{
+						int32 ItemRootGraphIndex = INDEX_NONE;
+						int32 SelectedRootGraphIndex = INDEX_NONE;
+
+						ItemStack->GetRootGraph(&ItemRootGraphIndex);
+						SelectedStack.GetRootGraph(&SelectedRootGraphIndex);
+
+						const TArray<FPCGStackFrame>& ItemStackFrames = ItemStack->GetStackFrames();
+						const TArray<FPCGStackFrame>& SelectedStackFrames = SelectedStack.GetStackFrames();
+
+						// If the stacks match from the RootGraph onwards, then our fuzzy match should succeed.
+						if (ItemRootGraphIndex != INDEX_NONE && ItemRootGraphIndex == SelectedRootGraphIndex && ItemStackFrames.Num() == SelectedStackFrames.Num())
+						{
+							bool bAllStackFramesMatch = true;
+
+							for (int I = ItemRootGraphIndex; I < ItemStackFrames.Num(); ++I)
+							{
+								if (!ItemStackFrames[I].IsValid() || ItemStackFrames[I] != SelectedStackFrames[I])
+								{
+									bAllStackFramesMatch = false;
+									break;
+								}
+							}
+
+							if (bAllStackFramesMatch)
+							{
+								bFuzzyMatch = true;
+							}
+						}
+					}
+				}
+			}
+
+			if (bFuzzyMatch)
+			{
+				// Force the selected object to re-expand.
+				FPCGEditorGraphDebugObjectItemPtr Parent = Item->GetParent();
+				while (Parent)
+				{
+					DebugObjectTreeView->SetItemExpansion(Parent, true);
+					Parent = Parent->GetParent();
+				}
+
+				if (!DebugObjectTreeView->IsItemSelected(Item))
+				{
+					DebugObjectTreeView->SetItemSelection(Item, true);
+				}
+				
+				break;
+			}
+		}
+	}
+}
+
+void SPCGEditorGraphDebugObjectTree::OnPreObjectPropertyChanged(UObject* InObject, const FEditPropertyChain& /*InPropertyChain*/)
+{
+	if (UPCGGraphInstance* PCGgraphInstance = Cast<UPCGGraphInstance>(InObject))
+	{
+		const UPCGGraph* PCGGraph = PCGgraphInstance->GetGraph();
+		if (PCGGraph && PCGGraph == GetPCGGraph())
+		{
+			RequestRefresh();
+		}
+	}
+}
+
+void SPCGEditorGraphDebugObjectTree::OnObjectPropertyChanged(UObject* InObject, FPropertyChangedEvent& InPropertyChangedEvent)
+{
+	if (const UPCGGraphInterface* PCGGraphInterface = Cast<UPCGGraphInterface>(InObject))
+	{
+		const UPCGGraph* PCGGraph = PCGGraphInterface->GetGraph();
+		if (PCGGraph && PCGGraph == GetPCGGraph())
+		{
+			RequestRefresh();
+		}
+	}
+}
+
+void SPCGEditorGraphDebugObjectTree::OnEditorSelectionChanged(UObject* InObject)
+{
+	IsSetDebugObjectFromSelectionEnabled.Reset();
+}
+
+void SPCGEditorGraphDebugObjectTree::OnObjectConstructed(UObject* InObject)
+{
+	if (const UPCGComponent* PCGComponent = Cast<UPCGComponent>(InObject))
+	{
+		RequestRefresh();
+	}
+}
+
+UPCGGraph* SPCGEditorGraphDebugObjectTree::GetPCGGraph() const
+{
+	const TSharedPtr<FPCGEditor> PCGEditorPtr = PCGEditor.Pin();
+	const UPCGEditorGraph* PCGEditorGraph = PCGEditorPtr.IsValid() ? PCGEditorPtr->GetPCGEditorGraph() : nullptr;
+
+	return PCGEditorGraph ? PCGEditorGraph->GetPCGGraph() : nullptr;
+}
+
+TSharedRef<ITableRow> SPCGEditorGraphDebugObjectTree::MakeTreeRowWidget(FPCGEditorGraphDebugObjectItemPtr InItem, const TSharedRef<STableViewBase>& InOwnerTable)
+{
+	return SNew( STableRow<TSharedPtr<SPCGEditorGraphDebugObjectItemRow>>, InOwnerTable)
+	[
+		SNew(SPCGEditorGraphDebugObjectItemRow, InOwnerTable, InItem)
+		.OnDoubleClick(this, &SPCGEditorGraphDebugObjectTree::ExpandAndSelectFirstLeafDebugObject)
+		.OnJumpTo(this, &SPCGEditorGraphDebugObjectTree::JumpToGraphInTree)
+		.CanJumpTo(this, &SPCGEditorGraphDebugObjectTree::CanJumpToGraphInTree)
+		.OnFocus(this, &SPCGEditorGraphDebugObjectTree::FocusOnItem)
+		.CanFocus(this, &SPCGEditorGraphDebugObjectTree::CanFocusOnItem)
+	];
+}
+
+void SPCGEditorGraphDebugObjectTree::OnGetChildren(FPCGEditorGraphDebugObjectItemPtr InItem, TArray<FPCGEditorGraphDebugObjectItemPtr>& OutChildren) const
+{
+	if (InItem.IsValid())
+	{
+		for (TWeakPtr<FPCGEditorGraphDebugObjectItem> ChildItem : InItem->GetChildren())
+		{
+			FPCGEditorGraphDebugObjectItemPtr ChildItemPtr = ChildItem.Pin();
+			OutChildren.Add(ChildItemPtr);
+		}
+	}
+}
+
+void SPCGEditorGraphDebugObjectTree::OnSelectionChanged(FPCGEditorGraphDebugObjectItemPtr InItem, ESelectInfo::Type InSelectInfo)
+{
+	// Unmark the previously selected object
+	if (FPCGEditorGraphDebugObjectItemPtr PreviousItem = GetItemFromStack(SelectedStack))
+	{
+		PreviousItem->SetSelected(false);
+	}
+
+	// Reset selected item information.
+	SelectedStack = FPCGStack();
+	SelectedGraph = nullptr;
+	SelectedOwner = nullptr;
+	SelectedGridSize = PCGHiGenGrid::UnboundedGridSize();
+	SelectedGridCoord = FIntVector::ZeroValue;
+	SelectedOriginalComponent = nullptr;
+
+	if (const FPCGStack* Stack = InItem ? InItem->GetPCGStack() : nullptr)
+	{
+		SelectedStack = *Stack;
+		SelectedGraph = SelectedStack.GetRootGraph();
+
+		if (const UPCGComponent* RootComponent = SelectedStack.GetRootComponent())
+		{
+			SelectedOwner = RootComponent->GetOwner();
+
+			if (const APCGPartitionActor* PartitionActor = Cast<APCGPartitionActor>(SelectedOwner.Get()))
+			{
+				SelectedGridSize = RootComponent->GetGenerationGridSize();
+				SelectedGridCoord = PartitionActor->GetGridCoord();
+				SelectedOriginalComponent = PartitionActor->GetOriginalComponent(RootComponent);
+			}
+			else
+			{
+				SelectedOriginalComponent = RootComponent;
+			}
+		}
+	}
+
+	const UPCGGraph* CurrentGraph = PCGEditor.Pin()->GetPCGGraph();
+	bool bStackInspectedSet = false;
+
+	// Only attempt to inspect stacks that correspond to the edited graph. Other graphs need to be inspected in their own editor.
+	if (InItem)
+	{
+		// Basically - if the selection is something from "upstream", then we can clear the stack being inspected.
+		if (SelectedStack.HasObject(CurrentGraph))
+		{
+			// If the last graph in the stack is the current graph, then we can just use that as the selected stack
+			TArray<FPCGStackFrame>& StackFrames = SelectedStack.GetStackFramesMutable();
+			while(!StackFrames.IsEmpty() && StackFrames.Last().Object.Get() != CurrentGraph)
+			{
+				StackFrames.Pop();
+			}
+
+			if (ensure(!SelectedStack.GetStackFrames().IsEmpty()))
+			{
+				PCGEditor.Pin()->SetStackBeingInspected(SelectedStack);
+				bStackInspectedSet = true;
+			}
+		}
+	}
+
+	if (!bStackInspectedSet)
+	{
+		PCGEditor.Pin()->ClearStackBeingInspected();
+	}
+
+	// Finally mark the selected item as selected
+	if (SelectedStack.GetGraphForCurrentFrame() == CurrentGraph)
+	{
+		if (FPCGEditorGraphDebugObjectItemPtr CurrentSelectedItem = GetItemFromStack(SelectedStack))
+		{
+			CurrentSelectedItem->SetSelected(true);
+		}
+	}
+}
+
+void SPCGEditorGraphDebugObjectTree::OnExpansionChanged(FPCGEditorGraphDebugObjectItemPtr InItem, bool bInIsExpanded)
+{
+	if (!InItem.IsValid())
+	{
+		return;
+	}
+
+	InItem->SetExpanded(bInIsExpanded);
+
+	if (const FPCGStack* Stack = InItem->GetPCGStack())
+	{
+		if (bInIsExpanded)
+		{
+			ExpandedStacks.Add(*Stack);
+		}
+		else
+		{
+			ExpandedStacks.Remove(*Stack);
+		}
+	}
+}
+
+void SPCGEditorGraphDebugObjectTree::OnSetExpansionRecursive(FPCGEditorGraphDebugObjectItemPtr InItem, bool bInExpand) const
+{
+	if (!InItem.IsValid() || !DebugObjectTreeView.IsValid())
+	{
+		return;
+	}
+
+	DebugObjectTreeView->SetItemExpansion(InItem, bInExpand);
+	InItem->SetExpanded(bInExpand);
+
+	for (const FPCGEditorGraphDebugObjectItemPtr& ChildItem : InItem->GetChildren())
+	{
+		if (ChildItem.IsValid())
+		{
+			OnSetExpansionRecursive(ChildItem, bInExpand);
+		}
+	}
+}
+
+void SPCGEditorGraphDebugObjectTree::ExpandAndSelectFirstLeafDebugObject(const FPCGEditorGraphDebugObjectItemPtr& InItem)
+{
+	if (!InItem.IsValid())
+	{
+		return;
+	}
+
+	// Regardless of what happens, we'll expand the currently selected item.
+	DebugObjectTreeView->SetItemExpansion(InItem, true);
+
+	// Find first occurrence of tree (breadth-first) of current graph in children.
+	FPCGEditorGraphDebugObjectItemPtr ItemToSelect;
+
+	const UPCGGraph* CurrentGraph = PCGEditor.Pin()->GetPCGGraph();
+	TArray<FPCGEditorGraphDebugObjectItemPtr> ToVisit = InItem->GetChildren().Array();
+	int32 VisitIndex = 0;
+	
+	while(VisitIndex < ToVisit.Num())
+	{
+		const FPCGEditorGraphDebugObjectItemPtr& Item = ToVisit[VisitIndex]; 
+		const FPCGStack* ItemStack = Item ? Item->GetPCGStack() : nullptr;
+		if (ItemStack)
+		{
+			if(!ItemStack->GetStackFrames().IsEmpty() && ItemStack->GetStackFrames().Last().Object == CurrentGraph)
+			{
+				ItemToSelect = Item;
+				break;
+			}
+			else
+			{
+				ToVisit.Append(Item->GetChildren().Array());
+				++VisitIndex;
+			}
+		}
+	}
+
+	// If we've found nothing, we'll select the original item, otherwise we'll pick the first occurrence we found.
+	ExpandAndSelectDebugObject(ItemToSelect.IsValid() ? ItemToSelect : InItem);
+}
+
+TSharedPtr<SWidget> SPCGEditorGraphDebugObjectTree::OpenContextMenu()
+{
+	FMenuBuilder MenuBuilder(/*bInShouldCloseWindowAfterMenuSelection=*/true, /*InCommandList=*/nullptr);
+	MenuBuilder.AddMenuEntry(
+		LOCTEXT("JumpToSelectedGraph", "Jump To"),
+		LOCTEXT("JumpToSelectedGraphTooltip", "Jumps to the selected graph."),
+		FSlateIcon(FPCGEditorStyle::Get().GetStyleSetName(), "PCG.Editor.JumpTo"),
+		FUIAction(FExecuteAction::CreateSP(this, &SPCGEditorGraphDebugObjectTree::ContextMenu_JumpToGraphInTree), FCanExecuteAction::CreateSP(this, &SPCGEditorGraphDebugObjectTree::ContextMenu_JumpToGraphInTree_CanExecute))
+	);
+
+	return MenuBuilder.MakeWidget();
+}
+
+void SPCGEditorGraphDebugObjectTree::ContextMenu_JumpToGraphInTree()
+{
+	const TArray<FPCGEditorGraphDebugObjectItemPtr> SelectedItems = DebugObjectTreeView->GetSelectedItems();
+	for (const FPCGEditorGraphDebugObjectItemPtr& SelectedItem : SelectedItems)
+	{
+		JumpToGraphInTree(SelectedItem);
+	}
+}
+
+void SPCGEditorGraphDebugObjectTree::JumpToGraphInTree(const FPCGEditorGraphDebugObjectItemPtr& Item)
+{
+	if (!Item || !Item->GetPCGStack())
+	{
+		return;
+	}
+
+	FPCGStack Stack = *Item->GetPCGStack();
+
+	// Example stack:
+	//     Component/TopGraph/SubgraphNode/Subgraph/LoopSubgraphNode/LoopIndex/LoopSubgraph
+	//                                       ^ static subgraph
+	//     Component/TopGraph/SubgraphNode/INDEX_NONE/Subgraph/...
+	//                                       ^ dynamic subgraph
+
+	UPCGGraph* JumpToPCGGraph = nullptr;
+	UPCGNode* JumpToPCGNode = nullptr;
+
+	// If the item is about this graph, then we'll jump to this instance's caller - we'll walk back until we find a graph (e.g. parent)
+	// and the next entry should be our caller node.
+	if (Item->IsDebuggable())
+	{
+		// Debuggable target graphs correspond to the currently edited graph. For this case open the parent graph and jump to the corresponding subgraph node.
+		// Search the stack for the parent graph.
+		for (int i = Stack.GetStackFrames().Num() - 2; i > 0; --i) 
+		{
+			if (Stack.GetStackFrames()[i].Object != nullptr && Stack.GetStackFrames()[i].Object->IsA<UPCGGraph>())
+			{
+				JumpToPCGGraph = const_cast<UPCGGraph*>(Cast<const UPCGGraph>(Stack.GetStackFrames()[i].Object.Get()));
+				JumpToPCGNode = const_cast<UPCGNode*>(Cast<const UPCGNode>(Stack.GetStackFrames()[i + 1].Object.Get()));
+
+				// Cull remaining remaining frames so the stack selection into the parent graph editor works.
+				Stack.GetStackFramesMutable().SetNum(i);
+				
+				break;
+			}
+		}
+	}
+	else 
+	{
+		// Two cases here: upstream or downstream. In both cases, there's no need to jump to a given node as there is no concept of caller here.
+		JumpToPCGGraph = const_cast<UPCGGraph*>(Item->GetPCGGraph());
+	}
+
+	if (JumpToPCGGraph)
+	{
+		GEditor->GetEditorSubsystem<UAssetEditorSubsystem>()->OpenEditorForAsset(JumpToPCGGraph);
+		IAssetEditorInstance* EditorInstance = GEditor->GetEditorSubsystem<UAssetEditorSubsystem>()->FindEditorForAsset(JumpToPCGGraph, /*bFocusIfOpen*/true);
+		FPCGEditor* OtherPCGEditor = static_cast<FPCGEditor*>(EditorInstance);
+
+		if (OtherPCGEditor)
+		{
+			// Implementation note - for selection purposes we had culled the stack to be in the local referential
+			// however in the case of the jump here we want the full stack
+			OtherPCGEditor->SetStackBeingInspectedFromAnotherEditor(Stack);
+
+			if (JumpToPCGNode)
+			{
+				OtherPCGEditor->JumpToNode(JumpToPCGNode);
+			}
+		}
+	}
+}
+
+bool SPCGEditorGraphDebugObjectTree::ContextMenu_JumpToGraphInTree_CanExecute() const
+{
+	const TArray<FPCGEditorGraphDebugObjectItemPtr> SelectedItems = DebugObjectTreeView->GetSelectedItems();
+	for (const FPCGEditorGraphDebugObjectItemPtr& SelectedItem : SelectedItems)
+	{
+		if (CanJumpToGraphInTree(SelectedItem))
+		{
+			return true;
+		}
+	}
+
+	return false;
+}
+
+bool SPCGEditorGraphDebugObjectTree::CanJumpToGraphInTree(const FPCGEditorGraphDebugObjectItemPtr& Item) const
+{
+	// Offer jump-to command if any selected item is a graph or a subgraph loop iteration.
+	return Item && (Item->GetPCGGraph() || Item->IsLoopIteration());
+}
+
+void SPCGEditorGraphDebugObjectTree::FocusOnItem(const FPCGEditorGraphDebugObjectItemPtr& Item)
+{
+	if (!Item || !Item->GetPCGStack())
+	{
+		return;
+	}
+
+	const UPCGGraph* CurrentGraph = PCGEditor.Pin()->GetPCGGraph();
+	const FPCGStack& Stack = *Item->GetPCGStack();
+
+	bool bFoundCurrentGraph = false;
+	const UPCGNode* LastNode = nullptr;
+
+	// Find first node after the current graph in the stack.
+	for (int i = Stack.GetStackFrames().Num() - 1; i >= 0; --i)
+	{
+		const FPCGStackFrame& StackFrame = Stack.GetStackFrames()[i];
+		const UObject* Object = StackFrame.Object.Get();
+
+		if (Object == CurrentGraph)
+		{
+			bFoundCurrentGraph = true;
+			break;
+		}
+		else if (const UPCGNode* Node = Cast<const UPCGNode>(Object))
+		{
+			LastNode = Node;
+		}
+	}
+
+	if (bFoundCurrentGraph && LastNode)
+	{
+		PCGEditor.Pin()->JumpToNode(LastNode);
+	}
+}
+
+bool SPCGEditorGraphDebugObjectTree::CanFocusOnItem(const FPCGEditorGraphDebugObjectItemPtr& Item) const
+{
+	if (!Item)
+	{
+		return false;
+	}
+
+	// Anything downstream from the currently editable graph (e.g. "IsDebuggable" is true) is avaiable for focus.
+	FPCGEditorGraphDebugObjectItemPtr Parent = Item->GetParent();
+	while (Parent)
+	{
+		if (Parent->IsDebuggable())
+		{
+			return true;
+		}
+		else
+		{
+			Parent = Parent->GetParent();
+		}
+	}
+
+	return false;
+}
+
+#undef LOCTEXT_NAMESPACE
diff --git a/Engine/Plugins/PCG/Source/PCGEditor/Private/Widgets/SPCGEditorGraphDebugObjectTree.h b/Engine/Plugins/PCG/Source/PCGEditor/Private/Widgets/SPCGEditorGraphDebugObjectTree.h
new file mode 100644
--- /dev/null
+++ b/Engine/Plugins/PCG/Source/PCGEditor/Private/Widgets/SPCGEditorGraphDebugObjectTree.h
@@ -0,0 +1,355 @@
+// Copyright Epic Games, Inc. All Rights Reserved.
+
+#pragma once
+
+#include "PCGComponent.h"
+#include "PCGGraph.h"
+#include "PCGNode.h"
+#include "Graph/PCGStackContext.h"
+
+#include "ToolMenus.h"
+#include "GameFramework/Actor.h"
+#include "Misc/Optional.h"
+#include "Widgets/SCompoundWidget.h"
+#include "Widgets/Views/STreeView.h"
+
+class FPCGEditorGraphDebugObjectItem;
+class FPCGEditor;
+
+typedef TSharedPtr<FPCGEditorGraphDebugObjectItem> FPCGEditorGraphDebugObjectItemPtr;
+
+class FPCGEditorGraphDebugObjectItem : public TSharedFromThis<FPCGEditorGraphDebugObjectItem>
+{
+public:
+	FPCGEditorGraphDebugObjectItem() = default;
+
+	explicit FPCGEditorGraphDebugObjectItem(bool bInGrayedOut)
+		: bGrayedOut(bInGrayedOut)
+	{
+	}
+	
+	virtual ~FPCGEditorGraphDebugObjectItem() = default;
+
+	void AddChild(TSharedRef<FPCGEditorGraphDebugObjectItem> InChild);
+	const TSet<TSharedPtr<FPCGEditorGraphDebugObjectItem>>& GetChildren() const;
+	FPCGEditorGraphDebugObjectItemPtr GetParent() const;
+	void SortChildren(bool bIsAscending, bool bIsRecursive);
+
+	bool IsExpanded() const { return bIsExpanded; }
+	void SetExpanded(bool bInIsExpanded) { bIsExpanded = bInIsExpanded; }
+
+	bool IsGrayedOut() const { return bGrayedOut; }
+
+	bool IsSelected() const { return bSelected; }
+	void SetSelected(bool bInSelected) { bSelected = bInSelected; }
+
+	/** Optional sort priority, if returns INDEX_NONE then sort will fall back to alphabetical. */
+	virtual int32 GetSortPriority() const { return INDEX_NONE; }
+
+	/** Whether this item represents a currently debuggable object for the current edited graph. */
+	virtual bool IsDebuggable() const { return false; }
+	/** Whether this item represents a dynamically executed element (as in dynamic subgraph). */
+	virtual bool IsDynamic() const { return false; }
+
+	virtual FString GetLabel(bool bForSorting = false) const = 0;
+	virtual UPCGComponent* GetPCGComponent() const = 0;
+	virtual const FPCGStack* GetPCGStack() const { return nullptr; }
+	virtual const UObject* GetObject() const  = 0;
+	virtual const UPCGGraph* GetPCGGraph() const { return nullptr; }
+	virtual bool IsLoopIteration() const { return false; }
+	virtual const FSlateBrush* GetIcon() const;
+
+	virtual void OnObjectsReplaced(const TMap<UObject*, UObject*>& ReplacementMap);
+
+protected:
+	virtual FPCGStack* GetMutablePCGStack() { return nullptr; }
+
+	TWeakPtr<FPCGEditorGraphDebugObjectItem> Parent;
+	TSet<TSharedPtr<FPCGEditorGraphDebugObjectItem>> Children;
+	bool bIsExpanded = false;
+	bool bGrayedOut = false;
+	bool bSelected = false;
+};
+
+class FPCGEditorGraphDebugObjectItem_Actor : public FPCGEditorGraphDebugObjectItem
+{
+public:
+	explicit FPCGEditorGraphDebugObjectItem_Actor(TWeakObjectPtr<AActor> InActor, bool bInHasInspectionData)
+		: FPCGEditorGraphDebugObjectItem(bInHasInspectionData)
+		, Actor(InActor)
+	{
+		PCGStack.PushFrame(InActor.Get());
+	}
+
+	virtual FString GetLabel(bool bForSorting = false) const override;
+	virtual UPCGComponent* GetPCGComponent() const override { return nullptr; }
+	virtual const UObject* GetObject() const override { return Actor.Get(); };
+	virtual const FPCGStack* GetPCGStack() const override { return &PCGStack; }
+	virtual const FSlateBrush* GetIcon() const override;
+
+protected:
+	virtual FPCGStack* GetMutablePCGStack() override { return &PCGStack; }
+
+	TWeakObjectPtr<AActor> Actor = nullptr;
+
+	FPCGStack PCGStack;
+};
+
+class FPCGEditorGraphDebugObjectItem_PCGComponent : public FPCGEditorGraphDebugObjectItem
+{
+public:
+	explicit FPCGEditorGraphDebugObjectItem_PCGComponent(
+		TWeakObjectPtr<UPCGComponent> InPCGComponent,
+		TWeakObjectPtr<const UPCGGraph> InPCGGraph,
+		const FPCGStack& InPCGStack,
+		bool bInIsDebuggable,
+		bool bInHasInspectionData)
+		: FPCGEditorGraphDebugObjectItem(bInHasInspectionData)
+		, PCGComponent(InPCGComponent)
+		, PCGGraph(InPCGGraph)
+		, PCGStack(InPCGStack)
+		, bIsDebuggable(bInIsDebuggable)
+	{}
+
+	virtual FString GetLabel(bool bForSorting = false) const override;
+	virtual UPCGComponent* GetPCGComponent() const override { return PCGComponent.Get(); }
+	virtual const UObject* GetObject() const override { return PCGComponent.Get(); };
+	virtual const FPCGStack* GetPCGStack() const override { return &PCGStack; }
+	virtual bool IsDebuggable() const override { return bIsDebuggable; }
+	virtual const UPCGGraph* GetPCGGraph() const override { return PCGGraph.Get(); }
+
+	virtual void OnObjectsReplaced(const TMap<UObject*, UObject*>& ReplacementMap) override;
+
+protected:
+	virtual FPCGStack* GetMutablePCGStack() override { return &PCGStack; }
+
+	TWeakObjectPtr<UPCGComponent> PCGComponent = nullptr;
+	TWeakObjectPtr<const UPCGGraph> PCGGraph = nullptr;
+
+	FPCGStack PCGStack;
+	bool bIsDebuggable = false;
+};
+
+enum class EPCGEditorSubgraphNodeType : uint8
+{
+	StaticSubgraph = 0,
+	DynamicSubgraph,
+	LoopSubgraph
+};
+
+class FPCGEditorGraphDebugObjectItem_PCGSubgraph : public FPCGEditorGraphDebugObjectItem
+{
+public:
+	explicit FPCGEditorGraphDebugObjectItem_PCGSubgraph(
+		TWeakObjectPtr<const UPCGNode> InPCGNode,
+		TWeakObjectPtr<const UPCGGraph> InPCGGraph,
+		const FPCGStack& InPCGStack,
+		bool bInIsDebuggable,
+		bool bInHasInspectionData,
+		EPCGEditorSubgraphNodeType InSubgraphType)
+		: FPCGEditorGraphDebugObjectItem(bInHasInspectionData)
+		, PCGNode(InPCGNode)
+		, PCGGraph(InPCGGraph)
+		, PCGStack(InPCGStack)
+		, bIsDebuggable(bInIsDebuggable)
+		, SubgraphType(InSubgraphType)
+	{}
+
+	virtual FString GetLabel(bool bForSorting = false) const override;
+	virtual UPCGComponent* GetPCGComponent() const override { return GetParent() ? GetParent()->GetPCGComponent() : nullptr; }
+	virtual const UObject* GetObject() const override { return PCGNode.Get(); };
+	virtual const FPCGStack* GetPCGStack() const override { return &PCGStack; }
+	virtual bool IsDebuggable() const override { return bIsDebuggable; }
+	virtual bool IsDynamic() const override { return SubgraphType != EPCGEditorSubgraphNodeType::StaticSubgraph; }
+	virtual const UPCGGraph* GetPCGGraph() const override { return PCGGraph.Get(); }
+	virtual const FSlateBrush* GetIcon() const override;
+
+protected:
+	virtual FPCGStack* GetMutablePCGStack() override { return &PCGStack; }
+
+	TWeakObjectPtr<const UPCGNode> PCGNode = nullptr;
+	TWeakObjectPtr<const UPCGGraph> PCGGraph = nullptr;
+
+	FPCGStack PCGStack;
+	bool bIsDebuggable = false;
+	EPCGEditorSubgraphNodeType SubgraphType = EPCGEditorSubgraphNodeType::StaticSubgraph;
+};
+
+class FPCGEditorGraphDebugObjectItem_PCGLoopIndex : public FPCGEditorGraphDebugObjectItem
+{
+public:
+	explicit FPCGEditorGraphDebugObjectItem_PCGLoopIndex(
+		int32 InLoopIndex,
+		TWeakObjectPtr<const UObject> InLoopedPCGGraph,
+		const FPCGStack& InPCGStack,
+		bool bInIsDebuggable,
+		bool bInHasInspectionData)
+		: FPCGEditorGraphDebugObjectItem(bInHasInspectionData)
+		, LoopIndex(InLoopIndex)
+		, LoopedPCGGraph(InLoopedPCGGraph)
+		, PCGStack(InPCGStack)
+		, bIsDebuggable(bInIsDebuggable)
+	{}
+
+	virtual int32 GetLoopIndex() const { return LoopIndex; }
+
+	virtual FString GetLabel(bool bForSorting = false) const override;
+	virtual UPCGComponent* GetPCGComponent() const override { return GetParent() ? GetParent()->GetPCGComponent() : nullptr; }
+	virtual const UObject* GetObject() const override { return nullptr; };
+	virtual const FPCGStack* GetPCGStack() const override { return &PCGStack; }
+	virtual int32 GetSortPriority() const override { return LoopIndex; }
+	virtual bool IsDebuggable() const override { return bIsDebuggable; }
+	virtual bool IsLoopIteration() const override { return true; }
+	virtual const UPCGGraph* GetPCGGraph() const override { return Cast<UPCGGraph>(LoopedPCGGraph.Get()); }
+
+protected:
+	virtual FPCGStack* GetMutablePCGStack() override { return &PCGStack; }
+
+	int32 LoopIndex = INDEX_NONE;
+	TWeakObjectPtr<const UObject> LoopedPCGGraph = nullptr;
+
+	FPCGStack PCGStack;
+	bool bIsDebuggable = false;
+};
+
+class SPCGEditorGraphDebugObjectItemRow : public SCompoundWidget
+{
+public:
+	DECLARE_DELEGATE_OneParam(FPCGDebugObjectItemRowAction, const FPCGEditorGraphDebugObjectItemPtr&);
+	DECLARE_DELEGATE_RetVal_OneParam(bool, FPCGDebugObjectItemRowPredicate, const FPCGEditorGraphDebugObjectItemPtr&);
+
+	SLATE_BEGIN_ARGS(SPCGEditorGraphDebugObjectItemRow)
+	{}
+		SLATE_EVENT(FPCGDebugObjectItemRowAction, OnDoubleClick)
+		SLATE_EVENT(FPCGDebugObjectItemRowAction, OnJumpTo)
+		SLATE_EVENT(FPCGDebugObjectItemRowPredicate, CanJumpTo)
+		SLATE_EVENT(FPCGDebugObjectItemRowAction, OnFocus)
+		SLATE_EVENT(FPCGDebugObjectItemRowPredicate, CanFocus)
+	SLATE_END_ARGS()
+
+	void Construct(const FArguments& InArgs, const TSharedRef<STableViewBase>& InOwnerTableView, FPCGEditorGraphDebugObjectItemPtr InItem);
+
+	//~Begin SWidget Interface
+	virtual FReply OnMouseButtonDoubleClick(const FGeometry& InMyGeometry, const FPointerEvent& InMouseEvent) override;
+	//~End SWidget Interface
+
+	FReply FocusClicked() const;
+	FReply JumpToClicked() const;
+	bool IsJumpToEnabled() const;
+	bool IsFocusEnabled() const;
+
+private:
+	FPCGEditorGraphDebugObjectItemPtr Item;
+
+	/** Invoked when the user double clicks on the row. */
+	FPCGDebugObjectItemRowAction OnDoubleClick;
+	/** Invoked when the 'Jump To' action is clicked on the row buttons. */
+	FPCGDebugObjectItemRowAction OnJumpTo;
+	/** Invoked when the 'Go to node' action is clicked on the row buttons. */
+	FPCGDebugObjectItemRowAction OnFocus;
+	/** Controls whether the jump to button will be enabled or not. */
+	FPCGDebugObjectItemRowPredicate CanJumpTo;
+	/** Controls whether the 'go to node' button will be enabled or not. */
+	FPCGDebugObjectItemRowPredicate CanFocus;
+};
+
+class SPCGEditorGraphDebugObjectTree : public SCompoundWidget
+{
+public:
+	SLATE_BEGIN_ARGS(SPCGEditorGraphDebugObjectTree) {}
+	SLATE_END_ARGS()
+
+	virtual ~SPCGEditorGraphDebugObjectTree();
+
+	void Construct(const FArguments& InArgs, TSharedPtr<FPCGEditor> InPCGEditor);
+
+	virtual void Tick(const FGeometry& AllottedGeometry, const double InCurrentTime, const float InDeltaTime) override;
+
+	void RequestRefresh();
+
+	void SetNodeBeingInspected(const UPCGNode* InPCGNode);
+
+	void OnObjectsReplaced(const TMap<UObject*, UObject*>& ReplacementMap);
+
+	/** If the stack matches to an item, Expands the tree view to make that item visible and select it, and return true. Otherwise, return false. */
+	bool SetDebugObjectFromStackFromAnotherEditor(const FPCGStack& InStack);
+
+	/** Return the first stack downstream of the currently selected stack with the provided node & graph. Return false if nothing is found. */
+	bool GetFirstStackFromSelection(const UPCGNode* InNode, const UPCGGraph* InGraph, FPCGStack& OutStack) const;
+
+private:
+	FReply FocusOnDebugObject_OnClicked() const;
+	bool IsFocusOnDebugObjectButtonEnabled() const;
+
+	void SetDebugObjectFromSelection_OnClicked();
+	bool IsSetDebugObjectFromSelectionButtonEnabled() const { return IsSetDebugObjectFromSelectionEnabled.IsSet() && IsSetDebugObjectFromSelectionEnabled.GetValue(); }
+	void UpdateIsSetDebugObjectFromSelectionEnabled();
+
+	void RefreshTree();
+	void SortTreeItems(bool bIsAscending = true, bool bIsRecursive = true);
+	void RestoreTreeState();
+
+	void AddStacksToTree(
+		const TArray<FPCGStack>& Stacks,
+		TMap<AActor*, TSharedPtr<FPCGEditorGraphDebugObjectItem_Actor>>& InOutActorItems,
+		TMap<const FPCGStack, FPCGEditorGraphDebugObjectItemPtr>& InOutStackToItem);
+
+	void OnPreObjectPropertyChanged(UObject* InObject, const FEditPropertyChain& InPropertyChain);
+	void OnObjectPropertyChanged(UObject* InObject, FPropertyChangedEvent& InPropertyChangedEvent);
+	void OnObjectConstructed(UObject* InObject);
+	void OnEditorSelectionChanged(UObject* InObject);
+
+	UPCGGraph* GetPCGGraph() const;
+
+	TSharedRef<ITableRow> MakeTreeRowWidget(FPCGEditorGraphDebugObjectItemPtr InItem, const TSharedRef<STableViewBase>& InOwnerTable);
+	void OnGetChildren(FPCGEditorGraphDebugObjectItemPtr InItem, TArray<FPCGEditorGraphDebugObjectItemPtr>& OutChildren) const;
+	void OnSelectionChanged(FPCGEditorGraphDebugObjectItemPtr InItem, ESelectInfo::Type InSelectInfo);
+	void OnExpansionChanged(FPCGEditorGraphDebugObjectItemPtr InItem, bool bInIsExpanded);
+	void OnSetExpansionRecursive(FPCGEditorGraphDebugObjectItemPtr InItem, bool bInExpand) const;
+
+	/** Returns the matching item from the AllGraphItems array if any. */
+	FPCGEditorGraphDebugObjectItemPtr GetItemFromStack(const FPCGStack& InStack) const;
+	/** Expands the tree view to make an item visible and select it. */
+	void ExpandAndSelectDebugObject(const FPCGEditorGraphDebugObjectItemPtr& InItem);
+
+	/** Expand the given row and select the next occurrence of the current graph. */
+	void ExpandAndSelectFirstLeafDebugObject(const FPCGEditorGraphDebugObjectItemPtr& InItem);
+
+	TSharedPtr<SWidget> OpenContextMenu();
+
+	/** Jump-to context menu command. */
+	void ContextMenu_JumpToGraphInTree();
+	bool ContextMenu_JumpToGraphInTree_CanExecute() const;
+
+	void JumpToGraphInTree(const FPCGEditorGraphDebugObjectItemPtr& Item);
+	bool CanJumpToGraphInTree(const FPCGEditorGraphDebugObjectItemPtr& Item) const;
+	void FocusOnItem(const FPCGEditorGraphDebugObjectItemPtr& Item);
+	bool CanFocusOnItem(const FPCGEditorGraphDebugObjectItemPtr& Item) const;
+
+	TWeakPtr<FPCGEditor> PCGEditor;
+
+	TSharedPtr<STreeView<FPCGEditorGraphDebugObjectItemPtr>> DebugObjectTreeView;
+	TArray<FPCGEditorGraphDebugObjectItemPtr> RootItems;
+	TArray<FPCGEditorGraphDebugObjectItemPtr> AllGraphItems;
+
+	bool bNeedsRefresh = false;
+
+	double NextRefreshTime = 0.0;
+
+	/** Latest value for IsSetDebugObjectFromSelectionButtonEnabled */
+	TOptional<bool> IsSetDebugObjectFromSelectionEnabled;
+
+	/** Used to retain item expansion state across tree refreshes. */
+	TSet<FPCGStack> ExpandedStacks;
+
+	/** Used to retain item selection state across tree refreshes. */
+	FPCGStack SelectedStack;
+
+	/** Used to retain item selection state across tree refreshes if the SelectedStack is invalidated (e.g. through BP reconstruction). */
+	TWeakObjectPtr<const UPCGGraph> SelectedGraph = nullptr;
+	TWeakObjectPtr<const AActor> SelectedOwner = nullptr;
+	uint32 SelectedGridSize = PCGHiGenGrid::UnboundedGridSize();
+	FIntVector SelectedGridCoord = FIntVector::ZeroValue;
+	TWeakObjectPtr<const UPCGComponent> SelectedOriginalComponent = nullptr;
+	TWeakObjectPtr<const UPCGNode> PCGNodeBeingInspected = nullptr;
+};
diff --git a/Engine/Plugins/PCG/Source/PCGEditor/Private/Widgets/SPCGEditorGraphDetailsView.cpp b/Engine/Plugins/PCG/Source/PCGEditor/Private/Widgets/SPCGEditorGraphDetailsView.cpp
new file mode 100644
--- /dev/null
+++ b/Engine/Plugins/PCG/Source/PCGEditor/Private/Widgets/SPCGEditorGraphDetailsView.cpp
@@ -0,0 +1,284 @@
+// Copyright Epic Games, Inc. All Rights Reserved.
+
+#include "SPCGEditorGraphDetailsView.h"
+
+#include "PCGGraph.h"
+#include "PCGNode.h"
+#include "PCGSettings.h"
+
+#include "PCGEditor.h"
+#include "Nodes/PCGEditorGraphNode.h"
+
+#include "DetailsViewArgs.h"
+#include "IDetailsView.h"
+#include "PropertyEditorModule.h"
+#include "Layout/Visibility.h"
+#include "Modules/ModuleManager.h"
+#include "Widgets/DeclarativeSyntaxSupport.h"
+#include "Widgets/Images/SImage.h"
+#include "Widgets/Input/SButton.h"
+#include "Widgets/Text/STextBlock.h"
+
+#define LOCTEXT_NAMESPACE "PCGEditorGraphDetailsView"
+
+void SPCGEditorGraphDetailsView::Construct(const FArguments& InArgs)
+{
+	FDetailsViewArgs DetailsViewArgs;
+	DetailsViewArgs.NameAreaSettings = FDetailsViewArgs::HideNameArea;
+	DetailsViewArgs.bCustomFilterAreaLocation = true;
+	DetailsViewArgs.bHideSelectionTip = true;
+	DetailsViewArgs.DefaultsOnlyVisibility = EEditDefaultsOnlyNodeVisibility::Automatic;
+
+	FPropertyEditorModule& PropertyEditorModule = FModuleManager::GetModuleChecked<FPropertyEditorModule>("PropertyEditor");
+	DetailsView = PropertyEditorModule.CreateDetailView(DetailsViewArgs);
+
+	DetailsView->SetIsPropertyReadOnlyDelegate(FIsPropertyReadOnly::CreateSP(this, &SPCGEditorGraphDetailsView::IsReadOnlyProperty));
+	DetailsView->SetIsPropertyVisibleDelegate(FIsPropertyVisible::CreateSP(this, &SPCGEditorGraphDetailsView::IsVisibleProperty));
+
+	ChildSlot
+	[
+		SNew(SVerticalBox)
+		+SVerticalBox::Slot()
+		.Padding(0.f, 0.f, 0.f, 0.f)
+		.AutoHeight()
+		[
+			SNew(SHorizontalBox)
+			+SHorizontalBox::Slot()
+			.AutoWidth()
+			[
+				SNew(SButton)
+				.ButtonStyle(FAppStyle::Get(), "SimpleButton")
+				.OnClicked(this, &SPCGEditorGraphDetailsView::OnLockButtonClicked)
+				.ContentPadding(FMargin(4, 2))
+				.HAlign(HAlign_Center)
+				.VAlign(VAlign_Center)
+				.ToolTipText(LOCTEXT("LockSelectionButton_ToolTip", "Locks the current attribute list view to this selection"))
+				[
+					SNew(SImage)
+					.ColorAndOpacity(FSlateColor::UseForeground())
+					.Image(this, &SPCGEditorGraphDetailsView::GetLockIcon)
+				]
+			]
+			+SHorizontalBox::Slot()
+			.AutoWidth()
+			.HAlign(HAlign_Center)
+			.VAlign(VAlign_Center)
+			[
+				SNew(SButton)
+				.ButtonStyle(FAppStyle::Get(), "SimpleButton")
+				.OnClicked(this, &SPCGEditorGraphDetailsView::OnNameClicked)
+				.Visibility(this, &SPCGEditorGraphDetailsView::GetNameVisibility)
+				[
+					SNew(STextBlock)
+					.Text(this, &SPCGEditorGraphDetailsView::GetName)
+				]
+			]
+			+SHorizontalBox::Slot()
+			[
+				DetailsView->GetFilterAreaWidget().ToSharedRef()
+			]
+		]
+		+SVerticalBox::Slot()
+		[
+			DetailsView.ToSharedRef()
+		]
+	];
+}
+
+void SPCGEditorGraphDetailsView::SetObject(UObject* InObject, bool bForceRefresh)
+{
+	TArray<TWeakObjectPtr<UObject>> InObjects;
+	InObjects.Add(InObject);
+
+	SetObjects(InObjects, bForceRefresh);
+}
+
+void SPCGEditorGraphDetailsView::SetObjects(const TArray<TWeakObjectPtr<UObject>>& InObjects, bool bForceRefresh, bool bOverrideLock)
+{
+	if (bOverrideLock || !bIsLocked)
+	{
+		SelectedObjects = InObjects;
+
+		// Filter only the types we're interested in, e.g. the settings and not the nodes themselves
+		TArray<TWeakObjectPtr<UObject>> ObjectsToView;
+		ObjectsToView.Reserve(InObjects.Num());
+
+		for (const TWeakObjectPtr<UObject>& InObject : InObjects)
+		{
+			if (!InObject.IsValid())
+			{
+				continue;
+			}
+
+			if (UPCGEditorGraphNodeBase* PCGGraphNode = Cast<UPCGEditorGraphNodeBase>(InObject.Get()))
+			{
+				if (UPCGNode* PCGNode = PCGGraphNode->GetPCGNode())
+				{
+					if (PCGNode->IsInstance())
+					{
+						ObjectsToView.Add(PCGNode->GetSettingsInterface());
+					}
+					else
+					{
+						ObjectsToView.Add(PCGNode->GetSettings());
+					}
+
+					continue;
+				}
+			}
+
+			ObjectsToView.Add(InObject);
+		}
+
+		DetailsView->SetObjects(ObjectsToView);
+
+		if (InObjects.IsEmpty())
+		{
+			bIsLocked = false;
+		}
+	}
+}
+
+const TArray<TWeakObjectPtr<UObject>>& SPCGEditorGraphDetailsView::GetSelectedObjects() const
+{
+	return SelectedObjects;
+}
+
+FReply SPCGEditorGraphDetailsView::OnLockButtonClicked()
+{
+	bIsLocked = !bIsLocked;
+	return FReply::Handled();
+}
+
+FReply SPCGEditorGraphDetailsView::OnNameClicked()
+{
+	if (SelectedObjects.Num() != 1)
+	{
+		return FReply::Handled();
+	}
+
+	if (UEdGraphNode* GraphNode = Cast<UEdGraphNode>(SelectedObjects[0].Get()))
+	{
+		if (FPCGEditor* Editor = EditorPtr.Pin().Get())
+		{
+			Editor->JumpToNode(GraphNode);
+		}
+	}
+
+	return FReply::Handled();
+}
+
+const FSlateBrush* SPCGEditorGraphDetailsView::GetLockIcon() const
+{
+	return FAppStyle::GetBrush(bIsLocked ? TEXT("PropertyWindow.Locked") : TEXT("PropertyWindow.Unlocked"));
+}
+
+EVisibility SPCGEditorGraphDetailsView::GetNameVisibility() const
+{
+	return bIsLocked ? EVisibility::Visible : EVisibility::Collapsed;
+}
+
+FText SPCGEditorGraphDetailsView::GetName() const
+{
+	const TArray<TWeakObjectPtr<UObject>>& DetailsSelectedObjects = DetailsView->GetSelectedObjects();
+	if (DetailsSelectedObjects.IsEmpty())
+	{
+		// Implementation note: this doesn't really happen since we reselect the graph settings when unselecting
+		return LOCTEXT("NoObjectsSelected", "Empty selection");
+	}
+	else if (DetailsSelectedObjects.Num() == 1)
+	{
+		UObject* SelectedObject = DetailsSelectedObjects[0].Get();
+		UPCGNode* OwnerNode = nullptr;
+
+		if (Cast<UPCGSettings>(SelectedObject) != nullptr || Cast<UPCGSettingsInstance>(SelectedObject) != nullptr)
+		{
+			OwnerNode = Cast<UPCGNode>(SelectedObject->GetOuter());
+		}
+
+		if (OwnerNode)
+		{
+			return OwnerNode->GetNodeTitle(EPCGNodeTitleType::ListView);
+		}
+		else if (UPCGGraph* Graph = Cast<UPCGGraph>(SelectedObject))
+		{
+			return LOCTEXT("GraphSettingsSelected", "Graph Settings");
+		}
+		else
+		{
+			return FText::FromName(SelectedObject->GetFName());
+		}
+	}
+	else
+	{
+		return FText::Format(LOCTEXT("MultipleObjectsSelectedFmt", "{0} nodes"), FText::AsNumber(DetailsSelectedObjects.Num()));
+	}
+}
+
+bool SPCGEditorGraphDetailsView::IsReadOnlyProperty(const FPropertyAndParent& InPropertyAndParent) const
+{
+	// Everything is writable when not in an instance
+	if (!DetailsView ||
+		InPropertyAndParent.ParentProperties.IsEmpty() ||
+		InPropertyAndParent.ParentProperties.Last()->GetFName() != GET_MEMBER_NAME_CHECKED(UPCGSettingsInstance, Settings))
+	{
+		return false;
+	}
+
+	const TArray<TWeakObjectPtr<UObject>>& DetailsSelectedObjects = DetailsView->GetSelectedObjects();
+	for (const TWeakObjectPtr<UObject>& SelectedObject : DetailsSelectedObjects)
+	{
+		if (!SelectedObject.IsValid())
+		{
+			continue;
+		}
+
+		if (UPCGSettingsInstance* Instance = Cast<UPCGSettingsInstance>(SelectedObject.Get()))
+		{
+			return true;
+		}
+	}
+
+	return false;
+}
+
+bool SPCGEditorGraphDetailsView::IsVisibleProperty(const FPropertyAndParent& InPropertyAndParent) const
+{
+	if (!DetailsView)
+	{
+		return true;
+	}
+
+	// Currently never hide anything from the graph settings
+	if (InPropertyAndParent.Objects.Num() == 1 && Cast<UPCGGraph>(InPropertyAndParent.Objects[0]))
+	{
+		return true;
+	}
+
+	// Always hide asset info information
+	if (InPropertyAndParent.Property.HasMetaData(TEXT("Category")) &&
+		InPropertyAndParent.Property.GetMetaData(TEXT("Category")) == TEXT("AssetInfo"))
+	{
+		return false;
+	}
+
+	// Otherwise, everything is visible when not in an instance
+	if (InPropertyAndParent.ParentProperties.IsEmpty() ||
+		InPropertyAndParent.ParentProperties.Last()->GetFName() != GET_MEMBER_NAME_CHECKED(UPCGSettingsInstance, Settings))
+	{
+		return true;
+	}
+
+	// Hide debug settings from the setting when showing the instance settings.
+	if (InPropertyAndParent.Property.GetFName() == GET_MEMBER_NAME_CHECKED(UPCGSettings, bEnabled) ||
+		InPropertyAndParent.Property.GetFName() == GET_MEMBER_NAME_CHECKED(UPCGSettings, bDebug) ||
+		(InPropertyAndParent.ParentProperties.Num() >= 2 &&
+			InPropertyAndParent.ParentProperties[InPropertyAndParent.ParentProperties.Num() - 2]->GetFName() == GET_MEMBER_NAME_CHECKED(UPCGSettings, DebugSettings)))
+	{
+		return false;
+	}
+
+	return true;
+}
+
+#undef LOCTEXT_NAMESPACE
\ No newline at end of file
diff --git a/Engine/Plugins/PCG/Source/PCGEditor/Private/Widgets/SPCGEditorGraphDetailsView.h b/Engine/Plugins/PCG/Source/PCGEditor/Private/Widgets/SPCGEditorGraphDetailsView.h
new file mode 100644
--- /dev/null
+++ b/Engine/Plugins/PCG/Source/PCGEditor/Private/Widgets/SPCGEditorGraphDetailsView.h
@@ -0,0 +1,66 @@
+// Copyright Epic Games, Inc. All Rights Reserved.
+
+#pragma once
+
+#include "Templates/SharedPointer.h"
+#include "Widgets/SCompoundWidget.h"
+
+class FPCGEditor;
+class IDetailsView;
+struct FPropertyAndParent;
+
+/** Thin wrapper on top of a SDetailsView to add some PCG-specific logic around locking and improve layout */
+class SPCGEditorGraphDetailsView : public SCompoundWidget
+{
+public:
+	SLATE_BEGIN_ARGS(SPCGEditorGraphDetailsView)
+		{}	
+	SLATE_END_ARGS()
+
+	/** Constructs this widget with InArgs */
+	void Construct(const FArguments& InArgs);
+
+	/** Sets the editor associated to this details view, used for interaction from the details view to the editor */
+	void SetEditor(TWeakPtr<FPCGEditor> InEditorPtr) { EditorPtr = InEditorPtr; }
+
+	/** Returns the details view under this widget */
+	TSharedPtr<IDetailsView> GetDetailsView() const { return DetailsView; }
+
+	//~Begin IDetailsView-like interface
+	/** Returns the lock state of the details view, e.g. whether the selected object should change on a new selection */
+	bool IsLocked() const { return bIsLocked; }
+
+	/** Changes the selected objects in the details view. Will unlock automatically. */
+	void SetObject(UObject* InObject, bool bForceRefresh = false);
+
+	/** Changes the selected objects in the details view. Will unlock automatically. */
+	void SetObjects(const TArray<TWeakObjectPtr<UObject>>& InObjects, bool bForceRefresh = false, bool bOverrideLock = false);
+
+	/** Retrieves the list of objects currently shown in the details view */
+	const TArray<TWeakObjectPtr<UObject>>& GetSelectedObjects() const;
+	//~End IDetailsView-like interface
+
+	/** Controls lock from the editor */
+	void SetIsLocked(bool bInIsLocked) { bIsLocked = bInIsLocked; }
+
+protected:
+	/** Returns whether a property should be readonly (used for instances) */
+	bool IsReadOnlyProperty(const FPropertyAndParent& InPropertyAndParent) const;
+	/** Returns whether a property should be visible (used for instance vs. settings properties) */
+	bool IsVisibleProperty(const FPropertyAndParent& InPropertyAndParent) const;
+
+	FReply OnLockButtonClicked();
+	FReply OnNameClicked();
+	const FSlateBrush* GetLockIcon() const;
+
+	/** Gets the current visibility state of the name section, which is related to whether the details view is locked. */
+	EVisibility GetNameVisibility() const;
+	/** Returns the name to show in the details-view name section, e.g. the settings title most of the time. */
+	FText GetName() const;
+
+	TSharedPtr<IDetailsView> DetailsView;
+	bool bIsLocked : 1 = false;
+
+	TArray<TWeakObjectPtr<UObject>> SelectedObjects;
+	TWeakPtr<FPCGEditor> EditorPtr = nullptr;
+};
\ No newline at end of file
diff --git a/Engine/Plugins/PCG/Source/PCGEditor/Private/Widgets/SPCGEditorGraphDeterminism.cpp b/Engine/Plugins/PCG/Source/PCGEditor/Private/Widgets/SPCGEditorGraphDeterminism.cpp
new file mode 100644
--- /dev/null
+++ b/Engine/Plugins/PCG/Source/PCGEditor/Private/Widgets/SPCGEditorGraphDeterminism.cpp
@@ -0,0 +1,275 @@
+// Copyright Epic Games, Inc. All Rights Reserved.
+
+#include "SPCGEditorGraphDeterminism.h"
+
+#include "Framework/Views/TableViewMetadata.h"
+#include "Tests/Determinism/PCGDeterminismTestsCommon.h"
+#include "Widgets/Views/SListView.h"
+
+#define LOCTEXT_NAMESPACE "PCGDeterminism"
+
+namespace
+{
+	const FName NAME_Index(TEXT("Index_ColumnID"));
+	const FName NAME_NodeTitle(TEXT("NodeTitle_ColumnID"));
+	const FName NAME_NodeName(TEXT("NodeName_ColumnID"));
+	const FName NAME_Seed(TEXT("Seed_ColumnID"));
+	const FName NAME_DataTypesTested(TEXT("DataTypesTested_ColumnID"));
+	const FName NAME_AdditionalDetails(TEXT("AdditionalDetails_ColumnID"));
+
+	const FText TEXT_Index;
+	const FText TEXT_NodeTitle(LOCTEXT("NodeTitle_Label", "Title"));
+	const FText TEXT_NodeName(LOCTEXT("NodeName_Label", "Name"));
+	const FText TEXT_Seed(LOCTEXT("Seed_Label", "Seed"));
+	const FText TEXT_DataTypesTested(LOCTEXT("DataTypesTested_Label", "Input Data"));
+	const FText TEXT_AdditionalDetails(LOCTEXT("AdditionalDetails_Label", "Additional Details"));
+
+	const FText TEXT_NotDeterministic(LOCTEXT("NotDeterministic", "Fail"));
+	const FText TEXT_Consistent(LOCTEXT("OrderConsistent", "Order Consistent"));
+	const FText TEXT_Independent(LOCTEXT("OrderIndependent", "Order Independent"));
+	const FText TEXT_Orthogonal(LOCTEXT("OrderOrthogonal", "Order Orthogonal"));
+	const FText TEXT_Basic(LOCTEXT("BasicDeterminism", "Pass"));
+
+	constexpr float SmallManualWidth = 25.f;
+	constexpr float MediumManualWidth = 80.f;
+	constexpr float LargeManualWidth = 160.f;
+}
+
+void SPCGEditorGraphDeterminismRow::Construct(const FArguments& InArgs, const TSharedRef<STableViewBase>& InOwnerTableView, const FPCGNodeTestResultPtr& Item, int32 ItemIndex)
+{
+	CurrentItem = Item;
+	SMultiColumnTableRow<FPCGNodeTestResultPtr>::Construct(SMultiColumnTableRow::FArguments(), InOwnerTableView);
+}
+
+TSharedRef<SWidget> SPCGEditorGraphDeterminismRow::GenerateWidgetForColumn(const FName& ColumnId)
+{
+	check(CurrentItem);
+
+	FText CellText;
+
+	auto ReturnColorCodedResultBlock = [this](const FText& OutCellText)
+	{
+		return SNew(STextBlock)
+			.Text(OutCellText)
+			.ColorAndOpacity(CurrentItem->bFlagRaised ? FColor::Red : FColor::Green);
+	};
+
+	// Permanent columns
+	if (ColumnId == NAME_Index)
+	{
+		return ReturnColorCodedResultBlock(FText::FromString(FString::FromInt(CurrentItem->Index)));
+	}
+	else if (ColumnId == NAME_NodeTitle)
+	{
+		return ReturnColorCodedResultBlock(FText::FromName(CurrentItem->TestResultTitle));
+	}
+	else if (ColumnId == NAME_NodeName)
+	{
+		CellText = FText::FromString(CurrentItem->TestResultName);
+	}
+	else if (ColumnId == NAME_Seed)
+	{
+		CellText = FText::AsNumber(CurrentItem->Seed);
+	}
+	else if (ColumnId == NAME_DataTypesTested)
+	{
+		FString DataTypesTestedString = *UEnum::GetValueAsString(CurrentItem->DataTypesTested);
+		DataTypesTestedString.RemoveFromStart("EPCGDataType::");
+		CellText = FText::FromString(DataTypesTestedString);
+	}
+	else if (ColumnId == NAME_AdditionalDetails)
+	{
+		// TODO: Add tooltip widget to display multiple error messages, but only display the first in the row for now
+		FString FullDetails = CurrentItem->AdditionalDetails.Num() > 0 ? CurrentItem->AdditionalDetails[0] : "";
+		CellText = FText::FromString(FullDetails);
+	}
+	// Test columns
+	else if (EDeterminismLevel* DeterminismLevel = CurrentItem->TestResults.Find(ColumnId))
+	{
+		switch (*DeterminismLevel)
+		{
+		case EDeterminismLevel::OrderOrthogonal:
+			return SNew(STextBlock)
+				.Text(TEXT_Orthogonal)
+				.ColorAndOpacity(FColor::Orange);
+		case EDeterminismLevel::OrderConsistent:
+			return SNew(STextBlock)
+				.Text(TEXT_Consistent)
+				.ColorAndOpacity(FColor::Yellow);
+		case EDeterminismLevel::OrderIndependent:
+			return SNew(STextBlock)
+				.Text(TEXT_Independent)
+				.ColorAndOpacity(FColor::Green);
+		case EDeterminismLevel::Basic:
+			return SNew(STextBlock)
+				.Text(TEXT_Basic)
+				.ColorAndOpacity(FColor::Turquoise);
+		case EDeterminismLevel::NoDeterminism:
+		default:
+			return SNew(STextBlock)
+				.Text(TEXT_NotDeterministic)
+				.ColorAndOpacity(FColor::Red);
+		}
+	}
+
+	return SNew(STextBlock)
+		.Text(CellText);
+}
+
+void SPCGEditorGraphDeterminismListView::Construct(const FArguments& InArgs, TWeakPtr<FPCGEditor> InPCGEditor)
+{
+	check(InPCGEditor.IsValid() && !bIsConstructed);
+	PCGEditorPtr = InPCGEditor;
+
+	GeneratedHeaderRow = SNew(SHeaderRow);
+	SortMode = EColumnSortMode::None;
+	SortingColumn = NAME_None;
+
+	SAssignNew(ListView, SListView<FPCGNodeTestResultPtr>)
+		.ListItemsSource(&ListViewItems)
+		.OnGenerateRow(this, &SPCGEditorGraphDeterminismListView::OnGenerateRow)
+		.HeaderRow(GeneratedHeaderRow);
+
+	ChildSlot
+	[
+		ListView->AsShared()
+	];
+
+	bIsConstructed = true;
+}
+
+void SPCGEditorGraphDeterminismListView::AddItem(const FPCGNodeTestResultPtr& Item)
+{
+	check(Item.IsValid());
+	ListViewItems.Emplace(Item);
+}
+
+void SPCGEditorGraphDeterminismListView::ClearItems()
+{
+	ItemIndexCounter = -1;
+	ListViewItems.Empty();
+	RefreshItems();
+}
+
+void SPCGEditorGraphDeterminismListView::RefreshItems()
+{
+	ListView->RequestListRefresh();
+}
+
+void SPCGEditorGraphDeterminismListView::AddColumn(const FTestColumnInfo& ColumnInfo)
+{
+	SHeaderRow::FColumn::FArguments Arguments;
+	Arguments.ColumnId(ColumnInfo.ColumnID);
+	Arguments.DefaultLabel(ColumnInfo.ColumnLabel);
+	if (ColumnInfo.Width > 0.f)
+	{
+		Arguments.ManualWidth(ColumnInfo.Width);
+	}
+	Arguments.HAlignHeader(ColumnInfo.HAlign);
+	Arguments.HAlignCell(ColumnInfo.HAlign);
+	Arguments.SortMode(this, &SPCGEditorGraphDeterminismListView::GetColumnSortMode, ColumnInfo.ColumnID);
+	Arguments.OnSort(this, &SPCGEditorGraphDeterminismListView::OnSortColumnHeader);
+	GeneratedHeaderRow->AddColumn(Arguments);
+}
+
+void SPCGEditorGraphDeterminismListView::BuildBaseColumns()
+{
+	ClearColumns();
+	// Permanent columns
+	AddColumn({NAME_Index, TEXT_Index, SmallManualWidth, HAlign_Center});
+	AddColumn({NAME_NodeTitle, TEXT_NodeTitle, LargeManualWidth, HAlign_Left});
+	AddColumn({NAME_NodeName, TEXT_NodeName, LargeManualWidth, HAlign_Left});
+	AddColumn({NAME_Seed, TEXT_Seed, MediumManualWidth, HAlign_Center});
+	AddColumn({NAME_DataTypesTested, TEXT_DataTypesTested, MediumManualWidth, HAlign_Center});
+}
+
+void SPCGEditorGraphDeterminismListView::AddDetailsColumn()
+{
+	// Final details column
+	AddColumn({NAME_AdditionalDetails, TEXT_AdditionalDetails, 0.f, HAlign_Left});
+}
+
+bool SPCGEditorGraphDeterminismListView::WidgetIsConstructed() const
+{
+	return bIsConstructed;
+}
+
+void SPCGEditorGraphDeterminismListView::ClearColumns()
+{
+	GeneratedHeaderRow->ClearColumns();
+}
+
+TSharedRef<ITableRow> SPCGEditorGraphDeterminismListView::OnGenerateRow(const FPCGNodeTestResultPtr Item, const TSharedRef<STableViewBase>& OwnerTable) const
+{
+	return SNew(SPCGEditorGraphDeterminismRow, OwnerTable, Item, ++ItemIndexCounter);
+}
+
+void SPCGEditorGraphDeterminismListView::OnSortColumnHeader(const EColumnSortPriority::Type SortPriority, const FName& ColumnId, const EColumnSortMode::Type NewSortMode)
+{
+	if (SortingColumn == ColumnId)
+	{
+		// Circling
+		SortMode = static_cast<EColumnSortMode::Type>((SortMode + 1) % 3);
+	}
+	else
+	{
+		SortingColumn = ColumnId;
+		SortMode = NewSortMode;
+	}
+
+	if (SortingColumn != NAME_None && SortMode != EColumnSortMode::None)
+	{
+		Algo::Sort(ListViewItems, [this, ColumnId](const FPCGNodeTestResultPtr& A, const FPCGNodeTestResultPtr& B)
+		{
+			bool bIsLess = false;
+			if (SortingColumn == NAME_Index)
+			{
+				bIsLess = A->Index < B->Index;
+			}
+			else if (SortingColumn == NAME_NodeTitle)
+			{
+				bIsLess = A->TestResultTitle.ToString() < B->TestResultTitle.ToString();
+			}
+			else if (SortingColumn == NAME_NodeName)
+			{
+				bIsLess = A->TestResultName < B->TestResultName;
+			}
+			else if (SortingColumn == NAME_Seed)
+			{
+				bIsLess = A->Seed < B->Seed;
+			}
+			else if (SortingColumn == NAME_DataTypesTested)
+			{
+				bIsLess = A->DataTypesTested < B->DataTypesTested;
+			}
+			else if (SortingColumn == NAME_AdditionalDetails)
+			{
+				bIsLess = (A->AdditionalDetails.IsEmpty() && !B->AdditionalDetails.IsEmpty()) ||
+					(!A->AdditionalDetails.IsEmpty() && !B->AdditionalDetails.IsEmpty() && A->AdditionalDetails[0] < B->AdditionalDetails[0]);
+			}
+			else
+			{
+				const EDeterminismLevel* DeterminismLevelA = A->TestResults.Find(ColumnId);
+				const EDeterminismLevel* DeterminismLevelB = B->TestResults.Find(ColumnId);
+				bIsLess = (DeterminismLevelA && !DeterminismLevelB) ||
+					(DeterminismLevelA && DeterminismLevelB && *DeterminismLevelA < *DeterminismLevelB);
+			}
+
+			return SortMode == EColumnSortMode::Ascending ? bIsLess : !bIsLess;
+		});
+	}
+
+	RefreshItems();
+}
+
+EColumnSortMode::Type SPCGEditorGraphDeterminismListView::GetColumnSortMode(const FName ColumnId) const
+{
+	if (SortingColumn != ColumnId)
+	{
+		return EColumnSortMode::None;
+	}
+
+	return SortMode;
+}
+
+#undef LOCTEXT_NAMESPACE
diff --git a/Engine/Plugins/PCG/Source/PCGEditor/Private/Widgets/SPCGEditorGraphDeterminism.h b/Engine/Plugins/PCG/Source/PCGEditor/Private/Widgets/SPCGEditorGraphDeterminism.h
new file mode 100644
--- /dev/null
+++ b/Engine/Plugins/PCG/Source/PCGEditor/Private/Widgets/SPCGEditorGraphDeterminism.h
@@ -0,0 +1,90 @@
+// Copyright Epic Games, Inc. All Rights Reserved.
+
+#pragma once
+
+#include "Widgets/Views/ITableRow.h"
+#include "Widgets/Views/SListView.h"
+#include "Widgets/Views/STableRow.h"
+
+class STableViewBase;
+struct FDeterminismTestResult;
+
+class FPCGEditor;
+
+typedef TSharedPtr<FDeterminismTestResult> FPCGNodeTestResultPtr;
+
+struct FTestColumnInfo
+{
+	FTestColumnInfo(FName ColumnID, FText ColumnLabel, float Width = 0.f, EHorizontalAlignment HAlign = HAlign_Left) :
+		ColumnID(ColumnID),
+		ColumnLabel(ColumnLabel),
+		Width(Width),
+		HAlign(HAlign) {}
+
+	FName ColumnID = TEXT("UnnamedColumn_ID");
+	FText ColumnLabel = NSLOCTEXT("PCGDeterminism", "Unnamed_Column", "Unnamed Column");
+	float Width = 0.f;
+	EHorizontalAlignment HAlign = HAlign_Left;
+};
+
+class SPCGEditorGraphDeterminismRow final : public SMultiColumnTableRow<FPCGNodeTestResultPtr>
+{
+	SLATE_BEGIN_ARGS(SPCGEditorGraphDeterminismRow) {}
+	SLATE_END_ARGS()
+
+	/** Construct a row of the ListView */
+	void Construct(const FArguments& InArgs, const TSharedRef<STableViewBase>& InOwnerTableView, const FPCGNodeTestResultPtr& Item, int32 ItemIndex);
+
+	/** Generates a column, given the column's ID */
+	virtual TSharedRef<SWidget> GenerateWidgetForColumn(const FName& ColumnId) override;
+
+protected:
+	FPCGNodeTestResultPtr CurrentItem;
+};
+
+class SPCGEditorGraphDeterminismListView final : public SCompoundWidget
+{
+public:
+	SLATE_BEGIN_ARGS(SPCGEditorGraphDeterminismListView) {}
+	SLATE_END_ARGS()
+
+	/** Construct the ListView */
+	void Construct(const FArguments& InArgs, TWeakPtr<FPCGEditor> InPCGEditor);
+	/** Add an item to the ListView */
+	void AddItem(const FPCGNodeTestResultPtr& Item);
+	/** Clear all items from the ListView */
+	void ClearItems();
+	/** Refreshes the ListView */
+	void RefreshItems();
+	/** Adds a test column to the ListView */
+	void AddColumn(const FTestColumnInfo& ColumnInfo);
+	/** Builds the default columns for the widget */
+	void BuildBaseColumns();
+	/** Adds the additional details column */
+	void AddDetailsColumn();
+
+	/** Validates if the ListView has been constructed */
+	bool WidgetIsConstructed() const;
+
+private:
+	/** Clears the columns of the widget */
+	void ClearColumns();
+	/** Generate the row widget */
+	TSharedRef<ITableRow> OnGenerateRow(const FPCGNodeTestResultPtr Item, const TSharedRef<STableViewBase>& OwnerTable) const;
+	void OnSortColumnHeader(const EColumnSortPriority::Type SortPriority, const FName& ColumnId, const EColumnSortMode::Type NewSortMode);
+	EColumnSortMode::Type GetColumnSortMode(const FName ColumnId) const;
+
+	TSharedPtr<SHeaderRow> GeneratedHeaderRow;
+
+	TWeakPtr<FPCGEditor> PCGEditorPtr;
+
+	TSharedPtr<SListView<FPCGNodeTestResultPtr>> ListView;
+	TArray<FPCGNodeTestResultPtr> ListViewItems;
+
+	// To allow sorting
+	FName SortingColumn = NAME_None;
+	EColumnSortMode::Type SortMode = EColumnSortMode::Type::None;
+
+	bool bIsConstructed = false;
+	mutable int32 ItemIndexCounter = -1;
+};
diff --git a/Engine/Plugins/PCG/Source/PCGEditor/Private/Widgets/SPCGEditorGraphFind.cpp b/Engine/Plugins/PCG/Source/PCGEditor/Private/Widgets/SPCGEditorGraphFind.cpp
new file mode 100644
--- /dev/null
+++ b/Engine/Plugins/PCG/Source/PCGEditor/Private/Widgets/SPCGEditorGraphFind.cpp
@@ -0,0 +1,950 @@
+// Copyright Epic Games, Inc. All Rights Reserved.
+
+#include "SPCGEditorGraphFind.h"
+
+#include "PCGEditor.h"
+#include "PCGEditorGraph.h"
+#include "PCGEditorStyle.h"
+#include "PCGGraph.h"
+#include "PCGNode.h"
+#include "PCGSubgraph.h"
+#include "PCGSubsystem.h"
+#include "Nodes/PCGEditorGraphNodeBase.h"
+
+#include "EdGraph/EdGraphSchema.h"
+#include "Framework/Application/SlateApplication.h"
+#include "Layout/WidgetPath.h"
+#include "Widgets/Input/SComboButton.h"
+#include "Widgets/Input/SSearchBox.h"
+#include "Widgets/Views/STreeView.h"
+
+#define LOCTEXT_NAMESPACE "PCGEditorGraphFind"
+
+FPCGEditorGraphFindResult::FPCGEditorGraphFindResult(const FString& InValue)
+: Parent(nullptr)
+, Value(InValue)
+, GraphNode(nullptr)
+, RootGraphNode(nullptr)
+, Pin()
+{
+}
+
+FPCGEditorGraphFindResult::FPCGEditorGraphFindResult(const FText& InValue)
+: FPCGEditorGraphFindResult(InValue.ToString())
+{
+}
+
+FPCGEditorGraphFindResult::FPCGEditorGraphFindResult(const FString& InValue, const TSharedPtr<FPCGEditorGraphFindResult>& InParent, UEdGraphNode* InNode)
+: Parent(InParent)
+, Value(InValue)
+, GraphNode(InNode)
+, Pin()
+{
+	RootGraphNode = (Parent.IsValid() && Parent.Pin()->RootGraphNode.IsValid()) ? Parent.Pin()->RootGraphNode : InNode;
+}
+
+FPCGEditorGraphFindResult::FPCGEditorGraphFindResult(const FString& InValue, const TSharedPtr<FPCGEditorGraphFindResult>& InParent, UEdGraphPin* InPin)
+: Parent(InParent)
+, Value(InValue)
+, GraphNode(nullptr)
+, Pin(InPin)
+{
+	RootGraphNode = (Parent.IsValid() && Parent.Pin()->RootGraphNode.IsValid()) ? Parent.Pin()->RootGraphNode : nullptr;
+}
+
+FReply FPCGEditorGraphFindResult::OnClick(TWeakPtr<FPCGEditor> InPCGEditorPtr)
+{
+	TSharedPtr<FPCGEditor> Editor = InPCGEditorPtr.Pin();
+
+	if (!Editor)
+	{
+		return FReply::Handled();
+	}
+
+	UEdGraphPin* ResolvedPin = Pin.Get();
+	if (ResolvedPin && ResolvedPin->GetOwningNode()->GetGraph() == Editor->GetPCGEditorGraph())
+	{
+		Editor->JumpToNode(ResolvedPin->GetOwningNode());
+	}
+	else
+	{
+		Editor->JumpToNode(RootGraphNode.Get());
+	}
+
+	return FReply::Handled();
+}
+
+FReply FPCGEditorGraphFindResult::OnDoubleClick(TWeakPtr<FPCGEditor> InPCGEditorPtr)
+{
+	TSharedPtr<FPCGEditor> CurrentEditor = InPCGEditorPtr.Pin();
+	TSharedPtr<FPCGEditor> Editor = CurrentEditor;
+
+	if (ParentGraph)
+	{
+		GEditor->GetEditorSubsystem<UAssetEditorSubsystem>()->OpenEditorForAsset(ParentGraph->GetPCGGraph());
+		Editor = ParentGraph->GetEditor().Pin();
+	}
+
+	if (Editor)
+	{
+		if (Editor != CurrentEditor && CurrentEditor->GetStackBeingInspected())
+		{
+			Editor->SetStackBeingInspectedFromAnotherEditor(*CurrentEditor->GetStackBeingInspected());
+		}
+
+		if (UEdGraphPin* ResolvedPin = Pin.Get())
+		{
+			Editor->JumpToNode(ResolvedPin->GetOwningNode());
+		}
+		else if (GraphNode.IsValid())
+		{
+			Editor->JumpToNode(GraphNode.Get());
+		}
+	}
+
+	return FReply::Handled();
+}
+
+FText FPCGEditorGraphFindResult::GetToolTip() const
+{
+	FText ToolTip;
+
+	if (UEdGraphPin* ResolvedPin = Pin.Get())
+	{
+		if (UEdGraphNode* OwningNode = ResolvedPin->GetOwningNode())
+		{
+			FString ToolTipString;
+			OwningNode->GetPinHoverText(*ResolvedPin, ToolTipString);
+			ToolTip = FText::FromString(ToolTipString);
+		}
+	}
+	else if (GraphNode.IsValid())
+	{
+		ToolTip = GraphNode->GetTooltipText();
+	}
+
+	return ToolTip;
+}
+
+FText FPCGEditorGraphFindResult::GetCategory() const
+{
+	if (Pin.Get())
+	{
+		return LOCTEXT("PinCategory", "Pin");
+	}
+	else if (GraphNode.IsValid())
+	{
+		return LOCTEXT("NodeCategory", "Node");
+	}
+	return FText::GetEmpty();
+}
+
+FText FPCGEditorGraphFindResult::GetComment() const
+{
+	if (GraphNode.IsValid())
+	{
+		const FString NodeComment = GraphNode->NodeComment;
+		if (!NodeComment.IsEmpty())
+		{
+			return FText::Format(LOCTEXT("NodeCommentFmt", "Node Comment:[{0}]"), FText::FromString(NodeComment));
+		}
+	}
+
+	return FText::GetEmpty();
+}
+
+TSharedRef<SWidget> FPCGEditorGraphFindResult::CreateIcon() const
+{
+	FSlateColor IconColor = FSlateColor::UseForeground();
+	const FSlateBrush* Brush = nullptr;
+
+	// TODO: consider the pin connection state (connected or not + single + multi)
+	// TODO: consider node state (enabled or not)
+	if (UEdGraphPin* ResolvedPin = Pin.Get())
+	{
+		// TODO get pin icon from nodebase?
+		Brush = FAppStyle::GetBrush(TEXT("GraphEditor.PinIcon"));
+		const UEdGraphSchema* Schema = ResolvedPin->GetSchema();
+		IconColor = Schema->GetPinTypeColor(ResolvedPin->PinType);
+	}
+	else if (GraphNode.IsValid())
+	{
+		// TODO get icon and tint from nodebase?
+		Brush = FAppStyle::GetBrush(TEXT("GraphEditor.NodeGlyph"));
+
+		if (const UPCGEditorGraphNodeBase* PCGNode = Cast<UPCGEditorGraphNodeBase>(GraphNode))
+		{
+			if (Cast<UPCGBaseSubgraphNode>(PCGNode->GetPCGNode()))
+			{
+				Brush = FPCGEditorStyle::Get().GetBrush(TEXT("ClassIcon.PCGGraphInterface"));
+			}
+
+			IconColor = FSlateColor(PCGNode->GetNodeTitleColor());
+		}
+	}
+
+	return SNew(SImage)
+		.Image(Brush)
+		.ColorAndOpacity(IconColor)
+		.ToolTipText(GetCategory());
+}
+
+SPCGEditorGraphFind::~SPCGEditorGraphFind()
+{
+	if (PCGEditorPtr.IsValid())
+	{
+		PCGEditorPtr.Pin()->OnInspectedStackChangedDelegate.RemoveAll(this);
+	}
+}
+
+void SPCGEditorGraphFind::Construct(const FArguments& InArgs, TSharedPtr<FPCGEditor> InPCGEditor)
+{
+	PCGEditorPtr = InPCGEditor;
+
+	const TSharedPtr<FPCGEditor> PCGEditor = PCGEditorPtr.Pin();
+
+	if (PCGEditor)
+	{
+		PCGEditor->OnInspectedStackChangedDelegate.AddSP(this, &SPCGEditorGraphFind::OnInspectedStackChanged);
+	}
+
+	this->ChildSlot
+	[
+		SNew(SVerticalBox)
+		+SVerticalBox::Slot()
+		.AutoHeight()
+		[
+			SNew(SHorizontalBox)
+			+SHorizontalBox::Slot()
+			.AutoWidth()
+			[
+				SNew(SComboButton)
+				.ForegroundColor(FSlateColor::UseStyle())
+				.HasDownArrow(false)
+				.OnGetMenuContent(this, &SPCGEditorGraphFind::OnFindFilterMenu)
+				.ContentPadding(1)
+				.ButtonContent()
+				[
+					SNew(SImage)
+					.Image(FAppStyle::Get().GetBrush("Icons.Filter"))
+					.ColorAndOpacity(FSlateColor::UseForeground())
+				]
+			]
+			+SHorizontalBox::Slot()
+			.AutoWidth()
+			[
+				SAssignNew(SearchTextField, SSearchBox)
+				.HintText(LOCTEXT("PCGGraphSearchHint", "Enter text to find nodes..."))
+				.OnTextChanged(this, &SPCGEditorGraphFind::OnSearchTextChanged)
+				.OnTextCommitted(this, &SPCGEditorGraphFind::OnSearchTextCommitted)
+			]
+		]
+		+ SVerticalBox::Slot()
+		.FillHeight(1.0f)
+		.Padding(0.f, 4.f, 0.f, 0.f)
+		[
+			SNew(SBorder)
+			.BorderImage(FAppStyle::GetBrush("Menu.Background"))
+			[
+				SAssignNew(TreeView, STreeViewType)
+				.TreeItemsSource(&ItemsFound)
+				.OnGenerateRow(this, &SPCGEditorGraphFind::OnGenerateRow)
+				.OnGetChildren(this, &SPCGEditorGraphFind::OnGetChildren)
+				.OnSelectionChanged(this, &SPCGEditorGraphFind::OnTreeSelectionChanged)
+				.OnMouseButtonDoubleClick(this, &SPCGEditorGraphFind::OnTreeDoubleClick)
+				.OnKeyDownHandler(this, &SPCGEditorGraphFind::OnTreeViewKeyDown)
+				.SelectionMode(ESelectionMode::Single)
+			]
+		]
+	];
+}
+
+void SPCGEditorGraphFind::FocusForUse()
+{
+	// NOTE: Careful, GeneratePathToWidget can be reentrant in that it can call visibility delegates and such
+	FWidgetPath FilterTextBoxWidgetPath;
+	FSlateApplication::Get().GeneratePathToWidgetUnchecked(SearchTextField.ToSharedRef(), FilterTextBoxWidgetPath);
+
+	// Set keyboard focus directly
+	FSlateApplication::Get().SetKeyboardFocus(FilterTextBoxWidgetPath, EFocusCause::SetDirectly);
+}
+
+void SPCGEditorGraphFind::OnSearchTextChanged(const FText& InText)
+{
+	SearchValue = InText.ToString();
+	InitiateSearch();
+}
+
+void SPCGEditorGraphFind::OnSearchTextCommitted(const FText& /*InText*/, ETextCommit::Type InCommitType)
+{
+	// Since we already initiate a search when the text changes, there's no real need to do anything on commit
+	if (InCommitType != ETextCommit::Type::OnUserMovedFocus)
+	{
+		InitiateSearch();
+	}
+}
+
+void SPCGEditorGraphFind::OnInspectedStackChanged(const FPCGStack& InPCGStack)
+{
+	if (!SearchValue.IsEmpty())
+	{
+		InitiateSearch();
+	}
+}
+
+void SPCGEditorGraphFind::OnGetChildren(FPCGEditorGraphFindResultPtr InItem, TArray<FPCGEditorGraphFindResultPtr>& OutChildren)
+{
+	OutChildren += InItem->Children;
+}
+
+void SPCGEditorGraphFind::OnTreeSelectionChanged(FPCGEditorGraphFindResultPtr Item, ESelectInfo::Type)
+{
+	if (Item.IsValid())
+	{
+		Item->OnClick(PCGEditorPtr);
+	}
+}
+
+void SPCGEditorGraphFind::OnTreeDoubleClick(FPCGEditorGraphFindResultPtr Item)
+{
+	if (Item.IsValid())
+	{
+		Item->OnDoubleClick(PCGEditorPtr);
+	}
+}
+
+FReply SPCGEditorGraphFind::OnTreeViewKeyDown(const FGeometry& /*InGeometry*/, const FKeyEvent& InKeyEvent) const
+{
+	if (InKeyEvent.GetKey() == EKeys::Enter)
+	{
+		TArray<FPCGEditorGraphFindResultPtr> SelectedItems = TreeView->GetSelectedItems();
+		if (SelectedItems.Num() == 1 && SelectedItems[0].IsValid())
+		{
+			return SelectedItems[0]->OnDoubleClick(PCGEditorPtr);
+		}
+	}
+
+	return FReply::Unhandled();
+}
+
+void SPCGEditorGraphFind::SetFindMode(EPCGGraphFindMode InFindMode)
+{
+	if (InFindMode != FindMode)
+	{
+		FindMode = InFindMode;
+		InitiateSearch();
+	}
+}
+
+bool SPCGEditorGraphFind::IsCurrentFindMode(EPCGGraphFindMode InFindMode) const
+{
+	return FindMode == InFindMode;
+}
+
+TSharedRef<SWidget> SPCGEditorGraphFind::OnFindFilterMenu()
+{
+	FMenuBuilder MenuBuilder(/*bInShouldCloseWindowAfterMenuSelection=*/true, nullptr);
+
+	MenuBuilder.AddMenuEntry(
+		LOCTEXT("MinimalExpansion", "Show minimum tree"),
+		LOCTEXT("MinimalExpansionTooltip", "Shows minimum subset of visited tree to perform search."),
+		FSlateIcon(),
+		FUIAction(
+			FExecuteAction::CreateSP(this, &SPCGEditorGraphFind::SetFindMode, EPCGGraphFindMode::ShowMinimumTree),
+			FCanExecuteAction(),
+			FIsActionChecked::CreateSP(this, &SPCGEditorGraphFind::IsCurrentFindMode, EPCGGraphFindMode::ShowMinimumTree)
+		),
+		NAME_None,
+		EUserInterfaceActionType::RadioButton);
+
+	MenuBuilder.AddMenuEntry(
+		LOCTEXT("FullExpansion", "Show full tree"),
+		LOCTEXT("FullExpansionTooltip", "Shows the search results from the fully expanded tree, e.g. will have every occurrence."),
+		FSlateIcon(),
+		FUIAction(
+			FExecuteAction::CreateSP(this, &SPCGEditorGraphFind::SetFindMode, EPCGGraphFindMode::ShowFullTree),
+			FCanExecuteAction(),
+			FIsActionChecked::CreateSP(this, &SPCGEditorGraphFind::IsCurrentFindMode, EPCGGraphFindMode::ShowFullTree)
+		),
+		NAME_None,
+		EUserInterfaceActionType::RadioButton);
+
+	MenuBuilder.AddMenuEntry(
+		LOCTEXT("FlatList", "Show flat graph list"),
+		LOCTEXT("FlatListTooltip", "Shows the occurrences in all visited graphs but no hierarchy."),
+		FSlateIcon(),
+		FUIAction(
+			FExecuteAction::CreateSP(this, &SPCGEditorGraphFind::SetFindMode, EPCGGraphFindMode::ShowFlatList),
+			FCanExecuteAction(),
+			FIsActionChecked::CreateSP(this, &SPCGEditorGraphFind::IsCurrentFindMode, EPCGGraphFindMode::ShowFlatList)
+		),
+		NAME_None,
+		EUserInterfaceActionType::RadioButton);
+
+	MenuBuilder.AddSeparator();
+
+	MenuBuilder.AddMenuEntry(
+		LOCTEXT("SearchPins", "Include pin names"),
+		LOCTEXT("SearchPinsTooltip", "TODO"),
+		FSlateIcon(),
+		FUIAction(
+			FExecuteAction::CreateLambda([this]()
+			{
+				bShowPinResults = !bShowPinResults;
+				InitiateSearch();
+			}),
+			FCanExecuteAction(),
+			FIsActionChecked::CreateLambda([this]()
+			{
+				return bShowPinResults;
+			})
+		),
+		NAME_None,
+		EUserInterfaceActionType::ToggleButton);
+
+	return MenuBuilder.MakeWidget();
+}
+
+TSharedRef<ITableRow> SPCGEditorGraphFind::OnGenerateRow(FPCGEditorGraphFindResultPtr InItem, const TSharedRef<STableViewBase>& InOwnerTable)
+{
+	return SNew(STableRow<TSharedPtr<FPCGEditorGraphFindResult>>, InOwnerTable)
+	[
+		SNew(SHorizontalBox)
+		+ SHorizontalBox::Slot()
+		.VAlign(VAlign_Center)
+		.AutoWidth()
+		[
+			InItem->CreateIcon()
+		]
+		+SHorizontalBox::Slot()
+		.VAlign(VAlign_Center)
+		.AutoWidth()
+		.Padding(2, 0)
+		[
+			SNew(STextBlock)
+			.Text(FText::FromString(InItem->Value))
+			.HighlightText(HighlightText)
+			.ToolTipText(InItem->GetToolTip())
+		]
+		+ SHorizontalBox::Slot()
+		.FillWidth(1)
+		.HAlign(HAlign_Right)
+		.VAlign(VAlign_Center)
+		.Padding(2, 0)
+		[
+			SNew(STextBlock)
+			.Text(InItem->GetComment()) 
+			.ColorAndOpacity(FLinearColor::Yellow)
+			.HighlightText(HighlightText)
+		]
+	];
+}
+
+void SPCGEditorGraphFind::InitiateSearch()
+{
+	TArray<FString> Tokens;
+	SearchValue.ParseIntoArray(Tokens, TEXT(" "), true);
+
+	ItemsFound.Empty();
+	if (Tokens.Num() > 0)
+	{
+		HighlightText = FText::FromString(SearchValue);
+		MatchTokens(Tokens);
+	}
+
+	// Insert a fake result to inform user if none found
+	if (ItemsFound.Num() == 0)
+	{
+		ItemsFound.Add(MakeShared<FPCGEditorGraphFindResult>(LOCTEXT("PCGGraphSearchNoResults", "No Results found")));
+	}
+
+	TreeView->RequestTreeRefresh();
+
+	// Expand so that all items that contain a matched token are shown
+	auto ExpandItem = [this](const FPCGEditorGraphFindResultPtr& Item, auto&& RecursiveCall) -> bool
+	{
+		check(Item);
+		bool bShouldExpandItem = Item->bIsMatch;
+
+		for (const FPCGEditorGraphFindResultPtr& Child : Item->Children)
+		{
+			bShouldExpandItem = RecursiveCall(Child, RecursiveCall);
+		}
+
+		if (bShouldExpandItem)
+		{
+			TreeView->SetItemExpansion(Item, true);
+		}
+
+		return bShouldExpandItem;
+	};
+
+	for (const FPCGEditorGraphFindResultPtr& Item : ItemsFound)
+	{
+		ExpandItem(Item, ExpandItem);
+	}
+}
+
+namespace PCGEditorGraphFind
+{
+	// Convenience struct to vist graphs & stacks and build a proper hierarchy
+	struct GraphTree
+	{
+		UPCGGraph* Graph = nullptr;
+		UEdGraphNode* Node = nullptr;
+
+		const TArray<GraphTree>& GetChildren() const { return Children; }
+		const GraphTree* GetParent() const { return Parent; }
+
+		bool HasVisitedGraph(UPCGGraph* InGraph) const
+		{
+			return Graph == InGraph || (Parent && Parent->HasVisitedGraph(InGraph));
+		}
+
+		GraphTree* FindOrCreateChild(UEdGraphNode* InNode, UPCGGraph* InGraph)
+		{
+			for (GraphTree& Child : Children)
+			{
+				if (Child.Node == InNode && Child.Graph == InGraph)
+				{
+					return &Child;
+				}
+			}
+
+			// Not found -> add new child
+			GraphTree& NewChild = AddChild();
+			NewChild.Node = InNode;
+			NewChild.Graph = InGraph;
+
+			return &NewChild;
+		}
+
+		// Collapse current subtree so that it is attached to the same parent graph, but closest to the root.
+		bool Collapse()
+		{
+			for (int ChildIndex = Children.Num() - 1; ChildIndex >= 0; --ChildIndex)
+			{
+				if (Children[ChildIndex].Collapse())
+				{
+					Children.RemoveAt(ChildIndex, EAllowShrinking::No);
+				}
+			}
+
+			GraphTree* CurrentParent = Parent;
+			GraphTree* TopmostOccurrenceOfThisGraph = nullptr;
+			while (CurrentParent)
+			{
+				if (CurrentParent->Graph == Graph)
+				{
+					TopmostOccurrenceOfThisGraph = CurrentParent;
+				}
+
+				CurrentParent = CurrentParent->Parent;
+			}
+
+			if (TopmostOccurrenceOfThisGraph)
+			{
+				TopmostOccurrenceOfThisGraph->AppendChildren(Children);
+				return true;
+			}
+			else
+			{
+				return false;
+			}
+		}
+
+		GraphTree* FindFirst(UPCGGraph* InGraph)
+		{
+			TArray<GraphTree*> BFSExpansion = GetBreadthFirstExpansion();
+			for (GraphTree* ExpandedNode : BFSExpansion)
+			{
+				if (ExpandedNode->Graph == InGraph)
+				{
+					return ExpandedNode;
+				}
+			}
+
+			return nullptr;
+		}
+
+	private:
+		TArray<GraphTree> Children;
+		GraphTree* Parent = nullptr;
+
+		TArray<GraphTree*> GetBreadthFirstExpansion()
+		{
+			TArray<GraphTree*> Expansion;
+			Expansion.Add(this);
+			int ExpansionIndex = 0;
+			while (ExpansionIndex < Expansion.Num())
+			{
+				GraphTree* Current = Expansion[ExpansionIndex];
+				for (GraphTree& Child : Current->Children)
+				{
+					Expansion.Add(&Child);
+				}
+
+				++ExpansionIndex;
+			}
+
+			return Expansion;
+		}
+
+		GraphTree& AddChild()
+		{
+			SIZE_T PreviouslyAllocatedSize = Children.GetAllocatedSize();
+			GraphTree& NewChild = Children.Emplace_GetRef();
+			NewChild.Parent = this;
+
+			if (Children.GetAllocatedSize() != PreviouslyAllocatedSize)
+			{
+				UpdateParentPointers();
+			}
+
+			return NewChild;
+		}
+
+		void AppendChildren(const TArray<GraphTree>& InChildren)
+		{
+			Children.Append(InChildren);
+			// Since we might have resized the array and because we're moving the children here, we must update the parent pointers.
+			UpdateParentPointers();
+		}
+
+		void UpdateParentPointers()
+		{
+			for (GraphTree& Child : Children)
+			{
+				Child.Parent = this;
+				Child.UpdateParentPointers();
+			}
+		}
+	};
+} // namespace PCGEditorGraphFind
+
+void SPCGEditorGraphFind::MatchTokens(const TArray<FString>& InTokens)
+{
+	TSharedPtr<FPCGEditor> PCGEditor = PCGEditorPtr.Pin();
+	UPCGEditorGraph* PCGEditorGraph = PCGEditor ? PCGEditor->GetPCGEditorGraph() : nullptr;
+	
+	if (!PCGEditor || !PCGEditorGraph)
+	{
+		return;
+	}
+
+	// Start by gathering all graphs to search through.
+	TSet<UPCGGraph*> AllGraphs;
+	UPCGGraph* ThisGraph = PCGEditorGraph->GetPCGGraph();
+	AllGraphs.Add(ThisGraph);
+
+	PCGEditorGraphFind::GraphTree GraphRoot;
+	GraphRoot.Graph = ThisGraph;
+	GraphRoot.Node = nullptr;
+
+	auto VisitAllNodes = [&AllGraphs](UPCGEditorGraph* EditorGraph, PCGEditorGraphFind::GraphTree& Parent, auto&& RecursiveCall)
+	{
+		if (!EditorGraph)
+		{
+			return;
+		}
+
+		for (UEdGraphNode* Node : EditorGraph->Nodes)
+		{
+			UPCGEditorGraphNodeBase* PCGEditorNode = Cast<UPCGEditorGraphNodeBase>(Node);
+			const UPCGBaseSubgraphNode* SubgraphNode = PCGEditorNode ? Cast<UPCGBaseSubgraphNode>(PCGEditorNode->GetPCGNode()) : nullptr;
+			if(UPCGGraph* Subgraph = SubgraphNode ? SubgraphNode->GetSubgraph() : nullptr)
+			{
+				// Check if this subgraph is already present in the hierarchy
+				if (!Parent.HasVisitedGraph(Subgraph))
+				{
+					AllGraphs.Add(Subgraph);
+					if (PCGEditorGraphFind::GraphTree* Child = Parent.FindOrCreateChild(Node, Subgraph))
+					{
+						RecursiveCall(FPCGEditor::GetPCGEditorGraph(Subgraph), *Child, RecursiveCall);
+					}
+				}
+			}
+		}
+	};
+
+	// Visit the static graph(s) starting from the root.
+	VisitAllNodes(PCGEditorGraph, GraphRoot, VisitAllNodes);
+
+	// Then gather all stacks starting from the stack being inspected, and link them to their matching subgraph node.
+	if (PCGEditor->GetStackBeingInspected() && PCGEditorPtr.Pin()->GetSubsystem())
+	{
+		const FPCGStack& CurrentStack = *PCGEditor->GetStackBeingInspected();
+		TArray<FPCGStack> Substacks = PCGEditorPtr.Pin()->GetSubsystem()->GetExecutedStacks(CurrentStack);
+		// For each stack, find the first subgraph node from the start, then get all graphs under it.
+		for (const FPCGStack& Substack : Substacks)
+		{
+			if (Substack.GetStackFrames().Num() <= CurrentStack.GetStackFrames().Num())
+			{
+				continue;
+			}
+
+			// We'll be trying to navigate down the stack and find the node + graph pairs, creating what we need in the graph tree.
+			PCGEditorGraphFind::GraphTree* Current = &GraphRoot;
+			int FrameIndex = CurrentStack.GetStackFrames().Num();
+
+			UEdGraphNode* SubgraphEditorNode = nullptr;
+			UPCGGraph* Subgraph = nullptr;
+			UPCGEditorGraph* CurrentEditorGraph = PCGEditorGraph;
+
+			// There are two kind of subgraph frame structures we are looking for:
+			// Subgraph node (UPCGNode) > Subgraph (UPCGGraph) -> this is when we have a static subgraphs
+			// Subgraph node (UPCGNode) > Loop Index / -1 -> this is when we have a dynamic subgraph or a loop
+			auto CheckFramePair = [&Substack, &SubgraphEditorNode, &Subgraph, &CurrentEditorGraph](int NodeFrameIndex, int SubgraphFrameIndex) -> bool
+			{
+				if(SubgraphFrameIndex >= Substack.GetStackFrames().Num())
+				{
+					return false;
+				}
+
+				const FPCGStackFrame& SubgraphNodeFrame = Substack.GetStackFrames()[NodeFrameIndex];
+				const UPCGNode* SubgraphNode = Cast<UPCGNode>(SubgraphNodeFrame.Object.Get());
+				SubgraphEditorNode = (SubgraphNode && CurrentEditorGraph) ? const_cast<UPCGEditorGraphNodeBase*>(CurrentEditorGraph->GetEditorNodeFromPCGNode(SubgraphNode)) : nullptr;
+				const FPCGStackFrame& SubgraphFrame = Substack.GetStackFrames()[SubgraphFrameIndex];
+				Subgraph = Cast<UPCGGraph>(const_cast<UObject*>(SubgraphFrame.Object.Get()));
+
+				return SubgraphEditorNode && Subgraph;
+			};
+
+			while (FrameIndex < Substack.GetStackFrames().Num() && Current)
+			{
+				bool bProcessPair = false;
+
+				// Test node + graph in a static subgraph configuration
+				if (CheckFramePair(FrameIndex, FrameIndex + 1))
+				{
+					bProcessPair = true;
+					FrameIndex += 2;
+				}
+				// Test node + graph in a dynamic subgraph/loop configuration
+				else if (CheckFramePair(FrameIndex, FrameIndex + 2))
+				{
+					bProcessPair = true;
+					FrameIndex += 3;
+				}
+				else
+				{
+					++FrameIndex;
+				}
+
+				if (bProcessPair)
+				{
+					AllGraphs.Add(Subgraph);
+					// Find/Create child on current
+					Current = Current->FindOrCreateChild(SubgraphEditorNode, Subgraph);
+					CurrentEditorGraph = Current && Subgraph ? FPCGEditor::GetPCGEditorGraph(Subgraph) : nullptr;
+
+					if (Current && CurrentEditorGraph)
+					{
+						VisitAllNodes(CurrentEditorGraph, *Current, VisitAllNodes);
+					}
+					else
+					{
+						// We can't progress here so leave the loop.
+						break;
+					}
+				}
+			}
+		}
+	}
+
+	// Implementation note: at this point, we could visit each graph only once to search for the tokens
+	// but we'd need additional mechanisms to do deep copies, which is not super significant at this point in time.
+	// Perform search by graph based on the tree expansion we'll visit
+	if (FindMode == EPCGGraphFindMode::ShowMinimumTree)
+	{
+		// Moves subtrees to their earliest appearance in the breadth-first tree
+		GraphRoot.Collapse();
+	}
+
+	if(FindMode == EPCGGraphFindMode::ShowMinimumTree || FindMode == EPCGGraphFindMode::ShowFullTree)
+	{
+		auto RecurseInGraphTree = [this, &InTokens](const PCGEditorGraphFind::GraphTree& Node, TFunctionRef<FPCGEditorGraphFindResultPtr()> GetParentFunc, auto&& RecurseCall) -> void
+		{
+			check(Node.Graph);
+
+			FString NodeString;
+			FPCGEditorGraphFindResultPtr ThisResult;
+
+			if (Node.Node)
+			{
+				NodeString = FString::Printf(TEXT("%s (%s)"), *FName::NameToDisplayString(Node.Graph->GetName(), /*bIsBool=*/false), *Node.Node->GetNodeTitle(ENodeTitleType::ListView).ToString());
+			}
+
+			auto GetOrCreateNodeResult = [this, &Node, &NodeString, &ThisResult, &GetParentFunc]()
+			{
+				FPCGEditorGraphFindResultPtr Parent = GetParentFunc();
+				check(Parent);
+
+				if(!ThisResult.IsValid() && Node.Node)
+				{
+					ThisResult = MakeShared<FPCGEditorGraphFindResult>(NodeString, Parent, Node.Node);
+					ThisResult->ParentGraph = FPCGEditor::GetPCGEditorGraph(Node.Graph);
+					Parent->Children.Add(ThisResult);
+				}
+
+				return ThisResult ? ThisResult : Parent;
+			};
+
+			// If current node (including dynamic graph name) matches tokens, add it
+			if (StringMatchesSearchTokens(InTokens, NodeString))
+			{
+				GetOrCreateNodeResult()->bIsMatch = true;
+			}
+
+			// Create local elements
+			MatchTokensInternal(InTokens, FPCGEditor::GetPCGEditorGraph(Node.Graph), GetOrCreateNodeResult/*, MatchedTokens*/);
+
+			// Continue through the tree
+			for (const PCGEditorGraphFind::GraphTree& Child : Node.GetChildren())
+			{
+				RecurseCall(Child, GetOrCreateNodeResult, RecurseCall);
+			}
+		};
+
+		FPCGEditorGraphFindResultPtr RootFindResult = MakeShared<FPCGEditorGraphFindResult>(FString("PCGTreeRoot"));
+		auto GetRootPtr = [&RootFindResult]() { return RootFindResult; };
+
+		RecurseInGraphTree(GraphRoot, GetRootPtr, RecurseInGraphTree);
+		ItemsFound.Append(RootFindResult->Children);
+	}
+	// - flat list (this graph + all downstream graphs after)
+	else if (FindMode == EPCGGraphFindMode::ShowFlatList)
+	{
+		for(UPCGGraph* Graph : AllGraphs)
+		{
+			FPCGEditorGraphFindResultPtr GraphNodePtr;
+			const FString GraphString = FName::NameToDisplayString(Graph->GetName(), /*bIsBool=*/false);
+
+			auto CreateGraphNode = [Graph, &GraphString, &GraphNodePtr, &GraphRoot]()
+			{
+				if (!GraphNodePtr.IsValid())
+				{
+					PCGEditorGraphFind::GraphTree* GraphTreeNode = GraphRoot.FindFirst(Graph);
+
+					UEdGraphNode* EditorNode = nullptr;
+					UEdGraphNode* RootEditorNode = nullptr;
+
+					if (GraphTreeNode)
+					{
+						EditorNode = GraphTreeNode->Node;
+
+						const PCGEditorGraphFind::GraphTree* GraphRootNode = GraphTreeNode;
+						while (GraphRootNode->GetParent() && GraphRootNode->GetParent() != &GraphRoot)
+						{
+							GraphRootNode = GraphRootNode->GetParent();
+						}
+
+						RootEditorNode = GraphRootNode->Node;
+					}
+
+					GraphNodePtr = MakeShared<FPCGEditorGraphFindResult>(GraphString, FPCGEditorGraphFindResultPtr(), EditorNode);
+					GraphNodePtr->ParentGraph = FPCGEditor::GetPCGEditorGraph(Graph);
+					GraphNodePtr->RootGraphNode = RootEditorNode;
+				}
+
+				return GraphNodePtr;
+			};
+
+			// If current node (including dynamic graph name) matches tokens, add it
+			if (StringMatchesSearchTokens(InTokens, GraphString))
+			{
+				CreateGraphNode()->bIsMatch = true;
+			}
+
+			MatchTokensInternal(InTokens, FPCGEditor::GetPCGEditorGraph(Graph), CreateGraphNode/*, Results*/);
+
+			if (GraphNodePtr)
+			{
+				ItemsFound.Add(GraphNodePtr);
+			}
+		}
+	}
+}
+
+void SPCGEditorGraphFind::MatchTokensInternal(const TArray<FString>& InTokens, UPCGEditorGraph* PCGEditorGraph, TFunctionRef<FPCGEditorGraphFindResultPtr()> GetParentFunc) const
+{
+	if (!PCGEditorGraph)
+	{
+		return;
+	}
+
+	for (UEdGraphNode* Node : PCGEditorGraph->Nodes)
+	{
+		check(Node);
+
+		const FString NodeString = Node->GetNodeTitle(ENodeTitleType::ListView).ToString();
+
+		// Search string has full title (both lines).
+		FString NodeSearchString = Node->GetNodeTitle(ENodeTitleType::FullTitle).ToString() + Node->NodeComment;
+
+		// Add internal object name which will still display here and there.
+		const UPCGEditorGraphNodeBase* PCGEditorGraphNodeBase = Cast<UPCGEditorGraphNodeBase>(Node);
+		if (const UPCGNode* PCGNode = PCGEditorGraphNodeBase ? PCGEditorGraphNodeBase->GetPCGNode() : nullptr)
+		{
+			NodeSearchString.Append(PCGNode->GetName());
+		}
+
+		NodeSearchString = NodeSearchString.Replace(TEXT(" "), TEXT(""));
+
+		FPCGEditorGraphFindResultPtr NodeResult;
+		auto GetOrCreateNodeResult = [&NodeResult, &NodeString, Node, PCGEditorGraph, &GetParentFunc]() -> FPCGEditorGraphFindResultPtr&
+		{
+			if (!NodeResult.IsValid())
+			{
+				NodeResult = MakeShared<FPCGEditorGraphFindResult>(NodeString, GetParentFunc(), Node);
+				NodeResult->ParentGraph = PCGEditorGraph;
+				GetParentFunc()->Children.Add(NodeResult);
+			}
+
+			return NodeResult;
+		};
+
+		if (StringMatchesSearchTokens(InTokens, NodeSearchString))
+		{
+			GetOrCreateNodeResult()->bIsMatch = true;
+		}
+
+		if (bShowPinResults)
+		{
+			for (UEdGraphPin* Pin : Node->Pins)
+			{
+				if (Pin && Pin->PinFriendlyName.CompareTo(FText::FromString(TEXT(" "))) != 0)
+				{
+					const FText PinName = Pin->GetSchema()->GetPinDisplayName(Pin);
+					FString PinSearchString = Pin->PinName.ToString() + Pin->PinFriendlyName.ToString() + Pin->DefaultValue + Pin->PinType.PinCategory.ToString() + Pin->PinType.PinSubCategory.ToString() + (Pin->PinType.PinSubCategoryObject.IsValid() ? Pin->PinType.PinSubCategoryObject.Get()->GetFullName() : TEXT(""));
+					PinSearchString = PinSearchString.Replace(TEXT(" "), TEXT(""));
+					if (StringMatchesSearchTokens(InTokens, PinSearchString))
+					{
+						FPCGEditorGraphFindResultPtr PinResult(MakeShared<FPCGEditorGraphFindResult>(PinName.ToString(), GetOrCreateNodeResult(), Pin));
+						PinResult->ParentGraph = PCGEditorGraph;
+						PinResult->bIsMatch = true;
+						NodeResult->Children.Add(PinResult);
+					}
+				}
+			}
+		}
+	}
+}
+
+bool SPCGEditorGraphFind::StringMatchesSearchTokens(const TArray<FString>& InTokens, const FString& InComparisonString)
+{
+	bool bFoundAllTokens = true;
+
+	//search the entry for each token, it must have all of them to pass
+	for (const FString& Token : InTokens)
+	{
+		if (!InComparisonString.Contains(Token))
+		{
+			bFoundAllTokens = false;
+			break;
+		}
+	}
+	return bFoundAllTokens;
+}
+
+#undef LOCTEXT_NAMESPACE
diff --git a/Engine/Plugins/PCG/Source/PCGEditor/Private/Widgets/SPCGEditorGraphFind.h b/Engine/Plugins/PCG/Source/PCGEditor/Private/Widgets/SPCGEditorGraphFind.h
new file mode 100644
--- /dev/null
+++ b/Engine/Plugins/PCG/Source/PCGEditor/Private/Widgets/SPCGEditorGraphFind.h
@@ -0,0 +1,169 @@
+// Copyright Epic Games, Inc. All Rights Reserved.
+
+#pragma once
+
+#include "EdGraph/EdGraphPin.h"
+#include "Widgets/SCompoundWidget.h"
+
+class ITableRow;
+class SSearchBox;
+class STableViewBase;
+namespace ESelectInfo { enum Type : int; }
+namespace ETextCommit { enum Type : int; }
+template <typename ItemType> class STreeView;
+
+class FPCGEditor;
+struct FPCGStack;
+class UEdGraphNode;
+class UPCGEditorGraph;
+class UPCGGraph;
+
+struct FPCGEditorGraphFindResult
+{
+	/** Create a root (or only text) result */
+	FPCGEditorGraphFindResult(const FString& InValue);
+
+	/** Create a root (or only text) result */
+	FPCGEditorGraphFindResult(const FText& InValue);
+
+	/** Create a listing for a node result */
+	FPCGEditorGraphFindResult(const FString& InValue, const TSharedPtr<FPCGEditorGraphFindResult>& InParent, UEdGraphNode* InNode);
+
+	/** Create a listing for a pin result */
+	FPCGEditorGraphFindResult(const FString& InValue, const TSharedPtr<FPCGEditorGraphFindResult>& InParent, UEdGraphPin* InPin);
+
+	/** Called when user clicks on the search item */
+	FReply OnClick(TWeakPtr<FPCGEditor> InPCGEditorPtr);
+
+	/** Called when user double-clicks on the search item */
+	FReply OnDoubleClick(TWeakPtr<FPCGEditor> InPCGEditorPtr);
+
+	/** Get ToolTip for this search result */
+	FText GetToolTip() const;
+
+	/** Get Category for this search result */
+	FText GetCategory() const;
+
+	/** Get Comment for this search result */
+	FText GetComment() const;
+
+	/** Create an icon to represent the result */
+	TSharedRef<SWidget> CreateIcon() const;
+
+	/** ParentGraph if not current */
+	UPCGEditorGraph* ParentGraph = nullptr;
+
+	/** Search result parent */
+	TWeakPtr<FPCGEditorGraphFindResult> Parent;
+
+	/** Any children listed under this category */
+	TArray<TSharedPtr<FPCGEditorGraphFindResult>> Children;
+	
+	/** The string value for this result */
+	FString Value;
+
+	/** The graph node that this search result refers to */
+	TWeakObjectPtr<UEdGraphNode> GraphNode;
+
+	/** The graph node that`s in the current editor graph upstream. */
+	TWeakObjectPtr<UEdGraphNode> RootGraphNode;
+
+	/** The pin that this search result refers to */
+	FEdGraphPinReference Pin;
+
+	/** Whether this result contains a found token or not. */
+	bool bIsMatch = false;
+};
+
+enum class EPCGGraphFindMode : uint8
+{
+	ShowMinimumTree,
+	ShowFullTree,
+	ShowFlatList
+};
+
+class SPCGEditorGraphFind : public SCompoundWidget
+{
+public:
+	SLATE_BEGIN_ARGS(SPCGEditorGraphFind) {}
+	SLATE_END_ARGS()
+
+	~SPCGEditorGraphFind();
+	void Construct(const FArguments& InArgs, TSharedPtr<FPCGEditor> InPCGEditor);
+
+	/** Focuses this widget's search box */
+	void FocusForUse();
+
+private:
+	typedef TSharedPtr<FPCGEditorGraphFindResult> FPCGEditorGraphFindResultPtr;
+	typedef STreeView<FPCGEditorGraphFindResultPtr> STreeViewType;
+
+	/** Called when user changes the text they are searching for */
+	void OnSearchTextChanged(const FText& InText);
+
+	/** Called when user changes commits text to the search box */
+	void OnSearchTextCommitted(const FText& InText, ETextCommit::Type InCommitType);
+
+	/** Called when the debug object selected changes, which should trigger a new search */
+	void OnInspectedStackChanged(const FPCGStack& InPCGStack);
+
+	/** Get the children of a row */
+	void OnGetChildren(FPCGEditorGraphFindResultPtr InItem, TArray<FPCGEditorGraphFindResultPtr>& OutChildren);
+
+	/** Called when user clicks on a new result */
+	void OnTreeSelectionChanged(FPCGEditorGraphFindResultPtr Item, ESelectInfo::Type);
+
+	/** Called when an element is double-clicked */
+	void OnTreeDoubleClick(FPCGEditorGraphFindResultPtr InItem);
+
+	/** Called when keys are entered when in the tree selection */
+	FReply OnTreeViewKeyDown(const FGeometry& InGeometry, const FKeyEvent& InKeyEvent) const;
+
+	/** Called when the user clicks on the filter icon */
+	TSharedRef<SWidget> OnFindFilterMenu();
+
+	/** Called when a new row is being generated */
+	TSharedRef<ITableRow> OnGenerateRow(FPCGEditorGraphFindResultPtr InItem, const TSharedRef<STableViewBase>& InOwnerTable);
+
+	/** Begins the search based on the SearchValue */
+	void InitiateSearch();
+
+	/** Find any results that contain all of the tokens */
+	void MatchTokens(const TArray<FString>& InTokens);
+
+	/** Recursive internal implementation of the MatchTokens */
+	void MatchTokensInternal(const TArray<FString>& InTokens, UPCGEditorGraph* PCGEditorGraph, TFunctionRef<FPCGEditorGraphFindResultPtr()> GetParentFunc) const;
+
+	/** Determines if a string matches the search tokens */
+	static bool StringMatchesSearchTokens(const TArray<FString>& InTokens, const FString& InComparisonString);
+
+	/** Called when the find mode is changed */
+	void SetFindMode(EPCGGraphFindMode InFindMode);
+
+	/** Returns true if the current find mode is the same as the one provided */
+	bool IsCurrentFindMode(EPCGGraphFindMode InFindMode) const;
+
+	/** Pointer back to the PCG editor that owns us */
+	TWeakPtr<FPCGEditor> PCGEditorPtr;
+
+	/** The tree view displays the results */
+	TSharedPtr<STreeViewType> TreeView;
+
+	/** The search text box */
+	TSharedPtr<SSearchBox> SearchTextField;
+
+	/** This buffer stores the currently displayed results */
+	TArray<FPCGEditorGraphFindResultPtr> ItemsFound;
+
+	/** The string to highlight in the results */
+	FText HighlightText;
+
+	/** The string to search for */
+	FString SearchValue;
+
+	/** Controls the way find results are presented. */
+	EPCGGraphFindMode FindMode = EPCGGraphFindMode::ShowMinimumTree;
+
+	/** Controls whether pin names are show in the find results. */
+	bool bShowPinResults = false;
+};
diff --git a/Engine/Plugins/PCG/Source/PCGEditor/Private/Widgets/SPCGEditorGraphLogView.cpp b/Engine/Plugins/PCG/Source/PCGEditor/Private/Widgets/SPCGEditorGraphLogView.cpp
new file mode 100644
--- /dev/null
+++ b/Engine/Plugins/PCG/Source/PCGEditor/Private/Widgets/SPCGEditorGraphLogView.cpp
@@ -0,0 +1,423 @@
+// Copyright Epic Games, Inc. All Rights Reserved.
+
+#include "SPCGEditorGraphLogView.h"
+
+#include "PCGComponent.h"
+#include "PCGGraph.h"
+#include "PCGSubgraph.h"
+
+#include "PCGEditor.h"
+#include "PCGEditorGraph.h"
+#include "Nodes/PCGEditorGraphNode.h"
+
+#include "Styling/StyleColors.h"
+#include "Widgets/Input/SButton.h"
+#include "Widgets/Layout/SScrollBox.h"
+#include "Widgets/Views/SListView.h"
+
+#define LOCTEXT_NAMESPACE "SPCGEditorGraphLogView"
+
+namespace PCGEditorGraphLogView
+{	
+	const FName NAME_Order = FName(TEXT("Order"));
+	const FName NAME_Node = FName(TEXT("Node"));
+	const FName NAME_Namespace = FName(TEXT("Namespace"));
+	const FName NAME_Message = FName(TEXT("Message"));
+
+	/** Labels of the columns */
+	const FText TEXT_OrderLabel = LOCTEXT("Order", "Order");
+	const FText TEXT_NodeLabel = LOCTEXT("Node", "Node");
+	const FText TEXT_NamespaceLabel = LOCTEXT("Namespace", "Namespace");
+	const FText TEXT_MessageLabel = LOCTEXT("Message", "Message");
+}
+
+void SPCGLogListViewItemRow::Construct(const FArguments& InArgs, const TSharedRef<STableViewBase>& InOwnerTableView, const PCGLogListViewItemPtr& Item)
+{
+	InternalItem = Item;
+
+	SMultiColumnTableRow<PCGLogListViewItemPtr>::Construct(
+		SMultiColumnTableRow::FArguments()
+		.Style(FAppStyle::Get(), "DataTableEditor.CellListViewRow"),
+		InOwnerTableView);
+}
+
+TSharedRef<SWidget> SPCGLogListViewItemRow::GenerateWidgetForColumn(const FName& ColumnId)
+{
+	FText ColumnData = LOCTEXT("ColumnError", "Unrecognized Column");
+
+	if (InternalItem.IsValid())
+	{
+		if (ColumnId == PCGEditorGraphLogView::NAME_Order)
+		{
+			ColumnData = FText::AsNumber(InternalItem->Order);
+		}
+		else if (ColumnId == PCGEditorGraphLogView::NAME_Node)
+		{
+			return SNew(STextBlock)
+				.ColorAndOpacity(InternalItem->EditorNode->GetNodeTitleColor())
+				.Text(FText::FromName(InternalItem->NodeName));
+		}
+		else if (ColumnId == PCGEditorGraphLogView::NAME_Namespace)
+		{
+			ColumnData = FText::FromName(InternalItem->Namespace);
+		}
+		else if (ColumnId == PCGEditorGraphLogView::NAME_Message)
+		{
+			EStyleColor Color = EStyleColor::Foreground;
+
+			switch (InternalItem->Verbosity)
+			{
+			case ELogVerbosity::Error:
+				Color = EStyleColor::Error;
+				break;
+
+			case ELogVerbosity::Warning:
+				Color = EStyleColor::Warning;
+				break;
+			}
+
+			return SNew(STextBlock)
+				.ColorAndOpacity(Color)
+				.Text(FText::FromString(InternalItem->Message));
+		}
+	}
+
+	return SNew(STextBlock)
+		.Text(ColumnData);
+}
+
+void SPCGEditorGraphLogView::OnItemDoubleClicked(PCGLogListViewItemPtr Item)
+{
+	const TSharedPtr<FPCGEditor> PCGEditor = PCGEditorPtr.Pin();
+	if (!PCGEditor.IsValid())
+	{
+		return;
+	}
+
+	if (!Item.IsValid() || !Item->EditorNode)
+	{
+		return;
+	}
+
+	PCGEditor->JumpToNode(Item->EditorNode);
+}
+
+SPCGEditorGraphLogView::~SPCGEditorGraphLogView()
+{
+	if (PCGEditorPtr.IsValid())
+	{
+		PCGEditorPtr.Pin()->OnInspectedStackChangedDelegate.RemoveAll(this);
+	}
+}
+
+void SPCGEditorGraphLogView::Construct(const FArguments& InArgs, TSharedPtr<FPCGEditor> InPCGEditor)
+{
+	PCGEditorPtr = InPCGEditor;
+
+	const TSharedPtr<FPCGEditor> PCGEditor = PCGEditorPtr.Pin();
+	if (PCGEditor)
+	{
+		PCGEditorGraph = PCGEditor->GetPCGEditorGraph();
+		PCGEditor->OnInspectedStackChangedDelegate.AddSP(this, &SPCGEditorGraphLogView::OnDebugStackChanged);
+	}
+
+	ListViewHeader = CreateHeaderRowWidget();
+
+	const TSharedRef<SScrollBar> HorizontalScrollBar = SNew(SScrollBar)
+		.Orientation(Orient_Horizontal)
+		.Thickness(FVector2D(12.0f, 12.0f));
+
+	const TSharedRef<SScrollBar> VerticalScrollBar = SNew(SScrollBar)
+		.Orientation(Orient_Vertical)
+		.Thickness(FVector2D(12.0f, 12.0f));
+	
+	SAssignNew(ListView, SListView<PCGLogListViewItemPtr>)
+		.ListItemsSource(&ListViewItems)
+		.HeaderRow(ListViewHeader)
+		.OnGenerateRow(this, &SPCGEditorGraphLogView::OnGenerateRow)
+		.OnMouseButtonDoubleClick(this, &SPCGEditorGraphLogView::OnItemDoubleClicked)
+		.AllowOverscroll(EAllowOverscroll::No)
+		.ExternalScrollbar(VerticalScrollBar)
+		.ConsumeMouseWheel(EConsumeMouseWheel::Always);
+	
+	this->ChildSlot
+	[
+		SNew(SVerticalBox)
+		+SVerticalBox::Slot()
+		.AutoHeight()
+		[
+			SNew(SHorizontalBox)
+			+SHorizontalBox::Slot()
+			.AutoWidth()
+			[
+				SNew(SButton)
+				.Text(LOCTEXT("ClearButton", "Clear"))
+				.OnClicked(this, &SPCGEditorGraphLogView::Clear)
+			]
+		]
+		+SVerticalBox::Slot()
+		.FillHeight(1.0f)
+		[
+			SNew(SHorizontalBox)
+			+SHorizontalBox::Slot()
+			.FillWidth(1.0f)
+			[
+				SNew(SScrollBox)
+				.Orientation(Orient_Horizontal)
+				.ExternalScrollbar(HorizontalScrollBar)
+				+SScrollBox::Slot()
+				[
+					ListView->AsShared()
+				]
+			]
+			+SHorizontalBox::Slot()
+			.AutoWidth()
+			[
+				VerticalScrollBar
+			]
+		]
+		+SVerticalBox::Slot()
+		.AutoHeight()
+		[
+			SNew(SHorizontalBox)
+			+SHorizontalBox::Slot()
+			[
+				HorizontalScrollBar
+			]
+		]
+	];
+
+	Refresh();
+}
+
+void SPCGEditorGraphLogView::Tick(const FGeometry& AllottedGeometry, const double InCurrentTime, const float InDeltaTime)
+{
+	SCompoundWidget::Tick(AllottedGeometry, InCurrentTime, InDeltaTime);
+
+	if (bNeedsRefresh)
+	{
+		bNeedsRefresh = false;
+		Refresh();
+	}
+}
+
+TSharedRef<SHeaderRow> SPCGEditorGraphLogView::CreateHeaderRowWidget()
+{
+	return SNew(SHeaderRow)
+		.ResizeMode(ESplitterResizeMode::FixedPosition)
+		.CanSelectGeneratedColumn(true)
+		+ SHeaderRow::Column(PCGEditorGraphLogView::NAME_Order)
+			.ManualWidth(64)
+			.DefaultLabel(PCGEditorGraphLogView::TEXT_OrderLabel)
+			.HAlignHeader(HAlign_Center)
+			.HAlignCell(HAlign_Center)
+			.SortMode(this, &SPCGEditorGraphLogView::GetColumnSortMode, PCGEditorGraphLogView::NAME_Order)
+			.OnSort(this, &SPCGEditorGraphLogView::OnSortColumnHeader)
+		+ SHeaderRow::Column(PCGEditorGraphLogView::NAME_Namespace)
+			.ManualWidth(180)
+			.DefaultLabel(PCGEditorGraphLogView::TEXT_NamespaceLabel)
+			.HAlignHeader(HAlign_Center)
+			.HAlignCell(HAlign_Left)
+			.SortMode(this, &SPCGEditorGraphLogView::GetColumnSortMode, PCGEditorGraphLogView::NAME_Namespace)
+			.OnSort(this, &SPCGEditorGraphLogView::OnSortColumnHeader)
+		+ SHeaderRow::Column(PCGEditorGraphLogView::NAME_Node)
+			.ManualWidth(180)
+			.DefaultLabel(PCGEditorGraphLogView::TEXT_NodeLabel)
+			.HAlignHeader(HAlign_Center)
+			.HAlignCell(HAlign_Left)
+			.SortMode(this, &SPCGEditorGraphLogView::GetColumnSortMode, PCGEditorGraphLogView::NAME_Node)
+			.OnSort(this, &SPCGEditorGraphLogView::OnSortColumnHeader)
+		+ SHeaderRow::Column(PCGEditorGraphLogView::NAME_Message)
+			.FillWidth(1.0)
+			.DefaultLabel(PCGEditorGraphLogView::TEXT_MessageLabel)
+			.HAlignHeader(HAlign_Center)
+			.HAlignCell(HAlign_Left)
+			.SortMode(this, &SPCGEditorGraphLogView::GetColumnSortMode, PCGEditorGraphLogView::NAME_Message)
+			.OnSort(this, &SPCGEditorGraphLogView::OnSortColumnHeader);
+}
+
+void SPCGEditorGraphLogView::OnSortColumnHeader(const EColumnSortPriority::Type SortPriority, const FName& ColumnId, const EColumnSortMode::Type NewSortMode)
+{
+	if (SortingColumn == ColumnId)
+	{
+		// Circling
+		SortMode = EColumnSortMode::Type((SortMode + 1) % 3);
+	}
+	else
+	{
+		SortingColumn = ColumnId;
+		SortMode = NewSortMode;
+	}
+
+	Refresh();
+}
+
+EColumnSortMode::Type SPCGEditorGraphLogView::GetColumnSortMode(const FName ColumnId) const
+{
+	if (SortingColumn != ColumnId)
+	{
+		return EColumnSortMode::None;
+	}
+
+	return SortMode;
+}
+
+FReply SPCGEditorGraphLogView::Clear()
+{
+	const TSharedPtr<FPCGEditor> PCGEditor = PCGEditorPtr.Pin();
+	if (!PCGEditor.IsValid())
+	{
+		return FReply::Handled();
+	}
+
+	if (!PCGEditorGraph)
+	{
+		return FReply::Handled();
+	}
+
+	if (const UPCGComponent* Component = GetPCGComponent().Get())
+	{
+		Component->ExtraCapture.ResetCapturedMessages();
+		Refresh();
+	}
+
+	return FReply::Handled();
+}
+
+FReply SPCGEditorGraphLogView::Refresh()
+{
+	ListViewItems.Empty();
+	ListView->RequestListRefresh();
+
+	const TSharedPtr<FPCGEditor> PCGEditor = PCGEditorPtr.Pin();
+	if (!PCGEditor.IsValid())
+	{
+		return FReply::Handled();
+	}
+
+	if (!PCGEditorGraph)
+	{
+		return FReply::Handled();
+	}
+
+	const UPCGComponent* Component = GetPCGComponent().Get();
+
+	if (!Component)
+	{
+		return FReply::Handled();
+	}
+
+	TArray<UPCGEditorGraphNode*> EditorNodes;
+	PCGEditorGraph->GetNodesOfClass<UPCGEditorGraphNode>(EditorNodes);
+
+	for (const UPCGEditorGraphNode* PCGEditorNode : EditorNodes)
+	{
+		const UPCGNode* PCGNode = PCGEditorNode ? PCGEditorNode->GetPCGNode() : nullptr;
+		if (PCGNode)
+		{
+			CreateAndAddItem(PCGEditorNode, PCGNode, *PCGNode->GetNodeTitle(EPCGNodeTitleType::ListView).ToString());
+		}
+	}
+
+	if (SortingColumn != NAME_None && SortMode != EColumnSortMode::None)
+	{
+		Algo::Sort(ListViewItems, [this](const PCGLogListViewItemPtr& A, const PCGLogListViewItemPtr& B)
+			{
+				bool isLess = false;
+				if (SortingColumn == PCGEditorGraphLogView::NAME_Order)
+				{
+					isLess = A->Order < B->Order;
+				}
+				else if (SortingColumn == PCGEditorGraphLogView::NAME_Namespace)
+				{
+					isLess = A->Namespace.LexicalLess(B->Namespace);
+				}
+				else if (SortingColumn == PCGEditorGraphLogView::NAME_Node)
+				{
+					isLess = A->NodeName.LexicalLess(B->NodeName);
+				}
+				else if (SortingColumn == PCGEditorGraphLogView::NAME_Message)
+				{
+					isLess = A->Message < B->Message;
+				}
+
+				return SortMode == EColumnSortMode::Ascending ? isLess : !isLess;
+			});
+	}
+
+	ListView->SetItemsSource(&ListViewItems);
+
+	return FReply::Handled();
+}
+
+void SPCGEditorGraphLogView::CreateAndAddItem(const UPCGEditorGraphNode* InPCGEditorNode, const UPCGNode* InPCGNode, const FName& InName)
+{
+	const TSharedPtr<FPCGEditor> PCGEditor = PCGEditorPtr.Pin();
+
+	check(PCGEditor);
+
+	const UPCGComponent* Component = GetPCGComponent().Get();
+
+	if (!Component || !InPCGEditorNode || !InPCGNode)
+	{
+		return;
+	}
+
+	const TArray<PCGUtils::FCapturedMessage>* CapturedMessagesPtr = Component->ExtraCapture.GetCapturedMessages().Find(InPCGNode);
+	if (!CapturedMessagesPtr)
+	{
+		return;
+	}
+
+	for (const PCGUtils::FCapturedMessage& Message : *CapturedMessagesPtr)
+	{
+		PCGLogListViewItemPtr ListViewItem = MakeShared<FPCGLogListViewItem>();
+		ListViewItem->EditorNode = InPCGEditorNode;
+		ListViewItem->PCGNode = InPCGNode;
+		ListViewItem->NodeName = InName;
+		ListViewItem->Order = Message.Index;
+		ListViewItem->Namespace = Message.Namespace;
+		ListViewItem->Message = Message.Message;
+		ListViewItem->Verbosity = Message.Verbosity;
+
+		ListViewItems.Add(ListViewItem);
+	}
+
+	// Also if the node is a subgraph, recurse on all child nodes
+	const UPCGSettings* Settings = InPCGNode->GetSettings();
+	if (Settings && Settings->IsA<UPCGBaseSubgraphSettings>())
+	{
+		if (const UPCGGraph* Graph = CastChecked<const UPCGBaseSubgraphSettings>(Settings)->GetSubgraph())
+		{
+			const FString Prefix = InName.ToString() + "/";
+			for (const UPCGNode* ChildNode : Graph->GetNodes())
+			{
+				if (!ChildNode)
+				{
+					continue;
+				}
+
+				const FText ChildTitle = ChildNode->GetNodeTitle(EPCGNodeTitleType::ListView);
+
+				CreateAndAddItem(InPCGEditorNode, ChildNode, FName(Prefix + ChildTitle.ToString()));
+			}
+		}
+	}
+}
+
+TWeakObjectPtr<UPCGComponent> SPCGEditorGraphLogView::GetPCGComponent() const
+{
+	return PCGEditorPtr.IsValid() ? PCGEditorPtr.Pin()->GetPCGComponentBeingInspected() : nullptr;
+}
+
+void SPCGEditorGraphLogView::OnDebugStackChanged(const FPCGStack& InPCGStack)
+{
+	RequestRefresh();
+}
+
+TSharedRef<ITableRow> SPCGEditorGraphLogView::OnGenerateRow(PCGLogListViewItemPtr Item, const TSharedRef<STableViewBase>& OwnerTable) const
+{
+	return SNew(SPCGLogListViewItemRow, OwnerTable, Item);
+}
+
+#undef LOCTEXT_NAMESPACE
diff --git a/Engine/Plugins/PCG/Source/PCGEditor/Private/Widgets/SPCGEditorGraphLogView.h b/Engine/Plugins/PCG/Source/PCGEditor/Private/Widgets/SPCGEditorGraphLogView.h
new file mode 100644
--- /dev/null
+++ b/Engine/Plugins/PCG/Source/PCGEditor/Private/Widgets/SPCGEditorGraphLogView.h
@@ -0,0 +1,97 @@
+// Copyright Epic Games, Inc. All Rights Reserved.
+
+#pragma once
+
+#include "Widgets/Views/ITableRow.h"
+#include "Widgets/Views/SListView.h"
+#include "Widgets/Views/STableRow.h"
+
+class STableViewBase;
+
+class FPCGEditor;
+class UPCGComponent;
+class UPCGEditorGraphNode;
+class UPCGEditorGraph;
+class UPCGNode;
+struct FPCGStack;
+
+namespace PCGEditorGraphLogView
+{	
+	/** Names of the columns in the attribute list */
+	extern const FName NAME_Order;
+}
+
+struct FPCGLogListViewItem
+{
+	const UPCGNode* PCGNode = nullptr;	
+	const UPCGEditorGraphNode* EditorNode = nullptr;
+
+	int32 Order = 0;
+	FName NodeName;
+	FName Namespace;
+	FString Message;
+
+	ELogVerbosity::Type Verbosity = ELogVerbosity::NoLogging;
+};
+
+typedef TSharedPtr<FPCGLogListViewItem> PCGLogListViewItemPtr;
+
+class SPCGLogListViewItemRow : public SMultiColumnTableRow<PCGLogListViewItemPtr>
+{
+public:
+	SLATE_BEGIN_ARGS(SPCGLogListViewItemRow) {}
+	SLATE_END_ARGS()
+
+	void Construct(const FArguments& InArgs, const TSharedRef<STableViewBase>& InOwnerTableView, const PCGLogListViewItemPtr& Item);
+
+	virtual TSharedRef<SWidget> GenerateWidgetForColumn(const FName& ColumnId) override;
+
+protected:
+	PCGLogListViewItemPtr InternalItem;
+};
+
+class SPCGEditorGraphLogView : public SCompoundWidget
+{
+public:
+	SLATE_BEGIN_ARGS(SPCGEditorGraphLogView) { }
+	SLATE_END_ARGS()
+
+	~SPCGEditorGraphLogView();
+
+	void Construct(const FArguments& InArgs, TSharedPtr<FPCGEditor> InPCGEditor);
+	virtual void Tick(const FGeometry& AllottedGeometry, const double InCurrentTime, const float InDeltaTime) override;
+	void RequestRefresh() { bNeedsRefresh = true; }
+
+private:
+	TSharedRef<SHeaderRow> CreateHeaderRowWidget();
+	void CreateAndAddItem(const UPCGEditorGraphNode* InPCGEditorNode, const UPCGNode* InPCGNode, const FName& InName);
+
+	void OnDebugStackChanged(const FPCGStack& InPCGStack);
+
+	// Callbacks
+	FReply Refresh();
+	FReply Clear();
+	TSharedRef<ITableRow> OnGenerateRow(PCGLogListViewItemPtr Item, const TSharedRef<STableViewBase>& OwnerTable) const;
+	void OnItemDoubleClicked(PCGLogListViewItemPtr Item);
+	void OnSortColumnHeader(const EColumnSortPriority::Type SortPriority, const FName& ColumnId, const EColumnSortMode::Type NewSortMode);
+	EColumnSortMode::Type GetColumnSortMode(const FName ColumnId) const;
+
+	/** Pointer back to the PCG editor that owns us */
+	TWeakPtr<FPCGEditor> PCGEditorPtr;
+
+	/** Cached PCGGraph being viewed */
+	UPCGEditorGraph* PCGEditorGraph = nullptr;
+
+	/** Cached PCGComponent being viewed */
+	TWeakObjectPtr<UPCGComponent> GetPCGComponent() const;
+
+	TSharedPtr<SHeaderRow> ListViewHeader;
+	TSharedPtr<SListView<PCGLogListViewItemPtr>> ListView;
+	TArray<PCGLogListViewItemPtr> ListViewItems;
+
+	bool bNeedsRefresh : 1 = false;
+
+	// To allow sorting
+	FName SortingColumn = PCGEditorGraphLogView::NAME_Order;
+	EColumnSortMode::Type SortMode = EColumnSortMode::Type::Ascending;
+};
diff --git a/Engine/Plugins/PCG/Source/PCGEditor/Private/Widgets/SPCGEditorGraphNodePalette.cpp b/Engine/Plugins/PCG/Source/PCGEditor/Private/Widgets/SPCGEditorGraphNodePalette.cpp
new file mode 100644
--- /dev/null
+++ b/Engine/Plugins/PCG/Source/PCGEditor/Private/Widgets/SPCGEditorGraphNodePalette.cpp
@@ -0,0 +1,178 @@
+// Copyright Epic Games, Inc. All Rights Reserved.
+
+#include "SPCGEditorGraphNodePalette.h"
+
+#include "AssetRegistry/AssetData.h"
+#include "PCGGraph.h"
+
+#include "PCGEditor.h"
+#include "PCGEditorGraph.h"
+#include "PCGEditorGraphSchema.h"
+#include "PCGEditorGraphSchemaActions.h"
+#include "PCGEditorUtils.h"
+#include "PCGSettingsDragDropAction.h"
+
+#include "SEnumCombo.h"
+#include "AssetRegistry/AssetRegistryModule.h"
+#include "Modules/ModuleManager.h"
+
+#define LOCTEXT_NAMESPACE "SPCGEditorGraphNodePalette"
+
+void SPCGEditorGraphNodePaletteItem::Construct(const FArguments& InArgs, FCreateWidgetForActionData* const InCreateData)
+{
+	check(InCreateData->Action.IsValid());
+
+	TSharedPtr<FEdGraphSchemaAction> GraphAction = InCreateData->Action;
+	ActionPtr = InCreateData->Action;
+
+	bool bIsReadOnly = false;
+	TSharedRef<SWidget> NameSlotWidget = CreateTextSlotWidget(InCreateData, bIsReadOnly);
+
+	this->ChildSlot
+	[
+		NameSlotWidget
+	];
+}
+
+FText SPCGEditorGraphNodePaletteItem::GetItemTooltip() const
+{
+	return ActionPtr.Pin()->GetTooltipDescription();
+}
+
+void SPCGEditorGraphNodePalette::Construct(const FArguments& InArgs, TSharedPtr<FPCGEditor> InPCGEditor)
+{
+	const UEnum* PCGElementTypeEnum = StaticEnum<EPCGElementType>();
+	PCGEditor = InPCGEditor;
+	
+	this->ChildSlot
+	[
+		SNew(SVerticalBox)
+		+SVerticalBox::Slot()
+		.AutoHeight()
+		[
+			SNew(SHorizontalBox)
+			+SHorizontalBox::Slot()
+			.AutoWidth()
+			.VAlign(VAlign_Center)
+			.Padding(2.0f, 0.0f)
+			[
+				SNew(STextBlock)
+				.Text(LOCTEXT("TypeTextBlock", "Type:"))
+			]
+			+SHorizontalBox::Slot()
+			.VAlign(VAlign_Center)
+			.FillWidth(1.0f)
+			[
+				SNew(SEnumComboBox, PCGElementTypeEnum)
+				.ContentPadding(FMargin(4, 0))
+				.OnEnumSelectionChanged(this, &SPCGEditorGraphNodePalette::OnTypeSelectionChanged)
+				.CurrentValue(this, &SPCGEditorGraphNodePalette::GetTypeValue)
+			]
+		]
+		+SVerticalBox::Slot()
+		.FillHeight(1.0f)
+		[
+			SAssignNew(GraphActionMenu, SGraphActionMenu)
+			.OnActionDragged(this, &SPCGEditorGraphNodePalette::OnActionDragged)
+			.OnCreateWidgetForAction(this, &SPCGEditorGraphNodePalette::OnCreateWidgetForAction)
+			.OnCollectAllActions(this, &SPCGEditorGraphNodePalette::CollectAllActions)
+			.AutoExpandActionMenu(true)
+		]
+	];
+
+	FAssetRegistryModule& AssetRegistryModule = FModuleManager::GetModuleChecked<FAssetRegistryModule>(TEXT("AssetRegistry"));
+	AssetRegistryModule.Get().OnAssetAdded().AddSP(this, &SPCGEditorGraphNodePalette::OnAssetChanged);
+	AssetRegistryModule.Get().OnAssetRemoved().AddSP(this, &SPCGEditorGraphNodePalette::OnAssetChanged);
+	AssetRegistryModule.Get().OnAssetUpdated().AddSP(this, &SPCGEditorGraphNodePalette::OnAssetChanged); // Todo, evaluate if this triggers too often
+	AssetRegistryModule.Get().OnAssetRenamed().AddSP(this, &SPCGEditorGraphNodePalette::OnAssetRenamed);
+}
+
+SPCGEditorGraphNodePalette::~SPCGEditorGraphNodePalette()
+{
+	if (FModuleManager::Get().IsModuleLoaded(TEXT("AssetRegistry")))
+	{
+		IAssetRegistry* AssetRegistry = FModuleManager::GetModuleChecked<FAssetRegistryModule>(TEXT("AssetRegistry")).TryGet();
+		if (AssetRegistry) // AssetRegistry can be null during EngineShutdown even if module is available
+		{
+			AssetRegistry->OnAssetAdded().RemoveAll(this);
+			AssetRegistry->OnAssetRemoved().RemoveAll(this);
+			AssetRegistry->OnAssetUpdated().RemoveAll(this);
+			AssetRegistry->OnAssetRenamed().RemoveAll(this);
+		}
+	}
+}
+
+void SPCGEditorGraphNodePalette::RequestRefresh()
+{
+	bNeedsRefresh = true;
+}
+
+TSharedRef<SWidget> SPCGEditorGraphNodePalette::OnCreateWidgetForAction(FCreateWidgetForActionData* const InCreateData)
+{
+	return SNew(SPCGEditorGraphNodePaletteItem, InCreateData);
+}
+
+FReply SPCGEditorGraphNodePalette::OnActionDragged(const TArray<TSharedPtr<FEdGraphSchemaAction>>& InActions, const FPointerEvent& MouseEvent)
+{
+	if (InActions.Num() > 0 && InActions[0].IsValid())
+	{
+		TSharedPtr<FEdGraphSchemaAction> InAction = InActions[0];
+		if (InAction->GetTypeId() == FPCGEditorGraphSchemaAction_NewSettingsElement::StaticGetTypeId())
+		{
+			FPCGEditorGraphSchemaAction_NewSettingsElement* SettingsAction = static_cast<FPCGEditorGraphSchemaAction_NewSettingsElement*>(InAction.Get());
+			return FReply::Handled().BeginDragDrop(FPCGSettingsDragDropAction::New(InAction, SettingsAction->SettingsObjectPath));
+		}
+	}
+
+	return SGraphPalette::OnActionDragged(InActions, MouseEvent);
+}
+
+void SPCGEditorGraphNodePalette::CollectAllActions(FGraphActionListBuilderBase& OutAllActions)
+{
+	const UPCGEditorGraphSchema* PCGSchema = GetDefault<UPCGEditorGraphSchema>();
+
+	FGraphActionMenuBuilder ActionMenuBuilder;
+	PCGSchema->GetPaletteActions(ActionMenuBuilder, FPCGActionsFilter(PCGEditor.Pin()->GetPCGEditorGraph(), ElementType));
+	OutAllActions.Append(ActionMenuBuilder);
+}
+
+void SPCGEditorGraphNodePalette::OnAssetChanged(const FAssetData& InAssetData)
+{
+	if (InAssetData.IsInstanceOf<UPCGGraphInterface>() ||
+		InAssetData.IsInstanceOf<UPCGSettings>() ||
+		PCGEditorUtils::IsAssetPCGBlueprint(InAssetData))
+	{
+		bNeedsRefresh = true;
+	}
+}
+
+void SPCGEditorGraphNodePalette::OnAssetRenamed(const FAssetData& InAssetData, const FString& /*InNewAssetName*/)
+{
+	OnAssetChanged(InAssetData);
+}
+
+void SPCGEditorGraphNodePalette::OnTypeSelectionChanged(int32 InValue, ESelectInfo::Type /*SelectInfo*/)
+{
+	ElementType = static_cast<EPCGElementType>(InValue);
+	bNeedsRefresh = true;
+}
+
+void SPCGEditorGraphNodePalette::Tick(const FGeometry& AllottedGeometry, const double InCurrentTime, const float InDeltaTime)
+{
+	SGraphPalette::Tick(AllottedGeometry, InCurrentTime, InDeltaTime);
+
+	if (bNeedsRefresh)
+	{
+		TRACE_CPUPROFILER_EVENT_SCOPE(SPCGEditorGraphNodePalette::RefreshActionsList);
+		RefreshActionsList(true);
+		bNeedsRefresh = false;
+	}
+}
+
+int32 SPCGEditorGraphNodePalette::GetTypeValue() const
+{
+	return static_cast<int32>(ElementType);
+}
+
+#undef LOCTEXT_NAMESPACE
+ 
diff --git a/Engine/Plugins/PCG/Source/PCGEditor/Private/Widgets/SPCGEditorGraphNodePalette.h b/Engine/Plugins/PCG/Source/PCGEditor/Private/Widgets/SPCGEditorGraphNodePalette.h
new file mode 100644
--- /dev/null
+++ b/Engine/Plugins/PCG/Source/PCGEditor/Private/Widgets/SPCGEditorGraphNodePalette.h
@@ -0,0 +1,58 @@
+// Copyright Epic Games, Inc. All Rights Reserved.
+
+#pragma once
+
+#include "PCGEditorCommon.h"
+
+#include "SGraphPalette.h"
+
+class FPCGEditor;
+
+class SPCGEditorGraphNodePaletteItem : public SGraphPaletteItem
+{
+public:
+	SLATE_BEGIN_ARGS(SPCGEditorGraphNodePaletteItem) {};
+	SLATE_END_ARGS()
+	
+	void Construct(const FArguments& InArgs, FCreateWidgetForActionData* const InCreateData);
+
+protected:
+	//~ Begin SGraphPaletteItem Interface
+	virtual FText GetItemTooltip() const override;
+	//~ End SGraphPaletteItem Interface
+};
+
+
+class SPCGEditorGraphNodePalette : public SGraphPalette
+{
+public:
+	SLATE_BEGIN_ARGS(SPCGEditorGraphNodePalette) {};
+	SLATE_END_ARGS()
+
+	void Construct(const FArguments& InArgs, TSharedPtr<FPCGEditor> InPCGEditor);
+	virtual ~SPCGEditorGraphNodePalette();
+
+	void RequestRefresh();
+
+protected:
+	//~ Begin SGraphPalette Interface
+	virtual TSharedRef<SWidget> OnCreateWidgetForAction(FCreateWidgetForActionData* const InCreateData) override;
+	virtual FReply OnActionDragged(const TArray<TSharedPtr<FEdGraphSchemaAction>>& InActions, const FPointerEvent& MouseEvent) override;
+	virtual void CollectAllActions(FGraphActionListBuilderBase& OutAllActions) override;
+	//~ End SGraphPalette Interface
+
+	//~ Begin SWidget interface
+	void Tick(const FGeometry& AllottedGeometry, const double InCurrentTime, const float InDeltaTime) override;
+	//~ End SWidget interface
+
+private:
+	void OnAssetChanged(const FAssetData& InAssetData);
+	void OnAssetRenamed(const FAssetData& InAssetData, const FString& InNewAssetName);
+	void OnTypeSelectionChanged(int32, ESelectInfo::Type SelectInfo);
+	
+	int32 GetTypeValue() const;
+	
+	TWeakPtr<FPCGEditor> PCGEditor;
+	EPCGElementType ElementType = EPCGElementType::All;
+	bool bNeedsRefresh = false;
+};
diff --git a/Engine/Plugins/PCG/Source/PCGEditor/Private/Widgets/SPCGEditorGraphProfilingView.cpp b/Engine/Plugins/PCG/Source/PCGEditor/Private/Widgets/SPCGEditorGraphProfilingView.cpp
new file mode 100644
--- /dev/null
+++ b/Engine/Plugins/PCG/Source/PCGEditor/Private/Widgets/SPCGEditorGraphProfilingView.cpp
@@ -0,0 +1,823 @@
+// Copyright Epic Games, Inc. All Rights Reserved.
+
+#include "SPCGEditorGraphProfilingView.h"
+
+#include "PCGComponent.h"
+#include "Graph/PCGStackContext.h"
+
+#include "PCGEditor.h"
+#include "PCGEditorGraph.h"
+#include "Nodes/PCGEditorGraphNode.h"
+
+#include "HAL/PlatformApplicationMisc.h"
+#include "Widgets/Input/SButton.h"
+#include "Widgets/Input/SCheckBox.h"
+#include "Widgets/Input/SSearchBox.h"
+#include "Widgets/Input/SSpinBox.h"
+#include "Widgets/Layout/SScrollBox.h"
+#include "Widgets/Views/SListView.h"
+
+#define LOCTEXT_NAMESPACE "SPCGEditorGraphProfilingView"
+
+namespace PCGEditorGraphProfilingView
+{
+	const FText NoDataAvailableText = LOCTEXT("NoDataAvailableText", "N/A");
+	
+	/** Names of the columns in the attribute list */
+	const FName NAME_Node = FName(TEXT("Node"));
+	const FName NAME_PrepareDataTime = FName(TEXT("PrepareDataTime"));
+	const FName NAME_PrepareDataWallTime = FName(TEXT("PrepareData_WallTime"));
+	const FName NAME_NbPrepareFrames = FName(TEXT("NbPrepareFrames"));
+	const FName NAME_MinExecutionFrameTime = FName(TEXT("MinFrameTime"));
+	const FName NAME_MaxExecutionFrameTime = FName(TEXT("MaxFrameTime"));
+	const FName NAME_ExecutionTime = FName(TEXT("ExecutionTime"));
+	const FName NAME_ExecutionWallTime = FName(TEXT("Execution_WallTime"));
+	const FName NAME_NbExecutionFrames = FName(TEXT("NbExecutionFrames"));
+	const FName NAME_TotalTime = FName(TEXT("TotalTime"));
+	const FName NAME_TotalWallTime = FName(TEXT("Total_WallTime"));
+	const FName NAME_MemoryInMB = FName(TEXT("Data size (MB)"));
+
+	/** Labels of the columns */
+	const FText TEXT_NodeLabel = LOCTEXT("NodeLabel", "Node");
+	const FText TEXT_PrepareDataTimeLabel = LOCTEXT("PrepareDataTimeLabel", "Prepare");
+	const FText TEXT_PrepareDataWallTimeLabel = LOCTEXT("PrepareDataWallTimeLabel", "Prepare WallTime");
+	const FText TEXT_NbPrepareFramesLabel = LOCTEXT("NbPrepareFramesLabel", "Prep Frames");
+	const FText TEXT_MinExecutionFrameTimeLabel = LOCTEXT("MinExecutionFrameTimeLabel", "Min FrameTime");
+	const FText TEXT_MaxExecutionFrameTimeLabel = LOCTEXT("MaxExecutionFrameTimeLabel", "Max FrameTime");
+	const FText TEXT_ExecutionTimeLabel = LOCTEXT("ExecutionTimeLabel", "Exec");
+	const FText TEXT_ExecutionWallTimeLabel = LOCTEXT("ExecutionWallTimeLabel", "Exec WallTime");
+	const FText TEXT_NbExecutionFramesLabel = LOCTEXT("NbExecutionFramesLabel", "Exec Frames");
+	const FText TEXT_TotalTimeLabel = LOCTEXT("TotalTimeLabel", "Total");
+	const FText TEXT_TotalWallTimeLabel = LOCTEXT("TotalWallTimeLabel", "Total WallTime");
+	const FText TEXT_MemoryInMBLabel = LOCTEXT("MemoryInMBLabel", "Data size (MB)");
+
+	/** Tooltips */
+	const FText TEXT_PrepareDataTimeTooltip = LOCTEXT("PrepareDataTimeTooltip", "Cost of the PrepareData execution phase which some nodes use to process the incoming data, in ms.");
+	const FText TEXT_PrepareDataWallTimeTooltip = LOCTEXT("PrepareDataWallTimeTooltip", "Total real time elapsed between prepare data first being called until completion, including any wait/sleep time, in ms.");
+	const FText TEXT_NbPrepareFramesTooltip = LOCTEXT("NbPrepareFramesTooltip", "The number of frames in which one or more prepare data phases were executed.");
+	const FText TEXT_MinExecutionFrameTimeTooltip = LOCTEXT("MinExecutionFrameTimeTooltip", "The minimum time spent of all execution frames, in ms.");
+	const FText TEXT_MaxExecutionFrameTimeTooltip = LOCTEXT("MaxExecutionFrameTimeTooltip", "The maximum time spent of all execution frames, in ms.");
+	const FText TEXT_ExecutionTimeTooltip = LOCTEXT("ExecutionTimeTooltip", "The total time spent for execution, summed over all execution frames, in ms.");
+	const FText TEXT_ExecutionWallTimeTooltip = LOCTEXT("ExecutionWallTimeTooltip", "Total real time elapsed between execute first being called until completion, including any wait/sleep time, in ms.");
+	const FText TEXT_NbExecutionFramesTooltip = LOCTEXT("NbExecutionFramesTooltip", "The number of frames in which one or more execution phases were executed.");
+	const FText TEXT_TotalTimeTooltip = LOCTEXT("TotalTimeTooltip", "The total time spent in this node, summed over all execution and prepare data frames, in ms.");
+	const FText TEXT_TotalWallTimeTooltip = LOCTEXT("TotalWallTimeTooltip", "Total real time elapsed between the first call until completion, including any wait/sleep time, in ms.");
+	const FText TEXT_MemoryInMBTooltip = LOCTEXT("MemoryInMBTooltip", "Size in MB of newly created data.");
+}
+
+void SPCGProfilingListViewItemRow::Construct(const FArguments& InArgs, const TSharedRef<STableViewBase>& InOwnerTableView, const PCGProfilingListViewItemPtr& Item)
+{
+	InternalItem = Item;
+
+	SMultiColumnTableRow<PCGProfilingListViewItemPtr>::Construct(
+		SMultiColumnTableRow::FArguments()
+		.Style(FAppStyle::Get(), "DataTableEditor.CellListViewRow"),
+		InOwnerTableView);
+}
+
+FText FPCGProfilingListViewItem::GetTextForColumn(FName ColumnId, bool bNoGrouping) const
+{
+	const FNumberFormattingOptions* NumberFormattingOptions = (bNoGrouping ? &FNumberFormattingOptions::DefaultNoGrouping() : nullptr);
+
+	if (ColumnId == PCGEditorGraphProfilingView::NAME_Node)
+	{
+		const UPCGSettings* PCGSettings = PCGNode ? PCGNode->GetSettings() : nullptr;
+
+		if (PCGSettings)
+		{
+			if (PCGSettings->ShouldExecuteOnGPU())
+			{
+				return FText::FromString(Name + TEXT(" (profiling GPU nodes not supported)"));
+			}
+			else if (!PCGSettings->bEnabled)
+			{
+				return FText::FromString(Name + TEXT(" (disabled)"));
+			}
+		}
+
+		if (bHasData)
+		{
+			return FText::FromString(Name);
+		}
+		else
+		{
+			return FText::FromString(Name + TEXT(" (cached)"));
+		}
+	}
+	else if (ColumnId == PCGEditorGraphProfilingView::NAME_NbExecutionFrames)
+	{
+		return ((CallTime.ExecutionFrameCount >= 0) ? FText::AsNumber(CallTime.ExecutionFrameCount, NumberFormattingOptions) : FText());
+	}
+	else if (ColumnId == PCGEditorGraphProfilingView::NAME_NbPrepareFrames)
+	{
+		return ((CallTime.PrepareDataFrameCount >= 0) ? FText::AsNumber(CallTime.PrepareDataFrameCount, NumberFormattingOptions) : FText());
+	}
+	else if (!bHasData)
+	{
+		return PCGEditorGraphProfilingView::NoDataAvailableText;
+	}
+	else if (ColumnId == PCGEditorGraphProfilingView::NAME_MinExecutionFrameTime)
+	{
+		// In ms
+		return ((CallTime.MinExecutionFrameTime >= 0) ? FText::AsNumber(CallTime.MinExecutionFrameTime * 1000.0, NumberFormattingOptions) : FText());
+	}
+	else if (ColumnId == PCGEditorGraphProfilingView::NAME_MaxExecutionFrameTime)
+	{
+		// In ms
+		return ((CallTime.MaxExecutionFrameTime >= 0) ? FText::AsNumber(CallTime.MaxExecutionFrameTime * 1000.0, NumberFormattingOptions) : FText());
+	}
+	else if (ColumnId == PCGEditorGraphProfilingView::NAME_ExecutionTime)
+	{
+		// In ms
+		return FText::AsNumber(CallTime.ExecutionTime * 1000.0, NumberFormattingOptions);
+	}
+	else if (ColumnId == PCGEditorGraphProfilingView::NAME_ExecutionWallTime)
+	{
+		// In ms
+		return FText::AsNumber(CallTime.ExecutionWallTime() * 1000.0, NumberFormattingOptions);
+	}
+	else if (ColumnId == PCGEditorGraphProfilingView::NAME_PrepareDataTime)
+	{
+		// In ms
+		return FText::AsNumber(CallTime.PrepareDataTime * 1000.0, NumberFormattingOptions);
+	}
+	else if (ColumnId == PCGEditorGraphProfilingView::NAME_PrepareDataWallTime)
+	{
+		// In ms
+		return FText::AsNumber(CallTime.PrepareDataWallTime() * 1000.0, NumberFormattingOptions);
+	}
+	else if (ColumnId == PCGEditorGraphProfilingView::NAME_TotalTime)
+	{
+		// In ms
+		return FText::AsNumber((CallTime.TotalTime()) * 1000.0, NumberFormattingOptions);
+	}
+	else if (ColumnId == PCGEditorGraphProfilingView::NAME_TotalWallTime)
+	{
+		// In ms
+		return FText::AsNumber((CallTime.TotalWallTime()) * 1000.0, NumberFormattingOptions);
+	}
+	else if (ColumnId == PCGEditorGraphProfilingView::NAME_MemoryInMB)
+	{
+		// In MB
+		return FText::AsNumber(CallTime.OutputDataMemorySize / (1024.0 * 1024.0), NumberFormattingOptions);
+	}
+	else
+	{
+		return LOCTEXT("ItemColumnError", "Unrecognized Column");
+	}
+}
+
+TSharedRef<SWidget> SPCGProfilingListViewItemRow::GenerateWidgetForColumn(const FName& ColumnId)
+{
+	FText ColumnData = LOCTEXT("ColumnError", "Unrecognized Column");
+	FColor ColumnDataColor = FColor::White;
+
+	if (InternalItem.IsValid())
+	{
+		ColumnData = InternalItem->GetTextForColumn(ColumnId, /*bNoGrouping=*/false);
+
+		if (!InternalItem->bHasData)
+		{
+			ColumnDataColor = FColor(75, 75, 75);
+		}
+	}
+
+	TSharedRef<SWidget> DataWidget = 
+		SNew(STextBlock)
+		.Text(ColumnData)
+		.ColorAndOpacity(ColumnDataColor)
+		.OverflowPolicy(ETextOverflowPolicy::Ellipsis);
+
+	// Add the internal name of the node as tooltip
+	if (ColumnId == PCGEditorGraphProfilingView::NAME_Node && InternalItem->PCGNode)
+	{
+		DataWidget->SetToolTipText(FText::FromString(InternalItem->PCGNode->GetName()));
+	}
+
+	return DataWidget;
+}
+
+void SPCGEditorGraphProfilingView::OnItemDoubleClicked(PCGProfilingListViewItemPtr Item)
+{
+	const TSharedPtr<FPCGEditor> PCGEditor = PCGEditorPtr.Pin();
+	if (!PCGEditor.IsValid())
+	{
+		return;
+	}
+
+	if (!Item.IsValid() || !Item->EditorNode)
+	{
+		return;
+	}
+
+	PCGEditor->JumpToNode(Item->EditorNode);
+}
+
+SPCGEditorGraphProfilingView::~SPCGEditorGraphProfilingView()
+{
+	if (PCGEditorPtr.IsValid())
+	{
+		PCGEditorPtr.Pin()->OnInspectedStackChangedDelegate.RemoveAll(this);
+	}
+}
+
+void SPCGEditorGraphProfilingView::Construct(const FArguments& InArgs, TSharedPtr<FPCGEditor> InPCGEditor)
+{
+	PCGEditorPtr = InPCGEditor;
+
+	const TSharedPtr<FPCGEditor> PCGEditor = PCGEditorPtr.Pin();
+	if (PCGEditor)
+	{
+		PCGEditorGraph = PCGEditor->GetPCGEditorGraph();
+		PCGEditor->OnInspectedStackChangedDelegate.AddSP(this, &SPCGEditorGraphProfilingView::OnDebugStackChanged);
+	}
+
+	SortingColumn = PCGEditorGraphProfilingView::NAME_TotalTime;
+	SortMode = EColumnSortMode::Descending;
+	ListViewHeader = CreateHeaderRowWidget();
+
+	ListViewCommands = MakeShareable(new FUICommandList);
+	ListViewCommands->MapAction(FGenericCommands::Get().Copy,
+		FExecuteAction::CreateSP(this, &SPCGEditorGraphProfilingView::CopySelectionToClipboard),
+		FCanExecuteAction::CreateSP(this, &SPCGEditorGraphProfilingView::CanCopySelectionToClipboard));
+
+	const TSharedRef<SScrollBar> HorizontalScrollBar = SNew(SScrollBar)
+		.Orientation(Orient_Horizontal)
+		.Thickness(FVector2D(12.0f, 12.0f));
+
+	const TSharedRef<SScrollBar> VerticalScrollBar = SNew(SScrollBar)
+		.Orientation(Orient_Vertical)
+		.Thickness(FVector2D(12.0f, 12.0f));
+	
+	SAssignNew(ListView, SListView<PCGProfilingListViewItemPtr>)
+		.ListItemsSource(&ListViewItems)
+		.HeaderRow(ListViewHeader)
+		.OnGenerateRow(this, &SPCGEditorGraphProfilingView::OnGenerateRow)
+		.OnMouseButtonDoubleClick(this, &SPCGEditorGraphProfilingView::OnItemDoubleClicked)
+		.AllowOverscroll(EAllowOverscroll::No)
+		.ExternalScrollbar(VerticalScrollBar)
+		.OnKeyDownHandler(this, &SPCGEditorGraphProfilingView::OnListViewKeyDown)
+		.ConsumeMouseWheel(EConsumeMouseWheel::Always);
+	
+	this->ChildSlot
+	[
+		SNew(SVerticalBox)
+		+SVerticalBox::Slot()
+		.AutoHeight()
+		[
+			SNew(SHorizontalBox)
+			+SHorizontalBox::Slot()
+			.AutoWidth()
+			[
+				SNew(SSearchBox)
+				.HintText(LOCTEXT("PCGGraphSearchHint", "Enter text to find nodes..."))
+				.OnTextChanged(this, &SPCGEditorGraphProfilingView::OnSearchTextChanged)
+				.OnTextCommitted(this, &SPCGEditorGraphProfilingView::OnSearchTextCommitted)
+			]
+			+SHorizontalBox::Slot()
+			.AutoWidth()
+			.VAlign(VAlign_Center)
+			.Padding(1.0f, 0.0f)
+			[
+				SNew(STextBlock)
+				.Text(LOCTEXT("ExpandSubgraphDepth", "Expand Subgraph Depth"))
+			]
+			+SHorizontalBox::Slot()
+			.AutoWidth()
+			.VAlign(VAlign_Center)
+			.Padding(1.0f, 0.0f)
+			[
+				SNew(SSpinBox<int32>)
+				.Value(this, &SPCGEditorGraphProfilingView::GetSubgraphExpandDepth)
+				.OnValueChanged(this, &SPCGEditorGraphProfilingView::OnSubgraphExpandDepthChanged)
+				.MinValue(0)
+				.MaxValue(20)
+			]
+			+SHorizontalBox::Slot()
+			.Padding(FMargin(30.0f, 0.0f, 2.0f, 0.0f))
+			.VAlign(EVerticalAlignment::VAlign_Center)
+			.AutoWidth()
+			[
+				SNew(STextBlock)
+				.Text(LOCTEXT("TotalExecutionTime", "Total Time:"))
+			]
+			+SHorizontalBox::Slot()
+			.VAlign(EVerticalAlignment::VAlign_Center)
+			.AutoWidth()
+			[
+				SNew(STextBlock)
+				.Text(this, &SPCGEditorGraphProfilingView::GetTotalTimeLabel)
+				.MinDesiredWidth(50.0f)
+			]
+			+SHorizontalBox::Slot()
+			.Padding(FMargin(30.0f, 0.0f, 2.0f, 0.0f))
+			.VAlign(EVerticalAlignment::VAlign_Center)
+			.AutoWidth()
+			[
+				SNew(STextBlock)
+				.Text(LOCTEXT("TotalExecutionWallTime", "Total Wall Time:"))
+			]
+			+SHorizontalBox::Slot()
+			.VAlign(EVerticalAlignment::VAlign_Center)
+			.AutoWidth()
+			[
+				SNew(STextBlock)
+				.Text(this, &SPCGEditorGraphProfilingView::GetTotalWallTimeLabel)
+				.MinDesiredWidth(50.0f)
+			]
+			+SHorizontalBox::Slot()
+			.Padding(FMargin(30.0f, 0.0f, 2.0f, 0.0f))
+			.VAlign(EVerticalAlignment::VAlign_Center)
+			.AutoWidth()
+			[
+				SNew(STextBlock)
+				.Text(LOCTEXT("TotalMemoryInM", "Total data size (MB):"))
+			]
+			+SHorizontalBox::Slot()
+			.VAlign(EVerticalAlignment::VAlign_Center)
+			.AutoWidth()
+			[
+				SNew(STextBlock)
+				.Text(this, &SPCGEditorGraphProfilingView::GetTotalDataSizeInMB)
+				.MinDesiredWidth(50.0f)
+			]
+		]
+		+SVerticalBox::Slot()
+		.FillHeight(1.0f)
+		[
+			SNew(SHorizontalBox)
+			+SHorizontalBox::Slot()
+			.FillWidth(1.0f)
+			[
+				SNew(SScrollBox)
+				.Orientation(Orient_Horizontal)
+				.ExternalScrollbar(HorizontalScrollBar)
+				+SScrollBox::Slot()
+				[
+					ListView->AsShared()
+				]
+			]
+			+SHorizontalBox::Slot()
+			.AutoWidth()
+			[
+				VerticalScrollBar
+			]
+		]
+		+SVerticalBox::Slot()
+		.AutoHeight()
+		[
+			SNew(SHorizontalBox)
+			+SHorizontalBox::Slot()
+			[
+				HorizontalScrollBar
+			]
+		]
+	];
+
+	Refresh();
+}
+
+void SPCGEditorGraphProfilingView::Tick(const FGeometry& AllottedGeometry, const double InCurrentTime, const float InDeltaTime)
+{
+	SCompoundWidget::Tick(AllottedGeometry, InCurrentTime, InDeltaTime);
+
+	if (bNeedsRefresh)
+	{
+		bNeedsRefresh = false;
+		Refresh();
+	}
+}
+
+TSharedRef<SHeaderRow> SPCGEditorGraphProfilingView::CreateHeaderRowWidget()
+{
+	return SNew(SHeaderRow)
+		.ResizeMode(ESplitterResizeMode::FixedPosition)
+		.CanSelectGeneratedColumn(true)
+		+ SHeaderRow::Column(PCGEditorGraphProfilingView::NAME_Node)
+		.ManualWidth(300.0f)
+		.DefaultLabel(PCGEditorGraphProfilingView::TEXT_NodeLabel)
+		.HAlignHeader(HAlign_Center)
+		.HAlignCell(HAlign_Left)
+		.SortMode(this, &SPCGEditorGraphProfilingView::GetColumnSortMode, PCGEditorGraphProfilingView::NAME_Node)
+		.OnSort(this, &SPCGEditorGraphProfilingView::OnSortColumnHeader)
+		+ SHeaderRow::Column(PCGEditorGraphProfilingView::NAME_PrepareDataTime)
+		.ManualWidth(75.0f)
+		.DefaultLabel(PCGEditorGraphProfilingView::TEXT_PrepareDataTimeLabel)
+		.HAlignHeader(HAlign_Center)
+		.HAlignCell(HAlign_Right)
+		.SortMode(this, &SPCGEditorGraphProfilingView::GetColumnSortMode, PCGEditorGraphProfilingView::NAME_PrepareDataTime)
+		.OnSort(this, &SPCGEditorGraphProfilingView::OnSortColumnHeader)
+		.InitialSortMode(EColumnSortMode::Descending)
+		.DefaultTooltip(PCGEditorGraphProfilingView::TEXT_PrepareDataTimeTooltip)
+		+ SHeaderRow::Column(PCGEditorGraphProfilingView::NAME_PrepareDataWallTime)
+		.ManualWidth(110.0f)
+		.DefaultLabel(PCGEditorGraphProfilingView::TEXT_PrepareDataWallTimeLabel)
+		.HAlignHeader(HAlign_Center)
+		.HAlignCell(HAlign_Right)
+		.SortMode(this, &SPCGEditorGraphProfilingView::GetColumnSortMode, PCGEditorGraphProfilingView::NAME_PrepareDataWallTime)
+		.OnSort(this, &SPCGEditorGraphProfilingView::OnSortColumnHeader)
+		.InitialSortMode(EColumnSortMode::Descending)
+		.DefaultTooltip(PCGEditorGraphProfilingView::TEXT_PrepareDataWallTimeTooltip)
+		+ SHeaderRow::Column(PCGEditorGraphProfilingView::NAME_NbPrepareFrames)
+		.ManualWidth(90.0f)
+		.DefaultLabel(PCGEditorGraphProfilingView::TEXT_NbPrepareFramesLabel)
+		.HAlignHeader(HAlign_Center)
+		.HAlignCell(HAlign_Right)
+		.SortMode(this, &SPCGEditorGraphProfilingView::GetColumnSortMode, PCGEditorGraphProfilingView::NAME_NbPrepareFrames)
+		.OnSort(this, &SPCGEditorGraphProfilingView::OnSortColumnHeader)
+		.InitialSortMode(EColumnSortMode::Descending)
+		.DefaultTooltip(PCGEditorGraphProfilingView::TEXT_NbPrepareFramesTooltip)
+		+ SHeaderRow::Column(PCGEditorGraphProfilingView::NAME_NbExecutionFrames)
+		.ManualWidth(90.0f)
+		.DefaultLabel(PCGEditorGraphProfilingView::TEXT_NbExecutionFramesLabel)
+		.HAlignHeader(HAlign_Center)
+		.HAlignCell(HAlign_Right)
+		.SortMode(this, &SPCGEditorGraphProfilingView::GetColumnSortMode, PCGEditorGraphProfilingView::NAME_NbExecutionFrames)
+		.OnSort(this, &SPCGEditorGraphProfilingView::OnSortColumnHeader)
+		.InitialSortMode(EColumnSortMode::Descending)
+		.DefaultTooltip(PCGEditorGraphProfilingView::TEXT_NbExecutionFramesTooltip)
+		+ SHeaderRow::Column(PCGEditorGraphProfilingView::NAME_MinExecutionFrameTime)
+		.ManualWidth(110.0f)
+		.DefaultLabel(PCGEditorGraphProfilingView::TEXT_MinExecutionFrameTimeLabel)
+		.HAlignHeader(HAlign_Center)
+		.HAlignCell(HAlign_Right)
+		.SortMode(this, &SPCGEditorGraphProfilingView::GetColumnSortMode, PCGEditorGraphProfilingView::NAME_MinExecutionFrameTime)
+		.OnSort(this, &SPCGEditorGraphProfilingView::OnSortColumnHeader)
+		.InitialSortMode(EColumnSortMode::Descending)
+		.DefaultTooltip(PCGEditorGraphProfilingView::TEXT_MinExecutionFrameTimeTooltip)
+		+ SHeaderRow::Column(PCGEditorGraphProfilingView::NAME_MaxExecutionFrameTime)
+		.ManualWidth(110.0f)
+		.DefaultLabel(PCGEditorGraphProfilingView::TEXT_MaxExecutionFrameTimeLabel)
+		.HAlignHeader(HAlign_Center)
+		.HAlignCell(HAlign_Right)
+		.SortMode(this, &SPCGEditorGraphProfilingView::GetColumnSortMode, PCGEditorGraphProfilingView::NAME_MaxExecutionFrameTime)
+		.OnSort(this, &SPCGEditorGraphProfilingView::OnSortColumnHeader)
+		.InitialSortMode(EColumnSortMode::Descending)
+		.DefaultTooltip(PCGEditorGraphProfilingView::TEXT_MaxExecutionFrameTimeTooltip)
+		+ SHeaderRow::Column(PCGEditorGraphProfilingView::NAME_ExecutionTime)
+		.ManualWidth(75.0f)
+		.DefaultLabel(PCGEditorGraphProfilingView::TEXT_ExecutionTimeLabel)
+		.HAlignHeader(HAlign_Center)
+		.HAlignCell(HAlign_Right)
+		.SortMode(this, &SPCGEditorGraphProfilingView::GetColumnSortMode, PCGEditorGraphProfilingView::NAME_ExecutionTime)
+		.OnSort(this, &SPCGEditorGraphProfilingView::OnSortColumnHeader)
+		.InitialSortMode(EColumnSortMode::Descending)
+		.DefaultTooltip(PCGEditorGraphProfilingView::TEXT_ExecutionTimeTooltip)
+		+ SHeaderRow::Column(PCGEditorGraphProfilingView::NAME_ExecutionWallTime)
+		.ManualWidth(100.0f)
+		.DefaultLabel(PCGEditorGraphProfilingView::TEXT_ExecutionWallTimeLabel)
+		.HAlignHeader(HAlign_Center)
+		.HAlignCell(HAlign_Right)
+		.SortMode(this, &SPCGEditorGraphProfilingView::GetColumnSortMode, PCGEditorGraphProfilingView::NAME_ExecutionWallTime)
+		.OnSort(this, &SPCGEditorGraphProfilingView::OnSortColumnHeader)
+		.InitialSortMode(EColumnSortMode::Descending)
+		.DefaultTooltip(PCGEditorGraphProfilingView::TEXT_ExecutionWallTimeTooltip)
+		+ SHeaderRow::Column(PCGEditorGraphProfilingView::NAME_TotalTime)
+		.ManualWidth(75.0f)
+		.DefaultLabel(PCGEditorGraphProfilingView::TEXT_TotalTimeLabel)
+		.HAlignHeader(HAlign_Center)
+		.HAlignCell(HAlign_Right)
+		.SortMode(this, &SPCGEditorGraphProfilingView::GetColumnSortMode, PCGEditorGraphProfilingView::NAME_TotalTime)
+		.OnSort(this, &SPCGEditorGraphProfilingView::OnSortColumnHeader)
+		.InitialSortMode(EColumnSortMode::Descending)
+		.DefaultTooltip(PCGEditorGraphProfilingView::TEXT_TotalTimeTooltip)
+		+ SHeaderRow::Column(PCGEditorGraphProfilingView::NAME_TotalWallTime)
+		.ManualWidth(100.0f)
+		.DefaultLabel(PCGEditorGraphProfilingView::TEXT_TotalWallTimeLabel)
+		.HAlignHeader(HAlign_Center)
+		.HAlignCell(HAlign_Right)
+		.SortMode(this, &SPCGEditorGraphProfilingView::GetColumnSortMode, PCGEditorGraphProfilingView::NAME_TotalWallTime)
+		.OnSort(this, &SPCGEditorGraphProfilingView::OnSortColumnHeader)
+		.InitialSortMode(EColumnSortMode::Descending)
+		.DefaultTooltip(PCGEditorGraphProfilingView::TEXT_TotalWallTimeTooltip)
+		+ SHeaderRow::Column(PCGEditorGraphProfilingView::NAME_MemoryInMB)
+		.ManualWidth(100.0f)
+		.DefaultLabel(PCGEditorGraphProfilingView::TEXT_MemoryInMBLabel)
+		.HAlignHeader(HAlign_Center)
+		.HAlignCell(HAlign_Right)
+		.SortMode(this, &SPCGEditorGraphProfilingView::GetColumnSortMode, PCGEditorGraphProfilingView::NAME_MemoryInMB)
+		.OnSort(this, &SPCGEditorGraphProfilingView::OnSortColumnHeader)
+		.InitialSortMode(EColumnSortMode::Descending)
+		.DefaultTooltip(PCGEditorGraphProfilingView::TEXT_MemoryInMBTooltip);
+}
+
+FText SPCGEditorGraphProfilingView::GetTotalTimeLabel() const
+{
+	return FText::Format(LOCTEXT("GraphTotalTimeLabel", "{0} s"), TotalTime);
+}
+
+FText SPCGEditorGraphProfilingView::GetTotalWallTimeLabel() const
+{
+	return FText::Format(LOCTEXT("GraphTotalWallTimeLabel", "{0} s"), TotalWallTime);
+}
+
+FText SPCGEditorGraphProfilingView::GetTotalDataSizeInMB() const
+{
+	return FText::Format(LOCTEXT("GraphTotalDataSizeInMBLabel", "{0} MB"), TotalDataSizeInMB);
+}
+
+void SPCGEditorGraphProfilingView::OnSortColumnHeader(const EColumnSortPriority::Type SortPriority, const FName& ColumnId, const EColumnSortMode::Type NewSortMode)
+{
+	if (SortingColumn == ColumnId)
+	{
+		// Circling
+		SortMode = static_cast<EColumnSortMode::Type>((SortMode + 1) % 3);
+	}
+	else
+	{
+		SortingColumn = ColumnId;
+		SortMode = NewSortMode;
+	}
+
+	RequestRefresh();
+}
+
+void SPCGEditorGraphProfilingView::OnSubgraphExpandDepthChanged(int32 NewValue)
+{
+	ExpandSubgraphDepth = NewValue;
+	RequestRefresh();
+}
+
+FReply SPCGEditorGraphProfilingView::OnListViewKeyDown(const FGeometry& /*InGeometry*/, const FKeyEvent& InKeyEvent) const
+{
+	if (ListViewCommands->ProcessCommandBindings(InKeyEvent))
+	{
+		return FReply::Handled();
+	}
+
+	return FReply::Unhandled();
+}
+
+void SPCGEditorGraphProfilingView::CopySelectionToClipboard() const
+{
+	constexpr TCHAR Delimiter = TEXT(',');
+	constexpr TCHAR LineEnd = TEXT('\n');
+
+	TStringBuilder<2048> CSVExport;
+	TArray<FName> VisibleColumns;
+
+	// Write column header row
+	bool bHasWrittenAColumn = false;
+	const TIndirectArray<SHeaderRow::FColumn>& Columns = ListViewHeader->GetColumns();
+	for (auto Column : Columns)
+	{
+		if (!Column.bIsVisible)
+		{
+			continue;
+		}
+
+		VisibleColumns.Add(Column.ColumnId);
+		const FText& ColumnTitle = Column.DefaultText.Get();
+
+		if (bHasWrittenAColumn)
+		{
+			CSVExport += Delimiter;
+		}
+
+		CSVExport += ColumnTitle.ToString();
+		bHasWrittenAColumn = true;
+	}
+
+	// Gather selected rows and sort them to match the display order instead of selection order
+	TArray<PCGProfilingListViewItemPtr> SelectedListViewItems = ListView->GetSelectedItems();
+
+	TArray<int32> SelectedListViewItemsIndices;
+	Algo::Transform(SelectedListViewItems, SelectedListViewItemsIndices, [this](const PCGProfilingListViewItemPtr& InItem) { return ListViewItems.IndexOfByKey(InItem); });
+	Algo::Sort(SelectedListViewItemsIndices, [](const int32& A, const int32& B) { return A < B; });
+
+	SelectedListViewItems.Reset();
+	Algo::Transform(SelectedListViewItemsIndices, SelectedListViewItems, [this](const int32& InItemIndex) { return ListViewItems[InItemIndex]; });
+
+	// Write each row
+	for (const PCGProfilingListViewItemPtr& ListViewItem : SelectedListViewItems)
+	{
+		CSVExport += LineEnd;
+
+		for (int ColumnIndex = 0; ColumnIndex < VisibleColumns.Num(); ++ColumnIndex)
+		{
+			if (ColumnIndex > 0)
+			{
+				CSVExport += Delimiter;
+			}
+
+			CSVExport += ListViewItem->GetTextForColumn(VisibleColumns[ColumnIndex], /*bNoGrouping=*/true).ToString();
+		}
+	}
+
+	FPlatformApplicationMisc::ClipboardCopy(*CSVExport);
+}
+
+bool SPCGEditorGraphProfilingView::CanCopySelectionToClipboard() const
+{
+	return ListView->GetNumItemsSelected() > 0;
+}
+
+EColumnSortMode::Type SPCGEditorGraphProfilingView::GetColumnSortMode(const FName ColumnId) const
+{
+	if (SortingColumn != ColumnId)
+	{
+		return EColumnSortMode::None;
+	}
+
+	return SortMode;
+}
+
+namespace PCGEditorGraphProfilingView
+{
+	void AddListItems(
+		TArray<PCGProfilingListViewItemPtr>& OutListViewItems,
+		const TArray<PCGUtils::FCallTreeInfo>& TreeInfo,
+		const TMap<const UPCGNode*, UPCGEditorGraphNode*>& EditorNodeLookup,
+		int ExpandSubgraphDepth,
+		const FString& FolderName,
+		const FString& SearchString,
+		UPCGEditorGraphNode* CurrentEditorNode = nullptr)
+	{
+		for (const PCGUtils::FCallTreeInfo& Info : TreeInfo)
+		{
+			UPCGEditorGraphNode*const* EditorNodeItr = EditorNodeLookup.Find(Info.Node);
+			UPCGEditorGraphNode* EditorNode = EditorNodeItr ? *EditorNodeItr : CurrentEditorNode;
+
+			FString Fullname = FolderName;
+			if (!Info.Name.IsEmpty())
+			{
+				Fullname += Info.Name;
+			}
+			else if (Info.Node)
+			{
+				Fullname += Info.Node->GetNodeTitle(EPCGNodeTitleType::ListView).ToString();
+			}
+			else if (Info.LoopIndex != INDEX_NONE)
+			{
+				Fullname += FString::Printf(TEXT("Loop_%d"), Info.LoopIndex);
+			}
+
+			const bool bFilteredIn = SearchString.IsEmpty() || Fullname.Find(SearchString) != INDEX_NONE || (Info.Node && Info.Node->GetName().Find(SearchString) != INDEX_NONE);
+
+			if (bFilteredIn && (Info.Children.IsEmpty() || ExpandSubgraphDepth == 0))
+			{
+				PCGProfilingListViewItemPtr ListViewItem = MakeShared<FPCGProfilingListViewItem>();
+
+				ListViewItem->PCGNode = Info.Node;
+				ListViewItem->EditorNode = EditorNode;
+				ListViewItem->Name = Fullname;
+				ListViewItem->bHasData = (Info.CallTime.MaxExecutionFrameTime > 0);
+
+				ListViewItem->CallTime = Info.CallTime;
+
+				OutListViewItems.Add(ListViewItem);
+			}
+
+			if(ExpandSubgraphDepth > 0 && !Info.Children.IsEmpty())
+			{
+				AddListItems(OutListViewItems, Info.Children, EditorNodeLookup, ExpandSubgraphDepth - 1, Fullname + "/", SearchString, EditorNode);
+			}
+		}
+	}
+}
+
+void SPCGEditorGraphProfilingView::RequestRefresh()
+{
+	bNeedsRefresh = true;
+}
+
+FReply SPCGEditorGraphProfilingView::Refresh()
+{
+	TotalTime = 0;
+	TotalWallTime = 0;
+	TotalDataSizeInMB = 0;
+
+	ListViewItems.Empty();
+	ListView->RequestListRefresh();
+
+	const TSharedPtr<FPCGEditor> PCGEditor = PCGEditorPtr.Pin();
+	if (!PCGEditor.IsValid())
+	{
+		return FReply::Handled();
+	}
+
+	if (!PCGEditorGraph)
+	{
+		return FReply::Handled();
+	}
+
+	const UPCGComponent* Component = GetPCGComponent().Get();
+	if (!Component)
+	{
+		return FReply::Handled();
+	}
+
+	TArray<UPCGEditorGraphNode*> EditorNodes;
+	PCGEditorGraph->GetNodesOfClass<UPCGEditorGraphNode>(EditorNodes);
+
+	TMap<const UPCGNode*, UPCGEditorGraphNode*> EditorNodeLookup;
+	for (UPCGEditorGraphNode* EditorNode : EditorNodes)
+	{
+		EditorNodeLookup.Add(EditorNode->GetPCGNode(), EditorNode);
+	}
+
+	PCGUtils::FCallTreeInfo TreeInfo = Component->ExtraCapture.CalculateCallTreeInfo(Component, PCGStack);
+
+	ListViewItems.Reserve(TreeInfo.Children.Num());
+
+	if (TreeInfo.Children.Num() > 0)
+	{
+		TotalTime = TreeInfo.CallTime.TotalTime();
+		TotalWallTime = TreeInfo.CallTime.TotalWallTime();
+		TotalDataSizeInMB = TreeInfo.CallTime.OutputDataMemorySize / (1024.0 * 1024.0);
+	}
+
+	//TODO: could turn this into a tree instead of expanding into a list
+	PCGEditorGraphProfilingView::AddListItems(ListViewItems, TreeInfo.Children, EditorNodeLookup, ExpandSubgraphDepth, FString(), SearchValue);
+
+	if (SortingColumn != NAME_None && SortMode != EColumnSortMode::None)
+	{
+		Algo::Sort(ListViewItems, [this](const PCGProfilingListViewItemPtr& A, const PCGProfilingListViewItemPtr& B)
+			{
+				bool isLess = false;
+				if (SortingColumn == PCGEditorGraphProfilingView::NAME_Node)
+				{
+					isLess = A->Name < B->Name;
+				}
+				else if (SortingColumn == PCGEditorGraphProfilingView::NAME_PrepareDataTime)
+				{
+					isLess = A->CallTime.PrepareDataTime < B->CallTime.PrepareDataTime;
+				}
+				else if (SortingColumn == PCGEditorGraphProfilingView::NAME_PrepareDataWallTime)
+				{
+					isLess = A->CallTime.PrepareDataWallTime() < B->CallTime.PrepareDataWallTime();
+				}
+				else if (SortingColumn == PCGEditorGraphProfilingView::NAME_NbPrepareFrames)
+				{
+					isLess = A->CallTime.PrepareDataFrameCount < B->CallTime.PrepareDataFrameCount;
+				}
+				else if (SortingColumn == PCGEditorGraphProfilingView::NAME_MinExecutionFrameTime)
+				{
+					isLess = A->CallTime.MinExecutionFrameTime < B->CallTime.MinExecutionFrameTime;
+				}
+				else if (SortingColumn == PCGEditorGraphProfilingView::NAME_MaxExecutionFrameTime)
+				{
+					isLess = A->CallTime.MaxExecutionFrameTime < B->CallTime.MaxExecutionFrameTime;
+				}
+				else if (SortingColumn == PCGEditorGraphProfilingView::NAME_ExecutionTime)
+				{
+					isLess = A->CallTime.ExecutionTime < B->CallTime.ExecutionTime;
+				}
+				else if (SortingColumn == PCGEditorGraphProfilingView::NAME_ExecutionWallTime)
+				{
+					isLess = A->CallTime.ExecutionWallTime() < B->CallTime.ExecutionWallTime();
+				}
+				else if (SortingColumn == PCGEditorGraphProfilingView::NAME_NbExecutionFrames)
+				{
+					isLess = A->CallTime.ExecutionFrameCount < B->CallTime.ExecutionFrameCount;
+				}
+				else if (SortingColumn == PCGEditorGraphProfilingView::NAME_TotalTime)
+				{
+					isLess = A->CallTime.TotalTime() < B->CallTime.TotalTime();
+				}
+				else if (SortingColumn == PCGEditorGraphProfilingView::NAME_TotalWallTime)
+				{
+					isLess = A->CallTime.TotalWallTime() < B->CallTime.TotalWallTime();
+				}
+				else if (SortingColumn == PCGEditorGraphProfilingView::NAME_MemoryInMB)
+				{
+					isLess = A->CallTime.OutputDataMemorySize < B->CallTime.OutputDataMemorySize;
+				}
+
+				return SortMode == EColumnSortMode::Ascending ? isLess : !isLess;
+			});
+	}
+
+	ListView->SetItemsSource(&ListViewItems);
+
+	return FReply::Handled();
+}
+
+void SPCGEditorGraphProfilingView::OnDebugStackChanged(const FPCGStack& InPCGStack)
+{
+	PCGStack = InPCGStack;
+	RequestRefresh();
+}
+
+TWeakObjectPtr<UPCGComponent> SPCGEditorGraphProfilingView::GetPCGComponent() const
+{
+	return PCGEditorPtr.IsValid() ? PCGEditorPtr.Pin()->GetPCGComponentBeingInspected() : nullptr;
+}
+
+TSharedRef<ITableRow> SPCGEditorGraphProfilingView::OnGenerateRow(PCGProfilingListViewItemPtr Item, const TSharedRef<STableViewBase>& OwnerTable) const
+{
+	return SNew(SPCGProfilingListViewItemRow, OwnerTable, Item);
+}
+
+void SPCGEditorGraphProfilingView::OnSearchTextChanged(const FText& InText)
+{
+	SearchValue = InText.ToString();
+	Refresh();
+}
+
+void SPCGEditorGraphProfilingView::OnSearchTextCommitted(const FText& InText, ETextCommit::Type InCommitType)
+{
+	OnSearchTextChanged(InText);
+}
+
+#undef LOCTEXT_NAMESPACE
diff --git a/Engine/Plugins/PCG/Source/PCGEditor/Private/Widgets/SPCGEditorGraphProfilingView.h b/Engine/Plugins/PCG/Source/PCGEditor/Private/Widgets/SPCGEditorGraphProfilingView.h
new file mode 100644
--- /dev/null
+++ b/Engine/Plugins/PCG/Source/PCGEditor/Private/Widgets/SPCGEditorGraphProfilingView.h
@@ -0,0 +1,122 @@
+// Copyright Epic Games, Inc. All Rights Reserved.
+
+#pragma once
+
+#include "Graph/PCGStackContext.h"
+#include "Utils/PCGExtraCapture.h"
+
+#include "Framework/Commands/GenericCommands.h"
+#include "Widgets/Views/ITableRow.h"
+#include "Widgets/Views/SListView.h"
+#include "Widgets/Views/STableRow.h"
+
+class FPCGEditor;
+class STableViewBase;
+class UPCGComponent;
+class UPCGEditorGraphNode;
+class UPCGEditorGraph;
+class UPCGNode;
+
+struct FPCGProfilingListViewItem
+{
+	FText GetTextForColumn(FName ColumnId, bool bNoGrouping) const;
+
+	const UPCGNode* PCGNode = nullptr;
+	const UPCGEditorGraphNode* EditorNode = nullptr;
+
+	FString Name;
+
+	PCGUtils::FCallTime CallTime;
+
+	bool bHasData = false;
+};
+
+typedef TSharedPtr<FPCGProfilingListViewItem> PCGProfilingListViewItemPtr;
+
+class SPCGProfilingListViewItemRow : public SMultiColumnTableRow<PCGProfilingListViewItemPtr>
+{
+public:
+	SLATE_BEGIN_ARGS(SPCGProfilingListViewItemRow) {}
+	SLATE_END_ARGS()
+
+	void Construct(const FArguments& InArgs, const TSharedRef<STableViewBase>& InOwnerTableView, const PCGProfilingListViewItemPtr& Item);
+
+	virtual TSharedRef<SWidget> GenerateWidgetForColumn(const FName& ColumnId) override;
+
+protected:
+	PCGProfilingListViewItemPtr InternalItem;
+};
+
+class SPCGEditorGraphProfilingView : public SCompoundWidget
+{
+public:
+	SLATE_BEGIN_ARGS(SPCGEditorGraphProfilingView) { }
+	SLATE_END_ARGS()
+
+	~SPCGEditorGraphProfilingView();
+
+	void Construct(const FArguments& InArgs, TSharedPtr<FPCGEditor> InPCGEditor);
+	virtual void Tick(const FGeometry& AllottedGeometry, const double InCurrentTime, const float InDeltaTime) override;
+
+private:
+	TSharedRef<SHeaderRow> CreateHeaderRowWidget();
+
+	int32 GetSubgraphExpandDepth() const { return ExpandSubgraphDepth; }
+	void OnSubgraphExpandDepthChanged(int32 NewValue);
+
+	void OnDebugStackChanged(const FPCGStack& InPCGStack);
+
+	// Callbacks
+	void RequestRefresh();
+	FReply Refresh();
+	TSharedRef<ITableRow> OnGenerateRow(PCGProfilingListViewItemPtr Item, const TSharedRef<STableViewBase>& OwnerTable) const;
+	void OnItemDoubleClicked(PCGProfilingListViewItemPtr Item);
+	void OnSortColumnHeader(const EColumnSortPriority::Type SortPriority, const FName& ColumnId, const EColumnSortMode::Type NewSortMode);
+	EColumnSortMode::Type GetColumnSortMode(const FName ColumnId) const;
+	FText GetTotalTimeLabel() const;
+	FText GetTotalWallTimeLabel() const;
+	FText GetTotalDataSizeInMB() const;
+
+	FReply OnListViewKeyDown(const FGeometry& InGeometry, const FKeyEvent& InKeyEvent) const;
+	void CopySelectionToClipboard() const;
+	bool CanCopySelectionToClipboard() const;
+
+	/** Called when user changes the text they are searching for */
+	void OnSearchTextChanged(const FText& InText);
+
+	/** Called when user changes commits text to the search box */
+	void OnSearchTextCommitted(const FText& InText, ETextCommit::Type InCommitType);
+
+	/** Returns the current PCG component */
+	TWeakObjectPtr<UPCGComponent> GetPCGComponent() const;
+
+	/** Pointer back to the PCG editor that owns us */
+	TWeakPtr<FPCGEditor> PCGEditorPtr;
+
+	/** Cached PCGGraph being viewed */
+	UPCGEditorGraph* PCGEditorGraph = nullptr;
+
+	/** Current stack being viewed */
+	FPCGStack PCGStack;
+
+	TSharedPtr<SHeaderRow> ListViewHeader;
+	TSharedPtr<SListView<PCGProfilingListViewItemPtr>> ListView;
+	TArray<PCGProfilingListViewItemPtr> ListViewItems;
+	TSharedPtr<FUICommandList> ListViewCommands;
+
+	// To allow sorting
+	FName SortingColumn = NAME_None;
+	EColumnSortMode::Type SortMode = EColumnSortMode::Type::None;
+
+	double TotalTime = 0.0;
+	double TotalWallTime = 0.0;
+	double TotalDataSizeInMB = 0.0;
+
+	bool bNeedsRefresh = false;
+
+	/** Currently depth of entry exposition - if 0, only the nodes in the currently debugged graph will be shown. */
+	int32 ExpandSubgraphDepth = 0;
+
+	/** The string to search for */
+	FString SearchValue;
+};
