Added default value providers for load asset node

diff --git a/Engine/Plugins/PCG/Source/PCG/Private/Data/PCGBasePointData.cpp b/Engine/Plugins/PCG/Source/PCG/Private/Data/PCGBasePointData.cpp
--- a/Engine/Plugins/PCG/Source/PCG/Private/Data/PCGBasePointData.cpp
+++ b/Engine/Plugins/PCG/Source/PCG/Private/Data/PCGBasePointData.cpp
@@ -1048,7 +1048,7 @@ void UPCGBasePointData::AddSinglePointFromActor(AActor* InActor, bool* bOutOptio
 	for (FName Tag : InActor->Tags)
 	{
 		PCG::Private::FParseTagResult TagData(Tag);
-		if (PCG::Private::SetAttributeFromTag(TagData, Metadata, PointMetadataEntry, /*bCanCreateAttribute=*/true))
+		if (PCG::Private::SetAttributeFromTag(TagData, Metadata, PointMetadataEntry, PCG::Private::ESetAttributeFromTagFlags::CreateAttribute))
 		{
 			bSanitizedAttributeNames |= TagData.HasBeenSanitized();
 		}
diff --git a/Engine/Plugins/PCG/Source/PCG/Private/Elements/IO/PCGLoadAssetElement.cpp b/Engine/Plugins/PCG/Source/PCG/Private/Elements/IO/PCGLoadAssetElement.cpp
--- a/Engine/Plugins/PCG/Source/PCG/Private/Elements/IO/PCGLoadAssetElement.cpp
+++ b/Engine/Plugins/PCG/Source/PCG/Private/Elements/IO/PCGLoadAssetElement.cpp
@@ -3,10 +3,23 @@
 #include "Elements/IO/PCGLoadAssetElement.h"
 
 #include "PCGModule.h"
+#include "PCGData.h"
+#include "PCGParamData.h"
 #include "Helpers/PCGDynamicTrackingHelpers.h"
+#include "Helpers/PCGHelpers.h"
+#include "Metadata/Accessors/PCGAttributeAccessor.h"
+#include "Metadata/Accessors/PCGAttributeAccessorHelpers.h"
 
+#include "Algo/Transform.h"
 #include "AssetRegistry/AssetData.h"
 
+#define LOCTEXT_NAMESPACE "PCGLoadDataAssetElement"
+
+namespace PCGLoadDataAsset
+{
+	const FName DefaultProviderPinLabel = TEXT("DefaultAttributeOverridesIn");
+}
+
 UPCGLoadDataAssetSettings::UPCGLoadDataAssetSettings()
 {
 	Pins = Super::OutputPinProperties();
@@ -40,7 +53,8 @@ EPCGChangeType UPCGLoadDataAssetSettings::GetChangeTypeForProperty(const FName&
 {
 	EPCGChangeType ChangeType = Super::GetChangeTypeForProperty(InPropertyName) | EPCGChangeType::Cosmetic;
 
-	if (InPropertyName == GET_MEMBER_NAME_CHECKED(UPCGLoadDataAssetSettings, bLoadFromInput))
+	if (InPropertyName == GET_MEMBER_NAME_CHECKED(UPCGLoadDataAssetSettings, bLoadFromInput) ||
+		InPropertyName == GET_MEMBER_NAME_CHECKED(UPCGLoadDataAssetSettings, bSetDefaultAttributeOverridesFromInput))
 	{
 		ChangeType |= EPCGChangeType::Structural;
 	}
@@ -63,6 +77,12 @@ TArray<FPCGPinProperties> UPCGLoadDataAssetSettings::InputPinProperties() const
 		FPCGPinProperties& InputPin = PinProperties.Emplace_GetRef(PCGPinConstants::DefaultInputLabel, EPCGDataType::Param);
 		InputPin.SetRequiredPin();
 	}
+
+	if (bSetDefaultAttributeOverridesFromInput)
+	{
+		FPCGPinProperties& DefaultsPin = PinProperties.Emplace_GetRef(PCGLoadDataAsset::DefaultProviderPinLabel, EPCGDataType::Param);
+		DefaultsPin.SetNormalPin();
+	}
 	
 	return PinProperties;
 }
@@ -166,6 +186,46 @@ bool FPCGLoadDataAssetElement::PrepareDataInternal(FPCGContext* InContext) const
 	const UPCGLoadDataAssetSettings* Settings = Context->GetInputSettings<UPCGLoadDataAssetSettings>();
 	check(Settings);
 
+	// Additional validation when we're asking to set default from the default pin -
+	// Basically we want that the number of data on the pin is either 0, 1 or the same cardinality as the actual input.
+	if (Settings->bSetDefaultAttributeOverridesFromInput)
+	{
+		const TArray<FPCGTaggedData> Inputs = Context->InputData.GetInputsByPin(PCGPinConstants::DefaultInputLabel);
+		const TArray<FPCGTaggedData> DefaultProviders = Context->InputData.GetInputsByPin(PCGLoadDataAsset::DefaultProviderPinLabel);
+
+		if (DefaultProviders.Num() != 0 && DefaultProviders.Num() != 1 && DefaultProviders.Num() != Inputs.Num())
+		{
+			PCGLog::InputOutput::LogInvalidCardinalityError(PCGPinConstants::DefaultInputLabel, PCGLoadDataAsset::DefaultProviderPinLabel, Context);
+			return true;
+		}
+
+		if (!DefaultProviders.IsEmpty())
+		{
+			Algo::Transform(DefaultProviders, Context->DefaultProviders, [](const FPCGTaggedData& TaggedData) { return Cast<UPCGParamData>(TaggedData.Data); });
+		}
+	}
+	else if(!Settings->CommaSeparatedDefaultAttributeOverrides.IsEmpty() || !Settings->DefaultAttributeOverrides.IsEmpty())
+	{
+		TArray<FString> OverriddenDefaultValues;
+		const TArray<FString>* DefaultValues = nullptr;
+
+		if (!Settings->CommaSeparatedDefaultAttributeOverrides.IsEmpty())
+		{
+			OverriddenDefaultValues = PCGHelpers::GetStringArrayFromCommaSeparatedList(Settings->CommaSeparatedDefaultAttributeOverrides);
+			DefaultValues = &OverriddenDefaultValues;
+		}
+		else
+		{
+			DefaultValues = &Settings->DefaultAttributeOverrides;
+		}
+
+		check(DefaultValues);
+		Context->DefaultValueTags.Reserve(DefaultValues->Num());
+		Algo::Transform(*DefaultValues, Context->DefaultValueTags, [](const FString& InTagValue) { return PCG::Private::FParseTagResult(InTagValue); });
+	}
+
+	Context->bDefaultsMatchInput = true;
+
 	return Context->InitializeAndRequestLoad(PCGPinConstants::DefaultInputLabel,
 		Settings->AssetReferenceSelector,
 		{ Settings->Asset.ToSoftObjectPath() },
@@ -183,6 +243,12 @@ bool FPCGLoadDataAssetElement::ExecuteInternal(FPCGContext* InContext) const
 	const UPCGLoadDataAssetSettings* Settings = Context->GetInputSettings<UPCGLoadDataAssetSettings>();
 	check(Settings);
 
+	// Early out if we have a data matching issue
+	if (!Context->bDefaultsMatchInput)
+	{
+		return true;
+	}
+
 #if WITH_EDITOR
 	FPCGDynamicTrackingHelper DynamicTracking;
 	const bool bRequiresDynamicTracking = Context->IsValueOverriden(GET_MEMBER_NAME_CHECKED(UPCGLoadDataAssetSettings, Asset)) || (Settings->bLoadFromInput && !Context->PathsToObjectsAndDataIndex.IsEmpty());
@@ -202,6 +268,98 @@ bool FPCGLoadDataAssetElement::ExecuteInternal(FPCGContext* InContext) const
 			const int TaggedDataOffset = Context->OutputData.TaggedData.Num();
 			Context->OutputData.TaggedData.Append(AssetData->Data.TaggedData);
 
+			// Apply defaults if needed
+			if (!Context->DefaultProviders.IsEmpty())
+			{
+				// Get matching param data based on input index (AssetPath.Get<1>())
+				const int32 InputIndex = AssetPath.Get<1>();
+				TObjectPtr<const UPCGParamData> DefaultProvider = Context->DefaultProviders[InputIndex % Context->DefaultProviders.Num()];
+				const UPCGMetadata* DefaultMetadata = DefaultProvider ? DefaultProvider->ConstMetadata() : nullptr;
+
+				// Skip applying defaults if the provider isn't valid or has no entries.
+				if (DefaultProvider && DefaultMetadata && DefaultMetadata->GetItemCountForChild() > 0)
+				{
+					TArray<FPCGAttributeIdentifier> DefaultAttributesNames;
+					TArray<EPCGMetadataTypes> DefaultAttributesTypes;
+					DefaultMetadata->GetAllAttributes(DefaultAttributesNames, DefaultAttributesTypes);
+
+					for (int TaggedDataIndex = TaggedDataOffset; TaggedDataIndex < Context->OutputData.TaggedData.Num(); ++TaggedDataIndex)
+					{
+						const UPCGData* OriginalData = Context->OutputData.TaggedData[TaggedDataIndex].Data.Get();
+						if (!OriginalData)
+						{
+							continue;
+						}
+
+						UPCGData* DuplicateData = OriginalData->DuplicateData(Context);
+						UPCGMetadata* Metadata = DuplicateData->MutableMetadata();
+
+						Context->OutputData.TaggedData[TaggedDataIndex].Data = DuplicateData;
+
+						for (const FPCGAttributeIdentifier& DefaultAttributeName : DefaultAttributesNames)
+						{
+							const FPCGMetadataAttributeBase* DefaultAttribute = DefaultMetadata->GetConstAttribute(DefaultAttributeName);
+							check(DefaultAttribute);
+
+							FPCGMetadataAttributeBase* Attribute = Metadata->GetMutableAttribute(DefaultAttributeName);
+							if (!Attribute)
+							{
+								// This already sets up the default value, so we will have nothing else to do.
+								Metadata->AddAttribute(DefaultMetadata, DefaultAttributeName);
+								Attribute = Metadata->GetMutableAttribute(DefaultAttributeName);
+							}
+							else
+							{
+								if (!PCG::Private::IsBroadcastableOrConstructible(DefaultAttribute->GetTypeId(), Attribute->GetTypeId()))
+								{
+									PCGLog::LogWarningOnGraph(FText::Format(LOCTEXT("DefaultTypeDoesNotMatchAssetType", "Default value '{0}' does not have a compatible type ('{1}') to its original type ('{2}') in the asset. Will be skipped."), FText::FromName(DefaultAttributeName.Name), PCG::Private::GetTypeNameText(DefaultAttribute->GetTypeId()), PCG::Private::GetTypeNameText(Attribute->GetTypeId())), Context);
+									continue;
+								}
+
+								// Create accessor on mutable metadata for the current attribute
+								TUniquePtr<IPCGAttributeAccessor> Accessor = PCGAttributeAccessorHelpers::CreateAccessor(Attribute, Metadata);
+								auto SetDefaultValueOperation = [&Accessor, DefaultMetadata, DefaultAttribute](auto Dummy)
+								{
+									using AttributeType = decltype(Dummy);
+									const FPCGMetadataAttribute<AttributeType>* TypedDefaultAttribute = static_cast<const FPCGMetadataAttribute<AttributeType>*>(DefaultAttribute);
+									const AttributeType DefaultValue = TypedDefaultAttribute->GetValueFromItemKey(0); // first key should always be zero.
+									FPCGAttributeAccessorKeysEntries DefaultEntry(PCGInvalidEntryKey);
+
+									Accessor->Set(DefaultValue, DefaultEntry, EPCGAttributeAccessorFlags::AllowSetDefaultValue | EPCGAttributeAccessorFlags::AllowBroadcastAndConstructible);
+								};
+
+								PCGMetadataAttribute::CallbackWithRightType(DefaultAttribute->GetTypeId(), SetDefaultValueOperation);
+							}
+						}
+					}
+				}
+			}
+			// Otherwise, normal Tag:Value case
+			else if (!Context->DefaultValueTags.IsEmpty())
+			{
+				for (int TaggedDataIndex = TaggedDataOffset; TaggedDataIndex < Context->OutputData.TaggedData.Num(); ++TaggedDataIndex)
+				{
+					const UPCGData* OriginalData = Context->OutputData.TaggedData[TaggedDataIndex].Data.Get();
+					if (!OriginalData)
+					{
+						continue;
+					}
+
+					UPCGData* DuplicateData = OriginalData->DuplicateData(Context);
+					UPCGMetadata* Metadata = DuplicateData->MutableMetadata();
+
+					Context->OutputData.TaggedData[TaggedDataIndex].Data = DuplicateData;
+
+					for (const PCG::Private::FParseTagResult& TagData : Context->DefaultValueTags)
+					{
+						if (!PCG::Private::SetAttributeFromTag(TagData, Metadata, PCGInvalidEntryKey, PCG::Private::ESetAttributeFromTagFlags::CreateAttribute | PCG::Private::ESetAttributeFromTagFlags::SetDefaultValue))
+						{
+							PCGLog::LogWarningOnGraph(FText::Format(LOCTEXT("TagFailedToSetDefaultValue", "Default Tag Value '{0}' failed to set its value to the target asset data, most likely due to type mismatch."), FText::FromString(TagData.GetOriginalAttribute())), Context);
+						}
+					}
+				}
+			}
+
 			if (Settings->bTagOutputsBasedOnOutputPins || Settings->InputIndexTag != NAME_None || Settings->DataIndexTag != NAME_None)
 			{
 				for (int TaggedDataIndex = TaggedDataOffset; TaggedDataIndex < Context->OutputData.TaggedData.Num(); ++TaggedDataIndex)
@@ -242,4 +400,6 @@ bool FPCGLoadDataAssetElement::ExecuteInternal(FPCGContext* InContext) const
 #endif
 
 	return true;
-}
\ No newline at end of file
+}
+
+#undef LOCTEXT_NAMESPACE
\ No newline at end of file
diff --git a/Engine/Plugins/PCG/Source/PCG/Private/Elements/PCGConvertToAttributeSet.cpp b/Engine/Plugins/PCG/Source/PCG/Private/Elements/PCGConvertToAttributeSet.cpp
--- a/Engine/Plugins/PCG/Source/PCG/Private/Elements/PCGConvertToAttributeSet.cpp
+++ b/Engine/Plugins/PCG/Source/PCG/Private/Elements/PCGConvertToAttributeSet.cpp
@@ -124,7 +124,7 @@ bool FPCGTagsToAttributeSetElement::ExecuteInternal(FPCGContext* Context) const
 
 		for (const FString& DataTag : Input.Tags)
 		{
-			PCG::Private::SetAttributeFromTag(DataTag, ParamData->Metadata, EntryKey, /*bCanCreateAttribute=*/true);
+			PCG::Private::SetAttributeFromTag(DataTag, ParamData->Metadata, EntryKey, PCG::Private::ESetAttributeFromTagFlags::CreateAttribute);
 		}
 	}
 
diff --git a/Engine/Plugins/PCG/Source/PCG/Private/Helpers/PCGTagHelpers.cpp b/Engine/Plugins/PCG/Source/PCG/Private/Helpers/PCGTagHelpers.cpp
--- a/Engine/Plugins/PCG/Source/PCG/Private/Helpers/PCGTagHelpers.cpp
+++ b/Engine/Plugins/PCG/Source/PCG/Private/Helpers/PCGTagHelpers.cpp
@@ -37,11 +37,11 @@ namespace Private
 				{
 					Result.NumericValue = FCString::Atod(*RightSide);
 				}
-				else if (RightSide.Equals(TEXT("True")))
+				else if (RightSide.Equals(TEXT("True"), ESearchCase::IgnoreCase))
 				{
 					Result.BooleanValue = true;
 				}
-				else if (RightSide.Equals(TEXT("False")))
+				else if (RightSide.Equals(TEXT("False"), ESearchCase::IgnoreCase))
 				{
 					Result.BooleanValue = false;
 				}
@@ -85,13 +85,13 @@ namespace Private
 
 	bool CreateAttributeFromTag(const FParseTagResult& InTagData, UPCGMetadata* InMetadata)
 	{
-		return SetAttributeFromTag(InTagData, InMetadata, PCGInvalidEntryKey, /*bCanCreateAttribute=*/true);
+		return SetAttributeFromTag(InTagData, InMetadata, PCGInvalidEntryKey, ESetAttributeFromTagFlags::OverwriteAttributeIfDifferentType);
 	}
 
-	bool SetAttributeFromTag(const FString& InTag, UPCGMetadata* InMetadata, PCGMetadataEntryKey InEntryKey, bool bInCanCreateAttribute, FParseTagResult* OutResult)
+	bool SetAttributeFromTag(const FString& InTag, UPCGMetadata* InMetadata, PCGMetadataEntryKey InEntryKey, ESetAttributeFromTagFlags Flags, FParseTagResult* OutResult)
 	{
 		FParseTagResult TagData = ParseTag(InTag);
-		const bool bSetSuccess = SetAttributeFromTag(TagData, InMetadata, InEntryKey, bInCanCreateAttribute);
+		const bool bSetSuccess = SetAttributeFromTag(TagData, InMetadata, InEntryKey, Flags);
 
 		if (OutResult)
 		{
@@ -101,7 +101,7 @@ namespace Private
 		return bSetSuccess;
 	}
 
-	bool SetAttributeFromTag(const FParseTagResult& TagData, UPCGMetadata* InMetadata, PCGMetadataEntryKey InEntryKey, bool bInCanCreateAttribute)
+	bool SetAttributeFromTag(const FParseTagResult& TagData, UPCGMetadata* InMetadata, PCGMetadataEntryKey InEntryKey, ESetAttributeFromTagFlags Flags)
 	{
 		check(InMetadata);
 
@@ -109,56 +109,72 @@ namespace Private
 		{
 			FName AttributeName(TagData.Attribute);
 
-			auto GetAttribute = [InMetadata, bInCanCreateAttribute, &AttributeName](auto DefaultValue)
+			auto SetValue = [InMetadata, &AttributeName, InEntryKey, Flags](auto DefaultValue, auto Value)
 			{
-				using AttributeType = decltype(DefaultValue);
-				if (bInCanCreateAttribute)
+				using AttributeType = decltype(Value);
+
+				const bool bCanCreateAttribute = !!((Flags & ESetAttributeFromTagFlags::CreateAttribute) | (Flags & ESetAttributeFromTagFlags::OverwriteAttributeIfDifferentType));
+				const bool bCanOverwriteAttribute = !!(Flags & ESetAttributeFromTagFlags::OverwriteAttributeIfDifferentType);
+
+				FPCGMetadataAttribute<AttributeType>* Attribute = nullptr;
+
+				if (bCanCreateAttribute)
 				{
-					return InMetadata->FindOrCreateAttribute<AttributeType>(AttributeName, DefaultValue, /*bAllowsInterpolation=*/false, /*bOverrideParent=*/false, /*bOverwriteIfTypeMismatch=*/true);
+					Attribute = InMetadata->FindOrCreateAttribute<AttributeType>(AttributeName, DefaultValue, /*bAllowsInterpolation=*/false, /*bOverrideParent=*/false, /*bOverwriteIfTypeMismatch=*/bCanOverwriteAttribute);
 				}
 				else
 				{
-					return InMetadata->GetMutableTypedAttribute<AttributeType>(AttributeName);
+					Attribute = InMetadata->GetMutableTypedAttribute<AttributeType>(AttributeName);
+				}
+
+				if (Attribute)
+				{
+					if (InEntryKey != PCGInvalidEntryKey)
+					{
+						Attribute->SetValue(InEntryKey, Value);
+						return true;
+					}
+					else if (!!(Flags & ESetAttributeFromTagFlags::SetDefaultValue))
+					{
+						Attribute->SetDefaultValue(Value);
+						return true;
+					}
 				}
+
+				return false;
 			};
 
 			if (TagData.HasNumericValue())
 			{
-				FPCGMetadataAttribute<double>* Attribute = GetAttribute(0.0);
-				if (Attribute && InEntryKey != PCGInvalidEntryKey)
-				{
-					Attribute->SetValue(InEntryKey, TagData.NumericValue.GetValue());
-				}
+				return SetValue(0.0, TagData.NumericValue.GetValue());
 			}
 			else if (TagData.HasBooleanValue())
 			{
-				FPCGMetadataAttribute<bool>* Attribute = GetAttribute(false);
-				if (Attribute && InEntryKey != PCGInvalidEntryKey)
-				{
-					Attribute->SetValue(InEntryKey, TagData.BooleanValue.GetValue());
-				}
+				return SetValue(false, TagData.BooleanValue.GetValue());
 			}
 			else if (TagData.HasValue())
 			{
-				FPCGMetadataAttribute<FString>* Attribute = GetAttribute(FString());
-				if (Attribute && InEntryKey != PCGInvalidEntryKey)
-				{
-					Attribute->SetValue(InEntryKey, TagData.Value.GetValue());
-				}
+				return SetValue(FString(), TagData.Value.GetValue());
 			}
-			else
+			else // Default boolean tag, assumes default is false.
 			{
-				FPCGMetadataAttribute<bool>* Attribute = GetAttribute(false);
-				if (Attribute && InEntryKey != PCGInvalidEntryKey)
-				{
-					Attribute->SetValue(InEntryKey, true);
-				}
+				return SetValue(false, true);
 			}
-
-			return true;
 		}
 		
 		return false;
 	}
+
+	// DEPRECATED 5.6
+	bool SetAttributeFromTag(const FString& InTag, UPCGMetadata* InMetadata, PCGMetadataEntryKey InEntryKey, bool bInCanCreateAttribute, FParseTagResult* OutResult)
+	{
+		return SetAttributeFromTag(InTag, InMetadata, InEntryKey, bInCanCreateAttribute ? ESetAttributeFromTagFlags::OverwriteAttributeIfDifferentType : ESetAttributeFromTagFlags::None, OutResult);
+	}
+
+	// DEPRECATED 5.6
+	bool SetAttributeFromTag(const FParseTagResult& TagData, UPCGMetadata* InMetadata, PCGMetadataEntryKey InEntryKey, bool bInCanCreateAttribute)
+	{
+		return SetAttributeFromTag(TagData, InMetadata, InEntryKey, bInCanCreateAttribute ? ESetAttributeFromTagFlags::OverwriteAttributeIfDifferentType : ESetAttributeFromTagFlags::None);
+	}
 }
 }
\ No newline at end of file
diff --git a/Engine/Plugins/PCG/Source/PCG/Public/Elements/IO/PCGLoadAssetElement.h b/Engine/Plugins/PCG/Source/PCG/Public/Elements/IO/PCGLoadAssetElement.h
--- a/Engine/Plugins/PCG/Source/PCG/Public/Elements/IO/PCGLoadAssetElement.h
+++ b/Engine/Plugins/PCG/Source/PCG/Public/Elements/IO/PCGLoadAssetElement.h
@@ -6,6 +6,7 @@
 #include "PCGDataAsset.h"
 #include "PCGSettings.h"
 #include "Elements/PCGLoadObjectsContext.h"
+#include "Helpers/PCGTagHelpers.h"
 
 struct FAssetData;
 
@@ -89,6 +90,17 @@ public:
 	UPROPERTY(BlueprintReadWrite, EditAnywhere, Category = Settings, meta = (EditCondition = "bLoadFromInput", EditConditionHides))
 	FName DataIndexTag = NAME_None;
 
+	// Exposes an attribute set pin to override defaults of the loaded data assets.
+	UPROPERTY(BlueprintReadWrite, EditAnywhere, Category = Settings)
+	bool bSetDefaultAttributeOverridesFromInput = false;
+
+	// List of Tag:Value default value overrides to apply on the loaded data assets.
+	UPROPERTY(BlueprintReadWrite, EditAnywhere, Category = Settings, meta = (Tooltip="Tag:Value default overrides to be applied to the loaded data.", EditCondition = "!bSetDefaultOverridesFromInput", EditConditionHides))
+	TArray<FString> DefaultAttributeOverrides;
+
+	UPROPERTY(meta = (Tooltip="Overridable-only value to set multiple tags from a single string. If provided, the entries in the Default Value Overrides will be ignored.", PCG_Overridable))
+	FString CommaSeparatedDefaultAttributeOverrides;
+
 	UPROPERTY(BlueprintReadWrite, EditAnywhere, Category = Settings, meta = (Tooltip="Warns if asset is null or couldn't be loaded"))
 	bool bWarnIfNoAsset = true;
 
@@ -101,7 +113,14 @@ public:
 	bool bSynchronousLoad = false;
 };
 
-struct FPCGLoadDataAssetContext : public FPCGLoadObjectsFromPathContext {};
+struct FPCGLoadDataAssetContext : public FPCGLoadObjectsFromPathContext 
+{
+	TArray<TObjectPtr<const UPCGParamData>> DefaultProviders;
+	TArray<PCG::Private::FParseTagResult> DefaultValueTags;
+
+	bool bDefaultsMatchInput = false;
+	bool bShouldApplyDefaults = false;
+};
 
 class PCG_API FPCGLoadDataAssetElement : public IPCGElementWithCustomContext<FPCGLoadDataAssetContext>
 {
diff --git a/Engine/Plugins/PCG/Source/PCG/Public/Helpers/PCGTagHelpers.h b/Engine/Plugins/PCG/Source/PCG/Public/Helpers/PCGTagHelpers.h
--- a/Engine/Plugins/PCG/Source/PCG/Public/Helpers/PCGTagHelpers.h
+++ b/Engine/Plugins/PCG/Source/PCG/Public/Helpers/PCGTagHelpers.h
@@ -29,6 +29,15 @@ namespace Private
 		TOptional<bool> BooleanValue;
 	};
 
+	enum class ESetAttributeFromTagFlags : uint8
+	{
+		None = 0,
+		CreateAttribute = 1,
+		OverwriteAttributeIfDifferentType = 2,
+		SetDefaultValue = 4
+	};
+	ENUM_CLASS_FLAGS(ESetAttributeFromTagFlags);
+
 	// Builds the tag result structure from the provided tag.
 	PCG_API FParseTagResult ParseTag(const FString& InTag);
 	PCG_API FParseTagResult ParseTag(FName InTag);
@@ -38,7 +47,14 @@ namespace Private
 	PCG_API bool CreateAttributeFromTag(const FParseTagResult& InTagData, UPCGMetadata* InMetadata);
 
 	// Parses a tag, optionally creates the corresponding attribute on the provided metadata and sets the value. If the name is invalid or the predicate rejects it, will return false.
+	PCG_API bool SetAttributeFromTag(const FString& InTag, UPCGMetadata* InMetadata, PCGMetadataEntryKey InKey, ESetAttributeFromTagFlags Flags = ESetAttributeFromTagFlags::None, FParseTagResult* OutResult = nullptr);
+	PCG_API bool SetAttributeFromTag(const FParseTagResult& InTagData, UPCGMetadata* InMetadata, PCGMetadataEntryKey InKey, ESetAttributeFromTagFlags Flags = ESetAttributeFromTagFlags::None);
+
+	// Deprecated methods
+	UE_DEPRECATED(5.6, "Use the version with the flags parameters")
 	PCG_API bool SetAttributeFromTag(const FString& InTag, UPCGMetadata* InMetadata, PCGMetadataEntryKey InKey, bool bInCanCreateAttribute, FParseTagResult* OutResult = nullptr);
+
+	UE_DEPRECATED(5.6, "Use the version with the flags parameter")
 	PCG_API bool SetAttributeFromTag(const FParseTagResult& InTagData, UPCGMetadata* InMetadata, PCGMetadataEntryKey InKey, bool bInCanCreateAttribute);
 }
 }
\ No newline at end of file
diff --git a/Engine/Plugins/PCG/Source/PCGEditor/Private/PCGLevelToAsset.cpp b/Engine/Plugins/PCG/Source/PCGEditor/Private/PCGLevelToAsset.cpp
--- a/Engine/Plugins/PCG/Source/PCGEditor/Private/PCGLevelToAsset.cpp
+++ b/Engine/Plugins/PCG/Source/PCGEditor/Private/PCGLevelToAsset.cpp
@@ -440,7 +440,7 @@ bool UPCGLevelToAsset::BP_ExportWorld_Implementation(UWorld* World, const FStrin
 
 				if (TagData.IsValid() && TagToAttributeMap.Contains(FName(TagData.GetOriginalAttribute())))
 				{
-					PCG::Private::SetAttributeFromTag(TagData, PointMetadata, MetadataEntry, /*bCanCreateAttribute=*/false);
+					PCG::Private::SetAttributeFromTag(TagData, PointMetadata, MetadataEntry);
 				}
 			}
 		};
