[PCG] Make the default metadata op type Unknown instead of Double and fix MakeVector issues.

* Explicitly defined default types for metadata ops, for robustness
* Added axis aligned defaults for MakeVector
* Moved PCGAttributeRemap into the Metadata folder
* Moved PCGParseString into the Metadata folder

#rnx

diff --git a/Engine/Plugins/PCG/Source/PCG/Private/Elements/Metadata/PCGMetadataMakeRotator.cpp b/Engine/Plugins/PCG/Source/PCG/Private/Elements/Metadata/PCGMetadataMakeRotator.cpp
--- a/Engine/Plugins/PCG/Source/PCG/Private/Elements/Metadata/PCGMetadataMakeRotator.cpp
+++ b/Engine/Plugins/PCG/Source/PCG/Private/Elements/Metadata/PCGMetadataMakeRotator.cpp
@@ -9,6 +9,34 @@
 
 #define LOCTEXT_NAMESPACE "PCGMetadataMakeRotatorSettings"
 
+namespace PCGMetadataMakeRotator::Helpers
+{
+	FVector GetAxisAlignedVector(const FName PinLabel)
+	{
+		if (PinLabel == PCGMetadataMakeRotatorConstants::XLabel || PinLabel == PCGMetadataMakeRotatorConstants::ForwardLabel)
+		{
+			return FVector::ForwardVector;
+		}
+		else if (PinLabel == PCGMetadataMakeRotatorConstants::YLabel || PinLabel == PCGMetadataMakeRotatorConstants::RightLabel)
+		{
+			return FVector::RightVector;
+		}
+		else if (PinLabel == PCGMetadataMakeRotatorConstants::ZLabel || PinLabel == PCGMetadataMakeRotatorConstants::UpLabel)
+		{
+			return FVector::UpVector;
+		}
+		else
+		{
+			return FVector::ZeroVector;
+		}
+	}
+
+	FString GetAxisAlignedVectorString(const FName PinLabel)
+	{
+		return PCG::Private::MetadataTraits<FVector>::ToString(GetAxisAlignedVector(PinLabel));
+	}
+}
+
 FName UPCGMetadataMakeRotatorSettings::GetInputPinLabel(uint32 Index) const
 {
 	switch (Operation)
@@ -145,20 +173,19 @@ FString UPCGMetadataMakeRotatorSettings::GetPinInitialDefaultValueString(FName P
 {
 	switch (Operation)
 	{
-		case EPCGMetadataMakeRotatorOp::MakeRotFromX: // fall-through
-		case EPCGMetadataMakeRotatorOp::MakeRotFromY: // fall-through
-		case EPCGMetadataMakeRotatorOp::MakeRotFromZ: // fall-through
-		case EPCGMetadataMakeRotatorOp::MakeRotFromAngles:
-			return PCG::Private::MetadataTraits<double>::ZeroValueString();
+		case EPCGMetadataMakeRotatorOp::MakeRotFromAxes:
+		case EPCGMetadataMakeRotatorOp::MakeRotFromX:  // fall-through
+		case EPCGMetadataMakeRotatorOp::MakeRotFromY:  // fall-through
+		case EPCGMetadataMakeRotatorOp::MakeRotFromZ:  // fall-through
 		case EPCGMetadataMakeRotatorOp::MakeRotFromXY: // fall-through
 		case EPCGMetadataMakeRotatorOp::MakeRotFromYX: // fall-through
 		case EPCGMetadataMakeRotatorOp::MakeRotFromXZ: // fall-through
 		case EPCGMetadataMakeRotatorOp::MakeRotFromZX: // fall-through
 		case EPCGMetadataMakeRotatorOp::MakeRotFromYZ: // fall-through
-		case EPCGMetadataMakeRotatorOp::MakeRotFromZY:
-			return PCG::Private::MetadataTraits<FVector2D>::ZeroValueString();
-		case EPCGMetadataMakeRotatorOp::MakeRotFromAxes:
-			return PCG::Private::MetadataTraits<FVector>::ZeroValueString();
+		case EPCGMetadataMakeRotatorOp::MakeRotFromZY: // fall-through
+			return PCGMetadataMakeRotator::Helpers::GetAxisAlignedVectorString(PinLabel);
+		case EPCGMetadataMakeRotatorOp::MakeRotFromAngles:
+			return PCG::Private::MetadataTraits<double>::ZeroValueString();
 		default:
 			checkNoEntry();
 			return FString();
@@ -166,6 +193,35 @@ FString UPCGMetadataMakeRotatorSettings::GetPinInitialDefaultValueString(FName P
 }
 #endif // WITH_EDITOR
 
+bool UPCGMetadataMakeRotatorSettings::CreateInitialDefaultValueAttribute(const FName PinLabel, UPCGMetadata* OutMetadata) const
+{
+	const FVector Value = PCGMetadataMakeRotator::Helpers::GetAxisAlignedVector(PinLabel);
+	return nullptr != OutMetadata->CreateAttribute(NAME_None, Value, /*bAllowsInterpolation=*/true, /*bOverrideParent=*/false);
+}
+
+EPCGMetadataTypes UPCGMetadataMakeRotatorSettings::GetPinInitialDefaultValueType(FName PinLabel) const
+{
+	switch (Operation)
+	{
+		case EPCGMetadataMakeRotatorOp::MakeRotFromAxes: // fall-through
+		case EPCGMetadataMakeRotatorOp::MakeRotFromX:    // fall-through
+		case EPCGMetadataMakeRotatorOp::MakeRotFromY:    // fall-through
+		case EPCGMetadataMakeRotatorOp::MakeRotFromZ:    // fall-through
+		case EPCGMetadataMakeRotatorOp::MakeRotFromXY:   // fall-through
+		case EPCGMetadataMakeRotatorOp::MakeRotFromYX:   // fall-through
+		case EPCGMetadataMakeRotatorOp::MakeRotFromXZ:   // fall-through
+		case EPCGMetadataMakeRotatorOp::MakeRotFromZX:   // fall-through
+		case EPCGMetadataMakeRotatorOp::MakeRotFromYZ:   // fall-through
+		case EPCGMetadataMakeRotatorOp::MakeRotFromZY:
+			return EPCGMetadataTypes::Vector;
+		case EPCGMetadataMakeRotatorOp::MakeRotFromAngles:
+			return EPCGMetadataTypes::Double;
+		default:
+			checkNoEntry();
+			return EPCGMetadataTypes::Unknown;
+	}
+}
+
 void UPCGMetadataMakeRotatorSettings::ApplyPreconfiguredSettings(const FPCGPreConfiguredSettingsInfo& PreconfiguredInfo)
 {
 	if (const UEnum* EnumPtr = StaticEnum<EPCGMetadataMakeRotatorOp>())
diff --git a/Engine/Plugins/PCG/Source/PCG/Private/Elements/Metadata/PCGMetadataOpElementBase.cpp b/Engine/Plugins/PCG/Source/PCG/Private/Elements/Metadata/PCGMetadataOpElementBase.cpp
--- a/Engine/Plugins/PCG/Source/PCG/Private/Elements/Metadata/PCGMetadataOpElementBase.cpp
+++ b/Engine/Plugins/PCG/Source/PCG/Private/Elements/Metadata/PCGMetadataOpElementBase.cpp
@@ -381,7 +381,13 @@ FString UPCGMetadataSettingsBase::GetPinDefaultValueAsString(const FName PinLabe
 
 EPCGMetadataTypes UPCGMetadataSettingsBase::GetPinInitialDefaultValueType(const FName PinLabel) const
 {
-	return EPCGMetadataTypes::Double;
+	// All overrides should exist for clarity and robustness.
+#if WITH_EDITOR
+	ensure(false);
+	PCGLog::LogErrorOnGraph(FText::Format(LOCTEXT("GetPinInitialDefaultValueTypeNotOverridden", "GetPinInitialDefaultValueType not overridden by node: {0}"), FText::FromName(GetDefaultNodeName())));
+#endif // WITH_EDITOR
+
+	return EPCGMetadataTypes::Unknown;
 }
 
 bool UPCGMetadataSettingsBase::IsInputPinRequiredByExecution(const UPCGPin* InPin) const
diff --git a/Engine/Plugins/PCG/Source/PCG/Private/Elements/PCGAttributeRemap.cpp b/Engine/Plugins/PCG/Source/PCG/Private/Elements/PCGAttributeRemap.cpp
--- a/Engine/Plugins/PCG/Source/PCG/Private/Elements/PCGAttributeRemap.cpp
+++ b/Engine/Plugins/PCG/Source/PCG/Private/Elements/PCGAttributeRemap.cpp
@@ -1,6 +1,6 @@
 // Copyright Epic Games, Inc. All Rights Reserved.
 
-#include "Elements/PCGAttributeRemap.h"
+#include "Elements/Metadata/PCGAttributeRemap.h"
 
 #include "PCGContext.h"
 #include "Helpers/PCGAsync.h"
diff --git a/Engine/Plugins/PCG/Source/PCG/Private/Elements/PCGParseString.cpp b/Engine/Plugins/PCG/Source/PCG/Private/Elements/PCGParseString.cpp
--- a/Engine/Plugins/PCG/Source/PCG/Private/Elements/PCGParseString.cpp
+++ b/Engine/Plugins/PCG/Source/PCG/Private/Elements/PCGParseString.cpp
@@ -1,6 +1,6 @@
 // Copyright Epic Games, Inc. All Rights Reserved.
 
-#include "Elements/PCGParseString.h"
+#include "Elements/Metadata/PCGParseString.h"
 
 #include "PCGContext.h"
 
diff --git a/Engine/Plugins/PCG/Source/PCG/Private/Tests/Elements/PCGDensityRemapTest.cpp b/Engine/Plugins/PCG/Source/PCG/Private/Tests/Elements/PCGDensityRemapTest.cpp
--- a/Engine/Plugins/PCG/Source/PCG/Private/Tests/Elements/PCGDensityRemapTest.cpp
+++ b/Engine/Plugins/PCG/Source/PCG/Private/Tests/Elements/PCGDensityRemapTest.cpp
@@ -6,7 +6,7 @@
 #include "PCGComponent.h"
 
 #include "Data/PCGPointData.h"
-#include "Elements/PCGAttributeRemap.h"
+#include "Elements/Metadata/PCGAttributeRemap.h"
 #include "PCGContext.h"
 
 IMPLEMENT_CUSTOM_SIMPLE_AUTOMATION_TEST(FPCGDensityRemapTest, FPCGTestBaseClass, "Plugins.PCG.DensityRemap.Basic", PCGTestsCommon::TestFlags)
diff --git a/Engine/Plugins/PCG/Source/PCG/Private/Tests/Elements/PCGParseStringTests.cpp b/Engine/Plugins/PCG/Source/PCG/Private/Tests/Elements/PCGParseStringTests.cpp
--- a/Engine/Plugins/PCG/Source/PCG/Private/Tests/Elements/PCGParseStringTests.cpp
+++ b/Engine/Plugins/PCG/Source/PCG/Private/Tests/Elements/PCGParseStringTests.cpp
@@ -4,7 +4,7 @@
 
 #include "PCGContext.h"
 #include "PCGParamData.h"
-#include "Elements/PCGParseString.h"
+#include "Elements/Metadata/PCGParseString.h"
 
 class FPCGParseStringTestsBase : public FPCGTestBaseClass
 {
diff --git a/Engine/Plugins/PCG/Source/PCG/Public/Elements/Metadata/PCGAttributeRemap.h b/Engine/Plugins/PCG/Source/PCG/Public/Elements/Metadata/PCGAttributeRemap.h
new file mode 100644
--- /dev/null
+++ b/Engine/Plugins/PCG/Source/PCG/Public/Elements/Metadata/PCGAttributeRemap.h
@@ -0,0 +1,88 @@
+// Copyright Epic Games, Inc. All Rights Reserved.
+
+#pragma once
+
+#include "PCGSettings.h"
+#include "Elements/Metadata/PCGMetadataOpElementBase.h"
+#include "Elements/PCGTimeSlicedElementBase.h"
+#include "Metadata/Accessors/IPCGAttributeAccessor.h"
+#include "Metadata/Accessors/PCGAttributeAccessorKeys.h"
+
+#include "Templates/UniquePtr.h"
+
+#include "PCGAttributeRemap.generated.h"
+
+/**
+* Remap attribute values from one range to another.
+*/
+UCLASS(MinimalAPI, BlueprintType, ClassGroup = (Procedural))
+class UPCGAttributeRemapSettings : public UPCGMetadataSettingsBase
+{
+	GENERATED_BODY()
+
+public:
+	UPCGAttributeRemapSettings();
+
+	//~Begin UPCGSettings interface
+#if WITH_EDITOR
+	virtual FName GetDefaultNodeName() const override;
+	virtual FText GetDefaultNodeTitle() const override;
+	virtual TArray<FPCGPreConfiguredSettingsInfo> GetPreconfiguredInfo() const override;
+	virtual bool GroupPreconfiguredSettings() const override { return false; }
+	virtual bool HasFlippedTitleLines() const override { return false; }
+#endif
+
+	virtual FString GetAdditionalTitleInformation() const override;
+	virtual void ApplyPreconfiguredSettings(const FPCGPreConfiguredSettingsInfo& PreconfigureInfo) override;
+
+protected:
+	virtual FPCGElementPtr CreateElement() const override;
+	//~End UPCGSettings interface
+
+public:
+	//~Begin UPCGMetadataSettingsBase interface
+	virtual FPCGAttributePropertyInputSelector GetInputSource(uint32 Index) const override;
+	virtual bool IsSupportedInputType(uint16 TypeId, uint32 InputIndex, bool& bHasSpecialRequirement) const override;
+	//~End UPCGMetadataSettingsBase interface
+
+	//~Begin IPCGSettingsDefaultValueProvider interface
+	virtual EPCGMetadataTypes GetPinInitialDefaultValueType(FName PinLabel) const override { return EPCGMetadataTypes::Double; }
+	//~End IPCGSettingsDefaultValueProvider interface
+
+public:
+	UPROPERTY(BlueprintReadWrite, EditAnywhere, Category = Settings, meta = (PCG_Overridable))
+	FPCGAttributePropertyInputSelector InputSource;
+
+	/** If InRangeMin = InRangeMax, then that attribute value is mapped to the average of OutRangeMin and OutRangeMax */
+	UPROPERTY(BlueprintReadWrite, EditAnywhere, Category = Settings, meta = (PCG_Overridable))
+	double InRangeMin = 0.f;
+
+	/** If InRangeMin = InRangeMax, then that attribute value is mapped to the average of OutRangeMin and OutRangeMax */
+	UPROPERTY(BlueprintReadWrite, EditAnywhere, Category = Settings, meta = (PCG_Overridable))
+	double InRangeMax = 1.f;
+
+	UPROPERTY(BlueprintReadWrite, EditAnywhere, Category = Settings, meta = (PCG_Overridable))
+	double OutRangeMin = 0.f;
+
+	UPROPERTY(BlueprintReadWrite, EditAnywhere, Category = Settings, meta = (PCG_Overridable))
+	double OutRangeMax = 1.f;
+
+	/** If checked, outside values will be clamped between 0 and 1. */
+	UPROPERTY(BlueprintReadWrite, EditAnywhere, Category = Settings, meta = (PCG_Overridable))
+	bool bClampToUnitRange = false;
+
+	/** Attribute values outside of the input range will be unaffected by the remapping */
+	UPROPERTY(BlueprintReadWrite, EditAnywhere, Category = Settings, meta = (PCG_Overridable))
+	bool bIgnoreValuesOutsideInputRange = false;
+
+	/** Allow remapping when Min is larger than Max, e.g. from [0.0, 1.0] -> [1.0, 0.0]. */
+	UPROPERTY(BlueprintReadWrite, EditAnywhere, Category = Settings, meta = (PCG_Overridable))
+	bool bAllowInverseRange = false; // Note that this is no longer the default value for new nodes, it is now 'true'
+};
+
+class FPCGAttributeRemapElement : public FPCGMetadataElementBase
+{
+protected:
+	virtual bool DoOperation(PCGMetadataOps::FOperationData& OperationData) const override;
+};
+
diff --git a/Engine/Plugins/PCG/Source/PCG/Public/Elements/Metadata/PCGHashAttribute.h b/Engine/Plugins/PCG/Source/PCG/Public/Elements/Metadata/PCGHashAttribute.h
--- a/Engine/Plugins/PCG/Source/PCG/Public/Elements/Metadata/PCGHashAttribute.h
+++ b/Engine/Plugins/PCG/Source/PCG/Public/Elements/Metadata/PCGHashAttribute.h
@@ -30,6 +30,10 @@ public:
 	virtual FPCGElementPtr CreateElement() const override;
 	//~End UPCGSettings interface
 
+	//~Begin IPCGSettingsDefaultValueProvider interface
+	virtual EPCGMetadataTypes GetPinInitialDefaultValueType(FName PinLabel) const override { return EPCGMetadataTypes::String; }
+	//~End IPCGSettingsDefaultValueProvider interface
+
 	UPROPERTY(BlueprintReadWrite, EditAnywhere, Category = Settings, meta = (PCG_Overridable))
 	FPCGAttributePropertyInputSelector InputSource1;
 };
diff --git a/Engine/Plugins/PCG/Source/PCG/Public/Elements/Metadata/PCGMetadataCompareOpElement.h b/Engine/Plugins/PCG/Source/PCG/Public/Elements/Metadata/PCGMetadataCompareOpElement.h
--- a/Engine/Plugins/PCG/Source/PCG/Public/Elements/Metadata/PCGMetadataCompareOpElement.h
+++ b/Engine/Plugins/PCG/Source/PCG/Public/Elements/Metadata/PCGMetadataCompareOpElement.h
@@ -48,6 +48,10 @@ public:
 	virtual uint16 GetOutputType(uint16 InputTypeId) const override;
 	//~End UPCGMetadataSettingsBase interface
 
+	//~Begin IPCGSettingsDefaultValueProvider interface
+	virtual EPCGMetadataTypes GetPinInitialDefaultValueType(FName PinLabel) const override { return EPCGMetadataTypes::Double; }
+	//~End IPCGSettingsDefaultValueProvider interface
+
 protected:
 	//~Begin UPCGSettings interface
 #if WITH_EDITOR
diff --git a/Engine/Plugins/PCG/Source/PCG/Public/Elements/Metadata/PCGMetadataMakeRotator.h b/Engine/Plugins/PCG/Source/PCG/Public/Elements/Metadata/PCGMetadataMakeRotator.h
--- a/Engine/Plugins/PCG/Source/PCG/Public/Elements/Metadata/PCGMetadataMakeRotator.h
+++ b/Engine/Plugins/PCG/Source/PCG/Public/Elements/Metadata/PCGMetadataMakeRotator.h
@@ -71,6 +71,8 @@ protected:
 #if WITH_EDITOR
 	virtual FString GetPinInitialDefaultValueString(FName PinLabel) const override;
 #endif // WITH_EDITOR
+	virtual bool CreateInitialDefaultValueAttribute(FName PinLabel, UPCGMetadata* OutMetadata) const override;
+	virtual EPCGMetadataTypes GetPinInitialDefaultValueType(FName PinLabel) const override;
 	//~End IPCGSettingsDefaultValueProvider interface
 
 public:
diff --git a/Engine/Plugins/PCG/Source/PCG/Public/Elements/Metadata/PCGMetadataMathsOpElement.h b/Engine/Plugins/PCG/Source/PCG/Public/Elements/Metadata/PCGMetadataMathsOpElement.h
--- a/Engine/Plugins/PCG/Source/PCG/Public/Elements/Metadata/PCGMetadataMathsOpElement.h
+++ b/Engine/Plugins/PCG/Source/PCG/Public/Elements/Metadata/PCGMetadataMathsOpElement.h
@@ -77,6 +77,10 @@ public:
 	virtual bool IsSupportedInputType(uint16 TypeId, uint32 InputIndex, bool& bHasSpecialRequirement) const override;
 	//~End UPCGMetadataSettingsBase interface
 
+	//~Begin IPCGSettingsDefaultValueProvider interface
+	virtual EPCGMetadataTypes GetPinInitialDefaultValueType(FName PinLabel) const override { return EPCGMetadataTypes::Double; }
+	//~End IPCGSettingsDefaultValueProvider interface
+
 protected:
 	//~Begin UPCGSettings interface
 #if WITH_EDITOR
diff --git a/Engine/Plugins/PCG/Source/PCG/Public/Elements/Metadata/PCGMetadataTrigOpElement.h b/Engine/Plugins/PCG/Source/PCG/Public/Elements/Metadata/PCGMetadataTrigOpElement.h
--- a/Engine/Plugins/PCG/Source/PCG/Public/Elements/Metadata/PCGMetadataTrigOpElement.h
+++ b/Engine/Plugins/PCG/Source/PCG/Public/Elements/Metadata/PCGMetadataTrigOpElement.h
@@ -59,6 +59,10 @@ protected:
 	virtual FPCGElementPtr CreateElement() const override;
 	//~End UPCGSettings interface
 
+	//~Begin IPCGSettingsDefaultValueProvider interface
+	virtual EPCGMetadataTypes GetPinInitialDefaultValueType(FName PinLabel) const override { return EPCGMetadataTypes::Double; }
+	//~End IPCGSettingsDefaultValueProvider interface
+
 public:
 	UPROPERTY(BlueprintReadWrite, EditAnywhere, Category = Settings)
 	EPCGMetadataTrigOperation Operation = EPCGMetadataTrigOperation::Acos;
diff --git a/Engine/Plugins/PCG/Source/PCG/Public/Elements/Metadata/PCGParseString.h b/Engine/Plugins/PCG/Source/PCG/Public/Elements/Metadata/PCGParseString.h
new file mode 100644
--- /dev/null
+++ b/Engine/Plugins/PCG/Source/PCG/Public/Elements/Metadata/PCGParseString.h
@@ -0,0 +1,58 @@
+// Copyright Epic Games, Inc. All Rights Reserved.
+
+#pragma once
+
+#include "Elements/Metadata/PCGMetadataOpElementBase.h"
+#include "Metadata/PCGMetadataAttributeTraits.h"
+
+#include "PCGParseString.generated.h"
+
+/**
+* Convert a string to a numerical value.
+*/
+UCLASS(MinimalAPI, BlueprintType, ClassGroup = (Procedural))
+class UPCGParseStringSettings : public UPCGMetadataSettingsBase
+{
+	GENERATED_BODY()
+
+public:
+	//~Begin UPCGSettings interface
+#if WITH_EDITOR
+	virtual FName GetDefaultNodeName() const override;
+	virtual FText GetDefaultNodeTitle() const override;
+	virtual FText GetNodeTooltipText() const override;
+#endif
+
+	virtual bool HasFlippedTitleLines() const override { return false; }
+	virtual FString GetAdditionalTitleInformation() const override;
+
+protected:
+	virtual FPCGElementPtr CreateElement() const override;
+	//~End UPCGSettings interface
+
+public:
+	//~Begin UPCGMetadataSettingsBase interface
+	virtual FPCGAttributePropertyInputSelector GetInputSource(uint32 Index) const override;
+	virtual bool IsSupportedInputType(uint16 TypeId, uint32 InputIndex, bool& bHasSpecialRequirement) const override;
+	virtual bool HasDifferentOutputTypes() const { return true; }
+	virtual TArray<uint16> GetAllOutputTypes() const;
+	//~End UPCGMetadataSettingsBase interface
+
+	//~Begin IPCGSettingsDefaultValueProvider interface
+	virtual EPCGMetadataTypes GetPinInitialDefaultValueType(FName PinLabel) const override { return EPCGMetadataTypes::String; }
+	//~End IPCGSettingsDefaultValueProvider interface
+
+public:
+	UPROPERTY(BlueprintReadWrite, EditAnywhere, Category = "Input", meta = (PCG_Overridable))
+	FPCGAttributePropertyInputSelector InputSource;
+
+	UPROPERTY(BlueprintReadWrite, EditAnywhere, Category = "Settings", meta = (PCG_Overridable))
+	EPCGMetadataTypes TargetType = EPCGMetadataTypes::Integer32;
+};
+
+class FPCGParseStringElement : public FPCGMetadataElementBase
+{
+protected:
+	virtual bool DoOperation(PCGMetadataOps::FOperationData& OperationData) const override;
+};
+
diff --git a/Engine/Plugins/PCG/Source/PCG/Public/Elements/PCGAttributeRemap.h b/Engine/Plugins/PCG/Source/PCG/Public/Elements/PCGAttributeRemap.h
--- a/Engine/Plugins/PCG/Source/PCG/Public/Elements/PCGAttributeRemap.h
+++ b/Engine/Plugins/PCG/Source/PCG/Public/Elements/PCGAttributeRemap.h
@@ -2,83 +2,8 @@
 
 #pragma once
 
-#include "PCGSettings.h"
-#include "Elements/Metadata/PCGMetadataOpElementBase.h"
-#include "Elements/PCGTimeSlicedElementBase.h"
-#include "Metadata/Accessors/IPCGAttributeAccessor.h"
-#include "Metadata/Accessors/PCGAttributeAccessorKeys.h"
+#include "CoreTypes.h"
 
-#include "Templates/UniquePtr.h"
-
-#include "PCGAttributeRemap.generated.h"
-
-/**
-* Remap attribute values from one range to another.
-*/
-UCLASS(MinimalAPI, BlueprintType, ClassGroup = (Procedural))
-class UPCGAttributeRemapSettings : public UPCGMetadataSettingsBase
-{
-	GENERATED_BODY()
-
-public:
-	UPCGAttributeRemapSettings();
-
-	//~Begin UPCGSettings interface
-#if WITH_EDITOR
-	virtual FName GetDefaultNodeName() const override;
-	virtual FText GetDefaultNodeTitle() const override;
-	virtual TArray<FPCGPreConfiguredSettingsInfo> GetPreconfiguredInfo() const override;
-	virtual bool GroupPreconfiguredSettings() const override { return false; }
-	virtual bool HasFlippedTitleLines() const override { return false; }
-#endif
-
-	virtual FString GetAdditionalTitleInformation() const override;
-	virtual void ApplyPreconfiguredSettings(const FPCGPreConfiguredSettingsInfo& PreconfigureInfo) override;
-
-protected:
-	virtual FPCGElementPtr CreateElement() const override;
-	//~End UPCGSettings interface
-
-public:
-	//~Begin UPCGMetadataSettingsBase interface
-	virtual FPCGAttributePropertyInputSelector GetInputSource(uint32 Index) const override;
-	virtual bool IsSupportedInputType(uint16 TypeId, uint32 InputIndex, bool& bHasSpecialRequirement) const override;
-	//~End UPCGMetadataSettingsBase interface
-
-public:
-	UPROPERTY(BlueprintReadWrite, EditAnywhere, Category = Settings, meta = (PCG_Overridable))
-	FPCGAttributePropertyInputSelector InputSource;
-
-	/** If InRangeMin = InRangeMax, then that attribute value is mapped to the average of OutRangeMin and OutRangeMax */
-	UPROPERTY(BlueprintReadWrite, EditAnywhere, Category = Settings, meta = (PCG_Overridable))
-	double InRangeMin = 0.f;
-
-	/** If InRangeMin = InRangeMax, then that attribute value is mapped to the average of OutRangeMin and OutRangeMax */
-	UPROPERTY(BlueprintReadWrite, EditAnywhere, Category = Settings, meta = (PCG_Overridable))
-	double InRangeMax = 1.f;
-
-	UPROPERTY(BlueprintReadWrite, EditAnywhere, Category = Settings, meta = (PCG_Overridable))
-	double OutRangeMin = 0.f;
-
-	UPROPERTY(BlueprintReadWrite, EditAnywhere, Category = Settings, meta = (PCG_Overridable))
-	double OutRangeMax = 1.f;
-
-	/** If checked, outside values will be clamped between 0 and 1. */
-	UPROPERTY(BlueprintReadWrite, EditAnywhere, Category = Settings, meta = (PCG_Overridable))
-	bool bClampToUnitRange = false;
-
-	/** Attribute values outside of the input range will be unaffected by the remapping */
-	UPROPERTY(BlueprintReadWrite, EditAnywhere, Category = Settings, meta = (PCG_Overridable))
-	bool bIgnoreValuesOutsideInputRange = false;
-
-	/** Allow remapping when Min is larger than Max, e.g. from [0.0, 1.0] -> [1.0, 0.0]. */
-	UPROPERTY(BlueprintReadWrite, EditAnywhere, Category = Settings, meta = (PCG_Overridable))
-	bool bAllowInverseRange = false; // Note that this is no longer the default value for new nodes, it is now 'true'
-};
-
-class FPCGAttributeRemapElement : public FPCGMetadataElementBase
-{
-protected:
-	virtual bool DoOperation(PCGMetadataOps::FOperationData& OperationData) const override;
-};
+UE_DEPRECATED_HEADER(5.6, "Use /Elements/Metadata/PCGAttributeRemap.h instead of /Elements/PCGAttributeRemap.h")
 
+#include "Metadata/PCGAttributeRemap.h"
\ No newline at end of file
diff --git a/Engine/Plugins/PCG/Source/PCG/Public/Elements/PCGParseString.h b/Engine/Plugins/PCG/Source/PCG/Public/Elements/PCGParseString.h
--- a/Engine/Plugins/PCG/Source/PCG/Public/Elements/PCGParseString.h
+++ b/Engine/Plugins/PCG/Source/PCG/Public/Elements/PCGParseString.h
@@ -2,53 +2,8 @@
 
 #pragma once
 
-#include "Elements/Metadata/PCGMetadataOpElementBase.h"
-#include "Metadata/PCGMetadataAttributeTraits.h"
+#include "CoreTypes.h"
 
-#include "PCGParseString.generated.h"
-
-/**
-* Convert a string to a numerical value.
-*/
-UCLASS(MinimalAPI, BlueprintType, ClassGroup = (Procedural))
-class UPCGParseStringSettings : public UPCGMetadataSettingsBase
-{
-	GENERATED_BODY()
-
-public:
-	//~Begin UPCGSettings interface
-#if WITH_EDITOR
-	virtual FName GetDefaultNodeName() const override;
-	virtual FText GetDefaultNodeTitle() const override;
-	virtual FText GetNodeTooltipText() const override;
-#endif
-
-	virtual bool HasFlippedTitleLines() const override { return false; }
-	virtual FString GetAdditionalTitleInformation() const override;
-
-protected:
-	virtual FPCGElementPtr CreateElement() const override;
-	//~End UPCGSettings interface
-
-public:
-	//~Begin UPCGMetadataSettingsBase interface
-	virtual FPCGAttributePropertyInputSelector GetInputSource(uint32 Index) const override;
-	virtual bool IsSupportedInputType(uint16 TypeId, uint32 InputIndex, bool& bHasSpecialRequirement) const override;
-	virtual bool HasDifferentOutputTypes() const { return true; }
-	virtual TArray<uint16> GetAllOutputTypes() const;
-	//~End UPCGMetadataSettingsBase interface
-
-public:
-	UPROPERTY(BlueprintReadWrite, EditAnywhere, Category = "Input", meta = (PCG_Overridable))
-	FPCGAttributePropertyInputSelector InputSource;
-
-	UPROPERTY(BlueprintReadWrite, EditAnywhere, Category = "Settings", meta = (PCG_Overridable))
-	EPCGMetadataTypes TargetType = EPCGMetadataTypes::Integer32;
-};
-
-class FPCGParseStringElement : public FPCGMetadataElementBase
-{
-protected:
-	virtual bool DoOperation(PCGMetadataOps::FOperationData& OperationData) const override;
-};
+UE_DEPRECATED_HEADER(5.6, "Use /Elements/Metadata/PCGParseString.h instead of /Elements/PCGParseString.h")
 
+#include "Elements/Metadata/PCGParseString.h"
\ No newline at end of file
