[PCG] Removed "Criteria as Input" as an overridable property in Select Grammar

"Criteria as Input" toggles manually a structural change adding the "Selection Criteria Data" input pin

#rnx

diff --git a/Engine/Plugins/PCG/Source/PCG/Public/Elements/Grammar/PCGSelectGrammar.h b/Engine/Plugins/PCG/Source/PCG/Public/Elements/Grammar/PCGSelectGrammar.h
--- a/Engine/Plugins/PCG/Source/PCG/Public/Elements/Grammar/PCGSelectGrammar.h
+++ b/Engine/Plugins/PCG/Source/PCG/Public/Elements/Grammar/PCGSelectGrammar.h
@@ -116,8 +116,8 @@ public:
 	UPROPERTY(BlueprintReadWrite, EditAnywhere, Category = Settings, meta = (PCG_Overridable))
 	FPCGAttributePropertyInputSelector ComparedValueAttribute;
 
-	/** Pass criteria info as an attribute set matching the structure type 'FPCGSelectGrammarCriteria'. */
-	UPROPERTY(BlueprintReadWrite, EditAnywhere, Category = Settings, meta = (PCG_Overridable))
+	/** Toggle for an additional input "Selection Criteria Data" which accepts criteria info as an attribute set matching the structure type 'FPCGSelectGrammarCriteria'. */
+	UPROPERTY(BlueprintReadWrite, EditAnywhere, Category = Settings)
 	bool bCriteriaAsInput = false;
 
 	/** Selection criteria that will be evaluated in order. */
