Fixed return type when trying to create an attribute from tag that already exists.

#rnx

diff --git a/Engine/Plugins/PCG/Source/PCG/Private/Helpers/PCGTagHelpers.cpp b/Engine/Plugins/PCG/Source/PCG/Private/Helpers/PCGTagHelpers.cpp
--- a/Engine/Plugins/PCG/Source/PCG/Private/Helpers/PCGTagHelpers.cpp
+++ b/Engine/Plugins/PCG/Source/PCG/Private/Helpers/PCGTagHelpers.cpp
@@ -132,16 +132,18 @@ namespace Private
 					if (InEntryKey != PCGInvalidEntryKey)
 					{
 						Attribute->SetValue(InEntryKey, Value);
-						return true;
 					}
 					else if (!!(Flags & ESetAttributeFromTagFlags::SetDefaultValue))
 					{
 						Attribute->SetDefaultValue(Value);
-						return true;
 					}
-				}
 
-				return false;
+					return true;
+				}
+				else
+				{
+					return false;
+				}
 			};
 
 			if (TagData.HasNumericValue())
