Prevent potential nullptr in -game mode

#rnx

diff --git a/Engine/Plugins/PCG/Source/PCG/Private/RuntimeGen/PCGGenSourceManager.cpp b/Engine/Plugins/PCG/Source/PCG/Private/RuntimeGen/PCGGenSourceManager.cpp
--- a/Engine/Plugins/PCG/Source/PCG/Private/RuntimeGen/PCGGenSourceManager.cpp
+++ b/Engine/Plugins/PCG/Source/PCG/Private/RuntimeGen/PCGGenSourceManager.cpp
@@ -189,7 +189,7 @@ void FPCGGenSourceManager::UpdatePerTickGenSources(const APCGWorldActor* InPCGWo
 	EditorCameraGenSource->EditorViewportClient = nullptr;
 
 	// Update the active editor viewport client for the EditorCameraGenSource, only if requested by the world actor, in-editor, and the viewport is visible.
-	if (InPCGWorldActor->bTreatEditorViewportAsGenerationSource && !World->IsGameWorld())
+	if (InPCGWorldActor->bTreatEditorViewportAsGenerationSource && !World->IsGameWorld() && GEditor)
 	{
 		if (FViewport* Viewport = GEditor->GetActiveViewport())
 		{
