[PCG] Add information when we cross boundaries with compiled tasks, to mark them used multiple times

* This was forgotten when initially fixing the linked Jira
* Might need to fix it properly in 5.6 (Can't fix it properly without an ABI break for hotfix)
* This specific instance would crash, so it is kinda important to fix

#rnx

diff --git a/Engine/Plugins/PCG/Source/PCG/Private/Graph/PCGGraphCompiler.cpp b/Engine/Plugins/PCG/Source/PCG/Private/Graph/PCGGraphCompiler.cpp
--- a/Engine/Plugins/PCG/Source/PCG/Private/Graph/PCGGraphCompiler.cpp
+++ b/Engine/Plugins/PCG/Source/PCG/Private/Graph/PCGGraphCompiler.cpp
@@ -1417,7 +1417,13 @@ void FPCGGraphCompiler::CompileTopGraph(UPCGGraph* InGraph, uint32 GenerationGri
 			}
 			else
 			{
-				TaskInput.bIsUsedMultipleTimes = false;
+				// @todo_pcg: This should not have that big of a cost, as most InputTasks do not already have an element when compiling.
+				// The virtual call will then not be often called. However, it would be cleaner to have a boolean state on the TaskInput 
+				// that is set to true when we know we cross boundaries.
+				// By definition, if we cross a boundary, then an input will be used multiple times.
+				const FPCGGraphTask* InputTask = CompiledTasks.IsValidIndex(TaskInput.TaskId) ? &CompiledTasks[TaskInput.TaskId] : nullptr;
+				const bool bCrossBoundaries = InputTask && InputTask->Element && InputTask->Element->IsGridLinkage();
+				TaskInput.bIsUsedMultipleTimes = bCrossBoundaries;
 				SeenInputs.Emplace(std::move(InputOrigin));
 			}
 		}
