[PCG] PCGActorSelector: Pre-filter Actor class (removes most of PCGActorSelector::FindActors  cost when selecting by class)

#rnx

diff --git a/Engine/Plugins/PCG/Source/PCG/Private/Elements/PCGActorSelector.cpp b/Engine/Plugins/PCG/Source/PCG/Private/Elements/PCGActorSelector.cpp
--- a/Engine/Plugins/PCG/Source/PCG/Private/Elements/PCGActorSelector.cpp
+++ b/Engine/Plugins/PCG/Source/PCG/Private/Elements/PCGActorSelector.cpp
@@ -109,8 +109,10 @@ namespace PCGActorSelector
 		// In case of iterating over all actors in the world, call our filtering function and get out.
 		if (Settings.ActorFilter == EPCGActorFilter::AllWorldActors)
 		{
+			// If we are going to filter by class in the FilteringFunction, might as well avoid the cost of iterating over all actor types (which can get really expensive)
+			const TSubclassOf<AActor> ActorSelectionClass = Settings.ActorSelection == EPCGActorSelection::ByClass && Settings.ActorSelectionClass ? Settings.ActorSelectionClass.Get() : AActor::StaticClass();
 			// A potential optimization if we know the sought actors are collide-able could be to obtain overlaps via a collision query.
-			UPCGActorHelpers::ForEachActorInWorld<AActor>(World, FilteringFunction);
+			UPCGActorHelpers::ForEachActorInWorld(World, ActorSelectionClass, FilteringFunction);
 
 			// FoundActor is set by the FilteringFunction (captured)
 			return FoundActors;
