[PCG] Fix missing registration of BP PCG Components
- This happened when the PCG Component was added to the BP after the BP actors were first saved to the map
- Implement delayed registration similar to delayed unregistration tracking code

#rnx

diff --git a/Engine/Plugins/PCG/Source/PCG/Private/PCGActorAndComponentMapping.cpp b/Engine/Plugins/PCG/Source/PCG/Private/PCGActorAndComponentMapping.cpp
--- a/Engine/Plugins/PCG/Source/PCG/Private/PCGActorAndComponentMapping.cpp
+++ b/Engine/Plugins/PCG/Source/PCG/Private/PCGActorAndComponentMapping.cpp
@@ -19,6 +19,7 @@
 #include "StaticMeshCompiler.h"
 #include "TextureCompiler.h"
 #include "Algo/AnyOf.h"
+#include "Algo/ForEach.h"
 #include "Engine/DataTable.h"
 #include "Engine/Engine.h"
 #include "Engine/Level.h"
@@ -100,11 +101,6 @@ namespace PCGActorAndComponentMapping
 		TEXT("Disable the optimization that keep track of components depending on others, as a safety measure."));
 #endif // WITH_EDITOR
 
-	static TAutoConsoleVariable<bool> CVarDisableDelayedUnregister(
-		TEXT("pcg.DisableDelayedUnregister"),
-		true,
-		TEXT("If delayed unregister for all is introducing bad behavior, disables it, allowing people to continue working while we investigate."));
-
 #if WITH_EDITOR
 	bool ShouldIgnoreActor(AActor* InActor)
 	{
@@ -199,16 +195,17 @@ void FPCGActorAndComponentMapping::Tick()
 {
 #if WITH_EDITOR
 
-	TSet<UPCGComponent*> ComponentToUnregister;
+	// Process delayed Registrations
+	TSet<UPCGComponent*> ComponentsToUnregister;
+	TSet<UPCGComponent*> ComponentsToRegister;
 	{
-		FScopeLock Lock(&DelayedComponentToUnregisterLock);
-		ComponentToUnregister = MoveTemp(DelayedComponentToUnregister);
+		FScopeLock Lock(&DelayedComponentsRegistrationLock);
+		ComponentsToUnregister = MoveTemp(DelayedComponentsToUnregister);
+		ComponentsToRegister = MoveTemp(DelayedComponentsToRegister);
 	}
 
-	for (UPCGComponent* Component : ComponentToUnregister)
-	{
-		UnregisterPCGComponent(Component, /*bForce=*/true);
-	}
+	Algo::ForEach(ComponentsToRegister, [this](UPCGComponent* Component) { RegisterOrUpdatePCGComponent(Component, Component->bGenerated, /*bForce=*/true); });
+	Algo::ForEach(ComponentsToUnregister, [this](UPCGComponent* Component) { UnregisterPCGComponent(Component, /*bForce=*/true); });
 
 	ProcessDelayedEvents();
 
@@ -303,6 +300,11 @@ TArray<FPCGTaskId> FPCGActorAndComponentMapping::DispatchToLocalComponents(UPCGC
 }
 
 bool FPCGActorAndComponentMapping::RegisterOrUpdatePCGComponent(UPCGComponent* InComponent, bool bDoActorMapping)
+{
+	return RegisterOrUpdatePCGComponent(InComponent, bDoActorMapping, /*bForce=*/false);
+}
+
+bool FPCGActorAndComponentMapping::RegisterOrUpdatePCGComponent(UPCGComponent* InComponent, bool bDoActorMapping, bool bForce)
 {
 	check(InComponent);
 
@@ -319,6 +321,24 @@ bool FPCGActorAndComponentMapping::RegisterOrUpdatePCGComponent(UPCGComponent* I
 		return false;
 	}
 
+#if WITH_EDITOR
+	{
+		FScopeLock Lock(&DelayedComponentsRegistrationLock);
+
+		DelayedComponentsToUnregister.Remove(InComponent);
+		DelayedComponentsToRegister.Remove(InComponent);
+
+		// When we are running Construction Script there are 2 possibilities
+		// 1. This component is replacing an existing PCG Component and we don't want to register, we want to replace in the RemapPCGComponent call (remap will remove this component from the Delayed list)
+		// 2. This comopnent is new and so it will be registered after the potential replacement call on the next frame in FPCGActorAndComponentMapping::Tick
+		if (!bForce && GetWorld()->bIsRunningConstructionScript && InComponent->IsCreatedByConstructionScript())
+		{
+			DelayedComponentsToRegister.Add(InComponent);
+			return true;
+		}
+	}
+#endif // WITH_EDITOR
+
 	const bool bWasAlreadyRegistered = IsComponentRegistered(InComponent);
 
 	// First check if the component has changed its partitioned flag.
@@ -346,12 +366,6 @@ bool FPCGActorAndComponentMapping::RegisterOrUpdatePCGComponent(UPCGComponent* I
 	}
 
 #if WITH_EDITOR
-	// If the component was previously marked as to be unregistered, remove it here.
-	{
-		FScopeLock Lock(&DelayedComponentToUnregisterLock);
-		DelayedComponentToUnregister.Remove(InComponent);
-	}
-
 	// And finally handle the tracking. Only do it when the component is registered for the first time.
 	if (!bWasAlreadyRegistered && bHasChanged)
 	{
@@ -441,8 +455,9 @@ bool FPCGActorAndComponentMapping::RemapPCGComponent(const UPCGComponent* OldCom
 #if WITH_EDITOR
 	// Remove it from the delayed
 	{
-		FScopeLock Lock(&DelayedComponentToUnregisterLock);
-		DelayedComponentToUnregister.Remove(OldComponent);
+		FScopeLock Lock(&DelayedComponentsRegistrationLock);
+		DelayedComponentsToUnregister.Remove(OldComponent);
+		DelayedComponentsToRegister.Remove(NewComponent);
 	}
 #endif
 
@@ -488,23 +503,23 @@ void FPCGActorAndComponentMapping::UnregisterPCGComponent(UPCGComponent* InCompo
 	}
 
 #if WITH_EDITOR
-	if ((PartitionedOctree.Contains(InComponent) || NonPartitionedOctree.Contains(InComponent)))
 	{
-		bool bShouldBeDelayed = true;
-		if (PCGActorAndComponentMapping::CVarDisableDelayedUnregister.GetValueOnAnyThread())
-		{
-			// We also need to check that our current PCG Component is not deleted while being reconstructed by a construction script.
-			// If so, it will be "re-created" at some point with the same properties.
-			// In this particular case, we don't remove the PCG component from the octree and we won't delete the mapping, but mark it to be removed
-			// at next Subsystem tick. If we call "RemapPCGComponent" before, we will re-connect everything correctly.
-			// Ignore this if we force (aka when we actually unregister the delayed one)
-			bShouldBeDelayed = InComponent->IsCreatedByConstructionScript();
-		}
+		FScopeLock Lock(&DelayedComponentsRegistrationLock);
+		DelayedComponentsToRegister.Remove(InComponent);
+		DelayedComponentsToUnregister.Remove(InComponent);
+	}
 
-		if (!bForce && bShouldBeDelayed)
+	if ((PartitionedOctree.Contains(InComponent) || NonPartitionedOctree.Contains(InComponent)))
+	{
+		// We also need to check that our current PCG Component is not deleted while being reconstructed by a construction script.
+		// If so, it will be "re-created" at some point with the same properties.
+		// In this particular case, we don't remove the PCG component from the octree and we won't delete the mapping, but mark it to be removed
+		// at next Subsystem tick. If we call "RemapPCGComponent" before, we will re-connect everything correctly.
+		// Ignore this if we force (aka when we actually unregister the delayed one)		
+		if (!bForce && InComponent->IsCreatedByConstructionScript())
 		{
-			FScopeLock Lock(&DelayedComponentToUnregisterLock);
-			DelayedComponentToUnregister.Add(InComponent);
+			FScopeLock Lock(&DelayedComponentsRegistrationLock);
+			DelayedComponentsToUnregister.Add(InComponent);
 			return;
 		}
 
@@ -514,14 +529,6 @@ void FPCGActorAndComponentMapping::UnregisterPCGComponent(UPCGComponent* InCompo
 
 	UnregisterPartitionedPCGComponent(InComponent);
 	UnregisterNonPartitionedPCGComponent(InComponent);
-
-#if WITH_EDITOR
-	FScopeLock Lock(&DelayedComponentToUnregisterLock);
-	if (DelayedComponentToUnregister.Contains(InComponent))
-	{
-		DelayedComponentToUnregister.Remove(InComponent);
-	}
-#endif
 }
 
 void FPCGActorAndComponentMapping::UnregisterPartitionedPCGComponent(UPCGComponent* InComponent)
diff --git a/Engine/Plugins/PCG/Source/PCG/Private/PCGComponent.cpp b/Engine/Plugins/PCG/Source/PCG/Private/PCGComponent.cpp
--- a/Engine/Plugins/PCG/Source/PCG/Private/PCGComponent.cpp
+++ b/Engine/Plugins/PCG/Source/PCG/Private/PCGComponent.cpp
@@ -2011,12 +2011,8 @@ void UPCGComponent::OnRegister()
 	{
 		if (UWorld* World = GetWorld())
 		{
-			// We won't be able to spawn any actors if we are currently running a construction script.
-			if (!World->bIsRunningConstructionScript)
-			{
-				ensure(!bGenerated || LastGeneratedBounds.IsValid);
-				GetSubsystem()->RegisterOrUpdatePCGComponent(this, bGenerated);
-			}
+			ensure(!bGenerated || LastGeneratedBounds.IsValid);
+			GetSubsystem()->RegisterOrUpdatePCGComponent(this, bGenerated);
 		}
 	}
 #endif //WITH_EDITOR
diff --git a/Engine/Plugins/PCG/Source/PCG/Public/PCGActorAndComponentMapping.h b/Engine/Plugins/PCG/Source/PCG/Public/PCGActorAndComponentMapping.h
--- a/Engine/Plugins/PCG/Source/PCG/Public/PCGActorAndComponentMapping.h
+++ b/Engine/Plugins/PCG/Source/PCG/Public/PCGActorAndComponentMapping.h
@@ -117,6 +117,7 @@ private:
 	void ProcessDelayedEvents();
 	bool ShouldDelayActor(AActor* InActor) const;
 #endif
+	bool RegisterOrUpdatePCGComponent(UPCGComponent* InComponent, bool bDoActorMapping, bool bForce);
 
 	// This class is only meant to be used as part of the PCG Subsytem and owned by it.
 	// So we put constructors private.
@@ -238,9 +239,10 @@ private:
 	FPCGComponentOctreeAndMap NonPartitionedOctree;
 
 #if WITH_EDITOR
-	/** Components to be unregister at the next frame. cf. UnregisterComponent for a better understanding on why it is needed. */
-	TSet<UPCGComponent*> DelayedComponentToUnregister;
-	mutable FCriticalSection DelayedComponentToUnregisterLock;
+	/** Components to be registered/unregistered on the next frame. see RegisterOrUpdatePCGComponent/UnregisterPCGComponent for a better understanding on why it is needed. */
+	mutable FCriticalSection DelayedComponentsRegistrationLock;
+	TSet<UPCGComponent*> DelayedComponentsToUnregister;
+	TSet<UPCGComponent*> DelayedComponentsToRegister;
 	
 	// Tracking actors
 	/** Keep a mapping between tracked keys and the components that track them, and the tracking needs to be culled.*/
