[PCG] Tweak copy points handling of no input data items to fix errors/ensures.

* Don't attempt to execute when there are no data items.
* Only flag errors about missing attributes to user if there are data items arriving on inputs.

#rnx

diff --git a/Engine/Plugins/PCG/Source/PCG/Private/Elements/Metadata/PCGMetadataPartitionKernel.cpp b/Engine/Plugins/PCG/Source/PCG/Private/Elements/Metadata/PCGMetadataPartitionKernel.cpp
--- a/Engine/Plugins/PCG/Source/PCG/Private/Elements/Metadata/PCGMetadataPartitionKernel.cpp
+++ b/Engine/Plugins/PCG/Source/PCG/Private/Elements/Metadata/PCGMetadataPartitionKernel.cpp
@@ -60,7 +60,7 @@ bool UPCGMetadataPartitionKernel::IsKernelValid(FPCGContext* InContext, bool bQu
 
 	if (Context && Context->DataBinding)
 	{
-		FName AttributeName = Settings->PartitionAttributeSelectors[0].GetAttributeName();
+		const FName AttributeName = Settings->PartitionAttributeSelectors[0].GetAttributeName();
 
 		FPCGDataCollectionDesc InputDesc;
 		ensure(Context->DataBinding->GetKernelPinDataDesc_GameThread(this, PCGPinConstants::DefaultInputLabel, /*bIsInput=*/true, InputDesc));
@@ -71,14 +71,21 @@ bool UPCGMetadataPartitionKernel::IsKernelValid(FPCGContext* InContext, bool bQu
 
 		if (!bAttributePresent)
 		{
-			PCG_KERNEL_VALIDATION_ERR(InContext, Settings, bQuiet, LOCTEXT("PartitionAttributeMissing", "Partition attribute was not present in GPU data."));
+			if (!InputDesc.DataDescs.IsEmpty())
+			{
+				PCG_KERNEL_VALIDATION_ERR(InContext, Settings, bQuiet, FText::Format(
+					LOCTEXT("PartitionAttributeMissing", "Partition attribute was not present in GPU data."),
+					FText::FromName(AttributeName)));
+			}
+
 			return false;
 		}
 
 		if (bConflictingTypesInData)
 		{
-			PCG_KERNEL_VALIDATION_ERR(InContext, Settings, bQuiet,
-				FText::Format(LOCTEXT("PartitionAttributeTypeConflict", "Attribute '{0}' encountered with multiple different types in input data."), FText::FromName(AttributeName)));
+			PCG_KERNEL_VALIDATION_ERR(InContext, Settings, bQuiet, FText::Format(
+				LOCTEXT("PartitionAttributeTypeConflict", "Attribute '{0}' encountered with multiple different types in input data."),
+				FText::FromName(AttributeName)));
 			return false;
 		}
 
diff --git a/Engine/Plugins/PCG/Source/PCG/Private/Elements/PCGCopyPointsKernelShared.cpp b/Engine/Plugins/PCG/Source/PCG/Private/Elements/PCGCopyPointsKernelShared.cpp
--- a/Engine/Plugins/PCG/Source/PCG/Private/Elements/PCGCopyPointsKernelShared.cpp
+++ b/Engine/Plugins/PCG/Source/PCG/Private/Elements/PCGCopyPointsKernelShared.cpp
@@ -44,7 +44,8 @@ bool PCGCopyPointsKernel::IsKernelValid(const UPCGComputeKernel* InKernel, const
 					}
 				}
 
-				return true;
+				// Valid for execution if we have some data to process.
+				return !PinDataDesc.DataDescs.IsEmpty();
 			};
 
 			if (!ValidateAttributeExists(PCGCopyPointsConstants::SourcePointsLabel) || !ValidateAttributeExists(PCGCopyPointsConstants::TargetPointsLabel))
