[PCG] Improvement on PCG Object Overrides

* Export functions to be used outside of the PCG plugin
* Add option to get all objects to load from override.

#rnx

diff --git a/Engine/Plugins/PCG/Source/PCG/Private/Metadata/PCGObjectPropertyOverride.cpp b/Engine/Plugins/PCG/Source/PCG/Private/Metadata/PCGObjectPropertyOverride.cpp
--- a/Engine/Plugins/PCG/Source/PCG/Private/Metadata/PCGObjectPropertyOverride.cpp
+++ b/Engine/Plugins/PCG/Source/PCG/Private/Metadata/PCGObjectPropertyOverride.cpp
@@ -93,7 +93,12 @@ void FPCGObjectSingleOverride::Initialize(const FPCGAttributePropertySelector& I
 	const TArray<FString>& ExtraNames = OutputSelector.GetExtraNames();
 	if (ExtraNames.IsEmpty())
 	{
-		ObjectOverrideOutputAccessor = PCGAttributeAccessorHelpers::CreatePropertyAccessor(FName(OutputProperty), TemplateClass);
+		if (const FProperty* Property = TemplateClass->FindPropertyByName(FName(OutputProperty)))
+		{
+			ObjectOverrideOutputAccessor = PCGAttributeAccessorHelpers::CreatePropertyAccessor(Property);
+			TArray<const FStructProperty*> EncounteredStructProps;
+			bWillNeedLoading = Property->ContainsObjectReference(EncounteredStructProps, EPropertyObjectReferenceType::Strong);
+		}
 	}
 	else
 	{
@@ -105,7 +110,14 @@ void FPCGObjectSingleOverride::Initialize(const FPCGAttributePropertySelector& I
 			PropertyNames.Add(FName(Name));
 		}
 
-		ObjectOverrideOutputAccessor = PCGAttributeAccessorHelpers::CreatePropertyChainAccessor(PropertyNames, TemplateClass);
+		TArray<const FProperty*> PropertyChain;
+		PCGAttributeAccessorHelpers::GetPropertyChain(PropertyNames, TemplateClass, PropertyChain);
+		if (!PropertyChain.IsEmpty())
+		{
+			ObjectOverrideOutputAccessor = PCGAttributeAccessorHelpers::CreatePropertyChainAccessor(std::move(PropertyChain));
+			TArray<const FStructProperty*> EncounteredStructProps;
+			bWillNeedLoading = PropertyChain.Last()->ContainsObjectReference(EncounteredStructProps, EPropertyObjectReferenceType::Strong);
+		}
 	}
 
 	if (!ObjectOverrideOutputAccessor.IsValid())
@@ -150,4 +162,17 @@ bool FPCGObjectSingleOverride::Apply(int32 InputKeyIndex, IPCGAttributeAccessorK
 	return Invoke(ObjectOverrideFunction, this, InputKeyIndex, OutputKey);
 }
 
+void FPCGObjectSingleOverride::GatherAllOverridesToLoad(TArray<FSoftObjectPath>& OutObjectsToLoad) const
+{
+	if (!bWillNeedLoading || !IsValid())
+	{
+		return;
+	}
+
+	PCGMetadataElementCommon::ApplyOnAccessor<FSoftObjectPath>(*InputKeys, *ObjectOverrideInputAccessor, [&OutObjectsToLoad](FSoftObjectPath& Value, int32 Index)
+	{
+		OutObjectsToLoad.AddUnique(std::move(Value));
+	});
+}
+
 #undef LOCTEXT_NAMESPACE
diff --git a/Engine/Plugins/PCG/Source/PCG/Public/Metadata/Accessors/PCGAttributeAccessorHelpers.h b/Engine/Plugins/PCG/Source/PCG/Public/Metadata/Accessors/PCGAttributeAccessorHelpers.h
--- a/Engine/Plugins/PCG/Source/PCG/Public/Metadata/Accessors/PCGAttributeAccessorHelpers.h
+++ b/Engine/Plugins/PCG/Source/PCG/Public/Metadata/Accessors/PCGAttributeAccessorHelpers.h
@@ -41,6 +41,8 @@ namespace PCGAttributeAccessorHelpers
 	PCG_API TUniquePtr<IPCGAttributeAccessor> CreateChainAccessor(TUniquePtr<IPCGAttributeAccessor> InAccessor, FName Name, bool& bOutSuccess);
 	PCG_API TUniquePtr<IPCGAttributeAccessor> CreateChainAccessor(TUniquePtr<IPCGAttributeAccessor> InAccessor, const FPCGAttributePropertySelector& InSelector, bool& bOutSuccess, bool bQuiet = false);
 
+	PCG_API bool GetPropertyChain(const TArray<FName>& InPropertyNames, const UStruct* InStruct, TArray<const FProperty*>& OutProperties);
+
 	struct AccessorParamResult
 	{
 		FName AttributeName = NAME_None;
diff --git a/Engine/Plugins/PCG/Source/PCG/Public/Metadata/PCGObjectPropertyOverride.h b/Engine/Plugins/PCG/Source/PCG/Public/Metadata/PCGObjectPropertyOverride.h
--- a/Engine/Plugins/PCG/Source/PCG/Public/Metadata/PCGObjectPropertyOverride.h
+++ b/Engine/Plugins/PCG/Source/PCG/Public/Metadata/PCGObjectPropertyOverride.h
@@ -5,6 +5,7 @@
 #include "PCGData.h"
 #include "PCGPin.h"
 #include "Accessors/PCGAttributeAccessorHelpers.h"
+#include "Elements/Metadata/PCGMetadataElementCommon.h"
 #include "Metadata/PCGAttributePropertySelector.h"
 #include "Metadata/PCGMetadata.h"
 #include "Metadata/Accessors/IPCGAttributeAccessor.h"
@@ -48,12 +49,12 @@ struct FPCGObjectPropertyOverrideDescription
 namespace PCGObjectPropertyOverrideHelpers
 {
 	/** Create an advanced ParamData pin for capturing property overrides. */
-	FPCGPinProperties CreateObjectPropertiesOverridePin(FName Label, const FText& Tooltip);
+	PCG_API FPCGPinProperties CreateObjectPropertiesOverridePin(FName Label, const FText& Tooltip);
 
 	/** Apply property overrides to the TargetObject directly from the ObjectPropertiesOverride pin. Use CreateObjectPropertiesOverridePin(). */
-	void ApplyOverridesFromParams(const TArray<FPCGObjectPropertyOverrideDescription>& InObjectPropertyOverrideDescriptions, UObject* TargetObject, FName OverridesPinLabel, FPCGContext* Context);
+	PCG_API void ApplyOverridesFromParams(const TArray<FPCGObjectPropertyOverrideDescription>& InObjectPropertyOverrideDescriptions, UObject* TargetObject, FName OverridesPinLabel, FPCGContext* Context);
 
-	void ApplyOverrides(const TArray<FPCGObjectPropertyOverrideDescription>& InObjectPropertyOverrideDescriptions, const TArray<TPair<UObject*, int32>>& TargetObjectAndIndex, FName OverridesPinLabel, int32 InputDataIndex, FPCGContext* Context);
+	PCG_API void ApplyOverrides(const TArray<FPCGObjectPropertyOverrideDescription>& InObjectPropertyOverrideDescriptions, const TArray<TPair<UObject*, int32>>& TargetObjectAndIndex, FName OverridesPinLabel, int32 InputDataIndex, FPCGContext* Context);
 }
 
 /**
@@ -65,19 +66,24 @@ namespace PCGObjectPropertyOverrideHelpers
 struct FPCGObjectSingleOverride
 {
 	/** Initialize the single object override. Call before using Apply(InputKeyIndex, OutputKey). */
-	void Initialize(const FPCGAttributePropertySelector& InputSelector, const FString& OutputProperty, const UStruct* TemplateClass, const UPCGData* SourceData, FPCGContext* Context);
+	PCG_API void Initialize(const FPCGAttributePropertySelector& InputSelector, const FString& OutputProperty, const UStruct* TemplateClass, const UPCGData* SourceData, FPCGContext* Context);
 
 	/** Returns true if initialization succeeded in creating the accessors and accessor keys. */
-	bool IsValid() const;
+	PCG_API bool IsValid() const;
 
 	/** Applies a single property override to the object by reading from the InputAccessor at the given KeyIndex, and writing to the OutputKey which represents the object property. */
-	bool Apply(int32 InputKeyIndex, IPCGAttributeAccessorKeys& OutputKey);
+	PCG_API bool Apply(int32 InputKeyIndex, IPCGAttributeAccessorKeys& OutputKey);
+
+	/** Utility function that gather all the overrides and fill the out array if the override needs loading. */
+	PCG_API void GatherAllOverridesToLoad(TArray<FSoftObjectPath>& OutObjectsToLoad) const;
 
 private:
 	TUniquePtr<const IPCGAttributeAccessorKeys> InputKeys;
 	TUniquePtr<const IPCGAttributeAccessor> ObjectOverrideInputAccessor;
 	TUniquePtr<IPCGAttributeAccessor> ObjectOverrideOutputAccessor;
 
+	bool bWillNeedLoading = false;
+
 	// InputKeyIndex, OutputKeys
 	using ApplyOverrideFunction = bool(FPCGObjectSingleOverride::*)(int32, IPCGAttributeAccessorKeys&);
 	ApplyOverrideFunction ObjectOverrideFunction = nullptr;
@@ -114,6 +120,18 @@ struct FPCGObjectOverrides
 	FPCGObjectOverrides(T* TemplateObject) : OutputKey(TemplateObject)
 	{}
 
+	constexpr const UStruct* GetClassOfTemplate() const
+	{
+		if constexpr (TIsDerivedFrom<T, UObject>::Value)
+		{
+			return T::StaticClass();
+		}
+		else
+		{
+			return T::StaticStruct();
+		}
+	}
+
 	/** Initialize the object overrides. Call before using Apply(InputKeyIndex). */
 	void Initialize(const TArray<FPCGObjectPropertyOverrideDescription>& OverrideDescriptions, T* TemplateObject, const UPCGData* SourceData, FPCGContext* Context)
 	{
@@ -125,15 +143,9 @@ struct FPCGObjectOverrides
 			return;
 		}
 
-		const UStruct* ClassObject = nullptr;
-		if constexpr (TIsDerivedFrom<T, UObject>::Value)
-		{
-			ClassObject = TemplateObject->GetClass();
-		}
-		else
-		{
-			ClassObject = TemplateObject->StaticStruct();
-		}
+		OutputKey = FPCGAttributeAccessorKeysSingleObjectPtr<T>(TemplateObject);
+
+		const UStruct* ClassObject = GetClassOfTemplate();
 
 		ObjectSingleOverrides.Empty(OverrideDescriptions.Num());
 
@@ -184,9 +196,27 @@ struct FPCGObjectOverrides
 		return bAllSucceeded;
 	}
 
+	/** Update the template object. Only if it was already initialized */
+	void UpdateTemplateObject(T* InTemplateObject)
+	{
+		if (IsValid())
+		{
+			OutputKey = FPCGAttributeAccessorKeysSingleObjectPtr<T>(InTemplateObject);
+		}
+	}
+
 	/** Returns true if we have any override to apply */
 	bool IsValid() const { return bInitialized && !ObjectSingleOverrides.IsEmpty(); }
 
+	/** Utility function that scans all the overrides and fill the out array if the object needs loading. */
+	void GatherAllOverridesToLoad(TArray<FSoftObjectPath>& OutObjectsToLoad) const
+	{
+		for (const FPCGObjectSingleOverride& SingleOverride : ObjectSingleOverrides)
+		{
+			SingleOverride.GatherAllOverridesToLoad(OutObjectsToLoad);
+		}
+	}
+
 private:
 	FPCGAttributeAccessorKeysSingleObjectPtr<T> OutputKey;
 	TArray<FPCGObjectSingleOverride> ObjectSingleOverrides;
