[PCG] Fix "wait until next frame"s that were actually busy waits due to PCGHelpers::ExecuteOnMainThread() executing the functor immediately when on game thread.

#rnx

diff --git a/Engine/Plugins/PCG/Source/PCG/Private/Compute/DataInterfaces/PCGDataCollectionDataInterface.cpp b/Engine/Plugins/PCG/Source/PCG/Private/Compute/DataInterfaces/PCGDataCollectionDataInterface.cpp
--- a/Engine/Plugins/PCG/Source/PCG/Private/Compute/DataInterfaces/PCGDataCollectionDataInterface.cpp
+++ b/Engine/Plugins/PCG/Source/PCG/Private/Compute/DataInterfaces/PCGDataCollectionDataInterface.cpp
@@ -543,7 +543,7 @@ FComputeDataProviderRenderProxy* UPCGDataCollectionDataProvider::GetRenderProxy(
 
 			This->bReadbackComplete = true;
 
-			PCGHelpers::ExecuteOnGameThread(UE_SOURCE_LOCATION, [ThisWeak=ThisWeak, OutputPinLabelAlias=This->OutputPinLabelAlias]()
+			ExecuteOnGameThread(UE_SOURCE_LOCATION, [ThisWeak=ThisWeak, OutputPinLabelAlias=This->OutputPinLabelAlias]()
 			{
 				TStrongObjectPtr<UPCGDataCollectionDataProvider> This = ThisWeak.Pin();
 				if (This)
@@ -682,7 +682,7 @@ void FPCGDataCollectionDataProviderProxy::AllocateResources(FRDGBuilder& GraphBu
 		const TRefCountPtr<FRDGPooledBuffer> ExportedPooledBuffer = GraphBuilder.ConvertToExternalBuffer(Buffer);
 		GraphBuilder.SetBufferAccessFinal(Buffer, ERHIAccess::CopySrc);
 
-		PCGHelpers::ExecuteOnGameThread(UE_SOURCE_LOCATION, [ExportedPooledBuffer, DataProviderWeakPtr=DataProviderWeakPtr, PinDesc=PinDesc, OutputPinLabel=OutputPinLabel, SizeBytes=SizeBytes, ExportMode=ExportMode, OutputPinLabelAlias=OutputPinLabelAlias]()
+		ExecuteOnGameThread(UE_SOURCE_LOCATION, [ExportedPooledBuffer, DataProviderWeakPtr=DataProviderWeakPtr, PinDesc=PinDesc, OutputPinLabel=OutputPinLabel, SizeBytes=SizeBytes, ExportMode=ExportMode, OutputPinLabelAlias=OutputPinLabelAlias]()
 		{
 			// Obtain objects. No ensures added because a graph cancellation could feasibly destroy some or all of these.
 			UPCGDataCollectionDataProvider* DataProvider = DataProviderWeakPtr.Get();
diff --git a/Engine/Plugins/PCG/Source/PCG/Private/Elements/PCGGenerateGrassMaps.cpp b/Engine/Plugins/PCG/Source/PCG/Private/Elements/PCGGenerateGrassMaps.cpp
--- a/Engine/Plugins/PCG/Source/PCG/Private/Elements/PCGGenerateGrassMaps.cpp
+++ b/Engine/Plugins/PCG/Source/PCG/Private/Elements/PCGGenerateGrassMaps.cpp
@@ -425,7 +425,7 @@ bool FPCGGenerateGrassMapsElement::ExecuteInternal(FPCGContext* InContext) const
 			GraphBuilder.Execute();
 
 			// Pass exported buffer back to game thread and wake up this element.
-			PCGHelpers::ExecuteOnGameThread(UE_SOURCE_LOCATION, [ContextHandle, GrassMapExported]()
+			ExecuteOnGameThread(UE_SOURCE_LOCATION, [ContextHandle, GrassMapExported]()
 			{
 				if (FPCGGenerateGrassMapsContext* Context = FPCGContext::GetContextFromHandle<FPCGGenerateGrassMapsContext>(ContextHandle))
 				{
diff --git a/Engine/Plugins/PCG/Source/PCG/Private/PCGElement.cpp b/Engine/Plugins/PCG/Source/PCG/Private/PCGElement.cpp
--- a/Engine/Plugins/PCG/Source/PCG/Private/PCGElement.cpp
+++ b/Engine/Plugins/PCG/Source/PCG/Private/PCGElement.cpp
@@ -417,7 +417,7 @@ bool IPCGElement::ConvertInputsIfNeeded(FPCGContext* Context) const
 			Context->bIsPaused = true;
 
 			// Not ready to execute and unlikely to be in the very short term, sleep until next frame.
-			PCGHelpers::ExecuteOnGameThread(UE_SOURCE_LOCATION, [ContextHandle = Context->GetOrCreateHandle()]()
+			ExecuteOnGameThread(UE_SOURCE_LOCATION, [ContextHandle = Context->GetOrCreateHandle()]()
 			{
 				if (TSharedPtr<FPCGContextHandle> SharedHandle = ContextHandle.Pin())
 				{
diff --git a/Engine/Plugins/PCG/Source/PCG/Public/Helpers/PCGHelpers.h b/Engine/Plugins/PCG/Source/PCG/Public/Helpers/PCGHelpers.h
--- a/Engine/Plugins/PCG/Source/PCG/Public/Helpers/PCGHelpers.h
+++ b/Engine/Plugins/PCG/Source/PCG/Public/Helpers/PCGHelpers.h
@@ -156,6 +156,7 @@ namespace PCGHelpers
 		Array = std::move(TempArray);
 	}
 
+	/** Execute given functor on game thread. If called from game thread, executes immediately. */
 	template<typename FunctorType>
 	void ExecuteOnGameThread(const TCHAR* DebugName, FunctorType&& Functor)
 	{
