Fixed retrieval error in data on a few nodes (esp. tags from data) where it would connect but not work with attribute sets.

#rnx

diff --git a/Engine/Plugins/PCG/Source/PCG/Private/Elements/PCGConvertToAttributeSet.cpp b/Engine/Plugins/PCG/Source/PCG/Private/Elements/PCGConvertToAttributeSet.cpp
--- a/Engine/Plugins/PCG/Source/PCG/Private/Elements/PCGConvertToAttributeSet.cpp
+++ b/Engine/Plugins/PCG/Source/PCG/Private/Elements/PCGConvertToAttributeSet.cpp
@@ -42,7 +42,7 @@ bool FPCGConvertToAttributeSetElement::ExecuteInternal(FPCGContext* Context) con
 	TRACE_CPUPROFILER_EVENT_SCOPE(FPCGConvertToAttributeSetElement::Execute);
 	check(Context);
 
-	TArray<FPCGTaggedData> Inputs = Context->InputData.GetInputs();
+	TArray<FPCGTaggedData> Inputs = Context->InputData.GetInputsByPin(PCGPinConstants::DefaultInputLabel);
 	TArray<FPCGTaggedData>& Outputs = Context->OutputData.TaggedData;
 
 	for (const FPCGTaggedData& Input : Inputs)
@@ -104,7 +104,7 @@ bool FPCGTagsToAttributeSetElement::ExecuteInternal(FPCGContext* Context) const
 	TRACE_CPUPROFILER_EVENT_SCOPE(FPCGTagsToAttributeSetElement::Execute);
 	check(Context);
 
-	TArray<FPCGTaggedData> Inputs = Context->InputData.GetInputs();
+	TArray<FPCGTaggedData> Inputs = Context->InputData.GetInputsByPin(PCGPinConstants::DefaultInputLabel);
 	TArray<FPCGTaggedData>& Outputs = Context->OutputData.TaggedData;
 
 	for (const FPCGTaggedData& Input : Inputs)
diff --git a/Engine/Plugins/PCG/Source/PCG/Private/Elements/PCGCreateSpline.cpp b/Engine/Plugins/PCG/Source/PCG/Private/Elements/PCGCreateSpline.cpp
--- a/Engine/Plugins/PCG/Source/PCG/Private/Elements/PCGCreateSpline.cpp
+++ b/Engine/Plugins/PCG/Source/PCG/Private/Elements/PCGCreateSpline.cpp
@@ -61,7 +61,7 @@ bool FPCGCreateSplineElement::ExecuteInternal(FPCGContext* Context) const
 	const UPCGCreateSplineSettings* Settings = Context->GetInputSettings<UPCGCreateSplineSettings>();
 	check(Settings);
 
-	TArray<FPCGTaggedData> Inputs = Context->InputData.GetInputs();
+	TArray<FPCGTaggedData> Inputs = Context->InputData.GetInputsByPin(PCGPinConstants::DefaultInputLabel);
 	TArray<FPCGTaggedData>& Outputs = Context->OutputData.TaggedData;
 
 	EPCGCreateSplineMode Mode = Settings->Mode;
diff --git a/Engine/Plugins/PCG/Source/PCG/Private/Elements/PCGDebugElement.cpp b/Engine/Plugins/PCG/Source/PCG/Private/Elements/PCGDebugElement.cpp
--- a/Engine/Plugins/PCG/Source/PCG/Private/Elements/PCGDebugElement.cpp
+++ b/Engine/Plugins/PCG/Source/PCG/Private/Elements/PCGDebugElement.cpp
@@ -43,7 +43,7 @@ bool FPCGDebugElement::ExecuteInternal(FPCGContext* Context) const
 	}
 
 	const FPCGDataVisualizationRegistry& DataVisRegistry = FPCGModule::GetConstPCGDataVisualizationRegistry();
-	TArray<FPCGTaggedData> Inputs = Context->InputData.GetInputs();
+	TArray<FPCGTaggedData> Inputs = Context->InputData.GetInputsByPin(PCGPinConstants::DefaultInputLabel);
 
 	for (const FPCGTaggedData& Input : Inputs)
 	{
diff --git a/Engine/Plugins/PCG/Source/PCG/Private/Elements/PCGDensityFilter.cpp b/Engine/Plugins/PCG/Source/PCG/Private/Elements/PCGDensityFilter.cpp
--- a/Engine/Plugins/PCG/Source/PCG/Private/Elements/PCGDensityFilter.cpp
+++ b/Engine/Plugins/PCG/Source/PCG/Private/Elements/PCGDensityFilter.cpp
@@ -24,7 +24,7 @@ bool FPCGDensityFilterElement::ExecuteInternal(FPCGContext* Context) const
 	const UPCGDensityFilterSettings* Settings = Context->GetInputSettings<UPCGDensityFilterSettings>();
 	check(Settings);
 
-	TArray<FPCGTaggedData> Inputs = Context->InputData.GetInputs();
+	TArray<FPCGTaggedData> Inputs = Context->InputData.GetInputsByPin(PCGPinConstants::DefaultInputLabel);
 	TArray<FPCGTaggedData>& Outputs = Context->OutputData.TaggedData;
 
 	const bool bInvertFilter = Settings->bInvertFilter;
diff --git a/Engine/Plugins/PCG/Source/PCG/Private/Elements/PCGDuplicatePoint.cpp b/Engine/Plugins/PCG/Source/PCG/Private/Elements/PCGDuplicatePoint.cpp
--- a/Engine/Plugins/PCG/Source/PCG/Private/Elements/PCGDuplicatePoint.cpp
+++ b/Engine/Plugins/PCG/Source/PCG/Private/Elements/PCGDuplicatePoint.cpp
@@ -24,7 +24,7 @@ bool FPCGDuplicatePointElement::ExecuteInternal(FPCGContext* Context) const
 	const UPCGDuplicatePointSettings* Settings = Context->GetInputSettings<UPCGDuplicatePointSettings>();
 	check(Settings);
 
-	const TArray<FPCGTaggedData> Inputs = Context->InputData.GetInputs();
+	const TArray<FPCGTaggedData> Inputs = Context->InputData.GetInputsByPin(PCGPinConstants::DefaultInputLabel);
 	TArray<FPCGTaggedData>& Outputs = Context->OutputData.TaggedData;
 
 	if (Settings->Iterations <= 0)
diff --git a/Engine/Plugins/PCG/Source/PCG/Private/Elements/PCGInnerIntersectionElement.cpp b/Engine/Plugins/PCG/Source/PCG/Private/Elements/PCGInnerIntersectionElement.cpp
--- a/Engine/Plugins/PCG/Source/PCG/Private/Elements/PCGInnerIntersectionElement.cpp
+++ b/Engine/Plugins/PCG/Source/PCG/Private/Elements/PCGInnerIntersectionElement.cpp
@@ -53,7 +53,7 @@ bool FPCGInnerIntersectionElement::ExecuteInternal(FPCGContext* Context) const
 	const UPCGInnerIntersectionSettings* Settings = Context->GetInputSettings<UPCGInnerIntersectionSettings>();
 	check(Settings);
 
-	TArray<FPCGTaggedData> Inputs = Context->InputData.GetInputs();
+	TArray<FPCGTaggedData> Inputs = Context->InputData.GetInputsByPin(PCGPinConstants::DefaultInputLabel);
 
 	const EPCGIntersectionDensityFunction DensityFunction = Settings->DensityFunction;
 	const bool bKeepZeroDensityPoints = Settings->bKeepZeroDensityPoints;
diff --git a/Engine/Plugins/PCG/Source/PCG/Private/Elements/PCGMakeConcreteElement.cpp b/Engine/Plugins/PCG/Source/PCG/Private/Elements/PCGMakeConcreteElement.cpp
--- a/Engine/Plugins/PCG/Source/PCG/Private/Elements/PCGMakeConcreteElement.cpp
+++ b/Engine/Plugins/PCG/Source/PCG/Private/Elements/PCGMakeConcreteElement.cpp
@@ -40,7 +40,7 @@ bool FPCGMakeConcreteElement::ExecuteInternal(FPCGContext* InContext) const
 {
 	TRACE_CPUPROFILER_EVENT_SCOPE(FPCGMakeConcreteElement::Execute);
 
-	TArray<FPCGTaggedData> SpatialInputs = InContext->InputData.GetInputs();
+	TArray<FPCGTaggedData> SpatialInputs = InContext->InputData.GetAllSpatialInputs();
 	for (const FPCGTaggedData& Input : SpatialInputs)
 	{
 		const UPCGSpatialData* InputSpatialData = Cast<UPCGSpatialData>(Input.Data);
diff --git a/Engine/Plugins/PCG/Source/PCG/Private/Elements/PCGOuterIntersectionElement.cpp b/Engine/Plugins/PCG/Source/PCG/Private/Elements/PCGOuterIntersectionElement.cpp
--- a/Engine/Plugins/PCG/Source/PCG/Private/Elements/PCGOuterIntersectionElement.cpp
+++ b/Engine/Plugins/PCG/Source/PCG/Private/Elements/PCGOuterIntersectionElement.cpp
@@ -75,7 +75,6 @@ bool FPCGOuterIntersectionElement::ExecuteInternal(FPCGContext* Context) const
 	const UPCGOuterIntersectionSettings* Settings = Context->GetInputSettings<UPCGOuterIntersectionSettings>();
 	check(Settings);
 	
-	const TArray<FPCGTaggedData> Inputs = Context->InputData.GetInputs();
 	const TArray<TObjectPtr<UPCGPin>>& InputPins = Context->Node->GetInputPins();
 	TArray<FPCGTaggedData>& Outputs = Context->OutputData.TaggedData;
 
diff --git a/Engine/Plugins/PCG/Source/PCG/Private/Elements/PCGPointMatchAndSet.cpp b/Engine/Plugins/PCG/Source/PCG/Private/Elements/PCGPointMatchAndSet.cpp
--- a/Engine/Plugins/PCG/Source/PCG/Private/Elements/PCGPointMatchAndSet.cpp
+++ b/Engine/Plugins/PCG/Source/PCG/Private/Elements/PCGPointMatchAndSet.cpp
@@ -197,7 +197,7 @@ bool FPCGPointMatchAndSetElement::ExecuteInternal(FPCGContext* Context) const
 		return true;
 	}
 
-	TArray<FPCGTaggedData> Inputs = Context->InputData.GetInputs();
+	TArray<FPCGTaggedData> Inputs = Context->InputData.GetInputsByPin(PCGPinConstants::DefaultInputLabel);
 	TArray<FPCGTaggedData>& Outputs = Context->OutputData.TaggedData;
 
 	for (const FPCGTaggedData& Input : Inputs)
diff --git a/Engine/Plugins/PCG/Source/PCG/Private/Elements/PCGPointNeighborhood.cpp b/Engine/Plugins/PCG/Source/PCG/Private/Elements/PCGPointNeighborhood.cpp
--- a/Engine/Plugins/PCG/Source/PCG/Private/Elements/PCGPointNeighborhood.cpp
+++ b/Engine/Plugins/PCG/Source/PCG/Private/Elements/PCGPointNeighborhood.cpp
@@ -50,7 +50,7 @@ bool FPCGPointNeighborhoodElement::ExecuteInternal(FPCGContext* Context) const
 	const UPCGPointNeighborhoodSettings* Settings = Context->GetInputSettings<UPCGPointNeighborhoodSettings>();
 	check(Settings);
 
-	TArray<FPCGTaggedData> Inputs = Context->InputData.GetInputs();
+	TArray<FPCGTaggedData> Inputs = Context->InputData.GetInputsByPin(PCGPinConstants::DefaultInputLabel);
 	TArray<FPCGTaggedData>& Outputs = Context->OutputData.TaggedData;
 
 	const double SearchDistance = Settings->SearchDistance;
diff --git a/Engine/Plugins/PCG/Source/PCG/Private/Elements/PCGSelectPoints.cpp b/Engine/Plugins/PCG/Source/PCG/Private/Elements/PCGSelectPoints.cpp
--- a/Engine/Plugins/PCG/Source/PCG/Private/Elements/PCGSelectPoints.cpp
+++ b/Engine/Plugins/PCG/Source/PCG/Private/Elements/PCGSelectPoints.cpp
@@ -33,8 +33,7 @@ bool FPCGSelectPointsElement::ExecuteInternal(FPCGContext* Context) const
 	const UPCGSelectPointsSettings* Settings = Context->GetInputSettings<UPCGSelectPointsSettings>();
 	check(Settings);
 
-	TArray<FPCGTaggedData> Inputs = Context->InputData.GetInputs();
-
+	TArray<FPCGTaggedData> Inputs = Context->InputData.GetInputsByPin(PCGPinConstants::DefaultInputLabel);
 	TArray<FPCGTaggedData>& Outputs = Context->OutputData.TaggedData;
 
 	const float Ratio = Settings->Ratio;
diff --git a/Engine/Plugins/PCG/Source/PCG/Private/Elements/PCGSelfPruning.cpp b/Engine/Plugins/PCG/Source/PCG/Private/Elements/PCGSelfPruning.cpp
--- a/Engine/Plugins/PCG/Source/PCG/Private/Elements/PCGSelfPruning.cpp
+++ b/Engine/Plugins/PCG/Source/PCG/Private/Elements/PCGSelfPruning.cpp
@@ -452,7 +452,7 @@ namespace PCGSelfPruningElement
 			return;
 		}
 
-		TArray<FPCGTaggedData> Inputs = Context->InputData.GetInputs();
+		TArray<FPCGTaggedData> Inputs = Context->InputData.GetInputsByPin(PCGPinConstants::DefaultInputLabel);
 		for (const FPCGTaggedData& Input : Inputs)
 		{
 			const UPCGSpatialData* SpatialInput = Cast<UPCGSpatialData>(Input.Data);
diff --git a/Engine/Plugins/PCG/Source/PCG/Private/Elements/PCGSkinnedMeshSpawner.cpp b/Engine/Plugins/PCG/Source/PCG/Private/Elements/PCGSkinnedMeshSpawner.cpp
--- a/Engine/Plugins/PCG/Source/PCG/Private/Elements/PCGSkinnedMeshSpawner.cpp
+++ b/Engine/Plugins/PCG/Source/PCG/Private/Elements/PCGSkinnedMeshSpawner.cpp
@@ -147,7 +147,7 @@ bool FPCGSkinnedMeshSpawnerElement::PrepareDataInternal(FPCGContext* InContext)
 	}
 
 	// perform mesh selection
-	TArray<FPCGTaggedData> Inputs = Context->InputData.GetInputs();
+	TArray<FPCGTaggedData> Inputs = Context->InputData.GetInputsByPin(PCGPinConstants::DefaultInputLabel);
 	TArray<FPCGTaggedData>& Outputs = Context->OutputData.TaggedData;
 
 	while(Context->CurrentInputIndex < Inputs.Num())
diff --git a/Engine/Plugins/PCG/Source/PCG/Private/Elements/PCGSpawnActor.cpp b/Engine/Plugins/PCG/Source/PCG/Private/Elements/PCGSpawnActor.cpp
--- a/Engine/Plugins/PCG/Source/PCG/Private/Elements/PCGSpawnActor.cpp
+++ b/Engine/Plugins/PCG/Source/PCG/Private/Elements/PCGSpawnActor.cpp
@@ -536,7 +536,7 @@ bool FPCGSpawnActorElement::SpawnAndPrepareSubgraphs(FPCGSubgraphContext* Contex
 
 	const bool bHasAuthority = !Context->SourceComponent.IsValid() || (Context->SourceComponent->GetOwner() && Context->SourceComponent->GetOwner()->HasAuthority());
 
-	TArray<FPCGTaggedData> Inputs = Context->InputData.GetInputs();
+	TArray<FPCGTaggedData> Inputs = Context->InputData.GetInputsByPin(PCGPinConstants::DefaultInputLabel);
 	for (const FPCGTaggedData& Input : Inputs)
 	{
 		const UPCGSpatialData* SpatialData = Cast<UPCGSpatialData>(Input.Data);
diff --git a/Engine/Plugins/PCG/Source/PCG/Private/Elements/PCGSpawnSplineMesh.cpp b/Engine/Plugins/PCG/Source/PCG/Private/Elements/PCGSpawnSplineMesh.cpp
--- a/Engine/Plugins/PCG/Source/PCG/Private/Elements/PCGSpawnSplineMesh.cpp
+++ b/Engine/Plugins/PCG/Source/PCG/Private/Elements/PCGSpawnSplineMesh.cpp
@@ -102,7 +102,7 @@ bool FPCGSpawnSplineMeshElement::ExecuteInternal(FPCGContext* Context) const
 	FPCGSplineMeshComponentBuilderParameters SMCBuilderParams;
 	SMCBuilderParams.SplineMeshParams = Settings->SplineMeshParams;
 
-	TArray<FPCGTaggedData> Inputs = Context->InputData.GetInputs();
+	TArray<FPCGTaggedData> Inputs = Context->InputData.GetInputsByPin(PCGPinConstants::DefaultInputLabel);
 
 	for (const FPCGTaggedData& Input : Inputs)
 	{
diff --git a/Engine/Plugins/PCG/Source/PCG/Private/Elements/PCGSplitPoints.cpp b/Engine/Plugins/PCG/Source/PCG/Private/Elements/PCGSplitPoints.cpp
--- a/Engine/Plugins/PCG/Source/PCG/Private/Elements/PCGSplitPoints.cpp
+++ b/Engine/Plugins/PCG/Source/PCG/Private/Elements/PCGSplitPoints.cpp
@@ -47,7 +47,7 @@ bool FPCGSplitPointsElement::ExecuteInternal(FPCGContext* Context) const
 	const UPCGSplitPointsSettings* Settings = Context->GetInputSettings<UPCGSplitPointsSettings>();
 	check(Settings);
 
-	const TArray<FPCGTaggedData> Inputs = Context->InputData.GetInputs();
+	const TArray<FPCGTaggedData> Inputs = Context->InputData.GetInputsByPin(PCGPinConstants::DefaultInputLabel);
 	TArray<FPCGTaggedData>& Outputs = Context->OutputData.TaggedData;
 
 	const float SplitPosition = FMath::Clamp(Settings->SplitPosition, 0.0f, 1.0f);
diff --git a/Engine/Plugins/PCG/Source/PCG/Private/Elements/PCGStaticMeshSpawner.cpp b/Engine/Plugins/PCG/Source/PCG/Private/Elements/PCGStaticMeshSpawner.cpp
--- a/Engine/Plugins/PCG/Source/PCG/Private/Elements/PCGStaticMeshSpawner.cpp
+++ b/Engine/Plugins/PCG/Source/PCG/Private/Elements/PCGStaticMeshSpawner.cpp
@@ -177,7 +177,7 @@ bool FPCGStaticMeshSpawnerElement::PrepareDataInternal(FPCGContext* InContext) c
 	}
 
 	// perform mesh selection
-	TArray<FPCGTaggedData> Inputs = Context->InputData.GetInputs();
+	TArray<FPCGTaggedData> Inputs = Context->InputData.GetInputsByPin(PCGPinConstants::DefaultInputLabel);
 	TArray<FPCGTaggedData>& Outputs = Context->OutputData.TaggedData;
 
 	while(Context->CurrentInputIndex < Inputs.Num())
diff --git a/Engine/Plugins/PCG/Source/PCG/Private/Elements/PCGSurfaceSampler.cpp b/Engine/Plugins/PCG/Source/PCG/Private/Elements/PCGSurfaceSampler.cpp
--- a/Engine/Plugins/PCG/Source/PCG/Private/Elements/PCGSurfaceSampler.cpp
+++ b/Engine/Plugins/PCG/Source/PCG/Private/Elements/PCGSurfaceSampler.cpp
@@ -589,7 +589,7 @@ bool FPCGSurfaceSamplerElement::ExecuteInternal(FPCGContext* InContext) const
 	// The execution would have resulted in an empty set of points for all iterations
 	if (TimeSlicedContext->GetExecutionStateResult() == EPCGTimeSliceInitResult::NoOperation)
 	{
-		for (FPCGTaggedData& Input : TimeSlicedContext->InputData.GetInputs())
+		for (FPCGTaggedData& Input : TimeSlicedContext->InputData.GetInputsByPin(PCGSurfaceSamplerConstants::SurfaceLabel))
 		{
 			FPCGTaggedData& Output = TimeSlicedContext->OutputData.TaggedData.Add_GetRef(Input);
 			UPCGBasePointData* PointData = FPCGContext::NewPointData_AnyThread(InContext);
diff --git a/Engine/Plugins/PCG/Source/PCG/Private/Elements/PCGTransformPoints.cpp b/Engine/Plugins/PCG/Source/PCG/Private/Elements/PCGTransformPoints.cpp
--- a/Engine/Plugins/PCG/Source/PCG/Private/Elements/PCGTransformPoints.cpp
+++ b/Engine/Plugins/PCG/Source/PCG/Private/Elements/PCGTransformPoints.cpp
@@ -27,7 +27,7 @@ bool FPCGTransformPointsElement::ExecuteInternal(FPCGContext* Context) const
 	const UPCGTransformPointsSettings* Settings = Context->GetInputSettings<UPCGTransformPointsSettings>();
 	check(Settings);
 
-	const TArray<FPCGTaggedData> Inputs = Context->InputData.GetInputs();
+	const TArray<FPCGTaggedData> Inputs = Context->InputData.GetInputsByPin(PCGPinConstants::DefaultInputLabel);
 	TArray<FPCGTaggedData>& Outputs = Context->OutputData.TaggedData;
 
 	const bool bApplyToAttribute = Settings->bApplyToAttribute;
diff --git a/Engine/Plugins/PCG/Source/PCG/Private/Elements/PCGVolumeSampler.cpp b/Engine/Plugins/PCG/Source/PCG/Private/Elements/PCGVolumeSampler.cpp
--- a/Engine/Plugins/PCG/Source/PCG/Private/Elements/PCGVolumeSampler.cpp
+++ b/Engine/Plugins/PCG/Source/PCG/Private/Elements/PCGVolumeSampler.cpp
@@ -398,7 +398,7 @@ bool FPCGVolumeSamplerElement::ExecuteInternal(FPCGContext* Context) const
 	// The execution would have resulted in an empty set of points for all iterations
 	if (TimeSlicedContext->GetExecutionStateResult() == EPCGTimeSliceInitResult::NoOperation)
 	{
-		for (FPCGTaggedData& Input : TimeSlicedContext->InputData.GetInputs())
+		for (FPCGTaggedData& Input : TimeSlicedContext->InputData.GetInputsByPin(PCGVolumeSamplerConstants::VolumeLabel))
 		{
 			// TODO: Empty point data (to preserve previous behavior). Eventually, should be replaced with no output at all
 			FPCGTaggedData& Output = TimeSlicedContext->OutputData.TaggedData.Add_GetRef(Input);
diff --git a/Engine/Plugins/PCG/Source/PCG/Private/PCGData.cpp b/Engine/Plugins/PCG/Source/PCG/Private/PCGData.cpp
--- a/Engine/Plugins/PCG/Source/PCG/Private/PCGData.cpp
+++ b/Engine/Plugins/PCG/Source/PCG/Private/PCGData.cpp
@@ -175,6 +175,11 @@ FPCGCrc FPCGTaggedData::ComputeCrc(bool bFullDataCrc) const
 }
 
 TArray<FPCGTaggedData> FPCGDataCollection::GetInputs() const
+{
+	return GetAllSpatialInputs();
+}
+
+TArray<FPCGTaggedData> FPCGDataCollection::GetAllSpatialInputs() const
 {
 	return TaggedData.FilterByPredicate([](const FPCGTaggedData& Data) {
 		return Cast<UPCGSpatialData>(Data.Data) != nullptr;
@@ -568,9 +573,10 @@ void UPCGDataFunctionLibrary::AddToCollection(FPCGDataCollection& InCollection,
 	}
 }
 
+// DEPRECATED
 TArray<FPCGTaggedData> UPCGDataFunctionLibrary::GetInputs(const FPCGDataCollection& InCollection)
 {
-	return InCollection.GetInputs();
+	return InCollection.GetAllSpatialInputs();
 }
 
 TArray<FPCGTaggedData> UPCGDataFunctionLibrary::GetInputsByPinLabel(const FPCGDataCollection& InCollection, const FName InPinLabel)
diff --git a/Engine/Plugins/PCG/Source/PCG/Private/Tests/Elements/Metadata/PCGMetadataBreakTransformTest.cpp b/Engine/Plugins/PCG/Source/PCG/Private/Tests/Elements/Metadata/PCGMetadataBreakTransformTest.cpp
--- a/Engine/Plugins/PCG/Source/PCG/Private/Tests/Elements/Metadata/PCGMetadataBreakTransformTest.cpp
+++ b/Engine/Plugins/PCG/Source/PCG/Private/Tests/Elements/Metadata/PCGMetadataBreakTransformTest.cpp
@@ -108,7 +108,7 @@ bool FPCGMetadataBreakTransformTest::RunTest(const FString& Parameters)
 		// If the test is not valid, just check that we have no output and early out
 		if (!bIsValid)
 		{
-			return TestTrue(TEXT("Invalid operation, no outputs"), Context->OutputData.GetInputs().IsEmpty());
+			return TestTrue(TEXT("Invalid operation, no outputs"), Context->OutputData.GetAllInputs().IsEmpty());
 		}
 
 		const TArray<FPCGTaggedData>& Inputs = Context->InputData.GetInputsByPin(PCGPinConstants::DefaultInputLabel);
diff --git a/Engine/Plugins/PCG/Source/PCG/Private/Tests/Elements/Metadata/PCGMetadataBreakVectorTest.cpp b/Engine/Plugins/PCG/Source/PCG/Private/Tests/Elements/Metadata/PCGMetadataBreakVectorTest.cpp
--- a/Engine/Plugins/PCG/Source/PCG/Private/Tests/Elements/Metadata/PCGMetadataBreakVectorTest.cpp
+++ b/Engine/Plugins/PCG/Source/PCG/Private/Tests/Elements/Metadata/PCGMetadataBreakVectorTest.cpp
@@ -126,7 +126,7 @@ bool FPCGMetadataBreakVectorTest::RunTest(const FString& Parameters)
 		// If the test is not valid, just check that we have no output and early out
 		if (!bIsValid)
 		{
-			return TestTrue(TEXT("Invalid operation, no outputs"), Context->OutputData.GetInputs().IsEmpty());
+			return TestTrue(TEXT("Invalid operation, no outputs"), Context->OutputData.GetAllInputs().IsEmpty());
 		}
 
 		const TArray<FPCGTaggedData>& Inputs = Context->InputData.GetInputsByPin(PCGPinConstants::DefaultInputLabel);
diff --git a/Engine/Plugins/PCG/Source/PCG/Private/Tests/Elements/PCGAttributeNoiseTest.cpp b/Engine/Plugins/PCG/Source/PCG/Private/Tests/Elements/PCGAttributeNoiseTest.cpp
--- a/Engine/Plugins/PCG/Source/PCG/Private/Tests/Elements/PCGAttributeNoiseTest.cpp
+++ b/Engine/Plugins/PCG/Source/PCG/Private/Tests/Elements/PCGAttributeNoiseTest.cpp
@@ -41,8 +41,8 @@ bool FPCGAttributeNoiseDensityTest::RunTest(const FString& Parameters)
 		while (!NoiseElement->Execute(Context.Get()))
 		{}
 
-		const TArray<FPCGTaggedData>& Inputs = Context->InputData.GetInputs();
-		const TArray<FPCGTaggedData>& Outputs = Context->OutputData.GetInputs();
+		const TArray<FPCGTaggedData>& Inputs = Context->InputData.GetAllSpatialInputs();
+		const TArray<FPCGTaggedData>& Outputs = Context->OutputData.GetAllSpatialInputs();
 
 		if (!TestEqual("Valid number of outputs", Outputs.Num(), Inputs.Num()))
 		{
diff --git a/Engine/Plugins/PCG/Source/PCG/Private/Tests/Elements/PCGCopyPointsTest.cpp b/Engine/Plugins/PCG/Source/PCG/Private/Tests/Elements/PCGCopyPointsTest.cpp
--- a/Engine/Plugins/PCG/Source/PCG/Private/Tests/Elements/PCGCopyPointsTest.cpp
+++ b/Engine/Plugins/PCG/Source/PCG/Private/Tests/Elements/PCGCopyPointsTest.cpp
@@ -97,7 +97,7 @@ bool FPCGCopyPointsTestBase::RunTestInternal(FParameters& Parameters)
 
 	const TArray<FPCGTaggedData> Sources = Context->InputData.GetInputsByPin(PCGCopyPointsConstants::SourcePointsLabel);
 	const TArray<FPCGTaggedData> Targets = Context->InputData.GetInputsByPin(PCGCopyPointsConstants::TargetPointsLabel);
-	const TArray<FPCGTaggedData>& Outputs = Context->OutputData.GetInputs();
+	const TArray<FPCGTaggedData>& Outputs = Context->OutputData.GetAllSpatialInputs();
 
 	UTEST_TRUE("Valid number of inputs", Sources.Num() == 1 && Targets.Num() == 1);
 	UTEST_EQUAL("Valid number of outputs", Sources.Num(), 1);
diff --git a/Engine/Plugins/PCG/Source/PCG/Private/Tests/Elements/PCGDensityRemapTest.cpp b/Engine/Plugins/PCG/Source/PCG/Private/Tests/Elements/PCGDensityRemapTest.cpp
--- a/Engine/Plugins/PCG/Source/PCG/Private/Tests/Elements/PCGDensityRemapTest.cpp
+++ b/Engine/Plugins/PCG/Source/PCG/Private/Tests/Elements/PCGDensityRemapTest.cpp
@@ -43,8 +43,8 @@ bool FPCGDensityRemapTest::RunTest(const FString& Parameters)
 		while (!DensityRemapElement->Execute(Context.Get()))
 		{}
 
-		const TArray<FPCGTaggedData>& Inputs = Context->InputData.GetInputs();
-		const TArray<FPCGTaggedData>& Outputs = Context->OutputData.GetInputs();
+		const TArray<FPCGTaggedData>& Inputs = Context->InputData.GetAllSpatialInputs();
+		const TArray<FPCGTaggedData>& Outputs = Context->OutputData.GetAllSpatialInputs();
 		
 		if (!TestEqual("Valid number of outputs", Outputs.Num(), Inputs.Num()))
 		{
diff --git a/Engine/Plugins/PCG/Source/PCG/Private/Tests/Elements/PCGFilterByIndexTest.cpp b/Engine/Plugins/PCG/Source/PCG/Private/Tests/Elements/PCGFilterByIndexTest.cpp
--- a/Engine/Plugins/PCG/Source/PCG/Private/Tests/Elements/PCGFilterByIndexTest.cpp
+++ b/Engine/Plugins/PCG/Source/PCG/Private/Tests/Elements/PCGFilterByIndexTest.cpp
@@ -32,7 +32,7 @@ namespace PCGFilterByIndexTestPrivate
 		{
 		}
 
-		return Context->OutputData.GetInputs();
+		return Context->OutputData.GetAllInputs();
 	}
 
 	bool ValidateData(const TArray<FPCGTaggedData>& Outputs)
diff --git a/Engine/Plugins/PCG/Source/PCG/Private/Tests/Elements/PCGGatherTest.cpp b/Engine/Plugins/PCG/Source/PCG/Private/Tests/Elements/PCGGatherTest.cpp
--- a/Engine/Plugins/PCG/Source/PCG/Private/Tests/Elements/PCGGatherTest.cpp
+++ b/Engine/Plugins/PCG/Source/PCG/Private/Tests/Elements/PCGGatherTest.cpp
@@ -28,7 +28,7 @@ bool FPCGGatherTest_Basic::RunTest(const FString& Parameters)
 	while (!TestElement->Execute(Context.Get())) {}
 
 	// test our point data
-	const TArray<FPCGTaggedData>& Outputs = Context->OutputData.GetInputs();
+	const TArray<FPCGTaggedData>& Outputs = Context->OutputData.GetAllInputs();
 
 	UTEST_EQUAL("Output count", Outputs.Num(), 1);
 
diff --git a/Engine/Plugins/PCG/Source/PCG/Private/Tests/Elements/PCGNormalToDensityTest.cpp b/Engine/Plugins/PCG/Source/PCG/Private/Tests/Elements/PCGNormalToDensityTest.cpp
--- a/Engine/Plugins/PCG/Source/PCG/Private/Tests/Elements/PCGNormalToDensityTest.cpp
+++ b/Engine/Plugins/PCG/Source/PCG/Private/Tests/Elements/PCGNormalToDensityTest.cpp
@@ -37,7 +37,7 @@ bool FPCGNormalToDensityTest_Set::RunTest(const FString& Parameters)
 	while (!Element->Execute(Context.Get())) {}
 
 	// test our point data
-	const TArray<FPCGTaggedData>& Outputs = Context->OutputData.GetInputs();
+	const TArray<FPCGTaggedData>& Outputs = Context->OutputData.GetAllInputs();
 
 	UTEST_EQUAL("Output count", Outputs.Num(), 1);
 
@@ -84,7 +84,7 @@ bool FPCGNormalToDensityTest_CustomNormal::RunTest(const FString& Parameters)
 	while (!Element->Execute(Context.Get())) {}
 
 	// test our point data
-	const TArray<FPCGTaggedData>& Outputs = Context->OutputData.GetInputs();
+	const TArray<FPCGTaggedData>& Outputs = Context->OutputData.GetAllInputs();
 
 	UTEST_EQUAL("Output count", Outputs.Num(), 1);
 
@@ -131,7 +131,7 @@ bool FPCGNormalToDensityTest_Strength::RunTest(const FString& Parameters)
 	while (!Element->Execute(Context.Get())) {}
 
 	// test our point data
-	const TArray<FPCGTaggedData>& Outputs = Context->OutputData.GetInputs();
+	const TArray<FPCGTaggedData>& Outputs = Context->OutputData.GetAllInputs();
 
 	UTEST_EQUAL("Output count", Outputs.Num(), 1);
 
@@ -179,7 +179,7 @@ bool FPCGNormalToDensityTest_Minimum::RunTest(const FString& Parameters)
 	while (!Element->Execute(Context.Get())) {}
 
 	// test our point data
-	const TArray<FPCGTaggedData>& Outputs = Context->OutputData.GetInputs();
+	const TArray<FPCGTaggedData>& Outputs = Context->OutputData.GetAllInputs();
 
 	UTEST_EQUAL("Output count", Outputs.Num(), 1);
 
@@ -227,7 +227,7 @@ bool FPCGNormalToDensityTest_Maximum::RunTest(const FString& Parameters)
 	while (!Element->Execute(Context.Get())) {}
 
 	// test our point data
-	const TArray<FPCGTaggedData>& Outputs = Context->OutputData.GetInputs();
+	const TArray<FPCGTaggedData>& Outputs = Context->OutputData.GetAllInputs();
 
 	UTEST_EQUAL("Output count", Outputs.Num(), 1);
 
@@ -276,7 +276,7 @@ bool FPCGNormalToDensityTest_Add::RunTest(const FString& Parameters)
 	while (!Element->Execute(Context.Get())) {}
 
 	// test our point data
-	const TArray<FPCGTaggedData>& Outputs = Context->OutputData.GetInputs();
+	const TArray<FPCGTaggedData>& Outputs = Context->OutputData.GetAllInputs();
 
 	UTEST_EQUAL("Output count", Outputs.Num(), 1);
 
@@ -324,7 +324,7 @@ bool FPCGNormalToDensityTest_Subtract::RunTest(const FString& Parameters)
 	while (!Element->Execute(Context.Get())) {}
 
 	// test our point data
-	const TArray<FPCGTaggedData>& Outputs = Context->OutputData.GetInputs();
+	const TArray<FPCGTaggedData>& Outputs = Context->OutputData.GetAllInputs();
 
 	UTEST_EQUAL("Output count", Outputs.Num(), 1);
 
@@ -372,7 +372,7 @@ bool FPCGNormalToDensityTest_Multiply::RunTest(const FString& Parameters)
 	while (!Element->Execute(Context.Get())) {}
 
 	// test our point data
-	const TArray<FPCGTaggedData>& Outputs = Context->OutputData.GetInputs();
+	const TArray<FPCGTaggedData>& Outputs = Context->OutputData.GetAllInputs();
 
 	UTEST_EQUAL("Output count", Outputs.Num(), 1);
 
@@ -420,7 +420,7 @@ bool FPCGNormalToDensityTest_Divide::RunTest(const FString& Parameters)
 	while (!Element->Execute(Context.Get())) {}
 
 	// test our point data
-	const TArray<FPCGTaggedData>& Outputs = Context->OutputData.GetInputs();
+	const TArray<FPCGTaggedData>& Outputs = Context->OutputData.GetAllInputs();
 
 	UTEST_EQUAL("Output count", Outputs.Num(), 1);
 
diff --git a/Engine/Plugins/PCG/Source/PCG/Private/Tests/Elements/PCGPointExtentsModifierTest.cpp b/Engine/Plugins/PCG/Source/PCG/Private/Tests/Elements/PCGPointExtentsModifierTest.cpp
--- a/Engine/Plugins/PCG/Source/PCG/Private/Tests/Elements/PCGPointExtentsModifierTest.cpp
+++ b/Engine/Plugins/PCG/Source/PCG/Private/Tests/Elements/PCGPointExtentsModifierTest.cpp
@@ -48,7 +48,7 @@ bool FPCGPointsBoundsModiferTest::RunTest(const FString& Parameters)
 		{}
 
 		const TArray<FPCGTaggedData>& Inputs = Context->InputData.GetInputsByPin(PCGPinConstants::DefaultInputLabel);
-		const TArray<FPCGTaggedData>& Outputs = Context->OutputData.GetInputs();
+		const TArray<FPCGTaggedData>& Outputs = Context->OutputData.GetAllInputs();
 
 		if (!TestEqual("Valid number of Outputs", Inputs.Num(), Outputs.Num()))
 		{
diff --git a/Engine/Plugins/PCG/Source/PCG/Private/Tests/Elements/PCGSanityCheckPointDataTest.cpp b/Engine/Plugins/PCG/Source/PCG/Private/Tests/Elements/PCGSanityCheckPointDataTest.cpp
--- a/Engine/Plugins/PCG/Source/PCG/Private/Tests/Elements/PCGSanityCheckPointDataTest.cpp
+++ b/Engine/Plugins/PCG/Source/PCG/Private/Tests/Elements/PCGSanityCheckPointDataTest.cpp
@@ -90,7 +90,7 @@ bool FPCGSanityCheckPointData_Success::RunTest(const FString& Parameters)
 
 	while (!TestElement->Execute(Context.Get())) {}
 
-	const TArray<FPCGTaggedData>& Outputs = Context->OutputData.GetInputs();
+	const TArray<FPCGTaggedData>& Outputs = Context->OutputData.GetAllSpatialInputs();
 
 	UTEST_EQUAL("Output count", Outputs.Num(), 1);
 
diff --git a/Engine/Plugins/PCG/Source/PCG/Private/Tests/Elements/PCGSpatialNoiseTest.cpp b/Engine/Plugins/PCG/Source/PCG/Private/Tests/Elements/PCGSpatialNoiseTest.cpp
--- a/Engine/Plugins/PCG/Source/PCG/Private/Tests/Elements/PCGSpatialNoiseTest.cpp
+++ b/Engine/Plugins/PCG/Source/PCG/Private/Tests/Elements/PCGSpatialNoiseTest.cpp
@@ -62,7 +62,7 @@ bool FPCGSpatialNoise_Perlin2D::RunTest(const FString& Parameters)
 
 	while (!TestElement->Execute(Context.Get())) {}	
 
-	const TArray<FPCGTaggedData>& Outputs = Context->OutputData.GetInputs();
+	const TArray<FPCGTaggedData>& Outputs = Context->OutputData.GetAllInputs();
 
 	UTEST_EQUAL("Output count", Outputs.Num(), 1);
 
@@ -100,7 +100,7 @@ bool FPCGSpatialNoise_Caustic2D::RunTest(const FString& Parameters)
 
 	while (!TestElement->Execute(Context.Get())) {}	
 
-	const TArray<FPCGTaggedData>& Outputs = Context->OutputData.GetInputs();
+	const TArray<FPCGTaggedData>& Outputs = Context->OutputData.GetAllInputs();
 
 	UTEST_EQUAL("Output count", Outputs.Num(), 1);
 
@@ -139,7 +139,7 @@ bool FPCGSpatialNoise_Voronoi2D::RunTest(const FString& Parameters)
 
 	while (!TestElement->Execute(Context.Get())) {}	
 
-	const TArray<FPCGTaggedData>& Outputs = Context->OutputData.GetInputs();
+	const TArray<FPCGTaggedData>& Outputs = Context->OutputData.GetAllInputs();
 
 	UTEST_EQUAL("Output count", Outputs.Num(), 1);
 
@@ -181,7 +181,7 @@ bool FPCGSpatialNoise_FractionalBrownian2D::RunTest(const FString& Parameters)
 
 	while (!TestElement->Execute(Context.Get())) {}	
 
-	const TArray<FPCGTaggedData>& Outputs = Context->OutputData.GetInputs();
+	const TArray<FPCGTaggedData>& Outputs = Context->OutputData.GetAllInputs();
 
 	UTEST_EQUAL("Output count", Outputs.Num(), 1);
 
@@ -219,7 +219,7 @@ bool FPCGSpatialNoise_EdgeMask2D::RunTest(const FString& Parameters)
 
 	while (!TestElement->Execute(Context.Get())) {}	
 
-	const TArray<FPCGTaggedData>& Outputs = Context->OutputData.GetInputs();
+	const TArray<FPCGTaggedData>& Outputs = Context->OutputData.GetAllInputs();
 
 	UTEST_EQUAL("Output count", Outputs.Num(), 1);
 
@@ -257,7 +257,7 @@ bool FPCGSpatialNoise_TilingPerlin2D::RunTest(const FString& Parameters)
 
 	while (!TestElement->Execute(Context.Get())) {}	
 
-	const TArray<FPCGTaggedData>& Outputs = Context->OutputData.GetInputs();
+	const TArray<FPCGTaggedData>& Outputs = Context->OutputData.GetAllInputs();
 
 	UTEST_EQUAL("Output count", Outputs.Num(), 1);
 
@@ -296,7 +296,7 @@ bool FPCGSpatialNoise_TilingVoronoi2D::RunTest(const FString& Parameters)
 
 	while (!TestElement->Execute(Context.Get())) {}	
 
-	const TArray<FPCGTaggedData>& Outputs = Context->OutputData.GetInputs();
+	const TArray<FPCGTaggedData>& Outputs = Context->OutputData.GetAllInputs();
 
 	UTEST_EQUAL("Output count", Outputs.Num(), 1);
 
diff --git a/Engine/Plugins/PCG/Source/PCG/Private/Tests/Elements/PCGSwitchTest.cpp b/Engine/Plugins/PCG/Source/PCG/Private/Tests/Elements/PCGSwitchTest.cpp
--- a/Engine/Plugins/PCG/Source/PCG/Private/Tests/Elements/PCGSwitchTest.cpp
+++ b/Engine/Plugins/PCG/Source/PCG/Private/Tests/Elements/PCGSwitchTest.cpp
@@ -54,7 +54,7 @@ protected:
 
 	bool ValidateOutput(const FPCGDataCollection& OutputData, const FName CorrectOutputLabel)
 	{
-		UTEST_EQUAL("Only one output exists", OutputData.GetInputs().Num(), 1);
+		UTEST_EQUAL("Only one output exists", OutputData.GetAllSpatialInputs().Num(), 1);
 
 		const TArray<FPCGTaggedData> OnPinData = OutputData.GetInputsByPin(CorrectOutputLabel);
 
@@ -165,7 +165,7 @@ bool FPCGSwitchTest_Enum::RunTest(const FString& Parameters)
 
 		if (TestSettings->EnumSelection.Class->HasMetaData(TEXT("Hidden"), i) || TestSettings->EnumSelection.Class->HasMetaData(TEXT("Spacer"), i))
 		{
-			UTEST_TRUE("Hidden enum selection output is empty", OutputData.GetInputs().IsEmpty());
+			UTEST_TRUE("Hidden enum selection output is empty", OutputData.GetAllSpatialInputs().IsEmpty());
 		}
 		else
 		{
diff --git a/Engine/Plugins/PCG/Source/PCG/Private/Tests/Elements/PCGTransformPointsTest.cpp b/Engine/Plugins/PCG/Source/PCG/Private/Tests/Elements/PCGTransformPointsTest.cpp
--- a/Engine/Plugins/PCG/Source/PCG/Private/Tests/Elements/PCGTransformPointsTest.cpp
+++ b/Engine/Plugins/PCG/Source/PCG/Private/Tests/Elements/PCGTransformPointsTest.cpp
@@ -42,12 +42,15 @@ bool FPCGTransformPointsTest::RunTest(const FString& Parameters)
 
 	FPCGTaggedData& EmptyTaggedData = TestData.InputData.TaggedData.Emplace_GetRef(FPCGTaggedData());
 	EmptyTaggedData.Data = EmptyData;
+	EmptyTaggedData.Pin = PCGPinConstants::DefaultInputLabel;
 
 	FPCGTaggedData& SimpleTaggedData = TestData.InputData.TaggedData.Emplace_GetRef(FPCGTaggedData());
 	SimpleTaggedData.Data = SimpleData;
+	SimpleTaggedData.Pin = PCGPinConstants::DefaultInputLabel;
 
 	FPCGTaggedData& ComplexTaggedData = TestData.InputData.TaggedData.Emplace_GetRef(FPCGTaggedData());
 	ComplexTaggedData.Data = ComplexData;
+	ComplexTaggedData.Pin = PCGPinConstants::DefaultInputLabel;
 
 	// Basic parameters we will use for all tests
 	Settings->OffsetMin = FVector(-1.f);
@@ -64,8 +67,8 @@ bool FPCGTransformPointsTest::RunTest(const FString& Parameters)
 		while (!TransformPointsElement->Execute(Context.Get()))
 		{}
 
-		const TArray<FPCGTaggedData>& Inputs = Context->InputData.GetInputs();
-		const TArray<FPCGTaggedData>& Outputs = Context->OutputData.GetInputs();
+		const TArray<FPCGTaggedData>& Inputs = Context->InputData.GetAllSpatialInputs();
+		const TArray<FPCGTaggedData>& Outputs = Context->OutputData.GetAllSpatialInputs();
 
 		if (!TestEqual("Valid number of outputs", Inputs.Num(), Outputs.Num()))
 		{
diff --git a/Engine/Plugins/PCG/Source/PCG/Public/PCGData.h b/Engine/Plugins/PCG/Source/PCG/Public/PCGData.h
--- a/Engine/Plugins/PCG/Source/PCG/Public/PCGData.h
+++ b/Engine/Plugins/PCG/Source/PCG/Public/PCGData.h
@@ -177,7 +177,15 @@ struct PCG_API FPCGDataCollection
 	GENERATED_BODY()
 
 	/** Returns all spatial data in the collection */
+	UE_DEPRECATED(5.6, "Use GetAllSpatialInputs for only spatial inputs, or use GetAllInputs.")
 	TArray<FPCGTaggedData> GetInputs() const;
+
+	/** Returns all spatial data in the collection. */
+	TArray<FPCGTaggedData> GetAllSpatialInputs() const;
+
+	/** Returns all inputs in the collection. */
+	const TArray<FPCGTaggedData>& GetAllInputs() const { return TaggedData; }
+
 	/** Returns all data on a given pin. */
 	TArray<FPCGTaggedData> GetInputsByPin(const FName& InPinLabel) const;
 	/** Returns all spatial data on a given pin */
@@ -352,7 +360,7 @@ public:
 	static void AddToCollection(UPARAM(ref) FPCGDataCollection& InCollection, const UPCGData* InData, FName InPinLabel, TArray<FString> InTags);
 
 	// Blueprint methods to support interaction with FPCGDataCollection
-	UFUNCTION(BlueprintCallable, Category = Data)
+	UFUNCTION(BlueprintCallable, Category = Data, meta = (DisplayName="Get All Spatial Inputs"))
 	static TArray<FPCGTaggedData> GetInputs(const FPCGDataCollection& InCollection);
 
 	UFUNCTION(BlueprintCallable, Category = Data)
