[PCG] GPU Static Mesh Data Interface
* New type UPCGStaticMeshResourceData. Thin wrapper around a static mesh object.
* New abstract base type UPCGResourceData. Thin wrapper around a non-world resource/asset.
* SM DI supports sampling per-vert and per-tri. Can get position, normal, tangent, bitangent, colors, and UVs
* Can convert to an Attribute Set via new Get Resource Path node

diff --git a/Engine/Plugins/PCG/Shaders/Private/PCGStaticMeshDataInterface.ush b/Engine/Plugins/PCG/Shaders/Private/PCGStaticMeshDataInterface.ush
new file mode 100644
--- /dev/null
+++ b/Engine/Plugins/PCG/Shaders/Private/PCGStaticMeshDataInterface.ush
@@ -0,0 +1,159 @@
+// Copyright Epic Games, Inc. All Rights Reserved.
+
+// Reference NiagaraDataInterfaceStaticMeshTemplate.ush
+
+/** @todo_pcg:
+ * Sections/Sockets
+ * Distance fields
+ * Spatial sampling
+ */
+
+int {DataInterfaceName}_NumVertices;
+int {DataInterfaceName}_NumTriangles;
+int {DataInterfaceName}_NumUVs;
+uint {DataInterfaceName}_HasColors;
+
+Buffer<uint> {DataInterfaceName}_IndexBuffer;
+Buffer<float> {DataInterfaceName}_PositionBuffer;
+Buffer<float4> {DataInterfaceName}_TangentBuffer;
+Buffer<float2> {DataInterfaceName}_UVBuffer;
+Buffer<float4> {DataInterfaceName}_ColorBuffer;
+
+float3 {DataInterfaceName}_BoundsExtents;
+
+// #################### VERTEX BASED SAMPLING ##########################
+
+int GetNumVertices_{DataInterfaceName}(int DataIndex)
+{
+	return {DataInterfaceName}_NumVertices;
+}
+
+void SampleVertex_{DataInterfaceName}(int DataIndex, int VertexIndex, out float3 OutPosition, out float3 OutNormal, out float3 OutTangent, out float3 OutBitangent)
+{
+	if ({DataInterfaceName}_NumVertices > 0)
+	{
+		VertexIndex = clamp(VertexIndex, 0, {DataInterfaceName}_NumVertices - 1);
+
+		OutPosition.x = {DataInterfaceName}_PositionBuffer[VertexIndex * 3 + 0];
+		OutPosition.y = {DataInterfaceName}_PositionBuffer[VertexIndex * 3 + 1];
+		OutPosition.z = {DataInterfaceName}_PositionBuffer[VertexIndex * 3 + 2];
+
+		float4 TangentX = TangentBias({DataInterfaceName}_TangentBuffer[VertexIndex * 2 + 0]);
+		float4 TangentZ = TangentBias({DataInterfaceName}_TangentBuffer[VertexIndex * 2 + 1]);
+		float3 TangentY = cross(TangentZ.xyz, TangentX.xyz) * TangentZ.w;
+
+		// Sometimes we have bogus tangent data in the buffer, i.e. X==Z so ensure we handle this and don't end up with NaNs in the buffer
+		if (length2(TangentY) > 0.001f)
+		{
+			OutNormal = TangentZ.xyz;
+			OutTangent = TangentX.xyz;
+			OutBitangent = TangentY;
+		}
+		else
+		{
+			OutNormal = float3(0, 0, 1);
+			OutTangent = float3(1, 0, 0);
+			OutBitangent = float3(0, 1, 0);
+		}
+	}
+	else
+	{
+		OutPosition = float3(0, 0, 0);
+		OutNormal = float3(0, 0, 1);
+		OutBitangent = float3(0, 1, 0);
+		OutTangent = float3(1, 0, 0);
+	}
+}
+
+float4 SampleVertexColor_{DataInterfaceName}(int DataIndex, int VertexIndex)
+{
+	VertexIndex = clamp(VertexIndex, 0, {DataInterfaceName}_NumVertices - 1);
+	return {DataInterfaceName}_HasColors ? {DataInterfaceName}_ColorBuffer[VertexIndex] FMANUALFETCH_COLOR_COMPONENT_SWIZZLE : float4(1, 1, 1, 1);
+}
+
+float2 SampleVertexUV_{DataInterfaceName}(int DataIndex, int VertexIndex, int UVSet)
+{
+	VertexIndex = clamp(VertexIndex, 0, {DataInterfaceName}_NumVertices - 1);
+	UVSet = clamp(UVSet, 0, {DataInterfaceName}_NumUVs - 1);
+	return UVSet >= 0 && UVSet < {DataInterfaceName}_NumUVs ? {DataInterfaceName}_UVBuffer[VertexIndex * {DataInterfaceName}_NumUVs + UVSet] : float2(0, 0);
+}
+
+// #################### TRIANGLE BASED SAMPLING ##########################
+
+int GetNumTriangles_{DataInterfaceName}(int DataIndex)
+{
+	return {DataInterfaceName}_NumTriangles;
+}
+
+void GetTriangleIndices_{DataInterfaceName}(int DataIndex, int TriangleIndex, out int OutIndex0, out int OutIndex1, out int OutIndex2)
+{
+	if ({DataInterfaceName}_NumTriangles > 0)
+	{
+		TriangleIndex = clamp(TriangleIndex, 0, {DataInterfaceName}_NumTriangles - 1);
+		OutIndex0 = {DataInterfaceName}_IndexBuffer[TriangleIndex * 3 + 0];
+		OutIndex1 = {DataInterfaceName}_IndexBuffer[TriangleIndex * 3 + 1];
+		OutIndex2 = {DataInterfaceName}_IndexBuffer[TriangleIndex * 3 + 2];
+	}
+	else
+	{
+		OutIndex0 = 0;
+		OutIndex1 = 0;
+		OutIndex2 = 0;
+	}
+}
+
+void SampleTriangle_{DataInterfaceName}(
+	int DataIndex,
+	int TriangleIndex,
+	float3 BaryCoord,
+	out float3 OutPosition,
+	out float3 OutNormal,
+	out float3 OutTangent,
+	out float3 OutBitangent)
+{
+	int Index[3];
+	GetTriangleIndices_{DataInterfaceName}(DataIndex, TriangleIndex, Index[0], Index[1], Index[2]);
+
+	float3 Positions[3], Normals[3], Bitangents[3], Tangents[3];
+	SampleVertex_{DataInterfaceName}(DataIndex, Index[0], Positions[0], Normals[0], Tangents[0], Bitangents[0]);
+	SampleVertex_{DataInterfaceName}(DataIndex, Index[1], Positions[1], Normals[1], Tangents[1], Bitangents[1]);
+	SampleVertex_{DataInterfaceName}(DataIndex, Index[2], Positions[2], Normals[2], Tangents[2], Bitangents[2]);
+
+	OutPosition = Positions[0] * BaryCoord.x + Positions[1] * BaryCoord.y + Positions[2] * BaryCoord.z;
+	OutNormal = Normals[0] * BaryCoord.x + Normals[1] * BaryCoord.y + Normals[2] * BaryCoord.x;
+	OutTangent = Tangents[0] * BaryCoord.x + Tangents[1] * BaryCoord.y + Tangents[2] * BaryCoord.x;
+	OutBitangent = Bitangents[0] * BaryCoord.x + Bitangents[1] * BaryCoord.y + Bitangents[2] * BaryCoord.x;
+}
+
+float4 SampleTriangleColor_{DataInterfaceName}(int DataIndex, int TriangleIndex, float3 BaryCoord)
+{
+	int Index[3];
+	GetTriangleIndices_{DataInterfaceName}(DataIndex, TriangleIndex, Index[0], Index[1], Index[2]);
+
+	float4 Colors[3];
+	Colors[0] = SampleVertexColor_{DataInterfaceName}(DataIndex, Index[0]);
+	Colors[1] = SampleVertexColor_{DataInterfaceName}(DataIndex, Index[1]);
+	Colors[2] = SampleVertexColor_{DataInterfaceName}(DataIndex, Index[2]);
+
+	return Colors[0] * BaryCoord.x + Colors[1] * BaryCoord.y + Colors[2] * BaryCoord.z;
+}
+
+float2 SampleTriangleUV_{DataInterfaceName}(int DataIndex, int TriangleIndex, float3 BaryCoord, int UVSet)
+{
+	int Index[3];
+	GetTriangleIndices_{DataInterfaceName}(DataIndex, TriangleIndex, Index[0], Index[1], Index[2]);
+
+	float2 UVs[3];
+	UVs[0] = SampleVertexUV_{DataInterfaceName}(DataIndex, Index[0], UVSet);
+	UVs[1] = SampleVertexUV_{DataInterfaceName}(DataIndex, Index[1], UVSet);
+	UVs[2] = SampleVertexUV_{DataInterfaceName}(DataIndex, Index[2], UVSet);
+
+	return UVs[0] * BaryCoord.x + UVs[1] * BaryCoord.y + UVs[2] * BaryCoord.z;
+}
+
+// #################### MISC ##########################
+
+float3 GetMeshBoundsExtents_{DataInterfaceName}(int DataIndex)
+{
+	return {DataInterfaceName}_BoundsExtents;
+}
diff --git a/Engine/Plugins/PCG/Source/PCG/Internal/Compute/PCGComputeCommon.h b/Engine/Plugins/PCG/Source/PCG/Internal/Compute/PCGComputeCommon.h
--- a/Engine/Plugins/PCG/Source/PCG/Internal/Compute/PCGComputeCommon.h
+++ b/Engine/Plugins/PCG/Source/PCG/Internal/Compute/PCGComputeCommon.h
@@ -16,6 +16,8 @@ struct FPCGDataDesc;
 struct FPCGPinPropertiesGPU;
 struct FPCGTaggedData;
 
+class FRHIShaderResourceView;
+
 #define PCG_KERNEL_LOGGING_ENABLED (!(UE_BUILD_SHIPPING || UE_BUILD_TEST) || USE_LOGGING_IN_SHIPPING)
 
 #if PCG_KERNEL_LOGGING_ENABLED
@@ -67,7 +69,13 @@ namespace PCGComputeConstants
 	constexpr int POINT_STEEPNESS_ATTRIBUTE_ID = 8;
 
 	/** PCG data types supported in GPU node inputs. */
-	constexpr EPCGDataType AllowedInputTypes = EPCGDataType::Point | EPCGDataType::Param | EPCGDataType::Landscape | EPCGDataType::BaseTexture | EPCGDataType::ProxyForGPU;
+	constexpr EPCGDataType AllowedInputTypes =
+		EPCGDataType::Point
+		| EPCGDataType::Param
+		| EPCGDataType::Landscape
+		| EPCGDataType::BaseTexture
+		| EPCGDataType::Resource
+		| EPCGDataType::ProxyForGPU;
 
 	/** PCG data types supported in GPU node outputs. */
 	constexpr EPCGDataType AllowedOutputTypes = EPCGDataType::Point | EPCGDataType::Param | EPCGDataType::ProxyForGPU;
@@ -125,6 +133,13 @@ namespace PCGComputeHelpers
 	FString GetDataLabelResolverName(FName InPinLabel);
 }
 
+namespace PCGComputeDummies
+{
+	FRHIShaderResourceView* GetDummyFloatBuffer();
+	FRHIShaderResourceView* GetDummyFloat2Buffer();
+	FRHIShaderResourceView* GetDummyFloat4Buffer();
+}
+
 /** A by-label reference to a pin, used for wiring kernels within a node. */
 struct FPCGPinReference
 {
diff --git a/Engine/Plugins/PCG/Source/PCG/Private/Compute/DataInterfaces/PCGStaticMeshDataInterface.cpp b/Engine/Plugins/PCG/Source/PCG/Private/Compute/DataInterfaces/PCGStaticMeshDataInterface.cpp
new file mode 100644
--- /dev/null
+++ b/Engine/Plugins/PCG/Source/PCG/Private/Compute/DataInterfaces/PCGStaticMeshDataInterface.cpp
@@ -0,0 +1,285 @@
+// Copyright Epic Games, Inc. All Rights Reserved.
+
+#include "PCGStaticMeshDataInterface.h"
+
+#include "PCGData.h"
+#include "PCGModule.h"
+#include "Compute/PCGDataBinding.h"
+#include "Data/PCGStaticMeshResourceData.h"
+
+#include "RHIResources.h"
+#include "RenderGraphBuilder.h"
+#include "ShaderCompilerCore.h"
+#include "ShaderCore.h"
+#include "ShaderParameterMetadataBuilder.h"
+#include "StaticMeshResources.h"
+#include "ComputeFramework/ShaderParamTypeDefinition.h"
+#include "Engine/StaticMesh.h"
+
+#include UE_INLINE_GENERATED_CPP_BY_NAME(PCGStaticMeshDataInterface)
+
+#define LOCTEXT_NAMESPACE "PCGStaticMeshDataInterface"
+
+void UPCGStaticMeshDataInterface::GetSupportedInputs(TArray<FShaderFunctionDefinition>& OutFunctions) const
+{
+	// Vertex functions
+	{
+		OutFunctions.AddDefaulted_GetRef()
+			.SetName(TEXT("GetNumVertices"))
+			.AddReturnType(EShaderFundamentalType::Int)
+			.AddParam(EShaderFundamentalType::Int); // DataIndex
+
+		OutFunctions.AddDefaulted_GetRef()
+			.SetName(TEXT("SampleVertex"))
+			.AddParam(EShaderFundamentalType::Int) // DataIndex
+			.AddParam(EShaderFundamentalType::Int) // VertexIndex
+			.AddParam(EShaderFundamentalType::Float, 3, 0, EShaderParamModifier::Out) // OutPosition
+			.AddParam(EShaderFundamentalType::Float, 3, 0, EShaderParamModifier::Out) // OutNormal
+			.AddParam(EShaderFundamentalType::Float, 3, 0, EShaderParamModifier::Out) // OutTangent
+			.AddParam(EShaderFundamentalType::Float, 3, 0, EShaderParamModifier::Out); // OutBitangent
+
+		OutFunctions.AddDefaulted_GetRef()
+			.SetName(TEXT("SampleVertexColor"))
+			.AddReturnType(EShaderFundamentalType::Float, 4)
+			.AddParam(EShaderFundamentalType::Int) // DataIndex
+			.AddParam(EShaderFundamentalType::Int); // VertexIndex
+
+		OutFunctions.AddDefaulted_GetRef()
+			.SetName(TEXT("SampleVertexUV"))
+			.AddReturnType(EShaderFundamentalType::Float, 2)
+			.AddParam(EShaderFundamentalType::Int) // DataIndex
+			.AddParam(EShaderFundamentalType::Int) // VertexIndex
+			.AddParam(EShaderFundamentalType::Int); // UVSet
+	}
+
+	// Triangle functions
+	{
+		OutFunctions.AddDefaulted_GetRef()
+			.SetName(TEXT("GetNumTriangles"))
+			.AddReturnType(EShaderFundamentalType::Int)
+			.AddParam(EShaderFundamentalType::Int); // DataIndex
+
+		OutFunctions.AddDefaulted_GetRef()
+			.SetName(TEXT("GetTriangleIndices"))
+			.AddParam(EShaderFundamentalType::Int) // DataIndex
+			.AddParam(EShaderFundamentalType::Int) // TriangleIndex
+			.AddParam(EShaderFundamentalType::Int, 1, 0, EShaderParamModifier::Out) // OutIndex0
+			.AddParam(EShaderFundamentalType::Int, 1, 0, EShaderParamModifier::Out) // OutIndex1
+			.AddParam(EShaderFundamentalType::Int, 1, 0, EShaderParamModifier::Out); // OutIndex2
+
+		OutFunctions.AddDefaulted_GetRef()
+			.SetName(TEXT("SampleTriangle"))
+			.AddParam(EShaderFundamentalType::Int) // DataIndex
+			.AddParam(EShaderFundamentalType::Int) // TriangleIndex
+			.AddParam(EShaderFundamentalType::Float, 3) // BaryCoord
+			.AddParam(EShaderFundamentalType::Float, 3, 0, EShaderParamModifier::Out) // OutPosition
+			.AddParam(EShaderFundamentalType::Float, 3, 0, EShaderParamModifier::Out) // OutNormal
+			.AddParam(EShaderFundamentalType::Float, 3, 0, EShaderParamModifier::Out) // OutTangent
+			.AddParam(EShaderFundamentalType::Float, 3, 0, EShaderParamModifier::Out); // OutBitangent
+
+		OutFunctions.AddDefaulted_GetRef()
+			.SetName(TEXT("SampleTriangleColor"))
+			.AddReturnType(EShaderFundamentalType::Float, 4)
+			.AddParam(EShaderFundamentalType::Int) // DataIndex
+			.AddParam(EShaderFundamentalType::Int) // TriangleIndex
+			.AddParam(EShaderFundamentalType::Float, 3); // BaryCoord
+
+		OutFunctions.AddDefaulted_GetRef()
+			.SetName(TEXT("SampleTriangleUV"))
+			.AddReturnType(EShaderFundamentalType::Float, 2)
+			.AddParam(EShaderFundamentalType::Int) // DataIndex
+			.AddParam(EShaderFundamentalType::Int) // TriangleIndex
+			.AddParam(EShaderFundamentalType::Float, 3) // BaryCoord
+			.AddParam(EShaderFundamentalType::Int); // UVSet
+	}
+
+	// Misc functions
+	{
+		OutFunctions.AddDefaulted_GetRef()
+			.SetName(TEXT("GetMeshBoundsExtents"))
+			.AddReturnType(EShaderFundamentalType::Float, 3)
+			.AddParam(EShaderFundamentalType::Int); // DataIndex
+	}
+}
+
+// Reference NiagaraDataInterfaceStaticMesh.h
+BEGIN_SHADER_PARAMETER_STRUCT(FPCGStaticMeshDataInterfaceParameters,)
+	SHADER_PARAMETER(int, NumVertices)
+	SHADER_PARAMETER(int, NumTriangles)
+	SHADER_PARAMETER(int, NumUVs)
+	SHADER_PARAMETER(uint32, HasColors)
+	SHADER_PARAMETER_SRV(Buffer<uint>, IndexBuffer)
+	SHADER_PARAMETER_SRV(Buffer<float>, PositionBuffer)
+	SHADER_PARAMETER_SRV(Buffer<float4>, TangentBuffer)
+	SHADER_PARAMETER_SRV(Buffer<float2>, UVBuffer)
+	SHADER_PARAMETER_SRV(Buffer<float4>, ColorBuffer)
+	SHADER_PARAMETER(FVector3f, BoundsExtents)
+END_SHADER_PARAMETER_STRUCT()
+
+void UPCGStaticMeshDataInterface::GetShaderParameters(TCHAR const* UID, FShaderParametersMetadataBuilder& InOutBuilder, FShaderParametersMetadataAllocations& InOutAllocations) const
+{
+	InOutBuilder.AddNestedStruct<FPCGStaticMeshDataInterfaceParameters>(UID);
+}
+
+TCHAR const* UPCGStaticMeshDataInterface::TemplateFilePath = TEXT("/Plugin/PCG/Private/PCGStaticMeshDataInterface.ush");
+
+TCHAR const* UPCGStaticMeshDataInterface::GetShaderVirtualPath() const
+{
+	return TemplateFilePath;
+}
+
+void UPCGStaticMeshDataInterface::GetShaderHash(FString& InOutKey) const
+{
+	GetShaderFileHash(TemplateFilePath, EShaderPlatform::SP_PCD3D_SM5).AppendString(InOutKey);
+}
+
+void UPCGStaticMeshDataInterface::GetHLSL(FString& OutHLSL, FString const& InDataInterfaceName) const
+{
+	TMap<FString, FStringFormatArg> TemplateArgs =
+	{
+		{ TEXT("DataInterfaceName"), InDataInterfaceName },
+	};
+
+	FString TemplateFile;
+	if (ensure(LoadShaderSourceFile(TemplateFilePath, EShaderPlatform::SP_PCD3D_SM5, &TemplateFile, nullptr)))
+	{
+		OutHLSL += FString::Format(*TemplateFile, TemplateArgs);
+	}
+}
+
+UComputeDataProvider* UPCGStaticMeshDataInterface::CreateDataProvider(TObjectPtr<UObject> InBinding, uint64 InInputMask, uint64 InOutputMask) const
+{
+	TRACE_CPUPROFILER_EVENT_SCOPE(UPCGStaticMeshDataInterface::CreateDataProvider);
+
+	TObjectPtr<UPCGStaticMeshDataProvider> DataProvider = NewObject<UPCGStaticMeshDataProvider>();
+
+	// Take any input pin label alias to obtain the data from the input data collection.
+	check(!DownstreamInputPinLabelAliases.IsEmpty());
+	DataProvider->DownstreamInputPinLabels = DownstreamInputPinLabelAliases;
+
+	return DataProvider;
+}
+
+bool UPCGStaticMeshDataProvider::PrepareForExecute_GameThread(UPCGDataBinding* InBinding)
+{
+	TRACE_CPUPROFILER_EVENT_SCOPE(UPCGStaticMeshDataProvider::PrepareForExecute_GameThread);
+
+	if (!LoadHandle.IsValid())
+	{
+		// Take any input pin label alias to obtain the data from the input data collection.
+		check(!DownstreamInputPinLabels.IsEmpty());
+
+		const UPCGDataBinding* Binding = CastChecked<UPCGDataBinding>(InBinding);
+		const TArray<FPCGTaggedData> TaggedDatas = Binding->InputDataCollection.GetInputsByPin(DownstreamInputPinLabels[0]);
+
+		if (!TaggedDatas.IsEmpty())
+		{
+			const UPCGStaticMeshResourceData* ResourceData = Cast<UPCGStaticMeshResourceData>(TaggedDatas[0].Data);
+			ensure(TaggedDatas.Num() == 1); // There should only be one static mesh data
+
+			if (ensure(ResourceData))
+			{
+				LoadHandle = ResourceData->RequestResourceLoad();
+			}
+		}
+	}
+
+	if (!LoadHandle.IsValid())
+	{
+		return true;
+	}
+
+	if (LoadHandle->HasLoadCompleted())
+	{
+		LoadedStaticMesh = Cast<UStaticMesh>(LoadHandle->GetLoadedAsset());
+		return true;
+	}
+
+	return false;
+}
+
+FComputeDataProviderRenderProxy* UPCGStaticMeshDataProvider::GetRenderProxy()
+{
+	return new FPCGStaticMeshDataProviderProxy(LoadedStaticMesh.Get());
+}
+
+FPCGStaticMeshDataProviderProxy::FPCGStaticMeshDataProviderProxy(const UStaticMesh* InStaticMesh)
+{
+	const FStaticMeshRenderData* RenderData = InStaticMesh ? InStaticMesh->GetRenderData() : nullptr;
+
+	if (RenderData && !RenderData->LODResources.IsEmpty())
+	{
+		MeshName = InStaticMesh->GetFName();
+		LODResources = &RenderData->LODResources[0];
+		Bounds = RenderData->Bounds;
+	}
+}
+
+FPCGStaticMeshDataProviderProxy::~FPCGStaticMeshDataProviderProxy()
+{
+	IndexBufferSRV.SafeRelease();
+}
+
+bool FPCGStaticMeshDataProviderProxy::IsValid(FValidationData const& InValidationData) const
+{
+	if (InValidationData.ParameterStructSize != sizeof(FParameters))
+	{
+		return false;
+	}
+
+	if (!LODResources)
+	{
+		return false;
+	}
+
+	return true;
+}
+
+void FPCGStaticMeshDataProviderProxy::GatherDispatchData(FDispatchData const& InDispatchData)
+{
+	check(LODResources);
+
+	const FStaticMeshVertexBuffers& VertexBuffers = LODResources->VertexBuffers;
+	FShaderResourceViewRHIRef PositionBufferSRV = VertexBuffers.PositionVertexBuffer.GetSRV();
+	FShaderResourceViewRHIRef TangentBufferSRV = VertexBuffers.StaticMeshVertexBuffer.GetTangentsSRV();
+	FShaderResourceViewRHIRef UVBufferSRV = VertexBuffers.StaticMeshVertexBuffer.GetTexCoordsSRV();
+	FShaderResourceViewRHIRef ColorBufferSRV = VertexBuffers.ColorVertexBuffer.GetColorComponentsSRV();
+
+	const TStridedView<FParameters> ParameterArray = MakeStridedParameterView<FParameters>(InDispatchData);
+	for (int32 InvocationIndex = 0; InvocationIndex < ParameterArray.Num(); ++InvocationIndex)
+	{
+		FParameters& Parameters = ParameterArray[InvocationIndex];
+		Parameters.NumTriangles = IndexBufferSRV.IsValid() ? LODResources->IndexBuffer.GetNumIndices() / 3 : 0;
+		Parameters.NumVertices = PositionBufferSRV.IsValid() ? VertexBuffers.PositionVertexBuffer.GetNumVertices() : 0;
+		Parameters.NumUVs = UVBufferSRV.IsValid() ? VertexBuffers.StaticMeshVertexBuffer.GetNumTexCoords() : 0;
+		Parameters.HasColors = ColorBufferSRV.IsValid();
+		Parameters.IndexBuffer = IndexBufferSRV;
+		Parameters.PositionBuffer = PositionBufferSRV;
+		Parameters.TangentBuffer = TangentBufferSRV.IsValid() ? TangentBufferSRV : static_cast<FShaderResourceViewRHIRef>(PCGComputeDummies::GetDummyFloat4Buffer());
+		Parameters.UVBuffer = UVBufferSRV.IsValid() ? UVBufferSRV : static_cast<FShaderResourceViewRHIRef>(PCGComputeDummies::GetDummyFloat2Buffer());
+		Parameters.ColorBuffer = ColorBufferSRV.IsValid() ? ColorBufferSRV : static_cast<FShaderResourceViewRHIRef>(PCGComputeDummies::GetDummyFloat4Buffer());
+		Parameters.BoundsExtents = static_cast<FVector3f>(Bounds.BoxExtent);
+	}
+}
+
+void FPCGStaticMeshDataProviderProxy::AllocateResources(FRDGBuilder& GraphBuilder, FAllocationData const& InAllocationData)
+{
+	FBufferRHIRef IndexBufferRHIRef = LODResources ? LODResources->IndexBuffer.IndexBufferRHI : nullptr;
+	const bool bCanCreateIndexSRV = IndexBufferRHIRef.IsValid() && ((IndexBufferRHIRef->GetUsage() & EBufferUsageFlags::ShaderResource) == EBufferUsageFlags::ShaderResource);
+
+	if (bCanCreateIndexSRV)
+	{
+		IndexBufferSRV = GraphBuilder.RHICmdList.CreateShaderResourceView(IndexBufferRHIRef,
+			FRHIViewDesc::CreateBufferSRV()
+				.SetType(FRHIViewDesc::EBufferType::Typed)
+				.SetFormat(LODResources->IndexBuffer.Is32Bit() ? PF_R32_UINT : PF_R16_UINT));
+	}
+	else
+	{
+		UE_LOG(LogPCG, Error,
+			TEXT("PCGStaticMeshDataInterface used by PCG Graph but does not have SRV access on this platform. Mesh: '%s'."),
+			*MeshName.ToString());
+	}
+}
+
+#undef LOCTEXT_NAMESPACE
diff --git a/Engine/Plugins/PCG/Source/PCG/Private/Compute/DataInterfaces/PCGStaticMeshDataInterface.h b/Engine/Plugins/PCG/Source/PCG/Private/Compute/DataInterfaces/PCGStaticMeshDataInterface.h
new file mode 100644
--- /dev/null
+++ b/Engine/Plugins/PCG/Source/PCG/Private/Compute/DataInterfaces/PCGStaticMeshDataInterface.h
@@ -0,0 +1,80 @@
+// Copyright Epic Games, Inc. All Rights Reserved.
+
+#pragma once
+
+#include "Compute/DataInterfaces/PCGComputeDataInterface.h"
+
+#include "Engine/StreamableManager.h"
+
+#include "PCGStaticMeshDataInterface.generated.h"
+
+class FPCGStaticMeshDataInterfaceParameters;
+class UPCGStaticMeshResourceData;
+
+class FRHIShaderResourceView;
+class UStaticMesh;
+struct FStaticMeshLODResources;
+
+/** Data Interface allowing sampling of a static mesh. */
+UCLASS(ClassGroup = (Procedural))
+class UPCGStaticMeshDataInterface : public UPCGComputeDataInterface
+{
+	GENERATED_BODY()
+
+public:
+	//~ Begin UComputeDataInterface Interface
+	TCHAR const* GetClassName() const override { return TEXT("PCGStaticMesh"); }
+	void GetSupportedInputs(TArray<FShaderFunctionDefinition>& OutFunctions) const override;
+	void GetShaderParameters(TCHAR const* UID, FShaderParametersMetadataBuilder& InOutBuilder, FShaderParametersMetadataAllocations& InOutAllocations) const override;
+	TCHAR const* GetShaderVirtualPath() const override;
+	void GetShaderHash(FString& InOutKey) const override;
+	void GetHLSL(FString& OutHLSL, FString const& InDataInterfaceName) const override;
+	UComputeDataProvider* CreateDataProvider(TObjectPtr<UObject> InBinding, uint64 InInputMask, uint64 InOutputMask) const override;
+	//~ End UComputeDataInterface Interface
+
+private:
+	static TCHAR const* TemplateFilePath;
+};
+
+/** Compute Framework Data Provider for reading a static mesh. */
+UCLASS()
+class UPCGStaticMeshDataProvider : public UPCGComputeDataProvider
+{
+	GENERATED_BODY()
+
+public:
+	//~ Begin UComputeDataProvider Interface
+	FComputeDataProviderRenderProxy* GetRenderProxy() override;
+	//~ End UComputeDataProvider Interface
+	
+	//~ Begin UPCGComputeDataProvider Interface
+	virtual bool PrepareForExecute_GameThread(UPCGDataBinding* InBinding) override;
+	//~ End UPCGComputeDataProvider Interface
+
+	UPROPERTY()
+	TObjectPtr<const UStaticMesh> LoadedStaticMesh = nullptr;
+
+	TArray<FName> DownstreamInputPinLabels;
+	TSharedPtr<FStreamableHandle> LoadHandle = nullptr;
+};
+
+class FPCGStaticMeshDataProviderProxy : public FComputeDataProviderRenderProxy
+{
+public:
+	FPCGStaticMeshDataProviderProxy(const UStaticMesh* InStaticMesh);
+	virtual ~FPCGStaticMeshDataProviderProxy();
+
+	//~ Begin FComputeDataProviderRenderProxy Interface
+	bool IsValid(FValidationData const& InValidationData) const override;
+	void GatherDispatchData(FDispatchData const& InDispatchData) override;
+	void AllocateResources(FRDGBuilder& GraphBuilder, FAllocationData const& InAllocationData) override;
+	//~ End FComputeDataProviderRenderProxy Interface
+
+protected:
+	using FParameters = FPCGStaticMeshDataInterfaceParameters;
+
+	FName MeshName;
+	TRefCountPtr<const FStaticMeshLODResources> LODResources = nullptr;
+	FBoxSphereBounds Bounds = FBoxSphereBounds(EForceInit::ForceInit);
+	TRefCountPtr<FRHIShaderResourceView> IndexBufferSRV = nullptr;
+};
diff --git a/Engine/Plugins/PCG/Source/PCG/Private/Compute/Elements/PCGCustomHLSL.cpp b/Engine/Plugins/PCG/Source/PCG/Private/Compute/Elements/PCGCustomHLSL.cpp
--- a/Engine/Plugins/PCG/Source/PCG/Private/Compute/Elements/PCGCustomHLSL.cpp
+++ b/Engine/Plugins/PCG/Source/PCG/Private/Compute/Elements/PCGCustomHLSL.cpp
@@ -409,6 +409,7 @@ void UPCGCustomHLSLSettings::UpdateInputDeclarations()
 	TArray<FString> LandscapeDataPins;
 	TArray<FString> TextureDataPins;
 	TArray<FString> RawBufferDataPins;
+	TArray<FString> StaticMeshDataPins;
 
 	for (const FPCGPinProperties& Pin : InputPinProperties())
 	{
@@ -436,6 +437,11 @@ void UPCGCustomHLSLSettings::UpdateInputDeclarations()
 		{
 			TextureDataPins.Add(Pin.Label.ToString());
 		}
+
+		if (!!(Pin.AllowedTypes & EPCGDataType::StaticMeshResource))
+		{
+			StaticMeshDataPins.Add(Pin.Label.ToString());
+		}
 	}
 
 	if (!DataCollectionDataPins.IsEmpty())
@@ -552,6 +558,38 @@ void UPCGCustomHLSLSettings::UpdateInputDeclarations()
 		InputDeclarations += TEXT("\n");
 	}
 
+	if (!StaticMeshDataPins.IsEmpty())
+	{
+		InputDeclarations += TEXT("/*** INPUT STATIC MESH DATA FUNCTIONS ***/\n\n");
+
+		const bool bMultiPin = StaticMeshDataPins.Num() > 1;
+
+		if (bMultiPin)
+		{
+			InputDeclarations += TEXT("// Valid pins: ") + FString::Join(StaticMeshDataPins, TEXT(", ")) + TEXT("\n\n");
+		}
+
+		InputDeclarations += FString::Format(TEXT(
+			"// Vertex functions\n"
+			"int {0}_GetNumVertices(int DataIndex);\n"
+			"void {0}_SampleVertex(int DataIndex, int VertexIndex, float3 OutPosition, float3 OutNormal, float3 OutTangent, float3 OutBitangent);\n"
+			"float4 {0}_SampleVertexColor(int DataIndex, int VertexIndex);\n"
+			"float2 {0}_SampleVertexUVs(int DataIndex, int VertexIndex, int UVSet);\n"
+			"\n"
+			"// Triangle functions\n"
+			"int {0}_GetNumTriangles(int DataIndex);\n"
+			"void {0}_GetTriangleIndices(int DataIndex, int TriangleIndex, int OutIndex0, int OutIndex1, int OutIndex2);\n"
+			"void {0}_SampleTriangle(int DataIndex, int TriangleIndex, float3 BaryCoords, float3 OutPosition, float3 OutNormal, float3 OutTangent, float3 OutBitangent);\n"
+			"float4 {0}_SampleTriangleColor(int DataIndex, int TriangleIndex, float3 BaryCoords);\n"
+			"float2 {0}_SampleTriangleUVs(int DataIndex, int TriangleIndex, float3 BaryCoords, int UVSet);\n"
+			"\n"
+			"// Get bounds extents of the static mesh.\n"
+			"float3 {0}_GetMeshBoundsExtents(int DataIndex);\n"),
+			{ bMultiPin ? PCGHLSLElement::PinDeclTemplateStr : StaticMeshDataPins[0] });
+
+		InputDeclarations += TEXT("\n");
+	}
+
 	InputDeclarations.TrimStartAndEndInline();
 }
 
diff --git a/Engine/Plugins/PCG/Source/PCG/Private/Compute/PCGComputeCommon.cpp b/Engine/Plugins/PCG/Source/PCG/Private/Compute/PCGComputeCommon.cpp
--- a/Engine/Plugins/PCG/Source/PCG/Private/Compute/PCGComputeCommon.cpp
+++ b/Engine/Plugins/PCG/Source/PCG/Private/Compute/PCGComputeCommon.cpp
@@ -13,6 +13,7 @@
 
 #include "DynamicRHI.h"
 #include "RHIStats.h"
+#include "RenderResource.h"
 
 namespace PCGComputeHelpers
 {
@@ -176,6 +177,65 @@ namespace PCGComputeHelpers
 	}
 }
 
+namespace PCGComputeDummies
+{
+	class FPCGEmptyBufferSRV : public FRenderResource
+	{
+	public:
+		FPCGEmptyBufferSRV(EPixelFormat InPixelFormat, const FString& InDebugName)
+			: PixelFormat(InPixelFormat)
+			, DebugName(InDebugName)
+		{}
+
+		EPixelFormat PixelFormat;
+		FString DebugName;
+		FBufferRHIRef Buffer = nullptr;
+		FShaderResourceViewRHIRef BufferSRV = nullptr;
+	
+		virtual void InitRHI(FRHICommandListBase& RHICmdList) override
+		{
+			// Create a buffer with one element.
+			const uint32 NumBytes = GPixelFormats[PixelFormat].BlockBytes;
+
+			FRHIBufferCreateDesc CreateDesc =
+				FRHIBufferCreateDesc::CreateVertex(*DebugName, NumBytes)
+				.AddUsage(EBufferUsageFlags::ShaderResource | EBufferUsageFlags::Static)
+				.DetermineInitialState();
+
+			Buffer = RHICmdList.CreateBuffer(CreateDesc.SetInitActionZeroData());
+			BufferSRV = RHICmdList.CreateShaderResourceView(
+				Buffer,
+				FRHIViewDesc::CreateBufferSRV()
+					.SetType(FRHIViewDesc::EBufferType::Typed)
+					.SetFormat(PixelFormat));
+		}
+	
+		virtual void ReleaseRHI() override
+		{
+			BufferSRV.SafeRelease();
+			Buffer.SafeRelease();
+		}
+	};
+	
+	FRHIShaderResourceView* GetDummyFloatBuffer()
+	{
+		static TGlobalResource<FPCGEmptyBufferSRV> DummyFloatBuffer(PF_R32_FLOAT, TEXT("PCGDummyFloat"));
+		return DummyFloatBuffer.BufferSRV;
+	}
+	
+	FRHIShaderResourceView* GetDummyFloat2Buffer()
+	{
+		static TGlobalResource<FPCGEmptyBufferSRV> DummyFloat2Buffer(PF_G32R32F, TEXT("PCGDummyFloat2"));
+		return DummyFloat2Buffer.BufferSRV;
+	}
+	
+	FRHIShaderResourceView* GetDummyFloat4Buffer()
+	{
+		static TGlobalResource<FPCGEmptyBufferSRV> DummyFloat4Buffer(PF_A32B32G32R32F, TEXT("PCGDummyFloat4"));
+		return DummyFloat4Buffer.BufferSRV;
+	}
+}
+
 uint32 GetTypeHash(const FPCGPinReference& In)
 {
 	return HashCombine(/*GetTypeHash(In.TaskId),*/ PointerHash(In.Kernel), GetTypeHash(In.Label));
diff --git a/Engine/Plugins/PCG/Source/PCG/Private/Data/PCGResourceData.cpp b/Engine/Plugins/PCG/Source/PCG/Private/Data/PCGResourceData.cpp
new file mode 100644
--- /dev/null
+++ b/Engine/Plugins/PCG/Source/PCG/Private/Data/PCGResourceData.cpp
@@ -0,0 +1,20 @@
+// Copyright Epic Games, Inc. All Rights Reserved.
+
+#include "Data/PCGResourceData.h"
+
+#include "Engine/AssetManager.h"
+
+#include UE_INLINE_GENERATED_CPP_BY_NAME(PCGResourceData)
+
+// @todo_pcg: A preferable API might be to hold the load handle as a member and provide functions to query the state.
+TSharedPtr<FStreamableHandle> UPCGResourceData::RequestResourceLoad(bool bAsynchronous) const
+{
+	if (!bAsynchronous)
+	{
+		return UAssetManager::GetStreamableManager().RequestSyncLoad(GetResourcePath());
+	}
+	else
+	{
+		return UAssetManager::GetStreamableManager().RequestAsyncLoad(GetResourcePath(), nullptr);
+	}
+}
diff --git a/Engine/Plugins/PCG/Source/PCG/Private/Data/PCGStaticMeshResourceData.cpp b/Engine/Plugins/PCG/Source/PCG/Private/Data/PCGStaticMeshResourceData.cpp
new file mode 100644
--- /dev/null
+++ b/Engine/Plugins/PCG/Source/PCG/Private/Data/PCGStaticMeshResourceData.cpp
@@ -0,0 +1,15 @@
+// Copyright Epic Games, Inc. All Rights Reserved.
+
+#include "Data/PCGStaticMeshResourceData.h"
+
+#include UE_INLINE_GENERATED_CPP_BY_NAME(PCGStaticMeshResourceData)
+
+void UPCGStaticMeshResourceData::Initialize(TSoftObjectPtr<UStaticMesh> InStaticMesh)
+{
+	StaticMesh = InStaticMesh;
+}
+
+FSoftObjectPath UPCGStaticMeshResourceData::GetResourcePath() const
+{
+	return StaticMesh.ToSoftObjectPath();
+}
diff --git a/Engine/Plugins/PCG/Source/PCG/Private/Elements/PCGGetResourcePath.cpp b/Engine/Plugins/PCG/Source/PCG/Private/Elements/PCGGetResourcePath.cpp
new file mode 100644
--- /dev/null
+++ b/Engine/Plugins/PCG/Source/PCG/Private/Elements/PCGGetResourcePath.cpp
@@ -0,0 +1,60 @@
+// Copyright Epic Games, Inc. All Rights Reserved.
+
+#include "Elements/PCGGetResourcePath.h"
+
+#include "PCGContext.h"
+#include "PCGParamData.h"
+#include "Data/PCGResourceData.h"
+
+#include UE_INLINE_GENERATED_CPP_BY_NAME(PCGGetResourcePath)
+
+TArray<FPCGPinProperties> UPCGGetResourcePath::InputPinProperties() const
+{
+	TArray<FPCGPinProperties> PinProperties;
+	FPCGPinProperties& InputPinProperty = PinProperties.Emplace_GetRef(PCGPinConstants::DefaultInputLabel, EPCGDataType::Resource);
+	InputPinProperty.SetRequiredPin();
+
+	return PinProperties;
+}
+
+TArray<FPCGPinProperties> UPCGGetResourcePath::OutputPinProperties() const
+{
+	TArray<FPCGPinProperties> PinProperties;
+	PinProperties.Emplace(PCGPinConstants::DefaultOutputLabel, EPCGDataType::Param);
+
+	return PinProperties;
+}
+
+FPCGElementPtr UPCGGetResourcePath::CreateElement() const
+{
+	return MakeShared<FPCGGetResourcePathElement>();
+}
+
+bool FPCGGetResourcePathElement::ExecuteInternal(FPCGContext* Context) const
+{
+	TRACE_CPUPROFILER_EVENT_SCOPE(FPCGGetResourcePathElement::Execute);
+
+	check(Context);
+
+	for (const FPCGTaggedData& Input : Context->InputData.TaggedData)
+	{
+		const UPCGResourceData* ResourceData = Cast<UPCGResourceData>(Input.Data);
+
+		if (!ResourceData)
+		{
+			continue;
+		}
+
+		UPCGParamData* ParamData = FPCGContext::NewObject_AnyThread<UPCGParamData>(Context);
+		check(ParamData);
+		check(ParamData->Metadata);
+		
+		ParamData->Metadata->CreateSoftObjectPathAttribute(TEXT("ResourceReference"), ResourceData->GetResourcePath(), /*bAllowsInterpolation=*/false, /*bOverrideParent=*/false);
+		ParamData->Metadata->AddEntry();
+
+		FPCGTaggedData& Output = Context->OutputData.TaggedData.Add_GetRef(Input);
+		Output.Data = ParamData;
+	}
+
+	return true;
+}
diff --git a/Engine/Plugins/PCG/Source/PCG/Private/Elements/PCGGetStaticMeshResourceData.cpp b/Engine/Plugins/PCG/Source/PCG/Private/Elements/PCGGetStaticMeshResourceData.cpp
new file mode 100644
--- /dev/null
+++ b/Engine/Plugins/PCG/Source/PCG/Private/Elements/PCGGetStaticMeshResourceData.cpp
@@ -0,0 +1,136 @@
+// Copyright Epic Games, Inc. All Rights Reserved.
+
+#include "Elements/PCGGetStaticMeshResourceData.h"
+
+#include "PCGContext.h"
+#include "PCGModule.h"
+#include "Metadata/Accessors/PCGAttributeAccessorHelpers.h"
+
+#include "Engine/StaticMesh.h"
+
+#include UE_INLINE_GENERATED_CPP_BY_NAME(PCGGetStaticMeshResourceData)
+
+#define LOCTEXT_NAMESPACE "PCGGetStaticMeshResourceDataElement"
+
+namespace PCGGetStaticMeshResourceDataConstants
+{
+	const FName MeshOverridesPinLabel = TEXT("Meshes");
+}
+
+TArray<FPCGPinProperties> UPCGGetStaticMeshResourceDataSettings::InputPinProperties() const
+{
+	TArray<FPCGPinProperties> Properties;
+
+	if (bOverrideFromInput)
+	{
+		FPCGPinProperties& MeshOverrides = Properties.Emplace_GetRef(PCGGetStaticMeshResourceDataConstants::MeshOverridesPinLabel, EPCGDataType::Param);
+		MeshOverrides.SetRequiredPin();
+	}
+
+	return Properties;
+}
+
+TArray<FPCGPinProperties> UPCGGetStaticMeshResourceDataSettings::OutputPinProperties() const
+{
+	TArray<FPCGPinProperties> Properties;
+	Properties.Emplace(
+		PCGPinConstants::DefaultOutputLabel,
+		EPCGDataType::StaticMeshResource,
+		/*bInAllowMultipleConnections=*/true,
+		/*bAllowMultipleData=*/true);
+
+	return Properties;
+}
+
+FPCGElementPtr UPCGGetStaticMeshResourceDataSettings::CreateElement() const
+{
+	return MakeShared<FPCGGetStaticMeshResourceDataElement>();
+}
+
+#if WITH_EDITOR
+EPCGChangeType UPCGGetStaticMeshResourceDataSettings::GetChangeTypeForProperty(const FName& InPropertyName) const
+{
+	EPCGChangeType ChangeType = Super::GetChangeTypeForProperty(InPropertyName);
+
+	if (InPropertyName == GET_MEMBER_NAME_CHECKED(UPCGGetStaticMeshResourceDataSettings, bOverrideFromInput))
+	{
+		ChangeType |= EPCGChangeType::Structural;
+	}
+
+	return ChangeType;
+}
+#endif // WITH_EDITOR
+
+
+bool FPCGGetStaticMeshResourceDataElement::ExecuteInternal(FPCGContext* InContext) const
+{
+	TRACE_CPUPROFILER_EVENT_SCOPE(FPCGGetStaticMeshResourceDataElement::Execute);
+	check(InContext);
+
+	const UPCGGetStaticMeshResourceDataSettings* Settings = InContext->GetInputSettings<UPCGGetStaticMeshResourceDataSettings>();
+	check(Settings);
+
+	TArray<TSoftObjectPtr<UStaticMesh>> StaticMeshes;
+
+	if (!Settings->bOverrideFromInput)
+	{
+		StaticMeshes = Settings->StaticMeshes;
+	}
+	else
+	{
+		const TArray<FPCGTaggedData> OverrideTaggedDatas = InContext->InputData.GetInputsByPin(PCGGetStaticMeshResourceDataConstants::MeshOverridesPinLabel);
+		TArray<FSoftObjectPath> MeshOverrides;
+
+		for (const FPCGTaggedData& OverrideTaggedData : OverrideTaggedDatas)
+		{
+			if (const UPCGData* OverrideData = OverrideTaggedData.Data)
+			{
+				const FPCGAttributePropertyInputSelector MeshSelector = Settings->MeshAttribute.CopyAndFixLast(OverrideData);
+
+				if (PCGAttributeAccessorHelpers::ExtractAllValues(OverrideData, MeshSelector, MeshOverrides, InContext))
+				{
+					for (const FSoftObjectPath& MeshOverride : MeshOverrides)
+					{
+						StaticMeshes.Emplace(MeshOverride);
+					}
+				}
+				else
+				{
+					PCGLog::LogWarningOnGraph(LOCTEXT("FailExtractMeshOverrides", "Failed to extract static mesh overrides."), InContext);
+				}
+			}
+		}
+	}
+
+	// We only want to create one resource data per unique mesh, so we'll cache them here while looping over the input.
+	TMap<TSoftObjectPtr<UStaticMesh>, UPCGStaticMeshResourceData*> CreatedResourceDatas;
+
+	for (const TSoftObjectPtr<UStaticMesh>& StaticMesh : StaticMeshes)
+	{
+		if (StaticMesh.IsNull())
+		{
+			continue;
+		}
+		
+		UPCGStaticMeshResourceData** StaticMeshResourceDataPtr = CreatedResourceDatas.Find(StaticMesh);
+		UPCGStaticMeshResourceData* StaticMeshResourceData = nullptr;
+
+		if (StaticMeshResourceDataPtr)
+		{
+			StaticMeshResourceData = *StaticMeshResourceDataPtr;
+		}
+		else
+		{
+			StaticMeshResourceData = FPCGContext::NewObject_AnyThread<UPCGStaticMeshResourceData>(InContext);
+			StaticMeshResourceData->Initialize(StaticMesh);
+
+			CreatedResourceDatas.Add(StaticMesh, StaticMeshResourceData);
+		}
+
+		InContext->OutputData.TaggedData.Emplace_GetRef().Data = StaticMeshResourceData;
+	}
+
+	return true;
+}
+
+#undef LOCTEXT_NAMESPACE
diff --git a/Engine/Plugins/PCG/Source/PCG/Private/Graph/PCGGraphCompilerGPU.cpp b/Engine/Plugins/PCG/Source/PCG/Private/Graph/PCGGraphCompilerGPU.cpp
--- a/Engine/Plugins/PCG/Source/PCG/Private/Graph/PCGGraphCompilerGPU.cpp
+++ b/Engine/Plugins/PCG/Source/PCG/Private/Graph/PCGGraphCompilerGPU.cpp
@@ -16,6 +16,7 @@
 #include "Compute/DataInterfaces/PCGDataLabelResolverDataInterface.h"
 #include "Compute/DataInterfaces/PCGDebugDataInterface.h"
 #include "Compute/DataInterfaces/PCGLandscapeDataInterface.h"
+#include "Compute/DataInterfaces/PCGStaticMeshDataInterface.h"
 #include "Compute/DataInterfaces/PCGTextureDataInterface.h"
 #include "Compute/Elements/PCGComputeGraphElement.h"
 #include "Elements/PCGStaticMeshSpawnerKernel.h"
@@ -900,6 +901,11 @@ void FPCGGraphCompilerGPU::CreateDataInterfaces(
 			DataInterface = InOutContext.NewObject_AnyThread<UPCGLandscapeDataInterface>(InOutComputeGraph);
 			break;
 		}
+		case EPCGDataType::StaticMeshResource:
+		{
+			DataInterface = InOutContext.NewObject_AnyThread<UPCGStaticMeshDataInterface>(InOutComputeGraph);
+			break;
+		}
 		default:
 		{
 			UE_LOG(LogPCG, Error, TEXT("Unsupported connected upstream pin '%s' on node '%s' with type %s. Consider adding a conversion to a supported type such as Point."),
diff --git a/Engine/Plugins/PCG/Source/PCG/Public/Data/PCGResourceData.h b/Engine/Plugins/PCG/Source/PCG/Public/Data/PCGResourceData.h
new file mode 100644
--- /dev/null
+++ b/Engine/Plugins/PCG/Source/PCG/Public/Data/PCGResourceData.h
@@ -0,0 +1,25 @@
+// Copyright Epic Games, Inc. All Rights Reserved.
+
+#pragma once
+
+#include "PCGData.h"
+
+#include "Engine/StreamableManager.h"
+
+#include "PCGResourceData.generated.h"
+
+/** Data that wrap/represent an asset, like a Static Mesh or Texture. */
+UCLASS(Abstract, BlueprintType, ClassGroup = (Procedural))
+class PCG_API UPCGResourceData : public UPCGData
+{
+	GENERATED_BODY()
+
+public:
+	//~ Begin UPCGData interface
+	virtual EPCGDataType GetDataType() const override { return EPCGDataType::Resource; }
+	//~ End UPCGData interface
+
+	virtual FSoftObjectPath GetResourcePath() const PURE_VIRTUAL(UPCGResourceData::GetResourcePath, return FSoftObjectPath(););
+
+	TSharedPtr<FStreamableHandle> RequestResourceLoad(bool bAsynchronous = true) const;
+};
diff --git a/Engine/Plugins/PCG/Source/PCG/Public/Data/PCGStaticMeshResourceData.h b/Engine/Plugins/PCG/Source/PCG/Public/Data/PCGStaticMeshResourceData.h
new file mode 100644
--- /dev/null
+++ b/Engine/Plugins/PCG/Source/PCG/Public/Data/PCGStaticMeshResourceData.h
@@ -0,0 +1,31 @@
+// Copyright Epic Games, Inc. All Rights Reserved.
+
+#pragma once
+
+#include "Data/PCGResourceData.h"
+
+#include "PCGStaticMeshResourceData.generated.h"
+
+class UStaticMesh;
+
+/** Data that wraps a Static Mesh soft object path. */
+UCLASS(BlueprintType, ClassGroup = (Procedural))
+class PCG_API UPCGStaticMeshResourceData : public UPCGResourceData
+{
+	GENERATED_BODY()
+
+public:
+	void Initialize(TSoftObjectPtr<UStaticMesh> InStaticMesh);
+
+	//~ Begin UPCGData interface
+	virtual EPCGDataType GetDataType() const override { return EPCGDataType::StaticMeshResource; }
+	//~ End UPCGData interface
+
+	//~ Begin UPCGResourceData interface
+	virtual FSoftObjectPath GetResourcePath() const override;
+	//~ End UPCGResourceData interface
+
+protected:
+	UPROPERTY(BlueprintReadWrite, EditAnywhere, Category = StaticMesh)
+	TSoftObjectPtr<UStaticMesh> StaticMesh = nullptr;
+};
\ No newline at end of file
diff --git a/Engine/Plugins/PCG/Source/PCG/Public/Elements/PCGGetResourcePath.h b/Engine/Plugins/PCG/Source/PCG/Public/Elements/PCGGetResourcePath.h
new file mode 100644
--- /dev/null
+++ b/Engine/Plugins/PCG/Source/PCG/Public/Elements/PCGGetResourcePath.h
@@ -0,0 +1,37 @@
+// Copyright Epic Games, Inc. All Rights Reserved.
+
+#pragma once
+
+#include "PCGContext.h"
+#include "PCGSettings.h"
+
+#include "PCGGetResourcePath.generated.h"
+
+/** Converts a resource data to an attribute set containing the resource path. */
+UCLASS(MinimalAPI, BlueprintType, ClassGroup = (Procedural))
+class UPCGGetResourcePath : public UPCGSettings
+{
+	GENERATED_BODY()
+
+public:
+	//~ Begin UPCGSettings interface
+#if WITH_EDITOR
+	virtual FName GetDefaultNodeName() const override { return FName(TEXT("GetResourcePath")); }
+	virtual FText GetDefaultNodeTitle() const override { return NSLOCTEXT("PCGGetResourcePathElement", "NodeTitle", "Get Resource Path"); }
+	virtual EPCGSettingsType GetType() const override { return EPCGSettingsType::Resource; }
+#endif
+	virtual bool HasExecutionDependencyPin() const override { return false; }
+
+protected:
+	virtual TArray<FPCGPinProperties> InputPinProperties() const override;
+	virtual TArray<FPCGPinProperties> OutputPinProperties() const override;
+	virtual FPCGElementPtr CreateElement() const override;
+	//~ End UPCGSettings interface
+};
+
+class FPCGGetResourcePathElement : public IPCGElement
+{
+protected:
+	virtual bool ExecuteInternal(FPCGContext* Context) const override;
+	virtual bool SupportsGPUResidentData(FPCGContext* InContext) const override { return true; }
+};
diff --git a/Engine/Plugins/PCG/Source/PCG/Public/Elements/PCGGetStaticMeshResourceData.h b/Engine/Plugins/PCG/Source/PCG/Public/Elements/PCGGetStaticMeshResourceData.h
new file mode 100644
--- /dev/null
+++ b/Engine/Plugins/PCG/Source/PCG/Public/Elements/PCGGetStaticMeshResourceData.h
@@ -0,0 +1,56 @@
+// Copyright Epic Games, Inc. All Rights Reserved.
+
+#pragma once
+
+#include "PCGContext.h"
+#include "PCGSettings.h"
+#include "Async/PCGAsyncLoadingContext.h"
+#include "Data/PCGStaticMeshResourceData.h"
+
+#include "PCGGetStaticMeshResourceData.generated.h"
+
+class UStaticMesh;
+
+/** Creates static mesh resource data from the given soft object paths. */
+UCLASS(MinimalAPI, BlueprintType, ClassGroup = (Procedural))
+class UPCGGetStaticMeshResourceDataSettings : public UPCGSettings
+{
+	GENERATED_BODY()
+
+public:
+	//~ Begin UPCGSettings interface
+#if WITH_EDITOR
+	virtual FName GetDefaultNodeName() const override { return FName(TEXT("GetStaticMeshResourceData")); }
+	virtual FText GetDefaultNodeTitle() const override { return NSLOCTEXT("PCGGetStaticMeshResourceDataElement", "NodeTitle", "Get Static Mesh Resource Data"); }
+	virtual EPCGSettingsType GetType() const override { return EPCGSettingsType::Resource; }
+#endif
+
+protected:
+	virtual TArray<FPCGPinProperties> InputPinProperties() const override;
+	virtual TArray<FPCGPinProperties> OutputPinProperties() const override;
+	virtual FPCGElementPtr CreateElement() const override;
+
+#if WITH_EDITOR
+	virtual EPCGChangeType GetChangeTypeForProperty(const FName& InPropertyName) const override;
+#endif
+	//~ End UPCGSettings interface
+
+public:
+	/** Produces one resource data per entry. */
+	UPROPERTY(BlueprintReadWrite, EditAnywhere, Category = Settings, meta = (EditCondition = "!bOverrideFromInput"))
+	TArray<TSoftObjectPtr<UStaticMesh>> StaticMeshes;
+
+	/** Override static meshes from input. */
+	UPROPERTY(BlueprintReadWrite, EditAnywhere, Category = Settings)
+	bool bOverrideFromInput = false;
+
+	/** Input attribute to pull meshes from. */
+	UPROPERTY(BlueprintReadWrite, EditAnywhere, Category = Settings, meta = (EditCondition = "bOverrideFromInput", PCG_DiscardPropertySelection, PCG_DiscardExtraSelection))
+	FPCGAttributePropertyInputSelector MeshAttribute;
+};
+
+class FPCGGetStaticMeshResourceDataElement : public IPCGElement
+{
+protected:
+	virtual bool ExecuteInternal(FPCGContext* InContext) const override;
+};
diff --git a/Engine/Plugins/PCG/Source/PCG/Public/PCGCommon.h b/Engine/Plugins/PCG/Source/PCG/Public/PCGCommon.h
--- a/Engine/Plugins/PCG/Source/PCG/Public/PCGCommon.h
+++ b/Engine/Plugins/PCG/Source/PCG/Public/PCGCommon.h
@@ -117,6 +117,8 @@ enum class EPCGDataType : uint32
 	Primitive = 1 << 8,
 	DynamicMesh = 1 << 10,
 
+	StaticMeshResource = 1 << 11,
+
 	/** Simple concrete data. */
 	Concrete = Point | PolyLine | Surface | Volume | Primitive | DynamicMesh,
 
@@ -126,6 +128,9 @@ enum class EPCGDataType : uint32
 	/** Combinations of concrete data and/or boolean operations. */
 	Spatial = Composite | Concrete,
 
+	/** Data that represent resources/assets. */
+	Resource = StaticMeshResource UMETA(Hidden),
+
 	/** Proxy for data that was created on the GPU and not necessarily read back to CPU. */
 	ProxyForGPU = 1 << 26 UMETA(Hidden),
 
@@ -164,7 +169,9 @@ enum class EPCGExclusiveDataType : uint8
 	Other,
 	Any,
 	PointOrParam UMETA(DisplayName = "Point or Attribute Set"),
-	DynamicMesh
+	DynamicMesh,
+	StaticMeshResource,
+	Resource UMETA(Hidden)
 };
 
 namespace PCGValueConstants
diff --git a/Engine/Plugins/PCG/Source/PCG/Public/PCGSettings.h b/Engine/Plugins/PCG/Source/PCG/Public/PCGSettings.h
--- a/Engine/Plugins/PCG/Source/PCG/Public/PCGSettings.h
+++ b/Engine/Plugins/PCG/Source/PCG/Public/PCGSettings.h
@@ -70,6 +70,7 @@ enum class EPCGSettingsType : uint8
 	GPU,
 	DynamicMesh,
 	DataLayers,
+	Resource,
 };
 
 #if WITH_EDITOR
diff --git a/Engine/Plugins/PCG/Source/PCGEditor/Private/Nodes/PCGEditorGraphNodeBase.cpp b/Engine/Plugins/PCG/Source/PCGEditor/Private/Nodes/PCGEditorGraphNodeBase.cpp
--- a/Engine/Plugins/PCG/Source/PCGEditor/Private/Nodes/PCGEditorGraphNodeBase.cpp
+++ b/Engine/Plugins/PCG/Source/PCGEditor/Private/Nodes/PCGEditorGraphNodeBase.cpp
@@ -1180,6 +1180,10 @@ FEdGraphPinType UPCGEditorGraphNodeBase::GetPinType(const UPCGPin* InPin)
 	{
 		EdPinType.PinCategory = FPCGEditorCommon::SpatialDataType;
 	}
+	else if (CheckType(EPCGDataType::StaticMeshResource))
+	{
+		EdPinType.PinSubCategory = FPCGEditorCommon::StaticMeshResourceDataType;
+	}
 	else if (CheckType(EPCGDataType::Param))
 	{
 		EdPinType.PinCategory = FPCGEditorCommon::ParamDataType;
diff --git a/Engine/Plugins/PCG/Source/PCGEditor/Private/PCGEditorCommon.h b/Engine/Plugins/PCG/Source/PCGEditor/Private/PCGEditorCommon.h
--- a/Engine/Plugins/PCG/Source/PCGEditor/Private/PCGEditorCommon.h
+++ b/Engine/Plugins/PCG/Source/PCGEditor/Private/PCGEditorCommon.h
@@ -49,6 +49,7 @@ namespace FPCGEditorCommon
 	const FName VolumeDataType = FName(TEXT("Volume Data"));
 	const FName PrimitiveDataType = FName(TEXT("Primitive Data"));
 	const FName DynamicMeshDataType = FName(TEXT("Dynamic Mesh Data"));
+	const FName StaticMeshResourceDataType = FName(TEXT("Static Mesh Resource Data"));
 
 	extern TAutoConsoleVariable<bool> CVarShowAdvancedAttributesFields;
 
diff --git a/Engine/Plugins/PCG/Source/PCGEditor/Private/PCGEditorSettings.cpp b/Engine/Plugins/PCG/Source/PCGEditor/Private/PCGEditorSettings.cpp
--- a/Engine/Plugins/PCG/Source/PCGEditor/Private/PCGEditorSettings.cpp
+++ b/Engine/Plugins/PCG/Source/PCGEditor/Private/PCGEditorSettings.cpp
@@ -33,6 +33,7 @@ UPCGEditorSettings::UPCGEditorSettings(const FObjectInitializer& ObjectInitializ
 	GraphParametersNodeColor = FLinearColor::Yellow;
 	RerouteNodeColor = FLinearColor(0.5f, 1.0f, 0.83f);
 	DynamicMeshNodeColor = FLinearColor(0.42f, 0.25f, 1.0f);
+	ResourceNodeColor = FLinearColor(1.0f, 0.0f, 1.0f);
 
 	DefaultPinColor = FLinearColor(0.29f, 0.29f, 0.29f);
 	SpatialDataPinColor = FLinearColor(1.0f, 1.0f, 1.0f);
@@ -47,6 +48,7 @@ UPCGEditorSettings::UPCGEditorSettings(const FObjectInitializer& ObjectInitializ
 	VolumeDataPinColor = FLinearColor(0.79f, 0.06f, 0.5f);
 	PrimitiveDataPinColor = FLinearColor(0.22f, 0.05f, 1.0f);
 	DynamicMeshPinColor = FLinearColor(0.42f, 0.25f, 1.0f);
+	StaticMeshResourcePinColor = FLinearColor(0.82f, 0.25f, 0.1f);
 
 	ParamDataPinColor = FLinearColor(1.0f, 0.38f, 0.02f);
 	UnknownDataPinColor = FLinearColor(0.3f, 0.3f, 0.3f);
@@ -102,6 +104,8 @@ FLinearColor UPCGEditorSettings::GetColor(UPCGSettings* Settings) const
 			return RerouteNodeColor;
 		case EPCGSettingsType::DynamicMesh:
 			return DynamicMeshNodeColor;
+		case EPCGSettingsType::Resource:
+			return ResourceNodeColor;
 		case EPCGSettingsType::Generic: // falls through
 		default:
 			// Finally, we couldn't find any match, so return the default value
@@ -163,6 +167,10 @@ FLinearColor UPCGEditorSettings::GetPinColor(const FEdGraphPinType& PinType) con
 	{
 		return SpatialDataPinColor;
 	}
+	else if (PinType.PinSubCategory == FPCGEditorCommon::StaticMeshResourceDataType)
+	{
+		return StaticMeshResourcePinColor;
+	}
 	else if (PinType.PinCategory == FPCGEditorCommon::ParamDataType)
 	{
 		return ParamDataPinColor;
diff --git a/Engine/Plugins/PCG/Source/PCGEditor/Public/PCGEditorSettings.h b/Engine/Plugins/PCG/Source/PCGEditor/Public/PCGEditorSettings.h
--- a/Engine/Plugins/PCG/Source/PCGEditor/Public/PCGEditorSettings.h
+++ b/Engine/Plugins/PCG/Source/PCGEditor/Public/PCGEditorSettings.h
@@ -119,6 +119,10 @@ public:
 	UPROPERTY(EditAnywhere, config, Category = Node, meta = (HideAlphaChannel))
 	FLinearColor DynamicMeshNodeColor;
 
+	/** Color used for resource nodes */
+	UPROPERTY(EditAnywhere, config, Category = Node, meta = (HideAlphaChannel))
+	FLinearColor ResourceNodeColor;
+
 	/** Default pin color */
 	UPROPERTY(EditAnywhere, config, Category = Node, meta = (HideAlphaChannel))
 	FLinearColor DefaultPinColor;
@@ -175,6 +179,10 @@ public:
 	UPROPERTY(EditAnywhere, config, Category = Node, meta = (HideAlphaChannel))
 	FLinearColor DynamicMeshPinColor;
 
+	/** Color used for data pins of type Static Mesh Resource */
+	UPROPERTY(EditAnywhere, config, Category = Node, meta = (HideAlphaChannel))
+	FLinearColor StaticMeshResourcePinColor;
+
 	/** Color used for other/unknown data types */
 	UPROPERTY(EditAnywhere, config, Category = Node, meta = (HideAlphaChannel))
 	FLinearColor UnknownDataPinColor;
