Improved the way we visit components in the debug object tree view to make it faster to update.

#rnx

diff --git a/Engine/Plugins/PCG/Source/PCG/Private/Graph/PCGStackContext.cpp b/Engine/Plugins/PCG/Source/PCG/Private/Graph/PCGStackContext.cpp
--- a/Engine/Plugins/PCG/Source/PCG/Private/Graph/PCGStackContext.cpp
+++ b/Engine/Plugins/PCG/Source/PCG/Private/Graph/PCGStackContext.cpp
@@ -299,6 +299,25 @@ FPCGCrc FPCGStack::GetCrc() const
 	return FPCGCrc(Ar.GetCrc());
 }
 
+FPCGStackContext::FPCGStackContext(const FPCGStackContext& InStackContext, const FPCGStack& InCommonParentStack)
+{
+	const TArray<FPCGStackFrame>& ParentFrames = InCommonParentStack.StackFrames;
+	const int ParentStackFrameCount = ParentFrames.Num();
+
+	const TArray<FPCGStack>& InStacks = InStackContext.Stacks;
+
+	Stacks.SetNum(InStacks.Num());
+	for(int StackIndex = 0; StackIndex < Stacks.Num(); ++StackIndex)
+	{
+		FPCGStack& Stack = Stacks[StackIndex];
+		const FPCGStack& InStack = InStacks[StackIndex];
+
+		Stack.StackFrames.Reserve(InStack.StackFrames.Num() + ParentStackFrameCount);
+		Stack.StackFrames.Append(ParentFrames);
+		Stack.StackFrames.Append(InStack.StackFrames);
+	}
+}
+
 int32 FPCGStackContext::PushFrame(const UObject* InFrameObject)
 {
 	if (CurrentStackIndex == INDEX_NONE)
diff --git a/Engine/Plugins/PCG/Source/PCG/Private/PCGSubsystem.cpp b/Engine/Plugins/PCG/Source/PCG/Private/PCGSubsystem.cpp
--- a/Engine/Plugins/PCG/Source/PCG/Private/PCGSubsystem.cpp
+++ b/Engine/Plugins/PCG/Source/PCG/Private/PCGSubsystem.cpp
@@ -1731,35 +1731,48 @@ void UPCGSubsystem::ClearLandscapeCache()
 	}
 }
 
-bool UPCGSubsystem::GetStackContext(const UPCGComponent* InComponent, FPCGStackContext& OutStackContext)
+bool UPCGSubsystem::GetStackContext(UPCGGraph* InGraph, uint32 InGridSize, bool bIsPartitioned, FPCGStackContext& OutStackContext)
 {
-	UPCGGraph* Graph = InComponent ? InComponent->GetGraph() : nullptr;
-	if (!InComponent || !Graph)
+	if (!InGraph)
 	{
 		return false;
 	}
 
 	// A non-partitioned component generally executes (original component or local component).
-	bool bDoesComponentExecute = !InComponent->IsPartitioned();
-
-	// A partitioned higen original component will execute if the graph has UB grid level.
-	if (!bDoesComponentExecute && Graph->IsHierarchicalGenerationEnabled())
+	if(bIsPartitioned)
 	{
-		PCGHiGenGrid::FSizeArray GridSizes;
-		Graph->GetGridSizes(GridSizes, /*bOutHasUnbounded=*/bDoesComponentExecute);
-	}
+		// A partitioned higen original component will execute if the graph has UB grid level.
+		if (InGraph->IsHierarchicalGenerationEnabled())
+		{
+			PCGHiGenGrid::FSizeArray GridSizes;
+			bool bHasUnbounded = false;
+			InGraph->GetGridSizes(GridSizes, bHasUnbounded);
 
-	if (bDoesComponentExecute)
-	{
-		if (FPCGGraphCompiler* GraphCompiler = GetGraphCompiler())
+			if (!bHasUnbounded)
+			{
+				return false;
+			}
+		}
+		else
 		{
-			GraphCompiler->GetCompiledTasks(Graph, InComponent->GetGenerationGridSize(), OutStackContext, /*bIsTopGraph=*/false);
-			return true;
+			return false;
 		}
 	}
 
-	OutStackContext = FPCGStackContext();
-	return false;
+	if (FPCGGraphCompiler* GraphCompiler = GetGraphCompiler())
+	{
+		GraphCompiler->GetCompiledTasks(InGraph, InGridSize, OutStackContext, /*bIsTopGraph=*/false);
+		return true;
+	}
+	else
+	{
+		return false;
+	}
+}
+
+bool UPCGSubsystem::GetStackContext(const UPCGComponent* InComponent, FPCGStackContext& OutStackContext)
+{
+	return InComponent && GetStackContext(InComponent->GetGraph(), InComponent->IsPartitioned() ? PCGHiGenGrid::UnboundedGridSize() : InComponent->GetGenerationGridSize(), InComponent->IsPartitioned(), OutStackContext);
 }
 
 uint32 UPCGSubsystem::GetGraphCacheEntryCount(IPCGElement* InElement) const
diff --git a/Engine/Plugins/PCG/Source/PCG/Public/Graph/PCGStackContext.h b/Engine/Plugins/PCG/Source/PCG/Public/Graph/PCGStackContext.h
--- a/Engine/Plugins/PCG/Source/PCG/Public/Graph/PCGStackContext.h
+++ b/Engine/Plugins/PCG/Source/PCG/Public/Graph/PCGStackContext.h
@@ -196,6 +196,10 @@ struct PCG_API FPCGStackContext
 	friend class FPCGGraphExecutor;
 
 public:
+	FPCGStackContext() = default;
+	/** Builds a stack context by prepending the stack context by the common parent stack. */
+	FPCGStackContext(const FPCGStackContext& InStackContext, const FPCGStack& InCommonParentStack);
+
 	int32 GetNumStacks() const { return Stacks.Num(); }
 	int32 GetCurrentStackIndex() const { return CurrentStackIndex; }
 	const FPCGStack* GetStack(int32 InStackIndex) const;
diff --git a/Engine/Plugins/PCG/Source/PCG/Public/PCGSubsystem.h b/Engine/Plugins/PCG/Source/PCG/Public/PCGSubsystem.h
--- a/Engine/Plugins/PCG/Source/PCG/Public/PCGSubsystem.h
+++ b/Engine/Plugins/PCG/Source/PCG/Public/PCGSubsystem.h
@@ -320,9 +320,12 @@ public:
 	/** Clears the landscape data cache */
 	void ClearLandscapeCache();
 
-	/** Get the execution stack information for the given component. */
+	/** Gets the execution stack information for the given component (depending on partitioning, grid size, etc.) but with no component frames. */
 	bool GetStackContext(const UPCGComponent* InComponent, FPCGStackContext& OutStackContext);
 
+	/** Gets the base execution stack information for a specific graph & grid size. */
+	bool GetStackContext(UPCGGraph* InGraph, uint32 GridSize, bool bIsPartitioned, FPCGStackContext& OutStackContext);
+
 	/** Returns how many times InElement is present in the cache. */
 	uint32 GetGraphCacheEntryCount(IPCGElement* InElement) const;
 
diff --git a/Engine/Plugins/PCG/Source/PCGEditor/Private/Widgets/SPCGEditorGraphDebugObjectTree.cpp b/Engine/Plugins/PCG/Source/PCGEditor/Private/Widgets/SPCGEditorGraphDebugObjectTree.cpp
--- a/Engine/Plugins/PCG/Source/PCGEditor/Private/Widgets/SPCGEditorGraphDebugObjectTree.cpp
+++ b/Engine/Plugins/PCG/Source/PCGEditor/Private/Widgets/SPCGEditorGraphDebugObjectTree.cpp
@@ -596,60 +596,9 @@ bool SPCGEditorGraphDebugObjectTree::IsFocusOnDebugObjectButtonEnabled() const
 
 void SPCGEditorGraphDebugObjectTree::SetDebugObjectFromSelection_OnClicked()
 {
-	const UPCGGraph* PCGGraph = GetPCGGraph();
-	if (!PCGGraph)
-	{
-		return;
-	}
-
-	check(GEditor);
-	USelection* SelectedActors = GEditor->GetSelectedActors();
-	if (!IsValid(SelectedActors))
+	if (FPCGEditorGraphDebugObjectItemPtr Item = GetFirstDebugObjectFromSelection())
 	{
-		return;
-	}
-
-	for (FSelectionIterator It(*SelectedActors); It; ++It)
-	{
-		const AActor* SelectedActor = Cast<AActor>(*It);
-		if (!IsValid(SelectedActor))
-		{
-			continue;
-		}
-
-		TArray<const UPCGComponent*> PCGComponents;
-		SelectedActor->GetComponents<const UPCGComponent>(PCGComponents, /*bIncludeFromChildActors=*/true);
-
-		for (const UPCGComponent* PCGComponent : PCGComponents)
-		{
-			if (!IsValid(PCGComponent))
-			{
-				continue;
-			}
-
-			FPCGStackContext StackContext;
-			if (!PCGComponent->GetStackContext(StackContext))
-			{
-				continue;
-			}
-
-			bool bSelectedDebugObject = false;
-
-			for (const FPCGStack& Stack : StackContext.GetStacks())
-			{
-				if (FPCGEditorGraphDebugObjectItemPtr Item = GetItemFromStack(Stack))
-				{
-					ExpandAndSelectFirstLeafDebugObject(Item);
-					bSelectedDebugObject = true;
-					break;
-				}
-			}
-
-			if (bSelectedDebugObject)
-			{
-				break;
-			}
-		}
+		ExpandAndSelectFirstLeafDebugObject(Item);
 	}
 }
 
@@ -778,20 +727,19 @@ bool SPCGEditorGraphDebugObjectTree::GetFirstStackFromSelection(const UPCGNode*
 	return false;
 }
 
-void SPCGEditorGraphDebugObjectTree::UpdateIsSetDebugObjectFromSelectionEnabled()
+FPCGEditorGraphDebugObjectItemPtr SPCGEditorGraphDebugObjectTree::GetFirstDebugObjectFromSelection() const
 {
-	IsSetDebugObjectFromSelectionEnabled = false;
-
 	const UPCGGraph* PCGGraph = GetPCGGraph();
 	if (!PCGGraph)
 	{
-		return;
+		return nullptr;
 	}
 
+	check(GEditor);
 	USelection* SelectedActors = GEditor->GetSelectedActors();
 	if (!IsValid(SelectedActors))
 	{
-		return;
+		return nullptr;
 	}
 
 	UPCGSubsystem* Subsystem = PCGEditor.Pin() ? PCGEditor.Pin()->GetSubsystem() : nullptr;
@@ -818,24 +766,36 @@ void SPCGEditorGraphDebugObjectTree::UpdateIsSetDebugObjectFromSelectionEnabled(
 			FPCGStackContext StackContext;
 			if (PCGComponent->GetStackContext(StackContext))
 			{
-				if (Algo::AnyOf(StackContext.GetStacks(), [&PCGGraph](const FPCGStack& InStack) { return InStack.GetGraphForCurrentFrame() == PCGGraph; }))
+				for (const FPCGStack& Stack : StackContext.GetStacks())
 				{
-					IsSetDebugObjectFromSelectionEnabled = true;
-					return;
+					if (FPCGEditorGraphDebugObjectItemPtr Item = GetItemFromStack(Stack))
+					{
+						return Item;
+					}
 				}
 			}
 
 			// Look for graph in dynamic stacks.
 			if (Subsystem)
 			{
-				if (Algo::AnyOf(Subsystem->GetExecutedStacks(PCGComponent, PCGGraph), [&PCGGraph](const FPCGStack& InStack) { return InStack.GetGraphForCurrentFrame() == PCGGraph; }))
+				TArray<FPCGStack> ExecutedStacks = Subsystem->GetExecutedStacks(PCGComponent, PCGGraph);
+				for (const FPCGStack& Stack : ExecutedStacks)
 				{
-					IsSetDebugObjectFromSelectionEnabled = true;
-					return;
+					if (FPCGEditorGraphDebugObjectItemPtr Item = GetItemFromStack(Stack))
+					{
+						return Item;
+					}
 				}
 			}
 		}
 	}
+
+	return nullptr;
+}
+
+void SPCGEditorGraphDebugObjectTree::UpdateIsSetDebugObjectFromSelectionEnabled()
+{
+	IsSetDebugObjectFromSelectionEnabled = (GetFirstDebugObjectFromSelection() != nullptr);
 }
 
 void SPCGEditorGraphDebugObjectTree::AddStacksToTree(const TArray<FPCGStack>& Stacks,
@@ -1054,7 +1014,7 @@ void SPCGEditorGraphDebugObjectTree::AddStacksToTree(const TArray<FPCGStack>& St
 
 void SPCGEditorGraphDebugObjectTree::RefreshTree()
 {
-	//TRACE_CPUPROFILER_EVENT_SCOPE(SPCGEditorGraphDebugObjectTree::RefreshTree);
+	TRACE_CPUPROFILER_EVENT_SCOPE(SPCGEditorGraphDebugObjectTree::RefreshTree);
 	RootItems.Empty();
 	AllGraphItems.Empty();
 	DebugObjectTreeView->RequestTreeRefresh();
@@ -1078,12 +1038,19 @@ void SPCGEditorGraphDebugObjectTree::RefreshTree()
 	{
 		TMap<AActor*, TSharedPtr<FPCGEditorGraphDebugObjectItem_Actor>> ActorItems;
 		TMap<const FPCGStack, FPCGEditorGraphDebugObjectItemPtr> StackToItem;
-
-		TArray<FPCGStack> GraphStacks = Subsystem->GetExecutedStacks(nullptr, PCGGraph, /*bOnlyWithSubgraphAsCurrentFrame=*/false);
 		TMap<const UPCGComponent*, TArray<FPCGStack>> ComponentsGraphStacks;
-		for (FPCGStack& GraphStack : GraphStacks)
+
+		using FGraphContextMapKey = TTuple<UPCGGraph*, uint32 /*GenerationGridSize*/, bool /*IsPartitioned*/>;
+		using FGraphContextMapValue = TTuple<FPCGStackContext, bool /*bContainsGraph*/>;
+		TMap<FGraphContextMapKey, FGraphContextMapValue> GraphToContexts;
+
 		{
-			ComponentsGraphStacks.FindOrAdd(GraphStack.GetRootComponent()).Add(MoveTemp(GraphStack));
+			TRACE_CPUPROFILER_EVENT_SCOPE(SPCGEditorGraphDebugObjectTree::RefreshTree::GetExecutedStacks);
+			TArray<FPCGStack> GraphStacks = Subsystem->GetExecutedStacks(nullptr, PCGGraph, /*bOnlyWithSubgraphAsCurrentFrame=*/false);
+			for (FPCGStack& GraphStack : GraphStacks)
+			{
+				ComponentsGraphStacks.FindOrAdd(GraphStack.GetRootComponent()).Add(MoveTemp(GraphStack));
+			}
 		}
 
 		for (UObject* PCGComponentObject : PCGComponents)
@@ -1094,23 +1061,46 @@ void SPCGEditorGraphDebugObjectTree::RefreshTree()
 			}
 
 			UPCGComponent* PCGComponent = Cast<UPCGComponent>(PCGComponentObject);
-			if (!PCGComponent || !PCGComponent->IsRegistered())
+			if (!PCGComponent || !PCGComponent->IsRegistered() || !PCGComponent->GetGraph())
 			{
 				continue;
 			}
 
 			// Process static stacks that can be read from the compiled graph.
-			FPCGStackContext StackContext;
-			if (!PCGComponent->GetStackContext(StackContext))
 			{
-				continue;
+				//TRACE_CPUPROFILER_EVENT_SCOPE(SPCGEditorGraphDebugObjectTree::RefreshTree::StaticStacks);
+				// Considering it's likely for multiple components to have the same graph, we'll do a per-graph check first prior to building the component + graph stack.
+				UPCGGraph* ComponentGraph = PCGComponent->GetGraph();
+				FGraphContextMapKey MapKey(ComponentGraph, PCGComponent->IsPartitioned() ? PCGHiGenGrid::UnboundedGridSize() : PCGComponent->GetGenerationGridSize(), PCGComponent->IsPartitioned());
+				FGraphContextMapValue* MapValue = GraphToContexts.Find(MapKey);
+
+				if (!MapValue)
+				{
+					FPCGStackContext GraphStackContext;
+					bool bSomeStacksContainCurrentGraph = false;
+					if (Subsystem->GetStackContext(MapKey.Get<0>(), MapKey.Get<1>(), MapKey.Get<2>(), GraphStackContext))
+					{
+						bSomeStacksContainCurrentGraph = Algo::AnyOf(GraphStackContext.GetStacks(), [PCGGraph](const FPCGStack& Stack) { return Stack.HasObject(PCGGraph); });
+					}
+
+					MapValue = &(GraphToContexts.Add(MapKey, FGraphContextMapValue(std::move(GraphStackContext), bSomeStacksContainCurrentGraph)));
+				}
+
+				check(MapValue);
+				if (MapValue->Get<1>())
+				{
+					FPCGStack ComponentStack;
+					ComponentStack.PushFrame(PCGComponent);
+					FPCGStackContext ComponentContext(MapValue->Get<0>(), ComponentStack);
+					AddStacksToTree(ComponentContext.GetStacks(), ActorItems, StackToItem);
+				}
 			}
-			AddStacksToTree(StackContext.GetStacks(), ActorItems, StackToItem);
 
 			// Process stacks encountered during execution so far, which will include dynamic subgraphs & loop subgraphs.
 			// There will be overlaps with the static stacks but only unique entries will be added to the tree.
 			if (TArray<FPCGStack>* ComponentGraphStacks = ComponentsGraphStacks.Find(PCGComponent))
 			{
+				//TRACE_CPUPROFILER_EVENT_SCOPE(SPCGEditorGraphDebugObjectTree::RefreshTree::DynamicStacks);
 				AddStacksToTree(*ComponentGraphStacks, ActorItems, StackToItem);
 			}
 		}
@@ -1127,6 +1117,7 @@ void SPCGEditorGraphDebugObjectTree::RefreshTree()
 
 void SPCGEditorGraphDebugObjectTree::SortTreeItems(bool bIsAscending, bool bIsRecursive)
 {
+	TRACE_CPUPROFILER_EVENT_SCOPE(SPCGEditorGraphDebugObjectTree::SortTreeItems);
 	RootItems.Sort([bIsAscending](const FPCGEditorGraphDebugObjectItemPtr& InLHS, const FPCGEditorGraphDebugObjectItemPtr& InRHS)
 	{
 		return (InLHS->GetLabel(/*bForSorting=*/true) < InRHS->GetLabel(/*bForSorting=*/true)) == bIsAscending;
@@ -1143,6 +1134,7 @@ void SPCGEditorGraphDebugObjectTree::SortTreeItems(bool bIsAscending, bool bIsRe
 
 void SPCGEditorGraphDebugObjectTree::RestoreTreeState()
 {
+	TRACE_CPUPROFILER_EVENT_SCOPE(SPCGEditorGraphDebugObjectTree::RestoreTreeState);
 	// Try to restore user item expansion. We can't modify ExpandedStack, so we'll rebuild it.
 	TSet<FPCGStack> ExpandedStacksBefore = std::move(ExpandedStacks);
 	ExpandedStacks.Empty(ExpandedStacksBefore.Num());
diff --git a/Engine/Plugins/PCG/Source/PCGEditor/Private/Widgets/SPCGEditorGraphDebugObjectTree.h b/Engine/Plugins/PCG/Source/PCGEditor/Private/Widgets/SPCGEditorGraphDebugObjectTree.h
--- a/Engine/Plugins/PCG/Source/PCGEditor/Private/Widgets/SPCGEditorGraphDebugObjectTree.h
+++ b/Engine/Plugins/PCG/Source/PCGEditor/Private/Widgets/SPCGEditorGraphDebugObjectTree.h
@@ -281,6 +281,7 @@ private:
 	FReply FocusOnDebugObject_OnClicked() const;
 	bool IsFocusOnDebugObjectButtonEnabled() const;
 
+	FPCGEditorGraphDebugObjectItemPtr GetFirstDebugObjectFromSelection() const;
 	void SetDebugObjectFromSelection_OnClicked();
 	bool IsSetDebugObjectFromSelectionButtonEnabled() const { return IsSetDebugObjectFromSelectionEnabled.IsSet() && IsSetDebugObjectFromSelectionEnabled.GetValue(); }
 	void UpdateIsSetDebugObjectFromSelectionEnabled();
