Added missing validation in apply hierarchy, as point keys have to have the same number of entries as their parent counterpart.

#rnx

diff --git a/Engine/Plugins/PCG/Source/PCG/Private/Elements/PCGApplyHierarchy.cpp b/Engine/Plugins/PCG/Source/PCG/Private/Elements/PCGApplyHierarchy.cpp
--- a/Engine/Plugins/PCG/Source/PCG/Private/Elements/PCGApplyHierarchy.cpp
+++ b/Engine/Plugins/PCG/Source/PCG/Private/Elements/PCGApplyHierarchy.cpp
@@ -70,7 +70,7 @@ bool FPCGApplyHierarchyElement::PrepareDataInternal(FPCGContext* InContext) cons
 		OutState.InputData = InputData;
 		OutState.OutputData = OutputData;
 		OutState.OutputDataIndex = Context->OutputData.TaggedData.Num();
-		Context->OutputData.TaggedData.Add_GetRef(Input).Data = OutputData;
+		Context->OutputData.TaggedData.Add(Input);
 		
 		// Create point index & parent index accessors
 		// Implementation note: we could have used PCGMetadataElementCommon::ApplyOnAccessorRange here, as we're doing a transformation
@@ -119,7 +119,13 @@ bool FPCGApplyHierarchyElement::PrepareDataInternal(FPCGContext* InContext) cons
 
 		if (Settings->PointKeyAttributes.IsEmpty() || Settings->ParentKeyAttributes.IsEmpty())
 		{
-			PCGLog::LogErrorOnGraph(LOCTEXT("RequiresAtLeastOneAttribute", "Both the point key attribute and parent key attribute require valid entries to perform the Apply Hierarchy operation."));
+			PCGLog::LogErrorOnGraph(LOCTEXT("RequiresAtLeastOneAttribute", "Both the point key attribute and parent key attribute require valid entries to perform the Apply Hierarchy operation."), Context);
+			return EPCGTimeSliceInitResult::NoOperation;
+		}
+
+		if (Settings->PointKeyAttributes.Num() != Settings->ParentKeyAttributes.Num())
+		{
+			PCGLog::LogErrorOnGraph(LOCTEXT("PointAndParentKeyAttributesMismatch", "There needs to be the same number of point key attributes as parent key attributes."), Context);
 			return EPCGTimeSliceInitResult::NoOperation;
 		}
 
@@ -211,6 +217,9 @@ bool FPCGApplyHierarchyElement::PrepareDataInternal(FPCGContext* InContext) cons
 			PCGLog::Metadata::LogFailToGetAttributeError<FTransform>(RelativeTransformSelector, OutState.RelativeTransformAccessor.Get(), Context);
 			return EPCGTimeSliceInitResult::NoOperation;
 		}
+
+		// If everything was validated correctly, move the "new" output result to the output data.
+		Context->OutputData.TaggedData.Last().Data = OutputData;
 		
 		return EPCGTimeSliceInitResult::Success;
 	});
