GetActorData with reference modes now returns tags like other modes.

#rnx

diff --git a/Engine/Plugins/PCG/Source/PCG/Private/Elements/PCGDataFromActor.cpp b/Engine/Plugins/PCG/Source/PCG/Private/Elements/PCGDataFromActor.cpp
--- a/Engine/Plugins/PCG/Source/PCG/Private/Elements/PCGDataFromActor.cpp
+++ b/Engine/Plugins/PCG/Source/PCG/Private/Elements/PCGDataFromActor.cpp
@@ -16,6 +16,7 @@
 #include "Grid/PCGPartitionActor.h"
 #include "Helpers/PCGDynamicTrackingHelpers.h"
 #include "Helpers/PCGHelpers.h"
+#include "Helpers/PCGTagHelpers.h"
 #include "Utils/PCGGraphExecutionLogging.h"
 
 #include "Algo/AnyOf.h"
@@ -32,6 +33,7 @@ namespace PCGDataFromActorConstants
 	static const FName SinglePointPinLabel = TEXT("Single Point");
 	static const FString PCGComponentDataGridSizeTagPrefix = TEXT("PCG_GridSize_");
 	static const FText TagNamesSanitizedWarning = LOCTEXT("TagAttributeNamesSanitized", "One or more tag names contained invalid characters and were sanitized when creating the corresponding attributes.");
+	static const FText TagNamesReservedWarning = LOCTEXT("TagAttributeNamesReserved", "One or more tag names match to reserved tags (ActorReference, ComponentReference) and were ignored when parsing the data.");
 }
 
 namespace PCGDataFromActorHelpers
@@ -705,22 +707,49 @@ void FPCGDataFromActorElement::CreateReferenceData(FPCGContext* InContext, const
 	FPCGDataFromActorContext* Context = static_cast<FPCGDataFromActorContext*>(InContext);
 	check(Settings && (Settings->Mode == EPCGGetDataFromActorMode::GetActorReference || Settings->Mode == EPCGGetDataFromActorMode::GetComponentsReference));
 
-	const bool bCreateActorReferences = (Settings->Mode == EPCGGetDataFromActorMode::GetActorReference || Settings->Mode == EPCGGetDataFromActorMode::GetComponentsReference);
 	const bool bCreateComponentReferences = (Settings->Mode == EPCGGetDataFromActorMode::GetComponentsReference);
 
 	UPCGParamData* ParamData = FPCGContext::NewObject_AnyThread<UPCGParamData>(Context);
-	FPCGMetadataAttribute<FSoftObjectPath>* ActorReference = nullptr;
-	FPCGMetadataAttribute<FSoftObjectPath>* ComponentReference = nullptr;
+	UPCGMetadata* Metadata = ParamData->MutableMetadata();
 
-	if (bCreateActorReferences)
+	TSet<FString> AllGatheredTags;
+
+	TSet<FName> ReservedTags;
+	ReservedTags.Add(PCGPointDataConstants::ActorReferenceAttribute);
+	ReservedTags.Add(PCGAddComponentConstants::ComponentReferenceAttribute);
+
+	FPCGMetadataAttribute<FSoftObjectPath>* ActorReference = Metadata->FindOrCreateAttribute(PCGPointDataConstants::ActorReferenceAttribute, FSoftObjectPath(), /*bAllowsInterpolation=*/false, /*bOverrideParent=*/false, /*bOverwriteIfTypeMismatch=*/true);
+	FPCGMetadataAttribute<FSoftObjectPath>* ComponentReference = bCreateComponentReferences ? Metadata->FindOrCreateAttribute(PCGAddComponentConstants::ComponentReferenceAttribute, FSoftObjectPath(), /*bAllowsInterpolation=*/false, /*bOverrideParent=*/false, /*bOverwriteIfTypeMismatch=*/true) : nullptr;
+
+	bool bHasReservedTag = false;
+	bool bHasSanitizedTag = false;
+
+	auto SetAttributesFromTags = [Metadata, &ReservedTags, &bHasReservedTag, &bHasSanitizedTag](const auto& Tags, PCGMetadataEntryKey Entry)
 	{
-		ActorReference = ParamData->MutableMetadata()->FindOrCreateAttribute(PCGPointDataConstants::ActorReferenceAttribute, FSoftObjectPath(), /*bAllowsInterpolation=*/false, /*bOverrideParent=*/false, /*bOverwriteIfTypeMismatch=*/true);
-	}
+		for(FName Tag : Tags)
+		{
+			PCG::Private::FParseTagResult TagData(Tag);
+			if (!ReservedTags.Contains(FName(TagData.Attribute)))
+			{
+				if (PCG::Private::SetAttributeFromTag(TagData, Metadata, Entry, PCG::Private::ESetAttributeFromTagFlags::CreateAttribute))
+				{
+					bHasSanitizedTag |= TagData.HasBeenSanitized();
+				}
+			}
+			else
+			{
+				bHasReservedTag = true;
+			}
+		}
+	};
 
-	if (bCreateComponentReferences)
+	auto AddTagsToGatheredTags = [&AllGatheredTags](const auto& Tags)
 	{
-		ComponentReference = ParamData->MutableMetadata()->FindOrCreateAttribute(PCGAddComponentConstants::ComponentReferenceAttribute, FSoftObjectPath(), /*bAllowsInterpolation=*/false, /*bOverrideParent=*/false, /*bOverwriteIfTypeMismatch=*/true);
-	}
+		for (FName Tag : Tags)
+		{
+			AllGatheredTags.Add(Tag.ToString());
+		}
+	};
 
 	for (AActor* Actor : Actors)
 	{
@@ -729,12 +758,14 @@ void FPCGDataFromActorElement::CreateReferenceData(FPCGContext* InContext, const
 			continue;
 		}
 
+		FSoftObjectPath ActorPath(Actor);
+
 		if (ComponentReference)
 		{
 			TInlineComponentArray<UActorComponent*, 16> Components;
 			Actor->GetComponents(Components);
 
-			FSoftObjectPath ActorPath(Actor);
+			bool bHasAddedActorTags = false;
 
 			for (UActorComponent* Component : Components)
 			{
@@ -743,14 +774,28 @@ void FPCGDataFromActorElement::CreateReferenceData(FPCGContext* InContext, const
 					continue;
 				}
 
+				if (!bHasAddedActorTags)
+				{
+					bHasAddedActorTags = true;
+					AddTagsToGatheredTags(Actor->Tags);
+				}
+
 				PCGMetadataEntryKey Entry = ParamData->MutableMetadata()->AddEntry();
 				ActorReference->SetValue(Entry, ActorPath);
 				ComponentReference->SetValue(Entry, FSoftObjectPath(Component));
+
+				SetAttributesFromTags(Actor->Tags, Entry);
+				SetAttributesFromTags(Component->ComponentTags, Entry);
+				AddTagsToGatheredTags(Component->ComponentTags);
 			}
 		}
-		else if (ActorReference)
+		else
 		{
-			ActorReference->SetValue(ParamData->MutableMetadata()->AddEntry(), FSoftObjectPath(Actor));
+			PCGMetadataEntryKey Entry = ParamData->MutableMetadata()->AddEntry();
+			ActorReference->SetValue(Entry, ActorPath);
+
+			SetAttributesFromTags(Actor->Tags, Entry);
+			AddTagsToGatheredTags(Actor->Tags);
 		}
 	}
 
@@ -758,6 +803,18 @@ void FPCGDataFromActorElement::CreateReferenceData(FPCGContext* InContext, const
 	{
 		FPCGTaggedData& TaggedData = Context->OutputData.TaggedData.Emplace_GetRef();
 		TaggedData.Data = ParamData;
+		TaggedData.Tags = MoveTemp(AllGatheredTags);
+	}
+
+	// Finally, log warnings when/if required
+	if (bHasReservedTag && !Settings->bSilenceReservedAttributeNameWarnings)
+	{
+		PCGE_LOG(Warning, GraphAndLog, PCGDataFromActorConstants::TagNamesReservedWarning);
+	}
+
+	if (bHasSanitizedTag && !Settings->bSilenceSanitizedAttributeNameWarnings)
+	{
+		PCGE_LOG(Warning, GraphAndLog, PCGDataFromActorConstants::TagNamesSanitizedWarning);
 	}
 }
 
diff --git a/Engine/Plugins/PCG/Source/PCG/Public/Elements/PCGDataFromActor.h b/Engine/Plugins/PCG/Source/PCG/Public/Elements/PCGDataFromActor.h
--- a/Engine/Plugins/PCG/Source/PCG/Public/Elements/PCGDataFromActor.h
+++ b/Engine/Plugins/PCG/Source/PCG/Public/Elements/PCGDataFromActor.h
@@ -125,6 +125,10 @@ public:
 	UPROPERTY(BlueprintReadWrite, EditAnywhere, Category = "Data Retrieval Settings|Advanced")
 	bool bSilenceSanitizedAttributeNameWarnings = false;
 
+	/** Silence warnings that attribute names were rejected because they clash with reserved names. */
+	UPROPERTY(BlueprintReadWrite, EditAnywhere, Category = "Data Retrieval Settings | Advanced")
+	bool bSilenceReservedAttributeNameWarnings = false;
+
 #if WITH_EDITORONLY_DATA
 	/** If this is checked, found actors that are outside component bounds will not trigger a refresh. Only works for tags for now in editor. */
 	UPROPERTY(BlueprintReadOnly, EditAnywhere, Category = "Data Retrieval Settings")
