[PCG] Fix regression on PropertyOverride

#rnx

diff --git a/Engine/Plugins/PCG/Source/PCG/Public/Metadata/PCGObjectPropertyOverride.h b/Engine/Plugins/PCG/Source/PCG/Public/Metadata/PCGObjectPropertyOverride.h
--- a/Engine/Plugins/PCG/Source/PCG/Public/Metadata/PCGObjectPropertyOverride.h
+++ b/Engine/Plugins/PCG/Source/PCG/Public/Metadata/PCGObjectPropertyOverride.h
@@ -120,18 +120,6 @@ struct FPCGObjectOverrides
 	FPCGObjectOverrides(T* TemplateObject) : OutputKey(TemplateObject)
 	{}
 
-	constexpr const UStruct* GetClassOfTemplate() const
-	{
-		if constexpr (TIsDerivedFrom<T, UObject>::Value)
-		{
-			return T::StaticClass();
-		}
-		else
-		{
-			return T::StaticStruct();
-		}
-	}
-
 	/** Initialize the object overrides. Call before using Apply(InputKeyIndex). */
 	void Initialize(const TArray<FPCGObjectPropertyOverrideDescription>& OverrideDescriptions, T* TemplateObject, const UPCGData* SourceData, FPCGContext* Context)
 	{
@@ -145,7 +133,15 @@ struct FPCGObjectOverrides
 
 		OutputKey = FPCGAttributeAccessorKeysSingleObjectPtr<T>(TemplateObject);
 
-		const UStruct* ClassObject = GetClassOfTemplate();
+		const UStruct* ClassObject = nullptr;
+		if constexpr (TIsDerivedFrom<T, UObject>::Value)
+		{
+			ClassObject = TemplateObject->GetClass();
+		}
+		else
+		{
+			ClassObject = TemplateObject->StaticStruct();
+		}
 
 		ObjectSingleOverrides.Empty(OverrideDescriptions.Num());
 
