[PCG] Fix build error in PCGStaticMeshDataInterface

#rnx

diff --git a/Engine/Plugins/PCG/Source/PCG/Private/Compute/DataInterfaces/PCGStaticMeshDataInterface.h b/Engine/Plugins/PCG/Source/PCG/Private/Compute/DataInterfaces/PCGStaticMeshDataInterface.h
--- a/Engine/Plugins/PCG/Source/PCG/Private/Compute/DataInterfaces/PCGStaticMeshDataInterface.h
+++ b/Engine/Plugins/PCG/Source/PCG/Private/Compute/DataInterfaces/PCGStaticMeshDataInterface.h
@@ -74,7 +74,7 @@ protected:
 	using FParameters = FPCGStaticMeshDataInterfaceParameters;
 
 	FName MeshName;
-	TRefCountPtr<const FStaticMeshLODResources> LODResources = nullptr;
+	TRefCountPtr<const FStaticMeshLODResources> LODResources;
 	FBoxSphereBounds Bounds = FBoxSphereBounds(EForceInit::ForceInit);
-	TRefCountPtr<FRHIShaderResourceView> IndexBufferSRV = nullptr;
+	TRefCountPtr<FRHIShaderResourceView> IndexBufferSRV;
 };
