[PCG] Fix error format string

#rnx

diff --git a/Engine/Plugins/PCG/Source/PCG/Private/Utils/PCGNodeVisualLogs.cpp b/Engine/Plugins/PCG/Source/PCG/Private/Utils/PCGNodeVisualLogs.cpp
--- a/Engine/Plugins/PCG/Source/PCG/Private/Utils/PCGNodeVisualLogs.cpp
+++ b/Engine/Plugins/PCG/Source/PCG/Private/Utils/PCGNodeVisualLogs.cpp
@@ -184,7 +184,7 @@ FText FPCGNodeVisualLogs::GetSummaryText(const FPCGPerNodeVisualLogs& InLogs, co
 		{
 			if (Summary.IsEmpty())
 			{
-				Summary = FText::Format(LOCTEXT("NodeTooltipLogEmpty", "{1}: {2}"), MessageVerbosity, LogEntry.Message);
+				Summary = FText::Format(LOCTEXT("NodeTooltipLogEmpty", "{0}: {1}"), MessageVerbosity, LogEntry.Message);
 			}
 			else
 			{
