Fixed an issue where collapsing nodes to a subgraph would suggest using a template when it's clearly not required.

#rnx

diff --git a/Engine/Plugins/PCG/Source/PCGEditor/Private/PCGEditor.cpp b/Engine/Plugins/PCG/Source/PCGEditor/Private/PCGEditor.cpp
--- a/Engine/Plugins/PCG/Source/PCGEditor/Private/PCGEditor.cpp
+++ b/Engine/Plugins/PCG/Source/PCGEditor/Private/PCGEditor.cpp
@@ -1497,6 +1497,7 @@ void FPCGEditor::OnCollapseNodesInSubgraph()
 	IAssetTools& AssetTools = FModuleManager::Get().LoadModuleChecked<FAssetToolsModule>("AssetTools").Get();
 
 	TObjectPtr<UPCGGraphFactory> Factory = NewObject<UPCGGraphFactory>();
+	Factory->bSkipTemplateSelection = true;
 
 	FString NewPackageName;
 	FString NewAssetName;
diff --git a/Engine/Plugins/PCG/Source/PCGEditor/Private/PCGGraphFactory.cpp b/Engine/Plugins/PCG/Source/PCGEditor/Private/PCGGraphFactory.cpp
--- a/Engine/Plugins/PCG/Source/PCGEditor/Private/PCGGraphFactory.cpp
+++ b/Engine/Plugins/PCG/Source/PCGEditor/Private/PCGGraphFactory.cpp
@@ -49,6 +49,12 @@ bool UPCGGraphFactory::ShouldShowInNewMenu() const
 
 bool UPCGGraphFactory::ConfigureProperties()
 {
+	// Early out if this is a standalone factory that does not allow templates.
+	if (bSkipTemplateSelection)
+	{
+		return true;
+	}
+
 	// Early out if the editor config says we don't want to have a dialog for it.
 	const UPCGEditorSettings* EditorSettings = GetDefault<UPCGEditorSettings>();
 	if (EditorSettings && !EditorSettings->bShowTemplatePickerOnNewGraph)
diff --git a/Engine/Plugins/PCG/Source/PCGEditor/Private/PCGGraphFactory.h b/Engine/Plugins/PCG/Source/PCGEditor/Private/PCGGraphFactory.h
--- a/Engine/Plugins/PCG/Source/PCGEditor/Private/PCGGraphFactory.h
+++ b/Engine/Plugins/PCG/Source/PCGEditor/Private/PCGGraphFactory.h
@@ -20,6 +20,10 @@ public:
 	UPROPERTY()
 	TObjectPtr<UPCGGraph> TemplateGraph;
 
+	/** Disables template selection, even if some exist. Useful when creating a graph from a standalone factory. */
+	UPROPERTY()
+	bool bSkipTemplateSelection = false;
+
 	//~UFactory interface
 	virtual UObject* FactoryCreateNew(UClass* InClass, UObject* InParent, FName InName, EObjectFlags Flags, UObject* Context, FFeedbackContext* Warn) override;
 	virtual bool ConfigureProperties() override;
