[PCG] Disable render-targetable on GPU texture readback copy.

#rnx

diff --git a/Engine/Plugins/PCG/Source/PCGCompute/Private/PCGTextureReadback.cpp b/Engine/Plugins/PCG/Source/PCGCompute/Private/PCGTextureReadback.cpp
--- a/Engine/Plugins/PCG/Source/PCGCompute/Private/PCGTextureReadback.cpp
+++ b/Engine/Plugins/PCG/Source/PCGCompute/Private/PCGTextureReadback.cpp
@@ -51,7 +51,6 @@ void FPCGTextureReadbackInterface::Dispatch_RenderThread(FRHICommandListImmediat
 		FRHITextureCreateDesc::Create2D(TEXT("PCGTexture Readback Compute Target"), Params.SourceDimensions.X, Params.SourceDimensions.Y, EPixelFormat::PF_B8G8R8A8)
 			.SetClearValue(FClearValueBinding::None)
 			.SetFlags(ETextureCreateFlags::UAV |
-				ETextureCreateFlags::RenderTargetable |
 				ETextureCreateFlags::ShaderResource |
 				ETextureCreateFlags::NoTiling
 			)
