[PCG] Name our kernels so that they will appear with legible names in profiling tools when the new GPU profiler is enabled.

Generates name of Custom HLSL node kernels from node title with sanitisation.

#rnx

diff --git a/Engine/Plugins/PCG/Shaders/Private/Elements/PCGCopyPoints.usf b/Engine/Plugins/PCG/Shaders/Private/Elements/PCGCopyPoints.usf
--- a/Engine/Plugins/PCG/Shaders/Private/Elements/PCGCopyPoints.usf
+++ b/Engine/Plugins/PCG/Shaders/Private/Elements/PCGCopyPoints.usf
@@ -22,7 +22,7 @@
 #define NUM_RESERVED_ATTRS 32
 
 [numthreads(64, 1, 1)]
-void Main(uint3 GroupId : SV_GroupID, uint GroupIndex : SV_GroupIndex)
+void PCGCopyPointsCS(uint3 GroupId : SV_GroupID, uint GroupIndex : SV_GroupIndex)
 {
 	// Mark the kernel as having executed. Must run before we early out via thread index, because the kernel is still 'executed' even if the number of
 	// threads to iterate on is zero. Even if GetNumThreads() returns 0, the kernel will still have been dispatched on a single thread to set this flag.
@@ -123,17 +123,17 @@ void Main(uint3 GroupId : SV_GroupID, uint GroupIndex : SV_GroupIndex)
 		return;
 	}
 
-	float3 SourcePosition = Source_GetPosition(SourceDataIndex, SourcePointIndexInData);
-	float4 SourceRotation = Source_GetRotation(SourceDataIndex, SourcePointIndexInData);
-	float3 SourceScale = Source_GetScale(SourceDataIndex, SourcePointIndexInData);
-	float4 SourceColor = Source_GetColor(SourceDataIndex, SourcePointIndexInData);
-	uint SourceSeed = Source_GetSeed(SourceDataIndex, SourcePointIndexInData);
+	const float3 SourcePosition = Source_GetPosition(SourceDataIndex, SourcePointIndexInData);
+	const float4 SourceRotation = Source_GetRotation(SourceDataIndex, SourcePointIndexInData);
+	const float3 SourceScale = Source_GetScale(SourceDataIndex, SourcePointIndexInData);
+	const float4 SourceColor = Source_GetColor(SourceDataIndex, SourcePointIndexInData);
+	const uint SourceSeed = Source_GetSeed(SourceDataIndex, SourcePointIndexInData);
 	
-	float3 TargetPosition = Target_GetPosition(TargetDataIndex, TargetPointIndexInData);
-	float4 TargetRotation = Target_GetRotation(TargetDataIndex, TargetPointIndexInData);
-	float3 TargetScale = Target_GetScale(TargetDataIndex, TargetPointIndexInData);
-	float4 TargetColor = Target_GetColor(TargetDataIndex, TargetPointIndexInData);
-	uint TargetSeed = Target_GetSeed(TargetDataIndex, TargetPointIndexInData);
+	const float3 TargetPosition = Target_GetPosition(TargetDataIndex, TargetPointIndexInData);
+	const float4 TargetRotation = Target_GetRotation(TargetDataIndex, TargetPointIndexInData);
+	const float3 TargetScale = Target_GetScale(TargetDataIndex, TargetPointIndexInData);
+	const float4 TargetColor = Target_GetColor(TargetDataIndex, TargetPointIndexInData);
+	const uint TargetSeed = Target_GetSeed(TargetDataIndex, TargetPointIndexInData);
 
 	const uint RotationInheritance = CopyPoints_GetRotationInheritance();
 	const uint ApplyTargetRotationToPositions = CopyPoints_GetApplyTargetRotationToPositions();
@@ -143,9 +143,10 @@ void Main(uint3 GroupId : SV_GroupID, uint GroupIndex : SV_GroupIndex)
 	const uint SeedInheritance = CopyPoints_GetSeedInheritance();
 	const uint AttributeInheritance = CopyPoints_GetAttributeInheritance();
 
-	float3 ScaleForPosition = ApplyTargetScaleToPositions > 0 ? TargetScale : float3(1, 1, 1);
+	const float3 ScaleForPosition = ApplyTargetScaleToPositions > 0 ? TargetScale : float3(1, 1, 1);
+
+	const float3 OutPosition = TargetPosition + (ApplyTargetRotationToPositions > 0 ? QuatRotateVector(TargetRotation, (SourcePosition * ScaleForPosition)) : (SourcePosition * ScaleForPosition));
 
-	float3 OutPosition = TargetPosition + (ApplyTargetRotationToPositions > 0 ? QuatRotateVector(TargetRotation, (SourcePosition * ScaleForPosition)) : (SourcePosition * ScaleForPosition));
 	float4 OutRotation = float4(0, 0, 0, 1);
 	float3 OutScale = (float3)0;
 	float4 OutColor = (float4)0;
diff --git a/Engine/Plugins/PCG/Shaders/Private/Elements/PCGCopyPointsAnalysis.usf b/Engine/Plugins/PCG/Shaders/Private/Elements/PCGCopyPointsAnalysis.usf
--- a/Engine/Plugins/PCG/Shaders/Private/Elements/PCGCopyPointsAnalysis.usf
+++ b/Engine/Plugins/PCG/Shaders/Private/Elements/PCGCopyPointsAnalysis.usf
@@ -10,7 +10,7 @@
 #define NUM_RESERVED_ATTRS 32
 
 [numthreads(64, 1, 1)]
-void Main(uint3 GroupId : SV_GroupID, uint GroupIndex : SV_GroupIndex)
+void PCGCopyPointsAnalysisCS(uint3 GroupId : SV_GroupID, uint GroupIndex : SV_GroupIndex)
 {
 	// Mark the kernel as having executed. Must run before we early out via thread index, because the kernel is still 'executed' even if the number of
 	// threads to iterate on is zero. Even if GetNumThreads() returns 0, the kernel will still have been dispatched on a single thread to set this flag.
diff --git a/Engine/Plugins/PCG/Shaders/Private/Elements/PCGCopyPointsConditional.usf b/Engine/Plugins/PCG/Shaders/Private/Elements/PCGCopyPointsConditional.usf
--- a/Engine/Plugins/PCG/Shaders/Private/Elements/PCGCopyPointsConditional.usf
+++ b/Engine/Plugins/PCG/Shaders/Private/Elements/PCGCopyPointsConditional.usf
@@ -24,7 +24,7 @@
 #define NUM_RESERVED_ATTRS 32
 
 [numthreads(64, 1, 1)]
-void Main(uint3 GroupId : SV_GroupID, uint GroupIndex : SV_GroupIndex)
+void PCGCopyPointsCS(uint3 GroupId : SV_GroupID, uint GroupIndex : SV_GroupIndex)
 {
 	// Mark the kernel as having executed. Must run before we early out via thread index, because the kernel is still 'executed' even if the number of
 	// threads to iterate on is zero. Even if GetNumThreads() returns 0, the kernel will still have been dispatched on a single thread to set this flag.
diff --git a/Engine/Plugins/PCG/Shaders/Private/Elements/PCGMetadataPartition.usf b/Engine/Plugins/PCG/Shaders/Private/Elements/PCGMetadataPartition.usf
--- a/Engine/Plugins/PCG/Shaders/Private/Elements/PCGMetadataPartition.usf
+++ b/Engine/Plugins/PCG/Shaders/Private/Elements/PCGMetadataPartition.usf
@@ -12,7 +12,7 @@
 #define NUM_RESERVED_ATTRS 32
 
 [numthreads(64, 1, 1)]
-void Main(uint3 GroupId : SV_GroupID, uint GroupIndex : SV_GroupIndex)
+void PCGMetadataPartitionCS(uint3 GroupId : SV_GroupID, uint GroupIndex : SV_GroupIndex)
 {
 	// Mark the kernel as having executed. Must run before we early out via thread index, because the kernel is still 'executed' even if the number of
 	// threads to iterate on is zero. Even if GetNumThreads() returns 0, the kernel will still have been dispatched on a single thread to set this flag.
diff --git a/Engine/Plugins/PCG/Shaders/Private/Elements/PCGStaticMeshSpawner.usf b/Engine/Plugins/PCG/Shaders/Private/Elements/PCGStaticMeshSpawner.usf
--- a/Engine/Plugins/PCG/Shaders/Private/Elements/PCGStaticMeshSpawner.usf
+++ b/Engine/Plugins/PCG/Shaders/Private/Elements/PCGStaticMeshSpawner.usf
@@ -15,7 +15,7 @@
 #define NUM_RESERVED_ATTRS 32
 
 [numthreads(64, 1, 1)]
-void Main(uint3 GroupId : SV_GroupID, uint GroupIndex : SV_GroupIndex)
+void PCGStaticMeshSpawnerCS(uint3 GroupId : SV_GroupID, uint GroupIndex : SV_GroupIndex)
 {
 	// Mark the kernel as having executed. Must run before we early out via thread index, because the kernel is still 'executed' even if the number of
 	// threads to iterate on is zero. Even if GetNumThreads() returns 0, the kernel will still have been dispatched on a single thread to set this flag.
diff --git a/Engine/Plugins/PCG/Source/PCG/Private/Compute/Elements/PCGCustomHLSLKernel.cpp b/Engine/Plugins/PCG/Source/PCG/Private/Compute/Elements/PCGCustomHLSLKernel.cpp
--- a/Engine/Plugins/PCG/Source/PCG/Private/Compute/Elements/PCGCustomHLSLKernel.cpp
+++ b/Engine/Plugins/PCG/Source/PCG/Private/Compute/Elements/PCGCustomHLSLKernel.cpp
@@ -44,12 +44,17 @@ namespace PCGCustomHLSLKernel
 	}
 }
 
+#if WITH_EDITOR
 void UPCGCustomHLSLKernel::Initialize(const UPCGCustomHLSLSettings* InSettings, bool bLogDescriptions)
 {
 	Super::Initialize(bLogDescriptions);
 
+	check(InSettings);
 	Settings = InSettings;
+
+	InitEntryPoint();
 }
+#endif // WITH_EDITOR
 
 const UPCGSettings* UPCGCustomHLSLKernel::GetSettings() const
 {
@@ -630,7 +635,7 @@ FString UPCGCustomHLSLKernel::GetCookedSource(const UPCGComputeGraph* InComputeG
 
 	const FString KernelFunc = FString::Printf(
 		TEXT("[numthreads(%d, %d, %d)]\nvoid %s(uint3 GroupId : SV_GroupID, uint GroupIndex : SV_GroupIndex)"),
-		GroupSize.X, GroupSize.Y, GroupSize.Z, TEXT("Main"));
+		GroupSize.X, GroupSize.Y, GroupSize.Z, *GetEntryPoint());
 
 	const FString UnWrappedDispatchThreadId = FString::Printf(
 		TEXT("GetUnWrappedDispatchThreadId(GroupId, GroupIndex, %d)"),
@@ -850,6 +855,40 @@ void UPCGCustomHLSLKernel::GetOutputPins(TArray<FPCGPinPropertiesGPU>& OutPins)
 	OutPins = Settings->OutputPins;
 }
 
+#if WITH_EDITOR
+void UPCGCustomHLSLKernel::InitEntryPoint()
+{
+	check(Settings);
+
+	// Record the node title - used to name the kernel and shows up in profiling.
+	const UPCGNode* Node = Cast<UPCGNode>(Settings->GetOuter());
+	if (Node && Node->HasAuthoredTitle())
+	{
+		EntryPoint = Node->GetAuthoredTitleName().ToString();
+	}
+	else
+	{
+		EntryPoint = Settings->GetDefaultNodeTitle().ToString();
+	}
+
+	// We append the PCG kernel name here because of a weird issue on Mac where if the last word of the name of a
+	// shader parameter matches the kernel name, the value for the parameter is always zero for unknown reasons
+	// (e.g. parameter named "XXXWeight" & kernel also named "Weight"). Adding this bit of "random" string after
+	// the user provided kernel name should greatly reduce the chance of things like that happening
+	EntryPoint += TEXT("_") + GetName();
+
+	const TCHAR InvalidCharacters[] = {
+		TCHAR('\"'), TCHAR('\\'), TCHAR('\''), TCHAR('\"'), TCHAR(' '), TCHAR(','), TCHAR('.'), TCHAR('|'), TCHAR('&'), TCHAR('!'),
+		TCHAR('~'), TCHAR('\n'), TCHAR('\r'), TCHAR('\t'), TCHAR('@'), TCHAR('#'), TCHAR('/'), TCHAR('('), TCHAR(')'), TCHAR('{'),
+		TCHAR('}'), TCHAR('['), TCHAR(']'), TCHAR('='), TCHAR(';'), TCHAR(':'), TCHAR('^'), TCHAR('%'), TCHAR('$'), TCHAR('`') };
+
+	for (TCHAR Char : InvalidCharacters)
+	{
+		EntryPoint.ReplaceCharInline(Char, TCHAR('_'));
+	}
+}
+#endif // WITH_EDITOR
+
 bool UPCGCustomHLSLKernel::IsThreadCountMultiplierInUse() const
 {
 	return Settings->KernelType == EPCGKernelType::Custom && Settings->DispatchThreadCount != EPCGDispatchThreadCount::Fixed;
diff --git a/Engine/Plugins/PCG/Source/PCG/Private/Compute/Elements/PCGCustomHLSLKernel.h b/Engine/Plugins/PCG/Source/PCG/Private/Compute/Elements/PCGCustomHLSLKernel.h
--- a/Engine/Plugins/PCG/Source/PCG/Private/Compute/Elements/PCGCustomHLSLKernel.h
+++ b/Engine/Plugins/PCG/Source/PCG/Private/Compute/Elements/PCGCustomHLSLKernel.h
@@ -16,13 +16,17 @@ class UPCGCustomHLSLKernel : public UPCGComputeKernel
 	GENERATED_BODY()
 
 public:
+#if WITH_EDITOR
+	/** Initialize kernel. Editor only as shader compilation only available in editor. */
 	void Initialize(const UPCGCustomHLSLSettings* InSettings, bool bLogDescriptions);
+#endif
 
 	virtual const UPCGSettings* GetSettings() const override;
 	virtual bool IsKernelValid(FPCGContext* InContext = nullptr, bool bQuiet = true) const override;
 	virtual bool ComputeOutputBindingDataDesc(const UPCGComputeGraph* InGraph, FName InOutputPinLabel, const UPCGDataBinding* InBinding, FPCGDataCollectionDesc& OutDataDesc) const override;
 #if WITH_EDITOR
 	virtual FString GetCookedSource(const UPCGComputeGraph* InComputeGraph) const override;
+	virtual FString GetEntryPoint() const override { return EntryPoint; }
 #endif
 	virtual int ComputeThreadCount(const UPCGDataBinding* Binding) const override;
 	virtual void GetDataLabels(FName InPinLabel, TArray<FString>& OutDataLabels) const override;
@@ -33,6 +37,10 @@ public:
 	virtual void GetOutputPins(TArray<FPCGPinPropertiesGPU>& OutPins) const override;
 
 protected:
+#if WITH_EDITOR
+	void InitEntryPoint();
+#endif
+
 	/** Will the ThreadCountMultiplier value be applied when calculating the dispatch thread count. */
 	bool IsThreadCountMultiplierInUse() const;
 
@@ -51,4 +59,10 @@ protected:
 public:
 	UPROPERTY()
 	TObjectPtr<const UPCGCustomHLSLSettings> Settings;
+
+#if WITH_EDITORONLY_DATA
+	/** The name of the main function in the shader. Generated from the node title. */
+	UPROPERTY()
+	FString EntryPoint;
+#endif
 };
diff --git a/Engine/Plugins/PCG/Source/PCG/Private/Compute/PCGComputeKernel.h b/Engine/Plugins/PCG/Source/PCG/Private/Compute/PCGComputeKernel.h
--- a/Engine/Plugins/PCG/Source/PCG/Private/Compute/PCGComputeKernel.h
+++ b/Engine/Plugins/PCG/Source/PCG/Private/Compute/PCGComputeKernel.h
@@ -41,6 +41,9 @@ public:
 	/** Produces the node specific portion of kernel shader source text, including the main entry point. */
 	virtual FString GetCookedSource(const UPCGComputeGraph* InComputeGraph) const PURE_VIRTUAL(UPCGComputeKernel::GetCookedSource, return {};);
 
+	/** Get the name of the main kernel function in the source. This name is also displayed in GPU profile scopes. */
+	virtual FString GetEntryPoint() const PURE_VIRTUAL(UPCGComputeKernel::GetEntryPoint, return TEXT("Main"););
+
 	/** Create additional input data interfaces to marshal any required input data. */
 	virtual void CreateAdditionalInputDataInterfaces(FPCGGPUCompilationContext& InOutContext, UObject* InObjectOuter, TArray<TObjectPtr<UComputeDataInterface>>& OutDataInterfaces) const {}
 
diff --git a/Engine/Plugins/PCG/Source/PCG/Private/Graph/PCGGraphCompilerGPU.cpp b/Engine/Plugins/PCG/Source/PCG/Private/Graph/PCGGraphCompilerGPU.cpp
--- a/Engine/Plugins/PCG/Source/PCG/Private/Graph/PCGGraphCompilerGPU.cpp
+++ b/Engine/Plugins/PCG/Source/PCG/Private/Graph/PCGGraphCompilerGPU.cpp
@@ -1396,7 +1396,7 @@ void FPCGGraphCompilerGPU::CompileComputeGraph(
 			KernelWithBindings.Kernel->KernelSource = KernelSource;
 
 			// These could be exposed through PCGSettings API later when the need arises (and/or when GPU feature matures).
-			KernelSource->EntryPoint = TEXT("Main");
+			KernelSource->EntryPoint = KernelWithBindings.Kernel->GetEntryPoint();
 			KernelSource->GroupSize = FIntVector(PCGComputeConstants::THREAD_GROUP_SIZE, 1, 1);
 
 			// TODO: Instead of passing a half-baked compute graph into this virtual call, we should pass in a compilation context per compute graph.
diff --git a/Engine/Plugins/PCG/Source/PCG/Public/Elements/Metadata/PCGMetadataPartitionKernel.h b/Engine/Plugins/PCG/Source/PCG/Public/Elements/Metadata/PCGMetadataPartitionKernel.h
--- a/Engine/Plugins/PCG/Source/PCG/Public/Elements/Metadata/PCGMetadataPartitionKernel.h
+++ b/Engine/Plugins/PCG/Source/PCG/Public/Elements/Metadata/PCGMetadataPartitionKernel.h
@@ -22,6 +22,7 @@ public:
 	virtual int ComputeThreadCount(const UPCGDataBinding* Binding) const override;
 #if WITH_EDITOR
 	virtual FString GetCookedSource(const UPCGComputeGraph* InComputeGraph) const override;
+	virtual FString GetEntryPoint() const override { return TEXT("PCGMetadataPartitionCS"); }
 	void CreateAdditionalInputDataInterfaces(FPCGGPUCompilationContext& InOutContext, UObject* InObjectOuter, TArray<TObjectPtr<UComputeDataInterface>>& OutDataInterfaces) const override;
 	// Split graph to update element counts on CPU before proceeding. Compacts data for downstream kernels.
 	virtual bool SplitGraphAtOutput() const override { return true; }
diff --git a/Engine/Plugins/PCG/Source/PCG/Public/Elements/PCGCopyPointsAnalysisKernel.h b/Engine/Plugins/PCG/Source/PCG/Public/Elements/PCGCopyPointsAnalysisKernel.h
--- a/Engine/Plugins/PCG/Source/PCG/Public/Elements/PCGCopyPointsAnalysisKernel.h
+++ b/Engine/Plugins/PCG/Source/PCG/Public/Elements/PCGCopyPointsAnalysisKernel.h
@@ -22,6 +22,7 @@ public:
 	virtual int ComputeThreadCount(const UPCGDataBinding* Binding) const override;
 #if WITH_EDITOR
 	virtual FString GetCookedSource(const UPCGComputeGraph* InComputeGraph) const override;
+	virtual FString GetEntryPoint() const override { return TEXT("PCGCopyPointsAnalysisCS"); }
 	void CreateAdditionalInputDataInterfaces(FPCGGPUCompilationContext& InOutContext, UObject* InObjectOuter, TArray<TObjectPtr<UComputeDataInterface>>& OutDataInterfaces) const override;
 	// Split graph to read back analysis results.
 	virtual bool SplitGraphAtOutput() const override { return true; }
diff --git a/Engine/Plugins/PCG/Source/PCG/Public/Elements/PCGCopyPointsKernel.h b/Engine/Plugins/PCG/Source/PCG/Public/Elements/PCGCopyPointsKernel.h
--- a/Engine/Plugins/PCG/Source/PCG/Public/Elements/PCGCopyPointsKernel.h
+++ b/Engine/Plugins/PCG/Source/PCG/Public/Elements/PCGCopyPointsKernel.h
@@ -24,6 +24,7 @@ public:
 	virtual int ComputeThreadCount(const UPCGDataBinding* Binding) const override;
 #if WITH_EDITOR
 	virtual FString GetCookedSource(const UPCGComputeGraph* InComputeGraph) const override;
+	virtual FString GetEntryPoint() const override { return TEXT("PCGCopyPointsCS"); }
 	void CreateAdditionalInputDataInterfaces(FPCGGPUCompilationContext& InOutContext, UObject* InObjectOuter, TArray<TObjectPtr<UComputeDataInterface>>& OutDataInterfaces) const override;
 #endif
 	virtual void GetInputPins(TArray<FPCGPinProperties>& OutPins) const override;
diff --git a/Engine/Plugins/PCG/Source/PCG/Public/Elements/PCGStaticMeshSpawnerKernel.h b/Engine/Plugins/PCG/Source/PCG/Public/Elements/PCGStaticMeshSpawnerKernel.h
--- a/Engine/Plugins/PCG/Source/PCG/Public/Elements/PCGStaticMeshSpawnerKernel.h
+++ b/Engine/Plugins/PCG/Source/PCG/Public/Elements/PCGStaticMeshSpawnerKernel.h
@@ -22,6 +22,7 @@ public:
 	virtual int ComputeThreadCount(const UPCGDataBinding* Binding) const override;
 #if WITH_EDITOR
 	virtual FString GetCookedSource(const UPCGComputeGraph* InComputeGraph) const override;
+	virtual FString GetEntryPoint() const override { return TEXT("PCGStaticMeshSpawnerCS"); }
 	void CreateAdditionalInputDataInterfaces(FPCGGPUCompilationContext& InOutContext, UObject* InObjectOuter, TArray<TObjectPtr<UComputeDataInterface>>& OutDataInterfaces) const override;
 	void CreateAdditionalOutputDataInterfaces(FPCGGPUCompilationContext& InOutContext, UObject* InObjectOuter, TArray<TObjectPtr<UComputeDataInterface>>& OutDataInterfaces) const override;
 #endif
diff --git a/Engine/Plugins/PCG/Source/PCG/Public/PCGNode.h b/Engine/Plugins/PCG/Source/PCG/Public/PCGNode.h
--- a/Engine/Plugins/PCG/Source/PCG/Public/PCGNode.h
+++ b/Engine/Plugins/PCG/Source/PCG/Public/PCGNode.h
@@ -77,6 +77,9 @@ public:
 	/** Authored part of node title (like "Create Attribute 1"). */
 	FText GetAuthoredTitleLine() const;
 
+	/** Authored node title as raw name or none if no title authored. */
+	FName GetAuthoredTitleName() const { return NodeTitle; }
+
 	/** Whether to flip the order of the title lines - display generated title first and authored second. */
 	bool HasFlippedTitleLines() const;
 
