[PCG] Update SM DI API
* Rename some functions and update the declarations

#rnx

diff --git a/Engine/Plugins/PCG/Shaders/Private/PCGStaticMeshDataInterface.ush b/Engine/Plugins/PCG/Shaders/Private/PCGStaticMeshDataInterface.ush
--- a/Engine/Plugins/PCG/Shaders/Private/PCGStaticMeshDataInterface.ush
+++ b/Engine/Plugins/PCG/Shaders/Private/PCGStaticMeshDataInterface.ush
@@ -28,7 +28,7 @@ int GetNumVertices_{DataInterfaceName}(int DataIndex)
 	return {DataInterfaceName}_NumVertices;
 }
 
-void SampleVertex_{DataInterfaceName}(int DataIndex, int VertexIndex, out float3 OutPosition, out float3 OutNormal, out float3 OutTangent, out float3 OutBitangent)
+void GetVertex_{DataInterfaceName}(int DataIndex, int VertexIndex, out float3 OutPosition, out float3 OutNormal, out float3 OutTangent, out float3 OutBitangent)
 {
 	if ({DataInterfaceName}_NumVertices > 0)
 	{
@@ -65,13 +65,13 @@ void SampleVertex_{DataInterfaceName}(int DataIndex, int VertexIndex, out float3
 	}
 }
 
-float4 SampleVertexColor_{DataInterfaceName}(int DataIndex, int VertexIndex)
+float4 GetVertexColor_{DataInterfaceName}(int DataIndex, int VertexIndex)
 {
 	VertexIndex = clamp(VertexIndex, 0, {DataInterfaceName}_NumVertices - 1);
 	return {DataInterfaceName}_HasColors ? {DataInterfaceName}_ColorBuffer[VertexIndex] FMANUALFETCH_COLOR_COMPONENT_SWIZZLE : float4(1, 1, 1, 1);
 }
 
-float2 SampleVertexUV_{DataInterfaceName}(int DataIndex, int VertexIndex, int UVSet)
+float2 GetVertexUV_{DataInterfaceName}(int DataIndex, int VertexIndex, int UVSet)
 {
 	VertexIndex = clamp(VertexIndex, 0, {DataInterfaceName}_NumVertices - 1);
 	UVSet = clamp(UVSet, 0, {DataInterfaceName}_NumUVs - 1);
@@ -115,9 +115,9 @@ void SampleTriangle_{DataInterfaceName}(
 	GetTriangleIndices_{DataInterfaceName}(DataIndex, TriangleIndex, Index[0], Index[1], Index[2]);
 
 	float3 Positions[3], Normals[3], Bitangents[3], Tangents[3];
-	SampleVertex_{DataInterfaceName}(DataIndex, Index[0], Positions[0], Normals[0], Tangents[0], Bitangents[0]);
-	SampleVertex_{DataInterfaceName}(DataIndex, Index[1], Positions[1], Normals[1], Tangents[1], Bitangents[1]);
-	SampleVertex_{DataInterfaceName}(DataIndex, Index[2], Positions[2], Normals[2], Tangents[2], Bitangents[2]);
+	GetVertex_{DataInterfaceName}(DataIndex, Index[0], Positions[0], Normals[0], Tangents[0], Bitangents[0]);
+	GetVertex_{DataInterfaceName}(DataIndex, Index[1], Positions[1], Normals[1], Tangents[1], Bitangents[1]);
+	GetVertex_{DataInterfaceName}(DataIndex, Index[2], Positions[2], Normals[2], Tangents[2], Bitangents[2]);
 
 	OutPosition = Positions[0] * BaryCoord.x + Positions[1] * BaryCoord.y + Positions[2] * BaryCoord.z;
 	OutNormal = Normals[0] * BaryCoord.x + Normals[1] * BaryCoord.y + Normals[2] * BaryCoord.x;
@@ -131,9 +131,9 @@ float4 SampleTriangleColor_{DataInterfaceName}(int DataIndex, int TriangleIndex,
 	GetTriangleIndices_{DataInterfaceName}(DataIndex, TriangleIndex, Index[0], Index[1], Index[2]);
 
 	float4 Colors[3];
-	Colors[0] = SampleVertexColor_{DataInterfaceName}(DataIndex, Index[0]);
-	Colors[1] = SampleVertexColor_{DataInterfaceName}(DataIndex, Index[1]);
-	Colors[2] = SampleVertexColor_{DataInterfaceName}(DataIndex, Index[2]);
+	Colors[0] = GetVertexColor_{DataInterfaceName}(DataIndex, Index[0]);
+	Colors[1] = GetVertexColor_{DataInterfaceName}(DataIndex, Index[1]);
+	Colors[2] = GetVertexColor_{DataInterfaceName}(DataIndex, Index[2]);
 
 	return Colors[0] * BaryCoord.x + Colors[1] * BaryCoord.y + Colors[2] * BaryCoord.z;
 }
@@ -144,9 +144,9 @@ float2 SampleTriangleUV_{DataInterfaceName}(int DataIndex, int TriangleIndex, fl
 	GetTriangleIndices_{DataInterfaceName}(DataIndex, TriangleIndex, Index[0], Index[1], Index[2]);
 
 	float2 UVs[3];
-	UVs[0] = SampleVertexUV_{DataInterfaceName}(DataIndex, Index[0], UVSet);
-	UVs[1] = SampleVertexUV_{DataInterfaceName}(DataIndex, Index[1], UVSet);
-	UVs[2] = SampleVertexUV_{DataInterfaceName}(DataIndex, Index[2], UVSet);
+	UVs[0] = GetVertexUV_{DataInterfaceName}(DataIndex, Index[0], UVSet);
+	UVs[1] = GetVertexUV_{DataInterfaceName}(DataIndex, Index[1], UVSet);
+	UVs[2] = GetVertexUV_{DataInterfaceName}(DataIndex, Index[2], UVSet);
 
 	return UVs[0] * BaryCoord.x + UVs[1] * BaryCoord.y + UVs[2] * BaryCoord.z;
 }
diff --git a/Engine/Plugins/PCG/Source/PCG/Private/Compute/DataInterfaces/PCGStaticMeshDataInterface.cpp b/Engine/Plugins/PCG/Source/PCG/Private/Compute/DataInterfaces/PCGStaticMeshDataInterface.cpp
--- a/Engine/Plugins/PCG/Source/PCG/Private/Compute/DataInterfaces/PCGStaticMeshDataInterface.cpp
+++ b/Engine/Plugins/PCG/Source/PCG/Private/Compute/DataInterfaces/PCGStaticMeshDataInterface.cpp
@@ -30,7 +30,7 @@ void UPCGStaticMeshDataInterface::GetSupportedInputs(TArray<FShaderFunctionDefin
 			.AddParam(EShaderFundamentalType::Int); // DataIndex
 
 		OutFunctions.AddDefaulted_GetRef()
-			.SetName(TEXT("SampleVertex"))
+			.SetName(TEXT("GetVertex"))
 			.AddParam(EShaderFundamentalType::Int) // DataIndex
 			.AddParam(EShaderFundamentalType::Int) // VertexIndex
 			.AddParam(EShaderFundamentalType::Float, 3, 0, EShaderParamModifier::Out) // OutPosition
@@ -39,13 +39,13 @@ void UPCGStaticMeshDataInterface::GetSupportedInputs(TArray<FShaderFunctionDefin
 			.AddParam(EShaderFundamentalType::Float, 3, 0, EShaderParamModifier::Out); // OutBitangent
 
 		OutFunctions.AddDefaulted_GetRef()
-			.SetName(TEXT("SampleVertexColor"))
+			.SetName(TEXT("GetVertexColor"))
 			.AddReturnType(EShaderFundamentalType::Float, 4)
 			.AddParam(EShaderFundamentalType::Int) // DataIndex
 			.AddParam(EShaderFundamentalType::Int); // VertexIndex
 
 		OutFunctions.AddDefaulted_GetRef()
-			.SetName(TEXT("SampleVertexUV"))
+			.SetName(TEXT("GetVertexUV"))
 			.AddReturnType(EShaderFundamentalType::Float, 2)
 			.AddParam(EShaderFundamentalType::Int) // DataIndex
 			.AddParam(EShaderFundamentalType::Int) // VertexIndex
diff --git a/Engine/Plugins/PCG/Source/PCG/Private/Compute/Elements/PCGCustomHLSL.cpp b/Engine/Plugins/PCG/Source/PCG/Private/Compute/Elements/PCGCustomHLSL.cpp
--- a/Engine/Plugins/PCG/Source/PCG/Private/Compute/Elements/PCGCustomHLSL.cpp
+++ b/Engine/Plugins/PCG/Source/PCG/Private/Compute/Elements/PCGCustomHLSL.cpp
@@ -572,14 +572,14 @@ void UPCGCustomHLSLSettings::UpdateInputDeclarations()
 		InputDeclarations += FString::Format(TEXT(
 			"// Vertex functions\n"
 			"int {0}_GetNumVertices(int DataIndex);\n"
-			"void {0}_SampleVertex(int DataIndex, int VertexIndex, float3 OutPosition, float3 OutNormal, float3 OutTangent, float3 OutBitangent);\n"
-			"float4 {0}_SampleVertexColor(int DataIndex, int VertexIndex);\n"
-			"float2 {0}_SampleVertexUVs(int DataIndex, int VertexIndex, int UVSet);\n"
+			"void {0}_GetVertex(int DataIndex, int VertexIndex, out float3 OutPosition, out float3 OutNormal, out float3 OutTangent, out float3 OutBitangent);\n"
+			"float4 {0}_GetVertexColor(int DataIndex, int VertexIndex);\n"
+			"float2 {0}_GetVertexUVs(int DataIndex, int VertexIndex, int UVSet);\n"
 			"\n"
 			"// Triangle functions\n"
 			"int {0}_GetNumTriangles(int DataIndex);\n"
-			"void {0}_GetTriangleIndices(int DataIndex, int TriangleIndex, int OutIndex0, int OutIndex1, int OutIndex2);\n"
-			"void {0}_SampleTriangle(int DataIndex, int TriangleIndex, float3 BaryCoords, float3 OutPosition, float3 OutNormal, float3 OutTangent, float3 OutBitangent);\n"
+			"void {0}_GetTriangleIndices(int DataIndex, int TriangleIndex, out int OutIndex0, out int OutIndex1, out int OutIndex2);\n"
+			"void {0}_SampleTriangle(int DataIndex, int TriangleIndex, float3 BaryCoords, out float3 OutPosition, out float3 OutNormal, out float3 OutTangent, out float3 OutBitangent);\n"
 			"float4 {0}_SampleTriangleColor(int DataIndex, int TriangleIndex, float3 BaryCoords);\n"
 			"float2 {0}_SampleTriangleUVs(int DataIndex, int TriangleIndex, float3 BaryCoords, int UVSet);\n"
 			"\n"
