[PCG] Remove internal include from public header

#rnx

diff --git a/Engine/Plugins/PCG/Source/PCG/Public/PCGGraph.h b/Engine/Plugins/PCG/Source/PCG/Public/PCGGraph.h
--- a/Engine/Plugins/PCG/Source/PCG/Public/PCGGraph.h
+++ b/Engine/Plugins/PCG/Source/PCG/Public/PCGGraph.h
@@ -5,7 +5,6 @@
 #include "PCGCommon.h"
 #include "PCGNode.h"
 #include "PCGSettings.h"
-#include "Compute/PCGCompilerDiagnostic.h"
 #include "Editor/PCGGraphCustomization.h"
 #include "Graph/PCGStackContext.h"
 #include "Helpers/PCGGraphParameterExtension.h"
@@ -16,6 +15,7 @@
 #include "PCGGraph.generated.h"
 
 struct FAssetData;
+struct FPCGCompilerDiagnostics;
 class UPCGGraphInterface;
 class UPCGGraphCompilationData;
 #if WITH_EDITOR
