Forced back actor reference / component reference as last attributes in the Get Actor Reference / Get Component Reference mode.

#rnx

diff --git a/Engine/Plugins/PCG/Source/PCG/Private/Elements/PCGDataFromActor.cpp b/Engine/Plugins/PCG/Source/PCG/Private/Elements/PCGDataFromActor.cpp
--- a/Engine/Plugins/PCG/Source/PCG/Private/Elements/PCGDataFromActor.cpp
+++ b/Engine/Plugins/PCG/Source/PCG/Private/Elements/PCGDataFromActor.cpp
@@ -799,6 +799,11 @@ void FPCGDataFromActorElement::CreateReferenceData(FPCGContext* InContext, const
 		}
 	}
 
+	// Mark the actor or component reference as the last attribute in order to behave like before wrt to this.
+	FPCGAttributePropertySelector LastSelector;
+	LastSelector.SetAttributeName(bCreateComponentReferences ? PCGAddComponentConstants::ComponentReferenceAttribute : PCGPointDataConstants::ActorReferenceAttribute);
+	ParamData->SetLastSelector(LastSelector);
+
 	if (ParamData->MutableMetadata()->GetItemCountForChild() > 0)
 	{
 		FPCGTaggedData& TaggedData = Context->OutputData.TaggedData.Emplace_GetRef();
