Minor optimization in tracking (used a const ref).
Passed recursive lambdas by ref.

#rnx

diff --git a/Engine/Plugins/PCG/Source/PCG/Private/Elements/Grammar/PCGPrintGrammar.cpp b/Engine/Plugins/PCG/Source/PCG/Private/Elements/Grammar/PCGPrintGrammar.cpp
--- a/Engine/Plugins/PCG/Source/PCG/Private/Elements/Grammar/PCGPrintGrammar.cpp
+++ b/Engine/Plugins/PCG/Source/PCG/Private/Elements/Grammar/PCGPrintGrammar.cpp
@@ -60,7 +60,7 @@ bool FPCGPrintGrammarElement::ExecuteInternal(FPCGContext* Context) const
 		return true;
 	}
 
-	auto PrintDescriptor = [&Description](const PCGGrammar::FModuleDescriptor& Descriptor, auto PrintRecurse) -> void
+	auto PrintDescriptor = [&Description](const PCGGrammar::FModuleDescriptor& Descriptor, auto&& PrintRecurse) -> void
 		{
 			Description.Append("Module[");
 			switch (Descriptor.Type)
diff --git a/Engine/Plugins/PCG/Source/PCG/Private/Helpers/PCGSubgraphHelpers.cpp b/Engine/Plugins/PCG/Source/PCG/Private/Helpers/PCGSubgraphHelpers.cpp
--- a/Engine/Plugins/PCG/Source/PCG/Private/Helpers/PCGSubgraphHelpers.cpp
+++ b/Engine/Plugins/PCG/Source/PCG/Private/Helpers/PCGSubgraphHelpers.cpp
@@ -663,7 +663,7 @@ namespace PCGSubgraphHelpersExtra
 			// Also keep track of all seen nodes, so our algorithm stops even if there are cycles already existing.
 			TSet<const UPCGNode*> SeenNodes;
 
-			auto Visitor = [&NodesToCollapse, &SeenNodes](const UPCGNode* CurrentNode, bool bInSubgraphBoundary, auto VisitorRecurse) -> bool
+			auto Visitor = [&NodesToCollapse, &SeenNodes](const UPCGNode* CurrentNode, bool bInSubgraphBoundary, auto&& VisitorRecurse) -> bool
 			{
 				if (!CurrentNode || SeenNodes.Contains(CurrentNode))
 				{
diff --git a/Engine/Plugins/PCG/Source/PCG/Private/PCGActorAndComponentMapping.cpp b/Engine/Plugins/PCG/Source/PCG/Private/PCGActorAndComponentMapping.cpp
--- a/Engine/Plugins/PCG/Source/PCG/Private/PCGActorAndComponentMapping.cpp
+++ b/Engine/Plugins/PCG/Source/PCG/Private/PCGActorAndComponentMapping.cpp
@@ -1521,7 +1521,7 @@ void FPCGActorAndComponentMapping::OnActorChanged_Recursive(AActor* InActor, UOb
 	check(!DelayedChangedActors.Contains(InActor));
 
 	// Process actors recursively
-	auto OnActorChanged = [this, InOriginatingChangeObject](AActor* InActor, auto OnActorChangedRecursive) -> void
+	auto OnActorChanged = [this, InOriginatingChangeObject](AActor* InActor, auto&& OnActorChangedRecursive) -> void
 	{
 		if (InActor)
 		{
@@ -1664,7 +1664,7 @@ void FPCGActorAndComponentMapping::OnObjectModified(UObject* InObject)
 	}
 #endif
 
-	auto StorePreviousData = [this](AActor* InActor, auto RecursiveCall) -> void
+	auto StorePreviousData = [this](AActor* InActor, auto&& RecursiveCall) -> void
 	{
 		if (!ActorToPreviousDataMap.Contains(InActor))
 		{
@@ -1760,7 +1760,7 @@ void FPCGActorAndComponentMapping::OnObjectPropertyChanged(UObject* InObject, FP
 		// Don't remove it as it will be needed when the delay will be called.
 		if (Actor && !(DelayedChangedActors.Contains(Actor)))
 		{
-			auto RemoveAll = [this](AActor* InActor, auto RecursiveCall) -> void
+			auto RemoveAll = [this](AActor* InActor, auto&& RecursiveCall) -> void
 			{
 				ActorToPreviousDataMap.Remove(InActor);
 				PCGActorAndComponentMapping::PropagateToLevelInstanceActors(InActor, [this, RecursiveCall](AActor* LevelActor)
@@ -1824,9 +1824,9 @@ void FPCGActorAndComponentMapping::OnObjectChanged(UObject* InObject, const FAct
 
 	FSoftObjectPath ObjectPath(InObject);
 	FSoftObjectPath InOriginatingChangeObjectPath(InOriginatingChangeObject);
-	auto Gather = [InObject, &ObjectPath, &RemovedTags, &MatchedKeys, InOriginatingChangeObject, &InOriginatingChangeObjectPath](TMap<FPCGSelectionKey, TSet<UPCGComponent*>> InMap, TSet<UPCGComponent*>& OutSet)
+	auto Gather = [InObject, &ObjectPath, &RemovedTags, &MatchedKeys, InOriginatingChangeObject, &InOriginatingChangeObjectPath](const TMap<FPCGSelectionKey, TSet<UPCGComponent*>>& InMap, TSet<UPCGComponent*>& OutSet)
 	{
-		for (auto& It : InMap)
+		for (const auto& It : InMap)
 		{
 			if (It.Key.IsMatching(InObject, ObjectPath, RemovedTags, It.Value, &OutSet) || It.Key.IsMatching(InOriginatingChangeObject, InOriginatingChangeObjectPath, RemovedTags, It.Value, &OutSet))
 			{
diff --git a/Engine/Plugins/PCG/Source/PCG/Private/PCGSubsystem.cpp b/Engine/Plugins/PCG/Source/PCG/Private/PCGSubsystem.cpp
--- a/Engine/Plugins/PCG/Source/PCG/Private/PCGSubsystem.cpp
+++ b/Engine/Plugins/PCG/Source/PCG/Private/PCGSubsystem.cpp
@@ -1496,7 +1496,7 @@ void UPCGSubsystem::DeleteSerializedPartitionActors(bool bOnlyDeleteUnused, bool
 	UWorldPartition* WorldPartition = World->GetWorldPartition();
 	TMap<FGuid, TArray<FGuid>> Attachments;
 
-	auto GetAttachments = [WorldPartition](const FGuid& ParentActor, const TMap<FGuid, TArray<FGuid>>& Attachments, auto GetAttachmentsRecursive, TArray<FGuid>& OutAttachedActors) -> void
+	auto GetAttachments = [WorldPartition](const FGuid& ParentActor, const TMap<FGuid, TArray<FGuid>>& Attachments, auto&& GetAttachmentsRecursive, TArray<FGuid>& OutAttachedActors) -> void
 	{
 		if (const TArray<FGuid>* Attached = Attachments.Find(ParentActor))
 		{
diff --git a/Engine/Plugins/PCG/Source/PCGEditor/Private/PCGEditorGraphSchema.cpp b/Engine/Plugins/PCG/Source/PCGEditor/Private/PCGEditorGraphSchema.cpp
--- a/Engine/Plugins/PCG/Source/PCGEditor/Private/PCGEditorGraphSchema.cpp
+++ b/Engine/Plugins/PCG/Source/PCGEditor/Private/PCGEditorGraphSchema.cpp
@@ -578,7 +578,7 @@ void UPCGEditorGraphSchema::GetSubgraphElementActions(FGraphActionMenuBuilder& A
 			// Only exposing the instance if it's parent graph is defined for instances. Otherwise, it is not interesting.
 			// Also march up the hierarchy to find overrides for category and description.
 			// We don't look for override in titles because if we have no override for an instance but the parent has an override we will have 2 times the same entry in the palette.
-			auto GetRecursiveTextsAndIsValid = [&AssetRegistryModule](const FAssetData& InAssetData, FText& OutCategory, FText& OutDescription, auto Recurse)
+			auto GetRecursiveTextsAndIsValid = [&AssetRegistryModule](const FAssetData& InAssetData, FText& OutCategory, FText& OutDescription, auto&& Recurse)
 			{
 				if (InAssetData.IsInstanceOf<UPCGGraphInstance>())
 				{
