[PCG] Fixed editor crash, caused by node reconstruction on pin change during pin disconnection

#rnx

diff --git a/Engine/Plugins/PCG/Source/PCGEditor/Private/Nodes/PCGEditorGraphNodeBase.cpp b/Engine/Plugins/PCG/Source/PCGEditor/Private/Nodes/PCGEditorGraphNodeBase.cpp
--- a/Engine/Plugins/PCG/Source/PCGEditor/Private/Nodes/PCGEditorGraphNodeBase.cpp
+++ b/Engine/Plugins/PCG/Source/PCGEditor/Private/Nodes/PCGEditorGraphNodeBase.cpp
@@ -30,6 +30,8 @@
 #include "ToolMenuSection.h"
 #include "Logging/TokenizedMessage.h"
 #include "Misc/TransactionObjectEvent.h"
+#include "UObject/FortniteMainBranchObjectVersion.h"
+
 #include "Widgets/Colors/SColorPicker.h"
 
 #define LOCTEXT_NAMESPACE "PCGEditorGraphNodeBase"
@@ -190,6 +192,16 @@ void UPCGEditorGraphNodeBase::PostTransacted(const FTransactionObjectEvent& Tran
 	}
 }
 
+void UPCGEditorGraphNodeBase::PostLoad()
+{
+	Super::PostLoad();
+
+	if (GetLinkerCustomVersion(FFortniteMainBranchObjectVersion::GUID) < FFortniteMainBranchObjectVersion::PCGInlineConstantDefaultValues)
+	{
+		bHasEverBeenConnected = true;
+	}
+}
+
 void UPCGEditorGraphNodeBase::GetNodeContextMenuActions(UToolMenu* Menu, class UGraphNodeContextMenuContext* Context) const
 {
 	using namespace PCGEditorGraphNodeBase::Constants;
@@ -1364,7 +1376,7 @@ void UPCGEditorGraphNodeBase::PinConnectionListChanged(UEdGraphPin* Pin)
 					if (InputPin && InputPin->Properties.Label != Pin->PinName && Interface->IsPinDefaultValueEnabled(InputPin->Properties.Label))
 					{
 						// Connecting a pin will dirty anyway, so no need to dirty again
-						Interface->SetPinDefaultValueIsActivated(InputPin->Properties.Label, /*bIsActive=*/true);
+						Interface->SetPinDefaultValueIsActivated(InputPin->Properties.Label, /*bIsActive=*/true, /*bDirtySettings=*/false);
 					}
 				}
 
diff --git a/Engine/Plugins/PCG/Source/PCGEditor/Private/Nodes/PCGEditorGraphNodeBase.h b/Engine/Plugins/PCG/Source/PCGEditor/Private/Nodes/PCGEditorGraphNodeBase.h
--- a/Engine/Plugins/PCG/Source/PCGEditor/Private/Nodes/PCGEditorGraphNodeBase.h
+++ b/Engine/Plugins/PCG/Source/PCGEditor/Private/Nodes/PCGEditorGraphNodeBase.h
@@ -35,6 +35,7 @@ public:
 	// ~Begin UObject interface
 	virtual void BeginDestroy() override;
 	virtual void PostTransacted(const FTransactionObjectEvent& TransactionEvent) override;
+	virtual void PostLoad() override;
 	// ~End UObject interface
 
 	// ~Begin UEdGraphNode interface
