[PCG] Fix "Clear PCG Link"
- Do not attach Child stamps to orignal stamp as this breaks streaming (because of references)
- Instead place Child stamps into a folder
- In WP worlds, load all PAs and Generated Actors so that Clear PCG Link is not partial
- Fix UPCGSubsystem::ForAllOverlappingCells not properly adding Cell Tasks to dependencies

diff --git a/Engine/Plugins/PCG/Source/PCG/Private/PCGComponent.cpp b/Engine/Plugins/PCG/Source/PCG/Private/PCGComponent.cpp
--- a/Engine/Plugins/PCG/Source/PCG/Private/PCGComponent.cpp
+++ b/Engine/Plugins/PCG/Source/PCG/Private/PCGComponent.cpp
@@ -80,32 +80,48 @@ namespace PCGComponent
 
 	struct FPCGManagedActorLoadingScope
 	{
+		FPCGManagedActorLoadingScope()
+		{
+#if WITH_EDITOR
+			bShouldLoadActors = CVarLoadManagedActorsForCleanup.GetValueOnAnyThread();
+#endif
+		}
+
+		FPCGManagedActorLoadingScope(bool bInShouldLoadActors)
+		{
+#if WITH_EDITOR
+			bShouldLoadActors = bInShouldLoadActors;
+#endif
+		}
+
 		void AddResources(UPCGComponent* InComponent, const TArray<UPCGManagedResource*>& InManagedResources)
 		{
 #if WITH_EDITOR
-			if (CVarLoadManagedActorsForCleanup.GetValueOnAnyThread())
+			if (!bShouldLoadActors)
 			{
-				if (!InComponent || !InComponent->GetWorld() || InComponent->GetWorld()->IsGameWorld())
-				{
-					return;
-				}
+				return;
+			}
 
-				UWorldPartition* WorldPartition = FWorldPartitionHelpers::GetWorldPartition(InComponent);
-				if (!WorldPartition)
-				{
-					return;
-				}
+			if (!InComponent || !InComponent->GetWorld() || InComponent->GetWorld()->IsGameWorld())
+			{
+				return;
+			}
 
-				for (UPCGManagedResource* ManagedResource : InManagedResources)
+			UWorldPartition* WorldPartition = FWorldPartitionHelpers::GetWorldPartition(InComponent);
+			if (!WorldPartition)
+			{
+				return;
+			}
+
+			for (UPCGManagedResource* ManagedResource : InManagedResources)
+			{
+				if (UPCGManagedActors* ManagedActors = Cast<UPCGManagedActors>(ManagedResource))
 				{
-					if (UPCGManagedActors* ManagedActors = Cast<UPCGManagedActors>(ManagedResource))
+					for (const TSoftObjectPtr<AActor>& ManagedActor : ManagedActors->GetConstGeneratedActors())
 					{
-						for (const TSoftObjectPtr<AActor>& ManagedActor : ManagedActors->GetConstGeneratedActors())
+						if (const FWorldPartitionActorDescInstance* ActorDescInstance = WorldPartition->GetActorDescInstanceByPath(ManagedActor.ToSoftObjectPath()))
 						{
-							if (const FWorldPartitionActorDescInstance* ActorDescInstance = WorldPartition->GetActorDescInstanceByPath(ManagedActor.ToSoftObjectPath()))
-							{
-								LoadedActors.Add(FWorldPartitionReference(WorldPartition, ActorDescInstance->GetGuid()));
-							}
+							LoadedActors.Add(FWorldPartitionReference(WorldPartition, ActorDescInstance->GetGuid()));
 						}
 					}
 				}
@@ -115,6 +131,7 @@ namespace PCGComponent
 
 #if WITH_EDITOR
 		TSet<FWorldPartitionReference> LoadedActors;
+		bool bShouldLoadActors;
 #endif
 	};
 }
@@ -1130,7 +1147,8 @@ AActor* UPCGComponent::ClearPCGLink(UClass* TemplateActorClass)
 	// First create a new actor that will be the new owner of all the resources
 	AActor* NewActor = UPCGActorHelpers::SpawnDefaultActor(SpawnDefaultActorParams);
 #if WITH_EDITOR
-	FActorLabelUtilities::SetActorLabelUnique(NewActor, DefaultName);
+	FActorLabelUtilities::SetActorLabelUnique(NewActor, GetOwner()->GetActorLabel() + TEXT("_") + DefaultName);
+	NewActor->SetFolderPath(GetOwner()->GetFolderPath());
 #endif
 
 	// Then move all resources linked to this component to this actor
@@ -1303,12 +1321,17 @@ bool UPCGComponent::MoveResourcesToNewActor(AActor* InNewActor, bool bCreateChil
 	}
 
 	check(InNewActor);
-	AActor* NewActor = InNewActor;
-
+	AActor* ParentActor = InNewActor;
+	AActor* CreatedChildActor = nullptr;
+	
 	bool bHasMovedResources = false;
 
 	Modify(!IsInPreviewMode());
 
+#if WITH_EDITOR
+	FName FolderPath;
+#endif
+
 	if (bCreateChild)
 	{
 		FActorSpawnParameters ActorSpawnParams;
@@ -1317,30 +1340,53 @@ bool UPCGComponent::MoveResourcesToNewActor(AActor* InNewActor, bool bCreateChil
 		ActorSpawnParams.NameMode = FActorSpawnParameters::ESpawnActorNameMode::Requested;
 		ActorSpawnParams.OverrideLevel = Owner->GetLevel();
 
-		UPCGActorHelpers::FSpawnDefaultActorParams SpawnDefaultActorParams(GetWorld(), NewActor->GetClass(), Owner->GetTransform(), ActorSpawnParams);
+		UPCGActorHelpers::FSpawnDefaultActorParams SpawnDefaultActorParams(GetWorld(), ParentActor->GetClass(), Owner->GetTransform(), ActorSpawnParams);
 
 #if WITH_EDITOR
 		SpawnDefaultActorParams.DataLayerInstances = Owner->GetDataLayerInstances();
 		SpawnDefaultActorParams.HLODLayer = Owner->GetHLODLayer();
 #endif
 
-		NewActor = UPCGActorHelpers::SpawnDefaultActor(SpawnDefaultActorParams);
+		CreatedChildActor = UPCGActorHelpers::SpawnDefaultActor(SpawnDefaultActorParams);
 #if WITH_EDITOR
-		FActorLabelUtilities::SetActorLabelUnique(NewActor, DefaultName);
+		FActorLabelUtilities::SetActorLabelUnique(CreatedChildActor, DefaultName);
+
+		if (InNewActor->GetFolderPath() != NAME_None)
+		{
+			FolderPath = FName(*(InNewActor->GetFolderPath().ToString() / InNewActor->GetActorLabel()));
+		}
+		else
+		{
+			FolderPath = FName(*InNewActor->GetActorLabel());
+		}
 #endif
-		NewActor->AttachToActor(InNewActor, FAttachmentTransformRules::KeepWorldTransform);
-		check(NewActor);
 	}
 
+#if WITH_EDITOR
 	// Trying to move all resources for now. Perhaps in the future we won't want that.
+	TSet<TSoftObjectPtr<AActor>> MovedActors;
+
+	// Do not rely on CVar to load actors when Clearing PCG Link
+	PCGComponent::FPCGManagedActorLoadingScope Scope(/*bInShouldLoadActors=*/true);
+#endif
 	{
 		FScopeLock ResourcesLock(&GeneratedResourcesLock);
 		check(!GeneratedResourcesInaccessible);
+#if WITH_EDITOR
+		Scope.AddResources(this, GeneratedResources);
+#endif
 		for (TObjectPtr<UPCGManagedResource>& GeneratedResource : GeneratedResources)
 		{
 			if (GeneratedResource)
 			{
-				GeneratedResource->MoveResourceToNewActor(NewActor, Owner);
+#if WITH_EDITOR
+				if (UPCGManagedActors* ManagedActors = Cast<UPCGManagedActors>(GeneratedResource))
+				{
+					MovedActors.Append(ManagedActors->GetConstGeneratedActors());
+				}
+#endif
+				AActor* MoveToActor = CreatedChildActor ? CreatedChildActor : ParentActor;
+				GeneratedResource->MoveResourceToNewActor(MoveToActor, Owner);
 				TSet<TSoftObjectPtr<AActor>> Dummy;
 				GeneratedResource->ReleaseIfUnused(Dummy);
 				bHasMovedResources = true;
@@ -1354,11 +1400,54 @@ bool UPCGComponent::MoveResourcesToNewActor(AActor* InNewActor, bool bCreateChil
 		GeneratedResources.Empty();
 	}
 
-	if (!bHasMovedResources && bCreateChild)
+	if (CreatedChildActor)
 	{
-		// There was no resource moved, delete the newly spawned actor.
-		GetWorld()->DestroyActor(NewActor);
-		return false;
+		// No resources moved destroy stamp actor
+		if (!bHasMovedResources)
+		{
+			GetWorld()->DestroyActor(CreatedChildActor);
+			return false;
+		}
+
+		USceneComponent* RootComponent = CreatedChildActor->GetRootComponent();
+		check(RootComponent);
+
+		TArray<AActor*> AttachedActors;
+		CreatedChildActor->GetAttachedActors(AttachedActors);
+
+		// Set Folder Path if actor is not going to be deleted or if we have moved actors
+		const bool bDestroyActor = RootComponent->GetNumChildrenComponents() == 0 && AttachedActors.IsEmpty();
+
+#if WITH_EDITOR
+		const bool bSetFolderPath = !MovedActors.IsEmpty() || !bDestroyActor;
+		if (bSetFolderPath)
+		{
+			CreatedChildActor->SetFolderPath(FolderPath);
+		}
+
+		// Move Generated Actors
+		if (!MovedActors.IsEmpty())
+		{
+			FString OutGeneratedFolderStr;
+			PCGHelpers::GetGeneratedActorsFolderPath(CreatedChildActor, OutGeneratedFolderStr);
+			FName GeneratedFolder(*OutGeneratedFolderStr);
+
+			for (const TSoftObjectPtr<AActor>& MovedActorPtr : MovedActors)
+			{
+				if (AActor* MovedActor = MovedActorPtr.Get())
+				{
+					MovedActor->SetFolderPath(GeneratedFolder);
+				}
+			}
+		}
+#endif
+
+		// No moved components destroy stamp actor
+		if (bDestroyActor)
+		{
+			GetWorld()->DestroyActor(CreatedChildActor);
+			return false;
+		}
 	}
 
 	return bHasMovedResources;
diff --git a/Engine/Plugins/PCG/Source/PCG/Private/PCGSubsystem.cpp b/Engine/Plugins/PCG/Source/PCG/Private/PCGSubsystem.cpp
--- a/Engine/Plugins/PCG/Source/PCG/Private/PCGSubsystem.cpp
+++ b/Engine/Plugins/PCG/Source/PCG/Private/PCGSubsystem.cpp
@@ -30,6 +30,7 @@
 #include "WorldPartition/WorldPartition.h"
 #include "WorldPartition/WorldPartitionHelpers.h"
 #include "WorldPartition/DataLayer/DataLayerManager.h"
+#include "WorldPartition/LoaderAdapter/LoaderAdapterShape.h"
 #else
 #include "Engine/Engine.h"
 #include "Engine/World.h"
@@ -130,6 +131,19 @@ namespace PCGSubsystemConsole
 
 #if WITH_EDITOR
 TSet<UWorld*> UPCGSubsystem::DisablePartitionActorCreationForWorld;
+
+struct FPCGPartitionActorLoaderAdapter : public FLoaderAdapterShape
+{
+public:
+	FPCGPartitionActorLoaderAdapter(UWorld* InWorld, const FBox& InBoundingBox, const FString& InLabel)
+		: FLoaderAdapterShape(InWorld, InBoundingBox, InLabel)
+	{ }
+
+	virtual bool PassActorDescFilter(const FWorldPartitionHandle& Actor) const override
+	{
+		return FLoaderAdapterShape::PassActorDescFilter(Actor) && Actor->GetActorNativeClass() && Actor->GetActorNativeClass()->IsChildOf<APCGPartitionActor>();
+	}
+};
 #endif
 
 UPCGSubsystem::UPCGSubsystem()
@@ -1027,7 +1041,7 @@ void UPCGSubsystem::ForAllOverlappingComponentsInHierarchy(UPCGComponent* InComp
 	});
 }
 
-FPCGTaskId UPCGSubsystem::ForAllOverlappingCells(UPCGComponent* InComponent, const FBox& InBounds, const PCGHiGenGrid::FSizeArray& InGridSizes, bool bCanCreateActor, const TArray<FPCGTaskId>& Dependencies, TFunctionRef<FPCGTaskId(APCGPartitionActor*, const FBox&)> InFunc) const
+FPCGTaskId UPCGSubsystem::ForAllOverlappingCells(UPCGComponent* InComponent, const FBox& InBounds, const PCGHiGenGrid::FSizeArray& InGridSizes, bool bCanCreateActor, const TArray<FPCGTaskId>& Dependencies, TFunctionRef<FPCGTaskId(APCGPartitionActor*, const FBox&)> InFunc, TFunctionRef<FPCGTaskId(const FPCGGridDescriptor&, const FIntVector&, const FBox&)> InUnloadedFunc) const
 {
 	if (!GraphExecutor || !PCGWorldActor)
 	{
@@ -1065,13 +1079,26 @@ FPCGTaskId UPCGSubsystem::ForAllOverlappingCells(UPCGComponent* InComponent, con
 		}
 
 		// Helper lambda to apply 'InFunc' to the specified grid cell. Only applies if the cell overlaps the modified bounds, and the actor can be found or created and is not unloaded.
-		auto ApplyOnCell = [this, &CellTasks, ModifiedInBounds, bCanCreateActor, &Descriptor, &InFunc](const FIntVector& CellCoord, const FBox& CellBounds)
+		auto ApplyOnCell = [this, &CellTasks, ModifiedInBounds, bCanCreateActor, &Descriptor, &InFunc, InUnloadedFunc](const FIntVector& CellCoord, const FBox& CellBounds)
 		{
 			const FBox IntersectedBounds = ModifiedInBounds.Overlap(CellBounds);
 
 			if (IntersectedBounds.IsValid)
 			{
-				if (APCGPartitionActor* Actor = FindOrCreatePCGPartitionActor(Descriptor, CellCoord, bCanCreateActor))
+				APCGPartitionActor* Actor = FindOrCreatePCGPartitionActor(Descriptor, CellCoord, bCanCreateActor);
+#if WITH_EDITOR
+				if (!Actor && !GetWorld()->IsGameWorld() && ActorAndComponentMapping.DoesPartitionActorRecordExist(Descriptor, CellCoord))
+				{
+					FPCGTaskId ExecuteTaskId = InUnloadedFunc(Descriptor, CellCoord, IntersectedBounds);
+					
+					if (ExecuteTaskId != InvalidPCGTaskId)
+					{
+						CellTasks.Add(ExecuteTaskId);
+					}
+				}
+				else 
+#endif
+				if(Actor)
 				{
 					FPCGTaskId ExecuteTaskId = InFunc(Actor, IntersectedBounds);
 
@@ -1108,8 +1135,9 @@ FPCGTaskId UPCGSubsystem::ForAllOverlappingCells(UPCGComponent* InComponent, con
 	// Create a dummy task to wait on dependencies, which creates a dummy task to wait on all cells.
 	if (!CellTasks.IsEmpty())
 	{
-		const FPCGTaskId DependenciesTaskId = GraphExecutor->ScheduleGeneric([]() { return true; }, /*InSourceComponent=*/nullptr, Dependencies);
-		return GraphExecutor->ScheduleGeneric([]() { return true; }, /*InSourceComponent=*/nullptr, { DependenciesTaskId });
+		TArray<FPCGTaskId> AllDependencies(Dependencies);
+		AllDependencies.Append(CellTasks);
+		return GraphExecutor->ScheduleGeneric([]() { return true; }, /*InSourceComponent=*/nullptr, AllDependencies);
 	}
 	else
 	{
@@ -1408,21 +1436,43 @@ void UPCGSubsystem::ClearPCGLink(UPCGComponent* InComponent, const FBox& InBound
 	TWeakObjectPtr<AActor> NewActorPtr(InNewActor);
 	TWeakObjectPtr<UPCGComponent> ComponentPtr(InComponent);
 
-	auto ScheduleTask = [this, NewActorPtr, ComponentPtr](APCGPartitionActor* PCGActor, const FBox& InIntersectedBounds)
+	auto MoveTask = [this, NewActorPtr, ComponentPtr](APCGPartitionActor* PCGActor, const FBox& InIntersectedBounds)
 	{
-		auto MoveTask = [this, NewActorPtr, ComponentPtr, PCGActor]()
+		check(NewActorPtr.IsValid() && ComponentPtr.IsValid() && PCGActor != nullptr);
+
+		if (TObjectPtr<UPCGComponent> LocalComponent = PCGActor->GetLocalComponent(ComponentPtr.Get()))
 		{
-			check(NewActorPtr.IsValid() && ComponentPtr.IsValid() && PCGActor != nullptr);
+			LocalComponent->MoveResourcesToNewActor(NewActorPtr.Get(), /*bCreateChild=*/true);
+		}
 
-			if (TObjectPtr<UPCGComponent> LocalComponent = PCGActor->GetLocalComponent(ComponentPtr.Get()))
+		return true;
+	};
+
+	auto ScheduleTask = [this, ComponentPtr, MoveTask](APCGPartitionActor* InPCGActor, const FBox& InIntersectedBounds)
+	{
+		auto MoveTaskInternal = [MoveTask, InPCGActor, InIntersectedBounds]()
+		{
+			return MoveTask(InPCGActor, InIntersectedBounds);
+		};
+
+		return GraphExecutor->ScheduleGeneric(MoveTaskInternal, ComponentPtr.Get(), /*TaskExecutionDependencies=*/{});
+	};
+
+	auto ScheduleUnloadedTask = [this, ComponentPtr, MoveTask](const FPCGGridDescriptor& InGridDescriptor, const FIntVector& InGridCoord, const FBox& InIntersectedBounds)
+	{
+		auto MoveTaskInternal = [this, MoveTask, InGridDescriptor, InGridCoord, InIntersectedBounds]()
+		{
+			TUniquePtr<IWorldPartitionActorLoaderInterface::ILoaderAdapter> LoaderAdapter = MakeUnique<FPCGPartitionActorLoaderAdapter>(GetWorld(), InIntersectedBounds, TEXT("UPCGSubsystem::ClearPCGLink"));
+			LoaderAdapter->Load();
+			if (APCGPartitionActor* PCGActor = FindOrCreatePCGPartitionActor(InGridDescriptor, InGridCoord, /*bCanCreateActor=*/false))
 			{
-				LocalComponent->MoveResourcesToNewActor(NewActorPtr.Get(), /*bCreateChild=*/true);
+				return MoveTask(PCGActor, InIntersectedBounds);
 			}
 
 			return true;
 		};
 
-		return GraphExecutor->ScheduleGeneric(MoveTask, ComponentPtr.Get(), /*TaskExecutionDependencies=*/{});
+		return GraphExecutor->ScheduleGeneric(MoveTaskInternal, ComponentPtr.Get(), /*TaskExecutionDependencies=*/{});
 	};
 
 	bool bHasUnbounded = false;
@@ -1432,7 +1482,8 @@ void UPCGSubsystem::ClearPCGLink(UPCGComponent* InComponent, const FBox& InBound
 	FPCGTaskId TaskId = InvalidPCGTaskId;
 	if (!GridSizes.IsEmpty())
 	{
-		TaskId = ForAllOverlappingCells(InComponent, InBounds, GridSizes, /*bCanCreateActor=*/false, /*Dependencies=*/{}, ScheduleTask);
+		// Load unloaded PAs to make sure Clear PCG Link is not partial
+		TaskId = ForAllOverlappingCells(InComponent, InBounds, GridSizes, /*bCanCreateActor=*/false, /*Dependencies=*/{}, ScheduleTask, ScheduleUnloadedTask);
 	}
 
 	// Verify if the NewActor has some components attached to its root or attached actors. If not, destroy it.
diff --git a/Engine/Plugins/PCG/Source/PCG/Public/PCGSubsystem.h b/Engine/Plugins/PCG/Source/PCG/Public/PCGSubsystem.h
--- a/Engine/Plugins/PCG/Source/PCG/Public/PCGSubsystem.h
+++ b/Engine/Plugins/PCG/Source/PCG/Public/PCGSubsystem.h
@@ -252,7 +252,14 @@ public:
 	UE_DEPRECATED(5.5, "Use version with UPCGComponent")
 	FPCGTaskId ForAllOverlappingCells(const FBox& InBounds, const PCGHiGenGrid::FSizeArray& InGridSizes, bool bCanCreateActor, const TArray<FPCGTaskId>& Dependencies, TFunctionRef<FPCGTaskId(APCGPartitionActor*, const FBox&)> InFunc) const { return InvalidPCGTaskId;  }
 
-	FPCGTaskId ForAllOverlappingCells(UPCGComponent* InPCGComponent, const FBox& InBounds, const PCGHiGenGrid::FSizeArray& InGridSizes, bool bCanCreateActor, const TArray<FPCGTaskId>& Dependencies, TFunctionRef<FPCGTaskId(APCGPartitionActor*, const FBox&)> InFunc) const;
+	FPCGTaskId ForAllOverlappingCells(
+		UPCGComponent* InPCGComponent,
+		const FBox& InBounds,
+		const PCGHiGenGrid::FSizeArray& InGridSizes,
+		bool bCanCreateActor,
+		const TArray<FPCGTaskId>& Dependencies,
+		TFunctionRef<FPCGTaskId(APCGPartitionActor*, const FBox&)> InFunc,
+		TFunctionRef<FPCGTaskId(const FPCGGridDescriptor&, const FIntVector&, const FBox&)> InUnloadedFunc = [](const FPCGGridDescriptor&, const FIntVector&, const FBox&) { return InvalidPCGTaskId; }) const;
 		
 	/** Immediately cleanup the local components associated with an original component. */
 	void CleanupLocalComponentsImmediate(UPCGComponent* InOriginalComponent, bool bRemoveComponents);
