[PCG] Fix Static Analysis

#rnx

diff --git a/Engine/Plugins/PCG/Source/PCG/Private/PCGComponent.cpp b/Engine/Plugins/PCG/Source/PCG/Private/PCGComponent.cpp
--- a/Engine/Plugins/PCG/Source/PCG/Private/PCGComponent.cpp
+++ b/Engine/Plugins/PCG/Source/PCG/Private/PCGComponent.cpp
@@ -1170,7 +1170,7 @@ AActor* UPCGComponent::ClearPCGLink(UClass* TemplateActorClass)
 		{
 			Cleanup(true);
 		}
-		else if (NewActor)
+		else
 		{
 			NewActor->Destroy();
 			NewActor = nullptr;
