[PCG] Fix perf regression with some nodes that didn't initialize FPCGPoint before being converted to support PointArrayData
- SetNumPoints now has a bInitializeValues param (defaults to true)
- Call SetNumPoints(N, false) for uses cases that didn't init data before conversion
- Convert Debug element and Grid linkage element to support PointArrayData
- Avoid allocating when flattening a non allocated parent property
- FPCGAttributeAccessorKeysPointIndices keep ranges as members so that range Stub value isn't local (since we fill an array of pointers in GetMetadataEntryKeys)

#rnx

diff --git a/Engine/Plugins/PCG/Source/PCG/Private/Data/PCGBasePointData.cpp b/Engine/Plugins/PCG/Source/PCG/Private/Data/PCGBasePointData.cpp
--- a/Engine/Plugins/PCG/Source/PCG/Private/Data/PCGBasePointData.cpp
+++ b/Engine/Plugins/PCG/Source/PCG/Private/Data/PCGBasePointData.cpp
@@ -634,7 +634,7 @@ void UPCGBasePointData::SetPoints(const UPCGBasePointData* From, UPCGBasePointDa
 	check(From);
 
 	const int32 NumPoints = bCopyAll ? From->GetNumPoints() : InDataIndices.Num();
-	To->SetNumPoints(NumPoints);
+	To->SetNumPoints(NumPoints, /*bInitializeValues=*/false);
 	To->AllocateProperties(From->GetAllocatedProperties());
 
 	if (NumPoints == 0)
diff --git a/Engine/Plugins/PCG/Source/PCG/Private/Data/PCGPointArrayData.cpp b/Engine/Plugins/PCG/Source/PCG/Private/Data/PCGPointArrayData.cpp
--- a/Engine/Plugins/PCG/Source/PCG/Private/Data/PCGPointArrayData.cpp
+++ b/Engine/Plugins/PCG/Source/PCG/Private/Data/PCGPointArrayData.cpp
@@ -192,7 +192,7 @@ void UPCGPointArrayData::FlattenPropertiesIfNeeded(EPCGPointNativeProperties Pro
 	}
 }
 
-void UPCGPointArrayData::SetNumPoints(int32 InNumPoints)
+void UPCGPointArrayData::SetNumPoints(int32 InNumPoints, bool bInitializeValues)
 {
 	if (ParentData && ParentData->GetNumPoints() != InNumPoints)
 	{
@@ -201,7 +201,7 @@ void UPCGPointArrayData::SetNumPoints(int32 InNumPoints)
 
 	if (InNumPoints != PointArray.GetNumPoints())
 	{
-		PointArray.SetNumPoints(InNumPoints);
+		PointArray.SetNumPoints(InNumPoints, bInitializeValues);
 		DirtyCache();
 	}
 }
diff --git a/Engine/Plugins/PCG/Source/PCG/Private/Data/PCGPointData.cpp b/Engine/Plugins/PCG/Source/PCG/Private/Data/PCGPointData.cpp
--- a/Engine/Plugins/PCG/Source/PCG/Private/Data/PCGPointData.cpp
+++ b/Engine/Plugins/PCG/Source/PCG/Private/Data/PCGPointData.cpp
@@ -314,11 +314,18 @@ void UPCGPointData::SetPoints(const TArray<FPCGPoint>& InPoints)
 	GetMutablePoints() = InPoints;
 }
 
-void UPCGPointData::SetNumPoints(int32 InNumPoints)
+void UPCGPointData::SetNumPoints(int32 InNumPoints, bool bInitializeValues)
 { 
 	if (Points.Num() != InNumPoints)
 	{
-		Points.SetNum(InNumPoints);
+		if (bInitializeValues)
+		{
+			Points.SetNum(InNumPoints);
+		}
+		else
+		{
+			Points.SetNumUninitialized(InNumPoints);
+		}
 		DirtyCache();
 	}
 }
diff --git a/Engine/Plugins/PCG/Source/PCG/Private/Data/PCGProjectionData.cpp b/Engine/Plugins/PCG/Source/PCG/Private/Data/PCGProjectionData.cpp
--- a/Engine/Plugins/PCG/Source/PCG/Private/Data/PCGProjectionData.cpp
+++ b/Engine/Plugins/PCG/Source/PCG/Private/Data/PCGProjectionData.cpp
@@ -308,7 +308,7 @@ const UPCGBasePointData* UPCGProjectionData::CreateBasePointData(FPCGContext* Co
 
 	auto InitializeFunc = [PointData, SourcePointData]()
 	{
-		PointData->SetNumPoints(SourcePointData->GetNumPoints());
+		PointData->SetNumPoints(SourcePointData->GetNumPoints(), /*bInitializeValues=*/false);
 
 		// @todo_pcg: try to optimize allocations here
 		PointData->AllocateProperties(SourcePointData->GetAllocatedProperties() | EPCGPointNativeProperties::Transform | EPCGPointNativeProperties::Color | EPCGPointNativeProperties::Density);
diff --git a/Engine/Plugins/PCG/Source/PCG/Private/Data/PCGUnionData.cpp b/Engine/Plugins/PCG/Source/PCG/Private/Data/PCGUnionData.cpp
--- a/Engine/Plugins/PCG/Source/PCG/Private/Data/PCGUnionData.cpp
+++ b/Engine/Plugins/PCG/Source/PCG/Private/Data/PCGUnionData.cpp
@@ -413,7 +413,7 @@ void UPCGUnionData::CreateSequentialPointData(FPCGContext* Context, TArray<const
 		PropertiesToAllocate |= PointInputDatas[Index]->GetAllocatedProperties();
 	}
 	
-	OutPointData->SetNumPoints(PointInputCount);
+	OutPointData->SetNumPoints(PointInputCount, /*bInitializeValues=*/false);
 	OutPointData->AllocateProperties(PropertiesToAllocate);
 
 	struct FIndexParams
diff --git a/Engine/Plugins/PCG/Source/PCG/Private/Elements/PCGCopyPoints.cpp b/Engine/Plugins/PCG/Source/PCG/Private/Elements/PCGCopyPoints.cpp
--- a/Engine/Plugins/PCG/Source/PCG/Private/Elements/PCGCopyPoints.cpp
+++ b/Engine/Plugins/PCG/Source/PCG/Private/Elements/PCGCopyPoints.cpp
@@ -285,7 +285,7 @@ bool FPCGCopyPointsElement::ExecuteInternal(FPCGContext* Context) const
 
 		auto InitializeFunc = [OutPointData, SourcePointData, TargetPointData, NumPoints, bProcessMetadata, ColorInheritance, SeedInheritance, OptionalProperties]()
 		{
-			OutPointData->SetNumPoints(NumPoints);
+			OutPointData->SetNumPoints(NumPoints, /*bInitializeValues=*/false);
 			EPCGPointNativeProperties SourceAllocatedProperties = SourcePointData->GetAllocatedProperties();
 			EPCGPointNativeProperties TargetAllocatedProperties = TargetPointData->GetAllocatedProperties();
 						
diff --git a/Engine/Plugins/PCG/Source/PCG/Private/Elements/PCGCreatePoints.cpp b/Engine/Plugins/PCG/Source/PCG/Private/Elements/PCGCreatePoints.cpp
--- a/Engine/Plugins/PCG/Source/PCG/Private/Elements/PCGCreatePoints.cpp
+++ b/Engine/Plugins/PCG/Source/PCG/Private/Elements/PCGCreatePoints.cpp
@@ -106,7 +106,7 @@ bool FPCGCreatePointsElement::ExecuteInternal(FPCGContext* Context) const
 
 	UPCGBasePointData* PointData = FPCGContext::NewPointData_AnyThread(Context);
 	check(PointData); 
-	PointData->SetNumPoints(PointsToLoopOn.Num());
+	PointData->SetNumPoints(PointsToLoopOn.Num(), /*bInitializeValues=*/false);
 
 	// @todo_pcg: Could probably be optimzed by comparing the source points to find the properties that are default but this could also be expensive
 	PointData->AllocateProperties(EPCGPointNativeProperties::All);
diff --git a/Engine/Plugins/PCG/Source/PCG/Private/Elements/PCGDensityFilter.cpp b/Engine/Plugins/PCG/Source/PCG/Private/Elements/PCGDensityFilter.cpp
--- a/Engine/Plugins/PCG/Source/PCG/Private/Elements/PCGDensityFilter.cpp
+++ b/Engine/Plugins/PCG/Source/PCG/Private/Elements/PCGDensityFilter.cpp
@@ -86,7 +86,7 @@ bool FPCGDensityFilterElement::ExecuteInternal(FPCGContext* Context) const
 
 		auto InitializeFunc = [FilteredData, OriginalData]()
 		{
-			FilteredData->SetNumPoints(OriginalData->GetNumPoints());
+			FilteredData->SetNumPoints(OriginalData->GetNumPoints(), /*bInitializeValues=*/false);
 			FilteredData->AllocateProperties(OriginalData->GetAllocatedProperties());
 			OriginalData->CopyUnallocatedPropertiesTo(FilteredData);
 		};
diff --git a/Engine/Plugins/PCG/Source/PCG/Private/Elements/PCGDistance.cpp b/Engine/Plugins/PCG/Source/PCG/Private/Elements/PCGDistance.cpp
--- a/Engine/Plugins/PCG/Source/PCG/Private/Elements/PCGDistance.cpp
+++ b/Engine/Plugins/PCG/Source/PCG/Private/Elements/PCGDistance.cpp
@@ -194,7 +194,7 @@ bool FPCGDistanceElement::ExecuteInternal(FPCGContext* Context) const
 
 		UPCGBasePointData* OutputData = FPCGContext::NewPointData_AnyThread(Context);
 		OutputData->InitializeFromData(SourcePointData);
-		OutputData->SetNumPoints(SourcePointData->GetNumPoints());
+		OutputData->SetNumPoints(SourcePointData->GetNumPoints(), /*bInitializeValues=*/false);
 
 		if (!OutputData->HasSpatialDataParent())
 		{
diff --git a/Engine/Plugins/PCG/Source/PCG/Private/Elements/PCGDuplicatePoint.cpp b/Engine/Plugins/PCG/Source/PCG/Private/Elements/PCGDuplicatePoint.cpp
--- a/Engine/Plugins/PCG/Source/PCG/Private/Elements/PCGDuplicatePoint.cpp
+++ b/Engine/Plugins/PCG/Source/PCG/Private/Elements/PCGDuplicatePoint.cpp
@@ -69,7 +69,7 @@ bool FPCGDuplicatePointElement::ExecuteInternal(FPCGContext* Context) const
 
 		auto InitializeFunc = [NumIterations, OutPointData, InputPointData]()
 		{
-			OutPointData->SetNumPoints(NumIterations);
+			OutPointData->SetNumPoints(NumIterations, /*bInitializeValues=*/false);
 			OutPointData->AllocateProperties(InputPointData->GetAllocatedProperties() | EPCGPointNativeProperties::Transform | EPCGPointNativeProperties::Seed);
 		};
 
diff --git a/Engine/Plugins/PCG/Source/PCG/Private/Elements/PCGPointOperationElementBase.cpp b/Engine/Plugins/PCG/Source/PCG/Private/Elements/PCGPointOperationElementBase.cpp
--- a/Engine/Plugins/PCG/Source/PCG/Private/Elements/PCGPointOperationElementBase.cpp
+++ b/Engine/Plugins/PCG/Source/PCG/Private/Elements/PCGPointOperationElementBase.cpp
@@ -57,7 +57,7 @@ bool FPCGPointOperationElementBase::PreparePointOperationData(ContextType* InCon
 			// Create and initialize the output points
 			OutState.OutputData = SupportsBasePointDataInputs(InContext) ? FPCGContext::NewPointData_AnyThread(InContext) : FPCGContext::NewObject_AnyThread<UPCGPointData>(InContext);
 			OutState.OutputData->InitializeFromData(OutState.InputData);
-			OutState.OutputData->SetNumPoints(OutState.NumPoints);
+			OutState.OutputData->SetNumPoints(OutState.NumPoints, /*bInitializeValues=*/false);
 
 			// Allocate properties that we are going to modify
 			EPCGPointNativeProperties PropertiesToAllocate = GetPropertiesToAllocate(InContext);
diff --git a/Engine/Plugins/PCG/Source/PCG/Private/Elements/PCGSelectPoints.cpp b/Engine/Plugins/PCG/Source/PCG/Private/Elements/PCGSelectPoints.cpp
--- a/Engine/Plugins/PCG/Source/PCG/Private/Elements/PCGSelectPoints.cpp
+++ b/Engine/Plugins/PCG/Source/PCG/Private/Elements/PCGSelectPoints.cpp
@@ -111,7 +111,7 @@ bool FPCGSelectPointsElement::ExecuteInternal(FPCGContext* Context) const
 
 			auto InitializeFunc = [SampledData, OriginalData]()
 			{
-				SampledData->SetNumPoints(OriginalData->GetNumPoints());
+				SampledData->SetNumPoints(OriginalData->GetNumPoints(), /*bInitializeValues=*/false);
 				SampledData->AllocateProperties(OriginalData->GetAllocatedProperties());
 				OriginalData->CopyUnallocatedPropertiesTo(SampledData);
 			};
diff --git a/Engine/Plugins/PCG/Source/PCG/Private/Elements/PCGSurfaceSampler.cpp b/Engine/Plugins/PCG/Source/PCG/Private/Elements/PCGSurfaceSampler.cpp
--- a/Engine/Plugins/PCG/Source/PCG/Private/Elements/PCGSurfaceSampler.cpp
+++ b/Engine/Plugins/PCG/Source/PCG/Private/Elements/PCGSurfaceSampler.cpp
@@ -184,9 +184,8 @@ namespace PCGSurfaceSampler
 				
 		auto InitializeFunc = [&SamplerData, SampledData]()
 		{
-			SampledData->SetNumPoints(SamplerData.CellCount);
+			SampledData->SetNumPoints(SamplerData.CellCount, /*bInitializeValues=*/false);
 			SampledData->SetSteepness(SamplerData.Params.PointSteepness);
-			SampledData->SetDensity(1.0f);
 
 			FPCGPoint DefaultPoint;
 			PCGPointHelpers::SetExtents(SamplerData.Params.PointExtents, DefaultPoint.BoundsMin, DefaultPoint.BoundsMax);
@@ -207,8 +206,6 @@ namespace PCGSurfaceSampler
 		{
 			TPCGValueRange<FTransform> TransformRange = SampledData->GetTransformValueRange(/*bAllocate=*/false);
 			TPCGValueRange<float> DensityRange = SampledData->GetDensityValueRange(/*bAllocate=*/false);
-			TPCGValueRange<FVector> BoundsMinRange = SampledData->GetBoundsMinValueRange(/*bAllocate=*/false);
-			TPCGValueRange<FVector> BoundsMaxRange = SampledData->GetBoundsMaxValueRange(/*bAllocate=*/false);
 			TPCGValueRange<int32> SeedRange = SampledData->GetSeedValueRange(/*bAllocate=*/false);
 			TPCGValueRange<int64> MetadataEntryRange = SampledData->GetMetadataEntryValueRange(/*bAllocate=*/false);
 			TPCGValueRange<FVector4> ColorRange = SampledData->GetColorValueRange(/*bAllocate=*/false);
@@ -258,7 +255,10 @@ namespace PCGSurfaceSampler
 				{
 					continue;
 				}
-											
+					
+				// Set default density
+				DensityRange[WriteIndex] = 1.0f;
+
 				// Now run gauntlet of shape network (if there is one) to accept or reject the point.
 				if (InBoundingShape)
 				{
diff --git a/Engine/Plugins/PCG/Source/PCG/Private/Elements/PCGTransformPoints.cpp b/Engine/Plugins/PCG/Source/PCG/Private/Elements/PCGTransformPoints.cpp
--- a/Engine/Plugins/PCG/Source/PCG/Private/Elements/PCGTransformPoints.cpp
+++ b/Engine/Plugins/PCG/Source/PCG/Private/Elements/PCGTransformPoints.cpp
@@ -97,7 +97,7 @@ bool FPCGTransformPointsElement::ExecuteInternal(FPCGContext* Context) const
 		UPCGBasePointData* OutputData = FPCGContext::NewPointData_AnyThread(Context);
 		
 		OutputData->InitializeFromData(PointData);
-		OutputData->SetNumPoints(PointData->GetNumPoints());
+		OutputData->SetNumPoints(PointData->GetNumPoints(), /*bInitializeValues=*/false);
 
 		const EPCGPointNativeProperties PropertiesToAllocate = bApplyToAttribute ? EPCGPointNativeProperties::MetadataEntry : (bRecomputeSeed ? (EPCGPointNativeProperties::Transform | EPCGPointNativeProperties::Seed) : EPCGPointNativeProperties::Transform);
 
diff --git a/Engine/Plugins/PCG/Source/PCG/Private/Elements/PCGVolumeSampler.cpp b/Engine/Plugins/PCG/Source/PCG/Private/Elements/PCGVolumeSampler.cpp
--- a/Engine/Plugins/PCG/Source/PCG/Private/Elements/PCGVolumeSampler.cpp
+++ b/Engine/Plugins/PCG/Source/PCG/Private/Elements/PCGVolumeSampler.cpp
@@ -107,7 +107,7 @@ namespace PCGVolumeSampler
 
 		auto InitializeFunc = [NumIterations, OutputData]()
 		{
-			OutputData->SetNumPoints(NumIterations);
+			OutputData->SetNumPoints(NumIterations, /*bInitializeValues=*/false);
 			// Volume->SamplePoint could write to all properties
 			OutputData->AllocateProperties(EPCGPointNativeProperties::All);
 		};
diff --git a/Engine/Plugins/PCG/Source/PCG/Private/Graph/PCGGraphExecutor.h b/Engine/Plugins/PCG/Source/PCG/Private/Graph/PCGGraphExecutor.h
--- a/Engine/Plugins/PCG/Source/PCG/Private/Graph/PCGGraphExecutor.h
+++ b/Engine/Plugins/PCG/Source/PCG/Private/Graph/PCGGraphExecutor.h
@@ -472,6 +472,7 @@ namespace PCGGraphExecutor
 		//~Begin IPCGElement interface
 		virtual bool IsGridLinkage() const override { return true; }
 		virtual bool SupportsGPUResidentData(FPCGContext* InContext) const override;
+		virtual bool SupportsBasePointDataInputs(FPCGContext* InContext) const override { return true; }
 		//~End IPCGElement interface
 
 		const UPCGPin* GetUpstreamPin() const { return UpstreamPin.Get(); }
diff --git a/Engine/Plugins/PCG/Source/PCG/Private/PCGPointArray.cpp b/Engine/Plugins/PCG/Source/PCG/Private/PCGPointArray.cpp
--- a/Engine/Plugins/PCG/Source/PCG/Private/PCGPointArray.cpp
+++ b/Engine/Plugins/PCG/Source/PCG/Private/PCGPointArray.cpp
@@ -15,61 +15,62 @@ FPCGPointArray::FPCGPointArray()
 	MetadataEntry.SetValue(DefaultPoint.MetadataEntry);
 }
 
-void FPCGPointArray::SetNumPoints(int32 InNumPoints)
+void FPCGPointArray::SetNumPoints(int32 InNumPoints, bool bInitializeValues)
 {
 	NumPoints = InNumPoints;
-	
+	bInitializedValues = bInitializeValues;
+
 	// Reallocate arrays that are currently allocated
-	Transform.SetNum(NumPoints, /*bAllocate=*/false);
-	Density.SetNum(NumPoints, /*bAllocate=*/false);
-	BoundsMin.SetNum(NumPoints, /*bAllocate=*/false);
-	BoundsMax.SetNum(NumPoints, /*bAllocate=*/false);
-	Color.SetNum(NumPoints, /*bAllocate=*/false);
-	Steepness.SetNum(NumPoints, /*bAllocate=*/false);
-	Seed.SetNum(NumPoints, /*bAllocate=*/false);
-	MetadataEntry.SetNum(NumPoints, /*bAllocate=*/false);
+	Transform.SetNum(NumPoints, /*bAllocate=*/false, bInitializedValues);
+	Density.SetNum(NumPoints, /*bAllocate=*/false, bInitializedValues);
+	BoundsMin.SetNum(NumPoints, /*bAllocate=*/false, bInitializedValues);
+	BoundsMax.SetNum(NumPoints, /*bAllocate=*/false, bInitializedValues);
+	Color.SetNum(NumPoints, /*bAllocate=*/false, bInitializedValues);
+	Steepness.SetNum(NumPoints, /*bAllocate=*/false, bInitializedValues);
+	Seed.SetNum(NumPoints, /*bAllocate=*/false, bInitializedValues);
+	MetadataEntry.SetNum(NumPoints, /*bAllocate=*/false, bInitializedValues);
 }
 
 void FPCGPointArray::Allocate(EPCGPointNativeProperties InProperties)
 {
 	if (EnumHasAnyFlags(InProperties, EPCGPointNativeProperties::Transform))
 	{
-		Transform.Allocate();
+		Transform.Allocate(bInitializedValues);
 	}
 
 	if (EnumHasAnyFlags(InProperties, EPCGPointNativeProperties::Density))
 	{
-		Density.Allocate();
+		Density.Allocate(bInitializedValues);
 	}
 
 	if (EnumHasAnyFlags(InProperties, EPCGPointNativeProperties::BoundsMin))
 	{
-		BoundsMin.Allocate();
+		BoundsMin.Allocate(bInitializedValues);
 	}
 
 	if (EnumHasAnyFlags(InProperties, EPCGPointNativeProperties::BoundsMax))
 	{
-		BoundsMax.Allocate();
+		BoundsMax.Allocate(bInitializedValues);
 	}
 
 	if (EnumHasAnyFlags(InProperties, EPCGPointNativeProperties::Color))
 	{
-		Color.Allocate();
+		Color.Allocate(bInitializedValues);
 	}
 
 	if (EnumHasAnyFlags(InProperties, EPCGPointNativeProperties::Steepness))
 	{
-		Steepness.Allocate();
+		Steepness.Allocate(bInitializedValues);
 	}
 
 	if (EnumHasAnyFlags(InProperties, EPCGPointNativeProperties::Seed))
 	{
-		Seed.Allocate();
+		Seed.Allocate(bInitializedValues);
 	}
 
 	if (EnumHasAnyFlags(InProperties, EPCGPointNativeProperties::MetadataEntry))
 	{
-		MetadataEntry.Allocate();
+		MetadataEntry.Allocate(bInitializedValues);
 	}
 }
 
diff --git a/Engine/Plugins/PCG/Source/PCG/Public/Data/PCGBasePointData.h b/Engine/Plugins/PCG/Source/PCG/Public/Data/PCGBasePointData.h
--- a/Engine/Plugins/PCG/Source/PCG/Public/Data/PCGBasePointData.h
+++ b/Engine/Plugins/PCG/Source/PCG/Public/Data/PCGBasePointData.h
@@ -58,7 +58,7 @@ public:
 
 	virtual int32 GetNumPoints() const PURE_VIRTUAL(UPCGBasePointData::GetNumPoints, return 0;);
 	bool IsEmpty() const { return GetNumPoints() == 0; }
-	virtual void SetNumPoints(int32 InNumPoints) PURE_VIRTUAL(UPCGBasePointData::SetNumPoints, return;);
+	virtual void SetNumPoints(int32 InNumPoints, bool bInitializeValues = true) PURE_VIRTUAL(UPCGBasePointData::SetNumPoints, return;);
 
 	virtual void AllocateProperties(EPCGPointNativeProperties Properties) {}
 	virtual void FreeProperties(EPCGPointNativeProperties Properties) {}
diff --git a/Engine/Plugins/PCG/Source/PCG/Public/Data/PCGPointArrayData.h b/Engine/Plugins/PCG/Source/PCG/Public/Data/PCGPointArrayData.h
--- a/Engine/Plugins/PCG/Source/PCG/Public/Data/PCGPointArrayData.h
+++ b/Engine/Plugins/PCG/Source/PCG/Public/Data/PCGPointArrayData.h
@@ -33,7 +33,7 @@ public:
 	virtual bool IsValidRef(const PCGPointOctree::FPointRef& InPointRef) const override { return InPointRef.Index >= 0 && InPointRef.Index < GetNumPoints(); }
 	
 	virtual int32 GetNumPoints() const override { return PointArray.GetNumPoints(); }
-	virtual void SetNumPoints(int32 InNumPoints) override;
+	virtual void SetNumPoints(int32 InNumPoints, bool bInitializeValues = true) override;
 		
 	virtual void AllocateProperties(EPCGPointNativeProperties Properties) override;
 	virtual void FreeProperties(EPCGPointNativeProperties Properties) override;
@@ -215,7 +215,11 @@ private:
 		// Should not be equal since InheritedProperties contains that property
 		if (ensure(InheritedProperty != Property))
 		{
-			Property->Allocate();
+			if (InheritedProperty->IsAllocated())
+			{
+				Property->Allocate(/*bInitializeValues=*/false);
+			}
+
 			check(InheritedProperty->Num() == Property->Num());
 			InheritedProperty->CopyTo(*Property, 0, 0, Property->Num());
 		}
diff --git a/Engine/Plugins/PCG/Source/PCG/Public/Data/PCGPointData.h b/Engine/Plugins/PCG/Source/PCG/Public/Data/PCGPointData.h
--- a/Engine/Plugins/PCG/Source/PCG/Public/Data/PCGPointData.h
+++ b/Engine/Plugins/PCG/Source/PCG/Public/Data/PCGPointData.h
@@ -130,7 +130,7 @@ public:
 	UFUNCTION(BlueprintCallable, Category = SpatialData)
 	virtual int32 GetNumPoints() const override { return Points.Num(); }
 
-	virtual void SetNumPoints(int32 InNumPoints) override;
+	virtual void SetNumPoints(int32 InNumPoints, bool bInitializeValues = true) override;
 
 	UFUNCTION(BlueprintCallable, Category = SpatialData)
 	bool IsEmpty() const { return Points.IsEmpty(); }
diff --git a/Engine/Plugins/PCG/Source/PCG/Public/Data/PCGSpatialDataTpl.h b/Engine/Plugins/PCG/Source/PCG/Public/Data/PCGSpatialDataTpl.h
--- a/Engine/Plugins/PCG/Source/PCG/Public/Data/PCGSpatialDataTpl.h
+++ b/Engine/Plugins/PCG/Source/PCG/Public/Data/PCGSpatialDataTpl.h
@@ -72,8 +72,9 @@ namespace FPCGSpatialDataProcessing
 
 		auto Initialize = [TargetData, SourceData, NumIterations, PropertiesToAllocate]()
 		{
-			TargetData->SetNumPoints(NumIterations);
+			TargetData->SetNumPoints(NumIterations, /*bInitializeValues=*/false);
 			TargetData->AllocateProperties(PropertiesToAllocate);
+			SourceData->CopyUnallocatedPropertiesTo(TargetData);
 		};
 
 		auto ProcessRange = [RangeFunc = MoveTemp(InProcessRange), SourceData, TargetData](int32 StartReadIndex, int32 StartWriteIndex, int32 Count) -> int32
diff --git a/Engine/Plugins/PCG/Source/PCG/Public/Elements/PCGDebugElement.h b/Engine/Plugins/PCG/Source/PCG/Public/Elements/PCGDebugElement.h
--- a/Engine/Plugins/PCG/Source/PCG/Public/Elements/PCGDebugElement.h
+++ b/Engine/Plugins/PCG/Source/PCG/Public/Elements/PCGDebugElement.h
@@ -42,4 +42,5 @@ public:
 
 protected:
 	virtual bool ExecuteInternal(FPCGContext* Context) const override;
+	virtual bool SupportsBasePointDataInputs(FPCGContext* InContext) const override { return true; }
 };
diff --git a/Engine/Plugins/PCG/Source/PCG/Public/Metadata/Accessors/PCGCustomAccessor.h b/Engine/Plugins/PCG/Source/PCG/Public/Metadata/Accessors/PCGCustomAccessor.h
--- a/Engine/Plugins/PCG/Source/PCG/Public/Metadata/Accessors/PCGCustomAccessor.h
+++ b/Engine/Plugins/PCG/Source/PCG/Public/Metadata/Accessors/PCGCustomAccessor.h
@@ -223,12 +223,15 @@ public:
 		: IPCGAttributeAccessorKeys(/*bInReadOnly=*/false)
 		, PointData(InPointData)
 	{
+		MetadataEntryKeys = PointData->GetMetadataEntryValueRange(/*bAllocate=*/false);
+		ConstMetadataEntryKeys = PointData->GetConstMetadataEntryValueRange();
 	}
 
 	FPCGAttributeAccessorKeysPointIndices(const UPCGBasePointData* InPointData)
 		: IPCGAttributeAccessorKeys(/*bInReadOnly=*/true)
 		, PointData(const_cast<UPCGBasePointData*>(InPointData))
 	{
+		ConstMetadataEntryKeys = PointData->GetConstMetadataEntryValueRange();
 	}
 
 	virtual bool GetKeyIndices(int32 InStart, int32 InCount, TArray<int32>& OutKeyIndices, bool& bOutContiguous) const override
@@ -267,7 +270,7 @@ protected:
 			return false;
 		}
 
-		TPCGValueRange<int64> MetadataEntryKeys = PointData->GetMetadataEntryValueRange();
+		check(!MetadataEntryKeys.IsEmpty());
 
 		for (int32 Index = 0; Index < OutEntryKeys.Num(); ++Index)
 		{
@@ -285,12 +288,12 @@ protected:
 			return false;
 		}
 
-		const TConstPCGValueRange<int64> MetadataEntryKeys = PointData->GetConstMetadataEntryValueRange();
+		check(!ConstMetadataEntryKeys.IsEmpty());
 
 		for (int32 Index = 0; Index < OutEntryKeys.Num(); ++Index)
 		{
-			int32 EntryIndex = (InStart + Index) % MetadataEntryKeys.Num();
-			OutEntryKeys[Index] = &MetadataEntryKeys[EntryIndex];
+			int32 EntryIndex = (InStart + Index) % ConstMetadataEntryKeys.Num();
+			OutEntryKeys[Index] = &ConstMetadataEntryKeys[EntryIndex];
 		}
 
 		return true;
@@ -298,6 +301,8 @@ protected:
 
 private:
 	UPCGBasePointData* PointData = nullptr;
+	TPCGValueRange<int64> MetadataEntryKeys;
+	TConstPCGValueRange<int64> ConstMetadataEntryKeys;
 };
 
 namespace PCGCustomAccessor
diff --git a/Engine/Plugins/PCG/Source/PCG/Public/PCGPointArray.h b/Engine/Plugins/PCG/Source/PCG/Public/PCGPointArray.h
--- a/Engine/Plugins/PCG/Source/PCG/Public/PCGPointArray.h
+++ b/Engine/Plugins/PCG/Source/PCG/Public/PCGPointArray.h
@@ -26,6 +26,11 @@ template<typename T>
 struct FPCGPointArrayProperty
 {
 public:
+	FPCGPointArrayProperty()
+	{
+		static_assert(std::is_trivially_copyable_v<T>, "FPCGPointArrayProperty type must be trivially copyable");
+	}
+
 	/** Defined for serialization purposes(FPCGPointArray trait WithIdenticalViaEquality) */
 	[[nodiscard]] FORCEINLINE bool operator==(const FPCGPointArrayProperty<T>& Other) const
 	{
@@ -63,13 +68,13 @@ public:
 	}
 	
 	/** If 'Values' isn't currently allocated, allocate it and copy 'Value' at every index of 'Values[0, NumValues - 1]' */
-	FORCEINLINE void Allocate()
+	FORCEINLINE void Allocate(bool bInitializeValues)
 	{
-		SetNum(NumValues, /*bAllocate=*/true);
+		SetNum(NumValues, /*bAllocate=*/true, bInitializeValues);
 	}
 
 	/** Set the 'NumValues', optionally allocating 'Values' */
-	FORCEINLINE void SetNum(int32 NewNum, bool bAllocate)
+	FORCEINLINE void SetNum(int32 NewNum, bool bAllocate, bool bInitializeValues)
 	{
 		NumValues = NewNum;
 
@@ -81,10 +86,13 @@ public:
 
 		if (OldNum > 0 || bAllocate)
 		{
-			Values.SetNum(NewNum);
-			for (int32 i = OldNum; i < NewNum; ++i)
+			Values.SetNumUninitialized(NewNum);
+			if (bInitializeValues)
 			{
-				Values[i] = Value;
+				for (int32 i = OldNum; i < NewNum; ++i)
+				{
+					Values[i] = Value;
+				}
 			}
 		}
 	};
@@ -158,12 +166,12 @@ public:
 	}
 
 	/** Returns a TPCGValueRange<T> of NumValues and optionally allocate 'Values' */
-	[[nodiscard]] FORCEINLINE TPCGValueRange<T> GetValueRange(bool bAllocate)
+	[[nodiscard]] FORCEINLINE TPCGValueRange<T> GetValueRange(bool bAllocate, bool bInitializeValues)
 	{
 		if (bAllocate)
 		{
 			// Allocate when getting non-const version
-			SetNum(NumValues, /*bForce=*/true);
+			SetNum(NumValues, /*bForce=*/true, bInitializeValues);
 		}
 
 		return TPCGValueRange<T>(IsAllocated() ? MakeStridedView(Values) : MakeStridedView(sizeof(Value), &Value, 1), NumValues);
@@ -237,7 +245,7 @@ public:
 
 	friend bool operator==(const FPCGPointArray&, const FPCGPointArray&) = default;
 
-	void SetNumPoints(int32 InNumPoints);
+	void SetNumPoints(int32 InNumPoints, bool bInitializeValues);
 	int32 GetNumPoints() const { return NumPoints; }
 		
 	EPCGPointNativeProperties GetAllocatedProperties() const;
@@ -259,14 +267,14 @@ public:
 
 	TArray<FTransform> GetTransformCopy() const;
 
-	TPCGValueRange<FTransform> GetTransformValueRange(bool bAllocate) { return Transform.GetValueRange(bAllocate); }
-	TPCGValueRange<float> GetDensityValueRange(bool bAllocate) { return Density.GetValueRange(bAllocate); }
-	TPCGValueRange<FVector> GetBoundsMinValueRange(bool bAllocate) { return BoundsMin.GetValueRange(bAllocate); }
-	TPCGValueRange<FVector> GetBoundsMaxValueRange(bool bAllocate) { return BoundsMax.GetValueRange(bAllocate); }
-	TPCGValueRange<FVector4> GetColorValueRange(bool bAllocate) { return Color.GetValueRange(bAllocate); }
-	TPCGValueRange<float> GetSteepnessValueRange(bool bAllocate) { return Steepness.GetValueRange(bAllocate); }
-	TPCGValueRange<int32> GetSeedValueRange(bool bAllocate) { return Seed.GetValueRange(bAllocate); }
-	TPCGValueRange<int64> GetMetadataEntryValueRange(bool bAllocate) { return MetadataEntry.GetValueRange(bAllocate); }
+	TPCGValueRange<FTransform> GetTransformValueRange(bool bAllocate) { return Transform.GetValueRange(bAllocate, bInitializedValues); }
+	TPCGValueRange<float> GetDensityValueRange(bool bAllocate) { return Density.GetValueRange(bAllocate, bInitializedValues); }
+	TPCGValueRange<FVector> GetBoundsMinValueRange(bool bAllocate) { return BoundsMin.GetValueRange(bAllocate, bInitializedValues); }
+	TPCGValueRange<FVector> GetBoundsMaxValueRange(bool bAllocate) { return BoundsMax.GetValueRange(bAllocate, bInitializedValues); }
+	TPCGValueRange<FVector4> GetColorValueRange(bool bAllocate) { return Color.GetValueRange(bAllocate, bInitializedValues); }
+	TPCGValueRange<float> GetSteepnessValueRange(bool bAllocate) { return Steepness.GetValueRange(bAllocate, bInitializedValues); }
+	TPCGValueRange<int32> GetSeedValueRange(bool bAllocate) { return Seed.GetValueRange(bAllocate, bInitializedValues); }
+	TPCGValueRange<int64> GetMetadataEntryValueRange(bool bAllocate) { return MetadataEntry.GetValueRange(bAllocate, bInitializedValues); }
 
 	TConstPCGValueRange<FTransform> GetConstTransformValueRange() const { return Transform.GetConstValueRange(); }
 	TConstPCGValueRange<float> GetConstDensityValueRange() const { return Density.GetConstValueRange(); }
@@ -353,6 +361,9 @@ private:
 	FPCGPointArrayProperty<int64>		MetadataEntry;
 
 	int32 NumPoints = 0;
+
+	// Last call to SetNumPoints param value
+	bool bInitializedValues = false;
 };
 
 template<>
