[PCG] Add ECastCheckedType::NullAllowed to CastChecked calls

#rnx

diff --git a/Engine/Plugins/PCG/Source/PCG/Private/Data/PCGCollisionShapeData.cpp b/Engine/Plugins/PCG/Source/PCG/Private/Data/PCGCollisionShapeData.cpp
--- a/Engine/Plugins/PCG/Source/PCG/Private/Data/PCGCollisionShapeData.cpp
+++ b/Engine/Plugins/PCG/Source/PCG/Private/Data/PCGCollisionShapeData.cpp
@@ -97,14 +97,14 @@ const UPCGPointData* UPCGCollisionShapeData::CreatePointData(FPCGContext* Contex
 {
 	TRACE_CPUPROFILER_EVENT_SCOPE(UPCGCollisionShapeData::CreatePointData);
 
-	return Cast<UPCGPointData>(CreateBasePointData(Context, UPCGPointData::StaticClass()));
+	return CastChecked<UPCGPointData>(CreateBasePointData(Context, UPCGPointData::StaticClass()), ECastCheckedType::NullAllowed);
 }
 
 const UPCGPointArrayData* UPCGCollisionShapeData::CreatePointArrayData(FPCGContext* Context, const FBox& InBounds) const
 {
 	TRACE_CPUPROFILER_EVENT_SCOPE(UPCGCollisionShapeData::CreatePointArrayData);
 
-	return Cast<UPCGPointArrayData>(CreateBasePointData(Context, UPCGPointArrayData::StaticClass()));
+	return CastChecked<UPCGPointArrayData>(CreateBasePointData(Context, UPCGPointArrayData::StaticClass()), ECastCheckedType::NullAllowed);
 }
 
 const UPCGBasePointData* UPCGCollisionShapeData::CreateBasePointData(FPCGContext* Context, TSubclassOf<UPCGBasePointData> PointDataClass) const
diff --git a/Engine/Plugins/PCG/Source/PCG/Private/Data/PCGDifferenceData.cpp b/Engine/Plugins/PCG/Source/PCG/Private/Data/PCGDifferenceData.cpp
--- a/Engine/Plugins/PCG/Source/PCG/Private/Data/PCGDifferenceData.cpp
+++ b/Engine/Plugins/PCG/Source/PCG/Private/Data/PCGDifferenceData.cpp
@@ -233,7 +233,7 @@ const UPCGPointData* UPCGDifferenceData::CreatePointData(FPCGContext* Context) c
 	
 	const UPCGBasePointData* SourcePointData = GetSource()->ToPointData(Context);
 
-	return CastChecked<UPCGPointData>(CreateBasePointData(Context, SourcePointData, UPCGPointData::StaticClass()));
+	return CastChecked<UPCGPointData>(CreateBasePointData(Context, SourcePointData, UPCGPointData::StaticClass()), ECastCheckedType::NullAllowed);
 }
 
 const UPCGPointArrayData* UPCGDifferenceData::CreatePointArrayData(FPCGContext* Context, const FBox& InBounds) const
@@ -242,7 +242,7 @@ const UPCGPointArrayData* UPCGDifferenceData::CreatePointArrayData(FPCGContext*
 
 	const UPCGBasePointData* SourcePointData = GetSource()->ToPointArrayData(Context);
 
-	return CastChecked<UPCGPointArrayData>(CreateBasePointData(Context, SourcePointData, UPCGPointArrayData::StaticClass()));
+	return CastChecked<UPCGPointArrayData>(CreateBasePointData(Context, SourcePointData, UPCGPointArrayData::StaticClass()), ECastCheckedType::NullAllowed);
 }
 
 const UPCGBasePointData* UPCGDifferenceData::CreateBasePointData(FPCGContext* Context, const UPCGBasePointData* SourcePointData, TSubclassOf<UPCGBasePointData> PointDataClass) const
diff --git a/Engine/Plugins/PCG/Source/PCG/Private/Data/PCGIntersectionData.cpp b/Engine/Plugins/PCG/Source/PCG/Private/Data/PCGIntersectionData.cpp
--- a/Engine/Plugins/PCG/Source/PCG/Private/Data/PCGIntersectionData.cpp
+++ b/Engine/Plugins/PCG/Source/PCG/Private/Data/PCGIntersectionData.cpp
@@ -186,13 +186,13 @@ const UPCGSpatialData* UPCGIntersectionData::FindFirstConcreteShapeFromNetwork()
 const UPCGPointData* UPCGIntersectionData::CreatePointData(FPCGContext* Context) const
 {
 	TRACE_CPUPROFILER_EVENT_SCOPE(UPCGIntersectionData::CreatePointData);
-	return CastChecked<UPCGPointData>(CreateBasePointData(Context, UPCGPointData::StaticClass()));
+	return CastChecked<UPCGPointData>(CreateBasePointData(Context, UPCGPointData::StaticClass()), ECastCheckedType::NullAllowed);
 }
 
 const UPCGPointArrayData* UPCGIntersectionData::CreatePointArrayData(FPCGContext* Context, const FBox& InBounds) const
 {
 	TRACE_CPUPROFILER_EVENT_SCOPE(UPCGIntersectionData::CreatePointArrayData);
-	return CastChecked<UPCGPointArrayData>(CreateBasePointData(Context, UPCGPointArrayData::StaticClass()));
+	return CastChecked<UPCGPointArrayData>(CreateBasePointData(Context, UPCGPointArrayData::StaticClass()), ECastCheckedType::NullAllowed);
 }
 
 UPCGBasePointData* UPCGIntersectionData::CreateBasePointData(FPCGContext* Context, TSubclassOf<UPCGBasePointData> PointDataClass) const
diff --git a/Engine/Plugins/PCG/Source/PCG/Private/Data/PCGLandscapeData.cpp b/Engine/Plugins/PCG/Source/PCG/Private/Data/PCGLandscapeData.cpp
--- a/Engine/Plugins/PCG/Source/PCG/Private/Data/PCGLandscapeData.cpp
+++ b/Engine/Plugins/PCG/Source/PCG/Private/Data/PCGLandscapeData.cpp
@@ -472,14 +472,14 @@ const UPCGPointData* UPCGLandscapeData::CreatePointData(FPCGContext* Context, co
 {
 	TRACE_CPUPROFILER_EVENT_SCOPE(UPCGLandscapeData::CreatePointData);
 
-	return Cast<UPCGPointData>(CreateBasePointData(Context, InBounds, UPCGPointData::StaticClass()));
+	return CastChecked<UPCGPointData>(CreateBasePointData(Context, InBounds, UPCGPointData::StaticClass()), ECastCheckedType::NullAllowed);
 }
 
 const UPCGPointArrayData* UPCGLandscapeData::CreatePointArrayData(FPCGContext* Context, const FBox& InBounds) const
 {
 	TRACE_CPUPROFILER_EVENT_SCOPE(UPCGLandscapeData::CreatePointArrayData);
 
-	return Cast<UPCGPointArrayData>(CreateBasePointData(Context, InBounds, UPCGPointArrayData::StaticClass()));
+	return CastChecked<UPCGPointArrayData>(CreateBasePointData(Context, InBounds, UPCGPointArrayData::StaticClass()), ECastCheckedType::NullAllowed);
 }
 
 const UPCGBasePointData* UPCGLandscapeData::CreateBasePointData(FPCGContext* Context, const FBox& InBounds, TSubclassOf<UPCGBasePointData> PointDataClass) const
diff --git a/Engine/Plugins/PCG/Source/PCG/Private/Data/PCGLandscapeSplineData.cpp b/Engine/Plugins/PCG/Source/PCG/Private/Data/PCGLandscapeSplineData.cpp
--- a/Engine/Plugins/PCG/Source/PCG/Private/Data/PCGLandscapeSplineData.cpp
+++ b/Engine/Plugins/PCG/Source/PCG/Private/Data/PCGLandscapeSplineData.cpp
@@ -265,14 +265,14 @@ const UPCGPointData* UPCGLandscapeSplineData::CreatePointData(FPCGContext* Conte
 {
 	TRACE_CPUPROFILER_EVENT_SCOPE(UPCGLandscapeSplineData::CreatePointData);
 
-	return CastChecked<UPCGPointData>(CreateBasePointData(Context, UPCGPointData::StaticClass()));
+	return CastChecked<UPCGPointData>(CreateBasePointData(Context, UPCGPointData::StaticClass()), ECastCheckedType::NullAllowed);
 }
 
 const UPCGPointArrayData* UPCGLandscapeSplineData::CreatePointArrayData(FPCGContext* Context, const FBox& InBounds) const
 {
 	TRACE_CPUPROFILER_EVENT_SCOPE(UPCGLandscapeSplineData::CreatePointArrayData);
 
-	return CastChecked<UPCGPointArrayData>(CreateBasePointData(Context, UPCGPointArrayData::StaticClass()));
+	return CastChecked<UPCGPointArrayData>(CreateBasePointData(Context, UPCGPointArrayData::StaticClass()), ECastCheckedType::NullAllowed);
 }
 
 const UPCGBasePointData* UPCGLandscapeSplineData::CreateBasePointData(FPCGContext* Context, TSubclassOf<UPCGBasePointData> PointDataClass) const
diff --git a/Engine/Plugins/PCG/Source/PCG/Private/Data/PCGPrimitiveData.cpp b/Engine/Plugins/PCG/Source/PCG/Private/Data/PCGPrimitiveData.cpp
--- a/Engine/Plugins/PCG/Source/PCG/Private/Data/PCGPrimitiveData.cpp
+++ b/Engine/Plugins/PCG/Source/PCG/Private/Data/PCGPrimitiveData.cpp
@@ -53,14 +53,14 @@ const UPCGPointData* UPCGPrimitiveData::CreatePointData(FPCGContext* Context) co
 {
 	TRACE_CPUPROFILER_EVENT_SCOPE(UPCGPrimitiveData::CreatePointData);
 
-	return Cast<UPCGPointData>(CreateBasePointData(Context, UPCGPointData::StaticClass()));
+	return CastChecked<UPCGPointData>(CreateBasePointData(Context, UPCGPointData::StaticClass()), ECastCheckedType::NullAllowed);
 }
 
 const UPCGPointArrayData* UPCGPrimitiveData::CreatePointArrayData(FPCGContext* Context, const FBox& InBounds) const
 {
 	TRACE_CPUPROFILER_EVENT_SCOPE(UPCGPrimitiveData::CreatePointArrayData);
 
-	return Cast<UPCGPointArrayData>(CreateBasePointData(Context, UPCGPointArrayData::StaticClass()));
+	return CastChecked<UPCGPointArrayData>(CreateBasePointData(Context, UPCGPointArrayData::StaticClass()), ECastCheckedType::NullAllowed);
 }
 
 const UPCGBasePointData* UPCGPrimitiveData::CreateBasePointData(FPCGContext* Context, TSubclassOf<UPCGBasePointData> PointDataClass) const
diff --git a/Engine/Plugins/PCG/Source/PCG/Private/Data/PCGProjectionData.cpp b/Engine/Plugins/PCG/Source/PCG/Private/Data/PCGProjectionData.cpp
--- a/Engine/Plugins/PCG/Source/PCG/Private/Data/PCGProjectionData.cpp
+++ b/Engine/Plugins/PCG/Source/PCG/Private/Data/PCGProjectionData.cpp
@@ -255,14 +255,14 @@ const UPCGPointData* UPCGProjectionData::CreatePointData(FPCGContext* Context) c
 {
 	TRACE_CPUPROFILER_EVENT_SCOPE(UPCGProjectionData::CreatePointData);
 
-	return CastChecked<UPCGPointData>(CreateBasePointData(Context, UPCGPointData::StaticClass()));
+	return CastChecked<UPCGPointData>(CreateBasePointData(Context, UPCGPointData::StaticClass()), ECastCheckedType::NullAllowed);
 }
 
 const UPCGPointArrayData* UPCGProjectionData::CreatePointArrayData(FPCGContext* Context, const FBox& InBounds) const
 {
 	TRACE_CPUPROFILER_EVENT_SCOPE(UPCGProjectionData::CreatePointArrayData);
 
-	return CastChecked<UPCGPointArrayData>(CreateBasePointData(Context, UPCGPointArrayData::StaticClass()));
+	return CastChecked<UPCGPointArrayData>(CreateBasePointData(Context, UPCGPointArrayData::StaticClass()), ECastCheckedType::NullAllowed);
 }
 
 const UPCGBasePointData* UPCGProjectionData::CreateBasePointData(FPCGContext* Context, TSubclassOf<UPCGBasePointData> PointDataClass) const
diff --git a/Engine/Plugins/PCG/Source/PCG/Private/Data/PCGSpatialData.cpp b/Engine/Plugins/PCG/Source/PCG/Private/Data/PCGSpatialData.cpp
--- a/Engine/Plugins/PCG/Source/PCG/Private/Data/PCGSpatialData.cpp
+++ b/Engine/Plugins/PCG/Source/PCG/Private/Data/PCGSpatialData.cpp
@@ -121,12 +121,12 @@ const UPCGBasePointData* UPCGSpatialDataWithPointCache::ToBasePointDataInternal(
 
 const UPCGPointData* UPCGSpatialDataWithPointCache::ToPointData(FPCGContext* Context, const FBox& InBounds) const
 {
-	return CastChecked<UPCGPointData>(ToBasePointDataInternal(Context, InBounds, [this](FPCGContext* InContext, const FBox& InBounds) { return CreatePointData(InContext, InBounds); }));
+	return CastChecked<UPCGPointData>(ToBasePointDataInternal(Context, InBounds, [this](FPCGContext* InContext, const FBox& InBounds) { return CreatePointData(InContext, InBounds); }), ECastCheckedType::NullAllowed);
 }
 
 const UPCGPointArrayData* UPCGSpatialDataWithPointCache::ToPointArrayData(FPCGContext* Context, const FBox& InBounds) const
 {
-	return CastChecked<UPCGPointArrayData>(ToBasePointDataInternal(Context, InBounds, [this](FPCGContext* InContext, const FBox& InBounds) { return CreatePointArrayData(InContext, InBounds); }));
+	return CastChecked<UPCGPointArrayData>(ToBasePointDataInternal(Context, InBounds, [this](FPCGContext* InContext, const FBox& InBounds) { return CreatePointArrayData(InContext, InBounds); }), ECastCheckedType::NullAllowed);
 }
 
 const UPCGPointArrayData* UPCGSpatialDataWithPointCache::CreatePointArrayData(FPCGContext* Context, const FBox& InBounds) const
diff --git a/Engine/Plugins/PCG/Source/PCG/Private/Data/PCGSplineData.cpp b/Engine/Plugins/PCG/Source/PCG/Private/Data/PCGSplineData.cpp
--- a/Engine/Plugins/PCG/Source/PCG/Private/Data/PCGSplineData.cpp
+++ b/Engine/Plugins/PCG/Source/PCG/Private/Data/PCGSplineData.cpp
@@ -333,14 +333,14 @@ const UPCGPointData* UPCGSplineData::CreatePointData(FPCGContext* Context) const
 {
 	TRACE_CPUPROFILER_EVENT_SCOPE(UPCGSplineData::CreatePointData);
 
-	return CastChecked<UPCGPointData>(CreateBasePointData(Context, UPCGPointData::StaticClass()));
+	return CastChecked<UPCGPointData>(CreateBasePointData(Context, UPCGPointData::StaticClass()), ECastCheckedType::NullAllowed);
 }
 
 const UPCGPointArrayData* UPCGSplineData::CreatePointArrayData(FPCGContext* Context, const FBox& InBounds) const
 {
 	TRACE_CPUPROFILER_EVENT_SCOPE(UPCGSplineData::CreatePointArrayData);
 
-	return CastChecked<UPCGPointArrayData>(CreateBasePointData(Context, UPCGPointArrayData::StaticClass()));
+	return CastChecked<UPCGPointArrayData>(CreateBasePointData(Context, UPCGPointArrayData::StaticClass()), ECastCheckedType::NullAllowed);
 }
 
 const UPCGBasePointData* UPCGSplineData::CreateBasePointData(FPCGContext* Context, TSubclassOf<UPCGBasePointData> PointDataClass) const
diff --git a/Engine/Plugins/PCG/Source/PCG/Private/Data/PCGUnionData.cpp b/Engine/Plugins/PCG/Source/PCG/Private/Data/PCGUnionData.cpp
--- a/Engine/Plugins/PCG/Source/PCG/Private/Data/PCGUnionData.cpp
+++ b/Engine/Plugins/PCG/Source/PCG/Private/Data/PCGUnionData.cpp
@@ -237,14 +237,14 @@ const UPCGPointData* UPCGUnionData::CreatePointData(FPCGContext* Context) const
 {
 	TRACE_CPUPROFILER_EVENT_SCOPE(UPCGUnionData::CreatePointData);
 
-	return CastChecked<UPCGPointData>(CreateBasePointData(Context, UPCGPointData::StaticClass(), [](FPCGContext* InContext, const UPCGSpatialData* InSpatialData) { return InSpatialData->ToPointData(InContext); }));
+	return CastChecked<UPCGPointData>(CreateBasePointData(Context, UPCGPointData::StaticClass(), [](FPCGContext* InContext, const UPCGSpatialData* InSpatialData) { return InSpatialData->ToPointData(InContext); }), ECastCheckedType::NullAllowed);
 }
 
 const UPCGPointArrayData* UPCGUnionData::CreatePointArrayData(FPCGContext* Context, const FBox& InBounds) const
 {
 	TRACE_CPUPROFILER_EVENT_SCOPE(UPCGUnionData::CreatePointArrayData);
 
-	return CastChecked<UPCGPointArrayData>(CreateBasePointData(Context, UPCGPointArrayData::StaticClass(), [](FPCGContext* InContext, const UPCGSpatialData* InSpatialData) { return InSpatialData->ToPointArrayData(InContext); }));
+	return CastChecked<UPCGPointArrayData>(CreateBasePointData(Context, UPCGPointArrayData::StaticClass(), [](FPCGContext* InContext, const UPCGSpatialData* InSpatialData) { return InSpatialData->ToPointArrayData(InContext); }), ECastCheckedType::NullAllowed);
 }
 
 const UPCGBasePointData* UPCGUnionData::CreateBasePointData(FPCGContext* Context, TSubclassOf<UPCGBasePointData> PointDataClass, TFunctionRef<const UPCGBasePointData*(FPCGContext*, const UPCGSpatialData*)> ToPointDataFunc) const
diff --git a/Engine/Plugins/PCG/Source/PCG/Private/Data/PCGVolumeData.cpp b/Engine/Plugins/PCG/Source/PCG/Private/Data/PCGVolumeData.cpp
--- a/Engine/Plugins/PCG/Source/PCG/Private/Data/PCGVolumeData.cpp
+++ b/Engine/Plugins/PCG/Source/PCG/Private/Data/PCGVolumeData.cpp
@@ -143,14 +143,14 @@ const UPCGPointData* UPCGVolumeData::CreatePointData(FPCGContext* Context) const
 {
 	TRACE_CPUPROFILER_EVENT_SCOPE(UPCGVolumeData::CreatePointData);
 
-	return Cast<UPCGPointData>(CreateBasePointData(Context, UPCGPointData::StaticClass()));
+	return CastChecked<UPCGPointData>(CreateBasePointData(Context, UPCGPointData::StaticClass()), ECastCheckedType::NullAllowed);
 }
 
 const UPCGPointArrayData* UPCGVolumeData::CreatePointArrayData(FPCGContext* Context, const FBox& InBounds) const
 {
 	TRACE_CPUPROFILER_EVENT_SCOPE(UPCGVolumeData::CreatePointArrayData);
 
-	return Cast<UPCGPointArrayData>(CreateBasePointData(Context, UPCGPointArrayData::StaticClass()));
+	return CastChecked<UPCGPointArrayData>(CreateBasePointData(Context, UPCGPointArrayData::StaticClass()), ECastCheckedType::NullAllowed);
 }
 
 const UPCGBasePointData* UPCGVolumeData::CreateBasePointData(FPCGContext* Context, TSubclassOf<UPCGBasePointData> PointDataClass) const
diff --git a/Engine/Plugins/PCG/Source/PCG/Private/Data/PCGWorldData.cpp b/Engine/Plugins/PCG/Source/PCG/Private/Data/PCGWorldData.cpp
--- a/Engine/Plugins/PCG/Source/PCG/Private/Data/PCGWorldData.cpp
+++ b/Engine/Plugins/PCG/Source/PCG/Private/Data/PCGWorldData.cpp
@@ -208,13 +208,13 @@ bool UPCGWorldVolumetricData::SamplePoint(const FTransform& InTransform, const F
 const UPCGPointData* UPCGWorldVolumetricData::CreatePointData(FPCGContext* Context, const FBox& InBounds) const
 {
 	TRACE_CPUPROFILER_EVENT_SCOPE(UPCGWorldVolumetricData::CreatePointData);
-	return Cast<UPCGPointData>(CreateBasePointData(Context, InBounds, UPCGPointData::StaticClass()));
+	return CastChecked<UPCGPointData>(CreateBasePointData(Context, InBounds, UPCGPointData::StaticClass()), ECastCheckedType::NullAllowed);
 }
 
 const UPCGPointArrayData* UPCGWorldVolumetricData::CreatePointArrayData(FPCGContext* Context, const FBox& InBounds) const
 {
 	TRACE_CPUPROFILER_EVENT_SCOPE(UPCGWorldVolumetricData::CreatePointArrayData);
-	return Cast<UPCGPointArrayData>(CreateBasePointData(Context, InBounds, UPCGPointArrayData::StaticClass()));
+	return CastChecked<UPCGPointArrayData>(CreateBasePointData(Context, InBounds, UPCGPointArrayData::StaticClass()), ECastCheckedType::NullAllowed);
 }
 
 const UPCGBasePointData* UPCGWorldVolumetricData::CreateBasePointData(FPCGContext* Context, const FBox& InBounds, TSubclassOf<UPCGBasePointData> PointDataClass) const
