[PCG] Pooling compute graph instances rather than creating from scratch each time, to avoid overhead of repeatably building data providers.

#rnx

diff --git a/Engine/Plugins/PCG/Source/PCG/Private/Compute/DataInterfaces/Elements/PCGCopyPointsAnalysisDataInterface.cpp b/Engine/Plugins/PCG/Source/PCG/Private/Compute/DataInterfaces/Elements/PCGCopyPointsAnalysisDataInterface.cpp
--- a/Engine/Plugins/PCG/Source/PCG/Private/Compute/DataInterfaces/Elements/PCGCopyPointsAnalysisDataInterface.cpp
+++ b/Engine/Plugins/PCG/Source/PCG/Private/Compute/DataInterfaces/Elements/PCGCopyPointsAnalysisDataInterface.cpp
@@ -55,23 +55,29 @@ void UPCGCopyPointsAnalysisDataInterface::GetHLSL(FString& OutHLSL, FString cons
 		), TemplateArgs);
 }
 
-UComputeDataProvider* UPCGCopyPointsAnalysisDataInterface::CreateDataProvider(TObjectPtr<UObject> InBinding, uint64 InInputMask, uint64 InOutputMask) const
+UComputeDataProvider* UPCGCopyPointsAnalysisDataInterface::CreateDataProvider() const
 {
-	TRACE_CPUPROFILER_EVENT_SCOPE(UPCGCopyPointsAnalysisDataInterface::CreateDataProvider);
+	return NewObject<UPCGCopyPointsAnalysisDataProvider>();
+}
+
+void UPCGCopyPointsAnalysisDataProvider::Initialize(const UComputeDataInterface* InDataInterface, UObject* InBinding, uint64 InInputMask, uint64 InOutputMask)
+{
+	TRACE_CPUPROFILER_EVENT_SCOPE(UPCGCopyPointsAnalysisDataProvider::Initialize);
+
+	Super::Initialize(InDataInterface, InBinding, InInputMask, InOutputMask);
+
+	const UPCGCopyPointsAnalysisDataInterface* DataInterface = CastChecked<UPCGCopyPointsAnalysisDataInterface>(InDataInterface);
 	UPCGDataBinding* Binding = CastChecked<UPCGDataBinding>(InBinding);
 
-	TObjectPtr<UPCGCopyPointsAnalysisDataProvider> DataProvider = NewObject<UPCGCopyPointsAnalysisDataProvider>();
-	DataProvider->Settings = CastChecked<UPCGCopyPointsSettings>(ProducerKernel->GetSettings());
+	Settings = CastChecked<UPCGCopyPointsSettings>(DataInterface->ProducerKernel->GetSettings());
 
 	// Resolve attribute IDs from incoming data.
-	const FName MatchAttributeName = DataProvider->Settings->MatchAttribute;
-	DataProvider->MatchAttributeId = Binding->GetAttributeId(MatchAttributeName, EPCGKernelAttributeType::Int);
-	ensure(DataProvider->MatchAttributeId != INDEX_NONE);
-
-	DataProvider->SelectedFlagAttributeId = Binding->GetAttributeId(PCGCopyPointsConstants::SelectedFlagAttributeName, EPCGKernelAttributeType::Bool);
-	ensure(DataProvider->SelectedFlagAttributeId != INDEX_NONE);
+	const FName MatchAttributeName = Settings->MatchAttribute;
+	MatchAttributeId = Binding->GetAttributeId(MatchAttributeName, EPCGKernelAttributeType::Int);
+	ensure(MatchAttributeId != INDEX_NONE);
 
-	return DataProvider;
+	SelectedFlagAttributeId = Binding->GetAttributeId(PCGCopyPointsConstants::SelectedFlagAttributeName, EPCGKernelAttributeType::Bool);
+	ensure(SelectedFlagAttributeId != INDEX_NONE);
 }
 
 FComputeDataProviderRenderProxy* UPCGCopyPointsAnalysisDataProvider::GetRenderProxy()
diff --git a/Engine/Plugins/PCG/Source/PCG/Private/Compute/DataInterfaces/Elements/PCGCopyPointsAnalysisDataInterface.h b/Engine/Plugins/PCG/Source/PCG/Private/Compute/DataInterfaces/Elements/PCGCopyPointsAnalysisDataInterface.h
--- a/Engine/Plugins/PCG/Source/PCG/Private/Compute/DataInterfaces/Elements/PCGCopyPointsAnalysisDataInterface.h
+++ b/Engine/Plugins/PCG/Source/PCG/Private/Compute/DataInterfaces/Elements/PCGCopyPointsAnalysisDataInterface.h
@@ -22,7 +22,7 @@ public:
 	void GetSupportedInputs(TArray<FShaderFunctionDefinition>& OutFunctions) const override;
 	void GetShaderParameters(TCHAR const* UID, FShaderParametersMetadataBuilder& InOutBuilder, FShaderParametersMetadataAllocations& InOutAllocations) const override;
 	void GetHLSL(FString& OutHLSL, FString const& InDataInterfaceName) const override;
-	UComputeDataProvider* CreateDataProvider(TObjectPtr<UObject> InBinding, uint64 InInputMask, uint64 InOutputMask) const override;
+	UComputeDataProvider* CreateDataProvider() const override;
 	//~ End UComputeDataInterface Interface
 };
 
@@ -33,6 +33,7 @@ class UPCGCopyPointsAnalysisDataProvider : public UPCGComputeDataProvider
 
 public:
 	//~ Begin UComputeDataProvider Interface
+	virtual void Initialize(const UComputeDataInterface* InDataInterface, UObject* InBinding, uint64 InInputMask, uint64 InOutputMask) override;
 	FComputeDataProviderRenderProxy* GetRenderProxy() override;
 	//~ End UComputeDataProvider Interface
 
diff --git a/Engine/Plugins/PCG/Source/PCG/Private/Compute/DataInterfaces/Elements/PCGCopyPointsDataInterface.cpp b/Engine/Plugins/PCG/Source/PCG/Private/Compute/DataInterfaces/Elements/PCGCopyPointsDataInterface.cpp
--- a/Engine/Plugins/PCG/Source/PCG/Private/Compute/DataInterfaces/Elements/PCGCopyPointsDataInterface.cpp
+++ b/Engine/Plugins/PCG/Source/PCG/Private/Compute/DataInterfaces/Elements/PCGCopyPointsDataInterface.cpp
@@ -115,14 +115,19 @@ void UPCGCopyPointsDataInterface::GetHLSL(FString& OutHLSL, FString const& InDat
 		), TemplateArgs);
 }
 
-UComputeDataProvider* UPCGCopyPointsDataInterface::CreateDataProvider(TObjectPtr<UObject> InBinding, uint64 InInputMask, uint64 InOutputMask) const
+UComputeDataProvider* UPCGCopyPointsDataInterface::CreateDataProvider() const
 {
-	TRACE_CPUPROFILER_EVENT_SCOPE(UPCGCopyPointsDataInterface::CreateDataProvider);
+	return NewObject<UPCGCopyPointsDataProvider>();
+}
+
+void UPCGCopyPointsDataProvider::Initialize(const UComputeDataInterface* InDataInterface, UObject* InBinding, uint64 InInputMask, uint64 InOutputMask)
+{
+	TRACE_CPUPROFILER_EVENT_SCOPE(UPCGCopyPointsDataProvider::Initialize);
 
-	TObjectPtr<UPCGCopyPointsDataProvider> DataProvider = NewObject<UPCGCopyPointsDataProvider>();
-	DataProvider->ProducerKernel = ProducerKernel;
+	Super::Initialize(InDataInterface, InBinding, InInputMask, InOutputMask);
 
-	return DataProvider;
+	const UPCGCopyPointsDataInterface* DataInterface = CastChecked<UPCGCopyPointsDataInterface>(InDataInterface);
+	ProducerKernel = DataInterface->ProducerKernel;
 }
 
 bool UPCGCopyPointsDataProvider::PerformPreExecuteReadbacks_GameThread(UPCGDataBinding* InBinding)
@@ -276,6 +281,14 @@ FComputeDataProviderRenderProxy* UPCGCopyPointsDataProvider::GetRenderProxy()
 	return new FPCGCopyPointsDataProviderProxy(MoveTemp(ProxyData));
 }
 
+void UPCGCopyPointsDataProvider::Reset()
+{
+	ProducerKernel = nullptr;
+	SourceAndTargetDataIndices.Empty();
+
+	Super::Reset();
+}
+
 bool FPCGCopyPointsDataProviderProxy::IsValid(FValidationData const& InValidationData) const
 {
 	if (InValidationData.ParameterStructSize != sizeof(FParameters))
diff --git a/Engine/Plugins/PCG/Source/PCG/Private/Compute/DataInterfaces/Elements/PCGCopyPointsDataInterface.h b/Engine/Plugins/PCG/Source/PCG/Private/Compute/DataInterfaces/Elements/PCGCopyPointsDataInterface.h
--- a/Engine/Plugins/PCG/Source/PCG/Private/Compute/DataInterfaces/Elements/PCGCopyPointsDataInterface.h
+++ b/Engine/Plugins/PCG/Source/PCG/Private/Compute/DataInterfaces/Elements/PCGCopyPointsDataInterface.h
@@ -22,7 +22,7 @@ public:
 	void GetSupportedInputs(TArray<FShaderFunctionDefinition>& OutFunctions) const override;
 	void GetShaderParameters(TCHAR const* UID, FShaderParametersMetadataBuilder& InOutBuilder, FShaderParametersMetadataAllocations& InOutAllocations) const override;
 	void GetHLSL(FString& OutHLSL, FString const& InDataInterfaceName) const override;
-	UComputeDataProvider* CreateDataProvider(TObjectPtr<UObject> InBinding, uint64 InInputMask, uint64 InOutputMask) const override;
+	UComputeDataProvider* CreateDataProvider() const override;
 	//~ End UComputeDataInterface Interface
 };
 
@@ -38,7 +38,9 @@ public:
 	//~ End UPCGComputeDataProvider Interface
 
 	//~ Begin UComputeDataProvider Interface
+	virtual void Initialize(const UComputeDataInterface* InDataInterface, UObject* InBinding, uint64 InInputMask, uint64 InOutputMask) override;
 	FComputeDataProviderRenderProxy* GetRenderProxy() override;
+	virtual void Reset() override;
 	//~ End UComputeDataProvider Interface
 
 public:
diff --git a/Engine/Plugins/PCG/Source/PCG/Private/Compute/DataInterfaces/Elements/PCGMetadataPartitionDataInterface.cpp b/Engine/Plugins/PCG/Source/PCG/Private/Compute/DataInterfaces/Elements/PCGMetadataPartitionDataInterface.cpp
--- a/Engine/Plugins/PCG/Source/PCG/Private/Compute/DataInterfaces/Elements/PCGMetadataPartitionDataInterface.cpp
+++ b/Engine/Plugins/PCG/Source/PCG/Private/Compute/DataInterfaces/Elements/PCGMetadataPartitionDataInterface.cpp
@@ -80,14 +80,19 @@ void UPCGMetadataPartitionDataInterface::GetHLSL(FString& OutHLSL, FString const
 	), TemplateArgs);
 }
 
-UComputeDataProvider* UPCGMetadataPartitionDataInterface::CreateDataProvider(TObjectPtr<UObject> InBinding, uint64 InInputMask, uint64 InOutputMask) const
+UComputeDataProvider* UPCGMetadataPartitionDataInterface::CreateDataProvider() const
 {
-	TRACE_CPUPROFILER_EVENT_SCOPE(UPCGMetadataPartitionDataInterface::CreateDataProvider);
+	return NewObject<UPCGMetaDataPartitionDataProvider>();
+}
+
+void UPCGMetaDataPartitionDataProvider::Initialize(const UComputeDataInterface* InDataInterface, UObject* InBinding, uint64 InInputMask, uint64 InOutputMask)
+{
+	TRACE_CPUPROFILER_EVENT_SCOPE(UPCGMetaDataPartitionDataProvider::Initialize);
 
-	UPCGMetaDataPartitionDataProvider* DataProvider = NewObject<UPCGMetaDataPartitionDataProvider>();
-	DataProvider->Kernel = ProducerKernel;
+	Super::Initialize(InDataInterface, InBinding, InInputMask, InOutputMask);
 
-	return DataProvider;
+	const UPCGMetadataPartitionDataInterface* DataInterface = CastChecked<UPCGMetadataPartitionDataInterface>(InDataInterface);
+	Kernel = DataInterface->ProducerKernel;
 }
 
 bool UPCGMetaDataPartitionDataProvider::PrepareForExecute_GameThread(UPCGDataBinding* InBinding)
@@ -132,6 +137,18 @@ FComputeDataProviderRenderProxy* UPCGMetaDataPartitionDataProvider::GetRenderPro
 	return new FPCGMetaDataPartitionProviderProxy(PartitionAttributeId, NumInputData, NumPartitions, MaxAttributeValue, UniqueStringKeyValues);
 }
 
+void UPCGMetaDataPartitionDataProvider::Reset()
+{
+	Kernel = nullptr;
+	PartitionAttributeId = INDEX_NONE;
+	NumInputData = 0;
+	NumPartitions = 0;
+	MaxAttributeValue = INDEX_NONE;
+	UniqueStringKeyValues.Empty();
+
+	Super::Reset();
+}
+
 bool FPCGMetaDataPartitionProviderProxy::IsValid(FValidationData const& InValidationData) const
 {
 	return (InValidationData.ParameterStructSize == sizeof(FParameters));
diff --git a/Engine/Plugins/PCG/Source/PCG/Private/Compute/DataInterfaces/Elements/PCGMetadataPartitionDataInterface.h b/Engine/Plugins/PCG/Source/PCG/Private/Compute/DataInterfaces/Elements/PCGMetadataPartitionDataInterface.h
--- a/Engine/Plugins/PCG/Source/PCG/Private/Compute/DataInterfaces/Elements/PCGMetadataPartitionDataInterface.h
+++ b/Engine/Plugins/PCG/Source/PCG/Private/Compute/DataInterfaces/Elements/PCGMetadataPartitionDataInterface.h
@@ -21,7 +21,7 @@ public:
 	void GetSupportedInputs(TArray<FShaderFunctionDefinition>& OutFunctions) const override;
 	void GetShaderParameters(TCHAR const* UID, FShaderParametersMetadataBuilder& InOutBuilder, FShaderParametersMetadataAllocations& InOutAllocations) const override;
 	void GetHLSL(FString& OutHLSL, FString const& InDataInterfaceName) const override;
-	UComputeDataProvider* CreateDataProvider(TObjectPtr<UObject> InBinding, uint64 InInputMask, uint64 InOutputMask) const override;
+	UComputeDataProvider* CreateDataProvider() const override;
 	//~ End UComputeDataInterface Interface
 };
 
@@ -36,7 +36,9 @@ public:
 	//~ End UPCGComputeDataProvider Interface
 
 	//~ Begin UComputeDataProvider Interface
+	virtual void Initialize(const UComputeDataInterface* InDataInterface, UObject* InBinding, uint64 InInputMask, uint64 InOutputMask) override;
 	FComputeDataProviderRenderProxy* GetRenderProxy() override;
+	virtual void Reset() override;
 	//~ End UComputeDataProvider Interface
 
 public:
diff --git a/Engine/Plugins/PCG/Source/PCG/Private/Compute/DataInterfaces/Elements/PCGStaticMeshSpawnerDataInterface.cpp b/Engine/Plugins/PCG/Source/PCG/Private/Compute/DataInterfaces/Elements/PCGStaticMeshSpawnerDataInterface.cpp
--- a/Engine/Plugins/PCG/Source/PCG/Private/Compute/DataInterfaces/Elements/PCGStaticMeshSpawnerDataInterface.cpp
+++ b/Engine/Plugins/PCG/Source/PCG/Private/Compute/DataInterfaces/Elements/PCGStaticMeshSpawnerDataInterface.cpp
@@ -171,34 +171,38 @@ void UPCGStaticMeshSpawnerDataInterface::GetHLSL(FString& OutHLSL, FString const
 	), TemplateArgs);
 }
 
-UComputeDataProvider* UPCGStaticMeshSpawnerDataInterface::CreateDataProvider(TObjectPtr<UObject> InBinding, uint64 InInputMask, uint64 InOutputMask) const
+UComputeDataProvider* UPCGStaticMeshSpawnerDataInterface::CreateDataProvider() const
 {
-	TRACE_CPUPROFILER_EVENT_SCOPE(UPCGStaticMeshSpawnerDataInterface::CreateDataProvider);
-	UPCGDataBinding* Binding = CastChecked<UPCGDataBinding>(InBinding);
+	return NewObject<UPCGStaticMeshSpawnerDataProvider>();
+}
+
+void UPCGStaticMeshSpawnerDataProvider::Initialize(const UComputeDataInterface* InDataInterface, UObject* InBinding, uint64 InInputMask, uint64 InOutputMask)
+{
+	TRACE_CPUPROFILER_EVENT_SCOPE(UPCGStaticMeshSpawnerDataProvider::Initialize);
+
+	Super::Initialize(InDataInterface, InBinding, InInputMask, InOutputMask);
 
-	const UPCGStaticMeshSpawnerKernel* Kernel = CastChecked<UPCGStaticMeshSpawnerKernel>(ProducerKernel);
+	const UPCGStaticMeshSpawnerDataInterface* DataInterface = CastChecked<UPCGStaticMeshSpawnerDataInterface>(InDataInterface);
+	const UPCGStaticMeshSpawnerKernel* Kernel = CastChecked<UPCGStaticMeshSpawnerKernel>(DataInterface->ProducerKernel);
+	UPCGDataBinding* Binding = CastChecked<UPCGDataBinding>(InBinding);
 	check(Kernel->Settings);
 
-	TObjectPtr<UPCGStaticMeshSpawnerDataProvider> DataProvider = NewObject<UPCGStaticMeshSpawnerDataProvider>();
-	
 	// If there were 0 input points for this execution, we will not have created any primitives, so check for null.
 	if (const FPCGSpawnerPrimitives* Primitives = Binding->MeshSpawnersToPrimitives.Find(Kernel))
 	{
-		DataProvider->AttributeIdOffsetStrides = Primitives->AttributeIdOffsetStrides;
-		DataProvider->SelectionCDF = Primitives->SelectionCDF;
-		DataProvider->PrimitiveStringKeys = Primitives->PrimitiveStringKeys;
-		DataProvider->PrimitiveMeshBounds = Primitives->PrimitiveMeshBounds;
+		AttributeIdOffsetStrides = Primitives->AttributeIdOffsetStrides;
+		SelectionCDF = Primitives->SelectionCDF;
+		PrimitiveStringKeys = Primitives->PrimitiveStringKeys;
+		PrimitiveMeshBounds = Primitives->PrimitiveMeshBounds;
 
 		// We only know name at compile time (not type), so we resolved this Id at execution time by inspecting the data arriving on the input pin. 
-		DataProvider->SelectorAttributeId = Primitives->SelectorAttributeId;
+		SelectorAttributeId = Primitives->SelectorAttributeId;
 
 		// We know the name and type of the selected mesh attribute statically and declared the attribute in GetKernelAttributeKeys,
 		// so the attribute ID should be present in the attribute table.
-		DataProvider->SelectedMeshAttributeId = Binding->GetAttributeId(Kernel->Settings->OutAttributeName, EPCGKernelAttributeType::StringKey);
-		ensure(DataProvider->SelectedMeshAttributeId != INDEX_NONE);
+		SelectedMeshAttributeId = Binding->GetAttributeId(Kernel->Settings->OutAttributeName, EPCGKernelAttributeType::StringKey);
+		ensure(SelectedMeshAttributeId != INDEX_NONE);
 	}
-
-	return DataProvider;
 }
 
 FComputeDataProviderRenderProxy* UPCGStaticMeshSpawnerDataProvider::GetRenderProxy()
@@ -206,6 +210,19 @@ FComputeDataProviderRenderProxy* UPCGStaticMeshSpawnerDataProvider::GetRenderPro
 	return new FPCGStaticMeshSpawnerDataProviderProxy(AttributeIdOffsetStrides, SelectorAttributeId, PrimitiveStringKeys, SelectionCDF, SelectedMeshAttributeId, PrimitiveMeshBounds);
 }
 
+void UPCGStaticMeshSpawnerDataProvider::Reset()
+{
+	AttributeIdOffsetStrides.Empty();
+	PrimitiveStringKeys.Empty();
+	PrimitiveMeshBounds.Empty();
+	SelectionCDF.Empty();
+	SelectorAttributeId = INDEX_NONE;
+	NumInputPoints = 0;
+	SelectedMeshAttributeId = INDEX_NONE;
+
+	Super::Reset();
+}
+
 bool FPCGStaticMeshSpawnerDataProviderProxy::IsValid(FValidationData const& InValidationData) const
 {
 	if (InValidationData.ParameterStructSize != sizeof(FParameters))
diff --git a/Engine/Plugins/PCG/Source/PCG/Private/Compute/DataInterfaces/Elements/PCGStaticMeshSpawnerDataInterface.h b/Engine/Plugins/PCG/Source/PCG/Private/Compute/DataInterfaces/Elements/PCGStaticMeshSpawnerDataInterface.h
--- a/Engine/Plugins/PCG/Source/PCG/Private/Compute/DataInterfaces/Elements/PCGStaticMeshSpawnerDataInterface.h
+++ b/Engine/Plugins/PCG/Source/PCG/Private/Compute/DataInterfaces/Elements/PCGStaticMeshSpawnerDataInterface.h
@@ -23,7 +23,7 @@ public:
 	void GetSupportedInputs(TArray<FShaderFunctionDefinition>& OutFunctions) const override;
 	void GetShaderParameters(TCHAR const* UID, FShaderParametersMetadataBuilder& InOutBuilder, FShaderParametersMetadataAllocations& InOutAllocations) const override;
 	void GetHLSL(FString& OutHLSL, FString const& InDataInterfaceName) const override;
-	UComputeDataProvider* CreateDataProvider(TObjectPtr<UObject> InBinding, uint64 InInputMask, uint64 InOutputMask) const override;
+	UComputeDataProvider* CreateDataProvider() const override;
 	//~ End UComputeDataInterface Interface
 
 	static constexpr uint32 MAX_ATTRIBUTES = 64;
@@ -36,7 +36,9 @@ class UPCGStaticMeshSpawnerDataProvider : public UComputeDataProvider
 
 public:
 	//~ Begin UComputeDataProvider Interface
+	virtual void Initialize(const UComputeDataInterface* InDataInterface, UObject* InBinding, uint64 InInputMask, uint64 InOutputMask) override;
 	FComputeDataProviderRenderProxy* GetRenderProxy() override;
+	virtual void Reset() override;
 	//~ End UComputeDataProvider Interface
 
 	UPROPERTY()
diff --git a/Engine/Plugins/PCG/Source/PCG/Private/Compute/DataInterfaces/PCGComputeDataInterface.h b/Engine/Plugins/PCG/Source/PCG/Private/Compute/DataInterfaces/PCGComputeDataInterface.h
--- a/Engine/Plugins/PCG/Source/PCG/Private/Compute/DataInterfaces/PCGComputeDataInterface.h
+++ b/Engine/Plugins/PCG/Source/PCG/Private/Compute/DataInterfaces/PCGComputeDataInterface.h
@@ -29,13 +29,13 @@ public:
 		OutputPinLabelAlias = InOptionalOutputPinLabelAlias ? *InOptionalOutputPinLabelAlias : InOutputPinLabel;
 	}
 
+public:
 	UPROPERTY()
 	int32 GraphBindingIndex = INDEX_NONE;
 
 	UPROPERTY()
 	TObjectPtr<const UPCGComputeKernel> ProducerKernel;
 
-protected:
 	/** Label of output pin that this data interface is associated with. */
 	UPROPERTY()
 	FName OutputPinLabel;
diff --git a/Engine/Plugins/PCG/Source/PCG/Private/Compute/DataInterfaces/PCGCustomKernelDataInterface.cpp b/Engine/Plugins/PCG/Source/PCG/Private/Compute/DataInterfaces/PCGCustomKernelDataInterface.cpp
--- a/Engine/Plugins/PCG/Source/PCG/Private/Compute/DataInterfaces/PCGCustomKernelDataInterface.cpp
+++ b/Engine/Plugins/PCG/Source/PCG/Private/Compute/DataInterfaces/PCGCustomKernelDataInterface.cpp
@@ -89,27 +89,32 @@ void UPCGCustomKernelDataInterface::GetHLSL(FString& OutHLSL, FString const& InD
 		TemplateArgs);
 }
 
-UComputeDataProvider* UPCGCustomKernelDataInterface::CreateDataProvider(TObjectPtr<UObject> InBinding, uint64 InInputMask, uint64 InOutputMask) const
+UComputeDataProvider* UPCGCustomKernelDataInterface::CreateDataProvider() const
 {
-	TRACE_CPUPROFILER_EVENT_SCOPE(UPCGCustomKernelDataInterface::CreateDataProvider);
-	check(Settings);
-	check(Kernel);
+	return NewObject<UPCGCustomComputeKernelDataProvider>();
+}
 
-	UPCGDataBinding* Binding = CastChecked<UPCGDataBinding>(InBinding);
+void UPCGCustomComputeKernelDataProvider::Initialize(const UComputeDataInterface* InDataInterface, UObject* InBinding, uint64 InInputMask, uint64 InOutputMask)
+{
+	TRACE_CPUPROFILER_EVENT_SCOPE(UPCGCustomComputeKernelDataProvider::Initialize);
+
+	Super::Initialize(InDataInterface, InBinding, InInputMask, InOutputMask);
+
+	const UPCGCustomKernelDataInterface* DataInterface = CastChecked<UPCGCustomKernelDataInterface>(InDataInterface);
 
+	UPCGDataBinding* Binding = CastChecked<UPCGDataBinding>(InBinding);
 	const UPCGComponent* SourceComponent = Binding->GetSourceComponent().Get();
 	check(SourceComponent);
+	check(DataInterface->Settings);
+	check(DataInterface->Kernel);
 
-	UPCGCustomComputeKernelDataProvider* Provider = NewObject<UPCGCustomComputeKernelDataProvider>();
-	Provider->Kernel = Kernel;
-	Provider->ThreadCountMultiplier = Kernel->GetThreadCountMultiplier();
-	Provider->Seed = static_cast<uint32>(Settings->GetSeed(SourceComponent));
-	Provider->SeedSettings = static_cast<uint32>(Settings->Seed);
-	Provider->SeedComponent = static_cast<uint32>(SourceComponent->Seed);
-
-	Provider->SourceComponentBounds = SourceComponent->GetGridBounds();
+	Kernel = DataInterface->Kernel;
+	ThreadCountMultiplier = DataInterface->Kernel->GetThreadCountMultiplier();
+	Seed = static_cast<uint32>(DataInterface->Settings->GetSeed(SourceComponent));
+	SeedSettings = static_cast<uint32>(DataInterface->Settings->Seed);
+	SeedComponent = static_cast<uint32>(SourceComponent->Seed);
 
-	return Provider;
+	SourceComponentBounds = SourceComponent->GetGridBounds();
 }
 
 bool UPCGCustomComputeKernelDataProvider::PrepareForExecute_GameThread(UPCGDataBinding* InBinding)
@@ -123,6 +128,19 @@ FComputeDataProviderRenderProxy* UPCGCustomComputeKernelDataProvider::GetRenderP
 	return new FPCGCustomComputeKernelDataProviderProxy(ThreadCount, ThreadCountMultiplier, Seed, SeedSettings, SeedComponent, SourceComponentBounds);
 }
 
+void UPCGCustomComputeKernelDataProvider::Reset()
+{
+	Kernel = nullptr;
+	ThreadCount = -1;
+	ThreadCountMultiplier = 0;
+	Seed = 42;
+	SeedSettings = 42;
+	SeedComponent = 42;
+	SourceComponentBounds = {};
+
+	Super::Reset();
+}
+
 bool FPCGCustomComputeKernelDataProviderProxy::IsValid(FValidationData const& InValidationData) const
 {
 	if (InValidationData.ParameterStructSize != sizeof(FParameters))
diff --git a/Engine/Plugins/PCG/Source/PCG/Private/Compute/DataInterfaces/PCGCustomKernelDataInterface.h b/Engine/Plugins/PCG/Source/PCG/Private/Compute/DataInterfaces/PCGCustomKernelDataInterface.h
--- a/Engine/Plugins/PCG/Source/PCG/Private/Compute/DataInterfaces/PCGCustomKernelDataInterface.h
+++ b/Engine/Plugins/PCG/Source/PCG/Private/Compute/DataInterfaces/PCGCustomKernelDataInterface.h
@@ -26,8 +26,7 @@ public:
 	void GetSupportedInputs(TArray<FShaderFunctionDefinition>& OutFunctions) const override;
 	void GetShaderParameters(TCHAR const* UID, FShaderParametersMetadataBuilder& InOutBuilder, FShaderParametersMetadataAllocations& InOutAllocations) const override;
 	void GetHLSL(FString& OutHLSL, FString const& InDataInterfaceName) const override;
-	UComputeDataProvider* CreateDataProvider(TObjectPtr<UObject> InBinding, uint64 InInputMask, uint64 InOutputMask) const override;
-
+	UComputeDataProvider* CreateDataProvider() const override;
 	// This DI will provide execution parameters like dispatch information.
 	bool IsExecutionInterface() const override { return true; }
 	//~ End UComputeDataInterface Interface
@@ -51,7 +50,9 @@ public:
 	//~ End UPCGComputeDataProvider Interface
 
 	//~ Begin UComputeDataProvider Interface
+	virtual void Initialize(const UComputeDataInterface* InDataInterface, UObject* InBinding, uint64 InInputMask, uint64 InOutputMask) override;
 	FComputeDataProviderRenderProxy* GetRenderProxy() override;
+	virtual void Reset() override;
 	//~ End UComputeDataProvider Interface
 
 	UPROPERTY()
diff --git a/Engine/Plugins/PCG/Source/PCG/Private/Compute/DataInterfaces/PCGDataCollectionDataInterface.cpp b/Engine/Plugins/PCG/Source/PCG/Private/Compute/DataInterfaces/PCGDataCollectionDataInterface.cpp
--- a/Engine/Plugins/PCG/Source/PCG/Private/Compute/DataInterfaces/PCGDataCollectionDataInterface.cpp
+++ b/Engine/Plugins/PCG/Source/PCG/Private/Compute/DataInterfaces/PCGDataCollectionDataInterface.cpp
@@ -440,44 +440,53 @@ void UPCGDataCollectionDataInterface::GetHLSL(FString& OutHLSL, FString const& I
 	}
 }
 
-UComputeDataProvider* UPCGDataCollectionDataInterface::CreateDataProvider(TObjectPtr<UObject> InBinding, uint64 InInputMask, uint64 InOutputMask) const
+UComputeDataProvider* UPCGDataCollectionDataInterface::CreateDataProvider() const
 {
-	TRACE_CPUPROFILER_EVENT_SCOPE(UPCGDataCollectionDataInterface::CreateDataProvider);
-	check(ProducerSettings);
+	return NewObject<UPCGDataCollectionDataProvider>();
+}
+
+void UPCGDataCollectionDataProvider::Initialize(const UComputeDataInterface* InDataInterface, UObject* InBinding, uint64 InInputMask, uint64 InOutputMask)
+{
+	TRACE_CPUPROFILER_EVENT_SCOPE(UPCGDataCollectionDataProvider::Initialize);
+
+	// Bump counter so any async callbacks from previous usages of this provider will be invalidated.
+	++GenerationCounter;
 
-	UPCGDataBinding* Binding = CastChecked<UPCGDataBinding>(InBinding);
+	Super::Initialize(InDataInterface, InBinding, InInputMask, InOutputMask);
 
-	UPCGDataCollectionDataProvider* Provider = NewObject<UPCGDataCollectionDataProvider>();
-	Provider->Binding = Binding;
-	Provider->ProducerSettings = ProducerSettings;
-	Provider->GraphBindingIndex = GraphBindingIndex;
+	const UPCGDataCollectionDataInterface* DataInterface = CastChecked<UPCGDataCollectionDataInterface>(InDataInterface);
 
-	Provider->ExportMode = bRequiresExport ? EPCGExportMode::ComputeGraphOutput : EPCGExportMode::NoExport;
+	Binding = CastChecked<UPCGDataBinding>(InBinding);
 
-	Provider->ElementCountMultiplier = ElementCountMultiplier;
+	check(DataInterface->ProducerSettings);
 
-	Provider->bExposeElementCounters = bExposeElementCounters;
+	ProducerSettings = DataInterface->ProducerSettings;
+	GraphBindingIndex = DataInterface->GraphBindingIndex;
+
+	ExportMode = DataInterface->bRequiresExport ? EPCGExportMode::ComputeGraphOutput : EPCGExportMode::NoExport;
+
+	ElementCountMultiplier = DataInterface->ElementCountMultiplier;
+
+	bExposeElementCounters = DataInterface->bExposeElementCounters;
 
 #if WITH_EDITOR
 	const UPCGComponent* SourceComponent = Binding->GetSourceComponent().Get();
 	if (ProducerSettings->bIsInspecting && SourceComponent && SourceComponent->IsInspecting())
 	{
-		Provider->ExportMode |= EPCGExportMode::Inspection;
+		ExportMode |= EPCGExportMode::Inspection;
 	}
 
 	if (ProducerSettings->bDebug)
 	{
-		Provider->ExportMode |= EPCGExportMode::DebugVisualization;
+		ExportMode |= EPCGExportMode::DebugVisualization;
 	}
 #endif
 
 	// Use the aliased label for normal data output as this is the output from the compute graph.
-	Provider->OutputPinLabelAlias = OutputPinLabelAlias;
+	OutputPinLabelAlias = DataInterface->OutputPinLabelAlias;
 
 	// The original label is needed to store inspection data.
-	Provider->OutputPinLabel = OutputPinLabel;
-
-	return Provider;
+	OutputPinLabel = DataInterface->OutputPinLabel;
 }
 
 bool UPCGDataCollectionDataProvider::PrepareForExecute_GameThread(UPCGDataBinding* InBinding)
@@ -522,7 +531,9 @@ FComputeDataProviderRenderProxy* UPCGDataCollectionDataProvider::GetRenderProxy(
 
 	if (bExposeElementCounters)
 	{
-		Proxy->AsyncReadbackCallback_RenderThread = [ThisWeak=TWeakObjectPtr<UPCGDataCollectionDataProvider>(this)](const void* InData, int InNumBytes)
+		const uint64 GenerationCount = GenerationCounter;
+
+		Proxy->AsyncReadbackCallback_RenderThread = [ThisWeak=TWeakObjectPtr<UPCGDataCollectionDataProvider>(this), /*InitStack=InitStack, ProviderName,*/ GenerationCount](const void* InData, int InNumBytes)
 		{
 			TRACE_CPUPROFILER_EVENT_SCOPE(UPCGDataCollectionDataProvider::ProcessReadbackData_RenderThread);
 			TStrongObjectPtr<UPCGDataCollectionDataProvider> This = ThisWeak.Pin();
@@ -531,6 +542,11 @@ FComputeDataProviderRenderProxy* UPCGDataCollectionDataProvider::GetRenderProxy(
 				return;
 			}
 
+			if (This->GenerationCounter != GenerationCount)
+			{
+				return;
+			}
+
 			// We should never find ourselves stomping existing data.
 			check(This->ReadbackCounters.IsEmpty());
 
@@ -543,10 +559,10 @@ FComputeDataProviderRenderProxy* UPCGDataCollectionDataProvider::GetRenderProxy(
 
 			This->bReadbackComplete = true;
 
-			ExecuteOnGameThread(UE_SOURCE_LOCATION, [ThisWeak=ThisWeak, OutputPinLabelAlias=This->OutputPinLabelAlias]()
+			ExecuteOnGameThread(UE_SOURCE_LOCATION, [ThisWeak=ThisWeak, OutputPinLabelAlias=This->OutputPinLabelAlias, GenerationCount]()
 			{
 				TStrongObjectPtr<UPCGDataCollectionDataProvider> This = ThisWeak.Pin();
-				if (This)
+				if (This && This->GenerationCounter == GenerationCount)
 				{
 					if (This->Binding.IsValid())
 					{
@@ -562,6 +578,28 @@ FComputeDataProviderRenderProxy* UPCGDataCollectionDataProvider::GetRenderProxy(
 	return Proxy;
 }
 
+void UPCGDataCollectionDataProvider::Reset()
+{
+	// Bump counter so any async callbacks from usages of this provider will be invalidated.
+	++GenerationCounter;
+
+	OnDataExported = {};
+	OnCountersReadBack = {};
+	ProducerSettings = nullptr;
+	Binding.Reset();
+	GraphBindingIndex = INDEX_NONE;
+	ExportMode = EPCGExportMode::NoExport;
+	ElementCountMultiplier = 0;
+	bExposeElementCounters = false;
+	OutputPinLabel = NAME_None;
+	OutputPinLabelAlias = NAME_None;
+	PinDesc = {};
+	bReadbackComplete = false;
+	ReadbackCounters.Empty();
+
+	Super::Reset();
+}
+
 FPCGDataCollectionDataProviderProxy::FPCGDataCollectionDataProviderProxy(
 	TWeakObjectPtr<UPCGDataCollectionDataProvider> InDataProvider,
 	const FPCGDataCollectionDesc& InPinDesc,
@@ -578,6 +616,8 @@ FPCGDataCollectionDataProviderProxy::FPCGDataCollectionDataProviderProxy(
 	, OutputPinLabel(InOutputPinLabel)
 	, OutputPinLabelAlias(InOutputPinLabelAlias)
 {
+	OriginatingGenerationCount = InDataProvider->GenerationCounter;
+
 	SizeBytes = PinDesc.ComputePackedSizeBytes();
 }
 
@@ -682,14 +722,25 @@ void FPCGDataCollectionDataProviderProxy::AllocateResources(FRDGBuilder& GraphBu
 		const TRefCountPtr<FRDGPooledBuffer> ExportedPooledBuffer = GraphBuilder.ConvertToExternalBuffer(Buffer);
 		GraphBuilder.SetBufferAccessFinal(Buffer, ERHIAccess::CopySrc);
 
-		ExecuteOnGameThread(UE_SOURCE_LOCATION, [ExportedPooledBuffer, DataProviderWeakPtr=DataProviderWeakPtr, PinDesc=PinDesc, OutputPinLabel=OutputPinLabel, SizeBytes=SizeBytes, ExportMode=ExportMode, OutputPinLabelAlias=OutputPinLabelAlias]()
+		ExecuteOnGameThread(UE_SOURCE_LOCATION, [ExportedPooledBuffer, DataProviderWeakPtr=DataProviderWeakPtr, PinDesc=PinDesc, OutputPinLabel=OutputPinLabel, SizeBytes=SizeBytes, ExportMode=ExportMode, OutputPinLabelAlias=OutputPinLabelAlias, GenerationCount=OriginatingGenerationCount]()
 		{
 			// Obtain objects. No ensures added because a graph cancellation could feasibly destroy some or all of these.
 			UPCGDataCollectionDataProvider* DataProvider = DataProviderWeakPtr.Get();
+			if (!DataProvider)
+			{
+				UE_LOG(LogPCG, Error, TEXT("Could not resolve UPCGDataCollectionDataProvider object to pass back buffer handle."));
+				return;
+			}
+
+			if (DataProvider->GenerationCounter != GenerationCount)
+			{
+				return;
+			}
+
 			UPCGDataBinding* Binding = DataProvider ? DataProvider->Binding.Get() : nullptr;
 			UPCGSubsystem* Subsystem = UPCGSubsystem::GetSubsystemForCurrentWorld();
 
-			if (DataProvider && Binding && Subsystem)
+			if (Binding && Subsystem)
 			{
 				TSharedPtr<FPCGProxyForGPUDataCollection> DataCollectionOnGPU = MakeShared<FPCGProxyForGPUDataCollection>(ExportedPooledBuffer, SizeBytes, PinDesc, Binding->GetStringTable());
 
@@ -701,14 +752,7 @@ void FPCGDataCollectionDataProviderProxy::AllocateResources(FRDGBuilder& GraphBu
 					// TODO - binding is doing a lot of work. Could store a context handle in the data provider instead?
 					Binding->ReceiveDataFromGPU_GameThread(Proxy, DataProvider->ProducerSettings, ExportMode, OutputPinLabel, OutputPinLabelAlias);
 				}
-			}
-			else
-			{
-				UE_LOG(LogPCG, Error, TEXT("Could not resolve UPCGDataCollectionDataProvider object to pass back buffer handle."));
-			}
 
-			if (DataProvider)
-			{
 				DataProvider->OnDataExported_GameThread().Broadcast();
 			}
 		});
diff --git a/Engine/Plugins/PCG/Source/PCG/Private/Compute/DataInterfaces/PCGDataCollectionDataInterface.h b/Engine/Plugins/PCG/Source/PCG/Private/Compute/DataInterfaces/PCGDataCollectionDataInterface.h
--- a/Engine/Plugins/PCG/Source/PCG/Private/Compute/DataInterfaces/PCGDataCollectionDataInterface.h
+++ b/Engine/Plugins/PCG/Source/PCG/Private/Compute/DataInterfaces/PCGDataCollectionDataInterface.h
@@ -35,7 +35,7 @@ public:
 	void GetHLSL(FString& OutHLSL, FString const& InDataInterfaceName) const override;
 	// This data interface only uses the CF readback functionality for the element counters.
 	bool GetRequiresReadback() const override { return bExposeElementCounters; }
-	UComputeDataProvider* CreateDataProvider(TObjectPtr<UObject> InBinding, uint64 InInputMask, uint64 InOutputMask) const override;
+	UComputeDataProvider* CreateDataProvider() const override;
 	//~ End UComputeDataInterface Interface
 
 	void SetProducerSettings(const UPCGSettings* InProducerSettings) { ProducerSettings = InProducerSettings; }
@@ -46,7 +46,7 @@ public:
 
 	void SetExposeElementCounters(bool bInExposeElementCounters) { bExposeElementCounters = bInExposeElementCounters; }
 
-protected:
+public:
 	/** Settings of node that produces this data, normally the upstream node. */
 	UPROPERTY()
 	TObjectPtr<const UPCGSettings> ProducerSettings;
@@ -61,6 +61,7 @@ protected:
 	UPROPERTY()
 	bool bExposeElementCounters = false;
 
+protected:
 	static TCHAR const* TemplateFilePath;
 };
 
@@ -76,7 +77,9 @@ public:
 	//~ End UPCGComputeDataProvider Interface
 
 	//~ Begin UComputeDataProvider Interface
+	virtual void Initialize(const UComputeDataInterface* InDataInterface, UObject* InBinding, uint64 InInputMask, uint64 InOutputMask) override;
 	FComputeDataProviderRenderProxy* GetRenderProxy() override;
+	virtual void Reset() override;
 	//~ End UComputeDataProvider Interface
 
 	/** Returns true if the generated buffer should be exported to a persistent buffer to be passed to downstream nodes. */
@@ -113,6 +116,9 @@ public:
 	UPROPERTY()
 	FName OutputPinLabelAlias;
 
+	/** Bumped each time the data provider is initialized or reset, so that async callbacks can detect if they originated from a previous usage of the data provider and no-op. */
+	std::atomic<uint64> GenerationCounter = 0;
+
 protected:
 	FPCGDataCollectionDesc PinDesc;
 
@@ -153,6 +159,9 @@ public:
 public:
 	FReadbackCallback AsyncReadbackCallback_RenderThread;
 
+	/** Generation count of the data provider when the proxy was created. */
+	uint64 OriginatingGenerationCount = 0;
+
 protected:
 	uint32 SizeBytes = 0;
 
diff --git a/Engine/Plugins/PCG/Source/PCG/Private/Compute/DataInterfaces/PCGDataCollectionUploadDataInterface.cpp b/Engine/Plugins/PCG/Source/PCG/Private/Compute/DataInterfaces/PCGDataCollectionUploadDataInterface.cpp
--- a/Engine/Plugins/PCG/Source/PCG/Private/Compute/DataInterfaces/PCGDataCollectionUploadDataInterface.cpp
+++ b/Engine/Plugins/PCG/Source/PCG/Private/Compute/DataInterfaces/PCGDataCollectionUploadDataInterface.cpp
@@ -19,20 +19,23 @@
 #include "RenderGraphResources.h"
 #include "RenderGraphUtils.h"
 
-UComputeDataProvider* UPCGDataCollectionUploadDataInterface::CreateDataProvider(TObjectPtr<UObject> InBinding, uint64 InInputMask, uint64 InOutputMask) const
+UComputeDataProvider* UPCGDataCollectionUploadDataInterface::CreateDataProvider() const
 {
-	TRACE_CPUPROFILER_EVENT_SCOPE(UPCGDataCollectionUploadDataInterface::CreateDataProvider);
-	UPCGDataBinding* Binding = CastChecked<UPCGDataBinding>(InBinding);
+	return NewObject<UPCGDataProviderDataCollectionUpload>();
+}
 
-	UPCGDataProviderDataCollectionUpload* Provider = NewObject<UPCGDataProviderDataCollectionUpload>();
-	Provider->Binding = Binding;
-	Provider->GraphBindingIndex = GraphBindingIndex;
+void UPCGDataProviderDataCollectionUpload::Initialize(const UComputeDataInterface* InDataInterface, UObject* InBinding, uint64 InInputMask, uint64 InOutputMask)
+{
+	TRACE_CPUPROFILER_EVENT_SCOPE(UPCGDataProviderDataCollectionUpload::Initialize);
 
-	// Pick the data items from input data collection using any of the compute graph element virtual input pin labels.
-	check(!DownstreamInputPinLabelAliases.IsEmpty());
-	Provider->DownstreamInputPinLabels = DownstreamInputPinLabelAliases;
+	const UPCGDataCollectionUploadDataInterface* DataInterface = CastChecked<UPCGDataCollectionUploadDataInterface>(InDataInterface);
 
-	return Provider;
+	Binding = CastChecked<UPCGDataBinding>(InBinding);
+	GraphBindingIndex = DataInterface->GraphBindingIndex;
+
+	// Pick the data items from input data collection using any of the compute graph element virtual input pin labels.
+	check(!DataInterface->DownstreamInputPinLabelAliases.IsEmpty());
+	DownstreamInputPinLabels = DataInterface->DownstreamInputPinLabelAliases;
 }
 
 bool UPCGDataProviderDataCollectionUpload::PrepareForExecute_GameThread(UPCGDataBinding* InBinding)
@@ -89,6 +92,14 @@ FComputeDataProviderRenderProxy* UPCGDataProviderDataCollectionUpload::GetRender
 	return Proxy;
 }
 
+void UPCGDataProviderDataCollectionUpload::Reset()
+{
+	DownstreamInputPinLabels.Empty();
+	DataAdaptor.Reset();
+
+	Super::Reset();
+}
+
 FPCGDataProviderDataCollectionUploadProxy::FPCGDataProviderDataCollectionUploadProxy(
 	TSharedPtr<FPCGDataUploadAdaptor> InDataAdaptor,
 	TWeakObjectPtr<UPCGDataCollectionDataProvider> InDataProvider,
diff --git a/Engine/Plugins/PCG/Source/PCG/Private/Compute/DataInterfaces/PCGDataCollectionUploadDataInterface.h b/Engine/Plugins/PCG/Source/PCG/Private/Compute/DataInterfaces/PCGDataCollectionUploadDataInterface.h
--- a/Engine/Plugins/PCG/Source/PCG/Private/Compute/DataInterfaces/PCGDataCollectionUploadDataInterface.h
+++ b/Engine/Plugins/PCG/Source/PCG/Private/Compute/DataInterfaces/PCGDataCollectionUploadDataInterface.h
@@ -18,7 +18,7 @@ class UPCGDataCollectionUploadDataInterface : public UPCGDataCollectionDataInter
 
 public:
 	//~ Begin UComputeDataInterface Interface
-	UComputeDataProvider* CreateDataProvider(TObjectPtr<UObject> InBinding, uint64 InInputMask, uint64 InOutputMask) const override;
+	UComputeDataProvider* CreateDataProvider() const override;
 	//~ End UComputeDataInterface Interface
 };
 
@@ -32,7 +32,9 @@ public:
 	virtual void ReleaseTransientResources(const TCHAR* InReason = nullptr) override;
 
 	//~ Begin UComputeDataProvider Interface
+	virtual void Initialize(const UComputeDataInterface* InDataInterface, UObject* InBinding, uint64 InInputMask, uint64 InOutputMask) override;
 	FComputeDataProviderRenderProxy* GetRenderProxy() override;
+	virtual void Reset() override;
 	//~ End UComputeDataProvider Interface
 
 	//~ Begin UPCGComputeDataProvider Interface
diff --git a/Engine/Plugins/PCG/Source/PCG/Private/Compute/DataInterfaces/PCGDataLabelResolverDataInterface.cpp b/Engine/Plugins/PCG/Source/PCG/Private/Compute/DataInterfaces/PCGDataLabelResolverDataInterface.cpp
--- a/Engine/Plugins/PCG/Source/PCG/Private/Compute/DataInterfaces/PCGDataLabelResolverDataInterface.cpp
+++ b/Engine/Plugins/PCG/Source/PCG/Private/Compute/DataInterfaces/PCGDataLabelResolverDataInterface.cpp
@@ -68,16 +68,22 @@ void UPCGDataLabelResolverDataInterface::GetHLSL(FString& OutHLSL, FString const
 		TemplateArgs);
 }
 
-UComputeDataProvider* UPCGDataLabelResolverDataInterface::CreateDataProvider(TObjectPtr<UObject> InBinding, uint64 InInputMask, uint64 InOutputMask) const
+UComputeDataProvider* UPCGDataLabelResolverDataInterface::CreateDataProvider() const
 {
-	TRACE_CPUPROFILER_EVENT_SCOPE(UPCGDataLabelResolverDataInterface::CreateDataProvider);
+	return NewObject<UPCGDataLabelResolverDataProvider>();
+}
+
+void UPCGDataLabelResolverDataProvider::Initialize(const UComputeDataInterface* InDataInterface, UObject* InBinding, uint64 InInputMask, uint64 InOutputMask)
+{
+	TRACE_CPUPROFILER_EVENT_SCOPE(UPCGDataLabelResolverDataProvider::Initialize);
+
+	Super::Initialize(InDataInterface, InBinding, InInputMask, InOutputMask);
 
-	UPCGDataLabelResolverDataProvider* DataProvider = NewObject<UPCGDataLabelResolverDataProvider>();
-	DataProvider->Kernel = Kernel;
-	DataProvider->PinLabel = PinLabel;
-	DataProvider->bIsInput = bIsInput;
+	const UPCGDataLabelResolverDataInterface* DataInterface = CastChecked<UPCGDataLabelResolverDataInterface>(InDataInterface);
 
-	return DataProvider;
+	Kernel = DataInterface->Kernel;
+	PinLabel = DataInterface->PinLabel;
+	bIsInput = DataInterface->bIsInput;
 }
 
 bool UPCGDataLabelResolverDataProvider::PrepareForExecute_GameThread(UPCGDataBinding* InBinding)
@@ -139,6 +145,16 @@ FComputeDataProviderRenderProxy* UPCGDataLabelResolverDataProvider::GetRenderPro
 	return new FPCGDataLabelResolverDataProviderProxy(DataIdToDataIndexMap);
 }
 
+void UPCGDataLabelResolverDataProvider::Reset()
+{
+	Kernel = nullptr;
+	PinLabel = NAME_None;
+	bIsInput = false;
+	DataIdToDataIndexMap.Empty();
+
+	Super::Reset();
+}
+
 bool FPCGDataLabelResolverDataProviderProxy::IsValid(FValidationData const& InValidationData) const
 {
 	return InValidationData.ParameterStructSize == sizeof(FParameters);
diff --git a/Engine/Plugins/PCG/Source/PCG/Private/Compute/DataInterfaces/PCGDataLabelResolverDataInterface.h b/Engine/Plugins/PCG/Source/PCG/Private/Compute/DataInterfaces/PCGDataLabelResolverDataInterface.h
--- a/Engine/Plugins/PCG/Source/PCG/Private/Compute/DataInterfaces/PCGDataLabelResolverDataInterface.h
+++ b/Engine/Plugins/PCG/Source/PCG/Private/Compute/DataInterfaces/PCGDataLabelResolverDataInterface.h
@@ -25,7 +25,7 @@ public:
 	void GetSupportedOutputs(TArray<FShaderFunctionDefinition>& OutFunctions) const override;
 	void GetShaderParameters(TCHAR const* UID, FShaderParametersMetadataBuilder& InOutBuilder, FShaderParametersMetadataAllocations& InOutAllocations) const override;
 	void GetHLSL(FString& OutHLSL, FString const& InDataInterfaceName) const override;
-	UComputeDataProvider* CreateDataProvider(TObjectPtr<UObject> InBinding, uint64 InInputMask, uint64 InOutputMask) const override;
+	UComputeDataProvider* CreateDataProvider() const override;
 	//~ End UComputeDataInterface Interface
 
 	/** Label of the pin being resolved. */
@@ -52,7 +52,9 @@ public:
 	//~ End UPCGComputeDataProvider Interface
 
 	//~ Begin UComputeDataProvider Interface
+	virtual void Initialize(const UComputeDataInterface* InDataInterface, UObject* InBinding, uint64 InInputMask, uint64 InOutputMask) override;
 	FComputeDataProviderRenderProxy* GetRenderProxy() override;
+	virtual void Reset() override;
 	//~ End UComputeDataProvider Interface
 
 public:
diff --git a/Engine/Plugins/PCG/Source/PCG/Private/Compute/DataInterfaces/PCGDebugDataInterface.cpp b/Engine/Plugins/PCG/Source/PCG/Private/Compute/DataInterfaces/PCGDebugDataInterface.cpp
--- a/Engine/Plugins/PCG/Source/PCG/Private/Compute/DataInterfaces/PCGDebugDataInterface.cpp
+++ b/Engine/Plugins/PCG/Source/PCG/Private/Compute/DataInterfaces/PCGDebugDataInterface.cpp
@@ -61,39 +61,46 @@ bool UPCGDebugDataInterface::GetRequiresReadback() const
 	return DebugBufferSizeFloats > 0;
 }
 
-UComputeDataProvider* UPCGDebugDataInterface::CreateDataProvider(TObjectPtr<UObject> InBinding, uint64 InInputMask, uint64 InOutputMask) const
+UComputeDataProvider* UPCGDebugDataInterface::CreateDataProvider() const
 {
-	UPCGDebugDataProvider* Provider = NewObject<UPCGDebugDataProvider>();
-	Provider->DebugBufferSizeFloats = DebugBufferSizeFloats;
-	return Provider;
+	return NewObject<UPCGDebugDataProvider>();
 }
 
-FComputeDataProviderRenderProxy* UPCGDebugDataProvider::GetRenderProxy()
+void UPCGDebugDataProvider::Initialize(const UComputeDataInterface* InDataInterface, UObject* InBinding, uint64 InInputMask, uint64 InOutputMask)
 {
-	TWeakObjectPtr<UPCGDebugDataProvider> ThisWeakPtr(this);
+	TRACE_CPUPROFILER_EVENT_SCOPE(UPCGDebugDataProvider::Initialize);
+
+	Super::Initialize(InDataInterface, InBinding, InInputMask, InOutputMask);
 
-	auto ProcessReadbackData_RenderThread = [ThisWeakPtr](const void* InData, int InNumBytes)
+	const UPCGDebugDataInterface* DataInterface = CastChecked<UPCGDebugDataInterface>(InDataInterface);
+
+	DebugBufferSizeFloats = DataInterface->DebugBufferSizeFloats;
+}
+
+FComputeDataProviderRenderProxy* UPCGDebugDataProvider::GetRenderProxy()
+{
+	auto ProcessReadbackData_RenderThread = [](const void* InData, int InNumBytes)
 	{
-		if (!ThisWeakPtr.IsValid() || !ThisWeakPtr.Get() || !InData)
+		if (!InData)
 		{
 			return;
 		}
 
-		check(ThisWeakPtr->DebugBufferSizeFloats * sizeof(float) == InNumBytes);
-
-		UPCGDebugDataProvider* DataProvider = ThisWeakPtr.Get();
-		
 		const float* DataAsFloat = static_cast<const float*>(InData);
-
-		UE_LOG(LogPCG, Warning, TEXT("Logging kernel debug..."));
-
 		for (int Index = 0; Index < InNumBytes / sizeof(float); ++Index)
 		{
 			UE_LOG(LogPCG, Warning, TEXT("\t%d:\t%f"), Index, DataAsFloat[Index]);
 		}
 	};
 
-	return new FPCGDebugDataProviderProxy(DebugBufferSizeFloats, ProcessReadbackData_RenderThread);
+	return new FPCGDebugDataProviderProxy(DebugBufferSizeFloats, MoveTemp(ProcessReadbackData_RenderThread));
+}
+
+void UPCGDebugDataProvider::Reset()
+{
+	DebugBufferSizeFloats = 0;
+
+	Super::Reset();
 }
 
 bool FPCGDebugDataProviderProxy::IsValid(FValidationData const& InValidationData) const
diff --git a/Engine/Plugins/PCG/Source/PCG/Private/Compute/DataInterfaces/PCGDebugDataInterface.h b/Engine/Plugins/PCG/Source/PCG/Private/Compute/DataInterfaces/PCGDebugDataInterface.h
--- a/Engine/Plugins/PCG/Source/PCG/Private/Compute/DataInterfaces/PCGDebugDataInterface.h
+++ b/Engine/Plugins/PCG/Source/PCG/Private/Compute/DataInterfaces/PCGDebugDataInterface.h
@@ -25,12 +25,12 @@ public:
 	void GetShaderHash(FString& InOutKey) const override;
 	void GetHLSL(FString& OutHLSL, FString const& InDataInterfaceName) const override;
 	bool GetRequiresReadback() const override;
-	UComputeDataProvider* CreateDataProvider(TObjectPtr<UObject> InBinding, uint64 InInputMask, uint64 InOutputMask) const override;
+	UComputeDataProvider* CreateDataProvider() const override;
 	//~ End UComputeDataInterface Interface
 
 	void SetDebugBufferSize(uint32 InDebugBufferSizeFloats) { DebugBufferSizeFloats = InDebugBufferSizeFloats; }
 
-private:
+public:
 	uint32 DebugBufferSizeFloats = 0;
 };
 
@@ -42,9 +42,12 @@ class UPCGDebugDataProvider : public UComputeDataProvider
 
 public:
 	//~ Begin UComputeDataProvider Interface
+	virtual void Initialize(const UComputeDataInterface* InDataInterface, UObject* InBinding, uint64 InInputMask, uint64 InOutputMask) override;
 	FComputeDataProviderRenderProxy* GetRenderProxy() override;
+	virtual void Reset() override;
 	//~ End UComputeDataProvider Interface
 
+public:
 	uint32 DebugBufferSizeFloats = 0;
 };
 
diff --git a/Engine/Plugins/PCG/Source/PCG/Private/Compute/DataInterfaces/PCGInstanceDataInterface.cpp b/Engine/Plugins/PCG/Source/PCG/Private/Compute/DataInterfaces/PCGInstanceDataInterface.cpp
--- a/Engine/Plugins/PCG/Source/PCG/Private/Compute/DataInterfaces/PCGInstanceDataInterface.cpp
+++ b/Engine/Plugins/PCG/Source/PCG/Private/Compute/DataInterfaces/PCGInstanceDataInterface.cpp
@@ -126,16 +126,23 @@ void UPCGInstanceDataInterface::GetHLSL(FString& OutHLSL, FString const& InDataI
 	}
 }
 
-UComputeDataProvider* UPCGInstanceDataInterface::CreateDataProvider(TObjectPtr<UObject> InBinding, uint64 InInputMask, uint64 InOutputMask) const
+UComputeDataProvider* UPCGInstanceDataInterface::CreateDataProvider() const
 {
-	TRACE_CPUPROFILER_EVENT_SCOPE(UPCGInstanceDataInterface::CreateDataProvider);
-	UPCGDataBinding* Binding = CastChecked<UPCGDataBinding>(InBinding);
-	check(ProducerKernel);
+	return NewObject<UPCGInstanceDataProvider>();
+}
+
+void UPCGInstanceDataProvider::Initialize(const UComputeDataInterface* InDataInterface, UObject* InBinding, uint64 InInputMask, uint64 InOutputMask)
+{
+	TRACE_CPUPROFILER_EVENT_SCOPE(UPCGInstanceDataProvider::Initialize);
+
+	Super::Initialize(InDataInterface, InBinding, InInputMask, InOutputMask);
 
-	TObjectPtr<UPCGInstanceDataProvider> DataProvider = NewObject<UPCGInstanceDataProvider>();
+	const UPCGInstanceDataInterface* DataInterface = CastChecked<UPCGInstanceDataInterface>(InDataInterface);
+	UPCGDataBinding* Binding = CastChecked<UPCGDataBinding>(InBinding);
+	check(DataInterface->ProducerKernel);
 
 	// If there were 0 input points (or too many input points) for this execution, we will not have created any primitives, so check for null.
-	if (FPCGSpawnerPrimitives* FoundPrimitives = Binding->MeshSpawnersToPrimitives.Find(Cast<UPCGStaticMeshSpawnerKernel>(ProducerKernel)))
+	if (FPCGSpawnerPrimitives* FoundPrimitives = Binding->MeshSpawnersToPrimitives.Find(Cast<UPCGStaticMeshSpawnerKernel>(DataInterface->ProducerKernel)))
 	{
 		if (!ensure(FoundPrimitives->Primitives.Num() <= PCGComputeConstants::MAX_PRIMITIVE_COMPONENTS_PER_SPAWNER))
 		{
@@ -143,12 +150,10 @@ UComputeDataProvider* UPCGInstanceDataInterface::CreateDataProvider(TObjectPtr<U
 			FoundPrimitives->Primitives.SetNum(PCGComputeConstants::MAX_PRIMITIVE_COMPONENTS_PER_SPAWNER);
 		}
 
-		DataProvider->Primitives = FoundPrimitives->Primitives;
-		DataProvider->NumInstancesAllPrimitives = FoundPrimitives->NumInstancesAllPrimitives;
-		DataProvider->NumCustomFloatsPerInstance = FoundPrimitives->NumCustomFloats;
+		Primitives = FoundPrimitives->Primitives;
+		NumInstancesAllPrimitives = FoundPrimitives->NumInstancesAllPrimitives;
+		NumCustomFloatsPerInstance = FoundPrimitives->NumCustomFloats;
 	}
-
-	return DataProvider;
 }
 
 FComputeDataProviderRenderProxy* UPCGInstanceDataProvider::GetRenderProxy()
@@ -170,6 +175,15 @@ FComputeDataProviderRenderProxy* UPCGInstanceDataProvider::GetRenderProxy()
 	return new FPCGInstanceDataProviderProxy(MoveTemp(PrimitiveProxies), NumInstancesAllPrimitives, NumCustomFloatsPerInstance);
 }
 
+void UPCGInstanceDataProvider::Reset()
+{
+	Primitives.Empty();
+	NumInstancesAllPrimitives = 0;
+	NumCustomFloatsPerInstance = 0;
+
+	Super::Reset();
+}
+
 FPCGInstanceDataProviderProxy::FPCGInstanceDataProviderProxy(TArray<FPrimitiveSceneProxy*>&& InPrimitiveSceneProxies, uint32 InNumInstancesAllPrimitives, uint32 InNumCustomFloatsPerInstance)
 	: PrimitiveSceneProxies(MoveTemp(InPrimitiveSceneProxies))
 	, NumInstancesAllPrimitives(InNumInstancesAllPrimitives)
diff --git a/Engine/Plugins/PCG/Source/PCG/Private/Compute/DataInterfaces/PCGInstanceDataInterface.h b/Engine/Plugins/PCG/Source/PCG/Private/Compute/DataInterfaces/PCGInstanceDataInterface.h
--- a/Engine/Plugins/PCG/Source/PCG/Private/Compute/DataInterfaces/PCGInstanceDataInterface.h
+++ b/Engine/Plugins/PCG/Source/PCG/Private/Compute/DataInterfaces/PCGInstanceDataInterface.h
@@ -28,7 +28,7 @@ public:
 	void GetShaderHash(FString& InOutKey) const override;
 	void GetHLSL(FString& OutHLSL, FString const& InDataInterfaceName) const override;
 	bool GetRequiresPostSubmitCall() const override { return true; }
-	UComputeDataProvider* CreateDataProvider(TObjectPtr<UObject> InBinding, uint64 InInputMask, uint64 InOutputMask) const override;
+	UComputeDataProvider* CreateDataProvider() const override;
 	//~ End UComputeDataInterface Interface
 
 public:
@@ -47,7 +47,9 @@ class UPCGInstanceDataProvider : public UComputeDataProvider
 
 public:
 	//~ Begin UComputeDataProvider Interface
+	virtual void Initialize(const UComputeDataInterface* InDataInterface, UObject* InBinding, uint64 InInputMask, uint64 InOutputMask) override;
 	FComputeDataProviderRenderProxy* GetRenderProxy() override;
+	virtual void Reset() override;
 	//~ End UComputeDataProvider Interface
 
 	UPROPERTY()
diff --git a/Engine/Plugins/PCG/Source/PCG/Private/Compute/DataInterfaces/PCGLandscapeDataInterface.cpp b/Engine/Plugins/PCG/Source/PCG/Private/Compute/DataInterfaces/PCGLandscapeDataInterface.cpp
--- a/Engine/Plugins/PCG/Source/PCG/Private/Compute/DataInterfaces/PCGLandscapeDataInterface.cpp
+++ b/Engine/Plugins/PCG/Source/PCG/Private/Compute/DataInterfaces/PCGLandscapeDataInterface.cpp
@@ -70,35 +70,6 @@ private:
 	TArray<float> HeightValues;
 };
 
-/** Manages the resources created by this DI. */
-class FPCGLandscapeResource
-{
-public:
-	struct FResourceKey
-	{
-		TWeakObjectPtr<const ALandscape> Source = nullptr;
-		TArray<FIntPoint> CapturedRegions;
-		FIntPoint MinCaptureRegion = FIntPoint(ForceInitToZero);
-		FIntPoint MaxCaptureRegion = FIntPoint(ForceInitToZero);
-		bool bIncludesHeight = false;
-	};
-
-	FPCGLandscapeResource() {}
-	FPCGLandscapeResource(const FResourceKey& InKey);
-	~FPCGLandscapeResource();
-
-	FPCGLandscapeTextureResource* LandscapeTextures = nullptr;
-	FVector3f LandscapeLWCTile = FVector3f::ZeroVector;
-	FMatrix ActorToWorldTransform = FMatrix::Identity;
-	FMatrix WorldToActorTransform = FMatrix::Identity;
-	FVector4 UvScaleBias = FVector4(1.0f, 1.0f, 0.0f, 0.0f);
-	FIntPoint CellCount = FIntPoint(ForceInitToZero);
-	FVector2D TextureWorldGridSize = FVector2D(1.0f, 1.0f);
-
-private:
-	FResourceKey ResourceKey;
-};
-
 FPCGLandscapeResource::FPCGLandscapeResource(const FResourceKey& InKey)
 	: ResourceKey(InKey)
 {
@@ -240,9 +211,18 @@ void UPCGLandscapeDataInterface::GetHLSL(FString& OutHLSL, FString const& InData
 	}
 }
 
-UComputeDataProvider* UPCGLandscapeDataInterface::CreateDataProvider(TObjectPtr<UObject> InBinding, uint64 InInputMask, uint64 InOutputMask) const
+UComputeDataProvider* UPCGLandscapeDataInterface::CreateDataProvider() const
 {
-	TRACE_CPUPROFILER_EVENT_SCOPE(UPCGLandscapeDataInterface::CreateDataProvider);
+	return NewObject<UPCGLandscapeDataProvider>();
+}
+
+void UPCGLandscapeDataProvider::Initialize(const UComputeDataInterface* InDataInterface, UObject* InBinding, uint64 InInputMask, uint64 InOutputMask)
+{
+	TRACE_CPUPROFILER_EVENT_SCOPE(UPCGLandscapeDataProvider::Initialize);
+
+	Super::Initialize(InDataInterface, InBinding, InInputMask, InOutputMask);
+
+	const UPCGLandscapeDataInterface* DataInterface = CastChecked<UPCGLandscapeDataInterface>(InDataInterface);
 	UPCGDataBinding* Binding = CastChecked<UPCGDataBinding>(InBinding);
 	const UPCGComponent* SourceComponent = Binding->GetSourceComponent().Get();
 	check(SourceComponent);
@@ -269,11 +249,9 @@ UComputeDataProvider* UPCGLandscapeDataInterface::CreateDataProvider(TObjectPtr<
 		return false;
 	};
 
-	TObjectPtr<UPCGLandscapeDataProvider> DataProvider = NewObject<UPCGLandscapeDataProvider>();
-
 	// Take any input pin label alias to obtain the data from the input data collection.
-	check(!DownstreamInputPinLabelAliases.IsEmpty());
-	const TArray<FPCGTaggedData> TaggedDatas = Binding->InputDataCollection.GetInputsByPin(DownstreamInputPinLabelAliases[0]);
+	check(!DataInterface->DownstreamInputPinLabelAliases.IsEmpty());
+	const TArray<FPCGTaggedData> TaggedDatas = Binding->InputDataCollection.GetInputsByPin(DataInterface->DownstreamInputPinLabelAliases[0]);
 
 	if (!TaggedDatas.IsEmpty())
 	{
@@ -290,7 +268,7 @@ UComputeDataProvider* UPCGLandscapeDataInterface::CreateDataProvider(TObjectPtr<
 					{
 						if (TestLandscape(Landscape))
 						{
-							DataProvider->Init(Landscape, ComponentBounds);
+							InitFromLandscape(Landscape, ComponentBounds);
 							break;
 						}
 					}
@@ -298,13 +276,11 @@ UComputeDataProvider* UPCGLandscapeDataInterface::CreateDataProvider(TObjectPtr<
 			}
 		}
 	}
-
-	return DataProvider;
 }
 
-void UPCGLandscapeDataProvider::Init(ALandscape* InLandscape, const FBox& Bounds)
+void UPCGLandscapeDataProvider::InitFromLandscape(ALandscape* InLandscape, const FBox& Bounds)
 {
-	TRACE_CPUPROFILER_EVENT_SCOPE(UPCGLandscapeDataProvider::Initialize);
+	TRACE_CPUPROFILER_EVENT_SCOPE(UPCGLandscapeDataProvider::Init);
 	check(IsInGameThread());
 
 	const ULandscapeInfo* LandscapeInfo = InLandscape ? InLandscape->GetLandscapeInfo() : nullptr;
@@ -397,13 +373,19 @@ void UPCGLandscapeDataProvider::Init(ALandscape* InLandscape, const FBox& Bounds
 		}
 	}
 
-	Resource = MakeUnique<FPCGLandscapeResource>(Key);
+	Resource = MakeShared<FPCGLandscapeResource>(Key);
 }
 
 FComputeDataProviderRenderProxy* UPCGLandscapeDataProvider::GetRenderProxy()
 {
-	FPCGLandscapeDataProviderProxy* Proxy = new FPCGLandscapeDataProviderProxy(Resource.Get());
-	return Proxy;
+	return new FPCGLandscapeDataProviderProxy(Resource);
+}
+
+void UPCGLandscapeDataProvider::Reset()
+{
+	Resource.Reset();
+
+	Super::Reset();
 }
 
 bool FPCGLandscapeDataProviderProxy::IsValid(FValidationData const& InValidationData) const
diff --git a/Engine/Plugins/PCG/Source/PCG/Private/Compute/DataInterfaces/PCGLandscapeDataInterface.h b/Engine/Plugins/PCG/Source/PCG/Private/Compute/DataInterfaces/PCGLandscapeDataInterface.h
--- a/Engine/Plugins/PCG/Source/PCG/Private/Compute/DataInterfaces/PCGLandscapeDataInterface.h
+++ b/Engine/Plugins/PCG/Source/PCG/Private/Compute/DataInterfaces/PCGLandscapeDataInterface.h
@@ -9,7 +9,36 @@
 
 class ALandscape;
 class FPCGLandscapeDataInterfaceParameters;
-class FPCGLandscapeResource;
+class FPCGLandscapeTextureResource;
+
+/** Manages the resources created by this DI. */
+class FPCGLandscapeResource
+{
+public:
+	struct FResourceKey
+	{
+		TWeakObjectPtr<const ALandscape> Source = nullptr;
+		TArray<FIntPoint> CapturedRegions;
+		FIntPoint MinCaptureRegion = FIntPoint(ForceInitToZero);
+		FIntPoint MaxCaptureRegion = FIntPoint(ForceInitToZero);
+		bool bIncludesHeight = false;
+	};
+
+	FPCGLandscapeResource() = default;
+	FPCGLandscapeResource(const FResourceKey& InKey);
+	~FPCGLandscapeResource();
+
+	FPCGLandscapeTextureResource* LandscapeTextures = nullptr;
+	FVector3f LandscapeLWCTile = FVector3f::ZeroVector;
+	FMatrix ActorToWorldTransform = FMatrix::Identity;
+	FMatrix WorldToActorTransform = FMatrix::Identity;
+	FVector4 UvScaleBias = FVector4(1.0f, 1.0f, 0.0f, 0.0f);
+	FIntPoint CellCount = FIntPoint(ForceInitToZero);
+	FVector2D TextureWorldGridSize = FVector2D(1.0f, 1.0f);
+
+private:
+	FResourceKey ResourceKey;
+};
 
 /** Data Interface allowing sampling of a Landscape. */
 UCLASS(ClassGroup = (Procedural))
@@ -25,7 +54,7 @@ public:
 	TCHAR const* GetShaderVirtualPath() const override;
 	void GetShaderHash(FString& InOutKey) const override;
 	void GetHLSL(FString& OutHLSL, FString const& InDataInterfaceName) const override;
-	UComputeDataProvider* CreateDataProvider(TObjectPtr<UObject> InBinding, uint64 InInputMask, uint64 InOutputMask) const override;
+	UComputeDataProvider* CreateDataProvider() const override;
 	//~ End UComputeDataInterface Interface
 
 private:
@@ -40,19 +69,22 @@ class UPCGLandscapeDataProvider : public UComputeDataProvider
 
 public:
 	void Init(ALandscape* InLandscape, const FBox& Bounds);
+	void InitFromLandscape(ALandscape* InLandscape, const FBox& Bounds);
 
 	//~ Begin UComputeDataProvider Interface
+	virtual void Initialize(const UComputeDataInterface* InDataInterface, UObject* InBinding, uint64 InInputMask, uint64 InOutputMask) override;
 	FComputeDataProviderRenderProxy* GetRenderProxy() override;
+	virtual void Reset() override;
 	//~ End UComputeDataProvider Interface
 
 private:
-	TUniquePtr<FPCGLandscapeResource> Resource;
+	TSharedPtr<FPCGLandscapeResource> Resource;
 };
 
 class FPCGLandscapeDataProviderProxy : public FComputeDataProviderRenderProxy
 {
 public:
-	FPCGLandscapeDataProviderProxy(const FPCGLandscapeResource* InResource)
+	FPCGLandscapeDataProviderProxy(TSharedPtr<const FPCGLandscapeResource> InResource)
 		: Resource(InResource)
 	{}
 
@@ -65,5 +97,5 @@ public:
 protected:
 	using FParameters = FPCGLandscapeDataInterfaceParameters;
 
-	const FPCGLandscapeResource* Resource = nullptr;
+	TSharedPtr<const FPCGLandscapeResource> Resource;
 };
diff --git a/Engine/Plugins/PCG/Source/PCG/Private/Compute/DataInterfaces/PCGStaticMeshDataInterface.cpp b/Engine/Plugins/PCG/Source/PCG/Private/Compute/DataInterfaces/PCGStaticMeshDataInterface.cpp
--- a/Engine/Plugins/PCG/Source/PCG/Private/Compute/DataInterfaces/PCGStaticMeshDataInterface.cpp
+++ b/Engine/Plugins/PCG/Source/PCG/Private/Compute/DataInterfaces/PCGStaticMeshDataInterface.cpp
@@ -147,17 +147,9 @@ void UPCGStaticMeshDataInterface::GetHLSL(FString& OutHLSL, FString const& InDat
 	}
 }
 
-UComputeDataProvider* UPCGStaticMeshDataInterface::CreateDataProvider(TObjectPtr<UObject> InBinding, uint64 InInputMask, uint64 InOutputMask) const
+UComputeDataProvider* UPCGStaticMeshDataInterface::CreateDataProvider() const
 {
-	TRACE_CPUPROFILER_EVENT_SCOPE(UPCGStaticMeshDataInterface::CreateDataProvider);
-
-	TObjectPtr<UPCGStaticMeshDataProvider> DataProvider = NewObject<UPCGStaticMeshDataProvider>();
-
-	// Take any input pin label alias to obtain the data from the input data collection.
-	check(!DownstreamInputPinLabelAliases.IsEmpty());
-	DataProvider->DownstreamInputPinLabels = DownstreamInputPinLabelAliases;
-
-	return DataProvider;
+	return NewObject<UPCGStaticMeshDataProvider>();;
 }
 
 bool UPCGStaticMeshDataProvider::PrepareForExecute_GameThread(UPCGDataBinding* InBinding)
@@ -198,11 +190,33 @@ bool UPCGStaticMeshDataProvider::PrepareForExecute_GameThread(UPCGDataBinding* I
 	return false;
 }
 
+void UPCGStaticMeshDataProvider::Initialize(const UComputeDataInterface* InDataInterface, UObject* InBinding, uint64 InInputMask, uint64 InOutputMask)
+{
+	TRACE_CPUPROFILER_EVENT_SCOPE(UPCGStaticMeshDataProvider::Initialize);
+
+	Super::Initialize(InDataInterface, InBinding, InInputMask, InOutputMask);
+
+	const UPCGStaticMeshDataInterface* DataInterface = CastChecked<UPCGStaticMeshDataInterface>(InDataInterface);
+
+	// Take any input pin label alias to obtain the data from the input data collection.
+	check(!DataInterface->DownstreamInputPinLabelAliases.IsEmpty());
+	DownstreamInputPinLabels = DataInterface->DownstreamInputPinLabelAliases;
+}
+
 FComputeDataProviderRenderProxy* UPCGStaticMeshDataProvider::GetRenderProxy()
 {
 	return new FPCGStaticMeshDataProviderProxy(LoadedStaticMesh.Get());
 }
 
+void UPCGStaticMeshDataProvider::Reset()
+{
+	LoadedStaticMesh = nullptr;
+	DownstreamInputPinLabels.Empty();
+	LoadHandle.Reset();
+
+	Super::Reset();
+}
+
 FPCGStaticMeshDataProviderProxy::FPCGStaticMeshDataProviderProxy(const UStaticMesh* InStaticMesh)
 {
 	const FStaticMeshRenderData* RenderData = InStaticMesh ? InStaticMesh->GetRenderData() : nullptr;
diff --git a/Engine/Plugins/PCG/Source/PCG/Private/Compute/DataInterfaces/PCGStaticMeshDataInterface.h b/Engine/Plugins/PCG/Source/PCG/Private/Compute/DataInterfaces/PCGStaticMeshDataInterface.h
--- a/Engine/Plugins/PCG/Source/PCG/Private/Compute/DataInterfaces/PCGStaticMeshDataInterface.h
+++ b/Engine/Plugins/PCG/Source/PCG/Private/Compute/DataInterfaces/PCGStaticMeshDataInterface.h
@@ -29,7 +29,7 @@ public:
 	TCHAR const* GetShaderVirtualPath() const override;
 	void GetShaderHash(FString& InOutKey) const override;
 	void GetHLSL(FString& OutHLSL, FString const& InDataInterfaceName) const override;
-	UComputeDataProvider* CreateDataProvider(TObjectPtr<UObject> InBinding, uint64 InInputMask, uint64 InOutputMask) const override;
+	UComputeDataProvider* CreateDataProvider() const override;
 	//~ End UComputeDataInterface Interface
 
 private:
@@ -44,7 +44,9 @@ class UPCGStaticMeshDataProvider : public UPCGComputeDataProvider
 
 public:
 	//~ Begin UComputeDataProvider Interface
+	virtual void Initialize(const UComputeDataInterface* InDataInterface, UObject* InBinding, uint64 InInputMask, uint64 InOutputMask) override;
 	FComputeDataProviderRenderProxy* GetRenderProxy() override;
+	virtual void Reset() override;
 	//~ End UComputeDataProvider Interface
 	
 	//~ Begin UPCGComputeDataProvider Interface
diff --git a/Engine/Plugins/PCG/Source/PCG/Private/Compute/DataInterfaces/PCGTextureDataInterface.cpp b/Engine/Plugins/PCG/Source/PCG/Private/Compute/DataInterfaces/PCGTextureDataInterface.cpp
--- a/Engine/Plugins/PCG/Source/PCG/Private/Compute/DataInterfaces/PCGTextureDataInterface.cpp
+++ b/Engine/Plugins/PCG/Source/PCG/Private/Compute/DataInterfaces/PCGTextureDataInterface.cpp
@@ -93,27 +93,34 @@ void UPCGTextureDataInterface::GetHLSL(FString& OutHLSL, FString const& InDataIn
 	}
 }
 
-UComputeDataProvider* UPCGTextureDataInterface::CreateDataProvider(TObjectPtr<UObject> InBinding, uint64 InInputMask, uint64 InOutputMask) const
+UComputeDataProvider* UPCGTextureDataInterface::CreateDataProvider() const
 {
-	TRACE_CPUPROFILER_EVENT_SCOPE(UPCGTextureDataInterface::CreateDataProvider);
-	UPCGDataBinding* Binding = CastChecked<UPCGDataBinding>(InBinding);
+	return NewObject<UPCGTextureDataProvider>();
+}
+
+void UPCGTextureDataProvider::Initialize(const UComputeDataInterface* InDataInterface, UObject* InBinding, uint64 InInputMask, uint64 InOutputMask)
+{
+	TRACE_CPUPROFILER_EVENT_SCOPE(UPCGTextureDataProvider::Initialize);
 
-	TObjectPtr<UPCGTextureDataProvider> DataProvider = NewObject<UPCGTextureDataProvider>();
+	Super::Initialize(InDataInterface, InBinding, InInputMask, InOutputMask);
+
+	const UPCGTextureDataInterface* DataInterface = CastChecked<UPCGTextureDataInterface>(InDataInterface);
+	UPCGDataBinding* Binding = CastChecked<UPCGDataBinding>(InBinding);
 
 	// Take any input pin label alias to obtain the data from the input data collection.
-	check(!DownstreamInputPinLabelAliases.IsEmpty());
-	const TArray<FPCGTaggedData> InputTaggedData = Binding->InputDataCollection.GetInputsByPin(DownstreamInputPinLabelAliases[0]);
-	
+	check(!DataInterface->DownstreamInputPinLabelAliases.IsEmpty());
+	const TArray<FPCGTaggedData> InputTaggedData = Binding->InputDataCollection.GetInputsByPin(DataInterface->DownstreamInputPinLabelAliases[0]);
+
 	if (InputTaggedData.Num() >= PCGTextureDataInterfaceConstants::MAX_NUM_TEXTURES)
 	{
 		UE_LOG(LogPCG, Warning, TEXT("Texture data interface on pin '%s' received too many data items (%d), only first %d valid textures will be used."),
-			*DownstreamInputPinLabelAliases[0].ToString(),
+			*DataInterface->DownstreamInputPinLabelAliases[0].ToString(),
 			InputTaggedData.Num(),
 			PCGTextureDataInterfaceConstants::MAX_NUM_TEXTURES);
 	}
 
 	const int NumInputTextures = FMath::Min(InputTaggedData.Num(), (int)PCGTextureDataInterfaceConstants::MAX_NUM_TEXTURES);
-	TArray<FPCGInputTextureInfo> TextureInfos;
+
 	TextureInfos.SetNum(NumInputTextures);
 
 	for (int InputTextureIndex = 0; InputTextureIndex < NumInputTextures; ++InputTextureIndex)
@@ -144,10 +151,6 @@ UComputeDataProvider* UPCGTextureDataInterface::CreateDataProvider(TObjectPtr<UO
 			TextureInfo.bIsTextureArray = TextureInfo.Texture->GetDesc().Dimension == ETextureDimension::Texture2DArray;
 		}
 	}
-
-	DataProvider->TextureInfos = MoveTemp(TextureInfos);
-
-	return DataProvider;
 }
 
 FComputeDataProviderRenderProxy* UPCGTextureDataProvider::GetRenderProxy()
@@ -155,6 +158,13 @@ FComputeDataProviderRenderProxy* UPCGTextureDataProvider::GetRenderProxy()
 	return new FPCGTextureDataProviderProxy(TextureInfos);
 }
 
+void UPCGTextureDataProvider::Reset()
+{
+	TextureInfos.Empty();
+
+	Super::Reset();
+}
+
 bool FPCGTextureDataProviderProxy::IsValid(FValidationData const& InValidationData) const
 {
 	if (InValidationData.ParameterStructSize != sizeof(FParameters))
@@ -207,26 +217,7 @@ void FPCGTextureDataProviderProxy::GatherDispatchData(FDispatchData const& InDis
 
 void FPCGTextureDataProviderProxy::AllocateResources(FRDGBuilder& GraphBuilder, FAllocationData const& InAllocationData)
 {
-	// TODO: Move this to somewhere more generic where it runs once, rather than for every PCGTextureDataInterface?
-	// Copied from GSystemTextures.GetBlackDummy(GraphBuilder));
-	const FRHITextureCreateDesc DummyTextureDesc =
-		FRHITextureCreateDesc::Create2D(TEXT("PCGDummyTexture"), 1, 1, PF_G8)
-		.SetFlags(ETextureCreateFlags::ShaderResource);
-
-	FTextureRHIRef Texture = RHICreateTexture(DummyTextureDesc);
-
-	{
-		uint32 DestStride;
-		float* DestBuffer = (float*)RHILockTexture2D(Texture, 0, RLM_WriteOnly, DestStride, false);
-		*DestBuffer = 0;
-		RHIUnlockTexture2D(Texture, 0, false);
-	}
-
-	// Create a dummy texture to bind.
-	TRefCountPtr<IPooledRenderTarget> BlackDummyTarget = CreateRenderTarget(Texture, TEXT("PCGDummyRenderTarget"));
-	check(BlackDummyTarget.IsValid());
-
-	FRDGTextureRef DummyTextureRDG = GraphBuilder.RegisterExternalTexture(BlackDummyTarget);
+	FRDGTextureRef DummyTextureRDG = nullptr;
 
 	for (int DataIndex = 0; DataIndex < PCGTextureDataInterfaceConstants::MAX_NUM_TEXTURES; ++DataIndex)
 	{
@@ -247,6 +238,29 @@ void FPCGTextureDataProviderProxy::AllocateResources(FRDGBuilder& GraphBuilder,
 
 		if (!TextureRDG)
 		{
+			if (!DummyTextureRDG)
+			{
+				// Copied from GSystemTextures.GetBlackDummy(GraphBuilder));
+				const FRHITextureCreateDesc DummyTextureDesc =
+					FRHITextureCreateDesc::Create2D(TEXT("PCGDummyTexture"), 1, 1, PF_G8)
+					.SetFlags(ETextureCreateFlags::ShaderResource);
+
+				FTextureRHIRef Texture = RHICreateTexture(DummyTextureDesc);
+
+				{
+					uint32 DestStride;
+					float* DestBuffer = (float*)RHILockTexture2D(Texture, 0, RLM_WriteOnly, DestStride, false);
+					*DestBuffer = 0;
+					RHIUnlockTexture2D(Texture, 0, false);
+				}
+
+				// Create a dummy texture to bind.
+				TRefCountPtr<IPooledRenderTarget> BlackDummyTarget = CreateRenderTarget(Texture, TEXT("PCGDummyRenderTarget"));
+				check(BlackDummyTarget.IsValid());
+
+				DummyTextureRDG = GraphBuilder.RegisterExternalTexture(BlackDummyTarget);
+			}
+
 			// Fallback to blank dummy texture if no texture was provided.
 			TextureRDG = DummyTextureRDG;
 		}
diff --git a/Engine/Plugins/PCG/Source/PCG/Private/Compute/DataInterfaces/PCGTextureDataInterface.h b/Engine/Plugins/PCG/Source/PCG/Private/Compute/DataInterfaces/PCGTextureDataInterface.h
--- a/Engine/Plugins/PCG/Source/PCG/Private/Compute/DataInterfaces/PCGTextureDataInterface.h
+++ b/Engine/Plugins/PCG/Source/PCG/Private/Compute/DataInterfaces/PCGTextureDataInterface.h
@@ -47,7 +47,7 @@ public:
 	TCHAR const* GetShaderVirtualPath() const override;
 	void GetShaderHash(FString& InOutKey) const override;
 	void GetHLSL(FString& OutHLSL, FString const& InDataInterfaceName) const override;
-	UComputeDataProvider* CreateDataProvider(TObjectPtr<UObject> InBinding, uint64 InInputMask, uint64 InOutputMask) const override;
+	UComputeDataProvider* CreateDataProvider() const override;
 	//~ End UComputeDataInterface Interface
 
 private:
@@ -62,7 +62,9 @@ class UPCGTextureDataProvider : public UComputeDataProvider
 
 public:
 	//~ Begin UComputeDataProvider Interface
+	virtual void Initialize(const UComputeDataInterface* InDataInterface, UObject* InBinding, uint64 InInputMask, uint64 InOutputMask) override;
 	FComputeDataProviderRenderProxy* GetRenderProxy() override;
+	virtual void Reset() override;
 	//~ End UComputeDataProvider Interface
 
 	TArray<FPCGInputTextureInfo> TextureInfos;
diff --git a/Engine/Plugins/PCG/Source/PCG/Private/Compute/Elements/PCGComputeGraphElement.cpp b/Engine/Plugins/PCG/Source/PCG/Private/Compute/Elements/PCGComputeGraphElement.cpp
--- a/Engine/Plugins/PCG/Source/PCG/Private/Compute/Elements/PCGComputeGraphElement.cpp
+++ b/Engine/Plugins/PCG/Source/PCG/Private/Compute/Elements/PCGComputeGraphElement.cpp
@@ -37,12 +37,25 @@ FPCGComputeGraphContext::~FPCGComputeGraphContext()
 		}
 	}
 
-	ComputeGraphInstance.DestroyDataProviders();
+	if (bComputeGraphInstanceInitialized)
+	{
+		ComputeGraphInstance.ResetDataProviders();
+
+		// Return instance to pool.
+		if (const UPCGGraph* Graph = Stack ? Stack->GetNearestNonInlinedGraphForCurrentFrame() : nullptr; ensure(Graph))
+		{
+			const UPCGGraph::FComputeGraphInstanceKey Key = { SourceComponent->GetGenerationGridSize(), ComputeGraphIndex };
+			Graph->ReturnComputeGraphInstanceToPool(Key, MoveTemp(ComputeGraphInstance));
+		}
+	}
 }
 
 void FPCGComputeGraphContext::AddExtraStructReferencedObjects(FReferenceCollector& Collector)
 {
-	Collector.AddPropertyReferences(FComputeGraphInstance::StaticStruct(), &ComputeGraphInstance);
+	if (bComputeGraphInstanceInitialized)
+	{
+		Collector.AddPropertyReferences(FComputeGraphInstance::StaticStruct(), &ComputeGraphInstance);
+	}
 
 	for (TObjectPtr<UComputeDataProvider> Provider : ProvidersWithBufferExports)
 	{
@@ -97,6 +110,8 @@ bool FPCGComputeGraphElement::ExecuteInternal(FPCGContext* InContext) const
 		{
 			return true;
 		}
+
+		Context->ComputeGraphIndex = ComputeGraphIndex;
 	}
 
 	// Drive the execution of a compute graph. The stages are commented below and numbered by the sequence in which they are executed.
@@ -191,11 +206,20 @@ bool FPCGComputeGraphElement::ExecuteInternal(FPCGContext* InContext) const
 			Context->SourceComponent->NotifyProceduralInstancesInUse();
 		}
 
+		const UPCGGraph* Graph = Context->Stack ? Context->Stack->GetNearestNonInlinedGraphForCurrentFrame() : nullptr;
+		const UPCGGraph::FComputeGraphInstanceKey Key = { Context->SourceComponent->GetGenerationGridSize(), ComputeGraphIndex };
+
+		if (Graph && Graph->RetrieveComputeGraphInstanceFromPool(Key, Context->ComputeGraphInstance))
+		{
+			Context->ComputeGraphInstance.InitializeDataProviders(Context->ComputeGraph.Get(), 0, Context->DataBinding.Get());
+		}
+		else
 		{
-			TRACE_CPUPROFILER_EVENT_SCOPE(ComputeGraphInstance.CreateDataProviders);
 			Context->ComputeGraphInstance.CreateDataProviders(Context->ComputeGraph.Get(), 0, Context->DataBinding.Get());
 		}
 
+		Context->bComputeGraphInstanceInitialized = true;
+
 		// Register all providers running async operations.
 		for (UComputeDataProvider* ComputeDataProvider : Context->ComputeGraphInstance.GetDataProviders())
 		{
@@ -213,7 +237,9 @@ bool FPCGComputeGraphElement::ExecuteInternal(FPCGContext* InContext) const
 
 				TWeakObjectPtr<UPCGDataCollectionDataProvider> DataProviderWeak = DataProvider;
 
-				DataProvider->OnDataExported_GameThread().AddLambda([ContextHandle=Context->GetOrCreateHandle(), DataProviderWeak]()
+				const uint64 OriginatingGenerationCount = DataProvider->GenerationCounter;
+
+				DataProvider->OnDataExported_GameThread().AddLambda([ContextHandle=Context->GetOrCreateHandle(), DataProviderWeak, GenerationCount=OriginatingGenerationCount]()
 				{
 					check(IsInGameThread());
 
@@ -221,13 +247,12 @@ bool FPCGComputeGraphElement::ExecuteInternal(FPCGContext* InContext) const
 					FPCGComputeGraphContext* ContextPtr = static_cast<FPCGComputeGraphContext*>(SharedHandle ? SharedHandle->GetContext() : nullptr);
 					UPCGDataCollectionDataProvider* DataProvider = DataProviderWeak.Get();
 
-					if (!DataProvider || !ContextPtr)
+					if (!DataProvider || !ContextPtr || GenerationCount != DataProvider->GenerationCounter)
 					{
 						// Safe to just jump out. The GPU buffer is ref counted.
 						return;
 					}
 
-					ensure(ContextPtr->ProvidersWithBufferExports.Contains(DataProvider));
 					ContextPtr->ProvidersWithBufferExports.Remove(DataProvider);
 
 					if (ContextPtr->ProvidersWithBufferExports.IsEmpty())
diff --git a/Engine/Plugins/PCG/Source/PCG/Private/Compute/Elements/PCGComputeGraphElement.h b/Engine/Plugins/PCG/Source/PCG/Private/Compute/Elements/PCGComputeGraphElement.h
--- a/Engine/Plugins/PCG/Source/PCG/Private/Compute/Elements/PCGComputeGraphElement.h
+++ b/Engine/Plugins/PCG/Source/PCG/Private/Compute/Elements/PCGComputeGraphElement.h
@@ -26,9 +26,14 @@ public:
 	TStrongObjectPtr<UPCGDataBinding> DataBinding = nullptr;
 	TStrongObjectPtr<UPCGComputeGraph> ComputeGraph = nullptr;
 
+	int32 ComputeGraphIndex = INDEX_NONE;
+
 	/** Data providers created from data interfaces and data bindings. */
 	FComputeGraphInstance ComputeGraphInstance;
 
+	/** Whether the compute graph instance has been initialized (data providers created, etc). */
+	bool bComputeGraphInstanceInitialized = false;
+
 	/** Graph enqueued (scheduled for execution by GPU). */
 	bool bGraphEnqueued = false;
 
diff --git a/Engine/Plugins/PCG/Source/PCG/Private/Graph/PCGStackContext.cpp b/Engine/Plugins/PCG/Source/PCG/Private/Graph/PCGStackContext.cpp
--- a/Engine/Plugins/PCG/Source/PCG/Private/Graph/PCGStackContext.cpp
+++ b/Engine/Plugins/PCG/Source/PCG/Private/Graph/PCGStackContext.cpp
@@ -203,6 +203,12 @@ const UPCGGraph* FPCGStack::GetNearestDynamicSubgraphForCurrentFrame() const
 	return nullptr;
 }
 
+const UPCGGraph* FPCGStack::GetNearestNonInlinedGraphForCurrentFrame() const
+{
+	const UPCGGraph* NearestDynamicSubgraph = GetNearestDynamicSubgraphForCurrentFrame();
+	return NearestDynamicSubgraph ? NearestDynamicSubgraph : GetRootGraph();
+}
+
 const UPCGNode* FPCGStack::GetCurrentFrameNode() const
 {
 	return StackFrames.IsEmpty() ? nullptr : Cast<const UPCGNode>(StackFrames.Last().Object.Get());
diff --git a/Engine/Plugins/PCG/Source/PCG/Private/PCGGraph.cpp b/Engine/Plugins/PCG/Source/PCG/Private/PCGGraph.cpp
--- a/Engine/Plugins/PCG/Source/PCG/Private/PCGGraph.cpp
+++ b/Engine/Plugins/PCG/Source/PCG/Private/PCGGraph.cpp
@@ -41,6 +41,11 @@ namespace PCGGraph
 		TEXT("pcg.Graph.FixInvalidEdgesOnPostLoad"),
 		true,
 		TEXT("Validates all edges are connected to valid pins/nodes and removes any invalid edges"));
+
+	static TAutoConsoleVariable<bool> CVarEnableComputeGraphInstancePooling(
+		TEXT("pcg.GPU.ComputeGraphInstancePooling"),
+		false,
+		TEXT("Caches compute graph instances rather than creating each instance and its data providers from scratch each time."));
 }
 
 namespace PCGGraphUtils
@@ -729,10 +734,20 @@ void UPCGGraph::BeginDestroy()
 
 void UPCGGraph::AddReferencedObjects(UObject* InThis, FReferenceCollector& Collector)
 {
-#if WITH_EDITOR
 	UPCGGraph* This = CastChecked<UPCGGraph>(InThis);
+
+	for (TPair<FComputeGraphInstanceKey, TArray<FComputeGraphInstance>>& Entry : This->ComputeGraphInstancePool)
+	{
+		for (FComputeGraphInstance& Instance : Entry.Value)
+		{
+			Collector.AddPropertyReferences(FComputeGraphInstance::StaticStruct(), &Instance);
+		}
+	}
+
+#if WITH_EDITOR
 	Collector.AddReferencedObject(This->PCGEditorGraph, This);
 #endif
+
 	Super::AddReferencedObjects(InThis, Collector);
 }
 
@@ -742,6 +757,36 @@ uint32 UPCGGraph::GetDefaultGridSize() const
 	return PCGHiGenGrid::IsValidGrid(HiGenGridSize) ? (PCGHiGenGrid::GridToGridSize(HiGenGridSize) * (1 << HiGenExponential)) : PCGHiGenGrid::UnboundedGridSize();
 }
 
+bool UPCGGraph::RetrieveComputeGraphInstanceFromPool(const FComputeGraphInstanceKey& InKey, FComputeGraphInstance& OutInstance) const
+{
+	check(IsInGameThread());
+
+	if (PCGGraph::CVarEnableComputeGraphInstancePooling.GetValueOnGameThread())
+	{
+		TArray<FComputeGraphInstance>* Instances = ComputeGraphInstancePool.Find(InKey);
+		if (Instances && !Instances->IsEmpty())
+		{
+			OutInstance = MoveTemp(Instances->Last());
+			Instances->RemoveAt(Instances->Num() - 1, EAllowShrinking::No);
+
+			return true;
+		}
+	}
+
+	return false;
+}
+
+void UPCGGraph::ReturnComputeGraphInstanceToPool(const FComputeGraphInstanceKey& InKey, FComputeGraphInstance InInstance) const
+{
+	check(IsInGameThread());
+
+	if (PCGGraph::CVarEnableComputeGraphInstancePooling.GetValueOnGameThread())
+	{
+		TArray<FComputeGraphInstance>& Instances = ComputeGraphInstancePool.FindOrAdd(InKey);
+		Instances.Add(MoveTemp(InInstance));
+	}
+}
+
 UPCGNode* UPCGGraph::AddNodeOfType(TSubclassOf<class UPCGSettings> InSettingsClass, UPCGSettings*& OutDefaultNodeSettings)
 {
 	UPCGSettings* Settings = NewObject<UPCGSettings>(GetTransientPackage(), InSettingsClass, NAME_None, RF_Transactional);
@@ -1571,6 +1616,10 @@ void UPCGGraph::NotifyGraphChanged(EPCGChangeType ChangeType)
 			FPropertyChangedEvent EmptyEvent{ nullptr };
 			FCoreUObjectDelegates::OnObjectPropertyChanged.Broadcast(this, EmptyEvent);
 		}
+
+		// Any non-trivial change to graph may change compute graphs.
+		check(IsInGameThread());
+		ComputeGraphInstancePool.Empty(ComputeGraphInstancePool.Num());
 	}
 
 	OnGraphChangedDelegate.Broadcast(this, ChangeType);
diff --git a/Engine/Plugins/PCG/Source/PCG/Public/Graph/PCGStackContext.h b/Engine/Plugins/PCG/Source/PCG/Public/Graph/PCGStackContext.h
--- a/Engine/Plugins/PCG/Source/PCG/Public/Graph/PCGStackContext.h
+++ b/Engine/Plugins/PCG/Source/PCG/Public/Graph/PCGStackContext.h
@@ -144,6 +144,9 @@ public:
 	/** Walks up the stack to find nearest dynamic subgraph frame and if found returns the graph. */
 	const UPCGGraph* GetNearestDynamicSubgraphForCurrentFrame() const;
 
+	/** Walks up the stack to find the nearest non-inlined graph (either dynamic subgraph or the top graph). */
+	const UPCGGraph* GetNearestNonInlinedGraphForCurrentFrame() const;
+
 	/** If current frame (top of stack) corresponds to a node returns that node, otherwise returns null. */
 	const UPCGNode* GetCurrentFrameNode() const;
 
diff --git a/Engine/Plugins/PCG/Source/PCG/Public/PCGGraph.h b/Engine/Plugins/PCG/Source/PCG/Public/PCGGraph.h
--- a/Engine/Plugins/PCG/Source/PCG/Public/PCGGraph.h
+++ b/Engine/Plugins/PCG/Source/PCG/Public/PCGGraph.h
@@ -9,6 +9,7 @@
 #include "Graph/PCGStackContext.h"
 #include "Helpers/PCGGraphParameterExtension.h"
 
+#include "ComputeFramework/ComputeGraphInstance.h"
 #include "StructUtils/PropertyBag.h"
 #include "UObject/ObjectPtr.h"
 
@@ -278,6 +279,15 @@ public:
 	bool IsHierarchicalGenerationEnabled() const { return bUseHierarchicalGeneration; }
 	bool Use2DGrid() const { return bUse2DGrid; }
 
+	using FComputeGraphInstanceKey = TPair</*GridSize*/uint32, /*Compute graph index*/int32>;
+	using FComputeGraphInstancePool = TMap<FComputeGraphInstanceKey, TArray<FComputeGraphInstance>>;
+
+	/** Attempt to retrieve a pooled compute graph instance for the given key. Returns true on success. */
+	bool RetrieveComputeGraphInstanceFromPool(const FComputeGraphInstanceKey& InKey, FComputeGraphInstance& OutInstance) const;
+
+	/** Places given compute graph instance into pool. No-ops if pooling is disabled. */
+	void ReturnComputeGraphInstanceToPool(const FComputeGraphInstanceKey& InKey, FComputeGraphInstance InInstance) const;
+
 #if WITH_EDITORONLY_DATA
 	UPROPERTY(BlueprintReadWrite, EditAnywhere, Category = AssetInfo, AssetRegistrySearchable)
 	FText Category;
@@ -526,8 +536,10 @@ public:
 	// Will call the callback function with a mutable property bag and will trigger the updates when it's done.
 	void UpdateUserParametersStruct(TFunctionRef<void(FInstancedPropertyBag&)> Callback);
 
-#if WITH_EDITOR
 private:
+	mutable FComputeGraphInstancePool ComputeGraphInstancePool;
+
+#if WITH_EDITOR
 	/** Sends a change notification. Demotes change if the compiled tasks are not significantly changed. */
 	void NotifyGraphStructureChanged(EPCGChangeType ChangeType, bool bForce = false);
 
