[PCG] Add multiple options and user preferences for double clicking nodes

+ Inspect node on double-click <new default>
+ Debug node on double-click
* Jump to Source Code (C++) <original default>

diff --git a/Engine/Plugins/PCG/Source/PCGEditor/Private/PCGEditor.cpp b/Engine/Plugins/PCG/Source/PCGEditor/Private/PCGEditor.cpp
--- a/Engine/Plugins/PCG/Source/PCGEditor/Private/PCGEditor.cpp
+++ b/Engine/Plugins/PCG/Source/PCGEditor/Private/PCGEditor.cpp
@@ -2827,11 +2827,34 @@ void FPCGEditor::OnNodeDoubleClicked(UEdGraphNode* Node)
 		UObject* Object = Node->GetJumpTargetForDoubleClick();
 
 		// "Normal" node
-		if (UPCGSettings* PCGSettings = Cast<UPCGSettings>(Object))
+		if (const UPCGSettings* PCGSettings = Cast<UPCGSettings>(Object))
 		{
-			JumpToDefinition(PCGSettings->GetClass());
+			// Functions may require the GraphEditorWidget's node selection, so set it manually to be safe.
+			GraphEditorWidget->SetNodeSelection(Node, /*bSelect=*/true);
+
+			switch (GetDefault<UPCGEditorSettings>()->NodeDoubleClickAction)
+			{
+				case EPCGEditorDoubleClickAction::ToggleInspectNode:
+					if (CanToggleInspected())
+					{
+						OnToggleInspected();
+					}
+					break;
+				case EPCGEditorDoubleClickAction::ToggleDebugNode:
+					if (CanToggleDebug())
+					{
+						OnToggleDebug();
+					}
+					break;
+				case EPCGEditorDoubleClickAction::JumpToSourceFile:
+					JumpToDefinition(PCGSettings->GetClass());
+					break;
+				case EPCGEditorDoubleClickAction::DoNothing: // fall-through
+				default:
+					break;
+			}
 		}
-		else
+		else // Special options with non-UPCGSettings based targets.
 		{
 			const UPCGEditorGraphNodeBase* PCGEditorGraphNode = Cast<const UPCGEditorGraphNodeBase>(Node);
 			const UPCGNode* PCGNode = PCGEditorGraphNode ? PCGEditorGraphNode->GetPCGNode() : nullptr;
@@ -2867,33 +2890,7 @@ void FPCGEditor::JumpToDefinition(const UClass* Class) const
 {
 	if (ensure(GUnrealEd) && GUnrealEd->GetUnrealEdOptions()->IsCPPAllowed())
 	{
-		const bool bEnableNavigateToNativeNodes = GetDefault<UPCGEditorSettings>()->bEnableNavigateToNativeNodes;
-		if (bEnableNavigateToNativeNodes)
-		{
-			FSourceCodeNavigation::NavigateToClass(Class);
-		}
-		else
-		{
-			// Inform user that the node is native, give them opportunity to enable navigation to native nodes:
-			FNotificationInfo Info(LOCTEXT("NavigateToNativeDisabled", "Navigation to Native (c++) PCG Nodes Disabled"));
-			Info.ExpireDuration = 10.0f;
-			Info.CheckBoxState = bEnableNavigateToNativeNodes ? ECheckBoxState::Checked : ECheckBoxState::Unchecked;
-
-			Info.CheckBoxStateChanged = FOnCheckStateChanged::CreateStatic(
-				[](ECheckBoxState NewState)
-				{
-					const FScopedTransaction Transaction(LOCTEXT("ChangeEnableNavigateToNativeNodes", "Change Enable Navigate to Native Nodes Setting"));
-
-					UPCGEditorSettings* MutableEditorSettings = GetMutableDefault<UPCGEditorSettings>();
-					MutableEditorSettings->Modify();
-					MutableEditorSettings->bEnableNavigateToNativeNodes = (NewState == ECheckBoxState::Checked) ? true : false;
-					MutableEditorSettings->SaveConfig();
-				}
-			);
-			Info.CheckBoxText = LOCTEXT("EnableNavigationToNative", "Enable Navigate to Native Nodes?");
-
-			FSlateNotificationManager::Get().AddNotification(Info);
-		}
+		FSourceCodeNavigation::NavigateToClass(Class);
 	}
 }
 
diff --git a/Engine/Plugins/PCG/Source/PCGEditor/Public/PCGEditorSettings.h b/Engine/Plugins/PCG/Source/PCGEditor/Public/PCGEditorSettings.h
--- a/Engine/Plugins/PCG/Source/PCGEditor/Public/PCGEditorSettings.h
+++ b/Engine/Plugins/PCG/Source/PCGEditor/Public/PCGEditorSettings.h
@@ -12,6 +12,15 @@ class UPCGBuilderSettings;
 
 struct FEdGraphPinType;
 
+UENUM(BlueprintType)
+enum class EPCGEditorDoubleClickAction : uint8
+{
+	DoNothing,
+	ToggleInspectNode UMETA(DisplayName = "Toggle Inspection on Node"),
+	ToggleDebugNode UMETA(DisplayName = "Toggle Debug on Node"),
+	JumpToSourceFile UMETA(DisplayName = "View Source File (C++)")
+};
+
 UCLASS(config=Editor, meta=(DisplayName="PCG"))
 class PCGEDITOR_API UPCGEditorProjectSettings : public UDeveloperSettings
 {
@@ -207,9 +216,9 @@ public:
 	UPROPERTY(EditAnywhere, config, Category = Wire, meta = (ClampMin="0", ClampMax="1"))
 	float EmptyEdgeDebugDesaturateFactor = 0.5f;
 
-	/** Specify if we want to jump to definition in case of double click on native PCG Nodes */
-	UPROPERTY(EditAnywhere, config, Category = Workflow)
-	bool bEnableNavigateToNativeNodes = true;
+	/** Specify an action to take upon double-clicking a native PCG node. */
+	UPROPERTY(EditAnywhere, Config, Category = "Workflow")
+	EPCGEditorDoubleClickAction NodeDoubleClickAction = EPCGEditorDoubleClickAction::ToggleInspectNode;
 
 	/** Specify if we want to disable CPU Throttling when a PCG Graph is executing, this will improve execution time when app is out of focus/minimized */
 	UPROPERTY(EditAnywhere, config, Category = "Editor Performance", meta = (DisplayName = "Disable CPU throttling during graph execution"))
@@ -258,4 +267,9 @@ public:
 	/** Display node 'Determinism' actions on the right click context menu. */
 	UPROPERTY(EditAnywhere, config, Category = "Graph UI")
 	bool bShowNodeDeterminismActionsRightClickContext = true;
+
+	// Begin deprecated properties
+	UE_DEPRECATED(5.6, "bEnableNavigateToNativeNodes has been replaced by NodeDoubleClickAction.")
+	UPROPERTY()
+	bool bEnableNavigateToNativeNodes = true;
 };
