[PCG] Set bGenerated = false on PCGComponent CDOs.
- It is currently possible to create a BP from a world actor containing a generated component...which will flag the CDO as bGenerated = true (looking into fixing this in a subsequent change)

#rnx

diff --git a/Engine/Plugins/PCG/Source/PCG/Private/PCGComponent.cpp b/Engine/Plugins/PCG/Source/PCG/Private/PCGComponent.cpp
--- a/Engine/Plugins/PCG/Source/PCG/Private/PCGComponent.cpp
+++ b/Engine/Plugins/PCG/Source/PCG/Private/PCGComponent.cpp
@@ -1927,6 +1927,13 @@ void UPCGComponent::PostLoad()
 	if (!GetGraph())
 	{
 		bGenerated = false;
+	} // We can never be generated if we are a CDO
+	else if (IsTemplate())
+	{
+		// It is possible to create a Blueprint from an Actor containing a generated PCG Component.
+		// Eventually this shouldn't be possible but for data that has already been saved, we fix it up here.
+		bGenerated = false;
+		GeneratedGraphOutput.Reset();
 	}
 
 	// If we have both default value (bIsComponentPartitioned = false and bIsPartitioned = true)
