[PCG] Add GetPinInitialDefaultValueType to AttributeTransformOp to fix ensure

#rnx

diff --git a/Engine/Plugins/PCG/Source/PCG/Private/Elements/Metadata/PCGMetadataTransformOpElement.cpp b/Engine/Plugins/PCG/Source/PCG/Private/Elements/Metadata/PCGMetadataTransformOpElement.cpp
--- a/Engine/Plugins/PCG/Source/PCG/Private/Elements/Metadata/PCGMetadataTransformOpElement.cpp
+++ b/Engine/Plugins/PCG/Source/PCG/Private/Elements/Metadata/PCGMetadataTransformOpElement.cpp
@@ -204,6 +204,11 @@ bool UPCGMetadataTransformSettings::IsPinDefaultValueEnabled(const FName PinLabe
 }
 
 EPCGMetadataTypes UPCGMetadataTransformSettings::GetPinDefaultValueType(const FName PinLabel) const
+{
+	return GetPinInitialDefaultValueType(PinLabel);
+}
+
+EPCGMetadataTypes UPCGMetadataTransformSettings::GetPinInitialDefaultValueType(const FName PinLabel) const
 {
 	if (PinLabel == PCGMetadataSettingsBaseConstants::LerpRatioLabel)
 	{
diff --git a/Engine/Plugins/PCG/Source/PCG/Public/Elements/Metadata/PCGMetadataTransformOpElement.h b/Engine/Plugins/PCG/Source/PCG/Public/Elements/Metadata/PCGMetadataTransformOpElement.h
--- a/Engine/Plugins/PCG/Source/PCG/Public/Elements/Metadata/PCGMetadataTransformOpElement.h
+++ b/Engine/Plugins/PCG/Source/PCG/Public/Elements/Metadata/PCGMetadataTransformOpElement.h
@@ -61,6 +61,7 @@ public:
 	//~Begin IPCGSettingsDefaultValueProvider interface
 	virtual bool IsPinDefaultValueEnabled(FName PinLabel) const override;
 	virtual EPCGMetadataTypes GetPinDefaultValueType(FName PinLabel) const override;
+	virtual EPCGMetadataTypes GetPinInitialDefaultValueType(FName PinLabel) const override;
 	//~End IPCGSettingsDefaultValueProvider interface
 protected:
 	//~Begin UPCGSettings interface
