[PCG] Small robustness update for Inline Constants

* Also remove inline constants from Parse String, since it is not a useful case

#rnx

diff --git a/Engine/Plugins/PCG/Source/PCG/Private/Elements/Metadata/PCGMetadataOpElementBase.cpp b/Engine/Plugins/PCG/Source/PCG/Private/Elements/Metadata/PCGMetadataOpElementBase.cpp
--- a/Engine/Plugins/PCG/Source/PCG/Private/Elements/Metadata/PCGMetadataOpElementBase.cpp
+++ b/Engine/Plugins/PCG/Source/PCG/Private/Elements/Metadata/PCGMetadataOpElementBase.cpp
@@ -114,12 +114,12 @@ uint32 UPCGMetadataSettingsBase::GetInputPinIndex(const FName InPinLabel) const
 
 bool UPCGMetadataSettingsBase::IsPinDefaultValueEnabled(const FName PinLabel) const
 {
-	return GetInputPinIndex(PinLabel) != INDEX_NONE && PCGMetadataHelpers::MetadataTypeSupportsDefaultValues(GetPinDefaultValueType(PinLabel));
+	return DefaultValuesAreEnabled() && (GetInputPinIndex(PinLabel) != INDEX_NONE) && PCGMetadataHelpers::MetadataTypeSupportsDefaultValues(GetPinDefaultValueType(PinLabel));
 }
 
 bool UPCGMetadataSettingsBase::IsPinDefaultValueActivated(const FName PinLabel) const
 {
-	if (!IsPinDefaultValueEnabled(PinLabel))
+	if (!DefaultValuesAreEnabled() || !IsPinDefaultValueEnabled(PinLabel))
 	{
 		return false;
 	}
diff --git a/Engine/Plugins/PCG/Source/PCG/Public/Elements/Metadata/PCGParseString.h b/Engine/Plugins/PCG/Source/PCG/Public/Elements/Metadata/PCGParseString.h
--- a/Engine/Plugins/PCG/Source/PCG/Public/Elements/Metadata/PCGParseString.h
+++ b/Engine/Plugins/PCG/Source/PCG/Public/Elements/Metadata/PCGParseString.h
@@ -39,7 +39,7 @@ public:
 	//~End UPCGMetadataSettingsBase interface
 
 	//~Begin IPCGSettingsDefaultValueProvider interface
-	virtual EPCGMetadataTypes GetPinInitialDefaultValueType(FName PinLabel) const override { return EPCGMetadataTypes::String; }
+	virtual bool DefaultValuesAreEnabled() const override { return false; }
 	//~End IPCGSettingsDefaultValueProvider interface
 
 public:
