[PCG] Provide slightly more information when a non-normalized transform rotation is encountered when packing data for GPU.

#rnx

diff --git a/Engine/Plugins/PCG/Source/PCG/Private/Compute/PCGDataForGPU.cpp b/Engine/Plugins/PCG/Source/PCG/Private/Compute/PCGDataForGPU.cpp
--- a/Engine/Plugins/PCG/Source/PCG/Private/Compute/PCGDataForGPU.cpp
+++ b/Engine/Plugins/PCG/Source/PCG/Private/Compute/PCGDataForGPU.cpp
@@ -219,7 +219,8 @@ namespace PCGDataForGPUHelpers
 			const FTransform Transform = Attribute->GetValue(ValueKey);
 
 			const bool bIsRotationNormalized = Transform.IsRotationNormalized();
-			ensureMsgf(bIsRotationNormalized, TEXT("Tried to pack transform for GPU data collection, but the transform's rotation is not normalized. Using identity instead."));
+			ensureMsgf(bIsRotationNormalized, TEXT("Tried to pack transform for GPU data collection, but the transform's rotation is not normalized (%f, %f, %f, %f). Using identity instead."),
+				Transform.GetRotation().X, Transform.GetRotation().Y, Transform.GetRotation().Z, Transform.GetRotation().W);
 
 			// Note: ToMatrixWithScale() crashes if the transform is not normalized.
 			const FMatrix Matrix = bIsRotationNormalized ? Transform.ToMatrixWithScale() : FMatrix::Identity;
