[PCG] Change data validation from an ensure to an error as it's a condition we don't have control over, and ensure() only fires once which does not help debugging.

#rnx

diff --git a/Engine/Plugins/PCG/Source/PCG/Private/Compute/PCGDataForGPU.cpp b/Engine/Plugins/PCG/Source/PCG/Private/Compute/PCGDataForGPU.cpp
--- a/Engine/Plugins/PCG/Source/PCG/Private/Compute/PCGDataForGPU.cpp
+++ b/Engine/Plugins/PCG/Source/PCG/Private/Compute/PCGDataForGPU.cpp
@@ -219,8 +219,11 @@ namespace PCGDataForGPUHelpers
 			const FTransform Transform = Attribute->GetValue(ValueKey);
 
 			const bool bIsRotationNormalized = Transform.IsRotationNormalized();
-			ensureMsgf(bIsRotationNormalized, TEXT("Tried to pack transform for GPU data collection, but the transform's rotation is not normalized (%f, %f, %f, %f). Using identity instead."),
-				Transform.GetRotation().X, Transform.GetRotation().Y, Transform.GetRotation().Z, Transform.GetRotation().W);
+			if (!bIsRotationNormalized)
+			{
+				UE_LOG(LogPCG, Error, TEXT("Tried to pack transform for GPU data collection, but the transform's rotation is not normalized (%f, %f, %f, %f). Using identity instead."),
+					Transform.GetRotation().X, Transform.GetRotation().Y, Transform.GetRotation().Z, Transform.GetRotation().W);
+			}
 
 			// Note: ToMatrixWithScale() crashes if the transform is not normalized.
 			const FMatrix Matrix = bIsRotationNormalized ? Transform.ToMatrixWithScale() : FMatrix::Identity;
