[PCG] Un-fork the copy points shader and support bCopyEachSourceOnEveryTarget=false with match attribute.

#rnx

diff --git a/Engine/Plugins/PCG/Shaders/Private/Elements/PCGCopyPoints.usf b/Engine/Plugins/PCG/Shaders/Private/Elements/PCGCopyPoints.usf
--- a/Engine/Plugins/PCG/Shaders/Private/Elements/PCGCopyPoints.usf
+++ b/Engine/Plugins/PCG/Shaders/Private/Elements/PCGCopyPoints.usf
@@ -34,92 +34,28 @@ void PCGCopyPointsCS(uint3 GroupId : SV_GroupID, uint GroupIndex : SV_GroupIndex
 	const uint ThreadIndex = GetUnWrappedDispatchThreadId(GroupId, GroupIndex, 64);
 	if (ThreadIndex >= GetNumThreads().x) return;
 
-	const uint SourceNumData = Source_GetNumData();
-	const uint TargetNumData = Target_GetNumData();
-	const uint TargetPointsCount = Target_GetNumElements();
-	const uint bCopyEachSourceOnEveryTarget = CopyPoints_GetCopyEachSourceOnEveryTarget();
-
-	// Data must be non-zero, and NxM (when bCopyEachSourceOnEveryTarget == true), N:N, 1:N, or N:1
-	if (SourceNumData == 0 || TargetNumData == 0 || (!bCopyEachSourceOnEveryTarget && SourceNumData != TargetNumData && SourceNumData != 1 && TargetNumData != 1))
+	uint Out_DataIndex, Out_ElemIndex;
+	if (!Out_GetThreadData(ThreadIndex, Out_DataIndex, Out_ElemIndex))
 	{
-		// TODO: Runtime error diagnostics
 		return;
 	}
-	
-	uint SourceDataIndex, SourcePointIndexInData;
-	uint TargetDataIndex, TargetPointIndexInData;
-
-	if (bCopyEachSourceOnEveryTarget)
-	{
-		const uint SourcePointIndex = ThreadIndex / TargetPointsCount;
-		const uint TargetPointIndex = ThreadIndex % TargetPointsCount;
-
-		if (!Source_GetThreadData(SourcePointIndex, SourceDataIndex, SourcePointIndexInData))
-		{
-			return;
-		}
-
-		if (!Target_GetThreadData(TargetPointIndex, TargetDataIndex, TargetPointIndexInData))
-		{
-			return;
-		}
-	}
-	else
-	{
-		const uint NumIterations = max(SourceNumData, TargetNumData);
-
-		bool bFoundInData = false;
-		uint TotalElementCount = 0;
-
-		for (uint DataIndex = 0; DataIndex < NumIterations; ++DataIndex)
-		{
-			SourceDataIndex = min(DataIndex, SourceNumData - 1);
-			TargetDataIndex = min(DataIndex, TargetNumData - 1);
-
-			const uint SourceElemCount = Source_GetNumElements(SourceDataIndex);
-			const uint TargetElemCount = Target_GetNumElements(TargetDataIndex);
-
-			const uint ElemCount = SourceElemCount * TargetElemCount;
 
-			if (TotalElementCount + ElemCount > ThreadIndex)
-			{
-				const uint ElementIndex = ThreadIndex - TotalElementCount;
-				SourcePointIndexInData = ElementIndex / TargetElemCount;
-				TargetPointIndexInData = ElementIndex % TargetElemCount;
-				bFoundInData = true;
-				break;
-			}
-
-			TotalElementCount += ElemCount;
-		}
-
-		if (!bFoundInData)
-		{
-			return;
-		}
-	}
+	const uint SourceNumData = Source_GetNumData();
+	const uint TargetNumData = Target_GetNumData();
+	const uint CopyEachSourceOnEveryTarget = CopyPoints_GetCopyEachSourceOnEveryTarget();
 
-	uint OutNumData = 0;
+	// Get source and target data indices for this output data index.
+	uint SourceDataIndex, TargetDataIndex;
+	CopyPoints_GetSourceAndTargetDataIndices(Out_DataIndex, SourceNumData, TargetNumData, SourceDataIndex, TargetDataIndex);
 
-	if (bCopyEachSourceOnEveryTarget)
-	{
-		OutNumData = SourceNumData * TargetNumData;
-	}
-	else
-	{
-		OutNumData = max(SourceNumData, TargetNumData);
-	}
-
-	uint OutDataIndex, OutElemIndex;
-	if (!Out_GetThreadData(ThreadIndex, OutDataIndex, OutElemIndex))
-	{
-		return;
-	}
+	const uint NumTargetPoints = Target_GetNumElements(TargetDataIndex);
+	uint SourcePointIndexInData = Out_ElemIndex / NumTargetPoints;
+	uint TargetPointIndexInData = Out_ElemIndex % NumTargetPoints;
 
 	// Propagate the 'Removed' status to the output point if either source or target point is removed. Otherwise, this point will not be culled.
 	if (Source_IsPointRemoved(SourceDataIndex, SourcePointIndexInData) || Target_IsPointRemoved(TargetDataIndex, TargetPointIndexInData))
 	{
-		Out_RemovePoint(OutDataIndex, OutElemIndex);
+		Out_RemovePoint(Out_DataIndex, Out_ElemIndex);
 		return;
 	}
 
@@ -205,39 +141,39 @@ void PCGCopyPointsCS(uint3 GroupId : SV_GroupID, uint GroupIndex : SV_GroupIndex
 	}
 	
 	// Write output
-	Out_SetPosition(OutDataIndex, OutElemIndex, OutPosition);
-	Out_SetRotation(OutDataIndex, OutElemIndex, OutRotation);
-	Out_SetScale(OutDataIndex, OutElemIndex, OutScale);
-	Out_SetColor(OutDataIndex, OutElemIndex, OutColor);
-	Out_SetSeed(OutDataIndex, OutElemIndex, OutSeed);
+	Out_SetPosition(Out_DataIndex, Out_ElemIndex, OutPosition);
+	Out_SetRotation(Out_DataIndex, Out_ElemIndex, OutRotation);
+	Out_SetScale(Out_DataIndex, Out_ElemIndex, OutScale);
+	Out_SetColor(Out_DataIndex, Out_ElemIndex, OutColor);
+	Out_SetSeed(Out_DataIndex, Out_ElemIndex, OutSeed);
 
 	// These properties do not have an inheritance mode, simply copy from the source data.
-	Out_SetBoundsMin(OutDataIndex, OutElemIndex, Source_GetBoundsMin(SourceDataIndex, SourcePointIndexInData));
-	Out_SetBoundsMax(OutDataIndex, OutElemIndex, Source_GetBoundsMax(SourceDataIndex, SourcePointIndexInData));
-	Out_SetDensity(OutDataIndex, OutElemIndex, Source_GetDensity(SourceDataIndex, SourcePointIndexInData));
-	Out_SetSteepness(OutDataIndex, OutElemIndex, Source_GetSteepness(SourceDataIndex, SourcePointIndexInData));
+	Out_SetBoundsMin(Out_DataIndex, Out_ElemIndex, Source_GetBoundsMin(SourceDataIndex, SourcePointIndexInData));
+	Out_SetBoundsMax(Out_DataIndex, Out_ElemIndex, Source_GetBoundsMax(SourceDataIndex, SourcePointIndexInData));
+	Out_SetDensity(Out_DataIndex, Out_ElemIndex, Source_GetDensity(SourceDataIndex, SourcePointIndexInData));
+	Out_SetSteepness(Out_DataIndex, Out_ElemIndex, Source_GetSteepness(SourceDataIndex, SourcePointIndexInData));
 
 	// Copy attributes
 	if (AttributeInheritance == COPY_POINTS_METADATA_INHERITANCE_MODE_SourceOnly)
 	{
-		PCG_COPY_ATTRIBUTES_TO_OUTPUT(Out, Source, OutDataIndex, OutElemIndex, SourceDataIndex, SourcePointIndexInData, /*bMetadataOnly=*/true, /*bInitNonCopiedAttributes=*/true);
+		PCG_COPY_ATTRIBUTES_TO_OUTPUT(Out, Source, Out_DataIndex, Out_ElemIndex, SourceDataIndex, SourcePointIndexInData, /*bMetadataOnly=*/true, /*bInitNonCopiedAttributes=*/true);
 	}
 	else if (AttributeInheritance == COPY_POINTS_METADATA_INHERITANCE_MODE_TargetOnly)
 	{
-		PCG_COPY_ATTRIBUTES_TO_OUTPUT(Out, Target, OutDataIndex, OutElemIndex, TargetDataIndex, TargetPointIndexInData, /*bMetadataOnly=*/true, /*bInitNonCopiedAttributes=*/true);
+		PCG_COPY_ATTRIBUTES_TO_OUTPUT(Out, Target, Out_DataIndex, Out_ElemIndex, TargetDataIndex, TargetPointIndexInData, /*bMetadataOnly=*/true, /*bInitNonCopiedAttributes=*/true);
 	}
 	else if (AttributeInheritance == COPY_POINTS_METADATA_INHERITANCE_MODE_SourceFirst)
 	{
-		PCG_COPY_ATTRIBUTES_TO_OUTPUT(Out, Target, OutDataIndex, OutElemIndex, TargetDataIndex, TargetPointIndexInData, /*bMetadataOnly=*/true, /*bInitNonCopiedAttributes=*/true);
+		PCG_COPY_ATTRIBUTES_TO_OUTPUT(Out, Target, Out_DataIndex, Out_ElemIndex, TargetDataIndex, TargetPointIndexInData, /*bMetadataOnly=*/true, /*bInitNonCopiedAttributes=*/true);
 
 		// Don't allow Source to initialize non-copied attributes in the output, otherwise it will overwrite attributes from Target.
-		PCG_COPY_ATTRIBUTES_TO_OUTPUT(Out, Source, OutDataIndex, OutElemIndex, SourceDataIndex, SourcePointIndexInData, /*bMetadataOnly=*/true, /*bInitNonCopiedAttributes=*/false);
+		PCG_COPY_ATTRIBUTES_TO_OUTPUT(Out, Source, Out_DataIndex, Out_ElemIndex, SourceDataIndex, SourcePointIndexInData, /*bMetadataOnly=*/true, /*bInitNonCopiedAttributes=*/false);
 	}
 	else if (AttributeInheritance == COPY_POINTS_METADATA_INHERITANCE_MODE_TargetFirst)
 	{
-		PCG_COPY_ATTRIBUTES_TO_OUTPUT(Out, Source, OutDataIndex, OutElemIndex, SourceDataIndex, SourcePointIndexInData, /*bMetadataOnly=*/true, /*bInitNonCopiedAttributes=*/true);
+		PCG_COPY_ATTRIBUTES_TO_OUTPUT(Out, Source, Out_DataIndex, Out_ElemIndex, SourceDataIndex, SourcePointIndexInData, /*bMetadataOnly=*/true, /*bInitNonCopiedAttributes=*/true);
 
 		// Don't allow Target to initialize non-copied attributes in the output, otherwise it will overwrite attributes from Source.
-		PCG_COPY_ATTRIBUTES_TO_OUTPUT(Out, Target, OutDataIndex, OutElemIndex, TargetDataIndex, TargetPointIndexInData, /*bMetadataOnly=*/true, /*bInitNonCopiedAttributes=*/false);
+		PCG_COPY_ATTRIBUTES_TO_OUTPUT(Out, Target, Out_DataIndex, Out_ElemIndex, TargetDataIndex, TargetPointIndexInData, /*bMetadataOnly=*/true, /*bInitNonCopiedAttributes=*/false);
 	}
 }
diff --git a/Engine/Plugins/PCG/Shaders/Private/Elements/PCGCopyPointsAnalysis.usf b/Engine/Plugins/PCG/Shaders/Private/Elements/PCGCopyPointsAnalysis.usf
--- a/Engine/Plugins/PCG/Shaders/Private/Elements/PCGCopyPointsAnalysis.usf
+++ b/Engine/Plugins/PCG/Shaders/Private/Elements/PCGCopyPointsAnalysis.usf
@@ -21,12 +21,10 @@ void PCGCopyPointsAnalysisCS(uint3 GroupId : SV_GroupID, uint GroupIndex : SV_Gr
 
 	const uint ThreadIndex = GetUnWrappedDispatchThreadId(GroupId, GroupIndex, 64);
 
-	const uint SourceNumData = Source_GetNumData();
-	const uint NumTargetPointsTotal = Target_GetNumElements();
-
-	if (ThreadIndex >= SourceNumData * NumTargetPointsTotal)
+	// We dispatched one thread per copy points output data and the output data is a single attribute set with num data elements.
+	uint Dummy, CopyOutputDataIndex;
+	if (!Out_GetThreadData(ThreadIndex, Dummy, CopyOutputDataIndex))
 	{
-		// Discard threads outside of data range.
 		return;
 	}
 
@@ -34,18 +32,29 @@ void PCGCopyPointsAnalysisCS(uint3 GroupId : SV_GroupID, uint GroupIndex : SV_Gr
 	if (MatchAttributeId == -1)
 	{
 		// No match attribute, early out with no copy.
-		Out_SetBool(0, ThreadIndex, CopyPoints_GetSelectedFlagAttributeId(), false);
+		Out_SetBool(0, CopyOutputDataIndex, CopyPoints_GetSelectedFlagAttributeId(), false);
 		return;
 	}
 
-	const uint NumDataTarget = Target_GetNumData();
+	const uint SourceNumData = Source_GetNumData();
+	const uint TargetNumData = Target_GetNumData();
 
-	// todo_pcg: Needs both modes of copy points.
+	uint SourceDataIndex = 0;
+	uint TargetDataIndex = 0;
 
-	// Output is attribute set with SourceNumData * TargetNumData bools. Kernel executes one thread per output.
-	// copypoints(S = ((src 0), (src 1)), T = ((tgt 0), (tgt 1))) = (src 0 x tgt 0), (src 0 x tgt 1), (src 1 x tgt 0), (src 1 tgt 1))
-	const uint SourceDataIndex = ThreadIndex / NumDataTarget;
-	const uint TargetDataIndex = ThreadIndex % NumDataTarget;
+	if (CopyPoints_GetCopyEachSourceOnEveryTarget())
+	{
+		// Output is attribute set with SourceNumData * TargetNumData bools. Kernel executes one thread per output.
+		// copypoints(S = ((src 0), (src 1)), T = ((tgt 0), (tgt 1))) = (src 0 x tgt 0), (src 0 x tgt 1), (src 1 x tgt 0), (src 1 tgt 1))
+		SourceDataIndex = CopyOutputDataIndex / TargetNumData;
+		TargetDataIndex = CopyOutputDataIndex % TargetNumData;
+	}
+	else
+	{
+		// Output is attribute set with max(SourceNumData, TargetNumData) bools. (S, T) can be (N, N), (1, N) or (N, 1).
+		SourceDataIndex = clamp(CopyOutputDataIndex, 0, -1 + (int)SourceNumData);
+		TargetDataIndex = clamp(CopyOutputDataIndex, 0, -1 + (int)TargetNumData);
+	}
 
 	if (Source_GetNumElements(SourceDataIndex) > 0 && Target_GetNumElements(TargetDataIndex) > 0)
 	{
@@ -58,11 +67,11 @@ void PCGCopyPointsAnalysisCS(uint3 GroupId : SV_GroupID, uint GroupIndex : SV_Gr
 
 			if (SourceValue == TargetValue)
 			{
-				Out_SetBool(/*DataIndex=*/0, ThreadIndex, CopyPoints_GetSelectedFlagAttributeId(), true);
+				Out_SetBool(/*DataIndex=*/0, CopyOutputDataIndex, CopyPoints_GetSelectedFlagAttributeId(), true);
 				return;
 			}
 		}
 	}
 
-	Out_SetBool(0, ThreadIndex, CopyPoints_GetSelectedFlagAttributeId(), false);
+	Out_SetBool(0, CopyOutputDataIndex, CopyPoints_GetSelectedFlagAttributeId(), false);
 }
diff --git a/Engine/Plugins/PCG/Shaders/Private/Elements/PCGCopyPointsConditional.usf b/Engine/Plugins/PCG/Shaders/Private/Elements/PCGCopyPointsConditional.usf
deleted file mode 100644
diff --git a/Engine/Plugins/PCG/Source/PCG/Private/Compute/DataInterfaces/Elements/PCGCopyPointsAnalysisDataInterface.cpp b/Engine/Plugins/PCG/Source/PCG/Private/Compute/DataInterfaces/Elements/PCGCopyPointsAnalysisDataInterface.cpp
--- a/Engine/Plugins/PCG/Source/PCG/Private/Compute/DataInterfaces/Elements/PCGCopyPointsAnalysisDataInterface.cpp
+++ b/Engine/Plugins/PCG/Source/PCG/Private/Compute/DataInterfaces/Elements/PCGCopyPointsAnalysisDataInterface.cpp
@@ -27,11 +27,16 @@ void UPCGCopyPointsAnalysisDataInterface::GetSupportedInputs(TArray<FShaderFunct
 	OutFunctions.AddDefaulted_GetRef()
 		.SetName(TEXT("CopyPoints_GetSelectedFlagAttributeId"))
 		.AddReturnType(FShaderValueType::Get(EShaderFundamentalType::Int));
+
+	OutFunctions.AddDefaulted_GetRef()
+		.SetName(TEXT("CopyPoints_GetCopyEachSourceOnEveryTarget"))
+		.AddReturnType(FShaderValueType::Get(EShaderFundamentalType::Uint));
 }
 
 BEGIN_SHADER_PARAMETER_STRUCT(FPCGCopyPointsAnalysisDataInterfaceParameters,)
 	SHADER_PARAMETER(int32, MatchAttributeId)
 	SHADER_PARAMETER(int32, SelectedFlagAttributeId)
+	SHADER_PARAMETER(uint32, CopyEachSourceOnEveryTarget)
 END_SHADER_PARAMETER_STRUCT()
 
 void UPCGCopyPointsAnalysisDataInterface::GetShaderParameters(TCHAR const* UID, FShaderParametersMetadataBuilder& InOutBuilder, FShaderParametersMetadataAllocations& InOutAllocations) const
@@ -49,9 +54,11 @@ void UPCGCopyPointsAnalysisDataInterface::GetHLSL(FString& OutHLSL, FString cons
 	OutHLSL += FString::Format(TEXT(
 		"int {DataInterfaceName}_MatchAttributeId;\n"
 		"int {DataInterfaceName}_SelectedFlagAttributeId;\n"
+		"uint {DataInterfaceName}_CopyEachSourceOnEveryTarget;\n"
 		"\n"
 		"int CopyPoints_GetMatchAttributeId_{DataInterfaceName}() { return {DataInterfaceName}_MatchAttributeId; }\n"
 		"int CopyPoints_GetSelectedFlagAttributeId_{DataInterfaceName}() { return {DataInterfaceName}_SelectedFlagAttributeId; }\n"
+		"uint CopyPoints_GetCopyEachSourceOnEveryTarget_{DataInterfaceName}() { return {DataInterfaceName}_CopyEachSourceOnEveryTarget; }\n"
 		), TemplateArgs);
 }
 
@@ -78,16 +85,16 @@ void UPCGCopyPointsAnalysisDataProvider::Initialize(const UComputeDataInterface*
 
 	SelectedFlagAttributeId = Binding->GetAttributeId(PCGCopyPointsConstants::SelectedFlagAttributeName, EPCGKernelAttributeType::Bool);
 	ensure(SelectedFlagAttributeId != INDEX_NONE);
+
+	bCopyEachSourceOnEveryTarget = Settings->bCopyEachSourceOnEveryTarget;
 }
 
 FComputeDataProviderRenderProxy* UPCGCopyPointsAnalysisDataProvider::GetRenderProxy()
 {
-	const UPCGCopyPointsSettings* SettingsPtr = Settings;
-	check(SettingsPtr);
-
 	FPCGCopyPointsAnalysisDataProviderProxy::FCopyPointsAnalysisData_RenderThread ProxyData;
 	ProxyData.MatchAttributeId = MatchAttributeId;
 	ProxyData.SelectedFlagAttributeId = SelectedFlagAttributeId;
+	ProxyData.bCopyEachSourceOnEveryTarget = bCopyEachSourceOnEveryTarget;
 
 	return new FPCGCopyPointsAnalysisDataProviderProxy(MoveTemp(ProxyData));
 }
@@ -105,5 +112,6 @@ void FPCGCopyPointsAnalysisDataProviderProxy::GatherDispatchData(FDispatchData c
 		FParameters& Parameters = ParameterArray[InvocationIndex];
 		Parameters.MatchAttributeId = Data.MatchAttributeId;
 		Parameters.SelectedFlagAttributeId = Data.SelectedFlagAttributeId;
+		Parameters.CopyEachSourceOnEveryTarget = Data.bCopyEachSourceOnEveryTarget ? 1 : 0;
 	}
 }
diff --git a/Engine/Plugins/PCG/Source/PCG/Private/Compute/DataInterfaces/Elements/PCGCopyPointsAnalysisDataInterface.h b/Engine/Plugins/PCG/Source/PCG/Private/Compute/DataInterfaces/Elements/PCGCopyPointsAnalysisDataInterface.h
--- a/Engine/Plugins/PCG/Source/PCG/Private/Compute/DataInterfaces/Elements/PCGCopyPointsAnalysisDataInterface.h
+++ b/Engine/Plugins/PCG/Source/PCG/Private/Compute/DataInterfaces/Elements/PCGCopyPointsAnalysisDataInterface.h
@@ -46,6 +46,9 @@ public:
 
 	UPROPERTY()
 	int32 SelectedFlagAttributeId = INDEX_NONE;
+
+	UPROPERTY()
+	bool bCopyEachSourceOnEveryTarget = true;
 };
 
 class FPCGCopyPointsAnalysisDataProviderProxy : public FComputeDataProviderRenderProxy
@@ -55,6 +58,7 @@ public:
 	{
 		int32 MatchAttributeId = 0;
 		int32 SelectedFlagAttributeId = 0;
+		bool bCopyEachSourceOnEveryTarget = true;
 	};
 
 	FPCGCopyPointsAnalysisDataProviderProxy(FCopyPointsAnalysisData_RenderThread InData)
diff --git a/Engine/Plugins/PCG/Source/PCG/Private/Compute/DataInterfaces/Elements/PCGCopyPointsDataInterface.cpp b/Engine/Plugins/PCG/Source/PCG/Private/Compute/DataInterfaces/Elements/PCGCopyPointsDataInterface.cpp
--- a/Engine/Plugins/PCG/Source/PCG/Private/Compute/DataInterfaces/Elements/PCGCopyPointsDataInterface.cpp
+++ b/Engine/Plugins/PCG/Source/PCG/Private/Compute/DataInterfaces/Elements/PCGCopyPointsDataInterface.cpp
@@ -17,6 +17,8 @@
 #include "ShaderCompilerCore.h"
 #include "ShaderCore.h"
 #include "ShaderParameterMetadataBuilder.h"
+#include "SystemTextures.h"
+#include "ComputeFramework/ComputeKernelPermutationVector.h"
 #include "ComputeFramework/ShaderParamTypeDefinition.h"
 
 #include UE_INLINE_GENERATED_CPP_BY_NAME(PCGCopyPointsDataInterface)
@@ -57,7 +59,9 @@ void UPCGCopyPointsDataInterface::GetSupportedInputs(TArray<FShaderFunctionDefin
 
 	OutFunctions.AddDefaulted_GetRef()
 		.SetName(TEXT("CopyPoints_GetSourceAndTargetDataIndices"))
-		.AddParam(EShaderFundamentalType::Uint) // InOutputIndex
+		.AddParam(EShaderFundamentalType::Uint) // InOutputDataIndex
+		.AddParam(EShaderFundamentalType::Uint) // InSourceDataCount
+		.AddParam(EShaderFundamentalType::Uint) // InTargetDataCount
 		.AddParam(EShaderFundamentalType::Uint, 0, 0, EShaderParamModifier::Out) // OutSourceIndex
 		.AddParam(EShaderFundamentalType::Uint, 0, 0, EShaderParamModifier::Out); // OutTargetIndex
 }
@@ -70,7 +74,7 @@ BEGIN_SHADER_PARAMETER_STRUCT(FPCGCopyPointsDataInterfaceParameters,)
 	SHADER_PARAMETER(uint32, ColorInheritance)
 	SHADER_PARAMETER(uint32, SeedInheritance)
 	SHADER_PARAMETER(uint32, AttributeInheritance)
-	SHADER_PARAMETER(uint32, bCopyEachSourceOnEveryTarget)
+	SHADER_PARAMETER(uint32, CopyEachSourceOnEveryTarget)
 	SHADER_PARAMETER_RDG_BUFFER_SRV(StructuredBuffer<FUintVector2>, SourceAndTargetDataIndices)
 END_SHADER_PARAMETER_STRUCT()
 
@@ -79,6 +83,11 @@ void UPCGCopyPointsDataInterface::GetShaderParameters(TCHAR const* UID, FShaderP
 	InOutBuilder.AddNestedStruct<FPCGCopyPointsDataInterfaceParameters>(UID);
 }
 
+void UPCGCopyPointsDataInterface::GetPermutations(FComputeKernelPermutationVector& OutPermutationVector) const
+{
+	OutPermutationVector.AddPermutation(TEXT("USE_INPUT_DATA_INDICES_BUFFER"), 2);
+}
+
 void UPCGCopyPointsDataInterface::GetHLSL(FString& OutHLSL, FString const& InDataInterfaceName) const
 {
 	TMap<FString, FStringFormatArg> TemplateArgs =
@@ -94,7 +103,7 @@ void UPCGCopyPointsDataInterface::GetHLSL(FString& OutHLSL, FString const& InDat
 		"uint {DataInterfaceName}_ColorInheritance;\n"
 		"uint {DataInterfaceName}_SeedInheritance;\n"
 		"uint {DataInterfaceName}_AttributeInheritance;\n"
-		"uint {DataInterfaceName}_bCopyEachSourceOnEveryTarget;\n"
+		"uint {DataInterfaceName}_CopyEachSourceOnEveryTarget;\n"
 		"StructuredBuffer<uint2> {DataInterfaceName}_SourceAndTargetDataIndices;\n"
 		"\n"
 		"uint CopyPoints_GetRotationInheritance_{DataInterfaceName}() { return {DataInterfaceName}_RotationInheritance; }\n"
@@ -104,13 +113,26 @@ void UPCGCopyPointsDataInterface::GetHLSL(FString& OutHLSL, FString const& InDat
 		"uint CopyPoints_GetColorInheritance_{DataInterfaceName}() { return {DataInterfaceName}_ColorInheritance; }\n"
 		"uint CopyPoints_GetSeedInheritance_{DataInterfaceName}() { return {DataInterfaceName}_SeedInheritance; }\n"
 		"uint CopyPoints_GetAttributeInheritance_{DataInterfaceName}() { return {DataInterfaceName}_AttributeInheritance; }\n"
-		"uint CopyPoints_GetCopyEachSourceOnEveryTarget_{DataInterfaceName}() { return {DataInterfaceName}_bCopyEachSourceOnEveryTarget; }\n"
+		"uint CopyPoints_GetCopyEachSourceOnEveryTarget_{DataInterfaceName}() { return {DataInterfaceName}_CopyEachSourceOnEveryTarget; }\n"
 		"\n"
-		"void CopyPoints_GetSourceAndTargetDataIndices_{DataInterfaceName}(uint InOutputIndex, out uint OutSourceIndex, out uint OutTargetIndex)\n"
+		"void CopyPoints_GetSourceAndTargetDataIndices_{DataInterfaceName}(uint InOutputDataIndex, int InSourceDataCount, int InTargetDataCount, out uint OutSourceIndex, out uint OutTargetIndex)\n"
 		"{\n"
-		"	const uint2 Indices = {DataInterfaceName}_SourceAndTargetDataIndices[InOutputIndex];\n"
+		"#if USE_INPUT_DATA_INDICES_BUFFER\n"
+		"	const uint2 Indices = {DataInterfaceName}_SourceAndTargetDataIndices[InOutputDataIndex];\n"
 		"	OutSourceIndex = Indices[0];\n"
 		"	OutTargetIndex = Indices[1];\n"
+		"#else\n"
+		"	if (CopyPoints_GetCopyEachSourceOnEveryTarget_{DataInterfaceName}())\n"
+		"	{\n"
+		"		OutSourceIndex = InOutputDataIndex / InTargetDataCount;\n"
+		"		OutTargetIndex = InOutputDataIndex % InTargetDataCount;\n"
+		"	}\n"
+		"	else\n"
+		"	{\n"
+		"		OutSourceIndex = clamp(InOutputDataIndex, 0, -1 + (int)InSourceDataCount);\n"
+		"		OutTargetIndex = clamp(InOutputDataIndex, 0, -1 + (int)InTargetDataCount);\n"
+		"	}\n"
+		"#endif\n"
 		"}\n"
 		), TemplateArgs);
 }
@@ -210,52 +232,69 @@ bool UPCGCopyPointsDataProvider::PrepareForExecute_GameThread(UPCGDataBinding* I
 	}
 
 	const UPCGCopyPointsSettings* Settings = CastChecked<UPCGCopyPointsSettings>(ProducerKernel->GetSettings());
-	if (!Settings->bMatchBasedOnAttribute)
-	{
-		// If we're not matching based on attribute then we don't need to create flags.
-		return true;
-	}
-
-	int32 FirstAnalysisDataIndex = INDEX_NONE;
 
-	const int32 BindingIndex = InBinding->Graph ? InBinding->Graph->GetBindingIndex(ProducerKernel->GetKernelIndex(), PCGCopyPointsConstants::SelectedFlagsPinLabel, /*bIsInput=*/true) : INDEX_NONE;
-	if (ensure(BindingIndex != INDEX_NONE))
+	if (Settings->bMatchBasedOnAttribute)
 	{
-		if (const FName* VirtualLabel = InBinding->Graph->GetVirtualPinLabelFromBindingIndex(BindingIndex))
+		int32 FirstAnalysisDataIndex = INDEX_NONE;
+
+		const int32 BindingIndex = InBinding->Graph ? InBinding->Graph->GetBindingIndex(ProducerKernel->GetKernelIndex(), PCGCopyPointsConstants::SelectedFlagsPinLabel, /*bIsInput=*/true) : INDEX_NONE;
+		if (ensure(BindingIndex != INDEX_NONE))
 		{
-			FirstAnalysisDataIndex = InBinding->InputDataCollection.TaggedData.IndexOfByPredicate([VirtualLabel](const FPCGTaggedData& InData)
+			if (const FName* VirtualLabel = InBinding->Graph->GetVirtualPinLabelFromBindingIndex(BindingIndex))
 			{
-				return InData.Pin == *VirtualLabel;
-			});
+				FirstAnalysisDataIndex = InBinding->InputDataCollection.TaggedData.IndexOfByPredicate([VirtualLabel](const FPCGTaggedData& InData)
+				{
+					return InData.Pin == *VirtualLabel;
+				});
+			}
 		}
-	}
 
-	if (FirstAnalysisDataIndex == INDEX_NONE)
-	{
-		return true;
-	}
+		if (FirstAnalysisDataIndex == INDEX_NONE)
+		{
+			return true;
+		}
 
-	FPCGDataCollectionDesc SourcePinDesc;
-	FPCGDataCollectionDesc TargetPinDesc;
-	ensure(InBinding->GetKernelPinDataDesc_GameThread(ProducerKernel, PCGCopyPointsConstants::SourcePointsLabel, /*bIsInput=*/true, SourcePinDesc));
-	ensure(InBinding->GetKernelPinDataDesc_GameThread(ProducerKernel, PCGCopyPointsConstants::TargetPointsLabel, /*bIsInput=*/true, TargetPinDesc));
+		FPCGDataCollectionDesc SourcePinDesc;
+		FPCGDataCollectionDesc TargetPinDesc;
+		ensure(InBinding->GetKernelPinDataDesc_GameThread(ProducerKernel, PCGCopyPointsConstants::SourcePointsLabel, /*bIsInput=*/true, SourcePinDesc));
+		ensure(InBinding->GetKernelPinDataDesc_GameThread(ProducerKernel, PCGCopyPointsConstants::TargetPointsLabel, /*bIsInput=*/true, TargetPinDesc));
 
-	const uint32 NumOutputs = SourcePinDesc.DataDescs.Num()* TargetPinDesc.DataDescs.Num();
-	SourceAndTargetDataIndices.Reserve(NumOutputs);
+		const int32 NumOutputs = Settings->bCopyEachSourceOnEveryTarget ? SourcePinDesc.DataDescs.Num() * TargetPinDesc.DataDescs.Num() : FMath::Max(SourcePinDesc.DataDescs.Num(), TargetPinDesc.DataDescs.Num());
+		SourceAndTargetDataIndices.Reserve(NumOutputs);
 
-	const UPCGParamData* AnalysisResultsData = Cast<UPCGParamData>(InBinding->InputDataCollection.TaggedData[FirstAnalysisDataIndex].Data);
-	const UPCGMetadata* AnalysisMetadata = AnalysisResultsData ? AnalysisResultsData->ConstMetadata() : nullptr;
-	const FPCGMetadataAttributeBase* AnalysisAttributeBase = AnalysisMetadata ? AnalysisMetadata->GetConstAttribute(PCGCopyPointsConstants::SelectedFlagAttributeName) : nullptr;
-	const int32 NumElements = AnalysisAttributeBase ? AnalysisMetadata->GetItemCountForChild() : INDEX_NONE;
-	const FPCGMetadataAttribute<bool>* SelectedFlagsAttribute = ((NumElements == NumOutputs) && AnalysisAttributeBase && (AnalysisAttributeBase->GetTypeId() == PCG::Private::MetadataTypes<bool>::Id)) ? static_cast<const FPCGMetadataAttribute<bool>*>(AnalysisAttributeBase) : nullptr;
+		const UPCGParamData* AnalysisResultsData = Cast<UPCGParamData>(InBinding->InputDataCollection.TaggedData[FirstAnalysisDataIndex].Data);
+		const UPCGMetadata* AnalysisMetadata = AnalysisResultsData ? AnalysisResultsData->ConstMetadata() : nullptr;
+		const FPCGMetadataAttributeBase* AnalysisAttributeBase = AnalysisMetadata ? AnalysisMetadata->GetConstAttribute(PCGCopyPointsConstants::SelectedFlagAttributeName) : nullptr;
+		const int32 NumElements = AnalysisAttributeBase ? AnalysisMetadata->GetItemCountForChild() : INDEX_NONE;
+		const FPCGMetadataAttribute<bool>* SelectedFlagsAttribute = ((NumElements == NumOutputs) && AnalysisAttributeBase && (AnalysisAttributeBase->GetTypeId() == PCG::Private::MetadataTypes<bool>::Id)) ? static_cast<const FPCGMetadataAttribute<bool>*>(AnalysisAttributeBase) : nullptr;
 
-	for (int32 S = 0; S < SourcePinDesc.DataDescs.Num(); ++S)
-	{
-		for (int32 T = 0; T < TargetPinDesc.DataDescs.Num(); ++T)
+		if (SelectedFlagsAttribute)
 		{
-			if (!SelectedFlagsAttribute || SelectedFlagsAttribute->GetValue(S * TargetPinDesc.DataDescs.Num() + T))
+			if (Settings->bCopyEachSourceOnEveryTarget)
+			{
+				for (int32 S = 0; S < SourcePinDesc.DataDescs.Num(); ++S)
+				{
+					for (int32 T = 0; T < TargetPinDesc.DataDescs.Num(); ++T)
+					{
+						if (SelectedFlagsAttribute->GetValue(S * TargetPinDesc.DataDescs.Num() + T))
+						{
+							SourceAndTargetDataIndices.Emplace(S, T);
+						}
+					}
+				}
+			}
+			else
 			{
-				SourceAndTargetDataIndices.Emplace(S, T);
+				for (int32 OutputIndex = 0; OutputIndex < NumOutputs; ++OutputIndex)
+				{
+					const int32 S = FMath::Clamp(OutputIndex, 0, SourcePinDesc.DataDescs.Num() - 1);
+					const int32 T = FMath::Clamp(OutputIndex, 0, TargetPinDesc.DataDescs.Num() - 1);
+
+					if (SelectedFlagsAttribute->GetValue(OutputIndex))
+					{
+						SourceAndTargetDataIndices.Emplace(S, T);
+					}
+				}
 			}
 		}
 	}
@@ -299,24 +338,48 @@ bool FPCGCopyPointsDataProviderProxy::IsValid(FValidationData const& InValidatio
 	return true;
 }
 
+struct FCopyPointsDataInterfacePermutationIds
+{
+	FCopyPointsDataInterfacePermutationIds(FComputeKernelPermutationVector const& PermutationVector)
+	{
+		{
+			static FString Name(TEXT("USE_INPUT_DATA_INDICES_BUFFER"));
+			static uint32 Hash = GetTypeHash(Name);
+			UseInputDataIndicesBuffer = PermutationVector.GetPermutationBits(Name, Hash, 1);
+		}
+	}
+
+	uint32 UseInputDataIndicesBuffer = 0;
+};
+
+void FPCGCopyPointsDataProviderProxy::GatherPermutations(FPermutationData& InOutPermutationData) const
+{
+	FCopyPointsDataInterfacePermutationIds PermutationIds(InOutPermutationData.PermutationVector);
+
+	for (int32 InvocationIndex = 0; InvocationIndex < InOutPermutationData.NumInvocations; ++InvocationIndex)
+	{
+		// Select permutation based on whether we have data indices or not.
+		InOutPermutationData.PermutationIds[InvocationIndex] |= (Data.SourceAndTargetDataIndices.IsEmpty() ? 0 : PermutationIds.UseInputDataIndicesBuffer);
+	}
+}
+
 void FPCGCopyPointsDataProviderProxy::AllocateResources(FRDGBuilder& GraphBuilder, FAllocationData const& InAllocationData)
 {
 	TRACE_CPUPROFILER_EVENT_SCOPE(FPCGCopyPointsDataProviderProxy::AllocateResources);
 
-	// Ensure we always allocate at least 1 element so we have something to bind.
-	const uint32 NumIndexPairs = FMath::Max(1, Data.SourceAndTargetDataIndices.Num());
-	FRDGBufferDesc IndexPairsDesc = FRDGBufferDesc::CreateStructuredDesc(sizeof(FUintVector2), NumIndexPairs);
+	if (!Data.SourceAndTargetDataIndices.IsEmpty())
+	{
+		const FRDGBufferDesc IndexPairsDesc = FRDGBufferDesc::CreateStructuredDesc(sizeof(FUintVector2), Data.SourceAndTargetDataIndices.Num());
 
-	SourceAndTargetDataIndicesBuffer = GraphBuilder.CreateBuffer(IndexPairsDesc, TEXT("PCGCopyPoints_SourceAndTargetDataIndices"));
-	SourceAndTargetDataIndicesBufferSRV = GraphBuilder.CreateSRV(SourceAndTargetDataIndicesBuffer);
+		SourceAndTargetDataIndicesBuffer = GraphBuilder.CreateBuffer(IndexPairsDesc, TEXT("PCGCopyPoints_SourceAndTargetDataIndices"));
+		SourceAndTargetDataIndicesBufferSRV = GraphBuilder.CreateSRV(SourceAndTargetDataIndicesBuffer);
 
-	if (Data.SourceAndTargetDataIndices.IsEmpty())
-	{
-		GraphBuilder.QueueBufferUpload(SourceAndTargetDataIndicesBuffer, &FUintVector2::ZeroValue, 1);
+		GraphBuilder.QueueBufferUpload(SourceAndTargetDataIndicesBuffer, MakeArrayView(Data.SourceAndTargetDataIndices));
 	}
 	else
 	{
-		GraphBuilder.QueueBufferUpload(SourceAndTargetDataIndicesBuffer, MakeArrayView(Data.SourceAndTargetDataIndices));
+		// Fallback
+		SourceAndTargetDataIndicesBufferSRV = GraphBuilder.CreateSRV(FRDGBufferSRVDesc(GSystemTextures.GetDefaultStructuredBuffer(GraphBuilder, sizeof(FUintVector2))));
 	}
 }
 
@@ -333,7 +396,7 @@ void FPCGCopyPointsDataProviderProxy::GatherDispatchData(FDispatchData const& In
 		Parameters.ColorInheritance = Data.ColorInheritance;
 		Parameters.SeedInheritance = Data.SeedInheritance;
 		Parameters.AttributeInheritance = Data.AttributeInheritance;
-		Parameters.bCopyEachSourceOnEveryTarget = Data.bCopyEachSourceOnEveryTarget;
+		Parameters.CopyEachSourceOnEveryTarget = Data.bCopyEachSourceOnEveryTarget;
 		Parameters.SourceAndTargetDataIndices = SourceAndTargetDataIndicesBufferSRV;
 	}
 }
diff --git a/Engine/Plugins/PCG/Source/PCG/Private/Compute/DataInterfaces/Elements/PCGCopyPointsDataInterface.h b/Engine/Plugins/PCG/Source/PCG/Private/Compute/DataInterfaces/Elements/PCGCopyPointsDataInterface.h
--- a/Engine/Plugins/PCG/Source/PCG/Private/Compute/DataInterfaces/Elements/PCGCopyPointsDataInterface.h
+++ b/Engine/Plugins/PCG/Source/PCG/Private/Compute/DataInterfaces/Elements/PCGCopyPointsDataInterface.h
@@ -21,6 +21,7 @@ public:
 	TCHAR const* GetClassName() const override { return TEXT("PCGCopyPoints"); }
 	void GetSupportedInputs(TArray<FShaderFunctionDefinition>& OutFunctions) const override;
 	void GetShaderParameters(TCHAR const* UID, FShaderParametersMetadataBuilder& InOutBuilder, FShaderParametersMetadataAllocations& InOutAllocations) const override;
+	void GetPermutations(FComputeKernelPermutationVector& OutPermutationVector) const override;
 	void GetHLSL(FString& OutHLSL, FString const& InDataInterfaceName) const override;
 	UComputeDataProvider* CreateDataProvider() const override;
 	//~ End UComputeDataInterface Interface
@@ -73,6 +74,7 @@ public:
 
 	//~ Begin FComputeDataProviderRenderProxy Interface
 	bool IsValid(FValidationData const& InValidationData) const override;
+	void GatherPermutations(FPermutationData& InOutPermutationData) const override;
 	void AllocateResources(FRDGBuilder& GraphBuilder, FAllocationData const& InAllocationData) override;
 	void GatherDispatchData(FDispatchData const& InDispatchData) override;
 	//~ End FComputeDataProviderRenderProxy Interface
diff --git a/Engine/Plugins/PCG/Source/PCG/Private/Elements/PCGCopyPointsAnalysisKernel.cpp b/Engine/Plugins/PCG/Source/PCG/Private/Elements/PCGCopyPointsAnalysisKernel.cpp
--- a/Engine/Plugins/PCG/Source/PCG/Private/Elements/PCGCopyPointsAnalysisKernel.cpp
+++ b/Engine/Plugins/PCG/Source/PCG/Private/Elements/PCGCopyPointsAnalysisKernel.cpp
@@ -81,7 +81,18 @@ bool UPCGCopyPointsAnalysisKernel::ComputeOutputBindingDataDesc(const UPCGComput
 
 	FPCGDataDesc& DataDesc = OutDesc.DataDescs.Emplace_GetRef();
 	DataDesc.Type = EPCGDataType::Param;
-	DataDesc.ElementCount = NumSourceDatas * NumTargetDatas;
+
+	const UPCGCopyPointsSettings* CopyPointSettings = CastChecked<UPCGCopyPointsSettings>(Settings);
+
+	if (CopyPointSettings->bCopyEachSourceOnEveryTarget)
+	{
+		DataDesc.ElementCount = NumSourceDatas * NumTargetDatas;
+	}
+	else
+	{
+		// Validation should have ensured already that these are valid: (S, T) = (N, N), (N, 1) or (1, N).
+		DataDesc.ElementCount = FMath::Max(NumSourceDatas, NumTargetDatas);
+	}
 
 	FPCGKernelAttributeKey SelectedFlagAttributeKey = FPCGKernelAttributeKey(PCGCopyPointsConstants::SelectedFlagAttributeName, EPCGKernelAttributeType::Bool);
 	DataDesc.AddAttribute(SelectedFlagAttributeKey, InBinding);
diff --git a/Engine/Plugins/PCG/Source/PCG/Private/Elements/PCGCopyPointsKernel.cpp b/Engine/Plugins/PCG/Source/PCG/Private/Elements/PCGCopyPointsKernel.cpp
--- a/Engine/Plugins/PCG/Source/PCG/Private/Elements/PCGCopyPointsKernel.cpp
+++ b/Engine/Plugins/PCG/Source/PCG/Private/Elements/PCGCopyPointsKernel.cpp
@@ -225,17 +225,7 @@ int UPCGCopyPointsKernel::ComputeThreadCount(const UPCGDataBinding* InBinding) c
 FString UPCGCopyPointsKernel::GetCookedSource(const UPCGComputeGraph* InComputeGraph) const
 {
 	FString TemplateFile;
-
-	// todo_pcg: Figure out how permutations work (GetPermutations/GatherPermutations) so we can unify the two shader files into one.
-	if (Settings && Settings->bMatchBasedOnAttribute)
-	{
-		ensure(LoadShaderSourceFile(TEXT("/Plugin/PCG/Private/Elements/PCGCopyPointsConditional.usf"), EShaderPlatform::SP_PCD3D_SM5, &TemplateFile, nullptr));
-	}
-	else
-	{
-		ensure(LoadShaderSourceFile(TEXT("/Plugin/PCG/Private/Elements/PCGCopyPoints.usf"), EShaderPlatform::SP_PCD3D_SM5, &TemplateFile, nullptr));
-	}
-
+	ensure(LoadShaderSourceFile(TEXT("/Plugin/PCG/Private/Elements/PCGCopyPoints.usf"), EShaderPlatform::SP_PCD3D_SM5, &TemplateFile, nullptr));
 	return TemplateFile;
 }
 
diff --git a/Engine/Plugins/PCG/Source/PCG/Private/Elements/PCGCopyPointsKernelShared.cpp b/Engine/Plugins/PCG/Source/PCG/Private/Elements/PCGCopyPointsKernelShared.cpp
--- a/Engine/Plugins/PCG/Source/PCG/Private/Elements/PCGCopyPointsKernelShared.cpp
+++ b/Engine/Plugins/PCG/Source/PCG/Private/Elements/PCGCopyPointsKernelShared.cpp
@@ -14,12 +14,6 @@ bool PCGCopyPointsKernel::IsKernelValid(const UPCGComputeKernel* InKernel, const
 {
 	if (InCopyPointSettings->bMatchBasedOnAttribute)
 	{
-		if (!InCopyPointSettings->bCopyEachSourceOnEveryTarget)
-		{
-			PCG_KERNEL_VALIDATION_ERR(InContext, InCopyPointSettings, bQuiet, LOCTEXT("CopyEachSourceOnEveryTargetNotSupported", "Copy Each Source On Every Target must be enabled when Match Based On Attribute is enabled."));
-			return false;
-		}
-
 		if (const UPCGDataBinding* DataBinding = InContext->DataBinding.Get(); ensure(DataBinding))
 		{
 			auto ValidateAttributeExists = [InContext, DataBinding, InKernel, InCopyPointSettings, bQuiet](FName InputPin) -> bool
