Fixed typo in apply rotation test in Apply Hierarchy node.
Added skip apply hierarchy option, to not affect points in some cases.

#rnx

diff --git a/Engine/Plugins/PCG/Source/PCG/Private/Elements/PCGApplyHierarchy.cpp b/Engine/Plugins/PCG/Source/PCG/Private/Elements/PCGApplyHierarchy.cpp
--- a/Engine/Plugins/PCG/Source/PCG/Private/Elements/PCGApplyHierarchy.cpp
+++ b/Engine/Plugins/PCG/Source/PCG/Private/Elements/PCGApplyHierarchy.cpp
@@ -21,6 +21,7 @@ UPCGApplyHierarchySettings::UPCGApplyHierarchySettings()
 	ApplyParentRotationAttribute.SetAttributeName(PCGLevelToAssetConstants::IgnoreParentRotationAttributeName);
 	ApplyParentScale = EPCGApplyHierarchyOption::OptOutByAttribute;
 	ApplyParentScaleAttribute.SetAttributeName(PCGLevelToAssetConstants::IgnoreParentScaleAttributeName);
+	ApplyHierarchy = EPCGApplyHierarchyOption::Always;
 }
 
 FPCGElementPtr UPCGApplyHierarchySettings::CreateElement() const
@@ -166,40 +167,42 @@ bool FPCGApplyHierarchyElement::PrepareDataInternal(FPCGContext* InContext) cons
 			return EPCGTimeSliceInitResult::NoOperation;
 		}
 
-		if (Settings->ApplyParentRotation == EPCGApplyHierarchyOption::Always || Settings->ApplyParentRotation == EPCGApplyHierarchyOption::Never)
+		auto PrepareApplyOptionsAccessor = [Context, InputData](const FPCGAttributePropertyInputSelector& InputSelector, EPCGApplyHierarchyOption Option, TUniquePtr<const IPCGAttributeAccessor>& OutAccessor, bool& bOutInvert) -> bool
 		{
-			OutState.ApplyRotationAccessor = MakeUnique<FPCGConstantValueAccessor<bool>>(true);
-			OutState.bInvertApplyRotation = Settings->ApplyParentRotation == EPCGApplyHierarchyOption::Never;
-		}
-		else
-		{
-			FPCGAttributePropertyInputSelector ApplyRotationSelector = Settings->ApplyParentRotationAttribute.CopyAndFixLast(InputData);
-			OutState.ApplyRotationAccessor = PCGAttributeAccessorHelpers::CreateConstAccessor(InputData, ApplyRotationSelector);
-			OutState.bInvertApplyRotation = Settings->ApplyParentScale == EPCGApplyHierarchyOption::OptOutByAttribute;
-
-			if (!OutState.ApplyRotationAccessor)
+			if (Option == EPCGApplyHierarchyOption::Always || Option == EPCGApplyHierarchyOption::Never)
 			{
-				PCGLog::Metadata::LogFailToCreateAccessorError(ApplyRotationSelector, Context);
-				return EPCGTimeSliceInitResult::NoOperation;
+				OutAccessor = MakeUnique<FPCGConstantValueAccessor<bool>>(true);
+				bOutInvert = Option == EPCGApplyHierarchyOption::Never;
 			}
-		}
+			else
+			{
+				FPCGAttributePropertyInputSelector Selector = InputSelector.CopyAndFixLast(InputData);
+				OutAccessor = PCGAttributeAccessorHelpers::CreateConstAccessor(InputData, Selector);
+				bOutInvert = Option == EPCGApplyHierarchyOption::OptOutByAttribute;
+
+				if (!OutAccessor)
+				{
+					PCGLog::Metadata::LogFailToCreateAccessorError(InputSelector, Context);
+					return false;
+				}
+			}
+
+			return true;
+		};
 
-		if (Settings->ApplyParentScale == EPCGApplyHierarchyOption::Always || Settings->ApplyParentScale == EPCGApplyHierarchyOption::Never)
+		if (!PrepareApplyOptionsAccessor(Settings->ApplyParentRotationAttribute, Settings->ApplyParentRotation, OutState.ApplyRotationAccessor, OutState.bInvertApplyRotation))
 		{
-			OutState.ApplyScaleAccessor = MakeUnique<FPCGConstantValueAccessor<bool>>(true);
-			OutState.bInvertApplyScale = Settings->ApplyParentScale == EPCGApplyHierarchyOption::Never;
+			return EPCGTimeSliceInitResult::NoOperation;
 		}
-		else
+
+		if (!PrepareApplyOptionsAccessor(Settings->ApplyParentScaleAttribute, Settings->ApplyParentScale, OutState.ApplyScaleAccessor, OutState.bInvertApplyScale))
 		{
-			FPCGAttributePropertyInputSelector ApplyScaleSelector = Settings->ApplyParentScaleAttribute.CopyAndFixLast(InputData);
-			OutState.ApplyScaleAccessor = PCGAttributeAccessorHelpers::CreateConstAccessor(InputData, ApplyScaleSelector);
-			OutState.bInvertApplyScale = Settings->ApplyParentScale == EPCGApplyHierarchyOption::OptOutByAttribute;
+			return EPCGTimeSliceInitResult::NoOperation;
+		}
 
-			if (!OutState.ApplyScaleAccessor)
-			{
-				PCGLog::Metadata::LogFailToCreateAccessorError(ApplyScaleSelector, Context);
-				return EPCGTimeSliceInitResult::NoOperation;
-			}
+		if (!PrepareApplyOptionsAccessor(Settings->ApplyHierarchyAttribute, Settings->ApplyHierarchy, OutState.ApplyHierarchyAccessor, OutState.bInvertApplyHierarchy))
+		{
+			return EPCGTimeSliceInitResult::NoOperation;
 		}
 
 		// Finally, get the relative transforms
@@ -498,11 +501,20 @@ bool FPCGApplyHierarchyElement::ExecuteInternal(FPCGContext* InContext) const
 						ApplyParentScale.SetNumUninitialized(Count);
 						IterState.ApplyScaleAccessor->GetRange<bool>(ApplyParentScale, StartReadIndex, *SubsetKeys, EPCGAttributeAccessorFlags::AllowConstructible);
 
+						TArray<bool, TInlineAllocator<ChunkSize>> ApplyHierarchy;
+						ApplyHierarchy.SetNumUninitialized(Count);
+						IterState.ApplyHierarchyAccessor->GetRange<bool>(ApplyHierarchy, StartReadIndex, *SubsetKeys, EPCGAttributeAccessorFlags::AllowConstructible);
+
 						for (int32 Index = 0; Index < Count; ++Index)
 						{
 							const int32 CurrentPointIndex = CurrentDepthIndices[Index + StartReadIndex];
 							const int32 CurrentParentIndex = IterState.ParentIndices[CurrentPointIndex];
 
+							if (ApplyHierarchy[Index] == IterState.bInvertApplyHierarchy)
+							{
+								continue;
+							}
+
 							// At this point, if the current parent is invalid, or it has been marked invalid, we'll mark this point as invalid too.
 							if(CurrentParentIndex < 0 || IterState.ParentIndices[CurrentParentIndex] == InvalidParentIndex)
 							{
@@ -522,7 +534,7 @@ bool FPCGApplyHierarchyElement::ExecuteInternal(FPCGContext* InContext) const
 							Transform = RelativeTransform * ParentTransform;
 
 							// Finally, apply options (ignore scale/rotation from parent)
-							if (ApplyParentRotation[Index] == IterState.bInvertApplyScale)
+							if (ApplyParentRotation[Index] == IterState.bInvertApplyRotation)
 							{
 								Transform.SetRotation(RelativeTransform.GetRotation());
 							}
diff --git a/Engine/Plugins/PCG/Source/PCG/Public/Elements/PCGApplyHierarchy.h b/Engine/Plugins/PCG/Source/PCG/Public/Elements/PCGApplyHierarchy.h
--- a/Engine/Plugins/PCG/Source/PCG/Public/Elements/PCGApplyHierarchy.h
+++ b/Engine/Plugins/PCG/Source/PCG/Public/Elements/PCGApplyHierarchy.h
@@ -70,6 +70,12 @@ public:
 	UPROPERTY(BlueprintReadWrite, EditAnywhere, Category = Settings, meta = (EditCondition = "ApplyParentScale==EPCGApplyHierarchyOption::OptInByAttribute||ApplyParentScale==EPCGApplyHierarchyOption::OptOutByAttribute", EditConditionHides, PCG_Overridable))
 	FPCGAttributePropertyInputSelector ApplyParentScaleAttribute;
 
+	UPROPERTY(BlueprintReadWrite, EditAnywhere, Category = Advanced, meta = (PCG_Overridable))
+	EPCGApplyHierarchyOption ApplyHierarchy;
+
+	UPROPERTY(BlueprintReadWrite, EditAnywhere, Category = Advanced, meta = (EditCondition = "ApplyHierarchy==EPCGApplyHierarchyOption::OptInByAttribute||ApplyHierarchy==EPCGApplyHierarchyOption::OptOutByAttribute", EditConditionHides, PCG_Overridable))
+	FPCGAttributePropertyInputSelector ApplyHierarchyAttribute;
+
 	UPROPERTY(BlueprintReadWrite, EditAnywhere, Category = Advanced, meta = (PCG_Overridable))
 	bool bWarnOnPointsWithInvalidParent = true;
 };
@@ -100,6 +106,9 @@ namespace PCGApplyHierarchyElement
 
 		TUniquePtr<const IPCGAttributeAccessor> ApplyScaleAccessor;
 		bool bInvertApplyScale = false;
+
+		TUniquePtr<const IPCGAttributeAccessor> ApplyHierarchyAccessor;
+		bool bInvertApplyHierarchy = false;
 		
 		// Built during execution
 		TArray<int32> ParentIndices;
