[PCG] Optimization: GraphExecutor improved Paused task scheduling
- Elements can add scheduled dependencies in FPCGContext::DynamicDependencies
- GraphExecutor will handle task completion and rescheduling when DynamicDependencies are completed
- Prevents usage of costly round trip to mainthread with ScheduleGeneric (Major improvements for LoopElements and SubGraphs)
- Disabled by default so it can be tested on a wider range of graphs first (pcg.Experimental.EnableDynamicDependencies)

#rnx

diff --git a/Engine/Plugins/PCG/Source/PCG/Private/Elements/PCGLoopElement.cpp b/Engine/Plugins/PCG/Source/PCG/Private/Elements/PCGLoopElement.cpp
--- a/Engine/Plugins/PCG/Source/PCG/Private/Elements/PCGLoopElement.cpp
+++ b/Engine/Plugins/PCG/Source/PCG/Private/Elements/PCGLoopElement.cpp
@@ -17,6 +17,11 @@
 
 #define LOCTEXT_NAMESPACE "PCGLoopElement"
 
+namespace PCGGraphExecutor
+{
+	extern TAutoConsoleVariable<bool> CVarEnableDynamicDependencies;
+}
+
 #if WITH_EDITOR
 FText UPCGLoopSettings::GetDefaultNodeTitle() const
 {
@@ -322,27 +327,42 @@ bool FPCGLoopElement::ExecuteInternal(FPCGContext* InContext) const
 			Context->bScheduledSubgraph = true;
 			Context->bIsPaused = true;
 
-			Subsystem->ScheduleGeneric(
-				[ContextHandle = Context->GetOrCreateHandle()]() // Normal execution: Wake up the current task
+			// In the case of a feedback loop we only need to depend on the last scheduled task since all other tasks depend on each other
+			if (PCGGraphExecutor::CVarEnableDynamicDependencies.GetValueOnAnyThread())
+			{
+				if (bIsFeedbackLoop)
 				{
-					if (FPCGSubgraphContext* ContextPtr = FPCGContext::GetContextFromHandle<FPCGSubgraphContext>(ContextHandle))
-					{
-						ContextPtr->bIsPaused = false;
-					}
-					return true;
-				},
-				[ContextHandle = Context->GetOrCreateHandle()]() // On abort: wakeup and cancel, forget subgraphs
+					Context->DynamicDependencies.Add(Context->SubgraphTaskIds.Last());
+				}
+				else
 				{
-					if (FPCGSubgraphContext* ContextPtr = FPCGContext::GetContextFromHandle<FPCGSubgraphContext>(ContextHandle))
+					Context->DynamicDependencies.Append(Context->SubgraphTaskIds);
+				}
+			}
+			else
+			{
+				Subsystem->ScheduleGeneric(
+					[ContextHandle = Context->GetOrCreateHandle()]() // Normal execution: Wake up the current task
 					{
-						ContextPtr->bIsPaused = false;
-						ContextPtr->SubgraphTaskIds.Reset();
-						ContextPtr->OutputData.bCancelExecution = true;
-					}
-					return true;
-				},
-				Context->SourceComponent.Get(),
-				Context->SubgraphTaskIds);
+						if (FPCGSubgraphContext* ContextPtr = FPCGContext::GetContextFromHandle<FPCGSubgraphContext>(ContextHandle))
+						{
+							ContextPtr->bIsPaused = false;
+						}
+						return true;
+					},
+					[ContextHandle = Context->GetOrCreateHandle()]() // On abort: wakeup and cancel, forget subgraphs
+					{
+						if (FPCGSubgraphContext* ContextPtr = FPCGContext::GetContextFromHandle<FPCGSubgraphContext>(ContextHandle))
+						{
+							ContextPtr->bIsPaused = false;
+							ContextPtr->SubgraphTaskIds.Reset();
+							ContextPtr->OutputData.bCancelExecution = true;
+						}
+						return true;
+					},
+					Context->SourceComponent.Get(),
+					Context->SubgraphTaskIds);
+			}
 
 			return false;
 		}
diff --git a/Engine/Plugins/PCG/Source/PCG/Private/Elements/PCGSpawnActor.cpp b/Engine/Plugins/PCG/Source/PCG/Private/Elements/PCGSpawnActor.cpp
--- a/Engine/Plugins/PCG/Source/PCG/Private/Elements/PCGSpawnActor.cpp
+++ b/Engine/Plugins/PCG/Source/PCG/Private/Elements/PCGSpawnActor.cpp
@@ -38,6 +38,11 @@ static TAutoConsoleVariable<bool> CVarAllowActorReuse(
 	true,
 	TEXT("Controls whether PCG spawned actors can be reused and skipped when re-executing"));
 
+namespace PCGGraphExecutor
+{
+	extern TAutoConsoleVariable<bool> CVarEnableDynamicDependencies;
+}
+
 class FPCGSpawnActorPartitionByAttribute : public FPCGDataPartitionBase<FPCGSpawnActorPartitionByAttribute, TSubclassOf<AActor>>
 {
 public:
@@ -675,25 +680,32 @@ bool FPCGSpawnActorElement::SpawnAndPrepareSubgraphs(FPCGSubgraphContext* Contex
 		Context->bScheduledSubgraph = true;
 		Context->bIsPaused = true;
 
-		Subsystem->ScheduleGeneric(
-			[ContextHandle = Context->GetOrCreateHandle()]() // Normal execution: Wake up the current task
-			{
-				if (FPCGSubgraphContext* ContextPtr = FPCGContext::GetContextFromHandle<FPCGSubgraphContext>(ContextHandle))
+		if (PCGGraphExecutor::CVarEnableDynamicDependencies.GetValueOnAnyThread())
+		{
+			Context->DynamicDependencies.Append(Context->SubgraphTaskIds);
+		}
+		else
+		{
+			Subsystem->ScheduleGeneric(
+				[ContextHandle = Context->GetOrCreateHandle()]() // Normal execution: Wake up the current task
 				{
-					ContextPtr->bIsPaused = false;
-				}
-				return true;
-			}, 
-			[ContextHandle = Context->GetOrCreateHandle()]() // On Abort: Wake up & cancel
-			{
-				if (FPCGSubgraphContext* ContextPtr = FPCGContext::GetContextFromHandle<FPCGSubgraphContext>(ContextHandle))
+					if (FPCGSubgraphContext* ContextPtr = FPCGContext::GetContextFromHandle<FPCGSubgraphContext>(ContextHandle))
+					{
+						ContextPtr->bIsPaused = false;
+					}
+					return true;
+				},
+				[ContextHandle = Context->GetOrCreateHandle()]() // On Abort: Wake up & cancel
 				{
-					ContextPtr->bIsPaused = false;
-					ContextPtr->OutputData.bCancelExecution = true;
-				}
-			},
-			Context->SourceComponent.Get(), 
-			Context->SubgraphTaskIds);
+					if (FPCGSubgraphContext* ContextPtr = FPCGContext::GetContextFromHandle<FPCGSubgraphContext>(ContextHandle))
+					{
+						ContextPtr->bIsPaused = false;
+						ContextPtr->OutputData.bCancelExecution = true;
+					}
+				},
+				Context->SourceComponent.Get(),
+				Context->SubgraphTaskIds);
+		}
 
 		return false;
 	}
diff --git a/Engine/Plugins/PCG/Source/PCG/Private/Graph/PCGGraphExecutor.cpp b/Engine/Plugins/PCG/Source/PCG/Private/Graph/PCGGraphExecutor.cpp
--- a/Engine/Plugins/PCG/Source/PCG/Private/Graph/PCGGraphExecutor.cpp
+++ b/Engine/Plugins/PCG/Source/PCG/Private/Graph/PCGGraphExecutor.cpp
@@ -96,6 +96,11 @@ namespace PCGGraphExecutor
 		true,
 		TEXT("Whether proxies for GPU data are cached in per pin output data and passed through grid links. If false data is read back to CPU."));
 
+	TAutoConsoleVariable<bool> CVarEnableDynamicDependencies(
+		TEXT("pcg.Experimental.EnableDynamicDependencies"),
+		false,
+		TEXT("Support re-scheduling of paused tasks through dynamic dependencies"));
+
 	void ClearAsyncFlags(TSet<TObjectPtr<UObject>>& AsyncObjects)
 	{
 		// Remove Async Flags from Objects created in Async PCG Tasks so that they can get tracked by the GarbageCollector properly
@@ -119,6 +124,16 @@ namespace PCGGraphExecutor
 		}
 	}
 
+	// Needs to be called by owner of LiveTasksLock
+	void InsertToActiveTaskArrayNoLock(int32 Index, TArray<TSharedPtr<FPCGGraphActiveTask>>& ActiveTaskArray, TSharedPtr<FPCGGraphActiveTask> ActiveTask)
+	{
+		ActiveTaskArray.Insert(ActiveTask, Index);
+		for (int32 i = Index; i < ActiveTaskArray.Num(); ++i)
+		{
+			ActiveTaskArray[i]->TaskIndex = i;
+		}
+	}
+
 	// Needs to be called by owner of LiveTasksLock
 	void AddToActiveTaskArrayNoLock(TArray<TSharedPtr<FPCGGraphActiveTask>>& ActiveTaskArray, TSharedPtr<FPCGGraphActiveTask> ActiveTask)
 	{
@@ -565,6 +580,7 @@ TSet<UPCGComponent*> FPCGGraphExecutor::Cancel(TFunctionRef<bool(TWeakObjectPtr<
 			TArray<TSharedPtr<FPCGGraphActiveTask>> CancelledPausedTasks;
 			{
 				PCGGraphExecutor::TScopeLock ScopeLock(LiveTasksLock);
+				PCGGraphExecutor::TScopeLock PauseScopeLock(PausedTaskSuccessorsLock);
 				for (int32 PausedTaskIndex = PausedTasks.Num() - 1; PausedTaskIndex >= 0; --PausedTaskIndex)
 				{
 					TSharedPtr<FPCGGraphActiveTask> Task = PausedTasks[PausedTaskIndex];
@@ -574,6 +590,26 @@ TSet<UPCGComponent*> FPCGGraphExecutor::Cancel(TFunctionRef<bool(TWeakObjectPtr<
 						CancelledPausedTasks.Add(Task);
 						PCGGraphExecutor::RemoveAtFromActiveTaskArrayNoLock(PausedTasks, PausedTaskIndex);
 						Task->TaskIndex = INDEX_NONE;
+
+						if (!Task->Context->DynamicDependencies.IsEmpty())
+						{
+							// Remove this task from the PausedTaskSuccessors
+							for (FPCGTaskId DynamicDependency : Task->Context->DynamicDependencies)
+							{
+								if (TArray<TSharedPtr<FPCGGraphActiveTask>>* Successors = PausedTaskSuccessors.Find(DynamicDependency))
+								{
+									Successors->Remove(Task);
+
+									// If there are no more successors, remove the entry all together
+									if (Successors->IsEmpty())
+									{
+										PausedTaskSuccessors.Remove(DynamicDependency);
+									}
+								}
+							}
+							Task->Context->OutputData.bCancelExecution = true;
+							Task->Context->DynamicDependencies.Empty();
+						}
 					}
 				}
 			}
@@ -874,6 +910,7 @@ bool FPCGGraphExecutor::ProcessScheduledTasks()
 
 			PCGGraphExecutor::TScopeLock ScopeLock(TasksLock);
 			PCGGraphExecutor::TScopeLock ChildScopeLock(LiveTasksLock);
+			PCGGraphExecutor::TScopeLock PausedLock(PausedTaskSuccessorsLock);
 			PCGGraphExecutor::TScopeLock CachingResultsScopeLock(CachingResultsLock);
 			// This is a safeguard to check if we're in a stuck state
 			if (CollectGCCachingResults.Num() == 0 && ReadyTasks.Num() == 0 && ActiveTasks.Num() == 0 && ActiveTasksGameThreadOnly.Num() == 0 && PausedTasks.Num() == 0 && Tasks.Num() > 0)
@@ -969,10 +1006,40 @@ void FPCGGraphExecutor::PostTaskExecute(TSharedPtr<FPCGGraphActiveTask> ActiveTa
 			const int32 TaskRemoveIndex = ActiveTask.TaskIndex;
 			check(TaskRemoveIndex != INDEX_NONE && ActiveTaskArrayToRemove[TaskRemoveIndex] == ActiveTaskPtr);
 
-			if (ActiveTask.Context->bIsPaused)
+			const bool bHasPausedDependencies = !ActiveTask.Context->DynamicDependencies.IsEmpty();
+			if (ActiveTask.Context->bIsPaused || bHasPausedDependencies)
 			{
-				PCGGraphExecutor::AddToActiveTaskArrayNoLock(PausedTasks, ActiveTaskPtr);
-				PCGGraphExecutor::RemoveAtFromActiveTaskArrayNoLock(ActiveTaskArrayToRemove, TaskRemoveIndex);
+				ActiveTask.Context->bIsPaused = true;
+								
+				bool bMoveFromActiveToPaused = !bHasPausedDependencies;
+
+				// When using DynamicDependencies the task will be un-paused  (bIsPaused will toggle back to false) after completion of those dependency tasks.
+				// It is not possible for a task to have DynamicDependencies and some other external mechanism to unpause (like polling).
+				// This could be supported but for now we have no use for this added complexity.
+				if (bHasPausedDependencies)
+				{
+					PCGGraphExecutor::TScopeLock PauseLock(PausedTaskSuccessorsLock);
+					PCGGraphExecutor::TScopeLock OutputLock(TaskOutputsLock);
+					for (FPCGTaskId DependencyId : ActiveTask.Context->DynamicDependencies)
+					{
+						if (!TaskOutputs.Contains(DependencyId))
+						{
+							PausedTaskSuccessors.FindOrAdd(DependencyId).Add(ActiveTaskPtr);
+							bMoveFromActiveToPaused = true;
+						}
+					}
+				}
+
+				if (bMoveFromActiveToPaused)
+				{
+					PCGGraphExecutor::AddToActiveTaskArrayNoLock(PausedTasks, ActiveTaskPtr);
+					PCGGraphExecutor::RemoveAtFromActiveTaskArrayNoLock(ActiveTaskArrayToRemove, TaskRemoveIndex);
+				}
+				else
+				{
+					// All the paused dependencies have completed
+					ActiveTask.Context->bIsPaused = false;
+				}							
 			}
 			else if (ActiveTask.bIsGameThreadOnly != ActiveTask.CanExecuteOnlyOnMainThread())
 			{
@@ -1391,7 +1458,7 @@ bool FPCGGraphExecutor::ExecuteScheduling(double EndTime, TSharedPtr<FPCGGraphAc
 					TArray<TSharedPtr<FPCGGraphActiveTask>>& ActiveTaskArray = bIsGameThreadOnly ? ActiveTasksGameThreadOnly : ActiveTasks;
 					PausedTask->bIsGameThreadOnly = bIsGameThreadOnly;
 
-					PCGGraphExecutor::AddToActiveTaskArrayNoLock(ActiveTaskArray, PausedTask);
+					PCGGraphExecutor::InsertToActiveTaskArrayNoLock(0, ActiveTaskArray, PausedTask);
 					PCGGraphExecutor::RemoveAtFromActiveTaskArrayNoLock(PausedTasks, PausedTaskIndex);
 					bStateChanged = true;
 				}
@@ -1545,8 +1612,14 @@ bool FPCGGraphExecutor::ExecuteScheduling(double EndTime, TSharedPtr<FPCGGraphAc
 
 							if (bIsDone || bIsPaused)
 							{
+								// If this task is paused it is possible it scheduled some new tasks, process those new tasks before calling ExecuteScheduling (so we are not one frame behind)
+								if (bIsPaused)
+								{
+									ProcessScheduledTasks();
+								}
+
 								// Budget is provided by CVarAsyncOutOfTickBudgetInMilliseconds
-								ExecuteScheduling(-1.0);
+								ExecuteScheduling(-1.0, nullptr,  /*bForceCheckPausedTasks=*/bIsPaused);
 							}
 
 							return bIsDone;
@@ -1669,6 +1742,7 @@ void FPCGGraphExecutor::ClearAllTasks()
 	ActiveTasks.Reset();
 	ActiveTasksGameThreadOnly.Reset();
 	PausedTasks.Reset();
+	PausedTaskSuccessors.Reset();
 }
 
 TArray<FPCGGraphExecutor::FCachedResult*> FPCGGraphExecutor::QueueNextTasksInternal(FPCGTaskId FinishedTask)
@@ -1719,6 +1793,24 @@ TArray<FPCGGraphExecutor::FCachedResult*> FPCGGraphExecutor::QueueNextTasksInter
 		OnTaskInputsReady(NewReadyTask, CachedResults, bIsInGameThread);
 		--NewReadyTaskCount;
 	}
+
+	{
+		PCGGraphExecutor::TScopeLock PauseScopeLock(PausedTaskSuccessorsLock);
+
+		TArray<TSharedPtr<FPCGGraphActiveTask>> ActiveTaskSuccessors;
+
+		if (PausedTaskSuccessors.RemoveAndCopyValue(FinishedTask, ActiveTaskSuccessors))
+		{
+			for (const TSharedPtr<FPCGGraphActiveTask>& ActiveTask : ActiveTaskSuccessors)
+			{
+				ActiveTask->Context->DynamicDependencies.Remove(FinishedTask);
+				if (ActiveTask->Context->DynamicDependencies.IsEmpty())
+				{
+					ActiveTask->Context->bIsPaused = false;
+				}
+			}
+		}
+	}
 	
 	return CachedResults;
 }
diff --git a/Engine/Plugins/PCG/Source/PCG/Private/Graph/PCGGraphExecutor.h b/Engine/Plugins/PCG/Source/PCG/Private/Graph/PCGGraphExecutor.h
--- a/Engine/Plugins/PCG/Source/PCG/Private/Graph/PCGGraphExecutor.h
+++ b/Engine/Plugins/PCG/Source/PCG/Private/Graph/PCGGraphExecutor.h
@@ -305,7 +305,7 @@ private:
 	TArray<TSharedPtr<FPCGGraphActiveTask>> PausedTasks;
 	bool bNeedToCheckPausedTasks = false;
 
-	/** Lock level 3 */
+	/** Lock level leaf */
 	FPCGLock CollectGCReferenceTasksLock;
 	TSet<TSharedPtr<FPCGGraphActiveTask>> CollectGCReferenceTasks;
 			
@@ -323,11 +323,15 @@ private:
 	};
 
 	/** Lock level 4 */
+	UE::FSpinLock PausedTaskSuccessorsLock;
+	TMap<FPCGTaskId, TArray<TSharedPtr<FPCGGraphActiveTask>>> PausedTaskSuccessors;
+	
+	/** Lock level leaf */
 	FPCGLock CachingResultsLock;
 	// Used to keep GC references to in flight caching results (not yet stored to output and might not be in cache anymore)
 	TMap<FPCGTaskId, TUniquePtr<FCachedResult>> CollectGCCachingResults;
 
-	/** Lock level 4 */
+	/** Lock level leaf */
 	FPCGLock TaskOutputsLock;
 	TMap<FPCGTaskId, FOutputDataInfo> TaskOutputs;
 	
diff --git a/Engine/Plugins/PCG/Source/PCG/Private/Graph/PCGGraphTask.h b/Engine/Plugins/PCG/Source/PCG/Private/Graph/PCGGraphTask.h
--- a/Engine/Plugins/PCG/Source/PCG/Private/Graph/PCGGraphTask.h
+++ b/Engine/Plugins/PCG/Source/PCG/Private/Graph/PCGGraphTask.h
@@ -200,7 +200,7 @@ struct FPCGGraphActiveTask : TSharedFromThis<FPCGGraphActiveTask>
 
 	// Used to know if task should be in ActiveTasks or ActiveTasksGameThreadOnly
 	bool bIsGameThreadOnly = false;
-	// TaskIndex inside ActiveTasks/ActiveTasksGameThreadOnly/SleepingTasks for fast removal
+	// TaskIndex inside ActiveTasks/ActiveTasksGameThreadOnly/PausedTasks for fast removal
 	int32 TaskIndex = INDEX_NONE;
 
 	static int32 NumExecuting;
diff --git a/Engine/Plugins/PCG/Source/PCG/Private/PCGSubgraph.cpp b/Engine/Plugins/PCG/Source/PCG/Private/PCGSubgraph.cpp
--- a/Engine/Plugins/PCG/Source/PCG/Private/PCGSubgraph.cpp
+++ b/Engine/Plugins/PCG/Source/PCG/Private/PCGSubgraph.cpp
@@ -18,6 +18,11 @@
 
 #define LOCTEXT_NAMESPACE "PCGSubgraphElement"
 
+namespace PCGGraphExecutor
+{
+	extern TAutoConsoleVariable<bool> CVarEnableDynamicDependencies;
+}
+
 namespace PCGSubgraphSettings
 {
 	void RemoveAdvancedAndInvisibleOnConnectedPins(const UPCGGraph* Subgraph, TArray<FPCGPinProperties>& InOutPinProperties, const bool bIsInput)
@@ -788,29 +793,36 @@ bool FPCGSubgraphElement::ExecuteInternal(FPCGContext* InContext) const
 					Context->SubgraphTaskIds.Add(SubgraphTaskId);
 					Context->bScheduledSubgraph = true;
 					Context->bIsPaused = true;
-					
-					// add a trivial task after the output task that wakes up this task
-					Subsystem->ScheduleGeneric(
-						[ContextHandle = Context->GetOrCreateHandle()]() // Normal execution: Wake up the current task
-						{
-							if (FPCGSubgraphContext* ContextPtr = FPCGContext::GetContextFromHandle<FPCGSubgraphContext>(ContextHandle))
+
+					if (PCGGraphExecutor::CVarEnableDynamicDependencies.GetValueOnAnyThread())
+					{
+						Context->DynamicDependencies.Add(SubgraphTaskId);
+					}
+					else
+					{
+						// add a trivial task after the output task that wakes up this task
+						Subsystem->ScheduleGeneric(
+							[ContextHandle = Context->GetOrCreateHandle()]() // Normal execution: Wake up the current task
 							{
-								ContextPtr->bIsPaused = false;
-							}					
-							return true;
-						}, 
-						[ContextHandle = Context->GetOrCreateHandle()]() // On Abort: Wake up and cancel the execution
-						{
-							if (FPCGSubgraphContext* ContextPtr = FPCGContext::GetContextFromHandle<FPCGSubgraphContext>(ContextHandle))
+								if (FPCGSubgraphContext* ContextPtr = FPCGContext::GetContextFromHandle<FPCGSubgraphContext>(ContextHandle))
+								{
+									ContextPtr->bIsPaused = false;
+								}
+								return true;
+							},
+							[ContextHandle = Context->GetOrCreateHandle()]() // On Abort: Wake up and cancel the execution
 							{
-								ContextPtr->bIsPaused = false;
-								ContextPtr->OutputData.bCancelExecution = true;
-								ContextPtr->SubgraphTaskIds.Reset();
-							}
-							return true;
-						},
-						Context->SourceComponent.Get(), 
-						Context->SubgraphTaskIds);
+								if (FPCGSubgraphContext* ContextPtr = FPCGContext::GetContextFromHandle<FPCGSubgraphContext>(ContextHandle))
+								{
+									ContextPtr->bIsPaused = false;
+									ContextPtr->OutputData.bCancelExecution = true;
+									ContextPtr->SubgraphTaskIds.Reset();
+								}
+								return true;
+							},
+							Context->SourceComponent.Get(),
+							Context->SubgraphTaskIds);
+					}
 
 					return false;
 				}
diff --git a/Engine/Plugins/PCG/Source/PCG/Public/PCGContext.h b/Engine/Plugins/PCG/Source/PCG/Public/PCGContext.h
--- a/Engine/Plugins/PCG/Source/PCG/Public/PCGContext.h
+++ b/Engine/Plugins/PCG/Source/PCG/Public/PCGContext.h
@@ -89,6 +89,7 @@ struct PCG_API FPCGContext
 	FPCGTaskId TaskId = InvalidPCGTaskId;
 	FPCGTaskId CompiledTaskId = InvalidPCGTaskId;
 	bool bIsPaused = false;
+	TSet<FPCGTaskId> DynamicDependencies;
 
 	EPCGExecutionPhase CurrentPhase = EPCGExecutionPhase::NotExecuted;
 	TArray<TPair<FPCGDataCollection, FPCGDataCollection>> CachedInputToOutputInternalResults;
