[PCG] Fix UPCGBasePointData::AddToCrc use const value ranges to avoid flattening of properties

#rnx

diff --git a/Engine/Plugins/PCG/Source/PCG/Private/Data/PCGBasePointData.cpp b/Engine/Plugins/PCG/Source/PCG/Private/Data/PCGBasePointData.cpp
--- a/Engine/Plugins/PCG/Source/PCG/Private/Data/PCGBasePointData.cpp
+++ b/Engine/Plugins/PCG/Source/PCG/Private/Data/PCGBasePointData.cpp
@@ -1332,24 +1332,24 @@ void UPCGBasePointData::AddToCrc(FArchiveCrc32& Ar, bool bFullDataCrc) const
 
 	// Crc point data.
 	{
-		auto CrcRange = [] <class T> (FArchiveCrc32 & Ar, const TPCGValueRange<T>&ValueRange)
+		auto CrcRange = [] <class T> (FArchiveCrc32 & Ar, const TConstPCGValueRange<T>&ValueRange)
 		{
 			for (int32 Index = 0; Index < ValueRange.ViewNum(); ++Index)
 			{
-				Ar << ValueRange[Index];
+				Ar << const_cast<T&>(ValueRange[Index]);
 			}
 		};
 
 		UPCGBasePointData* NonConstThis = const_cast<UPCGBasePointData*>(this);
 
 		// Skip Metadata entry keys
-		CrcRange(Ar, NonConstThis->GetTransformValueRange(false));
-		CrcRange(Ar, NonConstThis->GetDensityValueRange(false));
-		CrcRange(Ar, NonConstThis->GetBoundsMinValueRange(false));
-		CrcRange(Ar, NonConstThis->GetBoundsMaxValueRange(false));
-		CrcRange(Ar, NonConstThis->GetSteepnessValueRange(false));
-		CrcRange(Ar, NonConstThis->GetSeedValueRange(false));
-		CrcRange(Ar, NonConstThis->GetColorValueRange(false));
+		CrcRange(Ar, NonConstThis->GetConstTransformValueRange());
+		CrcRange(Ar, NonConstThis->GetConstDensityValueRange());
+		CrcRange(Ar, NonConstThis->GetConstBoundsMinValueRange());
+		CrcRange(Ar, NonConstThis->GetConstBoundsMaxValueRange());
+		CrcRange(Ar, NonConstThis->GetConstSteepnessValueRange());
+		CrcRange(Ar, NonConstThis->GetConstSeedValueRange());
+		CrcRange(Ar, NonConstThis->GetConstColorValueRange());
 	}
 
 	// Crc metadata.
