[Backout] - CL40168186
#rnx
Original CL Desc
-----------------------------------------------------------------
[PCG] Enable pcg.GraphMultithreading by default

diff --git a/Engine/Plugins/PCG/Source/PCG/Private/Graph/PCGGraphExecutor.cpp b/Engine/Plugins/PCG/Source/PCG/Private/Graph/PCGGraphExecutor.cpp
--- a/Engine/Plugins/PCG/Source/PCG/Private/Graph/PCGGraphExecutor.cpp
+++ b/Engine/Plugins/PCG/Source/PCG/Private/Graph/PCGGraphExecutor.cpp
@@ -66,7 +66,7 @@ namespace PCGGraphExecutor
 
 	TAutoConsoleVariable<bool> CVarGraphMultithreading(
 		TEXT("pcg.GraphMultithreading"),
-		true,
+		false,
 		TEXT("Controls whether the graph can dispatch multiple tasks at the same time"));
 
 #if WITH_EDITOR
