[PCG] Optimization: FPCGGraphCache::AddDataToAccountedMemory avoid calling GetResourceSizeEx multiple times on the same data

#rnx

diff --git a/Engine/Plugins/PCG/Source/PCG/Private/Graph/PCGGraphCache.cpp b/Engine/Plugins/PCG/Source/PCG/Private/Graph/PCGGraphCache.cpp
--- a/Engine/Plugins/PCG/Source/PCG/Private/Graph/PCGGraphCache.cpp
+++ b/Engine/Plugins/PCG/Source/PCG/Private/Graph/PCGGraphCache.cpp
@@ -331,27 +331,25 @@ void FPCGGraphCache::AddDataToAccountedMemory(const FPCGDataCollection& InCollec
 			{
 				if (Data)
 				{
-					FResourceSizeEx ResSize = FResourceSizeEx(EResourceSizeMode::Exclusive);
-					// Calculate data size. Function is non-const but is const-like, especially when
-					// resource mode is Exclusive. The other mode calls a function to find all outer'd
-					// objects which is non-const.
-					const_cast<UPCGData*>(Data)->GetResourceSizeEx(ResSize);
-					const SIZE_T DataSize = ResSize.GetDedicatedSystemMemoryBytes();
-
 					// Find or add record
-					FCachedMemoryRecord& NewRecord = MemoryRecords.FindOrAdd(Data->UID);
-
-					// Cache the size
-					NewRecord.MemoryPerInstance = DataSize;
-
-					if (NewRecord.InstanceCount == 0)
+					if (FCachedMemoryRecord* ExistingRecord = MemoryRecords.Find(Data->UID))
 					{
-						// Account for memory if first instance
+						ExistingRecord->InstanceCount++;
+					}
+					else
+					{
+						FResourceSizeEx ResSize = FResourceSizeEx(EResourceSizeMode::Exclusive);
+						// Calculate data size. Function is non-const but is const-like, especially when
+						// resource mode is Exclusive. The other mode calls a function to find all outer'd
+						// objects which is non-const.
+						const_cast<UPCGData*>(Data)->GetResourceSizeEx(ResSize);
+						const SIZE_T DataSize = ResSize.GetDedicatedSystemMemoryBytes();
+
+						FCachedMemoryRecord& NewRecord = MemoryRecords.Add(Data->UID);
+						NewRecord.MemoryPerInstance = DataSize;
+						NewRecord.InstanceCount = 1;
 						TotalMemoryUsed += DataSize;
 					}
-
-					// Count instances
-					NewRecord.InstanceCount++;
 				}
 			});
 		}
