[PCG] Add culling support to runtime gen scheduling policies
* Add frustum culling option to distance/direction scheduling policy

diff --git a/Engine/Plugins/PCG/Source/PCG/Private/RuntimeGen/GenSources/PCGGenSourceEditorCamera.cpp b/Engine/Plugins/PCG/Source/PCG/Private/RuntimeGen/GenSources/PCGGenSourceEditorCamera.cpp
--- a/Engine/Plugins/PCG/Source/PCG/Private/RuntimeGen/GenSources/PCGGenSourceEditorCamera.cpp
+++ b/Engine/Plugins/PCG/Source/PCG/Private/RuntimeGen/GenSources/PCGGenSourceEditorCamera.cpp
@@ -2,12 +2,44 @@
 
 #include "RuntimeGen/GenSources/PCGGenSourceEditorCamera.h"
 
+#include "SceneView.h"
+
 #if WITH_EDITOR
 #include "EditorViewportClient.h"
 #endif
 
 #include UE_INLINE_GENERATED_CPP_BY_NAME(PCGGenSourceEditorCamera)
 
+void UPCGGenSourceEditorCamera::Tick()
+{
+#if WITH_EDITOR
+	ViewFrustum = TOptional<FConvexVolume>();
+
+	if (EditorViewportClient)
+	{
+		const UWorld* World = EditorViewportClient->GetWorld();
+		const FViewport* Viewport = EditorViewportClient->Viewport;
+
+		if (World && Viewport)
+		{
+			FSceneViewFamilyContext ViewFamily(FSceneViewFamily::ConstructionValues(
+				Viewport,
+				World->Scene,
+				EditorViewportClient->EngineShowFlags)
+				.SetRealtimeUpdate(EditorViewportClient->IsRealtime()));
+
+			if (FSceneView* SceneView = EditorViewportClient->CalcSceneView(&ViewFamily))
+			{
+				FConvexVolume ConvexVolume;
+				GetViewFrustumBounds(ConvexVolume, SceneView->ViewMatrices.GetViewProjectionMatrix(), /*bUseNearPlane=*/true, /*bUseFarPlane=*/true);
+
+				ViewFrustum = ConvexVolume;
+			}
+		}
+	}
+#endif
+}
+
 TOptional<FVector> UPCGGenSourceEditorCamera::GetPosition() const
 {
 #if WITH_EDITOR
@@ -31,3 +63,12 @@ TOptional<FVector> UPCGGenSourceEditorCamera::GetDirection() const
 
 	return TOptional<FVector>();
 }
+
+TOptional<FConvexVolume> UPCGGenSourceEditorCamera::GetViewFrustum(bool bIs2DGrid) const
+{
+#if WITH_EDITOR
+	return ViewFrustum;
+#else
+	return TOptional<FConvexVolume>();
+#endif	
+}
\ No newline at end of file
diff --git a/Engine/Plugins/PCG/Source/PCG/Private/RuntimeGen/GenSources/PCGGenSourcePlayer.cpp b/Engine/Plugins/PCG/Source/PCG/Private/RuntimeGen/GenSources/PCGGenSourcePlayer.cpp
--- a/Engine/Plugins/PCG/Source/PCG/Private/RuntimeGen/GenSources/PCGGenSourcePlayer.cpp
+++ b/Engine/Plugins/PCG/Source/PCG/Private/RuntimeGen/GenSources/PCGGenSourcePlayer.cpp
@@ -2,11 +2,34 @@
 
 #include "RuntimeGen/GenSources/PCGGenSourcePlayer.h"
 
+#include "SceneView.h"
+#include "Engine/GameViewportClient.h"
+#include "Engine/LocalPlayer.h"
 #include "GameFramework/Pawn.h"
 #include "GameFramework/PlayerController.h"
 
 #include UE_INLINE_GENERATED_CPP_BY_NAME(PCGGenSourcePlayer)
 
+void UPCGGenSourcePlayer::Tick()
+{
+	ViewFrustum = TOptional<FConvexVolume>();
+
+	const ULocalPlayer* LocalPlayer = PlayerController.IsValid() ? PlayerController->GetLocalPlayer() : nullptr;
+
+	if (LocalPlayer && LocalPlayer->ViewportClient)
+	{
+		FSceneViewProjectionData ProjectionData;
+
+		if (LocalPlayer->GetProjectionData(LocalPlayer->ViewportClient->Viewport, ProjectionData))
+		{
+			FConvexVolume ConvexVolume;
+			GetViewFrustumBounds(ConvexVolume, ProjectionData.ComputeViewProjectionMatrix(), /*bUseNearPlane=*/true, /*bUseFarPlane=*/true);
+
+			ViewFrustum = ConvexVolume;
+		}
+	}
+}
+
 TOptional<FVector> UPCGGenSourcePlayer::GetPosition() const
 {
 	if (PlayerController.IsValid())
diff --git a/Engine/Plugins/PCG/Source/PCG/Private/RuntimeGen/PCGRuntimeGenScheduler.cpp b/Engine/Plugins/PCG/Source/PCG/Private/RuntimeGen/PCGRuntimeGenScheduler.cpp
--- a/Engine/Plugins/PCG/Source/PCG/Private/RuntimeGen/PCGRuntimeGenScheduler.cpp
+++ b/Engine/Plugins/PCG/Source/PCG/Private/RuntimeGen/PCGRuntimeGenScheduler.cpp
@@ -123,6 +123,14 @@ void FPCGRuntimeGenScheduler::Tick(APCGWorldActor* InPCGWorldActor)
 	{
 		GenSourceManager->Tick();
 		GenSources = GenSourceManager->GetAllGenSources(InPCGWorldActor);
+
+		for (IPCGGenSourceBase* GenSource : GenSources)
+		{
+			if (ensure(GenSource))
+			{
+				GenSource->Tick();
+			}
+		}
 	}
 
 	// Initialize RuntimeGen PA pool if necessary. If PoolSize is 0, then we have not initialized the pool yet.
@@ -283,6 +291,11 @@ void FPCGRuntimeGenScheduler::TickQueueComponentsForGeneration(
 		// TODO: For each execution domain (for now only GenAtRuntime/dynamic), assuming we run Preview through this scheduler, which it seems like we will.
 		if (OriginalComponent->IsManagedByRuntimeGenSystem())
 		{
+			if (!ensure(Policy))
+			{
+				continue;
+			}
+
 			bool bHasUnbounded = false;
 			PCGHiGenGrid::FSizeArray GridSizes;
 			ensure(PCGHelpers::GetGenerationGridSizes(OriginalComponent->GetGraph(), InPCGWorldActor, GridSizes, bHasUnbounded));
@@ -325,7 +338,8 @@ void FPCGRuntimeGenScheduler::TickQueueComponentsForGeneration(
 				{
 					// Ignore components that have already been generated or marked for generation. Unbounded grid size means not-partitioned.
 					FGridGenerationKey Key(PCGHiGenGrid::UnboundedGridSize(), FIntVector(0), OriginalComponent);
-					if (!GeneratedComponents.Contains(Key))
+					if (!GeneratedComponents.Contains(Key) && 
+						Policy->ShouldGenerate(GenSource, OriginalComponentBounds, bIs2DGrid))
 					{
 						check(Key.GetGridDescriptor().Is2DGrid() == bIs2DGrid);
 						AddComponentToGenerate(Key, GenSource, Policy, OriginalComponentBounds, bIs2DGrid);
@@ -346,16 +360,8 @@ void FPCGRuntimeGenScheduler::TickQueueComponentsForGeneration(
 					const double HalfGridSize = GridSize / 2.0f;
 					FVector HalfExtent(HalfGridSize, HalfGridSize, HalfGridSize);
 
-					FBox ModifiedBounds = OriginalComponent->GetGridBounds();
 					if (bIs2DGrid)
 					{
-						FVector MinBounds = ModifiedBounds.Min;
-						FVector MaxBounds = ModifiedBounds.Max;
-
-						MinBounds.Z = 0;
-						MaxBounds.Z = GridSize;
-						ModifiedBounds = FBox(MinBounds, MaxBounds);
-
 						// In case of 2D grid, it's like the actor has infinite bounds on the Z axis.
 						HalfExtent.Z = HALF_WORLD_MAX1;
 					}
@@ -380,7 +386,7 @@ void FPCGRuntimeGenScheduler::TickQueueComponentsForGeneration(
 								const FBox CellBounds(Center - HalfExtent, Center + HalfExtent);
 
 								// Overlap cell with the partitioned component.
-								const FBox IntersectedBounds = ModifiedBounds.Overlap(CellBounds);
+								const FBox IntersectedBounds = OriginalComponentBounds.Overlap(CellBounds);
 								if (!IntersectedBounds.IsValid || IntersectedBounds.GetVolume() <= UE_DOUBLE_SMALL_NUMBER)
 								{
 									continue;
@@ -394,7 +400,8 @@ void FPCGRuntimeGenScheduler::TickQueueComponentsForGeneration(
 								// Verify the grid cell actually lies within the generation radius.
 								// TODO: this is no longer necessary if we rasterize the sphere instead.
 								const double LocalDistanceSquared = CellBounds.ComputeSquaredDistanceToPoint(ModifiedGenSourcePosition);
-								if (LocalDistanceSquared <= GenerationRadius * GenerationRadius)
+								if (LocalDistanceSquared <= GenerationRadius * GenerationRadius && 
+									Policy->ShouldGenerate(GenSource, IntersectedBounds, Key.GetGridDescriptor().Is2DGrid()))
 								{
 									AddComponentToGenerate(Key, GenSource, Policy, IntersectedBounds, Key.GetGridDescriptor().Is2DGrid());
 								}
@@ -450,7 +457,8 @@ void FPCGRuntimeGenScheduler::TickQueueComponentsForGeneration(
 				const double DistanceSquared = OriginalComponentBounds.ComputeSquaredDistanceToPoint(ModifiedGenSourcePosition);
 
 				// Max radius for a non-partitioned component is just the base GenerationRadius.
-				if (DistanceSquared <= MaxGenerationRadius * MaxGenerationRadius)
+				if (DistanceSquared <= MaxGenerationRadius * MaxGenerationRadius && 
+					Policy->ShouldGenerate(GenSource, OriginalComponentBounds, /*bUse2DGrid=*/false))
 				{
 					// Unbounded grid size means not-partitioned.
 					FGridGenerationKey Key(PCGHiGenGrid::UnboundedGridSize(), FIntVector(0), OriginalComponent);
@@ -471,6 +479,46 @@ void FPCGRuntimeGenScheduler::TickCleanup(const TSet<IPCGGenSourceBase*>& InGenS
 	using PCGComponentToClean = TTuple<FGridGenerationKey, UPCGComponent*>;
 	TArray<PCGComponentToClean> ComponentsToClean;
 
+	auto CheckIfAllGenSourcesWantToCleanup = [&InGenSources](const UPCGSchedulingPolicyBase* Policy, const FPCGGridDescriptor& GridDescriptor, const FBox& GridBounds, const double& CleanupRadiusSquared)
+	{
+		bool bAllGenSourcesWantToCleanup = true;
+
+		for (const IPCGGenSourceBase* GenSource : InGenSources)
+		{
+			if (!ensure(GenSource))
+			{
+				continue;
+			}
+
+			const TOptional<FVector> GenSourcePositionOptional = GenSource->GetPosition();
+			if (!GenSourcePositionOptional.IsSet())
+			{
+				continue;
+			}
+
+			FVector GenSourcePosition = GenSourcePositionOptional.GetValue();
+
+			// Only consider 2D distance when using a 2D grid.
+			if (GridDescriptor.Is2DGrid())
+			{
+				GenSourcePosition.Z = GridBounds.Min.Z;
+			}
+
+			const double SquaredDistToGenSource = GridBounds.ComputeSquaredDistanceToPoint(GenSourcePosition);
+
+			// If the distance to the gen source is greater than the cleanup radius, it means this generation source votes for the component to be cleaned up.
+			// Otherwise, the gen source might still vote for culling regardless.
+			if (SquaredDistToGenSource <= CleanupRadiusSquared && 
+				!Policy->ShouldCull(GenSource, GridBounds, GridDescriptor.Is2DGrid()))
+			{
+				bAllGenSourcesWantToCleanup = false;
+				break;
+			}
+		}
+
+		return bAllGenSourcesWantToCleanup;
+	};
+
 	// Find any generated components which should be cleaned.
 	for (FGridGenerationKey GenerationKey : GeneratedComponents)
 	{
@@ -482,6 +530,15 @@ void FPCGRuntimeGenScheduler::TickCleanup(const TSet<IPCGGenSourceBase*>& InGenS
 		UPCGComponent* OriginalComponent = GenerationKey.GetOriginalComponent();
 		check(OriginalComponent);
 
+		const UPCGSchedulingPolicyBase* Policy = OriginalComponent->GetRuntimeGenSchedulingPolicy();
+		if(!ensure(Policy))
+		{
+			continue;
+		}
+
+		const double CleanupRadius = OriginalComponent->GetCleanupRadiusFromGrid(Grid);
+		const double CleanupRadiusSquared = CleanupRadius * CleanupRadius;
+
 		// If the Grid is unbounded, we have a non-partitioned or unbounded component.
 		if (Grid == EPCGHiGenGrid::Unbounded)
 		{
@@ -491,35 +548,10 @@ void FPCGRuntimeGenScheduler::TickCleanup(const TSet<IPCGGenSourceBase*>& InGenS
 				continue;
 			}
 
-			const FBox Bounds = OriginalComponent->GetGridBounds();
-
-			double MinSquaredDistanceToGenSource = UE_DOUBLE_BIG_NUMBER;
-			for (const IPCGGenSourceBase* GenSource : InGenSources)
-			{
-				if (!ensure(GenSource))
-				{
-					continue;
-				}
-
-				const TOptional<FVector> GenSourcePositionOptional = GenSource->GetPosition();
-				if (!GenSourcePositionOptional.IsSet())
-				{
-					continue;
-				}
-
-				FVector GenSourcePosition = GenSourcePositionOptional.GetValue();
-
-				// Only consider 2D distance when using a 2D grid.
-				if (GridDescriptor.Is2DGrid())
-				{
-					GenSourcePosition.Z = Bounds.Min.Z;
-				}
-
-				MinSquaredDistanceToGenSource = FMath::Min(MinSquaredDistanceToGenSource, Bounds.ComputeSquaredDistanceToPoint(GenSourcePosition));
-			}
+			const FBox GridBounds = OriginalComponent->GetGridBounds();
 
-			const double CleanupRadius = OriginalComponent->GetCleanupRadiusFromGrid(Grid);
-			if (MinSquaredDistanceToGenSource > CleanupRadius * CleanupRadius)
+			// Only clean up if all generation sources agreed to clean up.
+			if(CheckIfAllGenSourcesWantToCleanup(Policy, GridDescriptor, GridBounds, CleanupRadiusSquared))
 			{
 				ComponentsToClean.Add({ GenerationKey, OriginalComponent });
 			}
@@ -536,35 +568,10 @@ void FPCGRuntimeGenScheduler::TickCleanup(const TSet<IPCGGenSourceBase*>& InGenS
 				continue;
 			}
 
-			const FBox GridBounds = PartitionActor->GetFixedBounds();
-
-			double MinSquaredDistanceToGenSource = UE_DOUBLE_BIG_NUMBER;
-			for (const IPCGGenSourceBase* GenSource : InGenSources)
-			{
-				if (!ensure(GenSource))
-				{
-					continue;
-				}
-
-				const TOptional<FVector> GenSourcePositionOptional = GenSource->GetPosition();
-				if (!GenSourcePositionOptional.IsSet())
-				{
-					continue;
-				}
-
-				FVector GenSourcePosition = GenSourcePositionOptional.GetValue();
-
-				// Only consider 2D distance when using a 2D grid.
-				if (GridDescriptor.Is2DGrid())
-				{
-					GenSourcePosition.Z = GridBounds.Min.Z;
-				}
-
-				MinSquaredDistanceToGenSource = FMath::Min(MinSquaredDistanceToGenSource, GridBounds.ComputeSquaredDistanceToPoint(GenSourcePosition));
-			}
+			const FBox GridBounds = LocalComponent->GetGridBounds();
 
-			const double CleanupRadius = OriginalComponent->GetCleanupRadiusFromGrid(Grid);
-			if (MinSquaredDistanceToGenSource > CleanupRadius * CleanupRadius)
+			// Only clean up if all generation sources agreed to clean up.
+			if (CheckIfAllGenSourcesWantToCleanup(Policy, GridDescriptor, GridBounds, CleanupRadiusSquared))
 			{
 				ComponentsToClean.Add({ GenerationKey, LocalComponent });
 			}
diff --git a/Engine/Plugins/PCG/Source/PCG/Private/RuntimeGen/SchedulingPolicies/PCGSchedulingPolicyDistanceAndDirection.cpp b/Engine/Plugins/PCG/Source/PCG/Private/RuntimeGen/SchedulingPolicies/PCGSchedulingPolicyDistanceAndDirection.cpp
--- a/Engine/Plugins/PCG/Source/PCG/Private/RuntimeGen/SchedulingPolicies/PCGSchedulingPolicyDistanceAndDirection.cpp
+++ b/Engine/Plugins/PCG/Source/PCG/Private/RuntimeGen/SchedulingPolicies/PCGSchedulingPolicyDistanceAndDirection.cpp
@@ -56,6 +56,54 @@ double UPCGSchedulingPolicyDistanceAndDirection::CalculatePriority(const IPCGGen
 	return Priority;
 }
 
+bool UPCGSchedulingPolicyDistanceAndDirection::ShouldGenerate(const IPCGGenSourceBase* InGenSource, const FBox& GenerationBounds, bool bUse2DGrid) const
+{
+	if (!bUseFrustumCulling)
+	{
+		return true;
+	}
+
+	const TOptional<FConvexVolume> ViewFrustumOptional = InGenSource->GetViewFrustum(bUse2DGrid);
+
+	if (!ViewFrustumOptional.IsSet())
+	{
+		return true;
+	}
+
+	FVector Center, Extents;
+	GenerationBounds.GetCenterAndExtents(Center, Extents);
+
+	// @todo_pcg: Instead of having a modifier on the bounds, it would probably be better to increase the FOV of the frustum's projection matrix.
+	// Another idea would be to adjust scaling based on distance instead, which would end up most likely cheaper than recomputing the FOV.
+	const FConvexVolume ViewFrustum = ViewFrustumOptional.GetValue();
+	return ViewFrustum.IntersectBox(Center, Extents * GenerateBoundsModifier);
+}
+
+bool UPCGSchedulingPolicyDistanceAndDirection::ShouldCull(const IPCGGenSourceBase* InGenSource, const FBox& GenerationBounds, bool bUse2DGrid) const
+{
+	if (!bUseFrustumCulling)
+	{
+		// If frustum culling is not enabled, don't try cleaning up volumes within the cleanup radius.
+		return false;
+	}
+
+	const TOptional<FConvexVolume> ViewFrustumOptional = InGenSource->GetViewFrustum(bUse2DGrid);
+
+	if (!ViewFrustumOptional.IsSet())
+	{
+		return false;
+	}
+
+	FVector Center, Extents;
+	GenerationBounds.GetCenterAndExtents(Center, Extents);
+
+	// If the volume does not intersect the view frustum, allow it to be cleaned up.
+	const FConvexVolume ViewFrustum = ViewFrustumOptional.GetValue();
+
+	const float CleanupModifier = FMath::Max(GenerateBoundsModifier + 0.1, CleanupBoundsModifier);
+	return !ViewFrustum.IntersectBox(Center, Extents * CleanupModifier);
+}
+
 bool UPCGSchedulingPolicyDistanceAndDirection::IsEquivalent(const UPCGSchedulingPolicyBase* OtherSchedulingPolicy) const
 {
 	if (!OtherSchedulingPolicy)
@@ -73,7 +121,10 @@ bool UPCGSchedulingPolicyDistanceAndDirection::IsEquivalent(const UPCGScheduling
 		return bUseDistance == Other->bUseDistance
 			&& DistanceWeight == Other->DistanceWeight
 			&& bUseDirection == Other->bUseDirection
-			&& DirectionWeight == Other->DirectionWeight;
+			&& DirectionWeight == Other->DirectionWeight
+			&& bUseFrustumCulling == Other->bUseFrustumCulling
+			&& GenerateBoundsModifier == Other->GenerateBoundsModifier
+			&& CleanupBoundsModifier == Other->CleanupBoundsModifier;
 	}
 	else
 	{
diff --git a/Engine/Plugins/PCG/Source/PCG/Public/RuntimeGen/GenSources/PCGGenSourceBase.h b/Engine/Plugins/PCG/Source/PCG/Public/RuntimeGen/GenSources/PCGGenSourceBase.h
--- a/Engine/Plugins/PCG/Source/PCG/Public/RuntimeGen/GenSources/PCGGenSourceBase.h
+++ b/Engine/Plugins/PCG/Source/PCG/Public/RuntimeGen/GenSources/PCGGenSourceBase.h
@@ -2,6 +2,7 @@
 
 #pragma once
 
+#include "ConvexVolume.h"
 #include "UObject/Interface.h"
 
 #include "PCGGenSourceBase.generated.h"
@@ -21,11 +22,15 @@ class IPCGGenSourceBase
 	GENERATED_BODY()
 
 public:
+	/** Update the generation source so that it can cache data that is queried often (e.g. view frustum). Should be called every tick on any active generation sources. */
+	virtual void Tick() {}
+
 	/** Returns the world space position of this gen source. */
 	virtual TOptional<FVector> GetPosition() const PURE_VIRTUAL(UPCGGenSourceBase::GetPosition, return TOptional<FVector>(););
 
 	/** Returns the normalized forward vector of this gen source. */
 	virtual TOptional<FVector> GetDirection() const PURE_VIRTUAL(UPCGGenSourceBase::GetDirection, return TOptional<FVector>(););
 
-	// TODO: Functions to get the horizontal and vertical FOV would be useful for policies based around frustum/zooming.
+	/** Returns the view frustum of this gen source. */
+	virtual TOptional<FConvexVolume> GetViewFrustum(bool bIs2DGrid) const { return TOptional<FConvexVolume>(); }
 };
diff --git a/Engine/Plugins/PCG/Source/PCG/Public/RuntimeGen/GenSources/PCGGenSourceEditorCamera.h b/Engine/Plugins/PCG/Source/PCG/Public/RuntimeGen/GenSources/PCGGenSourceEditorCamera.h
--- a/Engine/Plugins/PCG/Source/PCG/Public/RuntimeGen/GenSources/PCGGenSourceEditorCamera.h
+++ b/Engine/Plugins/PCG/Source/PCG/Public/RuntimeGen/GenSources/PCGGenSourceEditorCamera.h
@@ -20,14 +20,21 @@ class PCG_API UPCGGenSourceEditorCamera : public UObject, public IPCGGenSourceBa
 	GENERATED_BODY()
 
 public:
+	/** Update the generation source so that it can cache data that is queried often (e.g. view frustum). Should be called every tick on any active generation sources. */
+	virtual void Tick() override;
+
 	/** Returns the world space position of this gen source. */
 	virtual TOptional<FVector> GetPosition() const override;
 
 	/** Returns the normalized forward vector of this gen source. */
 	virtual TOptional<FVector> GetDirection() const override;
 
+	/** Returns the view frustum of this gen source. */
+	virtual TOptional<FConvexVolume> GetViewFrustum(bool bIs2DGrid) const override;
+
 public:
 #if WITH_EDITORONLY_DATA
 	FEditorViewportClient* EditorViewportClient = nullptr;
+	TOptional<FConvexVolume> ViewFrustum;
 #endif
 };
diff --git a/Engine/Plugins/PCG/Source/PCG/Public/RuntimeGen/GenSources/PCGGenSourcePlayer.h b/Engine/Plugins/PCG/Source/PCG/Public/RuntimeGen/GenSources/PCGGenSourcePlayer.h
--- a/Engine/Plugins/PCG/Source/PCG/Public/RuntimeGen/GenSources/PCGGenSourcePlayer.h
+++ b/Engine/Plugins/PCG/Source/PCG/Public/RuntimeGen/GenSources/PCGGenSourcePlayer.h
@@ -17,12 +17,18 @@ class PCG_API UPCGGenSourcePlayer : public UObject, public IPCGGenSourceBase
 	GENERATED_BODY()
 
 public:
+	/** Update the generation source so that it can cache data that is queried often (e.g. view frustum). Should be called every tick on any active generation sources. */
+	virtual void Tick() override;
+
 	/** Returns the world space position of this gen source. */
 	virtual TOptional<FVector> GetPosition() const override;
 
 	/** Returns the normalized forward vector of this gen source. */
 	virtual TOptional<FVector> GetDirection() const override;
 
+	/** Returns the view frustum of this gen source. */
+	virtual TOptional<FConvexVolume> GetViewFrustum(bool bIs2DGrid) const override { return ViewFrustum; }
+
 	TWeakObjectPtr<const APlayerController> GetPlayerController() const { return PlayerController; }
 	void SetPlayerController(const APlayerController* InPlayerController);
 
@@ -30,4 +36,5 @@ public:
 
 protected:
 	TWeakObjectPtr<const APlayerController> PlayerController;
+	TOptional<FConvexVolume> ViewFrustum;
 };
diff --git a/Engine/Plugins/PCG/Source/PCG/Public/RuntimeGen/SchedulingPolicies/PCGSchedulingPolicyBase.h b/Engine/Plugins/PCG/Source/PCG/Public/RuntimeGen/SchedulingPolicies/PCGSchedulingPolicyBase.h
--- a/Engine/Plugins/PCG/Source/PCG/Public/RuntimeGen/SchedulingPolicies/PCGSchedulingPolicyBase.h
+++ b/Engine/Plugins/PCG/Source/PCG/Public/RuntimeGen/SchedulingPolicies/PCGSchedulingPolicyBase.h
@@ -23,6 +23,12 @@ public:
 	virtual double CalculatePriority(const IPCGGenSourceBase* InGenSource, const FBox& GenerationBounds, bool bUse2DGrid) const 
 		PURE_VIRTUAL(UPCGSchedulingPolicyBase::CalculatePriority, return 0.0;);
 
+	/** True if the generation source would consider the given bounds for generation. */
+	virtual bool ShouldGenerate(const IPCGGenSourceBase* InGenSource, const FBox& GenerationBounds, bool bUse2DGrid) const { return true; }
+
+	/** True if the generation source would cull the given bounds. Only applies to bounds within the cleanup generation radius. */
+	virtual bool ShouldCull(const IPCGGenSourceBase* InGenSource, const FBox& GenerationBounds, bool bUse2DGrid) const { return false; }
+
 	/** A SchedulingPolicy is equivalent to another SchedulingPolicy if they are the same (same ptr), or if they have the same type and parameter values. */
 	virtual bool IsEquivalent(const UPCGSchedulingPolicyBase* OtherSchedulingPolicy) const PURE_VIRTUAL(UPCGSchedulingPolicyBase::IsEquivalent, return false;);
 
diff --git a/Engine/Plugins/PCG/Source/PCG/Public/RuntimeGen/SchedulingPolicies/PCGSchedulingPolicyDistanceAndDirection.h b/Engine/Plugins/PCG/Source/PCG/Public/RuntimeGen/SchedulingPolicies/PCGSchedulingPolicyDistanceAndDirection.h
--- a/Engine/Plugins/PCG/Source/PCG/Public/RuntimeGen/SchedulingPolicies/PCGSchedulingPolicyDistanceAndDirection.h
+++ b/Engine/Plugins/PCG/Source/PCG/Public/RuntimeGen/SchedulingPolicies/PCGSchedulingPolicyDistanceAndDirection.h
@@ -23,6 +23,12 @@ public:
 	/** Calculate the runtime scheduling priority with respect to a Generation Source. Should return a value in the range [0, 1], where higher values will be scheduled sooner. */
 	virtual double CalculatePriority(const IPCGGenSourceBase* InGenSource, const FBox& GenerationBounds, bool bUse2DGrid) const override;
 
+	/** True if the generation source would consider the given bounds for generation. */
+	virtual bool ShouldGenerate(const IPCGGenSourceBase* InGenSource, const FBox& GenerationBounds, bool bUse2DGrid) const override;
+
+	/** True if the generation source would cull the given bounds. Only applies to bounds within the cleanup generation radius. */
+	virtual bool ShouldCull(const IPCGGenSourceBase* InGenSource, const FBox& GenerationBounds, bool bUse2DGrid) const override;
+
 	/** A SchedulingPolicy is equivalent to another SchedulingPolicy if they are the same (same ptr), or if they have the same type and parameter values. */
 	virtual bool IsEquivalent(const UPCGSchedulingPolicyBase* OtherSchedulingPolicy) const;
 
@@ -42,4 +48,16 @@ public:
 	/** Scalar value used to increase/decrease the impact of direction in the scheduling priority. */
 	UPROPERTY(BlueprintReadWrite, EditAnywhere, Category = "RuntimeGeneration|Scheduling Policy Parameters", meta = (ClampMin = 0.0f, ClampMax = 1.0f, EditCondition = "bShouldDisplayProperties && bUseDirection", EditConditionHides))
 	float DirectionWeight = 0.0025f;
+
+	/** With frustum culling enabled, only components whose bounds overlap the view frustum will be generated. */
+	UPROPERTY(BlueprintReadWrite, EditAnywhere, Category = "RuntimeGeneration|Scheduling Policy Parameters", meta = (EditCondition = "bShouldDisplayProperties", EditConditionHides, HideEditConditionToggle))
+	bool bUseFrustumCulling = false;
+
+	/** Multiplier to scale bounds by when comparing against the view frustum for generation. Can help if components on the edge of the frustum are not generating as soon as you would like. */
+	UPROPERTY(BlueprintReadWrite, EditAnywhere, Category = "RuntimeGeneration|Scheduling Policy Parameters", meta = (ClampMin = 1.0f, EditCondition = "bShouldDisplayProperties && bUseFrustumCulling", EditConditionHides))
+	float GenerateBoundsModifier = 1.0f;
+
+	/** Multiplier to scale bounds by when comparing against the view frustum for clean up. */
+	UPROPERTY(BlueprintReadWrite, EditAnywhere, Category = "RuntimeGeneration|Scheduling Policy Parameters", meta = (ClampMin = 1.0f, EditCondition = "bShouldDisplayProperties && bUseFrustumCulling", EditConditionHides))
+	float CleanupBoundsModifier = 1.2f;
 };
