[PCG] Optimize cost of PCGSubsystem::OnScheduleGraph which was growing linear with size of ExecutedStacks
- Also optimize FPCGStack::BeginsWith by comparing hash before using operator==

#rnx

diff --git a/Engine/Plugins/PCG/Source/PCG/Private/Graph/PCGStackContext.cpp b/Engine/Plugins/PCG/Source/PCG/Private/Graph/PCGStackContext.cpp
--- a/Engine/Plugins/PCG/Source/PCG/Private/Graph/PCGStackContext.cpp
+++ b/Engine/Plugins/PCG/Source/PCG/Private/Graph/PCGStackContext.cpp
@@ -143,6 +143,12 @@ bool FPCGStack::BeginsWith(const FPCGStack& Other) const
 
 	for (int32 StackIndex = 0; StackIndex < Other.GetStackFrames().Num(); ++StackIndex)
 	{
+		// Test hash first which is a lot faster (hash is cached) than testing for equality
+		if (GetTypeHash(Other.GetStackFrames()[StackIndex]) != GetTypeHash(GetStackFrames()[StackIndex]))
+		{
+			return false;
+		}
+
 		if (Other.GetStackFrames()[StackIndex] != GetStackFrames()[StackIndex])
 		{
 			return false;
diff --git a/Engine/Plugins/PCG/Source/PCG/Private/PCGSubsystem.cpp b/Engine/Plugins/PCG/Source/PCG/Private/PCGSubsystem.cpp
--- a/Engine/Plugins/PCG/Source/PCG/Private/PCGSubsystem.cpp
+++ b/Engine/Plugins/PCG/Source/PCG/Private/PCGSubsystem.cpp
@@ -560,7 +560,7 @@ void UPCGSubsystem::OnScheduleGraph(const FPCGStackContext& StackContext)
 		FWriteScopeLock Lock(ExecutedStacksLock);
 		for (const FPCGStack& ExecutedStack : StackContext.GetStacks())
 		{
-			ExecutedStacks.AddUnique(ExecutedStack);
+			ExecutedStacks.Add(ExecutedStack);
 		}
 	}
 }
@@ -1893,13 +1893,13 @@ void UPCGSubsystem::ClearExecutedStacks(FPCGStack BeginningWithStack)
 {
 	FWriteScopeLock Lock(ExecutedStacksLock);
 
-	for (int StackIndex = ExecutedStacks.Num() - 1; StackIndex >= 0; --StackIndex)
+	for (auto It = ExecutedStacks.CreateIterator(); It; ++It)
 	{
 		// Clear any dynamic stack that starts with the given base stack, we don't know what stacks will
 		// execute so clean out everything.
-		if (ExecutedStacks[StackIndex].BeginsWith(BeginningWithStack))
+		if (It->BeginsWith(BeginningWithStack))
 		{
-			ExecutedStacks.RemoveAtSwap(StackIndex);
+			It.RemoveCurrent();
 		}
 	}
 }
@@ -1908,13 +1908,13 @@ void UPCGSubsystem::ClearExecutedStacks(const UPCGGraph* InContainingGraph)
 {
 	FWriteScopeLock Lock(ExecutedStacksLock);
 
-	for (int StackIndex = ExecutedStacks.Num() - 1; StackIndex >= 0; --StackIndex)
+	for (auto It = ExecutedStacks.CreateIterator(); It; ++It)
 	{
 		// Clear any dynamic stack that starts with the given base stack, we don't know what stacks will
 		// execute so clean out everything.
-		if (ExecutedStacks[StackIndex].HasObject(InContainingGraph))
+		if (It->HasObject(InContainingGraph))
 		{
-			ExecutedStacks.RemoveAtSwap(StackIndex);
+			It.RemoveCurrent();
 		}
 	}
 }
diff --git a/Engine/Plugins/PCG/Source/PCG/Public/PCGSubsystem.h b/Engine/Plugins/PCG/Source/PCG/Public/PCGSubsystem.h
--- a/Engine/Plugins/PCG/Source/PCG/Public/PCGSubsystem.h
+++ b/Engine/Plugins/PCG/Source/PCG/Public/PCGSubsystem.h
@@ -408,7 +408,7 @@ private:
 	FPCGActorAndComponentMapping ActorAndComponentMapping;
 
 	/** A record of stacks that were executed. Used to populate debugging tool UIs. */
-	TArray<FPCGStack> ExecutedStacks;
+	TSet<FPCGStack> ExecutedStacks;
 	mutable FRWLock ExecutedStacksLock;
 
 	/** Functions will be executed at the beginning of the tick and then removed from this array. */
