[PCG] Fix null pointer crash when component is null.

Not clear why component is null but crash occurred entering PIE. This change mirrors similar checks elsewhere.

#rnx

diff --git a/Engine/Plugins/PCG/Source/PCG/Private/Elements/PCGGenerateGrassMaps.cpp b/Engine/Plugins/PCG/Source/PCG/Private/Elements/PCGGenerateGrassMaps.cpp
--- a/Engine/Plugins/PCG/Source/PCG/Private/Elements/PCGGenerateGrassMaps.cpp
+++ b/Engine/Plugins/PCG/Source/PCG/Private/Elements/PCGGenerateGrassMaps.cpp
@@ -101,7 +101,7 @@ void FPCGGenerateGrassMapsElement::GetDependenciesCrc(const FPCGDataCollection&
 	IPCGElement::GetDependenciesCrc(InInput, InSettings, InComponent, Crc);
 
 	// todo_pcg: Technically this could be fancier and hash the set of landscape components that overlap with the generation volume.
-	if (const UPCGData* Data = InComponent->GetActorPCGData())
+	if (const UPCGData* Data = InComponent ? InComponent->GetActorPCGData() : nullptr)
 	{
 		Crc.Combine(Data->GetOrComputeCrc(/*bFullDataCrc=*/false));
 	}
