[PCG] Fix memory ASAN issue - adding an attribute invalidates the StringKeys pointer.

#rnx

diff --git a/Engine/Plugins/PCG/Source/PCG/Private/Elements/PCGStaticMeshSpawnerKernel.cpp b/Engine/Plugins/PCG/Source/PCG/Private/Elements/PCGStaticMeshSpawnerKernel.cpp
--- a/Engine/Plugins/PCG/Source/PCG/Private/Elements/PCGStaticMeshSpawnerKernel.cpp
+++ b/Engine/Plugins/PCG/Source/PCG/Private/Elements/PCGStaticMeshSpawnerKernel.cpp
@@ -131,20 +131,22 @@ bool UPCGStaticMeshSpawnerKernel::ComputeOutputBindingDataDesc(const UPCGCompute
 			// By-attribute selection - pass on strings from input attribute.
 			if (SelectorByAttribute->AttributeName != NAME_None)
 			{
+				TArray<int32> StringKeys;
+
 				for (FPCGDataDesc& DataDesc : PinDesc.DataDescs)
 				{
-					const TArray<int32>* StringKeys = nullptr;
+					StringKeys.Empty();
 
 					for (const FPCGKernelAttributeDesc& AttrDesc : DataDesc.AttributeDescs)
 					{
 						if (AttrDesc.Name == SelectorByAttribute->AttributeName)
 						{
-							StringKeys = &AttrDesc.GetUniqueStringKeys();
+							StringKeys = AttrDesc.GetUniqueStringKeys();
 							break;
 						}
 					}
 
-					DataDesc.AddAttribute(FPCGKernelAttributeKey(Settings->OutAttributeName, EPCGKernelAttributeType::StringKey), InBinding, StringKeys);
+					DataDesc.AddAttribute(FPCGKernelAttributeKey(Settings->OutAttributeName, EPCGKernelAttributeType::StringKey), InBinding, &StringKeys);
 				}
 			}
 		}
