[PCG] Remove bIsPartitioned deprecation added 2 years ago

* bIsPartitioned was living along side bIsComponentPartitioned, to change the default value from true to false
* Deprecation code was in place to make sure new objects have this flag set as false
* But since it is editor only data, when cooked then loaded in Cooked editor, it was not deprecating correctly.
* Since it was 2 years ago when PCG was nowhere it is now, we can safely remove this deprecation.

#rnx

diff --git a/Engine/Plugins/PCG/Source/PCG/Private/PCGComponent.cpp b/Engine/Plugins/PCG/Source/PCG/Private/PCGComponent.cpp
--- a/Engine/Plugins/PCG/Source/PCG/Private/PCGComponent.cpp
+++ b/Engine/Plugins/PCG/Source/PCG/Private/PCGComponent.cpp
@@ -1936,15 +1936,6 @@ void UPCGComponent::PostLoad()
 		GeneratedGraphOutput.Reset();
 	}
 
-	// If we have both default value (bIsComponentPartitioned = false and bIsPartitioned = true)
-	// we will follow the value of bIsPartitioned.
-	// bIsPartitioned will be set to false to new objects
-	if (!bIsComponentPartitioned && bIsPartitioned)
-	{
-		bIsComponentPartitioned = bIsPartitioned;
-		bIsPartitioned = false;
-	}
-
 	// Components marked as partitioned in non-WP worlds from BEFORE partitioning was supported in non-WP worlds can leak resources. To fix this, we can just unset 'bIsComponentPartitioned'.
 	if (bDisableIsComponentPartitionedOnLoad)
 	{
@@ -2078,14 +2069,6 @@ void UPCGComponent::PostInitProperties()
 	GraphInstance->OnGraphChangedDelegate.AddUObject(this, &UPCGComponent::OnGraphChanged);
 #endif // WITH_EDITOR
 
-#if WITH_EDITOR
-	// Force bIsPartitioned at false for new objects
-	if (PCGHelpers::IsNewObjectAndNotDefault(this, /*bCheckHierarchy=*/true))
-	{
-		bIsPartitioned = false;
-	}
-#endif // WITH_EDITOR
-
 	Super::PostInitProperties();
 }
 
diff --git a/Engine/Plugins/PCG/Source/PCG/Public/PCGComponent.h b/Engine/Plugins/PCG/Source/PCG/Public/PCGComponent.h
--- a/Engine/Plugins/PCG/Source/PCG/Public/PCGComponent.h
+++ b/Engine/Plugins/PCG/Source/PCG/Public/PCGComponent.h
@@ -551,10 +551,6 @@ protected:
 	TObjectPtr<UPCGGraph> Graph_DEPRECATED;
 
 private:
-	// Note for upgrade: can be safely replaced by bIsComponentPartitioned. Needed a new variable to change the default value. Kept to allow proper value change, cannot be deprecated. Do not use.
-	UPROPERTY()
-	bool bIsPartitioned = true;
-
 	/** Track if component should disable 'bIsComponentPartitioned'. Used to deprecate partitioned components in non-WP levels before support for partitioning in non-WP levels. */
 	bool bDisableIsComponentPartitionedOnLoad = false;
 #endif // WITH_EDITORONLY_DATA
