[PCG] Support multi attribute set data in GPU nodes
* Also support multiple attributes of the same name but different type

#rnx

diff --git a/Engine/Plugins/PCG/Source/PCG/Private/Compute/Elements/PCGCustomHLSL.cpp b/Engine/Plugins/PCG/Source/PCG/Private/Compute/Elements/PCGCustomHLSL.cpp
--- a/Engine/Plugins/PCG/Source/PCG/Private/Compute/Elements/PCGCustomHLSL.cpp
+++ b/Engine/Plugins/PCG/Source/PCG/Private/Compute/Elements/PCGCustomHLSL.cpp
@@ -786,13 +786,10 @@ void UPCGCustomHLSLSettings::UpdatePinSettings()
 			Properties.AllowedTypes = EPCGDataType::Point;
 		}
 
-		if (!!(Properties.AllowedTypes & EPCGDataType::Landscape)
-			|| !!(Properties.AllowedTypes & EPCGDataType::Param))
+		if (!!(Properties.AllowedTypes & EPCGDataType::Landscape))
 		{
-			// Don't allow multiple data on this pin because we do not support a dynamic number of textures/landscapes bound to a
+			// Don't allow multiple data on this pin because we do not support a dynamic number of landscapes bound to a
 			// compute kernel.
-			// Also disallow multi-data for Attribute Sets, since we require attributes to be uniform on a pin, thus having different
-			// attribute sets wouldn't work as expected.
 			Properties.bAllowMultipleData = false;
 
 #if WITH_EDITOR
@@ -837,24 +834,10 @@ void UPCGCustomHLSLSettings::UpdatePinSettings()
 		Properties.bAllowEditMultipleConnections = false;
 #endif
 
-		if (!!(Properties.AllowedTypes & EPCGDataType::Param))
-		{
-			Properties.bAllowMultipleData = false;
-			Properties.PropertiesGPU.DataCountMode = EPCGDataCountMode::Fixed;
-			Properties.PropertiesGPU.DataCount = 1;
-
 #if WITH_EDITOR
-			Properties.bAllowEditMultipleData = false;
-			Properties.PropertiesGPU.bAllowEditDataCount = false;
+		Properties.bAllowEditMultipleData = true;
+		Properties.PropertiesGPU.bAllowEditDataCount = true;
 #endif
-		}
-		else
-		{
-#if WITH_EDITOR
-			Properties.bAllowEditMultipleData = true;
-			Properties.PropertiesGPU.bAllowEditDataCount = true;
-#endif
-		}
 	}
 }
 
diff --git a/Engine/Plugins/PCG/Source/PCG/Private/Compute/Elements/PCGCustomHLSLKernel.cpp b/Engine/Plugins/PCG/Source/PCG/Private/Compute/Elements/PCGCustomHLSLKernel.cpp
--- a/Engine/Plugins/PCG/Source/PCG/Private/Compute/Elements/PCGCustomHLSLKernel.cpp
+++ b/Engine/Plugins/PCG/Source/PCG/Private/Compute/Elements/PCGCustomHLSLKernel.cpp
@@ -567,14 +567,48 @@ FString UPCGCustomHLSLKernel::GetCookedSource(const UPCGComputeGraph* InComputeG
 	// We inject attribute IDs directly into the source. This is most efficient and saves us passing them into the kernel. However
 	// the trade off is that a shader with a shared source will generate multiple variants if used in different compute graphs with
 	// different attribute IDs.
-	for (const FPCGKernelAttributeKey& AttributeKey : Settings->KernelAttributeKeys)
+	auto ReplaceAttributeNamesWithIds = [InComputeGraph](FString& InOutSource)
 	{
-		const FString SourceDefinition = PCGCustomHLSLKernel::GetKernelAttributeKeyAsString(AttributeKey);
-		const FString AttributeIndexAsString = FString::FromInt(InComputeGraph->GetStaticAttributeTable().GetAttributeId(AttributeKey));
+		FRegexMatcher AttributeFunctionMatcher(FRegexPattern(PCGHLSLElement::AttributeFunctionPattern), InOutSource);
+		TArray<TTuple<int32, int32, int32>> Replacements;
 
-		Source.ReplaceInline(*SourceDefinition, *AttributeIndexAsString);
-		Functions.ReplaceInline(*SourceDefinition, *AttributeIndexAsString);
-	}
+		const UEnum* AttributeTypeEnum = StaticEnum<EPCGKernelAttributeType>();
+		check(AttributeTypeEnum);
+
+		while (AttributeFunctionMatcher.FindNext())
+		{
+			const FString TypeStr = AttributeFunctionMatcher.GetCaptureGroup(PCGHLSLElement::AttributeTypeCaptureGroup);
+			const FString NameStr = AttributeFunctionMatcher.GetCaptureGroup(PCGHLSLElement::AttributeNameCaptureGroup);
+
+			const int64 AttrType = AttributeTypeEnum->GetValueByName(FName(*TypeStr));
+
+			if (!ensure(AttrType != INDEX_NONE))
+			{
+				break;
+			}
+
+			const FPCGKernelAttributeKey AttributeKey(FName(*NameStr), static_cast<EPCGKernelAttributeType>(AttrType));
+			const FString SourceDefinition = PCGCustomHLSLKernel::GetKernelAttributeKeyAsString(AttributeKey);
+
+			const int32 AttributeIndex = InComputeGraph->GetStaticAttributeTable().GetAttributeId(AttributeKey);
+			const int32 ReplaceStartIndex = InOutSource.Find(SourceDefinition, ESearchCase::CaseSensitive, ESearchDir::FromStart, AttributeFunctionMatcher.GetMatchBeginning());
+			const int32 ReplaceEndIndex = ReplaceStartIndex + SourceDefinition.Len();
+
+			Replacements.Emplace(AttributeIndex, ReplaceStartIndex, ReplaceEndIndex);
+		}
+
+		for (int I = Replacements.Num() - 1; I >= 0; --I)
+		{
+			const int32 AttributeIndex = Replacements[I].Get<0>();
+			const int32 ReplaceStartIndex = Replacements[I].Get<1>();
+			const int32 ReplaceEndIndex = Replacements[I].Get<2>();
+
+			InOutSource = InOutSource.Left(ReplaceStartIndex) + FString::FromInt(AttributeIndex) + InOutSource.RightChop(ReplaceEndIndex);
+		}
+	};
+
+	ReplaceAttributeNamesWithIds(Source);
+	ReplaceAttributeNamesWithIds(Functions);
 
 	// We inject data IDs directly into the source. They will get remapped to data indices using a label resolver data interface.
 	if (const FPCGPinDataLabels* PinToDataLabels = InComputeGraph->GetStaticDataLabelsTable().Find(KernelIndex))
@@ -988,9 +1022,9 @@ bool UPCGCustomHLSLKernel::AreAttributesValid(FPCGContext* InContext, FText* Out
 			const UEnum* AttributeTypeEnum = StaticEnum<EPCGKernelAttributeType>();
 			check(AttributeTypeEnum);
 
-			const int64 AttributeType = AttributeTypeEnum->GetValueByName(FName(*TypeStr));
+			const int64 AttrType = AttributeTypeEnum->GetValueByName(FName(*TypeStr));
 
-			if (AttributeType == INDEX_NONE)
+			if (AttrType == INDEX_NONE)
 			{
 #if PCG_KERNEL_LOGGING_ENABLED
 				if (OutErrorText)
@@ -1007,45 +1041,48 @@ bool UPCGCustomHLSLKernel::AreAttributesValid(FPCGContext* InContext, FText* Out
 
 			const FPCGKernelAttributeDesc* AttrDesc = nullptr;
 			const FName AttrName = FName(*NameStr);
+			bool bFoundMatchingAttributeName = false;
 
-			if (!PinDesc->DataDescs.IsEmpty())
+			// Verify that the attribute exists on at least one data in pin data collection.
+			for (const FPCGDataDesc& DataDesc : PinDesc->DataDescs)
 			{
-				// Note: This assumes attributes are the same on all data on a pin, which is true for now
-				const FPCGDataDesc& DataDesc = PinDesc->DataDescs[0];
-				AttrDesc = DataDesc.AttributeDescs.FindByPredicate([AttrName](const FPCGKernelAttributeDesc& Desc) { return Desc.Name == AttrName; });
-			}
+				AttrDesc = DataDesc.AttributeDescs.FindByPredicate([AttrName, AttrType, &bFoundMatchingAttributeName](const FPCGKernelAttributeDesc& Desc)
+				{
+					const bool bAttributeNameMatches = Desc.Name == AttrName;
+					bFoundMatchingAttributeName |= bAttributeNameMatches;
 
-			if (!AttrDesc)
-			{
-#if PCG_KERNEL_LOGGING_ENABLED
-				if (OutErrorText)
+					return bAttributeNameMatches && Desc.Type == static_cast<EPCGKernelAttributeType>(AttrType);
+				});
+
+				if (AttrDesc)
 				{
-					*OutErrorText = FText::Format(
-						LOCTEXT("InvalidAttributeName", "Tried to call attribute function '{0}' on attribute '{1}' which does not exist."),
-						ConstructFunctionText(),
-						FText::FromName(AttrName));
+					break;
 				}
-#endif
-
-				return false;
 			}
 
-			if (AttrDesc->Type != static_cast<EPCGKernelAttributeType>(AttributeType))
+			if (!AttrDesc)
 			{
 #if PCG_KERNEL_LOGGING_ENABLED
 				if (OutErrorText)
 				{
-					const FString ActualTypeStr = AttributeTypeEnum->GetNameStringByIndex(static_cast<int64>(AttrDesc->Type));
-
-					*OutErrorText = FText::Format(
-						LOCTEXT("AttributeTypeMismatch", "Type mismatch for call to attribute function '{0}' on attribute '{1}'. Expected '{2}' but received '{3}'."),
-						ConstructFunctionText(),
-						FText::FromName(AttrName),
-						FText::FromString(TypeStr),
-						FText::FromString(ActualTypeStr));
+					if (bFoundMatchingAttributeName)
+					{
+						*OutErrorText = FText::Format(
+							LOCTEXT("InvalidAttributeType", "Tried to call attribute function '{0}' on attribute '{1}' which is not of type '{2}'."),
+							ConstructFunctionText(),
+							FText::FromName(AttrName),
+							FText::FromString(TypeStr));
+					}
+					else
+					{
+						*OutErrorText = FText::Format(
+							LOCTEXT("InvalidAttributeDNE", "Tried to call attribute function '{0}' on attribute '{1}' which does not exist."),
+							ConstructFunctionText(),
+							FText::FromName(AttrName),
+							FText::FromString(TypeStr));
+					}
 				}
 #endif
-
 				return false;
 			}
 		}
diff --git a/Engine/Plugins/PCG/Source/PCG/Private/Compute/PCGPinPropertiesGPU.h b/Engine/Plugins/PCG/Source/PCG/Private/Compute/PCGPinPropertiesGPU.h
--- a/Engine/Plugins/PCG/Source/PCG/Private/Compute/PCGPinPropertiesGPU.h
+++ b/Engine/Plugins/PCG/Source/PCG/Private/Compute/PCGPinPropertiesGPU.h
@@ -111,7 +111,7 @@ struct FPCGPinPropertiesGPUStruct
 	UPROPERTY(Transient)
 	bool bAllowEditInitMode = false;
 
-	/** Some types force single data, such as Attribute Sets. We should not allow data count to be edited for such types. */
+	/** Some types force single data. We should not allow data count to be edited for such types. */
 	UPROPERTY(Transient)
 	bool bAllowEditDataCount = false;
 
