[PCG] Fix bad initialization of dummy texture data (overwrites off end).

#rnx

diff --git a/Engine/Plugins/PCG/Source/PCG/Private/Compute/DataInterfaces/PCGTextureDataInterface.cpp b/Engine/Plugins/PCG/Source/PCG/Private/Compute/DataInterfaces/PCGTextureDataInterface.cpp
--- a/Engine/Plugins/PCG/Source/PCG/Private/Compute/DataInterfaces/PCGTextureDataInterface.cpp
+++ b/Engine/Plugins/PCG/Source/PCG/Private/Compute/DataInterfaces/PCGTextureDataInterface.cpp
@@ -249,7 +249,7 @@ void FPCGTextureDataProviderProxy::AllocateResources(FRDGBuilder& GraphBuilder,
 
 				{
 					uint32 DestStride;
-					float* DestBuffer = (float*)RHILockTexture2D(Texture, 0, RLM_WriteOnly, DestStride, false);
+					uint8* DestBuffer = (uint8*)RHILockTexture2D(Texture, 0, RLM_WriteOnly, DestStride, false);
 					*DestBuffer = 0;
 					RHIUnlockTexture2D(Texture, 0, false);
 				}
