Update landscape cache default to never serialize on new levels.

[RN] Made the PCG landscape cache opt-in by default in new maps instead of opt-out.

diff --git a/Engine/Plugins/PCG/Source/PCG/Private/Grid/PCGLandscapeCache.cpp b/Engine/Plugins/PCG/Source/PCG/Private/Grid/PCGLandscapeCache.cpp
--- a/Engine/Plugins/PCG/Source/PCG/Private/Grid/PCGLandscapeCache.cpp
+++ b/Engine/Plugins/PCG/Source/PCG/Private/Grid/PCGLandscapeCache.cpp
@@ -5,6 +5,7 @@
 #include "PCGCustomVersion.h"
 #include "PCGPoint.h"
 #include "Helpers/PCGBlueprintHelpers.h"
+#include "Helpers/PCGHelpers.h"
 #include "Metadata/PCGMetadata.h"
 
 #include "LandscapeComponent.h"
@@ -17,6 +18,7 @@
 #include "Engine/World.h"
 #include "Kismet/GameplayStatics.h"
 #include "Serialization/BufferWriter.h"
+#include "UObject/FortniteMainBranchObjectVersion.h"
 
 #include UE_INLINE_GENERATED_CPP_BY_NAME(PCGLandscapeCache)
 
@@ -475,8 +477,30 @@ void UPCGLandscapeCache::BeginDestroy()
 	Super::BeginDestroy();
 }
 
+#if WITH_EDITOR
+void UPCGLandscapeCache::PostEditChangeProperty(FPropertyChangedEvent& PropertyChangedEvent)
+{
+	if (PropertyChangedEvent.Property && PropertyChangedEvent.GetPropertyName() == GET_MEMBER_NAME_CHECKED(UPCGLandscapeCache, SerializationMode))
+	{
+		// Reset error notification if the serialization mode changes.
+		bLoggedNoCacheError = false;
+	}
+
+	Super::PostEditChangeProperty(PropertyChangedEvent);
+}
+
+#endif // WITH_EDITOR
+
 void UPCGLandscapeCache::Serialize(FArchive& Archive)
 {
+	Archive.UsingCustomVersion(FFortniteMainBranchObjectVersion::GUID);
+
+	if (Archive.IsLoading() && Archive.CustomVer(FFortniteMainBranchObjectVersion::GUID) < FFortniteMainBranchObjectVersion::PCGLandscapeCacheDefaultSerializationChanged)
+	{
+		// Default value has changed from "serialize at cook" to "never serialize", so put back original default.
+		SerializationMode = EPCGLandscapeCacheSerializationMode::SerializeOnlyAtCook;
+	}
+
 	Super::Serialize(Archive);
 
 	const bool bShouldSerializeEntries = (Archive.IsSaving() &&
@@ -772,6 +796,17 @@ const FPCGLandscapeCacheEntry* UPCGLandscapeCache::GetCacheEntry(ULandscapeCompo
 		return nullptr;
 	}
 
+	if (SerializationMode == EPCGLandscapeCacheSerializationMode::NeverSerialize && PCGHelpers::IsRuntimeOrPIE())
+	{
+		if (!bLoggedNoCacheError)
+		{
+			UE_LOG(LogPCG, Error, TEXT("PCG Landscape cache (on the PCG World Actor) is not set to be serialized and will not work in non-editor modes."));
+			bLoggedNoCacheError = true;
+		}
+
+		return nullptr;
+	}
+
 	const FGuid LandscapeGuid = (LandscapeComponent->GetLandscapeProxy() ? LandscapeComponent->GetLandscapeProxy()->GetOriginalLandscapeGuid() : FGuid());
 	const FPCGLandscapeCacheEntry* CacheEntry = GetCacheEntry(LandscapeComponent->GetLandscapeProxy(), LandscapeGuid, ComponentCoordinate);
 
@@ -820,6 +855,12 @@ const FPCGLandscapeCacheEntry* UPCGLandscapeCache::GetCacheEntry(AActor* HintAct
 	{
 #if WITH_EDITOR
 		FReadScopeLock ScopeLock(CacheLock);
+#else
+		if (!bLoggedNoCacheError && CachedData.IsEmpty())
+		{
+			UE_LOG(LogPCG, Error, TEXT("PCG Landscape cache is empty; it probably was never serialized. Check the landscape cache object on the PCG World Actor."));
+			bLoggedNoCacheError = true;
+		}
 #endif
 		if (FPCGLandscapeCacheEntry** FoundEntry = CachedData.Find(ComponentKey))
 		{
diff --git a/Engine/Plugins/PCG/Source/PCG/Private/PCGSubsystem.cpp b/Engine/Plugins/PCG/Source/PCG/Private/PCGSubsystem.cpp
--- a/Engine/Plugins/PCG/Source/PCG/Private/PCGSubsystem.cpp
+++ b/Engine/Plugins/PCG/Source/PCG/Private/PCGSubsystem.cpp
@@ -299,7 +299,7 @@ void UPCGSubsystem::Tick(float DeltaSeconds)
 #if WITH_EDITOR
 		if (PCGSubsystemConsole::CVarRebuildLandscapeOnPIE.GetValueOnAnyThread() && PCGHelpers::IsRuntimeOrPIE())
 		{
-			BuildLandscapeCache(/*bQuiet=*/true);
+			BuildLandscapeCache(/*bQuiet=*/true, /*bForceBuild=*/false);
 		}
 #endif
 
@@ -1761,12 +1761,15 @@ void UPCGSubsystem::CleanFromCache(const IPCGElement* InElement, const UPCGSetti
 	}
 }
 
-void UPCGSubsystem::BuildLandscapeCache(bool bQuiet)
+void UPCGSubsystem::BuildLandscapeCache(bool bQuiet, bool bForceBuild)
 {
 	TRACE_CPUPROFILER_EVENT_SCOPE(UPCGSubsystem::BuildLandscapeCache);
 	if (UPCGLandscapeCache* LandscapeCache = GetLandscapeCache())
 	{
-		LandscapeCache->PrimeCache();
+		if (bForceBuild || LandscapeCache->SerializationMode != EPCGLandscapeCacheSerializationMode::NeverSerialize)
+		{
+			LandscapeCache->PrimeCache();
+		}
 	}
 	else if(!bQuiet)
 	{
diff --git a/Engine/Plugins/PCG/Source/PCG/Public/Grid/PCGLandscapeCache.h b/Engine/Plugins/PCG/Source/PCG/Public/Grid/PCGLandscapeCache.h
--- a/Engine/Plugins/PCG/Source/PCG/Public/Grid/PCGLandscapeCache.h
+++ b/Engine/Plugins/PCG/Source/PCG/Public/Grid/PCGLandscapeCache.h
@@ -124,6 +124,9 @@ public:
 	//~Begin UObject interface
 	virtual void BeginDestroy() override;
 	virtual void Serialize(FArchive& Archive) override;
+#if WITH_EDITOR
+	virtual void PostEditChangeProperty(FPropertyChangedEvent& PropertyChangedEvent) override;
+#endif
 	//~End UObject interface
 
 	/** Initialize cache. Can be safely called multiple times. */
@@ -153,7 +156,7 @@ public:
 	void SampleMetadataOnPoint(ALandscapeProxy* Landscape, FPCGPoint& OutPoint, UPCGMetadata* OutMetadata);
 
 	UPROPERTY(BlueprintReadWrite, EditAnywhere, Category = "Cache")
-	EPCGLandscapeCacheSerializationMode SerializationMode = EPCGLandscapeCacheSerializationMode::SerializeOnlyAtCook;
+	EPCGLandscapeCacheSerializationMode SerializationMode = EPCGLandscapeCacheSerializationMode::NeverSerialize;
 
 	UPROPERTY(BlueprintReadWrite, EditAnywhere, Category = "Cache")
 	EPCGLandscapeCacheSerializationContents CookedSerializedContents = EPCGLandscapeCacheSerializationContents::SerializeAll;
@@ -225,6 +228,7 @@ private:
 	std::atomic<int32> CacheTouch = 0;
 	float TimeSinceLastCleanupInSeconds = 0;
 	bool bInitialized = false;
+	bool bLoggedNoCacheError = false;
 
 #if WITH_EDITOR
 	TSet<TWeakObjectPtr<ALandscapeProxy>> Landscapes;
diff --git a/Engine/Plugins/PCG/Source/PCG/Public/PCGSubsystem.h b/Engine/Plugins/PCG/Source/PCG/Public/PCGSubsystem.h
--- a/Engine/Plugins/PCG/Source/PCG/Public/PCGSubsystem.h
+++ b/Engine/Plugins/PCG/Source/PCG/Public/PCGSubsystem.h
@@ -321,8 +321,8 @@ public:
 	/** If the partition grid size change, call this to empty the Partition actors map */
 	void ResetPartitionActorsMap() { ActorAndComponentMapping.ResetPartitionActorsMap(); }
 
-	/** Builds the landscape data cache */
-	void BuildLandscapeCache(bool bQuiet = false);
+	/** Builds the landscape data cache. If force build is true, then it will build the cache even if it is never serialized. */
+	void BuildLandscapeCache(bool bQuiet = false, bool bForceBuild = true);
 
 	/** Clears the landscape data cache */
 	void ClearLandscapeCache();
diff --git a/Engine/Source/Runtime/Core/Public/UObject/FortniteMainBranchObjectVersions.inl b/Engine/Source/Runtime/Core/Public/UObject/FortniteMainBranchObjectVersions.inl
--- a/Engine/Source/Runtime/Core/Public/UObject/FortniteMainBranchObjectVersions.inl
+++ b/Engine/Source/Runtime/Core/Public/UObject/FortniteMainBranchObjectVersions.inl
@@ -583,5 +583,8 @@ FORTNITE_MAIN_VERSION(AddedRuntimeVirtualTextureUseStreamingMipsInEditorMode, 18
 // Media plate holdout composite components have been replaced by a checkbox
 FORTNITE_MAIN_VERSION(MediaPlateHoldoutComponentRemoval, 190)
 
+// Changed PCG landscape cache default from "serialize at cook" to "never serialize"
+FORTNITE_MAIN_VERSION(PCGLandscapeCacheDefaultSerializationChanged, 191)
+
 // -----<new versions can be added above this line>-------------------------------------------------
 #undef FORTNITE_MAIN_VERSION
\ No newline at end of file
