[PCG] Add UX buttons to assign debug materials for visualizing density, color, to assist new users.

diff --git a/Engine/Plugins/PCG/Source/PCGEditor/Private/Widgets/SPCGEditorGraphDetailsView.cpp b/Engine/Plugins/PCG/Source/PCGEditor/Private/Widgets/SPCGEditorGraphDetailsView.cpp
--- a/Engine/Plugins/PCG/Source/PCGEditor/Private/Widgets/SPCGEditorGraphDetailsView.cpp
+++ b/Engine/Plugins/PCG/Source/PCGEditor/Private/Widgets/SPCGEditorGraphDetailsView.cpp
@@ -10,17 +10,133 @@
 #include "Nodes/PCGEditorGraphNode.h"
 
 #include "DetailsViewArgs.h"
+#include "IDetailPropertyExtensionHandler.h"
 #include "IDetailsView.h"
 #include "PropertyEditorModule.h"
+#include "Framework/Notifications/NotificationManager.h"
 #include "Layout/Visibility.h"
 #include "Modules/ModuleManager.h"
+
 #include "Widgets/DeclarativeSyntaxSupport.h"
 #include "Widgets/Images/SImage.h"
 #include "Widgets/Input/SButton.h"
+#include "Widgets/Notifications/SNotificationList.h"
 #include "Widgets/Text/STextBlock.h"
 
 #define LOCTEXT_NAMESPACE "PCGEditorGraphDetailsView"
 
+namespace PCGEditorGraphDetailsView
+{
+	namespace Constants
+	{
+		constexpr TCHAR DefaultDebugDensityMaterialPath[] = TEXT("/PCG/DebugObjects/PCG_DebugMaterial.PCG_DebugMaterial");
+		constexpr TCHAR DefaultDebugMaterialColorPath[] = TEXT("/PCG/DebugObjects/PCG_DebugMaterial_Color.PCG_DebugMaterial_Color");
+
+		static const FText DebugMaterialDensityButtonText = LOCTEXT("DebugMaterialDensityButtonText", "Density");
+		static const FText DebugMaterialColorButtonText = LOCTEXT("DebugMaterialColorButtonText", "Color");
+
+		static const FText DebugMaterialButtonToastText = LOCTEXT("DebugMaterialButtonToastText", "[PCG] Assigned Debug Material");
+
+		static const FText DebugMaterialDensityButtonToastSubText = LOCTEXT("DebugMaterialDensityButtonToastSubText", "Density - PCG_DebugMaterial");
+		static const FText DebugMaterialColorButtonToastSubText = LOCTEXT("DebugMaterialColorButtonToastSubText", "Color - PCG_DebugMaterial_Color");
+
+		static const FText DefaultDebugMaterialDensityButtonTooltip = LOCTEXT("DefaultDebugMaterialDensityButtonTooltip", "Select the default debug material - 'PCG_DebugMaterial'.");
+		static const FText DefaultDebugMaterialColorButtonTooltip = LOCTEXT("DefaultDebugMaterialColorButtonTooltip", "Select the default  debug material - 'PCG_DebugMaterial_Color'");
+	}
+
+	namespace Helpers
+	{
+		void DispatchEditorToast(const FText& Text, const FText& SubText = FText{}, const float Duration = 3.f)
+		{
+			FNotificationInfo ToastInfo(Text);
+			ToastInfo.ExpireDuration = Duration;
+			ToastInfo.bFireAndForget = true;
+			ToastInfo.SubText = SubText;
+
+			FSlateNotificationManager::Get().AddNotification(ToastInfo);
+		}
+	}
+}
+
+class FPCGDebugDetailsExtension : public IDetailPropertyExtensionHandler
+{
+	virtual bool IsPropertyExtendable(const UClass* InObjectClass, const IPropertyHandle& PropertyHandle) const override
+	{
+		if (const FProperty* Property = PropertyHandle.GetProperty())
+		{
+			return Property->GetFName() == GET_MEMBER_NAME_CHECKED(FPCGDebugVisualizationSettings, MaterialOverride);
+		}
+		else
+		{
+			return false;
+		}
+	}
+
+	virtual void ExtendWidgetRow(FDetailWidgetRow& InWidgetRow, const IDetailLayoutBuilder& InDetailBuilder, const UClass* InObjectClass, TSharedPtr<IPropertyHandle> PropertyHandle) override
+	{
+		using namespace PCGEditorGraphDetailsView;
+		if (PropertyHandle.IsValid())
+		{
+			check(IsPropertyExtendable(InObjectClass, *PropertyHandle));
+
+			InWidgetRow.ExtensionContent()
+			[
+				SNew(SHorizontalBox)
+				+ SHorizontalBox::Slot()
+				.MinWidth(0)
+				.MaxWidth(80)
+				.HAlign(HAlign_Right)
+				[
+					SNew(SVerticalBox)
+					+ SVerticalBox::Slot()
+					.AutoHeight()
+					[
+						SNew(SButton)
+						.Text(Constants::DebugMaterialDensityButtonText)
+						.HAlign(HAlign_Center)
+						.VAlign(VAlign_Center)
+						.IsEnabled_Lambda([PropertyHandle]
+						{
+							FString CurrentString;
+							PropertyHandle->GetValueAsFormattedString(CurrentString);
+							return CurrentString != Constants::DefaultDebugDensityMaterialPath;
+						})
+						.OnClicked_Lambda([PropertyHandle]
+						{
+							PropertyHandle->SetValueFromFormattedString(Constants::DefaultDebugDensityMaterialPath);
+							Helpers::DispatchEditorToast(Constants::DebugMaterialButtonToastText, Constants::DebugMaterialDensityButtonToastSubText);
+							return FReply::Handled();
+						})
+						.ToolTipText(Constants::DefaultDebugMaterialDensityButtonTooltip)
+					]
+					+ SVerticalBox::Slot()
+					.AutoHeight()
+					[
+						SNew(SButton)
+						.Text(Constants::DebugMaterialColorButtonText)
+						.HAlign(HAlign_Center)
+						.VAlign(VAlign_Center)
+						.IsEnabled_Lambda([PropertyHandle]
+						{
+							FString CurrentString;
+							PropertyHandle->GetValueAsFormattedString(CurrentString);
+							return CurrentString != Constants::DefaultDebugMaterialColorPath;
+						})
+						.OnClicked_Lambda([PropertyHandle]
+						{
+							PropertyHandle->SetValueFromFormattedString(Constants::DefaultDebugMaterialColorPath);
+							Helpers::DispatchEditorToast(Constants::DebugMaterialButtonToastText, Constants::DebugMaterialColorButtonToastSubText);
+
+							return FReply::Handled();
+						})
+						.ToolTipText(Constants::DefaultDebugMaterialColorButtonTooltip)
+					]
+				]
+			];
+		}
+	}
+};
+
 void SPCGEditorGraphDetailsView::Construct(const FArguments& InArgs)
 {
 	FDetailsViewArgs DetailsViewArgs;
@@ -34,6 +150,8 @@ void SPCGEditorGraphDetailsView::Construct(const FArguments& InArgs)
 
 	DetailsView->SetIsPropertyReadOnlyDelegate(FIsPropertyReadOnly::CreateSP(this, &SPCGEditorGraphDetailsView::IsReadOnlyProperty));
 	DetailsView->SetIsPropertyVisibleDelegate(FIsPropertyVisible::CreateSP(this, &SPCGEditorGraphDetailsView::IsVisibleProperty));
+	TSharedRef<FPCGDebugDetailsExtension> DebugDetailsExtension = MakeShared<FPCGDebugDetailsExtension>();
+	DetailsView->SetExtensionHandler(DebugDetailsExtension);
 
 	ChildSlot
 	[
