Added missing PCG_API on the generation source manager, otherwise external users won't be able to add theirs.

#rnx

diff --git a/Engine/Plugins/PCG/Source/PCG/Public/RuntimeGen/GenSources/PCGGenSourceBase.h b/Engine/Plugins/PCG/Source/PCG/Public/RuntimeGen/GenSources/PCGGenSourceBase.h
--- a/Engine/Plugins/PCG/Source/PCG/Public/RuntimeGen/GenSources/PCGGenSourceBase.h
+++ b/Engine/Plugins/PCG/Source/PCG/Public/RuntimeGen/GenSources/PCGGenSourceBase.h
@@ -17,7 +17,7 @@ class PCG_API UPCGGenSourceBase : public UInterface
  * A PCG Generation Source represents an object in the world that provokes nearby 
  * PCG Components to generate through the Runtime Generation Scheduler.
  */
-class IPCGGenSourceBase
+class PCG_API IPCGGenSourceBase
 {
 	GENERATED_BODY()
 
diff --git a/Engine/Plugins/PCG/Source/PCG/Public/RuntimeGen/PCGGenSourceManager.h b/Engine/Plugins/PCG/Source/PCG/Public/RuntimeGen/PCGGenSourceManager.h
--- a/Engine/Plugins/PCG/Source/PCG/Public/RuntimeGen/PCGGenSourceManager.h
+++ b/Engine/Plugins/PCG/Source/PCG/Public/RuntimeGen/PCGGenSourceManager.h
@@ -27,13 +27,13 @@ public:
 	void Tick() { bDirty = true; }
 
 	/** Creates the set of all generation sources tracked by the manager. Call 'Tick()' to keep tracked generation sources up to date. */
-	TSet<IPCGGenSourceBase*> GetAllGenSources(const APCGWorldActor* InPCGWorldActor);
+	TSet<IPCGGenSourceBase*> PCG_API GetAllGenSources(const APCGWorldActor* InPCGWorldActor);
 
 	/** Adds a UPCGGenSource to be tracked by the GenSourceManager. */
-	bool RegisterGenSource(IPCGGenSourceBase* InGenSource);
+	bool PCG_API RegisterGenSource(IPCGGenSourceBase* InGenSource);
 
 	/** Removes a UPCGGenSource from being tracked by the GenSourceManager. */
-	bool UnregisterGenSource(const IPCGGenSourceBase* InGenSource);
+	bool PCG_API UnregisterGenSource(const IPCGGenSourceBase* InGenSource);
 
 protected:
 	void OnGameModePostLogin(AGameModeBase* InGameMode, APlayerController* InPlayerController);
