Added cvar to force landscape cache to never serialize on post load.

#rnx

diff --git a/Engine/Plugins/PCG/Source/PCG/Private/PCGWorldActor.cpp b/Engine/Plugins/PCG/Source/PCG/Private/PCGWorldActor.cpp
--- a/Engine/Plugins/PCG/Source/PCG/Private/PCGWorldActor.cpp
+++ b/Engine/Plugins/PCG/Source/PCG/Private/PCGWorldActor.cpp
@@ -24,6 +24,14 @@
 
 #include UE_INLINE_GENERATED_CPP_BY_NAME(PCGWorldActor)
 
+namespace PCGWorldActor
+{
+	static TAutoConsoleVariable<bool> CVarForceLandscapeCacheToNeverSerialize(
+		TEXT("pcg.ForceLandscapeCacheNeverSerialized"),
+		false,
+		TEXT("Changes the landscape serialization mode to Never Serialize on load."));
+}
+
 PRAGMA_DISABLE_DEPRECATION_WARNINGS
 bool FPCGPartitionActorRecord::operator==(const FPCGPartitionActorRecord& InOther) const
 {
@@ -65,6 +73,12 @@ void APCGWorldActor::PostLoad()
 	// Even if value is not used by APCGWorldActor::IsEditorOnly, this value can show as read-only in the actor's details panel
 	// to avoid confusion we set it to false here.
 	bIsEditorOnlyActor = false;
+
+	if (PCGWorldActor::CVarForceLandscapeCacheToNeverSerialize.GetValueOnAnyThread())
+	{
+		// LandscapeCacheObject->Modify(); // can't modify in PostLoad.
+		LandscapeCacheObject->SerializationMode = EPCGLandscapeCacheSerializationMode::NeverSerialize;
+	}
 }
 
 bool APCGWorldActor::CanEditChange(const FProperty* InProperty) const
