[PCG] Fix race condition when deleting Context
- ContextHandle will take care of deleting the Context if it has been created
- Also make sure the lifetime of the pinned ContextHandle is greater than its usage (replace GetContextFromHandle by a new FSharedContext struct)

#rnx

diff --git a/Engine/Plugins/PCG/Source/PCG/Private/Compute/Elements/PCGComputeGraphElement.cpp b/Engine/Plugins/PCG/Source/PCG/Private/Compute/Elements/PCGComputeGraphElement.cpp
--- a/Engine/Plugins/PCG/Source/PCG/Private/Compute/Elements/PCGComputeGraphElement.cpp
+++ b/Engine/Plugins/PCG/Source/PCG/Private/Compute/Elements/PCGComputeGraphElement.cpp
@@ -428,7 +428,8 @@ bool FPCGComputeGraphElement::ExecuteInternal(FPCGContext* InContext) const
 			/*InFallbackDelegate=*/FSimpleDelegate::CreateLambda([ContextHandle = Context->GetOrCreateHandle()]()
 			{
 				// This render thread delegate will be executed if SubmitWork fails at any stage.
-				FPCGComputeGraphContext* Context = FPCGContext::GetContextFromHandle<FPCGComputeGraphContext>(ContextHandle);
+				FPCGContext::FSharedContext<FPCGComputeGraphContext> SharedContext(ContextHandle);
+				FPCGComputeGraphContext* Context = SharedContext.Get();
 				if (!Context)
 				{
 					return;
diff --git a/Engine/Plugins/PCG/Source/PCG/Private/Elements/Landscape/PCGWaitLandscapeReady.cpp b/Engine/Plugins/PCG/Source/PCG/Private/Elements/Landscape/PCGWaitLandscapeReady.cpp
--- a/Engine/Plugins/PCG/Source/PCG/Private/Elements/Landscape/PCGWaitLandscapeReady.cpp
+++ b/Engine/Plugins/PCG/Source/PCG/Private/Elements/Landscape/PCGWaitLandscapeReady.cpp
@@ -104,7 +104,8 @@ bool FPCGWaitLandscapeReadyElement::ExecuteInternal(FPCGContext* InContext) cons
 		Subsystem->ScheduleGeneric(
 			[ContextHandle = Context->GetOrCreateHandle()]() // Normal execution: wake up the current task
 			{
-				if (FPCGWaitLandscapeReadyElementContext* ContextPtr = FPCGContext::GetContextFromHandle<FPCGWaitLandscapeReadyElementContext>(ContextHandle))
+				FPCGContext::FSharedContext<FPCGWaitLandscapeReadyElementContext> SharedContext(ContextHandle);
+				if (FPCGWaitLandscapeReadyElementContext* ContextPtr = SharedContext.Get())
 				{
 					ContextPtr->bIsPaused = false;
 				}
@@ -112,7 +113,8 @@ bool FPCGWaitLandscapeReadyElement::ExecuteInternal(FPCGContext* InContext) cons
 			},
 			[ContextHandle = Context->GetOrCreateHandle()]() // On abort: wakeup and cancel
 			{
-				if (FPCGWaitLandscapeReadyElementContext* ContextPtr = FPCGContext::GetContextFromHandle<FPCGWaitLandscapeReadyElementContext>(ContextHandle))
+				FPCGContext::FSharedContext<FPCGWaitLandscapeReadyElementContext> SharedContext(ContextHandle);
+				if (FPCGWaitLandscapeReadyElementContext* ContextPtr = SharedContext.Get())
 				{
 					ContextPtr->bIsPaused = false;
 					ContextPtr->OutputData.bCancelExecution = true;
diff --git a/Engine/Plugins/PCG/Source/PCG/Private/Elements/PCGDataFromActor.cpp b/Engine/Plugins/PCG/Source/PCG/Private/Elements/PCGDataFromActor.cpp
--- a/Engine/Plugins/PCG/Source/PCG/Private/Elements/PCGDataFromActor.cpp
+++ b/Engine/Plugins/PCG/Source/PCG/Private/Elements/PCGDataFromActor.cpp
@@ -532,7 +532,8 @@ bool FPCGDataFromActorElement::ExecuteInternal(FPCGContext* InContext) const
 					Subsystem->ScheduleGeneric(
 						[ContextHandle = Context->GetOrCreateHandle()]() // Normal execution: Wake up the current task
 						{
-							if (FPCGDataFromActorContext* ContextPtr = FPCGContext::GetContextFromHandle<FPCGDataFromActorContext>(ContextHandle))
+							FPCGContext::FSharedContext<FPCGDataFromActorContext> SharedContext(ContextHandle);
+							if (FPCGDataFromActorContext* ContextPtr = SharedContext.Get())
 							{
 								ContextPtr->bIsPaused = false;
 							}
@@ -540,7 +541,8 @@ bool FPCGDataFromActorElement::ExecuteInternal(FPCGContext* InContext) const
 						}, 
 						[ContextHandle = Context->GetOrCreateHandle()]() // On Abort: Wake up on abort, clear all results and mark as cancelled
 						{
-							if (FPCGDataFromActorContext* ContextPtr = FPCGContext::GetContextFromHandle<FPCGDataFromActorContext>(ContextHandle))
+							FPCGContext::FSharedContext<FPCGDataFromActorContext> SharedContext(ContextHandle);
+							if (FPCGDataFromActorContext* ContextPtr = SharedContext.Get())
 							{
 								ContextPtr->bIsPaused = false;
 								ContextPtr->FoundActors.Reset();
diff --git a/Engine/Plugins/PCG/Source/PCG/Private/Elements/PCGGenerateGrassMaps.cpp b/Engine/Plugins/PCG/Source/PCG/Private/Elements/PCGGenerateGrassMaps.cpp
--- a/Engine/Plugins/PCG/Source/PCG/Private/Elements/PCGGenerateGrassMaps.cpp
+++ b/Engine/Plugins/PCG/Source/PCG/Private/Elements/PCGGenerateGrassMaps.cpp
@@ -348,7 +348,8 @@ bool FPCGGenerateGrassMapsElement::ExecuteInternal(FPCGContext* InContext) const
 
 		ENQUEUE_RENDER_COMMAND(GenerateGrassMaps)([ContextHandle, LandscapeComponentExtent=Context->LandscapeComponentExtent, LandscapeTileCoords=MoveTemp(LandscapeTileCoords)](FRHICommandListImmediate& RHICmdList)
 		{
-			FPCGGenerateGrassMapsContext* Context = FPCGContext::GetContextFromHandle<FPCGGenerateGrassMapsContext>(ContextHandle);
+			FPCGContext::FSharedContext<FPCGGenerateGrassMapsContext> SharedContext(ContextHandle);
+			FPCGGenerateGrassMapsContext* Context = SharedContext.Get();
 			if (!Context)
 			{
 				return;
@@ -427,7 +428,8 @@ bool FPCGGenerateGrassMapsElement::ExecuteInternal(FPCGContext* InContext) const
 			// Pass exported buffer back to game thread and wake up this element.
 			ExecuteOnGameThread(UE_SOURCE_LOCATION, [ContextHandle, GrassMapExported]()
 			{
-				if (FPCGGenerateGrassMapsContext* Context = FPCGContext::GetContextFromHandle<FPCGGenerateGrassMapsContext>(ContextHandle))
+				FPCGContext::FSharedContext<FPCGGenerateGrassMapsContext> SharedContext(ContextHandle);
+				if (FPCGGenerateGrassMapsContext* Context = SharedContext.Get())
 				{
 					Context->GrassMapHandle = GrassMapExported;
 					Context->bIsPaused = false;
@@ -477,7 +479,8 @@ bool FPCGGenerateGrassMapsElement::ExecuteInternal(FPCGContext* InContext) const
 		Context->bIsPaused = true;
 		Context->SourceComponent->GetSubsystem()->RegisterBeginTickAction([ContextHandle = Context->GetOrCreateHandle()]()
 		{
-			if (FPCGGenerateGrassMapsContext* ContextPtr = FPCGContext::GetContextFromHandle<FPCGGenerateGrassMapsContext>(ContextHandle))
+			FPCGContext::FSharedContext<FPCGGenerateGrassMapsContext> SharedContext(ContextHandle);
+			if (FPCGGenerateGrassMapsContext* ContextPtr = SharedContext.Get())
 			{
 				ContextPtr->bIsPaused = false;
 			}
diff --git a/Engine/Plugins/PCG/Source/PCG/Private/Elements/PCGLoopElement.cpp b/Engine/Plugins/PCG/Source/PCG/Private/Elements/PCGLoopElement.cpp
--- a/Engine/Plugins/PCG/Source/PCG/Private/Elements/PCGLoopElement.cpp
+++ b/Engine/Plugins/PCG/Source/PCG/Private/Elements/PCGLoopElement.cpp
@@ -344,7 +344,8 @@ bool FPCGLoopElement::ExecuteInternal(FPCGContext* InContext) const
 				Subsystem->ScheduleGeneric(
 					[ContextHandle = Context->GetOrCreateHandle()]() // Normal execution: Wake up the current task
 					{
-						if (FPCGSubgraphContext* ContextPtr = FPCGContext::GetContextFromHandle<FPCGSubgraphContext>(ContextHandle))
+						FPCGContext::FSharedContext<FPCGSubgraphContext> SharedContext(ContextHandle);
+						if (FPCGSubgraphContext* ContextPtr = SharedContext.Get())
 						{
 							ContextPtr->bIsPaused = false;
 						}
@@ -352,7 +353,8 @@ bool FPCGLoopElement::ExecuteInternal(FPCGContext* InContext) const
 					},
 					[ContextHandle = Context->GetOrCreateHandle()]() // On abort: wakeup and cancel, forget subgraphs
 					{
-						if (FPCGSubgraphContext* ContextPtr = FPCGContext::GetContextFromHandle<FPCGSubgraphContext>(ContextHandle))
+						FPCGContext::FSharedContext<FPCGSubgraphContext> SharedContext(ContextHandle);
+						if (FPCGSubgraphContext* ContextPtr = SharedContext.Get())
 						{
 							ContextPtr->bIsPaused = false;
 							ContextPtr->SubgraphTaskIds.Reset();
diff --git a/Engine/Plugins/PCG/Source/PCG/Private/Elements/PCGSpawnActor.cpp b/Engine/Plugins/PCG/Source/PCG/Private/Elements/PCGSpawnActor.cpp
--- a/Engine/Plugins/PCG/Source/PCG/Private/Elements/PCGSpawnActor.cpp
+++ b/Engine/Plugins/PCG/Source/PCG/Private/Elements/PCGSpawnActor.cpp
@@ -689,7 +689,8 @@ bool FPCGSpawnActorElement::SpawnAndPrepareSubgraphs(FPCGSubgraphContext* Contex
 			Subsystem->ScheduleGeneric(
 				[ContextHandle = Context->GetOrCreateHandle()]() // Normal execution: Wake up the current task
 				{
-					if (FPCGSubgraphContext* ContextPtr = FPCGContext::GetContextFromHandle<FPCGSubgraphContext>(ContextHandle))
+					FPCGContext::FSharedContext<FPCGSubgraphContext> SharedContext(ContextHandle);
+					if (FPCGSubgraphContext* ContextPtr = SharedContext.Get())
 					{
 						ContextPtr->bIsPaused = false;
 					}
@@ -697,7 +698,8 @@ bool FPCGSpawnActorElement::SpawnAndPrepareSubgraphs(FPCGSubgraphContext* Contex
 				},
 				[ContextHandle = Context->GetOrCreateHandle()]() // On Abort: Wake up & cancel
 				{
-					if (FPCGSubgraphContext* ContextPtr = FPCGContext::GetContextFromHandle<FPCGSubgraphContext>(ContextHandle))
+					FPCGContext::FSharedContext<FPCGSubgraphContext> SharedContext(ContextHandle);
+					if (FPCGSubgraphContext* ContextPtr = SharedContext.Get())
 					{
 						ContextPtr->bIsPaused = false;
 						ContextPtr->OutputData.bCancelExecution = true;
diff --git a/Engine/Plugins/PCG/Source/PCG/Private/Elements/PCGTextureSampler.cpp b/Engine/Plugins/PCG/Source/PCG/Private/Elements/PCGTextureSampler.cpp
--- a/Engine/Plugins/PCG/Source/PCG/Private/Elements/PCGTextureSampler.cpp
+++ b/Engine/Plugins/PCG/Source/PCG/Private/Elements/PCGTextureSampler.cpp
@@ -267,7 +267,8 @@ bool FPCGTextureSamplerElement::ExecuteInternal(FPCGContext* InContext) const
 				Context->bIsPaused = true;
 				Context->SourceComponent->GetSubsystem()->RegisterBeginTickAction([ContextHandle = Context->GetOrCreateHandle()]()
 					{
-						if (FPCGTextureSamplerContext* ContextPtr = FPCGContext::GetContextFromHandle<FPCGTextureSamplerContext>(ContextHandle))
+						FPCGContext::FSharedContext<FPCGTextureSamplerContext> SharedContext(ContextHandle);
+						if (FPCGTextureSamplerContext* ContextPtr = SharedContext.Get())
 						{
 							ContextPtr->bIsPaused = false;
 						}
diff --git a/Engine/Plugins/PCG/Source/PCG/Private/Graph/PCGGraphTask.cpp b/Engine/Plugins/PCG/Source/PCG/Private/Graph/PCGGraphTask.cpp
--- a/Engine/Plugins/PCG/Source/PCG/Private/Graph/PCGGraphTask.cpp
+++ b/Engine/Plugins/PCG/Source/PCG/Private/Graph/PCGGraphTask.cpp
@@ -116,7 +116,17 @@ FPCGGraphActiveTask::~FPCGGraphActiveTask()
 		PCGHelpers::ExecuteOnGameThread(UE_SOURCE_LOCATION, [ContextPtr = Context.Release()]()
 		{
 			PCGGraphExecutor::ClearAsyncFlags(ContextPtr->AsyncObjects);
-			delete ContextPtr;
+
+			if (!ContextPtr->Handle)
+			{
+				delete ContextPtr;
+			}
+			else
+			{
+				// Release the Handle SharedPtr (at this point only captured and pinned Handles can keep the Context alive)
+				// The Context will get deleted when the handle gets released
+				ContextPtr->Handle = nullptr;
+			}
 		});	
 	}
 }
diff --git a/Engine/Plugins/PCG/Source/PCG/Private/PCGContext.cpp b/Engine/Plugins/PCG/Source/PCG/Private/PCGContext.cpp
--- a/Engine/Plugins/PCG/Source/PCG/Private/PCGContext.cpp
+++ b/Engine/Plugins/PCG/Source/PCG/Private/PCGContext.cpp
@@ -44,6 +44,11 @@ namespace PCGContextHelpers
 	}
 }
 
+FPCGContextHandle::~FPCGContextHandle()
+{
+	delete Context;
+}
+
 UPCGBasePointData* FPCGContext::NewPointData_AnyThread(FPCGContext* Context)
 {
 	if (CVarPCGEnablePointArrayData.GetValueOnAnyThread())
diff --git a/Engine/Plugins/PCG/Source/PCG/Private/PCGSubgraph.cpp b/Engine/Plugins/PCG/Source/PCG/Private/PCGSubgraph.cpp
--- a/Engine/Plugins/PCG/Source/PCG/Private/PCGSubgraph.cpp
+++ b/Engine/Plugins/PCG/Source/PCG/Private/PCGSubgraph.cpp
@@ -804,7 +804,8 @@ bool FPCGSubgraphElement::ExecuteInternal(FPCGContext* InContext) const
 						Subsystem->ScheduleGeneric(
 							[ContextHandle = Context->GetOrCreateHandle()]() // Normal execution: Wake up the current task
 							{
-								if (FPCGSubgraphContext* ContextPtr = FPCGContext::GetContextFromHandle<FPCGSubgraphContext>(ContextHandle))
+								FPCGContext::FSharedContext<FPCGSubgraphContext> SharedContext(ContextHandle);
+								if (FPCGSubgraphContext* ContextPtr = SharedContext.Get())
 								{
 									ContextPtr->bIsPaused = false;
 								}
@@ -812,7 +813,8 @@ bool FPCGSubgraphElement::ExecuteInternal(FPCGContext* InContext) const
 							},
 							[ContextHandle = Context->GetOrCreateHandle()]() // On Abort: Wake up and cancel the execution
 							{
-								if (FPCGSubgraphContext* ContextPtr = FPCGContext::GetContextFromHandle<FPCGSubgraphContext>(ContextHandle))
+								FPCGContext::FSharedContext<FPCGSubgraphContext> SharedContext(ContextHandle);
+								if (FPCGSubgraphContext* ContextPtr = SharedContext.Get())
 								{
 									ContextPtr->bIsPaused = false;
 									ContextPtr->OutputData.bCancelExecution = true;
diff --git a/Engine/Plugins/PCG/Source/PCG/Public/PCGContext.h b/Engine/Plugins/PCG/Source/PCG/Public/PCGContext.h
--- a/Engine/Plugins/PCG/Source/PCG/Public/PCGContext.h
+++ b/Engine/Plugins/PCG/Source/PCG/Public/PCGContext.h
@@ -58,6 +58,8 @@ public:
 	{
 	}
 
+	~FPCGContextHandle();
+
 	FPCGContext* GetContext() { return Context; }
 private:
 	FPCGContext* Context = nullptr;
@@ -172,6 +174,7 @@ struct PCG_API FPCGContext
 		return PCGContextHelpers::GetInputSettings<SettingsType>(Node, InputData);
 	}
 
+	// This is not thread safe, make sure it is not called concurrently on the same context
 	TWeakPtr<FPCGContextHandle> GetOrCreateHandle()
 	{
 		if (!Handle)
@@ -183,6 +186,7 @@ struct PCG_API FPCGContext
 	}
 
 	template <typename ContextType>
+	UE_DEPRECATED(5.6, "Use FSharedContext instead")
 	static ContextType* GetContextFromHandle(TWeakPtr<FPCGContextHandle> WeakHandle)
 	{
 		if (TSharedPtr<FPCGContextHandle> SharedHandle = WeakHandle.Pin())
@@ -193,6 +197,20 @@ struct PCG_API FPCGContext
 		return nullptr;
 	}
 
+	template <typename ContextType>
+	class FSharedContext
+	{
+	public:
+		FSharedContext(TWeakPtr<FPCGContextHandle> WeakHandle)
+		{
+			SharedHandle = WeakHandle.Pin();
+		}
+		ContextType* Get() const { return SharedHandle.IsValid() ? (ContextType*)SharedHandle->GetContext() : nullptr; }
+
+	private:
+		TSharedPtr<FPCGContextHandle> SharedHandle;
+	};
+
 protected:
 	virtual UObject* GetExternalContainerForOverridableParam(const FPCGSettingsOverridableParam& InParam) { return nullptr; }
 	virtual void* GetUnsafeExternalContainerForOverridableParam(const FPCGSettingsOverridableParam& InParam) { return nullptr; }
