Fixed issue with debug object tree view selection going away when doing changes in dynamically executed graphs.
Improved refresh mechanism for the debug object tree view for dynamic graphs.

#rnx

diff --git a/Engine/Plugins/PCG/Source/PCGEditor/Private/PCGEditor.cpp b/Engine/Plugins/PCG/Source/PCGEditor/Private/PCGEditor.cpp
--- a/Engine/Plugins/PCG/Source/PCGEditor/Private/PCGEditor.cpp
+++ b/Engine/Plugins/PCG/Source/PCGEditor/Private/PCGEditor.cpp
@@ -2913,7 +2913,14 @@ void FPCGEditor::OnComponentGenerationDone(UPCGSubsystem* Subsystem, UPCGCompone
 	// We want to refresh if the component that is done generating has generated the current graph being edited,
 	// or if it is the root of the current stack being inspected (for subgraphs to also be refreshed).
 	// If we don't have a component, we refresh nonetheless.
-	const bool bShouldRefresh = !Component || StackBeingInspected.GetRootComponent() == Component || Component->GetGraph() == PCGGraphBeingEdited;
+	bool bShouldRefresh = !Component || StackBeingInspected.GetRootComponent() == Component || Component->GetGraph() == PCGGraphBeingEdited;
+
+	// Additionally, if we are not inspecting but the component that's done executing contains this graph, then we should also update.
+	if (!bShouldRefresh && PCGComponentBeingInspected.IsExplicitlyNull() && PCGGraphBeingEdited && Subsystem)
+	{
+		TArray<FPCGStack> ExecutedStacksContainingThisGraph = Subsystem->GetExecutedStacks(Component, PCGGraphBeingEdited);
+		bShouldRefresh |= !ExecutedStacksContainingThisGraph.IsEmpty();
+	}
 
 	if (!bShouldRefresh)
 	{
diff --git a/Engine/Plugins/PCG/Source/PCGEditor/Private/Widgets/SPCGEditorGraphDebugObjectTree.cpp b/Engine/Plugins/PCG/Source/PCGEditor/Private/Widgets/SPCGEditorGraphDebugObjectTree.cpp
--- a/Engine/Plugins/PCG/Source/PCGEditor/Private/Widgets/SPCGEditorGraphDebugObjectTree.cpp
+++ b/Engine/Plugins/PCG/Source/PCGEditor/Private/Widgets/SPCGEditorGraphDebugObjectTree.cpp
@@ -410,9 +410,6 @@ bool SPCGEditorGraphDebugObjectItemRow::IsFocusEnabled() const
 
 SPCGEditorGraphDebugObjectTree::~SPCGEditorGraphDebugObjectTree()
 {
-	FCoreUObjectDelegates::OnPreObjectPropertyChanged.RemoveAll(this);
-	FCoreUObjectDelegates::OnObjectPropertyChanged.RemoveAll(this);
-	FCoreUObjectDelegates::OnObjectConstructed.RemoveAll(this);
 	USelection::SelectionChangedEvent.RemoveAll(this);
 }
 
@@ -424,9 +421,6 @@ void SPCGEditorGraphDebugObjectTree::Construct(const FArguments& InArgs, TShared
 	UPCGGraph* PCGGraph = GetPCGGraph();
 	check(PCGGraph);
 
-	FCoreUObjectDelegates::OnPreObjectPropertyChanged.AddSP(this, &SPCGEditorGraphDebugObjectTree::OnPreObjectPropertyChanged);
-	FCoreUObjectDelegates::OnObjectPropertyChanged.AddSP(this, &SPCGEditorGraphDebugObjectTree::OnObjectPropertyChanged);
-	FCoreUObjectDelegates::OnObjectConstructed.AddSP(this, &SPCGEditorGraphDebugObjectTree::OnObjectConstructed);
 	USelection::SelectionChangedEvent.AddSP(this, &SPCGEditorGraphDebugObjectTree::OnEditorSelectionChanged);
 
 	const TSharedRef<SScrollBar> VerticalScrollBar = SNew(SScrollBar)
@@ -729,6 +723,7 @@ bool SPCGEditorGraphDebugObjectTree::GetFirstStackFromSelection(const UPCGNode*
 
 FPCGEditorGraphDebugObjectItemPtr SPCGEditorGraphDebugObjectTree::GetFirstDebugObjectFromSelection() const
 {
+	TRACE_CPUPROFILER_EVENT_SCOPE(SPCGEditorGraphDebugObjectTree::GetFirstDebugObjectFromSelection);
 	const UPCGGraph* PCGGraph = GetPCGGraph();
 	if (!PCGGraph)
 	{
@@ -1242,43 +1237,11 @@ void SPCGEditorGraphDebugObjectTree::RestoreTreeState()
 	}
 }
 
-void SPCGEditorGraphDebugObjectTree::OnPreObjectPropertyChanged(UObject* InObject, const FEditPropertyChain& /*InPropertyChain*/)
-{
-	if (UPCGGraphInstance* PCGgraphInstance = Cast<UPCGGraphInstance>(InObject))
-	{
-		const UPCGGraph* PCGGraph = PCGgraphInstance->GetGraph();
-		if (PCGGraph && PCGGraph == GetPCGGraph())
-		{
-			RequestRefresh();
-		}
-	}
-}
-
-void SPCGEditorGraphDebugObjectTree::OnObjectPropertyChanged(UObject* InObject, FPropertyChangedEvent& InPropertyChangedEvent)
-{
-	if (const UPCGGraphInterface* PCGGraphInterface = Cast<UPCGGraphInterface>(InObject))
-	{
-		const UPCGGraph* PCGGraph = PCGGraphInterface->GetGraph();
-		if (PCGGraph && PCGGraph == GetPCGGraph())
-		{
-			RequestRefresh();
-		}
-	}
-}
-
 void SPCGEditorGraphDebugObjectTree::OnEditorSelectionChanged(UObject* InObject)
 {
 	IsSetDebugObjectFromSelectionEnabled.Reset();
 }
 
-void SPCGEditorGraphDebugObjectTree::OnObjectConstructed(UObject* InObject)
-{
-	if (const UPCGComponent* PCGComponent = Cast<UPCGComponent>(InObject))
-	{
-		RequestRefresh();
-	}
-}
-
 UPCGGraph* SPCGEditorGraphDebugObjectTree::GetPCGGraph() const
 {
 	const TSharedPtr<FPCGEditor> PCGEditorPtr = PCGEditor.Pin();
diff --git a/Engine/Plugins/PCG/Source/PCGEditor/Private/Widgets/SPCGEditorGraphDebugObjectTree.h b/Engine/Plugins/PCG/Source/PCGEditor/Private/Widgets/SPCGEditorGraphDebugObjectTree.h
--- a/Engine/Plugins/PCG/Source/PCGEditor/Private/Widgets/SPCGEditorGraphDebugObjectTree.h
+++ b/Engine/Plugins/PCG/Source/PCGEditor/Private/Widgets/SPCGEditorGraphDebugObjectTree.h
@@ -295,9 +295,6 @@ private:
 		TMap<AActor*, TSharedPtr<FPCGEditorGraphDebugObjectItem_Actor>>& InOutActorItems,
 		TMap<const FPCGStack, FPCGEditorGraphDebugObjectItemPtr>& InOutStackToItem);
 
-	void OnPreObjectPropertyChanged(UObject* InObject, const FEditPropertyChain& InPropertyChain);
-	void OnObjectPropertyChanged(UObject* InObject, FPropertyChangedEvent& InPropertyChangedEvent);
-	void OnObjectConstructed(UObject* InObject);
 	void OnEditorSelectionChanged(UObject* InObject);
 
 	UPCGGraph* GetPCGGraph() const;
