[PCG] Support selecting layers in GenerateGrassMaps

diff --git a/Engine/Plugins/PCG/Source/PCG/Private/Elements/PCGGenerateGrassMaps.cpp b/Engine/Plugins/PCG/Source/PCG/Private/Elements/PCGGenerateGrassMaps.cpp
--- a/Engine/Plugins/PCG/Source/PCG/Private/Elements/PCGGenerateGrassMaps.cpp
+++ b/Engine/Plugins/PCG/Source/PCG/Private/Elements/PCGGenerateGrassMaps.cpp
@@ -9,6 +9,7 @@
 #include "Data/PCGLandscapeData.h"
 #include "Data/PCGTextureData.h"
 #include "Helpers/PCGHelpers.h"
+#include "Metadata/Accessors/PCGAttributeAccessorHelpers.h"
 
 #include "DrawDebugHelpers.h"
 #include "Landscape.h"
@@ -28,6 +29,7 @@
 namespace PCGGenerateGrassMaps
 {
 	const FName InputPinLabel = TEXT("Landscape");
+	const FName GrassTypeOverridesPinLabel = TEXT("Grass Types");
 
 	const FText LandscapeComponentLostError = LOCTEXT("LandscapeComponentLost", "Reference to one or more landscape components lost, grass maps will not be generated.");
 
@@ -57,6 +59,12 @@ TArray<FPCGPinProperties> UPCGGenerateGrassMapsSettings::InputPinProperties() co
 		/*bInAllowMultipleData=*/false);
 	PinProp.SetRequiredPin();
 
+	if (bOverrideFromInput)
+	{
+		FPCGPinProperties& GrassTypeOverrides = PinProperties.Emplace_GetRef(PCGGenerateGrassMaps::GrassTypeOverridesPinLabel, EPCGDataType::Param);
+		GrassTypeOverrides.SetRequiredPin();
+	}
+
 	return PinProperties;
 }
 
@@ -79,6 +87,20 @@ FPCGElementPtr UPCGGenerateGrassMapsSettings::CreateElement() const
 	return MakeShared<FPCGGenerateGrassMapsElement>();
 }
 
+#if WITH_EDITOR
+EPCGChangeType UPCGGenerateGrassMapsSettings::GetChangeTypeForProperty(const FName& InPropertyName) const
+{
+	EPCGChangeType ChangeType = Super::GetChangeTypeForProperty(InPropertyName);
+
+	if (InPropertyName == GET_MEMBER_NAME_CHECKED(UPCGGenerateGrassMapsSettings, bOverrideFromInput))
+	{
+		ChangeType |= EPCGChangeType::Structural;
+	}
+
+	return ChangeType;
+}
+#endif // WITH_EDITOR
+
 FPCGGenerateGrassMapsContext::~FPCGGenerateGrassMapsContext()
 {
 	delete LandscapeGrassWeightExporter;
@@ -121,6 +143,9 @@ bool FPCGGenerateGrassMapsElement::ExecuteInternal(FPCGContext* InContext) const
 	FPCGGenerateGrassMapsContext* Context = static_cast<FPCGGenerateGrassMapsContext*>(InContext);
 	check(Context);
 
+	const UPCGGenerateGrassMapsSettings* Settings = Context->GetInputSettings<UPCGGenerateGrassMapsSettings>();
+	check(Settings);
+
 	// 1. Select landscape components that overlap the given bounds.
 	if (!Context->bLandscapeComponentsFiltered)
 	{
@@ -155,6 +180,35 @@ bool FPCGGenerateGrassMapsElement::ExecuteInternal(FPCGContext* InContext) const
 			return true;
 		}
 
+		TArray<FString> SelectedGrassTypes;
+
+		if (!Settings->bOverrideFromInput)
+		{
+			SelectedGrassTypes = Settings->SelectedGrassTypes;
+		}
+		else
+		{
+			const TArray<FPCGTaggedData> OverrideTaggedDatas = InContext->InputData.GetInputsByPin(PCGGenerateGrassMaps::GrassTypeOverridesPinLabel);
+			TArray<FString> GrassTypeOverrides;
+
+			for (const FPCGTaggedData& OverrideTaggedData : OverrideTaggedDatas)
+			{
+				if (const UPCGData* OverrideData = OverrideTaggedData.Data)
+				{
+					const FPCGAttributePropertyInputSelector Selector = Settings->GrassTypesAttribute.CopyAndFixLast(OverrideData);
+
+					if (PCGAttributeAccessorHelpers::ExtractAllValues(OverrideData, Selector, GrassTypeOverrides, InContext))
+					{
+						SelectedGrassTypes = std::move(GrassTypeOverrides);
+					}
+					else
+					{
+						PCGLog::LogWarningOnGraph(LOCTEXT("FailExtractGrassTypeOverrides", "Failed to extract grass type overrides."), InContext);
+					}
+				}
+			}
+		}
+
 		const FBox ComponentBounds = SourceComponent->GetGridBounds();
 
 		for (TSoftObjectPtr<ALandscapeProxy> LandscapeProxyPtr : LandscapeData->Landscapes)
@@ -203,9 +257,25 @@ bool FPCGGenerateGrassMapsElement::ExecuteInternal(FPCGContext* InContext) const
 
 					if (Context->LandscapeComponents.Num() == 1)
 					{
-						for (TObjectPtr<ULandscapeGrassType> GrassType : LandscapeComponent->GetGrassTypes())
+						const TArray<TObjectPtr<ULandscapeGrassType>>& GrassTypes = LandscapeComponent->GetGrassTypes();
+						Context->NumGrassTypes = GrassTypes.Num();
+
+						FString GrassTypeName;
+						GrassTypeName.Reserve(256);
+
+						for (int GrassTypeIndex = 0; GrassTypeIndex < GrassTypes.Num(); ++GrassTypeIndex)
 						{
-							Context->GrassTypes.Add(GrassType);
+							if (TObjectPtr<ULandscapeGrassType> GrassType = GrassTypes[GrassTypeIndex])
+							{
+								GrassType->GetName(GrassTypeName);
+
+								const bool bIsSelectedLayer = SelectedGrassTypes.Contains(GrassTypeName);
+
+								if (bIsSelectedLayer != Settings->bExcludeSelectedGrassTypes)
+								{
+									Context->SelectedGrassTypes.Emplace(GrassType, GrassTypeIndex);
+								}
+							}
 						}
 
 						Context->GrassMapBounds = LandscapeComponentBounds;
@@ -227,7 +297,7 @@ bool FPCGGenerateGrassMapsElement::ExecuteInternal(FPCGContext* InContext) const
 						Context->GrassMapBounds += LandscapeComponentBounds;
 
 						// We expect all landscape components for a single landscape to have the same grass types. Do weak validation here.
-						const uint32 ExpectedNumGrassTypes = Context->GrassTypes.Num();
+						const uint32 ExpectedNumGrassTypes = Context->NumGrassTypes;
 						const uint32 ActualNumGrassTypes = LandscapeComponent->GetGrassTypes().Num();
 						ensure(ExpectedNumGrassTypes == ActualNumGrassTypes);
 
@@ -250,7 +320,7 @@ bool FPCGGenerateGrassMapsElement::ExecuteInternal(FPCGContext* InContext) const
 			}
 		}
 
-		if (!Context->LandscapeProxy.Get() || Context->LandscapeComponents.IsEmpty() || Context->GrassTypes.IsEmpty())
+		if (!Context->LandscapeProxy.Get() || Context->LandscapeComponents.IsEmpty() || Context->SelectedGrassTypes.IsEmpty())
 		{
 			return true;
 		}
@@ -376,7 +446,7 @@ bool FPCGGenerateGrassMapsElement::ExecuteInternal(FPCGContext* InContext) const
 			// Each corner of a quad in the landscape corresponds to one texel in the grass map.
 			const uint32 LandscapeComponentResolution = Context->LandscapeProxy->ComponentSizeQuads + 1;
 			const FIntPoint GrassMapResolution(LandscapeComponentResolution * NumTilesX, LandscapeComponentResolution * NumTilesY);
-			const uint32 NumGrassTypes = Context->GrassTypes.Num();
+			const uint32 NumGrassTypes = Context->NumGrassTypes;
 
 			// Output texture is array of textures, one per grass map.
 			FRDGTextureDesc GrassMapDesc = FRDGTextureDesc::Create2DArray(
@@ -449,8 +519,8 @@ bool FPCGGenerateGrassMapsElement::ExecuteInternal(FPCGContext* InContext) const
 
 	if (Context->TextureDatas.IsEmpty())
 	{
-		// Create the texture data objects if they haven't been created already.
-		for (int DataIndex = 0; DataIndex < Context->GrassTypes.Num(); ++DataIndex)
+		// Create the texture data objects if they haven't been created already. There should be one per selected grass type.
+		for (int DataIndex = 0; DataIndex < Context->SelectedGrassTypes.Num(); ++DataIndex)
 		{
 			UPCGTextureData* TextureData = FPCGContext::NewObject_AnyThread<UPCGTextureData>(Context);
 			Context->TextureDatas.Add(TextureData);
@@ -462,14 +532,14 @@ bool FPCGGenerateGrassMapsElement::ExecuteInternal(FPCGContext* InContext) const
 
 	check(Context->GrassMapHandle.IsValid());
 
-	for (int DataIndex = 0; DataIndex < Context->GrassTypes.Num(); ++DataIndex)
+	for (int DataIndex = 0; DataIndex < Context->SelectedGrassTypes.Num(); ++DataIndex)
 	{
 		check(Context->TextureDatas.IsValidIndex(DataIndex));
 		UPCGTextureData* TextureData = Context->TextureDatas[DataIndex];
 		check(TextureData);
 
 		// Poll initialize (fine to be called even when initialization was already complete).
-		bAllTexturesInitialized &= TextureData->Initialize(Context->GrassMapHandle, /*TextureIndex=*/DataIndex, GrassMapTransform);
+		bAllTexturesInitialized &= TextureData->Initialize(Context->GrassMapHandle, /*TextureIndex=*/Context->SelectedGrassTypes[DataIndex].Get<1>(), GrassMapTransform);
 	}
 
 	if (!bAllTexturesInitialized)
@@ -490,8 +560,7 @@ bool FPCGGenerateGrassMapsElement::ExecuteInternal(FPCGContext* InContext) const
 	}
 
 	// 5. Emit texture data objects.
-
-	for (int DataIndex = 0; DataIndex < Context->GrassTypes.Num(); ++DataIndex)
+	for (int DataIndex = 0; DataIndex < Context->SelectedGrassTypes.Num(); ++DataIndex)
 	{
 		check(Context->TextureDatas.IsValidIndex(DataIndex));
 		UPCGTextureData* TextureData = Context->TextureDatas[DataIndex];
@@ -508,7 +577,7 @@ bool FPCGGenerateGrassMapsElement::ExecuteInternal(FPCGContext* InContext) const
 		
 		FString GrassTypeName = TEXT("");
 
-		if (ULandscapeGrassType* GrassType = Context->GrassTypes[DataIndex].Get())
+		if (ULandscapeGrassType* GrassType = Context->SelectedGrassTypes[DataIndex].Get<0>().Get())
 		{
 			GrassType->GetName(GrassTypeName);
 		}
diff --git a/Engine/Plugins/PCG/Source/PCG/Public/Elements/PCGGenerateGrassMaps.h b/Engine/Plugins/PCG/Source/PCG/Public/Elements/PCGGenerateGrassMaps.h
--- a/Engine/Plugins/PCG/Source/PCG/Public/Elements/PCGGenerateGrassMaps.h
+++ b/Engine/Plugins/PCG/Source/PCG/Public/Elements/PCGGenerateGrassMaps.h
@@ -35,7 +35,27 @@ protected:
 	virtual TArray<FPCGPinProperties> OutputPinProperties() const override;
 	virtual bool IsInputPinRequiredByExecution(const UPCGPin* InPin) const override { return true; }
 	virtual FPCGElementPtr CreateElement() const override;
+#if WITH_EDITOR
+	virtual EPCGChangeType GetChangeTypeForProperty(const FName& InPropertyName) const override;
+#endif
 	//~End UPCGSettings interface
+
+public:
+	/** Select which grass types to generate. */
+	UPROPERTY(BlueprintReadWrite, EditAnywhere, Category = Settings, meta = (EditCondition = "!bOverrideFromInput"))
+	TArray<FString> SelectedGrassTypes;
+
+	/** Override grass types from input. */
+	UPROPERTY(BlueprintReadWrite, EditAnywhere, Category = Settings)
+	bool bOverrideFromInput = false;
+
+	/** Input attribute to pull grass type strings from. */
+	UPROPERTY(BlueprintReadWrite, EditAnywhere, Category = Settings, meta = (EditCondition = "bOverrideFromInput", PCG_DiscardPropertySelection, PCG_DiscardExtraSelection))
+	FPCGAttributePropertyInputSelector GrassTypesAttribute;
+
+	/** If toggled, will only generate grass types which are not selected. */
+	UPROPERTY(BlueprintReadWrite, EditAnywhere, Category = Settings, meta = (PCG_Overridable))
+	bool bExcludeSelectedGrassTypes = false;
 };
 
 struct FPCGGenerateGrassMapsContext : public FPCGContext
@@ -51,7 +71,13 @@ public:
 	/** Exported result texture array. */
 	TRefCountPtr<IPooledRenderTarget> GrassMapHandle = nullptr;
 
-	TArray<TWeakObjectPtr<ULandscapeGrassType>> GrassTypes;
+	/** List of the grass types selected for generation. We have to hold their texture index as well, since
+	 * this array could be sparse, but the actual texture array we produce will always have all the grass types.
+	 */
+	TArray<TTuple<TWeakObjectPtr<ULandscapeGrassType>, /*TextureIndex=*/int32>> SelectedGrassTypes;
+
+	/** Total number of grass types used by the landscape component. Includes grass types which were not selected for generation. */
+	int32 NumGrassTypes = 0;
 
 	TWeakObjectPtr<ALandscapeProxy> LandscapeProxy = nullptr;
 	TArray<TWeakObjectPtr<ULandscapeComponent>> LandscapeComponents;
