Fixed potential aliasing issue because we didn't use the intersected bounds for the generation range distance test.
Fixed typo in GetCleanupRadius for unbounded.

diff --git a/Engine/Plugins/PCG/Source/PCG/Private/PCGGraph.cpp b/Engine/Plugins/PCG/Source/PCG/Private/PCGGraph.cpp
--- a/Engine/Plugins/PCG/Source/PCG/Private/PCGGraph.cpp
+++ b/Engine/Plugins/PCG/Source/PCG/Private/PCGGraph.cpp
@@ -1365,7 +1365,7 @@ double UPCGGraph::GetGridCleanupRadiusFromGrid(EPCGHiGenGrid Grid) const
 {
 	if (Grid == EPCGHiGenGrid::Unbounded || Grid == EPCGHiGenGrid::Uninitialized)
 	{
-		return GenerationRadii.GetGenerationRadiusFromGrid(Grid);
+		return GenerationRadii.GetCleanupRadiusFromGrid(Grid);
 	}
 	// If the queried grid is smaller than the min grid including the exponent, we'll take the min grid and scale it down
 	else if (static_cast<uint32>(Grid) < (static_cast<uint32>(EPCGHiGenGrid::GridMin) << HiGenExponential))
diff --git a/Engine/Plugins/PCG/Source/PCG/Private/RuntimeGen/PCGRuntimeGenScheduler.cpp b/Engine/Plugins/PCG/Source/PCG/Private/RuntimeGen/PCGRuntimeGenScheduler.cpp
--- a/Engine/Plugins/PCG/Source/PCG/Private/RuntimeGen/PCGRuntimeGenScheduler.cpp
+++ b/Engine/Plugins/PCG/Source/PCG/Private/RuntimeGen/PCGRuntimeGenScheduler.cpp
@@ -399,7 +399,7 @@ void FPCGRuntimeGenScheduler::TickQueueComponentsForGeneration(
 
 								// Verify the grid cell actually lies within the generation radius.
 								// TODO: this is no longer necessary if we rasterize the sphere instead.
-								const double LocalDistanceSquared = CellBounds.ComputeSquaredDistanceToPoint(ModifiedGenSourcePosition);
+								const double LocalDistanceSquared = IntersectedBounds.ComputeSquaredDistanceToPoint(ModifiedGenSourcePosition);
 								if (LocalDistanceSquared <= GenerationRadius * GenerationRadius && 
 									Policy->ShouldGenerate(GenSource, IntersectedBounds, Key.GetGridDescriptor().Is2DGrid()))
 								{
