[PCG] Hide metadata detail view for all data

* It was showing for PCG settings

#rnx

diff --git a/Engine/Plugins/PCG/Source/PCG/Public/PCGData.h b/Engine/Plugins/PCG/Source/PCG/Public/PCGData.h
--- a/Engine/Plugins/PCG/Source/PCG/Public/PCGData.h
+++ b/Engine/Plugins/PCG/Source/PCG/Public/PCGData.h
@@ -100,7 +100,7 @@ public:
 	virtual TArray<FPCGMetadataDomainID> GetAllSupportedMetadataDomainIDs() const { return {PCGMetadataDomainID::Data}; }
 
 	// Not accessible through blueprint to make sure the constness is preserved.
-	UPROPERTY(VisibleAnywhere, Category = Metadata)
+	UPROPERTY()
 	TObjectPtr<UPCGMetadata> Metadata = nullptr;
 
 	void MarkUsage(EPCGDataUsage InUsage) const;
