[PCG] Generate grass maps node - extend material validation to look at override materials.

#rnx

diff --git a/Engine/Plugins/PCG/Source/PCG/Private/Elements/PCGGenerateGrassMaps.cpp b/Engine/Plugins/PCG/Source/PCG/Private/Elements/PCGGenerateGrassMaps.cpp
--- a/Engine/Plugins/PCG/Source/PCG/Private/Elements/PCGGenerateGrassMaps.cpp
+++ b/Engine/Plugins/PCG/Source/PCG/Private/Elements/PCGGenerateGrassMaps.cpp
@@ -251,6 +251,21 @@ bool FPCGGenerateGrassMapsElement::ExecuteInternal(FPCGContext* InContext) const
 						break;
 					}
 
+					if (LandscapeComponent->OverrideMaterial && Context->LandscapeProxy->LandscapeMaterial->GetBaseMaterial() != LandscapeComponent->OverrideMaterial->GetBaseMaterial())
+					{
+#if !(UE_BUILD_SHIPPING || UE_BUILD_TEST) || USE_LOGGING_IN_SHIPPING
+						PCGLog::LogErrorOnGraph(
+							FText::Format(LOCTEXT("LandscapeOverrideMaterialMismatch",
+								"Encountered landscape component {0} with override material {1} that did not match the landscape proxy material {2}. Only one landscape material per landscape is supported at this time."),
+								FText::FromString(LandscapeComponent->GetName()),
+								FText::FromString(LandscapeComponent->OverrideMaterial->GetName()),
+								FText::FromString(Context->LandscapeProxy->LandscapeMaterial->GetName())),
+							Context);
+#endif
+
+						return true;
+					}
+
 					LandscapeComponent->UpdateGrassTypes();
 
 					Context->LandscapeComponents.Add(LandscapeComponent);
