[PCG] FPCGStackFrame provide a couple of methods to resolve the frame object safely

#rnx

diff --git a/Engine/Plugins/PCG/Source/PCG/Private/Elements/PCGGetSubgraphDepth.cpp b/Engine/Plugins/PCG/Source/PCG/Private/Elements/PCGGetSubgraphDepth.cpp
--- a/Engine/Plugins/PCG/Source/PCG/Private/Elements/PCGGetSubgraphDepth.cpp
+++ b/Engine/Plugins/PCG/Source/PCG/Private/Elements/PCGGetSubgraphDepth.cpp
@@ -86,33 +86,37 @@ bool FPCGGetSubgraphDepthElement::ExecuteInternal(FPCGContext* Context) const
 	int SkipGraphCount = 0;
 	int Depth = 0;
 
-	for (int FrameIndex = StackFrames.Num() - 1; FrameIndex >= 0; --FrameIndex)
+	if (StackFrames.Num() > 0)
 	{
-		const FPCGStackFrame& StackFrame = StackFrames[FrameIndex];
-		if (const UPCGGraph* Graph = Cast<UPCGGraph>(StackFrame.Object.Get()))
+		FGCScopeGuard Guard;
+		for (int FrameIndex = StackFrames.Num() - 1; FrameIndex >= 0; --FrameIndex)
 		{
-			if (!bRecursive)
+			const FPCGStackFrame& StackFrame = StackFrames[FrameIndex];
+			if (const UPCGGraph* Graph = StackFrame.GetObject_NoGuard<UPCGGraph>())
 			{
-				++Depth;
-			}
-			else
-			{
-				if (!TargetGraph)
+				if (!bRecursive)
+				{
+					++Depth;
+				}
+				else
 				{
-					if (SkipGraphCount == Settings->DistanceRelativeToUpstreamGraph)
+					if (!TargetGraph)
 					{
-						TargetGraph = Graph;
-						Depth = 1;
+						if (SkipGraphCount == Settings->DistanceRelativeToUpstreamGraph)
+						{
+							TargetGraph = Graph;
+							Depth = 1;
+						}
+						else
+						{
+							++SkipGraphCount;
+						}
 					}
-					else
+					else if (TargetGraph == Graph)
 					{
-						++SkipGraphCount;
+						++Depth;
 					}
 				}
-				else if (TargetGraph == Graph)
-				{
-					++Depth;
-				}
 			}
 		}
 	}
diff --git a/Engine/Plugins/PCG/Source/PCG/Private/Graph/PCGStackContext.cpp b/Engine/Plugins/PCG/Source/PCG/Private/Graph/PCGStackContext.cpp
--- a/Engine/Plugins/PCG/Source/PCG/Private/Graph/PCGStackContext.cpp
+++ b/Engine/Plugins/PCG/Source/PCG/Private/Graph/PCGStackContext.cpp
@@ -68,43 +68,46 @@ bool FPCGStack::CreateStackFramePath(FString& OutString, const UPCGNode* InNode,
 		}
 	};
 
-	for (const FPCGStackFrame& Frame : StackFrames)
+	if (StackFrames.Num() > 0)
 	{
-		if (Frame.Object.IsValid())
+		FGCScopeGuard Guard;
+		for (const FPCGStackFrame& Frame : StackFrames)
 		{
-			AddPathSeparator();
-			const UObject* Object = Frame.Object.Get();
-
-			if (!Object)
+			if (const UObject* Object = Frame.GetObject_NoGuard())
 			{
-				// If any object does not resolve, cannot build the string
-				return false;
-			}
+				AddPathSeparator();
 
-			if (Object->IsA<UPCGComponent>())
-			{
-				StringBuilder << TEXT("COMPONENT:") << Object->GetFullName();
-			}
-			else if (Object->IsA<UPCGGraph>())
-			{
-				StringBuilder << TEXT("GRAPH:") << Object->GetFullName();
-			}
-			else if (Object->IsA<UPCGNode>())
-			{
-				StringBuilder << TEXT("NODE:") << Object->GetFName();
+				if (!Object)
+				{
+					// If any object does not resolve, cannot build the string
+					return false;
+				}
+
+				if (Object->IsA<UPCGComponent>())
+				{
+					StringBuilder << TEXT("COMPONENT:") << Object->GetFullName();
+				}
+				else if (Object->IsA<UPCGGraph>())
+				{
+					StringBuilder << TEXT("GRAPH:") << Object->GetFullName();
+				}
+				else if (Object->IsA<UPCGNode>())
+				{
+					StringBuilder << TEXT("NODE:") << Object->GetFName();
+				}
+				else
+				{
+					// Unrecognized type, should not happen
+					ensure(false);
+					StringBuilder << TEXT("UNRECOGNIZED:") << Object->GetFullName();
+				}
 			}
-			else
+			else if (Frame.LoopIndex != -1)
 			{
-				// Unrecognized type, should not happen
-				ensure(false);
-				StringBuilder << TEXT("UNRECOGNIZED:") << Object->GetFullName();
+				AddPathSeparator();
+				StringBuilder << TEXT("LOOP:") << FString::FromInt(Frame.LoopIndex);
 			}
 		}
-		else if (Frame.LoopIndex != -1)
-		{
-			AddPathSeparator();
-			StringBuilder << TEXT("LOOP:") << FString::FromInt(Frame.LoopIndex);
-		}
 	}
 
 	if (InNode)
@@ -126,9 +129,16 @@ bool FPCGStack::CreateStackFramePath(FString& OutString, const UPCGNode* InNode,
 uint32 FPCGStack::GetNumGraphLevels() const
 {
 	uint32 GraphCount = 0;
-	for (const FPCGStackFrame& Frame : StackFrames)
+	if (StackFrames.Num() > 0)
 	{
-		GraphCount += (Frame.Object.IsValid() && Frame.Object->IsA<UPCGGraph>()) ? 1 : 0;
+		FGCScopeGuard Guard;
+		for (const FPCGStackFrame& Frame : StackFrames)
+		{
+			if (const UObject* FrameObject = Frame.GetObject_NoGuard<UPCGGraph>())
+			{
+				++GraphCount;
+			}
+		}
 	}
 
 	return GraphCount;
@@ -160,21 +170,25 @@ bool FPCGStack::BeginsWith(const FPCGStack& Other) const
 
 const UPCGComponent* FPCGStack::GetRootComponent() const
 {
-	return StackFrames.IsEmpty() ? nullptr : Cast<const UPCGComponent>(StackFrames[0].Object.Get());
+	return StackFrames.IsEmpty() ? nullptr : StackFrames[0].GetObject_AnyThread<UPCGComponent>();
 }
 
 const UPCGGraph* FPCGStack::GetRootGraph(int32* OutRootFrameIndex) const
 {
-	for (int StackIndex = 0; StackIndex < GetStackFrames().Num(); ++StackIndex)
+	if (GetStackFrames().Num() > 0)
 	{
-		if (const UPCGGraph* Graph = Cast<const UPCGGraph>(StackFrames[StackIndex].Object.Get()))
+		FGCScopeGuard Guard;
+		for (int StackIndex = 0; StackIndex < GetStackFrames().Num(); ++StackIndex)
 		{
-			if (OutRootFrameIndex)
+			if (const UPCGGraph* Graph = StackFrames[StackIndex].GetObject_NoGuard<UPCGGraph>())
 			{
-				*OutRootFrameIndex = StackIndex;
-			}
+				if (OutRootFrameIndex)
+				{
+					*OutRootFrameIndex = StackIndex;
+				}
 
-			return Graph;
+				return Graph;
+			}
 		}
 	}
 
@@ -183,11 +197,15 @@ const UPCGGraph* FPCGStack::GetRootGraph(int32* OutRootFrameIndex) const
 
 const UPCGGraph* FPCGStack::GetGraphForCurrentFrame() const
 {
-	for (int StackIndex = GetStackFrames().Num() - 1; StackIndex >= 0; --StackIndex)
+	if (GetStackFrames().Num() > 0)
 	{
-		if (const UPCGGraph* Graph = Cast<const UPCGGraph>(StackFrames[StackIndex].Object.Get()))
+		FGCScopeGuard Guard;
+		for (int StackIndex = GetStackFrames().Num() - 1; StackIndex >= 0; --StackIndex)
 		{
-			return Graph;
+			if (const UPCGGraph* Graph = StackFrames[StackIndex].GetObject_AnyThread<UPCGGraph>())
+			{
+				return Graph;
+			}
 		}
 	}
 
@@ -202,7 +220,7 @@ const UPCGGraph* FPCGStack::GetNearestDynamicSubgraphForCurrentFrame() const
 	{
 		if (StackFrames[StackIndex - 1].IsLoopIndexFrame())
 		{
-			return Cast<UPCGGraph>(StackFrames[StackIndex].Object.Get());
+			return StackFrames[StackIndex].GetObject_AnyThread<UPCGGraph>();
 		}
 	}
 
@@ -217,16 +235,20 @@ const UPCGGraph* FPCGStack::GetNearestNonInlinedGraphForCurrentFrame() const
 
 const UPCGNode* FPCGStack::GetCurrentFrameNode() const
 {
-	return StackFrames.IsEmpty() ? nullptr : Cast<const UPCGNode>(StackFrames.Last().Object.Get());
+	return StackFrames.IsEmpty() ? nullptr : StackFrames.Last().GetObject_AnyThread<UPCGNode>();
 }
 
 const UPCGNode* FPCGStack::GetNodeForCurrentFrame() const
 {
-	for (int StackIndex = GetStackFrames().Num() - 1; StackIndex >= 0; --StackIndex)
+	if (GetStackFrames().Num() > 0)
 	{
-		if (const UPCGNode* Node = Cast<const UPCGNode>(StackFrames[StackIndex].Object.Get()))
+		FGCScopeGuard Guard;
+		for (int StackIndex = GetStackFrames().Num() - 1; StackIndex >= 0; --StackIndex)
 		{
-			return Node;
+			if (const UPCGNode* Node = StackFrames[StackIndex].GetObject_AnyThread<UPCGNode>())
+			{
+				return Node;
+			}
 		}
 	}
 
@@ -235,9 +257,15 @@ const UPCGNode* FPCGStack::GetNodeForCurrentFrame() const
 
 bool FPCGStack::HasObject(const UObject* InObject) const
 {
+	if (GetStackFrames().IsEmpty())
+	{
+		return false;
+	}
+
+	FGCScopeGuard Guard;
 	return !!Algo::FindByPredicate(StackFrames, [InObject](const FPCGStackFrame& Frame)
 	{
-		return Frame.Object == InObject;
+		return Frame.GetObject_NoGuard() == InObject;
 	});
 }
 
@@ -245,7 +273,7 @@ bool FPCGStack::ReplaceRoot(const TMap<UObject*, UObject*>& ReplacementMap)
 {
 	if (!StackFrames.IsEmpty())
 	{
-		UObject* NewStackRoot = ReplacementMap.FindRef(StackFrames[0].Object.Get());
+		UObject* NewStackRoot = ReplacementMap.FindRef(StackFrames[0].GetObject_AnyThread());
 
 		// If the stack frame was marked as garbage, NewStackRoot will be nullptr, but we still match against the object ptr hash.
 		if (!NewStackRoot)
diff --git a/Engine/Plugins/PCG/Source/PCG/Private/PCGComponent.cpp b/Engine/Plugins/PCG/Source/PCG/Private/PCGComponent.cpp
--- a/Engine/Plugins/PCG/Source/PCG/Private/PCGComponent.cpp
+++ b/Engine/Plugins/PCG/Source/PCG/Private/PCGComponent.cpp
@@ -2745,9 +2745,9 @@ void UPCGComponent::StoreInspectionData(const FPCGStack* InStack, const UPCGNode
 			TArray<FPCGStackFrame>& StackFrames = StackToSubgraphNode.GetStackFramesMutable();
 			StackFrames.Pop();
 
-			if (StackFrames.Last().Object.IsValid() && StackFrames.Last().Object->IsA<UPCGSubgraphNode>())
+			if (const UPCGSubgraphNode* Node = StackFrames.Last().GetObject_AnyThread<UPCGSubgraphNode>())
 			{
-				StorePinInspectionDataFromNode(&StackToSubgraphNode, Cast<const UPCGSubgraphNode>(StackFrames.Last().Object.Get())->GetInputPins(), InInputData, InspectionCache);
+				StorePinInspectionDataFromNode(&StackToSubgraphNode, Node->GetInputPins(), InInputData, InspectionCache);
 			}
 		}
 
diff --git a/Engine/Plugins/PCG/Source/PCG/Private/PCGElement.cpp b/Engine/Plugins/PCG/Source/PCG/Private/PCGElement.cpp
--- a/Engine/Plugins/PCG/Source/PCG/Private/PCGElement.cpp
+++ b/Engine/Plugins/PCG/Source/PCG/Private/PCGElement.cpp
@@ -383,6 +383,11 @@ bool IPCGElement::PrepareData(FPCGContext* Context) const
 
 bool IPCGElement::ConvertInputsIfNeeded(FPCGContext* Context) const
 {
+	if (Context->InputData.TaggedData.IsEmpty())
+	{
+		return true;
+	}
+
 	if (!SupportsGPUResidentData(Context))
 	{
 		bool bHasPendingReadbacks = false;
diff --git a/Engine/Plugins/PCG/Source/PCG/Private/Utils/PCGExtraCapture.cpp b/Engine/Plugins/PCG/Source/PCG/Private/Utils/PCGExtraCapture.cpp
--- a/Engine/Plugins/PCG/Source/PCG/Private/Utils/PCGExtraCapture.cpp
+++ b/Engine/Plugins/PCG/Source/PCG/Private/Utils/PCGExtraCapture.cpp
@@ -171,7 +171,7 @@ PCGUtils::FCallTreeInfo PCGUtils::FExtraCapture::CalculateCallTreeInfo(const UPC
 		{
 			const FPCGStackFrame& StackFrame = StackFrames[StackFrameIndex];
 
-			if (const UPCGNode* Node = Cast<UPCGNode>(StackFrame.Object.Get()))
+			if (const UPCGNode* Node = StackFrame.GetObject_GameThread<UPCGNode>())
 			{
 				NodePath.Add(Node);
 				NodePathLoop.Add(-1);
diff --git a/Engine/Plugins/PCG/Source/PCG/Private/Utils/PCGNodeVisualLogs.cpp b/Engine/Plugins/PCG/Source/PCG/Private/Utils/PCGNodeVisualLogs.cpp
--- a/Engine/Plugins/PCG/Source/PCG/Private/Utils/PCGNodeVisualLogs.cpp
+++ b/Engine/Plugins/PCG/Source/PCG/Private/Utils/PCGNodeVisualLogs.cpp
@@ -42,7 +42,7 @@ void FPCGNodeVisualLogs::Log(const FPCGStack& InPCGStack, ELogVerbosity::Type In
 	{
 		for (const FPCGStackFrame& Frame : InPCGStack.GetStackFrames())
 		{
-			if (const UPCGNode* Node = Cast<const UPCGNode>(Frame.Object.Get()))
+			if (const UPCGNode* Node = Frame.GetObject_GameThread<UPCGNode>())
 			{
 				Node->OnNodeChangedDelegate.Broadcast(const_cast<UPCGNode*>(Node), EPCGChangeType::Cosmetic);
 			}
@@ -236,6 +236,7 @@ void FPCGNodeVisualLogs::ClearLogs(const FPCGStack& InPCGStack)
 	TSet<const UPCGNode*> TouchedNodes;
 
 	{
+		FGCScopeGuard Guard;
 		FWriteScopeLock ScopedWriteLock(LogsLock);
 
 		TArray<FPCGStack> StacksToRemove;
@@ -255,7 +256,7 @@ void FPCGNodeVisualLogs::ClearLogs(const FPCGStack& InPCGStack)
 
 				for (const FPCGStackFrame& Frame : Entry.Key.GetStackFrames())
 				{
-					if (const UPCGNode* Node = Cast<const UPCGNode>(Frame.Object.Get()))
+					if (const UPCGNode* Node = Frame.GetObject_NoGuard<UPCGNode>())
 					{
 						TouchedNodes.Add(Node);
 					}
diff --git a/Engine/Plugins/PCG/Source/PCG/Public/Graph/PCGStackContext.h b/Engine/Plugins/PCG/Source/PCG/Public/Graph/PCGStackContext.h
--- a/Engine/Plugins/PCG/Source/PCG/Public/Graph/PCGStackContext.h
+++ b/Engine/Plugins/PCG/Source/PCG/Public/Graph/PCGStackContext.h
@@ -54,6 +54,26 @@ struct PCG_API FPCGStackFrame
 		ComputeHash();
 	}
 
+	template <class T = UObject>
+	const T* GetObject_AnyThread() const
+	{
+		FGCScopeGuard Guard;
+		return Cast<T>(Object.Get());
+	}
+
+	template <class T = UObject>
+	const T* GetObject_NoGuard() const
+	{
+		return Cast<T>(Object.Get());
+	}
+
+	template <class T = UObject>
+	const T* GetObject_GameThread() const
+	{
+		check(IsInGameThread());
+		return Cast<T>(Object.Get());
+	}
+
 	void SetLoopIndex(int32 InLoopIndex)
 	{
 		Object = nullptr;
@@ -64,6 +84,7 @@ struct PCG_API FPCGStackFrame
 	void PostSerialize(const FArchive& Ar);
 
 	// Stores object this frame refers to. Use SetObject to change this and properly update the hash.
+	// Prefer using a variation of GetObject to avoid direct access.
 	UPROPERTY()
 	TSoftObjectPtr<const UObject> Object;
 
diff --git a/Engine/Plugins/PCG/Source/PCGEditor/Private/Widgets/SPCGEditorGraphDebugObjectTree.cpp b/Engine/Plugins/PCG/Source/PCGEditor/Private/Widgets/SPCGEditorGraphDebugObjectTree.cpp
--- a/Engine/Plugins/PCG/Source/PCGEditor/Private/Widgets/SPCGEditorGraphDebugObjectTree.cpp
+++ b/Engine/Plugins/PCG/Source/PCGEditor/Private/Widgets/SPCGEditorGraphDebugObjectTree.cpp
@@ -681,7 +681,7 @@ bool SPCGEditorGraphDebugObjectTree::GetFirstStackFromSelection(const UPCGNode*
 		int NodeIndex = SelectedStack.GetStackFrames().Num();
 		while (NodeIndex < ItemStackFrames.Num())
 		{
-			ItemNode = Cast<const UPCGNode>(ItemStackFrames[NodeIndex].Object.Get());
+			ItemNode = ItemStackFrames[NodeIndex].GetObject_GameThread<UPCGNode>();
 			if (ItemNode)
 			{
 				break;
@@ -703,12 +703,12 @@ bool SPCGEditorGraphDebugObjectTree::GetFirstStackFromSelection(const UPCGNode*
 
 		if (ItemStackFrames.IsValidIndex(NodeIndex + 1))
 		{
-			ItemGraph = Cast<const UPCGGraph>(ItemStackFrames[NodeIndex + 1].Object.Get());
+			ItemGraph = ItemStackFrames[NodeIndex + 1].GetObject_GameThread<UPCGGraph>();
 		}
 
 		if (!ItemGraph && ItemStackFrames.IsValidIndex(NodeIndex + 2))
 		{
-			ItemGraph = Cast<const UPCGGraph>(ItemStackFrames[NodeIndex + 2].Object.Get());
+			ItemGraph = ItemStackFrames[NodeIndex + 2].GetObject_GameThread<UPCGGraph>();
 		}
 
 		if (ItemGraph && (!InGraph || ItemGraph == InGraph))
@@ -944,7 +944,7 @@ void SPCGEditorGraphDebugObjectTree::AddStacksToTree(const TArray<FPCGStack>& St
 			FPCGEditorGraphDebugObjectItemPtr CurrentItem;
 
 			// When we encounter a graph, we look at the frame index and/or preceding frames to determine the graph type.
-			if (const UPCGGraph* StackGraph = Cast<const UPCGGraph>(StackFrame.Object.Get()))
+			if (const UPCGGraph* StackGraph = StackFrame.GetObject_GameThread<UPCGGraph>())
 			{
 				const bool bIsDebuggable = (GraphBeingEdited == StackGraph);
 
@@ -966,14 +966,14 @@ void SPCGEditorGraphDebugObjectTree::AddStacksToTree(const TArray<FPCGStack>& St
 					}
 				}
 				// Previous stack was node, therefore static subgraph.
-				else if (const UPCGNode* SubgraphNode = Cast<const UPCGNode>(PreviousStackFrame.Object.Get()))
+				else if (const UPCGNode* SubgraphNode = PreviousStackFrame.GetObject_GameThread<UPCGNode>())
 				{
 					AddSubgraphOrLoopItemToStack(StackGraph, SubgraphNode, FrameIndex + 1, bIsDebuggable, bDisplayGrayedOut, EPCGEditorSubgraphNodeType::StaticSubgraph);
 				}
 				// Previous stack was loop index, therefore loop subgraph.
 				else if (FrameIndex >= 2 && (PreviousStackFrame.LoopIndex != INDEX_NONE))
 				{
-					const UPCGNode* LoopSubgraphNode = Cast<const UPCGNode>(StackFrames[FrameIndex - 2].Object.Get());
+					const UPCGNode* LoopSubgraphNode = StackFrames[FrameIndex - 2].GetObject_GameThread<UPCGNode>();
 					if (ensure(LoopSubgraphNode))
 					{
 						// Take the stack up to the looped subgraph node, add a item for the node + graph
@@ -996,7 +996,7 @@ void SPCGEditorGraphDebugObjectTree::AddStacksToTree(const TArray<FPCGStack>& St
 				// Previous stack was invalid node / node with a INDEX_NONE loop, therefore most likely a dynamic subgraph
 				else if (FrameIndex >= 2 && !PreviousStackFrame.IsValid())
 				{
-					const UPCGNode* DynamicSubgraphNode = Cast<const UPCGNode>(StackFrames[FrameIndex - 2].Object.Get());
+					const UPCGNode* DynamicSubgraphNode = StackFrames[FrameIndex - 2].GetObject_GameThread<UPCGNode>();
 					if (DynamicSubgraphNode)
 					{
 						AddSubgraphOrLoopItemToStack(StackGraph, DynamicSubgraphNode, FrameIndex + 1, bIsDebuggable, bDisplayGrayedOut, EPCGEditorSubgraphNodeType::DynamicSubgraph);
@@ -1324,7 +1324,7 @@ void SPCGEditorGraphDebugObjectTree::OnSelectionChanged(FPCGEditorGraphDebugObje
 		{
 			// If the last graph in the stack is the current graph, then we can just use that as the selected stack
 			TArray<FPCGStackFrame>& StackFrames = SelectedStack.GetStackFramesMutable();
-			while(!StackFrames.IsEmpty() && StackFrames.Last().Object.Get() != CurrentGraph)
+			while(!StackFrames.IsEmpty() && StackFrames.Last().GetObject_GameThread<UPCGGraph>() != CurrentGraph)
 			{
 				StackFrames.Pop();
 			}
@@ -1416,7 +1416,7 @@ void SPCGEditorGraphDebugObjectTree::ExpandAndSelectFirstLeafDebugObject(const F
 		const FPCGStack* ItemStack = Item ? Item->GetPCGStack() : nullptr;
 		if (ItemStack)
 		{
-			if(!ItemStack->GetStackFrames().IsEmpty() && ItemStack->GetStackFrames().Last().Object == CurrentGraph)
+			if(!ItemStack->GetStackFrames().IsEmpty() && ItemStack->GetStackFrames().Last().GetObject_GameThread<UPCGGraph>() == CurrentGraph)
 			{
 				ItemToSelect = Item;
 				break;
@@ -1481,10 +1481,10 @@ void SPCGEditorGraphDebugObjectTree::JumpToGraphInTree(const FPCGEditorGraphDebu
 		// Search the stack for the parent graph.
 		for (int i = Stack.GetStackFrames().Num() - 2; i > 0; --i) 
 		{
-			if (Stack.GetStackFrames()[i].Object != nullptr && Stack.GetStackFrames()[i].Object->IsA<UPCGGraph>())
+			if (const UPCGGraph* PCGGraph = Stack.GetStackFrames()[i].GetObject_GameThread<UPCGGraph>())
 			{
-				JumpToPCGGraph = const_cast<UPCGGraph*>(Cast<const UPCGGraph>(Stack.GetStackFrames()[i].Object.Get()));
-				JumpToPCGNode = const_cast<UPCGNode*>(Cast<const UPCGNode>(Stack.GetStackFrames()[i + 1].Object.Get()));
+				JumpToPCGGraph = const_cast<UPCGGraph*>(PCGGraph);
+				JumpToPCGNode = const_cast<UPCGNode*>(Stack.GetStackFrames()[i + 1].GetObject_GameThread<UPCGNode>());
 
 				// Cull remaining remaining frames so the stack selection into the parent graph editor works.
 				Stack.GetStackFramesMutable().SetNum(i);
@@ -1556,7 +1556,7 @@ void SPCGEditorGraphDebugObjectTree::FocusOnItem(const FPCGEditorGraphDebugObjec
 	for (int i = Stack.GetStackFrames().Num() - 1; i >= 0; --i)
 	{
 		const FPCGStackFrame& StackFrame = Stack.GetStackFrames()[i];
-		const UObject* Object = StackFrame.Object.Get();
+		const UObject* Object = StackFrame.GetObject_GameThread();
 
 		if (Object == CurrentGraph)
 		{
diff --git a/Engine/Plugins/PCG/Source/PCGEditor/Private/Widgets/SPCGEditorGraphFind.cpp b/Engine/Plugins/PCG/Source/PCGEditor/Private/Widgets/SPCGEditorGraphFind.cpp
--- a/Engine/Plugins/PCG/Source/PCGEditor/Private/Widgets/SPCGEditorGraphFind.cpp
+++ b/Engine/Plugins/PCG/Source/PCGEditor/Private/Widgets/SPCGEditorGraphFind.cpp
@@ -702,10 +702,10 @@ void SPCGEditorGraphFind::MatchTokens(const TArray<FString>& InTokens)
 				}
 
 				const FPCGStackFrame& SubgraphNodeFrame = Substack.GetStackFrames()[NodeFrameIndex];
-				const UPCGNode* SubgraphNode = Cast<UPCGNode>(SubgraphNodeFrame.Object.Get());
+				const UPCGNode* SubgraphNode = SubgraphNodeFrame.GetObject_GameThread<UPCGNode>();
 				SubgraphEditorNode = (SubgraphNode && CurrentEditorGraph) ? const_cast<UPCGEditorGraphNodeBase*>(CurrentEditorGraph->GetEditorNodeFromPCGNode(SubgraphNode)) : nullptr;
 				const FPCGStackFrame& SubgraphFrame = Substack.GetStackFrames()[SubgraphFrameIndex];
-				Subgraph = Cast<UPCGGraph>(const_cast<UObject*>(SubgraphFrame.Object.Get()));
+				Subgraph = const_cast<UPCGGraph*>(SubgraphFrame.GetObject_GameThread<UPCGGraph>());
 
 				return SubgraphEditorNode && Subgraph;
 			};
