[PCG] Fixes for bad memory behaviour and crash due to unsafe context pointer.

#rnx

diff --git a/Engine/Plugins/PCG/Source/PCG/Private/Graph/PCGGraphCompiler.cpp b/Engine/Plugins/PCG/Source/PCG/Private/Graph/PCGGraphCompiler.cpp
--- a/Engine/Plugins/PCG/Source/PCG/Private/Graph/PCGGraphCompiler.cpp
+++ b/Engine/Plugins/PCG/Source/PCG/Private/Graph/PCGGraphCompiler.cpp
@@ -520,8 +520,11 @@ void FPCGGraphCompiler::CreateGridLinkages(
 	for (FPCGTaskId TaskId = 0; TaskId < NumCompiledTasksBefore; ++TaskId)
 	{
 		const EPCGHiGenGrid GraphGenerationGrid = InOutTaskGenerationGrid[TaskId];
-		for (FPCGGraphTaskInput& TaskInput : InOutCompiledTasks[TaskId].Inputs)
+
+		for (int32 InputIndex = 0; InputIndex < InOutCompiledTasks[TaskId].Inputs.Num(); ++InputIndex)
 		{
+			const FPCGGraphTaskInput TaskInput = InOutCompiledTasks[TaskId].Inputs[InputIndex];
+
 			const UPCGPin* UpstreamPin = nullptr;
 			if (TaskInput.UpstreamPin.IsSet())
 			{
@@ -608,10 +611,10 @@ void FPCGGraphCompiler::CreateGridLinkages(
 				}
 
 				// Now splice in the new task - redirect the downstream task to grab its input from the link task.
-				TaskInput.TaskId = LinkTask.NodeId;
+				InOutCompiledTasks[TaskId].Inputs[InputIndex].TaskId = LinkTask.NodeId;
 
 				// Set the task input's UpstreamPin so that it has the correct type.
-				TaskInput.UpstreamPin = UpstreamPinProps;
+				InOutCompiledTasks[TaskId].Inputs[InputIndex].UpstreamPin = UpstreamPinProps;
 
 				// The link needs to execute at both FROM grid size (store) and TO grid size (retrieve).
 				InOutTaskGenerationGrid.Add(FromGrid | ToGrid);
diff --git a/Engine/Plugins/PCG/Source/PCG/Private/Graph/PCGGraphExecutor.cpp b/Engine/Plugins/PCG/Source/PCG/Private/Graph/PCGGraphExecutor.cpp
--- a/Engine/Plugins/PCG/Source/PCG/Private/Graph/PCGGraphExecutor.cpp
+++ b/Engine/Plugins/PCG/Source/PCG/Private/Graph/PCGGraphExecutor.cpp
@@ -2951,10 +2951,15 @@ namespace PCGGraphExecutor
 
 			// At this point we could not get to the data, so we'll try executing the graph if we did not do that already.
 
-			auto WakeUpLambda = [InContext, ComponentWithData]()
+			auto WakeUpLambda = [ComponentWithData, ContextHandle=InContext->GetOrCreateHandle()]()
 			{
-				PCGGraphExecutionLogging::LogGridLinkageTaskExecuteRetrieveWakeUp(InContext, ComponentWithData);
-				InContext->bIsPaused = false;
+				TSharedPtr<FPCGContextHandle> SharedHandle = ContextHandle.Pin();
+				if (FPCGContext* ContextPtr = SharedHandle ? SharedHandle->GetContext() : nullptr)
+				{
+					PCGGraphExecutionLogging::LogGridLinkageTaskExecuteRetrieveWakeUp(ContextPtr, ComponentWithData);
+					ContextPtr->bIsPaused = false;
+				}
+
 				return true;
 			};
 
