[PCG] Fixing Clang specific tests failures

* GraphParameters: Capturing by copy was not doing the right thing on Clang, and it was a mistake to have a copy there, so pass by reference
* Divide: Maths op was not checking correctly that the denominator was 0.
* Also hide ClampMin/Max as they are duplicates of Min/Max

#rnx

diff --git a/Engine/Plugins/PCG/Source/PCG/Private/Elements/Metadata/PCGMetadataMathsOpElement.cpp b/Engine/Plugins/PCG/Source/PCG/Private/Elements/Metadata/PCGMetadataMathsOpElement.cpp
--- a/Engine/Plugins/PCG/Source/PCG/Private/Elements/Metadata/PCGMetadataMathsOpElement.cpp
+++ b/Engine/Plugins/PCG/Source/PCG/Private/Elements/Metadata/PCGMetadataMathsOpElement.cpp
@@ -132,7 +132,10 @@ namespace PCGMetadataMathsSettings
 		case EPCGMetadataMathsOperation::Multiply:
 			return Value1 * Value2;
 		case EPCGMetadataMathsOperation::Divide:
-			return (Value2 != T{0}) ? (Value1 / Value2) : T{0}; // To mirror FMath
+		{
+			static const T ZeroValue = PCG::Private::MetadataTraits<T>::ZeroValue();
+			return !PCG::Private::MetadataTraits<T>::Equal(Value2, ZeroValue) ? (Value1 / Value2) : ZeroValue; // To mirror FMath
+		}
 		case EPCGMetadataMathsOperation::Max: // fall-through
 		case EPCGMetadataMathsOperation::ClampMin:
 			return PCGMetadataMaths::Max(Value1, Value2);
diff --git a/Engine/Plugins/PCG/Source/PCG/Private/Tests/PCGGraphParametersTest.cpp b/Engine/Plugins/PCG/Source/PCG/Private/Tests/PCGGraphParametersTest.cpp
--- a/Engine/Plugins/PCG/Source/PCG/Private/Tests/PCGGraphParametersTest.cpp
+++ b/Engine/Plugins/PCG/Source/PCG/Private/Tests/PCGGraphParametersTest.cpp
@@ -111,9 +111,9 @@ bool FPCGGraphAddSingleUserParameters::RunTest(const FString& Parameters)
 	const FName MyPropertyName = TEXT("MyProperty");
 
 	// Creating a new property MyProperty, which is a double
-	PCGTests::EmulateModifyingUserParameters(Graph, [Graph, MyPropertyName]()
+	PCGTests::EmulateModifyingUserParameters(Graph, [&Graph, MyPropertyName]()
 	{
-		FInstancedPropertyBag* UserParameters = const_cast<FInstancedPropertyBag*>(Graph->GetUserParametersStruct());
+		FInstancedPropertyBag* UserParameters = Graph->GetMutableUserParametersStruct_Unsafe();
 		UserParameters->AddProperty(MyPropertyName, EPropertyBagPropertyType::Double, nullptr);
 	});
 
@@ -145,9 +145,9 @@ bool FPCGGraphAddSingleUserParametersPropagates::RunTest(const FString& Paramete
 	const FName MyPropertyName = TEXT("MyProperty");
 
 	// Creating a new property MyProperty, which is a double
-	PCGTests::EmulateModifyingUserParameters(Graph, [Graph, MyPropertyName]()
+	PCGTests::EmulateModifyingUserParameters(Graph, [&Graph, MyPropertyName]()
 	{
-		FInstancedPropertyBag* UserParameters = const_cast<FInstancedPropertyBag*>(Graph->GetUserParametersStruct());
+		FInstancedPropertyBag* UserParameters = Graph->GetMutableUserParametersStruct_Unsafe();
 		UserParameters->AddProperty(MyPropertyName, EPropertyBagPropertyType::Double, nullptr);
 	});
 
@@ -179,16 +179,16 @@ bool FPCGGraphRemoveSingleUserParameters::RunTest(const FString& Parameters)
 	const FName MyPropertyName = TEXT("MyProperty");
 
 	// Creating a new property MyProperty, which is a double
-	PCGTests::EmulateModifyingUserParameters(Graph, [Graph, MyPropertyName]()
+	PCGTests::EmulateModifyingUserParameters(Graph, [&Graph, MyPropertyName]()
 	{
-		FInstancedPropertyBag* UserParameters = const_cast<FInstancedPropertyBag*>(Graph->GetUserParametersStruct());
+		FInstancedPropertyBag* UserParameters = Graph->GetMutableUserParametersStruct_Unsafe();
 		UserParameters->AddProperty(MyPropertyName, EPropertyBagPropertyType::Double, nullptr);
 	});
 
 	// Remove that same property
-	PCGTests::EmulateModifyingUserParameters(Graph, [Graph, MyPropertyName]()
+	PCGTests::EmulateModifyingUserParameters(Graph, [&Graph, MyPropertyName]()
 	{
-		FInstancedPropertyBag* UserParameters = const_cast<FInstancedPropertyBag*>(Graph->GetUserParametersStruct());
+		FInstancedPropertyBag* UserParameters = Graph->GetMutableUserParametersStruct_Unsafe();
 		UserParameters->RemovePropertyByName(MyPropertyName);
 	});
 
@@ -215,16 +215,16 @@ bool FPCGGraphRemoveSingleUserParametersPropagates::RunTest(const FString& Param
 	const FName MyPropertyName = TEXT("MyProperty");
 
 	// Creating a new property MyProperty, which is a double
-	PCGTests::EmulateModifyingUserParameters(Graph, [Graph, MyPropertyName]()
+	PCGTests::EmulateModifyingUserParameters(Graph, [&Graph, MyPropertyName]()
 	{
-		FInstancedPropertyBag* UserParameters = const_cast<FInstancedPropertyBag*>(Graph->GetUserParametersStruct());
+		FInstancedPropertyBag* UserParameters = Graph->GetMutableUserParametersStruct_Unsafe();
 		UserParameters->AddProperty(MyPropertyName, EPropertyBagPropertyType::Double, nullptr);
 	});
 
 	// Remove that same property
-	PCGTests::EmulateModifyingUserParameters(Graph, [Graph, MyPropertyName]()
+	PCGTests::EmulateModifyingUserParameters(Graph, [&Graph, MyPropertyName]()
 	{
-		FInstancedPropertyBag* UserParameters = const_cast<FInstancedPropertyBag*>(Graph->GetUserParametersStruct());
+		FInstancedPropertyBag* UserParameters = Graph->GetMutableUserParametersStruct_Unsafe();
 		UserParameters->RemovePropertyByName(MyPropertyName);
 	});
 
@@ -260,16 +260,16 @@ bool FPCGGraphRenameUserParametersPropagates::RunTest(const FString& Parameters)
 	const FName MyNewPropertyName = TEXT("MyNewProperty");
 
 	// Creating a new property MyProperty, which is a double
-	PCGTests::EmulateModifyingUserParameters(Graph, [Graph, MyPropertyName]()
+	PCGTests::EmulateModifyingUserParameters(Graph, [&Graph, MyPropertyName]()
 	{
-		FInstancedPropertyBag* UserParameters = const_cast<FInstancedPropertyBag*>(Graph->GetUserParametersStruct());
+		FInstancedPropertyBag* UserParameters = Graph->GetMutableUserParametersStruct_Unsafe();
 		UserParameters->AddProperty(MyPropertyName, EPropertyBagPropertyType::Double, nullptr);
 	});
 
 	// Rename it to MyNewProperty, while keeping the type Double
-	PCGTests::EmulateModifyingUserParameters(Graph, [Graph, MyNewPropertyName]()
+	PCGTests::EmulateModifyingUserParameters(Graph, [&Graph, MyNewPropertyName]()
 	{
-		FInstancedPropertyBag* UserParameters = const_cast<FInstancedPropertyBag*>(Graph->GetUserParametersStruct());
+		FInstancedPropertyBag* UserParameters = Graph->GetMutableUserParametersStruct_Unsafe();
 
 		const UPropertyBag* NewPropertyBag = UPropertyBag::GetOrCreateFromDescs({ FPropertyBagPropertyDesc{MyNewPropertyName, EPropertyBagPropertyType::Double, nullptr} });
 		UserParameters->MigrateToNewBagStruct(NewPropertyBag);
@@ -314,16 +314,16 @@ bool FPCGGraphChangeTypeUserParametersPropagates::RunTest(const FString& Paramet
 	const FName MyPropertyName = TEXT("MyProperty");
 
 	// Creating a new property MyProperty, which is a double
-	PCGTests::EmulateModifyingUserParameters(Graph, [Graph, MyPropertyName]()
+	PCGTests::EmulateModifyingUserParameters(Graph, [&Graph, MyPropertyName]()
 	{
-		FInstancedPropertyBag* UserParameters = const_cast<FInstancedPropertyBag*>(Graph->GetUserParametersStruct());
+		FInstancedPropertyBag* UserParameters = Graph->GetMutableUserParametersStruct_Unsafe();
 		UserParameters->AddProperty(MyPropertyName, EPropertyBagPropertyType::Double, nullptr);
 	});
 
 	// Keeping the name MyProperty but changing the type to Int64
-	PCGTests::EmulateModifyingUserParameters(Graph, [Graph, MyPropertyName]()
+	PCGTests::EmulateModifyingUserParameters(Graph, [&Graph, MyPropertyName]()
 	{
-		FInstancedPropertyBag* UserParameters = const_cast<FInstancedPropertyBag*>(Graph->GetUserParametersStruct());
+		FInstancedPropertyBag* UserParameters = Graph->GetMutableUserParametersStruct_Unsafe();
 
 		const UPropertyBag* NewPropertyBag = UPropertyBag::GetOrCreateFromDescs({ FPropertyBagPropertyDesc{MyPropertyName, EPropertyBagPropertyType::Int64, nullptr} });
 		UserParameters->MigrateToNewBagStruct(NewPropertyBag);
@@ -363,16 +363,16 @@ bool FPCGGraphRemovingGraphCleansUserParameters::RunTest(const FString& Paramete
 	const FName MyPropertyName2 = TEXT("MyProperty2");
 
 	// Creating a new property MyProperty, which is a double
-	PCGTests::EmulateModifyingUserParameters(Graph, [Graph, MyPropertyName]()
+	PCGTests::EmulateModifyingUserParameters(Graph, [&Graph, MyPropertyName]()
 	{
-		FInstancedPropertyBag* UserParameters = const_cast<FInstancedPropertyBag*>(Graph->GetUserParametersStruct());
+		FInstancedPropertyBag* UserParameters = Graph->GetMutableUserParametersStruct_Unsafe();
 		UserParameters->AddProperty(MyPropertyName, EPropertyBagPropertyType::Double, nullptr);
 	});
 
 	// Creating a new property MyProperty2, which is a int64
-	PCGTests::EmulateModifyingUserParameters(Graph, [Graph, MyPropertyName2]()
+	PCGTests::EmulateModifyingUserParameters(Graph, [&Graph, MyPropertyName2]()
 	{
-		FInstancedPropertyBag* UserParameters = const_cast<FInstancedPropertyBag*>(Graph->GetUserParametersStruct());
+		FInstancedPropertyBag* UserParameters = Graph->GetMutableUserParametersStruct_Unsafe();
 		UserParameters->AddProperty(MyPropertyName2, EPropertyBagPropertyType::Int64, nullptr);
 	});
 
@@ -395,16 +395,16 @@ bool FPCGGraphSetGraphUpdatesUserParameters::RunTest(const FString& Parameters)
 	const FName MyPropertyName2 = TEXT("MyProperty2");
 
 	// Creating a new property MyProperty, which is a double
-	PCGTests::EmulateModifyingUserParameters(Graph, [Graph, MyPropertyName]()
+	PCGTests::EmulateModifyingUserParameters(Graph, [&Graph, MyPropertyName]()
 	{
-		FInstancedPropertyBag* UserParameters = const_cast<FInstancedPropertyBag*>(Graph->GetUserParametersStruct());
+		FInstancedPropertyBag* UserParameters = Graph->GetMutableUserParametersStruct_Unsafe();
 		UserParameters->AddProperty(MyPropertyName, EPropertyBagPropertyType::Double, nullptr);
 	});
 
 	// Creating a new property MyProperty2, which is a int64
-	PCGTests::EmulateModifyingUserParameters(Graph, [Graph, MyPropertyName2]()
+	PCGTests::EmulateModifyingUserParameters(Graph, [&Graph, MyPropertyName2]()
 	{
-		FInstancedPropertyBag* UserParameters = const_cast<FInstancedPropertyBag*>(Graph->GetUserParametersStruct());
+		FInstancedPropertyBag* UserParameters = Graph->GetMutableUserParametersStruct_Unsafe();
 		UserParameters->AddProperty(MyPropertyName2, EPropertyBagPropertyType::Int64, nullptr);
 	});
 
diff --git a/Engine/Plugins/PCG/Source/PCG/Public/Elements/Metadata/PCGMetadataMathsOpElement.h b/Engine/Plugins/PCG/Source/PCG/Public/Elements/Metadata/PCGMetadataMathsOpElement.h
--- a/Engine/Plugins/PCG/Source/PCG/Public/Elements/Metadata/PCGMetadataMathsOpElement.h
+++ b/Engine/Plugins/PCG/Source/PCG/Public/Elements/Metadata/PCGMetadataMathsOpElement.h
@@ -33,8 +33,9 @@ enum class EPCGMetadataMathsOperation : uint16
 	Max,
 	Min,
 	Pow,
-	ClampMin,
-	ClampMax,
+	// ClampMin is just Max and ClampMax is just Min, so hide them.
+	ClampMin UMETA(Hidden),
+	ClampMax UMETA(Hidden),
 	Modulo,
 	Set,
 
