[PCG] Robustness fixes for Generate Grass Maps node

#rnx

diff --git a/Engine/Plugins/PCG/Source/PCG/Private/Elements/PCGGenerateGrassMaps.cpp b/Engine/Plugins/PCG/Source/PCG/Private/Elements/PCGGenerateGrassMaps.cpp
--- a/Engine/Plugins/PCG/Source/PCG/Private/Elements/PCGGenerateGrassMaps.cpp
+++ b/Engine/Plugins/PCG/Source/PCG/Private/Elements/PCGGenerateGrassMaps.cpp
@@ -215,7 +215,7 @@ bool FPCGGenerateGrassMapsElement::ExecuteInternal(FPCGContext* InContext) const
 		{
 			ALandscapeProxy* LandscapeProxy = LandscapeProxyPtr.Get();
 
-			if (!LandscapeProxy)
+			if (!LandscapeProxy || !LandscapeProxy->LandscapeMaterial)
 			{
 				continue;
 			}
@@ -224,14 +224,11 @@ bool FPCGGenerateGrassMapsElement::ExecuteInternal(FPCGContext* InContext) const
 			{
 				Context->LandscapeProxy = LandscapeProxy;
 			}
-			else
+			// TODO: In order to support multiple-material we need to do a dispatch of one landscape grass exporter per LS proxy.
+			else if (Context->LandscapeProxy->LandscapeMaterial != LandscapeProxy->LandscapeMaterial)
 			{
-				// TODO: In order to support multiple-material we need to do a dispatch of one landscape grass exporter per LS proxy.
-				if (Context->LandscapeProxy->LandscapeMaterial != LandscapeProxy->LandscapeMaterial)
-				{
-					PCGLog::LogErrorOnGraph(LOCTEXT("LandscapeMaterialMismatch", "Encountered multiple landscape materials on the same landscape. Only one landscape material per landscape is supported at this time."), Context);
-					return true;
-				}
+				PCGLog::LogErrorOnGraph(LOCTEXT("LandscapeMaterialMismatch", "Encountered multiple landscape materials on the same landscape. Only one landscape material per landscape is supported at this time."), Context);
+				return true;
 			}
 
 			for (ULandscapeComponent* LandscapeComponent : LandscapeProxy->LandscapeComponents)
@@ -431,7 +428,7 @@ bool FPCGGenerateGrassMapsElement::ExecuteInternal(FPCGContext* InContext) const
 		
 		TWeakPtr<FPCGContextHandle> ContextHandle = Context->GetOrCreateHandle();
 
-		ENQUEUE_RENDER_COMMAND(GenerateGrassMaps)([ContextHandle, LandscapeComponentExtent=Context->LandscapeComponentExtent, LandscapeTileCoords=MoveTemp(LandscapeTileCoords)](FRHICommandListImmediate& RHICmdList)
+		ENQUEUE_RENDER_COMMAND(GenerateGrassMaps)([ContextHandle, LandscapeComponentExtent=Context->LandscapeComponentExtent, LandscapeTileCoords=MoveTemp(LandscapeTileCoords), ComponentSizeQuads=LandscapeProxy->ComponentSizeQuads](FRHICommandListImmediate& RHICmdList)
 		{
 			FPCGContext::FSharedContext<FPCGGenerateGrassMapsContext> SharedContext(ContextHandle);
 			FPCGGenerateGrassMapsContext* Context = SharedContext.Get();
@@ -459,7 +456,7 @@ bool FPCGGenerateGrassMapsElement::ExecuteInternal(FPCGContext* InContext) const
 			const uint32 NumTilesY = FMath::RoundToInt(GrassMapExtent.Y / LandscapeComponentExtent);
 
 			// Each corner of a quad in the landscape corresponds to one texel in the grass map.
-			const uint32 LandscapeComponentResolution = Context->LandscapeProxy->ComponentSizeQuads + 1;
+			const uint32 LandscapeComponentResolution = ComponentSizeQuads + 1;
 			const FIntPoint GrassMapResolution(LandscapeComponentResolution * NumTilesX, LandscapeComponentResolution * NumTilesY);
 			const uint32 NumGrassTypes = Context->NumGrassTypes;
 
@@ -545,7 +542,10 @@ bool FPCGGenerateGrassMapsElement::ExecuteInternal(FPCGContext* InContext) const
 	const FTransform GrassMapTransform(FQuat::Identity, Context->GrassMapBounds.GetCenter(), Context->GrassMapBounds.GetExtent());
 	bool bAllTexturesInitialized = true;
 
-	check(Context->GrassMapHandle.IsValid());
+	if (!Context->GrassMapHandle.IsValid())
+	{
+		return true;
+	}
 
 	for (int DataIndex = 0; DataIndex < Context->SelectedGrassTypes.Num(); ++DataIndex)
 	{
