[PCG] Make sure to clear handle context ptr if Context is deleted with a valid one (to avoid double delete)
- ReadyTasks could get deleted when cancelled
- IndirectionElement has an internal Context

#rnx

diff --git a/Engine/Plugins/PCG/Source/PCG/Private/Elements/PCGIndirectionElement.cpp b/Engine/Plugins/PCG/Source/PCG/Private/Elements/PCGIndirectionElement.cpp
--- a/Engine/Plugins/PCG/Source/PCG/Private/Elements/PCGIndirectionElement.cpp
+++ b/Engine/Plugins/PCG/Source/PCG/Private/Elements/PCGIndirectionElement.cpp
@@ -203,7 +203,7 @@ bool FPCGIndirectionElement::CanExecuteOnlyOnMainThread(FPCGContext* InContext)
 
 FPCGIndirectionContext::~FPCGIndirectionContext()
 {
-	delete InnerContext;
+	FPCGContext::Release(InnerContext);
 	InnerContext = nullptr;
 }
 
diff --git a/Engine/Plugins/PCG/Source/PCG/Private/Graph/PCGGraphExecutor.cpp b/Engine/Plugins/PCG/Source/PCG/Private/Graph/PCGGraphExecutor.cpp
--- a/Engine/Plugins/PCG/Source/PCG/Private/Graph/PCGGraphExecutor.cpp
+++ b/Engine/Plugins/PCG/Source/PCG/Private/Graph/PCGGraphExecutor.cpp
@@ -570,7 +570,7 @@ TSet<UPCGComponent*> FPCGGraphExecutor::Cancel(TFunctionRef<bool(TWeakObjectPtr<
 				FPCGTaskId CancelledTaskId = CancelledReadyTask.NodeId;
 				RemoveTaskFromInputSuccessors(CancelledTaskId, CancelledReadyTask.Inputs);
 
-				delete CancelledReadyTask.Context;
+				FPCGContext::Release(CancelledReadyTask.Context);
 				bStableCancellationSet &= !CancelNextTasks(CancelledTaskId, CancelledComponents);
 			}
 		}
@@ -1735,7 +1735,7 @@ void FPCGGraphExecutor::ClearAllTasks()
 	// Make sure we don't leak preallocated contexts
 	for (FPCGGraphTask& ReadyTask : ReadyTasks)
 	{
-		delete ReadyTask.Context;
+		FPCGContext::Release(ReadyTask.Context);
 	}
 
 	ReadyTasks.Reset();
diff --git a/Engine/Plugins/PCG/Source/PCG/Private/Graph/PCGGraphTask.cpp b/Engine/Plugins/PCG/Source/PCG/Private/Graph/PCGGraphTask.cpp
--- a/Engine/Plugins/PCG/Source/PCG/Private/Graph/PCGGraphTask.cpp
+++ b/Engine/Plugins/PCG/Source/PCG/Private/Graph/PCGGraphTask.cpp
@@ -117,16 +117,7 @@ FPCGGraphActiveTask::~FPCGGraphActiveTask()
 		{
 			PCGGraphExecutor::ClearAsyncFlags(ContextPtr->AsyncObjects);
 
-			if (!ContextPtr->Handle)
-			{
-				delete ContextPtr;
-			}
-			else
-			{
-				// Release the Handle SharedPtr (at this point only captured and pinned Handles can keep the Context alive)
-				// The Context will get deleted when the handle gets released
-				ContextPtr->Handle = nullptr;
-			}
+			FPCGContext::Release(ContextPtr);
 		});	
 	}
 }
diff --git a/Engine/Plugins/PCG/Source/PCG/Private/PCGContext.cpp b/Engine/Plugins/PCG/Source/PCG/Private/PCGContext.cpp
--- a/Engine/Plugins/PCG/Source/PCG/Private/PCGContext.cpp
+++ b/Engine/Plugins/PCG/Source/PCG/Private/PCGContext.cpp
@@ -44,9 +44,23 @@ namespace PCGContextHelpers
 	}
 }
 
+FPCGContext::~FPCGContext()
+{
+	// This should not happen, might be a missing call to FPCGContext::Release
+	// Clear Handle context pointer so that we don't double delete the FPCGContext
+	ensureMsgf(!Handle, TEXT("FPCGContext deleted without first setting the Handle to nullptr, probably a delete that should be replaced by a FPCGContext::Release call"));
+	if (Handle)
+	{
+		Handle->Context = nullptr;
+	}
+}
+
 FPCGContextHandle::~FPCGContextHandle()
 {
-	delete Context;
+	if (Context)
+	{
+		delete Context;
+	}
 }
 
 UPCGBasePointData* FPCGContext::NewPointData_AnyThread(FPCGContext* Context)
diff --git a/Engine/Plugins/PCG/Source/PCG/Public/PCGContext.h b/Engine/Plugins/PCG/Source/PCG/Public/PCGContext.h
--- a/Engine/Plugins/PCG/Source/PCG/Public/PCGContext.h
+++ b/Engine/Plugins/PCG/Source/PCG/Public/PCGContext.h
@@ -62,6 +62,7 @@ public:
 
 	FPCGContext* GetContext() { return Context; }
 private:
+	friend struct FPCGContext;
 	FPCGContext* Context = nullptr;
 };
 
@@ -74,7 +75,7 @@ struct PCG_API FPCGContext
 	friend struct FPCGGraphActiveTask;
 	friend struct FPCGGraphTask;
 
-	virtual ~FPCGContext() = default;
+	virtual ~FPCGContext();
 
 	FPCGDataCollection InputData;
 	FPCGDataCollection OutputData;
@@ -185,6 +186,25 @@ struct PCG_API FPCGContext
 		return Handle.ToWeakPtr();
 	}
 
+	static void Release(FPCGContext* InContext)
+	{
+		if(!InContext)
+		{
+			return;
+		}
+		
+		if (!InContext->Handle)
+		{
+			delete InContext;
+		}
+		else
+		{
+			// Release the Handle SharedPtr (at this point only captured and pinned Handles can keep the Context alive)
+			// The Context will get deleted when the handle gets released
+			InContext->Handle = nullptr;
+		}
+	}
+
 	template <typename ContextType>
 	UE_DEPRECATED(5.6, "Use FSharedContext instead")
 	static ContextType* GetContextFromHandle(TWeakPtr<FPCGContextHandle> WeakHandle)
