[PCG] Uniformize metadata initialization

* Is important to clarify what it means to initialize metadata from another data
* Add a concept of mapping from domain to the other, so we know where to create the new attributes
* Composite data will have a clearer way to add attributes

#rnx

diff --git a/Engine/Plugins/PCG/Source/PCG/Private/Data/PCGDifferenceData.cpp b/Engine/Plugins/PCG/Source/PCG/Private/Data/PCGDifferenceData.cpp
--- a/Engine/Plugins/PCG/Source/PCG/Private/Data/PCGDifferenceData.cpp
+++ b/Engine/Plugins/PCG/Source/PCG/Private/Data/PCGDifferenceData.cpp
@@ -43,9 +43,6 @@ void UPCGDifferenceData::Initialize(const UPCGSpatialData* InData)
 #if WITH_EDITOR
 	RawPointerSource = Source;
 #endif
-
-	check(Metadata);
-	Metadata->Initialize(Source->Metadata);
 }
 
 void UPCGDifferenceData::K2_AddDifference(const UPCGSpatialData* InDifference)
@@ -253,29 +250,28 @@ const UPCGBasePointData* UPCGDifferenceData::CreateBasePointData(FPCGContext* Co
 		return SourcePointData;
 	}
 
-	if (!GetDifference())
+	const UPCGSpatialData* DifferenceData = GetDifference();
+
+	if (!DifferenceData)
 	{
 		UE_LOG(LogPCG, Verbose, TEXT("Difference is trivial"));
 		return SourcePointData;
 	}
 		
 	UPCGBasePointData* TargetPointData = FPCGContext::NewObject_AnyThread<UPCGBasePointData>(Context, GetTransientPackage(), PointDataClass);
+	const UPCGMetadata* SourceMetadata = SourcePointData->ConstMetadata();
 
-	FPCGInitializeFromDataParams InitializeFromDataParams(this);
+	FPCGInitializeFromDataParams InitializeFromDataParams(SourcePointData);
 	InitializeFromDataParams.bInheritSpatialData = false;
-
-	const UPCGMetadata* SourceMetadata = SourcePointData->Metadata;
-	InitializeFromDataParams.MetadataParentOverride = SourceMetadata;
 	TargetPointData->InitializeFromDataWithParams(InitializeFromDataParams);
-	
+
 	UPCGMetadata* OutMetadata = TargetPointData->Metadata;
-	
-	const UPCGMetadata* DifferenceMetadata = GetDifference()->Metadata;
+
 	UPCGMetadata* TempDiffMetadata = nullptr;
-	if (bDiffMetadata && OutMetadata && DifferenceMetadata)
+	if (bDiffMetadata && OutMetadata)
 	{
 		TempDiffMetadata = FPCGContext::NewObject_AnyThread<UPCGMetadata>(Context);
-		TempDiffMetadata->Initialize(DifferenceMetadata);
+		DifferenceData->InitializeOtherMetadata(FPCGInitializeFromDataParams{DifferenceData}, TempDiffMetadata);
 	}
 
 	constexpr int ChunkSize = FPCGSpatialDataProcessing::DefaultSamplePointsChunkSize;
@@ -393,3 +389,20 @@ UPCGSpatialData* UPCGDifferenceData::CopyInternal(FPCGContext* Context) const
 
 	return NewDifferenceData;
 }
+
+void UPCGDifferenceData::InitializeOtherMetadataInternal(const FPCGInitializeFromDataParams& InParams, UPCGMetadata* MetadataToInitialize) const
+{
+	check(InParams.bInheritMetadata);
+	check(MetadataToInitialize);
+
+	// In the case of the difference, we initialize the data from the source, not the difference itself. The difference will add its own
+	// attributes afterward.
+	if (Source)
+	{
+		FPCGInitializeFromDataParams CopyParams = InParams;
+		CopyParams.Source = Source;
+		Source->InitializeOtherMetadata(InParams, MetadataToInitialize);
+	}
+
+	MetadataToInitialize->AddAttributes(Metadata);
+}
diff --git a/Engine/Plugins/PCG/Source/PCG/Private/Data/PCGIntersectionData.cpp b/Engine/Plugins/PCG/Source/PCG/Private/Data/PCGIntersectionData.cpp
--- a/Engine/Plugins/PCG/Source/PCG/Private/Data/PCGIntersectionData.cpp
+++ b/Engine/Plugins/PCG/Source/PCG/Private/Data/PCGIntersectionData.cpp
@@ -43,19 +43,6 @@ void UPCGIntersectionData::Initialize(const UPCGSpatialData* InA, const UPCGSpat
 
 	CachedBounds = PCGHelpers::OverlapBounds(GetA()->GetBounds(), GetB()->GetBounds());
 	CachedStrictBounds = PCGHelpers::OverlapBounds(GetA()->GetStrictBounds(), GetB()->GetStrictBounds());
-
-	check(Metadata);
-	// Note: this should behave the same way as the ToPointData
-	if (GetA()->GetDimension() <= GetB()->GetDimension())
-	{
-		Metadata->Initialize(GetA()->Metadata);
-		Metadata->AddAttributes(GetB()->Metadata);
-	}
-	else
-	{
-		Metadata->Initialize(GetB()->Metadata);
-		Metadata->AddAttributes(GetA()->Metadata);
-	}
 }
 
 #if WITH_EDITOR
@@ -226,26 +213,28 @@ UPCGBasePointData* UPCGIntersectionData::CreateAndFilterPointData(FPCGContext* C
 		UE_LOG(LogPCG, Error, TEXT("Intersection unable to get source points"));
 		return nullptr;
 	}
+	
+	const UPCGMetadata* SourceMetadata = SourcePointData->ConstMetadata();
+	const UPCGMetadata* ThisMetadata = ConstMetadata();
 
 	UPCGBasePointData* OutputData = FPCGContext::NewObject_AnyThread<UPCGBasePointData>(Context, GetTransientPackage(), PointDataClass);
+	UPCGMetadata* OutputMetadata = OutputData->MutableMetadata();
 
-	FPCGInitializeFromDataParams InitializeFromDataParams(this);
+	FPCGInitializeFromDataParams InitializeFromDataParams(SourcePointData);
 	InitializeFromDataParams.bInheritSpatialData = false;
-
-	const UPCGMetadata* SourceMetadata = SourcePointData->Metadata;
-	InitializeFromDataParams.MetadataParentOverride = SourceMetadata;
 	OutputData->InitializeFromDataWithParams(InitializeFromDataParams);
-	
-	OutputData->Metadata->AddAttributes(Y->Metadata);
 
 	UPCGMetadata* TempYMetadata = nullptr;
-	if (Y->Metadata)
+	if (Y && Y->ConstMetadata())
 	{
-		TempYMetadata = Y->Metadata ? FPCGContext::NewObject_AnyThread<UPCGMetadata>(Context) : nullptr;
-		TempYMetadata->Initialize(Y->Metadata);
+		TempYMetadata = FPCGContext::NewObject_AnyThread<UPCGMetadata>(Context);
+		Y->InitializeOtherMetadata(FPCGInitializeFromDataParams{Y}, TempYMetadata);
+		OutputMetadata->AddAttributes(TempYMetadata);
 	}
 
-	const bool bPointDataHasCommonAttributes = (SourceMetadata && Y->Metadata && SourceMetadata->HasCommonAttributes(Y->Metadata));
+	OutputMetadata->AddAttributes(ThisMetadata);
+
+	const bool bPointDataHasCommonAttributes = (SourceMetadata && TempYMetadata && SourceMetadata->HasCommonAttributes(TempYMetadata));
 
 	constexpr int ChunkSize = FPCGSpatialDataProcessing::DefaultSamplePointsChunkSize;
 
@@ -299,9 +288,10 @@ UPCGBasePointData* UPCGIntersectionData::CreateAndFilterPointData(FPCGContext* C
 				
 				// TODO: create an array-based MergePointsAttributeSubset..
 				// If either the point from Y has metadata or the merge would be a non-trivial value, then perform the full merge
-				if (TargetPointData->Metadata && (bPointDataHasCommonAttributes || PointFromY.MetadataEntry != PCGInvalidEntryKey))
+				UPCGMetadata* TargetMetadata = TargetPointData->MutableMetadata();
+				if (TargetMetadata && (bPointDataHasCommonAttributes || PointFromY.MetadataEntry != PCGInvalidEntryKey))
 				{
-					TargetPointData->Metadata->MergeAttributesSubset(SourceRanges.MetadataEntryRange[ReadIndex], SourceMetadata, SourceMetadata, PointFromY.MetadataEntry, TempYMetadata, TempYMetadata, TargetRanges.MetadataEntryRange[WriteIndex], EPCGMetadataOp::Min);
+					TargetMetadata->MergeAttributesSubset(SourceRanges.MetadataEntryRange[ReadIndex], SourceMetadata, SourceMetadata, PointFromY.MetadataEntry, TempYMetadata, TempYMetadata, TargetRanges.MetadataEntryRange[WriteIndex], EPCGMetadataOp::Min);
 				}
 
 				++NumWritten;
@@ -352,3 +342,29 @@ UPCGSpatialData* UPCGIntersectionData::CopyInternal(FPCGContext* Context) const
 
 	return NewIntersectionData;
 }
+
+void UPCGIntersectionData::InitializeOtherMetadataInternal(const FPCGInitializeFromDataParams& InParams, UPCGMetadata* MetadataToInitialize) const
+{
+	check(InParams.bInheritMetadata);
+	check(MetadataToInitialize);
+	check(GetA());
+	check(GetB());
+
+	// In the case of the intersection, we initialize the data from either A or B, depending on the dimension, not the intersection itself. The intersection will add its own
+	// attributes afterward.
+	FPCGInitializeFromDataParams ParamsCopy = InParams;
+	if (GetA()->GetDimension() <= GetB()->GetDimension())
+	{
+		ParamsCopy.Source = GetA();
+		GetA()->InitializeOtherMetadata(ParamsCopy, MetadataToInitialize);
+		MetadataToInitialize->AddAttributes(GetB()->Metadata);
+	}
+	else
+	{
+		ParamsCopy.Source = GetB();
+		GetB()->InitializeOtherMetadata(ParamsCopy, MetadataToInitialize);
+		MetadataToInitialize->AddAttributes(GetA()->Metadata);
+	}
+
+	MetadataToInitialize->AddAttributes(Metadata);
+}
diff --git a/Engine/Plugins/PCG/Source/PCG/Private/Data/PCGProjectionData.cpp b/Engine/Plugins/PCG/Source/PCG/Private/Data/PCGProjectionData.cpp
--- a/Engine/Plugins/PCG/Source/PCG/Private/Data/PCGProjectionData.cpp
+++ b/Engine/Plugins/PCG/Source/PCG/Private/Data/PCGProjectionData.cpp
@@ -276,7 +276,6 @@ const UPCGBasePointData* UPCGProjectionData::CreateBasePointData(FPCGContext* Co
 	
 	// Copy metadata attributes from source point including values
 	FPCGInitializeFromDataParams InitializeFromDataParams(this);
-	InitializeFromDataParams.MetadataParentOverride = SourceMetadata;
 	InitializeFromDataParams.bInheritSpatialData = false;
 	PointData->InitializeFromDataWithParams(InitializeFromDataParams);
 	UPCGMetadata* OutMetadata = PointData->Metadata;
@@ -452,3 +451,18 @@ UPCGSpatialData* UPCGProjectionData::CopyInternal(FPCGContext* Context) const
 
 	return NewProjectionData;
 }
+
+void UPCGProjectionData::InitializeOtherMetadataInternal(const FPCGInitializeFromDataParams& InParams, UPCGMetadata* MetadataToInitialize) const
+{
+	check(InParams.bInheritMetadata);
+	check(MetadataToInitialize);
+	
+	if (Source)
+	{
+		FPCGInitializeFromDataParams ParamsCopy = InParams;
+		ParamsCopy.Source = Source;
+		Source->InitializeOtherMetadata(ParamsCopy, MetadataToInitialize);
+	}
+
+	MetadataToInitialize->AddAttributes(Metadata);
+}
\ No newline at end of file
diff --git a/Engine/Plugins/PCG/Source/PCG/Private/Data/PCGSpatialData.cpp b/Engine/Plugins/PCG/Source/PCG/Private/Data/PCGSpatialData.cpp
--- a/Engine/Plugins/PCG/Source/PCG/Private/Data/PCGSpatialData.cpp
+++ b/Engine/Plugins/PCG/Source/PCG/Private/Data/PCGSpatialData.cpp
@@ -312,22 +312,20 @@ UPCGMetadata* UPCGSpatialData::CreateEmptyMetadata()
 	return Metadata;
 }
 
-void UPCGSpatialData::InitializeMetadataInternal(const FPCGInitializeFromDataParams& InParams)
+void UPCGSpatialData::InitializeOtherMetadata(const FPCGInitializeFromDataParams& InParams, UPCGMetadata* MetadataToInitialize) const
 {
-	if (!Metadata)
+	if (InParams.bInheritMetadata)
 	{
-		Metadata = NewObject<UPCGMetadata>(this);
+		InitializeOtherMetadataInternal(InParams, MetadataToInitialize);
 	}
+}
 
-	if (!InParams.bInheritMetadata || InParams.MetadataParentOverride || InParams.Source)
-	{
-		const UPCGMetadata* ParentMetadata = InParams.bInheritMetadata ? (InParams.MetadataParentOverride ? InParams.MetadataParentOverride : (InParams.Source ? InParams.Source->Metadata : nullptr)) : nullptr;
-		Metadata->Initialize(ParentMetadata, InParams.bInheritAttributes);
-	}
-	else
-	{
-		UE_LOG(LogPCG, Warning, TEXT("InitializeFromData has both no source and no metadata override"));
-	}
+void UPCGSpatialData::InitializeOtherMetadataInternal(const FPCGInitializeFromDataParams& InParams, UPCGMetadata* MetadataToInitialize) const
+{
+	check(InParams.bInheritMetadata);
+	check(MetadataToInitialize);
+
+	MetadataToInitialize->Initialize(Metadata, InParams.bInheritAttributes);
 }
 
 void UPCGSpatialData::InitializeSpatialDataInternal(const FPCGInitializeFromDataParams& InParams)
@@ -341,29 +339,51 @@ void UPCGSpatialData::InitializeSpatialDataInternal(const FPCGInitializeFromData
 void UPCGSpatialData::InitializeFromData(const UPCGSpatialData* InSource, const UPCGMetadata* InMetadataParentOverride, bool bInheritMetadata, bool bInheritAttributes)
 {
 	FPCGInitializeFromDataParams Params(InSource);
-	Params.MetadataParentOverride = InMetadataParentOverride;
 	Params.bInheritMetadata = bInheritMetadata;
 	Params.bInheritAttributes = bInheritAttributes;
 	Params.bInheritSpatialData = true;
-	InitializeFromDataWithParams(MoveTemp(Params));
+	
+	InitializeSpatialDataInternal(Params);
+	if (InSource)
+	{
+		if (InMetadataParentOverride)
+		{
+			// Log a warning and ask people to change it
+			// UE_DEPRECATED(5.6, "To be removed")
+			UE_LOG(LogPCG, Warning, TEXT("MetadataParentOverride is deprecated and should not be used anymore. The source data is now responsible to initialize it."));
+			Metadata->Initialize(InMetadataParentOverride, bInheritAttributes);
+		}
+		else
+		{
+			InSource->InitializeOtherMetadata(Params, Metadata);
+		}
+	}
 }
 
 void UPCGSpatialData::InitializeFromDataWithParams(const FPCGInitializeFromDataParams& InParams)
-{	
+{
+	// Make sure to keep InitializeFromData on par, anything implemented here should be replicated in InitializeFromData.
+	// We have to have a deprecation path for InMetadataParentOverride, but we don't want to pollute FPCGInitializeFromDataParams with deprecation
+	// as there was no major release in between this change and when it was added.
 	InitializeSpatialDataInternal(InParams);
-	InitializeMetadataInternal(InParams);
+	if (InParams.Source)
+	{
+		InParams.Source->InitializeOtherMetadata(InParams, Metadata);
+	}
 }
 
 UPCGSpatialData* UPCGSpatialData::DuplicateData(FPCGContext* Context, bool bInitializeMetadata) const
 {
 	UPCGSpatialData* NewSpatialData = CopyInternal(Context);
 	check(NewSpatialData);
-
+	check(NewSpatialData->Metadata);
+	
 	NewSpatialData->InitializeSpatialDataInternal(this);
 
 	if (bInitializeMetadata)
 	{
-		NewSpatialData->InitializeMetadataInternal(this);
+		FPCGInitializeFromDataParams Params(this);
+		InitializeOtherMetadata(Params, NewSpatialData->Metadata);
 	}
 
 	if (bHasCachedLastSelector)
diff --git a/Engine/Plugins/PCG/Source/PCG/Private/Data/PCGUnionData.cpp b/Engine/Plugins/PCG/Source/PCG/Private/Data/PCGUnionData.cpp
--- a/Engine/Plugins/PCG/Source/PCG/Private/Data/PCGUnionData.cpp
+++ b/Engine/Plugins/PCG/Source/PCG/Private/Data/PCGUnionData.cpp
@@ -59,14 +59,12 @@ void UPCGUnionData::AddData(const UPCGSpatialData* InData)
 		CachedBounds = InData->GetBounds();
 		CachedStrictBounds = InData->GetStrictBounds();
 		CachedDimension = InData->GetDimension();
-		Metadata->Initialize(InData->Metadata);
 	}
 	else
 	{
 		CachedBounds += InData->GetBounds();
 		CachedStrictBounds = PCGHelpers::OverlapBounds(CachedStrictBounds, InData->GetStrictBounds());
 		CachedDimension = FMath::Max(CachedDimension, InData->GetDimension());
-		Metadata->AddAttributes(InData->Metadata);
 	}
 
 	if (!FirstNonTrivialTransformData && InData->HasNonTrivialTransform())
@@ -267,27 +265,22 @@ const UPCGBasePointData* UPCGUnionData::CreateBasePointData(FPCGContext* Context
 	TArray<const UPCGSpatialData*> DataRawPtr = Data;
 	TArray<const UPCGMetadata*> InputMetadatas;
 	InputMetadatas.SetNumUninitialized(Data.Num());
-	for (int32 i = 0; i < DataRawPtr.Num(); i++)
+	
+	const UPCGBasePointData* FirstSource = ToPointDataFunc(Context, DataRawPtr[0]);
+	InputMetadatas[0] = FirstSource->Metadata;
+	
+	for (int32 i = 1; i < DataRawPtr.Num(); i++)
 	{
 		InputMetadatas[i] = ToPointDataFunc(Context, DataRawPtr[i])->Metadata;
 	}
 
 	UPCGBasePointData* PointData = FPCGContext::NewObject_AnyThread<UPCGBasePointData>(Context, GetTransientPackage(), PointDataClass);
+	UPCGMetadata* OutMetadata = PointData->MutableMetadata();
 
-	FPCGInitializeFromDataParams InitializeFromDataParams(this);
+	FPCGInitializeFromDataParams InitializeFromDataParams(FirstSource);
 	InitializeFromDataParams.bInheritSpatialData = false;
-	InitializeFromDataParams.MetadataParentOverride = InputMetadatas[0];
 	PointData->InitializeFromDataWithParams(InitializeFromDataParams);
-
-	UPCGMetadata* OutMetadata = PointData->Metadata;
-	check(OutMetadata);
-
-	// Initialize metadata
-	for (const UPCGMetadata* InputMetadata : InputMetadatas)
-	{
-		OutMetadata->AddAttributes(InputMetadata);
-	}
-
+	
 	switch (UnionType)
 	{
 	case EPCGUnionType::LeftToRightPriority:
@@ -576,3 +569,21 @@ UPCGSpatialData* UPCGUnionData::CopyInternal(FPCGContext* Context) const
 
 	return NewUnionData;
 }
+
+void UPCGUnionData::InitializeOtherMetadataInternal(const FPCGInitializeFromDataParams& InParams, UPCGMetadata* MetadataToInitialize) const
+{
+	check(InParams.bInheritMetadata);
+	check(MetadataToInitialize);
+
+	// In the case of the Union, we initialize the data from all the sources. The union will add its own
+	// attributes afterward.
+	FPCGInitializeFromDataParams CopyParams = InParams;
+	for (const UPCGSpatialData* UnionedData : Data)
+	{
+		CopyParams.Source = UnionedData;
+		UnionedData->InitializeOtherMetadata(InParams, MetadataToInitialize);
+	}
+
+	MetadataToInitialize->AddAttributes(Metadata);
+}
+
diff --git a/Engine/Plugins/PCG/Source/PCG/Private/Metadata/PCGMetadata.cpp b/Engine/Plugins/PCG/Source/PCG/Private/Metadata/PCGMetadata.cpp
--- a/Engine/Plugins/PCG/Source/PCG/Private/Metadata/PCGMetadata.cpp
+++ b/Engine/Plugins/PCG/Source/PCG/Private/Metadata/PCGMetadata.cpp
@@ -323,12 +323,12 @@ void UPCGMetadata::AddToCrc(FArchiveCrc32& Ar, bool bFullDataCrc) const
 	}
 }
 
-void UPCGMetadata::Initialize(const UPCGMetadata* InParent, bool bAddAttributesFromParent)
+void UPCGMetadata::Initialize(const UPCGMetadata* InParent, bool bAddAttributesFromParent, const TMap<FPCGMetadataDomainID, FPCGMetadataDomainID>& DomainMapping)
 {
 	// If we are adding attributes from parent, then we use exclude filter with empty list so
 	// that all parameters added. Otherwise use include filter with empty list so none are added.
 	const EPCGMetadataFilterMode bFilter = bAddAttributesFromParent ? EPCGMetadataFilterMode::ExcludeAttributes : EPCGMetadataFilterMode::IncludeAttributes;
-	InitializeWithAttributeFilter(InParent, TSet<FName>(), bFilter);
+	InitializeWithAttributeFilter(InParent, TSet<FPCGAttributeIdentifier>(), bFilter, EPCGStringMatchingOperator::Equal, DomainMapping);
 }
 
 void UPCGMetadata::InitializeWithAttributeFilter(const UPCGMetadata* InParent, const TSet<FName>& InFilteredAttributes, EPCGMetadataFilterMode InFilterMode, EPCGStringMatchingOperator InMatchOperator)
@@ -336,7 +336,7 @@ void UPCGMetadata::InitializeWithAttributeFilter(const UPCGMetadata* InParent, c
 	return InitializeWithAttributeFilter(InParent, FPCGAttributeIdentifier::TransformNameSet(InFilteredAttributes), InFilterMode, InMatchOperator);
 }
 
-void UPCGMetadata::InitializeWithAttributeFilter(const UPCGMetadata* InParent, const TSet<FPCGAttributeIdentifier>& InFilteredAttributes, EPCGMetadataFilterMode InFilterMode, EPCGStringMatchingOperator InMatchOperator)
+void UPCGMetadata::InitializeWithAttributeFilter(const UPCGMetadata* InParent, const TSet<FPCGAttributeIdentifier>& InFilteredAttributes, EPCGMetadataFilterMode InFilterMode, EPCGStringMatchingOperator InMatchOperator, const TMap<FPCGMetadataDomainID, FPCGMetadataDomainID>& DomainMapping)
 {
 	if (Parent)
 	{
@@ -346,15 +346,26 @@ void UPCGMetadata::InitializeWithAttributeFilter(const UPCGMetadata* InParent, c
 
 	if (InParent != this && InParent)
 	{
-		SetupDomainsFromOtherMetadataIfNeeded(InParent);
+		// If we have a domain mapping, we can't have "dynamic" domains
+		if (DomainMapping.IsEmpty())
+		{
+			SetupDomainsFromOtherMetadataIfNeeded(InParent);
+		}
+
 		Parent = InParent;
 		
 		TMap<FPCGMetadataDomainID, TSet<FName>> FilteredAttributeSplit = FPCGAttributeIdentifier::SplitByDomain(InFilteredAttributes);
 		
-		InParent->ForEachValidUniqueConstMetadataDomain([this, &FilteredAttributeSplit, InFilterMode, InMatchOperator](const FPCGMetadataDomainID& DomainID, const FPCGMetadataDomain& OtherMetadataDomain)
+		InParent->ForEachValidUniqueConstMetadataDomain([this, &FilteredAttributeSplit, InFilterMode, InMatchOperator, &DomainMapping](const FPCGMetadataDomainID& DomainID, const FPCGMetadataDomain& OtherMetadataDomain)
 		{
+			FPCGMetadataDomainID OtherDomain = DomainID;
+			if (const FPCGMetadataDomainID* MappedDomain = DomainMapping.Find(DomainID))
+			{
+				OtherDomain = *MappedDomain;
+			}
+
 			TSet<FName>* FilteredAttributesPtr = FilteredAttributeSplit.Find(DomainID);
-			FPCGMetadataDomain* CurrentMetadataDomain = FindOrCreateMetadataDomain(DomainID);
+			FPCGMetadataDomain* CurrentMetadataDomain = FindOrCreateMetadataDomain(OtherDomain);
 			CurrentMetadataDomain->InitializeWithAttributeFilter(&OtherMetadataDomain, (FilteredAttributesPtr ? *FilteredAttributesPtr : TSet<FName>{}), InFilterMode, InMatchOperator);
 		});
 	}
diff --git a/Engine/Plugins/PCG/Source/PCG/Public/Data/PCGDifferenceData.h b/Engine/Plugins/PCG/Source/PCG/Public/Data/PCGDifferenceData.h
--- a/Engine/Plugins/PCG/Source/PCG/Public/Data/PCGDifferenceData.h
+++ b/Engine/Plugins/PCG/Source/PCG/Public/Data/PCGDifferenceData.h
@@ -77,6 +77,7 @@ public:
 	virtual const UPCGSpatialData* FindFirstConcreteShapeFromNetwork() const override { return GetSource() ? GetSource()->FindFirstConcreteShapeFromNetwork() : nullptr; }
 
 protected:
+	virtual void InitializeOtherMetadataInternal(const FPCGInitializeFromDataParams& InParams, UPCGMetadata* MetadataToInitialize) const override;
 	virtual UPCGSpatialData* CopyInternal(FPCGContext* Context) const override;
 	//~End UPCGSpatialData interface
 
diff --git a/Engine/Plugins/PCG/Source/PCG/Public/Data/PCGIntersectionData.h b/Engine/Plugins/PCG/Source/PCG/Public/Data/PCGIntersectionData.h
--- a/Engine/Plugins/PCG/Source/PCG/Public/Data/PCGIntersectionData.h
+++ b/Engine/Plugins/PCG/Source/PCG/Public/Data/PCGIntersectionData.h
@@ -50,6 +50,7 @@ public:
 
 protected:
 	virtual UPCGSpatialData* CopyInternal(FPCGContext* Context) const override;
+	virtual void InitializeOtherMetadataInternal(const FPCGInitializeFromDataParams& InParams, UPCGMetadata* MetadataToInitialize) const override;
 	//~End UPCGSpatialData interface
 
 public:
diff --git a/Engine/Plugins/PCG/Source/PCG/Public/Data/PCGProjectionData.h b/Engine/Plugins/PCG/Source/PCG/Public/Data/PCGProjectionData.h
--- a/Engine/Plugins/PCG/Source/PCG/Public/Data/PCGProjectionData.h
+++ b/Engine/Plugins/PCG/Source/PCG/Public/Data/PCGProjectionData.h
@@ -36,6 +36,7 @@ public:
 	virtual bool HasNonTrivialTransform() const override;
 	virtual bool RequiresCollapseToSample() const override;
 protected:
+	virtual void InitializeOtherMetadataInternal(const FPCGInitializeFromDataParams& InParams, UPCGMetadata* MetadataToInitialize) const override;
 	virtual UPCGSpatialData* CopyInternal(FPCGContext* Context) const override;
 	//~End UPCGSpatialData interface
 
diff --git a/Engine/Plugins/PCG/Source/PCG/Public/Data/PCGSpatialData.h b/Engine/Plugins/PCG/Source/PCG/Public/Data/PCGSpatialData.h
--- a/Engine/Plugins/PCG/Source/PCG/Public/Data/PCGSpatialData.h
+++ b/Engine/Plugins/PCG/Source/PCG/Public/Data/PCGSpatialData.h
@@ -7,6 +7,7 @@
 #include "Elements/PCGProjectionParams.h"
 #include "Metadata/PCGAttributePropertySelector.h"
 #include "Metadata/PCGMetadata.h"
+#include "Metadata/PCGMetadataCommon.h"
 
 #include "PCGSpatialData.generated.h"
 
@@ -27,7 +28,7 @@ struct PCG_API FPCGInitializeFromDataParams
 	GENERATED_BODY()
 
 	FPCGInitializeFromDataParams() = default;
-
+	
 	FPCGInitializeFromDataParams(const UPCGSpatialData* InSource)
 		: Source(InSource)
 	{
@@ -35,9 +36,6 @@ struct PCG_API FPCGInitializeFromDataParams
 	
 	UPROPERTY(BlueprintReadWrite, EditAnywhere, Category = Default)
 	TObjectPtr<const UPCGSpatialData> Source = nullptr;
-	
-	UPROPERTY(BlueprintReadWrite, EditAnywhere, Category = Default)
-	TObjectPtr<const UPCGMetadata> MetadataParentOverride = nullptr;
 		
 	UPROPERTY(BlueprintReadWrite, EditAnywhere, Category = Default)
 	bool bInheritMetadata = true;
@@ -47,6 +45,10 @@ struct PCG_API FPCGInitializeFromDataParams
 
 	UPROPERTY(BlueprintReadWrite, EditAnywhere, Category = Default)
 	bool bInheritSpatialData = true;
+
+	/** When initializing metadata from different types, can provide an extra mapping to know how to initialize metadata across domains. If empty, will match with the same domains. */
+	UPROPERTY(BlueprintReadWrite, EditAnywhere, Category = Default)
+	TMap<FPCGMetadataDomainID, FPCGMetadataDomainID> DomainMapping;
 };
 
 /**
@@ -192,6 +194,7 @@ public:
 	UFUNCTION(BlueprintCallable, Category = Metadata, meta=(DeprecatedFunction, DeprecationMessage = "The Create Empty Metadata function is not needed anymore - it can safely be removed"))
 	UPCGMetadata* CreateEmptyMetadata();
 
+	/** Prefer using InitializeFromDataWithParams instead. Note that InMetadataParentOverride is deprecated on the code side, and should not be used anymore. */
 	UFUNCTION(BlueprintCallable, Category = SpatialData)
 	void InitializeFromData(const UPCGSpatialData* InSource, const UPCGMetadata* InMetadataParentOverride = nullptr, bool bInheritMetadata = true, bool bInheritAttributes = true);
 
@@ -217,15 +220,20 @@ public:
 	UPROPERTY(BlueprintReadWrite, EditAnywhere, Category = Data)
 	bool bKeepZeroDensityPoints = false;
 
+	/** Internal function meant to be used by composite data to initialize their metadata using other sources, or ourselves. Need to be public to be called by child classes.
+	 * InitializeFromData or InitializeFromDataWithParams are the one to use normally.
+	 */
+	void InitializeOtherMetadata(const FPCGInitializeFromDataParams& InParams, UPCGMetadata* MetadataToInitialize) const;
+
 protected:
 	UE_DEPRECATED(5.5, "Call/Implement version with FPCGContext parameter")
 	virtual UPCGSpatialData* CopyInternal() const { return nullptr; }
 
 	virtual UPCGSpatialData* CopyInternal(FPCGContext* Context) const PURE_VIRTUAL(UPCGSpatialData::CopyInternal, return nullptr;);
-
-	void InitializeMetadataInternal(const FPCGInitializeFromDataParams& InParams);
 	virtual void InitializeSpatialDataInternal(const FPCGInitializeFromDataParams& InParams);
 
+	virtual void InitializeOtherMetadataInternal(const FPCGInitializeFromDataParams& InParams, UPCGMetadata* MetadataToInitialize) const;
+
 private:
 	/** Cache to keep track of the latest attribute manipulated on this data. */
 	UPROPERTY()
diff --git a/Engine/Plugins/PCG/Source/PCG/Public/Data/PCGUnionData.h b/Engine/Plugins/PCG/Source/PCG/Public/Data/PCGUnionData.h
--- a/Engine/Plugins/PCG/Source/PCG/Public/Data/PCGUnionData.h
+++ b/Engine/Plugins/PCG/Source/PCG/Public/Data/PCGUnionData.h
@@ -55,6 +55,7 @@ public:
 	virtual const UPCGSpatialData* FindFirstConcreteShapeFromNetwork() const override;
 
 protected:
+	virtual void InitializeOtherMetadataInternal(const FPCGInitializeFromDataParams& InParams, UPCGMetadata* MetadataToInitialize) const override;
 	virtual UPCGSpatialData* CopyInternal(FPCGContext* Context) const override;
 	//~End UPCGSpatialData interface
 
diff --git a/Engine/Plugins/PCG/Source/PCG/Public/Metadata/PCGMetadata.h b/Engine/Plugins/PCG/Source/PCG/Public/Metadata/PCGMetadata.h
--- a/Engine/Plugins/PCG/Source/PCG/Public/Metadata/PCGMetadata.h
+++ b/Engine/Plugins/PCG/Source/PCG/Public/Metadata/PCGMetadata.h
@@ -32,7 +32,7 @@ public:
 	/** To be called by the outer pcg data to initialize the domains supported for this data. Can mark this one as default, will override previously set default. */
 	void SetupDomain(FPCGMetadataDomainID DomainID, bool bIsDefault);
 
-	/** If we have floating metadata, setup it to match a given data type. */
+	/** If we have floating metadata , setup it to match a given data type. */
 	template <typename DataType, std::enable_if_t<std::is_base_of_v<UPCGData, DataType>, int> = 0>
 	void SetupDomainsFromPCGDataType() { SetupFromPCGDataType(DataType::StaticClass()); }
 
@@ -43,7 +43,7 @@ public:
 	void Initialize(const UPCGMetadata* InParent) { InitializeWithAttributeFilter(InParent, TSet<FPCGAttributeIdentifier>()); }
 
 	/** Initializes the metadata from a parent metadata, if any (can be null) with the option to not add attributes from the parent. */
-	void Initialize(const UPCGMetadata* InParent, bool bAddAttributesFromParent);
+	void Initialize(const UPCGMetadata* InParent, bool bAddAttributesFromParent, const TMap<FPCGMetadataDomainID, FPCGMetadataDomainID>& DomainMapping = {});
 
 	/**
 	 * Initializes the metadata from a parent metadata. Copies attributes and values.
@@ -54,7 +54,7 @@ public:
 	UFUNCTION(BlueprintCallable, Category = "PCG|Metadata")
 	void InitializeWithAttributeFilter(const UPCGMetadata* InParent, const TSet<FName>& InFilteredAttributes, EPCGMetadataFilterMode InFilterMode = EPCGMetadataFilterMode::ExcludeAttributes, EPCGStringMatchingOperator InMatchOperator = EPCGStringMatchingOperator::Equal);
 	
-	void InitializeWithAttributeFilter(const UPCGMetadata* InParent, const TSet<FPCGAttributeIdentifier>& InFilteredAttributes, EPCGMetadataFilterMode InFilterMode = EPCGMetadataFilterMode::ExcludeAttributes, EPCGStringMatchingOperator InMatchOperator = EPCGStringMatchingOperator::Equal);
+	void InitializeWithAttributeFilter(const UPCGMetadata* InParent, const TSet<FPCGAttributeIdentifier>& InFilteredAttributes, EPCGMetadataFilterMode InFilterMode = EPCGMetadataFilterMode::ExcludeAttributes, EPCGStringMatchingOperator InMatchOperator = EPCGStringMatchingOperator::Equal, const TMap<FPCGMetadataDomainID, FPCGMetadataDomainID>& DomainMapping = {});
 
 	/** Initializes the metadata from a parent metadata by copying all attributes to it.
 	* @param InMetadataToCopy Metadata to copy from
