[PCG] Fix refresh issue on Metadata nodes when edge is connected.

Input Source Selectors should now be refreshing on the UI when edges are connected/disconnected

#rnx

diff --git a/Engine/Plugins/PCG/Source/PCGEditor/Private/Nodes/PCGEditorGraphNodeBase.cpp b/Engine/Plugins/PCG/Source/PCGEditor/Private/Nodes/PCGEditorGraphNodeBase.cpp
--- a/Engine/Plugins/PCG/Source/PCGEditor/Private/Nodes/PCGEditorGraphNodeBase.cpp
+++ b/Engine/Plugins/PCG/Source/PCGEditor/Private/Nodes/PCGEditorGraphNodeBase.cpp
@@ -1364,6 +1364,11 @@ void UPCGEditorGraphNodeBase::PinConnectionListChanged(UEdGraphPin* Pin)
 
 	Super::PinConnectionListChanged(Pin);
 
+	if (UPCGGraph* PCGGraph = PCGNode->GetGraph())
+	{
+		PCGGraph->OnGraphChangedDelegate.Broadcast(PCGGraph, EPCGChangeType::Edge);
+	}
+
 	// One-time UX functionality, where when a user connects a node with default values, it will activate them for all other pins
 	if (!bHasEverBeenConnected && Pin && !Pin->IsPendingKill())
 	{
diff --git a/Engine/Plugins/PCG/Source/PCGEditor/Private/PCGEditor.cpp b/Engine/Plugins/PCG/Source/PCGEditor/Private/PCGEditor.cpp
--- a/Engine/Plugins/PCG/Source/PCGEditor/Private/PCGEditor.cpp
+++ b/Engine/Plugins/PCG/Source/PCGEditor/Private/PCGEditor.cpp
@@ -2993,6 +2993,14 @@ void FPCGEditor::OnGraphChanged(UPCGGraphInterface* InGraph, EPCGChangeType Chan
 			PaletteWidget->RequestRefresh();
 		}		
 	}
+
+	if (!!(ChangeType & EPCGChangeType::Edge))
+	{
+		for (TSharedPtr<SPCGEditorGraphDetailsView> Widget : PropertyDetailsWidgets)
+		{
+			Widget->GetDetailsView()->ForceRefresh();
+		}
+	}
 }
 
 void FPCGEditor::OnNodeSourceCompiled(const UPCGNode* InNode, const FPCGCompilerDiagnostics& InDiagnostics)
