[PCG] Separate Graph Parameters into its own dockable panel

diff --git a/Engine/Plugins/PCG/Source/PCG/Public/PCGGraph.h b/Engine/Plugins/PCG/Source/PCG/Public/PCGGraph.h
--- a/Engine/Plugins/PCG/Source/PCG/Public/PCGGraph.h
+++ b/Engine/Plugins/PCG/Source/PCG/Public/PCGGraph.h
@@ -83,7 +83,7 @@ public:
 	/** Handle the sync between the parent parameters and the overrides. Will return true if something changed. */
 	bool RefreshParameters(const FInstancedPropertyBag* ParentUserParameters, EPCGGraphParameterEvent InChangeType, FName InChangedPropertyName);
 
-	UPROPERTY(EditAnywhere, Category = "")
+	UPROPERTY(EditAnywhere, Category = "", meta = (ShowOnlyInnerProperties))
 	FInstancedPropertyBag Parameters;
 
 	UPROPERTY(VisibleAnywhere, Category = "", meta = (Hidden))
@@ -469,8 +469,14 @@ protected:
 	TObjectPtr<UPCGEditorGraph> PCGEditorGraph = nullptr;
 #endif // WITH_EDITORONLY_DATA
 
+public:
+#if WITH_EDITOR
+	static FName GetUserParameterPropertyName() { return GET_MEMBER_NAME_CHECKED(UPCGGraph, UserParameters); }
+#endif // WITH_EDITOR
+
+protected:
 	// Parameters
-	UPROPERTY(EditAnywhere, Category = Instance, meta = (DisplayName = "Parameters", NoResetToDefault, DefaultType = "EPropertyBagPropertyType::Double", IsPinTypeAccepted = "UserParametersIsPinTypeAccepted", CanRemoveProperty = "UserParametersCanRemoveProperty"))
+	UPROPERTY(EditAnywhere, Category = "Graph Parameters", DisplayName = "Instance", meta = (NoResetToDefault, DefaultType = "EPropertyBagPropertyType::Double", IsPinTypeAccepted = "UserParametersIsPinTypeAccepted", CanRemoveProperty = "UserParametersCanRemoveProperty"))
 	FInstancedPropertyBag UserParameters;
 
 #if WITH_EDITOR
diff --git a/Engine/Plugins/PCG/Source/PCGEditor/Private/PCGEditor.cpp b/Engine/Plugins/PCG/Source/PCGEditor/Private/PCGEditor.cpp
--- a/Engine/Plugins/PCG/Source/PCGEditor/Private/PCGEditor.cpp
+++ b/Engine/Plugins/PCG/Source/PCGEditor/Private/PCGEditor.cpp
@@ -31,6 +31,7 @@
 #include "Widgets/SPCGEditorGraphFind.h"
 #include "Widgets/SPCGEditorGraphLogView.h"
 #include "Widgets/SPCGEditorGraphNodePalette.h"
+#include "Widgets/SPCGEditorGraphParamsView.h"
 #include "Widgets/SPCGEditorGraphProfilingView.h"
 #include "Nodes/PCGEditorGraphNode.h"
 #include "Nodes/PCGEditorGraphNodeInput.h"
@@ -44,6 +45,7 @@
 #include "GraphEditorActions.h"
 #include "IAssetTools.h"
 #include "LevelEditor.h"
+#include "PropertyPath.h"
 #include "ScopedTransaction.h"
 #include "ShaderCore.h"
 #include "SNodePanel.h"
@@ -89,6 +91,7 @@ namespace FPCGEditor_private
 	const FName ProfilingID = FName(TEXT("Profiling"));
 	const FName LogID = FName(TEXT("Log"));
 	const FName NodeSourceID = FName(TEXT("NodeSource"));
+	const FName UserParamsID = FName(TEXT("UserParams"));
 }
 
 UPCGEditorGraph* FPCGEditor::GetPCGEditorGraph(UPCGGraph* InGraph)
@@ -141,6 +144,7 @@ void FPCGEditor::Initialize(const EToolkitMode::Type InMode, const TSharedPtr<cl
 	ProfilingWidget = CreateProfilingWidget();
 	LogWidget = CreateLogWidget();
 	NodeSourceWidget = CreateNodeSourceWidget();
+	UserParamsWidget = CreateGraphParamsWidget();
 
 	BindCommands();
 	RegisterToolbar();
@@ -486,6 +490,72 @@ void FPCGEditor::JumpToNode(const UPCGNode* InNode)
 	}
 }
 
+void FPCGEditor::BringFocusToPanel(const EPCGEditorPanel Panel) const
+{
+	FName TabName = NAME_None;
+
+	switch (Panel)
+	{
+		case EPCGEditorPanel::Attributes1:
+			TabName = FPCGEditor_private::AttributesID[0];
+			break;
+		case EPCGEditorPanel::Attributes2:
+			TabName = FPCGEditor_private::AttributesID[1];
+			break;
+		case EPCGEditorPanel::Attributes3:
+			TabName = FPCGEditor_private::AttributesID[2];
+			break;
+		case EPCGEditorPanel::Attributes4:
+			TabName = FPCGEditor_private::AttributesID[3];
+			break;
+		case EPCGEditorPanel::DebugObjectTree:
+			TabName = FPCGEditor_private::DebugObjectID;
+			break;
+		case EPCGEditorPanel::Determinism:
+			TabName = FPCGEditor_private::DeterminismID;
+			break;
+		case EPCGEditorPanel::Find:
+			TabName = FPCGEditor_private::FindID;
+			break;
+		case EPCGEditorPanel::GraphEditor:
+			TabName = FPCGEditor_private::GraphEditorID;
+			break;
+		case EPCGEditorPanel::Log:
+			TabName = FPCGEditor_private::LogID;
+			break;
+		case EPCGEditorPanel::NodePalette:
+			TabName = FPCGEditor_private::PaletteID;
+			break;
+		case EPCGEditorPanel::NodeSource:
+			TabName = FPCGEditor_private::NodeSourceID;
+			break;
+		case EPCGEditorPanel::Profiling:
+			TabName = FPCGEditor_private::ProfilingID;
+			break;
+		case EPCGEditorPanel::PropertyDetails1:
+			TabName = FPCGEditor_private::PropertyDetailsID[0];
+			break;
+		case EPCGEditorPanel::PropertyDetails2:
+			TabName = FPCGEditor_private::PropertyDetailsID[1];
+			break;
+		case EPCGEditorPanel::PropertyDetails3:
+			TabName = FPCGEditor_private::PropertyDetailsID[2];
+			break;
+		case EPCGEditorPanel::PropertyDetails4:
+			TabName = FPCGEditor_private::PropertyDetailsID[3];
+			break;
+		case EPCGEditorPanel::UserParams:
+			TabName = FPCGEditor_private::UserParamsID;
+			break;
+	}
+
+	if (TabName != NAME_None)
+	{
+		const TSharedPtr<SDockTab> Tab = TabManager->TryInvokeTab(TabName);
+		Tab->FlashTab();
+	}
+}
+
 void FPCGEditor::RegisterTabSpawners(const TSharedRef<FTabManager>& InTabManager)
 {
 	WorkspaceMenuCategory = InTabManager->AddLocalWorkspaceMenuCategory(LOCTEXT("WorkspaceMenu_PCGEditor", "PCG Editor"));
@@ -559,6 +629,10 @@ void FPCGEditor::RegisterTabSpawners(const TSharedRef<FTabManager>& InTabManager
 	InTabManager->RegisterTabSpawner(FPCGEditor_private::NodeSourceID, FOnSpawnTab::CreateSP(this, &FPCGEditor::SpawnTab_NodeSource))
 		.SetDisplayName(LOCTEXT("NodeSourceTab", "Node Source Editor"))
 		.SetGroup(WorkspaceMenuCategoryRef);
+
+	InTabManager->RegisterTabSpawner(FPCGEditor_private::UserParamsID, FOnSpawnTab::CreateSP(this, &FPCGEditor::SpawnTab_UserParams))
+		.SetDisplayName(LOCTEXT("UserParamsTab", "Graph Parameters"))
+		.SetGroup(WorkspaceMenuCategoryRef);
 }
 
 void FPCGEditor::UnregisterTabSpawners(const TSharedRef<class FTabManager>& InTabManager)
@@ -579,6 +653,7 @@ void FPCGEditor::UnregisterTabSpawners(const TSharedRef<class FTabManager>& InTa
 	InTabManager->UnregisterTabSpawner(FPCGEditor_private::ProfilingID);
 	InTabManager->UnregisterTabSpawner(FPCGEditor_private::LogID);
 	InTabManager->UnregisterTabSpawner(FPCGEditor_private::NodeSourceID);
+	InTabManager->UnregisterTabSpawner(FPCGEditor_private::UserParamsID);
 
 	FAssetEditorToolkit::UnregisterTabSpawners(InTabManager);
 }
@@ -2766,6 +2841,11 @@ TSharedRef<SPCGEditorNodeSource> FPCGEditor::CreateNodeSourceWidget()
 	return SNew(SPCGEditorNodeSource, SharedThis(this));
 }
 
+TSharedRef<SPCGEditorGraphUserParametersView> FPCGEditor::CreateGraphParamsWidget()
+{
+	return SNew(SPCGEditorGraphUserParametersView, SharedThis(this));
+}
+
 void FPCGEditor::OnSelectedNodesChanged(const TSet<UObject*>& NewSelection)
 {
 	TArray<TWeakObjectPtr<UObject>> SelectedObjects;
@@ -3158,6 +3238,16 @@ TSharedRef<SDockTab> FPCGEditor::SpawnTab_NodeSource(const FSpawnTabArgs& Args)
 		];
 }
 
+TSharedRef<SDockTab> FPCGEditor::SpawnTab_UserParams(const FSpawnTabArgs& Args)
+{
+	return SNew(SDockTab)
+		.Label(LOCTEXT("PCGUserParams", "Graph Parameters"))
+		.TabColorScale(GetTabColorScale())
+		[
+			UserParamsWidget.ToSharedRef()
+		];
+}
+
 FText FPCGEditor::GetDetailsTabLabel(int DetailsIndex)
 {
 	if (DetailsIndex == 0)
diff --git a/Engine/Plugins/PCG/Source/PCGEditor/Private/PCGEditor.h b/Engine/Plugins/PCG/Source/PCGEditor/Private/PCGEditor.h
--- a/Engine/Plugins/PCG/Source/PCGEditor/Private/PCGEditor.h
+++ b/Engine/Plugins/PCG/Source/PCGEditor/Private/PCGEditor.h
@@ -23,6 +23,7 @@ class SPCGEditorGraphLogView;
 class SPCGEditorGraphNodePalette;
 class SPCGEditorGraphProfilingView;
 class SPCGEditorNodeSource;
+class SPCGEditorGraphUserParametersView;
 class UEdGraphNode;
 class UPCGComponent;
 class UPCGEditorGraph;
@@ -33,6 +34,27 @@ struct FPCGCompilerDiagnostics;
 DECLARE_MULTICAST_DELEGATE_OneParam(FOnInspectedStackChanged, const FPCGStack&);
 DECLARE_MULTICAST_DELEGATE_OneParam(FOnInspectedNodeChanged, UPCGEditorGraphNodeBase*);
 
+enum class EPCGEditorPanel
+{
+	Attributes1,
+	Attributes2,
+	Attributes3,
+	Attributes4,
+	DebugObjectTree,
+	Determinism,
+	Find,
+	GraphEditor,
+	Log,
+	NodePalette,
+	NodeSource,
+	Profiling,
+	PropertyDetails1,
+	PropertyDetails2,
+	PropertyDetails3,
+	PropertyDetails4,
+	UserParams
+};
+
 class FPCGEditor : public FAssetEditorToolkit, public FGCObject, public FSelfRegisteringEditorUndoClient
 {
 public:
@@ -69,6 +91,9 @@ public:
 	void JumpToNode(const UEdGraphNode* InNode);
 	void JumpToNode(const UPCGNode* InNode);
 
+	/** Focuses the user on a specific Details View panel and flashes the tab. */
+	void BringFocusToPanel(EPCGEditorPanel Panel) const;
+
 	/** Helper to get to the subsystem. */
 	static class UPCGSubsystem* GetSubsystem();
 
@@ -324,6 +349,9 @@ private:
 	/** Create a new node source editor tab widget */
 	TSharedRef<SPCGEditorNodeSource> CreateNodeSourceWidget();
 
+	/** Create a new user graph parameters tab widget */
+	TSharedRef<SPCGEditorGraphUserParametersView> CreateGraphParamsWidget();
+
 	/** Called when the selection changes in the GraphEditor */
 	void OnSelectedNodesChanged(const TSet<UObject*>& NewSelection);
 
@@ -380,6 +408,7 @@ private:
 	TSharedRef<SDockTab> SpawnTab_Profiling(const FSpawnTabArgs& Args);
 	TSharedRef<SDockTab> SpawnTab_Log(const FSpawnTabArgs& Args);
 	TSharedRef<SDockTab> SpawnTab_NodeSource(const FSpawnTabArgs& Args);
+	TSharedRef<SDockTab> SpawnTab_UserParams(const FSpawnTabArgs& Args);
 
 	FText GetDetailsTabLabel(int DetailsIndex);
 	FText GetDetailsViewObjectName(int DetailsIndex);
@@ -395,6 +424,7 @@ private:
 	TSharedPtr<SPCGEditorGraphProfilingView> ProfilingWidget;
 	TSharedPtr<SPCGEditorGraphLogView> LogWidget;
 	TSharedPtr<SPCGEditorNodeSource> NodeSourceWidget;
+	TSharedPtr<SPCGEditorGraphUserParametersView> UserParamsWidget;
 
 	TSharedPtr<FUICommandList> GraphEditorCommands;
 
diff --git a/Engine/Plugins/PCG/Source/PCGEditor/Private/Widgets/SPCGEditorGraphDetailsView.cpp b/Engine/Plugins/PCG/Source/PCGEditor/Private/Widgets/SPCGEditorGraphDetailsView.cpp
--- a/Engine/Plugins/PCG/Source/PCGEditor/Private/Widgets/SPCGEditorGraphDetailsView.cpp
+++ b/Engine/Plugins/PCG/Source/PCGEditor/Private/Widgets/SPCGEditorGraphDetailsView.cpp
@@ -1,6 +1,6 @@
 // Copyright Epic Games, Inc. All Rights Reserved.
 
-#include "SPCGEditorGraphDetailsView.h"
+#include "Widgets/SPCGEditorGraphDetailsView.h"
 
 #include "PCGGraph.h"
 #include "PCGNode.h"
@@ -20,6 +20,7 @@
 #include "Widgets/DeclarativeSyntaxSupport.h"
 #include "Widgets/Images/SImage.h"
 #include "Widgets/Input/SButton.h"
+#include "Widgets/Layout/SSeparator.h"
 #include "Widgets/Notifications/SNotificationList.h"
 #include "Widgets/Text/STextBlock.h"
 
@@ -58,7 +59,7 @@ namespace PCGEditorGraphDetailsView
 	}
 }
 
-class FPCGDebugDetailsExtension : public IDetailPropertyExtensionHandler
+class FPCGDebugDetailsExtension final : public IDetailPropertyExtensionHandler
 {
 	virtual bool IsPropertyExtendable(const UClass* InObjectClass, const IPropertyHandle& PropertyHandle) const override
 	{
@@ -150,7 +151,7 @@ void SPCGEditorGraphDetailsView::Construct(const FArguments& InArgs)
 
 	DetailsView->SetIsPropertyReadOnlyDelegate(FIsPropertyReadOnly::CreateSP(this, &SPCGEditorGraphDetailsView::IsReadOnlyProperty));
 	DetailsView->SetIsPropertyVisibleDelegate(FIsPropertyVisible::CreateSP(this, &SPCGEditorGraphDetailsView::IsVisibleProperty));
-	TSharedRef<FPCGDebugDetailsExtension> DebugDetailsExtension = MakeShared<FPCGDebugDetailsExtension>();
+	const TSharedRef<FPCGDebugDetailsExtension> DebugDetailsExtension = MakeShared<FPCGDebugDetailsExtension>();
 	DetailsView->SetExtensionHandler(DebugDetailsExtension);
 
 	ChildSlot
@@ -200,6 +201,48 @@ void SPCGEditorGraphDetailsView::Construct(const FArguments& InArgs)
 		[
 			DetailsView.ToSharedRef()
 		]
+		// UE_DEPRECATED(5.6, "The graph parameters have moved to their own panel.")
+		+ SVerticalBox::Slot()
+		.AutoHeight()
+		.Padding(10)
+		[
+			SNew(SSeparator)
+		]
+		+ SVerticalBox::Slot()
+		.AutoHeight()
+		[
+			SNew(SHorizontalBox)
+			+ SHorizontalBox::Slot()
+			.FillWidth(1.3)
+			.HAlign(HAlign_Right)
+			.VAlign(VAlign_Center)
+			.Padding(15, 0, 0, 0)
+			[
+				SNew(STextBlock)
+				.Text(LOCTEXT("GraphParamsHaveNewPanel", "Graph Parameters have a new panel. "))
+				.Font(FAppStyle::GetFontStyle(TEXT("PropertyWindow.ItalicFont")))
+				.ShadowColorAndOpacity(FLinearColor::Black)
+			]
+			+ SHorizontalBox::Slot()
+			.FillWidth(1)
+			.HAlign(HAlign_Left)
+			.VAlign(VAlign_Center)
+			.Padding(10, 0, 10, 0)
+			[
+				SNew(SButton)
+				.ButtonStyle(FAppStyle::Get(), "SimpleButton")
+				.Text(LOCTEXT("OpenGraphParamButtonText", "Open"))
+				.OnClicked_Lambda([this]
+				{
+					if (const FPCGEditor* Editor = EditorPtr.Pin().Get())
+					{
+						Editor->BringFocusToPanel(EPCGEditorPanel::UserParams);
+					}
+
+					return FReply::Handled();
+				})
+			]
+		]
 	];
 }
 
@@ -370,7 +413,7 @@ bool SPCGEditorGraphDetailsView::IsVisibleProperty(const FPropertyAndParent& InP
 	// Currently never hide anything from the graph settings
 	if (InPropertyAndParent.Objects.Num() == 1 && Cast<UPCGGraph>(InPropertyAndParent.Objects[0]))
 	{
-		return true;
+		return InPropertyAndParent.Property.GetFName() != UPCGGraph::GetUserParameterPropertyName();
 	}
 
 	// Always hide asset info information
diff --git a/Engine/Plugins/PCG/Source/PCGEditor/Private/Widgets/SPCGEditorGraphParamsView.cpp b/Engine/Plugins/PCG/Source/PCGEditor/Private/Widgets/SPCGEditorGraphParamsView.cpp
new file mode 100644
--- /dev/null
+++ b/Engine/Plugins/PCG/Source/PCGEditor/Private/Widgets/SPCGEditorGraphParamsView.cpp
@@ -0,0 +1,76 @@
+// Copyright Epic Games, Inc. All Rights Reserved.
+
+#include "SPCGEditorGraphParamsView.h"
+
+#include "PCGEditor.h"
+#include "PCGEditorGraph.h"
+#include "PCGGraph.h"
+
+#include "PropertyBagDetails.h"
+#include "PropertyEditorModule.h"
+#include "Modules/ModuleManager.h"
+
+void SPCGEditorGraphUserParametersView::Construct(const FArguments& InArgs, const TSharedPtr<FPCGEditor>& InPCGEditor)
+{
+	if (!InPCGEditor.IsValid())
+	{
+		return;
+	}
+
+	UPCGGraph* PCGGraph = InPCGEditor.IsValid() && InPCGEditor->GetPCGEditorGraph() ? InPCGEditor->GetPCGEditorGraph()->GetPCGGraph() : nullptr;
+
+	if (!PCGGraph)
+	{
+		return;
+	}
+
+	FPropertyEditorModule& PropertyEditorModule = FModuleManager::LoadModuleChecked<FPropertyEditorModule>("PropertyEditor");
+
+	FDetailsViewArgs DetailsViewArgs;
+	{
+		DetailsViewArgs.NameAreaSettings = FDetailsViewArgs::ENameAreaSettings::HideNameArea;
+		DetailsViewArgs.DefaultsOnlyVisibility = EEditDefaultsOnlyNodeVisibility::Automatic;
+		DetailsViewArgs.bAllowSearch = true;
+		DetailsViewArgs.bShowScrollBar = true;
+		DetailsViewArgs.bShowOptions = true;
+		DetailsViewArgs.bShowObjectLabel = true;
+
+		DetailsViewArgs.bShowModifiedPropertiesOption = false;
+		DetailsViewArgs.bShowDifferingPropertiesOption = false;
+		DetailsViewArgs.bShowHiddenPropertiesWhilePlayingOption = false;
+		DetailsViewArgs.bShowKeyablePropertiesOption = false;
+		DetailsViewArgs.bShowAnimatedPropertiesOption = false;
+		DetailsViewArgs.bShowCustomFilterOption = false;
+		DetailsViewArgs.bShowSectionSelector = false;
+		DetailsViewArgs.bShowLooseProperties = false;
+
+		DetailsViewArgs.bShowPropertyMatrixButton = false;
+		DetailsViewArgs.bUpdatesFromSelection = false;
+		DetailsViewArgs.bForceHiddenPropertyVisibility = false;
+	}
+
+	// Note: Single Property View Header notes that it doesn't work with arrays or structs, so using filters for now.
+	const TSharedRef<IDetailsView> DetailsView = PropertyEditorModule.CreateDetailView(DetailsViewArgs);
+	DetailsView->SetObject(PCGGraph);
+	DetailsView->SetIsPropertyVisibleDelegate(FIsPropertyVisible::CreateSP(this, &SPCGEditorGraphUserParametersView::IsPropertyVisible));
+
+	ChildSlot
+	[
+		SNew(SVerticalBox)
+		+ SVerticalBox::Slot()
+		.FillHeight(1.f)
+		[
+			DetailsView
+		]
+	];
+}
+
+bool SPCGEditorGraphUserParametersView::IsPropertyVisible(const FPropertyAndParent& PropertyAndParent) const
+{
+	const TArray<const FProperty*>& Props = PropertyAndParent.ParentProperties;
+	const FName UserParameterPropertyName = UPCGGraph::GetUserParameterPropertyName();
+
+	return
+		PropertyAndParent.Property.GetFName() == UserParameterPropertyName
+		|| (!Props.IsEmpty() && Props[Props.Num() - 1]->GetFName() == UserParameterPropertyName);
+}
diff --git a/Engine/Plugins/PCG/Source/PCGEditor/Private/Widgets/SPCGEditorGraphParamsView.h b/Engine/Plugins/PCG/Source/PCGEditor/Private/Widgets/SPCGEditorGraphParamsView.h
new file mode 100644
--- /dev/null
+++ b/Engine/Plugins/PCG/Source/PCGEditor/Private/Widgets/SPCGEditorGraphParamsView.h
@@ -0,0 +1,23 @@
+// Copyright Epic Games, Inc. All Rights Reserved.
+
+#pragma once
+
+#include "Templates/SharedPointer.h"
+
+#include "Widgets/SCompoundWidget.h"
+
+struct FPropertyAndParent;
+class FPCGEditor;
+
+/** Displays a PCGGraph's User Parameters details view. */
+class SPCGEditorGraphUserParametersView final : public SCompoundWidget
+{
+public:
+	SLATE_BEGIN_ARGS(SPCGEditorGraphUserParametersView)	{}
+	SLATE_END_ARGS()
+
+	void Construct(const FArguments& InArgs, const TSharedPtr<FPCGEditor>& InPCGEditor);
+
+private:
+	bool IsPropertyVisible(const FPropertyAndParent& PropertyAndParent) const;
+};
\ No newline at end of file
