[PCG] Remove `MarkAsGarbage` for temporary objects in CreateSurfaceFromSpline

* We can't do that in a MT setup

#rnx

diff --git a/Engine/Plugins/PCG/Source/PCG/Private/Data/PCGSplineInteriorSurfaceData.cpp b/Engine/Plugins/PCG/Source/PCG/Private/Data/PCGSplineInteriorSurfaceData.cpp
--- a/Engine/Plugins/PCG/Source/PCG/Private/Data/PCGSplineInteriorSurfaceData.cpp
+++ b/Engine/Plugins/PCG/Source/PCG/Private/Data/PCGSplineInteriorSurfaceData.cpp
@@ -176,9 +176,6 @@ void UPCGSplineInteriorSurfaceData::CacheData(FPCGContext* Context)
 		CachedSplinePoints.Add(Point.Transform.GetLocation());
 		CachedSplinePoints2D.Add(FVector2D(Point.Transform.GetLocation()));
 	}
-
-	SplineData->MarkAsGarbage();
-	PointData->MarkAsGarbage();
 }
 
 bool UPCGSplineInteriorSurfaceData::PointInsidePolygon(const FTransform& InTransform, const FBox& InBounds) const
