Instanciating PCG components loses the policy in some instances.
Added some robustness to whether the policy exists or not too.

#rnx

diff --git a/Engine/Plugins/PCG/Source/PCG/Private/PCGComponent.cpp b/Engine/Plugins/PCG/Source/PCG/Private/PCGComponent.cpp
--- a/Engine/Plugins/PCG/Source/PCG/Private/PCGComponent.cpp
+++ b/Engine/Plugins/PCG/Source/PCG/Private/PCGComponent.cpp
@@ -4097,6 +4097,9 @@ void FPCGComponentInstanceData::ApplyToComponent(UActorComponent* Component, con
 
 		bool bDoActorMapping = PCGComponent->bGenerated || PCGHelpers::IsRuntimeOrPIE();
 
+		// If the generation mode is set to runtime, we won't receive a PostEditChange event, because this is a BP. Since it will not trigger a policy update, do it here.
+		PCGComponent->RefreshSchedulingPolicy();
+
 		// Also remap
 		UPCGSubsystem* Subsystem = PCGComponent->GetSubsystem();
 		if (Subsystem && ConstructionSourceComponent)
diff --git a/Engine/Plugins/PCG/Source/PCG/Private/RuntimeGen/PCGRuntimeGenScheduler.cpp b/Engine/Plugins/PCG/Source/PCG/Private/RuntimeGen/PCGRuntimeGenScheduler.cpp
--- a/Engine/Plugins/PCG/Source/PCG/Private/RuntimeGen/PCGRuntimeGenScheduler.cpp
+++ b/Engine/Plugins/PCG/Source/PCG/Private/RuntimeGen/PCGRuntimeGenScheduler.cpp
@@ -291,11 +291,7 @@ void FPCGRuntimeGenScheduler::TickQueueComponentsForGeneration(
 		// TODO: For each execution domain (for now only GenAtRuntime/dynamic), assuming we run Preview through this scheduler, which it seems like we will.
 		if (OriginalComponent->IsManagedByRuntimeGenSystem())
 		{
-			if (!ensure(Policy))
-			{
-				continue;
-			}
-
+			ensure(Policy);
 			bool bHasUnbounded = false;
 			PCGHiGenGrid::FSizeArray GridSizes;
 			ensure(PCGHelpers::GetGenerationGridSizes(OriginalComponent->GetGraph(), InPCGWorldActor, GridSizes, bHasUnbounded));
@@ -339,7 +335,7 @@ void FPCGRuntimeGenScheduler::TickQueueComponentsForGeneration(
 					// Ignore components that have already been generated or marked for generation. Unbounded grid size means not-partitioned.
 					FGridGenerationKey Key(PCGHiGenGrid::UnboundedGridSize(), FIntVector(0), OriginalComponent);
 					if (!GeneratedComponents.Contains(Key) && 
-						Policy->ShouldGenerate(GenSource, OriginalComponentBounds, bIs2DGrid))
+						(!Policy || Policy->ShouldGenerate(GenSource, OriginalComponentBounds, bIs2DGrid)))
 					{
 						check(Key.GetGridDescriptor().Is2DGrid() == bIs2DGrid);
 						AddComponentToGenerate(Key, GenSource, Policy, OriginalComponentBounds, bIs2DGrid);
@@ -434,6 +430,7 @@ void FPCGRuntimeGenScheduler::TickQueueComponentsForGeneration(
 		// TODO: For each execution domain (for now only GenAtRuntime/dynamic), assuming we run Preview through this scheduler, which it seems like we will.
 		if (OriginalComponent->IsManagedByRuntimeGenSystem())
 		{
+			ensure(Policy);
 			// Unbounded will grab the base GenerationRadius used for non-partitioned and unbounded.
 			const double MaxGenerationRadius = OriginalComponent->GetGenerationRadiusFromGrid(EPCGHiGenGrid::Unbounded);
 
@@ -458,7 +455,7 @@ void FPCGRuntimeGenScheduler::TickQueueComponentsForGeneration(
 
 				// Max radius for a non-partitioned component is just the base GenerationRadius.
 				if (DistanceSquared <= MaxGenerationRadius * MaxGenerationRadius && 
-					Policy->ShouldGenerate(GenSource, OriginalComponentBounds, /*bUse2DGrid=*/false))
+					(!Policy || Policy->ShouldGenerate(GenSource, OriginalComponentBounds, /*bUse2DGrid=*/false)))
 				{
 					// Unbounded grid size means not-partitioned.
 					FGridGenerationKey Key(PCGHiGenGrid::UnboundedGridSize(), FIntVector(0), OriginalComponent);
@@ -509,7 +506,7 @@ void FPCGRuntimeGenScheduler::TickCleanup(const TSet<IPCGGenSourceBase*>& InGenS
 			// If the distance to the gen source is greater than the cleanup radius, it means this generation source votes for the component to be cleaned up.
 			// Otherwise, the gen source might still vote for culling regardless.
 			if (SquaredDistToGenSource <= CleanupRadiusSquared && 
-				!Policy->ShouldCull(GenSource, GridBounds, GridDescriptor.Is2DGrid()))
+				(!Policy || !Policy->ShouldCull(GenSource, GridBounds, GridDescriptor.Is2DGrid())))
 			{
 				bAllGenSourcesWantToCleanup = false;
 				break;
@@ -531,10 +528,7 @@ void FPCGRuntimeGenScheduler::TickCleanup(const TSet<IPCGGenSourceBase*>& InGenS
 		check(OriginalComponent);
 
 		const UPCGSchedulingPolicyBase* Policy = OriginalComponent->GetRuntimeGenSchedulingPolicy();
-		if(!ensure(Policy))
-		{
-			continue;
-		}
+		ensure(Policy);
 
 		const double CleanupRadius = OriginalComponent->GetCleanupRadiusFromGrid(Grid);
 		const double CleanupRadiusSquared = CleanupRadius * CleanupRadius;
