Fixed error in edit condition

#rnx

diff --git a/Engine/Plugins/PCG/Source/PCG/Public/Elements/IO/PCGLoadAssetElement.h b/Engine/Plugins/PCG/Source/PCG/Public/Elements/IO/PCGLoadAssetElement.h
--- a/Engine/Plugins/PCG/Source/PCG/Public/Elements/IO/PCGLoadAssetElement.h
+++ b/Engine/Plugins/PCG/Source/PCG/Public/Elements/IO/PCGLoadAssetElement.h
@@ -95,7 +95,7 @@ public:
 	bool bSetDefaultAttributeOverridesFromInput = false;
 
 	// List of Tag:Value default value overrides to apply on the loaded data assets.
-	UPROPERTY(BlueprintReadWrite, EditAnywhere, Category = Settings, meta = (Tooltip="Tag:Value default overrides to be applied to the loaded data.", EditCondition = "!bSetDefaultOverridesFromInput", EditConditionHides))
+	UPROPERTY(BlueprintReadWrite, EditAnywhere, Category = Settings, meta = (Tooltip="Tag:Value default overrides to be applied to the loaded data.", EditCondition = "!bSetDefaultAttributeOverridesFromInput", EditConditionHides))
 	TArray<FString> DefaultAttributeOverrides;
 
 	UPROPERTY(meta = (Tooltip="Overridable-only value to set multiple tags from a single string. If provided, the entries in the Default Value Overrides will be ignored.", PCG_Overridable))
