Fix server crash on nullptr

diff --git a/Engine/Plugins/PCG/Source/PCG/Private/Elements/PCGSpawnActor.cpp b/Engine/Plugins/PCG/Source/PCG/Private/Elements/PCGSpawnActor.cpp
--- a/Engine/Plugins/PCG/Source/PCG/Private/Elements/PCGSpawnActor.cpp
+++ b/Engine/Plugins/PCG/Source/PCG/Private/Elements/PCGSpawnActor.cpp
@@ -1002,7 +1002,8 @@ void FPCGSpawnActorElement::SpawnActors(FPCGSubgraphContext* Context, AActor* Ta
 		bool bAllActorOverridesSucceeded = true;
 
 		const FConstPCGPointValueRanges ValueRanges(PointData);
-		TPCGValueRange<int64> OutMetadataEntryRange = OutPointData->GetMetadataEntryValueRange();
+
+		TPCGValueRange<int64> OutMetadataEntryRange = OutPointData ? OutPointData->GetMetadataEntryValueRange() : TPCGValueRange<int64>();
 
 		for (int32 i = 0; i < PointData->GetNumPoints(); ++i)
 		{
