[PCG] Landscape Data Interface RVT Support

diff --git a/Engine/Plugins/PCG/Shaders/Private/PCGLandscapeDataInterface.ush b/Engine/Plugins/PCG/Shaders/Private/PCGLandscapeDataInterface.ush
--- a/Engine/Plugins/PCG/Shaders/Private/PCGLandscapeDataInterface.ush
+++ b/Engine/Plugins/PCG/Shaders/Private/PCGLandscapeDataInterface.ush
@@ -1,47 +1,233 @@
 // Copyright Epic Games, Inc. All Rights Reserved.
 
-Texture2D {DataInterfaceName}_HeightTexture;
-SamplerState {DataInterfaceName}_HeightTextureSampler;
-float3 {DataInterfaceName}_HeightTextureLWCTile;
-float4x4 {DataInterfaceName}_HeightTextureWorldToUvTransform;
-float4 {DataInterfaceName}_HeightTextureUvScaleBias;
-float2 {DataInterfaceName}_HeightTextureWorldGridSize;
-float3 {DataInterfaceName}_SystemLWCTile;
-
-float2 GetHeightTextureUv_{DataInterfaceName}(float3 InWorldPos)
+#define NUM_VIRTUALTEXTURE_SAMPLES 1
+// Workaround for unbound View uniform buffer. This pattern is used elsewhere in the engine as well.
+#define VT_DISABLE_VIEW_UNIFORM_BUFFER 1
+
+#include "/Engine/Private/VirtualTextureCommon.ush"
+#include "/Engine/Private/Landscape/LandscapeCommon.ush"
+
+// Mirrored from FPCGLandscapeDataProviderProxy::ENormalUnpackType
+#define VIRTUAL_TEXTURE_UNPACK_NORMAL_NONE 0
+#define VIRTUAL_TEXTURE_UNPACK_NORMAL_BC3BC3 1
+#define VIRTUAL_TEXTURE_UNPACK_NORMAL_BC5BC1 2
+#define VIRTUAL_TEXTURE_UNPACK_NORMAL_B5G6R5 3
+
+// Height RVT
+Texture2D			{DataInterfaceName}_HeightVirtualTexture;
+Texture2D<uint4>	{DataInterfaceName}_HeightVirtualTexturePageTable;
+SamplerState		{DataInterfaceName}_HeightVirtualTextureSampler;
+float3				{DataInterfaceName}_HeightVirtualTextureLWCTile;
+float4x4			{DataInterfaceName}_HeightVirtualTextureWorldToUvTransform;
+uint				{DataInterfaceName}_HeightVirtualTextureEnabled;
+uint4				{DataInterfaceName}_HeightVirtualTexturePackedUniform0;
+uint4				{DataInterfaceName}_HeightVirtualTexturePackedUniform1;
+uint4				{DataInterfaceName}_HeightVirtualTextureUniforms;
+
+// Normal RVT
+Texture2D			{DataInterfaceName}_NormalVirtualTexture0;
+Texture2D			{DataInterfaceName}_NormalVirtualTexture1;
+Texture2D<uint4>	{DataInterfaceName}_NormalVirtualTexturePageTable;
+SamplerState		{DataInterfaceName}_NormalVirtualTexture0Sampler;
+SamplerState		{DataInterfaceName}_NormalVirtualTexture1Sampler;
+float3				{DataInterfaceName}_NormalVirtualTextureLWCTile;
+float4x4			{DataInterfaceName}_NormalVirtualTextureWorldToUvTransform;
+uint4				{DataInterfaceName}_NormalVirtualTexturePackedUniform0;
+uint4				{DataInterfaceName}_NormalVirtualTexturePackedUniform1;
+uint4				{DataInterfaceName}_NormalVirtualTextureUniforms0;
+uint4				{DataInterfaceName}_NormalVirtualTextureUniforms1;
+int					{DataInterfaceName}_NormalVirtualTextureUnpackMode;
+uint				{DataInterfaceName}_NormalVirtualTextureEnabled;
+
+// Base Color RVT
+Texture2D			{DataInterfaceName}_BaseColorVirtualTexture;
+Texture2D<uint4>	{DataInterfaceName}_BaseColorVirtualTexturePageTable;
+SamplerState		{DataInterfaceName}_BaseColorVirtualTextureSampler;
+float3				{DataInterfaceName}_BaseColorVirtualTextureLWCTile;
+float4x4			{DataInterfaceName}_BaseColorVirtualTextureWorldToUvTransform;
+uint				{DataInterfaceName}_BaseColorVirtualTextureUnpackSRGB;
+uint				{DataInterfaceName}_BaseColorVirtualTextureUnpackYCoCg;
+uint				{DataInterfaceName}_BaseColorVirtualTextureEnabled;
+uint4				{DataInterfaceName}_BaseColorVirtualTexturePackedUniform0;
+uint4				{DataInterfaceName}_BaseColorVirtualTexturePackedUniform1;
+uint4				{DataInterfaceName}_BaseColorVirtualTextureUniforms;
+
+// Collision based height texture
+Texture2D			{DataInterfaceName}_CollisionHeightTexture;
+SamplerState		{DataInterfaceName}_CollisionHeightTextureSampler;
+uint				{DataInterfaceName}_CollisionHeightTextureEnabled;
+float3				{DataInterfaceName}_CollisionHeightTextureLWCTile;
+float4x4			{DataInterfaceName}_CollisionHeightTextureWorldToUvTransform;
+float4x4			{DataInterfaceName}_CollisionHeightTextureUvToWorldTransform;
+float4				{DataInterfaceName}_CollisionHeightTextureUvScaleBias;
+int2				{DataInterfaceName}_CollisionHeightTextureDimension;
+
+// Misc parameters
+float2				{DataInterfaceName}_LandscapeGridSize;
+
+float2 GetCollisionHeightTextureUv_{DataInterfaceName}(float3 InWorldPos)
 {
-	const FLWCVector4 LwcWorldPos = MakeLWCVector4({DataInterfaceName}_SystemLWCTile, float4(InWorldPos, 1.0f));
-	const FLWCInverseMatrix LwcWorldToUv = MakeLWCInverseMatrix({DataInterfaceName}_HeightTextureLWCTile, {DataInterfaceName}_HeightTextureWorldToUvTransform);
+	// @todo_pcg: Support LWC landscape sampling
+	const float3 LWCTile = (float3)0;
+	const FLWCVector4 LwcWorldPos = MakeLWCVector4(LWCTile, float4(InWorldPos, 1.0f));
+	const FLWCInverseMatrix LwcWorldToUv = MakeLWCInverseMatrix({DataInterfaceName}_CollisionHeightTextureLWCTile, {DataInterfaceName}_CollisionHeightTextureWorldToUvTransform);
 
 	float2 Uv = LWCMultiply(LwcWorldPos, LwcWorldToUv).xy;
-	Uv = Uv * {DataInterfaceName}_HeightTextureUvScaleBias.xy + {DataInterfaceName}_HeightTextureUvScaleBias.zw;
+	Uv = Uv * {DataInterfaceName}_CollisionHeightTextureUvScaleBias.xy + {DataInterfaceName}_CollisionHeightTextureUvScaleBias.zw;
 
 	return Uv;
 }
 
 float GetHeight_{DataInterfaceName}(float3 InWorldPos)
 {
-	const float2 Uv = GetHeightTextureUv_{DataInterfaceName}(InWorldPos);
+	if ({DataInterfaceName}_HeightVirtualTextureEnabled != 0)
+	{
+		const float3 UnpackHeightOrigin = {DataInterfaceName}_HeightVirtualTextureWorldToUvTransform[0].xyz;
+		const float3 UnpackHeightU = {DataInterfaceName}_HeightVirtualTextureWorldToUvTransform[1].xyz;
+		const float3 UnpackHeightV = {DataInterfaceName}_HeightVirtualTextureWorldToUvTransform[2].xyz;
+		const float2 UnpackHeightScaleBias = {DataInterfaceName}_HeightVirtualTextureWorldToUvTransform[3].xy;
+		const float3 LWCTile = (float3)0;
 
-	return {DataInterfaceName}_HeightTexture.SampleLevel({DataInterfaceName}_HeightTextureSampler, Uv, 0).x;
+		const FDFVector3 LwcWorldPos = DFFromTileOffset_Hack(MakeLWCVector3(LWCTile, InWorldPos));
+		const FDFVector3 LwcHeightOrigin = DFFromTileOffset_Hack(MakeLWCVector3({DataInterfaceName}_HeightVirtualTextureLWCTile, UnpackHeightOrigin));
+		const float2 SampleUv = VirtualTextureWorldToUV(LwcWorldPos, LwcHeightOrigin, UnpackHeightU, UnpackHeightV);
+
+		VTPageTableResult PageTable = TextureLoadVirtualPageTableLevel(
+			{DataInterfaceName}_HeightVirtualTexturePageTable,
+			VTPageTableUniform_Unpack({DataInterfaceName}_HeightVirtualTexturePackedUniform0, {DataInterfaceName}_HeightVirtualTexturePackedUniform1),
+			SampleUv,
+			VTADDRESSMODE_CLAMP,
+			VTADDRESSMODE_CLAMP,
+			0.0f);
+
+		const float4 PackedValue = TextureVirtualSample(
+			{DataInterfaceName}_HeightVirtualTexture,
+			{DataInterfaceName}_HeightVirtualTextureSampler,
+			PageTable,
+			0,
+			VTUniform_Unpack({DataInterfaceName}_HeightVirtualTextureUniforms));
+
+		return VirtualTextureUnpackHeight(PackedValue, UnpackHeightScaleBias);
+	}
+	else if ({DataInterfaceName}_CollisionHeightTextureEnabled != 0)
+	{
+		const float2 Uv = GetCollisionHeightTextureUv_{DataInterfaceName}(InWorldPos);
+		return {DataInterfaceName}_CollisionHeightTexture.SampleLevel({DataInterfaceName}_CollisionHeightTextureSampler, Uv, 0).x;
+	}
+
+	return InWorldPos.z;
 }
 
 float3 GetNormal_{DataInterfaceName}(float3 InWorldPos)
 {
-	float3 Pos1 = InWorldPos;
-	float3 Pos2 = InWorldPos + float3(0, {DataInterfaceName}_HeightTextureWorldGridSize.y, 0);
-	float3 Pos3 = InWorldPos + float3({DataInterfaceName}_HeightTextureWorldGridSize.x, 0, 0);
+	if ({DataInterfaceName}_NormalVirtualTextureEnabled != 0)
+	{
+		const float3 UnpackNormalOrigin = {DataInterfaceName}_NormalVirtualTextureWorldToUvTransform[0].xyz;
+		const float3 UnpackNormalU = {DataInterfaceName}_NormalVirtualTextureWorldToUvTransform[1].xyz;
+		const float3 UnpackNormalV = {DataInterfaceName}_NormalVirtualTextureWorldToUvTransform[2].xyz;
+		const float3 LWCTile = (float3)0;
+
+		const FDFVector3 LwcWorldPos = DFFromTileOffset_Hack(MakeLWCVector3(LWCTile, InWorldPos));
+		const FDFVector3 LwcNormalOrigin = DFFromTileOffset_Hack(MakeLWCVector3({DataInterfaceName}_NormalVirtualTextureLWCTile, UnpackNormalOrigin));
+		const float2 SampleUv = VirtualTextureWorldToUV(LwcWorldPos, LwcNormalOrigin, UnpackNormalU, UnpackNormalV);
+
+		VTPageTableResult PageTable = TextureLoadVirtualPageTableLevel(
+			{DataInterfaceName}_NormalVirtualTexturePageTable,
+			VTPageTableUniform_Unpack({DataInterfaceName}_NormalVirtualTexturePackedUniform0, {DataInterfaceName}_NormalVirtualTexturePackedUniform1),
+			SampleUv,
+			VTADDRESSMODE_CLAMP,
+			VTADDRESSMODE_CLAMP,
+			0.0f);
 
-	const float2 Uv1 = GetHeightTextureUv_{DataInterfaceName}(Pos1);
-	const float2 Uv2 = GetHeightTextureUv_{DataInterfaceName}(Pos2);
-	const float2 Uv3 = GetHeightTextureUv_{DataInterfaceName}(Pos3);
+		const float4 PackedValue0 = TextureVirtualSample(
+			{DataInterfaceName}_NormalVirtualTexture0,
+			{DataInterfaceName}_NormalVirtualTexture0Sampler,
+			PageTable,
+			0,
+			VTUniform_Unpack({DataInterfaceName}_NormalVirtualTextureUniforms0));
 
-	Pos1.z = {DataInterfaceName}_HeightTexture.SampleLevel({DataInterfaceName}_HeightTextureSampler, Uv1, 0).x;
-	Pos2.z = {DataInterfaceName}_HeightTexture.SampleLevel({DataInterfaceName}_HeightTextureSampler, Uv2, 0).x;
-	Pos3.z = {DataInterfaceName}_HeightTexture.SampleLevel({DataInterfaceName}_HeightTextureSampler, Uv3, 0).x;
+		const float4 PackedValue1 = TextureVirtualSample(
+			{DataInterfaceName}_NormalVirtualTexture1,
+			{DataInterfaceName}_NormalVirtualTexture1Sampler,
+			PageTable,
+			0,
+			VTUniform_Unpack({DataInterfaceName}_NormalVirtualTextureUniforms1));
+
+		if ({DataInterfaceName}_NormalVirtualTextureUnpackMode == VIRTUAL_TEXTURE_UNPACK_NORMAL_BC3BC3)
+		{
+			return VirtualTextureUnpackNormalBC3BC3(PackedValue0, PackedValue1);
+		}
+		else if ({DataInterfaceName}_NormalVirtualTextureUnpackMode == VIRTUAL_TEXTURE_UNPACK_NORMAL_BC5BC1)
+		{
+			return VirtualTextureUnpackNormalBC5BC1(PackedValue0, PackedValue1);
+		}
+		else if ({DataInterfaceName}_NormalVirtualTextureUnpackMode == VIRTUAL_TEXTURE_UNPACK_NORMAL_B5G6R5)
+		{
+			return VirtualTextureUnpackNormalBGR565(PackedValue1);
+		}
+	}
+
+	float3 Pos1 = InWorldPos;
+	float3 Pos2 = InWorldPos + float3(0.0f, {DataInterfaceName}_LandscapeGridSize.y, 0.0f);
+	float3 Pos3 = InWorldPos + float3({DataInterfaceName}_LandscapeGridSize.x, 0.0f, 0.0f);
+
+	Pos1.z = GetHeight_{DataInterfaceName}(Pos1);
+	Pos2.z = GetHeight_{DataInterfaceName}(Pos2);
+	Pos3.z = GetHeight_{DataInterfaceName}(Pos3);
 
 	const float3 DeltaX = Pos2 - Pos1;
 	const float3 DeltaY = Pos3 - Pos1;
 
 	return normalize(cross(DeltaY, DeltaX));
+
+}
+
+float3 GetBaseColor_{DataInterfaceName}(float3 InWorldPos)
+{
+	if ({DataInterfaceName}_BaseColorVirtualTextureEnabled != 0)
+	{
+		const float3 UnpackBaseColorOrigin = {DataInterfaceName}_BaseColorVirtualTextureWorldToUvTransform[0].xyz;
+		const float3 UnpackBaseColorU = {DataInterfaceName}_BaseColorVirtualTextureWorldToUvTransform[1].xyz;
+		const float3 UnpackBaseColorV = {DataInterfaceName}_BaseColorVirtualTextureWorldToUvTransform[2].xyz;
+		const float3 LWCTile = (float3)0;
+
+		const FDFVector3 LwcWorldPos = DFFromTileOffset_Hack(MakeLWCVector3(LWCTile, InWorldPos));
+		const FDFVector3 LwcBaseColorOrigin = DFFromTileOffset_Hack(MakeLWCVector3({DataInterfaceName}_BaseColorVirtualTextureLWCTile, UnpackBaseColorOrigin));
+		const float2 SampleUv = VirtualTextureWorldToUV(LwcWorldPos, LwcBaseColorOrigin, UnpackBaseColorU, UnpackBaseColorV);
+
+		VTPageTableResult PageTable = TextureLoadVirtualPageTableLevel(
+			{DataInterfaceName}_BaseColorVirtualTexturePageTable,
+			VTPageTableUniform_Unpack({DataInterfaceName}_BaseColorVirtualTexturePackedUniform0, {DataInterfaceName}_BaseColorVirtualTexturePackedUniform1),
+			SampleUv,
+			VTADDRESSMODE_CLAMP,
+			VTADDRESSMODE_CLAMP,
+			0.0f);
+
+		const float4 PackedValue = TextureVirtualSample(
+			{DataInterfaceName}_BaseColorVirtualTexture,
+			{DataInterfaceName}_BaseColorVirtualTextureSampler,
+			PageTable,
+			0,
+			VTUniform_Unpack({DataInterfaceName}_BaseColorVirtualTextureUniforms));
+
+		if ({DataInterfaceName}_BaseColorVirtualTextureUnpackSRGB)
+		{
+			return VirtualTextureUnpackBaseColorSRGB(PackedValue);
+		}
+		else if ({DataInterfaceName}_BaseColorVirtualTextureUnpackYCoCg)
+		{
+			return VirtualTextureUnpackBaseColorYCoCg(PackedValue);
+		}
+		else
+		{
+			return PackedValue.rgb;
+		}
+	}
+
+	return (float3)0.0f;
 }
+
+#undef VIRTUAL_TEXTURE_UNPACK_NORMAL_NONE
+#undef VIRTUAL_TEXTURE_UNPACK_NORMAL_BC3BC3
+#undef VIRTUAL_TEXTURE_UNPACK_NORMAL_BC5BC1
+#undef VIRTUAL_TEXTURE_UNPACK_NORMAL_B5G6R5
diff --git a/Engine/Plugins/PCG/Source/PCG/Private/Compute/DataInterfaces/PCGLandscapeDataInterface.cpp b/Engine/Plugins/PCG/Source/PCG/Private/Compute/DataInterfaces/PCGLandscapeDataInterface.cpp
--- a/Engine/Plugins/PCG/Source/PCG/Private/Compute/DataInterfaces/PCGLandscapeDataInterface.cpp
+++ b/Engine/Plugins/PCG/Source/PCG/Private/Compute/DataInterfaces/PCGLandscapeDataInterface.cpp
@@ -22,6 +22,7 @@
 #include "ComputeFramework/ShaderParamTypeDefinition.h"
 #include "Containers/ResourceArray.h"
 #include "Misc/LargeWorldRenderPosition.h"
+#include "VT/RuntimeVirtualTexture.h"
 
 #include UE_INLINE_GENERATED_CPP_BY_NAME(PCGLandscapeDataInterface)
 
@@ -85,9 +86,9 @@ FPCGLandscapeResource::FPCGLandscapeResource(const FResourceKey& InKey)
 	const FIntPoint CaptureVertexSpan = CaptureQuadSpan + FIntPoint(1, 1);
 	const int32 SampleCount = CaptureVertexSpan.X * CaptureVertexSpan.Y;
 
-	LandscapeTextures = new FPCGLandscapeTextureResource(CaptureVertexSpan);
+	LandscapeTexture = new FPCGLandscapeTextureResource(CaptureVertexSpan);
 
-	TArray<float>* HeightValues = ResourceKey.bIncludesHeight ? &LandscapeTextures->EditHeightValues(SampleCount) : nullptr;
+	TArray<float>* HeightValues = &LandscapeTexture->EditHeightValues(SampleCount);
 
 	const FIntPoint RegionVertexBase = ResourceKey.MinCaptureRegion * ComponentQuadCount;
 
@@ -129,16 +130,14 @@ FPCGLandscapeResource::FPCGLandscapeResource(const FResourceKey& InKey)
 
 	ActorToWorldTransform = LandscapeTransform.ToMatrixWithScale();
 	WorldToActorTransform = ActorToWorldTransform.Inverse();
-	TextureWorldGridSize = FVector2D(ResourceKey.Source->GetTransform().GetScale3D());
 }
 
 FPCGLandscapeResource::~FPCGLandscapeResource()
 {
-	if (LandscapeTextures)
+	if (LandscapeTexture)
 	{
-		// Release the LandscapeTextures resource handle on the RHI and clear the pointer.
-
-		ENQUEUE_RENDER_COMMAND(BeginDestroyCommand)([RT_Resource = LandscapeTextures](FRHICommandListImmediate& RHICmdList)
+		// Release the LandscapeTexture resource handle on the RHI and clear the pointer.
+		ENQUEUE_RENDER_COMMAND(BeginDestroyCommand)([RT_Resource = LandscapeTexture](FRHICommandListImmediate& RHICmdList)
 		{
 			check(RT_Resource);
 			RT_Resource->ReleaseResource();
@@ -151,7 +150,7 @@ FPCGLandscapeResource::~FPCGLandscapeResource()
 			});
 		});
 
-		LandscapeTextures = nullptr;
+		LandscapeTexture = nullptr;
 	}
 }
 
@@ -160,24 +159,71 @@ void UPCGLandscapeDataInterface::GetSupportedInputs(TArray<FShaderFunctionDefini
 	OutFunctions.AddDefaulted_GetRef()
 		.SetName(TEXT("GetHeight"))
 		.AddReturnType(EShaderFundamentalType::Float)
-		.AddParam(EShaderFundamentalType::Float, 3);
+		.AddParam(EShaderFundamentalType::Float, 3); // InWorldPos
 
 	OutFunctions.AddDefaulted_GetRef()
 		.SetName(TEXT("GetNormal"))
 		.AddReturnType(EShaderFundamentalType::Float, 3)
-		.AddParam(EShaderFundamentalType::Float, 3);
+		.AddParam(EShaderFundamentalType::Float, 3); // InWorldPos
 
-	// TODO: Color, PhysicalMaterialIndex, etc.
+	OutFunctions.AddDefaulted_GetRef()
+		.SetName(TEXT("GetBaseColor"))
+		.AddReturnType(EShaderFundamentalType::Float, 3)
+		.AddParam(EShaderFundamentalType::Float, 3); // InWorldPos
 }
 
 BEGIN_SHADER_PARAMETER_STRUCT(FPCGLandscapeDataInterfaceParameters,)
-	SHADER_PARAMETER_TEXTURE(Texture2D, HeightTexture)
-	SHADER_PARAMETER_SAMPLER(SamplerState, HeightTextureSampler)
-	SHADER_PARAMETER(FVector3f, HeightTextureLWCTile)
-	SHADER_PARAMETER(FMatrix44f, HeightTextureWorldToUvTransform)
-	SHADER_PARAMETER(FVector4f, HeightTextureUvScaleBias)
-	SHADER_PARAMETER(FVector2f, HeightTextureWorldGridSize)
-	SHADER_PARAMETER(FVector3f, SystemLWCTile)
+	// Height RVT
+	SHADER_PARAMETER_SRV(Texture2D,				HeightVirtualTexture)
+	SHADER_PARAMETER_TEXTURE(Texture2D<uint4>,	HeightVirtualTexturePageTable)
+	SHADER_PARAMETER_SAMPLER(SamplerState,		HeightVirtualTextureSampler)
+	SHADER_PARAMETER(FVector3f,					HeightVirtualTextureLWCTile)
+	SHADER_PARAMETER(FMatrix44f,				HeightVirtualTextureWorldToUvTransform)
+	SHADER_PARAMETER(uint32,					HeightVirtualTextureEnabled)
+	SHADER_PARAMETER(FUintVector4,				HeightVirtualTexturePackedUniform0)
+	SHADER_PARAMETER(FUintVector4,				HeightVirtualTexturePackedUniform1)
+	SHADER_PARAMETER(FUintVector4,				HeightVirtualTextureUniforms)
+
+	// Normal RVT
+	SHADER_PARAMETER_SRV(Texture2D,				NormalVirtualTexture0)
+	SHADER_PARAMETER_SRV(Texture2D,				NormalVirtualTexture1)
+	SHADER_PARAMETER_TEXTURE(Texture2D<uint4>,	NormalVirtualTexturePageTable)
+	SHADER_PARAMETER_SAMPLER(SamplerState,		NormalVirtualTexture0Sampler)
+	SHADER_PARAMETER_SAMPLER(SamplerState,		NormalVirtualTexture1Sampler)
+	SHADER_PARAMETER(FVector3f,					NormalVirtualTextureLWCTile)
+	SHADER_PARAMETER(FMatrix44f,				NormalVirtualTextureWorldToUvTransform)
+	SHADER_PARAMETER(FUintVector4,				NormalVirtualTexturePackedUniform0)
+	SHADER_PARAMETER(FUintVector4,				NormalVirtualTexturePackedUniform1)
+	SHADER_PARAMETER(FUintVector4,				NormalVirtualTextureUniforms0)
+	SHADER_PARAMETER(FUintVector4,				NormalVirtualTextureUniforms1)
+	SHADER_PARAMETER(int32,						NormalVirtualTextureUnpackMode)
+	SHADER_PARAMETER(uint32,					NormalVirtualTextureEnabled)
+
+	// Base Color RVT
+	SHADER_PARAMETER_SRV(Texture2D,				BaseColorVirtualTexture)
+	SHADER_PARAMETER_TEXTURE(Texture2D<uint4>,	BaseColorVirtualTexturePageTable)
+	SHADER_PARAMETER_SAMPLER(SamplerState,		BaseColorVirtualTextureSampler)
+	SHADER_PARAMETER(FVector3f,					BaseColorVirtualTextureLWCTile)
+	SHADER_PARAMETER(FMatrix44f,				BaseColorVirtualTextureWorldToUvTransform)
+	SHADER_PARAMETER(uint32,					BaseColorVirtualTextureUnpackSRGB)
+	SHADER_PARAMETER(uint32,					BaseColorVirtualTextureUnpackYCoCg)
+	SHADER_PARAMETER(uint32,					BaseColorVirtualTextureEnabled)
+	SHADER_PARAMETER(FUintVector4,				BaseColorVirtualTexturePackedUniform0)
+	SHADER_PARAMETER(FUintVector4,				BaseColorVirtualTexturePackedUniform1)
+	SHADER_PARAMETER(FUintVector4,				BaseColorVirtualTextureUniforms)
+
+	// Collision based height texture
+	SHADER_PARAMETER_TEXTURE(Texture2D,			CollisionHeightTexture)
+	SHADER_PARAMETER_SAMPLER(SamplerState,		CollisionHeightTextureSampler)
+	SHADER_PARAMETER(uint32,					CollisionHeightTextureEnabled)
+	SHADER_PARAMETER(FVector3f,					CollisionHeightTextureLWCTile)
+	SHADER_PARAMETER(FMatrix44f,				CollisionHeightTextureWorldToUvTransform)
+	SHADER_PARAMETER(FMatrix44f,				CollisionHeightTextureUvToWorldTransform)
+	SHADER_PARAMETER(FVector4f,					CollisionHeightTextureUvScaleBias)
+	SHADER_PARAMETER(FIntPoint,					CollisionHeightTextureDimension)
+
+	// Misc parameters
+	SHADER_PARAMETER(FVector2f,					LandscapeGridSize)
 END_SHADER_PARAMETER_STRUCT()
 
 void UPCGLandscapeDataInterface::GetShaderParameters(TCHAR const* UID, FShaderParametersMetadataBuilder& InOutBuilder, FShaderParametersMetadataAllocations& InOutAllocations) const
@@ -280,10 +326,104 @@ void UPCGLandscapeDataProvider::Initialize(const UComputeDataInterface* InDataIn
 
 void UPCGLandscapeDataProvider::InitFromLandscape(ALandscape* InLandscape, const FBox& Bounds)
 {
-	TRACE_CPUPROFILER_EVENT_SCOPE(UPCGLandscapeDataProvider::Init);
+	TRACE_CPUPROFILER_EVENT_SCOPE(UPCGLandscapeDataProvider::InitFromLandscape);
 	check(IsInGameThread());
 
-	const ULandscapeInfo* LandscapeInfo = InLandscape ? InLandscape->GetLandscapeInfo() : nullptr;
+	if (!ensure(InLandscape))
+	{
+		return;
+	}
+
+	InitializeRuntimeVirtualTextures(InLandscape);
+
+	// If the height can't be accessed via RVT, build a physics based height texture instead.
+	if (!HeightVirtualTexture)
+	{
+		InitializeLandscapeResource(InLandscape, Bounds);
+	}
+
+	LandscapeGridSize = FVector2D(InLandscape->GetTransform().GetScale3D());
+}
+
+void UPCGLandscapeDataProvider::InitializeRuntimeVirtualTextures(ALandscape* InLandscape)
+{
+	check(InLandscape);
+
+	int32 BaseColorVirtualTextureIndex = INDEX_NONE;
+	int32 HeightVirtualTextureIndex = INDEX_NONE;
+	int32 NormalVirtualTextureIndex = INDEX_NONE;
+
+	// Only worry about virtual textures if our current platform supports them.
+	if (UseVirtualTexturing(GMaxRHIShaderPlatform))
+	{
+		const int32 RuntimeVirtualTextureCount = InLandscape->RuntimeVirtualTextures.Num();
+		for (int32 TextureIndex = 0; TextureIndex < RuntimeVirtualTextureCount; ++TextureIndex)
+		{
+			if (const URuntimeVirtualTexture* RVT = InLandscape->RuntimeVirtualTextures[TextureIndex])
+			{
+				const ERuntimeVirtualTextureMaterialType VirtualMaterialType = RVT->GetMaterialType();
+
+				switch (VirtualMaterialType)
+				{
+				case ERuntimeVirtualTextureMaterialType::WorldHeight:
+					if (HeightVirtualTextureIndex == INDEX_NONE)
+					{
+						HeightVirtualTextureIndex = TextureIndex;
+					}
+					break;
+				case ERuntimeVirtualTextureMaterialType::BaseColor_Normal_Roughness:
+				case ERuntimeVirtualTextureMaterialType::BaseColor_Normal_Specular:
+				case ERuntimeVirtualTextureMaterialType::BaseColor_Normal_Specular_YCoCg:
+				case ERuntimeVirtualTextureMaterialType::BaseColor_Normal_Specular_Mask_YCoCg:
+					if (NormalVirtualTextureIndex == INDEX_NONE)
+					{
+						NormalVirtualTextureIndex = TextureIndex;
+						NormalVirtualTextureMode = VirtualMaterialType;
+					}
+					// intentionally falling through to the BaseColor option
+				case ERuntimeVirtualTextureMaterialType::BaseColor:
+					if (BaseColorVirtualTextureIndex == INDEX_NONE)
+					{
+						BaseColorVirtualTextureIndex = TextureIndex;
+					}
+					break;
+				}
+
+				if (BaseColorVirtualTextureIndex == TextureIndex)
+				{
+					bBaseColorSRGB = RVT->IsLayerSRGB(0);
+
+					if (RVT->IsLayerYCoCg(0))
+					{
+						BaseColorVirtualTextureUnpackType = EPCGVirtualTextureBaseColorUnpackType::YCoCgUnpack;
+					}
+					else if (VirtualMaterialType == ERuntimeVirtualTextureMaterialType::BaseColor_Normal_Roughness)
+					{
+						BaseColorVirtualTextureUnpackType = EPCGVirtualTextureBaseColorUnpackType::SRGBUnpack;
+					}
+				}
+			}
+		}
+	}
+
+	BaseColorVirtualTexture = InLandscape->RuntimeVirtualTextures.IsValidIndex(BaseColorVirtualTextureIndex)
+		? BaseColorVirtualTexture = InLandscape->RuntimeVirtualTextures[BaseColorVirtualTextureIndex]
+		: nullptr;
+
+	HeightVirtualTexture = InLandscape->RuntimeVirtualTextures.IsValidIndex(HeightVirtualTextureIndex)
+		?  InLandscape->RuntimeVirtualTextures[HeightVirtualTextureIndex]
+		: nullptr;
+
+	NormalVirtualTexture = InLandscape->RuntimeVirtualTextures.IsValidIndex(NormalVirtualTextureIndex)
+		? InLandscape->RuntimeVirtualTextures[NormalVirtualTextureIndex]
+		: nullptr;
+}
+
+void UPCGLandscapeDataProvider::InitializeLandscapeResource(ALandscape* InLandscape, const FBox& Bounds)
+{
+	check(InLandscape);
+
+	const ULandscapeInfo* LandscapeInfo = InLandscape->GetLandscapeInfo();
 
 	if (!ensure(LandscapeInfo))
 	{
@@ -339,7 +479,6 @@ void UPCGLandscapeDataProvider::InitFromLandscape(ALandscape* InLandscape, const
 	Key.CapturedRegions.Reserve(MaxRegionCount);
 	Key.MinCaptureRegion = FIntPoint(TNumericLimits<int32>::Max(), TNumericLimits<int32>::Max());
 	Key.MaxCaptureRegion = FIntPoint(TNumericLimits<int32>::Min(), TNumericLimits<int32>::Min());
-	Key.bIncludesHeight = true;
 
 	auto AddRegion = [&](const FIntPoint& Region)
 	{
@@ -378,16 +517,154 @@ void UPCGLandscapeDataProvider::InitFromLandscape(ALandscape* InLandscape, const
 
 FComputeDataProviderRenderProxy* UPCGLandscapeDataProvider::GetRenderProxy()
 {
-	return new FPCGLandscapeDataProviderProxy(Resource);
+	return new FPCGLandscapeDataProviderProxy(
+		Resource,
+		BaseColorVirtualTexture,
+		HeightVirtualTexture,
+		NormalVirtualTexture,
+		bBaseColorSRGB,
+		BaseColorVirtualTextureUnpackType,
+		NormalVirtualTextureMode,
+		LandscapeGridSize);
 }
 
 void UPCGLandscapeDataProvider::Reset()
 {
+	BaseColorVirtualTexture = nullptr;
+	HeightVirtualTexture = nullptr;
+	NormalVirtualTexture = nullptr;
+
 	Resource.Reset();
+	bBaseColorSRGB = false;
+	BaseColorVirtualTextureUnpackType = EPCGVirtualTextureBaseColorUnpackType::None;
+	NormalVirtualTextureMode = ERuntimeVirtualTextureMaterialType::Count;
+	LandscapeGridSize = FVector2D(1.0f, 1.0f);
 
 	Super::Reset();
 }
 
+void FPCGLandscapeDataProviderProxy::FVirtualTextureLayer::Initialize(const URuntimeVirtualTexture* VirtualTexture, uint32 LayerIndex, bool bSRGB)
+{
+	Reset();
+
+	if (VirtualTexture && ::IsValid(VirtualTexture))
+	{
+		if (IAllocatedVirtualTexture* AllocatedTexture = VirtualTexture->GetAllocatedVirtualTexture())
+		{
+			if (FRHIShaderResourceView* PhysicalTextureSRV = AllocatedTexture->GetPhysicalTextureSRV(LayerIndex, bSRGB))
+			{
+				TextureSRV = PhysicalTextureSRV;
+				AllocatedTexture->GetPackedUniform(&TextureUniforms, LayerIndex);
+			}
+		}
+	}
+}
+
+void FPCGLandscapeDataProviderProxy::FVirtualTextureLayer::Reset()
+{
+	TextureSRV = nullptr;
+	TextureUniforms = FUintVector4(0, 0, 0, 0);
+}
+
+bool FPCGLandscapeDataProviderProxy::FVirtualTextureLayer::IsValid() const
+{
+	return TextureSRV.IsValid();
+}
+
+void FPCGLandscapeDataProviderProxy::FVirtualTexturePageTable::Initialize(const URuntimeVirtualTexture* VirtualTexture, uint32 PageTableIndex, bool bIncludeWorldToUv, bool bIncludeHeightUnpack)
+{
+	Reset();
+
+	if (VirtualTexture && ::IsValid(VirtualTexture))
+	{
+		if (IAllocatedVirtualTexture* AllocatedTexture = VirtualTexture->GetAllocatedVirtualTexture())
+		{
+			PageTableRef = AllocatedTexture->GetPageTableTexture(PageTableIndex);
+			if (PageTableRef.IsValid())
+			{
+				AllocatedTexture->GetPackedPageTableUniform(PageTableUniforms);
+
+				if (bIncludeWorldToUv)
+				{
+					WorldToUvParameters[0] = VirtualTexture->GetUniformParameter(ERuntimeVirtualTextureShaderUniform_WorldToUVTransform0);
+					WorldToUvParameters[1] = VirtualTexture->GetUniformParameter(ERuntimeVirtualTextureShaderUniform_WorldToUVTransform1);
+					WorldToUvParameters[2] = VirtualTexture->GetUniformParameter(ERuntimeVirtualTextureShaderUniform_WorldToUVTransform2);
+				}
+
+				if (bIncludeHeightUnpack)
+				{
+					WorldToUvParameters[3] = VirtualTexture->GetUniformParameter(ERuntimeVirtualTextureShaderUniform_WorldHeightUnpack);
+				}
+			}
+		}
+	}
+}
+
+void FPCGLandscapeDataProviderProxy::FVirtualTexturePageTable::Reset()
+{
+	PageTableRef = nullptr;
+	PageTableUniforms[0] = FUintVector4(0, 0, 0, 0);
+	PageTableUniforms[1] = FUintVector4(0, 0, 0, 0);
+	WorldToUvParameters[0] = FVector4::Zero();
+	WorldToUvParameters[1] = FVector4::Zero();
+	WorldToUvParameters[2] = FVector4::Zero();
+	WorldToUvParameters[3] = FVector4::Zero();
+}
+
+bool FPCGLandscapeDataProviderProxy::FVirtualTexturePageTable::IsValid() const
+{
+	return PageTableRef.IsValid();
+}
+
+FPCGLandscapeDataProviderProxy::FPCGLandscapeDataProviderProxy(
+	TSharedPtr<const FPCGLandscapeResource> InResource,
+	const URuntimeVirtualTexture* InBaseColorVirtualTexture,
+	const URuntimeVirtualTexture* InHeightVirtualTexture,
+	const URuntimeVirtualTexture* InNormalVirtualTexture,
+	bool bInBaseColorSRGB,
+	EPCGVirtualTextureBaseColorUnpackType InBaseColorVirtualTextureUnpackType,
+	ERuntimeVirtualTextureMaterialType InNormalVirtualTextureMode,
+	FVector2D InLandscapeGridSize)
+	: Resource(InResource)
+	, BaseColorVirtualTextureUnpackType(InBaseColorVirtualTextureUnpackType)
+	, LandscapeGridSize(InLandscapeGridSize)
+{
+	BaseColorVirtualLayer.Initialize(InBaseColorVirtualTexture, /*PageTableIndex=*/0, bInBaseColorSRGB);
+	BaseColorVirtualPage.Initialize(InBaseColorVirtualTexture, /*PageTableIndex=*/0, /*bIncludeWorldToUv=*/true, /*bIncludeHeightUnpack=*/false);
+
+	HeightVirtualLayer.Initialize(InHeightVirtualTexture, /*PageTableIndex=*/0, /*bSRGB=*/false);
+	HeightVirtualPage.Initialize(InHeightVirtualTexture, /*PageTableIndex=*/0, /*bIncludeWorldToUv=*/true, /*bIncludeHeightUnpack=*/true);
+
+	NormalVirtualPage.Initialize(InNormalVirtualTexture, /*PageTableIndex=*/0, /*bIncludeWorldToUv=*/true, /*bIncludeHeightUnpack=*/false);
+
+	switch (InNormalVirtualTextureMode)
+	{
+	case ERuntimeVirtualTextureMaterialType::BaseColor_Normal_Roughness:
+		NormalVirtualLayer0.Initialize(InNormalVirtualTexture, /*PageTableIndex=*/0, /*bSRGB=*/false);
+		NormalVirtualLayer1.Initialize(InNormalVirtualTexture, /*PageTableIndex=*/1, /*bSRGB=*/false);
+		NormalUnpackMode = ENormalUnpackType::B5G6R5;
+		break;
+	case ERuntimeVirtualTextureMaterialType::BaseColor_Normal_Specular:
+		NormalVirtualLayer0.Initialize(InNormalVirtualTexture, /*PageTableIndex=*/0, /*bSRGB=*/false);
+		NormalVirtualLayer1.Initialize(InNormalVirtualTexture, /*PageTableIndex=*/1, /*bSRGB=*/false);
+		NormalUnpackMode = ENormalUnpackType::BC3BC3;
+		break;
+
+	case ERuntimeVirtualTextureMaterialType::BaseColor_Normal_Specular_YCoCg:
+	case ERuntimeVirtualTextureMaterialType::BaseColor_Normal_Specular_Mask_YCoCg:
+		NormalVirtualLayer0.Initialize(InNormalVirtualTexture, /*PageTableIndex=*/1, /*bSRGB=*/false);
+		NormalVirtualLayer1.Initialize(InNormalVirtualTexture, /*PageTableIndex=*/2, /*bSRGB=*/false);
+		NormalUnpackMode = ENormalUnpackType::BC5BC1;
+		break;
+	default:
+		NormalVirtualPage.Reset();
+		NormalVirtualLayer0.Reset();
+		NormalVirtualLayer1.Reset();
+		NormalUnpackMode = ENormalUnpackType::None;
+		break;
+	}
+}
+
 bool FPCGLandscapeDataProviderProxy::IsValid(FValidationData const& InValidationData) const
 {
 	if (InValidationData.ParameterStructSize != sizeof(FParameters))
@@ -403,47 +680,213 @@ void FPCGLandscapeDataProviderProxy::GatherDispatchData(FDispatchData const& InD
 	const TStridedView<FParameters> ParameterArray = MakeStridedParameterView<FParameters>(InDispatchData);
 	for (int32 InvocationIndex = 0; InvocationIndex < ParameterArray.Num(); ++InvocationIndex)
 	{
-		FParameters& Parameters = ParameterArray[InvocationIndex];
+		SetShaderParameters(ParameterArray[InvocationIndex]);
+	}
+}
 
-		// Set Samplers
-		FRHISamplerState* BilinearSamplerState = TStaticSamplerState<SF_Bilinear, AM_Clamp, AM_Clamp, AM_Clamp>::GetRHI();
-		FRHISamplerState* PointClampedSampler = TStaticSamplerState<SF_Point, AM_Clamp, AM_Clamp, AM_Clamp>::GetRHI();
-		Parameters.HeightTextureSampler = RHIPixelFormatHasCapabilities(EPixelFormat::PF_R32_FLOAT, EPixelFormatCapabilities::TextureFilterable) ? BilinearSamplerState : PointClampedSampler;
+void FPCGLandscapeDataProviderProxy::AllocateResources(FRDGBuilder& GraphBuilder, FAllocationData const& InAllocationData)
+{
+	if (Resource && Resource->LandscapeTexture)
+	{
+		Resource->LandscapeTexture->InitResource(GraphBuilder.RHICmdList);
+	}
+}
 
-		// Set Textures
-		FRHITexture* HeightTexture = (Resource && Resource->LandscapeTextures) ? Resource->LandscapeTextures->GetHeightTexture() : (FRHITexture*)GBlackTexture->TextureRHI;
-		check(HeightTexture);
+void FPCGLandscapeDataProviderProxy::SetShaderParameters(FParameters& OutShaderParameters)
+{
+	FRHISamplerState* BilinearSamplerState				= TStaticSamplerState<SF_Bilinear, AM_Clamp, AM_Clamp, AM_Clamp>::GetRHI();
+	FRHISamplerState* PointClampedSampler				= TStaticSamplerState<SF_Point, AM_Clamp, AM_Clamp, AM_Clamp>::GetRHI();
 
-		// TODO: Bindless resources to handle multiple landscapes
-		Parameters.HeightTexture = HeightTexture;
+	OutShaderParameters.BaseColorVirtualTextureSampler	= BilinearSamplerState;
+	OutShaderParameters.HeightVirtualTextureSampler		= BilinearSamplerState;
+	OutShaderParameters.NormalVirtualTexture0Sampler	= BilinearSamplerState;
+	OutShaderParameters.NormalVirtualTexture1Sampler	= BilinearSamplerState;
+	OutShaderParameters.CollisionHeightTextureSampler	= RHIPixelFormatHasCapabilities(EPixelFormat::PF_R32_FLOAT, EPixelFormatCapabilities::TextureFilterable) ? BilinearSamplerState : PointClampedSampler;
 
-		if (Resource)
-		{
-			Parameters.HeightTextureLWCTile = Resource->LandscapeLWCTile;
-			Parameters.HeightTextureWorldToUvTransform = (FMatrix44f)Resource->WorldToActorTransform;
-			Parameters.HeightTextureUvScaleBias = (FVector4f)Resource->UvScaleBias;
-			Parameters.HeightTextureWorldGridSize = (FVector2f)Resource->TextureWorldGridSize;
-		}
-		else
-		{
-			Parameters.HeightTextureLWCTile = FVector3f::ZeroVector;
-			Parameters.HeightTextureWorldToUvTransform = (FMatrix44f)FMatrix::Identity;
-			Parameters.HeightTextureUvScaleBias = (FVector4f)FVector4(1.0f, 1.0f, 0.0f, 0.0f);
-			Parameters.HeightTextureWorldGridSize = (FVector2f)FVector2D(1.0f, 1.0f);
-		}
+	if (!SetBaseColorVirtualTextureParameters(OutShaderParameters))
+	{
+		SetBaseColorVirtualTextureParameters_Default(OutShaderParameters);
+	}
+
+	if (!SetHeightVirtualTextureParameters(OutShaderParameters))
+	{
+		SetHeightVirtualTextureParameters_Default(OutShaderParameters);
+	}
 
-		// System tile for LWC
-		FVector3f LWCTile(0, 0, 0); // Used to offset in LWC for precision, default to 0-vector for now.
-		Parameters.SystemLWCTile = LWCTile;
+	if (!SetNormalVirtualTextureParameters(OutShaderParameters))
+	{
+		SetNormalVirtualTextureParameters_Default(OutShaderParameters);
 	}
+
+	if (!SetCollisionHeightTextureParameters(OutShaderParameters))
+	{
+		SetCollisionHeightTextureParameters_Defaults(OutShaderParameters);
+	}
+
+	OutShaderParameters.LandscapeGridSize = (FVector2f)LandscapeGridSize;
 }
 
-void FPCGLandscapeDataProviderProxy::AllocateResources(FRDGBuilder& GraphBuilder, FAllocationData const& InAllocationData)
+bool FPCGLandscapeDataProviderProxy::SetBaseColorVirtualTextureParameters(FParameters& OutShaderParameters) const
+{
+	if (!BaseColorVirtualPage.IsValid() || !BaseColorVirtualLayer.IsValid())
+	{
+		return false;
+	}
+
+	OutShaderParameters.BaseColorVirtualTexture = BaseColorVirtualLayer.TextureSRV;
+	OutShaderParameters.BaseColorVirtualTexturePageTable = BaseColorVirtualPage.PageTableRef;
+
+	FLargeWorldRenderPosition BaseColorVirtualTextureOrigin(BaseColorVirtualPage.WorldToUvParameters[0]);
+	OutShaderParameters.BaseColorVirtualTextureLWCTile = BaseColorVirtualTextureOrigin.GetTile();
+
+	OutShaderParameters.BaseColorVirtualTextureWorldToUvTransform = FMatrix44f(
+		BaseColorVirtualTextureOrigin.GetOffset(),
+		FVector3f((FVector4f)BaseColorVirtualPage.WorldToUvParameters[1]),
+		FVector3f((FVector4f)BaseColorVirtualPage.WorldToUvParameters[2]),
+		FVector3f(0, 0, 0));
+	OutShaderParameters.BaseColorVirtualTextureUnpackSRGB = BaseColorVirtualTextureUnpackType == EPCGVirtualTextureBaseColorUnpackType::SRGBUnpack;
+	OutShaderParameters.BaseColorVirtualTextureUnpackYCoCg = BaseColorVirtualTextureUnpackType == EPCGVirtualTextureBaseColorUnpackType::YCoCgUnpack;
+	OutShaderParameters.BaseColorVirtualTextureEnabled = 1;
+	OutShaderParameters.BaseColorVirtualTexturePackedUniform0 = BaseColorVirtualPage.PageTableUniforms[0];
+	OutShaderParameters.BaseColorVirtualTexturePackedUniform1 = BaseColorVirtualPage.PageTableUniforms[1];
+	OutShaderParameters.BaseColorVirtualTextureUniforms = BaseColorVirtualLayer.TextureUniforms;
+
+	return true;
+}
+
+bool FPCGLandscapeDataProviderProxy::SetHeightVirtualTextureParameters(FParameters& OutShaderParameters) const
 {
-	if (Resource && Resource->LandscapeTextures)
+	if (!HeightVirtualPage.IsValid() || !HeightVirtualLayer.IsValid())
 	{
-		Resource->LandscapeTextures->InitResource(GraphBuilder.RHICmdList);
+		return false;
+	}
+
+	OutShaderParameters.HeightVirtualTexture = HeightVirtualLayer.TextureSRV;
+	OutShaderParameters.HeightVirtualTexturePageTable = HeightVirtualPage.PageTableRef;
+
+	FLargeWorldRenderPosition HeightVirtualTextureOrigin(HeightVirtualPage.WorldToUvParameters[0]);
+	OutShaderParameters.HeightVirtualTextureLWCTile = HeightVirtualTextureOrigin.GetTile();
+
+	OutShaderParameters.HeightVirtualTextureWorldToUvTransform = FMatrix44f(
+		HeightVirtualTextureOrigin.GetOffset(),
+		FVector3f((FVector4f)HeightVirtualPage.WorldToUvParameters[1]),
+		FVector3f((FVector4f)HeightVirtualPage.WorldToUvParameters[2]),
+		FVector3f((FVector4f)HeightVirtualPage.WorldToUvParameters[3]));
+
+	OutShaderParameters.HeightVirtualTextureEnabled = 1;
+	OutShaderParameters.HeightVirtualTexturePackedUniform0 = HeightVirtualPage.PageTableUniforms[0];
+	OutShaderParameters.HeightVirtualTexturePackedUniform1 = HeightVirtualPage.PageTableUniforms[1];
+	OutShaderParameters.HeightVirtualTextureUniforms = HeightVirtualLayer.TextureUniforms;
+
+	return true;
+}
+
+bool FPCGLandscapeDataProviderProxy::SetNormalVirtualTextureParameters(FParameters& OutShaderParameters) const
+{
+	if (!NormalVirtualPage.IsValid() || !NormalVirtualLayer0.IsValid() || !NormalVirtualLayer1.IsValid())
+	{
+		return false;
+	}
+
+	OutShaderParameters.NormalVirtualTexture0 = NormalVirtualLayer0.TextureSRV;
+	OutShaderParameters.NormalVirtualTexture1 = NormalVirtualLayer1.TextureSRV;
+	OutShaderParameters.NormalVirtualTexturePageTable = NormalVirtualPage.PageTableRef;
+
+	FLargeWorldRenderPosition NormalVirtualTextureOrigin(NormalVirtualPage.WorldToUvParameters[0]);
+	OutShaderParameters.NormalVirtualTextureLWCTile = NormalVirtualTextureOrigin.GetTile();
+
+	OutShaderParameters.NormalVirtualTextureWorldToUvTransform = FMatrix44f(
+		FPlane4f(NormalVirtualTextureOrigin.GetOffset()),
+		FPlane4f((FVector4f)NormalVirtualPage.WorldToUvParameters[1]),
+		FPlane4f((FVector4f)NormalVirtualPage.WorldToUvParameters[2]),
+		FPlane4f(0.0f, 0.0f, 0.0f, 1.0f));
+
+	OutShaderParameters.NormalVirtualTextureEnabled = 1;
+	OutShaderParameters.NormalVirtualTextureUnpackMode = (int32)NormalUnpackMode;
+	OutShaderParameters.NormalVirtualTexturePackedUniform0 = NormalVirtualPage.PageTableUniforms[0];
+	OutShaderParameters.NormalVirtualTexturePackedUniform1 = NormalVirtualPage.PageTableUniforms[1];
+	OutShaderParameters.NormalVirtualTextureUniforms0 = NormalVirtualLayer0.TextureUniforms;
+	OutShaderParameters.NormalVirtualTextureUniforms1 = NormalVirtualLayer1.TextureUniforms;
+
+	return true;
+}
+
+bool FPCGLandscapeDataProviderProxy::SetCollisionHeightTextureParameters(FParameters& OutShaderParameters) const
+{
+	FRHITexture* CollisionHeightTexture = (Resource && Resource->LandscapeTexture) ? Resource->LandscapeTexture->GetHeightTexture() : nullptr;
+
+	if (!CollisionHeightTexture)
+	{
+		return false;
 	}
+
+	const FIntPoint TextureDimensions(Resource->LandscapeTexture->GetDimensions());
+
+	OutShaderParameters.CollisionHeightTextureLWCTile = Resource->LandscapeLWCTile;
+	OutShaderParameters.CollisionHeightTextureWorldToUvTransform = (FMatrix44f)Resource->WorldToActorTransform;
+	OutShaderParameters.CollisionHeightTextureUvToWorldTransform = (FMatrix44f)Resource->ActorToWorldTransform;
+	OutShaderParameters.CollisionHeightTextureUvScaleBias = (FVector4f)Resource->UvScaleBias;
+	OutShaderParameters.CollisionHeightTexture = CollisionHeightTexture;
+	OutShaderParameters.CollisionHeightTextureEnabled = 1;
+	OutShaderParameters.CollisionHeightTextureDimension = TextureDimensions;
+
+	return true;
+}
+
+void FPCGLandscapeDataProviderProxy::SetBaseColorVirtualTextureParameters_Default(FParameters& OutShaderParameters)
+{
+	const FUintVector4 DummyUint4(ForceInitToZero);
+
+	OutShaderParameters.BaseColorVirtualTexture = GBlackTextureWithSRV->ShaderResourceViewRHI;
+	OutShaderParameters.BaseColorVirtualTexturePageTable = GBlackUintTexture->TextureRHI;
+	OutShaderParameters.BaseColorVirtualTextureWorldToUvTransform = FMatrix44f::Identity;
+	OutShaderParameters.BaseColorVirtualTextureUnpackSRGB = 0;
+	OutShaderParameters.BaseColorVirtualTextureUnpackYCoCg = 0;
+	OutShaderParameters.BaseColorVirtualTextureEnabled = 0;
+	OutShaderParameters.BaseColorVirtualTexturePackedUniform0 = DummyUint4;
+	OutShaderParameters.BaseColorVirtualTexturePackedUniform1 = DummyUint4;
+	OutShaderParameters.BaseColorVirtualTextureUniforms = DummyUint4;
+}
+
+void FPCGLandscapeDataProviderProxy::SetHeightVirtualTextureParameters_Default(FParameters& OutShaderParameters)
+{
+	const FUintVector4 DummyUint4(ForceInitToZero);
+
+	OutShaderParameters.HeightVirtualTexture = GBlackTextureWithSRV->ShaderResourceViewRHI;
+	OutShaderParameters.HeightVirtualTexturePageTable = GBlackUintTexture->TextureRHI;
+	OutShaderParameters.HeightVirtualTextureWorldToUvTransform = FMatrix44f::Identity;
+	OutShaderParameters.HeightVirtualTextureEnabled = 0;
+	OutShaderParameters.HeightVirtualTexturePackedUniform0 = DummyUint4;
+	OutShaderParameters.HeightVirtualTexturePackedUniform1 = DummyUint4;
+	OutShaderParameters.HeightVirtualTextureUniforms = DummyUint4;
+}
+
+void FPCGLandscapeDataProviderProxy::SetNormalVirtualTextureParameters_Default(FParameters& OutShaderParameters)
+{
+	const FUintVector4 DummyUint4(ForceInitToZero);
+
+	OutShaderParameters.NormalVirtualTexture0 = GBlackTextureWithSRV->ShaderResourceViewRHI;
+	OutShaderParameters.NormalVirtualTexture1 = GBlackTextureWithSRV->ShaderResourceViewRHI;
+	OutShaderParameters.NormalVirtualTexturePageTable = GBlackUintTexture->TextureRHI;
+	OutShaderParameters.NormalVirtualTextureWorldToUvTransform = FMatrix44f::Identity;
+	OutShaderParameters.NormalVirtualTextureEnabled = 0;
+	OutShaderParameters.NormalVirtualTextureUnpackMode = (int32)ENormalUnpackType::None;
+	OutShaderParameters.NormalVirtualTexturePackedUniform0 = DummyUint4;
+	OutShaderParameters.NormalVirtualTexturePackedUniform1 = DummyUint4;
+	OutShaderParameters.NormalVirtualTextureUniforms0 = DummyUint4;
+	OutShaderParameters.NormalVirtualTextureUniforms1 = DummyUint4;
+}
+
+void FPCGLandscapeDataProviderProxy::SetCollisionHeightTextureParameters_Defaults(FParameters& OutShaderParameters)
+{
+	const FVector4f DummyVector4 = FVector4f::Zero();
+
+	OutShaderParameters.CollisionHeightTexture = GBlackTexture->TextureRHI;
+	OutShaderParameters.CollisionHeightTextureLWCTile = FVector3f::ZeroVector;
+	OutShaderParameters.CollisionHeightTextureWorldToUvTransform = FMatrix44f::Identity;
+	OutShaderParameters.CollisionHeightTextureUvToWorldTransform = FMatrix44f::Identity;
+	OutShaderParameters.CollisionHeightTextureUvScaleBias = DummyVector4;
+	OutShaderParameters.CollisionHeightTextureEnabled = 0;
 }
 
 #undef LOCTEXT_NAMESPACE
diff --git a/Engine/Plugins/PCG/Source/PCG/Private/Compute/DataInterfaces/PCGLandscapeDataInterface.h b/Engine/Plugins/PCG/Source/PCG/Private/Compute/DataInterfaces/PCGLandscapeDataInterface.h
--- a/Engine/Plugins/PCG/Source/PCG/Private/Compute/DataInterfaces/PCGLandscapeDataInterface.h
+++ b/Engine/Plugins/PCG/Source/PCG/Private/Compute/DataInterfaces/PCGLandscapeDataInterface.h
@@ -5,13 +5,27 @@
 #include "Compute/DataInterfaces/PCGComputeDataInterface.h"
 #include "ComputeFramework/ComputeDataProvider.h"
 
+#include "RenderResource.h"
+#include "VT/RuntimeVirtualTextureEnum.h"
+
 #include "PCGLandscapeDataInterface.generated.h"
 
 class ALandscape;
 class FPCGLandscapeDataInterfaceParameters;
 class FPCGLandscapeTextureResource;
+class URuntimeVirtualTexture;
+
+enum class EPCGVirtualTextureBaseColorUnpackType : int32
+{
+	/** No unpacking is required. */
+	None = 0,
+	/** Base color is manually packed as SRGB. */
+	SRGBUnpack,
+	/** Base color is manually packed as YCoCg. */
+	YCoCgUnpack,
+};
 
-/** Manages the resources created by this DI. */
+/** Manages the textures created from landscape height data. */
 class FPCGLandscapeResource
 {
 public:
@@ -21,20 +35,18 @@ public:
 		TArray<FIntPoint> CapturedRegions;
 		FIntPoint MinCaptureRegion = FIntPoint(ForceInitToZero);
 		FIntPoint MaxCaptureRegion = FIntPoint(ForceInitToZero);
-		bool bIncludesHeight = false;
 	};
 
 	FPCGLandscapeResource() = default;
 	FPCGLandscapeResource(const FResourceKey& InKey);
 	~FPCGLandscapeResource();
 
-	FPCGLandscapeTextureResource* LandscapeTextures = nullptr;
+	FPCGLandscapeTextureResource* LandscapeTexture = nullptr;
 	FVector3f LandscapeLWCTile = FVector3f::ZeroVector;
 	FMatrix ActorToWorldTransform = FMatrix::Identity;
 	FMatrix WorldToActorTransform = FMatrix::Identity;
 	FVector4 UvScaleBias = FVector4(1.0f, 1.0f, 0.0f, 0.0f);
 	FIntPoint CellCount = FIntPoint(ForceInitToZero);
-	FVector2D TextureWorldGridSize = FVector2D(1.0f, 1.0f);
 
 private:
 	FResourceKey ResourceKey;
@@ -68,9 +80,6 @@ class UPCGLandscapeDataProvider : public UComputeDataProvider
 	GENERATED_BODY()
 
 public:
-	void Init(ALandscape* InLandscape, const FBox& Bounds);
-	void InitFromLandscape(ALandscape* InLandscape, const FBox& Bounds);
-
 	//~ Begin UComputeDataProvider Interface
 	virtual void Initialize(const UComputeDataInterface* InDataInterface, UObject* InBinding, uint64 InInputMask, uint64 InOutputMask) override;
 	FComputeDataProviderRenderProxy* GetRenderProxy() override;
@@ -78,15 +87,38 @@ public:
 	//~ End UComputeDataProvider Interface
 
 private:
+	void InitFromLandscape(ALandscape* InLandscape, const FBox& Bounds);
+	void InitializeRuntimeVirtualTextures(ALandscape* InLandscape);
+	void InitializeLandscapeResource(ALandscape* InLandscape, const FBox& Bounds);
+
+private:
+	UPROPERTY()
+	TObjectPtr<const URuntimeVirtualTexture> BaseColorVirtualTexture = nullptr;
+
+	UPROPERTY()
+	TObjectPtr<const URuntimeVirtualTexture> HeightVirtualTexture = nullptr;
+
+	UPROPERTY()
+	TObjectPtr<const URuntimeVirtualTexture> NormalVirtualTexture = nullptr;
+
 	TSharedPtr<FPCGLandscapeResource> Resource;
+	bool bBaseColorSRGB = false;
+	EPCGVirtualTextureBaseColorUnpackType BaseColorVirtualTextureUnpackType = EPCGVirtualTextureBaseColorUnpackType::None;
+	ERuntimeVirtualTextureMaterialType NormalVirtualTextureMode = ERuntimeVirtualTextureMaterialType::Count;
+	FVector2D LandscapeGridSize = FVector2D(1.0f, 1.0f);
 };
 
 class FPCGLandscapeDataProviderProxy : public FComputeDataProviderRenderProxy
 {
 public:
-	FPCGLandscapeDataProviderProxy(TSharedPtr<const FPCGLandscapeResource> InResource)
-		: Resource(InResource)
-	{}
+	FPCGLandscapeDataProviderProxy(TSharedPtr<const FPCGLandscapeResource> InResource,
+		const URuntimeVirtualTexture* InBaseColorVirtualTexture,
+		const URuntimeVirtualTexture* InHeightVirtualTexture,
+		const URuntimeVirtualTexture* InNormalVirtualTexture,
+		bool bBaseColorSRGB,
+		EPCGVirtualTextureBaseColorUnpackType BaseColorVirtualTextureUnpackType,
+		ERuntimeVirtualTextureMaterialType NormalVirtualTextureMode,
+		FVector2D InWorldGridSize);
 
 	//~ Begin FComputeDataProviderRenderProxy Interface
 	bool IsValid(FValidationData const& InValidationData) const override;
@@ -97,5 +129,65 @@ public:
 protected:
 	using FParameters = FPCGLandscapeDataInterfaceParameters;
 
+	enum class ENormalUnpackType : int32
+	{
+		None = 0,
+		BC3BC3,
+		BC5BC1,
+		B5G6R5
+	};
+
+	struct FVirtualTextureLayer
+	{
+		void Initialize(const URuntimeVirtualTexture* VirtualTexture, uint32 LayerIndex, bool bSRGB);
+		void Reset();
+		bool IsValid() const;
+	
+		FShaderResourceViewRHIRef TextureSRV =  nullptr;
+		FUintVector4 TextureUniforms = FUintVector4(0, 0, 0, 0);
+	};
+	
+	struct FVirtualTexturePageTable
+	{
+		void Initialize(const URuntimeVirtualTexture* VirtualTexture, uint32 PageTableIndex, bool bIncludeWorldToUv, bool bIncludeHeightUnpack);
+		void Reset();
+		bool IsValid() const;
+	
+		FTextureRHIRef PageTableRef = nullptr;
+		FUintVector4 PageTableUniforms[2];
+		FVector4 WorldToUvParameters[4];
+	};
+
+	/** Sets all shader parameters. */
+	void SetShaderParameters(FParameters& OutShaderParameters);
+
+	/** Set shader parameters for each texture. */
+	bool SetBaseColorVirtualTextureParameters(FParameters& OutShaderParameters) const;
+	bool SetHeightVirtualTextureParameters(FParameters& OutShaderParameters) const;
+	bool SetNormalVirtualTextureParameters(FParameters& OutShaderParameters) const;
+	bool SetCollisionHeightTextureParameters(FParameters& OutShaderParameters) const;
+	
+	/** Set default values for shader parameters. Fallback in case the textures were invalid. */
+	static void SetBaseColorVirtualTextureParameters_Default(FParameters& OutShaderParameters);
+	static void SetHeightVirtualTextureParameters_Default(FParameters& OutShaderParameters);
+	static void SetNormalVirtualTextureParameters_Default(FParameters& OutShaderParameters);
+	static void SetCollisionHeightTextureParameters_Defaults(FParameters& OutShaderParameters);
+
+protected:
 	TSharedPtr<const FPCGLandscapeResource> Resource;
+
+	FVirtualTexturePageTable BaseColorVirtualPage;
+	FVirtualTextureLayer BaseColorVirtualLayer;
+	EPCGVirtualTextureBaseColorUnpackType BaseColorVirtualTextureUnpackType = EPCGVirtualTextureBaseColorUnpackType::None;
+	
+	FVirtualTexturePageTable HeightVirtualPage;
+	FVirtualTextureLayer HeightVirtualLayer;
+	
+	FVirtualTexturePageTable NormalVirtualPage;
+	FVirtualTextureLayer NormalVirtualLayer0;
+	FVirtualTextureLayer NormalVirtualLayer1;
+	ENormalUnpackType NormalUnpackMode;
+
+	FVector2D LandscapeGridSize = FVector2D(1.0f, 1.0f);
 };
+
diff --git a/Engine/Plugins/PCG/Source/PCG/Private/Compute/Elements/PCGCustomHLSL.cpp b/Engine/Plugins/PCG/Source/PCG/Private/Compute/Elements/PCGCustomHLSL.cpp
--- a/Engine/Plugins/PCG/Source/PCG/Private/Compute/Elements/PCGCustomHLSL.cpp
+++ b/Engine/Plugins/PCG/Source/PCG/Private/Compute/Elements/PCGCustomHLSL.cpp
@@ -512,7 +512,8 @@ void UPCGCustomHLSLSettings::UpdateInputDeclarations()
 
 		InputDeclarations += FString::Format(TEXT(
 			"float {0}_GetHeight(float3 WorldPos);\n"
-			"float3 {0}_GetNormal(float3 WorldPos);\n"),
+			"float3 {0}_GetNormal(float3 WorldPos);\n"
+			"float3 {0}_GetBaseColor(float3 WorldPos);\n"),
 			{ bMultiPin ? PCGHLSLElement::PinDeclTemplateStr : LandscapeDataPins[0] });
 
 		InputDeclarations += TEXT("\n");
