[PCG] Fix non unity error from initializing TRefCountPtr in header

#rnx

diff --git a/Engine/Plugins/PCG/Source/PCG/Private/Compute/DataInterfaces/PCGLandscapeDataInterface.h b/Engine/Plugins/PCG/Source/PCG/Private/Compute/DataInterfaces/PCGLandscapeDataInterface.h
--- a/Engine/Plugins/PCG/Source/PCG/Private/Compute/DataInterfaces/PCGLandscapeDataInterface.h
+++ b/Engine/Plugins/PCG/Source/PCG/Private/Compute/DataInterfaces/PCGLandscapeDataInterface.h
@@ -143,7 +143,7 @@ protected:
 		void Reset();
 		bool IsValid() const;
 	
-		FShaderResourceViewRHIRef TextureSRV =  nullptr;
+		FShaderResourceViewRHIRef TextureSRV;
 		FUintVector4 TextureUniforms = FUintVector4(0, 0, 0, 0);
 	};
 	
@@ -153,7 +153,7 @@ protected:
 		void Reset();
 		bool IsValid() const;
 	
-		FTextureRHIRef PageTableRef = nullptr;
+		FTextureRHIRef PageTableRef;
 		FUintVector4 PageTableUniforms[2];
 		FVector4 WorldToUvParameters[4];
 	};
