Take two on fixing the policy in BP PCG components, as the refresh must be called only when there's a transition.

#rnx

diff --git a/Engine/Plugins/PCG/Source/PCG/Private/PCGComponent.cpp b/Engine/Plugins/PCG/Source/PCG/Private/PCGComponent.cpp
--- a/Engine/Plugins/PCG/Source/PCG/Private/PCGComponent.cpp
+++ b/Engine/Plugins/PCG/Source/PCG/Private/PCGComponent.cpp
@@ -503,11 +503,8 @@ void UPCGComponent::SetPropertiesFromOriginal(const UPCGComponent* Original)
 		GraphInstance->CopyParameterOverrides(Original->GraphInstance);
 	}
 
-	if (!SchedulingPolicy || SchedulingPolicyClass != Original->SchedulingPolicyClass)
-	{
-		SchedulingPolicyClass = Original->SchedulingPolicyClass;
-		RefreshSchedulingPolicy();
-	}
+	SchedulingPolicyClass = Original->SchedulingPolicyClass;
+	RefreshSchedulingPolicy();
 
 	if (SchedulingPolicy && ensure(Original->SchedulingPolicy) && !SchedulingPolicy->IsEquivalent(Original->SchedulingPolicy))
 	{
@@ -1262,15 +1259,8 @@ void UPCGComponent::ClearPerPinGeneratedOutput()
 
 void UPCGComponent::SetSchedulingPolicyClass(TSubclassOf<UPCGSchedulingPolicyBase> InSchedulingPolicyClass) 
 {
-	if (!SchedulingPolicy || InSchedulingPolicyClass != SchedulingPolicyClass)
-	{
-		if (InSchedulingPolicyClass != SchedulingPolicyClass)
-		{
-			SchedulingPolicyClass = InSchedulingPolicyClass;
-		}
-		
-		RefreshSchedulingPolicy();
-	}
+	SchedulingPolicyClass = InSchedulingPolicyClass;
+	RefreshSchedulingPolicy();
 }
 
 double UPCGComponent::GetGenerationRadiusFromGrid(EPCGHiGenGrid Grid) const
@@ -2331,12 +2321,12 @@ void UPCGComponent::PostEditChangeProperty(FPropertyChangedEvent& PropertyChange
 	} 
 	else if (PropName == GET_MEMBER_NAME_CHECKED(UPCGComponent, GenerationTrigger))
 	{
+		// Implementation detail - if we get rid of the scheduling policy here, we'll trigger an editor error, so just create the policy if needed
 		if (!SchedulingPolicy)
 		{
 			RefreshSchedulingPolicy();
 		}
-
-		if (SchedulingPolicy)
+		else
 		{
 			// We should only display scheduling policy parameters when in runtime generation mode.
 			SchedulingPolicy->SetShouldDisplayProperties(IsManagedByRuntimeGenSystem());
@@ -3241,30 +3231,37 @@ FPCGDataCollection UPCGComponent::CreateActorPCGDataCollection(AActor* Actor, co
 
 void UPCGComponent::RefreshSchedulingPolicy()
 {
-	// Only delete it if we are the owner, it's for deprecation where local components had hard ref on original policy.
-	if (IsValid(SchedulingPolicy) && SchedulingPolicy->GetOuter() == this)
+	if (SchedulingPolicy && (!SchedulingPolicyClass || SchedulingPolicy->GetClass() != SchedulingPolicyClass || !IsManagedByRuntimeGenSystem()))
 	{
+		// Only delete it if we are the owner, it's for deprecation where local components had hard ref on original policy.
+		if (IsValid(SchedulingPolicy) && SchedulingPolicy->GetOuter() == this)
+		{
 #if WITH_EDITOR
-        // We are renaming to a new outer on an object that may still be loading. Since we are destroying this object
-        // pass REN_AllowPackageLinkerMismatch to allow the linker to remain on the object so we don't have to force a load 
-        // to complete before the rename 
-		SchedulingPolicy->Rename(nullptr, GetTransientPackage(), REN_DontCreateRedirectors | REN_AllowPackageLinkerMismatch);
+			// We are renaming to a new outer on an object that may still be loading. Since we are destroying this object
+			// pass REN_AllowPackageLinkerMismatch to allow the linker to remain on the object so we don't have to force a load 
+			// to complete before the rename 
+			SchedulingPolicy->Rename(nullptr, GetTransientPackage(), REN_DontCreateRedirectors | REN_AllowPackageLinkerMismatch);
 #endif
-		SchedulingPolicy->MarkAsGarbage();
-	}
 
-	SchedulingPolicy = nullptr;
+			SchedulingPolicy->MarkAsGarbage();
+		}
+
+		SchedulingPolicy = nullptr;
+	}
 
 	// We should never create the scheduling policy when not in runtime generation mode.
-	if (SchedulingPolicyClass && IsManagedByRuntimeGenSystem())
+	if (SchedulingPolicyClass && !SchedulingPolicy && IsManagedByRuntimeGenSystem())
 	{
 		const EObjectFlags Flags = GetMaskedFlags(RF_PropagateToSubObjects);
 		SchedulingPolicy = NewObject<UPCGSchedulingPolicyBase>(this, SchedulingPolicyClass, NAME_None, Flags);
+	}
 
 #if WITH_EDITOR
+	if (SchedulingPolicy)
+	{
 		SchedulingPolicy->SetShouldDisplayProperties(IsManagedByRuntimeGenSystem());
-#endif
 	}
+#endif
 }
 
 UPCGData* UPCGComponent::CreatePCGData() const
@@ -4097,7 +4094,7 @@ void FPCGComponentInstanceData::ApplyToComponent(UActorComponent* Component, con
 
 		bool bDoActorMapping = PCGComponent->bGenerated || PCGHelpers::IsRuntimeOrPIE();
 
-		// If the generation mode is set to runtime, we won't receive a PostEditChange event, because this is a BP. Since it will not trigger a policy update, do it here.
+		// If the generation mode or the policy class is changed, we won't receive a PostEditChange event, because this is a BP. Since it will not trigger a policy update, do it here.
 		PCGComponent->RefreshSchedulingPolicy();
 
 		// Also remap
