[PCG] Fix missing CancelNextTask handling of DynamicDependencies (when pcg.Experimental.EnableDynamicDependencies is enabled)

#rnx

diff --git a/Engine/Plugins/PCG/Source/PCG/Private/Graph/PCGGraphExecutor.cpp b/Engine/Plugins/PCG/Source/PCG/Private/Graph/PCGGraphExecutor.cpp
--- a/Engine/Plugins/PCG/Source/PCG/Private/Graph/PCGGraphExecutor.cpp
+++ b/Engine/Plugins/PCG/Source/PCG/Private/Graph/PCGGraphExecutor.cpp
@@ -1913,7 +1913,27 @@ bool FPCGGraphExecutor::CancelNextTasks(FPCGTaskId CancelledTask, TSet<UPCGCompo
 	{
 		bAddedComponents |= CancelNextTasks(Successor, OutCancelledComponents);
 	}
-		
+
+	{
+		PCGGraphExecutor::TScopeLock PauseScopeLock(PausedTaskSuccessorsLock);
+
+		TArray<TSharedPtr<FPCGGraphActiveTask>> ActiveTaskSuccessors;
+
+		if (PausedTaskSuccessors.RemoveAndCopyValue(CancelledTask, ActiveTaskSuccessors))
+		{
+			for (const TSharedPtr<FPCGGraphActiveTask>& ActiveTask : ActiveTaskSuccessors)
+			{
+				if (!OutCancelledComponents.Contains(ActiveTask->Context->SourceComponent.Get()))
+				{
+					OutCancelledComponents.Add(ActiveTask->Context->SourceComponent.Get());
+					bAddedComponents = true;
+				}
+
+				// Do not call abort as this is an ActiveTask and will be processed by the Cancellation of it's component
+			}
+		}
+	}
+
 	// Tasks cancelled might have an impact on scheduled-but-not-processed tasks
 	{
 		PCGGraphExecutor::TScopeLock ScopeLock(ScheduleLock);
