[PCG] Fix runtime gen scheduler failing to tick while there are components for deferred cleanup

#rnx

diff --git a/Engine/Plugins/PCG/Source/PCG/Private/RuntimeGen/PCGRuntimeGenScheduler.cpp b/Engine/Plugins/PCG/Source/PCG/Private/RuntimeGen/PCGRuntimeGenScheduler.cpp
--- a/Engine/Plugins/PCG/Source/PCG/Private/RuntimeGen/PCGRuntimeGenScheduler.cpp
+++ b/Engine/Plugins/PCG/Source/PCG/Private/RuntimeGen/PCGRuntimeGenScheduler.cpp
@@ -231,7 +231,7 @@ bool FPCGRuntimeGenScheduler::ShouldTick()
 	}
 
 	// We can stop ticking if there are no runtime gen components alive and there are no generated components that need cleaning up.
-	if (!bAnyRuntimeGenComponentsExist && GeneratedComponents.IsEmpty())
+	if (!bAnyRuntimeGenComponentsExist && GeneratedComponents.IsEmpty() && GeneratedComponentsToRemove.IsEmpty())
 	{
 		return false;
 	}
