[PCG] Fixed a crash in PCG UI when connecting/disconnecting an edge of a PCG node would sometimes cause a reconstruct during the add/remove pin operation

#rnx

diff --git a/Engine/Plugins/PCG/Source/PCG/Private/PCGGraph.cpp b/Engine/Plugins/PCG/Source/PCG/Private/PCGGraph.cpp
--- a/Engine/Plugins/PCG/Source/PCG/Private/PCGGraph.cpp
+++ b/Engine/Plugins/PCG/Source/PCG/Private/PCGGraph.cpp
@@ -989,7 +989,7 @@ bool UPCGGraph::AddLabeledEdge(UPCGNode* From, const FName& FromPinLabel, UPCGNo
 		bToPinBrokeOtherEdges = ToPin->BreakAllIncompatibleEdges(&TouchedNodes);
 	}
 
-	const EPCGChangeType ChangeType = PCGGraphUtils::NotifyTouchedNodes(TouchedNodes, EPCGChangeType::Structural);
+	const EPCGChangeType ChangeType = PCGGraphUtils::NotifyTouchedNodes(TouchedNodes, EPCGChangeType::Structural) | EPCGChangeType::Edge;
 
 #if WITH_EDITOR
 	// After all nodes are notified, re-enable graph notifications and send graph change notification.
@@ -1135,7 +1135,7 @@ bool UPCGGraph::RemoveEdge(UPCGNode* From, const FName& FromLabel, UPCGNode* To,
 		OutPin->BreakEdgeTo(InPin, &TouchedNodes);
 	}
 
-	const EPCGChangeType ChangeType = PCGGraphUtils::NotifyTouchedNodes(TouchedNodes, EPCGChangeType::Structural);
+	const EPCGChangeType ChangeType = PCGGraphUtils::NotifyTouchedNodes(TouchedNodes, EPCGChangeType::Structural) | EPCGChangeType::Edge;
 
 #if WITH_EDITOR
 	// After all nodes are notified, re-enable graph notifications and send graph change notification.
diff --git a/Engine/Plugins/PCG/Source/PCGEditor/Private/Nodes/PCGEditorGraphNodeBase.cpp b/Engine/Plugins/PCG/Source/PCGEditor/Private/Nodes/PCGEditorGraphNodeBase.cpp
--- a/Engine/Plugins/PCG/Source/PCGEditor/Private/Nodes/PCGEditorGraphNodeBase.cpp
+++ b/Engine/Plugins/PCG/Source/PCGEditor/Private/Nodes/PCGEditorGraphNodeBase.cpp
@@ -1364,11 +1364,6 @@ void UPCGEditorGraphNodeBase::PinConnectionListChanged(UEdGraphPin* Pin)
 
 	Super::PinConnectionListChanged(Pin);
 
-	if (UPCGGraph* PCGGraph = PCGNode->GetGraph())
-	{
-		PCGGraph->OnGraphChangedDelegate.Broadcast(PCGGraph, EPCGChangeType::Edge);
-	}
-
 	// One-time UX functionality, where when a user connects a node with default values, it will activate them for all other pins
 	if (!bHasEverBeenConnected && Pin && !Pin->IsPendingKill())
 	{
diff --git a/Engine/Plugins/PCG/Source/PCGEditor/Private/PCGEditor.cpp b/Engine/Plugins/PCG/Source/PCGEditor/Private/PCGEditor.cpp
--- a/Engine/Plugins/PCG/Source/PCGEditor/Private/PCGEditor.cpp
+++ b/Engine/Plugins/PCG/Source/PCGEditor/Private/PCGEditor.cpp
@@ -2998,7 +2998,11 @@ void FPCGEditor::OnGraphChanged(UPCGGraphInterface* InGraph, EPCGChangeType Chan
 	{
 		for (TSharedPtr<SPCGEditorGraphDetailsView> Widget : PropertyDetailsWidgets)
 		{
-			Widget->GetDetailsView()->ForceRefresh();
+			if (Widget.IsValid() && Widget->GetVisibility() == EVisibility::Visible)
+			{
+				const TSharedPtr<IDetailsView> DetailsViewPtr = Widget->GetDetailsView();
+				DetailsViewPtr->ForceRefresh();
+			}
 		}
 	}
 }
