Add built-in analysis to SM Spawner, so that instance counts will be set correctly without needing to use the Attribute Partition node.

The GPU Attribute Partition currently allocates a sizeable transient buffer which manifested as an expensive resource allocation on the render thread, which this will also avoid for the GPU SM Spawner.

* Adds a generic, re-usable kernel for counting string key values as it seems this will be useful for other things (and didn't benefit from being SM Spawner specific).
* The SM Spawner emits this new counting kernel and uses this to determine instance counts.
* The SM Spawner reads back the analysis data to CPU as part of primitive setup. This primitive setup is finally moved out of the ComputeGraphElement.
* GPU graph splits turned out to be broken (requested splits did not always work as flood fill would fill around the splitting pin). Moved the split functionality downstream to where the other type of splits was performed.
* Added atomic increment shader API for the counter to use. For this to work, also had to add support for 0-initializing buffer data, otherwise was adding to unitialized values.

#rnx

diff --git a/Engine/Plugins/PCG/Shaders/Private/BuiltInKernels/PCGCountUniqueAttributeValues.usf b/Engine/Plugins/PCG/Shaders/Private/BuiltInKernels/PCGCountUniqueAttributeValues.usf
new file mode 100644
--- /dev/null
+++ b/Engine/Plugins/PCG/Shaders/Private/BuiltInKernels/PCGCountUniqueAttributeValues.usf
@@ -0,0 +1,59 @@
+// Copyright Epic Games, Inc. All Rights Reserved.
+
+// Counts how many times each unique value of an attribute is encountered in the input data elements.
+
+// Input: Input elements with attribute.
+// Output: Attribute set with a ValueCount attribute which has the count of each value. Number of elements is set to the number of unique values (known CPU-side).
+
+#include "/Engine/Private/ComputeShaderUtils.ush"
+#include "/Engine/Private/Quaternion.ush"
+#include "/Plugin/PCG/Private/PCGShaderUtils.ush"
+#include "/Plugin/PCG/Private/PCGShaderUtilsInternal.ush"
+
+// TODO: should be code-gen'd from PCGComputeCommon.h
+#define MAX_NUM_ATTRS 128
+#define NUM_RESERVED_ATTRS 32
+
+[numthreads(64, 1, 1)]
+void PCGCountUniqueAttributeValuesCS(uint3 GroupId : SV_GroupID, uint GroupIndex : SV_GroupIndex)
+{
+	// Mark the kernel as having executed. Must run before we early out via thread index, because the kernel is still 'executed' even if the number of
+	// threads to iterate on is zero. Even if GetNumThreads() returns 0, the kernel will still have been dispatched on a single thread to set this flag.
+	if (all(GroupId == 0) && GroupIndex == 0)
+	{
+		Out_SetAsExecutedInternal();
+	}
+
+	const uint ThreadIndex = GetUnWrappedDispatchThreadId(GroupId, GroupIndex, 64);
+
+	uint In_DataIndex, In_ElementIndex;
+	if (!In_GetThreadData(ThreadIndex, In_DataIndex, In_ElementIndex))
+	{
+		return;
+	}
+
+	if (In_IsPointRemoved(In_DataIndex, In_ElementIndex))
+	{
+		return;
+	}
+
+	const int AttributeToCountId = CountUniqueValues_GetAttributeToCountId();
+
+	const int AttributeValue = In_GetStringKey(In_DataIndex, In_ElementIndex, AttributeToCountId);
+
+	const uint Out_DataIndex = 0;
+	const uint Out_ElementIndex = CountUniqueValues_GetValueIndex(AttributeValue);
+
+	// Write value of attribute.
+	const int OutputValueAttributeId = CountUniqueValues_GetOutputValueAttributeId();
+	Out_SetInt(Out_DataIndex, Out_ElementIndex, OutputValueAttributeId, AttributeValue);
+
+	const int OutputCountAttributeId = CountUniqueValues_GetOutputCountAttributeId();
+	const int ValueIndex = CountUniqueValues_GetValueIndex(AttributeValue);
+
+	// Increment count of this attribute.
+	if (ValueIndex != -1)
+	{
+		Out_AtomicAddInt(Out_DataIndex, ValueIndex, OutputCountAttributeId, 1);
+	}
+}
diff --git a/Engine/Plugins/PCG/Shaders/Private/PCGDataCollectionDataInterface.ush b/Engine/Plugins/PCG/Shaders/Private/PCGDataCollectionDataInterface.ush
--- a/Engine/Plugins/PCG/Shaders/Private/PCGDataCollectionDataInterface.ush
+++ b/Engine/Plugins/PCG/Shaders/Private/PCGDataCollectionDataInterface.ush
@@ -345,6 +345,19 @@ void SetName_{DataInterfaceName}(uint DataIndex, uint ElementIndex, int Attribut
 	SetInt_{DataInterfaceName}(DataIndex, ElementIndex, AttributeId, (int)Value.x);
 }
 
+// #################### ATOMICS ##########################
+
+// Returns value of attribute before incrementing.
+int AtomicAddInt_{DataInterfaceName}(uint DataIndex, uint ElementIndex, int AttributeId, int ValueToAdd)
+{
+	const uint AttributeAddress = GetFirstElementAddressInternal_{DataInterfaceName}(DataIndex, AttributeId);
+
+	uint OriginalValue;
+	{DataInterfaceName}_DataCollectionBuffer.InterlockedAdd(AttributeAddress + ElementIndex * 1 * 4, (uint)ValueToAdd, OriginalValue);
+
+	return (int)OriginalValue;
+}
+
 // #################### POINT ATTRIBUTE GETTERS ##########################
 
 float3 GetPosition_{DataInterfaceName}(uint DataIndex, uint ElementIndex)
diff --git a/Engine/Plugins/PCG/Source/PCG/Private/Compute/BuiltInKernels/PCGCountUniqueAttributeValuesKernel.cpp b/Engine/Plugins/PCG/Source/PCG/Private/Compute/BuiltInKernels/PCGCountUniqueAttributeValuesKernel.cpp
new file mode 100644
--- /dev/null
+++ b/Engine/Plugins/PCG/Source/PCG/Private/Compute/BuiltInKernels/PCGCountUniqueAttributeValuesKernel.cpp
@@ -0,0 +1,180 @@
+// Copyright Epic Games, Inc. All Rights Reserved.
+
+#include "Compute/BuiltInKernels/PCGCountUniqueAttributeValuesKernel.h"
+
+#include "PCGComponent.h"
+#include "PCGContext.h"
+#include "PCGModule.h"
+#include "Compute/PCGComputeGraph.h"
+#include "Compute/PCGDataBinding.h"
+#include "Compute/PCGPinPropertiesGPU.h"
+#include "Compute/DataInterfaces/BuiltInKernels/PCGCountUniqueAttributeValuesDataInterface.h"
+#include "Compute/Elements/PCGComputeGraphElement.h"
+#include "Graph/PCGGPUGraphCompilationContext.h"
+
+#include "ShaderCompilerCore.h"
+
+#include UE_INLINE_GENERATED_CPP_BY_NAME(PCGCountUniqueAttributeValuesKernel)
+
+#define LOCTEXT_NAMESPACE "PCGCountUniqueAttributeValuesKernel"
+
+void UPCGCountUniqueAttributeValuesKernel::Initialize(FName InAttributeName, bool bLogDescriptions)
+{
+	Super::Initialize(bLogDescriptions);
+
+	AttributeName = InAttributeName;
+}
+
+bool UPCGCountUniqueAttributeValuesKernel::IsKernelValid(FPCGContext* InContext, bool bQuiet) const
+{
+	TRACE_CPUPROFILER_EVENT_SCOPE(UPCGCountUniqueAttributeValuesKernel::IsKernelValid);
+
+	if (!Super::IsKernelValid(InContext, bQuiet))
+	{
+		return false;
+	}
+
+	if (!InContext || !InContext->IsComputeContext())
+	{
+		return true;
+	}
+
+	const FPCGComputeGraphContext* Context = static_cast<FPCGComputeGraphContext*>(InContext);
+	const UPCGDataBinding* DataBinding = Context->DataBinding.Get();
+	if (DataBinding)
+	{
+		FPCGDataCollectionDesc InputDataDesc;
+		ensure(DataBinding->GetKernelPinDataDesc_GameThread(this, PCGPinConstants::DefaultInputLabel, /*bIsInputPin=*/true, InputDataDesc));
+
+		FPCGKernelAttributeDesc AttributeDesc;
+		bool bConflictingTypesFound = false;
+
+		if (!InputDataDesc.GetAttributeDesc(AttributeName, AttributeDesc, bConflictingTypesFound))
+		{
+			if (!InputDataDesc.DataDescs.IsEmpty())
+			{
+				PCG_KERNEL_VALIDATION_ERR(InContext, /*Settings=*/nullptr, bQuiet, FText::Format(
+					LOCTEXT("AttributeMissing", "Count attribute '{0}' not found, this attribute must be present on all input data, and be of type String Key."),
+					FText::FromName(AttributeName)));
+			}
+
+			return false;
+		}
+
+		if (bConflictingTypesFound)
+		{
+			PCG_KERNEL_VALIDATION_ERR(InContext, /*Settings=*/nullptr, bQuiet, FText::Format(
+				LOCTEXT("AttributeHasMultipleTypes", "Count attribute '{0}' found with multiple types in input data, all attributes must be of type String Key."),
+				FText::FromName(AttributeName)));
+
+			return false;
+		}
+
+		if (AttributeDesc.Type != EPCGKernelAttributeType::StringKey)
+		{
+			// Attribute value counting only currently supported for attributes of type StringKey.
+			PCG_KERNEL_VALIDATION_ERR(Context, /*Settings=*/nullptr, bQuiet, FText::Format(
+				LOCTEXT("AttributeTypeInvalid", "Cannot count values for attribute '{0}', only attributes of type String Key are currently supported."),
+				FText::FromName(AttributeName)));
+
+			return false;
+		}
+	}
+
+	return true;
+}
+
+bool UPCGCountUniqueAttributeValuesKernel::ComputeOutputBindingDataDesc(const UPCGComputeGraph* InGraph, FName InOutputPinLabel, const UPCGDataBinding* InBinding, FPCGDataCollectionDesc& OutDataDesc) const
+{
+	check(InGraph);
+	check(InBinding);
+
+	// Code assumes single output pin.
+	if (!ensure(InOutputPinLabel == PCGPinConstants::DefaultOutputLabel))
+	{
+		return false;
+	}
+
+	FPCGDataCollectionDesc InputDesc;
+	ensure(InBinding->GetKernelPinDataDesc_GameThread(this, PCGPinConstants::DefaultInputLabel, /*bIsInput=*/true, InputDesc));
+
+	const int32 AttributeId = InBinding->GetAttributeId(AttributeName, EPCGKernelAttributeType::StringKey);
+
+	if (AttributeId != INDEX_NONE)
+	{
+		TArray<int32> UniqueStringKeyValues;
+		InputDesc.GetUniqueStringKeyValues(AttributeId, UniqueStringKeyValues);
+
+		FPCGDataDesc& OutputDataDesc = OutDataDesc.DataDescs.Emplace_GetRef();
+		OutputDataDesc.Type = EPCGDataType::Param;
+		OutputDataDesc.ElementCount = UniqueStringKeyValues.Num();
+
+		OutputDataDesc.AddAttribute(FPCGKernelAttributeKey(PCGCountUniqueAttributeValuesConstants::ValueAttributeName, EPCGKernelAttributeType::Int), InBinding);
+		OutputDataDesc.AddAttribute(FPCGKernelAttributeKey(PCGCountUniqueAttributeValuesConstants::ValueCountAttributeName, EPCGKernelAttributeType::Int), InBinding);
+	}
+
+	return true;
+}
+
+int UPCGCountUniqueAttributeValuesKernel::ComputeThreadCount(const UPCGDataBinding* InBinding) const
+{
+	int ThreadCount = 0;
+
+	// Execute a thread for every element in the input data.
+	FPCGDataCollectionDesc OutputPinDesc;
+	if (ensure(InBinding->GetKernelPinDataDesc_GameThread(this, PCGPinConstants::DefaultInputLabel, /*bIsInput=*/true, OutputPinDesc)))
+	{
+		for (const FPCGDataDesc& DataDesc : OutputPinDesc.DataDescs)
+		{
+			ThreadCount += DataDesc.ElementCount;
+		}
+	}
+
+	return ThreadCount;
+}
+
+bool UPCGCountUniqueAttributeValuesKernel::DoesOutputPinRequireZeroInitialization(FName InOutputPinLabel) const
+{
+	// We are atomically incrementing the values on the output so we need to ensure the values are 0-initialized.
+	return InOutputPinLabel == PCGPinConstants::DefaultOutputLabel;
+}
+
+#if WITH_EDITOR
+FString UPCGCountUniqueAttributeValuesKernel::GetCookedSource(const UPCGComputeGraph* InComputeGraph) const
+{
+	FString TemplateFile;
+	ensure(LoadShaderSourceFile(TEXT("/Plugin/PCG/Private/BuiltInKernels/PCGCountUniqueAttributeValues.usf"), EShaderPlatform::SP_PCD3D_SM5, &TemplateFile, nullptr));
+	return TemplateFile;
+}
+
+void UPCGCountUniqueAttributeValuesKernel::CreateAdditionalInputDataInterfaces(FPCGGPUCompilationContext& InOutContext, UObject* InObjectOuter, TArray<TObjectPtr<UComputeDataInterface>>& OutDataInterfaces) const
+{
+	Super::CreateAdditionalInputDataInterfaces(InOutContext, InObjectOuter, OutDataInterfaces);
+
+	TObjectPtr<UPCGCountUniqueAttributeValuesDataInterface> KernelDI = InOutContext.NewObject_AnyThread<UPCGCountUniqueAttributeValuesDataInterface>(InObjectOuter);
+	KernelDI->ProducerKernel = this;
+	KernelDI->AttributeToCountName = AttributeName;
+
+	OutDataInterfaces.Add(KernelDI);
+}
+#endif
+
+void UPCGCountUniqueAttributeValuesKernel::GetKernelAttributeKeys(TArray<FPCGKernelAttributeKey>& OutKeys) const
+{
+	// Register the attribute this node creates.
+	OutKeys.AddUnique(FPCGKernelAttributeKey(PCGCountUniqueAttributeValuesConstants::ValueAttributeName, EPCGKernelAttributeType::Int));
+	OutKeys.AddUnique(FPCGKernelAttributeKey(PCGCountUniqueAttributeValuesConstants::ValueCountAttributeName, EPCGKernelAttributeType::Int));
+}
+
+void UPCGCountUniqueAttributeValuesKernel::GetInputPins(TArray<FPCGPinProperties>& OutPins) const
+{
+	OutPins.Emplace(PCGPinConstants::DefaultInputLabel, EPCGDataType::Point);
+}
+
+void UPCGCountUniqueAttributeValuesKernel::GetOutputPins(TArray<FPCGPinPropertiesGPU>& OutPins) const
+{
+	// Attribute set with a value count attribute, element count equal to number of unique values of the counted attribute.
+	OutPins.Emplace(PCGPinConstants::DefaultOutputLabel, EPCGDataType::Param);
+}
+
+#undef LOCTEXT_NAMESPACE
diff --git a/Engine/Plugins/PCG/Source/PCG/Private/Compute/BuiltInKernels/PCGCountUniqueAttributeValuesKernel.h b/Engine/Plugins/PCG/Source/PCG/Private/Compute/BuiltInKernels/PCGCountUniqueAttributeValuesKernel.h
new file mode 100644
--- /dev/null
+++ b/Engine/Plugins/PCG/Source/PCG/Private/Compute/BuiltInKernels/PCGCountUniqueAttributeValuesKernel.h
@@ -0,0 +1,42 @@
+// Copyright Epic Games, Inc. All Rights Reserved.
+
+#pragma once
+
+#include "Compute/PCGComputeKernel.h"
+
+#include "PCGCountUniqueAttributeValuesKernel.generated.h"
+
+namespace PCGCountUniqueAttributeValuesConstants
+{
+	const FName ValueAttributeName = TEXT("UniqueValue");
+	const FName ValueCountAttributeName = TEXT("UniqueValueCount");
+}
+
+UCLASS()
+class UPCGCountUniqueAttributeValuesKernel : public UPCGComputeKernel
+{
+	GENERATED_BODY()
+
+public:
+	virtual void Initialize(FName InAttributeName, bool bLogDescriptions);
+
+	virtual bool IsKernelValid(FPCGContext* InContext = nullptr, bool bQuiet = true) const override;
+	virtual bool ComputeOutputBindingDataDesc(const UPCGComputeGraph* InGraph, FName InOutputPinLabel, const UPCGDataBinding* InBinding, FPCGDataCollectionDesc& OutDataDesc) const override;
+	virtual int ComputeThreadCount(const UPCGDataBinding* Binding) const override;
+	virtual bool DoesOutputPinRequireZeroInitialization(FName InOutputPinLabel) const override;
+#if WITH_EDITOR
+	virtual FString GetCookedSource(const UPCGComputeGraph* InComputeGraph) const override;
+	virtual FString GetEntryPoint() const override { return TEXT("PCGCountUniqueAttributeValuesCS"); }
+	void CreateAdditionalInputDataInterfaces(FPCGGPUCompilationContext& InOutContext, UObject* InObjectOuter, TArray<TObjectPtr<UComputeDataInterface>>& OutDataInterfaces) const override;
+	// Split graph to read back analysis results.
+	virtual bool SplitGraphAtOutput() const override { return true; }
+#endif
+	virtual void GetKernelAttributeKeys(TArray<FPCGKernelAttributeKey>& OutKeys) const override;
+
+	virtual void GetInputPins(TArray<FPCGPinProperties>& OutPins) const override;
+	virtual void GetOutputPins(TArray<FPCGPinPropertiesGPU>& OutPins) const override;
+
+public:
+	UPROPERTY()
+	FName AttributeName;
+};
diff --git a/Engine/Plugins/PCG/Source/PCG/Private/Compute/DataInterfaces/BuiltInKernels/PCGCountUniqueAttributeValuesDataInterface.cpp b/Engine/Plugins/PCG/Source/PCG/Private/Compute/DataInterfaces/BuiltInKernels/PCGCountUniqueAttributeValuesDataInterface.cpp
new file mode 100644
--- /dev/null
+++ b/Engine/Plugins/PCG/Source/PCG/Private/Compute/DataInterfaces/BuiltInKernels/PCGCountUniqueAttributeValuesDataInterface.cpp
@@ -0,0 +1,205 @@
+// Copyright Epic Games, Inc. All Rights Reserved.
+
+#include "PCGCountUniqueAttributeValuesDataInterface.h"
+
+#include "Compute/PCGComputeCommon.h"
+#include "Compute/PCGDataBinding.h"
+#include "Compute/BuiltInKernels/PCGCountUniqueAttributeValuesKernel.h"
+#include "Elements/Metadata/PCGMetadataPartition.h"
+#include "Elements/Metadata/PCGMetadataPartitionKernel.h"
+
+#include "GlobalRenderResources.h"
+#include "RenderGraphBuilder.h"
+#include "RenderGraphResources.h"
+#include "RHIResources.h"
+#include "ShaderCompilerCore.h"
+#include "ShaderCore.h"
+#include "ShaderParameterMetadataBuilder.h"
+#include "ComputeFramework/ShaderParamTypeDefinition.h"
+
+void UPCGCountUniqueAttributeValuesDataInterface::GetSupportedInputs(TArray<FShaderFunctionDefinition>& OutFunctions) const
+{
+	OutFunctions.AddDefaulted_GetRef()
+		.SetName(TEXT("CountUniqueValues_GetAttributeToCountId"))
+		.AddReturnType(EShaderFundamentalType::Int);
+
+	OutFunctions.AddDefaulted_GetRef()
+		.SetName(TEXT("CountUniqueValues_GetOutputCountAttributeId"))
+		.AddReturnType(EShaderFundamentalType::Int);
+
+	OutFunctions.AddDefaulted_GetRef()
+		.SetName(TEXT("CountUniqueValues_GetOutputValueAttributeId"))
+		.AddReturnType(EShaderFundamentalType::Int);
+
+	OutFunctions.AddDefaulted_GetRef()
+		.SetName(TEXT("CountUniqueValues_GetValueIndex"))
+		.AddReturnType(EShaderFundamentalType::Uint)
+		.AddParam(EShaderFundamentalType::Uint); // InAttributeValue
+}
+
+BEGIN_SHADER_PARAMETER_STRUCT(FPCGCountUniqueAttributeValuesDataInterfaceParameters,)
+	SHADER_PARAMETER(int32, AttributeToCountId)
+	SHADER_PARAMETER(int32, OutputValueAttributeId)
+	SHADER_PARAMETER(int32, OutputCountAttributeId)
+	SHADER_PARAMETER_RDG_BUFFER_SRV(StructuredBuffer<uint32>, AttributeValueToValueIndex)
+END_SHADER_PARAMETER_STRUCT()
+
+void UPCGCountUniqueAttributeValuesDataInterface::GetShaderParameters(TCHAR const* UID, FShaderParametersMetadataBuilder& InOutBuilder, FShaderParametersMetadataAllocations& InOutAllocations) const
+{
+	InOutBuilder.AddNestedStruct<FPCGCountUniqueAttributeValuesDataInterfaceParameters>(UID);
+}
+
+void UPCGCountUniqueAttributeValuesDataInterface::GetHLSL(FString& OutHLSL, FString const& InDataInterfaceName) const
+{
+	TMap<FString, FStringFormatArg> TemplateArgs =
+	{
+		{ TEXT("DataInterfaceName"), InDataInterfaceName },
+	};
+
+	OutHLSL += FString::Format(TEXT(
+		"StructuredBuffer<uint> {DataInterfaceName}_AttributeValueToValueIndex;\n"
+		"int {DataInterfaceName}_AttributeToCountId;\n"
+		"int {DataInterfaceName}_OutputValueAttributeId;\n"
+		"int {DataInterfaceName}_OutputCountAttributeId;\n"
+		"\n"
+		"int CountUniqueValues_GetAttributeToCountId_{DataInterfaceName}()\n"
+		"{\n"
+		"	return {DataInterfaceName}_AttributeToCountId;\n"
+		"}\n"
+		"\n"
+		"int CountUniqueValues_GetOutputValueAttributeId_{DataInterfaceName}()\n"
+		"{\n"
+		"	return {DataInterfaceName}_OutputValueAttributeId;\n"
+		"}\n"
+		"\n"
+		"int CountUniqueValues_GetOutputCountAttributeId_{DataInterfaceName}()\n"
+		"{\n"
+		"	return {DataInterfaceName}_OutputCountAttributeId;\n"
+		"}\n"
+		"\n"
+		"uint CountUniqueValues_GetValueIndex_{DataInterfaceName}(int InAttributeValue)\n"
+		"{\n"
+		"	return {DataInterfaceName}_AttributeValueToValueIndex[InAttributeValue];\n"
+		"}\n"
+	), TemplateArgs);
+}
+
+UComputeDataProvider* UPCGCountUniqueAttributeValuesDataInterface::CreateDataProvider() const
+{
+	return NewObject<UPCGCountUniqueAttributeValuesDataProvider>();
+}
+
+void UPCGCountUniqueAttributeValuesDataProvider::Initialize(const UComputeDataInterface* InDataInterface, UObject* InBinding, uint64 InInputMask, uint64 InOutputMask)
+{
+	TRACE_CPUPROFILER_EVENT_SCOPE(UPCGCountUniqueAttributeValuesDataProvider::Initialize);
+
+	Super::Initialize(InDataInterface, InBinding, InInputMask, InOutputMask);
+
+	const UPCGCountUniqueAttributeValuesDataInterface* DataInterface = CastChecked<UPCGCountUniqueAttributeValuesDataInterface>(InDataInterface);
+	Kernel = DataInterface->ProducerKernel;
+
+	AttributeToCountName = DataInterface->AttributeToCountName;
+}
+
+bool UPCGCountUniqueAttributeValuesDataProvider::PrepareForExecute_GameThread(UPCGDataBinding* InBinding)
+{
+	TRACE_CPUPROFILER_EVENT_SCOPE(UPCGCountUniqueAttributeValuesDataInterface::PrepareForExecute_GameThread);
+	check(InBinding);
+
+	if (!Super::PrepareForExecute_GameThread(InBinding))
+	{
+		return false;
+	}
+
+	OutputValueAttributeId = InBinding->GetAttributeId(PCGCountUniqueAttributeValuesConstants::ValueAttributeName, EPCGKernelAttributeType::Int);
+	check(OutputValueAttributeId != -1);
+	OutputCountAttributeId = InBinding->GetAttributeId(PCGCountUniqueAttributeValuesConstants::ValueCountAttributeName, EPCGKernelAttributeType::Int);
+	check(OutputCountAttributeId != -1);
+
+	FPCGDataCollectionDesc InputDataDesc;
+	ensure(InBinding->GetKernelPinDataDesc_GameThread(Kernel, PCGPinConstants::DefaultInputLabel, /*bIsInput*/true, InputDataDesc));
+
+	AttributeToCountId = InBinding->GetAttributeId(AttributeToCountName, EPCGKernelAttributeType::StringKey);
+
+	if (AttributeToCountId != INDEX_NONE)
+	{
+		InputDataDesc.GetUniqueStringKeyValues(AttributeToCountId, UniqueStringKeyValues);
+	}
+
+	MaxAttributeValue = INDEX_NONE;
+
+	for (int32 Key : UniqueStringKeyValues)
+	{
+		MaxAttributeValue = FMath::Max(MaxAttributeValue, Key);
+	}
+
+	return true;
+}
+
+FComputeDataProviderRenderProxy* UPCGCountUniqueAttributeValuesDataProvider::GetRenderProxy()
+{
+	FPCGCountUniqueAttributeValuesProviderProxy::FCountUniqueValuesData_RenderThread Data;
+	Data.AttributeToCountId = AttributeToCountId;
+	Data.MaxAttributeValue = MaxAttributeValue;
+	Data.OutputValueAttributeId = OutputValueAttributeId;
+	Data.OutputCountAttributeId = OutputCountAttributeId;
+	Data.UniqueStringKeyValues = UniqueStringKeyValues;
+
+	return new FPCGCountUniqueAttributeValuesProviderProxy(MoveTemp(Data));
+}
+
+void UPCGCountUniqueAttributeValuesDataProvider::Reset()
+{
+	Kernel = nullptr;
+	AttributeToCountName = NAME_None;
+	AttributeToCountId = INDEX_NONE;
+	OutputValueAttributeId = INDEX_NONE;
+	OutputCountAttributeId = INDEX_NONE;
+	MaxAttributeValue = INDEX_NONE;
+	UniqueStringKeyValues.Empty();
+
+	Super::Reset();
+}
+
+bool FPCGCountUniqueAttributeValuesProviderProxy::IsValid(FValidationData const& InValidationData) const
+{
+	return (InValidationData.ParameterStructSize == sizeof(FParameters));
+}
+
+void FPCGCountUniqueAttributeValuesProviderProxy::AllocateResources(FRDGBuilder& GraphBuilder, FAllocationData const& InAllocationData)
+{
+	const int32 NumIndicesToAllocate = FMath::Max(1, Data.MaxAttributeValue + 1);
+
+	{
+		FRDGBufferDesc Desc = FRDGBufferDesc::CreateStructuredDesc(sizeof(uint32), NumIndicesToAllocate);
+		Desc.Usage |= BUF_SourceCopy;
+
+		AttributeValueToValueIndex = GraphBuilder.CreateBuffer(Desc, TEXT("PCGAttributeValueToValueIndex"));
+		AttributeValueToValueIndexSRV = GraphBuilder.CreateSRV(AttributeValueToValueIndex);
+		
+		TArray<uint32> ValueIndices;
+		ValueIndices.SetNumZeroed(NumIndicesToAllocate);
+
+		for (int Value = 0; Value <= Data.MaxAttributeValue; ++Value)
+		{
+			const int32 ValueIndex = Data.UniqueStringKeyValues.IndexOfByKey(Value);
+			ValueIndices[Value] = static_cast<uint32>(ValueIndex);
+		}
+
+		GraphBuilder.QueueBufferUpload(AttributeValueToValueIndex, MakeArrayView(ValueIndices));
+	}
+}
+
+void FPCGCountUniqueAttributeValuesProviderProxy::GatherDispatchData(FDispatchData const& InDispatchData)
+{
+	const TStridedView<FParameters> ParameterArray = MakeStridedParameterView<FParameters>(InDispatchData);
+	for (int32 InvocationIndex = 0; InvocationIndex < ParameterArray.Num(); ++InvocationIndex)
+	{
+		FParameters& Parameters = ParameterArray[InvocationIndex];
+
+		Parameters.AttributeToCountId = Data.AttributeToCountId;
+		Parameters.OutputValueAttributeId = Data.OutputValueAttributeId;
+		Parameters.OutputCountAttributeId = Data.OutputCountAttributeId;
+		Parameters.AttributeValueToValueIndex = AttributeValueToValueIndexSRV;
+	}
+}
diff --git a/Engine/Plugins/PCG/Source/PCG/Private/Compute/DataInterfaces/BuiltInKernels/PCGCountUniqueAttributeValuesDataInterface.h b/Engine/Plugins/PCG/Source/PCG/Private/Compute/DataInterfaces/BuiltInKernels/PCGCountUniqueAttributeValuesDataInterface.h
new file mode 100644
--- /dev/null
+++ b/Engine/Plugins/PCG/Source/PCG/Private/Compute/DataInterfaces/BuiltInKernels/PCGCountUniqueAttributeValuesDataInterface.h
@@ -0,0 +1,100 @@
+// Copyright Epic Games, Inc. All Rights Reserved.
+
+#pragma once
+
+#include "Compute/DataInterfaces/PCGComputeDataInterface.h"
+#include "ComputeFramework/ComputeDataProvider.h"
+
+#include "PCGCountUniqueAttributeValuesDataInterface.generated.h"
+
+class FPCGCountUniqueAttributeValuesDataInterfaceParameters;
+
+/** Data Interface to marshal Count Unique Values kernel data to GPU. */
+UCLASS(ClassGroup = (Procedural))
+class UPCGCountUniqueAttributeValuesDataInterface : public UPCGComputeDataInterface
+{
+	GENERATED_BODY()
+
+public:
+	//~ Begin UComputeDataInterface Interface
+	TCHAR const* GetClassName() const override { return TEXT("PCGCountUniqueAttributeValues"); }
+	void GetSupportedInputs(TArray<FShaderFunctionDefinition>& OutFunctions) const override;
+	void GetShaderParameters(TCHAR const* UID, FShaderParametersMetadataBuilder& InOutBuilder, FShaderParametersMetadataAllocations& InOutAllocations) const override;
+	void GetHLSL(FString& OutHLSL, FString const& InDataInterfaceName) const override;
+	UComputeDataProvider* CreateDataProvider() const override;
+	//~ End UComputeDataInterface Interface
+
+public:
+	UPROPERTY()
+	FName AttributeToCountName;
+};
+
+UCLASS()
+class UPCGCountUniqueAttributeValuesDataProvider : public UPCGComputeDataProvider
+{
+	GENERATED_BODY()
+
+public:
+	//~ Begin UPCGComputeDataProvider Interface
+	virtual bool PrepareForExecute_GameThread(UPCGDataBinding* InBinding) override;
+	//~ End UPCGComputeDataProvider Interface
+
+	//~ Begin UComputeDataProvider Interface
+	virtual void Initialize(const UComputeDataInterface* InDataInterface, UObject* InBinding, uint64 InInputMask, uint64 InOutputMask) override;
+	FComputeDataProviderRenderProxy* GetRenderProxy() override;
+	virtual void Reset() override;
+	//~ End UComputeDataProvider Interface
+
+public:
+	UPROPERTY()
+	TObjectPtr<const UPCGComputeKernel> Kernel;
+
+	UPROPERTY()
+	FName AttributeToCountName;
+
+	UPROPERTY()
+	int32 AttributeToCountId = INDEX_NONE;
+
+	UPROPERTY()
+	int32 OutputValueAttributeId = INDEX_NONE;
+
+	UPROPERTY()
+	int32 OutputCountAttributeId = INDEX_NONE;
+
+	UPROPERTY()
+	int32 MaxAttributeValue = INDEX_NONE;
+
+	UPROPERTY()
+	TArray<int32> UniqueStringKeyValues;
+};
+
+class FPCGCountUniqueAttributeValuesProviderProxy : public FComputeDataProviderRenderProxy
+{
+public:
+	struct FCountUniqueValuesData_RenderThread
+	{
+		int32 AttributeToCountId = INDEX_NONE;
+		int32 OutputValueAttributeId = INDEX_NONE;
+		int32 OutputCountAttributeId = INDEX_NONE;
+		int32 MaxAttributeValue = INDEX_NONE;
+		TArray<int32> UniqueStringKeyValues;
+	};
+
+	FPCGCountUniqueAttributeValuesProviderProxy(FCountUniqueValuesData_RenderThread InData)
+		: Data(MoveTemp(InData))
+	{}
+
+	//~ Begin FComputeDataProviderRenderProxy Interface
+	bool IsValid(FValidationData const& InValidationData) const override;
+	void AllocateResources(FRDGBuilder& GraphBuilder, FAllocationData const& InAllocationData);
+	void GatherDispatchData(FDispatchData const& InDispatchData) override;
+	//~ End FComputeDataProviderRenderProxy Interface
+
+protected:
+	using FParameters = FPCGCountUniqueAttributeValuesDataInterfaceParameters;
+
+	FCountUniqueValuesData_RenderThread Data;
+
+	FRDGBufferRef AttributeValueToValueIndex = nullptr;
+	FRDGBufferSRVRef AttributeValueToValueIndexSRV = nullptr;
+};
diff --git a/Engine/Plugins/PCG/Source/PCG/Private/Compute/DataInterfaces/Elements/PCGCopyPointsDataInterface.cpp b/Engine/Plugins/PCG/Source/PCG/Private/Compute/DataInterfaces/Elements/PCGCopyPointsDataInterface.cpp
--- a/Engine/Plugins/PCG/Source/PCG/Private/Compute/DataInterfaces/Elements/PCGCopyPointsDataInterface.cpp
+++ b/Engine/Plugins/PCG/Source/PCG/Private/Compute/DataInterfaces/Elements/PCGCopyPointsDataInterface.cpp
@@ -169,56 +169,15 @@ bool UPCGCopyPointsDataProvider::PerformPreExecuteReadbacks_GameThread(UPCGDataB
 		return true;
 	}
 
-	int32 FirstAnalysisDataIndex = INDEX_NONE;
-
-	if (InBinding->Graph && ProducerKernel)
-	{
-		const int32 BindingIndex = InBinding->Graph->GetBindingIndex(ProducerKernel->GetKernelIndex(), PCGCopyPointsConstants::SelectedFlagsPinLabel, /*bIsInput=*/true);
-		if (ensure(BindingIndex != INDEX_NONE))
-		{
-			if (const FName* VirtualLabel = InBinding->Graph->GetVirtualPinLabelFromBindingIndex(BindingIndex))
-			{
-				FirstAnalysisDataIndex = InBinding->InputDataCollection.TaggedData.IndexOfByPredicate([VirtualLabel](const FPCGTaggedData& InData)
-				{
-					return InData.Pin == *VirtualLabel;
-				});
-			}
-		}
-	}
-
-	if (FirstAnalysisDataIndex == INDEX_NONE)
+	const int32 AnalysisDataIndex = ProducerKernel ? InBinding->GetFirstInputDataIndex(ProducerKernel, PCGCopyPointsConstants::SelectedFlagsPinLabel) : INDEX_NONE;
+	if (AnalysisDataIndex == INDEX_NONE)
 	{
 		ensure(false);
 		return true;
 	}
 
-	const UPCGProxyForGPUData* DataWithGPUSupport = Cast<UPCGProxyForGPUData>(InBinding->InputDataCollection.TaggedData[FirstAnalysisDataIndex].Data);
-	if (DataWithGPUSupport)
-	{
-		TSharedPtr<FPCGContextHandle> Context = InBinding->ContextHandle.Pin();
-
-		// Poll until readback is done.
-		UPCGProxyForGPUData::FReadbackResult Result = DataWithGPUSupport->GetCPUData(Context ? Context->GetContext() : nullptr);
-
-		if (!Result.bComplete)
-		{
-			// Read back pending - wait for it to complete.
-			return false;
-		}
-
-		if (Result.TaggedData.Data)
-		{
-			InBinding->InputDataCollection.TaggedData[FirstAnalysisDataIndex].Data = Result.TaggedData.Data;
-			InBinding->InputDataCollection.TaggedData[FirstAnalysisDataIndex].Tags = MoveTemp(Result.TaggedData.Tags);
-		}
-		else
-		{
-			ensure(false);
-			return true;
-		}
-	}
-
-	return true;
+	// Readback analysis data - poll until readback complete.
+	return InBinding->ReadbackInputDataToCPU(AnalysisDataIndex);
 }
 
 bool UPCGCopyPointsDataProvider::PrepareForExecute_GameThread(UPCGDataBinding* InBinding)
@@ -235,21 +194,8 @@ bool UPCGCopyPointsDataProvider::PrepareForExecute_GameThread(UPCGDataBinding* I
 
 	if (Settings->bMatchBasedOnAttribute)
 	{
-		int32 FirstAnalysisDataIndex = INDEX_NONE;
-
-		const int32 BindingIndex = InBinding->Graph ? InBinding->Graph->GetBindingIndex(ProducerKernel->GetKernelIndex(), PCGCopyPointsConstants::SelectedFlagsPinLabel, /*bIsInput=*/true) : INDEX_NONE;
-		if (ensure(BindingIndex != INDEX_NONE))
-		{
-			if (const FName* VirtualLabel = InBinding->Graph->GetVirtualPinLabelFromBindingIndex(BindingIndex))
-			{
-				FirstAnalysisDataIndex = InBinding->InputDataCollection.TaggedData.IndexOfByPredicate([VirtualLabel](const FPCGTaggedData& InData)
-				{
-					return InData.Pin == *VirtualLabel;
-				});
-			}
-		}
-
-		if (FirstAnalysisDataIndex == INDEX_NONE)
+		const int32 AnalysisDataIndex = InBinding->GetFirstInputDataIndex(ProducerKernel, PCGCopyPointsConstants::SelectedFlagsPinLabel);
+		if (AnalysisDataIndex == INDEX_NONE)
 		{
 			return true;
 		}
@@ -262,7 +208,7 @@ bool UPCGCopyPointsDataProvider::PrepareForExecute_GameThread(UPCGDataBinding* I
 		const int32 NumOutputs = Settings->bCopyEachSourceOnEveryTarget ? SourcePinDesc.DataDescs.Num() * TargetPinDesc.DataDescs.Num() : FMath::Max(SourcePinDesc.DataDescs.Num(), TargetPinDesc.DataDescs.Num());
 		SourceAndTargetDataIndices.Reserve(NumOutputs);
 
-		const UPCGParamData* AnalysisResultsData = Cast<UPCGParamData>(InBinding->InputDataCollection.TaggedData[FirstAnalysisDataIndex].Data);
+		const UPCGParamData* AnalysisResultsData = Cast<UPCGParamData>(InBinding->InputDataCollection.TaggedData[AnalysisDataIndex].Data);
 		const UPCGMetadata* AnalysisMetadata = AnalysisResultsData ? AnalysisResultsData->ConstMetadata() : nullptr;
 		const FPCGMetadataAttributeBase* AnalysisAttributeBase = AnalysisMetadata ? AnalysisMetadata->GetConstAttribute(PCGCopyPointsConstants::SelectedFlagAttributeName) : nullptr;
 		const int32 NumElements = AnalysisAttributeBase ? AnalysisMetadata->GetItemCountForChild() : INDEX_NONE;
diff --git a/Engine/Plugins/PCG/Source/PCG/Private/Compute/DataInterfaces/Elements/PCGStaticMeshSpawnerDataInterface.cpp b/Engine/Plugins/PCG/Source/PCG/Private/Compute/DataInterfaces/Elements/PCGStaticMeshSpawnerDataInterface.cpp
--- a/Engine/Plugins/PCG/Source/PCG/Private/Compute/DataInterfaces/Elements/PCGStaticMeshSpawnerDataInterface.cpp
+++ b/Engine/Plugins/PCG/Source/PCG/Private/Compute/DataInterfaces/Elements/PCGStaticMeshSpawnerDataInterface.cpp
@@ -2,18 +2,40 @@
 
 #include "PCGStaticMeshSpawnerDataInterface.h"
 
+#include "PCGComponent.h"
+#include "PCGContext.h"
+#include "PCGParamData.h"
+#include "Components/PCGProceduralISMComponent.h"
 #include "Compute/PCGComputeCommon.h"
+#include "Compute/PCGComputeGraph.h"
 #include "Compute/PCGDataBinding.h"
+#include "Compute/BuiltInKernels/PCGCountUniqueAttributeValuesKernel.h"
+#include "Compute/Data/PCGProxyForGPUData.h"
 #include "Elements/PCGStaticMeshSpawner.h"
 #include "Elements/PCGStaticMeshSpawnerKernel.h"
+#include "InstanceDataPackers/PCGInstanceDataPackerBase.h"
+#include "MeshSelectors/PCGMeshSelectorByAttribute.h"
+#include "MeshSelectors/PCGMeshSelectorWeighted.h"
 
 #include "GlobalRenderResources.h"
 #include "RHIResources.h"
 #include "ShaderCompilerCore.h"
 #include "ShaderCore.h"
 #include "ShaderParameterMetadataBuilder.h"
+#include "Engine/StaticMesh.h"
 #include "ComputeFramework/ShaderParamTypeDefinition.h"
 
+#include UE_INLINE_GENERATED_CPP_BY_NAME(PCGStaticMeshSpawnerDataInterface)
+
+#define LOCTEXT_NAMESPACE "PCGStaticMeshSpawnerKernel"
+
+namespace PCGStaticMeshSpawnerKernel
+{
+	static const FText CouldNotLoadStaticMeshFormat = LOCTEXT("CouldNotLoadStaticMesh", "Could not load static mesh from path '{0}'.");
+	static const FText TooManyPrimitivesFormat = LOCTEXT("TooManyPrimitives", "Attempted to emit too many primitive components, terminated after creating '{0}'.");
+	static const FText NoMeshEntriesFormat = LOCTEXT("NoMeshEntries", "No mesh entries provided in weighted mesh selector.");
+}
+
 void UPCGStaticMeshSpawnerDataInterface::GetSupportedInputs(TArray<FShaderFunctionDefinition>& OutFunctions) const
 {
 	OutFunctions.AddDefaulted_GetRef()
@@ -183,46 +205,495 @@ void UPCGStaticMeshSpawnerDataProvider::Initialize(const UComputeDataInterface*
 	Super::Initialize(InDataInterface, InBinding, InInputMask, InOutputMask);
 
 	const UPCGStaticMeshSpawnerDataInterface* DataInterface = CastChecked<UPCGStaticMeshSpawnerDataInterface>(InDataInterface);
-	const UPCGStaticMeshSpawnerKernel* Kernel = CastChecked<UPCGStaticMeshSpawnerKernel>(DataInterface->ProducerKernel);
 	UPCGDataBinding* Binding = CastChecked<UPCGDataBinding>(InBinding);
-	check(Kernel->Settings);
 
-	// If there were 0 input points for this execution, we will not have created any primitives, so check for null.
-	if (const FPCGSpawnerPrimitives* Primitives = Binding->MeshSpawnersToPrimitives.Find(Kernel))
+	ProducerKernel = CastChecked<UPCGStaticMeshSpawnerKernel>(DataInterface->ProducerKernel);
+	check(ProducerKernel->Settings);
+}
+
+bool UPCGStaticMeshSpawnerDataProvider::PerformPreExecuteReadbacks_GameThread(UPCGDataBinding* InBinding)
+{
+	TRACE_CPUPROFILER_EVENT_SCOPE(UPCGCopyPointsDataProvider::PerformPreExecuteReadbacks_GameThread);
+	check(InBinding);
+
+	if (!Super::PerformPreExecuteReadbacks_GameThread(InBinding))
+	{
+		return false;
+	}
+
+	const UPCGStaticMeshSpawnerSettings* Settings = CastChecked<UPCGStaticMeshSpawnerSettings>(ProducerKernel->GetSettings());
+
+	// Obtain index of analysis data that we need to read back.
+	if (AnalysisDataIndex == INDEX_NONE)
+	{
+		if (!Cast<UPCGMeshSelectorByAttribute>(Settings->MeshSelectorParameters))
+		{
+			// Non by-attribute selection does not need to do any readbacks.
+			return true;
+		}
+
+		AnalysisDataIndex = ProducerKernel ? InBinding->GetFirstInputDataIndex(ProducerKernel, PCGStaticMeshSpawnerConstants::InstanceCountsPinLabel) : INDEX_NONE;
+
+		if (AnalysisDataIndex == INDEX_NONE)
+		{
+			// No analysis data to read back.
+			return true;
+		}
+	}
+
+	// Readback analysis data - poll until readback complete (return true).
+	if (!InBinding->ReadbackInputDataToCPU(AnalysisDataIndex))
+	{
+		return false;
+	}
+
+	const UPCGParamData* AnalysisResultsData = Cast<UPCGParamData>(InBinding->InputDataCollection.TaggedData[AnalysisDataIndex].Data);
+	const UPCGMetadata* AnalysisMetadata = AnalysisResultsData ? AnalysisResultsData->ConstMetadata() : nullptr;
+
+	const FPCGMetadataAttributeBase* AnalysisValueAttributeBase = AnalysisMetadata ? AnalysisMetadata->GetConstAttribute(PCGCountUniqueAttributeValuesConstants::ValueAttributeName) : nullptr;
+	const FPCGMetadataAttributeBase* AnalysisCountAttributeBase = AnalysisMetadata ? AnalysisMetadata->GetConstAttribute(PCGCountUniqueAttributeValuesConstants::ValueCountAttributeName) : nullptr;
+
+	if (AnalysisValueAttributeBase && AnalysisValueAttributeBase->GetTypeId() == PCG::Private::MetadataTypes<int32>::Id &&
+		AnalysisCountAttributeBase && AnalysisCountAttributeBase->GetTypeId() == PCG::Private::MetadataTypes<int32>::Id)
+	{
+		const FPCGMetadataAttribute<int32>* ValueAttribute = static_cast<const FPCGMetadataAttribute<int32>*>(AnalysisValueAttributeBase);
+		const FPCGMetadataAttribute<int32>* CountAttribute = static_cast<const FPCGMetadataAttribute<int32>*>(AnalysisCountAttributeBase);
+
+		const int32 NumElements = AnalysisMetadata->GetItemCountForChild();
+
+		StringKeyToInstanceCount.Reserve(NumElements);
+
+		// TODO: Range based get would scale better.
+		for (int64 MetadataKey = 0; MetadataKey < NumElements; ++MetadataKey)
+		{
+			StringKeyToInstanceCount.Add(ValueAttribute->GetValue(MetadataKey), CountAttribute->GetValue(MetadataKey));
+		}
+	}
+	else
+	{
+		UE_LOG(LogPCG, Warning, TEXT("No analysis data received by static mesh spawner kernel, worst case instance allocations will be made."));
+
+		if (InBinding->InputDataCollection.TaggedData[AnalysisDataIndex].Data->IsA<UPCGProxyForGPUData>())
+		{
+			UE_LOG(LogPCG, Error, TEXT("Data was not read back."));
+		}
+
+		return true;
+	}
+
+	return true;
+}
+
+bool UPCGStaticMeshSpawnerDataProvider::PrepareForExecute_GameThread(UPCGDataBinding* InBinding)
+{
+	TRACE_CPUPROFILER_EVENT_SCOPE(UPCGCopyPointsDataProvider::PerformPreExecuteReadbacks_GameThread);
+	check(InBinding);
+	const UPCGStaticMeshSpawnerSettings* Settings = CastChecked<UPCGStaticMeshSpawnerSettings>(ProducerKernel->GetSettings());
+
+	if (!Super::PrepareForExecute_GameThread(InBinding))
+	{
+		return false;
+	}
+
+	if (const UPCGMeshSelectorByAttribute* SelectorByAttribute = Cast<UPCGMeshSelectorByAttribute>(Settings->MeshSelectorParameters))
+	{
+		const FName SelectorName = SelectorByAttribute->AttributeName;
+
+		FPCGDataCollectionDesc InputDataDesc;
+		InBinding->GetKernelPinDataDesc_GameThread(ProducerKernel, PCGPinConstants::DefaultInputLabel, /*bIsInput=*/true, InputDataDesc);
+
+		for (const FPCGDataDesc& Desc : InputDataDesc.DataDescs)
+		{
+			if (Desc.ElementCount <= 0)
+			{
+				continue;
+			}
+
+			for (const FPCGKernelAttributeDesc& AttributeDesc : Desc.AttributeDescs)
+			{
+				if (AttributeDesc.Name == SelectorName && AttributeDesc.Type == EPCGKernelAttributeType::StringKey)
+				{
+					SelectorAttributeId = AttributeDesc.Index;
+					break;
+				}
+			}
+
+			if (SelectorAttributeId != INDEX_NONE)
+			{
+				break;
+			}
+		}
+
+		if (SelectorAttributeId == INDEX_NONE)
+		{
+			if (!InputDataDesc.DataDescs.IsEmpty())
+			{
+				TSharedPtr<FPCGContextHandle> ContextHandle = InBinding->ContextHandle.Pin();
+				FPCGContext* Context = ContextHandle ? ContextHandle->GetContext() : nullptr;
+
+				PCG_KERNEL_VALIDATION_ERR(Context, Settings, /*bQuiet=*/false, FText::Format(
+					LOCTEXT("MeshSelectorAttributeNotFound", "Mesh selector attribute '{0}' not found."),
+					FText::FromName(SelectorName)));
+			}
+
+			return true;
+		}
+	}
+
+	if (!bPrimitivesSetUp)
+	{
+		bPrimitivesSetUp = true;
+
+		TSharedPtr<FPCGContextHandle> ContextHandle = InBinding->ContextHandle.Pin();
+		FPCGContext* Context = ContextHandle ? ContextHandle->GetContext() : nullptr;
+
+		if (!ensure(Context))
+		{
+			return true;
+		}
+
+		const bool bSetupComponent = SetupPrimitives(Context, InBinding);
+
+		if (bSetupComponent)
+		{
+			Context->SourceComponent->NotifyProceduralInstancesInUse();
+
+			// Wait a frame - gives time for GPU Scene to be updated.
+			return false;
+		}
+		else
+		{
+			InBinding->MeshSpawnersToPrimitives.Add(ProducerKernel, {});
+
+			// No component set up means we have no more work to do.
+			return true;
+		}
+	}
+
+	// If any proxy is not yet ready, wait a frame.
+	if (const FPCGSpawnerPrimitives* Primitives = InBinding->MeshSpawnersToPrimitives.Find(ProducerKernel); ensure(Primitives))
 	{
-		AttributeIdOffsetStrides = Primitives->AttributeIdOffsetStrides;
-		SelectionCDF = Primitives->SelectionCDF;
-		PrimitiveStringKeys = Primitives->PrimitiveStringKeys;
-		PrimitiveMeshBounds = Primitives->PrimitiveMeshBounds;
-
-		// We only know name at compile time (not type), so we resolved this Id at execution time by inspecting the data arriving on the input pin. 
-		SelectorAttributeId = Primitives->SelectorAttributeId;
-
-		// We know the name and type of the selected mesh attribute statically and declared the attribute in GetKernelAttributeKeys,
-		// so the attribute ID should be present in the attribute table.
-		SelectedMeshAttributeId = Binding->GetAttributeId(Kernel->Settings->OutAttributeName, EPCGKernelAttributeType::StringKey);
-		ensure(SelectedMeshAttributeId != INDEX_NONE);
+		for (const TObjectPtr<UPrimitiveComponent>& Primitive : Primitives->Primitives)
+		{
+			const UPCGProceduralISMComponent* PISMC = Cast<UPCGProceduralISMComponent>(Primitive);
+			if (PISMC && PISMC->GetNumInstances() > 0 && !PISMC->GetSceneProxy())
+			{
+				// Waiting on a proxy. Wait a frame.
+				return false;
+			}
+		}
 	}
+
+	// We know the name and type of the selected mesh attribute statically and declared the attribute in GetKernelAttributeKeys,
+	// so the attribute ID should be present in the attribute table.
+	SelectedMeshAttributeId = InBinding->GetAttributeId(ProducerKernel->Settings->OutAttributeName, EPCGKernelAttributeType::StringKey);
+	ensure(SelectedMeshAttributeId != INDEX_NONE);
+
+	return true;
 }
 
 FComputeDataProviderRenderProxy* UPCGStaticMeshSpawnerDataProvider::GetRenderProxy()
 {
-	return new FPCGStaticMeshSpawnerDataProviderProxy(AttributeIdOffsetStrides, SelectorAttributeId, PrimitiveStringKeys, SelectionCDF, SelectedMeshAttributeId, PrimitiveMeshBounds);
+	return new FPCGStaticMeshSpawnerDataProviderProxy(AttributeIdOffsetStrides, SelectorAttributeId, PrimitiveStringKeys, PrimitiveSelectionCDF, SelectedMeshAttributeId, PrimitiveMeshBounds);
 }
 
 void UPCGStaticMeshSpawnerDataProvider::Reset()
 {
+	ProducerKernel = nullptr;
 	AttributeIdOffsetStrides.Empty();
 	PrimitiveStringKeys.Empty();
 	PrimitiveMeshBounds.Empty();
-	SelectionCDF.Empty();
+	PrimitiveSelectionCDF.Empty();
 	SelectorAttributeId = INDEX_NONE;
 	NumInputPoints = 0;
 	SelectedMeshAttributeId = INDEX_NONE;
+	StringKeyToInstanceCount.Empty();
+	AnalysisDataIndex = INDEX_NONE;
+	bPrimitivesSetUp = false;
 
 	Super::Reset();
 }
 
+bool UPCGStaticMeshSpawnerDataProvider::SetupPrimitives(FPCGContext* InContext, UPCGDataBinding* InBinding)
+{
+	const UPCGStaticMeshSpawnerSettings* Settings = CastChecked<UPCGStaticMeshSpawnerSettings>(ProducerKernel->GetSettings());
+
+	AActor* TargetActor = Settings->TargetActor.Get() ? Settings->TargetActor.Get() : InContext->GetTargetActor(nullptr);
+	if (!ensure(Settings->MeshSelectorParameters) || !ensure(TargetActor))
+	{
+		return false;
+	}
+
+	FPCGDataCollectionDesc InputDataDesc;
+	InBinding->GetKernelPinDataDesc_GameThread(ProducerKernel, PCGPinConstants::DefaultInputLabel, /*bIsInput=*/true, InputDataDesc);
+
+	const uint32 TotalInputPointCount = InputDataDesc.ComputeDataElementCount(EPCGDataType::Point);
+	if (TotalInputPointCount == 0)
+	{
+		return false;
+	}
+
+	const FBox LocalBounds = InContext->SourceComponent->GetGridBounds().ShiftBy(-InContext->SourceComponent->GetOwner()->GetActorLocation());
+
+	uint32 CustomFloatCount = 0;
+
+	if (Settings->InstanceDataPackerParameters && Settings->InstanceDataPackerParameters->GetAttributeNames(/*OutNames=*/nullptr))
+	{
+		TArray<FName> AttributeNames;
+		Settings->InstanceDataPackerParameters->GetAttributeNames(&AttributeNames);
+
+		for (FName AttributeName : AttributeNames)
+		{
+			if (AttributeName == NAME_None)
+			{
+				PCG_KERNEL_VALIDATION_ERR(InContext, Settings, /*bQuiet=*/false,
+					FText::Format(LOCTEXT("InstanceDataAttributeInvalid", "Invalid instance data attribute specified '{0}'."), FText::FromName(AttributeName)));
+
+				continue;
+			}
+
+			if (!InputDataDesc.ContainsAttributeOnAnyData(AttributeName))
+			{
+				PCG_KERNEL_VALIDATION_ERR(InContext, Settings, /*bQuiet=*/false,
+					FText::Format(LOCTEXT("InstanceDataAttributeNotFound", "Instance data attribute '{0}' not found."), FText::FromName(AttributeName)));
+
+				continue;
+			}
+		}
+
+		PCGDataForGPUHelpers::ComputeCustomFloatPacking(InContext, Settings, AttributeNames, InBinding, InputDataDesc, CustomFloatCount, AttributeIdOffsetStrides);
+	}
+
+	TArray<FPCGProceduralISMComponentDescriptor> ComponentsToCreate;
+
+	if (const UPCGMeshSelectorByAttribute* SelectorByAttribute = Cast<UPCGMeshSelectorByAttribute>(Settings->MeshSelectorParameters))
+	{
+		const FName SelectorName = SelectorByAttribute->AttributeName;
+
+		// Compute how many instances we expect for each mesh, if we do not already have the answer from analysis.
+		if (StringKeyToInstanceCount.IsEmpty())
+		{
+			for (const FPCGDataDesc& Desc : InputDataDesc.DataDescs)
+			{
+				if (Desc.ElementCount <= 0)
+				{
+					continue;
+				}
+
+				for (const FPCGKernelAttributeDesc& AttributeDesc : Desc.AttributeDescs)
+				{
+					if (AttributeDesc.Name == SelectorName)
+					{
+						if (AttributeDesc.Type == EPCGKernelAttributeType::StringKey)
+						{
+							for (int32 AttributeStringKeyValue : AttributeDesc.GetUniqueStringKeys())
+							{
+								if (AttributeStringKeyValue > 0)
+								{
+									StringKeyToInstanceCount.FindOrAdd(AttributeStringKeyValue, 0) += Desc.ElementCount;
+								}
+							}
+
+							// We could early out but we currently continue looping to generate the below warning for mismatched types.
+						}
+						else
+						{
+							// Currently only a single type per attribute name is supported (the name/type in the attribute table). It's possible to wire a graph
+							// for which an attribute is present with multiple types. Warn if this is encountered.
+							PCG_KERNEL_VALIDATION_WARN(InContext, Settings, /*bQuiet=*/false, FText::Format(
+								LOCTEXT("MeshSelectorAttributeNotUsable", "Attribute '{0}' not usable for mesh selection, only attributes of type String Key are supported."),
+								FText::FromName(SelectorName)));
+						}
+					}
+				}
+			}
+		}
+
+		for (const TPair<int32, uint32>& StringKeyAndInstanceCount : StringKeyToInstanceCount)
+		{
+			if (!ensure(InBinding->GetStringTable().IsValidIndex(StringKeyAndInstanceCount.Key)))
+			{
+				continue;
+			}
+
+			const FString& MeshPathString = InBinding->GetStringTable()[StringKeyAndInstanceCount.Key];
+			if (MeshPathString.IsEmpty())
+			{
+				continue;
+			}
+
+			if (ComponentsToCreate.Num() >= PCGComputeConstants::MAX_PRIMITIVE_COMPONENTS_PER_SPAWNER)
+			{
+				PCG_KERNEL_VALIDATION_WARN(InContext, Settings, /*bQuiet=*/false,
+					FText::Format(PCGStaticMeshSpawnerKernel::TooManyPrimitivesFormat, FText::FromString(FString::FromInt(PCGComputeConstants::MAX_PRIMITIVE_COMPONENTS_PER_SPAWNER))));
+
+				break;
+			}
+
+			UStaticMesh* StaticMesh = Cast<UStaticMesh>(FSoftObjectPath(MeshPathString).TryLoad());
+			if (!StaticMesh)
+			{
+				PCG_KERNEL_VALIDATION_ERR(InContext, Settings, /*bQuiet=*/false,
+					FText::Format(PCGStaticMeshSpawnerKernel::CouldNotLoadStaticMeshFormat, FText::FromString(MeshPathString)));
+
+				continue;
+			}
+
+			FPCGProceduralISMComponentDescriptor Descriptor;
+			Descriptor = SelectorByAttribute->TemplateDescriptor;
+			Descriptor.NumInstances = StringKeyAndInstanceCount.Value;
+			Descriptor.LocalBounds = LocalBounds;
+			Descriptor.NumCustomFloats = CustomFloatCount;
+			Descriptor.StaticMesh = StaticMesh;
+
+			PrimitiveStringKeys.Add(StringKeyAndInstanceCount.Key);
+			ComponentsToCreate.Add(MoveTemp(Descriptor));
+
+			if (Settings->bApplyMeshBoundsToPoints)
+			{
+				PrimitiveMeshBounds.Add(Descriptor.StaticMesh->GetBoundingBox());
+			}
+		}
+
+		PrimitiveSelectionCDF.SetNumZeroed(ComponentsToCreate.Num());
+	}
+	else if (const UPCGMeshSelectorWeighted* SelectorWeighted = Cast<UPCGMeshSelectorWeighted>(Settings->MeshSelectorParameters))
+	{
+		if (SelectorWeighted->MeshEntries.IsEmpty())
+		{
+			PCG_KERNEL_VALIDATION_ERR(InContext, Settings, /*bQuiet=*/false, PCGStaticMeshSpawnerKernel::NoMeshEntriesFormat);
+
+			return false;
+		}
+
+		float CumulativeWeight = 0.0f;
+
+		float TotalWeight = 0.0f;
+		for (const FPCGMeshSelectorWeightedEntry& Entry : SelectorWeighted->MeshEntries)
+		{
+			TotalWeight += Entry.Weight;
+		}
+
+		if (!ensure(TotalWeight > UE_SMALL_NUMBER))
+		{
+			return false;
+		}
+
+		PrimitiveSelectionCDF.Reserve(SelectorWeighted->MeshEntries.Num());
+
+		for (const FPCGMeshSelectorWeightedEntry& Entry : SelectorWeighted->MeshEntries)
+		{
+			if (Entry.Descriptor.StaticMesh.IsNull())
+			{
+				continue;
+			}
+
+			if (ComponentsToCreate.Num() >= PCGComputeConstants::MAX_PRIMITIVE_COMPONENTS_PER_SPAWNER)
+			{
+				PCG_KERNEL_VALIDATION_WARN(InContext, Settings, /*bQuiet=*/false,
+					FText::Format(PCGStaticMeshSpawnerKernel::TooManyPrimitivesFormat, FText::FromString(FString::FromInt(PCGComputeConstants::MAX_PRIMITIVE_COMPONENTS_PER_SPAWNER))));
+
+				break;
+			}
+
+			if (UStaticMesh* StaticMesh = Entry.Descriptor.StaticMesh.LoadSynchronous())
+			{
+				const float Weight = float(Entry.Weight) / TotalWeight;
+				CumulativeWeight += Weight;
+				PrimitiveSelectionCDF.Add(CumulativeWeight);
+				PrimitiveStringKeys.Add(InBinding->GetStringTable().IndexOfByKey(Entry.Descriptor.StaticMesh.ToString()));
+
+				FPCGProceduralISMComponentDescriptor Descriptor;
+				Descriptor = Entry.Descriptor;
+				Descriptor.LocalBounds = LocalBounds;
+				Descriptor.NumCustomFloats = CustomFloatCount;
+				Descriptor.StaticMesh = StaticMesh;
+
+				uint32 InstanceCount = FMath::CeilToInt(TotalInputPointCount * Weight);
+
+				if (SelectorWeighted->MeshEntries.Num() > 1)
+				{
+					// Since we'll be selecting meshes based on random draws using the point random seeds which we don't have on CPU,
+					// we may pick more or less than the expected number of instances for each mesh. Use binomial variance to calculate
+					// the overallocation that gives 99.7% confidence (3 sigma).
+					const double Variance = double(TotalInputPointCount) * Weight * (1.0f - Weight);
+					const double Sigma = FMath::Sqrt(Variance);
+					const uint32 AdditionalAllocation = FMath::CeilToInt(3.0f * Sigma);
+
+					// Useful for debugging instance counts.
+					//UE_LOG(LogPCG, Warning, TEXT("Over allocation, N = %d, P = %.2f, NP = %d, 3sig = %f, Extra = %d (%.2f%%), Calculated = %d, Final = %d"), TotalInputPointCount, Weight, InstanceCount, float(3.0f * Sigma), AdditionalAllocation, float(AdditionalAllocation) / TotalInputPointCount, InstanceCount + AdditionalAllocation, FMath::Min(InstanceCount + AdditionalAllocation, TotalInputPointCount));
+
+					InstanceCount += AdditionalAllocation;
+				}
+
+				Descriptor.NumInstances = FMath::Min(InstanceCount, TotalInputPointCount);
+
+				ComponentsToCreate.Add(MoveTemp(Descriptor));
+
+				if (Settings->bApplyMeshBoundsToPoints)
+				{
+					PrimitiveMeshBounds.Add(Descriptor.StaticMesh->GetBoundingBox());
+				}
+			}
+			else
+			{
+				PCG_KERNEL_VALIDATION_ERR(InContext, Settings, /*bQuiet=*/false,
+					FText::Format(PCGStaticMeshSpawnerKernel::CouldNotLoadStaticMeshFormat, FText::FromString(Entry.Descriptor.StaticMesh.ToString())));
+			}
+		}
+	}
+
+	// Validate instance count is not too large.
+	{
+		uint64 TotalInstanceCount = 0;
+		for (const FPCGProceduralISMComponentDescriptor& Desc : ComponentsToCreate)
+		{
+			TotalInstanceCount += Desc.NumInstances;
+		}
+
+		if (TotalInstanceCount >= MAX_INSTANCE_ID)
+		{
+			PCG_KERNEL_VALIDATION_ERR(InContext, Settings, /*bQuiet=*/false, FText::Format(
+				LOCTEXT("TooManyInstances", "Tried to spawn too many instances ({0}), procedural ISM component creation skipped and instances will not be rendered."),
+				FText::FromString(FString::FromInt(TotalInstanceCount))));
+
+			return false;
+		}
+	}
+
+	FPCGSpawnerPrimitives& Primitives = InBinding->MeshSpawnersToPrimitives.FindOrAdd(ProducerKernel);
+	Primitives.NumCustomFloats = CustomFloatCount;
+	Primitives.AttributeIdOffsetStrides = /*MoveTemp*/(AttributeIdOffsetStrides);
+	Primitives.SelectorAttributeId = SelectorAttributeId;
+	Primitives.SelectionCDF = /*MoveTemp*/(PrimitiveSelectionCDF);
+	Primitives.PrimitiveStringKeys = /*MoveTemp*/(PrimitiveStringKeys);
+	Primitives.PrimitiveMeshBounds = /*MoveTemp*/(PrimitiveMeshBounds);
+
+	for (const FPCGProceduralISMComponentDescriptor& Desc : ComponentsToCreate)
+	{
+		FPCGProceduralISMCBuilderParameters Params;
+		Params.Descriptor = Desc;
+		Params.bAllowDescriptorChanges = false;
+
+		UPCGManagedProceduralISMComponent* MISMC = PCGManagedProceduralISMComponent::GetOrCreateManagedProceduralISMC(TargetActor, InContext->SourceComponent.Get(), Settings->UID, Params);
+
+		check(MISMC);
+		MISMC->SetCrc(InContext->DependenciesCrc);
+
+		// Don't bother registering the resource change as we're transient anyway.
+		//InContext->TouchedResources.Emplace(MISMC);
+
+		Primitives.Primitives.Add(MISMC->GetComponent());
+
+		Primitives.NumInstancesAllPrimitives += Desc.NumInstances;
+	}
+
+	// Useful for debugging instance counts.
+	//UE_LOG(LogPCG, Log, TEXT("Input points %d, total instance count %d (+ %.2f%%)"), TotalInputPointCount, Primitives.NumInstancesAllPrimitives, (float(Primitives.NumInstancesAllPrimitives - TotalInputPointCount) / TotalInputPointCount) * 100.0f);
+
+	return !ComponentsToCreate.IsEmpty();
+}
+
 bool FPCGStaticMeshSpawnerDataProviderProxy::IsValid(FValidationData const& InValidationData) const
 {
 	if (InValidationData.ParameterStructSize != sizeof(FParameters))
@@ -269,3 +740,5 @@ void FPCGStaticMeshSpawnerDataProviderProxy::GatherDispatchData(FDispatchData co
 		}
 	}
 }
+
+#undef LOCTEXT_NAMESPACE
diff --git a/Engine/Plugins/PCG/Source/PCG/Private/Compute/DataInterfaces/Elements/PCGStaticMeshSpawnerDataInterface.h b/Engine/Plugins/PCG/Source/PCG/Private/Compute/DataInterfaces/Elements/PCGStaticMeshSpawnerDataInterface.h
--- a/Engine/Plugins/PCG/Source/PCG/Private/Compute/DataInterfaces/Elements/PCGStaticMeshSpawnerDataInterface.h
+++ b/Engine/Plugins/PCG/Source/PCG/Private/Compute/DataInterfaces/Elements/PCGStaticMeshSpawnerDataInterface.h
@@ -10,6 +10,8 @@
 class FPCGStaticMeshSpawnerDataInterfaceParameters;
 class UPCGStaticMeshSpawnerSettings;
 class UPCGSettings;
+class UPCGStaticMeshSpawnerKernel;
+struct FPCGContext;
 
 /** Data Interface to marshal Static Mesh Spawner settings to the GPU. */
 UCLASS(ClassGroup = (Procedural))
@@ -30,17 +32,26 @@ public:
 };
 
 UCLASS()
-class UPCGStaticMeshSpawnerDataProvider : public UComputeDataProvider
+class UPCGStaticMeshSpawnerDataProvider : public UPCGComputeDataProvider
 {
 	GENERATED_BODY()
 
 public:
 	//~ Begin UComputeDataProvider Interface
 	virtual void Initialize(const UComputeDataInterface* InDataInterface, UObject* InBinding, uint64 InInputMask, uint64 InOutputMask) override;
+	virtual bool PerformPreExecuteReadbacks_GameThread(UPCGDataBinding* InBinding) override;
+	virtual bool PrepareForExecute_GameThread(UPCGDataBinding* InBinding) override;
 	FComputeDataProviderRenderProxy* GetRenderProxy() override;
 	virtual void Reset() override;
 	//~ End UComputeDataProvider Interface
 
+	/** Setup Procedural ISMCs. Returns true if at least one component was set up. */
+	virtual bool SetupPrimitives(FPCGContext* InContext, UPCGDataBinding* InBinding);
+
+public:
+	UPROPERTY()
+	TObjectPtr<const UPCGStaticMeshSpawnerKernel> ProducerKernel;
+
 	UPROPERTY()
 	TArray<FUintVector4> AttributeIdOffsetStrides;
 
@@ -51,7 +62,7 @@ public:
 	TArray<FBox> PrimitiveMeshBounds;
 
 	UPROPERTY()
-	TArray<float> SelectionCDF;
+	TArray<float> PrimitiveSelectionCDF;
 
 	/** Attribute Id for mesh selector. */
 	UPROPERTY()
@@ -63,6 +74,13 @@ public:
 	/** Output attribute Id for the output selected mesh. */
 	UPROPERTY()
 	int32 SelectedMeshAttributeId = INDEX_NONE;
+
+	/** The number of instances per string key value, used for by-attribute spawning. */
+	TMap<int32, uint32> StringKeyToInstanceCount;
+
+	int32 AnalysisDataIndex = INDEX_NONE;
+
+	bool bPrimitivesSetUp = false;
 };
 
 class FPCGStaticMeshSpawnerDataProviderProxy : public FComputeDataProviderRenderProxy
diff --git a/Engine/Plugins/PCG/Source/PCG/Private/Compute/DataInterfaces/PCGDataCollectionDataInterface.cpp b/Engine/Plugins/PCG/Source/PCG/Private/Compute/DataInterfaces/PCGDataCollectionDataInterface.cpp
--- a/Engine/Plugins/PCG/Source/PCG/Private/Compute/DataInterfaces/PCGDataCollectionDataInterface.cpp
+++ b/Engine/Plugins/PCG/Source/PCG/Private/Compute/DataInterfaces/PCGDataCollectionDataInterface.cpp
@@ -249,6 +249,17 @@ void UPCGDataCollectionDataInterface::GetSupportedInputs(TArray<FShaderFunctionD
 			.AddParam(EShaderFundamentalType::Uint, 2); // Name represented by 2 uints
 	}
 
+	// Atomics
+	{
+		OutFunctions.AddDefaulted_GetRef()
+			.SetName(TEXT("AtomicAddInt"))
+			.AddReturnType(EShaderFundamentalType::Int) // Value before it was incremented
+			.AddParam(EShaderFundamentalType::Uint) // DataIndex
+			.AddParam(EShaderFundamentalType::Uint) // ElementIndex
+			.AddParam(EShaderFundamentalType::Int) // AttributeId
+			.AddParam(EShaderFundamentalType::Int); // ValueToAdd
+	}
+
 	// Point Attribute Getters
 	{
 		OutFunctions.AddDefaulted_GetRef()
@@ -465,6 +476,8 @@ void UPCGDataCollectionDataProvider::Initialize(const UComputeDataInterface* InD
 
 	ExportMode = DataInterface->bRequiresExport ? EPCGExportMode::ComputeGraphOutput : EPCGExportMode::NoExport;
 
+	bRequiresZeroInitialization = DataInterface->bRequiresZeroInitialization;
+
 	ElementCountMultiplier = DataInterface->ElementCountMultiplier;
 
 	bExposeElementCounters = DataInterface->bExposeElementCounters;
@@ -513,7 +526,7 @@ FComputeDataProviderRenderProxy* UPCGDataCollectionDataProvider::GetRenderProxy(
 {
 	TRACE_CPUPROFILER_EVENT_SCOPE(UPCGDataCollectionDataProvider::GetRenderProxy);
 
-	FPCGDataCollectionDataProviderProxy* Proxy = new FPCGDataCollectionDataProviderProxy(this, PinDesc, ExportMode, ElementCountMultiplier, bExposeElementCounters, OutputPinLabel, OutputPinLabelAlias);
+	FPCGDataCollectionDataProviderProxy* Proxy = new FPCGDataCollectionDataProviderProxy(this, PinDesc, ExportMode, bRequiresZeroInitialization, ElementCountMultiplier, bExposeElementCounters, OutputPinLabel, OutputPinLabelAlias);
 
 #if WITH_EDITOR
 	if (const UPCGNode* Node = Cast<UPCGNode>(ProducerSettings ? ProducerSettings->GetOuter() : nullptr))
@@ -596,6 +609,7 @@ void UPCGDataCollectionDataProvider::Reset()
 	PinDesc = {};
 	bReadbackComplete = false;
 	ReadbackCounters.Empty();
+	bRequiresZeroInitialization = false;
 
 	Super::Reset();
 }
@@ -604,11 +618,13 @@ FPCGDataCollectionDataProviderProxy::FPCGDataCollectionDataProviderProxy(
 	TWeakObjectPtr<UPCGDataCollectionDataProvider> InDataProvider,
 	const FPCGDataCollectionDesc& InPinDesc,
 	EPCGExportMode InExportMode,
+	bool bInZeroInitialize,
 	uint32 InElementCountMultiplier,
 	bool bInExposeElementCounters,
 	FName InOutputPinLabel,
 	FName InOutputPinLabelAlias)
 	: ExportMode(InExportMode)
+	, bZeroInitialize(bInZeroInitialize)
 	, ElementCountMultiplier(InElementCountMultiplier)
 	, bExposeElementCounters(bInExposeElementCounters)
 	, PinDesc(InPinDesc)
@@ -709,11 +725,21 @@ void FPCGDataCollectionDataProviderProxy::AllocateResources(FRDGBuilder& GraphBu
 	// Initialize with an empty data collection. The kernel may not run, for example if indirect dispatch args end up being 0. This ensures
 	// there is something meaningful to readback.
 	TArray<uint32> PackedDataCollection;
+	PackedDataCollection.Reserve(bZeroInitialize ? (SizeBytes / PackedDataCollection.GetTypeSize()) : PinDesc.ComputePackedHeaderSizeBytes());
+
 	{
 		TRACE_CPUPROFILER_EVENT_SCOPE(FPCGDataCollectionDataProviderProxy::AllocateResources::WriteHeader);
 		PinDesc.WriteHeader(PackedDataCollection);
 	}
 
+	if (bZeroInitialize)
+	{
+		TRACE_CPUPROFILER_EVENT_SCOPE(FPCGDataCollectionDataProviderProxy::AllocateResources::ZeroInitialize);
+
+		// Fill remaining part of full buffer with 0s.
+		PackedDataCollection.SetNumZeroed(SizeBytes / PackedDataCollection.GetTypeSize());
+	}
+
 	GraphBuilder.QueueBufferUpload(Buffer, PackedDataCollection.GetData(), PackedDataCollection.Num() * PackedDataCollection.GetTypeSize(), ERDGInitialDataFlags::None);
 
 	// If buffer needs to be exported, can do it now. Pass back a reference to the main thread where it can be picked up by the compute graph element.
diff --git a/Engine/Plugins/PCG/Source/PCG/Private/Compute/DataInterfaces/PCGDataCollectionDataInterface.h b/Engine/Plugins/PCG/Source/PCG/Private/Compute/DataInterfaces/PCGDataCollectionDataInterface.h
--- a/Engine/Plugins/PCG/Source/PCG/Private/Compute/DataInterfaces/PCGDataCollectionDataInterface.h
+++ b/Engine/Plugins/PCG/Source/PCG/Private/Compute/DataInterfaces/PCGDataCollectionDataInterface.h
@@ -42,6 +42,8 @@ public:
 
 	void SetRequiresExport(bool bInRequiresExport) { bRequiresExport = bInRequiresExport; }
 
+	void SetRequiresZeroInitialization(bool bInZeroInit) { bRequiresZeroInitialization = bInZeroInit; }
+
 	void SetElementCountMultiplier(uint32 InElementCountMultiplier) { ElementCountMultiplier = InElementCountMultiplier; }
 
 	void SetExposeElementCounters(bool bInExposeElementCounters) { bExposeElementCounters = bInExposeElementCounters; }
@@ -51,6 +53,10 @@ public:
 	UPROPERTY()
 	TObjectPtr<const UPCGSettings> ProducerSettings;
 
+	/** Whether to perform full 0-initialization of the buffer. */
+	UPROPERTY()
+	bool bRequiresZeroInitialization = false;
+
 	/** Whether this data is passed to downstream tasks outside of this compute graph. */
 	UPROPERTY()
 	bool bRequiresExport = false;
@@ -107,6 +113,9 @@ public:
 	UPROPERTY()
 	uint32 ElementCountMultiplier = 0;
 
+	UPROPERTY()
+	bool bRequiresZeroInitialization = false;
+
 	UPROPERTY()
 	bool bExposeElementCounters = false;
 
@@ -139,6 +148,7 @@ public:
 		TWeakObjectPtr<UPCGDataCollectionDataProvider> InDataProvider,
 		const FPCGDataCollectionDesc& InPinDesc,
 		EPCGExportMode InExportMode,
+		bool bInZeroInitialize,
 		uint32 InElementCountMultiplier,
 		bool bInExposeElementCounters,
 		FName InOutputPinLabel,
@@ -167,6 +177,8 @@ protected:
 
 	EPCGExportMode ExportMode = EPCGExportMode::NoExport;
 
+	bool bZeroInitialize = false;
+
 	uint32 ElementCountMultiplier = 0;
 
 	bool bExposeElementCounters = false;
diff --git a/Engine/Plugins/PCG/Source/PCG/Private/Compute/DataInterfaces/PCGDataCollectionUploadDataInterface.cpp b/Engine/Plugins/PCG/Source/PCG/Private/Compute/DataInterfaces/PCGDataCollectionUploadDataInterface.cpp
--- a/Engine/Plugins/PCG/Source/PCG/Private/Compute/DataInterfaces/PCGDataCollectionUploadDataInterface.cpp
+++ b/Engine/Plugins/PCG/Source/PCG/Private/Compute/DataInterfaces/PCGDataCollectionUploadDataInterface.cpp
@@ -75,6 +75,7 @@ FComputeDataProviderRenderProxy* UPCGDataProviderDataCollectionUpload::GetRender
 		OutputPinLabel,
 		OutputPinLabelAlias,
 		ExportMode,
+		bRequiresZeroInitialization,
 		ElementCountMultiplier,
 		bExposeElementCounters);
 
@@ -108,9 +109,10 @@ FPCGDataProviderDataCollectionUploadProxy::FPCGDataProviderDataCollectionUploadP
 	FName InOutputPinLabel,
 	FName InOutputPinLabelAlias,
 	EPCGExportMode InExportMode,
+	bool bInZeroInitialize,
 	uint32 InElementCountMultiplier,
 	bool bInExposeElementCounters)
-	: FPCGDataCollectionDataProviderProxy(InDataProvider, InPinDesc, InExportMode, InElementCountMultiplier, bInExposeElementCounters, InOutputPinLabel, InOutputPinLabelAlias)
+	: FPCGDataCollectionDataProviderProxy(InDataProvider, InPinDesc, InExportMode, bInZeroInitialize, InElementCountMultiplier, bInExposeElementCounters, InOutputPinLabel, InOutputPinLabelAlias)
 	, DataAdaptor(InDataAdaptor)
 	, DownstreamInputPinLabels(InDownstreamInputPinLabels)
 {
diff --git a/Engine/Plugins/PCG/Source/PCG/Private/Compute/DataInterfaces/PCGDataCollectionUploadDataInterface.h b/Engine/Plugins/PCG/Source/PCG/Private/Compute/DataInterfaces/PCGDataCollectionUploadDataInterface.h
--- a/Engine/Plugins/PCG/Source/PCG/Private/Compute/DataInterfaces/PCGDataCollectionUploadDataInterface.h
+++ b/Engine/Plugins/PCG/Source/PCG/Private/Compute/DataInterfaces/PCGDataCollectionUploadDataInterface.h
@@ -57,6 +57,7 @@ public:
 		FName InOutputPinLabel,
 		FName InOutputPinLabelAlias,
 		EPCGExportMode InExportMode,
+		bool bInZeroInitialize,
 		uint32 InElementCountMultiplier,
 		bool bInExposeElementCounters);
 
diff --git a/Engine/Plugins/PCG/Source/PCG/Private/Compute/DataInterfaces/PCGInstanceDataInterface.cpp b/Engine/Plugins/PCG/Source/PCG/Private/Compute/DataInterfaces/PCGInstanceDataInterface.cpp
--- a/Engine/Plugins/PCG/Source/PCG/Private/Compute/DataInterfaces/PCGInstanceDataInterface.cpp
+++ b/Engine/Plugins/PCG/Source/PCG/Private/Compute/DataInterfaces/PCGInstanceDataInterface.cpp
@@ -139,21 +139,34 @@ void UPCGInstanceDataProvider::Initialize(const UComputeDataInterface* InDataInt
 
 	const UPCGInstanceDataInterface* DataInterface = CastChecked<UPCGInstanceDataInterface>(InDataInterface);
 	UPCGDataBinding* Binding = CastChecked<UPCGDataBinding>(InBinding);
-	check(DataInterface->ProducerKernel);
 
-	// If there were 0 input points (or too many input points) for this execution, we will not have created any primitives, so check for null.
-	if (FPCGSpawnerPrimitives* FoundPrimitives = Binding->MeshSpawnersToPrimitives.Find(Cast<UPCGStaticMeshSpawnerKernel>(DataInterface->ProducerKernel)))
+	ProducerKernel = CastChecked<UPCGStaticMeshSpawnerKernel>(DataInterface->ProducerKernel);
+	check(ProducerKernel);
+}
+
+bool UPCGInstanceDataProvider::PrepareForExecute_GameThread(UPCGDataBinding* InBinding)
+{
+	TRACE_CPUPROFILER_EVENT_SCOPE(UPCGInstanceDataProvider::PrepareForExecute_GameThread);
+	UPCGDataBinding* Binding = CastChecked<UPCGDataBinding>(InBinding);
+
+	FPCGSpawnerPrimitives* FoundPrimitives = Binding->MeshSpawnersToPrimitives.Find(Cast<UPCGStaticMeshSpawnerKernel>(ProducerKernel));
+
+	if (!FoundPrimitives)
 	{
-		if (!ensure(FoundPrimitives->Primitives.Num() <= PCGComputeConstants::MAX_PRIMITIVE_COMPONENTS_PER_SPAWNER))
-		{
-			// Last resort - should be clamped earlier during setup and not come to this.
-			FoundPrimitives->Primitives.SetNum(PCGComputeConstants::MAX_PRIMITIVE_COMPONENTS_PER_SPAWNER);
-		}
+		return true;
+	}
 
-		Primitives = FoundPrimitives->Primitives;
-		NumInstancesAllPrimitives = FoundPrimitives->NumInstancesAllPrimitives;
-		NumCustomFloatsPerInstance = FoundPrimitives->NumCustomFloats;
+	if (!ensure(FoundPrimitives->Primitives.Num() <= PCGComputeConstants::MAX_PRIMITIVE_COMPONENTS_PER_SPAWNER))
+	{
+		// Last resort - should be clamped earlier during setup and not come to this.
+		FoundPrimitives->Primitives.SetNum(PCGComputeConstants::MAX_PRIMITIVE_COMPONENTS_PER_SPAWNER);
 	}
+
+	Primitives = FoundPrimitives->Primitives;
+	NumInstancesAllPrimitives = FoundPrimitives->NumInstancesAllPrimitives;
+	NumCustomFloatsPerInstance = FoundPrimitives->NumCustomFloats;
+
+	return true;
 }
 
 FComputeDataProviderRenderProxy* UPCGInstanceDataProvider::GetRenderProxy()
@@ -177,6 +190,7 @@ FComputeDataProviderRenderProxy* UPCGInstanceDataProvider::GetRenderProxy()
 
 void UPCGInstanceDataProvider::Reset()
 {
+	ProducerKernel = nullptr;
 	Primitives.Empty();
 	NumInstancesAllPrimitives = 0;
 	NumCustomFloatsPerInstance = 0;
diff --git a/Engine/Plugins/PCG/Source/PCG/Private/Compute/DataInterfaces/PCGInstanceDataInterface.h b/Engine/Plugins/PCG/Source/PCG/Private/Compute/DataInterfaces/PCGInstanceDataInterface.h
--- a/Engine/Plugins/PCG/Source/PCG/Private/Compute/DataInterfaces/PCGInstanceDataInterface.h
+++ b/Engine/Plugins/PCG/Source/PCG/Private/Compute/DataInterfaces/PCGInstanceDataInterface.h
@@ -5,13 +5,13 @@
 #include "Compute/DataInterfaces/PCGComputeDataInterface.h"
 
 #include "RHIResources.h"
-#include "ComputeFramework/ComputeDataProvider.h"
 
 #include "PCGInstanceDataInterface.generated.h"
 
 class FPCGInstanceDataInterfaceParameters;
 class FPrimitiveSceneProxy;
 class UPCGSettings;
+class UPCGStaticMeshSpawnerKernel;
 
 /** Data Interface for writing instance data (transform and custom floats) to an intermediate buffer, which is then injected into the GPU Scene. */
 UCLASS(ClassGroup = (Procedural))
@@ -41,17 +41,21 @@ private:
 
 /** Compute Framework Data Provider for writing to the GPU Scene. */
 UCLASS()
-class UPCGInstanceDataProvider : public UComputeDataProvider
+class UPCGInstanceDataProvider : public UPCGComputeDataProvider
 {
 	GENERATED_BODY()
 
 public:
 	//~ Begin UComputeDataProvider Interface
 	virtual void Initialize(const UComputeDataInterface* InDataInterface, UObject* InBinding, uint64 InInputMask, uint64 InOutputMask) override;
+	virtual bool PrepareForExecute_GameThread(UPCGDataBinding* InBinding) override;
 	FComputeDataProviderRenderProxy* GetRenderProxy() override;
 	virtual void Reset() override;
 	//~ End UComputeDataProvider Interface
 
+	UPROPERTY()
+	TObjectPtr<const UPCGStaticMeshSpawnerKernel> ProducerKernel;
+
 	UPROPERTY()
 	TArray<TObjectPtr<UPrimitiveComponent>> Primitives;
 
diff --git a/Engine/Plugins/PCG/Source/PCG/Private/Compute/Elements/PCGComputeGraphElement.cpp b/Engine/Plugins/PCG/Source/PCG/Private/Compute/Elements/PCGComputeGraphElement.cpp
--- a/Engine/Plugins/PCG/Source/PCG/Private/Compute/Elements/PCGComputeGraphElement.cpp
+++ b/Engine/Plugins/PCG/Source/PCG/Private/Compute/Elements/PCGComputeGraphElement.cpp
@@ -199,13 +199,6 @@ bool FPCGComputeGraphElement::ExecuteInternal(FPCGContext* InContext) const
 			return true;
 		}
 
-		const bool bAnyComponentsSetup = SetupProceduralISMComponents(Context, DataBinding);
-
-		if (bAnyComponentsSetup)
-		{
-			Context->SourceComponent->NotifyProceduralInstancesInUse();
-		}
-
 		const UPCGGraph* Graph = Context->Stack ? Context->Stack->GetNearestNonInlinedGraphForCurrentFrame() : nullptr;
 		const UPCGGraph::FComputeGraphInstanceKey Key = { Context->SourceComponent->GetGenerationGridSize(), ComputeGraphIndex };
 
@@ -290,48 +283,10 @@ bool FPCGComputeGraphElement::ExecuteInternal(FPCGContext* InContext) const
 				});
 			}
 		}
-
-		if (bAnyComponentsSetup)
-		{
-			// Delay to give time for proxy to settle and for scene update to allocate space in the GPU scene.
-			SleepUntilNextFrame(Context);
-			return false;
-		}
 	}
 
 	check(Context->DataBinding && InContext->SourceComponent.Get());
 
-	if (!Context->bPrimitiveProxiesValidated)
-	{
-		bool bAllProxiesCreated = true;
-
-		for (const TPair<TObjectPtr<const UPCGStaticMeshSpawnerKernel>, FPCGSpawnerPrimitives>& SpawnerToPrimitives : Context->DataBinding->MeshSpawnersToPrimitives)
-		{
-			for (const TObjectPtr<UPrimitiveComponent>& Primitive : SpawnerToPrimitives.Value.Primitives)
-			{
-				const UPCGProceduralISMComponent* PISMC = Cast<UPCGProceduralISMComponent>(Primitive);
-				if (PISMC && PISMC->GetNumInstances() > 0 && !PISMC->GetSceneProxy())
-				{
-					bAllProxiesCreated = false;
-					break;
-				}
-			}
-
-			if (!bAllProxiesCreated)
-			{
-				break;
-			}
-		}
-
-		if (!bAllProxiesCreated)
-		{
-			SleepUntilNextFrame(Context);
-			return false;
-		}
-
-		Context->bPrimitiveProxiesValidated = true;
-	}
-
 	// Allow data providers to do early readbacks, which some may use to download analysis results which feed into final data descriptions.
 	if (!Context->bDataProvidersReadbacksComplete)
 	{
@@ -464,29 +419,6 @@ bool FPCGComputeGraphElement::ExecuteInternal(FPCGContext* InContext) const
 	}
 }
 
-bool FPCGComputeGraphElement::SetupProceduralISMComponents(FPCGContext* InContext, UPCGDataBinding* InBinding) const
-{
-	TRACE_CPUPROFILER_EVENT_SCOPE(FPCGComputeGraphElement::SetupProceduralISMComponents);
-	check(InContext);
-	check(InBinding);
-
-	FPCGComputeGraphContext* Context = static_cast<FPCGComputeGraphContext*>(InContext);
-
-	bool bAnyComponentsSetup = false;
-
-	for (const UPCGStaticMeshSpawnerKernel* Kernel : Context->ComputeGraph->StaticMeshSpawners)
-	{
-		if (!ensure(Kernel))
-		{
-			continue;
-		}
-
-		bAnyComponentsSetup |= Kernel->SetupProceduralISMComponents(InContext, InBinding);
-	}
-
-	return bAnyComponentsSetup;
-}
-
 void FPCGComputeGraphElement::PostExecuteInternal(FPCGContext* InContext) const
 {
 	check(InContext);
diff --git a/Engine/Plugins/PCG/Source/PCG/Private/Compute/Elements/PCGComputeGraphElement.h b/Engine/Plugins/PCG/Source/PCG/Private/Compute/Elements/PCGComputeGraphElement.h
--- a/Engine/Plugins/PCG/Source/PCG/Private/Compute/Elements/PCGComputeGraphElement.h
+++ b/Engine/Plugins/PCG/Source/PCG/Private/Compute/Elements/PCGComputeGraphElement.h
@@ -93,8 +93,6 @@ protected:
 	// TODO - need to accumulate dependencies from compute graph nodes.
 	virtual bool IsCacheable(const UPCGSettings* InSettings) const override { return false; }
 
-	bool SetupProceduralISMComponents(FPCGContext* InContext, UPCGDataBinding* Binding) const;
-
 	void ResetAsyncOperations(FPCGContext* InContext) const;
 
 #if WITH_EDITOR
diff --git a/Engine/Plugins/PCG/Source/PCG/Private/Compute/PCGComputeKernel.h b/Engine/Plugins/PCG/Source/PCG/Private/Compute/PCGComputeKernel.h
--- a/Engine/Plugins/PCG/Source/PCG/Private/Compute/PCGComputeKernel.h
+++ b/Engine/Plugins/PCG/Source/PCG/Private/Compute/PCGComputeKernel.h
@@ -63,6 +63,9 @@ public:
 	/** Get any data labels that are used statically by this node. */
 	virtual void GetDataLabels(FName InPinLabel, TArray<FString>& OutDataLabels) const {}
 
+	/** Initialize all buffer data to 0. This is not free so should only be used when kernel execution requires 0-initialized data, such as for counters. */
+	virtual bool DoesOutputPinRequireZeroInitialization(FName InOutputPinLabel) const { return false; }
+
 	/** Whether the output pin should have a buffer of counters attached, useful for compaction and other counting requirements. */
 	virtual bool DoesOutputPinRequireElementCounters(FName InOutputPinLabel) const { return false; }
 
diff --git a/Engine/Plugins/PCG/Source/PCG/Private/Compute/PCGDataBinding.cpp b/Engine/Plugins/PCG/Source/PCG/Private/Compute/PCGDataBinding.cpp
--- a/Engine/Plugins/PCG/Source/PCG/Private/Compute/PCGDataBinding.cpp
+++ b/Engine/Plugins/PCG/Source/PCG/Private/Compute/PCGDataBinding.cpp
@@ -160,6 +160,55 @@ TWeakObjectPtr<UPCGComponent> UPCGDataBinding::GetSourceComponent() const
 	return nullptr;
 }
 
+int32 UPCGDataBinding::GetFirstInputDataIndex(const UPCGComputeKernel* InKernel, FName InPinLabel) const
+{
+	int32 FirstDataIndex = INDEX_NONE;
+
+	const int32 BindingIndex = Graph ? Graph->GetBindingIndex(InKernel->GetKernelIndex(), InPinLabel, /*bIsInput=*/true) : INDEX_NONE;
+
+	if (ensure(BindingIndex != INDEX_NONE))
+	{
+		if (const FName* VirtualLabel = Graph->GetVirtualPinLabelFromBindingIndex(BindingIndex))
+		{
+			FirstDataIndex = InputDataCollection.TaggedData.IndexOfByPredicate([VirtualLabel](const FPCGTaggedData& InData)
+			{
+				return InData.Pin == *VirtualLabel;
+			});
+		}
+	}
+
+	return FirstDataIndex;
+}
+
+bool UPCGDataBinding::ReadbackInputDataToCPU(int32 InInputDataIndex)
+{
+	const UPCGProxyForGPUData* DataWithGPUSupport = InputDataCollection.TaggedData.IsValidIndex(InInputDataIndex) ? Cast<UPCGProxyForGPUData>(InputDataCollection.TaggedData[InInputDataIndex].Data) : nullptr;
+
+	if (!DataWithGPUSupport)
+	{
+		// No work to do, signal completion.
+		return true;
+	}
+
+	TSharedPtr<FPCGContextHandle> Context = ContextHandle.Pin();
+
+	UPCGProxyForGPUData::FReadbackResult Result = DataWithGPUSupport->GetCPUData(Context ? Context->GetContext() : nullptr);
+
+	if (!Result.bComplete)
+	{
+		// Read back pending - wait for it to complete.
+		return false;
+	}
+
+	if (ensure(Result.TaggedData.Data))
+	{
+		InputDataCollection.TaggedData[InInputDataIndex].Data = Result.TaggedData.Data;
+		InputDataCollection.TaggedData[InInputDataIndex].Tags = MoveTemp(Result.TaggedData.Tags);
+	}
+
+	return true;
+}
+
 void UPCGDataBinding::InitializeInputData(const FPCGDataCollection& InComputeGraphElementInputData)
 {
 	TRACE_CPUPROFILER_EVENT_SCOPE(UPCGDataBinding::InitializeInputData);
diff --git a/Engine/Plugins/PCG/Source/PCG/Private/Compute/PCGDataBinding.h b/Engine/Plugins/PCG/Source/PCG/Private/Compute/PCGDataBinding.h
--- a/Engine/Plugins/PCG/Source/PCG/Private/Compute/PCGDataBinding.h
+++ b/Engine/Plugins/PCG/Source/PCG/Private/Compute/PCGDataBinding.h
@@ -89,6 +89,14 @@ public:
 
 	TWeakObjectPtr<UPCGComponent> GetSourceComponent() const;
 
+	/** Helper to get index in input data collection of the first data item for the given kernel and input pin label. */
+	int32 GetFirstInputDataIndex(const UPCGComputeKernel* InKernel, FName InPinLabel) const;
+
+	/** If there is a GPU proxy at the given index of the input data collection, triggers a readback and replaces the data item with CPU data if readback succeeds.
+	* Returns false while readback is in progress.
+	*/
+	bool ReadbackInputDataToCPU(int32 InInputDataIndex);
+
 private:
 	void InitializeInputData(const FPCGDataCollection& InComputeGraphElementInputData);
 
diff --git a/Engine/Plugins/PCG/Source/PCG/Private/Compute/PCGDataForGPU.cpp b/Engine/Plugins/PCG/Source/PCG/Private/Compute/PCGDataForGPU.cpp
--- a/Engine/Plugins/PCG/Source/PCG/Private/Compute/PCGDataForGPU.cpp
+++ b/Engine/Plugins/PCG/Source/PCG/Private/Compute/PCGDataForGPU.cpp
@@ -613,12 +613,14 @@ uint32 GetTypeHash(const FPCGKernelAttributeKey& In)
 int32 FPCGKernelAttributeTable::GetAttributeId(const FPCGKernelAttributeKey& InAttribute) const
 {
 	const int TableIndex = AttributeTable.IndexOfByKey(InAttribute);
-	if (ensureMsgf(TableIndex >= 0, TEXT("Missing attribute '%s', has this attribute been declared using GetKernelAttributeKeys()?"), *InAttribute.Name.ToString()))
+	if (TableIndex >= 0)
 	{
 		return PCGComputeHelpers::GetAttributeIdFromMetadataAttributeIndex(TableIndex);
 	}
 	else
 	{
+		// Attribute with the given name and type is not present. If this is unexpected, check this attribute table for attributes
+		// in case the supplied name or type is wrong, and ensure any attribute created by a kernel is declared using GetKernelAttributeKeys().
 		return INDEX_NONE;
 	}
 }
diff --git a/Engine/Plugins/PCG/Source/PCG/Private/Elements/PCGCopyPointsKernel.cpp b/Engine/Plugins/PCG/Source/PCG/Private/Elements/PCGCopyPointsKernel.cpp
--- a/Engine/Plugins/PCG/Source/PCG/Private/Elements/PCGCopyPointsKernel.cpp
+++ b/Engine/Plugins/PCG/Source/PCG/Private/Elements/PCGCopyPointsKernel.cpp
@@ -76,26 +76,11 @@ bool UPCGCopyPointsKernel::ComputeOutputBindingDataDesc(const UPCGComputeGraph*
 	
 	if (Settings->bMatchBasedOnAttribute)
 	{
-		int32 FirstAnalysisDataIndex = INDEX_NONE;
+		const int32 AnalysisDataIndex = InBinding->GetFirstInputDataIndex(this, PCGCopyPointsConstants::SelectedFlagsPinLabel);
 
-		if (InBinding->Graph)
+		if (AnalysisDataIndex != INDEX_NONE)
 		{
-			const int32 BindingIndex = InBinding->Graph->GetBindingIndex(KernelIndex, PCGCopyPointsConstants::SelectedFlagsPinLabel, /*bIsInput=*/true);
-			if (ensure(BindingIndex != INDEX_NONE))
-			{
-				if (const FName* VirtualLabel = InGraph->GetVirtualPinLabelFromBindingIndex(BindingIndex))
-				{
-					FirstAnalysisDataIndex = InBinding->InputDataCollection.TaggedData.IndexOfByPredicate([VirtualLabel](const FPCGTaggedData& InData)
-					{
-						return InData.Pin == *VirtualLabel;
-					});
-				}
-			}
-		}
-
-		if (FirstAnalysisDataIndex != INDEX_NONE)
-		{
-			const UPCGParamData* AnalysisResultsData = Cast<UPCGParamData>(InBinding->InputDataCollection.TaggedData[FirstAnalysisDataIndex].Data);
+			const UPCGParamData* AnalysisResultsData = Cast<UPCGParamData>(InBinding->InputDataCollection.TaggedData[AnalysisDataIndex].Data);
 			const UPCGMetadata* AnalysisMetadata = AnalysisResultsData ? AnalysisResultsData->ConstMetadata() : nullptr;
 			const FPCGMetadataAttributeBase* AnalysisAttributeBase = AnalysisMetadata ? AnalysisMetadata->GetConstAttribute(PCGCopyPointsConstants::SelectedFlagAttributeName) : nullptr;
 
@@ -118,7 +103,7 @@ bool UPCGCopyPointsKernel::ComputeOutputBindingDataDesc(const UPCGComputeGraph*
 			else
 			{
 				UE_LOG(LogPCG, Warning, TEXT("No analysis data received by copy points kernel, no points will be copied."));
-				if (InBinding->InputDataCollection.TaggedData[FirstAnalysisDataIndex].Data->IsA<UPCGProxyForGPUData>())
+				if (InBinding->InputDataCollection.TaggedData[AnalysisDataIndex].Data->IsA<UPCGProxyForGPUData>())
 				{
 					UE_LOG(LogPCG, Error, TEXT("Data was not read back."));
 				}
diff --git a/Engine/Plugins/PCG/Source/PCG/Private/Elements/PCGStaticMeshSpawner.cpp b/Engine/Plugins/PCG/Source/PCG/Private/Elements/PCGStaticMeshSpawner.cpp
--- a/Engine/Plugins/PCG/Source/PCG/Private/Elements/PCGStaticMeshSpawner.cpp
+++ b/Engine/Plugins/PCG/Source/PCG/Private/Elements/PCGStaticMeshSpawner.cpp
@@ -6,6 +6,7 @@
 #include "PCGCustomVersion.h"
 #include "PCGManagedResource.h"
 #include "Compute/PCGComputeCommon.h"
+#include "Compute/BuiltInKernels/PCGCountUniqueAttributeValuesKernel.h"
 #include "Data/PCGBasePointData.h"
 #include "Data/PCGSpatialData.h"
 #include "Elements/PCGStaticMeshSpawnerContext.h"
@@ -55,6 +56,17 @@ void UPCGStaticMeshSpawnerSettings::CreateKernels(FPCGGPUCompilationContext& InO
 	// Connect node pins to kernel pins
 	OutEdges.Emplace(FPCGPinReference(PCGPinConstants::DefaultInputLabel), FPCGPinReference(Kernel, PCGPinConstants::DefaultInputLabel));
 	OutEdges.Emplace(FPCGPinReference(Kernel, PCGPinConstants::DefaultOutputLabel), FPCGPinReference(PCGPinConstants::DefaultOutputLabel));
+
+	// If doing by-attribute selection, add analysis kernel that will count how many instances of each mesh are present.
+	if (const UPCGMeshSelectorByAttribute* Selector = Cast<UPCGMeshSelectorByAttribute>(MeshSelectorParameters))
+	{
+		UPCGCountUniqueAttributeValuesKernel* CountKernel = InOutContext.NewObject_AnyThread<UPCGCountUniqueAttributeValuesKernel>(InObjectOuter);
+		CountKernel->Initialize(Selector->AttributeName, bDumpDataDescriptions);
+		OutKernels.Add(CountKernel);
+
+		OutEdges.Emplace(FPCGPinReference(PCGPinConstants::DefaultInputLabel), FPCGPinReference(CountKernel, PCGPinConstants::DefaultInputLabel));
+		OutEdges.Emplace(FPCGPinReference(CountKernel, PCGPinConstants::DefaultOutputLabel), FPCGPinReference(Kernel, PCGStaticMeshSpawnerConstants::InstanceCountsPinLabel));
+	}
 }
 
 FText UPCGStaticMeshSpawnerSettings::GetDefaultNodeTitle() const
diff --git a/Engine/Plugins/PCG/Source/PCG/Private/Elements/PCGStaticMeshSpawnerKernel.cpp b/Engine/Plugins/PCG/Source/PCG/Private/Elements/PCGStaticMeshSpawnerKernel.cpp
--- a/Engine/Plugins/PCG/Source/PCG/Private/Elements/PCGStaticMeshSpawnerKernel.cpp
+++ b/Engine/Plugins/PCG/Source/PCG/Private/Elements/PCGStaticMeshSpawnerKernel.cpp
@@ -5,10 +5,13 @@
 #include "PCGComponent.h"
 #include "PCGContext.h"
 #include "PCGModule.h"
+#include "PCGParamData.h"
 #include "Components/PCGProceduralISMComponent.h"
 #include "Compute/PCGComputeGraph.h"
 #include "Compute/PCGDataBinding.h"
 #include "Compute/PCGPinPropertiesGPU.h"
+#include "Compute/BuiltInKernels/PCGCountUniqueAttributeValuesKernel.h"
+#include "Compute/Data/PCGProxyForGPUData.h"
 #include "Compute/DataInterfaces/PCGInstanceDataInterface.h"
 #include "Compute/DataInterfaces/Elements/PCGStaticMeshSpawnerDataInterface.h"
 #include "Elements/PCGStaticMeshSpawner.h"
@@ -26,8 +29,6 @@
 
 namespace PCGStaticMeshSpawnerKernel
 {
-	static const FText CouldNotLoadStaticMeshFormat = LOCTEXT("CouldNotLoadStaticMesh", "Could not load static mesh from path '{0}'.");
-	static const FText TooManyPrimitivesFormat = LOCTEXT("TooManyPrimitives", "Attempted to emit too many primitive components, terminated after creating '{0}'.");
 	static const FText NoMeshEntriesFormat = LOCTEXT("NoMeshEntries", "No mesh entries provided in weighted mesh selector.");
 }
 
@@ -216,6 +217,7 @@ void UPCGStaticMeshSpawnerKernel::GetKernelAttributeKeys(TArray<FPCGKernelAttrib
 void UPCGStaticMeshSpawnerKernel::GetInputPins(TArray<FPCGPinProperties>& OutPins) const
 {
 	OutPins.Emplace(PCGPinConstants::DefaultInputLabel, EPCGDataType::Point);
+	OutPins.Emplace(PCGStaticMeshSpawnerConstants::InstanceCountsPinLabel, EPCGDataType::Param);
 }
 
 void UPCGStaticMeshSpawnerKernel::GetOutputPins(TArray<FPCGPinPropertiesGPU>& OutPins) const
@@ -223,298 +225,4 @@ void UPCGStaticMeshSpawnerKernel::GetOutputPins(TArray<FPCGPinPropertiesGPU>& Ou
 	OutPins.Emplace(PCGPinConstants::DefaultOutputLabel, EPCGDataType::Point);
 }
 
-bool UPCGStaticMeshSpawnerKernel::SetupProceduralISMComponents(FPCGContext* InContext, UPCGDataBinding* InBinding) const
-{
-	AActor* TargetActor = Settings->TargetActor.Get() ? Settings->TargetActor.Get() : InContext->GetTargetActor(nullptr);
-	if (!ensure(Settings->MeshSelectorParameters) || !ensure(TargetActor))
-	{
-		return false;
-	}
-
-	FPCGDataCollectionDesc InputDataDesc;
-	InBinding->GetKernelPinDataDesc_GameThread(this, PCGPinConstants::DefaultInputLabel, /*bIsInput=*/true, InputDataDesc);
-
-	const uint32 TotalInputPointCount = InputDataDesc.ComputeDataElementCount(EPCGDataType::Point);
-	if (TotalInputPointCount == 0)
-	{
-		return false;
-	}
-
-	const FBox LocalBounds = InContext->SourceComponent->GetGridBounds().ShiftBy(-InContext->SourceComponent->GetOwner()->GetActorLocation());
-
-	uint32 CustomFloatCount = 0;
-
-	TArray<FUint32Vector4> AttributeIdOffsetStrides;
-	if (Settings->InstanceDataPackerParameters && Settings->InstanceDataPackerParameters->GetAttributeNames(/*OutNames=*/nullptr))
-	{
-		TArray<FName> AttributeNames;
-		Settings->InstanceDataPackerParameters->GetAttributeNames(&AttributeNames);
-
-		for (FName AttributeName : AttributeNames)
-		{
-			if (AttributeName == NAME_None)
-			{
-				PCG_KERNEL_VALIDATION_ERR(InContext, Settings, /*bQuiet=*/false,
-					FText::Format(LOCTEXT("InstanceDataAttributeInvalid", "Invalid instance data attribute specified '{0}'."), FText::FromName(AttributeName)));
-
-				continue;
-			}
-
-			if (!InputDataDesc.ContainsAttributeOnAnyData(AttributeName))
-			{
-				PCG_KERNEL_VALIDATION_ERR(InContext, Settings, /*bQuiet=*/false,
-					FText::Format(LOCTEXT("InstanceDataAttributeNotFound", "Instance data attribute '{0}' not found."), FText::FromName(AttributeName)));
-
-				continue;
-			}
-		}
-
-		PCGDataForGPUHelpers::ComputeCustomFloatPacking(InContext, Settings, AttributeNames, InBinding, InputDataDesc, CustomFloatCount, AttributeIdOffsetStrides);
-	}
-
-	TArray<float> PrimitiveSelectionCDF;
-	int32 SelectorAttributeId = INDEX_NONE;
-	TArray<int32> PrimitiveStringKeys;
-	TArray<FBox> PrimitiveMeshBounds;
-	TArray<FPCGProceduralISMComponentDescriptor> ComponentsToCreate;
-
-	if (const UPCGMeshSelectorByAttribute* SelectorByAttribute = Cast<UPCGMeshSelectorByAttribute>(Settings->MeshSelectorParameters))
-	{
-		const FName SelectorName = SelectorByAttribute->AttributeName;
-
-		// Compute how many instances we expect for each mesh.
-		TMap<int32, uint32> StringKeyToInstanceCount;
-
-		for (const FPCGDataDesc& Desc : InputDataDesc.DataDescs)
-		{
-			if (Desc.ElementCount <= 0)
-			{
-				continue;
-			}
-
-			for (const FPCGKernelAttributeDesc& AttributeDesc : Desc.AttributeDescs)
-			{
-				if (AttributeDesc.Name == SelectorName)
-				{
-					if (AttributeDesc.Type == EPCGKernelAttributeType::StringKey)
-					{
-						ensure(SelectorAttributeId == INDEX_NONE || SelectorAttributeId == AttributeDesc.Index);
-						SelectorAttributeId = AttributeDesc.Index;
-
-						for (int32 AttributeStringKeyValue : AttributeDesc.GetUniqueStringKeys())
-						{
-							if (AttributeStringKeyValue > 0)
-							{
-								StringKeyToInstanceCount.FindOrAdd(AttributeStringKeyValue, 0) += Desc.ElementCount;
-							}
-						}
-
-						// We could early out but we currently continue looping to generate the below warning for mismatched types.
-					}
-					else
-					{
-						// Currently only a single type per attribute name is supported (the name/type in the attribute table). It's possible to wire a graph
-						// for which an attribute is present with multiple types. Warn if this is encountered.
-						PCG_KERNEL_VALIDATION_WARN(InContext, Settings, /*bQuiet=*/false, FText::Format(
-							LOCTEXT("MeshSelectorAttributeNotUsable", "Attribute '{0}' not usable for mesh selection, only attributes of type String Key are supported."),
-							FText::FromName(SelectorName)));
-					}
-				}
-			}
-		}
-
-		if (SelectorAttributeId == INDEX_NONE)
-		{
-			PCG_KERNEL_VALIDATION_ERR(InContext, Settings, /*bQuiet=*/false, FText::Format(
-				LOCTEXT("MeshSelectorAttributeNotFound", "Mesh selector attribute '{0}' not found."),
-				FText::FromName(SelectorName)));
-
-			return false;
-		}
-
-		for (const TPair<int32, uint32>& StringKeyAndInstanceCount : StringKeyToInstanceCount)
-		{
-			if (!ensure(InBinding->GetStringTable().IsValidIndex(StringKeyAndInstanceCount.Key)))
-			{
-				continue;
-			}
-
-			const FString& MeshPathString = InBinding->GetStringTable()[StringKeyAndInstanceCount.Key];
-			if (MeshPathString.IsEmpty())
-			{
-				continue;
-			}
-
-			if (ComponentsToCreate.Num() >= PCGComputeConstants::MAX_PRIMITIVE_COMPONENTS_PER_SPAWNER)
-			{
-				PCG_KERNEL_VALIDATION_WARN(InContext, Settings, /*bQuiet=*/false,
-					FText::Format(PCGStaticMeshSpawnerKernel::TooManyPrimitivesFormat, FText::FromString(FString::FromInt(PCGComputeConstants::MAX_PRIMITIVE_COMPONENTS_PER_SPAWNER))));
-
-				break;
-			}
-
-			UStaticMesh* StaticMesh = Cast<UStaticMesh>(FSoftObjectPath(MeshPathString).TryLoad());
-			if (!StaticMesh)
-			{
-				PCG_KERNEL_VALIDATION_ERR(InContext, Settings, /*bQuiet=*/false,
-					FText::Format(PCGStaticMeshSpawnerKernel::CouldNotLoadStaticMeshFormat, FText::FromString(MeshPathString)));
-
-				continue;
-			}
-
-			FPCGProceduralISMComponentDescriptor Descriptor;
-			Descriptor = SelectorByAttribute->TemplateDescriptor;
-			Descriptor.NumInstances = StringKeyAndInstanceCount.Value;
-			Descriptor.LocalBounds = LocalBounds;
-			Descriptor.NumCustomFloats = CustomFloatCount;
-			Descriptor.StaticMesh = StaticMesh;
-
-			PrimitiveStringKeys.Add(StringKeyAndInstanceCount.Key);
-			ComponentsToCreate.Add(MoveTemp(Descriptor));
-
-			if (Settings->bApplyMeshBoundsToPoints)
-			{
-				PrimitiveMeshBounds.Add(Descriptor.StaticMesh->GetBoundingBox());
-			}
-		}
-
-		PrimitiveSelectionCDF.SetNumZeroed(ComponentsToCreate.Num());
-	}
-	else if (const UPCGMeshSelectorWeighted* SelectorWeighted = Cast<UPCGMeshSelectorWeighted>(Settings->MeshSelectorParameters))
-	{
-		if (SelectorWeighted->MeshEntries.IsEmpty())
-		{
-			PCG_KERNEL_VALIDATION_ERR(InContext, Settings, /*bQuiet=*/false, PCGStaticMeshSpawnerKernel::NoMeshEntriesFormat);
-
-			return false;
-		}
-
-		float CumulativeWeight = 0.0f;
-
-		float TotalWeight = 0.0f;
-		for (const FPCGMeshSelectorWeightedEntry& Entry : SelectorWeighted->MeshEntries)
-		{
-			TotalWeight += Entry.Weight;
-		}
-
-		if (!ensure(TotalWeight > UE_SMALL_NUMBER))
-		{
-			return false;
-		}
-
-		PrimitiveSelectionCDF.Reserve(SelectorWeighted->MeshEntries.Num());
-
-		for (const FPCGMeshSelectorWeightedEntry& Entry : SelectorWeighted->MeshEntries)
-		{
-			if (Entry.Descriptor.StaticMesh.IsNull())
-			{
-				continue;
-			}
-
-			if (ComponentsToCreate.Num() >= PCGComputeConstants::MAX_PRIMITIVE_COMPONENTS_PER_SPAWNER)
-			{
-				PCG_KERNEL_VALIDATION_WARN(InContext, Settings, /*bQuiet=*/false,
-					FText::Format(PCGStaticMeshSpawnerKernel::TooManyPrimitivesFormat, FText::FromString(FString::FromInt(PCGComputeConstants::MAX_PRIMITIVE_COMPONENTS_PER_SPAWNER))));
-
-				break;
-			}
-
-			if (UStaticMesh* StaticMesh = Entry.Descriptor.StaticMesh.LoadSynchronous())
-			{
-				const float Weight = float(Entry.Weight) / TotalWeight;
-				CumulativeWeight += Weight;
-				PrimitiveSelectionCDF.Add(CumulativeWeight);
-				PrimitiveStringKeys.Add(InBinding->GetStringTable().IndexOfByKey(Entry.Descriptor.StaticMesh.ToString()));
-
-				FPCGProceduralISMComponentDescriptor Descriptor;
-				Descriptor = Entry.Descriptor;
-				Descriptor.LocalBounds = LocalBounds;
-				Descriptor.NumCustomFloats = CustomFloatCount;
-				Descriptor.StaticMesh = StaticMesh;
-
-				uint32 InstanceCount = FMath::CeilToInt(TotalInputPointCount * Weight);
-
-				if (SelectorWeighted->MeshEntries.Num() > 1)
-				{
-					// Since we'll be selecting meshes based on random draws using the point random seeds which we don't have on CPU,
-					// we may pick more or less than the expected number of instances for each mesh. Use binomial variance to calculate
-					// the overallocation that gives 99.7% confidence (3 sigma).
-					const double Variance = double(TotalInputPointCount) * Weight * (1.0f - Weight);
-					const double Sigma = FMath::Sqrt(Variance);
-					const uint32 AdditionalAllocation = FMath::CeilToInt(3.0f * Sigma);
-
-					// Useful for debugging instance counts.
-					//UE_LOG(LogPCG, Warning, TEXT("Over allocation, N = %d, P = %.2f, NP = %d, 3sig = %f, Extra = %d (%.2f%%), Calculated = %d, Final = %d"), TotalInputPointCount, Weight, InstanceCount, float(3.0f * Sigma), AdditionalAllocation, float(AdditionalAllocation) / TotalInputPointCount, InstanceCount + AdditionalAllocation, FMath::Min(InstanceCount + AdditionalAllocation, TotalInputPointCount));
-
-					InstanceCount += AdditionalAllocation;
-				}
-
-				Descriptor.NumInstances = FMath::Min(InstanceCount, TotalInputPointCount);
-
-				ComponentsToCreate.Add(MoveTemp(Descriptor));
-
-				if (Settings->bApplyMeshBoundsToPoints)
-				{
-					PrimitiveMeshBounds.Add(Descriptor.StaticMesh->GetBoundingBox());
-				}
-			}
-			else
-			{
-				PCG_KERNEL_VALIDATION_ERR(InContext, Settings, /*bQuiet=*/false,
-					FText::Format(PCGStaticMeshSpawnerKernel::CouldNotLoadStaticMeshFormat, FText::FromString(Entry.Descriptor.StaticMesh.ToString())));
-			}
-		}
-	}
-
-	// Validate instance count is not too large.
-	{
-		uint64 TotalInstanceCount = 0;
-		for (const FPCGProceduralISMComponentDescriptor& Desc : ComponentsToCreate)
-		{
-			TotalInstanceCount += Desc.NumInstances;
-		}
-
-		if (TotalInstanceCount >= MAX_INSTANCE_ID)
-		{
-			PCG_KERNEL_VALIDATION_ERR(InContext, Settings, /*bQuiet=*/false, FText::Format(
-				LOCTEXT("TooManyInstances", "Tried to spawn too many instances ({0}), procedural ISM component creation skipped and instances will not be rendered."),
-				FText::FromString(FString::FromInt(TotalInstanceCount))));
-
-			return false;
-		}
-	}
-
-	FPCGSpawnerPrimitives& Primitives = InBinding->MeshSpawnersToPrimitives.FindOrAdd(this);
-	Primitives.NumCustomFloats = CustomFloatCount;
-	Primitives.AttributeIdOffsetStrides = MoveTemp(AttributeIdOffsetStrides);
-	Primitives.SelectorAttributeId = SelectorAttributeId;
-	Primitives.SelectionCDF = MoveTemp(PrimitiveSelectionCDF);
-	Primitives.PrimitiveStringKeys = MoveTemp(PrimitiveStringKeys);
-	Primitives.PrimitiveMeshBounds = MoveTemp(PrimitiveMeshBounds);
-
-	for (const FPCGProceduralISMComponentDescriptor& Desc : ComponentsToCreate)
-	{
-		FPCGProceduralISMCBuilderParameters Params;
-		Params.Descriptor = Desc;
-		Params.bAllowDescriptorChanges = false;
-
-		UPCGManagedProceduralISMComponent* MISMC = PCGManagedProceduralISMComponent::GetOrCreateManagedProceduralISMC(TargetActor, InContext->SourceComponent.Get(), Settings->UID, Params);
-
-		check(MISMC);
-		MISMC->SetCrc(InContext->DependenciesCrc);
-
-		// Don't bother registering the resource change as we're transient anyway.
-		//InContext->TouchedResources.Emplace(MISMC);
-
-		Primitives.Primitives.Add(MISMC->GetComponent());
-
-		Primitives.NumInstancesAllPrimitives += Desc.NumInstances;
-	}
-
-	// Useful for debugging instance counts.
-	//UE_LOG(LogPCG, Log, TEXT("Input points %d, total instance count %d (+ %.2f%%)"), TotalInputPointCount, Primitives.NumInstancesAllPrimitives, (float(Primitives.NumInstancesAllPrimitives - TotalInputPointCount) / TotalInputPointCount) * 100.0f);
-
-	return !ComponentsToCreate.IsEmpty();
-}
-
 #undef LOCTEXT_NAMESPACE
diff --git a/Engine/Plugins/PCG/Source/PCG/Private/Graph/PCGGraphCompilerGPU.cpp b/Engine/Plugins/PCG/Source/PCG/Private/Graph/PCGGraphCompilerGPU.cpp
--- a/Engine/Plugins/PCG/Source/PCG/Private/Graph/PCGGraphCompilerGPU.cpp
+++ b/Engine/Plugins/PCG/Source/PCG/Private/Graph/PCGGraphCompilerGPU.cpp
@@ -58,33 +58,19 @@ void FPCGGraphCompilerGPU::LabelConnectedGPUNodeIslands(
 
 		for (const FPCGGraphTaskInput& Input : InCompiledTasks[InTaskId].Inputs)
 		{
-			if (Input.TaskId == InTraversedFromTaskId)
-			{
-				continue;
-			}
-
-			if (OutIslandIDs[Input.TaskId] == 0 && InGPUCompatibleTaskIds.Contains(Input.TaskId))
+			if (Input.TaskId != InTraversedFromTaskId && OutIslandIDs[Input.TaskId] == 0 && InGPUCompatibleTaskIds.Contains(Input.TaskId))
 			{
 				RecursiveCall(Input.TaskId, InIslandID, InTaskId, RecursiveCall);
 			}
 		}
 
-		UPCGComputeKernel* const* FoundKernel = InTaskIdToKernel.Find(InTaskId);
-		if (!FoundKernel || !(*FoundKernel)->SplitGraphAtOutput())
+		if (const TArray<FPCGTaskId>* Successors = InTaskSuccessors.Find(InTaskId))
 		{
-			if (const TArray<FPCGTaskId>* Successors = InTaskSuccessors.Find(InTaskId))
+			for (FPCGTaskId Successor : *Successors)
 			{
-				for (FPCGTaskId Successor : *Successors)
+				if (Successor != InTraversedFromTaskId && OutIslandIDs[Successor] == 0 && InGPUCompatibleTaskIds.Contains(Successor))
 				{
-					if (Successor == InTraversedFromTaskId)
-					{
-						continue;
-					}
-
-					if (OutIslandIDs[Successor] == 0 && InGPUCompatibleTaskIds.Contains(Successor))
-					{
-						RecursiveCall(Successor, InIslandID, InTaskId, RecursiveCall);
-					}
+					RecursiveCall(Successor, InIslandID, InTaskId, RecursiveCall);
 				}
 			}
 		}
@@ -114,12 +100,6 @@ void FPCGGraphCompilerGPU::CollectGPUNodeSubsets(
 	TArray<uint32> ConnectedGPUNodeIslandIDs;
 	LabelConnectedGPUNodeIslands(InCompiledTasks, InGPUCompatibleTaskIds, InTaskSuccessors, InTaskIdToKernel, ConnectedGPUNodeIslandIDs);
 
-	// For every CPU node that has one or more downstream GPU node connected, check for a GPU -> CPU -> GPU pattern where data flows
-	// from a GPU node island to the CPU and then back to the same island. For such cases, we traverse the entire tree of GPU nodes that
-	// are in the island and downstream of the CPU node and bump their island ID - splitting the island into a portion that is independent
-	// of the CPU node and a portion that is dependent on it, so that we can read the data back to CPU, execute the CPU portion, then
-	// re-upload to GPU.
-	
 	// Any new island IDs will be created from the task count which will be larger than any island IDs presently set.
 	uint32 NextIslandId = InCompiledTasks.Num();
 
@@ -130,15 +110,54 @@ void FPCGGraphCompilerGPU::CollectGPUNodeSubsets(
 	FPCGGraphCompiler::VisitTasksInExecutionOrder(
 		InCompiledTasks,
 		InTaskSuccessors,
-		[&InCompiledTasks, &InTaskSuccessors, &InGPUCompatibleTaskIds, &ConnectedGPUNodeIslandIDs, &NextIslandId, &CPUNodeIsDependentOnIslandCached](FPCGTaskId InTaskId) -> bool
+		[&InCompiledTasks, &InTaskSuccessors, &InGPUCompatibleTaskIds, &ConnectedGPUNodeIslandIDs, &InTaskIdToKernel, &NextIslandId, &CPUNodeIsDependentOnIslandCached](FPCGTaskId InTaskId) -> bool
 		{
+			auto PropagateIslandIDDownstream = [&InCompiledTasks, &InTaskSuccessors, &InGPUCompatibleTaskIds, &ConnectedGPUNodeIslandIDs](FPCGTaskId InTaskId, uint32 InOldIslandID, uint32 InNewIslandID, auto&& RecursiveCall) -> void
+			{
+				ConnectedGPUNodeIslandIDs[InTaskId] = InNewIslandID;
+
+				if (const TArray<FPCGTaskId>* Successors = InTaskSuccessors.Find(InTaskId))
+				{
+					for (FPCGTaskId Successor : *Successors)
+					{
+						if (ConnectedGPUNodeIslandIDs[Successor] == InOldIslandID && InGPUCompatibleTaskIds.Contains(Successor))
+						{
+							RecursiveCall(Successor, InOldIslandID, InNewIslandID, RecursiveCall);
+						}
+					}
+				}
+			};
+
+			// GPU task. Split at output if the kernel requests it.
 			if (InGPUCompatibleTaskIds.Contains(InTaskId))
 			{
-				// Skip GPU nodes.
-				return true;
-			}
+				UPCGComputeKernel* const* FoundKernel = InTaskIdToKernel.Find(InTaskId);
+				if (FoundKernel && (*FoundKernel)->SplitGraphAtOutput())
+				{
+					if (const TArray<FPCGTaskId>* Successors = InTaskSuccessors.Find(InTaskId))
+					{
+						for (FPCGTaskId SuccessorTaskId : *Successors)
+						{
+							if (InGPUCompatibleTaskIds.Contains(SuccessorTaskId))
+							{
+								const uint32 SuccessorIsland = ConnectedGPUNodeIslandIDs[SuccessorTaskId];
 
-			if (const TArray<FPCGTaskId>* Successors = InTaskSuccessors.Find(InTaskId))
+								if (ConnectedGPUNodeIslandIDs[InTaskId] == SuccessorIsland)
+								{
+									// Propagate a new island ID to all downstream GPU tasks within the island.
+									PropagateIslandIDDownstream(SuccessorTaskId, SuccessorIsland, NextIslandId++, PropagateIslandIDDownstream);
+								}
+							}
+						}
+					}
+				}
+			}
+			// Otherwise CPU task. For every CPU node that has one or more downstream GPU node connected, check for a GPU -> CPU -> GPU pattern where data flows
+			// from a GPU node island to the CPU and then back to the same island. For such cases, we traverse the entire tree of GPU nodes that
+			// are in the island and downstream of the CPU node and bump their island ID - splitting the island into a portion that is independent
+			// of the CPU node and a portion that is dependent on it, so that we can read the data back to CPU, execute the CPU portion, then
+			// re-upload to GPU.
+			else if (const TArray<FPCGTaskId>* Successors = InTaskSuccessors.Find(InTaskId))
 			{
 				for (FPCGTaskId SuccessorTaskId : *Successors)
 				{
@@ -181,23 +200,6 @@ void FPCGGraphCompilerGPU::CollectGPUNodeSubsets(
 					if (CPUNodeIsDependentOnIsland(InTaskId, SuccessorIsland, CPUNodeIsDependentOnIsland))
 					{
 						// Propagate a new island ID to all downstream GPU tasks within the island.
-
-						auto PropagateIslandIDDownstream = [&InCompiledTasks, &InTaskSuccessors, &InGPUCompatibleTaskIds, &ConnectedGPUNodeIslandIDs](FPCGTaskId InTaskId, uint32 InOldIslandID, uint32 InNewIslandID, auto&& RecursiveCall) -> void
-						{
-							ConnectedGPUNodeIslandIDs[InTaskId] = InNewIslandID;
-
-							if (const TArray<FPCGTaskId>* Successors = InTaskSuccessors.Find(InTaskId))
-							{
-								for (FPCGTaskId Successor : *Successors)
-								{
-									if (ConnectedGPUNodeIslandIDs[Successor] == InOldIslandID && InGPUCompatibleTaskIds.Contains(Successor))
-									{
-										RecursiveCall(Successor, InOldIslandID, InNewIslandID, RecursiveCall);
-									}
-								}
-							}
-						};
-
 						PropagateIslandIDDownstream(SuccessorTaskId, SuccessorIsland, NextIslandId++, PropagateIslandIDDownstream);
 					}
 				}
@@ -354,6 +356,8 @@ void FPCGGraphCompilerGPU::ExpandGPUNodeKernelsToTasks(
 
 				if (const UPCGComputeKernel* UpstreamKernel = Edge->UpstreamPin.Kernel)
 				{
+					checkf(KernelToTaskId.Contains(UpstreamKernel), TEXT("Missing kernel '%s', was kernel emitted from CreateKernels()?"), *UpstreamKernel->GetName());
+
 					// Connection from an upstream kernel from this node.
 					const FPCGTaskId UpstreamTaskId = KernelToTaskId[UpstreamKernel];
 
@@ -883,8 +887,13 @@ void FPCGGraphCompilerGPU::CreateDataInterfaces(
 			check(DataInterfacePCGData);
 			DataInterfacePCGData->SetProducerSettings(ProducerSettings);
 			DataInterfacePCGData->SetRequiresExport(bRequiresExport);
-			DataInterfacePCGData->SetElementCountMultiplier((ProducerKernel && ProducerPin) ? ProducerKernel->GetElementCountMultiplier(ProducerPin->Properties.Label) : 1);
-			DataInterfacePCGData->SetExposeElementCounters(ProducerKernel && ProducerPin && ProducerKernel->DoesOutputPinRequireElementCounters(ProducerPin->Properties.Label));
+
+			if (ProducerPin)
+			{
+				DataInterfacePCGData->SetElementCountMultiplier(ProducerKernel ? ProducerKernel->GetElementCountMultiplier(ProducerPin->Properties.Label) : 1);
+				DataInterfacePCGData->SetExposeElementCounters(ProducerKernel && ProducerKernel->DoesOutputPinRequireElementCounters(ProducerPin->Properties.Label));
+				DataInterfacePCGData->SetRequiresZeroInitialization(ProducerKernel && ProducerKernel->DoesOutputPinRequireZeroInitialization(ProducerPin->Properties.Label));
+			}
 
 			DataInterface = DataInterfacePCGData;
 			break;
diff --git a/Engine/Plugins/PCG/Source/PCG/Public/Elements/PCGStaticMeshSpawnerKernel.h b/Engine/Plugins/PCG/Source/PCG/Public/Elements/PCGStaticMeshSpawnerKernel.h
--- a/Engine/Plugins/PCG/Source/PCG/Public/Elements/PCGStaticMeshSpawnerKernel.h
+++ b/Engine/Plugins/PCG/Source/PCG/Public/Elements/PCGStaticMeshSpawnerKernel.h
@@ -8,6 +8,11 @@
 
 class UPCGStaticMeshSpawnerSettings;
 
+namespace PCGStaticMeshSpawnerConstants
+{
+	const FName InstanceCountsPinLabel = TEXT("InstanceCounts");
+}
+
 UCLASS()
 class UPCGStaticMeshSpawnerKernel : public UPCGComputeKernel
 {
@@ -32,9 +37,6 @@ public:
 	virtual void GetInputPins(TArray<FPCGPinProperties>& OutPins) const override;
 	virtual void GetOutputPins(TArray<FPCGPinPropertiesGPU>& OutPins) const override;
 
-	/** Setup Procedural ISMCs. Returns true if at least one component was set up. */
-	virtual bool SetupProceduralISMComponents(FPCGContext* InContext, UPCGDataBinding* InBinding) const;
-
 public:
 	UPROPERTY()
 	TObjectPtr<const UPCGStaticMeshSpawnerSettings> Settings;
