[PCG] Fix namespace to fix build issue

#rnx

diff --git a/Engine/Plugins/PCG/Source/PCG/Private/Compute/DataInterfaces/Elements/PCGStaticMeshSpawnerDataInterface.cpp b/Engine/Plugins/PCG/Source/PCG/Private/Compute/DataInterfaces/Elements/PCGStaticMeshSpawnerDataInterface.cpp
--- a/Engine/Plugins/PCG/Source/PCG/Private/Compute/DataInterfaces/Elements/PCGStaticMeshSpawnerDataInterface.cpp
+++ b/Engine/Plugins/PCG/Source/PCG/Private/Compute/DataInterfaces/Elements/PCGStaticMeshSpawnerDataInterface.cpp
@@ -27,9 +27,9 @@
 
 #include UE_INLINE_GENERATED_CPP_BY_NAME(PCGStaticMeshSpawnerDataInterface)
 
-#define LOCTEXT_NAMESPACE "PCGStaticMeshSpawnerKernel"
+#define LOCTEXT_NAMESPACE "PCGStaticMeshSpawnerDataInterface"
 
-namespace PCGStaticMeshSpawnerKernel
+namespace PCGStaticMeshSpawnerDataInterface
 {
 	static const FText CouldNotLoadStaticMeshFormat = LOCTEXT("CouldNotLoadStaticMesh", "Could not load static mesh from path '{0}'.");
 	static const FText TooManyPrimitivesFormat = LOCTEXT("TooManyPrimitives", "Attempted to emit too many primitive components, terminated after creating '{0}'.");
@@ -525,7 +525,7 @@ bool UPCGStaticMeshSpawnerDataProvider::SetupPrimitives(FPCGContext* InContext,
 			if (ComponentsToCreate.Num() >= PCGComputeConstants::MAX_PRIMITIVE_COMPONENTS_PER_SPAWNER)
 			{
 				PCG_KERNEL_VALIDATION_WARN(InContext, Settings, /*bQuiet=*/false,
-					FText::Format(PCGStaticMeshSpawnerKernel::TooManyPrimitivesFormat, FText::FromString(FString::FromInt(PCGComputeConstants::MAX_PRIMITIVE_COMPONENTS_PER_SPAWNER))));
+					FText::Format(PCGStaticMeshSpawnerDataInterface::TooManyPrimitivesFormat, FText::FromString(FString::FromInt(PCGComputeConstants::MAX_PRIMITIVE_COMPONENTS_PER_SPAWNER))));
 
 				break;
 			}
@@ -534,7 +534,7 @@ bool UPCGStaticMeshSpawnerDataProvider::SetupPrimitives(FPCGContext* InContext,
 			if (!StaticMesh)
 			{
 				PCG_KERNEL_VALIDATION_ERR(InContext, Settings, /*bQuiet=*/false,
-					FText::Format(PCGStaticMeshSpawnerKernel::CouldNotLoadStaticMeshFormat, FText::FromString(MeshPathString)));
+					FText::Format(PCGStaticMeshSpawnerDataInterface::CouldNotLoadStaticMeshFormat, FText::FromString(MeshPathString)));
 
 				continue;
 			}
@@ -561,7 +561,7 @@ bool UPCGStaticMeshSpawnerDataProvider::SetupPrimitives(FPCGContext* InContext,
 	{
 		if (SelectorWeighted->MeshEntries.IsEmpty())
 		{
-			PCG_KERNEL_VALIDATION_ERR(InContext, Settings, /*bQuiet=*/false, PCGStaticMeshSpawnerKernel::NoMeshEntriesFormat);
+			PCG_KERNEL_VALIDATION_ERR(InContext, Settings, /*bQuiet=*/false, PCGStaticMeshSpawnerDataInterface::NoMeshEntriesFormat);
 
 			return false;
 		}
@@ -591,7 +591,7 @@ bool UPCGStaticMeshSpawnerDataProvider::SetupPrimitives(FPCGContext* InContext,
 			if (ComponentsToCreate.Num() >= PCGComputeConstants::MAX_PRIMITIVE_COMPONENTS_PER_SPAWNER)
 			{
 				PCG_KERNEL_VALIDATION_WARN(InContext, Settings, /*bQuiet=*/false,
-					FText::Format(PCGStaticMeshSpawnerKernel::TooManyPrimitivesFormat, FText::FromString(FString::FromInt(PCGComputeConstants::MAX_PRIMITIVE_COMPONENTS_PER_SPAWNER))));
+					FText::Format(PCGStaticMeshSpawnerDataInterface::TooManyPrimitivesFormat, FText::FromString(FString::FromInt(PCGComputeConstants::MAX_PRIMITIVE_COMPONENTS_PER_SPAWNER))));
 
 				break;
 			}
@@ -638,7 +638,7 @@ bool UPCGStaticMeshSpawnerDataProvider::SetupPrimitives(FPCGContext* InContext,
 			else
 			{
 				PCG_KERNEL_VALIDATION_ERR(InContext, Settings, /*bQuiet=*/false,
-					FText::Format(PCGStaticMeshSpawnerKernel::CouldNotLoadStaticMeshFormat, FText::FromString(Entry.Descriptor.StaticMesh.ToString())));
+					FText::Format(PCGStaticMeshSpawnerDataInterface::CouldNotLoadStaticMeshFormat, FText::FromString(Entry.Descriptor.StaticMesh.ToString())));
 			}
 		}
 	}
