[PCG] ScheduleGeneric Params version
- Allows scheduling non-gamethread generic task
- Remove need for anythread element in PCGInternalHelpers
- Remove private access friend

#rnx

diff --git a/Engine/Plugins/PCG/Source/PCG/Private/Graph/PCGGraphExecutor.cpp b/Engine/Plugins/PCG/Source/PCG/Private/Graph/PCGGraphExecutor.cpp
--- a/Engine/Plugins/PCG/Source/PCG/Private/Graph/PCGGraphExecutor.cpp
+++ b/Engine/Plugins/PCG/Source/PCG/Private/Graph/PCGGraphExecutor.cpp
@@ -691,6 +691,12 @@ void FPCGGraphExecutor::SetExecutionCacheData(FPCGTaskId InGraphExecutionTaskId,
 	}
 }
 
+FPCGTaskId FPCGGraphExecutor::ScheduleGeneric(const FPCGScheduleGenericParams& Params)
+{
+	return ScheduleSingleElement(MakeShared<FPCGGenericElement>(Params.Operation, Params.AbortOperation, Params.bSupportBasePointDataInput, Params.bCanExecuteOnlyOnMainThread),
+		!!Params.AbortOperation, Params.SourceComponent, Params.ExecutionDependencies, Params.DataDependencies);
+}
+
 FPCGTaskId FPCGGraphExecutor::ScheduleGeneric(TFunction<bool()> InOperation, UPCGComponent* InSourceComponent, const TArray<FPCGTaskId>& TaskExecutionDependencies)
 {
 	return ScheduleGeneric(
@@ -733,11 +739,11 @@ FPCGTaskId FPCGGraphExecutor::ScheduleGenericWithContext(TFunction<bool(FPCGCont
 
 FPCGTaskId FPCGGraphExecutor::ScheduleGenericWithContext(TFunction<bool(FPCGContext*)> InOperation, TFunction<void(FPCGContext*)> InAbortOperation, UPCGComponent* InSourceComponent, const TArray<FPCGTaskId>& TaskExecutionDependencies, const TArray<FPCGTaskId>& TaskDataDependencies, bool bSupportBasePointDataInput)
 {
-	return ScheduleSingleElement(MakeShared<FPCGGenericElement>(InOperation, InAbortOperation, bSupportBasePointDataInput), 
-		!!InAbortOperation, InSourceComponent, TaskExecutionDependencies, TaskDataDependencies, bSupportBasePointDataInput);
+	return ScheduleSingleElement(MakeShared<FPCGGenericElement>(InOperation, InAbortOperation, bSupportBasePointDataInput),
+		!!InAbortOperation, InSourceComponent, TaskExecutionDependencies, TaskDataDependencies);
 }
 
-FPCGTaskId FPCGGraphExecutor::ScheduleSingleElement(TSharedRef<IPCGElement> Element, bool bHasAbortCallbacks, UPCGComponent* InSourceComponent, const TArray<FPCGTaskId>& TaskExecutionDependencies, const TArray<FPCGTaskId>& TaskDataDependencies, bool bSupportBasePointDataInput)
+FPCGTaskId FPCGGraphExecutor::ScheduleSingleElement(TSharedRef<IPCGElement> Element, bool bHasAbortCallbacks, UPCGComponent* InSourceComponent, const TArray<FPCGTaskId>& TaskExecutionDependencies, const TArray<FPCGTaskId>& TaskDataDependencies)
 {
 	// Build task & element to hold the operation to perform
 	FPCGGraphTask Task;
@@ -771,16 +777,6 @@ FPCGTaskId FPCGGraphExecutor::ScheduleSingleElement(TSharedRef<IPCGElement> Elem
 	return Task.NodeId;
 }
 
-// Used to expose ScheduleSingleElement to internal helpers without exposing everything
-FPCGTaskId UE::FPCGGraphExecutorInternalAccess::ScheduleSingleElement(FPCGGraphExecutor& Executor, 
-	TSharedRef<IPCGElement> Element, bool bHasAbortCallbacks, UPCGComponent* InSourceComponent, 
-	const TArray<FPCGTaskId>& TaskExecutionDependencies, const TArray<FPCGTaskId>& TaskDataDependencies, 
-	bool bSupportBasePointDataInput)
-{
-	return Executor.ScheduleSingleElement(Element, bHasAbortCallbacks, InSourceComponent, 
-		TaskExecutionDependencies, TaskDataDependencies, bSupportBasePointDataInput);
-}
-
 void FPCGGraphExecutor::MarkInputResults(TArrayView<const FPCGTaskId> InInputResults)
 {
 	TRACE_CPUPROFILER_EVENT_SCOPE(FPCGGraphExecutor::MarkInputResults);
@@ -2793,6 +2789,15 @@ FPCGGenericElement::FPCGGenericElement(TFunction<bool(FPCGContext*)> InOperation
 {
 }
 
+FPCGGenericElement::FPCGGenericElement(TFunction<bool(FPCGContext*)> InOperation, TFunction<void(FPCGContext*)> InAbortOperation, bool bInSupportBasePointDataInput, bool bInCanExecuteOnlyOnMainThread, const FContextAllocator& InContextAllocator)
+	: Operation(InOperation)
+	, AbortOperation(InAbortOperation)
+	, bSupportBasePointDataInput(bInSupportBasePointDataInput)
+	, bCanExecuteOnlyOnMainThread(bInCanExecuteOnlyOnMainThread)
+	, ContextAllocator(InContextAllocator)
+{
+}
+
 FPCGContext* FPCGGenericElement::Initialize(const FPCGDataCollection& InputData, TWeakObjectPtr<UPCGComponent> SourceComponent, const UPCGNode* Node)
 {
 	check(ContextAllocator);
@@ -2971,7 +2976,7 @@ namespace PCGGraphExecutor
 
 			// At this point we could not get to the data, so we'll try executing the graph if we did not do that already.
 
-			auto WakeUpLambda = [ComponentWithData, ContextHandle=InContext->GetOrCreateHandle()]()
+			auto WakeUpLambda = [ComponentWithData, ContextHandle=InContext->GetOrCreateHandle()](FPCGContext* InContext)
 			{
 				TSharedPtr<FPCGContextHandle> SharedHandle = ContextHandle.Pin();
 				if (FPCGContext* ContextPtr = SharedHandle ? SharedHandle->GetContext() : nullptr)
@@ -2999,7 +3004,18 @@ namespace PCGGraphExecutor
 				FPCGTaskId GenerationTask = ComponentWithData->GetGenerationTaskId();
 				if (ensure(GenerationTask != InvalidPCGTaskId))
 				{
-					Subsystem->ScheduleGeneric(WakeUpLambda, InContext->SourceComponent.Get(), { ComponentWithData->GetGenerationTaskId() });
+					if(PCGGraphExecutor::CVarEnableDynamicDependencies.GetValueOnAnyThread())
+					{
+						InContext->DynamicDependencies.Add(ComponentWithData->GetGenerationTaskId());
+					}
+					else
+					{
+						FPCGScheduleGenericParams Params(MoveTemp(WakeUpLambda), InContext->SourceComponent.Get());
+						Params.bCanExecuteOnlyOnMainThread = false;
+						Params.ExecutionDependencies = { ComponentWithData->GetGenerationTaskId() };
+
+						Subsystem->ScheduleGeneric(Params);
+					}
 				}
 
 				return false;
@@ -3016,7 +3032,19 @@ namespace PCGGraphExecutor
 
 				// Wake up this task after graph has generated.
 				const FPCGTaskId GraphTaskId = ComponentWithData->GenerateLocalGetTaskId(GenTrigger, /*bForce=*/true);
-				Subsystem->ScheduleGeneric(WakeUpLambda, InContext->SourceComponent.Get(), { GraphTaskId });
+				
+				if (PCGGraphExecutor::CVarEnableDynamicDependencies.GetValueOnAnyThread())
+				{
+					InContext->DynamicDependencies.Add(GraphTaskId);
+				}
+				else
+				{
+					FPCGScheduleGenericParams Params(MoveTemp(WakeUpLambda), InContext->SourceComponent.Get());
+					Params.bCanExecuteOnlyOnMainThread = false;
+					Params.ExecutionDependencies = { GraphTaskId };
+
+					Subsystem->ScheduleGeneric(Params);
+				}
 
 				// Update state and go to paused.
 				InContext->bScheduledGraph = true;
diff --git a/Engine/Plugins/PCG/Source/PCG/Private/Graph/PCGGraphExecutor.h b/Engine/Plugins/PCG/Source/PCG/Private/Graph/PCGGraphExecutor.h
--- a/Engine/Plugins/PCG/Source/PCG/Private/Graph/PCGGraphExecutor.h
+++ b/Engine/Plugins/PCG/Source/PCG/Private/Graph/PCGGraphExecutor.h
@@ -29,10 +29,6 @@ class UPCGGraph;
 class UPCGNode;
 struct FPCGStack;
 class FTextFormat;
-namespace UE
-{
-	class FPCGGraphExecutorInternalAccess;
-}
 
 namespace PCGGraphExecutor
 {
@@ -85,7 +81,10 @@ public:
 	/** Returns true if any task is scheduled or executing for any graph */
 	bool IsAnyGraphCurrentlyExecuting() const;
 	
-	// Back compatibility function. Use ScheduleGenericWithContext
+	// General job scheduling
+	FPCGTaskId ScheduleGeneric(const FPCGScheduleGenericParams& Params);
+
+	// Back compatibility function. Use ScheduleGenericWithContext or ScheduleGeneric taking a FPCGScheduleGenericParams
 	FPCGTaskId ScheduleGeneric(TFunction<bool()> InOperation, UPCGComponent* InSourceComponent, const TArray<FPCGTaskId>& TaskExecutionDependencies);
 	FPCGTaskId ScheduleGeneric(TFunction<bool()> InOperation, TFunction<void()> InAbortOperation, UPCGComponent* InSourceComponent, const TArray<FPCGTaskId>& TaskExecutionDependencies);
 
@@ -147,7 +146,6 @@ public:
 
 private:
 	friend UPCGSubsystem;
-	friend UE::FPCGGraphExecutorInternalAccess;
 
 	enum class EExecutionCacheDataType : int8
 	{
@@ -205,9 +203,8 @@ private:
 		FPCGGraphExecutor* Executor = nullptr;
 	};
 
-	// Used by ScheduleGenericWithContext and exposed to an internal helper to schedule non-game-thread generics via FPCGGraphExecutorInternalAccess
 	FPCGTaskId ScheduleSingleElement(TSharedRef<IPCGElement> Element, bool bHasAbortCallbacks, UPCGComponent* InSourceComponent,
-		const TArray<FPCGTaskId>& TaskExecutionDependencies, const TArray<FPCGTaskId>& TaskDataDependencies, bool bSupportBasePointDataInput);
+		const TArray<FPCGTaskId>& TaskExecutionDependencies, const TArray<FPCGTaskId>& TaskDataDependencies);
 
 	void PostTaskExecute(TSharedPtr<FPCGGraphActiveTask> ActiveTask, bool bIsDone);
 	bool ProcessScheduledTasks();
@@ -420,11 +417,21 @@ public:
 	{
 		return new FPCGContext();
 	});
+
+	FPCGGenericElement(
+		TFunction<bool(FPCGContext*)> InOperation,
+		TFunction<void(FPCGContext*)> InAbortOperation,
+		bool bInSupportBasePointDataInput,
+		bool bInCanExecuteOnlyOnMainThread,
+		const FContextAllocator& InContextAllocator = (FContextAllocator)[](const FPCGDataCollection&, TWeakObjectPtr<UPCGComponent>, const UPCGNode*)
+	{
+		return new FPCGContext();
+	});
 	
 	virtual FPCGContext* Initialize(const FPCGDataCollection& InputData, TWeakObjectPtr<UPCGComponent> SourceComponent, const UPCGNode* Node) override;
 
 	virtual bool IsCacheable(const UPCGSettings* InSettings) const override { return false; }
-	virtual bool CanExecuteOnlyOnMainThread(FPCGContext* Context) const override { return true; }
+	virtual bool CanExecuteOnlyOnMainThread(FPCGContext* Context) const override { return bCanExecuteOnlyOnMainThread; }
 
 protected:
 	// Important note: generic elements must always be run on the main thread
@@ -442,6 +449,7 @@ private:
 	TFunction<bool(FPCGContext*)> Operation;
 	TFunction<void(FPCGContext*)> AbortOperation;
 	bool bSupportBasePointDataInput = false;
+	bool bCanExecuteOnlyOnMainThread = true;
 
 	/** Creates a context object for this element. */
 	FContextAllocator ContextAllocator;
@@ -535,18 +543,3 @@ public:
 	UPROPERTY()
 	TSoftObjectPtr<const UPCGPin> UpstreamPin = nullptr;
 };
-
-namespace UE
-{
-	// Helper to give FPCGInternalsHelper limited access to FPCGGraphExecutor private functions
-	//  rather than friending it completely.
-	// Currently used to allow for scheduling a custom generic element that does not require
-	//  running on game thread.
-	class FPCGGraphExecutorInternalAccess
-	{
-		friend class FPCGInternalHelpers;
-		static FPCGTaskId ScheduleSingleElement(FPCGGraphExecutor& Executor, TSharedRef<IPCGElement> Element, 
-			bool bHasAbortCallbacks, UPCGComponent* InSourceComponent, const TArray<FPCGTaskId>& TaskExecutionDependencies, 
-			const TArray<FPCGTaskId>& TaskDataDependencies, bool bSupportBasePointDataInput = false);
-	};
-}
\ No newline at end of file
diff --git a/Engine/Plugins/PCG/Source/PCG/Private/Helpers/PCGInternalHelpers.cpp b/Engine/Plugins/PCG/Source/PCG/Private/Helpers/PCGInternalHelpers.cpp
--- a/Engine/Plugins/PCG/Source/PCG/Private/Helpers/PCGInternalHelpers.cpp
+++ b/Engine/Plugins/PCG/Source/PCG/Private/Helpers/PCGInternalHelpers.cpp
@@ -6,20 +6,6 @@
 #include "PCGComponent.h"
 #include "PCGSubgraph.h" // FPCGInputForwardingElement
 
-namespace PCGInternalHelpersLocals
-{
-	class FPCGGenericElementAnyThread : public FPCGGenericElement
-	{
-	public:
-		FPCGGenericElementAnyThread(
-			TFunction<bool(FPCGContext*)> InOperation)
-			: FPCGGenericElement(InOperation)
-		{}
-
-		virtual bool CanExecuteOnlyOnMainThread(FPCGContext* Context) const override { return false; }
-	};
-}
-
 TSharedPtr<FPCGGraphExecutor> UE::FPCGInternalHelpers::CreateExecutor()
 {
 	return MakeShared<FPCGGraphExecutor>();
@@ -41,8 +27,6 @@ bool UE::FPCGInternalHelpers::ScheduleOnExecutorAndWait(FPCGGraphExecutor& Execu
 	UPCGComponent* PCGComponent, const FPCGDataCollection& InputData,
 	FPCGDataCollection& OutputData, int32 CyclesToWait, double SleepDurationPerCycle)
 {
-	using namespace PCGInternalHelpersLocals;
-
 	if (!PCGComponent)
 	{
 		return false;
@@ -70,18 +54,22 @@ bool UE::FPCGInternalHelpers::ScheduleOnExecutorAndWait(FPCGGraphExecutor& Execu
 	};
 	TSharedPtr<FOutputState> OutputState = MakeShared<FOutputState>();
 
-	FPCGGraphExecutorInternalAccess::ScheduleSingleElement(Executor,
-		MakeShared<FPCGGenericElementAnyThread>([&Executor, OutputStateWeak = OutputState.ToWeakPtr(), TaskId](FPCGContext*) -> bool
+	auto Operation = [&Executor, OutputStateWeak = OutputState.ToWeakPtr(), TaskId](FPCGContext*) -> bool
+	{
+		TSharedPtr<FOutputState> OutputStatePinned = OutputStateWeak.Pin();
+		if (OutputStatePinned.IsValid())
 		{
-			TSharedPtr<FOutputState> OutputStatePinned = OutputStateWeak.Pin();
-			if (OutputStatePinned.IsValid())
-			{
-				Executor.GetOutputData(TaskId, OutputStatePinned->Output);
-				OutputStatePinned->bDone = true;
-			}
-			return true;
-		}), 
-		false, PCGComponent, { TaskId }, {});
+			Executor.GetOutputData(TaskId, OutputStatePinned->Output);
+			OutputStatePinned->bDone = true;
+		}
+		return true;
+	};
+
+	FPCGScheduleGenericParams Params(MoveTemp(Operation), PCGComponent);
+	Params.bCanExecuteOnlyOnMainThread = false;
+	Params.ExecutionDependencies = { TaskId };
+
+	Executor.ScheduleGeneric(Params);
 
 	int32 CycleCount = 0;
 	CyclesToWait = FMath::Max(0, CyclesToWait);
diff --git a/Engine/Plugins/PCG/Source/PCG/Private/PCGSubsystem.cpp b/Engine/Plugins/PCG/Source/PCG/Private/PCGSubsystem.cpp
--- a/Engine/Plugins/PCG/Source/PCG/Private/PCGSubsystem.cpp
+++ b/Engine/Plugins/PCG/Source/PCG/Private/PCGSubsystem.cpp
@@ -843,6 +843,12 @@ TStatId UPCGSubsystem::GetStatId() const
 	RETURN_QUICK_DECLARE_CYCLE_STAT(UPCGSubsystem, STATGROUP_Tickables);
 }
 
+FPCGTaskId UPCGSubsystem::ScheduleGeneric(const FPCGScheduleGenericParams& InParams)
+{
+	check(GraphExecutor);
+	return GraphExecutor->ScheduleGeneric(InParams);
+}
+
 FPCGTaskId UPCGSubsystem::ScheduleGeneric(TFunction<bool()> InOperation, UPCGComponent* SourceComponent, const TArray<FPCGTaskId>& TaskExecutionDependencies)
 {
 	check(GraphExecutor);
diff --git a/Engine/Plugins/PCG/Source/PCG/Public/PCGCommon.h b/Engine/Plugins/PCG/Source/PCG/Public/PCGCommon.h
--- a/Engine/Plugins/PCG/Source/PCG/Public/PCGCommon.h
+++ b/Engine/Plugins/PCG/Source/PCG/Public/PCGCommon.h
@@ -12,6 +12,7 @@ static const FPCGTaskId InvalidPCGTaskId = (uint64)-1;
 using FPCGPinId = uint64;
 
 class UPCGComponent;
+struct FPCGContext;
 
 namespace PCGLog
 {
@@ -499,6 +500,35 @@ enum class EPCGGenerationStatus : uint8
 	Aborted
 };
 
+struct FPCGScheduleGenericParams
+{
+	FPCGScheduleGenericParams(TFunction<bool(FPCGContext*)>&& InOperation, UPCGComponent* InSourceComponent)
+		: Operation(MoveTemp(InOperation)), SourceComponent(InSourceComponent)
+	{
+	}
+
+	// Callback that takes a Context as argument and returns true if the task is done, false otherwise
+	TFunction<bool(FPCGContext*)> Operation;
+
+	// PCG component associated with this task. Can be null.
+	UPCGComponent* SourceComponent = nullptr;
+
+	// Callback that is called if the task is aborted (cancelled) before fully executed.
+	TFunction<void(FPCGContext*)> AbortOperation;
+
+	// Task will wait on these tasks to execute and won't take their output data as input.
+	TArray<FPCGTaskId> ExecutionDependencies;
+
+	// Task will wait on these tasks to execute and will take their output data as input.
+	TArray<FPCGTaskId> DataDependencies;
+
+	// When true, generic element will not convert input to UPCGPointData, this is false by default to preserver backward compatibility.
+	bool bSupportBasePointDataInput = false;
+
+	// When false, generic element can be executed outside of the game thread.
+	bool bCanExecuteOnlyOnMainThread = true;
+};
+
 // Enable to debug if some PCG Elements are creating some data that should be pre-created by the graph PreGraph element
 #define PCG_EXECUTION_CACHE_VALIDATION_ENABLED 0
 
diff --git a/Engine/Plugins/PCG/Source/PCG/Public/PCGSubsystem.h b/Engine/Plugins/PCG/Source/PCG/Public/PCGSubsystem.h
--- a/Engine/Plugins/PCG/Source/PCG/Public/PCGSubsystem.h
+++ b/Engine/Plugins/PCG/Source/PCG/Public/PCGSubsystem.h
@@ -28,6 +28,7 @@ class FPCGGraphCompiler;
 class FPCGGraphExecutor;
 struct FPCGContext;
 struct FPCGDataCollection;
+struct FPCGScheduleGenericParams;
 struct FPCGStack;
 class UPCGSettings;
 
@@ -140,6 +141,9 @@ public:
 	// Schedule graph (used internally for dynamic subgraph execution)
 	FPCGTaskId ScheduleGraph(UPCGComponent* SourceComponent, const TArray<FPCGTaskId>& Dependencies);
 
+	// General job scheduling
+	FPCGTaskId ScheduleGeneric(const FPCGScheduleGenericParams& InParams);
+
 	/** General job scheduling
 	*  @param InOperation:               Callback that returns true if the task is done, false otherwise.
 	*  @param SourceComponent:           PCG component associated with this task. Can be null.
