[PCG] Fix missing null checks on Context, possible when debugging a node and creating its visualization data

#rnx

diff --git a/Engine/Plugins/PCG/Source/PCG/Private/Data/PCGProjectionData.cpp b/Engine/Plugins/PCG/Source/PCG/Private/Data/PCGProjectionData.cpp
--- a/Engine/Plugins/PCG/Source/PCG/Private/Data/PCGProjectionData.cpp
+++ b/Engine/Plugins/PCG/Source/PCG/Private/Data/PCGProjectionData.cpp
@@ -382,7 +382,7 @@ const UPCGBasePointData* UPCGProjectionData::CreateBasePointData(FPCGContext* Co
 	};
 
 	FPCGAsync::AsyncProcessingRangeEx(
-		&Context->AsyncState,
+		Context ? &Context->AsyncState : nullptr,
 		SourcePointData->GetNumPoints(),
 		InitializeFunc,
 		ProcessRangeFunc,
diff --git a/Engine/Plugins/PCG/Source/PCG/Private/Data/PCGUnionData.cpp b/Engine/Plugins/PCG/Source/PCG/Private/Data/PCGUnionData.cpp
--- a/Engine/Plugins/PCG/Source/PCG/Private/Data/PCGUnionData.cpp
+++ b/Engine/Plugins/PCG/Source/PCG/Private/Data/PCGUnionData.cpp
@@ -542,7 +542,7 @@ void UPCGUnionData::CreateSequentialPointData(FPCGContext* Context, TArray<const
 		const UPCGBasePointData* CurrentPointData = PointInputDatas[DataIndex];
 
 		FPCGAsync::AsyncProcessingRangeEx(
-			&Context->AsyncState,
+			Context ? &Context->AsyncState : nullptr,
 			CurrentPointData->GetNumPoints(),
 			[]() {},
 			ProcessRangeCurrentIndex,
