PCG: Remove temporary support for non-game-thread executor.

diff --git a/Engine/Plugins/PCG/Source/PCG/Private/Helpers/PCGInternalHelpers.cpp b/Engine/Plugins/PCG/Source/PCG/Private/Helpers/PCGInternalHelpers.cpp
deleted file mode 100644
diff --git a/Engine/Plugins/PCG/Source/PCG/Public/Helpers/PCGInternalHelpers.h b/Engine/Plugins/PCG/Source/PCG/Public/Helpers/PCGInternalHelpers.h
deleted file mode 100644
diff --git a/Engine/Plugins/PCG/Source/PCG/Public/PCGSubgraph.h b/Engine/Plugins/PCG/Source/PCG/Public/PCGSubgraph.h
--- a/Engine/Plugins/PCG/Source/PCG/Public/PCGSubgraph.h
+++ b/Engine/Plugins/PCG/Source/PCG/Public/PCGSubgraph.h
@@ -205,7 +205,7 @@ protected:
 class FPCGInputForwardingElement : public IPCGElement
 {
 public:
-	explicit FPCGInputForwardingElement(const FPCGDataCollection& InputToForward);
+	PCG_API explicit FPCGInputForwardingElement(const FPCGDataCollection& InputToForward);
 
 	// Since this class is stateful because it owns a FPCGDataCollection it can't be cached unless it implements a proper GetDependenciesCrc.
 	// For now since it only forwards its input, we just disable the caching.
