[PCG] Add accessor support for `[REDACTED_HANDLE]`

* Data will know how to convert from a domain ID to a selector and vice versa

#rnx

diff --git a/Engine/Plugins/PCG/Source/PCG/Private/Data/PCGBasePointData.cpp b/Engine/Plugins/PCG/Source/PCG/Private/Data/PCGBasePointData.cpp
--- a/Engine/Plugins/PCG/Source/PCG/Private/Data/PCGBasePointData.cpp
+++ b/Engine/Plugins/PCG/Source/PCG/Private/Data/PCGBasePointData.cpp
@@ -1076,6 +1076,33 @@ bool UPCGBasePointData::ProjectPoint(const FTransform& InTransform, const FBox&
 	return ProjectPoint(InTransform, InBounds, InParams, OutPoint, OutMetadata, true);
 }
 
+FPCGMetadataDomainID UPCGBasePointData::GetMetadataDomainIDFromSelector(const FPCGAttributePropertySelector& InSelector) const
+{
+	const FName DomainName = InSelector.GetDomainName();
+
+	if (DomainName == PCGPointDataConstants::ElementsDomainName)
+	{
+		return PCGMetadataDomainID::Elements;
+	}
+	else
+	{
+		return Super::GetMetadataDomainIDFromSelector(InSelector);
+	}
+}
+
+bool UPCGBasePointData::SetDomainFromDomainID(const FPCGMetadataDomainID& InDomainID, FPCGAttributePropertySelector& InOutSelector) const
+{
+	if (InDomainID == PCGMetadataDomainID::Elements)
+	{
+		InOutSelector.SetDomainName(PCGPointDataConstants::ElementsDomainName);
+		return true;
+	}
+	else
+	{
+		return Super::SetDomainFromDomainID(InDomainID, InOutSelector);
+	}
+}
+
 bool UPCGBasePointData::ProjectPoint(const FTransform& InTransform, const FBox& InBounds, const FPCGProjectionParams& InParams, FPCGPoint& OutPoint, UPCGMetadata* OutMetadata, bool bUseBounds) const
 {
 	//TRACE_CPUPROFILER_EVENT_SCOPE(UPCGBasePointData::ProjectPoint);
diff --git a/Engine/Plugins/PCG/Source/PCG/Private/Data/PCGSplineData.cpp b/Engine/Plugins/PCG/Source/PCG/Private/Data/PCGSplineData.cpp
--- a/Engine/Plugins/PCG/Source/PCG/Private/Data/PCGSplineData.cpp
+++ b/Engine/Plugins/PCG/Source/PCG/Private/Data/PCGSplineData.cpp
@@ -19,6 +19,12 @@
 
 #define LOCTEXT_NAMESPACE "PCGSplineData"
 
+UPCGSplineData::UPCGSplineData()
+{
+	check(Metadata);
+	Metadata->SetupDomain(PCGMetadataDomainID::Elements, /*bIsDefault=*/false);
+}
+
 void UPCGSplineData::Initialize(const USplineComponent* InSpline)
 {
 	check(InSpline);
@@ -96,6 +102,33 @@ void UPCGSplineData::AddToCrc(FArchiveCrc32& Ar, bool bFullDataCrc) const
 	Ar << SplineStruct;
 }
 
+FPCGMetadataDomainID UPCGSplineData::GetMetadataDomainIDFromSelector(const FPCGAttributePropertySelector& InSelector) const
+{
+	const FName DomainName = InSelector.GetDomainName();
+
+	if (DomainName == PCGSplineData::ControlPointDomainName)
+	{
+		return PCGMetadataDomainID::Elements;
+	}
+	else
+	{
+		return Super::GetMetadataDomainIDFromSelector(InSelector);
+	}
+}
+
+bool UPCGSplineData::SetDomainFromDomainID(const FPCGMetadataDomainID& InDomainID, FPCGAttributePropertySelector& InOutSelector) const
+{
+	if (InDomainID == PCGMetadataDomainID::Elements)
+	{
+		InOutSelector.SetDomainName(PCGSplineData::ControlPointDomainName);
+		return true;
+	}
+	else
+	{
+		return Super::SetDomainFromDomainID(InDomainID, InOutSelector);
+	}
+}
+
 FTransform UPCGSplineData::GetTransform() const
 {
 	return SplineStruct.GetTransform();
diff --git a/Engine/Plugins/PCG/Source/PCG/Private/Metadata/Accessors/PCGAttributeAccessorFactory.cpp b/Engine/Plugins/PCG/Source/PCG/Private/Metadata/Accessors/PCGAttributeAccessorFactory.cpp
--- a/Engine/Plugins/PCG/Source/PCG/Private/Metadata/Accessors/PCGAttributeAccessorFactory.cpp
+++ b/Engine/Plugins/PCG/Source/PCG/Private/Metadata/Accessors/PCGAttributeAccessorFactory.cpp
@@ -32,7 +32,14 @@ namespace PCGAttributeAccessorFactory
 		}
 		else if (InSelector.GetSelection() == EPCGAttributePropertySelection::Attribute)
 		{
-			UPCGMetadata* Metadata = InData->MutableMetadata();
+			const FPCGMetadataDomainID DomainID = InData->GetMetadataDomainIDFromSelector(InSelector);
+			if (!DomainID.IsValid())
+			{
+				PCGLog::Metadata::LogInvalidMetadataDomain(InSelector);
+				return TUniquePtr<IPCGAttributeAccessor>();
+			}
+
+			FPCGMetadataDomain* Metadata = InData->MutableMetadata()->GetMetadataDomain(DomainID);
 			FPCGMetadataAttributeBase* Attribute = Metadata ? Metadata->GetMutableAttribute(InSelector.GetAttributeName()) : nullptr;
 		
 			Accessor = PCGAttributeAccessorHelpers::CreateAccessor(Attribute, Metadata, bQuiet);
@@ -62,7 +69,14 @@ namespace PCGAttributeAccessorFactory
 		}
 		else if (InSelector.GetSelection() == EPCGAttributePropertySelection::Attribute)
 		{
-			const UPCGMetadata* Metadata = InData->ConstMetadata();
+			const FPCGMetadataDomainID DomainID = InData->GetMetadataDomainIDFromSelector(InSelector);
+			if (!DomainID.IsValid())
+			{
+				PCGLog::Metadata::LogInvalidMetadataDomain(InSelector);
+				return TUniquePtr<IPCGAttributeAccessor>();
+			}
+
+			const FPCGMetadataDomain* Metadata = InData->ConstMetadata()->GetConstMetadataDomain(DomainID);
 			const FPCGMetadataAttributeBase* Attribute = Metadata ? Metadata->GetConstAttribute(InSelector.GetAttributeName()) : nullptr;
 		
 			Accessor = PCGAttributeAccessorHelpers::CreateConstAccessor(Attribute, Metadata, bQuiet);
@@ -73,23 +87,29 @@ namespace PCGAttributeAccessorFactory
 	
 	TUniquePtr<IPCGAttributeAccessorKeys> CreateDefaultAccessorKeys(UPCGData* InData, const FPCGAttributePropertySelector& InSelector, const bool bQuiet)
 	{
-		UPCGMetadata* Metadata = InData->MutableMetadata();
+		const FPCGMetadataDomainID DomainID = InData->GetMetadataDomainIDFromSelector(InSelector);
+		if (!DomainID.IsValid())
+		{
+			PCGLog::Metadata::LogInvalidMetadataDomain(InSelector);
+			return  TUniquePtr<IPCGAttributeAccessorKeys>();
+		}
+
+		FPCGMetadataDomain* Metadata = InData->MutableMetadata()->GetMetadataDomain(DomainID);
 		return Metadata ? MakeUnique<FPCGAttributeAccessorKeysEntries>(Metadata) : TUniquePtr<IPCGAttributeAccessorKeys>();
 	}
 	
 	TUniquePtr<const IPCGAttributeAccessorKeys> CreateDefaultConstAccessorKeys(const UPCGData* InData, const FPCGAttributePropertySelector& InSelector, const bool bQuiet)
 	{
-		const UPCGMetadata* Metadata = InData->ConstMetadata();
-		return Metadata ? MakeUnique<const FPCGAttributeAccessorKeysEntries>(Metadata) : TUniquePtr<const IPCGAttributeAccessorKeys>();
-	}
-	
-	TUniquePtr<const IPCGAttributeAccessorKeys> CreateSpatialConstAccessorKeys(const UPCGData* InData, const FPCGAttributePropertySelector& InSelector, const bool bQuiet)
-	{
-		// For const metadata, that is on SpatialData (not points), we allow for Default Value to support "data-wide" metadata
-		// It's only for spatial data and not Param data because for ParamData, GetNum on the keys should return the number of entries.
-		// If we have a default key, GetNum will return 1, while there is actually 0 values.
-		const UPCGMetadata* Metadata = InData->ConstMetadata();
-		return Metadata ? MakeUnique<const FPCGAttributeAccessorKeysEntries>(Metadata, /*bAddDefaultValueIfEmpty=*/true) : TUniquePtr<const IPCGAttributeAccessorKeys>();
+		const FPCGMetadataDomainID DomainID = InData->GetMetadataDomainIDFromSelector(InSelector);
+		if (!DomainID.IsValid())
+		{
+			PCGLog::Metadata::LogInvalidMetadataDomain(InSelector);
+			return  TUniquePtr<IPCGAttributeAccessorKeys>();
+		}
+
+		const FPCGMetadataDomain* Metadata = InData->ConstMetadata()->GetConstMetadataDomain(DomainID);
+		const bool bAddDefaultValueIfEmpty = Metadata->GetDomainID() == PCGMetadataDomainID::Data;
+		return Metadata ? MakeUnique<const FPCGAttributeAccessorKeysEntries>(Metadata, /*bAddDefaultValueIfEmpty=*/bAddDefaultValueIfEmpty) : TUniquePtr<const IPCGAttributeAccessorKeys>();
 	}
 }
 
@@ -270,17 +290,10 @@ void FPCGAttributeAccessorFactory::RegisterDefaultMethods()
 		.CreateConstAccessorKeysFunc = PCGAttributeAccessorFactory::CreateDefaultConstAccessorKeys,
 	};
 
-	FPCGAttributeAccessorMethods SpatialMethods
-	{
-		.CreateConstAccessorKeysFunc = PCGAttributeAccessorFactory::CreateSpatialConstAccessorKeys,
-	};
-
 	RegisterMethods<UPCGData>(std::move(DefaultMethods));
-	RegisterMethods<UPCGSpatialData>(std::move(SpatialMethods));
 }
 
 void FPCGAttributeAccessorFactory::UnregisterDefaultMethods()
 {
-	UnregisterMethods<UPCGSpatialData>();
 	UnregisterMethods<UPCGData>();
 }
diff --git a/Engine/Plugins/PCG/Source/PCG/Private/Metadata/PCGMetadata.cpp b/Engine/Plugins/PCG/Source/PCG/Private/Metadata/PCGMetadata.cpp
--- a/Engine/Plugins/PCG/Source/PCG/Private/Metadata/PCGMetadata.cpp
+++ b/Engine/Plugins/PCG/Source/PCG/Private/Metadata/PCGMetadata.cpp
@@ -144,12 +144,24 @@ FPCGMetadataDomain* UPCGMetadata::GetMetadataDomain(const FPCGMetadataDomainID&
 	return FindOrCreateMetadataDomain(InMetadataDomainID);
 }
 
+FPCGMetadataDomain* UPCGMetadata::GetMetadataDomainFromSelector(const FPCGAttributePropertySelector& InSelector)
+{
+	const UPCGData* OwnerData = Cast<UPCGData>(GetOuter());
+	return OwnerData ? GetMetadataDomain(OwnerData->GetMetadataDomainIDFromSelector(InSelector)) : nullptr;
+}
+
 const FPCGMetadataDomain* UPCGMetadata::GetConstMetadataDomain(const FPCGMetadataDomainID& InMetadataDomainID) const
 {
 	const TSharedPtr<FPCGMetadataDomain>* FoundMetadataDomain = MetadataDomains.Find(InMetadataDomainID);
 	return FoundMetadataDomain ? FoundMetadataDomain->Get() : nullptr;
 }
 
+const FPCGMetadataDomain* UPCGMetadata::GetConstMetadataDomainFromSelector(const FPCGAttributePropertySelector& InSelector) const
+{
+	const UPCGData* OwnerData = Cast<UPCGData>(GetOuter());
+	return OwnerData ? GetConstMetadataDomain(OwnerData->GetMetadataDomainIDFromSelector(InSelector)) : nullptr;
+}
+
 FPCGMetadataDomain* UPCGMetadata::FindOrCreateMetadataDomain(const FPCGMetadataDomainID& InMetadataDomainID)
 {
 	TSharedPtr<FPCGMetadataDomain>* FoundMetadataDomain = MetadataDomains.Find(InMetadataDomainID);
@@ -1005,12 +1017,17 @@ void UPCGMetadata::AccumulatePointWeightedAttributes(const FPCGPoint& InPoint, c
 	AccumulateWeightedAttributes(InPoint.MetadataEntry, InMetadata, Weight, bSetNonInterpolableAttributes, OutPoint.MetadataEntry);
 }
 
-void UPCGMetadata::SetLastCachedSelectorOnOwner(FName AttributeName)
+void UPCGMetadata::SetLastCachedSelectorOnOwner(FName AttributeName, FPCGMetadataDomainID DomainID)
 {
 	if (UPCGData* OwnerData = Cast<UPCGData>(GetOuter()))
 	{
 		FPCGAttributePropertyInputSelector Selector;
 		Selector.SetAttributeName(AttributeName);
+		if (!DomainID.IsDefault() && DomainID != DefaultDomain)
+		{
+			OwnerData->SetDomainFromDomainID(DomainID, Selector);
+		}
+
 		OwnerData->SetLastSelector(Selector);
 	}
 }
diff --git a/Engine/Plugins/PCG/Source/PCG/Private/Metadata/PCGMetadataDomain.cpp b/Engine/Plugins/PCG/Source/PCG/Private/Metadata/PCGMetadataDomain.cpp
--- a/Engine/Plugins/PCG/Source/PCG/Private/Metadata/PCGMetadataDomain.cpp
+++ b/Engine/Plugins/PCG/Source/PCG/Private/Metadata/PCGMetadataDomain.cpp
@@ -564,7 +564,7 @@ void FPCGMetadataDomain::RemoveAttributeInternal(FName AttributeName)
 
 void FPCGMetadataDomain::SetLastCachedSelectorOnOwner(FName AttributeName)
 {
-	TopMetadata->SetLastCachedSelectorOnOwner(AttributeName);
+	TopMetadata->SetLastCachedSelectorOnOwner(AttributeName, DomainID);
 }
 
 FPCGMetadataAttributeBase* FPCGMetadataDomain::GetMutableAttribute(FName AttributeName)
diff --git a/Engine/Plugins/PCG/Source/PCG/Private/PCGData.cpp b/Engine/Plugins/PCG/Source/PCG/Private/PCGData.cpp
--- a/Engine/Plugins/PCG/Source/PCG/Private/PCGData.cpp
+++ b/Engine/Plugins/PCG/Source/PCG/Private/PCGData.cpp
@@ -82,6 +82,46 @@ void UPCGData::ClearUsage(EPCGDataUsage InUsage) const
 	}
 }
 
+bool UPCGData::IsSupportedMetadataDomainID(const FPCGMetadataDomainID& InDomainID) const
+{
+	return InDomainID == PCGMetadataDomainID::Default || GetAllSupportedMetadataDomainIDs().Contains(InDomainID);
+}
+
+FPCGMetadataDomainID UPCGData::GetMetadataDomainIDFromSelector(const FPCGAttributePropertySelector& InSelector) const
+{
+	const FName DomainName = InSelector.GetDomainName();
+	if (DomainName == NAME_None)
+	{
+		return PCGMetadataDomainID::Default;
+	}
+	else if (DomainName == PCGDataConstants::GlobalDomainName)
+	{
+		return PCGMetadataDomainID::Data;
+	}
+	else
+	{
+		return FPCGMetadataDomainID(EPCGMetadataDomainFlag::Invalid);
+	}
+}
+
+bool UPCGData::SetDomainFromDomainID(const FPCGMetadataDomainID& InDomainID, FPCGAttributePropertySelector& InOutSelector) const
+{
+	if (InDomainID == PCGMetadataDomainID::Data)
+	{
+		InOutSelector.SetDomainName(PCGDataConstants::GlobalDomainName);
+		return true;
+	}
+	else if (InDomainID == PCGMetadataDomainID::Default)
+	{
+		InOutSelector.SetDomainName(NAME_None);
+		return true;
+	}
+	else
+	{
+		return false;
+	}
+}
+
 FPCGCrc UPCGData::ComputeCrc(bool bFullDataCrc) const
 {
 	FArchiveCrc32 Ar;
diff --git a/Engine/Plugins/PCG/Source/PCG/Private/PCGParamData.cpp b/Engine/Plugins/PCG/Source/PCG/Private/PCGParamData.cpp
--- a/Engine/Plugins/PCG/Source/PCG/Private/PCGParamData.cpp
+++ b/Engine/Plugins/PCG/Source/PCG/Private/PCGParamData.cpp
@@ -155,4 +155,30 @@ UPCGParamData* UPCGParamData::DuplicateData(FPCGContext* Context, bool bInitiali
 	}
 
 	return NewParamData;
+}
+
+FPCGMetadataDomainID UPCGParamData::GetMetadataDomainIDFromSelector(const FPCGAttributePropertySelector& InSelector) const
+{
+	const FName DomainName = InSelector.GetDomainName();
+	if (DomainName == PCGParamDataConstants::ElementsDomainName)
+	{
+		return PCGMetadataDomainID::Elements;
+	}
+	else
+	{
+		return Super::GetMetadataDomainIDFromSelector(InSelector);
+	}
+}
+
+bool UPCGParamData::SetDomainFromDomainID(const FPCGMetadataDomainID& InDomainID, FPCGAttributePropertySelector& InOutSelector) const
+{
+	if (InDomainID == PCGMetadataDomainID::Elements)
+	{
+		InOutSelector.SetDomainName(PCGParamDataConstants::ElementsDomainName);
+		return true;
+	}
+	else
+	{
+		return Super::SetDomainFromDomainID(InDomainID, InOutSelector);
+	}
 }
\ No newline at end of file
diff --git a/Engine/Plugins/PCG/Source/PCG/Private/Tests/Accessors/PCGSplineAccessorTests.cpp.cpp b/Engine/Plugins/PCG/Source/PCG/Private/Tests/Accessors/PCGSplineAccessorTests.cpp.cpp
new file mode 100644
--- /dev/null
+++ b/Engine/Plugins/PCG/Source/PCG/Private/Tests/Accessors/PCGSplineAccessorTests.cpp.cpp
@@ -0,0 +1,169 @@
+// Copyright Epic Games, Inc. All Rights Reserved.
+
+#include "Tests/PCGTestsCommon.h"
+
+#include "Data/PCGSplineData.h"
+#include "Metadata/PCGAttributePropertySelector.h"
+#include "Metadata/PCGMetadata.h"
+#include "Metadata/PCGMetadataDomain.h"
+#include "Metadata/PCGMetadataCommon.h"
+#include "Metadata/Accessors/IPCGAttributeAccessor.h"
+#include "Metadata/Accessors/PCGAttributeAccessorHelpers.h"
+#include "Metadata/Accessors/PCGAttributeAccessorKeys.h"
+
+IMPLEMENT_CUSTOM_SIMPLE_AUTOMATION_TEST(FPCGSplineAccessorControlPointsPropertyTest, FPCGTestBaseClass, "Plugins.PCG.Accessor.Splines.ControlPointProperty", PCGTestsCommon::TestFlags)
+IMPLEMENT_CUSTOM_SIMPLE_AUTOMATION_TEST(FPCGSplineAccessorControlPointsMetadataTest, FPCGTestBaseClass, "Plugins.PCG.Accessor.Splines.ControlPointMetadata", PCGTestsCommon::TestFlags)
+
+bool FPCGSplineAccessorControlPointsPropertyTest::RunTest(const FString& Parameters)
+{
+	UPCGSplineData* SplineData = NewObject<UPCGSplineData>();
+
+	constexpr int NumPoints = 10;
+	TArray<FSplinePoint> SplinePoints;
+	SplinePoints.Reserve(NumPoints);
+	for (int i = 0; i < NumPoints; ++i)
+	{
+		SplinePoints.Emplace(
+			/*InputKey=*/ static_cast<float>(i),
+			/*InPosition=*/ FVector(i, i, i),
+			/*InArriveTangent=*/ FVector(2*i, 2*i,2*i),
+			/*InLeaveTangent=*/ FVector(3*i, 3*i, 3*i),
+			/*InRotation=*/ FRotator(i, i, i),
+			/*InScale=*/ FVector(i, i, i),
+			/*InType=*/ESplinePointType::CurveCustomTangent);
+	}
+
+	SplineData->Initialize(SplinePoints, /*bInClosedLoop*/ false, FTransform::Identity);
+
+	// Index needs to be const
+	TUniquePtr<const IPCGAttributeAccessor> IndexAccessor = PCGAttributeAccessorHelpers::CreateConstAccessor(SplineData, FPCGAttributePropertySelector::CreateExtraPropertySelector(EPCGExtraProperties::Index));
+	TUniquePtr<IPCGAttributeAccessor> PositionAccessor = PCGAttributeAccessorHelpers::CreateAccessor(SplineData, FPCGAttributePropertySelector::CreatePropertySelector(TEXT("Position")));
+	TUniquePtr<IPCGAttributeAccessor> ArriveTangentAccessor = PCGAttributeAccessorHelpers::CreateAccessor(SplineData, FPCGAttributePropertySelector::CreatePropertySelector(TEXT("ArriveTangent")));
+	TUniquePtr<IPCGAttributeAccessor> LeaveTangentAccessor = PCGAttributeAccessorHelpers::CreateAccessor(SplineData, FPCGAttributePropertySelector::CreatePropertySelector(TEXT("LeaveTangent")));
+	TUniquePtr<IPCGAttributeAccessor> RotationAccessor = PCGAttributeAccessorHelpers::CreateAccessor(SplineData, FPCGAttributePropertySelector::CreatePropertySelector(TEXT("Rotation")));
+	TUniquePtr<IPCGAttributeAccessor> ScaleAccessor = PCGAttributeAccessorHelpers::CreateAccessor(SplineData, FPCGAttributePropertySelector::CreatePropertySelector(TEXT("Scale")));
+	TUniquePtr<IPCGAttributeAccessor> InterpTypeAccessor = PCGAttributeAccessorHelpers::CreateAccessor(SplineData, FPCGAttributePropertySelector::CreatePropertySelector(TEXT("InterpType")));
+
+	// Keys can be used for any property, so only create one on the position.
+	TUniquePtr<IPCGAttributeAccessorKeys> SplineKeys = PCGAttributeAccessorHelpers::CreateKeys(SplineData, FPCGAttributePropertySelector::CreatePropertySelector(TEXT("Position")));
+
+	UTEST_TRUE("Index accessor is valid", IndexAccessor.IsValid());
+	UTEST_TRUE("Position accessor is valid", PositionAccessor.IsValid());
+	UTEST_TRUE("ArriveTangent accessor is valid", ArriveTangentAccessor.IsValid());
+	UTEST_TRUE("LeaveTangent accessor is valid", LeaveTangentAccessor.IsValid());
+	UTEST_TRUE("Rotation accessor is valid", RotationAccessor.IsValid());
+	UTEST_TRUE("Scale accessor is valid", ScaleAccessor.IsValid());
+	UTEST_TRUE("InterpType accessor is valid", InterpTypeAccessor.IsValid());
+	
+	UTEST_TRUE("Keys are valid", SplineKeys.IsValid());
+
+	UTEST_EQUAL("Number of keys", SplineKeys->GetNum(), NumPoints);
+
+	// Modify all the points by adding 1 to each numerical value
+	auto ReadModifySetAndValidate = [this, NumPoints, &SplineKeys]<typename T>(TUniquePtr<IPCGAttributeAccessor>& Accessor, const T& ValueToAdd, const TCHAR* What) -> bool
+	{
+		TArray<T> Values;
+		TArray<T> ModifiedValues;
+		Values.SetNum(NumPoints);
+		ModifiedValues.Reserve(NumPoints);
+		
+		UTEST_TRUE(FString::Printf(TEXT("GetRange on '%s' succeeded."), What), Accessor->GetRange<T>(Values, 0, *SplineKeys));
+		
+		for (T& Value : Values)
+		{
+			ModifiedValues.Emplace(Value + ValueToAdd);
+		}
+		
+		UTEST_TRUE(FString::Printf(TEXT("SetRange on '%s' succeeded."), What), Accessor->SetRange<T>(ModifiedValues, 0, *SplineKeys));
+
+		TArray<T> NewValues;
+		NewValues.SetNum(NumPoints);
+		UTEST_TRUE("Second GetRange succeeded.", Accessor->GetRange<T>(NewValues, 0, *SplineKeys));
+		
+		for (int i = 0; i < NumPoints; ++i)
+		{
+			UTEST_NOT_EQUAL(FString::Printf(TEXT("Value for '%s' at index '%d' was is different from the first get"), What, i), NewValues[i], Values[i]);
+			UTEST_EQUAL(FString::Printf(TEXT("Value for '%s' at index '%d' was set correctly"), What, i), NewValues[i], ModifiedValues[i]);
+		}
+
+		return true;
+	};
+
+	ReadModifySetAndValidate(PositionAccessor, FVector::OneVector, TEXT("Position"));
+	ReadModifySetAndValidate(ArriveTangentAccessor, FVector::OneVector, TEXT("ArriveTangent"));
+	ReadModifySetAndValidate(LeaveTangentAccessor, FVector::OneVector, TEXT("LeaveTangent"));
+	ReadModifySetAndValidate(RotationAccessor, FQuat::MakeFromEuler(FVector::OneVector), TEXT("Rotation"));
+	ReadModifySetAndValidate(ScaleAccessor, FVector::OneVector, TEXT("Scale"));
+
+	// Also validate that the index and interp type are also OK
+	TArray<int32> IndexValues;
+	TArray<int64> InterpTypeValues;
+	IndexValues.SetNum(NumPoints);
+	InterpTypeValues.SetNum(NumPoints);
+
+	UTEST_TRUE("GetRange on Index succeeded", IndexAccessor->GetRange<int32>(IndexValues, 0, *SplineKeys));
+	UTEST_TRUE("GetRange on InterpTypes succeeded", InterpTypeAccessor->GetRange<int64>(InterpTypeValues, 0, *SplineKeys));
+
+	for (int i = 0; i < NumPoints; ++i)
+	{
+		UTEST_EQUAL("Index", IndexValues[i], i);
+		UTEST_EQUAL("InterpType", InterpTypeValues[i], static_cast<int64>(ESplinePointType::CurveCustomTangent));
+	}
+	
+	return true;
+}
+
+bool FPCGSplineAccessorControlPointsMetadataTest::RunTest(const FString& Parameters)
+{
+	UPCGSplineData* SplineData = NewObject<UPCGSplineData>();
+
+	constexpr int NumPoints = 10;
+	TArray<FSplinePoint> SplinePoints;
+	SplinePoints.Reserve(NumPoints);
+	for (int i = 0; i < NumPoints; ++i)
+	{
+		SplinePoints.Emplace(/*InputKey=*/ static_cast<float>(i), FVector::ZeroVector);
+	}
+
+	SplineData->Initialize(SplinePoints, /*bInClosedLoop*/ false, FTransform::Identity);
+
+	UPCGMetadata* Metadata = SplineData->MutableMetadata();
+	FPCGMetadataDomain* MetadataDomain = Metadata ? Metadata->GetMetadataDomain(PCGMetadataDomainID::Elements) : nullptr;
+
+	UTEST_NOT_NULL("Metadata domain for control points exist", MetadataDomain);
+	check(MetadataDomain);
+
+	const FName IntAttributeName = TEXT("Int");
+	FPCGMetadataAttribute<int32>* IntAttribute = MetadataDomain->CreateAttribute<int32>(IntAttributeName, -1, false, false);
+	UTEST_NOT_NULL("Attribute was successfully created", IntAttribute);
+
+	const FPCGAttributePropertySelector Selector = FPCGAttributePropertySelector::CreateAttributeSelector(IntAttributeName, TEXT("ControlPoints"));
+	TUniquePtr<IPCGAttributeAccessor> IntAccessor = PCGAttributeAccessorHelpers::CreateAccessor(SplineData, Selector);
+	TUniquePtr<IPCGAttributeAccessorKeys> SplineKeys = PCGAttributeAccessorHelpers::CreateKeys(SplineData, Selector);
+
+	UTEST_TRUE("Int accessor is valid", IntAccessor.IsValid());
+	UTEST_TRUE("Keys are valid", SplineKeys.IsValid());
+	UTEST_EQUAL("Number of keys", SplineKeys->GetNum(), NumPoints);
+
+	TArray<int32> Values;
+	Values.SetNum(NumPoints);
+	UTEST_TRUE("GetRange succeeded.", IntAccessor->GetRange<int32>(Values, 0, *SplineKeys));
+	UTEST_TRUE("All have the default value", Algo::AllOf(Values, [](int32 Value) { return Value == -1; }));
+
+	for (int i = 0; i < NumPoints; ++i)
+	{
+		Values[i] = i;
+	}
+
+	UTEST_TRUE("SetRange succeeded.", IntAccessor->SetRange<int32>(Values, 0, *SplineKeys));
+	TArray<int32> NewValues;
+	NewValues.SetNum(NumPoints);
+
+	UTEST_TRUE("Second GetRange succeeded.", IntAccessor->GetRange<int32>(NewValues, 0, *SplineKeys));
+	for (int i = 0; i < NumPoints; ++i)
+	{
+		UTEST_EQUAL(FString::Printf(TEXT("Value at index '%d' was set correctly"), i), NewValues[i], i);
+	}
+	
+	return true;
+}
\ No newline at end of file
diff --git a/Engine/Plugins/PCG/Source/PCG/Private/Utils/PCGLogErrors.cpp b/Engine/Plugins/PCG/Source/PCG/Private/Utils/PCGLogErrors.cpp
--- a/Engine/Plugins/PCG/Source/PCG/Private/Utils/PCGLogErrors.cpp
+++ b/Engine/Plugins/PCG/Source/PCG/Private/Utils/PCGLogErrors.cpp
@@ -52,6 +52,7 @@ namespace PCGLog::Metadata
 		const FTextFormat SetTypedAttributeFailure = LOCTEXT("SetTypedAttributeFailure", "Could not set value in attribute '{0}'. Attribute type: {1}, Value Type: {2}.");
 		const FTextFormat SetTypedAttributeFailureNoAccessor = LOCTEXT("SetTypedAttributeFailureNoAccessor", "Could not set value in attribute '{0}' value of type: '{1}.");
 		const FTextFormat IncomparableTypesFailure = LOCTEXT("IncomparableTypesFailure", "Attributes '{0}' and '{1}' are incomparable. Ensure they are either of the same or compatible types.");
+		const FTextFormat InvalidMetadataDomain = LOCTEXT("InvalidMetadataDomain", "Metadata domain {0} is invalid for this data.");
 	}
 
 	void LogFailToCreateAccessorError(const FPCGAttributePropertySelector& Selector, const FPCGContext* InContext)
@@ -59,6 +60,11 @@ namespace PCGLog::Metadata
 		LogErrorOnGraph(FText::Format(Format::CreateAccessorFailure, Selector.GetDisplayText()), InContext);
 	}
 
+	void LogInvalidMetadataDomain(const FPCGAttributePropertySelector& Selector, const FPCGContext* InContext)
+	{
+		LogErrorOnGraph(FText::Format(Format::InvalidMetadataDomain, FText::FromName(Selector.GetDomainName())), InContext);
+	}
+
 	void LogFailToGetAttributeError(const FText& AttributeName, const FPCGContext* InContext)
 	{
 		LogErrorOnGraph(FText::Format(Format::GetAttributeFailure, AttributeName), InContext);
diff --git a/Engine/Plugins/PCG/Source/PCG/Public/Data/PCGBasePointData.h b/Engine/Plugins/PCG/Source/PCG/Public/Data/PCGBasePointData.h
--- a/Engine/Plugins/PCG/Source/PCG/Public/Data/PCGBasePointData.h
+++ b/Engine/Plugins/PCG/Source/PCG/Public/Data/PCGBasePointData.h
@@ -13,6 +13,7 @@
 namespace PCGPointDataConstants
 {
 	const FName ActorReferenceAttribute = TEXT("ActorReference");
+	const FName ElementsDomainName = "Points";
 }
 
 UCLASS(Abstract, BlueprintType, ClassGroup = (Procedural))
@@ -42,6 +43,12 @@ public:
 	virtual void Flatten() override;
 	virtual bool SupportsFullDataCrc() const override { return true; }
 	virtual void AddToCrc(FArchiveCrc32& Ar, bool bFullDataCrc) const override;
+
+	/** Metadata specific */
+	virtual FPCGMetadataDomainID GetDefaultMetadataDomainID() const override { return PCGMetadataDomainID::Elements; }
+	virtual TArray<FPCGMetadataDomainID> GetAllSupportedMetadataDomainIDs() const override { return {PCGMetadataDomainID::Data, PCGMetadataDomainID::Elements}; }
+	virtual FPCGMetadataDomainID GetMetadataDomainIDFromSelector(const FPCGAttributePropertySelector& InSelector) const override;
+	virtual bool SetDomainFromDomainID(const FPCGMetadataDomainID& InDomainID, FPCGAttributePropertySelector& InOutSelector) const;
 	// ~End UPCGData Interface
 
 	virtual bool ProjectPoint(const FTransform& InTransform, const FBox& InBounds, const FPCGProjectionParams& InParams, FPCGPoint& OutPoint, UPCGMetadata* OutMetadata, bool bUseBounds) const;
@@ -146,10 +153,6 @@ public:
 	virtual void CopyPointsTo(UPCGBasePointData* To, const TArrayView<const int32>& ReadIndices, const TArrayView<const int32>& WriteIndices) const;
 	virtual EPCGPointNativeProperties GetAllocatedProperties(bool bWithInheritance = true) const { return EPCGPointNativeProperties::All; }
 
-	/** Metadata specific */
-	virtual FPCGMetadataDomainID GetDefaultMetadataDomainID() const override { return PCGMetadataDomainID::Elements; }
-	virtual TArray<FPCGMetadataDomainID> GetAllSupportedMetadataDomainIDs() const override { return {PCGMetadataDomainID::Data, PCGMetadataDomainID::Elements}; }
-
 	static void SetPoints(const UPCGBasePointData* From, UPCGBasePointData* To, const TArrayView<const int32>& InDataIndices, bool bCopyAll);
 
 protected:
diff --git a/Engine/Plugins/PCG/Source/PCG/Public/Data/PCGSplineData.h b/Engine/Plugins/PCG/Source/PCG/Public/Data/PCGSplineData.h
--- a/Engine/Plugins/PCG/Source/PCG/Public/Data/PCGSplineData.h
+++ b/Engine/Plugins/PCG/Source/PCG/Public/Data/PCGSplineData.h
@@ -26,6 +26,8 @@ class PCG_API UPCGSplineData : public UPCGPolyLineData
 	GENERATED_BODY()
 
 public:
+	UPCGSplineData();
+
 	void Initialize(const USplineComponent* InSpline);
 	void Initialize(const TArray<FSplinePoint>& InSplinePoints, bool bInClosedLoop, const FTransform& InTransform, TArray<PCGMetadataEntryKey> InOptionalEntryKeys = {});
 	void Initialize(const FPCGSplineStruct& InSplineStruct);
@@ -34,6 +36,12 @@ public:
 	// ~Begin UPCGData interface
 	virtual EPCGDataType GetDataType() const override { return EPCGDataType::Spline; }
 	virtual void AddToCrc(FArchiveCrc32& Ar, bool bFullDataCrc) const override;
+
+	// To be enabled when we are sure Default translate well from Data to Elements (or we have a path for deprecation)
+	//virtual FPCGMetadataDomainID GetDefaultMetadataDomainID() const override { return PCGMetadataDomainID::Elements; }
+	virtual TArray<FPCGMetadataDomainID> GetAllSupportedMetadataDomainIDs() const override { return {PCGMetadataDomainID::Data, PCGMetadataDomainID::Elements}; }
+	virtual FPCGMetadataDomainID GetMetadataDomainIDFromSelector(const FPCGAttributePropertySelector& InSelector) const override;
+	virtual bool SetDomainFromDomainID(const FPCGMetadataDomainID& InDomainID, FPCGAttributePropertySelector& InOutSelector) const;
 	// ~End UPCGData interface
 
 	//~Begin UPCGPolyLineData interface
diff --git a/Engine/Plugins/PCG/Source/PCG/Public/Metadata/PCGMetadata.h b/Engine/Plugins/PCG/Source/PCG/Public/Metadata/PCGMetadata.h
--- a/Engine/Plugins/PCG/Source/PCG/Public/Metadata/PCGMetadata.h
+++ b/Engine/Plugins/PCG/Source/PCG/Public/Metadata/PCGMetadata.h
@@ -393,9 +393,11 @@ public:
 
 	/** Expected to be used when metadata domains are already setup. */
 	FPCGMetadataDomain* GetMetadataDomain(const FPCGMetadataDomainID& InMetadataDomainID);
+	FPCGMetadataDomain* GetMetadataDomainFromSelector(const FPCGAttributePropertySelector& InSelector);
 	
 	/** Expected to be used when metadata domain are already setup. */
 	const FPCGMetadataDomain* GetConstMetadataDomain(const FPCGMetadataDomainID& InMetadataDomainID) const;
+	const FPCGMetadataDomain* GetConstMetadataDomainFromSelector(const FPCGAttributePropertySelector& InSelector) const;
 
 	FPCGMetadataDomain* GetDefaultMetadataDomain() { return GetMetadataDomain(PCGMetadataDomainID::Default); }
 	const FPCGMetadataDomain* GetConstDefaultMetadataDomain() const { return GetConstMetadataDomain(PCGMetadataDomainID::Default); }
@@ -426,7 +428,7 @@ protected:
 
 	bool ParentHasAttribute(FPCGAttributeIdentifier AttributeName) const;
 
-	void SetLastCachedSelectorOnOwner(FName AttributeName);
+	void SetLastCachedSelectorOnOwner(FName AttributeName, FPCGMetadataDomainID DomainID);
 
 	void SetupDomainsFromOtherMetadataIfNeeded(const UPCGMetadata* OtherMetadata);
 
diff --git a/Engine/Plugins/PCG/Source/PCG/Public/Metadata/PCGMetadataCommon.h b/Engine/Plugins/PCG/Source/PCG/Public/Metadata/PCGMetadataCommon.h
--- a/Engine/Plugins/PCG/Source/PCG/Public/Metadata/PCGMetadataCommon.h
+++ b/Engine/Plugins/PCG/Source/PCG/Public/Metadata/PCGMetadataCommon.h
@@ -54,6 +54,9 @@ enum class EPCGMetadataDomainFlag : uint8
 
 	/** Metadata on elements like points on point data and entries on param data. */
 	Elements = 2,
+	
+	/** For invalid domain. */
+	Invalid = 254,
 
 	/** For data that can have more domains. */
 	Custom = 255
@@ -97,6 +100,7 @@ struct FPCGMetadataDomainID
 	}
 
 	bool IsDefault() const { return Flag == EPCGMetadataDomainFlag::Default; }
+	bool IsValid() const { return Flag != EPCGMetadataDomainFlag::Invalid; }
 
 	UPROPERTY(EditAnywhere, BlueprintReadWrite, Category = "")
 	EPCGMetadataDomainFlag Flag = EPCGMetadataDomainFlag::Default;
@@ -113,6 +117,7 @@ namespace PCGMetadataDomainID
 	static const FPCGMetadataDomainID Default{EPCGMetadataDomainFlag::Default, -1, TEXT("Default")};
 	static const FPCGMetadataDomainID Elements{EPCGMetadataDomainFlag::Elements, -1, TEXT("Elements")};
 	static const FPCGMetadataDomainID Data{EPCGMetadataDomainFlag::Data, -1, TEXT("Data")};
+	static const FPCGMetadataDomainID Invalid{EPCGMetadataDomainFlag::Invalid, -1, TEXT("Invalid")};
 }
 
 USTRUCT(BlueprintType)
diff --git a/Engine/Plugins/PCG/Source/PCG/Public/Metadata/PCGMetadataDomain.h b/Engine/Plugins/PCG/Source/PCG/Public/Metadata/PCGMetadataDomain.h
--- a/Engine/Plugins/PCG/Source/PCG/Public/Metadata/PCGMetadataDomain.h
+++ b/Engine/Plugins/PCG/Source/PCG/Public/Metadata/PCGMetadataDomain.h
@@ -229,6 +229,8 @@ public:
 	/** Computes Crc from all attributes & keys from outer's data. */
 	void AddToCrc(FArchiveCrc32& Ar, const UPCGData* Data, bool bFullDataCrc) const;
 
+	FPCGMetadataDomainID GetDomainID() const { return DomainID; }
+
 protected:
 	FPCGMetadataAttributeBase* CopyAttribute(FName AttributeToCopy, FName NewAttributeName, bool bKeepParent, bool bCopyEntries, bool bCopyValues);
 
diff --git a/Engine/Plugins/PCG/Source/PCG/Public/PCGData.h b/Engine/Plugins/PCG/Source/PCG/Public/PCGData.h
--- a/Engine/Plugins/PCG/Source/PCG/Public/PCGData.h
+++ b/Engine/Plugins/PCG/Source/PCG/Public/PCGData.h
@@ -32,6 +32,11 @@ enum class EPCGDataUsage : uint8
 };
 ENUM_CLASS_FLAGS(EPCGDataUsage);
 
+namespace PCGDataConstants
+{
+	const FName GlobalDomainName = "Global";
+}
+
 /**
 * Base class for any "data" class in the PCG framework.
 * This is an intentionally vague base class so we can have the required
@@ -73,8 +78,13 @@ public:
 	/** CRC for this object instance. */
 	mutable FPCGCrc Crc;
 
+	/** Returns true if the data has a cached last selector. Used to know how to convert `@Last` in an attribute selector. */
 	virtual bool HasCachedLastSelector() const { return false; }
+
+	/** Returns the cached last selector. Used to know how to convert `@Last` in an attribute selector. */
 	virtual FPCGAttributePropertyInputSelector GetCachedLastSelector() const { return FPCGAttributePropertyInputSelector{}; }
+
+	/** Set the last selector used to modify an attribute. */
 	virtual void SetLastSelector(const FPCGAttributePropertySelector& InSelector) {};
 
 	UE_DEPRECATED(5.5, "Call/Implement version with FPCGContext parameter")
@@ -96,9 +106,22 @@ public:
 	virtual UPCGMetadata* MutableMetadata() { return Metadata; }
 
 	virtual void Flatten();
+
+	/** Returns the default domain for this data. Any data that has other default domain should override this method. */
 	virtual FPCGMetadataDomainID GetDefaultMetadataDomainID() const { return PCGMetadataDomainID::Data; }
+
+	/** Returns all the supported domain for this data. Any data that has more than one domain should override this method. */
 	virtual TArray<FPCGMetadataDomainID> GetAllSupportedMetadataDomainIDs() const { return {PCGMetadataDomainID::Data}; }
 
+	/** Returns true if the domain is supported by this data's metadata. */
+	bool IsSupportedMetadataDomainID(const FPCGMetadataDomainID& InDomainID) const;
+
+	/** Return the associated domain ID for a given domain name in the selector. Must return "Invalid" if it is not supported. Each data is responsible to override this for every supported domain.*/
+	virtual FPCGMetadataDomainID GetMetadataDomainIDFromSelector(const FPCGAttributePropertySelector& InSelector) const;
+
+	/** Write the domain name into the selector associated with this domain ID. Each data is responsible to override this for every supported domain. */
+	virtual bool SetDomainFromDomainID(const FPCGMetadataDomainID& InDomainID, FPCGAttributePropertySelector& InOutSelector) const;
+
 	// Not accessible through blueprint to make sure the constness is preserved.
 	UPROPERTY()
 	TObjectPtr<UPCGMetadata> Metadata = nullptr;
diff --git a/Engine/Plugins/PCG/Source/PCG/Public/PCGParamData.h b/Engine/Plugins/PCG/Source/PCG/Public/PCGParamData.h
--- a/Engine/Plugins/PCG/Source/PCG/Public/PCGParamData.h
+++ b/Engine/Plugins/PCG/Source/PCG/Public/PCGParamData.h
@@ -7,6 +7,11 @@
 
 #include "PCGParamData.generated.h"
 
+namespace PCGParamDataConstants
+{
+	const FName ElementsDomainName = "Elements";
+}
+
 /**
 * Class to hold execution parameters that will be consumed in nodes of the graph
 */
@@ -31,6 +36,8 @@ public:
 
 	virtual FPCGMetadataDomainID GetDefaultMetadataDomainID() const override { return PCGMetadataDomainID::Elements; }
 	virtual TArray<FPCGMetadataDomainID> GetAllSupportedMetadataDomainIDs() const override { return {PCGMetadataDomainID::Data, PCGMetadataDomainID::Elements}; }
+	virtual bool SetDomainFromDomainID(const FPCGMetadataDomainID& InDomainID, FPCGAttributePropertySelector& InOutSelector) const override;
+	virtual FPCGMetadataDomainID GetMetadataDomainIDFromSelector(const FPCGAttributePropertySelector& InSelector) const override;
 
 	/** Returns the entry for the given name */
 	UFUNCTION(BlueprintCallable, Category = Params)
diff --git a/Engine/Plugins/PCG/Source/PCG/Public/Utils/PCGLogErrors.h b/Engine/Plugins/PCG/Source/PCG/Public/Utils/PCGLogErrors.h
--- a/Engine/Plugins/PCG/Source/PCG/Public/Utils/PCGLogErrors.h
+++ b/Engine/Plugins/PCG/Source/PCG/Public/Utils/PCGLogErrors.h
@@ -42,10 +42,12 @@ namespace PCGLog
 			PCG_API extern const FTextFormat GetTypedAttributeFailureNoAccessor;
 			PCG_API extern const FTextFormat SetTypedAttributeFailure;
 			PCG_API extern const FTextFormat SetTypedAttributeFailureNoAccessor;
+			PCG_API extern const FTextFormat InvalidMetadataDomain;
 		}
 
 		// Errors
 		PCG_API void LogFailToCreateAccessorError(const FPCGAttributePropertySelector& Selector, const FPCGContext* InContext = nullptr);
+		PCG_API void LogInvalidMetadataDomain(const FPCGAttributePropertySelector& Selector, const FPCGContext* InContext = nullptr);
 
 		template <typename T>
 		void LogFailToCreateAttributeError(const FText& AttributeName, const FPCGContext* InContext = nullptr)
