[PCG] Fix static analysis

#rnx

diff --git a/Engine/Plugins/PCG/Source/PCG/Private/Metadata/Accessors/PCGAttributeAccessorFactory.cpp b/Engine/Plugins/PCG/Source/PCG/Private/Metadata/Accessors/PCGAttributeAccessorFactory.cpp
--- a/Engine/Plugins/PCG/Source/PCG/Private/Metadata/Accessors/PCGAttributeAccessorFactory.cpp
+++ b/Engine/Plugins/PCG/Source/PCG/Private/Metadata/Accessors/PCGAttributeAccessorFactory.cpp
@@ -108,7 +108,7 @@ namespace PCGAttributeAccessorFactory
 		}
 
 		const FPCGMetadataDomain* Metadata = InData->ConstMetadata()->GetConstMetadataDomain(DomainID);
-		const bool bAddDefaultValueIfEmpty = Metadata->GetDomainID() == PCGMetadataDomainID::Data;
+		const bool bAddDefaultValueIfEmpty = Metadata ? (Metadata->GetDomainID() == PCGMetadataDomainID::Data) : false;
 		return Metadata ? MakeUnique<const FPCGAttributeAccessorKeysEntries>(Metadata, /*bAddDefaultValueIfEmpty=*/bAddDefaultValueIfEmpty) : TUniquePtr<const IPCGAttributeAccessorKeys>();
 	}
 }
