[PCG] Fix ensures when dumping data descriptions for pins named None

#rnx

diff --git a/Engine/Plugins/PCG/Source/PCG/Private/Compute/PCGComputeGraph.cpp b/Engine/Plugins/PCG/Source/PCG/Private/Compute/PCGComputeGraph.cpp
--- a/Engine/Plugins/PCG/Source/PCG/Private/Compute/PCGComputeGraph.cpp
+++ b/Engine/Plugins/PCG/Source/PCG/Private/Compute/PCGComputeGraph.cpp
@@ -248,20 +248,26 @@ void UPCGComputeGraph::DebugLogDataDescriptions(const UPCGDataBinding* InBinding
 			PCGKernel->GetInputPins(InputPins);
 			for (const FPCGPinProperties& Pin : InputPins)
 			{
-				FPCGDataCollectionDesc DataDesc;
-				ensure(InBinding->GetKernelPinDataDesc_GameThread(PCGKernel, Pin.Label, true, DataDesc));
-				UE_LOG(LogPCG, Display, TEXT("\tInput Pin: %s (%d data)"), *Pin.Label.ToString(), DataDesc.DataDescs.Num());
-				LogDataDescription(DataDesc);
+				if (Pin.Label != NAME_None)
+				{
+					FPCGDataCollectionDesc DataDesc;
+					ensure(InBinding->GetKernelPinDataDesc_GameThread(PCGKernel, Pin.Label, true, DataDesc));
+					UE_LOG(LogPCG, Display, TEXT("\tInput Pin: %s (%d data)"), *Pin.Label.ToString(), DataDesc.DataDescs.Num());
+					LogDataDescription(DataDesc);
+				}
 			}
 
 			TArray<FPCGPinPropertiesGPU> OutputPins;
 			PCGKernel->GetOutputPins(OutputPins);
 			for (const FPCGPinPropertiesGPU& Pin : OutputPins)
 			{
-				FPCGDataCollectionDesc DataDesc;
-				ensure(InBinding->GetKernelPinDataDesc_GameThread(PCGKernel, Pin.Label, false, DataDesc));
-				UE_LOG(LogPCG, Display, TEXT("\tOutput Pin: %s (%d data)"), *Pin.Label.ToString(), DataDesc.DataDescs.Num());
-				LogDataDescription(DataDesc);
+				if (Pin.Label != NAME_None)
+				{
+					FPCGDataCollectionDesc DataDesc;
+					ensure(InBinding->GetKernelPinDataDesc_GameThread(PCGKernel, Pin.Label, false, DataDesc));
+					UE_LOG(LogPCG, Display, TEXT("\tOutput Pin: %s (%d data)"), *Pin.Label.ToString(), DataDesc.DataDescs.Num());
+					LogDataDescription(DataDesc);
+				}
 			}
 		}
 	}
