[PCG] Fix DebugDataInterface crashing in standalone build

#rnx

diff --git a/Engine/Plugins/PCG/Source/PCG/Private/Compute/DataInterfaces/PCGDebugDataInterface.cpp b/Engine/Plugins/PCG/Source/PCG/Private/Compute/DataInterfaces/PCGDebugDataInterface.cpp
--- a/Engine/Plugins/PCG/Source/PCG/Private/Compute/DataInterfaces/PCGDebugDataInterface.cpp
+++ b/Engine/Plugins/PCG/Source/PCG/Private/Compute/DataInterfaces/PCGDebugDataInterface.cpp
@@ -126,19 +126,18 @@ void FPCGDebugDataProviderProxy::GatherDispatchData(FDispatchData const& InDispa
 
 void FPCGDebugDataProviderProxy::AllocateResources(FRDGBuilder& GraphBuilder, FAllocationData const& InAllocationData)
 {
-	if (DebugBufferSizeFloats > 0)
-	{
-		FRDGBufferDesc Desc = FRDGBufferDesc::CreateByteAddressDesc(sizeof(float) * DebugBufferSizeFloats);
-		Desc.Usage |= BUF_SourceCopy;
+	const uint32 NumFloats = FMath::Max(DebugBufferSizeFloats, 1u);
 
-		DebugBuffer = GraphBuilder.CreateBuffer(Desc, TEXT("PCGDebugBuffer"));
-		DebugBufferUAV = GraphBuilder.CreateUAV(DebugBuffer);
+	FRDGBufferDesc Desc = FRDGBufferDesc::CreateByteAddressDesc(sizeof(float) * NumFloats);
+	Desc.Usage |= BUF_SourceCopy;
 
-		TArray<float> ZeroInitializedBuffer;
-		ZeroInitializedBuffer.SetNumZeroed(DebugBufferSizeFloats);
+	DebugBuffer = GraphBuilder.CreateBuffer(Desc, TEXT("PCGDebugBuffer"));
+	DebugBufferUAV = GraphBuilder.CreateUAV(DebugBuffer);
 
-		GraphBuilder.QueueBufferUpload(DebugBuffer, ZeroInitializedBuffer.GetData(), ZeroInitializedBuffer.Num() * sizeof(float), ERDGInitialDataFlags::None);
-	}
+	TArray<float> ZeroInitializedBuffer;
+	ZeroInitializedBuffer.SetNumZeroed(NumFloats);
+
+	GraphBuilder.QueueBufferUpload(DebugBuffer, MakeArrayView(ZeroInitializedBuffer));
 }
 
 void FPCGDebugDataProviderProxy::GetReadbackData(TArray<FReadbackData>& OutReadbackData) const
diff --git a/Engine/Plugins/PCG/Source/PCG/Private/Compute/DataInterfaces/PCGDebugDataInterface.h b/Engine/Plugins/PCG/Source/PCG/Private/Compute/DataInterfaces/PCGDebugDataInterface.h
--- a/Engine/Plugins/PCG/Source/PCG/Private/Compute/DataInterfaces/PCGDebugDataInterface.h
+++ b/Engine/Plugins/PCG/Source/PCG/Private/Compute/DataInterfaces/PCGDebugDataInterface.h
@@ -31,6 +31,7 @@ public:
 	void SetDebugBufferSize(uint32 InDebugBufferSizeFloats) { DebugBufferSizeFloats = InDebugBufferSizeFloats; }
 
 public:
+	UPROPERTY()
 	uint32 DebugBufferSizeFloats = 0;
 };
 
