[PCG] Add CVar to control whether to use Landscape VTs and disable for now due to generating with low res pages, until we can get a better handle on VT behaviour.

diff --git a/Engine/Plugins/PCG/Source/PCG/Private/Compute/DataInterfaces/PCGLandscapeDataInterface.cpp b/Engine/Plugins/PCG/Source/PCG/Private/Compute/DataInterfaces/PCGLandscapeDataInterface.cpp
--- a/Engine/Plugins/PCG/Source/PCG/Private/Compute/DataInterfaces/PCGLandscapeDataInterface.cpp
+++ b/Engine/Plugins/PCG/Source/PCG/Private/Compute/DataInterfaces/PCGLandscapeDataInterface.cpp
@@ -28,6 +28,14 @@
 
 #define LOCTEXT_NAMESPACE "UPCGLandscapeDataInterface"
 
+namespace PCGLandscapeDataInterface
+{
+	static TAutoConsoleVariable<bool> CVarUseLandscapeVTs(
+		TEXT("pcg.GPU.UseLandscapeVTs"),
+		true,
+		TEXT("Allow the landscape data interface to use landscape virtual textures if available."));
+}
+
 /** RenderResource used to hold textures generated by this DI, pulled from the collision geometry of the terrain. */
 class FPCGLandscapeTextureResource : public FRenderResource
 {
@@ -334,7 +342,10 @@ void UPCGLandscapeDataProvider::InitFromLandscape(ALandscape* InLandscape, const
 		return;
 	}
 
-	InitializeRuntimeVirtualTextures(InLandscape);
+	if (PCGLandscapeDataInterface::CVarUseLandscapeVTs.GetValueOnAnyThread())
+	{
+		InitializeRuntimeVirtualTextures(InLandscape);
+	}
 
 	// If the height can't be accessed via RVT, build a physics based height texture instead.
 	if (!HeightVirtualTexture)
