Fix Clang Editor linker error
#rnx

diff --git a/Engine/Plugins/PCG/Source/PCG/Public/PCGSubgraph.h b/Engine/Plugins/PCG/Source/PCG/Public/PCGSubgraph.h
--- a/Engine/Plugins/PCG/Source/PCG/Public/PCGSubgraph.h
+++ b/Engine/Plugins/PCG/Source/PCG/Public/PCGSubgraph.h
@@ -214,7 +214,7 @@ public:
 	// Should not verify if outputs are used multiples times, as we are just outputting the input collection we have.
 	virtual bool ShouldVerifyIfOutputsAreUsedMultipleTimes(const UPCGSettings* InSettings) const override { return false; }
 protected:
-	virtual bool ExecuteInternal(FPCGContext* Context) const override;
+	PCG_API virtual bool ExecuteInternal(FPCGContext* Context) const override;
 	virtual bool IsPassthrough(const UPCGSettings* InSettings) const override { return true; }
 	virtual bool SupportsGPUResidentData(FPCGContext* InContext) const override { return true; }
 	virtual bool SupportsBasePointDataInputs(FPCGContext* InContext) const { return true; }
