[PCG] State machine for compute graph element

* Add yields between each phase which should help time slicing
* Paves the way to moving one or more phases off the GT

#rnx

diff --git a/Engine/Plugins/PCG/Source/PCG/Private/Compute/Elements/PCGComputeGraphElement.cpp b/Engine/Plugins/PCG/Source/PCG/Private/Compute/Elements/PCGComputeGraphElement.cpp
--- a/Engine/Plugins/PCG/Source/PCG/Private/Compute/Elements/PCGComputeGraphElement.cpp
+++ b/Engine/Plugins/PCG/Source/PCG/Private/Compute/Elements/PCGComputeGraphElement.cpp
@@ -50,6 +50,11 @@ FPCGComputeGraphContext::~FPCGComputeGraphContext()
 	}
 }
 
+bool FPCGComputeGraphContext::HasPendingAsyncOperations() const
+{
+	return !ProvidersWithBufferExports.IsEmpty() || !ProvidersWithReadbacks.IsEmpty();
+}
+
 void FPCGComputeGraphContext::AddExtraStructReferencedObjects(FReferenceCollector& Collector)
 {
 	if (bComputeGraphInstanceInitialized)
@@ -75,12 +80,30 @@ bool FPCGComputeGraphElement::ExecuteInternal(FPCGContext* InContext) const
 	TRACE_CPUPROFILER_EVENT_SCOPE(FPCGComputeGraphElement::ExecuteInternal);
 	check(InContext);
 	check(InContext->SourceComponent.IsValid());
-	check(IsInGameThread());
 
 	FPCGComputeGraphContext* Context = static_cast<FPCGComputeGraphContext*>(InContext);
 
-	if (!Context->ComputeGraph)
+	auto SleepUntilNextFrame = [Context]()
+	{
+		Context->bIsPaused = true;
+		Context->SourceComponent->GetSubsystem()->RegisterBeginTickAction([ContextHandle = Context->GetOrCreateHandle()]()
+		{
+			if (TSharedPtr<FPCGContextHandle> SharedHandle = ContextHandle.Pin())
+			{
+				if (FPCGContext* ContextPtr = SharedHandle->GetContext())
+				{
+					ContextPtr->bIsPaused = false;
+				}
+			}
+		});
+	};
+
+	switch (Context->ExecutionSubPhase)
+	{
+	case EPCGComputeGraphExecutionPhase::GetComputeGraph:
 	{
+		TRACE_CPUPROFILER_EVENT_SCOPE(GetComputeGraph);
+
 		const UPCGComponent* Component = InContext->SourceComponent.Get();
 		const UPCGGraph* TopGraph = Component ? Component->GetGraph() : nullptr;
 		if (!ensure(TopGraph))
@@ -112,82 +135,19 @@ bool FPCGComputeGraphElement::ExecuteInternal(FPCGContext* InContext) const
 		}
 
 		Context->ComputeGraphIndex = ComputeGraphIndex;
-	}
 
-	// Drive the execution of a compute graph. The stages are commented below and numbered by the sequence in which they are executed.
-	// The sequence appears out of order as calls so that paths that are executed multiple times (like checking for completion) are as
-	// short as possible.
+		Context->ExecutionSubPhase = EPCGComputeGraphExecutionPhase::InitializeDataBinding;
 
-	auto SleepUntilNextFrame = [](FPCGComputeGraphContext* InContext)
-	{
-		check(InContext);
-		InContext->bIsPaused = true;
-		InContext->SourceComponent->GetSubsystem()->RegisterBeginTickAction([ContextHandle = InContext->GetOrCreateHandle()]()
+		if (Context->AsyncState.ShouldStop())
 		{
-			if (TSharedPtr<FPCGContextHandle> SharedHandle = ContextHandle.Pin())
-			{
-				if (FPCGContext* ContextPtr = SharedHandle->GetContext())
-				{
-					ContextPtr->bIsPaused = false;
-				}
-			}
-		});
-	};
-
-	// 7. Execution is complete when graph was enqueued/dispatched and all exports and readbacks are complete.
-	if (Context->bGraphEnqueued && Context->ProvidersWithBufferExports.IsEmpty() && Context->ProvidersWithReadbacks.IsEmpty())
-	{
-		// Apply any read back element counts to the final data description.
-		Context->DataBinding->ApplyElementCounts_GameThread();
-
-#if WITH_EDITOR
-		// Call both to start all required readbacks.
-		const bool bDebugDrawDone = ExecuteDebugDraw(Context);
-		const bool bInspectionDone = StoreDataForInspection(Context);
-
-		if (!bDebugDrawDone || !bInspectionDone)
-		{
-			SleepUntilNextFrame(Context);
 			return false;
 		}
-#endif
-
-		Context->bExecutionSuccess = true;
-
-		// Currently we don't output anything if processing any readback data processing failed.
-		if (ensure(Context->bExecutionSuccess) && ensure(Context->DataBinding))
-		{
-			Context->OutputData = Context->DataBinding->OutputDataCollection;
-		}
-
-		return true;
 	}
-
-	// 3. If still compiling, try again next frame.
-	if (Context->ComputeGraph->HasKernelResourcesPendingShaderCompilation())
+	case EPCGComputeGraphExecutionPhase::InitializeDataBinding: // Fallthrough
 	{
-		UE_LOG(LogPCG, Verbose, TEXT("Deferring until next frame as the kernel has pending shader compilations."));
-		SleepUntilNextFrame(Context);
-		return false;
-	}
-
-	// 6. Keep waiting for execution to complete.
-	if (Context->bGraphEnqueued)
-	{
-		if (Context->bGraphSubmitFailed)
-		{
-			// Submit failed, no choice but to exit so that execution doesn't hang.
-			return true;
-		}
+		TRACE_CPUPROFILER_EVENT_SCOPE(InitializeDataBinding);
+		ensure(!Context->DataBinding);
 
-		// Likely we need a frame to pass in order to make progress.
-		SleepUntilNextFrame(Context);
-		return false;
-	}
-
-	// 4. Initialize and parse incoming data for data sizes, attributes, etc that will drive buffer allocations and dispatch thread counts.
-	if (!Context->DataBinding)
-	{
 		UPCGDataBinding* DataBinding = FPCGContext::NewObject_AnyThread<UPCGDataBinding>(Context);
 		Context->DataBinding.Reset(DataBinding);
 
@@ -232,7 +192,7 @@ bool FPCGComputeGraphElement::ExecuteInternal(FPCGContext* InContext) const
 
 				const uint64 OriginatingGenerationCount = DataProvider->GenerationCounter;
 
-				DataProvider->OnDataExported_GameThread().AddLambda([ContextHandle=Context->GetOrCreateHandle(), DataProviderWeak, GenerationCount=OriginatingGenerationCount]()
+				DataProvider->OnDataExported_GameThread().AddLambda([ContextHandle = Context->GetOrCreateHandle(), DataProviderWeak, GenerationCount = OriginatingGenerationCount]()
 				{
 					check(IsInGameThread());
 
@@ -248,7 +208,7 @@ bool FPCGComputeGraphElement::ExecuteInternal(FPCGContext* InContext) const
 
 					ContextPtr->ProvidersWithBufferExports.Remove(DataProvider);
 
-					if (ContextPtr->ProvidersWithBufferExports.IsEmpty())
+					if (!ContextPtr->HasPendingAsyncOperations())
 					{
 						ContextPtr->bIsPaused = false;
 					}
@@ -261,7 +221,7 @@ bool FPCGComputeGraphElement::ExecuteInternal(FPCGContext* InContext) const
 
 				TWeakObjectPtr<UPCGDataCollectionDataProvider> DataProviderWeak = DataProvider;
 
-				DataProvider->OnCountersReadBack_GameThread().AddLambda([ContextHandle=Context->GetOrCreateHandle(), DataProviderWeak]()
+				DataProvider->OnCountersReadBack_GameThread().AddLambda([ContextHandle = Context->GetOrCreateHandle(), DataProviderWeak]()
 				{
 					check(IsInGameThread());
 
@@ -276,20 +236,25 @@ bool FPCGComputeGraphElement::ExecuteInternal(FPCGContext* InContext) const
 					ensure(ContextPtr->ProvidersWithReadbacks.Contains(DataProvider));
 					ContextPtr->ProvidersWithReadbacks.Remove(DataProvider);
 
-					if (ContextPtr->ProvidersWithReadbacks.IsEmpty())
+					if (!ContextPtr->HasPendingAsyncOperations())
 					{
 						ContextPtr->bIsPaused = false;
 					}
 				});
 			}
 		}
-	}
 
-	check(Context->DataBinding && InContext->SourceComponent.Get());
+		Context->ExecutionSubPhase = EPCGComputeGraphExecutionPhase::PreExecuteReadbacks;
 
-	// Allow data providers to do early readbacks, which some may use to download analysis results which feed into final data descriptions.
-	if (!Context->bDataProvidersReadbacksComplete)
+		if (Context->AsyncState.ShouldStop())
+		{
+			return false;
+		}
+	}
+	case EPCGComputeGraphExecutionPhase::PreExecuteReadbacks: // Fallthrough
 	{
+		TRACE_CPUPROFILER_EVENT_SCOPE(PreExecuteReadbacks);
+
 		bool bAllReady = true;
 
 		for (UComputeDataProvider* ComputeDataProvider : Context->ComputeGraphInstance.GetDataProviders())
@@ -303,18 +268,25 @@ bool FPCGComputeGraphElement::ExecuteInternal(FPCGContext* InContext) const
 			}
 		}
 
-		if (!bAllReady)
+		if (bAllReady)
+		{
+			Context->ExecutionSubPhase = EPCGComputeGraphExecutionPhase::PrepareForExecute;
+		}
+		else
 		{
-			SleepUntilNextFrame(Context);
+			SleepUntilNextFrame();
 			return false;
 		}
 
-		Context->bDataProvidersReadbacksComplete = true;
+		if (Context->AsyncState.ShouldStop())
+		{
+			return false;
+		}
 	}
-
-	// Prepare data providers for execution.
-	if (!Context->bDataProvidersReady)
+	case EPCGComputeGraphExecutionPhase::PrepareForExecute: // Fallthrough
 	{
+		TRACE_CPUPROFILER_EVENT_SCOPE(PrepareForExecute);
+
 		bool bAllReady = true;
 
 		for (UComputeDataProvider* ComputeDataProvider : Context->ComputeGraphInstance.GetDataProviders())
@@ -328,54 +300,75 @@ bool FPCGComputeGraphElement::ExecuteInternal(FPCGContext* InContext) const
 			}
 		}
 
-		if (!bAllReady)
+		if (bAllReady)
 		{
-			SleepUntilNextFrame(Context);
+			Context->ExecutionSubPhase = EPCGComputeGraphExecutionPhase::ValidateComputeGraphCompilation;
+		}
+		else
+		{
+			SleepUntilNextFrame();
 			return false;
 		}
 
-		Context->bDataProvidersReady = true;
+		if (Context->AsyncState.ShouldStop())
+		{
+			return false;
+		}
 	}
-
-	// 1. Prepare render resources. In editor, this will trigger shader compile if not compiled already.
-	if (!Context->ComputeGraph->GetRenderProxy())
+	case EPCGComputeGraphExecutionPhase::ValidateComputeGraphCompilation: // Fallthrough
 	{
-		TRACE_CPUPROFILER_EVENT_SCOPE(Context->ComputeGraph->UpdateResources);
+		TRACE_CPUPROFILER_EVENT_SCOPE(ValidateComputeGraphCompilation);
 
-		Context->ComputeGraph->UpdateResources();
+		if (Context->ComputeGraph->HasKernelResourcesPendingShaderCompilation())
+		{
+			UE_LOG(LogPCG, Verbose, TEXT("Deferring until next frame as the kernel has pending shader compilations."));
+			SleepUntilNextFrame();
+			return false;
+		}
+		else if (!Context->ComputeGraph->GetRenderProxy())
+		{
+			TRACE_CPUPROFILER_EVENT_SCOPE(Context->ComputeGraph->UpdateResources);
 
-		SleepUntilNextFrame(Context);
-		return false;
-	}
+			Context->ComputeGraph->UpdateResources();
 
-	// 2. Validate compilation
-	{
-		// Add any messages that may have occurred during compilation to visual logs.
+			SleepUntilNextFrame();
+			return false;
+		}
+		else
+		{
+			// Add any messages that may have occurred during compilation to visual logs.
 #if WITH_EDITOR
-		LogCompilationMessages(Context);
+			LogCompilationMessages(Context);
 #endif
 
-		// If there was any error then we should abort.
-		using FNodeAndCompileMessages = const TPair<TObjectKey<const UPCGNode>, TArray<FComputeKernelCompileMessage>>;
-		for (FNodeAndCompileMessages& NodeAndCompileMessages : Context->ComputeGraph->KernelToCompileMessages)
-		{
-			for (const FComputeKernelCompileMessage& Message : NodeAndCompileMessages.Get<1>())
+			// If there was any error then we should abort.
+			using FNodeAndCompileMessages = const TPair<TObjectKey<const UPCGNode>, TArray<FComputeKernelCompileMessage>>;
+			for (FNodeAndCompileMessages& NodeAndCompileMessages : Context->ComputeGraph->KernelToCompileMessages)
 			{
-				// Some error messages were getting lost, and we were only getting the final 'failed' message. Treat this as failure and report for now.
-				// TODO: Revert the 'failed' part once we're happy all relevant issues are bubbling up.
-				if (Message.Type == FComputeKernelCompileMessage::EMessageType::Error || Message.Text.Contains(TEXT("failed"), ESearchCase::IgnoreCase))
+				for (const FComputeKernelCompileMessage& Message : NodeAndCompileMessages.Get<1>())
 				{
-					return true;
+					// Some error messages were getting lost, and we were only getting the final 'failed' message. Treat this as failure and report for now.
+					// TODO: Revert the 'failed' part once we're happy all relevant issues are bubbling up.
+					if (Message.Type == FComputeKernelCompileMessage::EMessageType::Error || Message.Text.Contains(TEXT("failed"), ESearchCase::IgnoreCase))
+					{
+						return true;
+					}
 				}
 			}
 		}
-	}
 
-	// 5. Enqueue work to be executed when the GPU processes the current frame.
+		Context->ExecutionSubPhase = EPCGComputeGraphExecutionPhase::ScheduleComputeGraph;
+
+		if (Context->AsyncState.ShouldStop())
+		{
+			return false;
+		}
+	}
+	case EPCGComputeGraphExecutionPhase::ScheduleComputeGraph: // Fallthrough
 	{
-		TRACE_CPUPROFILER_EVENT_SCOPE(ComputeGraphInstance.EnqueueWork);
+		TRACE_CPUPROFILER_EVENT_SCOPE(ScheduleComputeGraph);
 
-		Context->bGraphEnqueued = Context->ComputeGraphInstance.EnqueueWork(
+		const bool bGraphEnqueued = Context->ComputeGraphInstance.EnqueueWork(
 			Context->ComputeGraph.Get(),
 			InContext->SourceComponent->GetScene(),
 			ComputeTaskExecutionGroup::EndOfFrameUpdate,
@@ -394,11 +387,17 @@ bool FPCGComputeGraphElement::ExecuteInternal(FPCGContext* InContext) const
 				Context->bIsPaused = false;
 				Context->bGraphSubmitFailed = true;
 			}));
-	}
 
-	if (Context->bGraphEnqueued)
-	{
-		if (!Context->ProvidersWithBufferExports.IsEmpty())
+		if (!bGraphEnqueued)
+		{
+			PCGE_LOG(Error, GraphAndLog, LOCTEXT("EnqueueFailed", "Compute graph enqueue failed, check log for errors."));
+			ResetAsyncOperations(InContext);
+			return true;
+		}
+
+		Context->ExecutionSubPhase = EPCGComputeGraphExecutionPhase::WaitForExecutionComplete;
+
+		if (Context->HasPendingAsyncOperations())
 		{
 			// If we're running async operations (like preparing exports), go to sleep and let them wake us up later.
 			Context->bIsPaused = true;
@@ -411,12 +410,53 @@ bool FPCGComputeGraphElement::ExecuteInternal(FPCGContext* InContext) const
 			return true;
 		}
 	}
-	else
+	case EPCGComputeGraphExecutionPhase::WaitForExecutionComplete:
 	{
-		PCGE_LOG(Error, GraphAndLog, LOCTEXT("EnqueueFailed", "Compute graph enqueue failed, check log for errors."));
-		ResetAsyncOperations(InContext);
-		return true;
+		TRACE_CPUPROFILER_EVENT_SCOPE(WaitForExecutionComplete);
+
+		if (Context->bGraphSubmitFailed)
+		{
+			return true;
+		}
+
+		// Execution is complete when graph was enqueued/dispatched and all exports and readbacks are complete.
+		if (!Context->HasPendingAsyncOperations())
+		{
+			// Apply any read back element counts to the final data description.
+			Context->DataBinding->ApplyElementCounts_GameThread();
+
+#if WITH_EDITOR
+			// Call both to start all required readbacks.
+			const bool bDebugDrawDone = ExecuteDebugDraw(Context);
+			const bool bInspectionDone = StoreDataForInspection(Context);
+
+			if (!bDebugDrawDone || !bInspectionDone)
+			{
+				SleepUntilNextFrame();
+				return false;
+			}
+#endif
+
+			Context->bExecutionSuccess = true;
+
+			// Currently we don't output anything if processing any readback data processing failed.
+			if (ensure(Context->bExecutionSuccess) && ensure(Context->DataBinding))
+			{
+				Context->OutputData = Context->DataBinding->OutputDataCollection;
+			}
+
+			// Signal completion.
+			return true;
+		}
+
+		// Still running. Likely we need a frame to pass in order to make progress with readbacks etc.
+		SleepUntilNextFrame();
+
+		return false;
 	}
+	}
+
+	return false;
 }
 
 void FPCGComputeGraphElement::PostExecuteInternal(FPCGContext* InContext) const
@@ -456,7 +496,9 @@ void FPCGComputeGraphElement::ResetAsyncOperations(FPCGContext* InContext) const
 
 	if (InContext)
 	{
-		static_cast<FPCGComputeGraphContext*>(InContext)->ProvidersWithBufferExports.Reset();
+		FPCGComputeGraphContext* Context = static_cast<FPCGComputeGraphContext*>(InContext);
+		Context->ProvidersWithBufferExports.Reset();
+		Context->ProvidersWithReadbacks.Reset();
 	}
 }
 
diff --git a/Engine/Plugins/PCG/Source/PCG/Private/Compute/Elements/PCGComputeGraphElement.h b/Engine/Plugins/PCG/Source/PCG/Private/Compute/Elements/PCGComputeGraphElement.h
--- a/Engine/Plugins/PCG/Source/PCG/Private/Compute/Elements/PCGComputeGraphElement.h
+++ b/Engine/Plugins/PCG/Source/PCG/Private/Compute/Elements/PCGComputeGraphElement.h
@@ -13,12 +13,27 @@
 
 #include "PCGComputeGraphElement.generated.h"
 
+UENUM()
+enum class EPCGComputeGraphExecutionPhase : uint8
+{
+	None,
+	GetComputeGraph,
+	InitializeDataBinding,
+	PreExecuteReadbacks,
+	PrepareForExecute,
+	ValidateComputeGraphCompilation,
+	ScheduleComputeGraph,
+	WaitForExecutionComplete,
+};
+
 struct FPCGComputeGraphContext : public FPCGContext
 {
 public:
 	virtual bool IsComputeContext() const override { return true; }
 	virtual ~FPCGComputeGraphContext();
 
+	bool HasPendingAsyncOperations() const;
+
 protected:
 	virtual void AddExtraStructReferencedObjects(FReferenceCollector& Collector) override;
 
@@ -34,28 +49,20 @@ public:
 	/** Whether the compute graph instance has been initialized (data providers created, etc). */
 	bool bComputeGraphInstanceInitialized = false;
 
-	/** Graph enqueued (scheduled for execution by GPU). */
-	bool bGraphEnqueued = false;
-
 	/** Graph was enqueued but was invalid for some reason and the work could not be submitted. */
 	bool bGraphSubmitFailed = false;
 
 	/** Graph executed successfully. */
 	bool bExecutionSuccess = false;
 
-	/** Whether all required primitive proxies have been created. */
-	bool bPrimitiveProxiesValidated = false;
-
-	bool bDataProvidersReadbacksComplete = false;
-
-	bool bDataProvidersReady = false;
-
 	/** Data providers with buffers that are passed to downstream tasks. The buffer will be created on the render thread, and then passed
 	* back to main thread, upon which a reference is taken to the buffer and the provider is removed from this set to signal completion.
 	*/
 	TSet<TObjectPtr<UComputeDataProvider>> ProvidersWithBufferExports;
 
 	TSet<TObjectPtr<UComputeDataProvider>> ProvidersWithReadbacks;
+
+	EPCGComputeGraphExecutionPhase ExecutionSubPhase = EPCGComputeGraphExecutionPhase::GetComputeGraph;
 };
 
 /** Executes a CF graph. Created by the compiler when collapsing GPU nodes rather than by a settings/node. */
