[PCG] Fix check in Attribute Filter

#rnx

diff --git a/Engine/Plugins/PCG/Source/PCG/Private/Elements/PCGAttributeFilter.cpp b/Engine/Plugins/PCG/Source/PCG/Private/Elements/PCGAttributeFilter.cpp
--- a/Engine/Plugins/PCG/Source/PCG/Private/Elements/PCGAttributeFilter.cpp
+++ b/Engine/Plugins/PCG/Source/PCG/Private/Elements/PCGAttributeFilter.cpp
@@ -622,6 +622,7 @@ bool FPCGAttributeFilterElementBase::DoFiltering(FPCGContext* Context, EPCGAttri
 			// Disable Spatial Queries
 			FirstThresholdInfo.bUseSpatialQuery = false;
 			SecondThresholdInfo.bUseSpatialQuery = false;
+			OperationData.bIsInputPointData = false;
 		}
 		else
 		{
