[PCG] Time slice primitive setup part of GPU SM Spawner

#rnx

diff --git a/Engine/Plugins/PCG/Source/PCG/Private/Compute/DataInterfaces/Elements/PCGStaticMeshSpawnerDataInterface.cpp b/Engine/Plugins/PCG/Source/PCG/Private/Compute/DataInterfaces/Elements/PCGStaticMeshSpawnerDataInterface.cpp
--- a/Engine/Plugins/PCG/Source/PCG/Private/Compute/DataInterfaces/Elements/PCGStaticMeshSpawnerDataInterface.cpp
+++ b/Engine/Plugins/PCG/Source/PCG/Private/Compute/DataInterfaces/Elements/PCGStaticMeshSpawnerDataInterface.cpp
@@ -5,7 +5,6 @@
 #include "PCGComponent.h"
 #include "PCGContext.h"
 #include "PCGParamData.h"
-#include "Components/PCGProceduralISMComponent.h"
 #include "Compute/PCGComputeCommon.h"
 #include "Compute/PCGComputeGraph.h"
 #include "Compute/PCGDataBinding.h"
@@ -295,7 +294,15 @@ bool UPCGStaticMeshSpawnerDataProvider::PrepareForExecute_GameThread(UPCGDataBin
 		return false;
 	}
 
-	if (const UPCGMeshSelectorByAttribute* SelectorByAttribute = Cast<UPCGMeshSelectorByAttribute>(Settings->MeshSelectorParameters))
+	TSharedPtr<FPCGContextHandle> ContextHandle = InBinding->ContextHandle.Pin();
+	FPCGContext* Context = ContextHandle ? ContextHandle->GetContext() : nullptr;
+	if (!ensure(Context))
+	{
+		return true;
+	}
+
+	const UPCGMeshSelectorByAttribute* SelectorByAttribute = Cast<UPCGMeshSelectorByAttribute>(Settings->MeshSelectorParameters);
+	if (SelectorAttributeId == INDEX_NONE && SelectorByAttribute)
 	{
 		const FName SelectorName = SelectorByAttribute->AttributeName;
 
@@ -328,9 +335,6 @@ bool UPCGStaticMeshSpawnerDataProvider::PrepareForExecute_GameThread(UPCGDataBin
 		{
 			if (!InputDataDesc.DataDescs.IsEmpty())
 			{
-				TSharedPtr<FPCGContextHandle> ContextHandle = InBinding->ContextHandle.Pin();
-				FPCGContext* Context = ContextHandle ? ContextHandle->GetContext() : nullptr;
-
 				PCG_KERNEL_VALIDATION_ERR(Context, Settings, /*bQuiet=*/false, FText::Format(
 					LOCTEXT("MeshSelectorAttributeNotFound", "Mesh selector attribute '{0}' not found."),
 					FText::FromName(SelectorName)));
@@ -338,27 +342,39 @@ bool UPCGStaticMeshSpawnerDataProvider::PrepareForExecute_GameThread(UPCGDataBin
 
 			return true;
 		}
+
+		if (Context->AsyncState.ShouldStop())
+		{
+			return false;
+		}
 	}
 
-	if (!bPrimitivesSetUp)
+	if (!bPrimitiveDescriptorsCreated)
 	{
-		bPrimitivesSetUp = true;
+		CreatePrimitiveDescriptors(Context, InBinding);
 
-		TSharedPtr<FPCGContextHandle> ContextHandle = InBinding->ContextHandle.Pin();
-		FPCGContext* Context = ContextHandle ? ContextHandle->GetContext() : nullptr;
+		bPrimitiveDescriptorsCreated = true;
 
-		if (!ensure(Context))
+		if (Context->AsyncState.ShouldStop())
 		{
-			return true;
+			return false;
+		}
+	}
+
+	if (!bPrimitivesSetUp)
+	{
+		if (!SetupPrimitives(Context, InBinding))
+		{
+			return false;
 		}
 
-		const bool bSetupComponent = SetupPrimitives(Context, InBinding);
+		bPrimitivesSetUp = true;
 
-		if (bSetupComponent)
+		if (NumPrimitivesSetup > 0)
 		{
 			Context->SourceComponent->NotifyProceduralInstancesInUse();
 
-			// Wait a frame - gives time for GPU Scene to be updated.
+			// Signal not finished. Compute graph element will wait a frame, giving time for GPU Scene to be updated.
 			return false;
 		}
 		else
@@ -414,14 +430,16 @@ void UPCGStaticMeshSpawnerDataProvider::Reset()
 	Super::Reset();
 }
 
-bool UPCGStaticMeshSpawnerDataProvider::SetupPrimitives(FPCGContext* InContext, UPCGDataBinding* InBinding)
+void UPCGStaticMeshSpawnerDataProvider::CreatePrimitiveDescriptors(FPCGContext* InContext, UPCGDataBinding* InBinding)
 {
+	PrimitiveDescriptors.Empty();
+
 	const UPCGStaticMeshSpawnerSettings* Settings = CastChecked<UPCGStaticMeshSpawnerSettings>(ProducerKernel->GetSettings());
 
 	AActor* TargetActor = Settings->TargetActor.Get() ? Settings->TargetActor.Get() : InContext->GetTargetActor(nullptr);
 	if (!ensure(Settings->MeshSelectorParameters) || !ensure(TargetActor))
 	{
-		return false;
+		return;
 	}
 
 	FPCGDataCollectionDesc InputDataDesc;
@@ -430,13 +448,11 @@ bool UPCGStaticMeshSpawnerDataProvider::SetupPrimitives(FPCGContext* InContext,
 	const uint32 TotalInputPointCount = InputDataDesc.ComputeDataElementCount(EPCGDataType::Point);
 	if (TotalInputPointCount == 0)
 	{
-		return false;
+		return;
 	}
 
 	const FBox LocalBounds = InContext->SourceComponent->GetGridBounds().ShiftBy(-InContext->SourceComponent->GetOwner()->GetActorLocation());
 
-	uint32 CustomFloatCount = 0;
-
 	if (Settings->InstanceDataPackerParameters && Settings->InstanceDataPackerParameters->GetAttributeNames(/*OutNames=*/nullptr))
 	{
 		TArray<FName> AttributeNames;
@@ -464,8 +480,6 @@ bool UPCGStaticMeshSpawnerDataProvider::SetupPrimitives(FPCGContext* InContext,
 		PCGDataForGPUHelpers::ComputeCustomFloatPacking(InContext, Settings, AttributeNames, InBinding, InputDataDesc, CustomFloatCount, AttributeIdOffsetStrides);
 	}
 
-	TArray<FPCGProceduralISMComponentDescriptor> ComponentsToCreate;
-
 	if (const UPCGMeshSelectorByAttribute* SelectorByAttribute = Cast<UPCGMeshSelectorByAttribute>(Settings->MeshSelectorParameters))
 	{
 		const FName SelectorName = SelectorByAttribute->AttributeName;
@@ -509,6 +523,8 @@ bool UPCGStaticMeshSpawnerDataProvider::SetupPrimitives(FPCGContext* InContext,
 			}
 		}
 
+		PrimitiveDescriptors.Reserve(StringKeyToInstanceCount.Num());
+
 		for (const TPair<int32, uint32>& StringKeyAndInstanceCount : StringKeyToInstanceCount)
 		{
 			if (!ensure(InBinding->GetStringTable().IsValidIndex(StringKeyAndInstanceCount.Key)))
@@ -522,7 +538,7 @@ bool UPCGStaticMeshSpawnerDataProvider::SetupPrimitives(FPCGContext* InContext,
 				continue;
 			}
 
-			if (ComponentsToCreate.Num() >= PCGComputeConstants::MAX_PRIMITIVE_COMPONENTS_PER_SPAWNER)
+			if (PrimitiveDescriptors.Num() >= PCGComputeConstants::MAX_PRIMITIVE_COMPONENTS_PER_SPAWNER)
 			{
 				PCG_KERNEL_VALIDATION_WARN(InContext, Settings, /*bQuiet=*/false,
 					FText::Format(PCGStaticMeshSpawnerDataInterface::TooManyPrimitivesFormat, FText::FromString(FString::FromInt(PCGComputeConstants::MAX_PRIMITIVE_COMPONENTS_PER_SPAWNER))));
@@ -547,7 +563,7 @@ bool UPCGStaticMeshSpawnerDataProvider::SetupPrimitives(FPCGContext* InContext,
 			Descriptor.StaticMesh = StaticMesh;
 
 			PrimitiveStringKeys.Add(StringKeyAndInstanceCount.Key);
-			ComponentsToCreate.Add(MoveTemp(Descriptor));
+			PrimitiveDescriptors.Add(MoveTemp(Descriptor));
 
 			if (Settings->bApplyMeshBoundsToPoints)
 			{
@@ -555,15 +571,14 @@ bool UPCGStaticMeshSpawnerDataProvider::SetupPrimitives(FPCGContext* InContext,
 			}
 		}
 
-		PrimitiveSelectionCDF.SetNumZeroed(ComponentsToCreate.Num());
+		PrimitiveSelectionCDF.SetNumZeroed(PrimitiveDescriptors.Num());
 	}
 	else if (const UPCGMeshSelectorWeighted* SelectorWeighted = Cast<UPCGMeshSelectorWeighted>(Settings->MeshSelectorParameters))
 	{
 		if (SelectorWeighted->MeshEntries.IsEmpty())
 		{
 			PCG_KERNEL_VALIDATION_ERR(InContext, Settings, /*bQuiet=*/false, PCGStaticMeshSpawnerDataInterface::NoMeshEntriesFormat);
-
-			return false;
+			return;
 		}
 
 		float CumulativeWeight = 0.0f;
@@ -576,11 +591,13 @@ bool UPCGStaticMeshSpawnerDataProvider::SetupPrimitives(FPCGContext* InContext,
 
 		if (!ensure(TotalWeight > UE_SMALL_NUMBER))
 		{
-			return false;
+			return;
 		}
 
 		PrimitiveSelectionCDF.Reserve(SelectorWeighted->MeshEntries.Num());
 
+		PrimitiveDescriptors.Reserve(SelectorWeighted->MeshEntries.Num());
+
 		for (const FPCGMeshSelectorWeightedEntry& Entry : SelectorWeighted->MeshEntries)
 		{
 			if (Entry.Descriptor.StaticMesh.IsNull())
@@ -588,7 +605,7 @@ bool UPCGStaticMeshSpawnerDataProvider::SetupPrimitives(FPCGContext* InContext,
 				continue;
 			}
 
-			if (ComponentsToCreate.Num() >= PCGComputeConstants::MAX_PRIMITIVE_COMPONENTS_PER_SPAWNER)
+			if (PrimitiveDescriptors.Num() >= PCGComputeConstants::MAX_PRIMITIVE_COMPONENTS_PER_SPAWNER)
 			{
 				PCG_KERNEL_VALIDATION_WARN(InContext, Settings, /*bQuiet=*/false,
 					FText::Format(PCGStaticMeshSpawnerDataInterface::TooManyPrimitivesFormat, FText::FromString(FString::FromInt(PCGComputeConstants::MAX_PRIMITIVE_COMPONENTS_PER_SPAWNER))));
@@ -628,7 +645,7 @@ bool UPCGStaticMeshSpawnerDataProvider::SetupPrimitives(FPCGContext* InContext,
 
 				Descriptor.NumInstances = FMath::Min(InstanceCount, TotalInputPointCount);
 
-				ComponentsToCreate.Add(MoveTemp(Descriptor));
+				PrimitiveDescriptors.Add(MoveTemp(Descriptor));
 
 				if (Settings->bApplyMeshBoundsToPoints)
 				{
@@ -644,33 +661,51 @@ bool UPCGStaticMeshSpawnerDataProvider::SetupPrimitives(FPCGContext* InContext,
 	}
 
 	// Validate instance count is not too large.
+	uint64 TotalInstanceCount = 0;
+	for (const FPCGProceduralISMComponentDescriptor& Desc : PrimitiveDescriptors)
 	{
-		uint64 TotalInstanceCount = 0;
-		for (const FPCGProceduralISMComponentDescriptor& Desc : ComponentsToCreate)
-		{
-			TotalInstanceCount += Desc.NumInstances;
-		}
+		TotalInstanceCount += Desc.NumInstances;
+	}
 
-		if (TotalInstanceCount >= MAX_INSTANCE_ID)
-		{
-			PCG_KERNEL_VALIDATION_ERR(InContext, Settings, /*bQuiet=*/false, FText::Format(
-				LOCTEXT("TooManyInstances", "Tried to spawn too many instances ({0}), procedural ISM component creation skipped and instances will not be rendered."),
-				FText::FromString(FString::FromInt(TotalInstanceCount))));
+	if (TotalInstanceCount >= MAX_INSTANCE_ID)
+	{
+		PCG_KERNEL_VALIDATION_ERR(InContext, Settings, /*bQuiet=*/false, FText::Format(
+			LOCTEXT("TooManyInstances", "Tried to spawn too many instances ({0}), procedural ISM component creation skipped and instances will not be rendered."),
+			FText::FromString(FString::FromInt(TotalInstanceCount))));
 
-			return false;
-		}
+		return;
 	}
 
 	FPCGSpawnerPrimitives& Primitives = InBinding->MeshSpawnersToPrimitives.FindOrAdd(ProducerKernel);
 	Primitives.NumCustomFloats = CustomFloatCount;
-	Primitives.AttributeIdOffsetStrides = /*MoveTemp*/(AttributeIdOffsetStrides);
+	Primitives.AttributeIdOffsetStrides = AttributeIdOffsetStrides;
 	Primitives.SelectorAttributeId = SelectorAttributeId;
-	Primitives.SelectionCDF = /*MoveTemp*/(PrimitiveSelectionCDF);
-	Primitives.PrimitiveStringKeys = /*MoveTemp*/(PrimitiveStringKeys);
-	Primitives.PrimitiveMeshBounds = /*MoveTemp*/(PrimitiveMeshBounds);
+	Primitives.SelectionCDF = PrimitiveSelectionCDF;
+	Primitives.PrimitiveStringKeys = PrimitiveStringKeys;
+	Primitives.PrimitiveMeshBounds = PrimitiveMeshBounds;
+	Primitives.NumInstancesAllPrimitives = static_cast<uint32>(TotalInstanceCount);
+
+	// Useful for debugging instance counts.
+	//UE_LOG(LogPCG, Log, TEXT("Input points %d, total instance count %d (+ %.2f%%)"), TotalInputPointCount, Primitives.NumInstancesAllPrimitives, (float(Primitives.NumInstancesAllPrimitives - TotalInputPointCount) / TotalInputPointCount) * 100.0f);
+}
+
+bool UPCGStaticMeshSpawnerDataProvider::SetupPrimitives(FPCGContext* InContext, UPCGDataBinding* InBinding)
+{
+	const UPCGStaticMeshSpawnerSettings* Settings = CastChecked<UPCGStaticMeshSpawnerSettings>(ProducerKernel->GetSettings());
 
-	for (const FPCGProceduralISMComponentDescriptor& Desc : ComponentsToCreate)
+	AActor* TargetActor = Settings->TargetActor.Get() ? Settings->TargetActor.Get() : InContext->GetTargetActor(nullptr);
+	if (!ensure(TargetActor))
+	{
+		return true;
+	}
+
+	FPCGSpawnerPrimitives* Primitives = InBinding->MeshSpawnersToPrimitives.Find(ProducerKernel);
+	check(Primitives);
+
+	while (NumPrimitivesSetup < PrimitiveDescriptors.Num())
 	{
+		const FPCGProceduralISMComponentDescriptor& Desc = PrimitiveDescriptors[NumPrimitivesSetup];
+
 		FPCGProceduralISMCBuilderParameters Params;
 		Params.Descriptor = Desc;
 		Params.bAllowDescriptorChanges = false;
@@ -683,15 +718,17 @@ bool UPCGStaticMeshSpawnerDataProvider::SetupPrimitives(FPCGContext* InContext,
 		// Don't bother registering the resource change as we're transient anyway.
 		//InContext->TouchedResources.Emplace(MISMC);
 
-		Primitives.Primitives.Add(MISMC->GetComponent());
+		Primitives->Primitives.Add(MISMC->GetComponent());
 
-		Primitives.NumInstancesAllPrimitives += Desc.NumInstances;
-	}
+		++NumPrimitivesSetup;
 
-	// Useful for debugging instance counts.
-	//UE_LOG(LogPCG, Log, TEXT("Input points %d, total instance count %d (+ %.2f%%)"), TotalInputPointCount, Primitives.NumInstancesAllPrimitives, (float(Primitives.NumInstancesAllPrimitives - TotalInputPointCount) / TotalInputPointCount) * 100.0f);
+		if (NumPrimitivesSetup < PrimitiveDescriptors.Num() && InContext->AsyncState.ShouldStop())
+		{
+			return false;
+		}
+	}
 
-	return !ComponentsToCreate.IsEmpty();
+	return true;
 }
 
 bool FPCGStaticMeshSpawnerDataProviderProxy::IsValid(FValidationData const& InValidationData) const
diff --git a/Engine/Plugins/PCG/Source/PCG/Private/Compute/DataInterfaces/Elements/PCGStaticMeshSpawnerDataInterface.h b/Engine/Plugins/PCG/Source/PCG/Private/Compute/DataInterfaces/Elements/PCGStaticMeshSpawnerDataInterface.h
--- a/Engine/Plugins/PCG/Source/PCG/Private/Compute/DataInterfaces/Elements/PCGStaticMeshSpawnerDataInterface.h
+++ b/Engine/Plugins/PCG/Source/PCG/Private/Compute/DataInterfaces/Elements/PCGStaticMeshSpawnerDataInterface.h
@@ -5,6 +5,8 @@
 #include "Compute/DataInterfaces/PCGComputeDataInterface.h"
 #include "ComputeFramework/ComputeDataProvider.h"
 
+#include "Components/PCGProceduralISMComponent.h"
+
 #include "PCGStaticMeshSpawnerDataInterface.generated.h"
 
 class FPCGStaticMeshSpawnerDataInterfaceParameters;
@@ -45,6 +47,8 @@ public:
 	virtual void Reset() override;
 	//~ End UComputeDataProvider Interface
 
+	void CreatePrimitiveDescriptors(FPCGContext* InContext, UPCGDataBinding* InBinding);
+
 	/** Setup Procedural ISMCs. Returns true if at least one component was set up. */
 	virtual bool SetupPrimitives(FPCGContext* InContext, UPCGDataBinding* InBinding);
 
@@ -80,7 +84,16 @@ public:
 
 	int32 AnalysisDataIndex = INDEX_NONE;
 
+	bool bPrimitiveDescriptorsCreated = false;
+
+	TArray<FPCGProceduralISMComponentDescriptor> PrimitiveDescriptors;
+
+	uint32 CustomFloatCount = 0;
+
 	bool bPrimitivesSetUp = false;
+
+	/** Number of primitives created during primitive setup. */
+	int32 NumPrimitivesSetup = 0;
 };
 
 class FPCGStaticMeshSpawnerDataProviderProxy : public FComputeDataProviderRenderProxy
