[PCG] Fixing missing reset of newly added state in static mesh spawner data provider.

#rnx

diff --git a/Engine/Plugins/PCG/Source/PCG/Private/Compute/DataInterfaces/Elements/PCGStaticMeshSpawnerDataInterface.cpp b/Engine/Plugins/PCG/Source/PCG/Private/Compute/DataInterfaces/Elements/PCGStaticMeshSpawnerDataInterface.cpp
--- a/Engine/Plugins/PCG/Source/PCG/Private/Compute/DataInterfaces/Elements/PCGStaticMeshSpawnerDataInterface.cpp
+++ b/Engine/Plugins/PCG/Source/PCG/Private/Compute/DataInterfaces/Elements/PCGStaticMeshSpawnerDataInterface.cpp
@@ -425,7 +425,11 @@ void UPCGStaticMeshSpawnerDataProvider::Reset()
 	SelectedMeshAttributeId = INDEX_NONE;
 	StringKeyToInstanceCount.Empty();
 	AnalysisDataIndex = INDEX_NONE;
+	bPrimitiveDescriptorsCreated = false;
+	PrimitiveDescriptors.Empty();
+	CustomFloatCount = 0;
 	bPrimitivesSetUp = false;
+	NumPrimitivesSetup = 0;
 
 	Super::Reset();
 }
