Add trace events for PCG Landscape initialization

diff --git a/Engine/Plugins/PCG/Source/PCG/Private/Compute/DataInterfaces/PCGLandscapeDataInterface.cpp b/Engine/Plugins/PCG/Source/PCG/Private/Compute/DataInterfaces/PCGLandscapeDataInterface.cpp
--- a/Engine/Plugins/PCG/Source/PCG/Private/Compute/DataInterfaces/PCGLandscapeDataInterface.cpp
+++ b/Engine/Plugins/PCG/Source/PCG/Private/Compute/DataInterfaces/PCGLandscapeDataInterface.cpp
@@ -358,6 +358,7 @@ void UPCGLandscapeDataProvider::InitFromLandscape(ALandscape* InLandscape, const
 
 void UPCGLandscapeDataProvider::InitializeRuntimeVirtualTextures(ALandscape* InLandscape)
 {
+	TRACE_CPUPROFILER_EVENT_SCOPE(UPCGLandscapeDataProvider::InitializeRuntimeVirtualTextures);
 	check(InLandscape);
 
 	int32 BaseColorVirtualTextureIndex = INDEX_NONE;
@@ -432,6 +433,7 @@ void UPCGLandscapeDataProvider::InitializeRuntimeVirtualTextures(ALandscape* InL
 
 void UPCGLandscapeDataProvider::InitializeLandscapeResource(ALandscape* InLandscape, const FBox& Bounds)
 {
+	TRACE_CPUPROFILER_EVENT_SCOPE(UPCGLandscapeDataProvider::InitializeLandscapeResource);
 	check(InLandscape);
 
 	const ULandscapeInfo* LandscapeInfo = InLandscape->GetLandscapeInfo();
