[PCG] Remove Metadata projection in the Point Neighborhood as the result is discarded

#rnx

diff --git a/Engine/Plugins/PCG/Source/PCG/Private/Elements/PCGPointNeighborhood.cpp b/Engine/Plugins/PCG/Source/PCG/Private/Elements/PCGPointNeighborhood.cpp
--- a/Engine/Plugins/PCG/Source/PCG/Private/Elements/PCGPointNeighborhood.cpp
+++ b/Engine/Plugins/PCG/Source/PCG/Private/Elements/PCGPointNeighborhood.cpp
@@ -112,10 +112,11 @@ bool FPCGPointNeighborhoodElement::ExecuteInternal(FPCGContext* Context) const
 			OutPoint = InPoint;
 
 			const FTransform InTransform(InPoint.Transform.GetLocation());
-			UPCGMetadata* OutMetadata = OutputPointData->Metadata;
 
 			FPCGPoint ProjectionPoint = FPCGPoint();
-			InputPointData->ProjectPoint(InTransform, SearchBounds, Params, ProjectionPoint, OutMetadata, Settings->bWeightedAverage);
+			// Do not pass metadata as output as we don't care about attributes (Projected point will be discarded).
+			// @todo_pcg: Maybe support the weighting average of attributes as an option.
+			InputPointData->ProjectPoint(InTransform, SearchBounds, Params, ProjectionPoint, /*OutMetadata=*/nullptr, Settings->bWeightedAverage);
 
 			const double NormalizedDistance = FVector::Distance(InPoint.Transform.GetLocation(), ProjectionPoint.Transform.GetLocation()) / SearchDistance;
 
