[PCG] Fix component seed being incorporated into CRC even if the seed is unused.

In one case this was causing a Get Texture Data node to do a lot of unnecessary texture initializations.

#rnx

diff --git a/Engine/Plugins/PCG/Source/PCG/Private/PCGElement.cpp b/Engine/Plugins/PCG/Source/PCG/Private/PCGElement.cpp
--- a/Engine/Plugins/PCG/Source/PCG/Private/PCGElement.cpp
+++ b/Engine/Plugins/PCG/Source/PCG/Private/PCGElement.cpp
@@ -1061,7 +1061,7 @@ void IPCGElement::GetDependenciesCrc(const FPCGDataCollection& InInput, const UP
 		}
 	}
 
-	if (InComponent)
+	if (InComponent && (!InSettings || InSettings->UseSeed()))
 	{
 		Crc.Combine(InComponent->Seed);
 	}
