[PCG] Optimize kernel validation
* Split validation into static/compile-time and data/execution-time
* Compile time validation errors/warnings are cached and logged during execution time validation
* Refactor so that custom HLSL source code is only parsed once, instead of multiple times

#rnx

diff --git a/Engine/Plugins/PCG/Source/PCG/Internal/Compute/PCGComputeCommon.h b/Engine/Plugins/PCG/Source/PCG/Internal/Compute/PCGComputeCommon.h
--- a/Engine/Plugins/PCG/Source/PCG/Internal/Compute/PCGComputeCommon.h
+++ b/Engine/Plugins/PCG/Source/PCG/Internal/Compute/PCGComputeCommon.h
@@ -21,11 +21,11 @@ class FRHIShaderResourceView;
 #define PCG_KERNEL_LOGGING_ENABLED (!(UE_BUILD_SHIPPING || UE_BUILD_TEST) || USE_LOGGING_IN_SHIPPING)
 
 #if PCG_KERNEL_LOGGING_ENABLED
-#define PCG_KERNEL_VALIDATION_WARN(Context, Settings, bQuiet, ValidationMessage) if (!bQuiet) PCGComputeHelpers::LogKernelWarning(Context, Settings, ValidationMessage);
-#define PCG_KERNEL_VALIDATION_ERR(Context, Settings, bQuiet, ValidationMessage) if (!bQuiet) PCGComputeHelpers::LogKernelError(Context, Settings, ValidationMessage);
+#define PCG_KERNEL_VALIDATION_WARN(Context, Settings, ValidationMessage) PCGComputeHelpers::LogKernelWarning(Context, Settings, ValidationMessage);
+#define PCG_KERNEL_VALIDATION_ERR(Context, Settings, ValidationMessage) PCGComputeHelpers::LogKernelError(Context, Settings, ValidationMessage);
 #else
-#define PCG_KERNEL_VALIDATION_WARN(Context, Settings, bQuiet, ValidationMessage) // Log removed
-#define PCG_KERNEL_VALIDATION_ERR(Context, Settings, bQuiet, ValidationMessage) // Log removed
+#define PCG_KERNEL_VALIDATION_WARN(Context, Settings, ValidationMessage) // Log removed
+#define PCG_KERNEL_VALIDATION_ERR(Context, Settings, ValidationMessage) // Log removed
 #endif
 
 /** Modes for exporting the buffer from transient to persistent for downstream consumption. */
diff --git a/Engine/Plugins/PCG/Source/PCG/Private/Compute/BuiltInKernels/PCGCountUniqueAttributeValuesKernel.cpp b/Engine/Plugins/PCG/Source/PCG/Private/Compute/BuiltInKernels/PCGCountUniqueAttributeValuesKernel.cpp
--- a/Engine/Plugins/PCG/Source/PCG/Private/Compute/BuiltInKernels/PCGCountUniqueAttributeValuesKernel.cpp
+++ b/Engine/Plugins/PCG/Source/PCG/Private/Compute/BuiltInKernels/PCGCountUniqueAttributeValuesKernel.cpp
@@ -5,12 +5,14 @@
 #include "PCGComponent.h"
 #include "PCGContext.h"
 #include "PCGModule.h"
+#include "Compute/Elements/PCGComputeGraphElement.h"
 #include "Compute/PCGComputeGraph.h"
 #include "Compute/PCGDataBinding.h"
 #include "Compute/PCGPinPropertiesGPU.h"
 #include "Compute/DataInterfaces/BuiltInKernels/PCGCountUniqueAttributeValuesDataInterface.h"
-#include "Compute/Elements/PCGComputeGraphElement.h"
+#include "Elements/PCGStaticMeshSpawner.h"
 #include "Graph/PCGGPUGraphCompilationContext.h"
+#include "MeshSelectors/PCGMeshSelectorByAttribute.h"
 
 #include "ShaderCompilerCore.h"
 
@@ -18,18 +20,22 @@
 
 #define LOCTEXT_NAMESPACE "PCGCountUniqueAttributeValuesKernel"
 
-void UPCGCountUniqueAttributeValuesKernel::Initialize(FName InAttributeName, bool bLogDescriptions)
+#if WITH_EDITOR
+void UPCGCountUniqueAttributeValuesKernel::Initialize(const UPCGSettings* InSettings, bool bLogDescriptions)
 {
-	Super::Initialize(bLogDescriptions);
+	Super::Initialize(InSettings, bLogDescriptions);
 
-	AttributeName = InAttributeName;
+	const UPCGStaticMeshSpawnerSettings* SMSettings = CastChecked<UPCGStaticMeshSpawnerSettings>(InSettings);
+	const UPCGMeshSelectorByAttribute* Selector = CastChecked<UPCGMeshSelectorByAttribute>(SMSettings->MeshSelectorParameters);
+	AttributeName = Selector->AttributeName;
 }
+#endif
 
-bool UPCGCountUniqueAttributeValuesKernel::IsKernelValid(FPCGContext* InContext, bool bQuiet) const
+bool UPCGCountUniqueAttributeValuesKernel::IsKernelValid(FPCGContext* InContext) const
 {
 	TRACE_CPUPROFILER_EVENT_SCOPE(UPCGCountUniqueAttributeValuesKernel::IsKernelValid);
 
-	if (!Super::IsKernelValid(InContext, bQuiet))
+	if (!Super::IsKernelValid(InContext))
 	{
 		return false;
 	}
@@ -53,7 +59,7 @@ bool UPCGCountUniqueAttributeValuesKernel::IsKernelValid(FPCGContext* InContext,
 		{
 			if (!InputDataDesc.DataDescs.IsEmpty())
 			{
-				PCG_KERNEL_VALIDATION_ERR(InContext, /*Settings=*/nullptr, bQuiet, FText::Format(
+				PCG_KERNEL_VALIDATION_ERR(InContext, Settings, FText::Format(
 					LOCTEXT("AttributeMissing", "Count attribute '{0}' not found, this attribute must be present on all input data, and be of type String Key."),
 					FText::FromName(AttributeName)));
 			}
@@ -63,7 +69,7 @@ bool UPCGCountUniqueAttributeValuesKernel::IsKernelValid(FPCGContext* InContext,
 
 		if (bConflictingTypesFound)
 		{
-			PCG_KERNEL_VALIDATION_ERR(InContext, /*Settings=*/nullptr, bQuiet, FText::Format(
+			PCG_KERNEL_VALIDATION_ERR(InContext, Settings, FText::Format(
 				LOCTEXT("AttributeHasMultipleTypes", "Count attribute '{0}' found with multiple types in input data, all attributes must be of type String Key."),
 				FText::FromName(AttributeName)));
 
@@ -73,7 +79,7 @@ bool UPCGCountUniqueAttributeValuesKernel::IsKernelValid(FPCGContext* InContext,
 		if (AttributeDesc.Type != EPCGKernelAttributeType::StringKey)
 		{
 			// Attribute value counting only currently supported for attributes of type StringKey.
-			PCG_KERNEL_VALIDATION_ERR(Context, /*Settings=*/nullptr, bQuiet, FText::Format(
+			PCG_KERNEL_VALIDATION_ERR(Context, Settings, FText::Format(
 				LOCTEXT("AttributeTypeInvalid", "Cannot count values for attribute '{0}', only attributes of type String Key are currently supported."),
 				FText::FromName(AttributeName)));
 
diff --git a/Engine/Plugins/PCG/Source/PCG/Private/Compute/BuiltInKernels/PCGCountUniqueAttributeValuesKernel.h b/Engine/Plugins/PCG/Source/PCG/Private/Compute/BuiltInKernels/PCGCountUniqueAttributeValuesKernel.h
--- a/Engine/Plugins/PCG/Source/PCG/Private/Compute/BuiltInKernels/PCGCountUniqueAttributeValuesKernel.h
+++ b/Engine/Plugins/PCG/Source/PCG/Private/Compute/BuiltInKernels/PCGCountUniqueAttributeValuesKernel.h
@@ -18,9 +18,11 @@ class UPCGCountUniqueAttributeValuesKernel : public UPCGComputeKernel
 	GENERATED_BODY()
 
 public:
-	virtual void Initialize(FName InAttributeName, bool bLogDescriptions);
+#if WITH_EDITOR
+	virtual void Initialize(const UPCGSettings* InSettings, bool bLogDescriptions) override;
+#endif
 
-	virtual bool IsKernelValid(FPCGContext* InContext = nullptr, bool bQuiet = true) const override;
+	virtual bool IsKernelValid(FPCGContext* InContext) const override;
 	virtual bool ComputeOutputBindingDataDesc(const UPCGComputeGraph* InGraph, FName InOutputPinLabel, const UPCGDataBinding* InBinding, FPCGDataCollectionDesc& OutDataDesc) const override;
 	virtual int ComputeThreadCount(const UPCGDataBinding* Binding) const override;
 	virtual bool DoesOutputPinRequireZeroInitialization(FName InOutputPinLabel) const override;
diff --git a/Engine/Plugins/PCG/Source/PCG/Private/Compute/DataInterfaces/Elements/PCGMetadataPartitionDataInterface.cpp b/Engine/Plugins/PCG/Source/PCG/Private/Compute/DataInterfaces/Elements/PCGMetadataPartitionDataInterface.cpp
--- a/Engine/Plugins/PCG/Source/PCG/Private/Compute/DataInterfaces/Elements/PCGMetadataPartitionDataInterface.cpp
+++ b/Engine/Plugins/PCG/Source/PCG/Private/Compute/DataInterfaces/Elements/PCGMetadataPartitionDataInterface.cpp
@@ -105,7 +105,8 @@ bool UPCGMetaDataPartitionDataProvider::PrepareForExecute_GameThread(UPCGDataBin
 		return false;
 	}
 
-	const UPCGMetadataPartitionSettings* Settings = Cast<UPCGMetadataPartitionKernel>(Kernel)->Settings;
+	check(Kernel);
+	const UPCGMetadataPartitionSettings* Settings = Cast<UPCGMetadataPartitionSettings>(Kernel->GetSettings());
 	check(Settings);
 
 	if (ensure(Settings->PartitionAttributeSelectors.Num() == 1 && Settings->PartitionAttributeSelectors[0].IsBasicAttribute()))
diff --git a/Engine/Plugins/PCG/Source/PCG/Private/Compute/DataInterfaces/Elements/PCGStaticMeshSpawnerDataInterface.cpp b/Engine/Plugins/PCG/Source/PCG/Private/Compute/DataInterfaces/Elements/PCGStaticMeshSpawnerDataInterface.cpp
--- a/Engine/Plugins/PCG/Source/PCG/Private/Compute/DataInterfaces/Elements/PCGStaticMeshSpawnerDataInterface.cpp
+++ b/Engine/Plugins/PCG/Source/PCG/Private/Compute/DataInterfaces/Elements/PCGStaticMeshSpawnerDataInterface.cpp
@@ -204,10 +204,7 @@ void UPCGStaticMeshSpawnerDataProvider::Initialize(const UComputeDataInterface*
 	Super::Initialize(InDataInterface, InBinding, InInputMask, InOutputMask);
 
 	const UPCGStaticMeshSpawnerDataInterface* DataInterface = CastChecked<UPCGStaticMeshSpawnerDataInterface>(InDataInterface);
-	UPCGDataBinding* Binding = CastChecked<UPCGDataBinding>(InBinding);
-
 	ProducerKernel = CastChecked<UPCGStaticMeshSpawnerKernel>(DataInterface->ProducerKernel);
-	check(ProducerKernel->Settings);
 }
 
 bool UPCGStaticMeshSpawnerDataProvider::PerformPreExecuteReadbacks_GameThread(UPCGDataBinding* InBinding)
@@ -335,7 +332,7 @@ bool UPCGStaticMeshSpawnerDataProvider::PrepareForExecute_GameThread(UPCGDataBin
 		{
 			if (!InputDataDesc.DataDescs.IsEmpty())
 			{
-				PCG_KERNEL_VALIDATION_ERR(Context, Settings, /*bQuiet=*/false, FText::Format(
+				PCG_KERNEL_VALIDATION_ERR(Context, Settings, FText::Format(
 					LOCTEXT("MeshSelectorAttributeNotFound", "Mesh selector attribute '{0}' not found."),
 					FText::FromName(SelectorName)));
 			}
@@ -402,7 +399,7 @@ bool UPCGStaticMeshSpawnerDataProvider::PrepareForExecute_GameThread(UPCGDataBin
 
 	// We know the name and type of the selected mesh attribute statically and declared the attribute in GetKernelAttributeKeys,
 	// so the attribute ID should be present in the attribute table.
-	SelectedMeshAttributeId = InBinding->GetAttributeId(ProducerKernel->Settings->OutAttributeName, EPCGKernelAttributeType::StringKey);
+	SelectedMeshAttributeId = InBinding->GetAttributeId(Settings->OutAttributeName, EPCGKernelAttributeType::StringKey);
 	ensure(SelectedMeshAttributeId != INDEX_NONE);
 
 	return true;
@@ -466,7 +463,7 @@ void UPCGStaticMeshSpawnerDataProvider::CreatePrimitiveDescriptors(FPCGContext*
 		{
 			if (AttributeName == NAME_None)
 			{
-				PCG_KERNEL_VALIDATION_ERR(InContext, Settings, /*bQuiet=*/false,
+				PCG_KERNEL_VALIDATION_ERR(InContext, Settings,
 					FText::Format(LOCTEXT("InstanceDataAttributeInvalid", "Invalid instance data attribute specified '{0}'."), FText::FromName(AttributeName)));
 
 				continue;
@@ -474,7 +471,7 @@ void UPCGStaticMeshSpawnerDataProvider::CreatePrimitiveDescriptors(FPCGContext*
 
 			if (!InputDataDesc.ContainsAttributeOnAnyData(AttributeName))
 			{
-				PCG_KERNEL_VALIDATION_ERR(InContext, Settings, /*bQuiet=*/false,
+				PCG_KERNEL_VALIDATION_ERR(InContext, Settings,
 					FText::Format(LOCTEXT("InstanceDataAttributeNotFound", "Instance data attribute '{0}' not found."), FText::FromName(AttributeName)));
 
 				continue;
@@ -518,7 +515,7 @@ void UPCGStaticMeshSpawnerDataProvider::CreatePrimitiveDescriptors(FPCGContext*
 						{
 							// Currently only a single type per attribute name is supported (the name/type in the attribute table). It's possible to wire a graph
 							// for which an attribute is present with multiple types. Warn if this is encountered.
-							PCG_KERNEL_VALIDATION_WARN(InContext, Settings, /*bQuiet=*/false, FText::Format(
+							PCG_KERNEL_VALIDATION_WARN(InContext, Settings, FText::Format(
 								LOCTEXT("MeshSelectorAttributeNotUsable", "Attribute '{0}' not usable for mesh selection, only attributes of type String Key are supported."),
 								FText::FromName(SelectorName)));
 						}
@@ -544,7 +541,7 @@ void UPCGStaticMeshSpawnerDataProvider::CreatePrimitiveDescriptors(FPCGContext*
 
 			if (PrimitiveDescriptors.Num() >= PCGComputeConstants::MAX_PRIMITIVE_COMPONENTS_PER_SPAWNER)
 			{
-				PCG_KERNEL_VALIDATION_WARN(InContext, Settings, /*bQuiet=*/false,
+				PCG_KERNEL_VALIDATION_WARN(InContext, Settings,
 					FText::Format(PCGStaticMeshSpawnerDataInterface::TooManyPrimitivesFormat, FText::FromString(FString::FromInt(PCGComputeConstants::MAX_PRIMITIVE_COMPONENTS_PER_SPAWNER))));
 
 				break;
@@ -553,7 +550,7 @@ void UPCGStaticMeshSpawnerDataProvider::CreatePrimitiveDescriptors(FPCGContext*
 			UStaticMesh* StaticMesh = Cast<UStaticMesh>(FSoftObjectPath(MeshPathString).TryLoad());
 			if (!StaticMesh)
 			{
-				PCG_KERNEL_VALIDATION_ERR(InContext, Settings, /*bQuiet=*/false,
+				PCG_KERNEL_VALIDATION_ERR(InContext, Settings,
 					FText::Format(PCGStaticMeshSpawnerDataInterface::CouldNotLoadStaticMeshFormat, FText::FromString(MeshPathString)));
 
 				continue;
@@ -581,7 +578,7 @@ void UPCGStaticMeshSpawnerDataProvider::CreatePrimitiveDescriptors(FPCGContext*
 	{
 		if (SelectorWeighted->MeshEntries.IsEmpty())
 		{
-			PCG_KERNEL_VALIDATION_ERR(InContext, Settings, /*bQuiet=*/false, PCGStaticMeshSpawnerDataInterface::NoMeshEntriesFormat);
+			PCG_KERNEL_VALIDATION_ERR(InContext, Settings, PCGStaticMeshSpawnerDataInterface::NoMeshEntriesFormat);
 			return;
 		}
 
@@ -611,7 +608,7 @@ void UPCGStaticMeshSpawnerDataProvider::CreatePrimitiveDescriptors(FPCGContext*
 
 			if (PrimitiveDescriptors.Num() >= PCGComputeConstants::MAX_PRIMITIVE_COMPONENTS_PER_SPAWNER)
 			{
-				PCG_KERNEL_VALIDATION_WARN(InContext, Settings, /*bQuiet=*/false,
+				PCG_KERNEL_VALIDATION_WARN(InContext, Settings,
 					FText::Format(PCGStaticMeshSpawnerDataInterface::TooManyPrimitivesFormat, FText::FromString(FString::FromInt(PCGComputeConstants::MAX_PRIMITIVE_COMPONENTS_PER_SPAWNER))));
 
 				break;
@@ -658,7 +655,7 @@ void UPCGStaticMeshSpawnerDataProvider::CreatePrimitiveDescriptors(FPCGContext*
 			}
 			else
 			{
-				PCG_KERNEL_VALIDATION_ERR(InContext, Settings, /*bQuiet=*/false,
+				PCG_KERNEL_VALIDATION_ERR(InContext, Settings,
 					FText::Format(PCGStaticMeshSpawnerDataInterface::CouldNotLoadStaticMeshFormat, FText::FromString(Entry.Descriptor.StaticMesh.ToString())));
 			}
 		}
@@ -673,7 +670,7 @@ void UPCGStaticMeshSpawnerDataProvider::CreatePrimitiveDescriptors(FPCGContext*
 
 	if (TotalInstanceCount >= MAX_INSTANCE_ID)
 	{
-		PCG_KERNEL_VALIDATION_ERR(InContext, Settings, /*bQuiet=*/false, FText::Format(
+		PCG_KERNEL_VALIDATION_ERR(InContext, Settings, FText::Format(
 			LOCTEXT("TooManyInstances", "Tried to spawn too many instances ({0}), procedural ISM component creation skipped and instances will not be rendered."),
 			FText::FromString(FString::FromInt(TotalInstanceCount))));
 
diff --git a/Engine/Plugins/PCG/Source/PCG/Private/Compute/DataInterfaces/PCGDataLabelResolverDataInterface.cpp b/Engine/Plugins/PCG/Source/PCG/Private/Compute/DataInterfaces/PCGDataLabelResolverDataInterface.cpp
--- a/Engine/Plugins/PCG/Source/PCG/Private/Compute/DataInterfaces/PCGDataLabelResolverDataInterface.cpp
+++ b/Engine/Plugins/PCG/Source/PCG/Private/Compute/DataInterfaces/PCGDataLabelResolverDataInterface.cpp
@@ -129,7 +129,7 @@ bool UPCGDataLabelResolverDataProvider::PrepareForExecute_GameThread(UPCGDataBin
 			{
 				if (TSharedPtr<FPCGContextHandle> Context = InBinding->ContextHandle.Pin())
 				{
-					PCG_KERNEL_VALIDATION_ERR(Context->GetContext(), Kernel->GetSettings(), /*bQuiet=*/false, FText::Format(
+					PCG_KERNEL_VALIDATION_ERR(Context->GetContext(), Kernel->GetSettings(), FText::Format(
 						NSLOCTEXT("PCGDataLabelResolver", "LabelDoesNotExist", "Data label '{0}' was not present in the tags on the incoming data."),
 						FText::FromString(Label)));
 				}
diff --git a/Engine/Plugins/PCG/Source/PCG/Private/Compute/Elements/PCGComputeGraphElement.cpp b/Engine/Plugins/PCG/Source/PCG/Private/Compute/Elements/PCGComputeGraphElement.cpp
--- a/Engine/Plugins/PCG/Source/PCG/Private/Compute/Elements/PCGComputeGraphElement.cpp
+++ b/Engine/Plugins/PCG/Source/PCG/Private/Compute/Elements/PCGComputeGraphElement.cpp
@@ -154,7 +154,7 @@ bool FPCGComputeGraphElement::ExecuteInternal(FPCGContext* InContext) const
 		Context->DataBinding->Initialize(Context->ComputeGraph.Get(), Context);
 
 		// Perform validation after input data is initialized.
-		if (!Context->ComputeGraph->IsGraphValid(InContext, /*bQuiet=*/false))
+		if (!Context->ComputeGraph->IsGraphValid(InContext))
 		{
 			return true;
 		}
diff --git a/Engine/Plugins/PCG/Source/PCG/Private/Compute/Elements/PCGCustomHLSLKernel.cpp b/Engine/Plugins/PCG/Source/PCG/Private/Compute/Elements/PCGCustomHLSLKernel.cpp
--- a/Engine/Plugins/PCG/Source/PCG/Private/Compute/Elements/PCGCustomHLSLKernel.cpp
+++ b/Engine/Plugins/PCG/Source/PCG/Private/Compute/Elements/PCGCustomHLSLKernel.cpp
@@ -45,286 +45,73 @@ namespace PCGCustomHLSLKernel
 }
 
 #if WITH_EDITOR
-void UPCGCustomHLSLKernel::Initialize(const UPCGCustomHLSLSettings* InSettings, bool bLogDescriptions)
+void UPCGCustomHLSLKernel::Initialize(const UPCGSettings* InSettings, bool bLogDescriptions)
 {
-	Super::Initialize(bLogDescriptions);
-
-	check(InSettings);
-	Settings = InSettings;
+	Super::Initialize(InSettings, bLogDescriptions);
 
 	InitEntryPoint();
-}
-#endif // WITH_EDITOR
-
-const UPCGSettings* UPCGCustomHLSLKernel::GetSettings() const
-{
-	return Settings;
-}
 
-bool UPCGCustomHLSLKernel::IsKernelValid(FPCGContext* InContext, bool bQuiet) const
-{
-	TRACE_CPUPROFILER_EVENT_SCOPE(UPCGCustomHLSLKernel::IsKernelValid);
-
-	if (!Super::IsKernelValid(InContext, bQuiet) || !ensure(Settings))
-	{
-		return false;
-	}
-
-	if (!ValidatePCGNode(InContext, bQuiet, Cast<UPCGSettings>(Settings)))
-	{
-		return false;
-	}
+	const UPCGCustomHLSLSettings* CustomHLSLSettings = CastChecked<UPCGCustomHLSLSettings>(Settings);
 
-	if (Settings->OutputPins.IsEmpty())
-	{
-		PCG_KERNEL_VALIDATION_ERR(InContext, Settings, bQuiet, LOCTEXT("NoOutputs", "Custom HLSL nodes must have at least one output."));
-		return false;
-	}
+	FRegexMatcher FunctionsModuleMatcher(FRegexPattern(PCGHLSLElement::AttributeFunctionPattern), CustomHLSLSettings->ShaderFunctions);
 
-	auto CheckPinLabel = [Settings=Settings, InContext, bQuiet](FName PinLabel)
+	auto ParseAttributeFunctions = [](const FString& SourceToParse, TArray<FPCGParsedAttributeFunction>& OutParsedAttributeFunctions)
 	{
-		if (PinLabel == NAME_None)
-		{
-			PCG_KERNEL_VALIDATION_ERR(InContext, Settings, bQuiet, LOCTEXT("InvalidPinLabelNone", "Pin label 'None' is not a valid pin label."));
-			return false;
-		}
-
-		bool bFoundPinLabel = false;
-
-		auto IsAlreadyFound = [Settings, InContext, bQuiet, PinLabel, &bFoundPinLabel](const FPCGPinProperties PinProps)
-		{
-			if (PinProps.Label == PinLabel)
-			{
-				if (bFoundPinLabel)
-				{
-					PCG_KERNEL_VALIDATION_ERR(InContext, Settings, bQuiet, FText::Format(LOCTEXT("DuplicatedPinLabels", "Duplicate pin label '{0}', all labels must be unique."), FText::FromName(PinLabel)));
-					return true;
-				}
-
-				bFoundPinLabel = true;
-			}
-
-			return false;
-		};
+		const UEnum* AttributeTypeEnum = StaticEnum<EPCGKernelAttributeType>();
+		check(AttributeTypeEnum);
 
-		for (const FPCGPinProperties& PinProps : Settings->InputPins)
-		{
-			if (IsAlreadyFound(PinProps))
-			{
-				return false;
-			}
-		}
+		FRegexMatcher ModuleMatcher(FRegexPattern(PCGHLSLElement::AttributeFunctionPattern), SourceToParse);
 
-		for (const FPCGPinProperties& PinProps : Settings->OutputPins)
+		while (ModuleMatcher.FindNext())
 		{
-			if (IsAlreadyFound(PinProps))
-			{
-				return false;
-			}
+			const int64 AttributeType = AttributeTypeEnum->GetValueByNameString(ModuleMatcher.GetCaptureGroup(PCGHLSLElement::AttributeTypeCaptureGroup));
+
+			OutParsedAttributeFunctions.Emplace(
+				ModuleMatcher.GetCaptureGroup(PCGHLSLElement::AttributePinCaptureGroup),
+				ModuleMatcher.GetCaptureGroup(PCGHLSLElement::AttributeFunctionCaptureGroup),
+				AttributeType,
+				ModuleMatcher.GetCaptureGroup(PCGHLSLElement::AttributeNameCaptureGroup),
+				ModuleMatcher.GetMatchBeginning());
 		}
-
-		return true;
 	};
 
-	// Validate input pins
-	bool bIsFirstInputPin = true;
-	for (const FPCGPinProperties& Properties : Settings->InputPins)
-	{
-		if (!CheckPinLabel(Properties.Label))
-		{
-			return false;
-		}
-
-		if (bIsFirstInputPin && Settings->KernelType == EPCGKernelType::PointProcessor)
-		{
-			if (Properties.AllowedTypes != EPCGDataType::Point)
-			{
-				PCG_KERNEL_VALIDATION_ERR(InContext, Settings, bQuiet, FText::Format(
-					LOCTEXT("InvalidNonPointPrimaryInput", "'Point Processor' nodes require primary input pin to be of type 'Point', but found '{0}'."),
-					FText::FromString(PCGCustomHLSLKernel::GetDataTypeString(Properties.AllowedTypes))));
-
-				return false;
-			}
-		}
-
-		if (!PCGComputeHelpers::IsTypeAllowedAsInput(Properties.AllowedTypes))
-		{
-			PCG_KERNEL_VALIDATION_ERR(InContext, Settings, bQuiet, FText::Format(
-				LOCTEXT("InvalidInputType", "Unsupported input type '{0}', found on pin '{1}'."),
-				FText::FromString(PCGCustomHLSLKernel::GetDataTypeString(Properties.AllowedTypes)),
-				FText::FromName(Properties.Label)));
-
-			return false;
-		}
-
-		bIsFirstInputPin = false;
-	}
+	ParseAttributeFunctions(CustomHLSLSettings->ShaderSource, SourceParsedAttributeFunctions);
+	ParseAttributeFunctions(CustomHLSLSettings->ShaderFunctions, FunctionsParsedAttributeFunctions);
 
-	// Validate output pins
-	bool bIsFirstOutputPin = true;
-	for (const FPCGPinPropertiesGPU& Properties : Settings->OutputPins)
+	auto ParseCopyElementFunctions = [](const FString& SourceToParse, TArray<FPCGParsedCopyElementFunction>& OutParsedCopyElementFunctions)
 	{
-		if (!CheckPinLabel(Properties.Label))
-		{
-			return false;
-		}
-
-		const bool bPinIsDefinedByKernel = bIsFirstOutputPin && (Settings->KernelType == EPCGKernelType::PointGenerator || Settings->KernelType == EPCGKernelType::PointProcessor);
-
-		if (bPinIsDefinedByKernel)
-		{
-			if (Properties.AllowedTypes != EPCGDataType::Point)
-			{
-				PCG_KERNEL_VALIDATION_ERR(InContext, Settings, bQuiet, FText::Format(
-					LOCTEXT("InvalidNonPointPrimaryOutput", "'Point Processor' and 'Point Generator' nodes require primary output pin to be of type 'Point', but found '{0}'."),
-					FText::FromString(PCGCustomHLSLKernel::GetDataTypeString(Properties.AllowedTypes))));
-
-				return false;
-			}
-		}
-
-		if (!PCGComputeHelpers::IsTypeAllowedAsOutput(Properties.AllowedTypes))
-		{
-			PCG_KERNEL_VALIDATION_ERR(InContext, Settings, bQuiet, FText::Format(
-				LOCTEXT("InvalidOutputType", "Unsupported output type '{0}', found on pin '{1}'."),
-				FText::FromString(PCGCustomHLSLKernel::GetDataTypeString(Properties.AllowedTypes)),
-				FText::FromName(Properties.Label)));
+		/** Regex for parsing usages of CopyElementFrom function. */
+		constexpr TCHAR CopyElementPattern[] = { TEXT(/*TargetPin*/"([a-zA-Z0-9]+)" "_CopyElementFrom_" /*SourcePin*/"([a-zA-Z0-9]+)\\s*\\(") };
+		constexpr int CopyElementTargetPinCaptureGroup = 1;
+		constexpr int CopyElementSourcePinCaptureGroup = 2;
 
-			return false;
-		}
+		FRegexMatcher ModuleMatcher(FRegexPattern(CopyElementPattern), SourceToParse);
 
-		if (!bPinIsDefinedByKernel)
+		while (ModuleMatcher.FindNext())
 		{
-			const FPCGPinPropertiesGPUStruct& Props = Properties.PropertiesGPU;
-
-			if (Props.InitializationMode == EPCGPinInitMode::FromInputPins)
-			{
-				if (Props.PinsToInititalizeFrom.IsEmpty())
-				{
-					PCG_KERNEL_VALIDATION_ERR(InContext, Settings, bQuiet, FText::Format(
-						LOCTEXT("InitFromEmptyPins", "Output pin '{0}' tried to initialize from input pins, but no pins were specified."),
-						FText::FromName(Properties.Label)));
-
-					return false;
-				}
-
-				for (const FName InitPinName : Props.PinsToInititalizeFrom)
-				{
-					const FPCGPinProperties* InitPinProps = Settings->InputPins.FindByPredicate([InitPinName](const FPCGPinProperties& InPinProps)
-					{
-						return InPinProps.Label == InitPinName;
-					});
-
-					if (!InitPinProps)
-					{
-						PCG_KERNEL_VALIDATION_ERR(InContext, Settings, bQuiet, FText::Format(
-							LOCTEXT("InitFromNonExistentPin", "Output pin '{0}' tried to initialize from non-existent input pin '{1}'."),
-							FText::FromName(Properties.Label),
-							FText::FromName(InitPinName)));
-
-						return false;
-					}
-
-					if (!PCGComputeHelpers::IsTypeAllowedInDataCollection(InitPinProps->AllowedTypes))
-					{
-						PCG_KERNEL_VALIDATION_ERR(InContext, Settings, bQuiet, FText::Format(
-							LOCTEXT("InitFromInvalidPinType", "Output pin '{0}' tried to initialize from input pin '{1}', but pin '{1}' has an invalid type."),
-							FText::FromName(Properties.Label),
-							FText::FromName(InitPinName)));
-
-						return false;
-					}
-				}
-
-				// TODO: Could do validation on data multiplicity for Pairwise, checking that data counts are 1 or N, but maybe that should be a runtime error instead.
-			}
-
-			const bool bUsingFixedDataCount = Props.InitializationMode == EPCGPinInitMode::Custom || Props.DataCountMode == EPCGDataCountMode::Fixed;
-
-			if (bUsingFixedDataCount)
-			{
-				if (Props.DataCount < 1)
-				{
-					PCG_KERNEL_VALIDATION_ERR(InContext, Settings, bQuiet, FText::Format(
-						LOCTEXT("InvalidDataCount", "Invalid fixed data count {0} on output pin '{1}'. Must be greater than 0."),
-						FText::AsNumber(Props.DataCount),
-						FText::FromName(Properties.Label)));
-
-					return false;
-				}
-
-			}
-
-			const bool bUsingFixedElemCount = Props.InitializationMode == EPCGPinInitMode::Custom || Props.ElementCountMode == EPCGElementCountMode::Fixed;
-
-			if (bUsingFixedElemCount)
-			{
-				if (Props.ElementCount < 1)
-				{
-					PCG_KERNEL_VALIDATION_ERR(InContext, Settings, bQuiet, FText::Format(
-						LOCTEXT("InvalidElementCount", "Invalid fixed element count {0} on output pin '{1}'. Must be greater than 0."),
-						FText::AsNumber(Props.ElementCount),
-						FText::FromName(Properties.Label)));
-
-					return false;
-				}
-			}
-
-			if (Props.ElementCountMultiplier < 1)
-			{
-				PCG_KERNEL_VALIDATION_ERR(InContext, Settings, bQuiet, FText::Format(
-					LOCTEXT("InvalidElementCountMultiplier", "Invalid element count multiplier {0} on output pin '{1}'. Must be greater than 0."),
-					FText::AsNumber(Props.ElementCountMultiplier),
-					FText::FromName(Properties.Label)));
-
-				return false;
-			}
+			OutParsedCopyElementFunctions.Emplace(
+				ModuleMatcher.GetCaptureGroup(CopyElementSourcePinCaptureGroup),
+				ModuleMatcher.GetCaptureGroup(CopyElementTargetPinCaptureGroup));
 		}
+	};
 
-		bIsFirstOutputPin = false;
-	}
-
-	if (Settings->KernelType == EPCGKernelType::Custom && Settings->DispatchThreadCount == EPCGDispatchThreadCount::FromProductOfInputPins)
-	{
-		if (Settings->ThreadCountInputPinLabels.IsEmpty())
-		{
-			PCG_KERNEL_VALIDATION_ERR(InContext, Settings, bQuiet, LOCTEXT("MissingThreadCountPins", "Dispatch thread count is based on input pins but no labels have been set in Input Pins array."));
-			return false;
-		}
+	ParseCopyElementFunctions(CustomHLSLSettings->ShaderSource, SourceParsedCopyElementFunctions);
+	ParseCopyElementFunctions(CustomHLSLSettings->ShaderFunctions, FunctionsParsedCopyElementFunctions);
+}
+#endif // WITH_EDITOR
 
-		for (const FName& Label : Settings->ThreadCountInputPinLabels)
-		{
-			if (!Settings->InputPins.FindByPredicate([Label](const FPCGPinProperties& InProps) { return InProps.Label == Label; }))
-			{
-				PCG_KERNEL_VALIDATION_ERR(InContext, Settings, bQuiet, FText::Format(LOCTEXT("MissingThreadCountPin", "Invalid pin specified in Input Pins array: '{0}'."), FText::FromName(Label)));
-				return false;
-			}
-		}
-	}
+const UPCGSettings* UPCGCustomHLSLKernel::GetSettings() const
+{
+	return Settings;
+}
 
-	if (IsThreadCountMultiplierInUse())
-	{
-		if (Settings->ThreadCountMultiplier < 1)
-		{
-			PCG_KERNEL_VALIDATION_ERR(InContext, Settings, bQuiet, FText::Format(LOCTEXT("InvalidThreadCountMultiplier", "Thread Count Multiplier has invalid value ({0}). Must be greater than 0."), Settings->ThreadCountMultiplier));
-			return false;
-		}
-	}
+bool UPCGCustomHLSLKernel::IsKernelValid(FPCGContext* InContext) const
+{
+	TRACE_CPUPROFILER_EVENT_SCOPE(UPCGCustomHLSLKernel::IsKernelValid);
 
-	// Validate attributes
-	for (const FPCGKernelAttributeKey& AttributeKey : Settings->KernelAttributeKeys)
+	if (!Super::IsKernelValid(InContext))
 	{
-		if (AttributeKey.Type == EPCGKernelAttributeType::Invalid)
-		{
-			const UEnum* AttributeTypeEnum = StaticEnum<EPCGKernelAttributeType>();
-			check(AttributeTypeEnum);
-
-			PCG_KERNEL_VALIDATION_ERR(InContext, Settings, bQuiet, FText::Format(
-				LOCTEXT("InvalidAttributeTypeNone", "Attribute '{0}' has invalid GPU attribute type '{1}', check the 'Attributes to Create' array on your pins."),
-				FText::FromName(AttributeKey.Name),
-				FText::FromString(AttributeTypeEnum->GetNameStringByValue(static_cast<int64>(AttributeKey.Type)))));
-			return false;
-		}
+		return false;
 	}
 
 	if (InContext)
@@ -339,25 +126,7 @@ bool UPCGCustomHLSLKernel::IsKernelValid(FPCGContext* InContext, bool bQuiet) co
 		{
 			if (ErrorTextPtr)
 			{
-				PCG_KERNEL_VALIDATION_ERR(InContext, Settings, bQuiet, *ErrorTextPtr);
-			}
-			return false;
-		}
-
-		if (!ProcessShaderSource(/*InOutSource=*/nullptr, ErrorTextPtr))
-		{
-			if (ErrorTextPtr)
-			{
-				PCG_KERNEL_VALIDATION_ERR(InContext, Settings, bQuiet, *ErrorTextPtr);
-			}
-			return false;
-		}
-
-		if (!Settings->bMuteUnwrittenPinDataErrors && !AreAllOutputPinsWritten(ErrorTextPtr))
-		{
-			if (ErrorTextPtr)
-			{
-				PCG_KERNEL_VALIDATION_ERR(InContext, Settings, bQuiet, *ErrorTextPtr);
+				PCG_KERNEL_VALIDATION_ERR(InContext, Settings, *ErrorTextPtr);
 			}
 			return false;
 		}
@@ -370,20 +139,22 @@ bool UPCGCustomHLSLKernel::ComputeOutputBindingDataDesc(const UPCGComputeGraph*
 {
 	check(InGraph);
 	check(InBinding);
+	
+	const UPCGCustomHLSLSettings* CustomHLSLSettings = CastChecked<UPCGCustomHLSLSettings>(Settings);
 
 	FPCGDataCollectionDesc OutputDataDesc;
 
-	const FPCGPinPropertiesGPU* OutputPinProperties = Settings->OutputPins.FindByPredicate([InOutputPinLabel](const FPCGPinPropertiesGPU& InProps) {return InProps.Label == InOutputPinLabel; });
+	const FPCGPinPropertiesGPU* OutputPinProperties = CustomHLSLSettings->OutputPins.FindByPredicate([InOutputPinLabel](const FPCGPinPropertiesGPU& InProps) {return InProps.Label == InOutputPinLabel; });
 	if (!OutputPinProperties)
 	{
 		OutDataDesc = {};
 		return true;
 	}
 
-	const FName FirstOutputPinLabel = Settings->OutputPins[0].Label;
+	const FName FirstOutputPinLabel = CustomHLSLSettings->OutputPins[0].Label;
 
 	// The primary output pin follows any rules prescribed by kernel type.
-	if (InOutputPinLabel == FirstOutputPinLabel && Settings->KernelType == EPCGKernelType::PointProcessor)
+	if (InOutputPinLabel == FirstOutputPinLabel && CustomHLSLSettings->KernelType == EPCGKernelType::PointProcessor)
 	{
 		FPCGDataCollectionDesc InputDesc;
 
@@ -398,20 +169,20 @@ bool UPCGCustomHLSLKernel::ComputeOutputBindingDataDesc(const UPCGComputeGraph*
 
 		OutputDataDesc = MoveTemp(InputDesc);
 	}
-	else if (InOutputPinLabel == FirstOutputPinLabel && Settings->KernelType == EPCGKernelType::PointGenerator)
+	else if (InOutputPinLabel == FirstOutputPinLabel && CustomHLSLSettings->KernelType == EPCGKernelType::PointGenerator)
 	{
 		// Generators always produce a single point data with known point count.
-		OutputDataDesc.DataDescs.Emplace(EPCGDataType::Point, Settings->PointCount);
+		OutputDataDesc.DataDescs.Emplace(EPCGDataType::Point, CustomHLSLSettings->PointCount);
 	}
 	else if (ensure(OutputPinProperties))
 	{
-		ComputeDataDescFromPinProperties(*OutputPinProperties, MakeArrayView(Settings->InputPins), InBinding, OutputDataDesc);
+		ComputeDataDescFromPinProperties(*OutputPinProperties, MakeArrayView(CustomHLSLSettings->InputPins), InBinding, OutputDataDesc);
 	}
 
-	for (const FPCGKernelAttributeKey& AttributeKey : Settings->KernelAttributeKeys)
+	for (const FPCGKernelAttributeKey& AttributeKey : CustomHLSLSettings->KernelAttributeKeys)
 	{
 		// Add attributes that will be created for this pin on the GPU.
-		if (const TArray<TTuple<FPCGKernelAttributeKey, bool>>* Keys = Settings->PinToAttributeKeys.Find(InOutputPinLabel))
+		if (const TArray<TTuple<FPCGKernelAttributeKey, bool>>* Keys = CustomHLSLSettings->PinToAttributeKeys.Find(InOutputPinLabel))
 		{
 			const TTuple<FPCGKernelAttributeKey, bool>* Pair = Keys->FindByPredicate([AttributeKey](const TTuple<FPCGKernelAttributeKey, bool>& Pair) { return Pair.Key == AttributeKey; });
 			if (!Pair || !Pair->Value)
@@ -457,7 +228,7 @@ bool UPCGCustomHLSLKernel::ComputeOutputBindingDataDesc(const UPCGComputeGraph*
 		TArray<FPCGDataCollectionDesc> RelevantInputDataDescs;
 
 		// Collect descriptions of input data items that have string key attributes.
-		for(const FPCGPinProperties& PinProps : Settings->InputPins)
+		for(const FPCGPinProperties& PinProps : CustomHLSLSettings->InputPins)
 		{
 			FPCGDataCollectionDesc InputPinDesc;
 			ensure(InBinding->GetKernelPinDataDesc_GameThread(this, PinProps.Label, /*bIsInputPin=*/true, InputPinDesc));
@@ -550,123 +321,31 @@ FString UPCGCustomHLSLKernel::GetCookedSource(const UPCGComputeGraph* InComputeG
 {
 	check(InComputeGraph);
 
-	const FIntVector GroupSize(PCGComputeConstants::THREAD_GROUP_SIZE, 1, 1);
+	const UPCGCustomHLSLSettings* CustomHLSLSettings = CastChecked<UPCGCustomHLSLSettings>(Settings);
+	FString Source = CustomHLSLSettings->ShaderSource;
+	FString Functions = TEXT("\n\n// User kernel functions\n") + CustomHLSLSettings->ShaderFunctions;
 
-	FString Source = Settings->ShaderSource;
-	FString Functions = TEXT("\n\n// User kernel functions\n") + Settings->ShaderFunctions;
+	ProcessShaderSource(InComputeGraph, Source, Functions);
 
-#if PLATFORM_WINDOWS
-	// Remove old-school stuff.
-	Source.ReplaceInline(TEXT("\r"), TEXT(""));
-	Functions.ReplaceInline(TEXT("\r"), TEXT(""));
-#endif
+	FString ShaderPathName = GetPathName();
+	PCGCustomHLSLKernel::ConvertObjectPathToShaderFilePath(ShaderPathName);
 
-	ProcessShaderSource(&Source);
-	ProcessShaderSource(&Functions);
+	//const bool bHasKernelKeyword = Source.Contains(TEXT("KERNEL"), ESearchCase::CaseSensitive);
 
-	// We inject attribute IDs directly into the source. This is most efficient and saves us passing them into the kernel. However
-	// the trade off is that a shader with a shared source will generate multiple variants if used in different compute graphs with
-	// different attribute IDs.
-	auto ReplaceAttributeNamesWithIds = [InComputeGraph](FString& InOutSource)
+	FString Includes;
 	{
-		FRegexMatcher AttributeFunctionMatcher(FRegexPattern(PCGHLSLElement::AttributeFunctionPattern), InOutSource);
-		TArray<TTuple<int32, int32, int32>> Replacements;
-
-		const UEnum* AttributeTypeEnum = StaticEnum<EPCGKernelAttributeType>();
-		check(AttributeTypeEnum);
+		// Add with caution: Pulling in external includes has the danger that 1) we may pull in more than we expect
+		// if they include additional things in the future, and 2) if their functions change it could break user's
+		// kernel source. The latter may be mitigated by branching our own PCG version of these if needed in the future.
 
-		while (AttributeFunctionMatcher.FindNext())
-		{
-			const FString TypeStr = AttributeFunctionMatcher.GetCaptureGroup(PCGHLSLElement::AttributeTypeCaptureGroup);
-			const FString NameStr = AttributeFunctionMatcher.GetCaptureGroup(PCGHLSLElement::AttributeNameCaptureGroup);
-
-			const int64 AttrType = AttributeTypeEnum->GetValueByName(FName(*TypeStr));
-
-			if (!ensure(AttrType != INDEX_NONE))
-			{
-				break;
-			}
-
-			const FPCGKernelAttributeKey AttributeKey(FName(*NameStr), static_cast<EPCGKernelAttributeType>(AttrType));
-			const FString SourceDefinition = PCGCustomHLSLKernel::GetKernelAttributeKeyAsString(AttributeKey);
-
-			const int32 AttributeIndex = InComputeGraph->GetStaticAttributeTable().GetAttributeId(AttributeKey);
-			const int32 ReplaceStartIndex = InOutSource.Find(SourceDefinition, ESearchCase::CaseSensitive, ESearchDir::FromStart, AttributeFunctionMatcher.GetMatchBeginning());
-			const int32 ReplaceEndIndex = ReplaceStartIndex + SourceDefinition.Len();
-
-			Replacements.Emplace(AttributeIndex, ReplaceStartIndex, ReplaceEndIndex);
-		}
-
-		for (int I = Replacements.Num() - 1; I >= 0; --I)
-		{
-			const int32 AttributeIndex = Replacements[I].Get<0>();
-			const int32 ReplaceStartIndex = Replacements[I].Get<1>();
-			const int32 ReplaceEndIndex = Replacements[I].Get<2>();
-
-			InOutSource = InOutSource.Left(ReplaceStartIndex) + FString::FromInt(AttributeIndex) + InOutSource.RightChop(ReplaceEndIndex);
-		}
-	};
-
-	ReplaceAttributeNamesWithIds(Source);
-	ReplaceAttributeNamesWithIds(Functions);
-
-	// We inject data IDs directly into the source. They will get remapped to data indices using a label resolver data interface.
-	if (const FPCGPinDataLabels* PinToDataLabels = InComputeGraph->GetStaticDataLabelsTable().Find(KernelIndex))
-	{
-		for (const TPair<FName, FPCGDataLabels>& Pair : PinToDataLabels->PinToDataLabels)
-		{
-			const FName PinLabel = Pair.Key;
-			const TArray<FString>& DataLabels = Pair.Value.Labels;
-
-			for (int DataId = 0; DataId < DataLabels.Num(); ++DataId)
-			{
-				const FString& DataLabel = DataLabels[DataId];
-				const FString ReplacementStr = FString::Format(TEXT("{0}_GetDataIndexFromIdInternal(/*DataId=*/{1})"),
-					{
-						PCGComputeHelpers::GetDataLabelResolverName(PinLabel),
-						FString::FromInt(DataId)
-					});
-
-				// Matches against {PinName}_AnyFunction('{DataLabel}'...
-				// First capture group is the data label, so that we can find & replace it by index.
-				const FString Pattern = FString::Format(TEXT("{0}_.*?[\\s]*?\\([\\s]*?('{1}')"), { *PinLabel.ToString(), DataLabel });
-
-				auto ReplaceDataLabels = [&Pattern, &ReplacementStr](FString& InOutSource)
-				{
-					FRegexMatcher ModuleMatcher(FRegexPattern(Pattern), InOutSource);
-					while (ModuleMatcher.FindNext())
-					{
-						const int32 DataLabelCharIndexStart = ModuleMatcher.GetCaptureGroupBeginning(1);
-						const int32 DataLabelCharIndexEnd = ModuleMatcher.GetCaptureGroupEnding(1);
-
-						InOutSource = InOutSource.Left(DataLabelCharIndexStart) + ReplacementStr + InOutSource.RightChop(DataLabelCharIndexEnd);
-					}
-				};
-
-				ReplaceDataLabels(Source);
-				ReplaceDataLabels(Functions);
-			}
-		}
-	}
-
-	FString ShaderPathName = GetPathName();
-	PCGCustomHLSLKernel::ConvertObjectPathToShaderFilePath(ShaderPathName);
-
-	//const bool bHasKernelKeyword = Source.Contains(TEXT("KERNEL"), ESearchCase::CaseSensitive);
-
-	FString Includes;
-	{
-		// Add with caution: Pulling in external includes has the danger that 1) we may pull in more than we expect
-		// if they include additional things in the future, and 2) if their functions change it could break user's
-		// kernel source. The latter may be mitigated by branching our own PCG version of these if needed in the future.
-
-		// TODO bring these in via additional sources so that their contents get hashed. Also these could generate declarations.
-		Includes += TEXT("#include \"/Engine/Private/ComputeShaderUtils.ush\"\n");
-		Includes += TEXT("#include \"/Engine/Private/Quaternion.ush\"\n");
-		Includes += TEXT("#include \"/Plugin/PCG/Private/PCGShaderUtils.ush\"\n");
-		Includes += TEXT("#include \"/Plugin/PCG/Private/PCGShaderUtilsInternal.ush\"\n");
-	}
+		// TODO bring these in via additional sources so that their contents get hashed. Also these could generate declarations.
+		Includes += TEXT("#include \"/Engine/Private/ComputeShaderUtils.ush\"\n");
+		Includes += TEXT("#include \"/Engine/Private/Quaternion.ush\"\n");
+		Includes += TEXT("#include \"/Plugin/PCG/Private/PCGShaderUtils.ush\"\n");
+		Includes += TEXT("#include \"/Plugin/PCG/Private/PCGShaderUtilsInternal.ush\"\n");
+	}
 
+	const FIntVector GroupSize(PCGComputeConstants::THREAD_GROUP_SIZE, 1, 1);
 	const FString KernelFunc = FString::Printf(
 		TEXT("[numthreads(%d, %d, %d)]\nvoid %s(uint3 GroupId : SV_GroupID, uint GroupIndex : SV_GroupIndex)"),
 		GroupSize.X, GroupSize.Y, GroupSize.Z, *GetEntryPoint());
@@ -679,7 +358,7 @@ FString UPCGCustomHLSLKernel::GetCookedSource(const UPCGComputeGraph* InComputeG
 	// Used to signal that a kernel has executed. Set the most significant bit in NumData.
 	FString SetAsExecuted = TEXT("    // Signal kernel executed by setting the most significant bit of NumData.\n");
 
-	for (const FPCGPinPropertiesGPU& PinProps : Settings->OutputPins)
+	for (const FPCGPinPropertiesGPU& PinProps : CustomHLSLSettings->OutputPins)
 	{
 		SetAsExecuted += FString::Format(TEXT("    if (all(GroupId == 0u) && GroupIndex == 0) {0}_SetAsExecutedInternal();\n"), { PinProps.Label.ToString() });
 	}
@@ -695,7 +374,7 @@ FString UPCGCustomHLSLKernel::GetCookedSource(const UPCGComputeGraph* InComputeG
 				{ PinLabel.ToString() });
 		};
 
-	if (Settings->KernelType == EPCGKernelType::PointProcessor)
+	if (CustomHLSLSettings->KernelType == EPCGKernelType::PointProcessor)
 	{
 		const FPCGPinProperties* InputPin = GetFirstPointInputPin();
 		const FPCGPinPropertiesGPU* OutputPin = GetFirstPointOutputPin();
@@ -725,9 +404,9 @@ FString UPCGCustomHLSLKernel::GetCookedSource(const UPCGComputeGraph* InComputeG
 				{ InputPin->Label.ToString(), OutputPin->Label.ToString() });
 		}
 	}
-	else if (Settings->KernelType == EPCGKernelType::PointGenerator)
+	else if (CustomHLSLSettings->KernelType == EPCGKernelType::PointGenerator)
 	{
-		KernelSpecificPreamble += FString::Format(TEXT("    const uint NumPoints = {0};\n"), { Settings->PointCount });
+		KernelSpecificPreamble += FString::Format(TEXT("    const uint NumPoints = {0};\n"), { CustomHLSLSettings->PointCount });
 
 		if (const FPCGPinPropertiesGPU* OutputPin = GetFirstPointOutputPin())
 		{
@@ -782,14 +461,15 @@ FString UPCGCustomHLSLKernel::GetCookedSource(const UPCGComputeGraph* InComputeG
 
 int UPCGCustomHLSLKernel::ComputeThreadCount(const UPCGDataBinding* InBinding) const
 {
+	const UPCGCustomHLSLSettings* CustomHLSLSettings = CastChecked<UPCGCustomHLSLSettings>(Settings);
 	int ThreadCount = 0;
 
-	if (Settings->KernelType == EPCGKernelType::PointGenerator)
+	if (CustomHLSLSettings->KernelType == EPCGKernelType::PointGenerator)
 	{
 		// Point generator has fixed thread count.
-		ThreadCount = Settings->PointCount;
+		ThreadCount = CustomHLSLSettings->PointCount;
 	}
-	else if (Settings->KernelType == EPCGKernelType::PointProcessor)
+	else if (CustomHLSLSettings->KernelType == EPCGKernelType::PointProcessor)
 	{
 		// Processing volume depends on data arriving on primary pin.
 		if (const FPCGPinProperties* InputPin = GetFirstPointInputPin())
@@ -797,11 +477,11 @@ int UPCGCustomHLSLKernel::ComputeThreadCount(const UPCGDataBinding* InBinding) c
 			ThreadCount = GetElementCountForInputPin(*InputPin, InBinding);
 		}
 	}
-	else if (Settings->KernelType == EPCGKernelType::Custom)
+	else if (CustomHLSLSettings->KernelType == EPCGKernelType::Custom)
 	{
-		if (Settings->DispatchThreadCount == EPCGDispatchThreadCount::FromFirstOutputPin)
+		if (CustomHLSLSettings->DispatchThreadCount == EPCGDispatchThreadCount::FromFirstOutputPin)
 		{
-			if (const FPCGPinPropertiesGPU* OutputPin = Settings->OutputPins.IsEmpty() ? nullptr : &Settings->OutputPins[0])
+			if (const FPCGPinPropertiesGPU* OutputPin = CustomHLSLSettings->OutputPins.IsEmpty() ? nullptr : &CustomHLSLSettings->OutputPins[0])
 			{
 				FPCGDataCollectionDesc Desc;
 				if (ensure(ComputeOutputBindingDataDesc(InBinding->Graph, OutputPin->Label, InBinding, Desc)))
@@ -810,19 +490,19 @@ int UPCGCustomHLSLKernel::ComputeThreadCount(const UPCGDataBinding* InBinding) c
 				}
 			}
 		}
-		else if (Settings->DispatchThreadCount == EPCGDispatchThreadCount::FromProductOfInputPins)
+		else if (CustomHLSLSettings->DispatchThreadCount == EPCGDispatchThreadCount::FromProductOfInputPins)
 		{
-			for (const FName& PinLabel : Settings->ThreadCountInputPinLabels)
+			for (const FName& PinLabel : CustomHLSLSettings->ThreadCountInputPinLabels)
 			{
-				if (const FPCGPinProperties* InputPin = Settings->InputPins.FindByPredicate([PinLabel](const FPCGPinProperties& InProps) { return InProps.Label == PinLabel; }))
+				if (const FPCGPinProperties* InputPin = CustomHLSLSettings->InputPins.FindByPredicate([PinLabel](const FPCGPinProperties& InProps) { return InProps.Label == PinLabel; }))
 				{
 					ThreadCount = FMath::Max(ThreadCount, 1) * GetElementCountForInputPin(*InputPin, InBinding);
 				}
 			}
 		}
-		else if (Settings->DispatchThreadCount == EPCGDispatchThreadCount::Fixed)
+		else if (CustomHLSLSettings->DispatchThreadCount == EPCGDispatchThreadCount::Fixed)
 		{
-			ThreadCount = Settings->FixedThreadCount;
+			ThreadCount = CustomHLSLSettings->FixedThreadCount;
 		}
 	}
 	else
@@ -832,7 +512,7 @@ int UPCGCustomHLSLKernel::ComputeThreadCount(const UPCGDataBinding* InBinding) c
 
 	if (IsThreadCountMultiplierInUse())
 	{
-		ThreadCount *= Settings->ThreadCountMultiplier;
+		ThreadCount *= CustomHLSLSettings->ThreadCountMultiplier;
 	}
 
 	return ThreadCount;
@@ -840,9 +520,11 @@ int UPCGCustomHLSLKernel::ComputeThreadCount(const UPCGDataBinding* InBinding) c
 
 void UPCGCustomHLSLKernel::GetDataLabels(FName InPinLabel, TArray<FString>& OutDataLabels) const
 {
+	const UPCGCustomHLSLSettings* CustomHLSLSettings = CastChecked<UPCGCustomHLSLSettings>(Settings);
+
 	// Matches against {PinName}_AnyFunction('{DataLabel}'...
 	const FString Pattern = FString::Format(TEXT("{0}_.*?[\\s]*?\\([\\s]*?'([a-zA-Z0-9_].*?)'"), { *InPinLabel.ToString() });
-	const FString Source = Settings->ShaderFunctions + TEXT("\n") + Settings->ShaderSource;
+	const FString Source = CustomHLSLSettings->ShaderFunctions + TEXT("\n") + CustomHLSLSettings->ShaderSource;
 
 	// First capture: Data label (supports a - z, A - Z, 0 - 9, and underscores).
 	constexpr int LabelCaptureGroup = 1;
@@ -856,17 +538,21 @@ void UPCGCustomHLSLKernel::GetDataLabels(FName InPinLabel, TArray<FString>& OutD
 
 void UPCGCustomHLSLKernel::GetKernelAttributeKeys(TArray<FPCGKernelAttributeKey>& OutKeys) const
 {
-	OutKeys.Append(Settings->KernelAttributeKeys);
+	const UPCGCustomHLSLSettings* CustomHLSLSettings = CastChecked<UPCGCustomHLSLSettings>(Settings);
+	OutKeys.Append(CustomHLSLSettings->KernelAttributeKeys);
 }
 
 uint32 UPCGCustomHLSLKernel::GetThreadCountMultiplier() const
 {
-	return Settings->ThreadCountMultiplier;
+	const UPCGCustomHLSLSettings* CustomHLSLSettings = CastChecked<UPCGCustomHLSLSettings>(Settings);
+	return CustomHLSLSettings->ThreadCountMultiplier;
 }
 
 uint32 UPCGCustomHLSLKernel::GetElementCountMultiplier(FName InOutputPinLabel) const
 {
-	if (const FPCGPinPropertiesGPU* PinProps = Settings->OutputPins.FindByPredicate([InOutputPinLabel](const FPCGPinPropertiesGPU& InProps) { return InProps.Label == InOutputPinLabel; }))
+	const UPCGCustomHLSLSettings* CustomHLSLSettings = CastChecked<UPCGCustomHLSLSettings>(Settings);
+
+	if (const FPCGPinPropertiesGPU* PinProps = CustomHLSLSettings->OutputPins.FindByPredicate([InOutputPinLabel](const FPCGPinPropertiesGPU& InProps) { return InProps.Label == InOutputPinLabel; }))
 	{
 		return PinProps->GetElementCountMultiplier();
 	}
@@ -879,14 +565,14 @@ uint32 UPCGCustomHLSLKernel::GetElementCountMultiplier(FName InOutputPinLabel) c
 
 void UPCGCustomHLSLKernel::GetInputPins(TArray<FPCGPinProperties>& OutPins) const
 {
-	check(Settings);
-	OutPins = Settings->InputPins;
+	const UPCGCustomHLSLSettings* CustomHLSLSettings = CastChecked<UPCGCustomHLSLSettings>(Settings);
+	OutPins = CustomHLSLSettings->InputPins;
 }
 
 void UPCGCustomHLSLKernel::GetOutputPins(TArray<FPCGPinPropertiesGPU>& OutPins) const
 {
-	check(Settings);
-	OutPins = Settings->OutputPins;
+	const UPCGCustomHLSLSettings* CustomHLSLSettings = CastChecked<UPCGCustomHLSLSettings>(Settings);
+	OutPins = CustomHLSLSettings->OutputPins;
 }
 
 #if WITH_EDITOR
@@ -921,295 +607,713 @@ void UPCGCustomHLSLKernel::InitEntryPoint()
 		EntryPoint.ReplaceCharInline(Char, TCHAR('_'));
 	}
 }
-#endif // WITH_EDITOR
-
-bool UPCGCustomHLSLKernel::IsThreadCountMultiplierInUse() const
-{
-	return Settings->KernelType == EPCGKernelType::Custom && Settings->DispatchThreadCount != EPCGDispatchThreadCount::Fixed;
-}
 
-bool UPCGCustomHLSLKernel::AreAttributesValid(FPCGContext* InContext, FText* OutErrorText) const
+bool UPCGCustomHLSLKernel::PerformStaticValidation()
 {
-	TRACE_CPUPROFILER_EVENT_SCOPE(UPCGCustomHLSLKernel::AreAttributesValid);
+	TRACE_CPUPROFILER_EVENT_SCOPE(UPCGCustomHLSLKernel::PerformStaticValidation);
 
-	// The context can either be a compute graph element context (if the compute graph was successfully created), otherwise
-	// it will be the original CPU node context. We need the former to run the following validation.
-	if (!InContext || !InContext->IsComputeContext())
+	if (!Super::PerformStaticValidation())
 	{
-		return true;
+		return false;
 	}
 
-	const FPCGComputeGraphContext* Context = static_cast<FPCGComputeGraphContext*>(InContext);
-	const UPCGDataBinding* DataBinding = Context->DataBinding.Get();
+	const UPCGCustomHLSLSettings* CustomHLSLSettings = CastChecked<UPCGCustomHLSLSettings>(Settings);
 
-	if (DataBinding)
+	if (CustomHLSLSettings->OutputPins.IsEmpty())
 	{
-		TMap<FName, FPCGDataCollectionDesc> InputPinDescs;
-
-		for (const FPCGPinProperties& Pin : Settings->InputPins)
-		{
-			FPCGDataCollectionDesc DataDesc;
-			ensure(DataBinding->GetKernelPinDataDesc_GameThread(this, Pin.Label, /*bIsInputPin=*/true, DataDesc));
-			InputPinDescs.Add(Pin.Label, MoveTemp(DataDesc));
-		}
-
-		TMap<FName, FPCGDataCollectionDesc> OutputPinDescs;
+#if PCG_KERNEL_LOGGING_ENABLED
+		StaticLogEntries.Emplace(LOCTEXT("NoOutputs", "Custom HLSL nodes must have at least one output."), EPCGKernelLogVerbosity::Error);
+#endif
+		return false;
+	}
 
-		for (const FPCGPinProperties& Pin : Settings->OutputPins)
+	auto CheckPinLabel = [Settings=CustomHLSLSettings, &StaticLogEntries=StaticLogEntries](FName PinLabel)
+	{
+		if (PinLabel == NAME_None)
 		{
-			FPCGDataCollectionDesc DataDesc;
-			ensure(DataBinding->GetKernelPinDataDesc_GameThread(this, Pin.Label, /*bIsInputPin=*/false, DataDesc));
-			OutputPinDescs.Add(Pin.Label, MoveTemp(DataDesc));
+#if PCG_KERNEL_LOGGING_ENABLED
+			StaticLogEntries.Emplace(LOCTEXT("InvalidPinLabelNone", "Pin label 'None' is not a valid pin label."), EPCGKernelLogVerbosity::Error);
+#endif
+			return false;
 		}
 
-		const FString Source = Settings->ShaderFunctions + Settings->ShaderSource;
-		FRegexMatcher ModuleMatcher(FRegexPattern(PCGHLSLElement::AttributeFunctionPattern), Source);
+		bool bFoundPinLabel = false;
 
-		while (ModuleMatcher.FindNext())
+		auto IsAlreadyFound = [Settings, &StaticLogEntries, PinLabel, &bFoundPinLabel](const FPCGPinProperties PinProps)
 		{
-			const FString PinStr = ModuleMatcher.GetCaptureGroup(PCGHLSLElement::AttributePinCaptureGroup);
-			const FString FuncStr = ModuleMatcher.GetCaptureGroup(PCGHLSLElement::AttributeFunctionCaptureGroup);
-			const FString TypeStr = ModuleMatcher.GetCaptureGroup(PCGHLSLElement::AttributeTypeCaptureGroup);
-			const FString NameStr = ModuleMatcher.GetCaptureGroup(PCGHLSLElement::AttributeNameCaptureGroup);
-
-			const FName PinName = FName(*PinStr);
-			const FPCGDataCollectionDesc* PinDesc = nullptr;
-
-			auto ConstructFunctionText = [&PinStr, &FuncStr, &TypeStr]()
-				{
-					return FText::FromString(PinStr + TEXT("_") + FuncStr + TypeStr);
-				};
-
-			if (FuncStr == PCGCustomHLSLKernel::AttributeFunctionSet)
+			if (PinProps.Label == PinLabel)
 			{
-				PinDesc = OutputPinDescs.Find(FName(*PinStr));
-
-				if (!PinDesc && InputPinDescs.Find(PinName))
+				if (bFoundPinLabel)
 				{
 #if PCG_KERNEL_LOGGING_ENABLED
-					if (OutErrorText)
-					{
-						*OutErrorText = FText::Format(
-							LOCTEXT("InvalidSetAttributeUsage", "Tried to call attribute function '{0}' on read-only input pin '{1}'."),
-							ConstructFunctionText(),
-							FText::FromName(PinName));
-					}
+					StaticLogEntries.Emplace(FText::Format(LOCTEXT("DuplicatedPinLabels", "Duplicate pin label '{0}', all labels must be unique."), FText::FromName(PinLabel)), EPCGKernelLogVerbosity::Error);
 #endif
-
-					return false;
+					return true;
 				}
+
+				bFoundPinLabel = true;
 			}
-			else if (ensure(FuncStr == PCGCustomHLSLKernel::AttributeFunctionGet))
+
+			return false;
+		};
+
+		for (const FPCGPinProperties& PinProps : Settings->InputPins)
+		{
+			if (IsAlreadyFound(PinProps))
 			{
-				PinDesc = InputPinDescs.Find(FName(*PinStr));
+				return false;
 			}
+		}
 
-			if (!PinDesc)
+		for (const FPCGPinProperties& PinProps : Settings->OutputPins)
+		{
+			if (IsAlreadyFound(PinProps))
 			{
-#if PCG_KERNEL_LOGGING_ENABLED
-				if (OutErrorText)
-				{
-					*OutErrorText = FText::Format(
-						LOCTEXT("InvalidAttributePinName", "Tried to call attribute function '{0}' on non-existent pin '{1}'."),
-						ConstructFunctionText(),
-						FText::FromName(PinName));
-				}
-#endif
-
 				return false;
 			}
+		}
 
-			const UEnum* AttributeTypeEnum = StaticEnum<EPCGKernelAttributeType>();
-			check(AttributeTypeEnum);
+		return true;
+	};
 
-			const int64 AttrType = AttributeTypeEnum->GetValueByName(FName(*TypeStr));
+	// Validate input pins
+	bool bIsFirstInputPin = true;
+	for (const FPCGPinProperties& Properties : CustomHLSLSettings->InputPins)
+	{
+		if (!CheckPinLabel(Properties.Label))
+		{
+			return false;
+		}
 
-			if (AttrType == INDEX_NONE)
+		if (bIsFirstInputPin && CustomHLSLSettings->KernelType == EPCGKernelType::PointProcessor)
+		{
+			if (Properties.AllowedTypes != EPCGDataType::Point)
 			{
 #if PCG_KERNEL_LOGGING_ENABLED
-				if (OutErrorText)
-				{
-					*OutErrorText = FText::Format(
-						LOCTEXT("InvalidAttributePinType", "Tried to call attribute function '{0}' on non-existent type '{1}'."),
-						ConstructFunctionText(),
-						FText::FromString(TypeStr));
-				}
+				StaticLogEntries.Emplace(FText::Format(
+					LOCTEXT("InvalidNonPointPrimaryInput", "'Point Processor' nodes require primary input pin to be of type 'Point', but found '{0}'."),
+					FText::FromString(PCGCustomHLSLKernel::GetDataTypeString(Properties.AllowedTypes))),
+					EPCGKernelLogVerbosity::Error);
 #endif
-
 				return false;
 			}
+		}
 
-			const FPCGKernelAttributeDesc* AttrDesc = nullptr;
-			const FName AttrName = FName(*NameStr);
-			bool bFoundMatchingAttributeName = false;
+		if (!PCGComputeHelpers::IsTypeAllowedAsInput(Properties.AllowedTypes))
+		{
+#if PCG_KERNEL_LOGGING_ENABLED
+			StaticLogEntries.Emplace(FText::Format(
+				LOCTEXT("InvalidInputType", "Unsupported input type '{0}', found on pin '{1}'."),
+				FText::FromString(PCGCustomHLSLKernel::GetDataTypeString(Properties.AllowedTypes)),
+				FText::FromName(Properties.Label)),
+				EPCGKernelLogVerbosity::Error);
+#endif
+			return false;
+		}
 
-			// Verify that the attribute exists on at least one data in pin data collection.
-			for (const FPCGDataDesc& DataDesc : PinDesc->DataDescs)
-			{
-				AttrDesc = DataDesc.AttributeDescs.FindByPredicate([AttrName, AttrType, &bFoundMatchingAttributeName](const FPCGKernelAttributeDesc& Desc)
-				{
-					const bool bAttributeNameMatches = Desc.Name == AttrName;
-					bFoundMatchingAttributeName |= bAttributeNameMatches;
+		bIsFirstInputPin = false;
+	}
 
-					return bAttributeNameMatches && Desc.Type == static_cast<EPCGKernelAttributeType>(AttrType);
-				});
+	// Validate output pins
+	bool bIsFirstOutputPin = true;
+	for (const FPCGPinPropertiesGPU& Properties : CustomHLSLSettings->OutputPins)
+	{
+		if (!CheckPinLabel(Properties.Label))
+		{
+			return false;
+		}
 
-				if (AttrDesc)
-				{
-					break;
-				}
-			}
+		const bool bPinIsDefinedByKernel = bIsFirstOutputPin && (CustomHLSLSettings->KernelType == EPCGKernelType::PointGenerator || CustomHLSLSettings->KernelType == EPCGKernelType::PointProcessor);
 
-			if (!AttrDesc)
+		if (bPinIsDefinedByKernel)
+		{
+			if (Properties.AllowedTypes != EPCGDataType::Point)
 			{
 #if PCG_KERNEL_LOGGING_ENABLED
-				if (OutErrorText)
-				{
-					if (bFoundMatchingAttributeName)
-					{
-						*OutErrorText = FText::Format(
-							LOCTEXT("InvalidAttributeType", "Tried to call attribute function '{0}' on attribute '{1}' which is not of type '{2}'."),
-							ConstructFunctionText(),
-							FText::FromName(AttrName),
-							FText::FromString(TypeStr));
-					}
-					else
-					{
-						*OutErrorText = FText::Format(
-							LOCTEXT("InvalidAttributeDNE", "Tried to call attribute function '{0}' on attribute '{1}' which does not exist."),
-							ConstructFunctionText(),
-							FText::FromName(AttrName),
-							FText::FromString(TypeStr));
-					}
-				}
+				StaticLogEntries.Emplace(FText::Format(
+					LOCTEXT("InvalidNonPointPrimaryOutput", "'Point Processor' and 'Point Generator' nodes require primary output pin to be of type 'Point', but found '{0}'."),
+					FText::FromString(PCGCustomHLSLKernel::GetDataTypeString(Properties.AllowedTypes))),
+					EPCGKernelLogVerbosity::Error);
 #endif
 				return false;
 			}
 		}
-	}
-
-	return true;
-}
-
-bool UPCGCustomHLSLKernel::AreAllOutputPinsWritten(FText* OutWarningText) const
-{
-	TRACE_CPUPROFILER_EVENT_SCOPE(UPCGCustomHLSLKernel::AreAllOutputPinsWritten);
-
-	const FString AllSource = Settings->ShaderFunctions + Settings->ShaderSource;
-
-	// Non-custom kernels initialize the first output pin data automatically.
-	if (Settings->KernelType != EPCGKernelType::Custom && !Settings->OutputPins.IsEmpty())
-	{
-		const FName AutomaticallyInitializedPin = Settings->OutputPins[0].Label;
 
-		for (const FPCGPinPropertiesGPU& Pin : Settings->OutputPins)
+		if (!PCGComputeHelpers::IsTypeAllowedAsOutput(Properties.AllowedTypes))
 		{
-			if (Pin.Label == AutomaticallyInitializedPin || Pin.Label == NAME_None)
-			{
-				continue;
-			}
-
-			const TStaticArray<FString, 3> WriteStrings = {
-				Pin.Label.ToString() + TEXT("_Set"),
-				Pin.Label.ToString() + TEXT("_CopyElementFrom_"),
-				Pin.Label.ToString() + TEXT("_Initialize"),
-			};
+#if PCG_KERNEL_LOGGING_ENABLED
+			StaticLogEntries.Emplace(FText::Format(
+				LOCTEXT("InvalidOutputType", "Unsupported output type '{0}', found on pin '{1}'."),
+				FText::FromString(PCGCustomHLSLKernel::GetDataTypeString(Properties.AllowedTypes)),
+				FText::FromName(Properties.Label)),
+				EPCGKernelLogVerbosity::Error);
+#endif
+			return false;
+		}
 
-			bool bFoundDataWrite = false;
+		if (!bPinIsDefinedByKernel)
+		{
+			const FPCGPinPropertiesGPUStruct& Props = Properties.PropertiesGPU;
 
-			for (const FString& WriteString : WriteStrings)
+			if (Props.InitializationMode == EPCGPinInitMode::FromInputPins)
 			{
-				if (AllSource.Contains(WriteString))
+				if (Props.PinsToInititalizeFrom.IsEmpty())
 				{
-					bFoundDataWrite = true;
-					break;
+#if PCG_KERNEL_LOGGING_ENABLED
+					StaticLogEntries.Emplace(FText::Format(
+						LOCTEXT("InitFromEmptyPins", "Output pin '{0}' tried to initialize from input pins, but no pins were specified."),
+						FText::FromName(Properties.Label)),
+						EPCGKernelLogVerbosity::Error);
+#endif
+					return false;
 				}
-			}
 
-			if (!bFoundDataWrite)
-			{
-				if (OutWarningText)
+				for (const FName InitPinName : Props.PinsToInititalizeFrom)
 				{
-					*OutWarningText = FText::Format(
-						LOCTEXT("PinMayNotBeWritten", "Data on pin '{0}' may be uninitialized. Add code to write to this data, or mute this error in the node settings."),
-						{ FText::FromName(Pin.Label) });
+					const FPCGPinProperties* InitPinProps = CustomHLSLSettings->InputPins.FindByPredicate([InitPinName](const FPCGPinProperties& InPinProps)
+					{
+						return InPinProps.Label == InitPinName;
+					});
+
+					if (InitPinProps)
+					{
+						if (!PCGComputeHelpers::IsTypeAllowedInDataCollection(InitPinProps->AllowedTypes))
+						{
+#if PCG_KERNEL_LOGGING_ENABLED
+							StaticLogEntries.Emplace(FText::Format(
+								LOCTEXT("InitFromInvalidPinType", "Output pin '{0}' tried to initialize from input pin '{1}', but pin '{1}' has an invalid type."),
+								FText::FromName(Properties.Label),
+								FText::FromName(InitPinName)),
+								EPCGKernelLogVerbosity::Error);
+#endif
+							return false;
+						}
+					}
+					else
+					{
+#if PCG_KERNEL_LOGGING_ENABLED
+						StaticLogEntries.Emplace(FText::Format(
+							LOCTEXT("InitFromNonExistentPin", "Output pin '{0}' tried to initialize from non-existent input pin '{1}'."),
+							FText::FromName(Properties.Label),
+							FText::FromName(InitPinName)),
+							EPCGKernelLogVerbosity::Error);
+#endif
+						return false;
+					}
+				}
+
+				// TODO: Could do validation on data multiplicity for Pairwise, checking that data counts are 1 or N, but maybe that should be a runtime error instead.
+			}
+
+			const bool bUsingFixedDataCount = Props.InitializationMode == EPCGPinInitMode::Custom || Props.DataCountMode == EPCGDataCountMode::Fixed;
+
+			if (bUsingFixedDataCount)
+			{
+				if (Props.DataCount < 1)
+				{
+#if PCG_KERNEL_LOGGING_ENABLED
+					StaticLogEntries.Emplace(FText::Format(
+						LOCTEXT("InvalidDataCount", "Invalid fixed data count {0} on output pin '{1}'. Must be greater than 0."),
+						FText::AsNumber(Props.DataCount),
+						FText::FromName(Properties.Label)),
+						EPCGKernelLogVerbosity::Error);
+#endif
+					return false;
 				}
+			}
 
+			const bool bUsingFixedElemCount = Props.InitializationMode == EPCGPinInitMode::Custom || Props.ElementCountMode == EPCGElementCountMode::Fixed;
+
+			if (bUsingFixedElemCount)
+			{
+				if (Props.ElementCount < 1)
+				{
+#if PCG_KERNEL_LOGGING_ENABLED
+					StaticLogEntries.Emplace(FText::Format(
+						LOCTEXT("InvalidElementCount", "Invalid fixed element count {0} on output pin '{1}'. Must be greater than 0."),
+						FText::AsNumber(Props.ElementCount),
+						FText::FromName(Properties.Label)),
+						EPCGKernelLogVerbosity::Error);
+#endif
+					return false;
+				}
+			}
+
+			if (Props.ElementCountMultiplier < 1)
+			{
+#if PCG_KERNEL_LOGGING_ENABLED
+				StaticLogEntries.Emplace(FText::Format(
+					LOCTEXT("InvalidElementCountMultiplier", "Invalid element count multiplier {0} on output pin '{1}'. Must be greater than 0."),
+					FText::AsNumber(Props.ElementCountMultiplier),
+					FText::FromName(Properties.Label)),
+					EPCGKernelLogVerbosity::Error);
+#endif
 				return false;
 			}
 		}
+
+		bIsFirstOutputPin = false;
+	}
+
+	if (CustomHLSLSettings->KernelType == EPCGKernelType::Custom && CustomHLSLSettings->DispatchThreadCount == EPCGDispatchThreadCount::FromProductOfInputPins)
+	{
+		if (CustomHLSLSettings->ThreadCountInputPinLabels.IsEmpty())
+		{
+#if PCG_KERNEL_LOGGING_ENABLED
+			StaticLogEntries.Emplace(LOCTEXT("MissingThreadCountPins", "Dispatch thread count is based on input pins but no labels have been set in Input Pins array."), EPCGKernelLogVerbosity::Error);
+#endif
+			return false;
+		}
+
+		for (const FName& Label : CustomHLSLSettings->ThreadCountInputPinLabels)
+		{
+			if (!CustomHLSLSettings->InputPins.FindByPredicate([Label](const FPCGPinProperties& InProps) { return InProps.Label == Label; }))
+			{
+#if PCG_KERNEL_LOGGING_ENABLED
+				StaticLogEntries.Emplace(FText::Format(LOCTEXT("MissingThreadCountPin", "Invalid pin specified in Input Pins array: '{0}'."), FText::FromName(Label)), EPCGKernelLogVerbosity::Error);
+#endif
+				return false;
+			}
+		}
+	}
+
+	if (IsThreadCountMultiplierInUse())
+	{
+		if (CustomHLSLSettings->ThreadCountMultiplier < 1)
+		{
+#if PCG_KERNEL_LOGGING_ENABLED
+			StaticLogEntries.Emplace(FText::Format(LOCTEXT("InvalidThreadCountMultiplier", "Thread Count Multiplier has invalid value ({0}). Must be greater than 0."), CustomHLSLSettings->ThreadCountMultiplier), EPCGKernelLogVerbosity::Error);
+#endif
+			return false;
+		}
+	}
+
+	for (const FPCGKernelAttributeKey& AttributeKey : CustomHLSLSettings->KernelAttributeKeys)
+	{
+		if (AttributeKey.Type == EPCGKernelAttributeType::Invalid)
+		{
+			const UEnum* AttributeTypeEnum = StaticEnum<EPCGKernelAttributeType>();
+			check(AttributeTypeEnum);
+
+#if PCG_KERNEL_LOGGING_ENABLED
+			StaticLogEntries.Emplace(FText::Format(
+				LOCTEXT("InvalidAttributeTypeNone", "Attribute '{0}' has invalid GPU attribute type '{1}', check the 'Attributes to Create' array on your pins."),
+				FText::FromName(AttributeKey.Name),
+				FText::FromString(AttributeTypeEnum->GetNameStringByValue(static_cast<int64>(AttributeKey.Type)))),
+				EPCGKernelLogVerbosity::Error);
+#endif
+			return false;
+		}
+	}
+
+	if (!ValidateShaderSource(StaticLogEntries))
+	{
+		return false;
+	}
+
+	if (!CustomHLSLSettings->bMuteUnwrittenPinDataErrors && !AreAllOutputPinsWritten(StaticLogEntries))
+	{
+		return false;
 	}
 
 	return true;
 }
 
-bool UPCGCustomHLSLKernel::ProcessShaderSource(FString* InOutSource, FText* OutErrorText) const
+bool UPCGCustomHLSLKernel::ValidateShaderSource(TArray<FPCGKernelLogEntry>& OutLogEntries) const
 {
-	TSet<TPair<FString, FString>> PinsRequiringElementCopy;
+	auto ParseCopyElementFunctions = [Settings=Settings, &OutLogEntries](const TArray<FPCGParsedCopyElementFunction>& ParsedFunctions)
+	{
+		const UPCGCustomHLSLSettings* CustomHLSLSettings = CastChecked<UPCGCustomHLSLSettings>(Settings);
 
-	const FString Source = Settings->ShaderFunctions + TEXT("\n") + Settings->ShaderSource;
+		for (const FPCGParsedCopyElementFunction& ParsedFunction : ParsedFunctions)
+		{
+			const FPCGPinProperties* SourcePinProps = CustomHLSLSettings->InputPins.FindByPredicate([Label=FName(*ParsedFunction.SourcePin)](const FPCGPinProperties& InProps) { return InProps.Label == Label; });
+			const FPCGPinPropertiesGPU* TargetPinProps = CustomHLSLSettings->OutputPins.FindByPredicate([Label=FName(*ParsedFunction.TargetPin)](const FPCGPinPropertiesGPU& InProps) { return InProps.Label == Label; });
 
-	/** Regex for parsing usages of CopyElementFrom function. */
-	constexpr TCHAR CopyElementPattern[] = { TEXT(/*TargetPin*/"([a-zA-Z0-9]+)" "_CopyElementFrom_" /*SourcePin*/"([a-zA-Z0-9]+)\\s*\\(") };
-	constexpr int CopyElementTargetPinCaptureGroup = 1;
-	constexpr int CopyElementSourcePinCaptureGroup = 2;
+			if (!SourcePinProps)
+			{
+#if PCG_KERNEL_LOGGING_ENABLED
+				OutLogEntries.Emplace(FText::Format(LOCTEXT("InvalidCopySourcePin", "Invalid input pin '{0}' used in copy element expression."), FText::FromString(ParsedFunction.SourcePin)), EPCGKernelLogVerbosity::Error);
+#endif
+				return false;
+			}
 
-	FRegexMatcher ModuleMatcher(FRegexPattern(CopyElementPattern), Source);
+			if (!TargetPinProps)
+			{
+#if PCG_KERNEL_LOGGING_ENABLED
+				OutLogEntries.Emplace(FText::Format(LOCTEXT("InvalidCopyTargetPin", "Invalid output pin '{0}' used in copy element expression."), FText::FromString(ParsedFunction.TargetPin)), EPCGKernelLogVerbosity::Error);
+#endif
+				return false;
+			}
 
-	while (ModuleMatcher.FindNext())
+			if (SourcePinProps && TargetPinProps && SourcePinProps->AllowedTypes != TargetPinProps->AllowedTypes)
+			{
+#if PCG_KERNEL_LOGGING_ENABLED
+				OutLogEntries.Emplace(LOCTEXT("InvalidCopyPinTypes", "CopyFromElement only allowed across pins of matching type."), EPCGKernelLogVerbosity::Error);
+#endif
+				return false;
+			}
+		}
+
+		return true;
+	};
+
+	if (!ParseCopyElementFunctions(SourceParsedCopyElementFunctions))
+	{
+		return false;
+	}
+
+	if (!ParseCopyElementFunctions(FunctionsParsedCopyElementFunctions))
+	{
+		return false;
+	}
+
+	// @todo_pcg: Validation of parsed attribute functions could be done here?
+
+	return true;
+}
+
+void UPCGCustomHLSLKernel::ProcessShaderSource(const UPCGComputeGraph* InComputeGraph, FString& OutShaderSource, FString& OutShaderFunctions) const
+{
+	// We inject attribute IDs directly into the source. This is most efficient and saves us passing them into the kernel. However
+	// the trade off is that a shader with a shared source will generate multiple variants if used in different compute graphs with
+	// different attribute IDs.
+	auto ReplaceAttributeNamesWithIds = [InComputeGraph](FString& InOutSource, const TArray<FPCGParsedAttributeFunction>& ParsedAttributeFunctions)
 	{
-		const FString SourcePinStr = ModuleMatcher.GetCaptureGroup(CopyElementSourcePinCaptureGroup);
-		const FPCGPinProperties* SourcePinProps = Settings->InputPins.FindByPredicate([Label=FName(*SourcePinStr)](const FPCGPinProperties& InProps) { return InProps.Label == Label; });
-		if (!SourcePinProps)
+		TArray<TTuple<int32, int32, int32>> Replacements;
+
+		const UEnum* AttributeTypeEnum = StaticEnum<EPCGKernelAttributeType>();
+		check(AttributeTypeEnum);
+
+		for (const FPCGParsedAttributeFunction& ParsedFunction : ParsedAttributeFunctions)
 		{
-#if PCG_KERNEL_LOGGING_ENABLED
-			if (OutErrorText)
+			const int64 AttrType = ParsedFunction.AttributeType;
+
+			if (!ensure(AttrType != INDEX_NONE))
 			{
-				*OutErrorText = FText::Format(LOCTEXT("InvalidCopySourcePin", "Invalid input pin '{0}' used in copy element expression."), FText::FromString(SourcePinStr));
+				break;
 			}
+
+			const FPCGKernelAttributeKey AttributeKey(FName(*ParsedFunction.AttributeName), static_cast<EPCGKernelAttributeType>(AttrType));
+			const FString SourceDefinition = PCGCustomHLSLKernel::GetKernelAttributeKeyAsString(AttributeKey);
+
+			const int32 AttributeIndex = InComputeGraph->GetStaticAttributeTable().GetAttributeId(AttributeKey);
+			const int32 ReplaceStartIndex = InOutSource.Find(SourceDefinition, ESearchCase::CaseSensitive, ESearchDir::FromStart, ParsedFunction.MatchBeginning);
+			const int32 ReplaceEndIndex = ReplaceStartIndex + SourceDefinition.Len();
+
+			Replacements.Emplace(AttributeIndex, ReplaceStartIndex, ReplaceEndIndex);
+		}
+
+		for (int I = Replacements.Num() - 1; I >= 0; --I)
+		{
+			const int32 AttributeIndex = Replacements[I].Get<0>();
+			const int32 ReplaceStartIndex = Replacements[I].Get<1>();
+			const int32 ReplaceEndIndex = Replacements[I].Get<2>();
+
+			InOutSource = InOutSource.Left(ReplaceStartIndex) + FString::FromInt(AttributeIndex) + InOutSource.RightChop(ReplaceEndIndex);
+		}
+	};
+
+	// Implementation note: Attribute name replacement relies on precomputed indices into the source strings, therefore the attribute name replacement must take place
+	// before any other modifications. Otherwise, the indices will be incorrect and the source will become gibberish.
+	ReplaceAttributeNamesWithIds(OutShaderSource, SourceParsedAttributeFunctions);
+	ReplaceAttributeNamesWithIds(OutShaderFunctions, FunctionsParsedAttributeFunctions);
+
+#if PLATFORM_WINDOWS
+	// Remove old-school stuff.
+	OutShaderSource.ReplaceInline(TEXT("\r"), TEXT(""));
+	OutShaderFunctions.ReplaceInline(TEXT("\r"), TEXT(""));
 #endif
-			return false;
+
+	auto ReplaceCopyElementFrom = [](FString& OutSource, const TArray<FPCGParsedCopyElementFunction>& ParsedFunctions)
+	{
+		// Replace function calls like Out_CopyElementFrom_In(...) with macro PCG_COPY_ALL_ATTRIBUTES_TO_OUTPUT(Out, In, ...).
+		for (const FPCGParsedCopyElementFunction& ParsedFunction : ParsedFunctions)
+		{
+			OutSource.ReplaceInline(
+				*FString::Format(TEXT("{1}_CopyElementFrom_{0}("), { ParsedFunction.SourcePin, ParsedFunction.TargetPin }),
+				*FString::Format(TEXT("PCG_COPY_ALL_ATTRIBUTES_TO_OUTPUT({1}, {0}, "), { ParsedFunction.SourcePin, ParsedFunction.TargetPin })
+			);
 		}
+	};
 
-		const FString TargetPinStr = ModuleMatcher.GetCaptureGroup(CopyElementTargetPinCaptureGroup);
-		const FPCGPinPropertiesGPU* TargetPinProps = Settings->OutputPins.FindByPredicate([Label=FName(*TargetPinStr)](const FPCGPinPropertiesGPU& InProps) { return InProps.Label == Label; });
-		if (!TargetPinProps)
+	ReplaceCopyElementFrom(OutShaderSource, SourceParsedCopyElementFunctions);
+	ReplaceCopyElementFrom(OutShaderSource, FunctionsParsedCopyElementFunctions);
+
+	// We inject data IDs directly into the source. They will get remapped to data indices using a label resolver data interface.
+	if (const FPCGPinDataLabels* PinToDataLabels = InComputeGraph->GetStaticDataLabelsTable().Find(KernelIndex))
+	{
+		for (const TPair<FName, FPCGDataLabels>& Pair : PinToDataLabels->PinToDataLabels)
 		{
-#if PCG_KERNEL_LOGGING_ENABLED
-			if (OutErrorText)
+			const FName PinLabel = Pair.Key;
+			const TArray<FString>& DataLabels = Pair.Value.Labels;
+
+			for (int DataId = 0; DataId < DataLabels.Num(); ++DataId)
 			{
-				*OutErrorText = FText::Format(LOCTEXT("InvalidCopyTargetPin", "Invalid output pin '{0}' used in copy element expression."), FText::FromString(TargetPinStr));
+				const FString& DataLabel = DataLabels[DataId];
+				const FString ReplacementStr = FString::Format(TEXT("{0}_GetDataIndexFromIdInternal(/*DataId=*/{1})"),
+				{
+					PCGComputeHelpers::GetDataLabelResolverName(PinLabel),
+					FString::FromInt(DataId)
+				});
+
+				// Matches against {PinName}_AnyFunction('{DataLabel}'...
+				// First capture group is the data label, so that we can find & replace it by index.
+				const FString Pattern = FString::Format(TEXT("{0}_.*?[\\s]*?\\([\\s]*?('{1}')"), { *PinLabel.ToString(), DataLabel });
+
+				auto ReplaceDataLabels = [&Pattern, &ReplacementStr](FString& InOutSource)
+				{
+					FRegexMatcher ModuleMatcher(FRegexPattern(Pattern), InOutSource);
+					while (ModuleMatcher.FindNext())
+					{
+						const int32 DataLabelCharIndexStart = ModuleMatcher.GetCaptureGroupBeginning(1);
+						const int32 DataLabelCharIndexEnd = ModuleMatcher.GetCaptureGroupEnding(1);
+
+						InOutSource = InOutSource.Left(DataLabelCharIndexStart) + ReplacementStr + InOutSource.RightChop(DataLabelCharIndexEnd);
+					}
+				};
+
+				ReplaceDataLabels(OutShaderSource);
+				ReplaceDataLabels(OutShaderFunctions);
 			}
-#endif
-			return false;
 		}
+	}
+}
 
-		if (SourcePinProps->AllowedTypes != TargetPinProps->AllowedTypes)
+bool UPCGCustomHLSLKernel::AreAllOutputPinsWritten(TArray<FPCGKernelLogEntry>& OutLogEntries) const
+{
+	TRACE_CPUPROFILER_EVENT_SCOPE(UPCGCustomHLSLKernel::AreAllOutputPinsWritten);
+
+	const UPCGCustomHLSLSettings* CustomHLSLSettings = CastChecked<UPCGCustomHLSLSettings>(Settings);
+	const FString AllSource = CustomHLSLSettings->ShaderFunctions + CustomHLSLSettings->ShaderSource;
+
+	// Non-custom kernels initialize the first output pin data automatically.
+	if (CustomHLSLSettings->KernelType != EPCGKernelType::Custom && !CustomHLSLSettings->OutputPins.IsEmpty())
+	{
+		const FName AutomaticallyInitializedPin = CustomHLSLSettings->OutputPins[0].Label;
+
+		for (const FPCGPinPropertiesGPU& Pin : CustomHLSLSettings->OutputPins)
 		{
-#if PCG_KERNEL_LOGGING_ENABLED
-			if (OutErrorText)
+			if (Pin.Label == AutomaticallyInitializedPin || Pin.Label == NAME_None)
 			{
-				*OutErrorText = LOCTEXT("InvalidCopyPinTypes", "CopyFromElement only allowed across pins of matching type.");
+				continue;
 			}
+
+			const TStaticArray<FString, 3> WriteStrings = {
+				Pin.Label.ToString() + TEXT("_Set"),
+				Pin.Label.ToString() + TEXT("_CopyElementFrom_"),
+				Pin.Label.ToString() + TEXT("_Initialize"),
+			};
+
+			bool bFoundDataWrite = false;
+
+			for (const FString& WriteString : WriteStrings)
+			{
+				if (AllSource.Contains(WriteString))
+				{
+					bFoundDataWrite = true;
+					break;
+				}
+			}
+
+			if (!bFoundDataWrite)
+			{
+#if PCG_KERNEL_LOGGING_ENABLED
+				OutLogEntries.Emplace(FText::Format(
+					LOCTEXT("PinMayNotBeWritten", "Data on pin '{0}' may be uninitialized. Add code to write to this data, or mute this error in the node settings."),
+					{ FText::FromName(Pin.Label) }),
+					EPCGKernelLogVerbosity::Error);
 #endif
-			return false;
+				return false;
+			}
 		}
-
-		PinsRequiringElementCopy.Emplace({ SourcePinStr, TargetPinStr });
 	}
 
-	if (!InOutSource)
+	return true;
+}
+#endif // WITH_EDITOR
+
+bool UPCGCustomHLSLKernel::IsThreadCountMultiplierInUse() const
+{
+	const UPCGCustomHLSLSettings* CustomHLSLSettings = CastChecked<UPCGCustomHLSLSettings>(Settings);
+	return CustomHLSLSettings->KernelType == EPCGKernelType::Custom && CustomHLSLSettings->DispatchThreadCount != EPCGDispatchThreadCount::Fixed;
+}
+
+bool UPCGCustomHLSLKernel::AreAttributesValid(FPCGContext* InContext, FText* OutErrorText) const
+{
+	TRACE_CPUPROFILER_EVENT_SCOPE(UPCGCustomHLSLKernel::AreAttributesValid);
+
+	// The context can either be a compute graph element context (if the compute graph was successfully created), otherwise
+	// it will be the original CPU node context. We need the former to run the following validation.
+	if (!InContext || !InContext->IsComputeContext())
 	{
 		return true;
 	}
 
-	// Replace function calls like Out_CopyElementFrom_In(...) with macro PCG_COPY_ALL_ATTRIBUTES_TO_OUTPUT(Out, In, ...).
-	for (const TPair<FString, FString>& SourcePinTargetPin : PinsRequiringElementCopy)
+	const UPCGCustomHLSLSettings* CustomHLSLSettings = CastChecked<UPCGCustomHLSLSettings>(Settings);
+	const FPCGComputeGraphContext* Context = static_cast<FPCGComputeGraphContext*>(InContext);
+	const UPCGDataBinding* DataBinding = Context->DataBinding.Get();
+
+	if (DataBinding)
 	{
-		InOutSource->ReplaceInline(
-			*FString::Format(TEXT("{1}_CopyElementFrom_{0}("), { SourcePinTargetPin.Key, SourcePinTargetPin.Value }),
-			*FString::Format(TEXT("PCG_COPY_ALL_ATTRIBUTES_TO_OUTPUT({1}, {0}, "), { SourcePinTargetPin.Key, SourcePinTargetPin.Value })
-		);
+		TMap<FName, FPCGDataCollectionDesc> InputPinDescs;
+		TMap<FName, FPCGDataCollectionDesc> OutputPinDescs;
+
+		{
+			TRACE_CPUPROFILER_EVENT_SCOPE(UPCGCustomHLSLKernel::AreAttributesValid::GetPinDataDescs);
+
+			for (const FPCGPinProperties& Pin : CustomHLSLSettings->InputPins)
+			{
+				FPCGDataCollectionDesc DataDesc;
+				ensure(DataBinding->GetKernelPinDataDesc_GameThread(this, Pin.Label, /*bIsInputPin=*/true, DataDesc));
+				InputPinDescs.Add(Pin.Label, MoveTemp(DataDesc));
+			}
+
+			for (const FPCGPinProperties& Pin : CustomHLSLSettings->OutputPins)
+			{
+				FPCGDataCollectionDesc DataDesc;
+				ensure(DataBinding->GetKernelPinDataDesc_GameThread(this, Pin.Label, /*bIsInputPin=*/false, DataDesc));
+				OutputPinDescs.Add(Pin.Label, MoveTemp(DataDesc));
+			}
+		}
+
+		auto ValidateParsedAttributeFunctions = [&InputPinDescs, &OutputPinDescs, OutErrorText](const TArray<FPCGParsedAttributeFunction>& ParsedAttributeFunctions)
+		{
+			for (const auto& ParsedFunction : ParsedAttributeFunctions)
+			{
+				const UEnum* AttributeTypeEnum = StaticEnum<EPCGKernelAttributeType>();
+				check(AttributeTypeEnum);
+
+				const FString& PinLabelStr = ParsedFunction.PinLabel;
+				const FString& FunctionName = ParsedFunction.FunctionName;
+				const FString& AttributeName = ParsedFunction.AttributeName;
+				const FString TypeStr = AttributeTypeEnum->GetNameStringByValue(ParsedFunction.AttributeType);
+
+				const FName PinLabel = FName(*PinLabelStr);
+				const FPCGDataCollectionDesc* PinDesc = nullptr;
+
+				auto ConstructFunctionText = [&ParsedFunction, &TypeStr]()
+				{
+					return FText::FromString(ParsedFunction.PinLabel + TEXT("_") + ParsedFunction.FunctionName + TypeStr);
+				};
+
+				if (FunctionName == PCGCustomHLSLKernel::AttributeFunctionSet)
+				{
+					PinDesc = OutputPinDescs.Find(PinLabel);
+
+					if (!PinDesc && InputPinDescs.Find(PinLabel))
+					{
+#if PCG_KERNEL_LOGGING_ENABLED
+						if (OutErrorText)
+						{
+							*OutErrorText = FText::Format(
+								LOCTEXT("InvalidSetAttributeUsage", "Tried to call attribute function '{0}' on read-only input pin '{1}'."),
+								ConstructFunctionText(),
+								FText::FromName(PinLabel));
+						}
+#endif
+
+						return false;
+					}
+				}
+				else if (ensure(FunctionName == PCGCustomHLSLKernel::AttributeFunctionGet))
+				{
+					PinDesc = InputPinDescs.Find(PinLabel);
+				}
+
+				if (!PinDesc)
+				{
+#if PCG_KERNEL_LOGGING_ENABLED
+					if (OutErrorText)
+					{
+						*OutErrorText = FText::Format(
+							LOCTEXT("InvalidAttributePinName", "Tried to call attribute function '{0}' on non-existent pin '{1}'."),
+							ConstructFunctionText(),
+							FText::FromName(PinLabel));
+					}
+#endif
+
+					return false;
+				}
+
+
+				const int64 AttrType = AttributeTypeEnum->GetValueByName(FName(*TypeStr));
+
+				if (AttrType == INDEX_NONE)
+				{
+#if PCG_KERNEL_LOGGING_ENABLED
+					if (OutErrorText)
+					{
+						*OutErrorText = FText::Format(
+							LOCTEXT("InvalidAttributePinType", "Tried to call attribute function '{0}' on non-existent type '{1}'."),
+							ConstructFunctionText(),
+							FText::FromString(TypeStr));
+					}
+#endif
+
+					return false;
+				}
+
+				const FPCGKernelAttributeDesc* AttrDesc = nullptr;
+				const FName AttrName = FName(*AttributeName);
+				bool bFoundMatchingAttributeName = false;
+
+				// Verify that the attribute exists on at least one data in pin data collection.
+				for (const FPCGDataDesc& DataDesc : PinDesc->DataDescs)
+				{
+					AttrDesc = DataDesc.AttributeDescs.FindByPredicate([AttrName, AttrType, &bFoundMatchingAttributeName](const FPCGKernelAttributeDesc& Desc)
+					{
+						const bool bAttributeNameMatches = Desc.Name == AttrName;
+						bFoundMatchingAttributeName |= bAttributeNameMatches;
+
+						return bAttributeNameMatches && Desc.Type == static_cast<EPCGKernelAttributeType>(AttrType);
+					});
+
+					if (AttrDesc)
+					{
+						break;
+					}
+				}
+
+				if (!AttrDesc)
+				{
+#if PCG_KERNEL_LOGGING_ENABLED
+					if (OutErrorText)
+					{
+						if (bFoundMatchingAttributeName)
+						{
+							*OutErrorText = FText::Format(
+								LOCTEXT("InvalidAttributeType", "Tried to call attribute function '{0}' on attribute '{1}' which is not of type '{2}'."),
+								ConstructFunctionText(),
+								FText::FromName(AttrName),
+								FText::FromString(TypeStr));
+						}
+						else
+						{
+							*OutErrorText = FText::Format(
+								LOCTEXT("InvalidAttributeDNE", "Tried to call attribute function '{0}' on attribute '{1}' which does not exist."),
+								ConstructFunctionText(),
+								FText::FromName(AttrName),
+								FText::FromString(TypeStr));
+						}
+					}
+#endif
+					return false;
+				}
+			}
+
+			return true;
+		};
+
+		if (!ValidateParsedAttributeFunctions(SourceParsedAttributeFunctions))
+		{
+			return false;
+		}
+
+		if (!ValidateParsedAttributeFunctions(FunctionsParsedAttributeFunctions))
+		{
+			return false;
+		}
 	}
 
 	return true;
@@ -1218,13 +1322,15 @@ bool UPCGCustomHLSLKernel::ProcessShaderSource(FString* InOutSource, FText* OutE
 const FPCGPinProperties* UPCGCustomHLSLKernel::GetFirstPointInputPin() const
 {
 	// Point processing input pin is first pin of type Point.
-	return Settings->InputPins.FindByPredicate([](const FPCGPinProperties& InProps) { return InProps.AllowedTypes == EPCGDataType::Point; });
+	const UPCGCustomHLSLSettings* CustomHLSLSettings = CastChecked<UPCGCustomHLSLSettings>(Settings);
+	return CustomHLSLSettings->InputPins.FindByPredicate([](const FPCGPinProperties& InProps) { return InProps.AllowedTypes == EPCGDataType::Point; });
 }
 
 const FPCGPinPropertiesGPU* UPCGCustomHLSLKernel::GetFirstPointOutputPin() const
 {
 	// Point processing output pin is first pin of type Point.
-	return Settings->OutputPins.FindByPredicate([](const FPCGPinPropertiesGPU& InProps) { return InProps.AllowedTypes == EPCGDataType::Point; });
+	const UPCGCustomHLSLSettings* CustomHLSLSettings = CastChecked<UPCGCustomHLSLSettings>(Settings);
+	return CustomHLSLSettings->OutputPins.FindByPredicate([](const FPCGPinPropertiesGPU& InProps) { return InProps.AllowedTypes == EPCGDataType::Point; });
 }
 
 int UPCGCustomHLSLKernel::GetElementCountForInputPin(const FPCGPinProperties& InInputPinProps, const UPCGDataBinding* InBinding) const
diff --git a/Engine/Plugins/PCG/Source/PCG/Private/Compute/Elements/PCGCustomHLSLKernel.h b/Engine/Plugins/PCG/Source/PCG/Private/Compute/Elements/PCGCustomHLSLKernel.h
--- a/Engine/Plugins/PCG/Source/PCG/Private/Compute/Elements/PCGCustomHLSLKernel.h
+++ b/Engine/Plugins/PCG/Source/PCG/Private/Compute/Elements/PCGCustomHLSLKernel.h
@@ -10,6 +10,41 @@
 
 class UPCGCustomHLSLSettings;
 
+USTRUCT()
+struct FPCGParsedAttributeFunction
+{
+	GENERATED_BODY()
+
+public:
+	UPROPERTY()
+	FString PinLabel;
+
+	UPROPERTY()
+	FString FunctionName;
+
+	UPROPERTY()
+	int64 AttributeType = INDEX_NONE;
+
+	UPROPERTY()
+	FString AttributeName;
+
+	UPROPERTY()
+	int32 MatchBeginning = INDEX_NONE;
+};
+
+USTRUCT()
+struct FPCGParsedCopyElementFunction
+{
+	GENERATED_BODY()
+
+public:
+	UPROPERTY()
+	FString SourcePin;
+
+	UPROPERTY()
+	FString TargetPin;
+};
+
 UCLASS()
 class UPCGCustomHLSLKernel : public UPCGComputeKernel
 {
@@ -17,12 +52,11 @@ class UPCGCustomHLSLKernel : public UPCGComputeKernel
 
 public:
 #if WITH_EDITOR
-	/** Initialize kernel. Editor only as shader compilation only available in editor. */
-	void Initialize(const UPCGCustomHLSLSettings* InSettings, bool bLogDescriptions);
+	virtual void Initialize(const UPCGSettings* InSettings, bool bLogDescriptions) override;
 #endif
 
 	virtual const UPCGSettings* GetSettings() const override;
-	virtual bool IsKernelValid(FPCGContext* InContext = nullptr, bool bQuiet = true) const override;
+	virtual bool IsKernelValid(FPCGContext* InContext) const override;
 	virtual bool ComputeOutputBindingDataDesc(const UPCGComputeGraph* InGraph, FName InOutputPinLabel, const UPCGDataBinding* InBinding, FPCGDataCollectionDesc& OutDataDesc) const override;
 #if WITH_EDITOR
 	virtual FString GetCookedSource(const UPCGComputeGraph* InComputeGraph) const override;
@@ -39,30 +73,44 @@ public:
 protected:
 #if WITH_EDITOR
 	void InitEntryPoint();
+	virtual bool PerformStaticValidation() override;
+	bool ValidateShaderSource(TArray<FPCGKernelLogEntry>& OutLogEntries) const;
+	bool AreAllOutputPinsWritten(TArray<FPCGKernelLogEntry>& OutLogEntries) const;
+
+	/** Apply any code gen/transformations/expansions to the shader source. */
+	virtual void ProcessShaderSource(const UPCGComputeGraph* InComputeGraph, FString& OutShaderSource, FString& OutShaderFunctions) const;
 #endif
 
 	/** Will the ThreadCountMultiplier value be applied when calculating the dispatch thread count. */
 	bool IsThreadCountMultiplierInUse() const;
-
 	bool AreAttributesValid(FPCGContext* InContext, FText* OutErrorText) const;
 
-	bool AreAllOutputPinsWritten(FText* OutErrorText) const;
-
-	/** Apply any code gen or transformations or expansions to the shader source. */
-	virtual bool ProcessShaderSource(FString* InOutSource, FText* OutErrorText = nullptr) const;
-
 	const FPCGPinProperties* GetFirstPointInputPin() const;
 	const FPCGPinPropertiesGPU* GetFirstPointOutputPin() const;
 
 	int GetElementCountForInputPin(const FPCGPinProperties& InInputPinProps, const UPCGDataBinding* InBinding) const;
 
 public:
-	UPROPERTY()
-	TObjectPtr<const UPCGCustomHLSLSettings> Settings;
-
 #if WITH_EDITORONLY_DATA
 	/** The name of the main function in the shader. Generated from the node title. */
 	UPROPERTY()
 	FString EntryPoint;
 #endif
+
+protected:
+	/** Attribute functions parsed and cached from HLSL shader source. */
+	UPROPERTY()
+	TArray<FPCGParsedAttributeFunction> SourceParsedAttributeFunctions;
+
+	/** Attribute functions parsed and cached from HLSL shader functions source. */
+	UPROPERTY()
+	TArray<FPCGParsedAttributeFunction> FunctionsParsedAttributeFunctions;
+
+	/** CopyElementFrom functions parsed and cached from HLSL shader source. */
+	UPROPERTY()
+	TArray<FPCGParsedCopyElementFunction> SourceParsedCopyElementFunctions;
+
+	/** CopyElementFrom functions parsed and cached from HLSL shader functions source. */
+	UPROPERTY()
+	TArray<FPCGParsedCopyElementFunction> FunctionsParsedCopyElementFunctions;
 };
diff --git a/Engine/Plugins/PCG/Source/PCG/Private/Compute/PCGComputeGraph.cpp b/Engine/Plugins/PCG/Source/PCG/Private/Compute/PCGComputeGraph.cpp
--- a/Engine/Plugins/PCG/Source/PCG/Private/Compute/PCGComputeGraph.cpp
+++ b/Engine/Plugins/PCG/Source/PCG/Private/Compute/PCGComputeGraph.cpp
@@ -176,14 +176,14 @@ int32 UPCGComputeGraph::GetBindingIndex(int32 InKernelIndex, FName InPinLabel, b
 	}
 }
 
-bool UPCGComputeGraph::IsGraphValid(FPCGContext* InContext, bool bQuiet) const
+bool UPCGComputeGraph::IsGraphValid(FPCGContext* InContext) const
 {
 	bool bAllKernelsValid = true;
 
 	// Run validation on all kernels.
 	for (const UComputeKernel* Kernel : KernelInvocations)
 	{
-		bAllKernelsValid &= CastChecked<UPCGComputeKernel>(Kernel)->IsKernelValid(InContext, bQuiet);
+		bAllKernelsValid &= CastChecked<UPCGComputeKernel>(Kernel)->IsKernelValid(InContext);
 	}
 
 	return bAllKernelsValid;
diff --git a/Engine/Plugins/PCG/Source/PCG/Private/Compute/PCGComputeGraph.h b/Engine/Plugins/PCG/Source/PCG/Private/Compute/PCGComputeGraph.h
--- a/Engine/Plugins/PCG/Source/PCG/Private/Compute/PCGComputeGraph.h
+++ b/Engine/Plugins/PCG/Source/PCG/Private/Compute/PCGComputeGraph.h
@@ -100,7 +100,7 @@ public:
 
 	const FName* GetVirtualPinLabelFromBindingIndex(int32 InBindingIndex) const { return CPUDataBindingToVirtualPinLabel.Find(InBindingIndex); }
 
-	bool IsGraphValid(FPCGContext* InContext = nullptr, bool bQuiet = true) const;
+	bool IsGraphValid(FPCGContext* InContext) const;
 
 	void DebugLogDataDescriptions(const UPCGDataBinding* InBinding) const;
 
diff --git a/Engine/Plugins/PCG/Source/PCG/Private/Compute/PCGComputeKernel.cpp b/Engine/Plugins/PCG/Source/PCG/Private/Compute/PCGComputeKernel.cpp
--- a/Engine/Plugins/PCG/Source/PCG/Private/Compute/PCGComputeKernel.cpp
+++ b/Engine/Plugins/PCG/Source/PCG/Private/Compute/PCGComputeKernel.cpp
@@ -2,7 +2,10 @@
 
 #include "Compute/PCGComputeKernel.h"
 
+#include "PCGContext.h"
+#include "PCGEdge.h"
 #include "PCGModule.h"
+#include "PCGNode.h"
 #include "PCGSettings.h"
 #include "Compute/PCGComputeGraph.h"
 #include "Compute/PCGDataBinding.h"
@@ -20,9 +23,40 @@ namespace PCGComputeKernel
 		TEXT("Enables warnings when parameters are overidden on GPU nodes."));
 }
 
-void UPCGComputeKernel::Initialize(bool bLogDescriptions)
+#if WITH_EDITOR
+void UPCGComputeKernel::Initialize(const UPCGSettings* InSettings, bool bLogDescriptions)
 {
+	Settings = InSettings;
 	bLogDataDescriptions = bLogDescriptions;
+	bInitialized = true;
+
+	bHasStaticValidationErrors = !PerformStaticValidation();
+}
+#endif
+
+bool UPCGComputeKernel::IsKernelValid(FPCGContext* InContext) const
+{
+	if (!bInitialized)
+	{
+		PCG_KERNEL_VALIDATION_ERR(InContext, Settings, LOCTEXT("UninitializeKernel", "Kernel was not initialized during compilation. Make sure to call Initialize() when creating your kernels."));
+		return false;
+	}
+
+#if PCG_KERNEL_LOGGING_ENABLED
+	for (const FPCGKernelLogEntry& StaticLogEntry : StaticLogEntries)
+	{
+		if (StaticLogEntry.Verbosity == EPCGKernelLogVerbosity::Error)
+		{
+			PCG_KERNEL_VALIDATION_ERR(InContext, Settings, StaticLogEntry.Message);
+		}
+		else
+		{
+			PCG_KERNEL_VALIDATION_WARN(InContext, Settings, StaticLogEntry.Message);
+		}
+	}
+#endif
+
+	return !bHasStaticValidationErrors;
 }
 
 void UPCGComputeKernel::ComputeDataDescFromPinProperties(const FPCGPinPropertiesGPU& OutputPinProps, const TArrayView<const FPCGPinProperties>& InInputPinProps, const UPCGDataBinding* InBinding, FPCGDataCollectionDesc& OutPinDesc) const
@@ -261,30 +295,76 @@ void UPCGComputeKernel::ComputeDataDescFromPinProperties(const FPCGPinProperties
 	}
 }
 
-bool UPCGComputeKernel::ValidatePCGNode(FPCGContext* InContext, bool bQuiet, const UPCGSettings* InSettings) const
+#if WITH_EDITOR
+bool UPCGComputeKernel::PerformStaticValidation()
+{
+	return ValidatePCGNode(StaticLogEntries);
+}
+
+bool UPCGComputeKernel::ValidatePCGNode(TArray<FPCGKernelLogEntry>& OutLogEntries) const
 {
-	check(InSettings);
+	if (!Settings)
+	{
+		return true;
+	}
 
-	if (PCGComputeKernel::CVarWarnOnOverridePinUsage.GetValueOnAnyThread() && !bQuiet)
+	if (PCGComputeKernel::CVarWarnOnOverridePinUsage.GetValueOnAnyThread())
 	{
-		for (const FPCGSettingsOverridableParam& Param : InSettings->OverridableParams())
+		for (const FPCGSettingsOverridableParam& Param : Settings->OverridableParams())
 		{
 			if (ensure(!Param.PropertiesNames.IsEmpty()))
 			{
 				const FName PropertyName = Param.PropertiesNames[0];
 
-				if (InSettings->IsPropertyOverriddenByPin(PropertyName))
+				if (Settings->IsPropertyOverriddenByPin(PropertyName))
 				{
-					PCG_KERNEL_VALIDATION_WARN(InContext, InSettings, bQuiet, FText::Format(
-						LOCTEXT("ParamOverrideGPU", "Tried to override pin '{0}', but overrides are not supported on GPU nodes."),
-						FText::FromName(PropertyName)));
+#if PCG_KERNEL_LOGGING_ENABLED
+					OutLogEntries.Emplace(FText::Format(LOCTEXT("ParamOverrideGPU", "Tried to override pin '{0}', but overrides are not supported on GPU nodes."), FText::FromName(PropertyName)), EPCGKernelLogVerbosity::Warning);
+#endif
 				}
 			}
 		}
 	}
 
 	// Validate types of incident edges to make sure we catch invalid cases like Spatial -> Point.
-	return InSettings->AreInputEdgesValid(InContext, bQuiet);
+	return AreInputEdgesValid(OutLogEntries);
+}
+
+bool UPCGComputeKernel::AreInputEdgesValid(TArray<FPCGKernelLogEntry>& OutLogEntries) const
+{
+	bool bAllEdgesValid = true;
+
+	if (const UPCGNode* Node = Cast<UPCGNode>(GetOuter()))
+	{
+		for (const UPCGPin* InputPin : Node->GetInputPins())
+		{
+			if (!InputPin)
+			{
+				continue;
+			}
+
+			for (const UPCGEdge* InputEdge : InputPin->Edges)
+			{
+				const UPCGPin* UpstreamPin = InputEdge ? InputEdge->GetOtherPin(InputPin) : nullptr;
+				if (UpstreamPin && InputPin->GetRequiredTypeConversion(UpstreamPin) != EPCGTypeConversion::NoConversionRequired)
+				{
+#if PCG_KERNEL_LOGGING_ENABLED
+					OutLogEntries.Emplace(FText::Format(
+						LOCTEXT("InvalidInputPinEdge", "Unsupported connected upstream pin '{0}' on node '{1}' with type {2}. Recreate the edge to add required conversion nodes."),
+						FText::FromName(UpstreamPin->Properties.Label),
+						Node->GetNodeTitle(EPCGNodeTitleType::ListView),
+						StaticEnum<EPCGDataType>() ? StaticEnum<EPCGDataType>()->GetDisplayNameTextByValue(static_cast<int64>(UpstreamPin->Properties.AllowedTypes)) : FText::FromString(TEXT("MISSING"))),
+						EPCGKernelLogVerbosity::Error);
+#endif
+
+					bAllEdgesValid = false;
+				}
+			}
+		}
+	}
+
+	return bAllEdgesValid;
 }
+#endif
 
 #undef LOCTEXT_NAMESPACE
diff --git a/Engine/Plugins/PCG/Source/PCG/Private/Compute/PCGComputeKernel.h b/Engine/Plugins/PCG/Source/PCG/Private/Compute/PCGComputeKernel.h
--- a/Engine/Plugins/PCG/Source/PCG/Private/Compute/PCGComputeKernel.h
+++ b/Engine/Plugins/PCG/Source/PCG/Private/Compute/PCGComputeKernel.h
@@ -3,6 +3,7 @@
 #pragma once
 
 #include "ComputeFramework/ComputeKernel.h"
+#include "Compute/PCGComputeCommon.h"
 
 #include "ComputeFramework/ComputeGraph.h"
 #include "UObject/ObjectKey.h"
@@ -19,23 +20,53 @@ struct FPCGKernelAttributeKey;
 struct FPCGPinProperties;
 struct FPCGPinPropertiesGPU;
 
+UENUM()
+enum class EPCGKernelLogVerbosity
+{
+	Verbose,
+	Warning,
+	Error
+};
+
+USTRUCT()
+struct FPCGKernelLogEntry
+{
+	GENERATED_BODY()
+
+	FPCGKernelLogEntry() = default;
+
+	explicit FPCGKernelLogEntry(const FText& InMessage, EPCGKernelLogVerbosity InVerbosity)
+		: Message(InMessage)
+		, Verbosity(InVerbosity)
+	{}
+
+	UPROPERTY()
+	FText Message;
+
+	UPROPERTY()
+	EPCGKernelLogVerbosity Verbosity = EPCGKernelLogVerbosity::Verbose;
+};
+
 UCLASS(Abstract)
 class UPCGComputeKernel : public UComputeKernel
 {
 	GENERATED_BODY()
 
 public:
-	void Initialize(bool bLogDescriptions);
+#if WITH_EDITOR
+	/** Initialize kernel. Editor only as shader compilation only available in editor. */
+	virtual void Initialize(const UPCGSettings* InSettings, bool bLogDescriptions);
+#endif
 
 	/** Get/set index of kernel in kernels array on owning compute graph. */
 	int32 GetKernelIndex() const { return KernelIndex; }
 	void SetKernelIndex(int32 InKernelIndex) { KernelIndex = InKernelIndex; }
 
 	/** Gets settings for node associated with this kernel, if any. */
-	virtual const UPCGSettings* GetSettings() const { return nullptr; }
+	virtual const UPCGSettings* GetSettings() const { return Settings; }
 
-	/** Performs validation and returns true if this node is suitable for deployment to the GPU. */
-	virtual bool IsKernelValid(FPCGContext* InContext = nullptr, bool bQuiet = true) const { return true; }
+	/** Performs data validation and returns true if this node is suitable for deployment to the GPU. */
+	virtual bool IsKernelValid(FPCGContext* InContext) const;
 
 #if WITH_EDITOR
 	/** Produces the node specific portion of kernel shader source text, including the main entry point. */
@@ -87,19 +118,42 @@ public:
 	/** Helper to compute data descriptions when data is entirely determined by pin properties. */
 	void ComputeDataDescFromPinProperties(const FPCGPinPropertiesGPU& OutputPinProps, const TArrayView<const FPCGPinProperties>& InInputPinProps, const UPCGDataBinding* InBinding, FPCGDataCollectionDesc& OutPinDesc) const;
 
-	/** Do any validation of the PCG node. Returns true if node is valid and GPU execution can proceed. */
-	bool ValidatePCGNode(FPCGContext* InContext, bool bQuiet, const UPCGSettings* InSettings) const;
-
 	virtual void GetInputPins(TArray<FPCGPinProperties>& OutPins) const {}
 	virtual void GetOutputPins(TArray<FPCGPinPropertiesGPU>& OutPins) const {}
 
 	virtual bool GetLogDataDescriptions() const { return bLogDataDescriptions; }
 
+protected:
+#if WITH_EDITOR
+	/** Performs validation on compile time information (e.g. Node, Settings, etc.). Caches validation errors/warnings. */
+	virtual bool PerformStaticValidation();
+
+	/** Do any validation of the PCG node. Returns true if node is valid and GPU execution can proceed. */
+	bool ValidatePCGNode(TArray<FPCGKernelLogEntry>& OutLogEntries) const;
+
+	/** Validate input edges, for example checking connected pin types are compatible. */
+	bool AreInputEdgesValid(TArray<FPCGKernelLogEntry>& OutLogEntries) const;
+#endif
+
 protected:
 	/** Index into kernels array on owning compute graph. */
 	UPROPERTY()
 	int32 KernelIndex = INDEX_NONE;
 
+	UPROPERTY()
+	TObjectPtr<const UPCGSettings> Settings;
+
 	UPROPERTY()
 	bool bLogDataDescriptions = false;
+
+	UPROPERTY()
+	bool bInitialized = false;
+
+	UPROPERTY()
+	bool bHasStaticValidationErrors = false;
+
+	// @todo_pcg: These are a bit of a hack to avoid supporting static/compilation logging in PCGNodeVisualLogs.h
+	/** Log entries created and cached on initialization. Are logged during IsKernelValid(), unless PCG_KERNEL_LOGGING_ENABLED is false. */
+	UPROPERTY()
+	TArray<FPCGKernelLogEntry> StaticLogEntries;
 };
diff --git a/Engine/Plugins/PCG/Source/PCG/Private/Compute/PCGDataForGPU.cpp b/Engine/Plugins/PCG/Source/PCG/Private/Compute/PCGDataForGPU.cpp
--- a/Engine/Plugins/PCG/Source/PCG/Private/Compute/PCGDataForGPU.cpp
+++ b/Engine/Plugins/PCG/Source/PCG/Private/Compute/PCGDataForGPU.cpp
@@ -580,7 +580,7 @@ namespace PCGDataForGPUHelpers
 			if (bConflictingTypesInData)
 			{
 				// TODO: in future we could partition the execution - run it once per attribute type in input data.
-				PCG_KERNEL_VALIDATION_ERR(InContext, InSettings, /*bQuiet=*/false,
+				PCG_KERNEL_VALIDATION_ERR(InContext, InSettings,
 					FText::Format(LOCTEXT("AttributePackingTypeConflict", "Attribute '{0}' encountered with multiple different types in input data, custom float packing failed."), FText::FromName(AttributeName)));
 
 				OutAttributeIdOffsetStrides.Empty();
diff --git a/Engine/Plugins/PCG/Source/PCG/Private/Elements/Metadata/PCGMetadataPartitionKernel.cpp b/Engine/Plugins/PCG/Source/PCG/Private/Elements/Metadata/PCGMetadataPartitionKernel.cpp
--- a/Engine/Plugins/PCG/Source/PCG/Private/Elements/Metadata/PCGMetadataPartitionKernel.cpp
+++ b/Engine/Plugins/PCG/Source/PCG/Private/Elements/Metadata/PCGMetadataPartitionKernel.cpp
@@ -18,49 +18,26 @@
 
 #define LOCTEXT_NAMESPACE "PCGMetadataPartitionKernel"
 
-void UPCGMetadataPartitionKernel::Initialize(const UPCGMetadataPartitionSettings* InSettings, bool bLogDescriptions)
-{
-	Super::Initialize(bLogDescriptions);
-
-	Settings = InSettings;
-}
-
 const UPCGSettings* UPCGMetadataPartitionKernel::GetSettings() const
 {
 	return Settings;
 }
 
-bool UPCGMetadataPartitionKernel::IsKernelValid(FPCGContext* InContext, bool bQuiet) const
+bool UPCGMetadataPartitionKernel::IsKernelValid(FPCGContext* InContext) const
 {
 	TRACE_CPUPROFILER_EVENT_SCOPE(UPCGMetadataPartitionKernel::IsKernelValid);
 
-	if (!Super::IsKernelValid(InContext, bQuiet) || !ensure(Settings))
-	{
-		return false;
-	}
-
-	if (!ValidatePCGNode(InContext, bQuiet, Cast<UPCGSettings>(Settings)))
+	if (!Super::IsKernelValid(InContext))
 	{
 		return false;
 	}
 
-	if (Settings->PartitionAttributeSelectors.Num() != 1)
-	{
-		PCG_KERNEL_VALIDATION_ERR(InContext, Settings, bQuiet, LOCTEXT("MustProvideOnePartitionAttribute", "GPU implementation currently only supports a single partition attribute."));
-		return false;
-	}
-
-	if (!Settings->PartitionAttributeSelectors[0].IsBasicAttribute())
-	{
-		PCG_KERNEL_VALIDATION_ERR(InContext, Settings, bQuiet, LOCTEXT("OnlyBasicAttributesSupported", "GPU implementation currently only supports basic attributes."));
-		return false;
-	}
-
 	FPCGComputeGraphContext* Context = InContext->IsComputeContext() ? static_cast<FPCGComputeGraphContext*>(InContext) : nullptr;
 
 	if (Context && Context->DataBinding)
 	{
-		const FName AttributeName = Settings->PartitionAttributeSelectors[0].GetAttributeName();
+		const UPCGMetadataPartitionSettings* MPSettings = CastChecked<UPCGMetadataPartitionSettings>(Settings);
+		const FName AttributeName = MPSettings->PartitionAttributeSelectors[0].GetAttributeName();
 
 		FPCGDataCollectionDesc InputDesc;
 		ensure(Context->DataBinding->GetKernelPinDataDesc_GameThread(this, PCGPinConstants::DefaultInputLabel, /*bIsInput=*/true, InputDesc));
@@ -73,7 +50,7 @@ bool UPCGMetadataPartitionKernel::IsKernelValid(FPCGContext* InContext, bool bQu
 		{
 			if (!InputDesc.DataDescs.IsEmpty())
 			{
-				PCG_KERNEL_VALIDATION_ERR(InContext, Settings, bQuiet, FText::Format(
+				PCG_KERNEL_VALIDATION_ERR(InContext, Settings, FText::Format(
 					LOCTEXT("PartitionAttributeMissing", "Partition attribute was not present in GPU data."),
 					FText::FromName(AttributeName)));
 			}
@@ -83,7 +60,7 @@ bool UPCGMetadataPartitionKernel::IsKernelValid(FPCGContext* InContext, bool bQu
 
 		if (bConflictingTypesInData)
 		{
-			PCG_KERNEL_VALIDATION_ERR(InContext, Settings, bQuiet, FText::Format(
+			PCG_KERNEL_VALIDATION_ERR(InContext, Settings, FText::Format(
 				LOCTEXT("PartitionAttributeTypeConflict", "Attribute '{0}' encountered with multiple different types in input data."),
 				FText::FromName(AttributeName)));
 			return false;
@@ -91,7 +68,7 @@ bool UPCGMetadataPartitionKernel::IsKernelValid(FPCGContext* InContext, bool bQu
 
 		if (AttributeDesc.Type != EPCGKernelAttributeType::StringKey)
 		{
-			PCG_KERNEL_VALIDATION_ERR(Context, Settings, bQuiet, FText::Format(
+			PCG_KERNEL_VALIDATION_ERR(Context, Settings, FText::Format(
 				LOCTEXT("PartitionAttributeTypeInvalid", "Attribute '{0}' not usable for partitioning, only attributes of type String Key are currently supported."),
 				FText::FromName(AttributeName)));
 			return false;
@@ -106,6 +83,8 @@ bool UPCGMetadataPartitionKernel::ComputeOutputBindingDataDesc(const UPCGCompute
 	check(InGraph);
 	check(InBinding);
 
+	const UPCGMetadataPartitionSettings* MPSettings = CastChecked<UPCGMetadataPartitionSettings>(Settings);
+
 	// Code assumes single output pin.
 	if (!ensure(InOutputPinLabel == PCGPinConstants::DefaultOutputLabel))
 	{
@@ -115,7 +94,7 @@ bool UPCGMetadataPartitionKernel::ComputeOutputBindingDataDesc(const UPCGCompute
 	FPCGDataCollectionDesc InputPinDesc;
 	ensure(InBinding->GetKernelPinDataDesc_GameThread(this, PCGPinConstants::DefaultInputLabel, /*bIsInput=*/true, InputPinDesc));
 
-	const int32 PartitionAttributeId = InBinding->GetAttributeId(Settings->PartitionAttributeSelectors[0].GetAttributeName(), EPCGKernelAttributeType::StringKey);
+	const int32 PartitionAttributeId = InBinding->GetAttributeId(MPSettings->PartitionAttributeSelectors[0].GetAttributeName(), EPCGKernelAttributeType::StringKey);
 
 	FPCGDataCollectionDesc PinDesc;
 
@@ -146,9 +125,9 @@ bool UPCGMetadataPartitionKernel::ComputeOutputBindingDataDesc(const UPCGCompute
 	}
 
 	// Add output attribute (partition index).
-	if (Settings->bAssignIndexPartition)
+	if (MPSettings->bAssignIndexPartition)
 	{
-		PinDesc.AddAttributeToAllData(FPCGKernelAttributeKey(Settings->PartitionIndexAttributeName, EPCGKernelAttributeType::Int), InBinding);
+		PinDesc.AddAttributeToAllData(FPCGKernelAttributeKey(MPSettings->PartitionIndexAttributeName, EPCGKernelAttributeType::Int), InBinding);
 	}
 
 	OutDataDesc = MoveTemp(PinDesc);
@@ -183,10 +162,12 @@ void UPCGMetadataPartitionKernel::CreateAdditionalInputDataInterfaces(FPCGGPUCom
 
 void UPCGMetadataPartitionKernel::GetKernelAttributeKeys(TArray<FPCGKernelAttributeKey>& OutKeys) const
 {
+	const UPCGMetadataPartitionSettings* MPSettings = CastChecked<UPCGMetadataPartitionSettings>(Settings);
+
 	// Register the attributes this node reads or writes for which we know the attribute type. Currently
 	// only StringKey attributes can be partitioned when executing on GPU, so declare the attribute here
 	// rather than doing work to resolve at runtime.
-	for (const FPCGAttributePropertyInputSelector& Selector : Settings->PartitionAttributeSelectors)
+	for (const FPCGAttributePropertyInputSelector& Selector : MPSettings->PartitionAttributeSelectors)
 	{
 		if (Selector.IsBasicAttribute())
 		{
@@ -195,9 +176,9 @@ void UPCGMetadataPartitionKernel::GetKernelAttributeKeys(TArray<FPCGKernelAttrib
 	}
 
 	// Register output attribute created by this node.
-	if (Settings->bAssignIndexPartition)
+	if (MPSettings->bAssignIndexPartition)
 	{
-		OutKeys.AddUnique(FPCGKernelAttributeKey(Settings->PartitionIndexAttributeName, EPCGKernelAttributeType::Int));
+		OutKeys.AddUnique(FPCGKernelAttributeKey(MPSettings->PartitionIndexAttributeName, EPCGKernelAttributeType::Int));
 	}
 }
 
@@ -211,4 +192,34 @@ void UPCGMetadataPartitionKernel::GetOutputPins(TArray<FPCGPinPropertiesGPU>& Ou
 	OutPins.Emplace(PCGPinConstants::DefaultOutputLabel, EPCGDataType::Point);
 }
 
+#if WITH_EDITOR
+bool UPCGMetadataPartitionKernel::PerformStaticValidation()
+{
+	if (!Super::PerformStaticValidation())
+	{
+		return false;
+	}
+	
+	const UPCGMetadataPartitionSettings* MPSettings = CastChecked<UPCGMetadataPartitionSettings>(Settings);
+
+	if (MPSettings->PartitionAttributeSelectors.Num() != 1)
+	{
+#if PCG_KERNEL_LOGGING_ENABLED
+		StaticLogEntries.Emplace(LOCTEXT("MustProvideOnePartitionAttribute", "GPU implementation currently only supports a single partition attribute."), EPCGKernelLogVerbosity::Error);
+#endif
+		return false;
+	}
+
+	if (!MPSettings->PartitionAttributeSelectors[0].IsBasicAttribute())
+	{
+#if PCG_KERNEL_LOGGING_ENABLED
+		StaticLogEntries.Emplace(LOCTEXT("OnlyBasicAttributesSupported", "GPU implementation currently only supports basic attributes."), EPCGKernelLogVerbosity::Error);
+#endif
+		return false;
+	}
+
+	return true;
+}
+#endif
+
 #undef LOCTEXT_NAMESPACE
diff --git a/Engine/Plugins/PCG/Source/PCG/Private/Elements/PCGCopyPointsAnalysisKernel.cpp b/Engine/Plugins/PCG/Source/PCG/Private/Elements/PCGCopyPointsAnalysisKernel.cpp
--- a/Engine/Plugins/PCG/Source/PCG/Private/Elements/PCGCopyPointsAnalysisKernel.cpp
+++ b/Engine/Plugins/PCG/Source/PCG/Private/Elements/PCGCopyPointsAnalysisKernel.cpp
@@ -21,29 +21,16 @@
 
 #define LOCTEXT_NAMESPACE "PCGCopyPointsAnalysisKernel"
 
-void UPCGCopyPointsAnalysisKernel::Initialize(const UPCGCopyPointsSettings* InSettings, bool bLogDescriptions)
-{
-	Super::Initialize(bLogDescriptions);
-
-	Settings = InSettings;
-}
-
 const UPCGSettings* UPCGCopyPointsAnalysisKernel::GetSettings() const
 {
 	return Settings;
 }
 
-bool UPCGCopyPointsAnalysisKernel::IsKernelValid(FPCGContext* InContext, bool bQuiet) const
+bool UPCGCopyPointsAnalysisKernel::IsKernelValid(FPCGContext* InContext) const
 {
 	TRACE_CPUPROFILER_EVENT_SCOPE(UPCGCopyPointsAnalysisKernel::IsKernelValid);
-	const UPCGCopyPointsSettings* CopyPointSettings = Cast<UPCGCopyPointsSettings>(Settings);
 
-	if (!Super::IsKernelValid(InContext, bQuiet) || !ensure(CopyPointSettings))
-	{
-		return false;
-	}
-
-	if (!ValidatePCGNode(InContext, bQuiet, CopyPointSettings))
+	if (!Super::IsKernelValid(InContext))
 	{
 		return false;
 	}
@@ -53,9 +40,7 @@ bool UPCGCopyPointsAnalysisKernel::IsKernelValid(FPCGContext* InContext, bool bQ
 		return true;
 	}
 
-	const FPCGComputeGraphContext* Context = static_cast<FPCGComputeGraphContext*>(InContext);
-
-	return PCGCopyPointsKernel::IsKernelValid(this, CopyPointSettings, Context, bQuiet);
+	return PCGCopyPointsKernel::IsKernelValid(this, CastChecked<UPCGCopyPointsSettings>(Settings), static_cast<FPCGComputeGraphContext*>(InContext));
 }
 
 bool UPCGCopyPointsAnalysisKernel::ComputeOutputBindingDataDesc(const UPCGComputeGraph* InGraph, FName InOutputPinLabel, const UPCGDataBinding* InBinding, FPCGDataCollectionDesc& OutDataDesc) const
diff --git a/Engine/Plugins/PCG/Source/PCG/Private/Elements/PCGCopyPointsKernel.cpp b/Engine/Plugins/PCG/Source/PCG/Private/Elements/PCGCopyPointsKernel.cpp
--- a/Engine/Plugins/PCG/Source/PCG/Private/Elements/PCGCopyPointsKernel.cpp
+++ b/Engine/Plugins/PCG/Source/PCG/Private/Elements/PCGCopyPointsKernel.cpp
@@ -22,29 +22,16 @@
 
 #define LOCTEXT_NAMESPACE "PCGCopyPointsKernel"
 
-void UPCGCopyPointsKernel::Initialize(const UPCGCopyPointsSettings* InSettings, bool bLogDescriptions)
-{
-	Super::Initialize(bLogDescriptions);
-
-	Settings = InSettings;
-}
-
 const UPCGSettings* UPCGCopyPointsKernel::GetSettings() const
 {
 	return Settings;
 }
 
-bool UPCGCopyPointsKernel::IsKernelValid(FPCGContext* InContext, bool bQuiet) const
+bool UPCGCopyPointsKernel::IsKernelValid(FPCGContext* InContext) const
 {
 	TRACE_CPUPROFILER_EVENT_SCOPE(UPCGCopyPointsKernel::IsKernelValid);
-	const UPCGCopyPointsSettings* CopyPointSettings = Cast<UPCGCopyPointsSettings>(Settings);
 
-	if (!Super::IsKernelValid(InContext, bQuiet) || !ensure(CopyPointSettings))
-	{
-		return false;
-	}
-
-	if (!ValidatePCGNode(InContext, bQuiet, CopyPointSettings))
+	if (!Super::IsKernelValid(InContext))
 	{
 		return false;
 	}
@@ -54,9 +41,7 @@ bool UPCGCopyPointsKernel::IsKernelValid(FPCGContext* InContext, bool bQuiet) co
 		return true;
 	}
 
-	const FPCGComputeGraphContext* Context = static_cast<FPCGComputeGraphContext*>(InContext);
-
-	return PCGCopyPointsKernel::IsKernelValid(this, CopyPointSettings, Context, bQuiet);
+	return PCGCopyPointsKernel::IsKernelValid(this, CastChecked<UPCGCopyPointsSettings>(Settings), static_cast<FPCGComputeGraphContext*>(InContext));
 }
 
 bool UPCGCopyPointsKernel::ComputeOutputBindingDataDesc(const UPCGComputeGraph* InGraph, FName InOutputPinLabel, const UPCGDataBinding* InBinding, FPCGDataCollectionDesc& OutDataDesc) const
@@ -64,6 +49,8 @@ bool UPCGCopyPointsKernel::ComputeOutputBindingDataDesc(const UPCGComputeGraph*
 	check(InGraph);
 	check(InBinding);
 
+	const UPCGCopyPointsSettings* CPSettings = CastChecked<UPCGCopyPointsSettings>(Settings);
+
 	// Code assumes single output pin.
 	if (!ensure(InOutputPinLabel == PCGPinConstants::DefaultOutputLabel))
 	{
@@ -74,7 +61,7 @@ bool UPCGCopyPointsKernel::ComputeOutputBindingDataDesc(const UPCGComputeGraph*
 	int NumSelected = INDEX_NONE;
 	TArray<bool> Selected;
 	
-	if (Settings->bMatchBasedOnAttribute)
+	if (CPSettings->bMatchBasedOnAttribute)
 	{
 		const int32 AnalysisDataIndex = InBinding->GetFirstInputDataIndex(this, PCGCopyPointsConstants::SelectedFlagsPinLabel);
 
@@ -122,9 +109,9 @@ bool UPCGCopyPointsKernel::ComputeOutputBindingDataDesc(const UPCGComputeGraph*
 
 	const int32 NumSources = SourcePinDesc.DataDescs.Num();
 	const int32 NumTargets = TargetPinDesc.DataDescs.Num();
-	const int32 NumIterations = Settings->bCopyEachSourceOnEveryTarget ? NumSources * NumTargets : FMath::Max(NumSources, NumTargets);
+	const int32 NumIterations = CPSettings->bCopyEachSourceOnEveryTarget ? NumSources * NumTargets : FMath::Max(NumSources, NumTargets);
 
-	if (NumSources > 0 && NumTargets > 0 && (Settings->bCopyEachSourceOnEveryTarget || NumSources == NumTargets || NumSources == 1 || NumTargets == 1))
+	if (NumSources > 0 && NumTargets > 0 && (CPSettings->bCopyEachSourceOnEveryTarget || NumSources == NumTargets || NumSources == 1 || NumTargets == 1))
 	{
 		for (int32 I = 0; I < NumIterations; ++I)
 		{
@@ -133,15 +120,15 @@ bool UPCGCopyPointsKernel::ComputeOutputBindingDataDesc(const UPCGComputeGraph*
 				continue;
 			}
 
-			const int32 SourceIndex = Settings->bCopyEachSourceOnEveryTarget ? (I / NumTargets) : FMath::Min(I, NumSources - 1);
-			const int32 TargetIndex = Settings->bCopyEachSourceOnEveryTarget ? (I % NumTargets) : FMath::Min(I, NumTargets - 1);
+			const int32 SourceIndex = CPSettings->bCopyEachSourceOnEveryTarget ? (I / NumTargets) : FMath::Min(I, NumSources - 1);
+			const int32 TargetIndex = CPSettings->bCopyEachSourceOnEveryTarget ? (I % NumTargets) : FMath::Min(I, NumTargets - 1);
 
 			const FPCGDataDesc& SourceDesc = SourcePinDesc.DataDescs[SourceIndex];
 			const FPCGDataDesc& TargetDesc = TargetPinDesc.DataDescs[TargetIndex];
 
 			FPCGDataDesc& ResultDataDesc = PinDesc.DataDescs.Emplace_GetRef(EPCGDataType::Point, SourceDesc.ElementCount * TargetDesc.ElementCount);
 
-			if (Settings->AttributeInheritance == EPCGCopyPointsMetadataInheritanceMode::SourceFirst)
+			if (CPSettings->AttributeInheritance == EPCGCopyPointsMetadataInheritanceMode::SourceFirst)
 			{
 				ResultDataDesc.AttributeDescs = SourceDesc.AttributeDescs;
 
@@ -150,7 +137,7 @@ bool UPCGCopyPointsKernel::ComputeOutputBindingDataDesc(const UPCGComputeGraph*
 					ResultDataDesc.AttributeDescs.AddUnique(AttrDesc);
 				}
 			}
-			else if (Settings->AttributeInheritance == EPCGCopyPointsMetadataInheritanceMode::TargetFirst)
+			else if (CPSettings->AttributeInheritance == EPCGCopyPointsMetadataInheritanceMode::TargetFirst)
 			{
 				ResultDataDesc.AttributeDescs = TargetDesc.AttributeDescs;
 
@@ -159,16 +146,16 @@ bool UPCGCopyPointsKernel::ComputeOutputBindingDataDesc(const UPCGComputeGraph*
 					ResultDataDesc.AttributeDescs.AddUnique(AttrDesc);
 				}
 			}
-			else if (Settings->AttributeInheritance == EPCGCopyPointsMetadataInheritanceMode::SourceOnly)
+			else if (CPSettings->AttributeInheritance == EPCGCopyPointsMetadataInheritanceMode::SourceOnly)
 			{
 				ResultDataDesc.AttributeDescs = SourceDesc.AttributeDescs;
 			}
-			else if (Settings->AttributeInheritance == EPCGCopyPointsMetadataInheritanceMode::TargetOnly)
+			else if (CPSettings->AttributeInheritance == EPCGCopyPointsMetadataInheritanceMode::TargetOnly)
 			{
 				ResultDataDesc.AttributeDescs = TargetDesc.AttributeDescs;
 			}
 
-			if (Settings->TagInheritance == EPCGCopyPointsTagInheritanceMode::Source || Settings->TagInheritance == EPCGCopyPointsTagInheritanceMode::Both)
+			if (CPSettings->TagInheritance == EPCGCopyPointsTagInheritanceMode::Source || CPSettings->TagInheritance == EPCGCopyPointsTagInheritanceMode::Both)
 			{
 				for (const int32 SourceTagStringKey : SourceDesc.TagStringKeys)
 				{
@@ -176,7 +163,7 @@ bool UPCGCopyPointsKernel::ComputeOutputBindingDataDesc(const UPCGComputeGraph*
 				}
 			}
 			
-			if (Settings->TagInheritance == EPCGCopyPointsTagInheritanceMode::Target || Settings->TagInheritance == EPCGCopyPointsTagInheritanceMode::Both)
+			if (CPSettings->TagInheritance == EPCGCopyPointsTagInheritanceMode::Target || CPSettings->TagInheritance == EPCGCopyPointsTagInheritanceMode::Both)
 			{
 				for (const int32 TargetTagStringKey : TargetDesc.TagStringKeys)
 				{
@@ -227,10 +214,12 @@ void UPCGCopyPointsKernel::CreateAdditionalInputDataInterfaces(FPCGGPUCompilatio
 
 void UPCGCopyPointsKernel::GetInputPins(TArray<FPCGPinProperties>& OutPins) const
 {
+	const UPCGCopyPointsSettings* CPSettings = CastChecked<UPCGCopyPointsSettings>(Settings);
+
 	OutPins.Emplace(PCGCopyPointsConstants::SourcePointsLabel, EPCGDataType::Point);
 	OutPins.Emplace(PCGCopyPointsConstants::TargetPointsLabel, EPCGDataType::Point);
 
-	if (Settings->bMatchBasedOnAttribute)
+	if (CPSettings->bMatchBasedOnAttribute)
 	{
 		// One bool flag per output data that signals if the output data should be computed or not.
 		OutPins.Emplace(PCGCopyPointsConstants::SelectedFlagsPinLabel, EPCGDataType::Param);
diff --git a/Engine/Plugins/PCG/Source/PCG/Private/Elements/PCGCopyPointsKernelShared.cpp b/Engine/Plugins/PCG/Source/PCG/Private/Elements/PCGCopyPointsKernelShared.cpp
--- a/Engine/Plugins/PCG/Source/PCG/Private/Elements/PCGCopyPointsKernelShared.cpp
+++ b/Engine/Plugins/PCG/Source/PCG/Private/Elements/PCGCopyPointsKernelShared.cpp
@@ -10,13 +10,16 @@
 
 #define LOCTEXT_NAMESPACE "PCGCopyPointsKernel"
 
-bool PCGCopyPointsKernel::IsKernelValid(const UPCGComputeKernel* InKernel, const UPCGCopyPointsSettings* InCopyPointSettings, const FPCGComputeGraphContext* InContext, bool bQuiet)
+bool PCGCopyPointsKernel::IsKernelValid(const UPCGComputeKernel* InKernel, const UPCGCopyPointsSettings* InCopyPointSettings, const FPCGComputeGraphContext* InContext)
 {
+	check(InCopyPointSettings);
+	check(InContext);
+
 	if (InCopyPointSettings->bMatchBasedOnAttribute)
 	{
 		if (const UPCGDataBinding* DataBinding = InContext->DataBinding.Get(); ensure(DataBinding))
 		{
-			auto ValidateAttributeExists = [InContext, DataBinding, InKernel, InCopyPointSettings, bQuiet](FName InputPin) -> bool
+			auto ValidateAttributeExists = [InContext, DataBinding, InKernel, InCopyPointSettings](FName InputPin) -> bool
 			{
 				FPCGDataCollectionDesc PinDataDesc;
 				ensure(DataBinding->GetKernelPinDataDesc_GameThread(InKernel, InputPin, /*bIsInputPin=*/true, PinDataDesc));
@@ -30,7 +33,7 @@ bool PCGCopyPointsKernel::IsKernelValid(const UPCGComputeKernel* InKernel, const
 				{
 					if (!DataDesc.ContainsAttribute(MatchAttributeName, MatchAttributeType))
 					{
-						PCG_KERNEL_VALIDATION_ERR(InContext, InCopyPointSettings, bQuiet, FText::Format(
+						PCG_KERNEL_VALIDATION_ERR(InContext, InCopyPointSettings, FText::Format(
 							LOCTEXT("MatchAttributeMissing", "Match attribute '{0}' not found, this attribute must be present on all input data, and be of type Integer."),
 							FText::FromName(MatchAttributeName)));
 
diff --git a/Engine/Plugins/PCG/Source/PCG/Private/Elements/PCGStaticMeshSpawner.cpp b/Engine/Plugins/PCG/Source/PCG/Private/Elements/PCGStaticMeshSpawner.cpp
--- a/Engine/Plugins/PCG/Source/PCG/Private/Elements/PCGStaticMeshSpawner.cpp
+++ b/Engine/Plugins/PCG/Source/PCG/Private/Elements/PCGStaticMeshSpawner.cpp
@@ -61,7 +61,7 @@ void UPCGStaticMeshSpawnerSettings::CreateKernels(FPCGGPUCompilationContext& InO
 	if (const UPCGMeshSelectorByAttribute* Selector = Cast<UPCGMeshSelectorByAttribute>(MeshSelectorParameters))
 	{
 		UPCGCountUniqueAttributeValuesKernel* CountKernel = InOutContext.NewObject_AnyThread<UPCGCountUniqueAttributeValuesKernel>(InObjectOuter);
-		CountKernel->Initialize(Selector->AttributeName, bDumpDataDescriptions);
+		CountKernel->Initialize(this, bDumpDataDescriptions);
 		OutKernels.Add(CountKernel);
 
 		OutEdges.Emplace(FPCGPinReference(PCGPinConstants::DefaultInputLabel), FPCGPinReference(CountKernel, PCGPinConstants::DefaultInputLabel));
diff --git a/Engine/Plugins/PCG/Source/PCG/Private/Elements/PCGStaticMeshSpawnerKernel.cpp b/Engine/Plugins/PCG/Source/PCG/Private/Elements/PCGStaticMeshSpawnerKernel.cpp
--- a/Engine/Plugins/PCG/Source/PCG/Private/Elements/PCGStaticMeshSpawnerKernel.cpp
+++ b/Engine/Plugins/PCG/Source/PCG/Private/Elements/PCGStaticMeshSpawnerKernel.cpp
@@ -32,66 +32,11 @@ namespace PCGStaticMeshSpawnerKernel
 	static const FText NoMeshEntriesFormat = LOCTEXT("NoMeshEntries", "No mesh entries provided in weighted mesh selector.");
 }
 
-void UPCGStaticMeshSpawnerKernel::Initialize(const UPCGStaticMeshSpawnerSettings* InSettings, bool bLogDescriptions)
-{
-	Super::Initialize(bLogDescriptions);
-
-	Settings = InSettings;
-}
-
 const UPCGSettings* UPCGStaticMeshSpawnerKernel::GetSettings() const
 {
 	return Settings;
 }
 
-bool UPCGStaticMeshSpawnerKernel::IsKernelValid(FPCGContext* InContext, bool bQuiet) const
-{
-	TRACE_CPUPROFILER_EVENT_SCOPE(UPCGStaticMeshSpawnerKernel::IsKernelValid);
-
-	if (!Super::IsKernelValid(InContext, bQuiet) || !ensure(Settings))
-	{
-		return false;
-	}
-
-	if (!ValidatePCGNode(InContext, bQuiet, Cast<UPCGSettings>(Settings)))
-	{
-		return false;
-	}
-
-	if (UPCGMeshSelectorWeighted* SelectorWeighted = Cast<UPCGMeshSelectorWeighted>(Settings->MeshSelectorParameters))
-	{
-		if (SelectorWeighted->MeshEntries.IsEmpty())
-		{
-			PCG_KERNEL_VALIDATION_ERR(InContext, Settings, bQuiet, PCGStaticMeshSpawnerKernel::NoMeshEntriesFormat);
-			return false;
-		}
-
-		for (const FPCGMeshSelectorWeightedEntry& Entry : SelectorWeighted->MeshEntries)
-		{
-			if (Entry.Descriptor.StaticMesh.IsNull())
-			{
-				PCG_KERNEL_VALIDATION_ERR(InContext, Settings, bQuiet, LOCTEXT("UnassignedMesh", "Unassigned mesh."));
-				return false;
-			}
-		}
-	}
-	else if (!Settings->MeshSelectorParameters || !Settings->MeshSelectorParameters->IsA<UPCGMeshSelectorByAttribute>())
-	{
-		PCG_KERNEL_VALIDATION_ERR(InContext, Settings, bQuiet, LOCTEXT("InvalidMeshSelector", "Currently GPU Static Mesh Spawner nodes must use PCGMeshSelectorWeighted or UPCGMeshSelectorByAttribute as the mesh selector type."));
-		return false;
-	}
-
-	// Currently instance packers must be able to specify a full list of attribute names upfront, to build the attribute table at compile time.
-	// TODO: We should be able to augment a static attribute table with new attributes at execution time, which will allow other types like regex.
-	if (Settings->InstanceDataPackerParameters && !Settings->InstanceDataPackerParameters->GetAttributeNames(/*OutNames=*/nullptr))
-	{
-		PCG_KERNEL_VALIDATION_ERR(InContext, Settings, bQuiet, LOCTEXT("InvalidInstancePacker", "Selected instance packer does not support GPU execution."));
-		return false;
-	}
-
-	return true;
-}
-
 bool UPCGStaticMeshSpawnerKernel::ComputeOutputBindingDataDesc(const UPCGComputeGraph* InGraph, FName InOutputPinLabel, const UPCGDataBinding* InBinding, FPCGDataCollectionDesc& OutDataDesc) const
 {
 	check(InGraph);
@@ -109,8 +54,10 @@ bool UPCGStaticMeshSpawnerKernel::ComputeOutputBindingDataDesc(const UPCGCompute
 
 	// Add output attribute (selected mesh).
 	{
+		const UPCGStaticMeshSpawnerSettings* SMSettings = CastChecked<UPCGStaticMeshSpawnerSettings>(Settings);
+
 		// Create unique value keys for the output string.
-		if (UPCGMeshSelectorWeighted* SelectorWeighted = Cast<UPCGMeshSelectorWeighted>(Settings->MeshSelectorParameters))
+		if (UPCGMeshSelectorWeighted* SelectorWeighted = Cast<UPCGMeshSelectorWeighted>(SMSettings->MeshSelectorParameters))
 		{
 			TArray<int32> UniqueStringKeys;
 
@@ -125,9 +72,9 @@ bool UPCGStaticMeshSpawnerKernel::ComputeOutputBindingDataDesc(const UPCGCompute
 				}
 			}
 
-			PinDesc.AddAttributeToAllData(FPCGKernelAttributeKey(Settings->OutAttributeName, EPCGKernelAttributeType::StringKey), InBinding, &UniqueStringKeys);
+			PinDesc.AddAttributeToAllData(FPCGKernelAttributeKey(SMSettings->OutAttributeName, EPCGKernelAttributeType::StringKey), InBinding, &UniqueStringKeys);
 		}
-		else if (UPCGMeshSelectorByAttribute* SelectorByAttribute = Cast<UPCGMeshSelectorByAttribute>(Settings->MeshSelectorParameters))
+		else if (UPCGMeshSelectorByAttribute* SelectorByAttribute = Cast<UPCGMeshSelectorByAttribute>(SMSettings->MeshSelectorParameters))
 		{
 			// By-attribute selection - pass on strings from input attribute.
 			if (SelectorByAttribute->AttributeName != NAME_None)
@@ -147,13 +94,13 @@ bool UPCGStaticMeshSpawnerKernel::ComputeOutputBindingDataDesc(const UPCGCompute
 						}
 					}
 
-					DataDesc.AddAttribute(FPCGKernelAttributeKey(Settings->OutAttributeName, EPCGKernelAttributeType::StringKey), InBinding, &StringKeys);
+					DataDesc.AddAttribute(FPCGKernelAttributeKey(SMSettings->OutAttributeName, EPCGKernelAttributeType::StringKey), InBinding, &StringKeys);
 				}
 			}
 		}
-		else if (Settings->MeshSelectorParameters)
+		else if (SMSettings->MeshSelectorParameters)
 		{
-			UE_LOG(LogPCG, Error, TEXT("Mesh selector not supported by GPU Static Mesh Spawner: %s"), *Settings->MeshSelectorParameters->GetName());
+			UE_LOG(LogPCG, Error, TEXT("Mesh selector not supported by GPU Static Mesh Spawner: %s"), *SMSettings->MeshSelectorParameters->GetName());
 		}
 	}
 
@@ -199,7 +146,8 @@ void UPCGStaticMeshSpawnerKernel::CreateAdditionalOutputDataInterfaces(FPCGGPUCo
 
 void UPCGStaticMeshSpawnerKernel::AddStaticCreatedStrings(TArray<FString>& InOutStringTable) const
 {
-	if (UPCGMeshSelectorWeighted* SelectorWeighted = Cast<UPCGMeshSelectorWeighted>(Settings->MeshSelectorParameters))
+	const UPCGStaticMeshSpawnerSettings* SMSettings = CastChecked<UPCGStaticMeshSpawnerSettings>(Settings);
+	if (UPCGMeshSelectorWeighted* SelectorWeighted = Cast<UPCGMeshSelectorWeighted>(SMSettings->MeshSelectorParameters))
 	{
 		for (const FPCGMeshSelectorWeightedEntry& Entry : SelectorWeighted->MeshEntries)
 		{
@@ -211,7 +159,8 @@ void UPCGStaticMeshSpawnerKernel::AddStaticCreatedStrings(TArray<FString>& InOut
 void UPCGStaticMeshSpawnerKernel::GetKernelAttributeKeys(TArray<FPCGKernelAttributeKey>& OutKeys) const
 {
 	// Register the attribute this node creates.
-	OutKeys.AddUnique(FPCGKernelAttributeKey(Settings->OutAttributeName, EPCGKernelAttributeType::StringKey));
+	const UPCGStaticMeshSpawnerSettings* SMSettings = CastChecked<UPCGStaticMeshSpawnerSettings>(Settings);
+	OutKeys.AddUnique(FPCGKernelAttributeKey(SMSettings->OutAttributeName, EPCGKernelAttributeType::StringKey));
 }
 
 void UPCGStaticMeshSpawnerKernel::GetInputPins(TArray<FPCGPinProperties>& OutPins) const
@@ -225,4 +174,59 @@ void UPCGStaticMeshSpawnerKernel::GetOutputPins(TArray<FPCGPinPropertiesGPU>& Ou
 	OutPins.Emplace(PCGPinConstants::DefaultOutputLabel, EPCGDataType::Point);
 }
 
+#if WITH_EDITOR
+bool UPCGStaticMeshSpawnerKernel::PerformStaticValidation()
+{
+	TRACE_CPUPROFILER_EVENT_SCOPE(UPCGStaticMeshSpawnerKernel::PerformStaticValidation);
+
+	if (!Super::PerformStaticValidation())
+	{
+		return false;
+	}
+
+	const UPCGStaticMeshSpawnerSettings* SMSettings = CastChecked<UPCGStaticMeshSpawnerSettings>(Settings);
+
+	if (UPCGMeshSelectorWeighted* SelectorWeighted = Cast<UPCGMeshSelectorWeighted>(SMSettings->MeshSelectorParameters))
+	{
+		if (SelectorWeighted->MeshEntries.IsEmpty())
+		{
+#if PCG_KERNEL_LOGGING_ENABLED
+			StaticLogEntries.Emplace(PCGStaticMeshSpawnerKernel::NoMeshEntriesFormat, EPCGKernelLogVerbosity::Error);
+#endif
+			return false;
+		}
+
+		for (const FPCGMeshSelectorWeightedEntry& Entry : SelectorWeighted->MeshEntries)
+		{
+			if (Entry.Descriptor.StaticMesh.IsNull())
+			{
+#if PCG_KERNEL_LOGGING_ENABLED
+				StaticLogEntries.Emplace(LOCTEXT("UnassignedMesh", "Unassigned mesh."), EPCGKernelLogVerbosity::Error);
+#endif
+				return false;
+			}
+		}
+	}
+	else if (!SMSettings->MeshSelectorParameters || !SMSettings->MeshSelectorParameters->IsA<UPCGMeshSelectorByAttribute>())
+	{
+#if PCG_KERNEL_LOGGING_ENABLED
+		StaticLogEntries.Emplace(LOCTEXT("InvalidMeshSelector", "Currently GPU Static Mesh Spawner nodes must use PCGMeshSelectorWeighted or UPCGMeshSelectorByAttribute as the mesh selector type."), EPCGKernelLogVerbosity::Error);
+#endif
+		return false;
+	}
+
+	// Currently instance packers must be able to specify a full list of attribute names upfront, to build the attribute table at compile time.
+	// TODO: We should be able to augment a static attribute table with new attributes at execution time, which will allow other types like regex.
+	if (SMSettings->InstanceDataPackerParameters && !SMSettings->InstanceDataPackerParameters->GetAttributeNames(/*OutNames=*/nullptr))
+	{
+#if PCG_KERNEL_LOGGING_ENABLED
+		StaticLogEntries.Emplace(LOCTEXT("InvalidInstancePacker", "Selected instance packer does not support GPU execution."), EPCGKernelLogVerbosity::Error);
+#endif
+		return false;
+	}
+	
+	return true;
+}
+#endif
+
 #undef LOCTEXT_NAMESPACE
diff --git a/Engine/Plugins/PCG/Source/PCG/Private/PCGSettings.cpp b/Engine/Plugins/PCG/Source/PCG/Private/PCGSettings.cpp
--- a/Engine/Plugins/PCG/Source/PCG/Private/PCGSettings.cpp
+++ b/Engine/Plugins/PCG/Source/PCG/Private/PCGSettings.cpp
@@ -560,39 +560,6 @@ EPCGDataType UPCGSettings::GetTypeUnionOfIncidentEdges(const FName& PinLabel) co
 	return Result;
 }
 
-bool UPCGSettings::AreInputEdgesValid(FPCGContext* InContext, bool bQuiet) const
-{
-	bool bAllEdgesValid = true;
-
-	if (const UPCGNode* Node = Cast<UPCGNode>(GetOuter()))
-	{
-		for (const UPCGPin* InputPin : Node->GetInputPins())
-		{
-			if (!InputPin)
-			{
-				continue;
-			}
-
-			for (const UPCGEdge* InputEdge : InputPin->Edges)
-			{
-				const UPCGPin* UpstreamPin = InputEdge ? InputEdge->GetOtherPin(InputPin) : nullptr;
-				if (UpstreamPin && InputPin->GetRequiredTypeConversion(UpstreamPin) != EPCGTypeConversion::NoConversionRequired)
-				{
-					PCG_KERNEL_VALIDATION_ERR(InContext, this, bQuiet, FText::Format(
-						LOCTEXT("InvalidInputPinEdge", "Unsupported connected upstream pin '{0}' on node '{1}' with type {2}. Recreate the edge to add required conversion nodes."),
-						FText::FromName(UpstreamPin->Properties.Label),
-						Node->GetNodeTitle(EPCGNodeTitleType::ListView),
-						StaticEnum<EPCGDataType>() ? StaticEnum<EPCGDataType>()->GetDisplayNameTextByValue(static_cast<int64>(UpstreamPin->Properties.AllowedTypes)) : FText::FromString(TEXT("MISSING"))));
-
-					bAllEdgesValid = false;
-				}
-			}
-		}
-	}
-
-	return bAllEdgesValid;
-}
-
 TArray<FPCGPinProperties> UPCGSettings::AllOutputPinProperties() const
 {
 	return OutputPinProperties();
diff --git a/Engine/Plugins/PCG/Source/PCG/Public/Elements/Metadata/PCGMetadataPartitionKernel.h b/Engine/Plugins/PCG/Source/PCG/Public/Elements/Metadata/PCGMetadataPartitionKernel.h
--- a/Engine/Plugins/PCG/Source/PCG/Public/Elements/Metadata/PCGMetadataPartitionKernel.h
+++ b/Engine/Plugins/PCG/Source/PCG/Public/Elements/Metadata/PCGMetadataPartitionKernel.h
@@ -14,10 +14,8 @@ class UPCGMetadataPartitionKernel : public UPCGComputeKernel
 	GENERATED_BODY()
 
 public:
-	void Initialize(const UPCGMetadataPartitionSettings* InSettings, bool bLogDescriptions);
-
 	virtual const UPCGSettings* GetSettings() const override;
-	virtual bool IsKernelValid(FPCGContext* InContext = nullptr, bool bQuiet = true) const override;
+	virtual bool IsKernelValid(FPCGContext* InContext) const override;
 	virtual bool ComputeOutputBindingDataDesc(const UPCGComputeGraph* InGraph, FName InOutputPinLabel, const UPCGDataBinding* InBinding, FPCGDataCollectionDesc& OutDataDesc) const override;
 	virtual int ComputeThreadCount(const UPCGDataBinding* Binding) const override;
 #if WITH_EDITOR
@@ -32,7 +30,8 @@ public:
 	virtual void GetInputPins(TArray<FPCGPinProperties>& OutPins) const override;
 	virtual void GetOutputPins(TArray<FPCGPinPropertiesGPU>& OutPins) const override;
 
-public:
-	UPROPERTY()
-	TObjectPtr<const UPCGMetadataPartitionSettings> Settings;
+protected:
+#if WITH_EDITOR
+	virtual bool PerformStaticValidation() override;
+#endif
 };
diff --git a/Engine/Plugins/PCG/Source/PCG/Public/Elements/PCGCopyPointsAnalysisKernel.h b/Engine/Plugins/PCG/Source/PCG/Public/Elements/PCGCopyPointsAnalysisKernel.h
--- a/Engine/Plugins/PCG/Source/PCG/Public/Elements/PCGCopyPointsAnalysisKernel.h
+++ b/Engine/Plugins/PCG/Source/PCG/Public/Elements/PCGCopyPointsAnalysisKernel.h
@@ -14,10 +14,8 @@ class UPCGCopyPointsAnalysisKernel : public UPCGComputeKernel
 	GENERATED_BODY()
 
 public:
-	void Initialize(const UPCGCopyPointsSettings* InSettings, bool bLogDescriptions);
-
 	virtual const UPCGSettings* GetSettings() const override;
-	virtual bool IsKernelValid(FPCGContext* InContext = nullptr, bool bQuiet = true) const override;
+	virtual bool IsKernelValid(FPCGContext* InContext) const override;
 	virtual bool ComputeOutputBindingDataDesc(const UPCGComputeGraph* InGraph, FName InOutputPinLabel, const UPCGDataBinding* InBinding, FPCGDataCollectionDesc& OutDataDesc) const override;
 	virtual int ComputeThreadCount(const UPCGDataBinding* Binding) const override;
 #if WITH_EDITOR
@@ -31,8 +29,4 @@ public:
 
 	virtual void GetInputPins(TArray<FPCGPinProperties>& OutPins) const override;
 	virtual void GetOutputPins(TArray<FPCGPinPropertiesGPU>& OutPins) const override;
-
-public:
-	UPROPERTY()
-	TObjectPtr<const UPCGCopyPointsSettings> Settings;
 };
diff --git a/Engine/Plugins/PCG/Source/PCG/Public/Elements/PCGCopyPointsKernel.h b/Engine/Plugins/PCG/Source/PCG/Public/Elements/PCGCopyPointsKernel.h
--- a/Engine/Plugins/PCG/Source/PCG/Public/Elements/PCGCopyPointsKernel.h
+++ b/Engine/Plugins/PCG/Source/PCG/Public/Elements/PCGCopyPointsKernel.h
@@ -16,10 +16,8 @@ class UPCGCopyPointsKernel : public UPCGComputeKernel
 	GENERATED_BODY()
 
 public:
-	void Initialize(const UPCGCopyPointsSettings* InSettings, bool bLogDescriptions);
-
 	virtual const UPCGSettings* GetSettings() const override;
-	virtual bool IsKernelValid(FPCGContext* InContext = nullptr, bool bQuiet = true) const override;
+	virtual bool IsKernelValid(FPCGContext* InContext) const override;
 	virtual bool ComputeOutputBindingDataDesc(const UPCGComputeGraph* InGraph, FName InOutputPinLabel, const UPCGDataBinding* InBinding, FPCGDataCollectionDesc& OutDataDesc) const override;
 	virtual int ComputeThreadCount(const UPCGDataBinding* Binding) const override;
 #if WITH_EDITOR
@@ -29,8 +27,4 @@ public:
 #endif
 	virtual void GetInputPins(TArray<FPCGPinProperties>& OutPins) const override;
 	virtual void GetOutputPins(TArray<FPCGPinPropertiesGPU>& OutPins) const override;
-
-public:
-	UPROPERTY()
-	TObjectPtr<const UPCGCopyPointsSettings> Settings;
 };
diff --git a/Engine/Plugins/PCG/Source/PCG/Public/Elements/PCGCopyPointsKernelShared.h b/Engine/Plugins/PCG/Source/PCG/Public/Elements/PCGCopyPointsKernelShared.h
--- a/Engine/Plugins/PCG/Source/PCG/Public/Elements/PCGCopyPointsKernelShared.h
+++ b/Engine/Plugins/PCG/Source/PCG/Public/Elements/PCGCopyPointsKernelShared.h
@@ -10,5 +10,5 @@ struct FPCGContext;
 namespace PCGCopyPointsKernel
 {
 	/** Performs validation common to all copy points kernels. */
-	bool IsKernelValid(const UPCGComputeKernel* InKernel, const UPCGCopyPointsSettings* InCopyPointSettings, const FPCGComputeGraphContext* InContext, bool bQuiet);
+	bool IsKernelValid(const UPCGComputeKernel* InKernel, const UPCGCopyPointsSettings* InCopyPointSettings, const FPCGComputeGraphContext* InContext);
 }
diff --git a/Engine/Plugins/PCG/Source/PCG/Public/Elements/PCGStaticMeshSpawnerKernel.h b/Engine/Plugins/PCG/Source/PCG/Public/Elements/PCGStaticMeshSpawnerKernel.h
--- a/Engine/Plugins/PCG/Source/PCG/Public/Elements/PCGStaticMeshSpawnerKernel.h
+++ b/Engine/Plugins/PCG/Source/PCG/Public/Elements/PCGStaticMeshSpawnerKernel.h
@@ -19,10 +19,7 @@ class UPCGStaticMeshSpawnerKernel : public UPCGComputeKernel
 	GENERATED_BODY()
 
 public:
-	void Initialize(const UPCGStaticMeshSpawnerSettings* InSettings, bool bLogDescriptions);
-
 	virtual const UPCGSettings* GetSettings() const override;
-	virtual bool IsKernelValid(FPCGContext* InContext = nullptr, bool bQuiet = true) const override;
 	virtual bool ComputeOutputBindingDataDesc(const UPCGComputeGraph* InGraph, FName InOutputPinLabel, const UPCGDataBinding* InBinding, FPCGDataCollectionDesc& OutDataDesc) const override;
 	virtual int ComputeThreadCount(const UPCGDataBinding* Binding) const override;
 #if WITH_EDITOR
@@ -37,7 +34,8 @@ public:
 	virtual void GetInputPins(TArray<FPCGPinProperties>& OutPins) const override;
 	virtual void GetOutputPins(TArray<FPCGPinPropertiesGPU>& OutPins) const override;
 
-public:
-	UPROPERTY()
-	TObjectPtr<const UPCGStaticMeshSpawnerSettings> Settings;
+protected:
+#if WITH_EDITOR
+	virtual bool PerformStaticValidation() override;
+#endif
 };
diff --git a/Engine/Plugins/PCG/Source/PCG/Public/PCGSettings.h b/Engine/Plugins/PCG/Source/PCG/Public/PCGSettings.h
--- a/Engine/Plugins/PCG/Source/PCG/Public/PCGSettings.h
+++ b/Engine/Plugins/PCG/Source/PCG/Public/PCGSettings.h
@@ -227,9 +227,6 @@ public:
 	/** Bitwise union of the allowed types of each incident edge on pin. Returns None if no common bits, or no edges. */
 	EPCGDataType GetTypeUnionOfIncidentEdges(const FName& PinLabel) const;
 
-	/** Validate input edges, for example checking connected pin types are compatible. */
-	bool AreInputEdgesValid(FPCGContext* InContext, bool bQuiet) const;
-
 	// Internal functions, should not be used by any user.
 	// Return a different subset for input/output pin properties, in case of a default object.
 	virtual TArray<FPCGPinProperties> DefaultInputPinProperties() const;
