Find actor/component by tag now allows wildcards (* and ?).

diff --git a/Engine/Plugins/PCG/Source/PCG/Private/Elements/PCGActorSelector.cpp b/Engine/Plugins/PCG/Source/PCG/Private/Elements/PCGActorSelector.cpp
--- a/Engine/Plugins/PCG/Source/PCG/Private/Elements/PCGActorSelector.cpp
+++ b/Engine/Plugins/PCG/Source/PCG/Private/Elements/PCGActorSelector.cpp
@@ -7,6 +7,7 @@
 #include "Grid/PCGPartitionActor.h"
 #include "Helpers/PCGActorHelpers.h"
 
+#include "Algo/AnyOf.h"
 #include "Components/ActorComponent.h"
 #include "GameFramework/Actor.h"
 
@@ -14,6 +15,43 @@
 
 namespace PCGActorSelector
 {
+	template<typename NameRangeType>
+	bool TagListMatchesTag(const NameRangeType& Range, const FName& TargetTag, bool bUseWildcardSearch, const FString& WildcardTag)
+	{
+		if (bUseWildcardSearch)
+		{
+			return Algo::AnyOf(Range, [&WildcardTag](const FName& TagName) { return TagName.ToString().MatchesWildcard(WildcardTag); });
+		}
+		else
+		{
+			return Range.Contains(TargetTag);
+		}
+	}
+
+	void PrepareTag(bool bWillUseTag, FName InTag, bool& OutTagContainsWildcard, FString& OutTagString)
+	{
+		if (bWillUseTag)
+		{
+			int32 DummyIndex;
+			FString TagString = InTag.ToString();
+			if (TagString.FindChar(TCHAR('*'), DummyIndex) || TagString.FindChar(TCHAR('?'), DummyIndex))
+			{
+				OutTagString = std::move(TagString);
+				OutTagContainsWildcard = true;
+			}
+			else
+			{
+				OutTagString = FString();
+				OutTagContainsWildcard = false;
+			}
+		}
+		else
+		{
+			OutTagString = FString();
+			OutTagContainsWildcard = false;
+		}
+	}
+
 	// Filter is required if it is not disabled and if we are gathering all world actors or gathering all children.
 	bool FilterRequired(const FPCGActorSelectorSettings& InSettings)
 	{
@@ -21,8 +59,11 @@ namespace PCGActorSelector
 	}
 
 	// Need to pass a pointer of pointer to the found actor. The lambda will capture this pointer and modify its value when an actor is found.
+	// Note: depending on the filtering type, the settings must stay alive longer than the returned filtering function here.
 	TFunction<bool(AActor*)> GetFilteringFunction(const FPCGActorSelectorSettings& InSettings, const TFunction<bool(const AActor*)>& BoundsCheck, const TFunction<bool(const AActor*)>& SelfIgnoreCheck, TArray<AActor*>& InFoundActors)
 	{
+		InSettings.PrepareForFiltering(/*bForce=*/true);
+
 		if (!FilterRequired(InSettings))
 		{
 			return [&InFoundActors, &BoundsCheck, &SelfIgnoreCheck](AActor* Actor) -> bool
@@ -40,9 +81,9 @@ namespace PCGActorSelector
 		switch (InSettings.ActorSelection)
 		{
 		case EPCGActorSelection::ByTag:
-			return[ActorSelectionTag = InSettings.ActorSelectionTag, &InFoundActors, bMultiSelect, &BoundsCheck, &SelfIgnoreCheck](AActor* Actor) -> bool
+			return[&InSettings, &InFoundActors, bMultiSelect, &BoundsCheck, &SelfIgnoreCheck](AActor* Actor) -> bool
 			{
-				if (Actor && Actor->ActorHasTag(ActorSelectionTag) && BoundsCheck(Actor) && SelfIgnoreCheck(Actor))
+				if(Actor && InSettings.MatchesTag(Actor) && BoundsCheck(Actor) && SelfIgnoreCheck(Actor))
 				{
 					InFoundActors.Add(Actor);
 					return bMultiSelect;
@@ -284,6 +325,7 @@ FPCGSelectionKey::FPCGSelectionKey(FName InTag)
 	Selection = EPCGActorSelection::ByTag;
 	Tag = InTag;
 	ActorFilter = EPCGActorFilter::AllWorldActors;
+	UpdateAfterTagChange();
 }
 
 FPCGSelectionKey::FPCGSelectionKey(TSubclassOf<UObject> InSelectionClass)
@@ -356,7 +398,7 @@ bool FPCGSelectionKey::IsMatching(const UObject* InObject, const FSoftObjectPath
 	case EPCGActorSelection::ByTag:
 	{
 		const AActor* InActor = Cast<const AActor>(InObject);
-		return InActor && InActor->ActorHasTag(Tag);
+		return InActor && PCGActorSelector::TagListMatchesTag(InActor->Tags, Tag, bTagContainsWildcard, CachedTagString);
 	}
 	case EPCGActorSelection::ByClass:
 		return InObject && InObject->GetClass()->IsChildOf(SelectionClass);
@@ -441,7 +483,8 @@ bool FPCGSelectionKey::IsMatching(const UObject* InObject, const FSoftObjectPath
 	case EPCGActorSelection::ByTag:
 	{
 		const AActor* InActor = Cast<const AActor>(InObject);
-		bIsMatched = InRemovedTags.Contains(Tag) || (InActor && InActor->ActorHasTag(Tag));
+		bIsMatched = PCGActorSelector::TagListMatchesTag(InRemovedTags, Tag, bTagContainsWildcard, CachedTagString) ||
+			(InActor && PCGActorSelector::TagListMatchesTag(InActor->Tags, Tag, bTagContainsWildcard, CachedTagString));
 		break;
 	}
 	case EPCGActorSelection::ByClass:
@@ -489,6 +532,11 @@ bool FPCGSelectionKey::operator==(const FPCGSelectionKey& InOther) const
 	}
 }
 
+void FPCGSelectionKey::UpdateAfterTagChange()
+{
+	PCGActorSelector::PrepareTag(Selection == EPCGActorSelection::ByTag, Tag, bTagContainsWildcard, CachedTagString);
+}
+
 FArchive& operator<<(FArchive& Ar, FPCGSelectionKey& Key)
 {
 	// Serialize the normal UPROPERTY data
@@ -498,6 +546,11 @@ FArchive& operator<<(FArchive& Ar, FPCGSelectionKey& Key)
 		{
 			ThisStruct->SerializeTaggedProperties(Ar, reinterpret_cast<uint8*>(&Key), ThisStruct, nullptr);
 		}
+
+		if (Ar.IsLoading())
+		{
+			Key.UpdateAfterTagChange();
+		}
 	}
 
 	return Ar;
@@ -590,10 +643,35 @@ FPCGActorSelectorSettings FPCGActorSelectorSettings::ReconstructFromKey(const FP
 	Result.ActorSelection = InKey.Selection;
 	Result.ActorSelectionTag = InKey.Tag;
 	Result.ActorSelectionClass = InKey.SelectionClass;
+	Result.PrepareForFiltering();
 
 	return Result;
 }
 
+void FPCGActorSelectorSettings::PrepareForFiltering(bool bForce) const
+{
+	if (bForce || !bHasPreparedTag)
+	{
+		PCGActorSelector::PrepareTag(ActorSelection == EPCGActorSelection::ByTag, ActorSelectionTag, bTagContainsWildcards, ActorSelectionTagString);
+		bHasPreparedTag = true;
+	}
+}
+
+bool FPCGActorSelectorSettings::MatchesTag(AActor* Actor) const
+{
+	PrepareForFiltering();
+	return PCGActorSelector::TagListMatchesTag(Actor->Tags, ActorSelectionTag, bTagContainsWildcards, ActorSelectionTagString);
+}
+
+void FPCGComponentSelectorSettings::PrepareForFiltering(bool bForce) const
+{
+	if (bForce || !bHasPreparedTag)
+	{
+		PCGActorSelector::PrepareTag(ComponentSelection == EPCGComponentSelection::ByTag, ComponentSelectionTag, bTagContainsWildcards, ComponentSelectionTagString);
+		bHasPreparedTag = true;
+	}
+}
+
 bool FPCGComponentSelectorSettings::FilterComponent(UActorComponent* InComponent) const
 {
 	check(InComponent);
@@ -604,8 +682,19 @@ bool FPCGComponentSelectorSettings::FilterComponent(UActorComponent* InComponent
 	}
 	else
 	{
-		return (ComponentSelection == EPCGComponentSelection::ByTag) ? (ComponentSelectionTag == NAME_None || InComponent->ComponentTags.Contains(ComponentSelectionTag)) :
-			(ComponentSelectionClass == nullptr || ComponentSelectionClass == UActorComponent::StaticClass() || InComponent->GetClass()->IsChildOf(ComponentSelectionClass));
+		if (ComponentSelection == EPCGComponentSelection::ByTag)
+		{
+			PrepareForFiltering();
+
+			return ComponentSelectionTag == NAME_None ||
+				PCGActorSelector::TagListMatchesTag(InComponent->ComponentTags, ComponentSelectionTag, bTagContainsWildcards, ComponentSelectionTagString);
+		}
+		else
+		{
+			return ComponentSelectionClass == nullptr || 
+				ComponentSelectionClass == UActorComponent::StaticClass() || 
+				InComponent->GetClass()->IsChildOf(ComponentSelectionClass);
+		}
 	}
 }
 
@@ -650,9 +739,11 @@ TArray<UActorComponent*> FPCGComponentSelectorSettings::FilterComponents(TArrayV
 	}
 	else if (ComponentSelection == EPCGComponentSelection::ByTag)
 	{
+		PrepareForFiltering();
+
 		for (UActorComponent* Component : Components)
 		{
-			if (Component && Component->ComponentTags.Contains(ComponentSelectionTag))
+			if (Component && PCGActorSelector::TagListMatchesTag(Component->ComponentTags, ComponentSelectionTag, bTagContainsWildcards, ComponentSelectionTagString))
 			{
 				FilteredComponents.Add(Component);
 			}
diff --git a/Engine/Plugins/PCG/Source/PCG/Public/Elements/PCGActorSelector.h b/Engine/Plugins/PCG/Source/PCG/Public/Elements/PCGActorSelector.h
--- a/Engine/Plugins/PCG/Source/PCG/Public/Elements/PCGActorSelector.h
+++ b/Engine/Plugins/PCG/Source/PCG/Public/Elements/PCGActorSelector.h
@@ -83,6 +83,7 @@ struct PCG_API FPCGSelectionKey
 	bool IsMatching(const UObject* InObject, const FSoftObjectPath& InObjectPath, const TSet<FName>& InRemovedTags, const TSet<UPCGComponent*>& InComponents, TSet<UPCGComponent*>* OptionalMatchedComponents = nullptr) const;
 
 	void SetExtraDependency(const UClass* InExtraDependency);
+	void UpdateAfterTagChange();
 
 	UPROPERTY()
 	EPCGActorFilter ActorFilter = EPCGActorFilter::AllWorldActors;
@@ -109,6 +110,10 @@ struct PCG_API FPCGSelectionKey
 	bool IsMatching(const UObject* InObject, const UPCGComponent* InComponent) const;
 	UE_DEPRECATED(5.6, "Replaced by the version with a soft object path.")
 	bool IsMatching(const UObject* InObject, const TSet<FName>& InRemovedTags, const TSet<UPCGComponent*>& InComponents, TSet<UPCGComponent*>* OptionalMatchedComponents = nullptr) const;
+
+protected:
+	FString CachedTagString;
+	bool bTagContainsWildcard = false;
 };
 
 PCG_API FArchive& operator<<(FArchive& Ar, FPCGSelectionKey& Key);
@@ -185,6 +190,14 @@ struct PCG_API FPCGActorSelectorSettings
 
 	FPCGSelectionKey GetAssociatedKey() const;
 	static FPCGActorSelectorSettings ReconstructFromKey(const FPCGSelectionKey& InKey);
+
+	void PrepareForFiltering(bool bForce = false) const;
+	bool MatchesTag(AActor* Actor) const;
+
+protected:
+	mutable FString ActorSelectionTagString;
+	mutable bool bTagContainsWildcards = false;
+	mutable bool bHasPreparedTag = false;
 };
 
 USTRUCT(BlueprintType)
@@ -215,6 +228,13 @@ struct PCG_API FPCGComponentSelectorSettings
 	bool FilterComponent(UActorComponent* InComponent) const;
 	bool FilterActor(AActor* InActor) const;
 	TArray<UActorComponent*> FilterComponents(TArrayView<UActorComponent*> InComponents) const;
+
+	void PrepareForFiltering(bool bForce = false) const;
+
+protected:
+	mutable FString ComponentSelectionTagString;
+	mutable bool bTagContainsWildcards = false;
+	mutable bool bHasPreparedTag = false;
 };
 
 namespace PCGActorSelector
