[PCG] Fix mac build issue with parsed HLSL functions

#rnx

diff --git a/Engine/Plugins/PCG/Source/PCG/Private/Compute/Elements/PCGCustomHLSLKernel.h b/Engine/Plugins/PCG/Source/PCG/Private/Compute/Elements/PCGCustomHLSLKernel.h
--- a/Engine/Plugins/PCG/Source/PCG/Private/Compute/Elements/PCGCustomHLSLKernel.h
+++ b/Engine/Plugins/PCG/Source/PCG/Private/Compute/Elements/PCGCustomHLSLKernel.h
@@ -16,6 +16,15 @@ struct FPCGParsedAttributeFunction
 	GENERATED_BODY()
 
 public:
+	FPCGParsedAttributeFunction() = default;
+	FPCGParsedAttributeFunction(const FString& InPinLabel, const FString& InFunctionName, int64 InAttributeType, const FString& InAttributeName, int32 InMatchBeginning)
+		: PinLabel(InPinLabel)
+		, FunctionName(InFunctionName)
+		, AttributeType(InAttributeType)
+		, AttributeName(InAttributeName)
+		, MatchBeginning(InMatchBeginning)
+	{}
+
 	UPROPERTY()
 	FString PinLabel;
 
@@ -38,6 +47,12 @@ struct FPCGParsedCopyElementFunction
 	GENERATED_BODY()
 
 public:
+	FPCGParsedCopyElementFunction() = default;
+	FPCGParsedCopyElementFunction(const FString& InSourcePin, const FString& InTargetPin)
+		: SourcePin(InSourcePin)
+		, TargetPin(InTargetPin)
+	{}
+
 	UPROPERTY()
 	FString SourcePin;
 
