Lowered tolerance for module placement on splines as it would create too much seams and exposed a control for it.

diff --git a/Engine/Plugins/PCG/Source/PCG/Private/Elements/Grammar/PCGSubdivideSpline.cpp b/Engine/Plugins/PCG/Source/PCG/Private/Elements/Grammar/PCGSubdivideSpline.cpp
--- a/Engine/Plugins/PCG/Source/PCG/Private/Elements/Grammar/PCGSubdivideSpline.cpp
+++ b/Engine/Plugins/PCG/Source/PCG/Private/Elements/Grammar/PCGSubdivideSpline.cpp
@@ -23,6 +23,7 @@ namespace PCGSubdivideSplineHelpers
 		bool bAcceptIncompleteSubdivision = false;
 		double ModuleHeight = 0.0;
 		int32 AdditionalSeed = 0;
+		double BisectionTolerance = 0.01;
 
 		const UPCGPolyLineData* PolyLineData = nullptr;
 		TArray<FPCGPoint>* OutPoints = nullptr;
@@ -47,7 +48,7 @@ namespace PCGSubdivideSplineHelpers
 	}
 
 	// Using a numerical method (s.a. bisection), determine the spline alpha from a segment length. Bisection can be slower, but guaranteed to converge.
-	static double FindRootAtLinearDistance_Bisection(const UPCGPolyLineData* PolyLineData, const double SegmentLength, const double StartingAlpha, const double Tolerance = 1)
+	static double FindRootAtLinearDistance_Bisection(const UPCGPolyLineData* PolyLineData, const double SegmentLength, const double StartingAlpha, const double Tolerance)
 	{
 		check(PolyLineData);
 
@@ -160,7 +161,8 @@ namespace PCGSubdivideSplineHelpers
 				if (!bAtSplineEnd)
 				{
 					const double PreviousAlpha = ModuleAlphas.Last();
-					const double CurrentAlpha = FindRootAtLinearDistance_Bisection(InOutParameters.PolyLineData, SubmoduleSize, PreviousAlpha);
+					const double BisectionTolerance = FMath::Max(FMath::Min(InOutParameters.BisectionTolerance, SubmoduleSize), PCGSubdivideSplineConstants::MinimumBisectionTolerance);
+					const double CurrentAlpha = FindRootAtLinearDistance_Bisection(InOutParameters.PolyLineData, SubmoduleSize, PreviousAlpha, BisectionTolerance);
 
 					ModuleAlphas.Add(CurrentAlpha);
 
@@ -331,7 +333,8 @@ bool FPCGSubdivideSplineElement::ExecuteInternal(FPCGContext* InContext) const
 	{
 		.Context = InContext,
 		.ModulesInfo = ModulesInfo,
-		.bAcceptIncompleteSubdivision = Settings->bAcceptIncompleteSubdivision
+		.bAcceptIncompleteSubdivision = Settings->bAcceptIncompleteSubdivision,
+		.BisectionTolerance = Settings->ModulePlacementTolerance
 	};
 
 	for (const FPCGTaggedData& Input : Inputs)
diff --git a/Engine/Plugins/PCG/Source/PCG/Public/Elements/Grammar/PCGSubdivideSpline.h b/Engine/Plugins/PCG/Source/PCG/Public/Elements/Grammar/PCGSubdivideSpline.h
--- a/Engine/Plugins/PCG/Source/PCG/Public/Elements/Grammar/PCGSubdivideSpline.h
+++ b/Engine/Plugins/PCG/Source/PCG/Public/Elements/Grammar/PCGSubdivideSpline.h
@@ -8,6 +8,11 @@
 
 #include "PCGSubdivideSpline.generated.h"
 
+namespace PCGSubdivideSplineConstants
+{
+	constexpr double MinimumBisectionTolerance = 1.0e-6;
+}
+
 UCLASS(MinimalAPI, BlueprintType, ClassGroup = (Procedural))
 class UPCGSubdivideSplineSettings : public UPCGSubdivisionBaseSettings
 {
@@ -61,6 +66,10 @@ public:
 	/** Name of the attribute labeling the final output point from the final module. */
 	UPROPERTY(BlueprintReadWrite, EditAnywhere, Category = "Settings|Extra Output Attributes", meta = (PCG_Overridable, EditCondition = "bOutputExtremityAttributes", EditConditionHides))
 	FName IsFinalAttributeName = TEXT("IsFinal");
+
+	/** Maximum acceptable distance tolerance between placed modules (overlap or gap) - a smaller number is more precise. Affects performance if the number is very small. */
+	UPROPERTY(BlueprintReadWrite, EditAnywhere, Category = "Settings|Advanced", meta = (PCG_Overridable, ClampMin="0.000001"))
+	double ModulePlacementTolerance = 0.01; 
 };
 
 class FPCGSubdivideSplineElement : public FPCGSubdivisionBaseElement
