[PCG] Add StartOffset/EndOffset to FPCGSplineMeshParams

* Make them editable in the Spline Param struct
* Remove the possibility to override the un-exposed parameters

#rnx

diff --git a/Engine/Plugins/PCG/Source/PCG/Private/Helpers/PCGActorHelpers.cpp b/Engine/Plugins/PCG/Source/PCG/Private/Helpers/PCGActorHelpers.cpp
--- a/Engine/Plugins/PCG/Source/PCG/Private/Helpers/PCGActorHelpers.cpp
+++ b/Engine/Plugins/PCG/Source/PCG/Private/Helpers/PCGActorHelpers.cpp
@@ -475,6 +475,8 @@ UPCGManagedSplineMeshComponent* UPCGActorHelpers::GetOrCreateManagedSplineMeshCo
 		SplineMeshComponent->SetEndScale(SplineMeshParams.EndScale);
 		SplineMeshComponent->SetForwardAxis((ESplineMeshAxis::Type)SplineMeshParams.ForwardAxis);
 		SplineMeshComponent->SetSplineUpDir(SplineMeshParams.SplineUpDir);
+		SplineMeshComponent->SetStartOffset(SplineMeshParams.StartOffset);
+		SplineMeshComponent->SetEndOffset(SplineMeshParams.EndOffset);
 		SplineMeshComponent->SplineParams.NaniteClusterBoundsScale = SplineMeshParams.NaniteClusterBoundsScale;
 		SplineMeshComponent->SplineBoundaryMin = SplineMeshParams.SplineBoundaryMin;
 		SplineMeshComponent->SplineBoundaryMax = SplineMeshParams.SplineBoundaryMax;
diff --git a/Engine/Plugins/PCG/Source/PCG/Public/Elements/PCGSplineMeshParams.h b/Engine/Plugins/PCG/Source/PCG/Public/Elements/PCGSplineMeshParams.h
--- a/Engine/Plugins/PCG/Source/PCG/Public/Elements/PCGSplineMeshParams.h
+++ b/Engine/Plugins/PCG/Source/PCG/Public/Elements/PCGSplineMeshParams.h
@@ -53,30 +53,37 @@ struct FPCGSplineMeshParams
 	UPROPERTY(BlueprintReadWrite, EditAnywhere, Category = "Spline Mesh Params", meta = (PCG_Overridable))
 	bool bSmoothInterpRollScale = true;
 
-	// Below properties are not exposed for edit because they are computed internally.
+	/** Starting offset of the mesh from the spline, in component space. */
+	UPROPERTY(BlueprintReadWrite, EditAnywhere, Category = "Spline Mesh Params", meta = (PCG_Overridable))
+	FVector2D StartOffset = FVector2D::ZeroVector;
 
-	UPROPERTY()
+	/** Ending offset of the mesh from the spline, in component space. */
+	UPROPERTY(BlueprintReadWrite, EditAnywhere, Category = "Spline Mesh Params", meta = (PCG_Overridable))
+	FVector2D EndOffset = FVector2D::ZeroVector;
+
+	// Below properties are not exposed for edit because they are computed internally.
+	UPROPERTY(meta = (PCG_NotOverridable))
 	FVector StartPosition = FVector::ZeroVector;
 
-	UPROPERTY()
+	UPROPERTY(meta = (PCG_NotOverridable))
 	FVector StartTangent = FVector::ZeroVector;
 
-	UPROPERTY()
+	UPROPERTY(meta = (PCG_NotOverridable))
 	float StartRollDegrees = 0.0f;
 
-	UPROPERTY()
+	UPROPERTY(meta = (PCG_NotOverridable))
 	FVector2D StartScale = FVector2D::ZeroVector;
 
-	UPROPERTY()
+	UPROPERTY(meta = (PCG_NotOverridable))
 	FVector EndPosition = FVector::ZeroVector;
 
-	UPROPERTY()
+	UPROPERTY(meta = (PCG_NotOverridable))
 	FVector EndTangent = FVector::ZeroVector;
 
-	UPROPERTY()
+	UPROPERTY(meta = (PCG_NotOverridable))
 	float EndRollDegrees = 0.0f;
 
-	UPROPERTY()
+	UPROPERTY(meta = (PCG_NotOverridable))
 	FVector2D EndScale = FVector2D::ZeroVector;
 
 	friend uint32 GetTypeHash(const FPCGSplineMeshParams& Key);
