[PCG] Allow for sequential load in AsyncLoading

* Now also use the ContextHandle for safety to pass the context as parameter of the completion callback lambda

#rnx

diff --git a/Engine/Plugins/PCG/Source/PCG/Private/Async/PCGAsyncLoadingContext.cpp b/Engine/Plugins/PCG/Source/PCG/Private/Async/PCGAsyncLoadingContext.cpp
--- a/Engine/Plugins/PCG/Source/PCG/Private/Async/PCGAsyncLoadingContext.cpp
+++ b/Engine/Plugins/PCG/Source/PCG/Private/Async/PCGAsyncLoadingContext.cpp
@@ -22,6 +22,20 @@ void IPCGAsyncLoadingContext::CancelLoading()
 	LoadHandle.Reset();
 }
 
+bool IPCGAsyncLoadingContext::ResetLoadRequestIfDone()
+{
+	if (bLoadRequested && LoadHandle.IsValid() && LoadHandle->HasLoadCompleted())
+	{
+		LoadHandle.Reset();
+		bLoadRequested = false;
+		return true;
+	}
+	else
+	{
+		return false;
+	}
+}
+
 void IPCGAsyncLoadingContext::ResolveAlreadyLoadedObjects(TArray<FSoftObjectPath>& ObjectsToLoad)
 {
 	for (int32 i = ObjectsToLoad.Num() - 1; i >= 0; --i)
@@ -62,8 +76,16 @@ bool IPCGAsyncLoadingContext::RequestResourceLoad(FPCGContext* ThisContext, TArr
 		{
 			ThisContext->bIsPaused = true;
 
-			// It is a bit unsafe to pass this to a delegate lambda. But if the context dies before the completion of the loading, the context will cancel the loading in its dtor.
-			LoadHandle = UAssetManager::GetStreamableManager().RequestAsyncLoad(std::move(ObjectsToLoad), [ThisContext]() { ThisContext->bIsPaused = false; });
+			LoadHandle = UAssetManager::GetStreamableManager().RequestAsyncLoad(std::move(ObjectsToLoad),
+				[ContextHandle = ThisContext->GetOrCreateHandle()]()
+				{
+					FPCGContext::FSharedContext<FPCGContext> SharedContext(ContextHandle);
+					if (FPCGContext* Context = SharedContext.Get())
+					{
+						Context->bIsPaused = false;
+					}
+				});
+
 			bLoadRequested = true;
 
 			// If the load handle is not active it means objects were invalid
diff --git a/Engine/Plugins/PCG/Source/PCG/Public/Async/PCGAsyncLoadingContext.h b/Engine/Plugins/PCG/Source/PCG/Public/Async/PCGAsyncLoadingContext.h
--- a/Engine/Plugins/PCG/Source/PCG/Public/Async/PCGAsyncLoadingContext.h
+++ b/Engine/Plugins/PCG/Source/PCG/Public/Async/PCGAsyncLoadingContext.h
@@ -27,6 +27,13 @@ public:
 
 	bool WasLoadRequested() const { return bLoadRequested; }
 
+	/**
+	 * If we need multiple loads (like ObjectA first load, and has a soft ObjectB that also needs to be loaded), we can use this function to reset the bLoadRequested flag.
+	 * Only works if the loading was requested and is done.
+	 * Returns true if a new loading can be requested.
+	 */
+	bool ResetLoadRequestIfDone();
+
 private:
 	/** Resolve already loaded objects and store a reference on them and also remove null paths, will modify in place the array of objects that actually require loading. */
 	void ResolveAlreadyLoadedObjects(TArray<FSoftObjectPath>& ObjectsToLoad);
