[PCG] Fix typos in profiling scopes

#rnx

diff --git a/Engine/Plugins/PCG/Source/PCG/Private/Compute/DataInterfaces/Elements/PCGStaticMeshSpawnerDataInterface.cpp b/Engine/Plugins/PCG/Source/PCG/Private/Compute/DataInterfaces/Elements/PCGStaticMeshSpawnerDataInterface.cpp
--- a/Engine/Plugins/PCG/Source/PCG/Private/Compute/DataInterfaces/Elements/PCGStaticMeshSpawnerDataInterface.cpp
+++ b/Engine/Plugins/PCG/Source/PCG/Private/Compute/DataInterfaces/Elements/PCGStaticMeshSpawnerDataInterface.cpp
@@ -209,7 +209,7 @@ void UPCGStaticMeshSpawnerDataProvider::Initialize(const UComputeDataInterface*
 
 bool UPCGStaticMeshSpawnerDataProvider::PerformPreExecuteReadbacks_GameThread(UPCGDataBinding* InBinding)
 {
-	TRACE_CPUPROFILER_EVENT_SCOPE(UPCGCopyPointsDataProvider::PerformPreExecuteReadbacks_GameThread);
+	TRACE_CPUPROFILER_EVENT_SCOPE(UPCGStaticMeshSpawnerDataProvider::PerformPreExecuteReadbacks_GameThread);
 	check(InBinding);
 
 	if (!Super::PerformPreExecuteReadbacks_GameThread(InBinding))
@@ -282,7 +282,7 @@ bool UPCGStaticMeshSpawnerDataProvider::PerformPreExecuteReadbacks_GameThread(UP
 
 bool UPCGStaticMeshSpawnerDataProvider::PrepareForExecute_GameThread(UPCGDataBinding* InBinding)
 {
-	TRACE_CPUPROFILER_EVENT_SCOPE(UPCGCopyPointsDataProvider::PerformPreExecuteReadbacks_GameThread);
+	TRACE_CPUPROFILER_EVENT_SCOPE(UPCGStaticMeshSpawnerDataProvider::PerformPreExecuteReadbacks_GameThread);
 	check(InBinding);
 	const UPCGStaticMeshSpawnerSettings* Settings = CastChecked<UPCGStaticMeshSpawnerSettings>(ProducerKernel->GetSettings());
 
