[PCG] New UPCGSpatialData::PrepareForSpatialQuery API : Allows generating new tasks that will prepare the data instead of loading/creating data on usage
- First usage: PCGSurfaceSampler (pcg.Experimental.EnableDynamicDependencies is enabled)
- API implemented by PCGLandscapeData (current implementation doesn't use async bulk serialize but it could to avoid waits)
- Disabled by default CVar: pcg.Experimental.SpatialData.EnablePrepareForSpatialQuery
- Depends on pcg.Experimental.EnableDynamicDependencies being enabled

#rnx

diff --git a/Engine/Plugins/PCG/Source/PCG/Private/Data/PCGLandscapeData.cpp b/Engine/Plugins/PCG/Source/PCG/Private/Data/PCGLandscapeData.cpp
--- a/Engine/Plugins/PCG/Source/PCG/Private/Data/PCGLandscapeData.cpp
+++ b/Engine/Plugins/PCG/Source/PCG/Private/Data/PCGLandscapeData.cpp
@@ -2,6 +2,7 @@
 
 #include "Data/PCGLandscapeData.h"
 
+#include "PCGComponent.h"
 #include "PCGContext.h"
 #include "PCGSubsystem.h"
 #include "Data/PCGPointArrayData.h"
@@ -380,12 +381,7 @@ bool UPCGLandscapeData::ProjectPoint(const FTransform& InTransform, const FBox&
 	const FVector LocalPoint = LandscapeTransform.InverseTransformPosition(InTransform.GetLocation());
 	const FIntPoint ComponentMapKey(FMath::FloorToInt(LocalPoint.X / LandscapeInfo->ComponentSizeQuads), FMath::FloorToInt(LocalPoint.Y / LandscapeInfo->ComponentSizeQuads));
 
-#if WITH_EDITOR
-	ULandscapeComponent* LandscapeComponent = LandscapeInfo->XYtoComponentMap.FindRef(ComponentMapKey);
-	const FPCGLandscapeCacheEntry * LandscapeCacheEntry = (LandscapeComponent ? LandscapeCache->GetCacheEntry(LandscapeComponent, ComponentMapKey) : nullptr);
-#else
-	const FPCGLandscapeCacheEntry* LandscapeCacheEntry = LandscapeCache->GetCacheEntry(LandscapeProxy, LandscapeProxy->GetOriginalLandscapeGuid(), ComponentMapKey);
-#endif
+	const FPCGLandscapeCacheEntry* LandscapeCacheEntry = LandscapeCache->GetCacheEntry(LandscapeInfo, ComponentMapKey);
 
 	if (!LandscapeCacheEntry)
 	{
@@ -568,12 +564,8 @@ const UPCGBasePointData* UPCGLandscapeData::CreateBasePointData(FPCGContext* Con
 			for (int32 ComponentY = MinComponentY; ComponentY <= MaxComponentY; ++ComponentY)
 			{
 				FIntPoint ComponentMapKey(ComponentX, ComponentY);
-#if WITH_EDITOR
-				ULandscapeComponent* LandscapeComponent = LandscapeInfo->XYtoComponentMap.FindRef(ComponentMapKey);
-				const FPCGLandscapeCacheEntry* LandscapeCacheEntry = LandscapeComponent ? LandscapeCache->GetCacheEntry(LandscapeComponent, ComponentMapKey) : nullptr;
-#else
-				const FPCGLandscapeCacheEntry* LandscapeCacheEntry = LandscapeCache->GetCacheEntry(LandscapeProxy, LandscapeGuid, ComponentMapKey);
-#endif
+
+				const FPCGLandscapeCacheEntry* LandscapeCacheEntry = LandscapeCache->GetCacheEntry(LandscapeInfo, ComponentMapKey);
 
 				if (!LandscapeCacheEntry)
 				{
@@ -661,3 +653,57 @@ UPCGSpatialData* UPCGLandscapeData::CopyInternal(FPCGContext* Context) const
 
 	return NewLandscapeData;
 }
+
+TArray<FPCGTaskId> UPCGLandscapeData::PrepareForSpatialQuery(FPCGContext* InContext, const FBox& InBounds) const
+{
+	TRACE_CPUPROFILER_EVENT_SCOPE(UPCGLandscapeData::PrepareForSpatialQuery);
+	if (!PCGSpatialData::CVarEnablePrepareForSpatialQuery.GetValueOnAnyThread())
+	{
+		return {};
+	}
+
+	UPCGSubsystem* Subsystem = InContext->SourceComponent.IsValid() ? InContext->SourceComponent->GetSubsystem() : nullptr;
+	if (!Subsystem)
+	{
+		return {};
+	}
+
+	const ULandscapeInfo* LandscapeInfo = GetLandscapeInfo(InBounds.GetCenter());
+	ALandscapeProxy* LandscapeProxy = LandscapeInfo ? LandscapeInfo->GetLandscapeProxy() : nullptr;
+	if (!LandscapeProxy)
+	{
+		return {};
+	}
+
+	const FTransform LandscapeTransform = LandscapeProxy->LandscapeActorToWorld();
+
+	FBox InverseTransformedBounds = InBounds.InverseTransformBy(LandscapeTransform);
+
+	const FVector MinPoint = InverseTransformedBounds.Min;
+	const FVector MaxPoint = InverseTransformedBounds.Max;
+
+	const FIntPoint MinComponentMapKey(FMath::FloorToInt(MinPoint.X / LandscapeInfo->ComponentSizeQuads), FMath::FloorToInt(MinPoint.Y / LandscapeInfo->ComponentSizeQuads));
+	const FIntPoint MaxComponentMapKey(FMath::FloorToInt(MaxPoint.X / LandscapeInfo->ComponentSizeQuads), FMath::FloorToInt(MaxPoint.Y / LandscapeInfo->ComponentSizeQuads));
+
+	// Already cached
+	if (LandscapeCache->AreCacheEntriesReady(LandscapeInfo, MinComponentMapKey, MaxComponentMapKey))
+	{
+		return {};
+	}
+
+	FPCGScheduleGenericParams Params([WeakThis = TWeakObjectPtr<const UPCGLandscapeData>(this), WeakLandscapeInfo = TWeakObjectPtr<const ULandscapeInfo>(LandscapeInfo), MinComponentMapKey, MaxComponentMapKey](FPCGContext* InContext)
+	{
+		const ULandscapeInfo* LandscapeInfo = WeakLandscapeInfo.Get();
+		const UPCGLandscapeData* LandscapeData = WeakThis.Get();
+		if (LandscapeInfo && LandscapeData)
+		{
+			LandscapeData->LandscapeCache->PrimeCache(LandscapeInfo, MinComponentMapKey, MaxComponentMapKey);
+		}
+				
+		return true;
+	},
+	InContext->SourceComponent.Get());
+	Params.bCanExecuteOnlyOnMainThread = false;
+
+	return { Subsystem->ScheduleGeneric(Params) };
+}
\ No newline at end of file
diff --git a/Engine/Plugins/PCG/Source/PCG/Private/Data/PCGSpatialData.cpp b/Engine/Plugins/PCG/Source/PCG/Private/Data/PCGSpatialData.cpp
--- a/Engine/Plugins/PCG/Source/PCG/Private/Data/PCGSpatialData.cpp
+++ b/Engine/Plugins/PCG/Source/PCG/Private/Data/PCGSpatialData.cpp
@@ -15,6 +15,14 @@
 
 #include UE_INLINE_GENERATED_CPP_BY_NAME(PCGSpatialData)
 
+namespace PCGSpatialData
+{
+	TAutoConsoleVariable<bool> CVarEnablePrepareForSpatialQuery(
+		TEXT("pcg.Experimental.SpatialData.EnablePrepareForSpatialQuery"),
+		false,
+		TEXT("Enable UPCGSpatialData subclass PrepareForSpatialQuery task scheduling"));
+}
+
 UPCGSpatialData::UPCGSpatialData(const FObjectInitializer& ObjectInitializer)
 	: Super(ObjectInitializer)
 {
diff --git a/Engine/Plugins/PCG/Source/PCG/Private/Elements/PCGSurfaceSampler.cpp b/Engine/Plugins/PCG/Source/PCG/Private/Elements/PCGSurfaceSampler.cpp
--- a/Engine/Plugins/PCG/Source/PCG/Private/Elements/PCGSurfaceSampler.cpp
+++ b/Engine/Plugins/PCG/Source/PCG/Private/Elements/PCGSurfaceSampler.cpp
@@ -22,6 +22,11 @@
 
 #define LOCTEXT_NAMESPACE "PCGSurfaceSamplerElement"
 
+namespace PCGGraphExecutor
+{
+	extern TAutoConsoleVariable<bool> CVarEnableDynamicDependencies;
+}
+
 namespace PCGSurfaceSampler
 {
 	void FSurfaceSamplerParams::InitializeFromSettings(const UPCGSurfaceSamplerSettings* Settings)
@@ -499,9 +504,11 @@ bool FPCGSurfaceSamplerElement::PrepareDataInternal(FPCGContext* InContext) cons
 	TArray<const UPCGSurfaceData*>& GeneratingShapes = Context->GetPerExecutionState().GeneratingShapes;
 	TArray<FPCGTaggedData>& Outputs = Context->OutputData.TaggedData;
 
+	TArray<FPCGTaskId> PrepareTasks;
+
 	// Initialize the per-iteration data, using the generating shapes as the source of iteration
 	Context->InitializePerIterationStates(GeneratingShapes.Num(),
-		[&GeneratingShapes, &Outputs, &Settings, &Context](IterStateType& OutState, const ExecStateType& ExecState, const uint32 IterationIndex)
+		[&GeneratingShapes, &Outputs, &Settings, &Context, &PrepareTasks](IterStateType& OutState, const ExecStateType& ExecState, const uint32 IterationIndex)
 		{
 			// If we have generating shape inputs, use them
 			const UPCGSurfaceData* GeneratingShape = GeneratingShapes[IterationIndex];
@@ -559,6 +566,18 @@ PRAGMA_ENABLE_DEPRECATION_WARNINGS
 				return EPCGTimeSliceInitResult::NoOperation;
 			}
 
+			if (PCGGraphExecutor::CVarEnableDynamicDependencies.GetValueOnAnyThread())
+			{
+				FBox PrepareSpatialQueryBounds = EffectiveGridBounds;
+				if (OutState.SamplerData.bNeedsLocalTransformation)
+				{
+					// Transform the pre-projected bounds around the surface's origin from local to world space
+					PrepareSpatialQueryBounds = EffectiveGridBounds.TransformBy(OutState.SamplerData.PreProjectionTransform);
+				}
+
+				PrepareTasks.Append(GeneratingShape->PrepareForSpatialQuery(Context, PrepareSpatialQueryBounds));
+			}
+
 			// Assigning this here prevents the need to root
 			Outputs[IterationIndex].Data = OutState.OutputData;
 			return EPCGTimeSliceInitResult::Success;
@@ -569,6 +588,13 @@ PRAGMA_ENABLE_DEPRECATION_WARNINGS
 		PCGE_LOG(Warning, GraphAndLog, LOCTEXT("CouldNotInitializeIterationState", "Could not initialize per-iteration timeslice state data"));
 	}
 
+	if (!PrepareTasks.IsEmpty())
+	{
+		Context->bIsPaused = true;
+		Context->DynamicDependencies.Append(PrepareTasks);
+		// ExecuteInternal will early out on Context->bIsPaused, we still want to advance to the Execute phase
+	}
+
 	return true;
 }
 
@@ -578,6 +604,12 @@ bool FPCGSurfaceSamplerElement::ExecuteInternal(FPCGContext* InContext) const
 	ContextType* TimeSlicedContext = static_cast<ContextType*>(InContext);
 	check(TimeSlicedContext);
 
+	// Paused waiting on PrepareForSpatialQuery tasks
+	if (TimeSlicedContext->bIsPaused)
+	{
+		return false;
+	}
+
 	// Prepare data failed, no need to execute. Return an empty output
 	if (!TimeSlicedContext->DataIsPreparedForExecution())
 	{
diff --git a/Engine/Plugins/PCG/Source/PCG/Private/Grid/PCGLandscapeCache.cpp b/Engine/Plugins/PCG/Source/PCG/Private/Grid/PCGLandscapeCache.cpp
--- a/Engine/Plugins/PCG/Source/PCG/Private/Grid/PCGLandscapeCache.cpp
+++ b/Engine/Plugins/PCG/Source/PCG/Private/Grid/PCGLandscapeCache.cpp
@@ -788,9 +788,85 @@ void UPCGLandscapeCache::TakeOwnership(UPCGLandscapeCache* InLandscapeCache)
 	}
 }
 
+const bool UPCGLandscapeCache::AreCacheEntriesReady(const ULandscapeInfo* InLandscapeInfo, const FIntPoint& InMinComponentKey, const FIntPoint& InMaxComponentKey)
+{
+	ALandscapeProxy* LandscapeProxy = InLandscapeInfo->GetLandscapeProxy();
+	for (int32 ComponentKeyX = InMinComponentKey.X; ComponentKeyX <= InMaxComponentKey.X; ++ComponentKeyX)
+	{
+		for (int32 ComponentKeyY = InMinComponentKey.Y; ComponentKeyY <= InMaxComponentKey.Y; ++ComponentKeyY)
+		{
+			const FIntPoint ComponentKey(ComponentKeyX, ComponentKeyY);
+
+			bool bOutIsValidEntry = false;
+			const FPCGLandscapeCacheEntry* CacheEntry = GetCacheEntryInternal(InLandscapeInfo, ComponentKey, LandscapeProxy, /*bAllowLoadOrCreate=*/false, bOutIsValidEntry);
+
+			if (!bOutIsValidEntry)
+			{
+				continue;
+			}
+
+			// Entry not created yet or not loaded
+			if (!CacheEntry || !CacheEntry->bDataLoaded)
+			{
+				return false;
+			}
+		}
+	}
+
+	return true;
+}
+
+void UPCGLandscapeCache::PrimeCache(const ULandscapeInfo* InLandscapeInfo, const FIntPoint& InMinComponentKey, const FIntPoint& InMaxComponentKey)
+{
+	TRACE_CPUPROFILER_EVENT_SCOPE(UPCGLandscapeData::PrimeCache);
+	for (int32 ComponentKeyX = InMinComponentKey.X; ComponentKeyX <= InMaxComponentKey.X; ++ComponentKeyX)
+	{
+		for (int32 ComponentKeyY = InMinComponentKey.Y; ComponentKeyY <= InMaxComponentKey.Y; ++ComponentKeyY)
+		{
+			GetCacheEntry(InLandscapeInfo, FIntPoint(ComponentKeyX, ComponentKeyY));
+		}
+	}
+}
+
+const FPCGLandscapeCacheEntry* UPCGLandscapeCache::GetCacheEntry(const ULandscapeInfo* InLandscapeInfo, const FIntPoint& InComponentKey, ALandscapeProxy* InLandscapeProxy, bool bInAllowLoadOrCreate)
+{
+	bool bOutIsValidEntry = false;
+
+	if (!InLandscapeInfo)
+	{
+		return nullptr;
+	}
+
+	return GetCacheEntryInternal(InLandscapeInfo, InComponentKey, InLandscapeProxy, bInAllowLoadOrCreate, bOutIsValidEntry);
+}
+
+const FPCGLandscapeCacheEntry* UPCGLandscapeCache::GetCacheEntryInternal(const ULandscapeInfo* InLandscapeInfo, const FIntPoint& InComponentKey, ALandscapeProxy* InLandscapeProxy, bool bInAllowLoadOrCreate, bool& bOutIsValidEntry)
+{
+	// If entry can't exist this will be false (at runtime it means we don't have the data, in editor it means component key is invalid)
+	bOutIsValidEntry = false;
+
+#if WITH_EDITOR
+	// ComponentKey exists so this entry should be valid but it might not exist yet
+	return GetOrCreateCacheEntryInternal(InLandscapeInfo, InComponentKey, bInAllowLoadOrCreate, bOutIsValidEntry);
+#else
+	ALandscapeProxy* LandscapeProxy = InLandscapeProxy ? InLandscapeProxy : InLandscapeInfo->GetLandscapeProxy();
+	const FGuid LandscapeGuid = LandscapeProxy ? LandscapeProxy->GetOriginalLandscapeGuid() : FGuid();
+	const FPCGLandscapeCacheEntry* CacheEntry = GetCacheEntryInternal(LandscapeProxy, LandscapeGuid, InComponentKey, bInAllowLoadOrCreate);
+
+	// Entry is valid if it exists (it might not be loaded yet)
+	bOutIsValidEntry = !!CacheEntry;
+
+	return CacheEntry;
+#endif
+}
+
+
 #if WITH_EDITOR
-const FPCGLandscapeCacheEntry* UPCGLandscapeCache::GetCacheEntry(ULandscapeComponent* LandscapeComponent, const FIntPoint& ComponentCoordinate)
+const FPCGLandscapeCacheEntry* UPCGLandscapeCache::GetOrCreateCacheEntryInternal(const ULandscapeInfo* InLandscapeInfo, const FIntPoint& InComponentKey, bool bInAllowLoadOrCreate, bool& bOutIsValidEntry)
 {
+	bOutIsValidEntry = false;
+	
+	ULandscapeComponent* LandscapeComponent = InLandscapeInfo->XYtoComponentMap.FindRef(InComponentKey);
 	if (!LandscapeComponent)
 	{
 		return nullptr;
@@ -807,39 +883,43 @@ const FPCGLandscapeCacheEntry* UPCGLandscapeCache::GetCacheEntry(ULandscapeCompo
 		return nullptr;
 	}
 
-	const FGuid LandscapeGuid = (LandscapeComponent->GetLandscapeProxy() ? LandscapeComponent->GetLandscapeProxy()->GetOriginalLandscapeGuid() : FGuid());
-	const FPCGLandscapeCacheEntry* CacheEntry = GetCacheEntry(LandscapeComponent->GetLandscapeProxy(), LandscapeGuid, ComponentCoordinate);
+	// Entry is valid, it might not be created/loaded yet
+	bOutIsValidEntry = true;
 
-	if (!CacheEntry && LandscapeComponent && LandscapeComponent->GetLandscapeInfo())
+	const FGuid LandscapeGuid = (LandscapeComponent->GetLandscapeProxy() ? LandscapeComponent->GetLandscapeProxy()->GetOriginalLandscapeGuid() : FGuid());
+	const FPCGLandscapeCacheEntry* CacheEntry = GetCacheEntryInternal(LandscapeComponent->GetLandscapeProxy(), LandscapeGuid, InComponentKey, bInAllowLoadOrCreate);
+	if (CacheEntry || !bInAllowLoadOrCreate)
 	{
-		FWriteScopeLock ScopeLock(CacheLock);
-		CacheMapKey ComponentKey(LandscapeGuid, ComponentCoordinate, LandscapeComponent->GetOwner());
-		if (FPCGLandscapeCacheEntry** FoundEntry = CachedData.Find(ComponentKey))
+		return CacheEntry;
+	}
+		
+	FWriteScopeLock ScopeLock(CacheLock);
+	CacheMapKey ComponentKey(LandscapeGuid, InComponentKey, LandscapeComponent->GetOwner());
+	if (FPCGLandscapeCacheEntry** FoundEntry = CachedData.Find(ComponentKey))
+	{
+		CacheEntry = *FoundEntry;
+	}
+	else
+	{
+		check(LandscapeComponent->SectionBaseX / LandscapeComponent->ComponentSizeQuads == ComponentKey.Coordinate.X && LandscapeComponent->SectionBaseY / LandscapeComponent->ComponentSizeQuads == ComponentKey.Coordinate.Y);
+		if (FPCGLandscapeCacheEntry* NewEntry = FPCGLandscapeCacheEntry::CreateCacheEntry(LandscapeComponent->GetLandscapeInfo(), LandscapeComponent))
 		{
-			CacheEntry = *FoundEntry;
+			NewEntry->OwningCache = this;
+			CacheEntry = NewEntry;
+			CachedData.Add(ComponentKey, NewEntry);
+			++CacheEntryCount;
 		}
-		else
+	}
+
+	if (CacheEntry)
+	{
+		if (CacheEntry->TouchAndLoad(CacheTouch++))
 		{
-			check(LandscapeComponent->SectionBaseX / LandscapeComponent->ComponentSizeQuads == ComponentKey.Coordinate.X && LandscapeComponent->SectionBaseY / LandscapeComponent->ComponentSizeQuads == ComponentKey.Coordinate.Y);
-			if (FPCGLandscapeCacheEntry* NewEntry = FPCGLandscapeCacheEntry::CreateCacheEntry(LandscapeComponent->GetLandscapeInfo(), LandscapeComponent))
-			{
-				NewEntry->OwningCache = this;
-				CacheEntry = NewEntry;
-				CachedData.Add(ComponentKey, NewEntry);
-				++CacheEntryCount;
-			}
+			CacheMemorySize += CacheEntry->GetMemorySize();
 		}
-
-		if (CacheEntry)
+		else if(!CacheEntry->bDataLoaded)
 		{
-			if (CacheEntry->TouchAndLoad(CacheTouch++))
-			{
-				CacheMemorySize += CacheEntry->GetMemorySize();
-			}
-			else if(!CacheEntry->bDataLoaded)
-			{
-				CacheEntry = nullptr;
-			}
+			CacheEntry = nullptr;
 		}
 	}
 
@@ -847,7 +927,7 @@ const FPCGLandscapeCacheEntry* UPCGLandscapeCache::GetCacheEntry(ULandscapeCompo
 }
 #endif
 
-const FPCGLandscapeCacheEntry* UPCGLandscapeCache::GetCacheEntry(AActor* HintActor, const FGuid& LandscapeGuid, const FIntPoint& ComponentCoordinate)
+const FPCGLandscapeCacheEntry* UPCGLandscapeCache::GetCacheEntryInternal(AActor* HintActor, const FGuid& LandscapeGuid, const FIntPoint& ComponentCoordinate, bool bInAllowLoad)
 {
 	const FPCGLandscapeCacheEntry* CacheEntry = nullptr;
 	CacheMapKey ComponentKey(LandscapeGuid, ComponentCoordinate, HintActor);
@@ -868,7 +948,7 @@ const FPCGLandscapeCacheEntry* UPCGLandscapeCache::GetCacheEntry(AActor* HintAct
 		}
 	}
 
-	if (CacheEntry)
+	if (CacheEntry && bInAllowLoad)
 	{
 		if (CacheEntry->TouchAndLoad(CacheTouch++))
 		{
@@ -909,12 +989,7 @@ void UPCGLandscapeCache::SampleMetadataOnPoint(ALandscapeProxy* Landscape, FPCGP
 	const FVector LocalPoint = LandscapeTransform.InverseTransformPosition(InOutPoint.Transform.GetLocation());
 	const FIntPoint ComponentMapKey(FMath::FloorToInt(LocalPoint.X / LandscapeInfo->ComponentSizeQuads), FMath::FloorToInt(LocalPoint.Y / LandscapeInfo->ComponentSizeQuads));
 
-#if WITH_EDITOR
-	ULandscapeComponent* LandscapeComponent = LandscapeInfo->XYtoComponentMap.FindRef(ComponentMapKey);
-	const FPCGLandscapeCacheEntry* CacheEntry = GetCacheEntry(LandscapeComponent, ComponentMapKey);
-#else
-	const FPCGLandscapeCacheEntry* CacheEntry = GetCacheEntry(Landscape, Landscape->GetOriginalLandscapeGuid(), ComponentMapKey);
-#endif
+	const FPCGLandscapeCacheEntry* CacheEntry = GetCacheEntry(LandscapeInfo, ComponentMapKey, Landscape);
 
 	if (!CacheEntry || CacheEntry->LayerData.IsEmpty())
 	{
diff --git a/Engine/Plugins/PCG/Source/PCG/Private/Helpers/PCGBlueprintHelpers.cpp b/Engine/Plugins/PCG/Source/PCG/Private/Helpers/PCGBlueprintHelpers.cpp
--- a/Engine/Plugins/PCG/Source/PCG/Private/Helpers/PCGBlueprintHelpers.cpp
+++ b/Engine/Plugins/PCG/Source/PCG/Private/Helpers/PCGBlueprintHelpers.cpp
@@ -178,12 +178,7 @@ TArray<FPCGLandscapeLayerWeight> UPCGBlueprintHelpers::GetInterpolatedPCGLandsca
 		const FVector LocalPoint = Landscape->LandscapeActorToWorld().InverseTransformPosition(Location);
 		const FIntPoint ComponentMapKey(FMath::FloorToInt(LocalPoint.X / LandscapeInfo->ComponentSizeQuads), FMath::FloorToInt(LocalPoint.Y / LandscapeInfo->ComponentSizeQuads));
 
-		#if WITH_EDITOR
-		ULandscapeComponent* LandscapeComponent = LandscapeInfo->XYtoComponentMap.FindRef(ComponentMapKey);
-		const FPCGLandscapeCacheEntry* CacheEntry = LandscapeCache->GetCacheEntry(LandscapeComponent, ComponentMapKey);
-		#else
-		const FPCGLandscapeCacheEntry* CacheEntry = LandscapeCache->GetCacheEntry(Landscape, Landscape->GetOriginalLandscapeGuid(), ComponentMapKey);
-		#endif
+		const FPCGLandscapeCacheEntry* CacheEntry = LandscapeCache->GetCacheEntry(LandscapeInfo, ComponentMapKey, Landscape);
 
 		if (!CacheEntry)
 		{
diff --git a/Engine/Plugins/PCG/Source/PCG/Public/Data/PCGLandscapeData.h b/Engine/Plugins/PCG/Source/PCG/Public/Data/PCGLandscapeData.h
--- a/Engine/Plugins/PCG/Source/PCG/Public/Data/PCGLandscapeData.h
+++ b/Engine/Plugins/PCG/Source/PCG/Public/Data/PCGLandscapeData.h
@@ -66,6 +66,7 @@ public:
 	virtual void SamplePoints(const TArrayView<const TPair<FTransform, FBox>>& Samples, const TArrayView<FPCGPoint>& OutPoints, UPCGMetadata* OutMetadata) const override;
 	virtual bool ProjectPoint(const FTransform& InTransform, const FBox& InBounds, const FPCGProjectionParams& InParams, FPCGPoint& OutPoint, UPCGMetadata* OutMetadata) const override;
 	virtual bool HasNonTrivialTransform() const override { return true; }
+	virtual TArray<FPCGTaskId> PrepareForSpatialQuery(FPCGContext* InContext, const FBox& InBounds) const override;
 protected:
 	virtual UPCGSpatialData* CopyInternal(FPCGContext* Context) const override;
 	//~End UPCGSpatialData interface
diff --git a/Engine/Plugins/PCG/Source/PCG/Public/Data/PCGSpatialData.h b/Engine/Plugins/PCG/Source/PCG/Public/Data/PCGSpatialData.h
--- a/Engine/Plugins/PCG/Source/PCG/Public/Data/PCGSpatialData.h
+++ b/Engine/Plugins/PCG/Source/PCG/Public/Data/PCGSpatialData.h
@@ -22,6 +22,11 @@ class UPCGUnionData;
 class UPCGDifferenceData;
 class UPCGProjectionData;
 
+namespace PCGSpatialData
+{
+	extern TAutoConsoleVariable<bool> CVarEnablePrepareForSpatialQuery;
+}
+
 USTRUCT(BlueprintType)
 struct PCG_API FPCGInitializeFromDataParams
 {
@@ -213,6 +218,9 @@ public:
 	/** True if data inherits from parent */
 	virtual bool HasSpatialDataParent() const { return false; }
 
+	/** Optionally return a list of scheduled tasks that will prepare the data for spatial queries. Calling task can then wait on those tasks to finish before progressing further. */
+	virtual TArray<FPCGTaskId> PrepareForSpatialQuery(FPCGContext* InContext, const FBox& InBounds) const { return {}; }
+
 	/** Recipient of any artifacts generated using this data. */
 	UPROPERTY(Transient, BlueprintReadWrite, EditAnywhere, Category = Data)
 	TWeakObjectPtr<AActor> TargetActor = nullptr;
diff --git a/Engine/Plugins/PCG/Source/PCG/Public/Grid/PCGLandscapeCache.h b/Engine/Plugins/PCG/Source/PCG/Public/Grid/PCGLandscapeCache.h
--- a/Engine/Plugins/PCG/Source/PCG/Public/Grid/PCGLandscapeCache.h
+++ b/Engine/Plugins/PCG/Source/PCG/Public/Grid/PCGLandscapeCache.h
@@ -142,13 +142,13 @@ public:
 
 	void Tick(float DeltaSeconds);
 
-#if WITH_EDITOR
-	/** Gets (and creates if needed) the cache entry - available only in Editor */
-	const FPCGLandscapeCacheEntry* GetCacheEntry(ULandscapeComponent* LandscapeComponent, const FIntPoint& ComponentKey);
-#endif
+	const bool AreCacheEntriesReady(const ULandscapeInfo* InLandscapeInfo, const FIntPoint& InMinComponentKey, const FIntPoint& InMaxComponentKey);
 
-	/** Gets landscape cache entry, works both in editor (but does not create) but works in game mode too. */
-	const FPCGLandscapeCacheEntry* GetCacheEntry(AActor* HintActor, const FGuid& LandscapeGuid, const FIntPoint& ComponentKey);
+	/** Gets landscape cache entry for specified component key */
+	const FPCGLandscapeCacheEntry* GetCacheEntry(const ULandscapeInfo* InLandscapeInfo, const FIntPoint& InComponentKey, ALandscapeProxy* InLandscapeProxy = nullptr, bool bInAllowLoadOrCreate = true);
+
+	/** Load/Create cache entries between Min/Max component key inclusively */
+	void PrimeCache(const ULandscapeInfo* InLandscapeInfo, const FIntPoint& InMinComponentKey, const FIntPoint& InMaxComponentKey);
 
 	TArray<FName> GetLayerNames(ALandscapeProxy* Landscape);
 
@@ -178,7 +178,13 @@ private:
 	void CacheLayerNames(ALandscapeProxy* InLandscape);
 	void CacheLayerNames();
 	void RemoveComponentFromCache(const ALandscapeProxy* LandscapeProxy);
+
+	/** Gets (and creates if needed) the cache entry - available only in Editor */
+	const FPCGLandscapeCacheEntry* GetOrCreateCacheEntryInternal(const ULandscapeInfo* InLandscapeInfo, const FIntPoint& InComponentKey, bool bInAllowLoadOrCreate, bool& bOutValidEntry);
 #endif
+	/** Gets landscape cache entry, works both in editor (but does not create) but works in game mode too. */
+	const FPCGLandscapeCacheEntry* GetCacheEntryInternal(AActor* InHintActor, const FGuid& InLandscapeGuid, const FIntPoint& InComponentKey, bool bInAllowLoad);
+	const FPCGLandscapeCacheEntry* GetCacheEntryInternal(const ULandscapeInfo* InLandscapeInfo, const FIntPoint& InComponentKey, ALandscapeProxy* InLandscapeProxy, bool bInAllowLoadOrCreate, bool& bOutValidEntry);
 
 	struct CacheMapKey
 	{
