[PCG] Check in default Accessor factory that metadata is valid

* Add generic log for invalid metadata
* Also update Reduce element to discard data that doesn't have metadata
* Needed to update some exptected warnings in unit tests that are different now

#rnx

diff --git a/Engine/Plugins/PCG/Source/PCG/Private/Elements/PCGAttributeReduceElement.cpp b/Engine/Plugins/PCG/Source/PCG/Private/Elements/PCGAttributeReduceElement.cpp
--- a/Engine/Plugins/PCG/Source/PCG/Private/Elements/PCGAttributeReduceElement.cpp
+++ b/Engine/Plugins/PCG/Source/PCG/Private/Elements/PCGAttributeReduceElement.cpp
@@ -276,14 +276,13 @@ bool FPCGAttributeReduceElement::ExecuteInternal(FPCGContext* Context) const
 	{
 		const UPCGData* InputData = Inputs[i].Data;
 
-		if (!InputData)
+		if (!InputData || !InputData->ConstMetadata())
 		{
-			PCGE_LOG(Warning, GraphAndLog, FText::Format(LOCTEXT("InputNotSpatialData", "Input {0} is invalid, skipped"), FText::AsNumber(i)));
+			PCGLog::InputOutput::LogInvalidInputDataError(Context);
 			continue;
 		}
 
-		FPCGAttributePropertyInputSelector InputSource = Settings->InputSource.CopyAndFixLast(InputData);
-
+		const FPCGAttributePropertyInputSelector InputSource = Settings->InputSource.CopyAndFixLast(InputData);
 		const FName OutputAttributeName = (Settings->OutputAttributeName == PCGMetadataAttributeConstants::SourceNameAttributeName) ? InputSource.GetName() : Settings->OutputAttributeName;
 
 		TUniquePtr<const IPCGAttributeAccessor> Accessor = PCGAttributeAccessorHelpers::CreateConstAccessor(InputData, InputSource);
@@ -291,7 +290,7 @@ bool FPCGAttributeReduceElement::ExecuteInternal(FPCGContext* Context) const
 
 		if (!Accessor.IsValid() || !Keys.IsValid())
 		{
-			PCGE_LOG(Warning, GraphAndLog, FText::Format(LOCTEXT("AttributeDoesNotExist", "Input attribute/property '{0}' does not exist on input {1}, skipped"), InputSource.GetDisplayText(), FText::AsNumber(i)));
+			PCGLog::Metadata::LogFailToCreateAccessorError(InputSource, Context);
 			continue;
 		}
 
diff --git a/Engine/Plugins/PCG/Source/PCG/Private/Metadata/Accessors/PCGAttributeAccessorFactory.cpp b/Engine/Plugins/PCG/Source/PCG/Private/Metadata/Accessors/PCGAttributeAccessorFactory.cpp
--- a/Engine/Plugins/PCG/Source/PCG/Private/Metadata/Accessors/PCGAttributeAccessorFactory.cpp
+++ b/Engine/Plugins/PCG/Source/PCG/Private/Metadata/Accessors/PCGAttributeAccessorFactory.cpp
@@ -27,22 +27,29 @@ namespace PCGAttributeAccessorFactory
 					UE_LOG(LogPCG, Error, TEXT("[FPCGAttributeAccessorFactory::DefaultAccessor] Expected to select an extra property but the data doesn't support this property."));
 				}
 				
-				return TUniquePtr<IPCGAttributeAccessor>();
+				return nullptr;
 			}
 		}
 		else if (InSelector.GetSelection() == EPCGAttributePropertySelection::Attribute)
 		{
+			UPCGMetadata* Metadata = InData->MutableMetadata();
+			if (!Metadata)
+			{
+				PCGLog::Metadata::LogInvalidMetadata();
+				return nullptr;
+			}
+			
 			const FPCGMetadataDomainID DomainID = InData->GetMetadataDomainIDFromSelector(InSelector);
 			if (!DomainID.IsValid())
 			{
 				PCGLog::Metadata::LogInvalidMetadataDomain(InSelector);
-				return TUniquePtr<IPCGAttributeAccessor>();
+				return nullptr;
 			}
 
-			FPCGMetadataDomain* Metadata = InData->MutableMetadata()->GetMetadataDomain(DomainID);
-			FPCGMetadataAttributeBase* Attribute = Metadata ? Metadata->GetMutableAttribute(InSelector.GetAttributeName()) : nullptr;
+			FPCGMetadataDomain* MetadataDomain = Metadata->GetMetadataDomain(DomainID);
+			FPCGMetadataAttributeBase* Attribute = MetadataDomain ? MetadataDomain->GetMutableAttribute(InSelector.GetAttributeName()) : nullptr;
 		
-			Accessor = PCGAttributeAccessorHelpers::CreateAccessor(Attribute, Metadata, bQuiet);
+			Accessor = PCGAttributeAccessorHelpers::CreateAccessor(Attribute, MetadataDomain, bQuiet);
 		}
 
 		return Accessor;
@@ -64,22 +71,29 @@ namespace PCGAttributeAccessorFactory
 					UE_LOG(LogPCG, Error, TEXT("[FPCGAttributeAccessorFactory::DefaultConstAccessor] Expected to select an extra property but the data doesn't support this property."));
 				}
 				
-				return TUniquePtr<IPCGAttributeAccessor>();
+				return nullptr;
 			}
 		}
 		else if (InSelector.GetSelection() == EPCGAttributePropertySelection::Attribute)
 		{
+			const UPCGMetadata* Metadata = InData->ConstMetadata();
+			if (!Metadata)
+			{
+				PCGLog::Metadata::LogInvalidMetadata();
+				return nullptr;
+			}
+			
 			const FPCGMetadataDomainID DomainID = InData->GetMetadataDomainIDFromSelector(InSelector);
 			if (!DomainID.IsValid())
 			{
 				PCGLog::Metadata::LogInvalidMetadataDomain(InSelector);
-				return TUniquePtr<IPCGAttributeAccessor>();
+				return nullptr;
 			}
 
-			const FPCGMetadataDomain* Metadata = InData->ConstMetadata()->GetConstMetadataDomain(DomainID);
-			const FPCGMetadataAttributeBase* Attribute = Metadata ? Metadata->GetConstAttribute(InSelector.GetAttributeName()) : nullptr;
+			const FPCGMetadataDomain* MetadataDomain = Metadata->GetConstMetadataDomain(DomainID);
+			const FPCGMetadataAttributeBase* Attribute = MetadataDomain ? MetadataDomain->GetConstAttribute(InSelector.GetAttributeName()) : nullptr;
 		
-			Accessor = PCGAttributeAccessorHelpers::CreateConstAccessor(Attribute, Metadata, bQuiet);
+			Accessor = PCGAttributeAccessorHelpers::CreateConstAccessor(Attribute, MetadataDomain, bQuiet);
 		}
 
 		return Accessor;
@@ -87,29 +101,43 @@ namespace PCGAttributeAccessorFactory
 	
 	TUniquePtr<IPCGAttributeAccessorKeys> CreateDefaultAccessorKeys(UPCGData* InData, const FPCGAttributePropertySelector& InSelector, const bool bQuiet)
 	{
+		UPCGMetadata* Metadata = InData->MutableMetadata();
+		if (!Metadata)
+		{
+			PCGLog::Metadata::LogInvalidMetadata();
+			return nullptr;
+		}
+		
 		const FPCGMetadataDomainID DomainID = InData->GetMetadataDomainIDFromSelector(InSelector);
 		if (!DomainID.IsValid())
 		{
 			PCGLog::Metadata::LogInvalidMetadataDomain(InSelector);
-			return  TUniquePtr<IPCGAttributeAccessorKeys>();
+			return nullptr;
 		}
 
-		FPCGMetadataDomain* Metadata = InData->MutableMetadata()->GetMetadataDomain(DomainID);
-		return Metadata ? MakeUnique<FPCGAttributeAccessorKeysEntries>(Metadata) : TUniquePtr<IPCGAttributeAccessorKeys>();
+		FPCGMetadataDomain* MetadataDomain = Metadata->GetMetadataDomain(DomainID);
+		return MetadataDomain ? MakeUnique<FPCGAttributeAccessorKeysEntries>(MetadataDomain) : nullptr;
 	}
 	
 	TUniquePtr<const IPCGAttributeAccessorKeys> CreateDefaultConstAccessorKeys(const UPCGData* InData, const FPCGAttributePropertySelector& InSelector, const bool bQuiet)
 	{
+		const UPCGMetadata* Metadata = InData->ConstMetadata();
+		if (!Metadata)
+		{
+			PCGLog::Metadata::LogInvalidMetadata();
+			return nullptr;
+		}
+		
 		const FPCGMetadataDomainID DomainID = InData->GetMetadataDomainIDFromSelector(InSelector);
 		if (!DomainID.IsValid())
 		{
 			PCGLog::Metadata::LogInvalidMetadataDomain(InSelector);
-			return  TUniquePtr<IPCGAttributeAccessorKeys>();
+			return nullptr;
 		}
 
-		const FPCGMetadataDomain* Metadata = InData->ConstMetadata()->GetConstMetadataDomain(DomainID);
-		const bool bAddDefaultValueIfEmpty = Metadata ? (Metadata->GetDomainID() == PCGMetadataDomainID::Data) : false;
-		return Metadata ? MakeUnique<const FPCGAttributeAccessorKeysEntries>(Metadata, /*bAddDefaultValueIfEmpty=*/bAddDefaultValueIfEmpty) : TUniquePtr<const IPCGAttributeAccessorKeys>();
+		const FPCGMetadataDomain* MetadataDomain = Metadata->GetConstMetadataDomain(DomainID);
+		const bool bAddDefaultValueIfEmpty = MetadataDomain ? (MetadataDomain->GetDomainID() == PCGMetadataDomainID::Data) : false;
+		return MetadataDomain ? MakeUnique<const FPCGAttributeAccessorKeysEntries>(MetadataDomain, /*bAddDefaultValueIfEmpty=*/bAddDefaultValueIfEmpty) : nullptr;
 	}
 }
 
diff --git a/Engine/Plugins/PCG/Source/PCG/Private/Tests/Elements/PCGAttributeReduceTest.cpp b/Engine/Plugins/PCG/Source/PCG/Private/Tests/Elements/PCGAttributeReduceTest.cpp
--- a/Engine/Plugins/PCG/Source/PCG/Private/Tests/Elements/PCGAttributeReduceTest.cpp
+++ b/Engine/Plugins/PCG/Source/PCG/Private/Tests/Elements/PCGAttributeReduceTest.cpp
@@ -302,7 +302,7 @@ bool FPCGAttributeReduceFirstInputInvalid_NoMerge::RunTest(const FString& Parame
 	FPCGTaggedData& TaggedData = Inputs.EmplaceAt_GetRef(0);
 	TaggedData.Data = NewObject<UPCGAttributeReduceSettings>();
 
-	AddExpectedError(TEXT("Input attribute/property 'double' does not exist on input 0, skipped"));
+	AddExpectedError(TEXT("Invalid input data."));
 
 	return RunAndValidate(Inputs, EPCGAttributeReduceOperation::Min, std::move(InputSource), /*bMerged=*/ false, 0.1, 0.2);
 }
@@ -316,7 +316,7 @@ bool FPCGAttributeReduceFirstInputInvalid_Merge::RunTest(const FString& Paramete
 	FPCGTaggedData& TaggedData = Inputs.EmplaceAt_GetRef(0);
 	TaggedData.Data = NewObject<UPCGAttributeReduceSettings>();
 
-	AddExpectedError(TEXT("Input attribute/property 'double' does not exist on input 0, skipped"));
+	AddExpectedError(TEXT("Invalid input data."));
 
 	return RunAndValidate(Inputs, EPCGAttributeReduceOperation::Min, std::move(InputSource), /*bMerged=*/ true, 0.1, 0.2);
 }
diff --git a/Engine/Plugins/PCG/Source/PCG/Private/Tests/Elements/PCGFilterByIndexTest.cpp b/Engine/Plugins/PCG/Source/PCG/Private/Tests/Elements/PCGFilterByIndexTest.cpp
--- a/Engine/Plugins/PCG/Source/PCG/Private/Tests/Elements/PCGFilterByIndexTest.cpp
+++ b/Engine/Plugins/PCG/Source/PCG/Private/Tests/Elements/PCGFilterByIndexTest.cpp
@@ -63,6 +63,9 @@ bool FPCGFilterByIndexTest_Basic::RunTest(const FString& Parameters)
 
 	// Empty index input string
 	{
+		// Empty will throw a warning
+		AddExpectedError(TEXT("Empty expression in parsed string"));
+		
 		Settings->SelectedIndices.Empty();
 		const TArray<FPCGTaggedData> Outputs = PCGFilterByIndexTestPrivate::ResetAndExecute(TestData, Settings);
 
diff --git a/Engine/Plugins/PCG/Source/PCG/Private/Utils/PCGLogErrors.cpp b/Engine/Plugins/PCG/Source/PCG/Private/Utils/PCGLogErrors.cpp
--- a/Engine/Plugins/PCG/Source/PCG/Private/Utils/PCGLogErrors.cpp
+++ b/Engine/Plugins/PCG/Source/PCG/Private/Utils/PCGLogErrors.cpp
@@ -53,6 +53,7 @@ namespace PCGLog::Metadata
 		const FTextFormat SetTypedAttributeFailureNoAccessor = LOCTEXT("SetTypedAttributeFailureNoAccessor", "Could not set value in attribute '{0}' value of type: '{1}.");
 		const FTextFormat IncomparableTypesFailure = LOCTEXT("IncomparableTypesFailure", "Attributes '{0}' and '{1}' are incomparable. Ensure they are either of the same or compatible types.");
 		const FTextFormat InvalidMetadataDomain = LOCTEXT("InvalidMetadataDomain", "Metadata domain {0} is invalid for this data.");
+		const FText InvalidMetadata = LOCTEXT("InvalidMetadata", "Metadata is unsupported or invalid for this data.");
 	}
 
 	void LogFailToCreateAccessorError(const FPCGAttributePropertySelector& Selector, const FPCGContext* InContext)
@@ -60,6 +61,11 @@ namespace PCGLog::Metadata
 		LogErrorOnGraph(FText::Format(Format::CreateAccessorFailure, Selector.GetDisplayText()), InContext);
 	}
 
+	void LogInvalidMetadata(const FPCGContext* InContext)
+	{
+		LogErrorOnGraph(Format::InvalidMetadata, InContext);
+	}
+
 	void LogInvalidMetadataDomain(const FPCGAttributePropertySelector& Selector, const FPCGContext* InContext)
 	{
 		LogErrorOnGraph(FText::Format(Format::InvalidMetadataDomain, FText::FromName(Selector.GetDomainName())), InContext);
diff --git a/Engine/Plugins/PCG/Source/PCG/Public/Utils/PCGLogErrors.h b/Engine/Plugins/PCG/Source/PCG/Public/Utils/PCGLogErrors.h
--- a/Engine/Plugins/PCG/Source/PCG/Public/Utils/PCGLogErrors.h
+++ b/Engine/Plugins/PCG/Source/PCG/Public/Utils/PCGLogErrors.h
@@ -43,11 +43,13 @@ namespace PCGLog
 			PCG_API extern const FTextFormat SetTypedAttributeFailure;
 			PCG_API extern const FTextFormat SetTypedAttributeFailureNoAccessor;
 			PCG_API extern const FTextFormat InvalidMetadataDomain;
+			PCG_API extern const FText InvalidMetadata;
 		}
 
 		// Errors
 		PCG_API void LogFailToCreateAccessorError(const FPCGAttributePropertySelector& Selector, const FPCGContext* InContext = nullptr);
 		PCG_API void LogInvalidMetadataDomain(const FPCGAttributePropertySelector& Selector, const FPCGContext* InContext = nullptr);
+		PCG_API void LogInvalidMetadata(const FPCGContext* InContext = nullptr);
 
 		template <typename T>
 		void LogFailToCreateAttributeError(const FText& AttributeName, const FPCGContext* InContext = nullptr)
