[PCG] Improve Metadata CRC computation

* No need to CRC the default as it is the same as another domain
* Only create the keys once (with factory) and skip the factory query for the accessor for each attribute

#rnx

diff --git a/Engine/Plugins/PCG/Source/PCG/Private/Metadata/PCGMetadata.cpp b/Engine/Plugins/PCG/Source/PCG/Private/Metadata/PCGMetadata.cpp
--- a/Engine/Plugins/PCG/Source/PCG/Private/Metadata/PCGMetadata.cpp
+++ b/Engine/Plugins/PCG/Source/PCG/Private/Metadata/PCGMetadata.cpp
@@ -325,6 +325,12 @@ void UPCGMetadata::AddToCrc(FArchiveCrc32& Ar, bool bFullDataCrc) const
 
 	for (const FPCGMetadataDomainID& DomainID : DomainIDs)
 	{
+		// No need to CRC the default (since it is an alias to one domain)
+		if (DomainID.IsDefault())
+		{
+			continue;
+		}
+
 		bool bIsValid = MetadataDomains[DomainID].IsValid();
 		Ar << bIsValid;
 
diff --git a/Engine/Plugins/PCG/Source/PCG/Private/Metadata/PCGMetadataDomain.cpp b/Engine/Plugins/PCG/Source/PCG/Private/Metadata/PCGMetadataDomain.cpp
--- a/Engine/Plugins/PCG/Source/PCG/Private/Metadata/PCGMetadataDomain.cpp
+++ b/Engine/Plugins/PCG/Source/PCG/Private/Metadata/PCGMetadataDomain.cpp
@@ -156,32 +156,52 @@ void FPCGMetadataDomain::AddToCrc(FArchiveCrc32& Ar, const UPCGData* Data, bool
 {
 	check(Data);
 
-	// Sort attributes so we have a consistent processing path
-	TArray<FName> AttributeNames;
+	TArray<const FPCGMetadataAttributeBase*> AllAttributes;
 
 	{
-		AttributeLock.ReadLock();
-		AttributeNames.Reserve(Attributes.Num());
+		FReadScopeLock Lock(AttributeLock);
+		AllAttributes.Reserve(Attributes.Num());
 
 		for (const TPair<FName, FPCGMetadataAttributeBase*>& Attribute : Attributes)
 		{
-			AttributeNames.Add(Attribute.Key);
+			if (!ensure(Attribute.Value))
+			{
+				continue;
+			}
+			
+			AllAttributes.Add(Attribute.Value);
 		}
-		AttributeLock.ReadUnlock();
+	}
 
-		Algo::Sort(AttributeNames, [this](const FName& A, const FName& B) { return A.LexicalLess(B); });
+	if (AllAttributes.IsEmpty())
+	{
+		return;
 	}
 
-	// Add attributes to CRC
-	for (FName AttributeName : AttributeNames)
+	// Sort attributes so we have a consistent processing path
+	Algo::Sort(AllAttributes, [this](const FPCGMetadataAttributeBase* A, const FPCGMetadataAttributeBase* B) { return A->Name.LexicalLess(B->Name); });
+
+	// Create the keys only once, as they are the same for all the attributes.
+	FPCGAttributePropertyInputSelector InputSource;
+	Data->SetDomainFromDomainID(DomainID, InputSource);
+	InputSource.SetAttributeName(AllAttributes[0]->Name);
+
+	TUniquePtr<const IPCGAttributeAccessorKeys> InputKeys = PCGAttributeAccessorHelpers::CreateConstKeys(Data, InputSource);
+	if (!ensure(InputKeys.IsValid()))
 	{
-		Ar << AttributeName;
+		return;
+	}
 
-		FPCGAttributePropertyInputSelector InputSource;
-		InputSource.SetAttributeName(AttributeName);
+	// Then for each attribute, serialize the name and its values.
+	for (const FPCGMetadataAttributeBase* Attribute : AllAttributes)
+	{
+		Ar << const_cast<FName&>(Attribute->Name);
 
-		TUniquePtr<const IPCGAttributeAccessor> InputAccessor = PCGAttributeAccessorHelpers::CreateConstAccessor(Data, InputSource);
-		TUniquePtr<const IPCGAttributeAccessorKeys> InputKeys = PCGAttributeAccessorHelpers::CreateConstKeys(Data, InputSource);
+		TUniquePtr<const IPCGAttributeAccessor> InputAccessor = PCGAttributeAccessorHelpers::CreateConstAccessor(Attribute, this);
+		if (!ensure(InputAccessor.IsValid()))
+		{
+			continue;
+		}
 
 		auto Callback = [&InputAccessor, &InputKeys, &Ar](auto&& Dummy)
 		{
@@ -195,7 +215,7 @@ void FPCGMetadataDomain::AddToCrc(FArchiveCrc32& Ar, const UPCGData* Data, bool
 			{
 				Values.SetNum(InputKeys->GetNum());
 			}
-			
+
 			InputAccessor->GetRange<AttributeType>(Values, 0, *InputKeys);
 
 			for (AttributeType& Value : Values)
