[PCG] Separate Graph Parameters into its own dockable panel

+ Added some helper functionality surrounding focusing dockable panels
+ Added a button where users would have found the graph parameters to take them to the panel
* Consolidated some similar code between the Open Source Editor button for HLSL node and the new button

diff --git a/Engine/Plugins/PCG/Source/PCG/Public/PCGGraph.h b/Engine/Plugins/PCG/Source/PCG/Public/PCGGraph.h
--- a/Engine/Plugins/PCG/Source/PCG/Public/PCGGraph.h
+++ b/Engine/Plugins/PCG/Source/PCG/Public/PCGGraph.h
@@ -83,7 +83,7 @@ public:
 	/** Handle the sync between the parent parameters and the overrides. Will return true if something changed. */
 	bool RefreshParameters(const FInstancedPropertyBag* ParentUserParameters, EPCGGraphParameterEvent InChangeType, FName InChangedPropertyName);
 
-	UPROPERTY(EditAnywhere, Category = "")
+	UPROPERTY(EditAnywhere, Category = "", meta = (ShowOnlyInnerProperties))
 	FInstancedPropertyBag Parameters;
 
 	UPROPERTY(VisibleAnywhere, Category = "", meta = (Hidden))
@@ -470,7 +470,7 @@ protected:
 #endif // WITH_EDITORONLY_DATA
 
 	// Parameters
-	UPROPERTY(EditAnywhere, Category = Instance, meta = (DisplayName = "Parameters", NoResetToDefault, DefaultType = "EPropertyBagPropertyType::Double", IsPinTypeAccepted = "UserParametersIsPinTypeAccepted", CanRemoveProperty = "UserParametersCanRemoveProperty"))
+	UPROPERTY(EditAnywhere, Category = "Graph Parameters", DisplayName = "Parameters", meta = (NoResetToDefault, DefaultType = "EPropertyBagPropertyType::Double", IsPinTypeAccepted = "UserParametersIsPinTypeAccepted", CanRemoveProperty = "UserParametersCanRemoveProperty"))
 	FInstancedPropertyBag UserParameters;
 
 #if WITH_EDITOR
diff --git a/Engine/Plugins/PCG/Source/PCGEditor/Private/Details/PCGCustomHLSLSettingsDetails.cpp b/Engine/Plugins/PCG/Source/PCGEditor/Private/Details/PCGCustomHLSLSettingsDetails.cpp
--- a/Engine/Plugins/PCG/Source/PCGEditor/Private/Details/PCGCustomHLSLSettingsDetails.cpp
+++ b/Engine/Plugins/PCG/Source/PCGEditor/Private/Details/PCGCustomHLSLSettingsDetails.cpp
@@ -21,63 +21,62 @@ void FPCGCustomHLSLSettingsDetails::CustomizeDetails(IDetailLayoutBuilder& Detai
 
 	IDetailCategoryBuilder& SettingsCategory = DetailBuilder.EditCategory(TEXT("Settings"));
 
-	TWeakPtr<FPCGEditor> PCGEditor = GetPCGEditor(ObjectsBeingCustomized);
-	
-	// Add a "open source editor" button if the source editor tab is not already visible.
-	if (PCGEditor.IsValid() && PCGEditor.Pin()->IsNodeSourceEditorTabClosed())
+	if (!ObjectsBeingCustomized.IsEmpty() && ObjectsBeingCustomized[0].IsValid())
 	{
-		SettingsCategory.AddCustomRow(FText::GetEmpty()).ValueContent()
-			.MaxDesiredWidth(120.f)
-			[
-				SNew(SHorizontalBox)
-				+SHorizontalBox::Slot()
-				.AutoWidth()
-				.Padding(2.0f, 0.0f)
-				.VAlign(VAlign_Fill)
+		const UPCGSettings* FirstSettings = Cast<UPCGSettings>(ObjectsBeingCustomized[0].Get());
+		const UPCGEditorGraph* PCGEditorGraph = FPCGEditor::GetPCGEditorGraph(FirstSettings);
+		const TWeakPtr<const FPCGEditor> PCGEditorWeakPtr = PCGEditorGraph ? PCGEditorGraph->GetEditor() : nullptr;
+		const TSharedPtr<const FPCGEditor> PCGEditorSharedPtr = PCGEditorWeakPtr.Pin();
+
+		// Add a "open source editor" button if the source editor tab is not already visible.
+		if (PCGEditorSharedPtr.IsValid() && PCGEditorSharedPtr->IsPanelCurrentlyOpen(EPCGEditorPanel::NodeSource))
+		{
+			SettingsCategory.AddCustomRow(FText::GetEmpty())
+				.Visibility(
+				TAttribute<EVisibility>::CreateLambda(
+				[PCGEditorWeakPtr]
+				{
+					const TSharedPtr<const FPCGEditor> SharedPtr = PCGEditorWeakPtr.Pin();
+					return (SharedPtr.IsValid() && SharedPtr->IsPanelCurrentlyOpen(EPCGEditorPanel::NodeSource)) ? EVisibility::Collapsed : EVisibility::Visible;
+				}))
+				.ValueContent()
+				.MaxDesiredWidth(120.f)
 				[
-					SNew(SButton)
-					.OnClicked_Lambda([PCGEditor]() -> FReply
-					{
-						if (PCGEditor.IsValid())
+					SNew(SHorizontalBox)
+					+ SHorizontalBox::Slot()
+					.AutoWidth()
+					.Padding(2.0f, 0.0f)
+					.VAlign(VAlign_Fill)
+					[
+						SNew(SButton)
+						.OnClicked_Lambda([PCGEditorWeakPtr]() -> FReply
 						{
-							PCGEditor.Pin()->SpawnNodeSourceEditorTab();
-						}
+							if (const TSharedPtr<const FPCGEditor> SharedPtr = PCGEditorWeakPtr.Pin())
+							{
+								SharedPtr->BringFocusToPanel(EPCGEditorPanel::NodeSource);
+							}
 
-						return FReply::Handled();
-					})
-					.ToolTipText(FText::FromString("Opens Node Source Editor tab."))
-					[
-						SNew(STextBlock)
-						.Font(IDetailLayoutBuilder::GetDetailFont())
-						.Text(LOCTEXT("ButtonOpenSourceEditor", "Open Source Editor"))
+							return FReply::Handled();
+						})
+						.ToolTipText(FText::FromString("Opens Node Source Editor Panel."))
+						[
+							SNew(STextBlock)
+							.Font(IDetailLayoutBuilder::GetDetailFont())
+							.Text(LOCTEXT("ButtonOpenSourceEditor", "Open Source Editor"))
+						]
 					]
-				]
-			];
+				];
+		}
 	}
 
 	// Add default properties in the category in order
 	TArray<TSharedRef<IPropertyHandle>> AllProperties;
 	SettingsCategory.GetDefaultProperties(AllProperties);
 
-	for (auto& Property : AllProperties)
+	for (TSharedRef<IPropertyHandle>& Property : AllProperties)
 	{
 		SettingsCategory.AddProperty(Property);
 	}
 }
 
-TWeakPtr<FPCGEditor> FPCGCustomHLSLSettingsDetails::GetPCGEditor(const TArray<TWeakObjectPtr<UObject>>& InObjectsBeingCustomized)
-{
-	if (InObjectsBeingCustomized.IsEmpty() || !InObjectsBeingCustomized[0].IsValid())
-	{
-		return nullptr;
-	}
-
-	const UPCGSettings* FirstSettings = Cast<UPCGSettings>(InObjectsBeingCustomized[0].Get());
-	UPCGNode* Node = Cast<UPCGNode>(FirstSettings ? FirstSettings->GetOuter() : nullptr);
-	UPCGGraph* Graph = Cast<UPCGGraph>(Node ? Node->GetOuter() : nullptr);
-	const UPCGEditorGraph* EditorGraph = Graph ? FPCGEditor::GetPCGEditorGraph(Graph) : nullptr;
-
-	return EditorGraph ? EditorGraph->GetEditor() : nullptr;
-}
-
 #undef LOCTEXT_NAMESPACE
diff --git a/Engine/Plugins/PCG/Source/PCGEditor/Private/Details/PCGGraphDetails.cpp b/Engine/Plugins/PCG/Source/PCGEditor/Private/Details/PCGGraphDetails.cpp
--- a/Engine/Plugins/PCG/Source/PCGEditor/Private/Details/PCGGraphDetails.cpp
+++ b/Engine/Plugins/PCG/Source/PCGEditor/Private/Details/PCGGraphDetails.cpp
@@ -1,11 +1,15 @@
 // Copyright Epic Games, Inc. All Rights Reserved.
 
 #include "Details/PCGGraphDetails.h"
+
+#include "PCGEditor.h"
+#include "PCGEditorGraph.h"
 #include "PCGGraph.h"
 
-#include "PropertyHandle.h"
-#include "DetailLayoutBuilder.h"
 #include "DetailCategoryBuilder.h"
+#include "DetailLayoutBuilder.h"
+#include "DetailWidgetRow.h"
+#include "PropertyHandle.h"
 
 #define LOCTEXT_NAMESPACE "PCGGraphDetails"
 
@@ -16,9 +20,6 @@ TSharedRef<IDetailCustomization> FPCGGraphDetails::MakeInstance()
 
 void FPCGGraphDetails::CustomizeDetails(IDetailLayoutBuilder& DetailBuilder)
 {
-	const FName PCGCategoryName("PCG");
-	IDetailCategoryBuilder& PCGCategory = DetailBuilder.EditCategory(PCGCategoryName);
-
 	TArray<TWeakObjectPtr<UObject>> ObjectsBeingCustomized;
 	DetailBuilder.GetObjectsBeingCustomized(ObjectsBeingCustomized);
 
@@ -31,13 +32,75 @@ void FPCGGraphDetails::CustomizeDetails(IDetailLayoutBuilder& DetailBuilder)
 		}
 	}
 
+	// UE_DEPRECATED(5.6, "Added for familiarity and convenience to the Graph Parameters panel moving.")
+	IDetailCategoryBuilder& SettingsCategory = DetailBuilder.EditCategory(TEXT("Instance"));
+	if (!ObjectsBeingCustomized.IsEmpty() && ObjectsBeingCustomized[0].IsValid())
+	{
+		const UPCGEditorGraph* PCGEditorGraph = nullptr;
+
+		if (!SelectedGraphs.IsEmpty())
+		{
+			PCGEditorGraph = FPCGEditor::GetPCGEditorGraph(SelectedGraphs[0].Get());
+		}
+		else if (const UPCGSettings* FirstSettings = Cast<UPCGSettings>(ObjectsBeingCustomized[0].Get()))
+		{
+			PCGEditorGraph = FPCGEditor::GetPCGEditorGraph(FirstSettings);
+		}
+
+		const TWeakPtr<const FPCGEditor> PCGEditorWeakPtr = PCGEditorGraph ? PCGEditorGraph->GetEditor() : nullptr;
+
+		// Add an "Open Graph Parameters" button if the source editor tab is not already visible.
+		if (PCGEditorWeakPtr.IsValid())
+		{
+			SettingsCategory.AddCustomRow(FText::GetEmpty())
+				.Visibility(
+				TAttribute<EVisibility>::CreateLambda(
+				[PCGEditorWeakPtr]
+				{
+					const TSharedPtr<const FPCGEditor> PCGEditorSharedPtr = PCGEditorWeakPtr.Pin();
+					return (PCGEditorSharedPtr.IsValid() && PCGEditorSharedPtr->IsPanelCurrentlyOpen(EPCGEditorPanel::UserParams)) ? EVisibility::Collapsed : EVisibility::Visible;
+				}))
+				.ValueContent()
+				.MaxDesiredWidth(120.f)
+				[
+					SNew(SHorizontalBox)
+					+ SHorizontalBox::Slot()
+					.AutoWidth()
+					.Padding(2.0f, 0.0f)
+					.VAlign(VAlign_Fill)
+					[
+						SNew(SButton)
+						.OnClicked_Lambda(
+						[PCGEditorWeakPtr]() -> FReply
+						{
+							if (const TSharedPtr<const FPCGEditor> SharedPtr = PCGEditorWeakPtr.Pin())
+							{
+								SharedPtr->BringFocusToPanel(EPCGEditorPanel::UserParams);
+							}
+
+							return FReply::Handled();
+						})
+						.ToolTipText(LOCTEXT("OpenGraphParamPanelTooltip", "Opens the Graph Parameters Panel."))
+						[
+							SNew(STextBlock)
+							.Font(IDetailLayoutBuilder::GetDetailFont())
+							.Text(LOCTEXT("OpenGraphParamPanelButtonText", "Open Graph Parameters"))
+						]
+					]
+				];
+		}
+	}
+
+	const FName PCGCategoryName("PCG");
+	IDetailCategoryBuilder& PCGCategory = DetailBuilder.EditCategory(PCGCategoryName);
+
 	TArray<TSharedRef<IPropertyHandle>> AllProperties;
 	bool bSimpleProperties = true;
 	bool bAdvancedProperties = false;
 	// Add all properties in the category in order
 	PCGCategory.GetDefaultProperties(AllProperties, bSimpleProperties, bAdvancedProperties);
 
-	for (auto& Property : AllProperties)
+	for (TSharedRef<IPropertyHandle>& Property : AllProperties)
 	{
 		PCGCategory.AddProperty(Property);
 	}
diff --git a/Engine/Plugins/PCG/Source/PCGEditor/Private/PCGEditor.cpp b/Engine/Plugins/PCG/Source/PCGEditor/Private/PCGEditor.cpp
--- a/Engine/Plugins/PCG/Source/PCGEditor/Private/PCGEditor.cpp
+++ b/Engine/Plugins/PCG/Source/PCGEditor/Private/PCGEditor.cpp
@@ -24,6 +24,11 @@
 #include "PCGEditorSettings.h"
 #include "PCGEditorStyle.h"
 #include "PCGEditorUtils.h"
+#include "Nodes/PCGEditorGraphNode.h"
+#include "Nodes/PCGEditorGraphNodeInput.h"
+#include "Nodes/PCGEditorGraphNodeOutput.h"
+#include "Nodes/PCGEditorGraphNodeReroute.h"
+
 #include "Widgets/SPCGEditorGraphAttributeListView.h"
 #include "Widgets/SPCGEditorGraphDebugObjectTree.h"
 #include "Widgets/SPCGEditorGraphDetailsView.h"
@@ -31,11 +36,8 @@
 #include "Widgets/SPCGEditorGraphFind.h"
 #include "Widgets/SPCGEditorGraphLogView.h"
 #include "Widgets/SPCGEditorGraphNodePalette.h"
+#include "Widgets/SPCGEditorGraphParamsView.h"
 #include "Widgets/SPCGEditorGraphProfilingView.h"
-#include "Nodes/PCGEditorGraphNode.h"
-#include "Nodes/PCGEditorGraphNodeInput.h"
-#include "Nodes/PCGEditorGraphNodeOutput.h"
-#include "Nodes/PCGEditorGraphNodeReroute.h"
 #include "Widgets/SPCGEditorNodeSource.h"
 
 #include "AssetToolsModule.h"
@@ -89,6 +91,9 @@ namespace FPCGEditor_private
 	const FName ProfilingID = FName(TEXT("Profiling"));
 	const FName LogID = FName(TEXT("Log"));
 	const FName NodeSourceID = FName(TEXT("NodeSource"));
+	const FName UserParamsID = FName(TEXT("UserParams"));
+
+	const FText UserParamsTabName = LOCTEXT("UserParamsTab", "Graph Parameters");
 }
 
 UPCGEditorGraph* FPCGEditor::GetPCGEditorGraph(UPCGGraph* InGraph)
@@ -108,6 +113,18 @@ UPCGEditorGraph* FPCGEditor::GetPCGEditorGraph(UPCGGraph* InGraph)
 	return InGraph->PCGEditorGraph;
 }
 
+UPCGEditorGraph* FPCGEditor::GetPCGEditorGraph(const UPCGNode* InNode)
+{
+	UPCGGraph* PCGGraph = InNode ? Cast<UPCGGraph>(InNode->GetOuter()) : nullptr;
+	return GetPCGEditorGraph(PCGGraph);
+}
+
+UPCGEditorGraph* FPCGEditor::GetPCGEditorGraph(const UPCGSettings* InSettings)
+{
+	UPCGNode* PCGNode = InSettings ? Cast<UPCGNode>(InSettings->GetOuter()) : nullptr;
+	return GetPCGEditorGraph(PCGNode);
+}
+
 void FPCGEditor::Initialize(const EToolkitMode::Type InMode, const TSharedPtr<class IToolkitHost>& InToolkitHost, UPCGGraph* InPCGGraph)
 {
 	FCoreUObjectDelegates::OnObjectsReplaced.AddSP(this, &FPCGEditor::OnObjectsReplaced);
@@ -141,6 +158,7 @@ void FPCGEditor::Initialize(const EToolkitMode::Type InMode, const TSharedPtr<cl
 	ProfilingWidget = CreateProfilingWidget();
 	LogWidget = CreateLogWidget();
 	NodeSourceWidget = CreateNodeSourceWidget();
+	UserParamsWidget = CreateGraphParamsWidget();
 
 	BindCommands();
 	RegisterToolbar();
@@ -486,6 +504,64 @@ void FPCGEditor::JumpToNode(const UPCGNode* InNode)
 	}
 }
 
+FName FPCGEditor::GetPanelID(const EPCGEditorPanel Panel) const
+{
+	switch (Panel)
+	{
+		case EPCGEditorPanel::Attributes1:
+			return FPCGEditor_private::AttributesID[0];
+		case EPCGEditorPanel::Attributes2:
+			return FPCGEditor_private::AttributesID[1];
+		case EPCGEditorPanel::Attributes3:
+			return FPCGEditor_private::AttributesID[2];
+		case EPCGEditorPanel::Attributes4:
+			return FPCGEditor_private::AttributesID[3];
+		case EPCGEditorPanel::DebugObjectTree:
+			return FPCGEditor_private::DebugObjectID;
+		case EPCGEditorPanel::Determinism:
+			return FPCGEditor_private::DeterminismID;
+		case EPCGEditorPanel::Find:
+			return FPCGEditor_private::FindID;
+		case EPCGEditorPanel::GraphEditor:
+			return FPCGEditor_private::GraphEditorID;
+		case EPCGEditorPanel::Log:
+			return FPCGEditor_private::LogID;
+		case EPCGEditorPanel::NodePalette:
+			return FPCGEditor_private::PaletteID;
+		case EPCGEditorPanel::NodeSource:
+			return FPCGEditor_private::NodeSourceID;
+		case EPCGEditorPanel::Profiling:
+			return FPCGEditor_private::ProfilingID;
+		case EPCGEditorPanel::PropertyDetails1:
+			return FPCGEditor_private::PropertyDetailsID[0];
+		case EPCGEditorPanel::PropertyDetails2:
+			return FPCGEditor_private::PropertyDetailsID[1];
+		case EPCGEditorPanel::PropertyDetails3:
+			return FPCGEditor_private::PropertyDetailsID[2];
+		case EPCGEditorPanel::PropertyDetails4:
+			return FPCGEditor_private::PropertyDetailsID[3];
+		case EPCGEditorPanel::UserParams:
+			return FPCGEditor_private::UserParamsID;
+		default:
+			return NAME_None;
+	}
+}
+
+void FPCGEditor::BringFocusToPanel(const EPCGEditorPanel Panel) const
+{
+	const FName PanelID = GetPanelID(Panel);
+	if (PanelID != NAME_None)
+	{
+		const TSharedPtr<SDockTab> Tab = TabManager->TryInvokeTab(PanelID);
+		Tab->DrawAttention(); // Bring the panel to focus and flash the tab
+	}
+}
+
+bool FPCGEditor::IsPanelCurrentlyOpen(const EPCGEditorPanel Panel) const
+{
+	return TabManager.IsValid() && TabManager->FindExistingLiveTab(GetPanelID(Panel));
+}
+
 void FPCGEditor::RegisterTabSpawners(const TSharedRef<FTabManager>& InTabManager)
 {
 	WorkspaceMenuCategory = InTabManager->AddLocalWorkspaceMenuCategory(LOCTEXT("WorkspaceMenu_PCGEditor", "PCG Editor"));
@@ -559,6 +635,10 @@ void FPCGEditor::RegisterTabSpawners(const TSharedRef<FTabManager>& InTabManager
 	InTabManager->RegisterTabSpawner(FPCGEditor_private::NodeSourceID, FOnSpawnTab::CreateSP(this, &FPCGEditor::SpawnTab_NodeSource))
 		.SetDisplayName(LOCTEXT("NodeSourceTab", "Node Source Editor"))
 		.SetGroup(WorkspaceMenuCategoryRef);
+
+	InTabManager->RegisterTabSpawner(FPCGEditor_private::UserParamsID, FOnSpawnTab::CreateSP(this, &FPCGEditor::SpawnTab_UserParams))
+		.SetDisplayName(FPCGEditor_private::UserParamsTabName)
+		.SetGroup(WorkspaceMenuCategoryRef);
 }
 
 void FPCGEditor::UnregisterTabSpawners(const TSharedRef<class FTabManager>& InTabManager)
@@ -579,6 +659,7 @@ void FPCGEditor::UnregisterTabSpawners(const TSharedRef<class FTabManager>& InTa
 	InTabManager->UnregisterTabSpawner(FPCGEditor_private::ProfilingID);
 	InTabManager->UnregisterTabSpawner(FPCGEditor_private::LogID);
 	InTabManager->UnregisterTabSpawner(FPCGEditor_private::NodeSourceID);
+	InTabManager->UnregisterTabSpawner(FPCGEditor_private::UserParamsID);
 
 	FAssetEditorToolkit::UnregisterTabSpawners(InTabManager);
 }
@@ -1419,19 +1500,6 @@ bool FPCGEditor::CanSelectNamedRerouteDeclaration() const
 	return false;
 }
 
-bool FPCGEditor::IsNodeSourceEditorTabClosed() const
-{
-	return !TabManager.IsValid() || !TabManager->FindExistingLiveTab(FPCGEditor_private::NodeSourceID).IsValid();
-}
-
-void FPCGEditor::SpawnNodeSourceEditorTab()
-{
-	if (TabManager.IsValid() && !TabManager->FindExistingLiveTab(FPCGEditor_private::NodeSourceID))
-	{
-		TabManager->TryInvokeTab(FPCGEditor_private::NodeSourceID);
-	}
-}
-
 FReply FPCGEditor::OnSpawnNodeByShortcut(FInputChord InChord, const FVector2D& InPosition, UPCGEditorGraph* InGraph)
 {
 	const TSharedPtr<FEdGraphSchemaAction> Action = FPCGEditorSpawnNodeCommands::Get().GetGraphActionByChord(InChord);
@@ -2766,6 +2834,11 @@ TSharedRef<SPCGEditorNodeSource> FPCGEditor::CreateNodeSourceWidget()
 	return SNew(SPCGEditorNodeSource, SharedThis(this));
 }
 
+TSharedRef<SPCGEditorGraphUserParametersView> FPCGEditor::CreateGraphParamsWidget()
+{
+	return SNew(SPCGEditorGraphUserParametersView, SharedThis(this));
+}
+
 void FPCGEditor::OnSelectedNodesChanged(const TSet<UObject*>& NewSelection)
 {
 	TArray<TWeakObjectPtr<UObject>> SelectedObjects;
@@ -3162,6 +3235,16 @@ TSharedRef<SDockTab> FPCGEditor::SpawnTab_NodeSource(const FSpawnTabArgs& Args)
 		];
 }
 
+TSharedRef<SDockTab> FPCGEditor::SpawnTab_UserParams(const FSpawnTabArgs& Args)
+{
+	return SNew(SDockTab)
+		.Label(FPCGEditor_private::UserParamsTabName)
+		.TabColorScale(GetTabColorScale())
+		[
+			UserParamsWidget.ToSharedRef()
+		];
+}
+
 FText FPCGEditor::GetDetailsTabLabel(int DetailsIndex)
 {
 	if (DetailsIndex == 0)
diff --git a/Engine/Plugins/PCG/Source/PCGEditor/Private/PCGEditor.h b/Engine/Plugins/PCG/Source/PCGEditor/Private/PCGEditor.h
--- a/Engine/Plugins/PCG/Source/PCGEditor/Private/PCGEditor.h
+++ b/Engine/Plugins/PCG/Source/PCGEditor/Private/PCGEditor.h
@@ -23,6 +23,7 @@ class SPCGEditorGraphLogView;
 class SPCGEditorGraphNodePalette;
 class SPCGEditorGraphProfilingView;
 class SPCGEditorNodeSource;
+class SPCGEditorGraphUserParametersView;
 class UEdGraphNode;
 class UPCGComponent;
 class UPCGEditorGraph;
@@ -33,6 +34,27 @@ struct FPCGCompilerDiagnostics;
 DECLARE_MULTICAST_DELEGATE_OneParam(FOnInspectedStackChanged, const FPCGStack&);
 DECLARE_MULTICAST_DELEGATE_OneParam(FOnInspectedNodeChanged, UPCGEditorGraphNodeBase*);
 
+enum class EPCGEditorPanel
+{
+	Attributes1,
+	Attributes2,
+	Attributes3,
+	Attributes4,
+	DebugObjectTree,
+	Determinism,
+	Find,
+	GraphEditor,
+	Log,
+	NodePalette,
+	NodeSource,
+	Profiling,
+	PropertyDetails1,
+	PropertyDetails2,
+	PropertyDetails3,
+	PropertyDetails4,
+	UserParams
+};
+
 class FPCGEditor : public FAssetEditorToolkit, public FGCObject, public FSelfRegisteringEditorUndoClient
 {
 public:
@@ -44,6 +66,8 @@ public:
 
 	/** Gets/Creates the PCG graph editor for a given PCG graph */
 	static UPCGEditorGraph* GetPCGEditorGraph(UPCGGraph* InGraph);
+	static UPCGEditorGraph* GetPCGEditorGraph(const UPCGNode* InNode);
+	static UPCGEditorGraph* GetPCGEditorGraph(const UPCGSettings* InSettings);
 
 	/** Get the PCG graph being edited. */
 	const UPCGGraph* GetPCGGraph() { return PCGGraphBeingEdited; }
@@ -69,6 +93,13 @@ public:
 	void JumpToNode(const UEdGraphNode* InNode);
 	void JumpToNode(const UPCGNode* InNode);
 
+	/** Get the TabID of the editor panel. */
+	FName GetPanelID(EPCGEditorPanel Panel) const;
+	/** Focuses the user on a specific Details View panel and flashes the tab. */
+	void BringFocusToPanel(EPCGEditorPanel Panel) const;
+	/** Returns true if the selected tab is currently open. */
+	bool IsPanelCurrentlyOpen(EPCGEditorPanel Panel) const;
+
 	/** Helper to get to the subsystem. */
 	static class UPCGSubsystem* GetSubsystem();
 
@@ -100,12 +131,6 @@ public:
 	virtual void InitToolMenuContext(FToolMenuContext& MenuContext) override;
 	// ~End FAssetEditorToolkit interface
 
-	/** Returns true if the Node Source Editor tab is not currently open. */
-	bool IsNodeSourceEditorTabClosed() const;
-
-	/** Opens the Node Source Editor tab if it is not already visible. */
-	void SpawnNodeSourceEditorTab();
-
 	/**
 	 * Handles spawning a graph node in the current graph using the passed in chord.
 	 *
@@ -324,6 +349,9 @@ private:
 	/** Create a new node source editor tab widget */
 	TSharedRef<SPCGEditorNodeSource> CreateNodeSourceWidget();
 
+	/** Create a new user graph parameters tab widget */
+	TSharedRef<SPCGEditorGraphUserParametersView> CreateGraphParamsWidget();
+
 	/** Called when the selection changes in the GraphEditor */
 	void OnSelectedNodesChanged(const TSet<UObject*>& NewSelection);
 
@@ -380,6 +408,7 @@ private:
 	TSharedRef<SDockTab> SpawnTab_Profiling(const FSpawnTabArgs& Args);
 	TSharedRef<SDockTab> SpawnTab_Log(const FSpawnTabArgs& Args);
 	TSharedRef<SDockTab> SpawnTab_NodeSource(const FSpawnTabArgs& Args);
+	TSharedRef<SDockTab> SpawnTab_UserParams(const FSpawnTabArgs& Args);
 
 	FText GetDetailsTabLabel(int DetailsIndex);
 	FText GetDetailsViewObjectName(int DetailsIndex);
@@ -395,6 +424,7 @@ private:
 	TSharedPtr<SPCGEditorGraphProfilingView> ProfilingWidget;
 	TSharedPtr<SPCGEditorGraphLogView> LogWidget;
 	TSharedPtr<SPCGEditorNodeSource> NodeSourceWidget;
+	TSharedPtr<SPCGEditorGraphUserParametersView> UserParamsWidget;
 
 	TSharedPtr<FUICommandList> GraphEditorCommands;
 
diff --git a/Engine/Plugins/PCG/Source/PCGEditor/Private/Widgets/SPCGEditorGraphDetailsView.cpp b/Engine/Plugins/PCG/Source/PCGEditor/Private/Widgets/SPCGEditorGraphDetailsView.cpp
--- a/Engine/Plugins/PCG/Source/PCGEditor/Private/Widgets/SPCGEditorGraphDetailsView.cpp
+++ b/Engine/Plugins/PCG/Source/PCGEditor/Private/Widgets/SPCGEditorGraphDetailsView.cpp
@@ -2,17 +2,17 @@
 
 #include "SPCGEditorGraphDetailsView.h"
 
+#include "PCGEditor.h"
 #include "PCGGraph.h"
 #include "PCGNode.h"
 #include "PCGSettings.h"
-
-#include "PCGEditor.h"
 #include "Nodes/PCGEditorGraphNode.h"
 
 #include "DetailsViewArgs.h"
 #include "IDetailPropertyExtensionHandler.h"
 #include "IDetailsView.h"
 #include "PropertyEditorModule.h"
+#include "Editor/PropertyEditor/Private/PropertyNode.h"
 #include "Framework/Notifications/NotificationManager.h"
 #include "Layout/Visibility.h"
 #include "Modules/ModuleManager.h"
@@ -144,13 +144,29 @@ void SPCGEditorGraphDetailsView::Construct(const FArguments& InArgs)
 	DetailsViewArgs.bCustomFilterAreaLocation = true;
 	DetailsViewArgs.bHideSelectionTip = true;
 	DetailsViewArgs.DefaultsOnlyVisibility = EEditDefaultsOnlyNodeVisibility::Automatic;
+	// Filter out the graph parameters from the details view
+	DetailsViewArgs.ShouldForceHideProperty = FDetailsViewArgs::FShouldForceHideProperty::CreateLambda(
+	[](const TSharedRef<FPropertyNode>& PropertyNode)
+	{
+		if (const FProperty* Property = PropertyNode->GetProperty())
+		{
+			if (const FStructProperty* PropertyStruct = CastField<FStructProperty>(Property))
+			{
+				// Explicit because FInstancedPropertyBag holds the Custom DetailsView.
+				return PropertyStruct->Struct == FInstancedPropertyBag::StaticStruct();
+			}
+		}
+
+		return false;
+	});
 
 	FPropertyEditorModule& PropertyEditorModule = FModuleManager::GetModuleChecked<FPropertyEditorModule>("PropertyEditor");
 	DetailsView = PropertyEditorModule.CreateDetailView(DetailsViewArgs);
 
 	DetailsView->SetIsPropertyReadOnlyDelegate(FIsPropertyReadOnly::CreateSP(this, &SPCGEditorGraphDetailsView::IsReadOnlyProperty));
 	DetailsView->SetIsPropertyVisibleDelegate(FIsPropertyVisible::CreateSP(this, &SPCGEditorGraphDetailsView::IsVisibleProperty));
-	TSharedRef<FPCGDebugDetailsExtension> DebugDetailsExtension = MakeShared<FPCGDebugDetailsExtension>();
+	// Adds helper buttons to the Debug Details View display.
+	const TSharedRef<FPCGDebugDetailsExtension> DebugDetailsExtension = MakeShared<FPCGDebugDetailsExtension>();
 	DetailsView->SetExtensionHandler(DebugDetailsExtension);
 
 	ChildSlot
diff --git a/Engine/Plugins/PCG/Source/PCGEditor/Private/Widgets/SPCGEditorGraphParamsView.cpp b/Engine/Plugins/PCG/Source/PCGEditor/Private/Widgets/SPCGEditorGraphParamsView.cpp
new file mode 100644
--- /dev/null
+++ b/Engine/Plugins/PCG/Source/PCGEditor/Private/Widgets/SPCGEditorGraphParamsView.cpp
@@ -0,0 +1,98 @@
+// Copyright Epic Games, Inc. All Rights Reserved.
+
+#include "SPCGEditorGraphParamsView.h"
+
+#include "PCGEditor.h"
+#include "PCGEditorGraph.h"
+#include "PCGGraph.h"
+
+#include "PropertyBagDetails.h"
+#include "PropertyEditorModule.h"
+#include "Editor/PropertyEditor/Private/PropertyNode.h"
+#include "Editor/PropertyEditor/Private/StructurePropertyNode.h"
+#include "Modules/ModuleManager.h"
+
+void SPCGEditorGraphUserParametersView::Construct(const FArguments& InArgs, const TSharedPtr<FPCGEditor>& InPCGEditor)
+{
+	if (!InPCGEditor.IsValid())
+	{
+		return;
+	}
+
+	UPCGGraph* PCGGraph = InPCGEditor.IsValid() && InPCGEditor->GetPCGEditorGraph() ? InPCGEditor->GetPCGEditorGraph()->GetPCGGraph() : nullptr;
+
+	if (!PCGGraph)
+	{
+		return;
+	}
+
+	FPropertyEditorModule& PropertyEditorModule = FModuleManager::LoadModuleChecked<FPropertyEditorModule>("PropertyEditor");
+
+	FDetailsViewArgs DetailsViewArgs;
+	{
+		DetailsViewArgs.NameAreaSettings = FDetailsViewArgs::ENameAreaSettings::HideNameArea;
+		DetailsViewArgs.DefaultsOnlyVisibility = EEditDefaultsOnlyNodeVisibility::Automatic;
+		DetailsViewArgs.bAllowSearch = true;
+		DetailsViewArgs.bShowScrollBar = true;
+		DetailsViewArgs.bShowOptions = true;
+		DetailsViewArgs.bShowObjectLabel = true;
+
+		DetailsViewArgs.bShowModifiedPropertiesOption = false;
+		DetailsViewArgs.bShowDifferingPropertiesOption = false;
+		DetailsViewArgs.bShowHiddenPropertiesWhilePlayingOption = false;
+		DetailsViewArgs.bShowKeyablePropertiesOption = false;
+		DetailsViewArgs.bShowAnimatedPropertiesOption = false;
+		DetailsViewArgs.bShowCustomFilterOption = false;
+		DetailsViewArgs.bShowSectionSelector = false;
+		DetailsViewArgs.bShowLooseProperties = false;
+
+		DetailsViewArgs.bShowPropertyMatrixButton = false;
+		DetailsViewArgs.bUpdatesFromSelection = false;
+		DetailsViewArgs.bForceHiddenPropertyVisibility = false;
+		// Filter all but the graph parameters
+		DetailsViewArgs.ShouldForceHideProperty = FDetailsViewArgs::FShouldForceHideProperty::CreateLambda(
+		[](const TSharedRef<FPropertyNode>& PropertyNode)
+		{
+			auto PropertyIsPropertyBag = [](const FProperty* Property)
+			{
+				const FStructProperty* PropertyStruct = CastField<FStructProperty>(Property);
+				// Explicitly check against FInstancedPropertyBag, because that is what holds the Custom DetailsView.
+				return PropertyStruct && PropertyStruct->Struct == FInstancedPropertyBag::StaticStruct();
+			};
+
+			if (const FProperty* Property = PropertyNode->GetProperty())
+			{
+				const bool bPropertyIsPropertyBag = PropertyIsPropertyBag(Property);
+				bool bPropertyIsInPropertyBag = false;
+				// Traverse the parentage to see if the property is contained within the property bag.
+				for (const FPropertyNode* ParentNode = PropertyNode->GetParentNode(); ParentNode != nullptr; ParentNode = ParentNode->GetParentNode())
+				{
+					if (PropertyIsPropertyBag(ParentNode->GetProperty()))
+					{
+						bPropertyIsInPropertyBag = true;
+						break;
+					}
+				}
+
+				// Filter if the property isn't the property bag and the property is not inside the property bag.
+				return !(bPropertyIsPropertyBag || bPropertyIsInPropertyBag);
+			}
+
+			return true;
+		});
+	}
+
+	// Note: Single Property View Header notes that it doesn't work with arrays or structs, so using filters for now.
+	const TSharedRef<IDetailsView> DetailsView = PropertyEditorModule.CreateDetailView(DetailsViewArgs);
+	DetailsView->SetObject(PCGGraph);
+
+	ChildSlot
+	[
+		SNew(SVerticalBox)
+		+ SVerticalBox::Slot()
+		.FillHeight(1.f)
+		[
+			DetailsView
+		]
+	];
+}
diff --git a/Engine/Plugins/PCG/Source/PCGEditor/Private/Widgets/SPCGEditorGraphParamsView.h b/Engine/Plugins/PCG/Source/PCGEditor/Private/Widgets/SPCGEditorGraphParamsView.h
new file mode 100644
--- /dev/null
+++ b/Engine/Plugins/PCG/Source/PCGEditor/Private/Widgets/SPCGEditorGraphParamsView.h
@@ -0,0 +1,20 @@
+// Copyright Epic Games, Inc. All Rights Reserved.
+
+#pragma once
+
+#include "Templates/SharedPointer.h"
+
+#include "Widgets/SCompoundWidget.h"
+
+struct FPropertyAndParent;
+class FPCGEditor;
+
+/** Displays a PCGGraph's User Parameters details view. */
+class SPCGEditorGraphUserParametersView final : public SCompoundWidget
+{
+public:
+	SLATE_BEGIN_ARGS(SPCGEditorGraphUserParametersView)	{}
+	SLATE_END_ARGS()
+
+	void Construct(const FArguments& InArgs, const TSharedPtr<FPCGEditor>& InPCGEditor);
+};
\ No newline at end of file
