[PCG] Fix warning when creating a Level Instance from a generated PCG Actor
- The managed resources were not marked as garbage and were reporting invalid references
- This was also causing the actors in the level instance keeping their offset because of an early save (prior to setting bAlreadyMovedActors being set to true) caused by the reference
- Level Instance Offset could happen with a different system (separate fix incoming)

#rnx

diff --git a/Engine/Plugins/PCG/Source/PCG/Private/PCGComponent.cpp b/Engine/Plugins/PCG/Source/PCG/Private/PCGComponent.cpp
--- a/Engine/Plugins/PCG/Source/PCG/Private/PCGComponent.cpp
+++ b/Engine/Plugins/PCG/Source/PCG/Private/PCGComponent.cpp
@@ -1811,8 +1811,33 @@ void UPCGComponent::OnUnregister()
 	Super::OnUnregister();
 }
 
+void UPCGComponent::MarkSubObjectsAsGarbage()
+{
+	auto ModifyAndMarkForGarbage = [](UObject* SubObject)
+	{
+		if (!SubObject)
+		{
+			return;
+		}
+
+#if WITH_EDITOR
+		// Modify object for Undo/Redo support
+		if (GEditor)
+		{
+			// We don't have the info here to know if we should dirty or not but we can rely on the fact that this component will have been modified prior to its OnComponentDestroyed call
+			SubObject->Modify(/*bAlwaysMarkDirty=*/false);
+		}
+#endif
+		SubObject->MarkAsGarbage();
+	};
+
+	// Make sure all outered objects of this component are marked as garbage so that references from those objects to our owner actor do not get reported
+	ForEachObjectWithOuter(this, ModifyAndMarkForGarbage, true);
+}
+
 void UPCGComponent::OnComponentDestroyed(bool bDestroyingHierarchy)
 {
+	TRACE_CPUPROFILER_EVENT_SCOPE(UPCGComponent::OnComponentDestroyed);
 #if WITH_EDITOR
 	// BeginDestroy is not called immediately when a component is destroyed. Therefore callbacks are not cleaned
 	// until GC is ran, and can stack up with BP reconstruction scripts. Force the removal of callbacks here. If the component
@@ -1832,6 +1857,14 @@ void UPCGComponent::OnComponentDestroyed(bool bDestroyingHierarchy)
 	}
 
 #if WITH_EDITOR
+	// Only do this in Editor as it hasn't proven to cause issues at runtime and the call can be costly:
+	// If Component is getting destroyed and it is not a BP constructed component or we are not currently reconstructing the BP
+	// We need to mark subobjects (GeneratedOutputData, ManagedResources) for Garbage so that they aren't found by reference checking code paths
+	if (!IsCreatedByConstructionScript() || !GIsReconstructingBlueprintInstances)
+	{
+		MarkSubObjectsAsGarbage();
+	}
+
 	// Don't do the unregister in OnUnregister, because we have flows where the component gets Unregistered/Registered without getting destroyed.
 	if (UPCGSubsystem* Subsystem = GetSubsystem())
 	{
@@ -4106,9 +4139,15 @@ void FPCGComponentInstanceData::ApplyToComponent(UActorComponent* Component, con
 
 #if WITH_EDITOR
 		// Disconnect callbacks on source.
-		if (ConstructionSourceComponent && ConstructionSourceComponent->GraphInstance)
+		if (ConstructionSourceComponent)
 		{
-			ConstructionSourceComponent->GraphInstance->TeardownCallbacks();
+			if (ConstructionSourceComponent->GraphInstance)
+			{
+				ConstructionSourceComponent->GraphInstance->TeardownCallbacks();
+			}
+			// Make sure previous component destroys all its sub-objects to avoid leaving references to the new component
+			// This can happen in Copy/Paste code paths (ex: Level Instance creation)
+			const_cast<UPCGComponent*>(ConstructionSourceComponent)->MarkSubObjectsAsGarbage();
 		}
 
 		// Finally, start a delayed refresh task (if there is not one already), in editor only
diff --git a/Engine/Plugins/PCG/Source/PCG/Public/PCGComponent.h b/Engine/Plugins/PCG/Source/PCG/Public/PCGComponent.h
--- a/Engine/Plugins/PCG/Source/PCG/Public/PCGComponent.h
+++ b/Engine/Plugins/PCG/Source/PCG/Public/PCGComponent.h
@@ -521,6 +521,8 @@ protected:
 	/** Purges Actors and Components generated by PCG but are no longer managed by any PCG Component */
 	static void PurgeUnlinkedResources(const AActor* InActor);
 
+	void MarkSubObjectsAsGarbage();
+
 protected:
 	UPROPERTY(BlueprintReadOnly, VisibleAnywhere, Category = Settings, Instanced, meta = (NoResetToDefault, DisplayPriority = 100))
 	TObjectPtr<UPCGGraphInstance> GraphInstance;
