Added additional per-data MISMC reuse capabilities and option to never merge for different input data.

diff --git a/Engine/Plugins/PCG/Source/PCG/Private/Elements/PCGStaticMeshSpawner.cpp b/Engine/Plugins/PCG/Source/PCG/Private/Elements/PCGStaticMeshSpawner.cpp
--- a/Engine/Plugins/PCG/Source/PCG/Private/Elements/PCGStaticMeshSpawner.cpp
+++ b/Engine/Plugins/PCG/Source/PCG/Private/Elements/PCGStaticMeshSpawner.cpp
@@ -137,7 +137,7 @@ bool FPCGStaticMeshSpawnerElement::PrepareDataInternal(FPCGContext* InContext) c
 	const bool bGenerateOutput = Context->Node && Context->Node->IsOutputPinConnected(PCGPinConstants::DefaultOutputLabel);
 #endif
 
-	// Check if we can reuse existing resources
+	// Check if we can reuse existing resources on an "all-or-nothing" basis
 	bool& bSkippedDueToReuse = Context->bSkippedDueToReuse;
 
 	if (!Context->bReuseCheckDone && CVarAllowISMReuse.GetValueOnAnyThread())
@@ -222,6 +222,58 @@ bool FPCGStaticMeshSpawnerElement::PrepareDataInternal(FPCGContext* InContext) c
 				continue;
 			}
 
+			// Prior to selection, if we have MISMC resources that have the same settings crc & the same data crc,
+			// We can skip the selection here, unless we need to output the data (again, could/should be cached instead).
+			FPCGCrc DataCrc;
+			if(!bSkippedDueToReuse && CVarAllowISMReuse.GetValueOnAnyThread())
+			{
+				check(Context->DependenciesCrc.IsValid()); // should have been done earlier
+
+				const FPCGCrc SettingsCrc = Settings->GetSettingsCrc();
+				ensure(SettingsCrc.IsValid());
+
+				// Compute this specific data crc as-if it were alone
+				{
+					FPCGDataCollection SubCollection;
+					SubCollection.TaggedData.Add(Input);
+					SubCollection.ComputeCrcs(ShouldComputeFullOutputDataCrc(Context));
+
+					GetDependenciesCrc(SubCollection, Settings, Context->SourceComponent.Get(), DataCrc);
+				}
+
+				TArray<UPCGManagedISMComponent*> MISMCs;
+				Context->SourceComponent->ForEachManagedResource([&MISMCs, &SettingsCrc, &DataCrc, &Context](UPCGManagedResource* InResource)
+				{
+					if (UPCGManagedISMComponent* Resource = Cast<UPCGManagedISMComponent>(InResource))
+					{
+						if (Resource->GetSettingsCrc().IsValid() && Resource->GetSettingsCrc() == SettingsCrc &&
+							Resource->GetDataCrc().IsValid() && Resource->GetDataCrc() == DataCrc)
+						{
+							MISMCs.Add(Resource);
+						}
+					}
+				});
+
+				for (UPCGManagedISMComponent* MISMC : MISMCs)
+				{
+					if (!MISMC->IsMarkedUnused() && Settings->bWarnOnIdenticalSpawn)
+					{
+						// TODO: Revisit if the stack is added to the managed components at creation
+						PCGLog::LogWarningOnGraph(LOCTEXT("IdenticalISMCSpawn", "Identical ISM Component spawn occurred. It may be beneficial to re-check graph logic for identical spawn conditions (same mesh descriptor at same location, etc) or repeated nodes."), Context);
+					}
+
+					MISMC->MarkAsReused();
+
+					// Update global crc, otherwise these resources wouldn't get picked up in a subsequent update
+					MISMC->SetCrc(Context->DependenciesCrc);
+				}
+
+				if (!MISMCs.IsEmpty())
+				{
+					Context->bCurrentDataSkippedDueToReuse = true;
+				}
+			}
+
 			if (bGenerateOutput)
 			{
 				FPCGTaggedData& Output = Outputs.Add_GetRef(Input);
@@ -249,12 +301,22 @@ bool FPCGStaticMeshSpawnerElement::PrepareDataInternal(FPCGContext* InContext) c
 				Context->CurrentOutputPointData = OutputPointData;
 			}
 
-			FPCGStaticMeshSpawnerContext::FPackedInstanceListData& InstanceListData = Context->MeshInstancesData.Emplace_GetRef();
-			InstanceListData.TargetActor = TargetActor;
-			InstanceListData.SpatialData = PointData;
-
-			Context->CurrentPointData = PointData;
-			Context->bCurrentInputSetup = true;
+			if (bGenerateOutput || !Context->bCurrentDataSkippedDueToReuse)
+			{
+				FPCGStaticMeshSpawnerContext::FPackedInstanceListData& InstanceListData = Context->MeshInstancesData.Emplace_GetRef();
+				InstanceListData.TargetActor = TargetActor;
+				InstanceListData.SpatialData = PointData;
+				InstanceListData.DataCrc = DataCrc;
+				InstanceListData.bSkippedDueToReuse = Context->bSkippedDueToReuse || Context->bCurrentDataSkippedDueToReuse;
+
+				Context->CurrentPointData = PointData;
+				Context->bCurrentInputSetup = true;
+			}
+			else
+			{
+				// skip selection
+				Context->bSelectionDone = true;
+			}
 		}
 
 		// TODO: If we know we re-use the ISMCs, we should not run the Selection, as it can be pretty costly.
@@ -271,7 +333,7 @@ bool FPCGStaticMeshSpawnerElement::PrepareDataInternal(FPCGContext* InContext) c
 		}
 
 		// If we need the output but would otherwise skip the resource creation, we don't need to run the instance packing part of the processing
-		if (!bSkippedDueToReuse)
+		if (!bSkippedDueToReuse && !Context->bCurrentDataSkippedDueToReuse)
 		{
 			TArray<FPCGPackedCustomData>& PackedCustomData = Context->MeshInstancesData.Last().PackedCustomData;
 			const TArray<FPCGMeshInstanceList>& MeshInstances = Context->MeshInstancesData.Last().MeshInstances;
@@ -341,7 +403,7 @@ bool FPCGStaticMeshSpawnerElement::ExecuteInternal(FPCGContext* InContext) const
 	while(!Context->MeshInstancesData.IsEmpty())
 	{
 		const FPCGStaticMeshSpawnerContext::FPackedInstanceListData& InstanceList = Context->MeshInstancesData.Last();
-		check(Context->bSkippedDueToReuse || InstanceList.MeshInstances.Num() == InstanceList.PackedCustomData.Num());
+		check(InstanceList.bSkippedDueToReuse || InstanceList.MeshInstances.Num() == InstanceList.PackedCustomData.Num());
 
 		const bool bTargetActorValid = (InstanceList.TargetActor && IsValid(InstanceList.TargetActor));
 
@@ -352,7 +414,7 @@ bool FPCGStaticMeshSpawnerElement::ExecuteInternal(FPCGContext* InContext) const
 				const FPCGMeshInstanceList& MeshInstance = InstanceList.MeshInstances[Context->CurrentDataIndex];
 				// We always have mesh instances, but if we are in re-use, we don't compute the packed custom data.
 				const FPCGPackedCustomData* PackedCustomData = InstanceList.PackedCustomData.IsValidIndex(Context->CurrentDataIndex) ? &InstanceList.PackedCustomData[Context->CurrentDataIndex] : nullptr;
-				SpawnStaticMeshInstances(Context, MeshInstance, InstanceList.TargetActor, PackedCustomData);
+				SpawnStaticMeshInstances(Context, InstanceList, MeshInstance, PackedCustomData);
 
 				// Now that the mesh is loaded/spawned, set the bounds to out points if requested.
 				if (MeshInstance.Descriptor.StaticMesh && Settings->bApplyMeshBoundsToPoints)
@@ -420,11 +482,15 @@ bool FPCGStaticMeshSpawnerElement::CanExecuteOnlyOnMainThread(FPCGContext* Conte
 	return !Context || Context->CurrentPhase == EPCGExecutionPhase::Execute || Context->CurrentPhase == EPCGExecutionPhase::PrepareData;
 }
 
-void FPCGStaticMeshSpawnerElement::SpawnStaticMeshInstances(FPCGStaticMeshSpawnerContext* Context, const FPCGMeshInstanceList& InstanceList, AActor* TargetActor, const FPCGPackedCustomData* InPackedCustomData) const
+void FPCGStaticMeshSpawnerElement::SpawnStaticMeshInstances(FPCGStaticMeshSpawnerContext* Context, const FPCGStaticMeshSpawnerContext::FPackedInstanceListData& InstanceListData, const FPCGMeshInstanceList& InstanceList, const FPCGPackedCustomData* InPackedCustomData) const
 {
 	// Populate the (H)ISM from the previously prepared entries
 	TRACE_CPUPROFILER_EVENT_SCOPE(FPCGStaticMeshSpawnerElement::Execute::PopulateISMs);
 
+	AActor* TargetActor = InstanceListData.TargetActor;
+	const FPCGCrc& DataCrc = InstanceListData.DataCrc;
+	const bool bSkippedDueToReuse = InstanceListData.bSkippedDueToReuse;
+
 	if (InstanceList.Instances.Num() == 0)
 	{
 		return;
@@ -451,7 +517,7 @@ void FPCGStaticMeshSpawnerElement::SpawnStaticMeshInstances(FPCGStaticMeshSpawne
 	}
 
 	// Don't spawn meshes if we reuse the ISMCs, but we still want to be sure that the mesh is loaded at least (for operations downstream).
-	if (Context->bSkippedDueToReuse)
+	if (bSkippedDueToReuse)
 	{
 		return;
 	}
@@ -490,6 +556,12 @@ void FPCGStaticMeshSpawnerElement::SpawnStaticMeshInstances(FPCGStaticMeshSpawne
 	Params.SettingsCrc = Settings->GetSettingsCrc();
 	ensure(Params.SettingsCrc.IsValid());
 
+	// Enforce selection limitation when getting the MISMC based on whether we want to share the ISM or not.
+	if (!Settings->bAllowMergeDifferentDataInSameInstancedComponents)
+	{
+		Params.DataCrc = DataCrc;
+	}
+
 	// Implementation note: in order to prevent components from merging together when they have tags, we need to incorporate the tags hash
 	// into the settings crc. This will make sure we don't select component that don't match up with what we're expecting.
 	// However, this validation is potentially too restrictive, as the order and composition of the tags ultimately don't really matter.
@@ -508,6 +580,20 @@ void FPCGStaticMeshSpawnerElement::SpawnStaticMeshInstances(FPCGStaticMeshSpawne
 	check(MISMC);
 	MISMC->SetCrc(Context->DependenciesCrc);
 
+	ensure(Settings->bAllowMergeDifferentDataInSameInstancedComponents || !MISMC->GetDataCrc().IsValid() && MISMC->GetComponent() && MISMC->GetComponent()->GetNumInstances() == 0);
+
+	// Mutate MISMC data crc, so this will ensure that when we write multiple data to the same ISM that we don't trigger reuse.
+	if (MISMC->GetDataCrc().IsValid())
+	{
+		FPCGCrc UpdatedCrc = MISMC->GetDataCrc();
+		UpdatedCrc.Combine(DataCrc);
+		MISMC->SetDataCrc(UpdatedCrc);
+	}
+	else
+	{
+		MISMC->SetDataCrc(DataCrc);
+	}
+
 	// Keep track of all touched resources in the context, because if the execution is cancelled during the SMS execution
 	// we cannot easily guarantee that the state (esp. vs CRCs) is going to be entirely valid
 	Context->TouchedResources.Emplace(MISMC);
diff --git a/Engine/Plugins/PCG/Source/PCG/Private/Elements/PCGStaticMeshSpawnerContext.cpp b/Engine/Plugins/PCG/Source/PCG/Private/Elements/PCGStaticMeshSpawnerContext.cpp
--- a/Engine/Plugins/PCG/Source/PCG/Private/Elements/PCGStaticMeshSpawnerContext.cpp
+++ b/Engine/Plugins/PCG/Source/PCG/Private/Elements/PCGStaticMeshSpawnerContext.cpp
@@ -7,6 +7,7 @@ void FPCGStaticMeshSpawnerContext::ResetInputIterationData()
 	bCurrentInputSetup = false;
 	bSelectionDone = false;
 	bPartitionDone = false;
+	bCurrentDataSkippedDueToReuse = false;
 	CurrentPointData = nullptr;
 	CurrentOutputPointData = nullptr;
 	MaterialOverrideHelper.Reset();
diff --git a/Engine/Plugins/PCG/Source/PCG/Private/Helpers/PCGActorHelpers.cpp b/Engine/Plugins/PCG/Source/PCG/Private/Helpers/PCGActorHelpers.cpp
--- a/Engine/Plugins/PCG/Source/PCG/Private/Helpers/PCGActorHelpers.cpp
+++ b/Engine/Plugins/PCG/Source/PCG/Private/Helpers/PCGActorHelpers.cpp
@@ -153,7 +153,9 @@ UPCGManagedISMComponent* UPCGActorHelpers::GetOrCreateManagedISMC(AActor* InTarg
 
 				if (UPCGManagedISMComponent* Resource = Cast<UPCGManagedISMComponent>(InResource))
 				{
-					if (!Resource->GetSettingsCrc().IsValid() || Resource->GetSettingsCrc() != InParams.SettingsCrc || !Resource->CanBeUsed())
+					if (!Resource->GetSettingsCrc().IsValid() || Resource->GetSettingsCrc() != InParams.SettingsCrc || 
+						(InParams.DataCrc.IsValid() && Resource->GetDataCrc() != InParams.DataCrc) ||
+						!Resource->CanBeUsed())
 					{
 						return;
 					}
diff --git a/Engine/Plugins/PCG/Source/PCG/Private/PCGManagedResource.cpp b/Engine/Plugins/PCG/Source/PCG/Private/PCGManagedResource.cpp
--- a/Engine/Plugins/PCG/Source/PCG/Private/PCGManagedResource.cpp
+++ b/Engine/Plugins/PCG/Source/PCG/Private/PCGManagedResource.cpp
@@ -786,6 +786,8 @@ bool UPCGManagedISMComponent::ReleaseIfUnused(TSet<TSoftObjectPtr<AActor>>& OutA
 
 void UPCGManagedISMComponent::ResetComponent()
 {
+	DataCrc = FPCGCrc();
+
 	if (UInstancedStaticMeshComponent* ISMC = GetComponent())
 	{
 		ISMC->ClearInstances();
diff --git a/Engine/Plugins/PCG/Source/PCG/Public/Elements/PCGStaticMeshSpawner.h b/Engine/Plugins/PCG/Source/PCG/Public/Elements/PCGStaticMeshSpawner.h
--- a/Engine/Plugins/PCG/Source/PCG/Public/Elements/PCGStaticMeshSpawner.h
+++ b/Engine/Plugins/PCG/Source/PCG/Public/Elements/PCGStaticMeshSpawner.h
@@ -3,7 +3,7 @@
 #pragma once
 
 #include "PCGSettings.h"
-
+#include "Elements/PCGStaticMeshSpawnerContext.h"
 #include "Metadata/PCGObjectPropertyOverride.h"
 
 #include "Engine/CollisionProfile.h"
@@ -14,9 +14,6 @@ class UPCGInstanceDataPackerBase;
 class UPCGMeshSelectorBase;
 class UPCGSpatialData;
 struct FPCGContext;
-struct FPCGMeshInstanceList;
-struct FPCGPackedCustomData;
-struct FPCGStaticMeshSpawnerContext;
 class UStaticMesh;
 
 UCLASS(BlueprintType, ClassGroup = (Procedural), meta = (PrioritizeCategories = "Settings"))
@@ -106,6 +103,10 @@ public:
 	UPROPERTY(BlueprintReadWrite, EditAnywhere, Category = "Settings|Debug")
 	bool bSynchronousLoad = false;
 
+	/** Controls whether instances stemming from different data can end up in the same ISM. Performance warning if set to false. */
+	UPROPERTY(BlueprintReadWrite, EditAnywhere, Category = "Settings|Advanced", meta = (PCG_Overridable))
+	bool bAllowMergeDifferentDataInSameInstancedComponents = true;
+
 	/** Opt-in option to silence errors when the property override attributes are not found. */
 	UPROPERTY(BlueprintReadWrite, EditAnywhere, Category = "Settings|Advanced")
 	bool bSilenceOverrideAttributeNotFoundErrors = false;
@@ -130,6 +131,6 @@ protected:
 	virtual bool PrepareDataInternal(FPCGContext* Context) const override;
 	virtual bool ExecuteInternal(FPCGContext* Context) const override;
 	virtual void AbortInternal(FPCGContext* Context) const override;
-	void SpawnStaticMeshInstances(FPCGStaticMeshSpawnerContext* Context, const FPCGMeshInstanceList& InstanceList, AActor* TargetActor, const FPCGPackedCustomData* InPackedCustomData) const;
+	void SpawnStaticMeshInstances(FPCGStaticMeshSpawnerContext* Context, const FPCGStaticMeshSpawnerContext::FPackedInstanceListData& InstanceListData, const FPCGMeshInstanceList& InstanceList, const FPCGPackedCustomData* InPackedCustomData) const;
 	virtual bool SupportsBasePointDataInputs(FPCGContext* InContext) const { return true; }
 };
diff --git a/Engine/Plugins/PCG/Source/PCG/Public/Elements/PCGStaticMeshSpawnerContext.h b/Engine/Plugins/PCG/Source/PCG/Public/Elements/PCGStaticMeshSpawnerContext.h
--- a/Engine/Plugins/PCG/Source/PCG/Public/Elements/PCGStaticMeshSpawnerContext.h
+++ b/Engine/Plugins/PCG/Source/PCG/Public/Elements/PCGStaticMeshSpawnerContext.h
@@ -32,6 +32,8 @@ struct PCG_API FPCGStaticMeshSpawnerContext : public FPCGContext, public IPCGAsy
 		const UPCGSpatialData* SpatialData = nullptr;
 		TArray<FPCGMeshInstanceList> MeshInstances;
 		TArray<FPCGPackedCustomData> PackedCustomData;
+		FPCGCrc DataCrc;
+		bool bSkippedDueToReuse = false;
 	};
 
 	TArray<FPackedInstanceListData> MeshInstancesData;
@@ -42,12 +44,13 @@ struct PCG_API FPCGStaticMeshSpawnerContext : public FPCGContext, public IPCGAsy
 
 	// Whole-execution variables
 	bool bReuseCheckDone = false;
-	bool bSkippedDueToReuse = false;
+	bool bSkippedDueToReuse = false; // fully skipped
 
 	// Per-input context variables
 	bool bCurrentInputSetup = false;
 	bool bSelectionDone = false;
 	bool bPartitionDone = false;
+	bool bCurrentDataSkippedDueToReuse = false;
 
 	const UPCGBasePointData* CurrentPointData = nullptr;
 	UPCGBasePointData* CurrentOutputPointData = nullptr;
diff --git a/Engine/Plugins/PCG/Source/PCG/Public/Helpers/PCGActorHelpers.h b/Engine/Plugins/PCG/Source/PCG/Public/Helpers/PCGActorHelpers.h
--- a/Engine/Plugins/PCG/Source/PCG/Public/Helpers/PCGActorHelpers.h
+++ b/Engine/Plugins/PCG/Source/PCG/Public/Helpers/PCGActorHelpers.h
@@ -55,6 +55,7 @@ struct FPCGISMComponentBuilderParams
 	FPCGSoftISMComponentDescriptor Descriptor;
 	int32 NumCustomDataFloats = 0;
 	FPCGCrc SettingsCrc;
+	FPCGCrc DataCrc;
 	bool bAllowDescriptorChanges = true;
 	bool bTransient = false;
 
diff --git a/Engine/Plugins/PCG/Source/PCG/Public/PCGManagedResource.h b/Engine/Plugins/PCG/Source/PCG/Public/PCGManagedResource.h
--- a/Engine/Plugins/PCG/Source/PCG/Public/PCGManagedResource.h
+++ b/Engine/Plugins/PCG/Source/PCG/Public/PCGManagedResource.h
@@ -286,6 +286,9 @@ public:
 	FPCGCrc GetSettingsCrc() const { return SettingsCrc; }
 	void SetSettingsCrc(const FPCGCrc& InSettingsCrc) { SettingsCrc = InSettingsCrc; }
 
+	FPCGCrc GetDataCrc() const { return DataCrc; }
+	void SetDataCrc(const FPCGCrc& InDataCrc) { DataCrc = InDataCrc; }
+
 	UE_DEPRECATED(5.6, "Use GetSettingsCrc instead")
 	uint64 GetSettingsUID() const { return -1; }
 	
@@ -307,6 +310,9 @@ protected:
 
 	UPROPERTY()
 	FPCGCrc SettingsCrc;
+	
+	UPROPERTY()
+	FPCGCrc DataCrc;
 };
 
 UCLASS(BlueprintType)
