[PCG] Fix ensure when accidentally caching data that is not cacheable (for example, TextureData that have an exported GPU buffer ref)

#rnx

diff --git a/Engine/Plugins/PCG/Source/PCG/Private/PCGElement.cpp b/Engine/Plugins/PCG/Source/PCG/Private/PCGElement.cpp
--- a/Engine/Plugins/PCG/Source/PCG/Private/PCGElement.cpp
+++ b/Engine/Plugins/PCG/Source/PCG/Private/PCGElement.cpp
@@ -639,16 +639,30 @@ void IPCGElement::PostExecutePrimaryLoopElement(FPCGContext* Context, const UPCG
 				}
 			}
 
-			FPCGDataCollection SubCollection = PrimaryDataCollection;
-			SubCollection.ComputeCrcs(bShouldComputeFullOutputDataCrc);
+			bool bCacheable = true;
 
-			FPCGCrc DependenciesCrc;
-			GetDependenciesCrc(SubCollection, Settings, Context->SourceComponent.Get(), DependenciesCrc);
+			for (const FPCGTaggedData& TaggedData : SubCollectionOutput.TaggedData)
+			{
+				if (TaggedData.Data && !TaggedData.Data->IsCacheable())
+				{
+					bCacheable = false;
+					break;
+				}
+			}
+
+			if (bCacheable)
+			{
+				FPCGDataCollection SubCollection = PrimaryDataCollection;
+				SubCollection.ComputeCrcs(bShouldComputeFullOutputDataCrc);
+
+				FPCGCrc DependenciesCrc;
+				GetDependenciesCrc(SubCollection, Settings, Context->SourceComponent.Get(), DependenciesCrc);
 
-			SubCollectionOutput.ComputeCrcs(bShouldComputeFullOutputDataCrc);
+				SubCollectionOutput.ComputeCrcs(bShouldComputeFullOutputDataCrc);
 
-			FPCGStoreInCacheParams Params = { .Element = this, .Crc = DependenciesCrc };
-			Subsystem->GetCache()->StoreInCache(Params, SubCollectionOutput);
+				FPCGStoreInCacheParams Params = { .Element = this, .Crc = DependenciesCrc };
+				Subsystem->GetCache()->StoreInCache(Params, SubCollectionOutput);
+			}
 		}
 	}
 
