[PCG] Progress Notification fix: total tasks wasn't properly updated

diff --git a/Engine/Plugins/PCG/Source/PCGEditor/Private/PCGEditorProgressNotification.cpp b/Engine/Plugins/PCG/Source/PCGEditor/Private/PCGEditorProgressNotification.cpp
--- a/Engine/Plugins/PCG/Source/PCGEditor/Private/PCGEditorProgressNotification.cpp
+++ b/Engine/Plugins/PCG/Source/PCGEditor/Private/PCGEditorProgressNotification.cpp
@@ -180,7 +180,7 @@ namespace PCGEditorProgressNotification
 }
 
 FPCGEditorProgressNotification::FPCGEditorProgressNotification(const FTextFormat& InTextFormat, bool bInCanCancel)
-	: NumTasks(0), NumTasksDone(0), bCanCancel(bInCanCancel), TextFormat(InTextFormat)
+	: NumTasks(0), NumTasksDone(0), NumTasksRemaining(0), bCanCancel(bInCanCancel), TextFormat(InTextFormat)
 {
 }
 
@@ -195,11 +195,12 @@ FPCGEditorProgressNotification::~FPCGEditorProgressNotification()
 
 void FPCGEditorProgressNotification::Update(uint32 InNumRemainingTasks)
 {
-	if (InNumRemainingTasks > NumTasks)
-	{
-		NumTasks = InNumRemainingTasks;
-	}
+	uint32 NumNewTasks = FMath::Max(0, (int32)InNumRemainingTasks - (int32)NumTasksRemaining);
+	
+	NumTasks += NumNewTasks;
+	
 	NumTasksDone = NumTasks - InNumRemainingTasks;
+	NumTasksRemaining = InNumRemainingTasks;
 
 	if (NumTasks > 0 && !NotificationItem)
 	{
diff --git a/Engine/Plugins/PCG/Source/PCGEditor/Private/PCGEditorProgressNotification.h b/Engine/Plugins/PCG/Source/PCGEditor/Private/PCGEditorProgressNotification.h
--- a/Engine/Plugins/PCG/Source/PCGEditor/Private/PCGEditorProgressNotification.h
+++ b/Engine/Plugins/PCG/Source/PCGEditor/Private/PCGEditorProgressNotification.h
@@ -29,6 +29,7 @@ public:
 private:
 	uint32 NumTasks = 0;
 	uint32 NumTasksDone = 0;
+	uint32 NumTasksRemaining = 0;
 	
 	bool bCanCancel = false;
 	bool bIsCompleted = false;
