[PCG] Fix CreateTargetActor HLODLayer assignment with HLOD Source Type = Template and Allow Template Actor Editing = false

#rnx

diff --git a/Engine/Plugins/PCG/Source/PCG/Private/Elements/PCGCreateTargetActor.cpp b/Engine/Plugins/PCG/Source/PCG/Private/Elements/PCGCreateTargetActor.cpp
--- a/Engine/Plugins/PCG/Source/PCG/Private/Elements/PCGCreateTargetActor.cpp
+++ b/Engine/Plugins/PCG/Source/PCG/Private/Elements/PCGCreateTargetActor.cpp
@@ -395,7 +395,8 @@ bool FPCGCreateTargetActorElement::ExecuteInternal(FPCGContext* Context) const
 	SpawnDefaultActorParams.DataLayerInstances = PCGDataLayerHelpers::GetDataLayerInstancesAndCrc(Context, Settings->DataLayerSettings, TargetActor, DataLayerCrc);
 
 	int32 HLODLayerCrc;
-	SpawnDefaultActorParams.HLODLayer = PCGHLODHelpers::GetHLODLayerAndCrc(Context, Settings->HLODSettings, TargetActor, TemplateActor, HLODLayerCrc);
+	AActor* TemplateActorOrDefault = TemplateActor ? TemplateActor : Cast<AActor>(Settings->TemplateActorClass->GetDefaultObject());
+	SpawnDefaultActorParams.HLODLayer = PCGHLODHelpers::GetHLODLayerAndCrc(Context, Settings->HLODSettings, TargetActor, TemplateActorOrDefault, HLODLayerCrc);
 #endif
 
 	UPCGManagedActors* ReusableResource = nullptr;
