[PCG] Prevent HLODLayer editing on PCG Partition Actors
- PCG Partition Actor names are based on the HLOD Layer and can't change

#rnx

diff --git a/Engine/Plugins/PCG/Source/PCG/Private/Grid/PCGPartitionActor.cpp b/Engine/Plugins/PCG/Source/PCG/Private/Grid/PCGPartitionActor.cpp
--- a/Engine/Plugins/PCG/Source/PCG/Private/Grid/PCGPartitionActor.cpp
+++ b/Engine/Plugins/PCG/Source/PCG/Private/Grid/PCGPartitionActor.cpp
@@ -623,6 +623,21 @@ AActor* APCGPartitionActor::GetSceneOutlinerParent() const
 		return Super::GetSceneOutlinerParent();
 	}	
 }
+
+bool APCGPartitionActor::CanEditChange(const FProperty* InProperty) const
+{
+	if (!Super::CanEditChange(InProperty))
+	{
+		return false;
+	}
+
+	if (InProperty && InProperty->GetFName() == AActor::GetHLODLayerPropertyName())
+	{
+		return false;
+	}
+
+	return true;
+}
 #endif // WITH_EDITOR
 
 void APCGPartitionActor::PostCreation(const FPCGGridDescriptor& GridDescriptor)
diff --git a/Engine/Plugins/PCG/Source/PCG/Public/Grid/PCGPartitionActor.h b/Engine/Plugins/PCG/Source/PCG/Public/Grid/PCGPartitionActor.h
--- a/Engine/Plugins/PCG/Source/PCG/Public/Grid/PCGPartitionActor.h
+++ b/Engine/Plugins/PCG/Source/PCG/Public/Grid/PCGPartitionActor.h
@@ -43,6 +43,7 @@ public:
 	virtual AActor* GetSceneOutlinerParent() const override;
 	virtual bool ShouldExport() override { return false; }
 	virtual bool ShouldImport(FStringView ActorPropString, bool IsMovingLevel) override { return false; }
+	virtual bool CanEditChange(const FProperty* InProperty) const override;
 #endif
 	//~End AActor Interface
 
