[PCG] Make sure FPCGContext handle is not copied by the BP PCGContext created on the stack (Fix ensure)

#rnx

diff --git a/Engine/Plugins/PCG/Source/PCG/Private/Elements/PCGExecuteBlueprint.cpp b/Engine/Plugins/PCG/Source/PCG/Private/Elements/PCGExecuteBlueprint.cpp
--- a/Engine/Plugins/PCG/Source/PCG/Private/Elements/PCGExecuteBlueprint.cpp
+++ b/Engine/Plugins/PCG/Source/PCG/Private/Elements/PCGExecuteBlueprint.cpp
@@ -807,6 +807,10 @@ bool FPCGExecuteBlueprintElement::ExecuteInternal(FPCGContext* InContext) const
 		/** Finally, execute the actual blueprint */
 		Context->BlueprintElementInstance->SetCurrentContext(Context);
 		
+		// Make sure handle is not copied when calling ExecuteWithContext which creates a FPCGContext copy on the stack
+		TSharedPtr<FPCGContextHandle> Handle = Context->Handle;
+		Context->Handle = nullptr;
+
 		// When running outside of main thread make sure GC can't run (BP Nodes might create objects which can't happen while GC runs)
 		if(!IsInGameThread())
 		{
@@ -819,6 +823,8 @@ bool FPCGExecuteBlueprintElement::ExecuteInternal(FPCGContext* InContext) const
 			Context->BlueprintElementInstance->ExecuteWithContext(*Context, Context->InputData, Context->OutputData);
 		}
 
+		// Restore handle
+		Context->Handle = Handle;
 		Context->BlueprintElementInstance->SetCurrentContext(nullptr);
 	}
 	else
diff --git a/Engine/Plugins/PCG/Source/PCG/Public/PCGContext.h b/Engine/Plugins/PCG/Source/PCG/Public/PCGContext.h
--- a/Engine/Plugins/PCG/Source/PCG/Public/PCGContext.h
+++ b/Engine/Plugins/PCG/Source/PCG/Public/PCGContext.h
@@ -71,6 +71,7 @@ struct PCG_API FPCGContext
 {
 	GENERATED_BODY()
 
+	friend class FPCGExecuteBlueprintElement;
 	friend class FPCGGraphExecutor;
 	friend struct FPCGGraphActiveTask;
 	friend struct FPCGGraphTask;
@@ -272,7 +273,6 @@ private:
 	// Lazy initialized shared handle pointer that can be used in lambda captures to test if Context is still valid before accessing it
 	TSharedPtr<FPCGContextHandle> Handle;
 
-	friend struct FPCGGraphActiveTask;
 #if WITH_EDITOR
 	friend class PCGUtils::FExtraCapture;
 	PCGUtils::FCallTime Timer;
