[PCG] Add AbortCleanup lambda when scheduling a CleanupTask to reset GeneratedResourcesInaccessible to false

#rnx

diff --git a/Engine/Plugins/PCG/Source/PCG/Private/PCGComponent.cpp b/Engine/Plugins/PCG/Source/PCG/Private/PCGComponent.cpp
--- a/Engine/Plugins/PCG/Source/PCG/Private/PCGComponent.cpp
+++ b/Engine/Plugins/PCG/Source/PCG/Private/PCGComponent.cpp
@@ -1553,6 +1553,20 @@ FPCGTaskId UPCGComponent::CreateCleanupTask(bool bRemoveComponents, const TArray
 	TWeakObjectPtr<UPCGComponent> ThisComponentWeakPtr(this);
 	TWeakObjectPtr<UWorld> WorldPtr(GetWorld());
 
+	auto AbortCleanup = [Context, ThisComponentWeakPtr]()
+	{
+		UPCGComponent* ThisComponent = ThisComponentWeakPtr.Get();
+		// If the component is not valid anymore, just early out
+		if (!IsValid(ThisComponent))
+		{
+			return true;
+		}
+
+		FScopeLock ResourcesLock(&ThisComponent->GeneratedResourcesLock);
+		ThisComponent->GeneratedResourcesInaccessible = false;
+		return true;
+	};
+
 	auto CleanupTask = [Context, ThisComponentWeakPtr, WorldPtr, bRemoveComponents]()
 	{
 		if (UPCGComponent* ThisComponent = ThisComponentWeakPtr.Get())
@@ -1693,7 +1707,7 @@ FPCGTaskId UPCGComponent::CreateCleanupTask(bool bRemoveComponents, const TArray
 		return true;
 	};
 
-	return GetSubsystem()->ScheduleGeneric(CleanupTask, this, *AllDependencies);
+	return GetSubsystem()->ScheduleGeneric(CleanupTask, AbortCleanup, this, *AllDependencies);
 }
 
 void UPCGComponent::CleanupUnusedManagedResources()
