[PCG] Fix "Open Source Editor" button disappearing from Custom HLSL node details panel after recent refactor.

#rnx

diff --git a/Engine/Plugins/PCG/Source/PCGEditor/Private/Details/PCGCustomHLSLSettingsDetails.cpp b/Engine/Plugins/PCG/Source/PCGEditor/Private/Details/PCGCustomHLSLSettingsDetails.cpp
--- a/Engine/Plugins/PCG/Source/PCGEditor/Private/Details/PCGCustomHLSLSettingsDetails.cpp
+++ b/Engine/Plugins/PCG/Source/PCGEditor/Private/Details/PCGCustomHLSLSettingsDetails.cpp
@@ -29,7 +29,7 @@ void FPCGCustomHLSLSettingsDetails::CustomizeDetails(IDetailLayoutBuilder& Detai
 		const TSharedPtr<const FPCGEditor> PCGEditorSharedPtr = PCGEditorWeakPtr.Pin();
 
 		// Add a "open source editor" button if the source editor tab is not already visible.
-		if (PCGEditorSharedPtr.IsValid() && PCGEditorSharedPtr->IsPanelCurrentlyOpen(EPCGEditorPanel::NodeSource))
+		if (PCGEditorSharedPtr.IsValid() && !PCGEditorSharedPtr->IsPanelCurrentlyOpen(EPCGEditorPanel::NodeSource))
 		{
 			SettingsCategory.AddCustomRow(FText::GetEmpty())
 				.Visibility(
