[PCG] Add FPCGContext::ScheduleGraph & FPCGContext::ScheduleGeneric to remove dependency on PCGSubsystem
- FPCGContext get assigned a FPCGGraphExecutor on creation and can forward the scheduling calls to it
- This is a first step in removing dependency on PCGComponent to execute graphs

#rnx

diff --git a/Engine/Plugins/PCG/Source/PCG/Private/Data/PCGLandscapeData.cpp b/Engine/Plugins/PCG/Source/PCG/Private/Data/PCGLandscapeData.cpp
--- a/Engine/Plugins/PCG/Source/PCG/Private/Data/PCGLandscapeData.cpp
+++ b/Engine/Plugins/PCG/Source/PCG/Private/Data/PCGLandscapeData.cpp
@@ -662,12 +662,6 @@ TArray<FPCGTaskId> UPCGLandscapeData::PrepareForSpatialQuery(FPCGContext* InCont
 		return {};
 	}
 
-	UPCGSubsystem* Subsystem = InContext->SourceComponent.IsValid() ? InContext->SourceComponent->GetSubsystem() : nullptr;
-	if (!Subsystem)
-	{
-		return {};
-	}
-
 	const ULandscapeInfo* LandscapeInfo = GetLandscapeInfo(InBounds.GetCenter());
 	ALandscapeProxy* LandscapeProxy = LandscapeInfo ? LandscapeInfo->GetLandscapeProxy() : nullptr;
 	if (!LandscapeProxy)
@@ -705,5 +699,5 @@ TArray<FPCGTaskId> UPCGLandscapeData::PrepareForSpatialQuery(FPCGContext* InCont
 	InContext->SourceComponent.Get());
 	Params.bCanExecuteOnlyOnMainThread = false;
 
-	return { Subsystem->ScheduleGeneric(Params) };
+	return { InContext->ScheduleGeneric(Params) };
 }
\ No newline at end of file
diff --git a/Engine/Plugins/PCG/Source/PCG/Private/Elements/PCGIndirectionElement.cpp b/Engine/Plugins/PCG/Source/PCG/Private/Elements/PCGIndirectionElement.cpp
--- a/Engine/Plugins/PCG/Source/PCG/Private/Elements/PCGIndirectionElement.cpp
+++ b/Engine/Plugins/PCG/Source/PCG/Private/Elements/PCGIndirectionElement.cpp
@@ -207,6 +207,27 @@ FPCGIndirectionContext::~FPCGIndirectionContext()
 	InnerContext = nullptr;
 }
 
+void FPCGIndirectionContext::SetupInnerContext()
+{
+	// Note: we need to pass a null node here + add the inner settings as part of the input so that they are retrieved properly
+	FPCGDataCollection InnerInput = InputData;
+	InnerInput.TaggedData.Emplace_GetRef().Data = InnerSettings;
+
+	// TODO: there are some types of settings that might require a bit more information to be able to do their processing correctly
+	// namely (dynamic) subgraphs, so YMMV.
+	InnerContext = InnerElement->Initialize(InnerInput, SourceComponent, nullptr);
+	check(InnerContext);
+	InitializeGraphExecutor(InnerContext);
+	InnerContext->InitializeSettings();
+
+	// Unclear whether we should give those new values
+	InnerContext->TaskId = TaskId;
+	InnerContext->CompiledTaskId = CompiledTaskId;
+	InnerContext->DependenciesCrc = DependenciesCrc;
+
+	InnerContext->AsyncState = AsyncState;
+}
+
 void FPCGIndirectionContext::AddExtraStructReferencedObjects(FReferenceCollector& Collector)
 {
 	if (InnerSettings)
@@ -262,22 +283,7 @@ bool FPCGIndirectionElement::PrepareDataInternal(FPCGContext* InContext) const
 			return false;
 		}
 
-		// Note: we need to pass a null node here + add the inner settings as part of the input so that they are retrieved properly
-		FPCGDataCollection InnerInput = Context->InputData;
-		InnerInput.TaggedData.Emplace_GetRef().Data = InnerSettings;
-
-		// TODO: there are some types of settings that might require a bit more information to be able to do their processing correctly
-		// namely (dynamic) subgraphs, so YMMV.
-		Context->InnerContext = Context->InnerElement->Initialize(InnerInput, Context->SourceComponent, nullptr);
-		check(Context->InnerContext);
-		Context->InnerContext->InitializeSettings();
-
-		// Unclear whether we should give those new values
-		Context->InnerContext->TaskId = Context->TaskId;
-		Context->InnerContext->CompiledTaskId = Context->CompiledTaskId;
-		Context->InnerContext->DependenciesCrc = Context->DependenciesCrc;
-
-		Context->InnerContext->AsyncState = Context->AsyncState;
+		Context->SetupInnerContext();
 
 #if WITH_EDITOR
 		// If we have an override, register for dynamic tracking.
diff --git a/Engine/Plugins/PCG/Source/PCG/Private/Elements/PCGLoopElement.cpp b/Engine/Plugins/PCG/Source/PCG/Private/Elements/PCGLoopElement.cpp
--- a/Engine/Plugins/PCG/Source/PCG/Private/Elements/PCGLoopElement.cpp
+++ b/Engine/Plugins/PCG/Source/PCG/Private/Elements/PCGLoopElement.cpp
@@ -6,7 +6,6 @@
 #include "PCGCustomVersion.h"
 #include "PCGModule.h"
 #include "PCGPin.h"
-#include "PCGSubsystem.h"
 #include "Graph/PCGStackContext.h"
 #include "Helpers/PCGDynamicTrackingHelpers.h"
 #include "Helpers/PCGHelpers.h"
@@ -206,7 +205,6 @@ bool FPCGLoopElement::ExecuteInternal(FPCGContext* InContext) const
 		}
 
 		UPCGGraph* Subgraph = Settings->GetSubgraph();
-		UPCGSubsystem* Subsystem = Context->SourceComponent.IsValid() ? Context->SourceComponent->GetSubsystem() : nullptr;
 
 		if (!Subgraph)
 		{
@@ -214,12 +212,6 @@ bool FPCGLoopElement::ExecuteInternal(FPCGContext* InContext) const
 			Context->OutputData.TaggedData = Context->InputData.GetInputsByPin(PCGPinConstants::DefaultInputLabel);
 			return true;
 		}
-		else if (!Subsystem)
-		{
-			// Job cannot run; cancel
-			Context->OutputData.bCancelExecution = true;
-			return true;
-		}
 
 		TArray<FPCGDataCollection> LoopDataCollection;
 		TArray<FName> FeedbackPinNames;
@@ -301,14 +293,14 @@ bool FPCGLoopElement::ExecuteInternal(FPCGContext* InContext) const
 			Context->AddToReferencedObjects(InputDataCollection);
 
 			FPCGElementPtr InputElement = MakeShared<FPCGLoopInputForwardingElement>(InputDataCollection, PreviousTaskId, FeedbackPinNames);
-			FPCGTaskId SubgraphTaskId = Subsystem->ScheduleGraph(
+			FPCGTaskId SubgraphTaskId = Context->ScheduleGraph(FPCGScheduleGraphParams(
 				Subgraph,
 				Context->SourceComponent.Get(),
 				PreGraphElement,
 				InputElement,
 				Dependencies,
 				&InvocationStack,
-				/*bAllowHierarchicalGeneration=*/false);
+				/*bAllowHierarchicalGeneration=*/false));
 
 			if (SubgraphTaskId != InvalidPCGTaskId)
 			{
@@ -341,8 +333,8 @@ bool FPCGLoopElement::ExecuteInternal(FPCGContext* InContext) const
 			}
 			else
 			{
-				Subsystem->ScheduleGeneric(
-					[ContextHandle = Context->GetOrCreateHandle()]() // Normal execution: Wake up the current task
+				Context->ScheduleGeneric(FPCGScheduleGenericParams(
+					[ContextHandle = Context->GetOrCreateHandle()](FPCGContext* NewContext) // Normal execution: Wake up the current task
 					{
 						FPCGContext::FSharedContext<FPCGSubgraphContext> SharedContext(ContextHandle);
 						if (FPCGSubgraphContext* ContextPtr = SharedContext.Get())
@@ -351,7 +343,7 @@ bool FPCGLoopElement::ExecuteInternal(FPCGContext* InContext) const
 						}
 						return true;
 					},
-					[ContextHandle = Context->GetOrCreateHandle()]() // On abort: wakeup and cancel, forget subgraphs
+					[ContextHandle = Context->GetOrCreateHandle()](FPCGContext* NewContext) // On abort: wakeup and cancel, forget subgraphs
 					{
 						FPCGContext::FSharedContext<FPCGSubgraphContext> SharedContext(ContextHandle);
 						if (FPCGSubgraphContext* ContextPtr = SharedContext.Get())
@@ -360,10 +352,9 @@ bool FPCGLoopElement::ExecuteInternal(FPCGContext* InContext) const
 							ContextPtr->SubgraphTaskIds.Reset();
 							ContextPtr->OutputData.bCancelExecution = true;
 						}
-						return true;
 					},
 					Context->SourceComponent.Get(),
-					Context->SubgraphTaskIds);
+					Context->SubgraphTaskIds));
 			}
 
 			return false;
@@ -383,40 +374,31 @@ bool FPCGLoopElement::ExecuteInternal(FPCGContext* InContext) const
 	{
 		// when woken up, get the output data from the subgraph
 		// and copy it to the current context output data, and finally return true
-		UPCGSubsystem* Subsystem = Context->SourceComponent.IsValid() ? Context->SourceComponent->GetSubsystem() : nullptr;
-		if (Subsystem)
+		// If this was running iterations, we need to build a list of pins we'll ignore for non-terminal tasks
+		TArray<FName> LoopPinNames;
+		TArray<FName> FeedbackPinNames;
+		Settings->GetLoopPinNames(Context, LoopPinNames, FeedbackPinNames, /*bQuiet=*/true);
+
+		for(int SubtaskIndex = 0; SubtaskIndex < Context->SubgraphTaskIds.Num(); ++SubtaskIndex)
 		{
-			// If this was running iterations, we need to build a list of pins we'll ignore for non-terminal tasks
-			TArray<FName> LoopPinNames;
-			TArray<FName> FeedbackPinNames;
-			Settings->GetLoopPinNames(Context, LoopPinNames, FeedbackPinNames, /*bQuiet=*/true);
+			FPCGTaskId SubgraphTaskId = Context->SubgraphTaskIds[SubtaskIndex];
+			const bool bIsLastTask = (SubtaskIndex == Context->SubgraphTaskIds.Num() - 1);
 
-			for(int SubtaskIndex = 0; SubtaskIndex < Context->SubgraphTaskIds.Num(); ++SubtaskIndex)
+			FPCGDataCollection SubgraphOutput;
+			// While this should be always return true, if a scheduled task was cancelled, this can still happen because the dependency has been removed but this code isn't aware of this
+			if (Context->GetOutputData(SubgraphTaskId, SubgraphOutput))
 			{
-				FPCGTaskId SubgraphTaskId = Context->SubgraphTaskIds[SubtaskIndex];
-				const bool bIsLastTask = (SubtaskIndex == Context->SubgraphTaskIds.Num() - 1);
-
-				FPCGDataCollection SubgraphOutput;
-				// While this should be always return true, if a scheduled task was cancelled, this can still happen because the dependency has been removed but this code isn't aware of this
-				if (Subsystem->GetOutputData(SubgraphTaskId, SubgraphOutput))
+				if (FeedbackPinNames.IsEmpty() || bIsLastTask)
 				{
-					if (FeedbackPinNames.IsEmpty() || bIsLastTask)
-					{
-						Context->OutputData.TaggedData.Append(SubgraphOutput.TaggedData);
-					}
-					else
-					{
-						Algo::CopyIf(SubgraphOutput.TaggedData, Context->OutputData.TaggedData, [&FeedbackPinNames](const FPCGTaggedData& InTaggedData) { return !FeedbackPinNames.Contains(InTaggedData.Pin); });
-					}
-					Subsystem->ClearOutputData(SubgraphTaskId);
+					Context->OutputData.TaggedData.Append(SubgraphOutput.TaggedData);
 				}
+				else
+				{
+					Algo::CopyIf(SubgraphOutput.TaggedData, Context->OutputData.TaggedData, [&FeedbackPinNames](const FPCGTaggedData& InTaggedData) { return !FeedbackPinNames.Contains(InTaggedData.Pin); });
+				}
+				Context->ClearOutputData(SubgraphTaskId);
 			}
 		}
-		else
-		{
-			// Job cannot run, cancel
-			Context->OutputData.bCancelExecution = true;
-		}
 
 		return true;
 	}
@@ -469,13 +451,10 @@ bool FPCGLoopInputForwardingElement::ExecuteInternal(FPCGContext* Context) const
 	if (PreviousIterationTaskId != InvalidPCGTaskId && !FeedbackPinNames.IsEmpty())
 	{
 		check(Context && Context->SourceComponent.IsValid());
-		if (UPCGSubsystem* Subsystem = Context->SourceComponent->GetSubsystem())
+		FPCGDataCollection PreviousTaskOutput;
+		if (Context->GetOutputData(PreviousIterationTaskId, PreviousTaskOutput))
 		{
-			FPCGDataCollection PreviousTaskOutput;
-			if (Subsystem->GetOutputData(PreviousIterationTaskId, PreviousTaskOutput))
-			{
-				Algo::CopyIf(PreviousTaskOutput.TaggedData, Context->OutputData.TaggedData, [this](const FPCGTaggedData& InTaggedData) { return FeedbackPinNames.Contains(InTaggedData.Pin); });
-			}
+			Algo::CopyIf(PreviousTaskOutput.TaggedData, Context->OutputData.TaggedData, [this](const FPCGTaggedData& InTaggedData) { return FeedbackPinNames.Contains(InTaggedData.Pin); });
 		}
 	}
 
diff --git a/Engine/Plugins/PCG/Source/PCG/Private/Graph/PCGGraphCompiler.h b/Engine/Plugins/PCG/Source/PCG/Private/Graph/PCGGraphCompiler.h
--- a/Engine/Plugins/PCG/Source/PCG/Private/Graph/PCGGraphCompiler.h
+++ b/Engine/Plugins/PCG/Source/PCG/Private/Graph/PCGGraphCompiler.h
@@ -10,8 +10,6 @@ class IPCGElement;
 class UPCGComputeGraph;
 class UPCGGraph;
 
-typedef TSharedPtr<IPCGElement, ESPMode::ThreadSafe> FPCGElementPtr;
-
 /** Holds compilation results from FPCGGraphCompiler. */
 struct FPCGGraphCompilerCache
 {
diff --git a/Engine/Plugins/PCG/Source/PCG/Private/Graph/PCGGraphExecutor.cpp b/Engine/Plugins/PCG/Source/PCG/Private/Graph/PCGGraphExecutor.cpp
--- a/Engine/Plugins/PCG/Source/PCG/Private/Graph/PCGGraphExecutor.cpp
+++ b/Engine/Plugins/PCG/Source/PCG/Private/Graph/PCGGraphExecutor.cpp
@@ -249,39 +249,44 @@ FPCGTaskId FPCGGraphExecutor::Schedule(
 	const FPCGStack* InFromStack,
 	bool bAllowHierarchicalGeneration)
 {
-	check(SourceComponent);
+	return ScheduleGraph(FPCGScheduleGraphParams(Graph, SourceComponent, PreGraphElement, InputElement, ExternalDependencies, InFromStack, bAllowHierarchicalGeneration));
+}
+
+FPCGTaskId FPCGGraphExecutor::ScheduleGraph(const FPCGScheduleGraphParams& InParams)
+{
+	check(InParams.SourceComponent);
 
-	PCGGraphExecutionLogging::LogGraphSchedule(SourceComponent, Graph);
+	PCGGraphExecutionLogging::LogGraphSchedule(InParams.SourceComponent, InParams.Graph);
 	
 	FPCGTaskId ScheduledId = InvalidPCGTaskId;
 
 	uint32 GenerationGridSize = PCGHiGenGrid::UninitializedGridSize();
-	if (bAllowHierarchicalGeneration && Graph->IsHierarchicalGenerationEnabled())
+	if (InParams.bAllowHierarchicalGeneration && InParams.Graph->IsHierarchicalGenerationEnabled())
 	{
-		if (SourceComponent->IsLocalComponent() || SourceComponent->IsPartitioned())
+		if (InParams.SourceComponent->IsLocalComponent() || InParams.SourceComponent->IsPartitioned())
 		{
-			GenerationGridSize = SourceComponent->GetGenerationGridSize();
+			GenerationGridSize = InParams.SourceComponent->GetGenerationGridSize();
 		}
 	}
 
 	// Get compiled tasks from compiler
 	TSharedPtr<FPCGStackContext> StackContextPtr = MakeShared<FPCGStackContext>();
-	TArray<FPCGGraphTask> CompiledTasks = GraphCompiler.GetCompiledTasks(Graph, GenerationGridSize, *StackContextPtr);
+	TArray<FPCGGraphTask> CompiledTasks = GraphCompiler.GetCompiledTasks(InParams.Graph, GenerationGridSize, *StackContextPtr);
 
 	// Create the final stack context by including the current stack frames
-	if (InFromStack)
+	if (InParams.FromStack)
 	{
-		StackContextPtr->PrependParentStack(InFromStack);
+		StackContextPtr->PrependParentStack(InParams.FromStack);
 	}
 	else
 	{
 		FPCGStack ComponentStack;
-		ComponentStack.PushFrame(SourceComponent);
+		ComponentStack.PushFrame(InParams.SourceComponent);
 		StackContextPtr->PrependParentStack(&ComponentStack);		
 	}
 
 #if WITH_EDITOR
-	if (UPCGSubsystem* Subsystem = SourceComponent->GetSubsystem())
+	if (UPCGSubsystem* Subsystem = InParams.SourceComponent->GetSubsystem())
 	{
 		Subsystem->OnScheduleGraph(*StackContextPtr);
 	}
@@ -290,30 +295,30 @@ FPCGTaskId FPCGGraphExecutor::Schedule(
 	// Assign this component to the tasks
 	for (FPCGGraphTask& Task : CompiledTasks)
 	{
-		Task.SourceComponent = SourceComponent;
+		Task.SourceComponent = InParams.SourceComponent;
 		Task.StackContext = StackContextPtr;
 	}
 
 	// Prepare scheduled task that will be promoted in the next Execute call.
 	if (CompiledTasks.Num() > 0)
 	{ 
-		check(CompiledTasks[0].Node == Graph->GetInputNode());
+		check(CompiledTasks[0].Node == InParams.Graph->GetInputNode());
 
 		// Setup fetch task on input node
-		CompiledTasks[0].Element = InputElement;
+		CompiledTasks[0].Element = InParams.InputElement;
 
 		PCGGraphExecutor::TScopeLock ScopeLock(ScheduleLock);
 
 		FPCGGraphScheduleTask& ScheduledTask = ScheduledTasks.Emplace_GetRef();
 		ScheduledTask.Tasks = MoveTemp(CompiledTasks);
-		ScheduledTask.SourceComponent = SourceComponent;
+		ScheduledTask.SourceComponent = InParams.SourceComponent;
 
 		// Offset task node ids
 		FPCGGraphCompiler::OffsetNodeIds(ScheduledTask.Tasks, NextTaskId, InvalidPCGTaskId);
 		NextTaskId += ScheduledTask.Tasks.Num();
 		ScheduledId = NextTaskId - 1; // This is true because the last task is from the output node or is the post-execute task
 
-		const FPCGTaskId GraphExecutionTaskId = InFromStack ? InFromStack->GetGraphExecutionTaskId() : ScheduledId;
+		const FPCGTaskId GraphExecutionTaskId = InParams.FromStack ? InParams.FromStack->GetGraphExecutionTaskId() : ScheduledId;
 		StackContextPtr->SetGraphExecutionTaskId(GraphExecutionTaskId);
 
 		// Push task (not data) dependencies on the pre-execute task
@@ -321,12 +326,12 @@ FPCGTaskId FPCGGraphExecutor::Schedule(
 		check(ScheduledTask.Tasks.Num() >= 2 && ScheduledTask.Tasks[ScheduledTask.Tasks.Num() - 2].Node == nullptr);
 		FPCGGraphTask& PreGraphTask = ScheduledTask.Tasks[ScheduledTask.Tasks.Num() - 2];
 
-		if(PreGraphElement.IsValid())
+		if(InParams.PreGraphElement.IsValid())
 		{
-			PreGraphTask.Element = PreGraphElement;
+			PreGraphTask.Element = InParams.PreGraphElement;
 		}
 
-		for (FPCGTaskId ExternalDependency : ExternalDependencies)
+		for (FPCGTaskId ExternalDependency : InParams.ExternalDependencies)
 		{
 			// For the pre-task, we don't consume any input
 			PreGraphTask.Inputs.Emplace(ExternalDependency, /*InUpstreamPin=*/FPCGGraphTaskInput::NoPin, /*InDownstreamPin=*/FPCGGraphTaskInput::NoPin, /*bInProvideData=*/false);
@@ -1330,8 +1335,9 @@ void FPCGGraphExecutor::PrepareForExecute(FPCGGraphTask& Task, FCachedResult*& O
 	
 	{
 		Task.Context = Task.Element->Initialize(Task.TaskInput, Task.SourceComponent, Task.Node);
-		
-		// Skip PostLoad only when in GeV2 and not on GameThread
+		Task.Context->GraphExecutor = this;
+
+		// Skip PostLoad only when not on GameThread
 		// The reason is that some nodes like SpawnActor / CreateTargetActor have TemplateActor subobjects that need PostLoad calls to function properly
 		// If a node supports being prepared outside of the GameThread it needs to make sure that it and its sub-object do not depend on PostLoad being called.
 		// This is true for most nodes as they only run some deprecation code in their PostLoad which isn't needed after duplication
@@ -3014,7 +3020,7 @@ namespace PCGGraphExecutor
 						Params.bCanExecuteOnlyOnMainThread = false;
 						Params.ExecutionDependencies = { ComponentWithData->GetGenerationTaskId() };
 
-						Subsystem->ScheduleGeneric(Params);
+						InContext->ScheduleGeneric(Params);
 					}
 				}
 
diff --git a/Engine/Plugins/PCG/Source/PCG/Private/Graph/PCGGraphExecutor.h b/Engine/Plugins/PCG/Source/PCG/Private/Graph/PCGGraphExecutor.h
--- a/Engine/Plugins/PCG/Source/PCG/Private/Graph/PCGGraphExecutor.h
+++ b/Engine/Plugins/PCG/Source/PCG/Private/Graph/PCGGraphExecutor.h
@@ -57,6 +57,8 @@ public:
 	FPCGTaskId Schedule(UPCGComponent* InComponent, const TArray<FPCGTaskId>& TaskDependency = TArray<FPCGTaskId>(), const FPCGStack* InFromStack = nullptr);
 
 	/** Schedules the execution of a given graph with specified inputs. This call is threadsafe */
+	FPCGTaskId ScheduleGraph(const FPCGScheduleGraphParams& InParams);
+
 	FPCGTaskId Schedule(
 		UPCGGraph* Graph,
 		UPCGComponent* InSourceComponent,
diff --git a/Engine/Plugins/PCG/Source/PCG/Private/PCGContext.cpp b/Engine/Plugins/PCG/Source/PCG/Private/PCGContext.cpp
--- a/Engine/Plugins/PCG/Source/PCG/Private/PCGContext.cpp
+++ b/Engine/Plugins/PCG/Source/PCG/Private/PCGContext.cpp
@@ -10,6 +10,7 @@
 #include "Data/PCGSpatialData.h"
 #include "Data/PCGPointData.h"
 #include "Data/PCGPointArrayData.h"
+#include "Graph/PCGGraphExecutor.h"
 #include "Helpers/PCGHelpers.h"
 #include "Metadata/PCGAttributePropertySelector.h"
 #include "Metadata/PCGMetadata.h"
@@ -519,4 +520,54 @@ void FPCGContext::AddStructReferencedObjects(FReferenceCollector& Collector)
 	AddExtraStructReferencedObjects(Collector);
 }
 
+FPCGTaskId FPCGContext::ScheduleGraph(const FPCGScheduleGraphParams& InParams)
+{
+	if (GraphExecutor && InParams.SourceComponent)
+	{
+		return GraphExecutor->ScheduleGraph(InParams);
+	}
+	else
+	{
+		return InvalidPCGTaskId;
+	}
+}
+
+FPCGTaskId FPCGContext::ScheduleGeneric(const FPCGScheduleGenericParams& InParams)
+{
+	if (GraphExecutor)
+	{
+		return GraphExecutor->ScheduleGeneric(InParams);
+	}
+	else
+	{
+		return InvalidPCGTaskId;
+	}
+}
+
+bool FPCGContext::GetOutputData(FPCGTaskId InTaskId, FPCGDataCollection& OutData)
+{
+	if (GraphExecutor)
+	{
+		return GraphExecutor->GetOutputData(InTaskId, OutData);
+	}
+	else
+	{
+		return false;
+	}
+}
+
+void FPCGContext::ClearOutputData(FPCGTaskId InTaskId)
+{
+	if (GraphExecutor)
+	{
+		GraphExecutor->ClearOutputData(InTaskId);
+	}
+}
+
+void FPCGContext::InitializeGraphExecutor(FPCGContext* InContext)
+{
+	check(InContext);
+	InContext->GraphExecutor = GraphExecutor;
+}
+
 #undef LOCTEXT_NAMESPACE
\ No newline at end of file
diff --git a/Engine/Plugins/PCG/Source/PCG/Private/PCGSubgraph.cpp b/Engine/Plugins/PCG/Source/PCG/Private/PCGSubgraph.cpp
--- a/Engine/Plugins/PCG/Source/PCG/Private/PCGSubgraph.cpp
+++ b/Engine/Plugins/PCG/Source/PCG/Private/PCGSubgraph.cpp
@@ -4,7 +4,6 @@
 
 #include "PCGComponent.h"
 #include "PCGPin.h"
-#include "PCGSubsystem.h"
 #include "Data/PCGUserParametersData.h"
 #include "Graph/PCGStackContext.h"
 #include "Helpers/PCGDynamicTrackingHelpers.h"
@@ -753,9 +752,7 @@ bool FPCGSubgraphElement::ExecuteInternal(FPCGContext* InContext) const
 
 		if (!Context->bScheduledSubgraph)
 		{
-			UPCGSubsystem* Subsystem = Context->SourceComponent.IsValid() ? Context->SourceComponent->GetSubsystem() : nullptr;
-
-			if (Subsystem && Subgraph)
+			if (Subgraph)
 			{
 				// Dispatch graph to execute with the given information we have
 				// using this node's task id as additional inputs
@@ -779,14 +776,14 @@ bool FPCGSubgraphElement::ExecuteInternal(FPCGContext* InContext) const
 				InvocationStack.GetStackFramesMutable().Emplace(INDEX_NONE); // not a loop index
 
 				// Higen is not allowed in dynamic subgraphs, entire subgraph is executed on the same grid as this subgraph node.
-				FPCGTaskId SubgraphTaskId = Subsystem->ScheduleGraph(
+				FPCGTaskId SubgraphTaskId = Context->ScheduleGraph(FPCGScheduleGraphParams(
 					Subgraph,
 					Context->SourceComponent.Get(),
 					MakeShared<FPCGInputForwardingElement>(PreSubgraphInputData),
 					MakeShared<FPCGInputForwardingElement>(SubgraphInputData),
 					/*Dependencies=*/{},
 					&InvocationStack,
-					/*bAllowHierarchicalGeneration=*/false);
+					/*bAllowHierarchicalGeneration=*/false));
 
 				if (SubgraphTaskId != InvalidPCGTaskId)
 				{
@@ -801,8 +798,8 @@ bool FPCGSubgraphElement::ExecuteInternal(FPCGContext* InContext) const
 					else
 					{
 						// add a trivial task after the output task that wakes up this task
-						Subsystem->ScheduleGeneric(
-							[ContextHandle = Context->GetOrCreateHandle()]() // Normal execution: Wake up the current task
+						Context->ScheduleGeneric(FPCGScheduleGenericParams(
+							[ContextHandle = Context->GetOrCreateHandle()](FPCGContext* NewContext) // Normal execution: Wake up the current task
 							{
 								FPCGContext::FSharedContext<FPCGSubgraphContext> SharedContext(ContextHandle);
 								if (FPCGSubgraphContext* ContextPtr = SharedContext.Get())
@@ -811,7 +808,7 @@ bool FPCGSubgraphElement::ExecuteInternal(FPCGContext* InContext) const
 								}
 								return true;
 							},
-							[ContextHandle = Context->GetOrCreateHandle()]() // On Abort: Wake up and cancel the execution
+							[ContextHandle = Context->GetOrCreateHandle()](FPCGContext* NewContext) // On Abort: Wake up and cancel the execution
 							{
 								FPCGContext::FSharedContext<FPCGSubgraphContext> SharedContext(ContextHandle);
 								if (FPCGSubgraphContext* ContextPtr = SharedContext.Get())
@@ -820,10 +817,9 @@ bool FPCGSubgraphElement::ExecuteInternal(FPCGContext* InContext) const
 									ContextPtr->OutputData.bCancelExecution = true;
 									ContextPtr->SubgraphTaskIds.Reset();
 								}
-								return true;
 							},
 							Context->SourceComponent.Get(),
-							Context->SubgraphTaskIds);
+							Context->SubgraphTaskIds));
 					}
 
 					return false;
@@ -857,26 +853,17 @@ bool FPCGSubgraphElement::ExecuteInternal(FPCGContext* InContext) const
 		{
 			// when woken up, get the output data from the subgraph
 			// and copy it to the current context output data, and finally return true
-			UPCGSubsystem* Subsystem = Context->SourceComponent.IsValid() ? Context->SourceComponent->GetSubsystem() : nullptr;
-			if (Subsystem)
+			if (Context->SubgraphTaskIds.Num() > 0)
 			{
-				if (Context->SubgraphTaskIds.Num() > 0)
+				// This element does not support multiple results/dispatches. Note that this would be easy to fix;
+				// as we just can call the GetOutputData on a fresh data collection and merge it afterwards in the output data,
+				// but this is not needed here so we will keep the full assignment & benefit from the crc as well.
+				ensure(Context->SubgraphTaskIds.Num() == 1);
+				if (Context->GetOutputData(Context->SubgraphTaskIds[0], Context->OutputData))
 				{
-					// This element does not support multiple results/dispatches. Note that this would be easy to fix;
-					// as we just can call the GetOutputData on a fresh data collection and merge it afterwards in the output data,
-					// but this is not needed here so we will keep the full assignment & benefit from the crc as well.
-					ensure(Context->SubgraphTaskIds.Num() == 1);
-					if (Subsystem->GetOutputData(Context->SubgraphTaskIds[0], Context->OutputData))
-					{
-						Subsystem->ClearOutputData(Context->SubgraphTaskIds[0]);
-					}
+					Context->ClearOutputData(Context->SubgraphTaskIds[0]);
 				}
 			}
-			else
-			{
-				// Job cannot run, cancel
-				Context->OutputData.bCancelExecution = true;
-			}
 
 			return true;
 		}
diff --git a/Engine/Plugins/PCG/Source/PCG/Private/PCGSubsystem.cpp b/Engine/Plugins/PCG/Source/PCG/Private/PCGSubsystem.cpp
--- a/Engine/Plugins/PCG/Source/PCG/Private/PCGSubsystem.cpp
+++ b/Engine/Plugins/PCG/Source/PCG/Private/PCGSubsystem.cpp
@@ -823,9 +823,14 @@ FPCGTaskId UPCGSubsystem::ScheduleGraph(
 	const FPCGStack* InFromStack, 
 	bool bAllowHierarchicalGeneration)
 {
-	if (SourceComponent)
+	return ScheduleGraph(FPCGScheduleGraphParams(Graph, SourceComponent, PreGraphElement, InputElement, Dependencies, InFromStack, bAllowHierarchicalGeneration));
+}
+
+FPCGTaskId UPCGSubsystem::ScheduleGraph(const FPCGScheduleGraphParams& InParams)
+{
+	if (InParams.SourceComponent)
 	{
-		return GraphExecutor->Schedule(Graph, SourceComponent, PreGraphElement, InputElement, Dependencies, InFromStack, bAllowHierarchicalGeneration);
+		return GraphExecutor->ScheduleGraph(InParams);
 	}
 	else
 	{
diff --git a/Engine/Plugins/PCG/Source/PCG/Public/Elements/PCGIndirectionElement.h b/Engine/Plugins/PCG/Source/PCG/Public/Elements/PCGIndirectionElement.h
--- a/Engine/Plugins/PCG/Source/PCG/Public/Elements/PCGIndirectionElement.h
+++ b/Engine/Plugins/PCG/Source/PCG/Public/Elements/PCGIndirectionElement.h
@@ -73,6 +73,7 @@ struct FPCGIndirectionContext : public FPCGContext
 public:
 	virtual ~FPCGIndirectionContext();
 
+	void SetupInnerContext();
 protected:
 	virtual void AddExtraStructReferencedObjects(FReferenceCollector& Collector) override;
 
diff --git a/Engine/Plugins/PCG/Source/PCG/Public/PCGCommon.h b/Engine/Plugins/PCG/Source/PCG/Public/PCGCommon.h
--- a/Engine/Plugins/PCG/Source/PCG/Public/PCGCommon.h
+++ b/Engine/Plugins/PCG/Source/PCG/Public/PCGCommon.h
@@ -11,8 +11,11 @@ static const FPCGTaskId InvalidPCGTaskId = (uint64)-1;
 
 using FPCGPinId = uint64;
 
+class IPCGElement;
 class UPCGComponent;
+class UPCGGraph;
 struct FPCGContext;
+struct FPCGStack;
 
 namespace PCGLog
 {
@@ -500,13 +503,87 @@ enum class EPCGGenerationStatus : uint8
 	Aborted
 };
 
+using FPCGElementPtr = TSharedPtr<IPCGElement, ESPMode::ThreadSafe>;
+
+struct FPCGScheduleGraphParams
+{
+	FPCGScheduleGraphParams(
+		UPCGGraph* InGraph,
+		UPCGComponent* InSourceComponent,
+		FPCGElementPtr InPreGraphElement,
+		FPCGElementPtr InInputElement,
+		TArray<FPCGTaskId> InExternalDependencies,
+		const FPCGStack* InFromStack,
+		bool bInAllowHierarchicalGeneration)
+		: Graph(InGraph)
+		, SourceComponent(InSourceComponent)
+		, PreGraphElement(InPreGraphElement)
+		, InputElement(InInputElement)
+		, ExternalDependencies(MoveTemp(InExternalDependencies))
+		, FromStack(InFromStack)
+		, bAllowHierarchicalGeneration(bInAllowHierarchicalGeneration)
+	{
+	}
+
+	// Graph to execute
+	UPCGGraph* Graph = nullptr;
+
+	// PCG component associated with this task. Can be null.
+	UPCGComponent* SourceComponent = nullptr;
+
+	// First task to run.
+	FPCGElementPtr PreGraphElement;
+
+	// Task to run as an input to the provided graph.
+	FPCGElementPtr InputElement;
+	
+	// PreGraph Task dependencies (will wait on those to finish before executing anything).
+	TArray<FPCGTaskId> ExternalDependencies;
+	
+	// When scheduling sub-graphs, this is the parent execution stack.
+	const FPCGStack* FromStack;
+	
+	// If graph is allowed to use hierarchical generation.
+	bool bAllowHierarchicalGeneration = true;
+};
+
 struct FPCGScheduleGenericParams
 {
-	FPCGScheduleGenericParams(TFunction<bool(FPCGContext*)>&& InOperation, UPCGComponent* InSourceComponent)
+	FPCGScheduleGenericParams(TFunction<bool(FPCGContext*)> InOperation, UPCGComponent* InSourceComponent)
 		: Operation(MoveTemp(InOperation)), SourceComponent(InSourceComponent)
 	{
 	}
 
+	FPCGScheduleGenericParams(
+		TFunction<bool(FPCGContext*)> InOperation, 
+		UPCGComponent* InSourceComponent, 
+		const TArray<FPCGTaskId>& InExecutionDependencies, 
+		const TArray<FPCGTaskId>& InDataDependencies, 
+		bool bInSupportBasePointDataInput)
+		: Operation(MoveTemp(InOperation))
+		, SourceComponent(InSourceComponent)
+		, ExecutionDependencies(InExecutionDependencies)
+		, DataDependencies(InDataDependencies)
+		, bSupportBasePointDataInput(bInSupportBasePointDataInput)
+	{
+	}
+		
+	FPCGScheduleGenericParams(
+		TFunction<bool(FPCGContext*)> InOperation, 
+		TFunction<void(FPCGContext*)> InAbortOperation, 
+		UPCGComponent* InSourceComponent, 
+		const TArray<FPCGTaskId>& InExecutionDependencies, 
+		const TArray<FPCGTaskId>& InDataDependencies = {}, 
+		bool bInSupportBasePointDataInput = false)
+		: Operation(MoveTemp(InOperation))
+		, SourceComponent(InSourceComponent)
+		, AbortOperation(MoveTemp(InAbortOperation))
+		, ExecutionDependencies(InExecutionDependencies)
+		, DataDependencies(InDataDependencies)
+		, bSupportBasePointDataInput(bInSupportBasePointDataInput)
+	{
+	}
+
 	// Callback that takes a Context as argument and returns true if the task is done, false otherwise
 	TFunction<bool(FPCGContext*)> Operation;
 
diff --git a/Engine/Plugins/PCG/Source/PCG/Public/PCGContext.h b/Engine/Plugins/PCG/Source/PCG/Public/PCGContext.h
--- a/Engine/Plugins/PCG/Source/PCG/Public/PCGContext.h
+++ b/Engine/Plugins/PCG/Source/PCG/Public/PCGContext.h
@@ -14,6 +14,7 @@
 
 class UPCGBasePointData;
 class UPCGComponent;
+class FPCGGraphExecutor;
 class UPCGGraphInterface;
 class UPCGSettingsInterface;
 class UPCGSpatialData;
@@ -194,6 +195,8 @@ struct PCG_API FPCGContext
 			return;
 		}
 		
+		InContext->GraphExecutor = nullptr;
+
 		if (!InContext->Handle)
 		{
 			delete InContext;
@@ -232,11 +235,19 @@ struct PCG_API FPCGContext
 		TSharedPtr<FPCGContextHandle> SharedHandle;
 	};
 
+	FPCGTaskId ScheduleGraph(const FPCGScheduleGraphParams& InParams);
+	FPCGTaskId ScheduleGeneric(const FPCGScheduleGenericParams& InParams);
+
+	bool GetOutputData(FPCGTaskId InTaskId, FPCGDataCollection& OutData);
+	void ClearOutputData(FPCGTaskId InTaskId);
+
 protected:
 	virtual UObject* GetExternalContainerForOverridableParam(const FPCGSettingsOverridableParam& InParam) { return nullptr; }
 	virtual void* GetUnsafeExternalContainerForOverridableParam(const FPCGSettingsOverridableParam& InParam) { return nullptr; }
 	virtual void AddExtraStructReferencedObjects(FReferenceCollector& Collector) {}
 
+	void InitializeGraphExecutor(FPCGContext* InContext);
+
 private:
 	template<class T, typename... Args>
 	T* NewObject_AnyThread_Impl(Args&&... InArgs)
@@ -273,6 +284,9 @@ private:
 	// Lazy initialized shared handle pointer that can be used in lambda captures to test if Context is still valid before accessing it
 	TSharedPtr<FPCGContextHandle> Handle;
 
+	// Graph Executor that scheduled the task
+	FPCGGraphExecutor* GraphExecutor = nullptr;
+
 #if WITH_EDITOR
 	friend class PCGUtils::FExtraCapture;
 	PCGUtils::FCallTime Timer;
diff --git a/Engine/Plugins/PCG/Source/PCG/Public/PCGSubsystem.h b/Engine/Plugins/PCG/Source/PCG/Public/PCGSubsystem.h
--- a/Engine/Plugins/PCG/Source/PCG/Public/PCGSubsystem.h
+++ b/Engine/Plugins/PCG/Source/PCG/Public/PCGSubsystem.h
@@ -129,6 +129,8 @@ public:
 	FPCGTaskId ScheduleCleanup(UPCGComponent* PCGComponent, bool bRemoveComponents, const TArray<FPCGTaskId>& Dependencies);
 
 	// Schedule graph (used internally for dynamic subgraph execution)
+	FPCGTaskId ScheduleGraph(const FPCGScheduleGraphParams& InParams);
+		
 	FPCGTaskId ScheduleGraph(
 		UPCGGraph* Graph,
 		UPCGComponent* SourceComponent,
