[PCG] Add long overdue robustness for shader API - return 0 from get functions if attribute does not exist.

#rnx

diff --git a/Engine/Plugins/PCG/Shaders/Private/Elements/PCGStaticMeshSpawner.usf b/Engine/Plugins/PCG/Shaders/Private/Elements/PCGStaticMeshSpawner.usf
--- a/Engine/Plugins/PCG/Shaders/Private/Elements/PCGStaticMeshSpawner.usf
+++ b/Engine/Plugins/PCG/Shaders/Private/Elements/PCGStaticMeshSpawner.usf
@@ -106,8 +106,11 @@ void PCGStaticMeshSpawnerCS(uint3 GroupId : SV_GroupID, uint GroupIndex : SV_Gro
 			for (uint FloatIndex = 0; FloatIndex < AttributeStride; ++FloatIndex)
 			{
 				const uint AttributeValueAddress = In_GetElementAddressInternal(In_DataIndex, In_ElementIndex, AttributeId);
-				// Assumes floating point data. TODO add type information and deal with different types in different ways?
-				InstanceData_WriteCustomFloatRaw(InstanceIndex, AttributeOffset + FloatIndex, In_LoadBufferInternal(AttributeValueAddress + 4 * FloatIndex));
+				if (AttributeValueAddress > 0)
+				{
+					// Assumes floating point data. TODO add type information and deal with different types in different ways?
+					InstanceData_WriteCustomFloatRaw(InstanceIndex, AttributeOffset + FloatIndex, In_LoadBufferInternal(AttributeValueAddress + 4 * FloatIndex));
+				}
 			}
 		}
 	}
diff --git a/Engine/Plugins/PCG/Shaders/Private/PCGDataCollectionDataInterface.ush b/Engine/Plugins/PCG/Shaders/Private/PCGDataCollectionDataInterface.ush
--- a/Engine/Plugins/PCG/Shaders/Private/PCGDataCollectionDataInterface.ush
+++ b/Engine/Plugins/PCG/Shaders/Private/PCGDataCollectionDataInterface.ush
@@ -184,60 +184,63 @@ void SetAsExecutedInternal_{DataInterfaceName}()
 bool GetBool_{DataInterfaceName}(uint DataIndex, uint ElementIndex, int AttributeId)
 {
 	const uint AttributeAddress = GetFirstElementAddressInternal_{DataInterfaceName}(DataIndex, AttributeId);
-	return {DataInterfaceName}_DataCollectionBuffer.Load(AttributeAddress + ElementIndex * 1 * 4) != 0;
+	return AttributeAddress > 0 && {DataInterfaceName}_DataCollectionBuffer.Load(AttributeAddress + ElementIndex * 1 * 4) != 0;
 }
 
 int GetInt_{DataInterfaceName}(uint DataIndex, uint ElementIndex, int AttributeId)
 {
 	const uint AttributeAddress = GetFirstElementAddressInternal_{DataInterfaceName}(DataIndex, AttributeId);
-	return asint({DataInterfaceName}_DataCollectionBuffer.Load(AttributeAddress + ElementIndex * 1 * 4));
+	return AttributeAddress > 0 ? asint({DataInterfaceName}_DataCollectionBuffer.Load(AttributeAddress + ElementIndex * 1 * 4)) : 0;
 }
 
 int GetUint_{DataInterfaceName}(uint DataIndex, uint ElementIndex, int AttributeId)
 {
 	const uint AttributeAddress = GetFirstElementAddressInternal_{DataInterfaceName}(DataIndex, AttributeId);
-	return {DataInterfaceName}_DataCollectionBuffer.Load(AttributeAddress + ElementIndex * 1 * 4);
+	return AttributeAddress > 0 ? {DataInterfaceName}_DataCollectionBuffer.Load(AttributeAddress + ElementIndex * 1 * 4) : 0;
 }
 
 float GetFloat_{DataInterfaceName}(uint DataIndex, uint ElementIndex, int AttributeId)
 {
 	const uint AttributeAddress = GetFirstElementAddressInternal_{DataInterfaceName}(DataIndex, AttributeId);
-	return asfloat({DataInterfaceName}_DataCollectionBuffer.Load(AttributeAddress + ElementIndex * 1 * 4));
+	return AttributeAddress > 0 ? asfloat({DataInterfaceName}_DataCollectionBuffer.Load(AttributeAddress + ElementIndex * 1 * 4)) : 0.0f;
 }
 
 float2 GetFloat2_{DataInterfaceName}(uint DataIndex, uint ElementIndex, int AttributeId)
 {
 	const uint AttributeAddress = GetFirstElementAddressInternal_{DataInterfaceName}(DataIndex, AttributeId);
-	return asfloat({DataInterfaceName}_DataCollectionBuffer.Load2(AttributeAddress + ElementIndex * 2 * 4));
+	return AttributeAddress > 0 ? asfloat({DataInterfaceName}_DataCollectionBuffer.Load2(AttributeAddress + ElementIndex * 2 * 4)) : (float2)0;
 }
 
 float3 GetFloat3_{DataInterfaceName}(uint DataIndex, uint ElementIndex, int AttributeId)
 {
 	const uint AttributeAddress = GetFirstElementAddressInternal_{DataInterfaceName}(DataIndex, AttributeId);
-	return asfloat({DataInterfaceName}_DataCollectionBuffer.Load3(AttributeAddress + ElementIndex * 3 * 4));
+	return AttributeAddress > 0 ? asfloat({DataInterfaceName}_DataCollectionBuffer.Load3(AttributeAddress + ElementIndex * 3 * 4)) : (float3)0;
 }
 
 float4 GetFloat4_{DataInterfaceName}(uint DataIndex, uint ElementIndex, int AttributeId)
 {
 	const uint AttributeAddress = GetFirstElementAddressInternal_{DataInterfaceName}(DataIndex, AttributeId);
-	return asfloat({DataInterfaceName}_DataCollectionBuffer.Load4(AttributeAddress + ElementIndex * 4 * 4));
+	return AttributeAddress > 0 ? asfloat({DataInterfaceName}_DataCollectionBuffer.Load4(AttributeAddress + ElementIndex * 4 * 4)) : (float4)0;
 }
 
 float3 GetRotator_{DataInterfaceName}(uint DataIndex, uint ElementIndex, int AttributeId)
 {
-	const uint AttributeAddress = GetFirstElementAddressInternal_{DataInterfaceName}(DataIndex, AttributeId);
-	return asfloat({DataInterfaceName}_DataCollectionBuffer.Load3(AttributeAddress + ElementIndex * 3 * 4));
+	return GetFloat3_{DataInterfaceName}(DataIndex, ElementIndex, AttributeId);
 }
 
 float4 GetQuat_{DataInterfaceName}(uint DataIndex, uint ElementIndex, int AttributeId)
 {
-	const uint AttributeAddress = GetFirstElementAddressInternal_{DataInterfaceName}(DataIndex, AttributeId);
-	return asfloat({DataInterfaceName}_DataCollectionBuffer.Load4(AttributeAddress + ElementIndex * 4 * 4));
+	return GetFloat4_{DataInterfaceName}(DataIndex, ElementIndex, AttributeId);
 }
 
 float4x4 GetTransform_{DataInterfaceName}(uint DataIndex, uint ElementIndex, int AttributeId)
 {
 	const uint AttributeAddress = GetFirstElementAddressInternal_{DataInterfaceName}(DataIndex, AttributeId);
+	if (AttributeAddress == 0)
+	{
+		return (float4x4)0;
+	}
+
 	const uint StartAddress = ElementIndex * 16 * 4 + AttributeAddress;
 
 	return float4x4(
@@ -271,65 +274,95 @@ uint GetElementCountMultiplier_{DataInterfaceName}()
 void SetBool_{DataInterfaceName}(uint DataIndex, uint ElementIndex, int AttributeId, bool Value)
 {
 	const uint AttributeAddress = GetFirstElementAddressInternal_{DataInterfaceName}(DataIndex, AttributeId);
-	{DataInterfaceName}_DataCollectionBuffer.Store(AttributeAddress + ElementIndex * 1 * 4, Value ? 1u : 0u);
+	if (AttributeAddress > 0)
+	{
+		{DataInterfaceName}_DataCollectionBuffer.Store(AttributeAddress + ElementIndex * 1 * 4, Value ? 1u : 0u);
+	}
 }
 
 void SetInt_{DataInterfaceName}(uint DataIndex, uint ElementIndex, int AttributeId, int Value)
 {
 	const uint AttributeAddress = GetFirstElementAddressInternal_{DataInterfaceName}(DataIndex, AttributeId);
-	{DataInterfaceName}_DataCollectionBuffer.Store(AttributeAddress + ElementIndex * 1 * 4, asuint(Value));
+	if (AttributeAddress > 0)
+	{
+		{DataInterfaceName}_DataCollectionBuffer.Store(AttributeAddress + ElementIndex * 1 * 4, asuint(Value));
+	}
 }
 
 void SetUint_{DataInterfaceName}(uint DataIndex, uint ElementIndex, int AttributeId, uint Value)
 {
 	const uint AttributeAddress = GetFirstElementAddressInternal_{DataInterfaceName}(DataIndex, AttributeId);
-	{DataInterfaceName}_DataCollectionBuffer.Store(AttributeAddress + ElementIndex * 1 * 4, Value);
+	if (AttributeAddress > 0)
+	{
+		{DataInterfaceName}_DataCollectionBuffer.Store(AttributeAddress + ElementIndex * 1 * 4, Value);
+	}
 }
 
 void SetFloat_{DataInterfaceName}(uint DataIndex, uint ElementIndex, int AttributeId, float Value)
 {
 	const uint AttributeAddress = GetFirstElementAddressInternal_{DataInterfaceName}(DataIndex, AttributeId);
-	{DataInterfaceName}_DataCollectionBuffer.Store(AttributeAddress + ElementIndex * 1 * 4, asuint(Value));
+	if (AttributeAddress > 0)
+	{
+		{DataInterfaceName}_DataCollectionBuffer.Store(AttributeAddress + ElementIndex * 1 * 4, asuint(Value));
+	}
 }
 
 void SetFloat2_{DataInterfaceName}(uint DataIndex, uint ElementIndex, int AttributeId, float2 Value)
 {
 	const uint AttributeAddress = GetFirstElementAddressInternal_{DataInterfaceName}(DataIndex, AttributeId);
-	{DataInterfaceName}_DataCollectionBuffer.Store2(AttributeAddress + ElementIndex * 2 * 4, asuint(Value));
+	if (AttributeAddress > 0)
+	{
+		{DataInterfaceName}_DataCollectionBuffer.Store2(AttributeAddress + ElementIndex * 2 * 4, asuint(Value));
+	}
 }
 
 void SetFloat3_{DataInterfaceName}(uint DataIndex, uint ElementIndex, int AttributeId, float3 Value)
 {
 	const uint AttributeAddress = GetFirstElementAddressInternal_{DataInterfaceName}(DataIndex, AttributeId);
-	{DataInterfaceName}_DataCollectionBuffer.Store3(AttributeAddress + ElementIndex * 3 * 4, asuint(Value));
+	if (AttributeAddress > 0)
+	{
+		{DataInterfaceName}_DataCollectionBuffer.Store3(AttributeAddress + ElementIndex * 3 * 4, asuint(Value));
+	}
 }
 
 void SetFloat4_{DataInterfaceName}(uint DataIndex, uint ElementIndex, int AttributeId, float4 Value)
 {
 	const uint AttributeAddress = GetFirstElementAddressInternal_{DataInterfaceName}(DataIndex, AttributeId);
-	{DataInterfaceName}_DataCollectionBuffer.Store4(AttributeAddress + ElementIndex * 4 * 4, asuint(Value));
+	if (AttributeAddress > 0)
+	{
+		{DataInterfaceName}_DataCollectionBuffer.Store4(AttributeAddress + ElementIndex * 4 * 4, asuint(Value));
+	}
 }
 
 void SetRotator_{DataInterfaceName}(uint DataIndex, uint ElementIndex, int AttributeId, float3 Value)
 {
 	const uint AttributeAddress = GetFirstElementAddressInternal_{DataInterfaceName}(DataIndex, AttributeId);
-	{DataInterfaceName}_DataCollectionBuffer.Store3(AttributeAddress + ElementIndex * 3 * 4, asuint(Value));
+	if (AttributeAddress > 0)
+	{
+		{DataInterfaceName}_DataCollectionBuffer.Store3(AttributeAddress + ElementIndex * 3 * 4, asuint(Value));
+	}
 }
 
 void SetQuat_{DataInterfaceName}(uint DataIndex, uint ElementIndex, int AttributeId, float4 Value)
 {
 	const uint AttributeAddress = GetFirstElementAddressInternal_{DataInterfaceName}(DataIndex, AttributeId);
-	{DataInterfaceName}_DataCollectionBuffer.Store4(AttributeAddress + ElementIndex * 4 * 4, asuint(Value));
+	if (AttributeAddress > 0)
+	{
+		{DataInterfaceName}_DataCollectionBuffer.Store4(AttributeAddress + ElementIndex * 4 * 4, asuint(Value));
+	}
 }
 
 void SetTransform_{DataInterfaceName}(uint DataIndex, uint ElementIndex, int AttributeId, float4x4 Value)
 {
 	const uint AttributeAddress = GetFirstElementAddressInternal_{DataInterfaceName}(DataIndex, AttributeId);
-	const uint StartAddress = ElementIndex * 16 * 4 + AttributeAddress;
-	{DataInterfaceName}_DataCollectionBuffer.Store4(StartAddress, asuint(Value[0]));
-	{DataInterfaceName}_DataCollectionBuffer.Store4(StartAddress + 4 * 4, asuint(Value[1]));
-	{DataInterfaceName}_DataCollectionBuffer.Store4(StartAddress + 8 * 4, asuint(Value[2]));
-	{DataInterfaceName}_DataCollectionBuffer.Store4(StartAddress + 12 * 4, asuint(Value[3]));
+	if (AttributeAddress > 0)
+	{
+		const uint StartAddress = ElementIndex * 16 * 4 + AttributeAddress;
+		{DataInterfaceName}_DataCollectionBuffer.Store4(StartAddress, asuint(Value[0]));
+		{DataInterfaceName}_DataCollectionBuffer.Store4(StartAddress + 4 * 4, asuint(Value[1]));
+		{DataInterfaceName}_DataCollectionBuffer.Store4(StartAddress + 8 * 4, asuint(Value[2]));
+		{DataInterfaceName}_DataCollectionBuffer.Store4(StartAddress + 12 * 4, asuint(Value[3]));
+	}
 }
 
 void SetStringKey_{DataInterfaceName}(uint DataIndex, uint ElementIndex, int AttributeId, int Value)
@@ -352,10 +385,17 @@ int AtomicAddInt_{DataInterfaceName}(uint DataIndex, uint ElementIndex, int Attr
 {
 	const uint AttributeAddress = GetFirstElementAddressInternal_{DataInterfaceName}(DataIndex, AttributeId);
 
-	uint OriginalValue;
-	{DataInterfaceName}_DataCollectionBuffer.InterlockedAdd(AttributeAddress + ElementIndex * 1 * 4, (uint)ValueToAdd, OriginalValue);
+	if (AttributeAddress > 0)
+	{
+		uint OriginalValue;
+		{DataInterfaceName}_DataCollectionBuffer.InterlockedAdd(AttributeAddress + ElementIndex * 1 * 4, (uint)ValueToAdd, OriginalValue);
 
-	return (int)OriginalValue;
+		return (int)OriginalValue;
+	}
+	else
+	{
+		return 0;
+	}
 }
 
 // #################### POINT ATTRIBUTE GETTERS ##########################
@@ -523,8 +563,10 @@ void InitializePoint_{DataInterfaceName}(uint InDataIndex, uint InElementIndex)
 			continue;
 		}
 
-		const uint ElementAddress = GetElementAddressInternal_{DataInterfaceName}(InDataIndex, InElementIndex, AttributeId);
-		
+		// We don't bother checking element address is nonzero as we already check Stride above.
+		const uint FirstElementAddress = GetFirstElementAddressInternal_{DataInterfaceName}(InDataIndex, AttributeId);
+		const uint ElementAddress = InElementIndex * Stride + FirstElementAddress;
+
 		for (uint Offset = 0; Offset < Stride; Offset += 4)
 		{
 			// TODO: In future could upload actual default values for attributes rather than 0-initializing.
