[PCG] Fix validation error from mismatching texture dimension in PCGTextureReadback

#rnx

diff --git a/Engine/Plugins/PCG/Shaders/Private/PCGTextureReadback.usf b/Engine/Plugins/PCG/Shaders/Private/PCGTextureReadback.usf
--- a/Engine/Plugins/PCG/Shaders/Private/PCGTextureReadback.usf
+++ b/Engine/Plugins/PCG/Shaders/Private/PCGTextureReadback.usf
@@ -3,6 +3,7 @@
 #include "/Engine/Private/Common.ush"
 #include "/Engine/Public/Platform.ush"
 
+Texture2D<float4> SourceTexture;
 Texture2DArray<float4> SourceTextureArray;
 SamplerState SourceSampler;
 float2 SourceDimensions;
@@ -16,5 +17,12 @@ void PCGTextureReadback_CS(uint3 DispatchThreadId : SV_DispatchThreadID, uint Gr
 	float2 ThreadID = DispatchThreadId.xy + 0.5;
 	float2 TexCoord = ThreadID / SourceDimensions;
 
-	OutputTexture[ThreadID] = Texture2DArraySample(SourceTextureArray, SourceSampler, float3(TexCoord.xy, SourceTextureIndex));
+	if (SourceTextureIndex == -1)
+	{
+		OutputTexture[ThreadID] = Texture2DSample(SourceTexture, SourceSampler, TexCoord.xy);
+	}
+	else
+	{
+		OutputTexture[ThreadID] = Texture2DArraySample(SourceTextureArray, SourceSampler, float3(TexCoord.xy, SourceTextureIndex));
+	}
 }
diff --git a/Engine/Plugins/PCG/Source/PCGCompute/Private/PCGTextureReadback.cpp b/Engine/Plugins/PCG/Source/PCGCompute/Private/PCGTextureReadback.cpp
--- a/Engine/Plugins/PCG/Source/PCGCompute/Private/PCGTextureReadback.cpp
+++ b/Engine/Plugins/PCG/Source/PCGCompute/Private/PCGTextureReadback.cpp
@@ -17,10 +17,11 @@ public:
 	SHADER_USE_PARAMETER_STRUCT(FPCGTextureReadbackCS, FGlobalShader);
  
 	BEGIN_SHADER_PARAMETER_STRUCT(FParameters, )
+		SHADER_PARAMETER_TEXTURE(Texture2D<float4>, SourceTexture)
 		SHADER_PARAMETER_TEXTURE(Texture2DArray<float4>, SourceTextureArray)
 		SHADER_PARAMETER_SAMPLER(SamplerState, SourceSampler)
 		SHADER_PARAMETER(FVector2f, SourceDimensions)
-		SHADER_PARAMETER(uint32, SourceTextureIndex)
+		SHADER_PARAMETER(int32, SourceTextureIndex)
 		SHADER_PARAMETER_UAV(RWTexture2D<float4>, OutputTexture)
 	END_SHADER_PARAMETER_STRUCT()
  
@@ -41,11 +42,50 @@ void FPCGTextureReadbackInterface::Dispatch_RenderThread(FRHICommandListImmediat
 {
 	check(Params.SourceTexture && Params.SourceSampler);
 
+	const bool bIsTextureArray = Params.SourceTexture->GetDesc().Dimension == ETextureDimension::Texture2DArray;
+
 	FPCGTextureReadbackCS::FParameters PassParameters;
-	PassParameters.SourceTextureArray = Params.SourceTexture;
 	PassParameters.SourceSampler = Params.SourceSampler;
 	PassParameters.SourceDimensions = { (float)Params.SourceDimensions.X, (float)Params.SourceDimensions.Y };
-	PassParameters.SourceTextureIndex = Params.SourceTextureIndex;
+
+	if (bIsTextureArray)
+	{
+		const FRHITextureCreateDesc DummyTextureDesc =
+			FRHITextureCreateDesc::Create2D(TEXT("PCGDummyTexture"), 1, 1, PF_G8)
+			.SetFlags(ETextureCreateFlags::ShaderResource);
+
+		FTextureRHIRef DummyTexture = RHICreateTexture(DummyTextureDesc);
+
+		{
+			uint32 DestStride;
+			uint8* DestBuffer = (uint8*)RHICmdList.LockTexture2D(DummyTexture, 0, RLM_WriteOnly, DestStride, false);
+			*DestBuffer = 0;
+			RHICmdList.UnlockTexture2D(DummyTexture, 0, false);
+		}
+
+		PassParameters.SourceTexture = MoveTemp(DummyTexture);
+		PassParameters.SourceTextureArray = Params.SourceTexture;
+		PassParameters.SourceTextureIndex = Params.SourceTextureIndex;
+	}
+	else
+	{
+		const FRHITextureCreateDesc DummyTextureDesc =
+			FRHITextureCreateDesc::Create2DArray(TEXT("PCGDummyTextureArray"), 1, 1, 1, PF_G8)
+			.SetFlags(ETextureCreateFlags::ShaderResource);
+
+		FTextureRHIRef DummyTexture = RHICreateTexture(DummyTextureDesc);
+
+		{
+			uint32 DestStride;
+			uint8* DestBuffer = (uint8*)RHICmdList.LockTexture2DArray(DummyTexture, 0, 0, RLM_WriteOnly, DestStride, false);
+			*DestBuffer = 0;
+			RHICmdList.UnlockTexture2DArray(DummyTexture, 0, 0, false);
+		}
+
+		PassParameters.SourceTexture = Params.SourceTexture;
+		PassParameters.SourceTextureArray = MoveTemp(DummyTexture);
+		PassParameters.SourceTextureIndex = -1;
+	}
 
 	FRHITextureCreateDesc TargetTextureDesc =
 		FRHITextureCreateDesc::Create2D(TEXT("PCGTexture Readback Compute Target"), Params.SourceDimensions.X, Params.SourceDimensions.Y, EPixelFormat::PF_B8G8R8A8)
