[PCG] Fix validation error when using byte address buffer instead of structured buffer

#rnx

diff --git a/Engine/Plugins/PCG/Source/PCG/Private/Compute/DataInterfaces/PCGDataLabelResolverDataInterface.cpp b/Engine/Plugins/PCG/Source/PCG/Private/Compute/DataInterfaces/PCGDataLabelResolverDataInterface.cpp
--- a/Engine/Plugins/PCG/Source/PCG/Private/Compute/DataInterfaces/PCGDataLabelResolverDataInterface.cpp
+++ b/Engine/Plugins/PCG/Source/PCG/Private/Compute/DataInterfaces/PCGDataLabelResolverDataInterface.cpp
@@ -12,6 +12,7 @@
 #include "RenderGraphResources.h"
 #include "ShaderCompilerCore.h"
 #include "ShaderParameterMetadataBuilder.h"
+#include "SystemTextures.h"
 #include "ComputeFramework/ComputeMetadataBuilder.h"
 #include "ComputeFramework/ShaderParamTypeDefinition.h"
 
@@ -173,22 +174,18 @@ void FPCGDataLabelResolverDataProviderProxy::GatherDispatchData(FDispatchData co
 
 void FPCGDataLabelResolverDataProviderProxy::AllocateResources(FRDGBuilder& GraphBuilder, FAllocationData const& InAllocationData)
 {
-	const uint32 SizeBytes = FMath::Max(DataIdToDataIndexMap.Num(), 1) * DataIdToDataIndexMap.GetTypeSize();
-	const FRDGBufferDesc Desc = FRDGBufferDesc::CreateByteAddressDesc(SizeBytes);
-
-	if (DataIdToDataIndexMap.IsEmpty())
+	if (!DataIdToDataIndexMap.IsEmpty())
 	{
-		FRDGBufferRef NullBuffer = GraphBuilder.CreateBuffer(Desc, TEXT("PCGDataLabelResolver_DataIdToDataIndexMap_Empty"));
-		GraphBuilder.QueueBufferUpload(NullBuffer, 0, sizeof(int32));
+		const FRDGBufferDesc BufferDesc = FRDGBufferDesc::CreateStructuredDesc(DataIdToDataIndexMap.GetTypeSize(), DataIdToDataIndexMap.Num());
+
+		FRDGBufferRef Buffer = GraphBuilder.CreateBuffer(BufferDesc, TEXT("PCGDataLabelResolver_DataIdToDataIndexMap"));
+		DataIdToDataIndexBufferSRV = GraphBuilder.CreateSRV(Buffer);
 
-		DataIdToDataIndexBufferSRV = GraphBuilder.CreateSRV(NullBuffer);
+		GraphBuilder.QueueBufferUpload(Buffer, MakeArrayView(DataIdToDataIndexMap));
 	}
 	else
 	{
-		FRDGBufferRef Buffer = GraphBuilder.CreateBuffer(Desc, TEXT("PCGDataLabelResolver_DataIdToDataIndexMap"));
-		GraphBuilder.QueueBufferUpload(Buffer, DataIdToDataIndexMap.GetData(), SizeBytes);
-
-		DataIdToDataIndexBufferSRV = GraphBuilder.CreateSRV(Buffer);
+		DataIdToDataIndexBufferSRV = GraphBuilder.CreateSRV(FRDGBufferSRVDesc(GSystemTextures.GetDefaultStructuredBuffer(GraphBuilder, DataIdToDataIndexMap.GetTypeSize())));
 	}
 }
 
