[PCG] Support texture arrays as a single binding in PCGTextureDataInterface
* We can now have unlimited texture datas in the texture data interface, as long as there are only 8 unique texture bindings. This is useful because a TextureArray is one binding, but potentially many texture datas.
* Add shader permutation that uses only one binding

diff --git a/Engine/Plugins/PCG/Shaders/Private/PCGTextureDataInterface.ush b/Engine/Plugins/PCG/Shaders/Private/PCGTextureDataInterface.ush
--- a/Engine/Plugins/PCG/Shaders/Private/PCGTextureDataInterface.ush
+++ b/Engine/Plugins/PCG/Shaders/Private/PCGTextureDataInterface.ush
@@ -1,9 +1,5 @@
 // Copyright Epic Games, Inc. All Rights Reserved.
 
-// We can't use full bindless on all platforms, so fallback to emulating. Supports maximum 8 texture data objects.
-#define MAX_NUM_TEXTURES 8
-
-// TODO: Define out multiple bindings if we know statically there is only one.
 Texture2D {DataInterfaceName}_Texture_0;
 Texture2D {DataInterfaceName}_Texture_1;
 Texture2D {DataInterfaceName}_Texture_2;
@@ -13,6 +9,15 @@ Texture2D {DataInterfaceName}_Texture_5;
 Texture2D {DataInterfaceName}_Texture_6;
 Texture2D {DataInterfaceName}_Texture_7;
 
+Texture2DArray {DataInterfaceName}_TextureArray_0;
+Texture2DArray {DataInterfaceName}_TextureArray_1;
+Texture2DArray {DataInterfaceName}_TextureArray_2;
+Texture2DArray {DataInterfaceName}_TextureArray_3;
+Texture2DArray {DataInterfaceName}_TextureArray_4;
+Texture2DArray {DataInterfaceName}_TextureArray_5;
+Texture2DArray {DataInterfaceName}_TextureArray_6;
+Texture2DArray {DataInterfaceName}_TextureArray_7;
+
 // TODO: We only need 2 samplers + an indirection array, de-duplicate these.
 SamplerState {DataInterfaceName}_Sampler_0;
 SamplerState {DataInterfaceName}_Sampler_1;
@@ -23,7 +28,27 @@ SamplerState {DataInterfaceName}_Sampler_5;
 SamplerState {DataInterfaceName}_Sampler_6;
 SamplerState {DataInterfaceName}_Sampler_7;
 
-float4 {DataInterfaceName}_TextureBounds[MAX_NUM_TEXTURES];
+float4 {DataInterfaceName}_TextureBounds[MAX_NUM_BINDINGS];
+
+uint {DataInterfaceName}_NumTextureInfos;
+StructuredBuffer<uint3> {DataInterfaceName}_TextureInfos;
+
+void GetTextureInfoInternal_{DataInterfaceName}(int InDataIndex, out uint OutBindingIndex, out uint OutSliceIndex, out uint OutIsTextureArray)
+{
+	if (InDataIndex >= 0 && InDataIndex < {DataInterfaceName}_NumTextureInfos)
+	{
+		const uint3 TextureInfo = {DataInterfaceName}_TextureInfos[InDataIndex];
+		OutBindingIndex = TextureInfo.x;
+		OutSliceIndex = TextureInfo.y;
+		OutIsTextureArray = TextureInfo.z;
+	}
+	else
+	{
+		OutBindingIndex = 0;
+		OutSliceIndex = 0;
+		OutIsTextureArray = 0;
+	}
+}
 
 float2 GetTexCoords_{DataInterfaceName}(float2 WorldPos, float2 Min, float2 Max)
 {
@@ -40,13 +65,26 @@ float2 GetTexCoords_{DataInterfaceName}(float2 WorldPos, float2 Min, float2 Max)
 }
 
 #define PCG_CONDITIONAL_TEXTURE_SAMPLE(Index) \
-	if (InDataIndex == Index) \
+	if (BindingIndex == Index) \
 	{ \
-		return {DataInterfaceName}_Texture_##Index.SampleLevel({DataInterfaceName}_Sampler_##Index, InTextureUV, 0); \
+		if (IsTextureArray == 0) \
+		{ \
+			return {DataInterfaceName}_Texture_##Index.SampleLevel({DataInterfaceName}_Sampler_##Index, InTextureUV, 0); \
+		} \
+		else \
+		{ \
+			return {DataInterfaceName}_TextureArray_##Index.SampleLevel({DataInterfaceName}_Sampler_##Index, float3(InTextureUV, SliceIndex), 0); \
+		} \
 	}
 
 float4 Sample_{DataInterfaceName}(int InDataIndex, float2 InTextureUV)
 {
+	uint BindingIndex = 0;
+	uint SliceIndex = 0;
+	uint IsTextureArray = 0;
+	GetTextureInfoInternal_{DataInterfaceName}(InDataIndex, BindingIndex, SliceIndex, IsTextureArray);
+
+#if ENABLE_MULTIPLE_TEXTURE_OBJECTS
 	PCG_CONDITIONAL_TEXTURE_SAMPLE(0)
 	PCG_CONDITIONAL_TEXTURE_SAMPLE(1)
 	PCG_CONDITIONAL_TEXTURE_SAMPLE(2)
@@ -58,19 +96,42 @@ float4 Sample_{DataInterfaceName}(int InDataIndex, float2 InTextureUV)
 
 	// TODO: Runtime diagnostic error.
 	return (float4)0;
+#else
+	if (IsTextureArray == 0)
+	{
+		return {DataInterfaceName}_Texture_0.SampleLevel({DataInterfaceName}_Sampler_0, InTextureUV, 0);
+	}
+	else
+	{
+		return {DataInterfaceName}_TextureArray_0.SampleLevel({DataInterfaceName}_Sampler_0, float3(InTextureUV, SliceIndex), 0);
+	}
+#endif
 }
 
 #undef PCG_CONDITIONAL_TEXTURE_SAMPLE
 
 #define PCG_CONDITIONAL_TEXTURE_SAMPLE_WORLD_POS(Index) \
-	if (InDataIndex == Index) \
+	if (BindingIndex == Index) \
 	{ \
 		const float2 TextureUV = GetTexCoords_{DataInterfaceName}(WorldPos, {DataInterfaceName}_TextureBounds[Index].xy, {DataInterfaceName}_TextureBounds[Index].zw); \
-		return {DataInterfaceName}_Texture_##Index.SampleLevel({DataInterfaceName}_Sampler_##Index, TextureUV, 0); \
+		if (IsTextureArray == 0) \
+		{ \
+			return {DataInterfaceName}_Texture_##Index.SampleLevel({DataInterfaceName}_Sampler_##Index, TextureUV, 0); \
+		} \
+		else \
+		{ \
+			return {DataInterfaceName}_TextureArray_##Index.SampleLevel({DataInterfaceName}_Sampler_##Index, float3(TextureUV, SliceIndex), 0); \
+		} \
 	}
 
 float4 SampleWorldPos_{DataInterfaceName}(int InDataIndex, float2 WorldPos)
 {
+	uint BindingIndex = 0;
+	uint SliceIndex = 0;
+	uint IsTextureArray = 0;
+	GetTextureInfoInternal_{DataInterfaceName}(InDataIndex, BindingIndex, SliceIndex, IsTextureArray);
+
+#if ENABLE_MULTIPLE_TEXTURE_OBJECTS
 	PCG_CONDITIONAL_TEXTURE_SAMPLE_WORLD_POS(0)
 	PCG_CONDITIONAL_TEXTURE_SAMPLE_WORLD_POS(1)
 	PCG_CONDITIONAL_TEXTURE_SAMPLE_WORLD_POS(2)
@@ -82,6 +143,18 @@ float4 SampleWorldPos_{DataInterfaceName}(int InDataIndex, float2 WorldPos)
 
 	// TODO: Runtime diagnostic error.
 	return (float4)0;
+#else
+	const float2 TextureUV = GetTexCoords_{DataInterfaceName}(WorldPos, {DataInterfaceName}_TextureBounds[0].xy, {DataInterfaceName}_TextureBounds[0].zw);
+
+	if (IsTextureArray == 0)
+	{
+		return {DataInterfaceName}_Texture_0.SampleLevel({DataInterfaceName}_Sampler_0, TextureUV, 0);
+	}
+	else
+	{
+		return {DataInterfaceName}_TextureArray_0.SampleLevel({DataInterfaceName}_Sampler_0, float3(TextureUV, SliceIndex), 0);
+	}
+#endif
 }
 
 #undef PCG_CONDITIONAL_TEXTURE_SAMPLE_WORLD_POS
diff --git a/Engine/Plugins/PCG/Source/PCG/Private/Compute/DataInterfaces/PCGTextureDataInterface.cpp b/Engine/Plugins/PCG/Source/PCG/Private/Compute/DataInterfaces/PCGTextureDataInterface.cpp
--- a/Engine/Plugins/PCG/Source/PCG/Private/Compute/DataInterfaces/PCGTextureDataInterface.cpp
+++ b/Engine/Plugins/PCG/Source/PCG/Private/Compute/DataInterfaces/PCGTextureDataInterface.cpp
@@ -15,6 +15,8 @@
 #include "ShaderParameterMetadataBuilder.h"
 #include "SystemTextures.h"
 #include "TextureResource.h"
+#include "ComputeFramework/ComputeKernelPermutationSet.h"
+#include "ComputeFramework/ComputeKernelPermutationVector.h"
 #include "ComputeFramework/ShaderParamTypeDefinition.h"
 #include "Engine/Texture.h"
 
@@ -22,6 +24,41 @@
 
 #define LOCTEXT_NAMESPACE "PCGTextureDataInterface"
 
+namespace PCGTextureDataInterfaceConstants
+{
+	const TCHAR* EnableMultipleTextureObjectsPermutationName = TEXT("ENABLE_MULTIPLE_TEXTURE_OBJECTS");
+}
+
+FPCGTextureBindingInfo::FPCGTextureBindingInfo(const UPCGBaseTextureData* InTextureData)
+{
+	check(InTextureData);
+
+	ResourceType = InTextureData->GetTextureResourceType();
+	Texture = InTextureData->GetTextureRHI();
+	Filter = InTextureData->Filter;
+	TextureBounds = InTextureData->GetBounds();
+
+	if (ResourceType == EPCGTextureResourceType::ExportedTexture)
+	{
+		ExportedTexture = InTextureData->GetRefCountedTexture();
+	}
+
+	if (Texture)
+	{
+		bIsTextureArray = Texture->GetDesc().Dimension == ETextureDimension::Texture2DArray;
+	}
+}
+
+bool FPCGTextureBindingInfo::operator==(const FPCGTextureBindingInfo& Other) const
+{
+	return ResourceType == Other.ResourceType
+		&& Texture == Other.Texture
+		&& ExportedTexture == Other.ExportedTexture
+		&& bIsTextureArray == Other.bIsTextureArray
+		&& Filter == Other.Filter
+		&& TextureBounds == Other.TextureBounds;
+}
+
 void UPCGTextureDataInterface::GetSupportedInputs(TArray<FShaderFunctionDefinition>& OutFunctions) const
 {
 	// TODO: Support texture3d, etc. UNiagaraDataInterfaceRenderTarget2D and related are useful references.
@@ -57,9 +94,12 @@ void UPCGTextureDataInterface::GetSupportedInputs(TArray<FShaderFunctionDefiniti
 }
 
 BEGIN_SHADER_PARAMETER_STRUCT(FPCGTextureDataInterfaceParameters,)
-	SHADER_PARAMETER_RDG_TEXTURE_SRV_ARRAY(Texture2D, Texture, [PCGTextureDataInterfaceConstants::MAX_NUM_TEXTURES]) // TODO: Pick a max value and document somewhere + do validation
-	SHADER_PARAMETER_SAMPLER_ARRAY(SamplerState, Sampler, [PCGTextureDataInterfaceConstants::MAX_NUM_TEXTURES])
-	SHADER_PARAMETER_ARRAY(FVector4f, TextureBounds, [PCGTextureDataInterfaceConstants::MAX_NUM_TEXTURES])
+	SHADER_PARAMETER_RDG_TEXTURE_SRV_ARRAY(Texture2D, Texture, [PCGTextureDataInterfaceConstants::MAX_NUM_BINDINGS])
+	SHADER_PARAMETER_RDG_TEXTURE_SRV_ARRAY(Texture2DArray, TextureArray, [PCGTextureDataInterfaceConstants::MAX_NUM_BINDINGS])
+	SHADER_PARAMETER_SAMPLER_ARRAY(SamplerState, Sampler, [PCGTextureDataInterfaceConstants::MAX_NUM_BINDINGS])
+	SHADER_PARAMETER_ARRAY(FVector4f, TextureBounds, [PCGTextureDataInterfaceConstants::MAX_NUM_BINDINGS])
+	SHADER_PARAMETER(uint32, NumTextureInfos)
+	SHADER_PARAMETER_RDG_BUFFER_SRV(StructuredBuffer<FUintVector3>, TextureInfos)
 END_SHADER_PARAMETER_STRUCT()
 
 void UPCGTextureDataInterface::GetShaderParameters(TCHAR const* UID, FShaderParametersMetadataBuilder& InOutBuilder, FShaderParametersMetadataAllocations& InOutAllocations) const
@@ -93,6 +133,16 @@ void UPCGTextureDataInterface::GetHLSL(FString& OutHLSL, FString const& InDataIn
 	}
 }
 
+void UPCGTextureDataInterface::GetDefines(FComputeKernelDefinitionSet& OutDefinitionSet) const
+{
+	OutDefinitionSet.Defines.Add(FComputeKernelDefinition(TEXT("MAX_NUM_BINDINGS"), FString::FromInt(PCGTextureDataInterfaceConstants::MAX_NUM_BINDINGS)));
+}
+
+void UPCGTextureDataInterface::GetPermutations(FComputeKernelPermutationVector& OutPermutationVector) const
+{
+	OutPermutationVector.AddPermutation(PCGTextureDataInterfaceConstants::EnableMultipleTextureObjectsPermutationName, /*NumValues=*/2);
+}
+
 UComputeDataProvider* UPCGTextureDataInterface::CreateDataProvider() const
 {
 	return NewObject<UPCGTextureDataProvider>();
@@ -111,55 +161,49 @@ void UPCGTextureDataProvider::Initialize(const UComputeDataInterface* InDataInte
 	check(!DataInterface->DownstreamInputPinLabelAliases.IsEmpty());
 	const TArray<FPCGTaggedData> InputTaggedData = Binding->InputDataCollection.GetInputsByPin(DataInterface->DownstreamInputPinLabelAliases[0]);
 
-	if (InputTaggedData.Num() >= PCGTextureDataInterfaceConstants::MAX_NUM_TEXTURES)
+	for (int I = 0; I < InputTaggedData.Num(); ++I)
 	{
-		UE_LOG(LogPCG, Warning, TEXT("Texture data interface on pin '%s' received too many data items (%d), only first %d valid textures will be used."),
-			*DataInterface->DownstreamInputPinLabelAliases[0].ToString(),
-			InputTaggedData.Num(),
-			PCGTextureDataInterfaceConstants::MAX_NUM_TEXTURES);
-	}
+		const FPCGTaggedData& TaggedData = InputTaggedData[I];
+		const UPCGBaseTextureData* TextureData = Cast<UPCGBaseTextureData>(TaggedData.Data);
 
-	const int NumInputTextures = FMath::Min(InputTaggedData.Num(), (int)PCGTextureDataInterfaceConstants::MAX_NUM_TEXTURES);
-
-	TextureInfos.SetNum(NumInputTextures);
-
-	for (int InputTextureIndex = 0; InputTextureIndex < NumInputTextures; ++InputTextureIndex)
-	{
-		const UPCGBaseTextureData* TextureData = Cast<UPCGBaseTextureData>(InputTaggedData[InputTextureIndex].Data);
 		if (!ensure(TextureData))
 		{
-			UE_LOG(LogPCG, Error, TEXT("Unsupported data type encountered by texture data interface: '%s'"),
-				InputTaggedData[InputTextureIndex].Data ? *InputTaggedData[InputTextureIndex].Data->GetName() : TEXT("NULL"));
+			UE_LOG(LogPCG, Error, TEXT("Unsupported data type encountered by texture data interface: '%s'"), TaggedData.Data ? *TaggedData.Data->GetName() : TEXT("NULL"));
 			continue;
 		}
 
-		FPCGInputTextureInfo& TextureInfo = TextureInfos[InputTextureIndex];
-
-		TextureInfo.Texture = TextureData->GetTextureRHI();
-		TextureInfo.ResourceType = TextureData->GetTextureResourceType();
-		TextureInfo.Filter = TextureData->Filter;
-		TextureInfo.TextureBounds = TextureData->GetBounds();
-		TextureInfo.SliceIndex = TextureData->GetTextureSlice();
+		FPCGTextureBindingInfo BindingInfo(TextureData);
+		int32 BindingIndex = BindingInfos.Find(BindingInfo);
 
-		if (TextureInfo.ResourceType == EPCGTextureResourceType::ExportedTexture)
+		if (BindingIndex == INDEX_NONE)
 		{
-			TextureInfo.ExportedTexture = TextureData->GetRefCountedTexture();
+			if (BindingInfos.Num() < PCGTextureDataInterfaceConstants::MAX_NUM_BINDINGS)
+			{
+				BindingIndex = BindingInfos.Add(MoveTemp(BindingInfo));
+			}
+			else
+			{
+				UE_LOG(LogPCG, Warning, TEXT("Texture data interface on pin '%s' received too many textures to bind. Only the first %d textures will be bound."),
+					*DataInterface->DownstreamInputPinLabelAliases[0].ToString(),
+					PCGTextureDataInterfaceConstants::MAX_NUM_BINDINGS);
+				continue;
+			}
 		}
 
-		if (TextureInfo.Texture)
-		{
-			TextureInfo.bIsTextureArray = TextureInfo.Texture->GetDesc().Dimension == ETextureDimension::Texture2DArray;
-		}
+		FPCGTextureInfo& TextureInfo = TextureInfos.Emplace_GetRef();
+		TextureInfo.BindingIndex = BindingIndex;
+		TextureInfo.SliceIndex = TextureData->GetTextureSlice();
 	}
 }
 
 FComputeDataProviderRenderProxy* UPCGTextureDataProvider::GetRenderProxy()
 {
-	return new FPCGTextureDataProviderProxy(TextureInfos);
+	return new FPCGTextureDataProviderProxy(BindingInfos, TextureInfos);
 }
 
 void UPCGTextureDataProvider::Reset()
 {
+	BindingInfos.Empty();
 	TextureInfos.Empty();
 
 	Super::Reset();
@@ -175,6 +219,27 @@ bool FPCGTextureDataProviderProxy::IsValid(FValidationData const& InValidationDa
 	return true;
 }
 
+struct FPCGTextureDataInterfacePermutationIds
+{
+	uint32 EnableMultipleTextureObjects = 0;
+ 
+	FPCGTextureDataInterfacePermutationIds(const FComputeKernelPermutationVector& PermutationVector)
+	{
+		static FString Name(PCGTextureDataInterfaceConstants::EnableMultipleTextureObjectsPermutationName);
+		static uint32 Hash = GetTypeHash(Name);
+		EnableMultipleTextureObjects = PermutationVector.GetPermutationBits(Name, Hash, /*Value=*/1);
+	}
+};
+
+void FPCGTextureDataProviderProxy::GatherPermutations(FPermutationData& InOutPermutationData) const
+{
+	FPCGTextureDataInterfacePermutationIds PermutationIds(InOutPermutationData.PermutationVector);
+	for (int32 InvocationIndex = 0; InvocationIndex < InOutPermutationData.NumInvocations; ++InvocationIndex)
+	{
+		InOutPermutationData.PermutationIds[InvocationIndex] |= (BindingInfos.Num() > 1 ? PermutationIds.EnableMultipleTextureObjects : 0);
+	}
+}
+
 void FPCGTextureDataProviderProxy::GatherDispatchData(FDispatchData const& InDispatchData)
 {
 	const TStridedView<FParameters> ParameterArray = MakeStridedParameterView<FParameters>(InDispatchData);
@@ -182,19 +247,20 @@ void FPCGTextureDataProviderProxy::GatherDispatchData(FDispatchData const& InDis
 	{
 		FParameters& Parameters = ParameterArray[InvocationIndex];
 
-		for (int BindingIndex = 0; BindingIndex < PCGTextureDataInterfaceConstants::MAX_NUM_TEXTURES; ++BindingIndex)
+		for (int BindingIndex = 0; BindingIndex < PCGTextureDataInterfaceConstants::MAX_NUM_BINDINGS; ++BindingIndex)
 		{
 			Parameters.Texture[BindingIndex] = TextureSRV[BindingIndex];
+			Parameters.TextureArray[BindingIndex] = TextureArraySRV[BindingIndex];
 
-			if (BindingIndex < TextureInfos.Num())
+			if (BindingIndex < BindingInfos.Num())
 			{
 				Parameters.TextureBounds[BindingIndex] = FVector4f(
-					TextureInfos[BindingIndex].TextureBounds.Min.X,
-					TextureInfos[BindingIndex].TextureBounds.Min.Y,
-					TextureInfos[BindingIndex].TextureBounds.Max.X,
-					TextureInfos[BindingIndex].TextureBounds.Max.Y);
+					BindingInfos[BindingIndex].TextureBounds.Min.X,
+					BindingInfos[BindingIndex].TextureBounds.Min.Y,
+					BindingInfos[BindingIndex].TextureBounds.Max.X,
+					BindingInfos[BindingIndex].TextureBounds.Max.Y);
 
-				switch (TextureInfos[BindingIndex].Filter)
+				switch (BindingInfos[BindingIndex].Filter)
 				{
 				case EPCGTextureFilter::Point:
 					Parameters.Sampler[BindingIndex] = TStaticSamplerState<SF_Point, AM_Clamp, AM_Clamp, AM_Clamp>::GetRHI();
@@ -211,68 +277,92 @@ void FPCGTextureDataProviderProxy::GatherDispatchData(FDispatchData const& InDis
 				Parameters.TextureBounds[BindingIndex] = FVector4f::Zero();
 				Parameters.Sampler[BindingIndex] = TStaticSamplerState<SF_Bilinear, AM_Clamp, AM_Clamp, AM_Clamp>::GetRHI();
 			}
+
+			Parameters.NumTextureInfos = TextureInfos.Num();
+			Parameters.TextureInfos = TextureInfosBufferSRV;
 		}
 	}
 }
 
 void FPCGTextureDataProviderProxy::AllocateResources(FRDGBuilder& GraphBuilder, FAllocationData const& InAllocationData)
 {
-	FRDGTextureRef DummyTextureRDG = nullptr;
+	CreateTextureSRVs(GraphBuilder);
+	PackTextureInfos(GraphBuilder);
+}
 
-	for (int DataIndex = 0; DataIndex < PCGTextureDataInterfaceConstants::MAX_NUM_TEXTURES; ++DataIndex)
+void FPCGTextureDataProviderProxy::CreateTextureSRVs(FRDGBuilder& GraphBuilder)
+{
+	const FRDGTextureRef DummyTextureRDG = GSystemTextures.GetDefaultTexture2D(GraphBuilder, EPixelFormat::PF_G8, 0.0f);
+	const FRDGTextureRef DummyTextureArrayRDG = GSystemTextures.GetDefaultTexture(GraphBuilder, ETextureDimension::Texture2DArray, EPixelFormat::PF_G8, 0.0f);
+
+	for (int BindingIndex = 0; BindingIndex < PCGTextureDataInterfaceConstants::MAX_NUM_BINDINGS; ++BindingIndex)
 	{
-		FRDGTextureRef TextureRDG = nullptr;
+		// Default to blank dummy texture if no texture was provided.
+		FRDGTextureRef TextureRDG = DummyTextureRDG;
+		FRDGTextureRef TextureArrayRDG = DummyTextureArrayRDG;
 
-		if (DataIndex < TextureInfos.Num())
+		if (BindingIndex < BindingInfos.Num())
 		{
-			if (TextureInfos[DataIndex].ResourceType == EPCGTextureResourceType::TextureObject && TextureInfos[DataIndex].Texture)
+			FRDGTextureRef TempRDG = nullptr;
+
+			if (BindingInfos[BindingIndex].ResourceType == EPCGTextureResourceType::TextureObject && BindingInfos[BindingIndex].Texture)
 			{
-				TextureRDG = GraphBuilder.RegisterExternalTexture(
-					CreateRenderTarget(TextureInfos[DataIndex].Texture, *FText::Format(LOCTEXT("RenderTargetDebugName", "PCGDataInterfaceRenderTarget_{0}"), DataIndex).ToString()));
+				TempRDG = GraphBuilder.RegisterExternalTexture(
+					CreateRenderTarget(BindingInfos[BindingIndex].Texture, *FText::Format(LOCTEXT("RenderTargetDebugName", "PCGDataInterfaceRenderTarget_{0}"), BindingIndex).ToString()));
 			}
-			else if (TextureInfos[DataIndex].ResourceType == EPCGTextureResourceType::ExportedTexture && TextureInfos[DataIndex].ExportedTexture)
+			else if (BindingInfos[BindingIndex].ResourceType == EPCGTextureResourceType::ExportedTexture && BindingInfos[BindingIndex].ExportedTexture)
 			{
-				TextureRDG = GraphBuilder.RegisterExternalTexture(TextureInfos[DataIndex].ExportedTexture);
+				TempRDG = GraphBuilder.RegisterExternalTexture(BindingInfos[BindingIndex].ExportedTexture);
 			}
-		}
 
-		if (!TextureRDG)
-		{
-			if (!DummyTextureRDG)
+			if (BindingInfos[BindingIndex].bIsTextureArray)
 			{
-				// Copied from GSystemTextures.GetBlackDummy(GraphBuilder));
-				const FRHITextureCreateDesc DummyTextureDesc =
-					FRHITextureCreateDesc::Create2D(TEXT("PCGDummyTexture"), 1, 1, PF_G8)
-					.SetFlags(ETextureCreateFlags::ShaderResource);
-
-				FTextureRHIRef Texture = RHICreateTexture(DummyTextureDesc);
+				TextureArrayRDG = TempRDG;
+			}
+			else
+			{
+				TextureRDG = TempRDG;
+			}
+		}
 
-				{
-					uint32 DestStride;
-					uint8* DestBuffer = (uint8*)RHILockTexture2D(Texture, 0, RLM_WriteOnly, DestStride, false);
-					*DestBuffer = 0;
-					RHIUnlockTexture2D(Texture, 0, false);
-				}
+		TextureSRV[BindingIndex] = GraphBuilder.CreateSRV(TextureRDG);
+		TextureArraySRV[BindingIndex] = GraphBuilder.CreateSRV(TextureArrayRDG);
+	}
+}
 
-				// Create a dummy texture to bind.
-				TRefCountPtr<IPooledRenderTarget> BlackDummyTarget = CreateRenderTarget(Texture, TEXT("PCGDummyRenderTarget"));
-				check(BlackDummyTarget.IsValid());
+void FPCGTextureDataProviderProxy::PackTextureInfos(FRDGBuilder& GraphBuilder)
+{
+	// Layout:
+	// - TextureInfo0
+	//   - uint32 - BindingIndex
+	//   - uint32 - SliceIndex (aka TextureIndex if bound to a TextureArray, 0 otherwise.)
+	//   - uint32 - IsTextureArray
+	// - TextureInfo1
+	// - TextureInfo2
+	// - ...
+
+	if (!TextureInfos.IsEmpty())
+	{
+		const FRDGBufferDesc TextureInfosDesc = FRDGBufferDesc::CreateStructuredDesc(sizeof(FUintVector3), TextureInfos.Num());
 
-				DummyTextureRDG = GraphBuilder.RegisterExternalTexture(BlackDummyTarget);
-			}
+		FRDGBufferRef TextureInfosBuffer = GraphBuilder.CreateBuffer(TextureInfosDesc, TEXT("PCGTextureInfosBuffer"));
+		TextureInfosBufferSRV = GraphBuilder.CreateSRV(TextureInfosBuffer);
 
-			// Fallback to blank dummy texture if no texture was provided.
-			TextureRDG = DummyTextureRDG;
-		}
+		TArray<FUintVector3> TextureInfosPackedData;
+		TextureInfosPackedData.SetNumUninitialized(TextureInfos.Num());
 
-		if (DataIndex < TextureInfos.Num() && TextureInfos[DataIndex].bIsTextureArray && TextureRDG != DummyTextureRDG)
-		{
-			TextureSRV[DataIndex] = GraphBuilder.CreateSRV(FRDGTextureSRVDesc::CreateForSlice(TextureRDG, TextureInfos[DataIndex].SliceIndex));
-		}
-		else
+		for (int TextureInfoIndex = 0; TextureInfoIndex < TextureInfos.Num(); ++TextureInfoIndex)
 		{
-			TextureSRV[DataIndex] = GraphBuilder.CreateSRV(TextureRDG);
+			const FPCGTextureInfo& TextureInfo = TextureInfos[TextureInfoIndex];
+			const FPCGTextureBindingInfo& BindingInfo = BindingInfos[TextureInfo.BindingIndex];
+			TextureInfosPackedData[TextureInfoIndex] = FUintVector3(TextureInfo.BindingIndex, TextureInfo.SliceIndex, BindingInfo.bIsTextureArray);
 		}
+
+		GraphBuilder.QueueBufferUpload(TextureInfosBuffer, MakeArrayView(TextureInfosPackedData));
+	}
+	else
+	{
+		TextureInfosBufferSRV = GraphBuilder.CreateSRV(FRDGBufferSRVDesc(GSystemTextures.GetDefaultStructuredBuffer(GraphBuilder, sizeof(FUintVector3))));
 	}
 }
 
diff --git a/Engine/Plugins/PCG/Source/PCG/Private/Compute/DataInterfaces/PCGTextureDataInterface.h b/Engine/Plugins/PCG/Source/PCG/Private/Compute/DataInterfaces/PCGTextureDataInterface.h
--- a/Engine/Plugins/PCG/Source/PCG/Private/Compute/DataInterfaces/PCGTextureDataInterface.h
+++ b/Engine/Plugins/PCG/Source/PCG/Private/Compute/DataInterfaces/PCGTextureDataInterface.h
@@ -12,25 +12,37 @@
 #include "PCGTextureDataInterface.generated.h"
 
 class FPCGTextureDataInterfaceParameters;
+class UPCGBaseTextureData;
 
 namespace PCGTextureDataInterfaceConstants
 {
 	/**
 	 * The maximum number of textures which can be bound to a PCGTextureDataInterface.
+	 * We can't use full bindless on all platforms, so fallback to emulating.
 	 * If you change this number, make sure update PCGTextureDataInterface.ush as well.
 	 */
-	constexpr uint32 MAX_NUM_TEXTURES = 8;
+	constexpr uint32 MAX_NUM_BINDINGS = 8;
 }
 
-struct FPCGInputTextureInfo
+struct FPCGTextureBindingInfo
 {
+	FPCGTextureBindingInfo(const UPCGBaseTextureData* InTextureData);
+	bool operator==(const FPCGTextureBindingInfo& Other) const;
+
 	EPCGTextureResourceType ResourceType = EPCGTextureResourceType::TextureObject;
 	FTextureRHIRef Texture = nullptr;
 	TRefCountPtr<IPooledRenderTarget> ExportedTexture = nullptr;
+	bool bIsTextureArray = false;
+
+	// @todo_pcg: Filter and TextureBounds could potentially be independent of the texture binding, and queried via lookup from the texture info buffer instead.
 	EPCGTextureFilter Filter = EPCGTextureFilter::Bilinear;
 	FBox TextureBounds = FBox(EForceInit::ForceInit);
+};
+
+struct FPCGTextureInfo
+{
+	int BindingIndex = 0;
 	int SliceIndex = 0;
-	bool bIsTextureArray = false;
 };
 
 /** Data Interface allowing sampling of a texture. */
@@ -47,6 +59,8 @@ public:
 	TCHAR const* GetShaderVirtualPath() const override;
 	void GetShaderHash(FString& InOutKey) const override;
 	void GetHLSL(FString& OutHLSL, FString const& InDataInterfaceName) const override;
+	void GetDefines(FComputeKernelDefinitionSet& OutDefinitionSet) const override;
+	void GetPermutations(FComputeKernelPermutationVector& OutPermutationVector) const override;
 	UComputeDataProvider* CreateDataProvider() const override;
 	//~ End UComputeDataInterface Interface
 
@@ -67,26 +81,35 @@ public:
 	virtual void Reset() override;
 	//~ End UComputeDataProvider Interface
 
-	TArray<FPCGInputTextureInfo> TextureInfos;
+	TArray<FPCGTextureBindingInfo> BindingInfos;
+	TArray<FPCGTextureInfo> TextureInfos;
 };
 
 class FPCGTextureDataProviderProxy : public FComputeDataProviderRenderProxy
 {
 public:
-	FPCGTextureDataProviderProxy(const TArray<FPCGInputTextureInfo>& InTextureInfos)
-		: TextureInfos(InTextureInfos)
+	FPCGTextureDataProviderProxy(const TArray<FPCGTextureBindingInfo>& InBindingInfos, const TArray<FPCGTextureInfo>& InTextureInfos)
+		: BindingInfos(InBindingInfos)
+		, TextureInfos(InTextureInfos)
 	{}
 
 	//~ Begin FComputeDataProviderRenderProxy Interface
 	bool IsValid(FValidationData const& InValidationData) const override;
+	void GatherPermutations(FPermutationData& InOutPermutationData) const override;
 	void GatherDispatchData(FDispatchData const& InDispatchData) override;
 	void AllocateResources(FRDGBuilder& GraphBuilder, FAllocationData const& InAllocationData) override;
 	//~ End FComputeDataProviderRenderProxy Interface
 
+	void CreateTextureSRVs(FRDGBuilder& GraphBuilder);
+	void PackTextureInfos(FRDGBuilder& GraphBuilder);
+
 protected:
 	using FParameters = FPCGTextureDataInterfaceParameters;
 
-	TArray<FPCGInputTextureInfo> TextureInfos;
+	TArray<FPCGTextureBindingInfo> BindingInfos;
+	TArray<FPCGTextureInfo> TextureInfos;
 
-	FRDGTextureSRVRef TextureSRV[PCGTextureDataInterfaceConstants::MAX_NUM_TEXTURES];
+	FRDGTextureSRVRef TextureSRV[PCGTextureDataInterfaceConstants::MAX_NUM_BINDINGS];
+	FRDGTextureSRVRef TextureArraySRV[PCGTextureDataInterfaceConstants::MAX_NUM_BINDINGS];
+	FRDGBufferSRVRef TextureInfosBufferSRV = nullptr;
 };
