[PCG] Add option to skip readback during texture initialization
* Expose option to skip readback from Get Texture Data and Generate Grass Maps nodes

diff --git a/Engine/Plugins/PCG/Source/PCG/Private/Data/PCGRenderTargetData.cpp b/Engine/Plugins/PCG/Source/PCG/Private/Data/PCGRenderTargetData.cpp
--- a/Engine/Plugins/PCG/Source/PCG/Private/Data/PCGRenderTargetData.cpp
+++ b/Engine/Plugins/PCG/Source/PCG/Private/Data/PCGRenderTargetData.cpp
@@ -8,14 +8,15 @@
 
 #include UE_INLINE_GENERATED_CPP_BY_NAME(PCGRenderTargetData)
 
-void UPCGRenderTargetData::Initialize(UTextureRenderTarget2D* InRenderTarget, const FTransform& InTransform)
+void UPCGRenderTargetData::Initialize(UTextureRenderTarget2D* InRenderTarget, const FTransform& InTransform, bool bInSkipReadbackToCPU)
 {
 	RenderTarget = InRenderTarget;
 	Transform = InTransform;
+	bSkipReadbackToCPU = bInSkipReadbackToCPU;
 
 	ColorData.Reset();
 
-	if (RenderTarget)
+	if (RenderTarget && !bSkipReadbackToCPU)
 	{
 		if (FTextureRenderTargetResource* RTResource = RenderTarget->GameThread_GetRenderTargetResource())
 		{
diff --git a/Engine/Plugins/PCG/Source/PCG/Private/Data/PCGTextureData.cpp b/Engine/Plugins/PCG/Source/PCG/Private/Data/PCGTextureData.cpp
--- a/Engine/Plugins/PCG/Source/PCG/Private/Data/PCGTextureData.cpp
+++ b/Engine/Plugins/PCG/Source/PCG/Private/Data/PCGTextureData.cpp
@@ -239,6 +239,17 @@ bool UPCGBaseTextureData::SamplePoint(const FTransform& InTransform, const FBox&
 		return false;
 	}
 
+	if (bSkipReadbackToCPU)
+	{
+		if (!bEmittedNoReadbackDataError)
+		{
+			UE_LOG(LogPCG, Error, TEXT("Texture data was initialized with bSkipReadbackToCPU enabled - point cannot be sampled."));
+			bEmittedNoReadbackDataError = true;
+		}
+		
+		return false;
+	}
+
 	// Compute transform
 	// TODO: embed local bounds center offset at this time?
 	OutPoint.Transform = InTransform;
@@ -278,7 +289,13 @@ const UPCGPointData* UPCGBaseTextureData::CreatePointData(FPCGContext* Context)
 	// Early out for invalid data
 	if (!IsValid())
 	{
-		UE_LOG(LogPCG, Error, TEXT("Texture data does not have valid sizes - will return empty data"));
+		UE_LOG(LogPCG, Error, TEXT("Texture data does not have valid sizes - will return empty data."));
+		return Data;
+	}
+
+	if (bSkipReadbackToCPU)
+	{
+		UE_LOG(LogPCG, Error, TEXT("Texture data was initialized with bSkipReadbackToCPU enabled - will return empty data."));
 		return Data;
 	}
 
@@ -334,7 +351,7 @@ const UPCGPointData* UPCGBaseTextureData::CreatePointData(FPCGContext* Context)
 
 bool UPCGBaseTextureData::IsValid() const
 {
-	return Height > 0 && Width > 0 && !ColorData.IsEmpty();
+	return (Height > 0 && Width > 0 && !ColorData.IsEmpty()) || bSkipReadbackToCPU;
 }
 
 bool UPCGBaseTextureData::SamplePointLocal(const FVector2D& LocalPosition, FVector4& OutColor, float& OutDensity) const
@@ -344,6 +361,17 @@ bool UPCGBaseTextureData::SamplePointLocal(const FVector2D& LocalPosition, FVect
 		return false;
 	}
 
+	if (bSkipReadbackToCPU)
+	{
+		if (!bEmittedNoReadbackDataError)
+		{
+			UE_LOG(LogPCG, Error, TEXT("Texture data was initialized with bSkipReadbackToCPU enabled - point cannot be sampled."));
+			bEmittedNoReadbackDataError = true;
+		}
+		
+		return false;
+	}
+
 	FVector2D Pos;
 	Pos.X = FMath::Frac(LocalPosition.X);
 	Pos.Y = FMath::Frac(LocalPosition.Y);
@@ -373,11 +401,12 @@ void UPCGBaseTextureData::CopyBaseTextureData(UPCGBaseTextureData* NewTextureDat
 	NewTextureData->Bounds = Bounds;
 	NewTextureData->Height = Height;
 	NewTextureData->Width = Width;
+	NewTextureData->bSkipReadbackToCPU = bSkipReadbackToCPU;
 }
 
-void UPCGTextureData::InitializeInternal(UTexture* InTexture, uint32 InTextureIndex, const FTransform& InTransform, bool* bOutInitializeCompleted, bool bCreateCPUDuplicateEditorOnly)
+void UPCGTextureData::InitializeInternal(UTexture* InTexture, uint32 InTextureIndex, const FTransform& InTransform, bool* bOutInitializeCompleted, bool bCreateCPUDuplicateEditorOnly, bool bInSkipReadbackToCPU)
 {
-	TRACE_CPUPROFILER_EVENT_SCOPE(UPCGTextureData::Initialize);
+	TRACE_CPUPROFILER_EVENT_SCOPE(UPCGTextureData::InitializeInternal);
 
 	auto SetInitCompleted = [bOutInitializeCompleted](bool bIsDone)
 	{
@@ -411,12 +440,20 @@ void UPCGTextureData::InitializeInternal(UTexture* InTexture, uint32 InTextureIn
 	Transform = InTransform;
 	Width = 0;
 	Height = 0;
+	bSkipReadbackToCPU = bInSkipReadbackToCPU;
 
 	Bounds = FBox(EForceInit::ForceInit);
 	Bounds += FVector(-1.0f, -1.0f, 0.0f);
 	Bounds += FVector(1.0f, 1.0f, 0.0f);
 	Bounds = Bounds.TransformBy(Transform);
 
+	if (bInSkipReadbackToCPU)
+	{
+		bSuccessfullyInitialized = true;
+		SetInitCompleted(true);
+		return;
+	}
+
 	// Prioritize initializing from a CPU texture when the provided texture is marked as CPU accessible
 	TOptional<bool> InitializedFromCPUTexture = InitializeFromCPUTexture();
 	if (!InitializedFromCPUTexture.IsSet())
@@ -500,17 +537,17 @@ void UPCGTextureData::InitializeInternal(UTexture* InTexture, uint32 InTextureIn
 #endif
 
 	// Finally try the GPU -> CPU readback path. We don't flag success yet though - this will be done when the readback data arrives.
-	SetInitCompleted(InitializeFromGPUTexture());
+	SetInitCompleted(ReadbackFromGPUTexture());
 
 	return;
 }
 
-bool UPCGTextureData::Initialize(UTexture* InTexture, uint32 InTextureIndex, const FTransform& InTransform, bool bCreateCPUDuplicateEditorOnly)
+bool UPCGTextureData::Initialize(UTexture* InTexture, uint32 InTextureIndex, const FTransform& InTransform, bool bCreateCPUDuplicateEditorOnly, bool bInSkipReadbackToCPU)
 {
 	ResourceType = EPCGTextureResourceType::TextureObject;
 
 	bool bInitializeDone = false;
-	InitializeInternal(InTexture, InTextureIndex, InTransform, &bInitializeDone, bCreateCPUDuplicateEditorOnly);
+	InitializeInternal(InTexture, InTextureIndex, InTransform, &bInitializeDone, bCreateCPUDuplicateEditorOnly, bInSkipReadbackToCPU);
 	return bInitializeDone;
 }
 
@@ -519,7 +556,7 @@ void UPCGTextureData::Initialize(UTexture* InTexture, uint32 InTextureIndex, con
 	ResourceType = EPCGTextureResourceType::TextureObject;
 
 	bool bInitializeDone = false;
-	InitializeInternal(InTexture, InTextureIndex, InTransform, &bInitializeDone, bCreateCPUDuplicateEditorOnly);
+	InitializeInternal(InTexture, InTextureIndex, InTransform, &bInitializeDone, bCreateCPUDuplicateEditorOnly, bSkipReadbackToCPU);
 
 	if (bInitializeDone)
 	{
@@ -531,19 +568,27 @@ void UPCGTextureData::Initialize(UTexture* InTexture, uint32 InTextureIndex, con
 	}
 }
 
-bool UPCGTextureData::Initialize(TRefCountPtr<IPooledRenderTarget> InTextureHandle, uint32 InTextureIndex, const FTransform& InTransform)
+bool UPCGTextureData::Initialize(TRefCountPtr<IPooledRenderTarget> InTextureHandle, uint32 InTextureIndex, const FTransform& InTransform, bool bInSkipReadbackToCPU)
 {
 	ResourceType = EPCGTextureResourceType::ExportedTexture;
 	TextureHandle = InTextureHandle;
 	TextureIndex = InTextureIndex;
 	Transform = InTransform;
+	bSkipReadbackToCPU = bInSkipReadbackToCPU;
 
 	Bounds = FBox(EForceInit::ForceInit);
 	Bounds += FVector(-1.0f, -1.0f, 0.0f);
 	Bounds += FVector(1.0f, 1.0f, 0.0f);
 	Bounds = Bounds.TransformBy(Transform);
 
-	InitializeFromGPUTexture();
+	if (bInSkipReadbackToCPU)
+	{
+		bSuccessfullyInitialized = true;
+	}
+	else
+	{
+		ReadbackFromGPUTexture();
+	}
 
 	return bSuccessfullyInitialized;
 }
@@ -715,8 +760,10 @@ TOptional<bool> UPCGTextureData::InitializeFromCPUTexture()
 	return true;
 }
 
-bool UPCGTextureData::InitializeFromGPUTexture()
+bool UPCGTextureData::ReadbackFromGPUTexture()
 {
+	TRACE_CPUPROFILER_EVENT_SCOPE(UPCGTextureData::ReadbackFromGPUTexture);
+
 	if (bReadbackFromGPUInitiated)
 	{
 		return false;
@@ -828,7 +875,6 @@ bool UPCGTextureData::InitializeFromGPUTexture()
 		else
 		{
 			UE_LOG(LogPCG, Error, TEXT("PCGTextureData failed to acquire texture resource."));
-
 		}
 
 		return true;
diff --git a/Engine/Plugins/PCG/Source/PCG/Private/Elements/PCGGenerateGrassMaps.cpp b/Engine/Plugins/PCG/Source/PCG/Private/Elements/PCGGenerateGrassMaps.cpp
--- a/Engine/Plugins/PCG/Source/PCG/Private/Elements/PCGGenerateGrassMaps.cpp
+++ b/Engine/Plugins/PCG/Source/PCG/Private/Elements/PCGGenerateGrassMaps.cpp
@@ -554,7 +554,7 @@ bool FPCGGenerateGrassMapsElement::ExecuteInternal(FPCGContext* InContext) const
 		check(TextureData);
 
 		// Poll initialize (fine to be called even when initialization was already complete).
-		bAllTexturesInitialized &= TextureData->Initialize(Context->GrassMapHandle, /*TextureIndex=*/Context->SelectedGrassTypes[DataIndex].Get<1>(), GrassMapTransform);
+		bAllTexturesInitialized &= TextureData->Initialize(Context->GrassMapHandle, /*TextureIndex=*/Context->SelectedGrassTypes[DataIndex].Get<1>(), GrassMapTransform, Settings->bSkipReadbackToCPU);
 	}
 
 	if (!bAllTexturesInitialized)
diff --git a/Engine/Plugins/PCG/Source/PCG/Private/Elements/PCGTextureSampler.cpp b/Engine/Plugins/PCG/Source/PCG/Private/Elements/PCGTextureSampler.cpp
--- a/Engine/Plugins/PCG/Source/PCG/Private/Elements/PCGTextureSampler.cpp
+++ b/Engine/Plugins/PCG/Source/PCG/Private/Elements/PCGTextureSampler.cpp
@@ -154,11 +154,6 @@ bool FPCGTextureSamplerElement::ExecuteInternal(FPCGContext* InContext) const
 		return false;
 	}
 
-	if (Context->bTextureReadbackDone)
-	{
-		return true;
-	}
-
 	const UPCGTextureSamplerSettings* Settings = Context->GetInputSettings<UPCGTextureSamplerSettings>();
 	check(Settings);
 
@@ -229,7 +224,7 @@ bool FPCGTextureSamplerElement::ExecuteInternal(FPCGContext* InContext) const
 	if (UTextureRenderTarget2D* RenderTarget = Cast<UTextureRenderTarget2D>(Texture))
 	{
 		UPCGRenderTargetData* RenderTargetData = FPCGContext::NewObject_AnyThread<UPCGRenderTargetData>(Context);
-		RenderTargetData->Initialize(RenderTarget, Context->Transform);
+		RenderTargetData->Initialize(RenderTarget, Context->Transform, Settings->bSkipReadbackToCPU);
 		BaseTextureData = RenderTargetData;
 	}
 	else
@@ -261,7 +256,7 @@ bool FPCGTextureSamplerElement::ExecuteInternal(FPCGContext* InContext) const
 			const bool bForceEditorOnlyCPUSampling = false;
 #endif
 
-			if (!TextureData->Initialize(Texture, TextureArrayIndex, Context->Transform, bForceEditorOnlyCPUSampling))
+			if (!TextureData->Initialize(Texture, TextureArrayIndex, Context->Transform, bForceEditorOnlyCPUSampling, Settings->bSkipReadbackToCPU))
 			{
 				// Initialization not complete. Could be waiting on async texture processing or for GPU readback. Sleep until next frame.
 				Context->bIsPaused = true;
diff --git a/Engine/Plugins/PCG/Source/PCG/Public/Data/PCGRenderTargetData.h b/Engine/Plugins/PCG/Source/PCG/Public/Data/PCGRenderTargetData.h
--- a/Engine/Plugins/PCG/Source/PCG/Public/Data/PCGRenderTargetData.h
+++ b/Engine/Plugins/PCG/Source/PCG/Public/Data/PCGRenderTargetData.h
@@ -35,7 +35,7 @@ protected:
 
 public:
 	UFUNCTION(BlueprintCallable, Category = RenderTarget)
-	PCG_API void Initialize(UTextureRenderTarget2D* InRenderTarget, const FTransform& InTransform);
+	PCG_API void Initialize(UTextureRenderTarget2D* InRenderTarget, const FTransform& InTransform, bool bInSkipReadbackToCPU = false);
 
 	UPROPERTY(BlueprintReadOnly, VisibleAnywhere, Category = Properties)
 	TObjectPtr<UTextureRenderTarget2D> RenderTarget = nullptr;
diff --git a/Engine/Plugins/PCG/Source/PCG/Public/Data/PCGTextureData.h b/Engine/Plugins/PCG/Source/PCG/Public/Data/PCGTextureData.h
--- a/Engine/Plugins/PCG/Source/PCG/Public/Data/PCGTextureData.h
+++ b/Engine/Plugins/PCG/Source/PCG/Public/Data/PCGTextureData.h
@@ -160,6 +160,12 @@ protected:
 	UPROPERTY(BlueprintReadOnly, VisibleAnywhere, Category = SpatialData)
 	int32 Width = 0;
 
+	UPROPERTY()
+	bool bSkipReadbackToCPU = false;
+
+	/** Used to make sure errors are only logged once when trying to sample points from a data which hasn't been read back into a CPU buffer. */
+	mutable bool bEmittedNoReadbackDataError = false;
+
 	void CopyBaseTextureData(UPCGBaseTextureData* NewTextureData) const;
 };
 
@@ -173,8 +179,8 @@ public:
 	 * Initialize this data. Can depend on async texture operations / async GPU readbacks. Should be polled until it returns true signaling completion,
 	 * and then IsInitialized() is used to verify the initialization was successful and data is ready to use.
 	 */
-	PCG_API bool Initialize(UTexture* InTexture, uint32 InTextureIndex, const FTransform& InTransform, bool bCreateCPUDuplicateEditorOnly = false);
-	PCG_API bool Initialize(TRefCountPtr<IPooledRenderTarget> InTextureHandle, uint32 InTextureIndex, const FTransform& InTransform);
+	PCG_API bool Initialize(UTexture* InTexture, uint32 InTextureIndex, const FTransform& InTransform, bool bCreateCPUDuplicateEditorOnly = false, bool bInSkipReadbackToCPU = false);
+	PCG_API bool Initialize(TRefCountPtr<IPooledRenderTarget> InTextureHandle, uint32 InTextureIndex, const FTransform& InTransform, bool bInSkipReadbackToCPU = false);
 
 	/** Data is successfully initialized and is ready to use. */
 	PCG_API bool IsSuccessfullyInitialized() const { return bSuccessfullyInitialized; }
@@ -201,17 +207,14 @@ protected:
 	virtual UPCGSpatialData* CopyInternal(FPCGContext* Context) const override;
 	//~End UPCGSpatialData interface
 
-	void InitializeInternal(UTexture* InTexture, uint32 InTextureIndex, const FTransform& InTransform, bool* bOutInitializeDone, bool bCreateCPUDuplicateEditorOnly = false);
+	void InitializeInternal(UTexture* InTexture, uint32 InTextureIndex, const FTransform& InTransform, bool* bOutInitializeDone, bool bCreateCPUDuplicateEditorOnly, bool bInSkipReadbackToCPU);
 
 private:
 	/** Attempts to initialize the UPCGTextureData from a CPU-accessible texture. Returns true if CPU initialization succeeds. */
 	TOptional<bool> InitializeFromCPUTexture();
 
-	/** 
-	* Attempts to initialize the UPCGTextureData from a GPU-accessible texture. Returns true if GPU texture readback can be dispatched. 
-	* The PostInitializeCallback is only executed if this function succeeds.
-	*/
-	bool InitializeFromGPUTexture();
+	/** Attempts to read back from a GPU-accessible texture. Returns true if GPU texture readback can be dispatched. */
+	bool ReadbackFromGPUTexture();
 
 #if WITH_EDITOR
 	/** Attempts to initialize the UPCGTextureData from a GPU-accessible texture, but with CPU-accessible memory. Returns true if initialization succeeds. */
diff --git a/Engine/Plugins/PCG/Source/PCG/Public/Elements/PCGGenerateGrassMaps.h b/Engine/Plugins/PCG/Source/PCG/Public/Elements/PCGGenerateGrassMaps.h
--- a/Engine/Plugins/PCG/Source/PCG/Public/Elements/PCGGenerateGrassMaps.h
+++ b/Engine/Plugins/PCG/Source/PCG/Public/Elements/PCGGenerateGrassMaps.h
@@ -56,6 +56,10 @@ public:
 	/** If toggled, will only generate grass types which are not selected. */
 	UPROPERTY(BlueprintReadWrite, EditAnywhere, Category = Settings, meta = (PCG_Overridable))
 	bool bExcludeSelectedGrassTypes = true;
+
+	/** Skip CPU readback of emitted textures during initialization of the texture datas. */
+	UPROPERTY(BlueprintReadWrite, EditAnywhere, Category = Settings)
+	bool bSkipReadbackToCPU = false;
 };
 
 struct FPCGGenerateGrassMapsContext : public FPCGContext
diff --git a/Engine/Plugins/PCG/Source/PCG/Public/Elements/PCGTextureSampler.h b/Engine/Plugins/PCG/Source/PCG/Public/Elements/PCGTextureSampler.h
--- a/Engine/Plugins/PCG/Source/PCG/Public/Elements/PCGTextureSampler.h
+++ b/Engine/Plugins/PCG/Source/PCG/Public/Elements/PCGTextureSampler.h
@@ -125,6 +125,10 @@ public:
 	UPROPERTY(BlueprintReadWrite, EditAnywhere, Category = "Settings|Debug")
 	bool bSynchronousLoad = false;
 
+	/** Skip CPU readback during initialization of the texture data. */
+	UPROPERTY(BlueprintReadWrite, EditAnywhere, Category = Settings)
+	bool bSkipReadbackToCPU = false;
+
 protected:
 	/** Texture specific parameters */
 	UPROPERTY(BlueprintReadWrite, EditAnywhere, Category = Settings, meta = (PCG_Overridable))
@@ -142,8 +146,6 @@ protected:
 struct FPCGTextureSamplerContext : public FPCGContext, public IPCGAsyncLoadingContext
 {
 public:
-	bool bTextureReadbackDone = false;
-
 	TObjectPtr<UPCGTextureData> TextureData;
 
 	bool bTransformInitialized = false;
