[PCG] Fix crash when accessing GSystemTextures too early by reverting to creating dummy textures on the fly.

#rnx

diff --git a/Engine/Plugins/PCG/Source/PCG/Private/Compute/DataInterfaces/PCGTextureDataInterface.cpp b/Engine/Plugins/PCG/Source/PCG/Private/Compute/DataInterfaces/PCGTextureDataInterface.cpp
--- a/Engine/Plugins/PCG/Source/PCG/Private/Compute/DataInterfaces/PCGTextureDataInterface.cpp
+++ b/Engine/Plugins/PCG/Source/PCG/Private/Compute/DataInterfaces/PCGTextureDataInterface.cpp
@@ -292,14 +292,15 @@ void FPCGTextureDataProviderProxy::AllocateResources(FRDGBuilder& GraphBuilder,
 
 void FPCGTextureDataProviderProxy::CreateTextureSRVs(FRDGBuilder& GraphBuilder)
 {
-	const FRDGTextureRef DummyTextureRDG = GSystemTextures.GetDefaultTexture2D(GraphBuilder, EPixelFormat::PF_G8, 0.0f);
-	const FRDGTextureRef DummyTextureArrayRDG = GSystemTextures.GetDefaultTexture(GraphBuilder, ETextureDimension::Texture2DArray, EPixelFormat::PF_G8, 0.0f);
+	// Outside loop so we only create them when necessary, and then share across bindings.
+	FRDGTextureRef DummyTextureRDG = nullptr;
+	FRDGTextureRef DummyTextureArrayRDG = nullptr;
 
 	for (int BindingIndex = 0; BindingIndex < PCGTextureDataInterfaceConstants::MAX_NUM_BINDINGS; ++BindingIndex)
 	{
 		// Default to blank dummy texture if no texture was provided.
-		FRDGTextureRef TextureRDG = DummyTextureRDG;
-		FRDGTextureRef TextureArrayRDG = DummyTextureArrayRDG;
+		FRDGTextureRef TextureRDG = nullptr;
+		FRDGTextureRef TextureArrayRDG = nullptr;
 
 		if (BindingIndex < BindingInfos.Num())
 		{
@@ -325,6 +326,28 @@ void FPCGTextureDataProviderProxy::CreateTextureSRVs(FRDGBuilder& GraphBuilder)
 			}
 		}
 
+		if (TextureRDG == nullptr)
+		{
+			if (DummyTextureRDG == nullptr)
+			{
+				DummyTextureRDG = CreateDummyTexture(GraphBuilder);
+				check(DummyTextureRDG);
+			}
+
+			TextureRDG = DummyTextureRDG;
+		}
+
+		if (TextureArrayRDG == nullptr)
+		{
+			if (DummyTextureArrayRDG == nullptr)
+			{
+				DummyTextureArrayRDG = CreateDummyTextureArray(GraphBuilder);
+				check(DummyTextureArrayRDG);
+			}
+
+			TextureArrayRDG = DummyTextureArrayRDG;
+		}
+
 		TextureSRV[BindingIndex] = GraphBuilder.CreateSRV(TextureRDG);
 		TextureArraySRV[BindingIndex] = GraphBuilder.CreateSRV(TextureArrayRDG);
 	}
@@ -366,4 +389,52 @@ void FPCGTextureDataProviderProxy::PackTextureInfos(FRDGBuilder& GraphBuilder)
 	}
 }
 
+FRDGTextureRef FPCGTextureDataProviderProxy::CreateDummyTexture(FRDGBuilder& GraphBuilder)
+{
+	// Copied from GSystemTextures.GetBlackDummy(). We cannot use GSystemTextures.GetDefaultTexture2D because it
+	// can fail depending on the callstack leading to the compute graph execution.
+
+	const FRHITextureCreateDesc Desc = FRHITextureCreateDesc::Create2D(TEXT("PCGDummyTexture"), 1, 1, PF_G8)
+		.SetFlags(ETextureCreateFlags::ShaderResource);
+
+	FTextureRHIRef Texture = RHICreateTexture(Desc);
+
+	{
+		uint32 DestStride;
+		uint8* DestBuffer = (uint8*)RHILockTexture2D(Texture, 0, RLM_WriteOnly, DestStride, false);
+		*DestBuffer = 0;
+		RHIUnlockTexture2D(Texture, 0, false);
+	}
+
+	// Create a dummy texture to bind.
+	TRefCountPtr<IPooledRenderTarget> BlackDummyTarget = CreateRenderTarget(Texture, TEXT("PCGDummyRenderTarget"));
+	check(BlackDummyTarget.IsValid());
+
+	return GraphBuilder.RegisterExternalTexture(BlackDummyTarget);
+}
+
+FRDGTextureRef FPCGTextureDataProviderProxy::CreateDummyTextureArray(FRDGBuilder& GraphBuilder)
+{
+	// Copied from GSystemTextures.GetBlackDummy(). We cannot use GSystemTextures.GetDefaultTexture because it
+	// can fail depending on the callstack leading to the compute graph execution.
+
+	const FRHITextureCreateDesc Desc = FRHITextureCreateDesc::Create2DArray(TEXT("PCGDummyTextureArray"), 1, 1, 1, PF_G8)
+		.SetFlags(ETextureCreateFlags::ShaderResource);
+
+	FTextureRHIRef Texture = RHICreateTexture(Desc);
+
+	{
+		uint32 DestStride;
+		uint8* DestBuffer = (uint8*)RHILockTexture2D(Texture, 0, RLM_WriteOnly, DestStride, false);
+		*DestBuffer = 0;
+		RHIUnlockTexture2D(Texture, 0, false);
+	}
+
+	// Create a dummy texture to bind.
+	TRefCountPtr<IPooledRenderTarget> BlackDummyTarget = CreateRenderTarget(Texture, TEXT("PCGDummyArrayRenderTarget"));
+	check(BlackDummyTarget.IsValid());
+
+	return GraphBuilder.RegisterExternalTexture(BlackDummyTarget);
+}
+
 #undef LOCTEXT_NAMESPACE
diff --git a/Engine/Plugins/PCG/Source/PCG/Private/Compute/DataInterfaces/PCGTextureDataInterface.h b/Engine/Plugins/PCG/Source/PCG/Private/Compute/DataInterfaces/PCGTextureDataInterface.h
--- a/Engine/Plugins/PCG/Source/PCG/Private/Compute/DataInterfaces/PCGTextureDataInterface.h
+++ b/Engine/Plugins/PCG/Source/PCG/Private/Compute/DataInterfaces/PCGTextureDataInterface.h
@@ -103,6 +103,10 @@ public:
 	void CreateTextureSRVs(FRDGBuilder& GraphBuilder);
 	void PackTextureInfos(FRDGBuilder& GraphBuilder);
 
+protected:
+	FRDGTextureRef CreateDummyTexture(FRDGBuilder& GraphBuilder);
+	FRDGTextureRef CreateDummyTextureArray(FRDGBuilder& GraphBuilder);
+
 protected:
 	using FParameters = FPCGTextureDataInterfaceParameters;
 
