[PCG] Harden code around read back instance counts in GPU SM Spawner.

#rnx

diff --git a/Engine/Plugins/PCG/Source/PCG/Private/Compute/DataInterfaces/Elements/PCGStaticMeshSpawnerDataInterface.cpp b/Engine/Plugins/PCG/Source/PCG/Private/Compute/DataInterfaces/Elements/PCGStaticMeshSpawnerDataInterface.cpp
--- a/Engine/Plugins/PCG/Source/PCG/Private/Compute/DataInterfaces/Elements/PCGStaticMeshSpawnerDataInterface.cpp
+++ b/Engine/Plugins/PCG/Source/PCG/Private/Compute/DataInterfaces/Elements/PCGStaticMeshSpawnerDataInterface.cpp
@@ -446,7 +446,7 @@ void UPCGStaticMeshSpawnerDataProvider::CreatePrimitiveDescriptors(FPCGContext*
 	FPCGDataCollectionDesc InputDataDesc;
 	InBinding->GetKernelPinDataDesc_GameThread(ProducerKernel, PCGPinConstants::DefaultInputLabel, /*bIsInput=*/true, InputDataDesc);
 
-	const uint32 TotalInputPointCount = InputDataDesc.ComputeDataElementCount(EPCGDataType::Point);
+	const int32 TotalInputPointCount = InputDataDesc.ComputeDataElementCount(EPCGDataType::Point);
 	if (TotalInputPointCount == 0)
 	{
 		return;
@@ -528,6 +528,11 @@ void UPCGStaticMeshSpawnerDataProvider::CreatePrimitiveDescriptors(FPCGContext*
 
 		for (const TPair<int32, uint32>& StringKeyAndInstanceCount : StringKeyToInstanceCount)
 		{
+			if (StringKeyAndInstanceCount.Value == 0)
+			{
+				continue;
+			}
+
 			if (!ensure(InBinding->GetStringTable().IsValidIndex(StringKeyAndInstanceCount.Key)))
 			{
 				continue;
@@ -563,6 +568,12 @@ void UPCGStaticMeshSpawnerDataProvider::CreatePrimitiveDescriptors(FPCGContext*
 			Descriptor.NumCustomFloats = CustomFloatCount;
 			Descriptor.StaticMesh = StaticMesh;
 
+			// Sanity check instance count.
+			if (!ensure(Descriptor.NumInstances <= TotalInputPointCount))
+			{
+				Descriptor.NumInstances = TotalInputPointCount;
+			}
+
 			PrimitiveStringKeys.Add(StringKeyAndInstanceCount.Key);
 			PrimitiveDescriptors.Add(MoveTemp(Descriptor));
 
@@ -627,7 +638,7 @@ void UPCGStaticMeshSpawnerDataProvider::CreatePrimitiveDescriptors(FPCGContext*
 				Descriptor.NumCustomFloats = CustomFloatCount;
 				Descriptor.StaticMesh = StaticMesh;
 
-				uint32 InstanceCount = FMath::CeilToInt(TotalInputPointCount * Weight);
+				int32 InstanceCount = FMath::CeilToInt(TotalInputPointCount * Weight);
 
 				if (SelectorWeighted->MeshEntries.Num() > 1)
 				{
