[PCG] Add additional yield in compute graph element execute between data binding initialization and compute graph instance initialization which are both currently heavy weight operations.

#rnx

diff --git a/Engine/Plugins/PCG/Source/PCG/Private/Compute/Elements/PCGComputeGraphElement.cpp b/Engine/Plugins/PCG/Source/PCG/Private/Compute/Elements/PCGComputeGraphElement.cpp
--- a/Engine/Plugins/PCG/Source/PCG/Private/Compute/Elements/PCGComputeGraphElement.cpp
+++ b/Engine/Plugins/PCG/Source/PCG/Private/Compute/Elements/PCGComputeGraphElement.cpp
@@ -159,6 +159,17 @@ bool FPCGComputeGraphElement::ExecuteInternal(FPCGContext* InContext) const
 			return true;
 		}
 
+		Context->ExecutionSubPhase = EPCGComputeGraphExecutionPhase::InitializeComputeGraphInstance;
+
+		if (Context->AsyncState.ShouldStop())
+		{
+			return false;
+		}
+	}
+	case EPCGComputeGraphExecutionPhase::InitializeComputeGraphInstance: // Fallthrough
+	{
+		TRACE_CPUPROFILER_EVENT_SCOPE(InitializeComputeGraphInstance);
+
 		const UPCGGraph* Graph = Context->Stack ? Context->Stack->GetNearestNonInlinedGraphForCurrentFrame() : nullptr;
 		const UPCGGraph::FComputeGraphInstanceKey Key = { Context->SourceComponent->GetGenerationGridSize(), ComputeGraphIndex };
 
diff --git a/Engine/Plugins/PCG/Source/PCG/Private/Compute/Elements/PCGComputeGraphElement.h b/Engine/Plugins/PCG/Source/PCG/Private/Compute/Elements/PCGComputeGraphElement.h
--- a/Engine/Plugins/PCG/Source/PCG/Private/Compute/Elements/PCGComputeGraphElement.h
+++ b/Engine/Plugins/PCG/Source/PCG/Private/Compute/Elements/PCGComputeGraphElement.h
@@ -19,6 +19,7 @@ enum class EPCGComputeGraphExecutionPhase : uint8
 	None,
 	GetComputeGraph,
 	InitializeDataBinding,
+	InitializeComputeGraphInstance,
 	PreExecuteReadbacks,
 	PrepareForExecute,
 	ValidateComputeGraphCompilation,
