[PCG] Fix ensure where GeneratedGraphOutputs were marked as garbage (from CL 40571447) when they should instead be moved to the transient package
- Make sure to clear GeneratedGraphOutputs in BeginDestroy for cases like Editor exit purge where we want to avoid renaming subobjects while GC is running

#rnx

diff --git a/Engine/Plugins/PCG/Source/PCG/Private/PCGComponent.cpp b/Engine/Plugins/PCG/Source/PCG/Private/PCGComponent.cpp
--- a/Engine/Plugins/PCG/Source/PCG/Private/PCGComponent.cpp
+++ b/Engine/Plugins/PCG/Source/PCG/Private/PCGComponent.cpp
@@ -1827,6 +1827,9 @@ void UPCGComponent::OnUnregister()
 
 void UPCGComponent::MarkSubObjectsAsGarbage()
 {
+	// First move out the generated output
+	ClearGraphGeneratedOutput();
+
 	auto ModifyAndMarkForGarbage = [](UObject* SubObject)
 	{
 		if (!SubObject)
@@ -2079,6 +2082,10 @@ bool UPCGComponent::CanEditChange(const FProperty* InProperty) const
 void UPCGComponent::BeginDestroy()
 {
 #if WITH_EDITOR
+	// BeginDestroy will end up calling OnComponentDestroyed and in this specific case we don't want to rename our Outputs to transient package.
+	// Known case where GeneratedGraphOutput can be non empty here is the editor exit purge.
+	GeneratedGraphOutput.Reset();
+
 	if (GraphInstance)
 	{
 		GraphInstance->OnGraphChangedDelegate.RemoveAll(this);
