Added map for easier execution metadata cleanup

#rnx

diff --git a/Engine/Plugins/PCG/Source/PCG/Private/PCGSubsystem.cpp b/Engine/Plugins/PCG/Source/PCG/Private/PCGSubsystem.cpp
--- a/Engine/Plugins/PCG/Source/PCG/Private/PCGSubsystem.cpp
+++ b/Engine/Plugins/PCG/Source/PCG/Private/PCGSubsystem.cpp
@@ -17,6 +17,7 @@
 #include "RuntimeGen/PCGGenSourceManager.h"
 #include "RuntimeGen/PCGRuntimeGenScheduler.h"
 
+#include "Algo/Transform.h"
 #include "Engine/Level.h"
 
 #include UE_INLINE_GENERATED_CPP_BY_NAME(PCGSubsystem)
@@ -539,28 +540,47 @@ void UPCGSubsystem::OnOriginalComponentUnregistered(UPCGComponent* InComponent)
 #if WITH_EDITOR
 void UPCGSubsystem::OnScheduleGraph(const FPCGStackContext& StackContext)
 {
+	TRACE_CPUPROFILER_EVENT_SCOPE(UPCGSubsystem::OnScheduleGraph);
 	// Always clear any possibly related warnings/errors on schedule.
 	for (const FPCGStack& Stack : StackContext.GetStacks())
 	{
 		GetNodeVisualLogsMutable().ClearLogs(Stack);
 	}
 
+	UPCGComponent* RootComponent = nullptr;
+
 	// Flush out all stacks that begin from the component / top graph as the existing dynamic
 	// stacks may not be occur during the next execution.
 	if (const FPCGStack* BaseStack = StackContext.GetStack(0))
 	{
+		RootComponent = const_cast<UPCGComponent*>(BaseStack->GetRootComponent());
+
 		if (BaseStack->IsCurrentFrameInRootGraph())
 		{
-			ClearExecutionMetadata(*BaseStack);
+			ClearExecutionMetadata(RootComponent);
 		}
 	}
 
 	// Record executed stacks.
 	{
+		TRACE_CPUPROFILER_EVENT_SCOPE(UPCGSubsystem::OnScheduleGraph::RecordExecutedStacks);
 		FWriteScopeLock Lock(ExecutedStacksLock);
+
+		TSoftObjectPtr<UPCGComponent> RootComponentSoftPtr(RootComponent);
+
 		for (const FPCGStack& ExecutedStack : StackContext.GetStacks())
 		{
-			ExecutedStacks.Add(ExecutedStack);
+			FPCGStackSharedPtr SharedStack(MakeShared<FPCGStack>(ExecutedStack));
+
+			bool bIsAlreadyInSet = false;
+
+			// Push to the set
+			ExecutedStacks.Add(SharedStack, &bIsAlreadyInSet);
+
+			if (!bIsAlreadyInSet && RootComponent)
+			{
+				ExecutedStacksPerComponent.FindOrAdd(RootComponentSoftPtr).Add(SharedStack);
+			}
 		}
 	}
 }
@@ -1850,35 +1870,64 @@ void UPCGSubsystem::NotifyLandscapeEditModeExited()
 	ActorAndComponentMapping.NotifyLandscapeEditModeExited();
 }
 
+/* UE_DEPRECATED(5.6) */
 void UPCGSubsystem::ClearExecutionMetadata(const FPCGStack& BaseStack)
 {
-	ClearExecutedStacks(BaseStack);
-
 	if (UPCGComponent* Component = const_cast<UPCGComponent*>(BaseStack.GetRootComponent()))
 	{
-		Component->ClearInspectionData();
-
-		NodeVisualLogs.ClearLogs(Component);
+		ClearExecutionMetadata(Component);
 	}
 }
 
 void UPCGSubsystem::ClearExecutionMetadata(UPCGComponent* InComponent)
 {
-	FPCGStack Stack;
-	Stack.PushFrame(InComponent);
-	ClearExecutionMetadata(Stack);
+	if (InComponent)
+	{
+		ClearExecutedStacks(InComponent);
+		InComponent->ClearInspectionData();
+		NodeVisualLogs.ClearLogs(InComponent);
+	}
 }
 
+/* UE_DEPRECATED(5.6) */
 TArray<FPCGStack> UPCGSubsystem::GetExecutedStacks(const UPCGComponent* InComponent, const UPCGGraph* InSubgraph, bool bOnlyWithSubgraphAsCurrentFrame)
 {
+	TArray<FPCGStackSharedPtr> StackPtrs = GetExecutedStacksPtrs(InComponent, InSubgraph, bOnlyWithSubgraphAsCurrentFrame);
+	TArray<FPCGStack> Stacks;
+	Stacks.Reserve(StackPtrs.Num());
+	Algo::Transform(StackPtrs, Stacks, [](const FPCGStackSharedPtr& StackPtr) { return *StackPtr; });
+
+	return Stacks;
+}
+
+TArray<FPCGStackSharedPtr> UPCGSubsystem::GetExecutedStacksPtrs(const UPCGComponent* InComponent, const UPCGGraph* InSubgraph, bool bOnlyWithSubgraphAsCurrentFrame)
+{
+	TRACE_CPUPROFILER_EVENT_SCOPE(UPCGSubsystem::GetExecutedStacksPtrs);
+	TArray<FPCGStackSharedPtr> MatchingStacks;
 	FReadScopeLock Lock(ExecutedStacksLock);
 
-	TArray<FPCGStack> MatchingStacks;
+	TSet<FPCGStackSharedPtr>* ComponentSet = nullptr;
+	TSet<FPCGStackSharedPtr>* GraphSet = nullptr;
+
+	if (InComponent)
+	{
+		ComponentSet = ExecutedStacksPerComponent.Find(InComponent);
+	}
+	
+	if(!ComponentSet)
+	{
+		ComponentSet = &ExecutedStacks;
+	}
+
+	auto StackMatches = [bOnlyWithSubgraphAsCurrentFrame, InSubgraph](const FPCGStackSharedPtr& Stack) -> bool
+	{
+		return ((bOnlyWithSubgraphAsCurrentFrame && Stack->GetGraphForCurrentFrame() == InSubgraph) ||
+			(!bOnlyWithSubgraphAsCurrentFrame && Stack->HasObject(InSubgraph)));
+	};
 
-	for (const FPCGStack& Stack : ExecutedStacks)
+	for (const FPCGStackSharedPtr& Stack : *ComponentSet)
 	{
-		if ((!InComponent || Stack.GetRootComponent() == InComponent) && 
-			((bOnlyWithSubgraphAsCurrentFrame && Stack.GetGraphForCurrentFrame() == InSubgraph) || (!bOnlyWithSubgraphAsCurrentFrame && Stack.HasObject(InSubgraph))))
+		if (StackMatches(Stack))
 		{
 			MatchingStacks.Add(Stack);
 		}
@@ -1887,14 +1936,40 @@ TArray<FPCGStack> UPCGSubsystem::GetExecutedStacks(const UPCGComponent* InCompon
 	return MatchingStacks;
 }
 
+/* UE_DEPRECATED(5.6) */
 TArray<FPCGStack> UPCGSubsystem::GetExecutedStacks(const FPCGStack& BeginningWithStack)
 {
+	TRACE_CPUPROFILER_EVENT_SCOPE(UPCGSubsystem::GetExecutedStacks);
+	TArray<FPCGStackSharedPtr> StackPtrs = GetExecutedStacksPtrs(BeginningWithStack);
+	TArray<FPCGStack> Stacks;
+
+	Stacks.Reserve(StackPtrs.Num());
+
+	for (FPCGStackSharedPtr StackPtr : StackPtrs)
+	{
+		Stacks.Add(*StackPtr);
+	}
+
+	return Stacks;
+}
+
+TArray<FPCGStackSharedPtr> UPCGSubsystem::GetExecutedStacksPtrs(const FPCGStack& BeginningWithStack)
+{
+	TRACE_CPUPROFILER_EVENT_SCOPE(UPCGSubsystem::GetExecutedStacksPtrsWithBeginning);
+	TArray<FPCGStackSharedPtr> MatchingStacks;
+
 	FReadScopeLock Lock(ExecutedStacksLock);
-	TArray<FPCGStack> MatchingStacks;
+	TSet<FPCGStackSharedPtr>* StacksToTest = ExecutedStacksPerComponent.Find(BeginningWithStack.GetRootComponent());
+
+	// Stack doesn't have a root component (unlikely), we'll check everything then
+	if (!StacksToTest)
+	{
+		StacksToTest = &ExecutedStacks;
+	}
 
-	for (const FPCGStack& Stack : ExecutedStacks)
+	for (const FPCGStackSharedPtr& Stack : *StacksToTest)
 	{
-		if (Stack.BeginsWith(BeginningWithStack))
+		if (Stack->BeginsWith(BeginningWithStack))
 		{
 			MatchingStacks.Add(Stack);
 		}
@@ -1903,32 +1978,52 @@ TArray<FPCGStack> UPCGSubsystem::GetExecutedStacks(const FPCGStack& BeginningWit
 	return MatchingStacks;
 }
 
-void UPCGSubsystem::ClearExecutedStacks(FPCGStack BeginningWithStack)
+void UPCGSubsystem::ClearExecutedStacks(const FPCGStack& BeginningWithStack)
 {
-	FWriteScopeLock Lock(ExecutedStacksLock);
+	return ClearExecutedStacks(BeginningWithStack.GetRootComponent());
+}
 
-	for (auto It = ExecutedStacks.CreateIterator(); It; ++It)
+void UPCGSubsystem::ClearExecutedStacks(const UPCGComponent* InRootComponent)
+{
+	TRACE_CPUPROFILER_EVENT_SCOPE(UPCGSubsystem::ClearExecutedStacksWithComponent);
+	if (InRootComponent)
 	{
-		// Clear any dynamic stack that starts with the given base stack, we don't know what stacks will
-		// execute so clean out everything.
-		if (It->BeginsWith(BeginningWithStack))
+		FWriteScopeLock Lock(ExecutedStacksLock);
+
+		if (TSet<FPCGStackSharedPtr>* AllStackForComponent = ExecutedStacksPerComponent.Find(InRootComponent))
 		{
-			It.RemoveCurrent();
+			for (const FPCGStackSharedPtr& Stack : *AllStackForComponent)
+			{
+				ExecutedStacks.Remove(Stack);
+			}
+
+			ExecutedStacksPerComponent.Remove(InRootComponent);
 		}
 	}
 }
 
 void UPCGSubsystem::ClearExecutedStacks(const UPCGGraph* InContainingGraph)
 {
-	FWriteScopeLock Lock(ExecutedStacksLock);
+	// TODO - it's fine to take more time here since it'll happen in rare(r) occurrences.
+	TRACE_CPUPROFILER_EVENT_SCOPE(UPCGSubsystem::ClearExecutedStacksWithGraph);
+	if (!InContainingGraph)
+	{
+		return;
+	}
 
-	for (auto It = ExecutedStacks.CreateIterator(); It; ++It)
+	FWriteScopeLock Lock(ExecutedStacksLock);
+	for(auto Iterator = ExecutedStacks.CreateIterator(); Iterator; ++Iterator)
 	{
-		// Clear any dynamic stack that starts with the given base stack, we don't know what stacks will
-		// execute so clean out everything.
-		if (It->HasObject(InContainingGraph))
+		FPCGStackSharedPtr StackPtr = *Iterator;
+
+		if (StackPtr->HasObject(InContainingGraph))
 		{
-			It.RemoveCurrent();
+			ExecutedStacks.Remove(StackPtr);
+
+			if (TSet<FPCGStackSharedPtr>* StacksForComponent = ExecutedStacksPerComponent.Find(StackPtr->GetRootComponent()))
+			{
+				StacksForComponent->Remove(StackPtr);
+			}
 		}
 	}
 }
diff --git a/Engine/Plugins/PCG/Source/PCG/Public/Graph/PCGStackContext.h b/Engine/Plugins/PCG/Source/PCG/Public/Graph/PCGStackContext.h
--- a/Engine/Plugins/PCG/Source/PCG/Public/Graph/PCGStackContext.h
+++ b/Engine/Plugins/PCG/Source/PCG/Public/Graph/PCGStackContext.h
@@ -211,6 +211,43 @@ public:
 #endif
 };
 
+struct FPCGStackSharedPtr : public TSharedPtr<FPCGStack>
+{
+	using Super = TSharedPtr<FPCGStack>;
+	using Super::operator bool;
+	using Super::operator->;
+	using Super::operator*;
+
+	FPCGStackSharedPtr() = default;
+	FPCGStackSharedPtr(TSharedPtr<FPCGStack>&& Other) : Super(Other) {}
+
+	bool operator==(const FPCGStackSharedPtr& Other) const
+	{
+		if (IsValid() && Other.IsValid())
+		{
+			return *Get() == *Other.Get();
+		}
+		else
+		{
+			return IsValid() == Other.IsValid();
+		}
+	}
+
+	bool operator!=(const FPCGStackSharedPtr& Other) const { return !operator==(Other); }
+
+	friend uint32 GetTypeHash(const FPCGStackSharedPtr& In)
+	{
+		if (In)
+		{
+			return GetTypeHash(*In);
+		}
+		else
+		{
+			return 0;
+		}
+	}
+};
+
 /** A collection of call stacks. */
 USTRUCT()
 struct PCG_API FPCGStackContext
diff --git a/Engine/Plugins/PCG/Source/PCG/Public/PCGSubsystem.h b/Engine/Plugins/PCG/Source/PCG/Public/PCGSubsystem.h
--- a/Engine/Plugins/PCG/Source/PCG/Public/PCGSubsystem.h
+++ b/Engine/Plugins/PCG/Source/PCG/Public/PCGSubsystem.h
@@ -350,14 +350,26 @@ public:
 	void NotifyLandscapeEditModeExited();
 
 	/** Clear any data collected during execution, normally called prior to generating the component. */
-	void ClearExecutionMetadata(const FPCGStack& BaseStack);
 	void ClearExecutionMetadata(UPCGComponent* InComponent);
 
+	UE_DEPRECATED(5.6, "Deprecated in favor of the version with the component pointer")
+	void ClearExecutionMetadata(const FPCGStack& BaseStack);
+
 	/** Get a list of stacks that were executed during the last execution. */
+	TArray<FPCGStackSharedPtr> GetExecutedStacksPtrs(const UPCGComponent* InComponent, const UPCGGraph* InSubgraph, bool bOnlyWithSubgraphAsCurrentFrame = true);
+	TArray<FPCGStackSharedPtr> GetExecutedStacksPtrs(const FPCGStack& BeginningWithStack);
+	
+	void ClearExecutedStacks(const UPCGComponent* InRootComponent);
+	void ClearExecutedStacks(const UPCGGraph* InContainingGraph);
+
+	UE_DEPRECATED(5.6, "Deprecated in favor of GetExecutedStacksPtrs")
 	TArray<FPCGStack> GetExecutedStacks(const UPCGComponent* InComponent, const UPCGGraph* InSubgraph, bool bOnlyWithSubgraphAsCurrentFrame = true);
+
+	UE_DEPRECATED(5.6, "Deprecated in favor of GetExecutedStacksPtrs")
 	TArray<FPCGStack> GetExecutedStacks(const FPCGStack& BeginningWithStack);
-	void ClearExecutedStacks(FPCGStack BeginningWithStack);
-	void ClearExecutedStacks(const UPCGGraph* InContainingGraph);
+
+	UE_DEPRECATED(5.6, "Deprecated in favor of ClearExecutedStacks with the component pointer")
+	void ClearExecutedStacks(const FPCGStack& BeginningWithStack);
 
 	UE_DEPRECATED(5.5, "Deprecated in favor of OnPCGComponentUnregistered, will not be notified anymore")
 	FPCGOnComponentUnregistered OnComponentUnregistered;
@@ -414,7 +426,8 @@ private:
 	FPCGActorAndComponentMapping ActorAndComponentMapping;
 
 	/** A record of stacks that were executed. Used to populate debugging tool UIs. */
-	TSet<FPCGStack> ExecutedStacks;
+	TSet<FPCGStackSharedPtr> ExecutedStacks;
+	TMap<TSoftObjectPtr<UPCGComponent>, TSet<FPCGStackSharedPtr>> ExecutedStacksPerComponent;
 	mutable FRWLock ExecutedStacksLock;
 
 	/** Functions will be executed at the beginning of the tick and then removed from this array. */
diff --git a/Engine/Plugins/PCG/Source/PCGEditor/Private/PCGEditor.cpp b/Engine/Plugins/PCG/Source/PCGEditor/Private/PCGEditor.cpp
--- a/Engine/Plugins/PCG/Source/PCGEditor/Private/PCGEditor.cpp
+++ b/Engine/Plugins/PCG/Source/PCGEditor/Private/PCGEditor.cpp
@@ -2991,7 +2991,7 @@ void FPCGEditor::OnComponentGenerationDone(UPCGSubsystem* Subsystem, UPCGCompone
 	// Additionally, if we are not inspecting but the component that's done executing contains this graph, then we should also update.
 	if (!bShouldRefresh && PCGComponentBeingInspected.IsExplicitlyNull() && PCGGraphBeingEdited && Subsystem)
 	{
-		TArray<FPCGStack> ExecutedStacksContainingThisGraph = Subsystem->GetExecutedStacks(Component, PCGGraphBeingEdited);
+		TArray<FPCGStackSharedPtr> ExecutedStacksContainingThisGraph = Subsystem->GetExecutedStacksPtrs(Component, PCGGraphBeingEdited);
 		bShouldRefresh |= !ExecutedStacksContainingThisGraph.IsEmpty();
 	}
 
diff --git a/Engine/Plugins/PCG/Source/PCGEditor/Private/Widgets/SPCGEditorGraphDebugObjectTree.cpp b/Engine/Plugins/PCG/Source/PCGEditor/Private/Widgets/SPCGEditorGraphDebugObjectTree.cpp
--- a/Engine/Plugins/PCG/Source/PCGEditor/Private/Widgets/SPCGEditorGraphDebugObjectTree.cpp
+++ b/Engine/Plugins/PCG/Source/PCGEditor/Private/Widgets/SPCGEditorGraphDebugObjectTree.cpp
@@ -773,10 +773,10 @@ FPCGEditorGraphDebugObjectItemPtr SPCGEditorGraphDebugObjectTree::GetFirstDebugO
 			// Look for graph in dynamic stacks.
 			if (Subsystem)
 			{
-				TArray<FPCGStack> ExecutedStacks = Subsystem->GetExecutedStacks(PCGComponent, PCGGraph);
-				for (const FPCGStack& Stack : ExecutedStacks)
+				TArray<FPCGStackSharedPtr> ExecutedStacks = Subsystem->GetExecutedStacksPtrs(PCGComponent, PCGGraph);
+				for (const FPCGStackSharedPtr& Stack : ExecutedStacks)
 				{
-					if (FPCGEditorGraphDebugObjectItemPtr Item = GetItemFromStack(Stack))
+					if (FPCGEditorGraphDebugObjectItemPtr Item = GetItemFromStack(*Stack))
 					{
 						return Item;
 					}
@@ -793,31 +793,59 @@ void SPCGEditorGraphDebugObjectTree::UpdateIsSetDebugObjectFromSelectionEnabled(
 	IsSetDebugObjectFromSelectionEnabled = (GetFirstDebugObjectFromSelection() != nullptr);
 }
 
+void SPCGEditorGraphDebugObjectTree::AddStacksToTree(const TArray<FPCGStackSharedPtr>& Stacks,
+	TMap<AActor*, TSharedPtr<FPCGEditorGraphDebugObjectItem_Actor>>& InOutActorItems,
+	TMap<const FPCGStack, FPCGEditorGraphDebugObjectItemPtr>& InOutStackToItem)
+{
+	TRACE_CPUPROFILER_EVENT_SCOPE(SPCGEditorGraphDebugObjectTree::AddStacksToTree);
+	const UPCGGraph* GraphBeingEdited = GetPCGGraph();
+	if (!GraphBeingEdited)
+	{
+		return;
+	}
+
+	for (const FPCGStackSharedPtr& Stack : Stacks)
+	{
+		AddStackToTree(*Stack, GraphBeingEdited, InOutActorItems, InOutStackToItem);
+	}
+}
+
 void SPCGEditorGraphDebugObjectTree::AddStacksToTree(const TArray<FPCGStack>& Stacks,
 	TMap<AActor*, TSharedPtr<FPCGEditorGraphDebugObjectItem_Actor>>& InOutActorItems,
 	TMap<const FPCGStack, FPCGEditorGraphDebugObjectItemPtr>& InOutStackToItem)
 {
-	//TRACE_CPUPROFILER_EVENT_SCOPE(SPCGEditorGraphDebugObjectTree::AddStacksToTree);
+	TRACE_CPUPROFILER_EVENT_SCOPE(SPCGEditorGraphDebugObjectTree::AddStacksToTree);
 	const UPCGGraph* GraphBeingEdited = GetPCGGraph();
-	UPCGSubsystem* Subsystem = PCGEditor.Pin()->GetSubsystem();
-	if (!GraphBeingEdited || !Subsystem)
+	if (!GraphBeingEdited)
 	{
 		return;
 	}
 
 	for (const FPCGStack& Stack : Stacks)
+	{
+		AddStackToTree(Stack, GraphBeingEdited, InOutActorItems, InOutStackToItem);
+	}
+}
+
+void SPCGEditorGraphDebugObjectTree::AddStackToTree(const FPCGStack& Stack,
+	const UPCGGraph* GraphBeingEdited,
+	TMap<AActor*, TSharedPtr<FPCGEditorGraphDebugObjectItem_Actor>>& InOutActorItems,
+	TMap<const FPCGStack, FPCGEditorGraphDebugObjectItemPtr>& InOutStackToItem)
+{
+	check(GraphBeingEdited);
+
 	{
 		//TRACE_CPUPROFILER_EVENT_SCOPE(SPCGEditorGraphDebugObjectTree::AddStacksToTree::PerStackLoop);
 		// If the current graph is not in the stack at all then skip.
 		if (!Stack.HasObject(GraphBeingEdited))
 		{
-			continue;
+			return;
 		}
 
 		UPCGComponent* PCGComponent = const_cast<UPCGComponent*>(Stack.GetRootComponent());
 		if (!PCGComponent)
 		{
-			continue;
+			return;
 		}
 
 		// Prevent duplicate entries from the editor world while in PIE.
@@ -827,7 +855,7 @@ void SPCGEditorGraphDebugObjectTree::AddStacksToTree(const TArray<FPCGStack>& St
 			{
 				if (!World->IsGameWorld())
 				{
-					continue;
+					return;
 				}
 			}
 		}
@@ -836,7 +864,7 @@ void SPCGEditorGraphDebugObjectTree::AddStacksToTree(const TArray<FPCGStack>& St
 		UPCGGraph* TopGraph = const_cast<UPCGGraph*>(Stack.GetRootGraph(&TopGraphIndex));
 		if (!TopGraph)
 		{
-			continue;
+			return;
 		}
 
 		// If we're inspecting a node which has not logged inspection data in a previous execution, display grayed out.
@@ -854,7 +882,7 @@ void SPCGEditorGraphDebugObjectTree::AddStacksToTree(const TArray<FPCGStack>& St
 		AActor* Actor = PCGComponent->GetOwner();
 		if (!Actor)
 		{
-			continue;
+			return;
 		}
 
 		// Add actor item if not already added.
@@ -1033,21 +1061,11 @@ void SPCGEditorGraphDebugObjectTree::RefreshTree()
 	{
 		TMap<AActor*, TSharedPtr<FPCGEditorGraphDebugObjectItem_Actor>> ActorItems;
 		TMap<const FPCGStack, FPCGEditorGraphDebugObjectItemPtr> StackToItem;
-		TMap<const UPCGComponent*, TArray<FPCGStack>> ComponentsGraphStacks;
 
 		using FGraphContextMapKey = TTuple<UPCGGraph*, uint32 /*GenerationGridSize*/, bool /*IsPartitioned*/>;
 		using FGraphContextMapValue = TTuple<FPCGStackContext, bool /*bContainsGraph*/>;
 		TMap<FGraphContextMapKey, FGraphContextMapValue> GraphToContexts;
 
-		{
-			TRACE_CPUPROFILER_EVENT_SCOPE(SPCGEditorGraphDebugObjectTree::RefreshTree::GetExecutedStacks);
-			TArray<FPCGStack> GraphStacks = Subsystem->GetExecutedStacks(nullptr, PCGGraph, /*bOnlyWithSubgraphAsCurrentFrame=*/false);
-			for (FPCGStack& GraphStack : GraphStacks)
-			{
-				ComponentsGraphStacks.FindOrAdd(GraphStack.GetRootComponent()).Add(MoveTemp(GraphStack));
-			}
-		}
-
 		for (UObject* PCGComponentObject : PCGComponents)
 		{
 			if (!IsValid(PCGComponentObject))
@@ -1090,14 +1108,15 @@ void SPCGEditorGraphDebugObjectTree::RefreshTree()
 					AddStacksToTree(ComponentContext.GetStacks(), ActorItems, StackToItem);
 				}
 			}
+		}
 
-			// Process stacks encountered during execution so far, which will include dynamic subgraphs & loop subgraphs.
-			// There will be overlaps with the static stacks but only unique entries will be added to the tree.
-			if (TArray<FPCGStack>* ComponentGraphStacks = ComponentsGraphStacks.Find(PCGComponent))
-			{
-				//TRACE_CPUPROFILER_EVENT_SCOPE(SPCGEditorGraphDebugObjectTree::RefreshTree::DynamicStacks);
-				AddStacksToTree(*ComponentGraphStacks, ActorItems, StackToItem);
-			}
+		// Process stacks encountered during execution so far, which will include dynamic subgraphs & loop subgraphs.
+		// There will be overlaps with the static stacks but only unique entries will be added to the tree.
+		{
+			TRACE_CPUPROFILER_EVENT_SCOPE(SPCGEditorGraphDebugObjectTree::RefreshTree::DynamicStacks);
+			TArray<FPCGStackSharedPtr> GraphStacks = Subsystem->GetExecutedStacksPtrs(nullptr, PCGGraph, /*bOnlyWithSubgraphAsCurrentFrame=*/false);
+			// TODO: do per component filtering?
+			AddStacksToTree(GraphStacks, ActorItems, StackToItem);
 		}
 
 		for (TPair<AActor*, TSharedPtr<FPCGEditorGraphDebugObjectItem_Actor>>& ActorItem : ActorItems)
diff --git a/Engine/Plugins/PCG/Source/PCGEditor/Private/Widgets/SPCGEditorGraphDebugObjectTree.h b/Engine/Plugins/PCG/Source/PCGEditor/Private/Widgets/SPCGEditorGraphDebugObjectTree.h
--- a/Engine/Plugins/PCG/Source/PCGEditor/Private/Widgets/SPCGEditorGraphDebugObjectTree.h
+++ b/Engine/Plugins/PCG/Source/PCGEditor/Private/Widgets/SPCGEditorGraphDebugObjectTree.h
@@ -290,10 +290,18 @@ private:
 	void SortTreeItems(bool bIsAscending = true, bool bIsRecursive = true);
 	void RestoreTreeState();
 
+	void AddStacksToTree(
+		const TArray<FPCGStackSharedPtr>& Stacks,
+		TMap<AActor*, TSharedPtr<FPCGEditorGraphDebugObjectItem_Actor>>& InOutActorItems,
+		TMap<const FPCGStack, FPCGEditorGraphDebugObjectItemPtr>& InOutStackToItem);
 	void AddStacksToTree(
 		const TArray<FPCGStack>& Stacks,
 		TMap<AActor*, TSharedPtr<FPCGEditorGraphDebugObjectItem_Actor>>& InOutActorItems,
 		TMap<const FPCGStack, FPCGEditorGraphDebugObjectItemPtr>& InOutStackToItem);
+	void AddStackToTree(const FPCGStack& Stack,
+		const UPCGGraph* GraphBeingEdited,
+		TMap<AActor*, TSharedPtr<FPCGEditorGraphDebugObjectItem_Actor>>& InOutActorItems,
+		TMap<const FPCGStack, FPCGEditorGraphDebugObjectItemPtr>& InOutStackToItem);
 
 	void OnEditorSelectionChanged(UObject* InObject);
 
diff --git a/Engine/Plugins/PCG/Source/PCGEditor/Private/Widgets/SPCGEditorGraphFind.cpp b/Engine/Plugins/PCG/Source/PCGEditor/Private/Widgets/SPCGEditorGraphFind.cpp
--- a/Engine/Plugins/PCG/Source/PCGEditor/Private/Widgets/SPCGEditorGraphFind.cpp
+++ b/Engine/Plugins/PCG/Source/PCGEditor/Private/Widgets/SPCGEditorGraphFind.cpp
@@ -674,10 +674,12 @@ void SPCGEditorGraphFind::MatchTokens(const TArray<FString>& InTokens)
 	if (PCGEditor->GetStackBeingInspected() && PCGEditorPtr.Pin()->GetSubsystem())
 	{
 		const FPCGStack& CurrentStack = *PCGEditor->GetStackBeingInspected();
-		TArray<FPCGStack> Substacks = PCGEditorPtr.Pin()->GetSubsystem()->GetExecutedStacks(CurrentStack);
+		TArray<FPCGStackSharedPtr> Substacks = PCGEditorPtr.Pin()->GetSubsystem()->GetExecutedStacksPtrs(CurrentStack);
 		// For each stack, find the first subgraph node from the start, then get all graphs under it.
-		for (const FPCGStack& Substack : Substacks)
+		for (const FPCGStackSharedPtr& SubstackPtr : Substacks)
 		{
+			const FPCGStack& Substack = *SubstackPtr;
+
 			if (Substack.GetStackFrames().Num() <= CurrentStack.GetStackFrames().Num())
 			{
 				continue;
