[PCG] Add more information in the CRC for Get Actor Data and Get Actor Property

* Add tags  in Self/Original case, since we add the tags to the output data.
* Add actor reference when we output it, as we could get the wrong reference if the actor data was identical between multiple actors (same volume same place in the world).
* Also add the option to always requery actors in Get Actor Data

#rnx

diff --git a/Engine/Plugins/PCG/Source/PCG/Private/Elements/PCGDataFromActor.cpp b/Engine/Plugins/PCG/Source/PCG/Private/Elements/PCGDataFromActor.cpp
--- a/Engine/Plugins/PCG/Source/PCG/Private/Elements/PCGDataFromActor.cpp
+++ b/Engine/Plugins/PCG/Source/PCG/Private/Elements/PCGDataFromActor.cpp
@@ -15,6 +15,7 @@
 #include "Elements/PCGMergeAttributes.h"
 #include "Grid/PCGPartitionActor.h"
 #include "Helpers/PCGDynamicTrackingHelpers.h"
+#include "Helpers/PCGActorHelpers.h"
 #include "Helpers/PCGHelpers.h"
 #include "Helpers/PCGTagHelpers.h"
 #include "Utils/PCGGraphExecutionLogging.h"
@@ -137,6 +138,13 @@ namespace PCGDataFromActorHelpers
 #if WITH_EDITOR
 void UPCGDataFromActorSettings::GetStaticTrackedKeys(FPCGSelectionKeyToSettingsMap& OutKeysToSettings, TArray<TObjectPtr<const UPCGGraph>>& OutVisitedGraphs) const
 {
+	// If input pin is connected, tracking is dynamic.
+	const UPCGNode* Node = Cast<const UPCGNode>(GetOuter());
+	if (Node && Node->IsInputPinConnected(PCGPinConstants::DefaultInputLabel))
+	{
+		return;
+	}
+
 	FPCGSelectionKey Key = ActorSelector.GetAssociatedKey();
 	if (Mode == EPCGGetDataFromActorMode::GetDataFromPCGComponent || Mode == EPCGGetDataFromActorMode::GetDataFromPCGComponentOrParseComponents)
 	{
@@ -1018,6 +1026,18 @@ void FPCGDataFromActorElement::ProcessActor(FPCGContext* InContext, const UPCGDa
 	}
 }
 
+bool FPCGDataFromActorElement::IsCacheable(const UPCGSettings* InSettings) const
+{
+	if (const UPCGDataFromActorSettings* Settings = Cast<UPCGDataFromActorSettings>(InSettings))
+	{
+		return !Settings->bAlwaysRequeryActors;
+	}
+	else
+	{
+		return false;
+	}
+}
+
 void FPCGDataFromActorElement::GetDependenciesCrc(const FPCGDataCollection& InInput, const UPCGSettings* InSettings, UPCGComponent* InComponent, FPCGCrc& OutCrc) const
 {
 	FPCGCrc Crc;
@@ -1026,6 +1046,11 @@ void FPCGDataFromActorElement::GetDependenciesCrc(const FPCGDataCollection& InIn
 	// If we track self or original, we are dependent on the actor data
 	if (const UPCGDataFromActorSettings* Settings = Cast<const UPCGDataFromActorSettings>(InSettings))
 	{
+		if (Settings->bAlwaysRequeryActors)
+		{
+			return;
+		}
+		
 		const bool bDependsOnSelfOrHierarchy = (Settings->ActorSelector.ActorFilter == EPCGActorFilter::Self || Settings->ActorSelector.ActorFilter == EPCGActorFilter::Original);
 		const bool bDependsOnSelfBounds = Settings->ActorSelector.bMustOverlapSelf;
 
@@ -1037,6 +1062,21 @@ void FPCGDataFromActorElement::GetDependenciesCrc(const FPCGDataCollection& InIn
 			if (ActorData)
 			{
 				Crc.Combine(ActorData->GetOrComputeCrc(/*bFullDataCrc=*/false));
+
+				// Also if it is depending on Self, we need to CRC the tags (as they will be passed to the output data)
+				// Same with anything that would pull the actor reference.
+				if (bDependsOnSelfOrHierarchy && ComponentToCheck)
+				{
+					if (const AActor* Owner = ComponentToCheck->GetOwner())
+					{
+						const bool bOutputActorReference = Settings->bAlsoOutputSinglePointData
+							|| Settings->Mode == EPCGGetDataFromActorMode::GetActorReference
+							|| Settings->Mode == EPCGGetDataFromActorMode::GetComponentsReference
+							|| Settings->Mode == EPCGGetDataFromActorMode::GetSinglePoint;
+						
+						Crc.Combine(UPCGActorHelpers::ComputeHashFromActorTagsAndReference(Owner, /*bIncludeTags=*/true, bOutputActorReference));
+					}
+				}
 			}
 		}
 
diff --git a/Engine/Plugins/PCG/Source/PCG/Private/Elements/PCGGetActorProperty.cpp b/Engine/Plugins/PCG/Source/PCG/Private/Elements/PCGGetActorProperty.cpp
--- a/Engine/Plugins/PCG/Source/PCG/Private/Elements/PCGGetActorProperty.cpp
+++ b/Engine/Plugins/PCG/Source/PCG/Private/Elements/PCGGetActorProperty.cpp
@@ -7,6 +7,7 @@
 #include "PCGComponent.h"
 #include "PCGParamData.h"
 #include "Data/PCGPointData.h"
+#include "Helpers/PCGActorHelpers.h"
 #include "Helpers/PCGBlueprintHelpers.h"
 #include "Helpers/PCGDynamicTrackingHelpers.h"
 #include "Helpers/PCGPropertyHelpers.h"
@@ -164,6 +165,35 @@ bool FPCGGetActorPropertyElement::IsCacheable(const UPCGSettings* InSettings) co
 	}
 }
 
+void FPCGGetActorPropertyElement::GetDependenciesCrc(const FPCGDataCollection& InInput, const UPCGSettings* InSettings, UPCGComponent* InComponent,
+	FPCGCrc& OutCrc) const
+{
+	FPCGCrc Crc;
+	IPCGElement::GetDependenciesCrc(InInput, InSettings, InComponent, Crc);
+
+	// If we track self or original, we are dependent on the actor data
+	if (const UPCGGetActorPropertySettings* Settings = Cast<const UPCGGetActorPropertySettings>(InSettings))
+	{
+		if (Settings->bAlwaysRequeryActors)
+		{
+			return;
+		}
+
+		const bool bDependsOnSelfOrHierarchy = (Settings->ActorSelector.ActorFilter == EPCGActorFilter::Self || Settings->ActorSelector.ActorFilter == EPCGActorFilter::Original);
+		UPCGComponent* ComponentToCheck = (Settings->ActorSelector.ActorFilter == EPCGActorFilter::Original) ? InComponent->GetOriginalComponent() : InComponent;
+
+		// If it is depending on Self, we need to CRC the tags (as they will be passed to the output data)
+		// and also the owner path.
+		if (bDependsOnSelfOrHierarchy && ComponentToCheck)
+		{
+			if (const AActor* Owner = ComponentToCheck->GetOwner())
+			{
+				Crc.Combine(UPCGActorHelpers::ComputeHashFromActorTagsAndReference(Owner, /*bIncludeTags=*/true, /*bIncludeActorReference=*/true));
+			}
+		}
+	}
+}
+
 bool FPCGGetActorPropertyElement::PrepareDataInternal(FPCGContext* Context) const
 {
 	TRACE_CPUPROFILER_EVENT_SCOPE(FPCGGetActorPropertyElement::PrepareData);
diff --git a/Engine/Plugins/PCG/Source/PCG/Private/Helpers/PCGActorHelpers.cpp b/Engine/Plugins/PCG/Source/PCG/Private/Helpers/PCGActorHelpers.cpp
--- a/Engine/Plugins/PCG/Source/PCG/Private/Helpers/PCGActorHelpers.cpp
+++ b/Engine/Plugins/PCG/Source/PCG/Private/Helpers/PCGActorHelpers.cpp
@@ -805,3 +805,27 @@ FIntVector UPCGActorHelpers::GetCellCoord(FVector InPosition, int InGridSize, bo
 		bUse2DGrid ? 0 : FMath::FloorToInt(Temp.Z)
 	);
 }
+
+int UPCGActorHelpers::ComputeHashFromActorTagsAndReference(const AActor* InActor, const bool bIncludeTags, const bool bIncludeActorReference)
+{
+	if (!InActor)
+	{
+		return 0;
+	}
+
+	int Result = 0;
+	if (bIncludeTags)
+	{
+		for (const FName Tag : InActor->Tags)
+		{
+			Result = HashCombineFast(Result, GetTypeHash(Tag));
+		}
+	}
+
+	if (bIncludeActorReference)
+	{
+		Result = HashCombineFast(Result, GetTypeHash(InActor->GetPathName()));
+	}
+
+	return Result;
+}
diff --git a/Engine/Plugins/PCG/Source/PCG/Public/Elements/PCGDataFromActor.h b/Engine/Plugins/PCG/Source/PCG/Public/Elements/PCGDataFromActor.h
--- a/Engine/Plugins/PCG/Source/PCG/Public/Elements/PCGDataFromActor.h
+++ b/Engine/Plugins/PCG/Source/PCG/Public/Elements/PCGDataFromActor.h
@@ -121,6 +121,10 @@ public:
 	UPROPERTY(BlueprintReadWrite, EditAnywhere, Category = "Data Retrieval Settings", meta = (EditCondition = "Mode == EPCGGetDataFromActorMode::GetDataFromProperty", EditConditionHides))
 	FName PropertyName = NAME_None;
 
+	/** If this is true, we will never put this element in cache, and will always try to re-query the actors and read the latest data from them. */
+	UPROPERTY(BlueprintReadWrite, EditAnywhere, Category = "Data Retrieval Settings")
+	bool bAlwaysRequeryActors = false;
+
 	/** Silence warnings that attribute names were sanitized to replace invalid characters. */
 	UPROPERTY(BlueprintReadWrite, EditAnywhere, Category = "Data Retrieval Settings|Advanced")
 	bool bSilenceSanitizedAttributeNameWarnings = false;
@@ -155,6 +159,7 @@ struct FPCGDataFromActorContext : public FPCGLoadObjectsFromPathContext
 class PCG_API FPCGDataFromActorElement : public IPCGElementWithCustomContext<FPCGDataFromActorContext>
 {
 public:
+	virtual bool IsCacheable(const UPCGSettings* InSettings) const override;
 	virtual bool CanExecuteOnlyOnMainThread(FPCGContext* Context) const override { return true; }
 	virtual void GetDependenciesCrc(const FPCGDataCollection& InInput, const UPCGSettings* InSettings, UPCGComponent* InComponent, FPCGCrc& OutCrc) const override;
 	virtual bool ShouldComputeFullOutputDataCrc(FPCGContext* Context) const override { return true; }
diff --git a/Engine/Plugins/PCG/Source/PCG/Public/Elements/PCGGetActorProperty.h b/Engine/Plugins/PCG/Source/PCG/Public/Elements/PCGGetActorProperty.h
--- a/Engine/Plugins/PCG/Source/PCG/Public/Elements/PCGGetActorProperty.h
+++ b/Engine/Plugins/PCG/Source/PCG/Public/Elements/PCGGetActorProperty.h
@@ -124,6 +124,7 @@ class FPCGGetActorPropertyElement : public IPCGElementWithCustomContext<FPCGLoad
 public:
 	virtual bool CanExecuteOnlyOnMainThread(FPCGContext* Context) const override { return true; }
 	virtual bool IsCacheable(const UPCGSettings* InSettings) const override;
+	virtual void GetDependenciesCrc(const FPCGDataCollection& InInput, const UPCGSettings* InSettings, UPCGComponent* InComponent, FPCGCrc& OutCrc) const override;
 protected:
 	virtual bool PrepareDataInternal(FPCGContext* Context) const override;
 	virtual bool ExecuteInternal(FPCGContext* Context) const override;
diff --git a/Engine/Plugins/PCG/Source/PCG/Public/Helpers/PCGActorHelpers.h b/Engine/Plugins/PCG/Source/PCG/Public/Helpers/PCGActorHelpers.h
--- a/Engine/Plugins/PCG/Source/PCG/Public/Helpers/PCGActorHelpers.h
+++ b/Engine/Plugins/PCG/Source/PCG/Public/Helpers/PCGActorHelpers.h
@@ -229,4 +229,7 @@ PRAGMA_ENABLE_DEPRECATION_WARNINGS
 	 * Return the grid cell coordinates on the PCG partition grid given a position and the grid size.
 	 */
 	static FIntVector GetCellCoord(FVector InPosition, int InGridSize, bool bUse2DGrid);
+
+	/** Extract the tags and the actor reference of the given actor and hash it. Useful for CRC dependencies that depends on the tags or the instance of the actor.*/
+	static int ComputeHashFromActorTagsAndReference(const AActor* InActor, const bool bIncludeTags, const bool bIncludeActorReference);
 };
