[PCG] Temporary fix for the flatten on generated output

* By default, will also use the `bUsedMultipleTimes` to duplicate data that can be used multiple times in the graph, to avoid memory corruption
* Also add a CVar if this is not enough to completely disable the stealing.
* Still not a fully acceptable solution, will fix it fully for 5.6

#rnx

diff --git a/Engine/Plugins/PCG/Source/PCG/Private/Graph/PCGGraphCompiler.cpp b/Engine/Plugins/PCG/Source/PCG/Private/Graph/PCGGraphCompiler.cpp
--- a/Engine/Plugins/PCG/Source/PCG/Private/Graph/PCGGraphCompiler.cpp
+++ b/Engine/Plugins/PCG/Source/PCG/Private/Graph/PCGGraphCompiler.cpp
@@ -1476,7 +1476,8 @@ void FPCGGraphCompiler::CompileTopGraph(UPCGGraph* InGraph, uint32 GenerationGri
 			// It is necessary for any post generation task to get the content of the output node
 			// and only this content.
 			const bool bProvideData = Task.Node == GraphOutputNode;
-			PostExecuteTask.Inputs.Emplace(Task.NodeId, /*InUpstreamPin=*/FPCGGraphTaskInput::NoPin, /*InDownstreamPin=*/FPCGGraphTaskInput::NoPin, bProvideData);
+			FPCGGraphTaskInput& TaskInput = PostExecuteTask.Inputs.Emplace_GetRef(Task.NodeId, /*InUpstreamPin=*/FPCGGraphTaskInput::NoPin, /*InDownstreamPin=*/FPCGGraphTaskInput::NoPin, bProvideData);
+			TaskInput.bIsUsedMultipleTimes = false;
 		}
 
 		for (const FPCGGraphTaskInput& Input : Task.Inputs)
diff --git a/Engine/Plugins/PCG/Source/PCG/Private/Graph/PCGGraphExecutor.cpp b/Engine/Plugins/PCG/Source/PCG/Private/Graph/PCGGraphExecutor.cpp
--- a/Engine/Plugins/PCG/Source/PCG/Private/Graph/PCGGraphExecutor.cpp
+++ b/Engine/Plugins/PCG/Source/PCG/Private/Graph/PCGGraphExecutor.cpp
@@ -762,7 +762,9 @@ FPCGTaskId FPCGGraphExecutor::ScheduleSingleElement(TSharedRef<IPCGElement> Elem
 	for (FPCGTaskId TaskDependency : TaskDataDependencies)
 	{
 		ensure(TaskDependency != InvalidPCGTaskId);
-		Task.Inputs.Emplace(TaskDependency);
+		// By default the input of a generic element is not used multiple times.
+		FPCGGraphTaskInput& TaskInput = Task.Inputs.Emplace_GetRef(TaskDependency);
+		TaskInput.bIsUsedMultipleTimes = false;
 	}
 
 	Task.SourceComponent = InSourceComponent;
diff --git a/Engine/Plugins/PCG/Source/PCG/Private/PCGComponent.cpp b/Engine/Plugins/PCG/Source/PCG/Private/PCGComponent.cpp
--- a/Engine/Plugins/PCG/Source/PCG/Private/PCGComponent.cpp
+++ b/Engine/Plugins/PCG/Source/PCG/Private/PCGComponent.cpp
@@ -77,6 +77,11 @@ namespace PCGComponent
 		false,
 		TEXT("In Editor and with runtime gen, a change with one tracked element will refresh all the local components. If it is too resource intensive, it can be disabled."));
 #endif
+	
+	static TAutoConsoleVariable<bool> CVarDisableStealingForFlattenInPostProcessGraph(
+		TEXT("pcg.Advanced.DisableStealingForFlattenInPostProcessGraph"),
+		false,
+		TEXT("At the end of execution, generated output data is flattened in place if it can be stolen, but it can be disabled if it has unwanted side effects."));
 
 	struct FPCGManagedActorLoadingScope
 	{
@@ -676,23 +681,18 @@ void UPCGComponent::PostProcessGraph(const FBox& InNewBounds, bool bInGenerated,
 			// TODO: should we filter based on supported serialized types here?
 			for (const FPCGTaggedData& TaggedData : Context->InputData.TaggedData)
 			{
-				if (ensure(TaggedData.Data))
+				// Proxies should never get cached on the component output. These can hold onto large chunks of video memory and are not
+				// directly serializable.
+				if (ensure(TaggedData.Data && TaggedData.Data->CanBeSerialized()))
 				{
-					// Proxies should never get cached on the component output. These can hold onto large chunks of video memory and are not
-					// directly serializable.
-					for (const FPCGTaggedData& Data : Context->InputData.TaggedData)
-					{
-						ensure(!Data.Data || Data.Data->CanBeSerialized());
-					}
-
 					// Visit the generated data; if the outer is the transient package or this component for all data, then
 					// we don't need to duplicate the data and can change the outer & flatten the data without any additional copies.
 					// Caveat: since that data could be in the cache currently, we shouldn't reouter it to a component that's not in the persistent level, otherwise when the
 					// sublevel is unloaded, it could leak references.
+					// We also need to duplicate the data if the data is used multiple times.
+					bool bDataCanBeStolen = !TaggedData.bIsUsedMultipleTimes && !PCGComponent::CVarDisableStealingForFlattenInPostProcessGraph.GetValueOnAnyThread();
 #if WITH_EDITOR
-					bool bDataCanBeStolen = GetOwner() && GetOwner()->GetWorld() && (GetOwner()->GetLevel() == GetOwner()->GetWorld()->PersistentLevel);
-#else
-					bool bDataCanBeStolen = true;
+					bDataCanBeStolen &= (GetOwner() && GetOwner()->GetWorld() && (GetOwner()->GetLevel() == GetOwner()->GetWorld()->PersistentLevel));
 #endif
 
 					if (bDataCanBeStolen)
