Cleaned up pin type retrieval in metadata ops to prevent combinatorial explosion

#rnx

diff --git a/Engine/Plugins/PCG/Source/PCG/Private/Elements/Metadata/PCGMetadataOpElementBase.cpp b/Engine/Plugins/PCG/Source/PCG/Private/Elements/Metadata/PCGMetadataOpElementBase.cpp
--- a/Engine/Plugins/PCG/Source/PCG/Private/Elements/Metadata/PCGMetadataOpElementBase.cpp
+++ b/Engine/Plugins/PCG/Source/PCG/Private/Elements/Metadata/PCGMetadataOpElementBase.cpp
@@ -136,7 +136,7 @@ bool UPCGMetadataSettingsBase::IsPinDefaultValueActivated(const FName PinLabel)
 	return false;
 }
 
-uint32 UPCGMetadataSettingsBase::GetInputPinToForward() const
+uint32 UPCGMetadataSettingsBase::GetInputPinToForward(EPCGDataType* OutDataType) const
 {
 	const uint32 OperandNum = GetOperandNum();
 	uint32 InputPinToForward = 0;
@@ -162,6 +162,12 @@ uint32 UPCGMetadataSettingsBase::GetInputPinToForward() const
 				if ((PinType != EPCGDataType::Any) && !!(PinType & EPCGDataType::Spatial))
 				{
 					InputPinToForward = InputPinIndex;
+
+					if (OutDataType)
+					{
+						*OutDataType = PinType;
+					}
+
 					break;
 				}
 			}
@@ -200,8 +206,12 @@ EPCGDataType UPCGMetadataSettingsBase::GetCurrentPinTypes(const UPCGPin* InPin)
 		return Super::GetCurrentPinTypes(InPin);
 	}
 
-	// Output pin narrows to union of inputs on pin to forward
-	return GetInputPinType(GetInputPinToForward());
+	// Output pin narrows to union of inputs on pin to forward.
+	// Implementation note - getting the input type of the pin to forward like 'GetInputPinType(GetInputPinToForward())' 
+	// causes a combinatorial explosion in long dynamic type chains, hence why we retrieve the type from the other call.
+	EPCGDataType PinTypeOfInputPinToForward = EPCGDataType::None;
+	GetInputPinToForward(&PinTypeOfInputPinToForward);
+	return PinTypeOfInputPinToForward;
 }
 
 #if WITH_EDITOR
diff --git a/Engine/Plugins/PCG/Source/PCG/Public/Elements/Metadata/PCGMetadataOpElementBase.h b/Engine/Plugins/PCG/Source/PCG/Public/Elements/Metadata/PCGMetadataOpElementBase.h
--- a/Engine/Plugins/PCG/Source/PCG/Public/Elements/Metadata/PCGMetadataOpElementBase.h
+++ b/Engine/Plugins/PCG/Source/PCG/Public/Elements/Metadata/PCGMetadataOpElementBase.h
@@ -177,8 +177,8 @@ public:
 	UE_DEPRECATED(5.6, "Replaced by inline DefaultValue system. Use CreateDefaultValueParamData instead.")
 	virtual UPCGParamData* CreateDefaultValueParam(FPCGContext* Context, uint32 Index) const { return nullptr; }
 
-	/* Return the current input pin to forward to the output. */
-	uint32 GetInputPinToForward() const;
+	/* Return the current input pin to forward to the output, and optionally the type of the input pin found. */
+	uint32 GetInputPinToForward(EPCGDataType* OutDataType = nullptr) const;
 
 	UPROPERTY(BlueprintReadWrite, EditAnywhere, Category = "Output", meta = (PCG_Overridable))
 	FPCGAttributePropertyOutputSelector OutputTarget;
