[PCG] Spawn Spline Mesh Improvements

* Add an output pin for execution dependency
* Supports time slicing
* Add override per control points
* Add option on the Time Sliced eleemnt to derive from another context (useful for async)
* Add option to get all objects to load from override.

diff --git a/Engine/Plugins/PCG/Source/PCG/Private/Elements/PCGSpawnSplineMesh.cpp b/Engine/Plugins/PCG/Source/PCG/Private/Elements/PCGSpawnSplineMesh.cpp
--- a/Engine/Plugins/PCG/Source/PCG/Private/Elements/PCGSpawnSplineMesh.cpp
+++ b/Engine/Plugins/PCG/Source/PCG/Private/Elements/PCGSpawnSplineMesh.cpp
@@ -7,18 +7,15 @@
 #include "PCGManagedResource.h"
 #include "PCGPin.h"
 #include "PCGSettings.h"
-#include "PCGSubsystem.h"
 #include "Data/PCGLandscapeSplineData.h"
 #include "Data/PCGSplineData.h"
 #include "Helpers/PCGActorHelpers.h"
 #include "Helpers/PCGHelpers.h"
 
 #include "DrawDebugHelpers.h"
-#include "Components/SplineMeshComponent.h"
 #include "Engine/StaticMesh.h"
 #include "Engine/World.h"
 #include "Materials/MaterialInterface.h"
-#include "VT/RuntimeVirtualTexture.h"
 
 #define LOCTEXT_NAMESPACE "PCGCreateSplineMeshElement"
 
@@ -32,7 +29,15 @@ FText UPCGSpawnSplineMeshSettings::GetNodeTooltipText() const
 TArray<FPCGPinProperties> UPCGSpawnSplineMeshSettings::InputPinProperties() const
 {
 	TArray<FPCGPinProperties> PinProperties;
-	PinProperties.Emplace(PCGPinConstants::DefaultInputLabel, EPCGDataType::PolyLine);
+	PinProperties.Emplace_GetRef(PCGPinConstants::DefaultInputLabel, EPCGDataType::PolyLine).SetRequiredPin();
+
+	return PinProperties;
+}
+
+TArray<FPCGPinProperties> UPCGSpawnSplineMeshSettings::OutputPinProperties() const
+{
+	TArray<FPCGPinProperties> PinProperties;
+	PinProperties.Emplace(PCGPinConstants::DefaultOutputLabel, EPCGDataType::PolyLine);
 
 	return PinProperties;
 }
@@ -46,19 +51,16 @@ bool FPCGSpawnSplineMeshElement::PrepareDataInternal(FPCGContext* InContext) con
 {
 	TRACE_CPUPROFILER_EVENT_SCOPE(FPCGCreateSplineMeshElement::PrepareDataInternal);
 
-	FPCGSpawnSplineMeshContext* Context = static_cast<FPCGSpawnSplineMeshContext*>(InContext);
+	ContextType* Context = static_cast<ContextType*>(InContext);
 	check(Context);
 
 	const UPCGSpawnSplineMeshSettings* Settings = Context->GetInputSettings<UPCGSpawnSplineMeshSettings>();
 	check(Settings);
+	
+	const TArray<FPCGTaggedData> Inputs = Context->InputData.GetInputsByPin(PCGPinConstants::DefaultInputLabel);
 
 	const FSoftSplineMeshComponentDescriptor& Descriptor = Settings->SplineMeshDescriptor;
 
-	if (Descriptor.StaticMesh.IsNull())
-	{
-		return true;
-	}
-
 	if (!Context->WasLoadRequested())
 	{
 		TArray<FSoftObjectPath> ObjectsToLoad = { Descriptor.StaticMesh.ToSoftObjectPath() };
@@ -84,121 +86,172 @@ bool FPCGSpawnSplineMeshElement::PrepareDataInternal(FPCGContext* InContext) con
 			}
 		}
 
+		EPCGTimeSliceInitResult ExecResult = Context->InitializePerExecutionState([Settings](ContextType* Context, FPCGSpawnSplineMeshPerExecutionState& OutState)
+		{
+			OutState.TargetActor = Settings->TargetActor.Get() ? Settings->TargetActor.Get() : Context->GetTargetActor(nullptr);
+
+			if (!OutState.TargetActor)
+			{
+				PCGLog::LogErrorOnGraph(LOCTEXT("InvalidTargetActor", "Invalid target actor."), Context);
+				return EPCGTimeSliceInitResult::AbortExecution;
+			}
+			else
+			{
+				return EPCGTimeSliceInitResult::Success;
+			}
+		});
+
+		if (ExecResult == EPCGTimeSliceInitResult::AbortExecution)
+		{
+			return true;
+		}
+		
+		Context->InitializePerIterationStates(Inputs.Num(), [&Inputs, &Context, Settings, &ObjectsToLoad](FPCGSpawnSplineMeshPerIterationState& OutState, const ExecStateType&, const uint32 IterationIndex) -> EPCGTimeSliceInitResult
+		{
+			const FPCGTaggedData& Input = Inputs[IterationIndex];
+			OutState.SplineData = Cast<UPCGPolyLineData>(Input.Data);
+			if (!OutState.SplineData)
+			{
+				return EPCGTimeSliceInitResult::NoOperation;
+			}
+
+			// Forward input
+			Context->OutputData.TaggedData.Add(Input);
+
+			// If the spline is empty we have nothing else to do
+			if (OutState.SplineData->GetNumSegments() <= 0)
+			{
+				return EPCGTimeSliceInitResult::NoOperation;
+			}
+			
+			OutState.LandscapeSplineData = Cast<UPCGLandscapeSplineData>(OutState.SplineData);
+			OutState.SMCBuilderParams.SplineMeshParams = Settings->SplineMeshParams;
+
+			// Initialize the override for the descriptor, and extract all the objects to load
+			OutState.DescriptionOverrides.Initialize(Settings->SplineMeshOverrideDescriptions, &OutState.SMCBuilderParams.Descriptor, OutState.SplineData, Context);
+			if (OutState.DescriptionOverrides.IsValid())
+			{
+				OutState.DescriptionOverrides.GatherAllOverridesToLoad(ObjectsToLoad);
+			}
+
+			// Also initialize the override for the params
+			OutState.ParamsOverrides.Initialize(Settings->SplineMeshParamsOverride, &OutState.SMCBuilderParams.SplineMeshParams, OutState.SplineData, Context);
+
+			return EPCGTimeSliceInitResult::Success;
+		});
+
 		return Context->RequestResourceLoad(Context, std::move(ObjectsToLoad), !Settings->bSynchronousLoad);
 	}
 
 	return true;
 }
 
-bool FPCGSpawnSplineMeshElement::ExecuteInternal(FPCGContext* Context) const
+bool FPCGSpawnSplineMeshElement::ExecuteInternal(FPCGContext* InContext) const
 {
 	TRACE_CPUPROFILER_EVENT_SCOPE(FPCGCreateSplineMeshElement::Execute);
 
+	ContextType* Context = static_cast<ContextType*>(InContext);
 	const UPCGSpawnSplineMeshSettings* Settings = Context->GetInputSettings<UPCGSpawnSplineMeshSettings>();
 	check(Settings);
 
-	FSplineMeshComponentDescriptor Descriptor(Settings->SplineMeshDescriptor);
-
-	FPCGSplineMeshComponentBuilderParameters SMCBuilderParams;
-	SMCBuilderParams.SplineMeshParams = Settings->SplineMeshParams;
-
-	TArray<FPCGTaggedData> Inputs = Context->InputData.GetInputsByPin(PCGPinConstants::DefaultInputLabel);
-
-	for (const FPCGTaggedData& Input : Inputs)
+	if (!Context->DataIsPreparedForExecution())
 	{
-		const UPCGPolyLineData* SplineData = Cast<UPCGPolyLineData>(Input.Data);
-		const UPCGLandscapeSplineData* LandscapeSplineData = Cast<UPCGLandscapeSplineData>(SplineData);
-
-		if (!SplineData)
-		{
-			PCGE_LOG(Error, GraphAndLog, LOCTEXT("InvalidInputData", "Invalid input data, must be of type PolyLine."));
-			continue;
-		}
-
-		AActor* TargetActor = Settings->TargetActor.Get() ? Settings->TargetActor.Get() : Context->GetTargetActor(nullptr);
+		return true;
+	}
 
-		if (!TargetActor)
+	return ExecuteSlice(Context, [this, Settings](ContextType* Context, const ExecStateType& ExecState, IterStateType& IterState, const uint32 IterIndex)
+	{
+		if (Context->GetIterationStateResult(IterIndex) != EPCGTimeSliceInitResult::Success)
 		{
-			PCGE_LOG(Error, GraphAndLog, LOCTEXT("InvalidTargetActor", "Invalid target actor."));
-			continue;
+			return true;
 		}
 
+		const UPCGPolyLineData* SplineData = IterState.SplineData;
+		
 		const int NumSegments = SplineData->GetNumSegments();
 		const bool bIsClosed = SplineData->IsClosed();
 
-		// Copy the descriptor, for the overrides (and reset the descriptor if it was overriden by previous inputs)
-		SMCBuilderParams.Descriptor = Descriptor;
-
-		// TODO: Support overrides on control points when we have better support for those.
-		FPCGObjectOverrides DescriptorOverride(&SMCBuilderParams.Descriptor);
-		DescriptorOverride.Initialize(Settings->SplineMeshOverrideDescriptions, &SMCBuilderParams.Descriptor, SplineData, Context);
-		if (DescriptorOverride.IsValid() && !DescriptorOverride.Apply(0))
+		if (IterState.ElementIndex == 0)
 		{
-			const int32 Index = static_cast<int32>(&Input - Inputs.GetData());
-			PCGLog::LogWarningOnGraph(FText::Format(LOCTEXT("FailOverride", "Failed to override descriptor for input {0}"), Index));
-		}
-
-		UStaticMesh* StaticMesh = SMCBuilderParams.Descriptor.StaticMesh.Get();
-
-		if (!StaticMesh)
+			// Copy the descriptor, for the overrides
+			IterState.SMCBuilderParams.Descriptor = FSplineMeshComponentDescriptor(Settings->SplineMeshDescriptor);
+			IterState.SMCBuilderParams.SplineMeshParams = Settings->SplineMeshParams;
+			IterState.DescriptionOverrides.UpdateTemplateObject(&IterState.SMCBuilderParams.Descriptor);
+			IterState.ParamsOverrides.UpdateTemplateObject(&IterState.SMCBuilderParams.SplineMeshParams);
+		}		
+
+		while (IterState.ElementIndex < NumSegments)
 		{
-			continue;
-		}
+			if (IterState.DescriptionOverrides.IsValid() && !IterState.DescriptionOverrides.Apply(IterState.ElementIndex))
+			{
+				PCGLog::LogWarningOnGraph(FText::Format(LOCTEXT("FailOverrideDescriptor", "Failed to override descriptor for input {0} element {1}"), IterIndex, IterState.ElementIndex));
+			}
+			
+			UStaticMesh* StaticMesh = IterState.SMCBuilderParams.Descriptor.StaticMesh.Get();
 
-		const FVector MeshExtents = StaticMesh->GetBounds().BoxExtent;
+			if (!StaticMesh)
+			{
+				IterState.ElementIndex++;
+				continue;
+			}
 
-		for (int Index = 0; Index < NumSegments; ++Index)
-		{
-			const int NextIndex = (bIsClosed && Index == NumSegments - 1) ? 0 : (Index + 1);
+			if (IterState.ParamsOverrides.IsValid() && !IterState.ParamsOverrides.Apply(IterState.ElementIndex))
+			{
+				PCGLog::LogWarningOnGraph(FText::Format(LOCTEXT("FailOverrideParams", "Failed to override params for input {0} element {1}"), IterIndex, IterState.ElementIndex));
+			}
+			
+			const FVector MeshExtents = StaticMesh->GetBounds().BoxExtent;
+			
+			const int NextIndex = (bIsClosed && IterState.ElementIndex == NumSegments - 1) ? 0 : (IterState.ElementIndex + 1);
 
 			FBox Bounds, NextBounds;
-			FTransform Transform = SplineData->GetTransformAtDistance(Index, /*Distance=*/0.0, /*bWorldSpace=*/true, &Bounds);
-			FTransform NextTransform = SplineData->GetTransformAtDistance(Index, SplineData->GetSegmentLength(Index), /*bWorldSpace=*/true, &NextBounds);
+			FTransform Transform = SplineData->GetTransformAtDistance(IterState.ElementIndex, /*Distance=*/0.0, /*bWorldSpace=*/true, &Bounds);
+			FTransform NextTransform = SplineData->GetTransformAtDistance(IterState.ElementIndex, SplineData->GetSegmentLength(IterState.ElementIndex), /*bWorldSpace=*/true, &NextBounds);
 
-			if (LandscapeSplineData && !SMCBuilderParams.SplineMeshParams.bScaleMeshToLandscapeSplineFullWidth)
+			if (IterState.LandscapeSplineData && !IterState.SMCBuilderParams.SplineMeshParams.bScaleMeshToLandscapeSplineFullWidth)
 			{
 				Bounds = FBox::BuildAABB(FVector::ZeroVector, FVector::OneVector);
 				NextBounds = FBox::BuildAABB(FVector::ZeroVector, FVector::OneVector);
 			}
 
 			FVector LeaveTangent, ArriveTangent, Dummy;
-			SplineData->GetTangentsAtSegmentStart(Index, Dummy, LeaveTangent);
+			SplineData->GetTangentsAtSegmentStart(IterState.ElementIndex, Dummy, LeaveTangent);
 			SplineData->GetTangentsAtSegmentStart(NextIndex, ArriveTangent, Dummy);
 
 			// Set Position
-			SMCBuilderParams.SplineMeshParams.StartPosition = Transform.GetLocation();
-			SMCBuilderParams.SplineMeshParams.StartTangent = LeaveTangent;
-			SMCBuilderParams.SplineMeshParams.EndPosition = NextTransform.GetLocation();
-			SMCBuilderParams.SplineMeshParams.EndTangent = ArriveTangent;
+			IterState.SMCBuilderParams.SplineMeshParams.StartPosition = Transform.GetLocation();
+			IterState.SMCBuilderParams.SplineMeshParams.StartTangent = LeaveTangent;
+			IterState.SMCBuilderParams.SplineMeshParams.EndPosition = NextTransform.GetLocation();
+			IterState.SMCBuilderParams.SplineMeshParams.EndTangent = ArriveTangent;
 
 			// Set Roll (Rotation)
 			const FQuat StartRotation = Transform.GetRotation();
 			const FQuat EndRotation = NextTransform.GetRotation();
 
-			SMCBuilderParams.SplineMeshParams.StartRollDegrees = StartRotation.Rotator().Roll;
-			SMCBuilderParams.SplineMeshParams.EndRollDegrees = EndRotation.Rotator().Roll;
+			IterState.SMCBuilderParams.SplineMeshParams.StartRollDegrees = StartRotation.Rotator().Roll;
+			IterState.SMCBuilderParams.SplineMeshParams.EndRollDegrees = EndRotation.Rotator().Roll;
 
 			// Set Scale
 			FVector StartScale = Transform.GetScale3D() * Bounds.GetExtent();
 			FVector EndScale = NextTransform.GetScale3D() * NextBounds.GetExtent();
 
-			if (SMCBuilderParams.SplineMeshParams.bScaleMeshToBounds)
+			if (IterState.SMCBuilderParams.SplineMeshParams.bScaleMeshToBounds)
 			{
 				FVector::FReal ScaleFactorY = 1.0f;
 				FVector::FReal ScaleFactorZ = 1.0f;
 
 				// We only scale in two dimensions since we are extruding along one of the axes. Scale on the two axes we are not extruding along.
-				if (SMCBuilderParams.SplineMeshParams.ForwardAxis == EPCGSplineMeshForwardAxis::X)
+				if (IterState.SMCBuilderParams.SplineMeshParams.ForwardAxis == EPCGSplineMeshForwardAxis::X)
 				{
 					ScaleFactorY = MeshExtents.Y;
 					ScaleFactorZ = MeshExtents.Z;
 				}
-				else if (SMCBuilderParams.SplineMeshParams.ForwardAxis == EPCGSplineMeshForwardAxis::Y)
+				else if (IterState.SMCBuilderParams.SplineMeshParams.ForwardAxis == EPCGSplineMeshForwardAxis::Y)
 				{
 					ScaleFactorY = MeshExtents.X;
 					ScaleFactorZ = MeshExtents.Z;
 				}
-				else if (SMCBuilderParams.SplineMeshParams.ForwardAxis == EPCGSplineMeshForwardAxis::Z)
+				else if (IterState.SMCBuilderParams.SplineMeshParams.ForwardAxis == EPCGSplineMeshForwardAxis::Z)
 				{
 					ScaleFactorY = MeshExtents.X;
 					ScaleFactorZ = MeshExtents.Y;
@@ -210,28 +263,37 @@ bool FPCGSpawnSplineMeshElement::ExecuteInternal(FPCGContext* Context) const
 				EndScale.Z /= ScaleFactorZ;
 			}
 
-			SMCBuilderParams.SplineMeshParams.StartScale = FVector2D(StartScale.Y, StartScale.Z);
-			SMCBuilderParams.SplineMeshParams.EndScale = FVector2D(EndScale.Y, EndScale.Z);
-			SMCBuilderParams.SettingsCrc = Settings->GetSettingsCrc();
-			ensure(SMCBuilderParams.SettingsCrc.IsValid());
+			IterState.SMCBuilderParams.SplineMeshParams.StartScale = FVector2D(StartScale.Y, StartScale.Z);
+			IterState.SMCBuilderParams.SplineMeshParams.EndScale = FVector2D(EndScale.Y, EndScale.Z);
+			IterState.SMCBuilderParams.SettingsCrc = Settings->GetSettingsCrc();
+			ensure(IterState.SMCBuilderParams.SettingsCrc.IsValid());
 
-			USplineMeshComponent* SplineMeshComponent = UPCGActorHelpers::GetOrCreateSplineMeshComponent(TargetActor, Context->SourceComponent.Get(), SMCBuilderParams, Context);
+			USplineMeshComponent* SplineMeshComponent = UPCGActorHelpers::GetOrCreateSplineMeshComponent(ExecState.TargetActor, Context->SourceComponent.Get(), IterState.SMCBuilderParams, Context);
 
 			// TODO: Write out the geometry to a dynamic mesh type.
 			//SplineMeshComponent->BodySetup->TriMeshGeometries
+
+			IterState.ElementIndex++;
+
+			if (Context->ShouldStop())
+			{
+				break;
+			}
 		}
 
+		const bool bDone = IterState.ElementIndex == NumSegments;
+		
 		// Execute PostProcess Functions
-		if (TargetActor)
+		if (bDone && ExecState.TargetActor)
 		{
-			for (UFunction* Function : PCGHelpers::FindUserFunctions(TargetActor->GetClass(), Settings->PostProcessFunctionNames, { UPCGFunctionPrototypes::GetPrototypeWithNoParams() }, Context))
+			for (UFunction* Function : PCGHelpers::FindUserFunctions(ExecState.TargetActor->GetClass(), Settings->PostProcessFunctionNames, { UPCGFunctionPrototypes::GetPrototypeWithNoParams() }, Context))
 			{
-				TargetActor->ProcessEvent(Function, nullptr);
+				ExecState.TargetActor->ProcessEvent(Function, nullptr);
 			}
 		}
-	}
 
-	return true;
+		return bDone;
+	});
 }
 
 #undef LOCTEXT_NAMESPACE
diff --git a/Engine/Plugins/PCG/Source/PCG/Private/Metadata/PCGObjectPropertyOverride.cpp b/Engine/Plugins/PCG/Source/PCG/Private/Metadata/PCGObjectPropertyOverride.cpp
--- a/Engine/Plugins/PCG/Source/PCG/Private/Metadata/PCGObjectPropertyOverride.cpp
+++ b/Engine/Plugins/PCG/Source/PCG/Private/Metadata/PCGObjectPropertyOverride.cpp
@@ -84,7 +84,7 @@ void FPCGObjectSingleOverride::Initialize(const FPCGAttributePropertySelector& I
 
 	if (!ObjectOverrideInputAccessor.IsValid())
 	{
-		PCGLog::LogWarningOnGraph(FText::Format(NSLOCTEXT("PCGObjectPropertyOverride", "OverrideInputInvalid", "ObjectOverride for input '{0}' is invalid or unsupported. Check the attribute or property selection."), InputSelector.GetDisplayText()), Context);
+		PCGLog::LogWarningOnGraph(FText::Format(NSLOCTEXT("PCGObjectPropertyOverride", "OverrideInputInvalid", "ObjectOverride for input '{0}' does not exist or is unsupported. Check the attribute or property selection."), InputSelector.GetDisplayText()), Context);
 		return;
 	}
 
@@ -172,7 +172,11 @@ void FPCGObjectSingleOverride::GatherAllOverridesToLoad(TArray<FSoftObjectPath>&
 
 	PCGMetadataElementCommon::ApplyOnAccessor<FSoftObjectPath>(*InputKeys, *ObjectOverrideInputAccessor, [&OutObjectsToLoad](FSoftObjectPath& Value, int32 Index)
 	{
-		OutObjectsToLoad.AddUnique(std::move(Value));
+		// Only keep valid paths
+		if (!Value.IsNull())
+		{
+			OutObjectsToLoad.AddUnique(std::move(Value));
+		}
 	});
 }
 
diff --git a/Engine/Plugins/PCG/Source/PCG/Public/Elements/PCGSpawnSplineMesh.h b/Engine/Plugins/PCG/Source/PCG/Public/Elements/PCGSpawnSplineMesh.h
--- a/Engine/Plugins/PCG/Source/PCG/Public/Elements/PCGSpawnSplineMesh.h
+++ b/Engine/Plugins/PCG/Source/PCG/Public/Elements/PCGSpawnSplineMesh.h
@@ -7,7 +7,9 @@
 #include "PCGElement.h"
 #include "PCGSettings.h"
 #include "PCGSplineMeshParams.h"
+#include "PCGTimeSlicedElementBase.h"
 #include "Async/PCGAsyncLoadingContext.h"
+#include "Helpers/PCGActorHelpers.h"
 #include "Metadata/PCGObjectPropertyOverride.h"
 
 #include "Engine/SplineMeshComponentDescriptor.h"
@@ -15,6 +17,8 @@
 #include "PCGSpawnSplineMesh.generated.h"
 
 struct FPCGContext;
+class UPCGLandscapeSplineData;
+class UPCGPolyLineData;
 
 /** Create a USplineMeshComponent for each segment along a given spline. */
 UCLASS(MinimalAPI, BlueprintType, ClassGroup = (Procedural))
@@ -33,7 +37,7 @@ public:
 
 protected:
 	virtual TArray<FPCGPinProperties> InputPinProperties() const override;
-	virtual TArray<FPCGPinProperties> OutputPinProperties() const override { return TArray<FPCGPinProperties>(); }
+	virtual TArray<FPCGPinProperties> OutputPinProperties() const override;
 	virtual FPCGElementPtr CreateElement() const override;
 	//~End UPCGSettings interface
 
@@ -55,14 +59,33 @@ public:
 	UPROPERTY(BlueprintReadWrite, EditAnywhere, Category = "Settings|Debug")
 	bool bSynchronousLoad = false;
 
-	/** Overrides for spline mesh descriptor. For now it is only support data wide attributes. Does not work for per-control point overrides for now. */
+	/** Overrides for spline mesh descriptor. */
 	UPROPERTY(EditAnywhere, BlueprintReadWrite, Category = Settings)
 	TArray<FPCGObjectPropertyOverrideDescription> SplineMeshOverrideDescriptions;
+
+	/** Overrides for spline mesh params. */
+	UPROPERTY(EditAnywhere, BlueprintReadWrite, Category = Settings)
+	TArray<FPCGObjectPropertyOverrideDescription> SplineMeshParamsOverride;
+};
+
+struct FPCGSpawnSplineMeshPerExecutionState
+{
+	AActor* TargetActor = nullptr;
+};
+
+struct FPCGSpawnSplineMeshPerIterationState
+{
+	FPCGSplineMeshComponentBuilderParameters SMCBuilderParams;
+	FPCGObjectOverrides<FSplineMeshComponentDescriptor> DescriptionOverrides = FPCGObjectOverrides(&SMCBuilderParams.Descriptor);
+	FPCGObjectOverrides<FPCGSplineMeshParams> ParamsOverrides = FPCGObjectOverrides(&SMCBuilderParams.SplineMeshParams);
+	const UPCGPolyLineData* SplineData = nullptr;
+	const UPCGLandscapeSplineData* LandscapeSplineData = nullptr;
+	int32 ElementIndex = 0;
 };
 
 struct FPCGSpawnSplineMeshContext : public FPCGContext, public IPCGAsyncLoadingContext {};
 
-class FPCGSpawnSplineMeshElement : public IPCGElementWithCustomContext<FPCGSpawnSplineMeshContext>
+class FPCGSpawnSplineMeshElement : public TPCGTimeSlicedElementBase<FPCGSpawnSplineMeshPerExecutionState, FPCGSpawnSplineMeshPerIterationState, FPCGSpawnSplineMeshContext>
 {
 public:
 	virtual bool CanExecuteOnlyOnMainThread(FPCGContext* Context) const override { return true; }
diff --git a/Engine/Plugins/PCG/Source/PCG/Public/Elements/PCGTimeSlicedElementBase.h b/Engine/Plugins/PCG/Source/PCG/Public/Elements/PCGTimeSlicedElementBase.h
--- a/Engine/Plugins/PCG/Source/PCG/Public/Elements/PCGTimeSlicedElementBase.h
+++ b/Engine/Plugins/PCG/Source/PCG/Public/Elements/PCGTimeSlicedElementBase.h
@@ -22,7 +22,7 @@
  */
 
 // Forward declaration for friending
-template <typename PerExecutionStateT, typename PerIterationStateT> class TPCGTimeSlicedElementBase;
+template <typename PerExecutionStateT, typename PerIterationStateT, typename BaseContext = FPCGContext> class TPCGTimeSlicedElementBase;
 
 /** The result of initializing the Per Execution or Per Iteration Time Slice States */
 enum class EPCGTimeSliceInitResult : uint8
@@ -42,9 +42,10 @@ namespace PCGTimeSlice
  * A PCG context with helper utility to enable element authors to more easily implement timeslicing.
  * @tparam PerExecutionStateT Struct type of the "per-execution" static data state
  * @tparam PerIterationStateT Struct type of the "per-iteration" data state
+ * @tparam BaseContext        To derive from something else than FPCGContext (useful for extra interfaces or even a custom context)
  */
-template <typename PerExecutionStateT = PCGTimeSlice::FEmptyStruct, typename PerIterationStateT = PCGTimeSlice::FEmptyStruct>
-struct TPCGTimeSlicedContext : public FPCGContext
+template <typename PerExecutionStateT = PCGTimeSlice::FEmptyStruct, typename PerIterationStateT = PCGTimeSlice::FEmptyStruct, typename BaseContext = FPCGContext>
+struct TPCGTimeSlicedContext : public BaseContext
 {
 	virtual bool TimeSliceIsEnabled() const override final { return bTimeSliceIsEnabled; }
 	void SetTimeSliceIsEnabled(const bool bEnableTimeSlice = true) { bTimeSliceIsEnabled = bEnableTimeSlice; }
@@ -110,7 +111,7 @@ protected:
 
 private:
 	// Allow exposure to iteration index, etc
-	friend class TPCGTimeSlicedElementBase<PerExecutionStateT, PerIterationStateT>;
+	friend class TPCGTimeSlicedElementBase<PerExecutionStateT, PerIterationStateT, BaseContext>;
 
 	/** The number of times this context has been time sliced */
 	uint32 ExecutionCount = 0u;
@@ -143,15 +144,16 @@ private:
  * A PCG Element that will utilize a Time Slice Context.
  * @tparam PerExecutionStateT Struct type of the "per-execution" static data state
  * @tparam PerIterationStateT Struct type of the "per-iteration" data state
+ * @tparam BaseContext        To derive from something else than FPCGContext (useful for extra interfaces or even a custom context)
  */
-template <typename PerExecutionStateT, typename PerIterationStateT>
+template <typename PerExecutionStateT, typename PerIterationStateT, typename BaseContext>
 class TPCGTimeSlicedElementBase : public IPCGElement
 {
 public:
 	// Aliases, for ease of use with the template
 	using ExecStateType = PerExecutionStateT;
 	using IterStateType = PerIterationStateT;
-	using ContextType = TPCGTimeSlicedContext<ExecStateType, IterStateType>;
+	using ContextType = TPCGTimeSlicedContext<ExecStateType, IterStateType, BaseContext>;
 
 	virtual FPCGContext* CreateContext() override { return new ContextType(); }
 
@@ -163,14 +165,14 @@ public:
 	// TODO: [FUTURE WORK] Consider callback 'on completion' and other Execution styles, like parallel iterations
 };
 
-template <typename PerExecutionStateT, typename PerIterationStateT>
-void TPCGTimeSlicedContext<PerExecutionStateT, PerIterationStateT>::AddExtraStructReferencedObjects(FReferenceCollector& Collector)
+template <typename PerExecutionStateT, typename PerIterationStateT, typename BaseContext>
+void TPCGTimeSlicedContext<PerExecutionStateT, PerIterationStateT, BaseContext>::AddExtraStructReferencedObjects(FReferenceCollector& Collector)
 {
 	Collector.AddReferencedObjects(TrackedObjectArray);
 }
 
-template <typename PerExecutionStateT, typename PerIterationStateT>
-EPCGTimeSliceInitResult TPCGTimeSlicedContext<PerExecutionStateT, PerIterationStateT>::InitializePerExecutionState(TFunctionRef<InitExecSignature> InitFunc)
+template <typename PerExecutionStateT, typename PerIterationStateT, typename BaseContext>
+EPCGTimeSliceInitResult TPCGTimeSlicedContext<PerExecutionStateT, PerIterationStateT, BaseContext>::InitializePerExecutionState(TFunctionRef<InitExecSignature> InitFunc)
 {
 	// Should only ever be initialized once, so if its called again, ignore it. This allows flexibility for the call to be somewhere that might be invoked numerous times
 	if (ExecutionStateResult == EPCGTimeSliceInitResult::Uninitialized)
@@ -181,8 +183,8 @@ EPCGTimeSliceInitResult TPCGTimeSlicedContext<PerExecutionStateT, PerIterationSt
 	return ExecutionStateResult;
 }
 
-template <typename PerExecutionStateT, typename PerIterationStateT>
-const TArray<EPCGTimeSliceInitResult>& TPCGTimeSlicedContext<PerExecutionStateT, PerIterationStateT>::InitializePerIterationStates(int32 NumIterations, TFunctionRef<InitIterSignature> IterFunc)
+template <typename PerExecutionStateT, typename PerIterationStateT, typename BaseContext>
+const TArray<EPCGTimeSliceInitResult>& TPCGTimeSlicedContext<PerExecutionStateT, PerIterationStateT, BaseContext>::InitializePerIterationStates(int32 NumIterations, TFunctionRef<InitIterSignature> IterFunc)
 {
 	// Same as InitializePerExecutionState. Should only ever be initialized once, so if its called again, ignore it.
 	if (bPerIterationStateIsInitialized)
@@ -219,15 +221,15 @@ const TArray<EPCGTimeSliceInitResult>& TPCGTimeSlicedContext<PerExecutionStateT,
 	return PerIterationStateResultArray;
 }
 
-template <typename PerExecutionStateT, typename PerIterationStateT>
-void TPCGTimeSlicedContext<PerExecutionStateT, PerIterationStateT>::TrackObject(const UObject* Object)
+template <typename PerExecutionStateT, typename PerIterationStateT, typename BaseContext>
+void TPCGTimeSlicedContext<PerExecutionStateT, PerIterationStateT, BaseContext>::TrackObject(const UObject* Object)
 {
 	check(Object && IsValid(Object));
 	TrackedObjectArray.AddUnique(Object);
 }
 
-template <typename PerExecutionStateT, typename PerIterationStateT>
-bool TPCGTimeSlicedElementBase<PerExecutionStateT, PerIterationStateT>::ExecuteSlice(ContextType* Context, TFunctionRef<ExecSignature> ExecFunc) const
+template <typename PerExecutionStateT, typename PerIterationStateT, typename BaseContext>
+bool TPCGTimeSlicedElementBase<PerExecutionStateT, PerIterationStateT, BaseContext>::ExecuteSlice(ContextType* Context, TFunctionRef<ExecSignature> ExecFunc) const
 {
 	++Context->ExecutionCount;
 
diff --git a/Engine/Plugins/PCG/Source/PCG/Public/Metadata/Accessors/PCGAttributeAccessorHelpers.h b/Engine/Plugins/PCG/Source/PCG/Public/Metadata/Accessors/PCGAttributeAccessorHelpers.h
--- a/Engine/Plugins/PCG/Source/PCG/Public/Metadata/Accessors/PCGAttributeAccessorHelpers.h
+++ b/Engine/Plugins/PCG/Source/PCG/Public/Metadata/Accessors/PCGAttributeAccessorHelpers.h
@@ -26,21 +26,40 @@ class UStruct;
 
 namespace PCGAttributeAccessorHelpers
 {
+	/** Returns true if the property is supported by PCG types and has a conversion to a Metadata type. */
 	PCG_API bool IsPropertyAccessorSupported(const FProperty* InProperty);
+
+	/** Returns true if the property is supported by PCG types and has a conversion to a Metadata type. */
 	PCG_API bool IsPropertyAccessorSupported(const FName InPropertyName, const UStruct* InStruct);
+
+	/** Returns true if the property chain exists,  and the last property is supported by PCG types and has a conversion to a Metadata type. */
 	PCG_API bool IsPropertyAccessorChainSupported(const TArray<FName>& InPropertyNames, const UStruct* InStruct);
 
+	/** Create an accessor for the given property. Property needs to be supported by PCG (cf. IsPropertyAccessorSupported). */
 	PCG_API TUniquePtr<IPCGAttributeAccessor> CreatePropertyAccessor(const FProperty* InProperty);
+	
+	/** Look for a property in the provided class/struct and create an accessor for it. Property needs to be supported by PCG (cf. IsPropertyAccessorSupported). */
 	PCG_API TUniquePtr<IPCGAttributeAccessor> CreatePropertyAccessor(const FName InPropertyName, const UStruct* InStruct);
 
+	/**
+	 * Create a chain accessor for the given properties. Last property needs to be supported by PCG (cf. IsPropertyAccessorSupported).
+	 * USE WITH CAUTION: There is no validation that the properties are related and may corrupt memory if used incorrectly.
+	 */
 	PCG_API TUniquePtr<IPCGAttributeAccessor> CreatePropertyChainAccessor(TArray<const FProperty*>&& InProperties);
+
+	/** Create a chain accessor for the given properties, starting from the provided class/struct. Last property needs to be supported by PCG (cf. IsPropertyAccessorSupported). */
 	PCG_API TUniquePtr<IPCGAttributeAccessor> CreatePropertyChainAccessor(const TArray<FName>& InPropertyNames, const UStruct* InStruct);
 
+	/** Create a special accessor for all the supported extra properties (cf. EPCGExtraProperties). */
 	PCG_API TUniquePtr<IPCGAttributeAccessor> CreateExtraAccessor(EPCGExtraProperties InExtraProperties);
 
+	/** Create a chain accessor to extract the "Name" out of the Accessor. Example: InAccessor is of type Vector, and name is `X`, it will create an accessor to access the X component on vectors. */
 	PCG_API TUniquePtr<IPCGAttributeAccessor> CreateChainAccessor(TUniquePtr<IPCGAttributeAccessor> InAccessor, FName Name, bool& bOutSuccess);
+	
+	/** Create a chain accessor using the extra names from the selector. Example: InAccessor is of type Vector, and extra names are [`X`], it will create an accessor to access the X component on vectors. */
 	PCG_API TUniquePtr<IPCGAttributeAccessor> CreateChainAccessor(TUniquePtr<IPCGAttributeAccessor> InAccessor, const FPCGAttributePropertySelector& InSelector, bool& bOutSuccess, bool bQuiet = false);
 
+	/** From a list of property names, starting from provided class/struct, return the property chain. */
 	PCG_API bool GetPropertyChain(const TArray<FName>& InPropertyNames, const UStruct* InStruct, TArray<const FProperty*>& OutProperties);
 
 	struct AccessorParamResult
diff --git a/Engine/Plugins/PCG/Source/PCG/Public/Metadata/PCGObjectPropertyOverride.h b/Engine/Plugins/PCG/Source/PCG/Public/Metadata/PCGObjectPropertyOverride.h
--- a/Engine/Plugins/PCG/Source/PCG/Public/Metadata/PCGObjectPropertyOverride.h
+++ b/Engine/Plugins/PCG/Source/PCG/Public/Metadata/PCGObjectPropertyOverride.h
@@ -74,7 +74,7 @@ struct FPCGObjectSingleOverride
 	/** Applies a single property override to the object by reading from the InputAccessor at the given KeyIndex, and writing to the OutputKey which represents the object property. */
 	PCG_API bool Apply(int32 InputKeyIndex, IPCGAttributeAccessorKeys& OutputKey);
 
-	/** Utility function that gather all the overrides and fill the out array if the override needs loading. */
+	/** Gathers overrides into an array, if they need loading. */
 	PCG_API void GatherAllOverridesToLoad(TArray<FSoftObjectPath>& OutObjectsToLoad) const;
 
 private:
@@ -147,20 +147,7 @@ struct FPCGObjectOverrides
 
 		for (int32 i = 0; i < OverrideDescriptions.Num(); ++i)
 		{
-			FPCGAttributePropertyInputSelector InputSelector = OverrideDescriptions[i].InputSource.CopyAndFixLast(SourceData);
-
-			if (InputSelector.GetSelection() == EPCGAttributePropertySelection::Attribute)
-			{
-				const UPCGMetadata* Metadata = SourceData ? SourceData->ConstMetadata() : nullptr;
-				const FName AttributeName = InputSelector.GetAttributeName();
-
-				if (!Metadata || !Metadata->HasAttribute(AttributeName))
-				{
-					PCGLog::LogWarningOnGraph(FText::Format(NSLOCTEXT("PCGObjectPropertyOverride", "InvalidAttribute", "Tried to initialize ObjectOverride for input '{0}', but the attribute '{1}' does not exist."), InputSelector.GetDisplayText(), FText::FromName(AttributeName)), Context);
-					continue;
-				}
-			}
-
+			const FPCGAttributePropertyInputSelector InputSelector = OverrideDescriptions[i].InputSource.CopyAndFixLast(SourceData);
 			const FString& OutputProperty = OverrideDescriptions[i].PropertyTarget;
 
 			FPCGObjectSingleOverride Override;
@@ -192,7 +179,7 @@ struct FPCGObjectOverrides
 		return bAllSucceeded;
 	}
 
-	/** Update the template object. Only if it was already initialized */
+	/** Update the template object, only if it was already initialized */
 	void UpdateTemplateObject(T* InTemplateObject)
 	{
 		if (IsValid())
@@ -204,7 +191,7 @@ struct FPCGObjectOverrides
 	/** Returns true if we have any override to apply */
 	bool IsValid() const { return bInitialized && !ObjectSingleOverrides.IsEmpty(); }
 
-	/** Utility function that scans all the overrides and fill the out array if the object needs loading. */
+	/** Gathers overrides into an array, if they need loading. */
 	void GatherAllOverridesToLoad(TArray<FSoftObjectPath>& OutObjectsToLoad) const
 	{
 		for (const FPCGObjectSingleOverride& SingleOverride : ObjectSingleOverrides)
