[PCG] Add visualization for different metadata domain

* Also add Hyperlink support for asset paths
* And replace the "Type" attribute name when the attribute name is "None", by "None (Type)"

#rnx

diff --git a/Engine/Plugins/PCG/Source/PCG/Public/PCGDataVisualization.h b/Engine/Plugins/PCG/Source/PCG/Public/PCGDataVisualization.h
--- a/Engine/Plugins/PCG/Source/PCG/Public/PCGDataVisualization.h
+++ b/Engine/Plugins/PCG/Source/PCG/Public/PCGDataVisualization.h
@@ -2,6 +2,9 @@
 
 #pragma once
 
+#include "PCGData.h"
+#include "Metadata/PCGAttributePropertySelector.h"
+#include "Metadata/PCGMetadataCommon.h"
 #include "Metadata/Accessors/IPCGAttributeAccessor.h"
 
 #if WITH_EDITOR
@@ -51,6 +54,34 @@ class IPCGDataVisualization
 public:
 	virtual ~IPCGDataVisualization() = default;
 	virtual void ExecuteDebugDisplay(FPCGContext* Context, const UPCGSettingsInterface* SettingsInterface, const UPCGData* Data, AActor* TargetActor) const = 0;
-	virtual FPCGTableVisualizerInfo GetTableVisualizerInfo(const UPCGData* Data) const = 0;
+	
+	UE_DEPRECATED(5.6, "Use and implement GetTableVisualizerInfoWithDomain instead.")
+	virtual FPCGTableVisualizerInfo GetTableVisualizerInfo(const UPCGData* Data) const
+	{
+		ensure("Should never be called, child class need to override GetTableVisualizerInfoWithDomain");
+		return {};
+	}
+
+	/** To be overriden by child classes. For deprecation, call the simple version one by default. */
+	virtual FPCGTableVisualizerInfo GetTableVisualizerInfoWithDomain(const UPCGData* Data, const FPCGMetadataDomainID& DomainID) const
+	{
+PRAGMA_DISABLE_DEPRECATION_WARNINGS
+		return GetTableVisualizerInfo(Data);
+PRAGMA_ENABLE_DEPRECATION_WARNINGS
+	}
+	
+	virtual FPCGMetadataDomainID GetDefaultDomainForInspection(const UPCGData* Data) const { return Data ? Data->GetDefaultMetadataDomainID() : PCGMetadataDomainID::Invalid; }
+	virtual TArray<FPCGMetadataDomainID> GetAllSupportedDomainsForInspection(const UPCGData* Data) const { return Data ? Data->GetAllSupportedMetadataDomainIDs() : TArray<FPCGMetadataDomainID>(); }
+	virtual FString GetDomainDisplayNameForInspection(const UPCGData* Data, const FPCGMetadataDomainID& DomainID) const
+	{
+		if (!Data)
+		{
+			return {};
+		}
+
+		FPCGAttributePropertySelector Selector;
+		Data->SetDomainFromDomainID(DomainID, Selector);
+		return Selector.GetDomainName().ToString();
+	}
 };
 #endif
diff --git a/Engine/Plugins/PCG/Source/PCGEditor/Private/DataVisualizations/PCGDataVisualizationHelpers.cpp b/Engine/Plugins/PCG/Source/PCGEditor/Private/DataVisualizations/PCGDataVisualizationHelpers.cpp
--- a/Engine/Plugins/PCG/Source/PCGEditor/Private/DataVisualizations/PCGDataVisualizationHelpers.cpp
+++ b/Engine/Plugins/PCG/Source/PCGEditor/Private/DataVisualizations/PCGDataVisualizationHelpers.cpp
@@ -12,6 +12,12 @@
 namespace PCGDataVisualizationConstants
 {
 	const FString NoneAttributeId = TEXT("@None");
+
+	/** Names of the columns in the attribute list. */
+	const FName NAME_Index = FName(TEXT("$Index"));
+
+	/** Labels of the columns. */
+	const FText TEXT_Index = LOCTEXT("Index", "$Index");
 }
 
 namespace PCGDataVisualizationHelpers
@@ -33,7 +39,7 @@ namespace PCGDataVisualizationHelpers
 		return ColumnInfo;
 	}
 
-	FPCGTableVisualizerColumnInfo CreateMetadataColumnInfo(const UPCGData* Data, const FName& Id, EPCGMetadataTypes MetadataType, const TCHAR* PostFix = nullptr)
+	FPCGTableVisualizerColumnInfo CreateMetadataColumnInfo(const FPCGMetadataDomain* MetadataDomain, const FPCGMetadataAttributeBase* Attribute, const FName& Id, EPCGMetadataTypes MetadataType, const TCHAR* PostFix = nullptr)
 	{
 		FString OriginalId = Id.ToString();
 		if (PostFix)
@@ -47,7 +53,7 @@ namespace PCGDataVisualizationHelpers
 		if (Id == NAME_None)
 		{
 			IdString = FString::Printf(TEXT("%s%s"), *PCGDataVisualizationConstants::NoneAttributeId, PostFix);
-			Label = FText::Format(LOCTEXT("NoneLabelFormat", "{0}{1}"), UEnum::GetDisplayValueAsText(MetadataType), FText::FromString(PostFix));
+			Label = FText::Format(LOCTEXT("NoneLabelFormat", "None{1} ({0})"), UEnum::GetDisplayValueAsText(MetadataType), FText::FromString(PostFix));
 		}
 		else if (Label.IsEmpty())
 		{
@@ -64,14 +70,18 @@ namespace PCGDataVisualizationHelpers
 			ColumnInfo.Width = PCGEditorGraphAttributeListView::MaxColumnWidth;
 		}
 
-		if (Data)
+		if (MetadataDomain && Attribute)
 		{
 			FPCGAttributePropertyInputSelector Selector;
 			Selector.Update(OriginalId);
-	
-			ColumnInfo.Accessor = TSharedPtr<const IPCGAttributeAccessor>(PCGAttributeAccessorHelpers::CreateConstAccessor(Data, Selector).Release());
 
-			if (ColumnInfo.Accessor)
+			// Need this gymnastic for the chain accessor
+			TUniquePtr<IPCGAttributeAccessor> Accessor = TUniquePtr<IPCGAttributeAccessor>(const_cast<IPCGAttributeAccessor*>(PCGAttributeAccessorHelpers::CreateConstAccessor(Attribute, MetadataDomain).Release()));
+
+			bool bSuccess = false;
+			ColumnInfo.Accessor = TSharedPtr<const IPCGAttributeAccessor>(PCGAttributeAccessorHelpers::CreateChainAccessor(std::move(Accessor), Selector, bSuccess, /*bQuiet=*/true).Release());
+
+			if (bSuccess && ColumnInfo.Accessor)
 			{
 				ColumnInfo.Tooltip = FText::FromString(PCG::Private::GetTypeName(ColumnInfo.Accessor->GetUnderlyingType()));
 			}
@@ -80,19 +90,31 @@ namespace PCGDataVisualizationHelpers
 		return ColumnInfo;
 	}
 
-	void CreateMetadataColumnInfos(const UPCGData* Data, TArray<FPCGTableVisualizerColumnInfo>& OutColumnInfos)
+	void CreateMetadataColumnInfos(const UPCGData* Data, TArray<FPCGTableVisualizerColumnInfo>& OutColumnInfos, const FPCGMetadataDomainID& DomainID)
 	{
 		const UPCGMetadata* Metadata = Data->ConstMetadata();
+		if (!Metadata)
+		{
+			return;
+		}
+
+		const FPCGMetadataDomain* MetadataDomain = Metadata->GetConstMetadataDomain(DomainID);
+		if (!MetadataDomain)
+		{
+			return;
+		}
 
 		TArray<FName> AttributeNames;
 		TArray<EPCGMetadataTypes> AttributeTypes;
-		Metadata->GetAttributes(AttributeNames, AttributeTypes);
+		MetadataDomain->GetAttributes(AttributeNames, AttributeTypes);
 
 		for (int32 I = 0; I < AttributeNames.Num(); I++)
 		{
 			const FName& AttributeName = AttributeNames[I];
 			const EPCGMetadataTypes AttributeType = AttributeTypes[I];
 			FName ColumnName = AttributeName;
+			const FPCGMetadataAttributeBase* Attribute = MetadataDomain->GetConstAttribute(AttributeName);
+			check(Attribute);
 
 			switch (AttributeType)
 			{
@@ -106,49 +128,49 @@ namespace PCGDataVisualizationHelpers
 			case EPCGMetadataTypes::SoftObjectPath:
 			case EPCGMetadataTypes::SoftClassPath:
 			{
-				OutColumnInfos.Add(CreateMetadataColumnInfo(Data, ColumnName, AttributeType, TEXT("")));
+				OutColumnInfos.Add(CreateMetadataColumnInfo(MetadataDomain, Attribute, ColumnName, AttributeType, TEXT("")));
 				break;
 			}
 			case EPCGMetadataTypes::Vector2:
 			{
-				OutColumnInfos.Add(CreateMetadataColumnInfo(Data, ColumnName, AttributeType, TEXT(".X")));
-				OutColumnInfos.Add(CreateMetadataColumnInfo(Data, ColumnName, AttributeType, TEXT(".Y")));
+				OutColumnInfos.Add(CreateMetadataColumnInfo(MetadataDomain, Attribute, ColumnName, AttributeType, TEXT(".X")));
+				OutColumnInfos.Add(CreateMetadataColumnInfo(MetadataDomain, Attribute, ColumnName, AttributeType, TEXT(".Y")));
 				break;
 			}
 			case EPCGMetadataTypes::Vector:
 			{
-				OutColumnInfos.Add(CreateMetadataColumnInfo(Data, ColumnName, AttributeType, TEXT(".X")));
-				OutColumnInfos.Add(CreateMetadataColumnInfo(Data, ColumnName, AttributeType, TEXT(".Y")));
-				OutColumnInfos.Add(CreateMetadataColumnInfo(Data, ColumnName, AttributeType, TEXT(".Z")));
+				OutColumnInfos.Add(CreateMetadataColumnInfo(MetadataDomain, Attribute, ColumnName, AttributeType, TEXT(".X")));
+				OutColumnInfos.Add(CreateMetadataColumnInfo(MetadataDomain, Attribute, ColumnName, AttributeType, TEXT(".Y")));
+				OutColumnInfos.Add(CreateMetadataColumnInfo(MetadataDomain, Attribute, ColumnName, AttributeType, TEXT(".Z")));
 				break;
 			}
 			case EPCGMetadataTypes::Vector4:
 			case EPCGMetadataTypes::Quaternion:
 			{
-				OutColumnInfos.Add(CreateMetadataColumnInfo(Data, ColumnName, AttributeType, TEXT(".X")));
-				OutColumnInfos.Add(CreateMetadataColumnInfo(Data, ColumnName, AttributeType, TEXT(".Y")));
-				OutColumnInfos.Add(CreateMetadataColumnInfo(Data, ColumnName, AttributeType, TEXT(".Z")));
-				OutColumnInfos.Add(CreateMetadataColumnInfo(Data, ColumnName, AttributeType, TEXT(".W")));
+				OutColumnInfos.Add(CreateMetadataColumnInfo(MetadataDomain, Attribute, ColumnName, AttributeType, TEXT(".X")));
+				OutColumnInfos.Add(CreateMetadataColumnInfo(MetadataDomain, Attribute, ColumnName, AttributeType, TEXT(".Y")));
+				OutColumnInfos.Add(CreateMetadataColumnInfo(MetadataDomain, Attribute, ColumnName, AttributeType, TEXT(".Z")));
+				OutColumnInfos.Add(CreateMetadataColumnInfo(MetadataDomain, Attribute, ColumnName, AttributeType, TEXT(".W")));
 				break;
 			}
 			case EPCGMetadataTypes::Rotator:
 			{
-				OutColumnInfos.Add(CreateMetadataColumnInfo(Data, ColumnName, AttributeType, TEXT(".Roll")));
-				OutColumnInfos.Add(CreateMetadataColumnInfo(Data, ColumnName, AttributeType, TEXT(".Pitch")));
-				OutColumnInfos.Add(CreateMetadataColumnInfo(Data, ColumnName, AttributeType, TEXT(".Yaw")));
+				OutColumnInfos.Add(CreateMetadataColumnInfo(MetadataDomain, Attribute, ColumnName, AttributeType, TEXT(".Roll")));
+				OutColumnInfos.Add(CreateMetadataColumnInfo(MetadataDomain, Attribute, ColumnName, AttributeType, TEXT(".Pitch")));
+				OutColumnInfos.Add(CreateMetadataColumnInfo(MetadataDomain, Attribute, ColumnName, AttributeType, TEXT(".Yaw")));
 				break;
 			}
 			case EPCGMetadataTypes::Transform:
 			{
-				OutColumnInfos.Add(CreateMetadataColumnInfo(Data, ColumnName, AttributeType, TEXT(".Position.X")));
-				OutColumnInfos.Add(CreateMetadataColumnInfo(Data, ColumnName, AttributeType, TEXT(".Position.Y")));
-				OutColumnInfos.Add(CreateMetadataColumnInfo(Data, ColumnName, AttributeType, TEXT(".Position.Z")));
-				OutColumnInfos.Add(CreateMetadataColumnInfo(Data, ColumnName, AttributeType, TEXT(".Rotation.Roll")));
-				OutColumnInfos.Add(CreateMetadataColumnInfo(Data, ColumnName, AttributeType, TEXT(".Rotation.Pitch")));
-				OutColumnInfos.Add(CreateMetadataColumnInfo(Data, ColumnName, AttributeType, TEXT(".Rotation.Yaw")));
-				OutColumnInfos.Add(CreateMetadataColumnInfo(Data, ColumnName, AttributeType, TEXT(".Scale.X")));
-				OutColumnInfos.Add(CreateMetadataColumnInfo(Data, ColumnName, AttributeType, TEXT(".Scale.Y")));
-				OutColumnInfos.Add(CreateMetadataColumnInfo(Data, ColumnName, AttributeType, TEXT(".Scale.Z")));
+				OutColumnInfos.Add(CreateMetadataColumnInfo(MetadataDomain, Attribute, ColumnName, AttributeType, TEXT(".Position.X")));
+				OutColumnInfos.Add(CreateMetadataColumnInfo(MetadataDomain, Attribute, ColumnName, AttributeType, TEXT(".Position.Y")));
+				OutColumnInfos.Add(CreateMetadataColumnInfo(MetadataDomain, Attribute, ColumnName, AttributeType, TEXT(".Position.Z")));
+				OutColumnInfos.Add(CreateMetadataColumnInfo(MetadataDomain, Attribute, ColumnName, AttributeType, TEXT(".Rotation.Roll")));
+				OutColumnInfos.Add(CreateMetadataColumnInfo(MetadataDomain, Attribute, ColumnName, AttributeType, TEXT(".Rotation.Pitch")));
+				OutColumnInfos.Add(CreateMetadataColumnInfo(MetadataDomain, Attribute, ColumnName, AttributeType, TEXT(".Rotation.Yaw")));
+				OutColumnInfos.Add(CreateMetadataColumnInfo(MetadataDomain, Attribute, ColumnName, AttributeType, TEXT(".Scale.X")));
+				OutColumnInfos.Add(CreateMetadataColumnInfo(MetadataDomain, Attribute, ColumnName, AttributeType, TEXT(".Scale.Y")));
+				OutColumnInfos.Add(CreateMetadataColumnInfo(MetadataDomain, Attribute, ColumnName, AttributeType, TEXT(".Scale.Z")));
 				break;
 			}
 			default:
@@ -156,6 +178,30 @@ namespace PCGDataVisualizationHelpers
 			}
 		}
 	}
+
+	FPCGTableVisualizerInfo CreateDefaultMetadataColumnInfos(const UPCGData* Data, const FPCGMetadataDomainID& DomainID)
+	{
+		const UPCGMetadata* Metadata = Data->ConstMetadata();
+		const FPCGMetadataDomain* MetadataDomain = Metadata ? Metadata->GetConstMetadataDomain(DomainID) : nullptr;
+		if (!MetadataDomain)
+		{
+			return {};
+		}
+		
+		FPCGTableVisualizerInfo Info;
+		Info.Data = Data;
+		Info.ColumnInfos.Add(CreateColumnInfo(Data, PCGDataVisualizationConstants::NAME_Index, PCGDataVisualizationConstants::TEXT_Index));
+		Info.SortingColumn = PCGDataVisualizationConstants::NAME_Index;
+
+		
+		const bool bAddDefaultValueIfEmpty = DomainID == PCGMetadataDomainID::Data;
+		Info.AccessorKeys = MakeShared<const FPCGAttributeAccessorKeysEntries>(MetadataDomain, /*bAddDefaultValueIfEmpty=*/bAddDefaultValueIfEmpty);
+
+		// Add Metadata Columns
+		CreateMetadataColumnInfos(Data, Info.ColumnInfos, DomainID);
+
+		return Info;
+	}
 }
 
 #undef LOCTEXT_NAMESPACE
diff --git a/Engine/Plugins/PCG/Source/PCGEditor/Private/DataVisualizations/PCGParamDataVisualization.cpp b/Engine/Plugins/PCG/Source/PCGEditor/Private/DataVisualizations/PCGParamDataVisualization.cpp
--- a/Engine/Plugins/PCG/Source/PCGEditor/Private/DataVisualizations/PCGParamDataVisualization.cpp
+++ b/Engine/Plugins/PCG/Source/PCGEditor/Private/DataVisualizations/PCGParamDataVisualization.cpp
@@ -8,27 +8,9 @@
 
 #define LOCTEXT_NAMESPACE "PCGParamDataVisualization"
 
-namespace PCGParamDataVisualizationConstants
+FPCGTableVisualizerInfo IPCGParamDataVisualization::GetTableVisualizerInfoWithDomain(const UPCGData* Data, const FPCGMetadataDomainID& DomainID) const
 {
-	/** Names of the columns in the attribute list. */
-	const FName NAME_Index = FName(TEXT("$Index"));
-
-	/** Labels of the columns. */
-	const FText TEXT_Index = LOCTEXT("Index", "$Index");
-}
-
-FPCGTableVisualizerInfo IPCGParamDataVisualization::GetTableVisualizerInfo(const UPCGData* Data) const
-{
-	FPCGTableVisualizerInfo Info;
-	Info.Data = Data;
-	Info.ColumnInfos.Add(PCGDataVisualizationHelpers::CreateColumnInfo(Data, PCGParamDataVisualizationConstants::NAME_Index, PCGParamDataVisualizationConstants::TEXT_Index));
-	Info.SortingColumn = PCGParamDataVisualizationConstants::NAME_Index;
-	Info.AccessorKeys = TSharedPtr<const IPCGAttributeAccessorKeys>(PCGAttributeAccessorHelpers::CreateConstKeys(Data, FPCGAttributePropertyInputSelector()).Release());
-
-	// Add Metadata Columns
-	PCGDataVisualizationHelpers::CreateMetadataColumnInfos(Data, Info.ColumnInfos);
-
-	return Info;
+	return PCGDataVisualizationHelpers::CreateDefaultMetadataColumnInfos(Data, DomainID);
 }
 
 #undef LOCTEXT_NAMESPACE
diff --git a/Engine/Plugins/PCG/Source/PCGEditor/Private/DataVisualizations/PCGSpatialDataVisualization.cpp b/Engine/Plugins/PCG/Source/PCGEditor/Private/DataVisualizations/PCGSpatialDataVisualization.cpp
--- a/Engine/Plugins/PCG/Source/PCGEditor/Private/DataVisualizations/PCGSpatialDataVisualization.cpp
+++ b/Engine/Plugins/PCG/Source/PCGEditor/Private/DataVisualizations/PCGSpatialDataVisualization.cpp
@@ -268,7 +268,7 @@ void IPCGSpatialDataVisualization::ExecuteDebugDisplay(FPCGContext* Context, con
 	}
 }
 
-FPCGTableVisualizerInfo IPCGSpatialDataVisualization::GetTableVisualizerInfo(const UPCGData* Data) const
+FPCGTableVisualizerInfo IPCGSpatialDataVisualization::GetTableVisualizerInfoWithDomain(const UPCGData* Data, const FPCGMetadataDomainID& DomainID) const
 {
 	using namespace PCGDataVisualizationHelpers;
 	using namespace PCGPointDataVisualizationConstants;
@@ -276,6 +276,11 @@ FPCGTableVisualizerInfo IPCGSpatialDataVisualization::GetTableVisualizerInfo(con
 	// Collapse to point representation for visualization.
 	const UPCGBasePointData* PointData = CollapseToDebugBasePointData(/*Context=*/nullptr, Data);
 
+	if (DomainID == PCGMetadataDomainID::Data)
+	{
+		return CreateDefaultMetadataColumnInfos(PointData, DomainID);
+	}
+
 	FPCGTableVisualizerInfo Info;
 	Info.Data = PointData;
 
@@ -329,7 +334,7 @@ FPCGTableVisualizerInfo IPCGSpatialDataVisualization::GetTableVisualizerInfo(con
 	}
 
 	// Add Metadata Columns
-	CreateMetadataColumnInfos(PointData, Info.ColumnInfos);
+	CreateMetadataColumnInfos(PointData, Info.ColumnInfos, PCGMetadataDomainID::Elements);
 
 	// Column Sorting
 	Info.SortingColumn = NAME_Index;
@@ -371,7 +376,7 @@ FPCGTableVisualizerInfo IPCGSpatialDataVisualization::GetTableVisualizerInfo(con
 	};
 
 	// Accessor Keys
-	Info.AccessorKeys = TSharedPtr<const IPCGAttributeAccessorKeys>(PCGAttributeAccessorHelpers::CreateConstKeys(PointData, FPCGAttributePropertyInputSelector()).Release());
+	Info.AccessorKeys = TSharedPtr<const IPCGAttributeAccessorKeys>(PCGAttributeAccessorHelpers::CreateConstKeys(PointData, FPCGAttributePropertySelector()).Release());
 
 	return Info;
 }
@@ -405,4 +410,22 @@ PRAGMA_ENABLE_DEPRECATION_WARNINGS
 	return nullptr;
 }
 
+FString IPCGSpatialDataVisualization::GetDomainDisplayNameForInspection(const UPCGData* Data, const FPCGMetadataDomainID& DomainID) const
+{
+	if (DomainID != PCGMetadataDomainID::Elements || Data->IsSupportedMetadataDomainID(DomainID))
+	{
+		return IPCGDataVisualization::GetDomainDisplayNameForInspection(Data, DomainID);
+	}
+	else
+	{
+		// For sampled points, clearly indicate that it is the default sampled points and not just "points"
+		return TEXT("Default Sampled Points");
+	}
+}
+
+TArray<FPCGMetadataDomainID> IPCGSpatialDataVisualization::GetAllSupportedDomainsForInspection(const UPCGData* Data) const
+{
+	return GetDefault<UPCGBasePointData>()->GetAllSupportedMetadataDomainIDs();
+}
+
 #undef LOCTEXT_NAMESPACE
diff --git a/Engine/Plugins/PCG/Source/PCGEditor/Private/DataVisualizations/PCGSplineDataVisualization.cpp b/Engine/Plugins/PCG/Source/PCGEditor/Private/DataVisualizations/PCGSplineDataVisualization.cpp
--- a/Engine/Plugins/PCG/Source/PCGEditor/Private/DataVisualizations/PCGSplineDataVisualization.cpp
+++ b/Engine/Plugins/PCG/Source/PCGEditor/Private/DataVisualizations/PCGSplineDataVisualization.cpp
@@ -12,43 +12,46 @@ namespace PCGSplineDataVisualizationConstants
 	const FVector HalfExtents = FVector(50.0); // Set the points to fill up 1 meter of space by default.
 }
 
-FPCGTableVisualizerInfo IPCGSplineDataVisualization::GetTableVisualizerInfo(const UPCGData* Data) const
+FPCGTableVisualizerInfo IPCGSplineDataVisualization::GetTableVisualizerInfoWithDomain(const UPCGData* Data, const FPCGMetadataDomainID& DomainID) const
 {
-	FPCGTableVisualizerInfo Info = IPCGSpatialDataVisualization::GetTableVisualizerInfo(Data);
+	FPCGTableVisualizerInfo Info = IPCGSpatialDataVisualization::GetTableVisualizerInfoWithDomain(Data, DomainID);
 
-	Info.FocusOnDataCallback = [](const UPCGData* Data, TArrayView<const int> Indices)
+	if (DomainID == PCGMetadataDomainID::Elements)
 	{
-		if (const UPCGSplineData* SplineData = Cast<UPCGSplineData>(Data))
+		Info.FocusOnDataCallback = [](const UPCGData* Data, TArrayView<const int> Indices)
 		{
-			FBox BoundingBox(EForceInit::ForceInit);
-
-			if (Indices.IsEmpty())
-			{
-				BoundingBox = SplineData->GetBounds();
-			}
-			else
+			if (const UPCGSplineData* SplineData = Cast<UPCGSplineData>(Data))
 			{
-				const TArray<FInterpCurvePointVector>& Positions = SplineData->SplineStruct.GetSplinePointsPosition().Points;
-				const TArray<FInterpCurvePointVector>& Scales = SplineData->SplineStruct.GetSplinePointsScale().Points;
+				FBox BoundingBox(EForceInit::ForceInit);
 
-				for (const int& Index : Indices)
+				if (Indices.IsEmpty())
+				{
+					BoundingBox = SplineData->GetBounds();
+				}
+				else
 				{
-					check(Positions.IsValidIndex(Index) && Scales.IsValidIndex(Index));
+					const TArray<FInterpCurvePointVector>& Positions = SplineData->SplineStruct.GetSplinePointsPosition().Points;
+					const TArray<FInterpCurvePointVector>& Scales = SplineData->SplineStruct.GetSplinePointsScale().Points;
+
+					for (const int& Index : Indices)
+					{
+						check(Positions.IsValidIndex(Index) && Scales.IsValidIndex(Index));
 
-					const FVector& Position = SplineData->GetTransform().TransformPosition(Positions[Index].OutVal);
-					const FVector HalfExtent = Scales[Index].OutVal * PCGSplineDataVisualizationConstants::HalfExtents;
+						const FVector& Position = SplineData->GetTransform().TransformPosition(Positions[Index].OutVal);
+						const FVector HalfExtent = Scales[Index].OutVal * PCGSplineDataVisualizationConstants::HalfExtents;
 
-					FBox PointBoundingBox(Position + HalfExtent, Position - HalfExtent);
-					BoundingBox += PointBoundingBox;
+						FBox PointBoundingBox(Position + HalfExtent, Position - HalfExtent);
+						BoundingBox += PointBoundingBox;
+					}
 				}
-			}
 
-			if (GEditor && BoundingBox.IsValid)
-			{
-				GEditor->MoveViewportCamerasToBox(BoundingBox, /*bActiveViewportOnly=*/true, /*DrawDebugBoxTimeInSeconds=*/2.5f);
+				if (GEditor && BoundingBox.IsValid)
+				{
+					GEditor->MoveViewportCamerasToBox(BoundingBox, /*bActiveViewportOnly=*/true, /*DrawDebugBoxTimeInSeconds=*/2.5f);
+				}
 			}
-		}
-	};
+		};
+	}
 
 	return Info;
 }
@@ -58,7 +61,8 @@ const UPCGBasePointData* IPCGSplineDataVisualization::CollapseToDebugBasePointDa
 	if (const UPCGSplineData* SplineData = Cast<UPCGSplineData>(Data))
 	{
 		UPCGBasePointData* PointData = FPCGContext::NewPointData_AnyThread(Context);
-		PointData->InitializeFromData(SplineData);
+		FPCGInitializeFromDataParams InitializeFromDataParams(SplineData);
+		PointData->InitializeFromDataWithParams(InitializeFromDataParams);
 
 		const int32 NumControlPoints = SplineData->SplineStruct.GetSplinePointsPosition().Points.Num();
 		TConstArrayView<PCGMetadataEntryKey> EntryKeys = SplineData->SplineStruct.GetConstControlPointsEntryKeys();
@@ -76,6 +80,11 @@ const UPCGBasePointData* IPCGSplineDataVisualization::CollapseToDebugBasePointDa
 			MetadataEntryRange[ControlPointIndex] = EntryKeys.IsValidIndex(ControlPointIndex) ? EntryKeys[ControlPointIndex] : PCGInvalidEntryKey;
 		}
 
+		if (SplineData->HasCachedLastSelector())
+		{
+			PointData->SetLastSelector(SplineData->GetCachedLastSelector());
+		}
+
 		return PointData;
 	}
 
diff --git a/Engine/Plugins/PCG/Source/PCGEditor/Private/Widgets/SPCGEditorGraphAttributeListView.cpp b/Engine/Plugins/PCG/Source/PCGEditor/Private/Widgets/SPCGEditorGraphAttributeListView.cpp
--- a/Engine/Plugins/PCG/Source/PCGEditor/Private/Widgets/SPCGEditorGraphAttributeListView.cpp
+++ b/Engine/Plugins/PCG/Source/PCGEditor/Private/Widgets/SPCGEditorGraphAttributeListView.cpp
@@ -34,6 +34,7 @@
 #include "Widgets/Images/SThrobber.h"
 #include "Widgets/Input/SButton.h"
 #include "Widgets/Input/SComboBox.h"
+#include "Widgets/Input/SHyperlink.h"
 #include "Widgets/Input/SSearchBox.h"
 #include "Widgets/Layout/SScrollBox.h"
 
@@ -135,12 +136,15 @@ TSharedRef<SWidget> SPCGListViewItemRow::GenerateWidgetForColumn(const FName& Co
 	// We will only compute cell width for types are that not numbers.
 	bool bShouldComputeCellWidth = false;
 
+	// Make the text clickable for soft object paths
+	bool bIsClickable = false;
+
 	if (FPCGColumnData* PCGColumnData = SharedAttributeListView->PCGColumnData.Find(ColumnId))
 	{
 		if (PCGColumnData->DataAccessor.IsValid() && PCGColumnData->DataKeys.IsValid())
 		{
 			int32 Index = InternalItem->Index;
-			auto Callback = [&PCGColumnData, Index, &RowText, &bShouldComputeCellWidth] (auto Dummy)
+			auto Callback = [&PCGColumnData, Index, &RowText, &bShouldComputeCellWidth, &bIsClickable] (auto Dummy)
 			{
 				using ValueType = decltype(Dummy);
 				ValueType Value{};
@@ -168,11 +172,13 @@ TSharedRef<SWidget> SPCGListViewItemRow::GenerateWidgetForColumn(const FName& Co
 					{
 						RowText = FText::FromString(Value.ToString());
 						bShouldComputeCellWidth = true;
+						bIsClickable = true;
 					}
 					else if constexpr (PCG::Private::IsOfTypes<ValueType, FSoftClassPath>())
 					{
 						RowText = FText::FromString(Value.ToString());
 						bShouldComputeCellWidth = true;
+						bIsClickable = true;
 					}
 					else
 					{
@@ -196,10 +202,38 @@ TSharedRef<SWidget> SPCGListViewItemRow::GenerateWidgetForColumn(const FName& Co
 		}
 	}
 
-	return SNew(STextBlock)
-		.Text(RowText)
-		.OverflowPolicy(ETextOverflowPolicy::Ellipsis)
-		.Margin(FMargin(2.0f, 0.0f));
+	const FMargin Margin = FMargin(2.0f, 0.0f);
+	const ETextOverflowPolicy OverflowPolicy = ETextOverflowPolicy::Ellipsis;
+
+	if (bIsClickable)
+	{
+		// @todo_pcg: We're missing Margin and Overflow policy with this widget, maybe something we can re-add in the future.
+
+		// Put it in a shared to pass it as the text of the hyperlink and to the OnNavigate callback
+		TSharedPtr<FText> Text = MakeShared<FText>(std::move(RowText));
+		
+		TSharedPtr<SHyperlink> Hyperlink;
+		SAssignNew(Hyperlink, SHyperlink)
+			.Text(TAttribute<FText>::CreateLambda([Text]() { return *Text;}))
+			.ToolTipText(LOCTEXT("Hyperlink", "Ctrl + Click to jump to the asset."))
+			.Style(FAppStyle::Get(), "HoverOnlyHyperlink")
+			.OnNavigate(FSimpleDelegate::CreateLambda([Text]()
+			{
+				if (FSlateApplication::Get().GetModifierKeys().IsControlDown() && GEditor && GEditor->GetEditorSubsystem<UAssetEditorSubsystem>())
+				{
+					GEditor->GetEditorSubsystem<UAssetEditorSubsystem>()->OpenEditorForAsset(Text->ToString());
+				}
+			}));
+
+		return Hyperlink.ToSharedRef();
+	}
+	else
+	{
+		return SNew(STextBlock)
+			.Text(RowText)
+			.OverflowPolicy(OverflowPolicy)
+			.Margin(Margin);
+	}
 }
 
 FPCGPointFilterExpressionContext::FPCGPointFilterExpressionContext(const FPCGListViewItem* InRowItem, const TMap<FName, FPCGColumnData>* InPCGColumnData)
@@ -344,6 +378,15 @@ void SPCGEditorGraphAttributeListView::Construct(const FArguments& InArgs, TShar
 			.Text(this, &SPCGEditorGraphAttributeListView::OnGenerateSelectedDataText)
 		];
 
+	SAssignNew(DomainsComboBox, SComboBox<TSharedPtr<FString>>)
+		.OptionsSource(&DomainsComboBoxItems)
+		.OnGenerateWidget(this, &SPCGEditorGraphAttributeListView::OnGenerateDataWidget)
+		.OnSelectionChanged(this, &SPCGEditorGraphAttributeListView::OnSelectionChanged)
+		[
+			SNew(STextBlock)
+			.Text(this, &SPCGEditorGraphAttributeListView::OnGenerateSelectedDomainText)
+		];
+
 	TSharedPtr<SLayeredImage> FilterImage = SNew(SLayeredImage)
 		.Image(FAppStyle::Get().GetBrush("Icons.Filter"))
 		.ColorAndOpacity(FSlateColor::UseForeground());
@@ -446,6 +489,12 @@ void SPCGEditorGraphAttributeListView::Construct(const FArguments& InArgs, TShar
 				DataComboBox->AsShared()
 			]
 			+SHorizontalBox::Slot()
+			.AutoWidth()
+			.Padding(1.0f, 0.0f)
+			[
+				DomainsComboBox->AsShared()
+			]
+			+SHorizontalBox::Slot()
 			.MinWidth(100.0)
 			.MaxWidth(300.0)
 			.Padding(1.0f, 0.0f)
@@ -710,8 +759,15 @@ void SPCGEditorGraphAttributeListView::RefreshAttributeList()
 		return;
 	}
 
+	const int32 DomainIndex = GetSelectedDomainIndex();
+	if (!DomainsComboBoxIds.IsValidIndex(DomainIndex))
+	{
+		return;
+	}
+
 	const FPCGTaggedData& TaggedData = InspectionData->TaggedData[DataIndex];
 	const UPCGData* PCGData = TaggedData.Data;
+	const FPCGMetadataDomainID& MetadataDomainID = DomainsComboBoxIds[DomainIndex];
 	const FPCGCrc Crc = InspectionData->DataCrcs.IsValidIndex(DataIndex) ? InspectionData->DataCrcs[DataIndex] : FPCGCrc(0);
 
 	if (!PCGData)
@@ -745,7 +801,7 @@ void SPCGEditorGraphAttributeListView::RefreshAttributeList()
 
 	if (const IPCGDataVisualization* DataVisualization = DataVisRegistry.GetDataVisualization(PCGData->GetClass()))
 	{
-		const FPCGTableVisualizerInfo TableVisualizerInfo = DataVisualization->GetTableVisualizerInfo(PCGData);
+		const FPCGTableVisualizerInfo TableVisualizerInfo = DataVisualization->GetTableVisualizerInfoWithDomain(PCGData, MetadataDomainID);
 		DataToVisualize = TableVisualizerInfo.Data;
 
 		for (const FPCGTableVisualizerColumnInfo& ColumnInfo : TableVisualizerInfo.ColumnInfos)
@@ -844,6 +900,7 @@ void SPCGEditorGraphAttributeListView::RefreshDataComboBox()
 	const FPCGDataCollection* InspectionData = GetInspectionData();
 	if (!InspectionData)
 	{
+		RefreshDomainComboBox();
 		return;
 	}
 
@@ -867,6 +924,73 @@ void SPCGEditorGraphAttributeListView::RefreshDataComboBox()
 	{
 		DataComboBox->SetSelectedItem(DataComboBoxItems[0]);
 	}
+
+	RefreshDomainComboBox();
+}
+
+void SPCGEditorGraphAttributeListView::RefreshDomainComboBox()
+{
+	// Try to stick to the same metadata domain if possible
+	const int32 DomainIndex = GetSelectedDomainIndex();
+	const FPCGMetadataDomainID PreviousDomain = DomainsComboBoxIds.IsValidIndex(DomainIndex) ? DomainsComboBoxIds[DomainIndex] : PCGMetadataDomainID::Invalid;
+	
+	DomainsComboBoxItems.Empty();
+	DomainsComboBoxIds.Empty();
+	DomainsComboBox->ClearSelection();
+	DomainsComboBox->RefreshOptions();
+
+	const FPCGDataCollection* InspectionData = GetInspectionData();
+	if (!InspectionData)
+	{
+		return;
+	}
+	
+	const int32 DataIndex = GetSelectedDataIndex();
+	if (!InspectionData->TaggedData.IsValidIndex(DataIndex))
+	{
+		return;
+	}
+
+	const FPCGTaggedData& TaggedData = InspectionData->TaggedData[DataIndex];
+	const UPCGData* PCGData = TaggedData.Data;
+
+	if (!PCGData)
+	{
+		return;
+	}
+	
+	const FPCGDataVisualizationRegistry& DataVisRegistry = FPCGModule::GetConstPCGDataVisualizationRegistry();
+
+	const IPCGDataVisualization* DataVisualization = DataVisRegistry.GetDataVisualization(PCGData->GetClass());
+	if (!DataVisualization)
+	{
+		return;
+	}
+
+	const FPCGMetadataDomainID DefaultDomainID = DataVisualization->GetDefaultDomainForInspection(PCGData);
+	int32 DefaultIndex = 0;
+	int32 PreviousDomainIndex = INDEX_NONE;
+	for (const FPCGMetadataDomainID& DomainID : DataVisualization->GetAllSupportedDomainsForInspection(PCGData))
+	{
+		if (PreviousDomain == DomainID)
+		{
+			PreviousDomainIndex = DomainsComboBoxItems.Num();
+		}
+		
+		if (DomainID == DefaultDomainID)
+		{
+			DefaultIndex = DomainsComboBoxItems.Num();
+		}
+		
+		DomainsComboBoxItems.Add(MakeShared<FString>(DataVisualization->GetDomainDisplayNameForInspection(PCGData, DomainID)));
+		DomainsComboBoxIds.Add(DomainID);
+	}
+
+	if (DomainsComboBoxItems.Num() > 0)
+	{
+		const int32 Index = PreviousDomainIndex != INDEX_NONE ? PreviousDomainIndex : DefaultIndex; 
+		DomainsComboBox->SetSelectedItem(DomainsComboBoxItems[Index]);
+	}
 }
 
 void SPCGEditorGraphAttributeListView::LaunchUpdateTask()
@@ -1094,6 +1218,18 @@ FText SPCGEditorGraphAttributeListView::OnGenerateSelectedDataText() const
 	}
 }
 
+FText SPCGEditorGraphAttributeListView::OnGenerateSelectedDomainText() const
+{
+	if (const TSharedPtr<FString> SelectedDataName = DomainsComboBox->GetSelectedItem())
+	{
+		return FText::FromString(*SelectedDataName);
+	}
+	else
+	{
+		return PCGEditorGraphAttributeListView::NoDataAvailableText;
+	}
+}
+
 int32 SPCGEditorGraphAttributeListView::GetSelectedDataIndex() const
 {
 	int32 Index = INDEX_NONE;
@@ -1105,6 +1241,17 @@ int32 SPCGEditorGraphAttributeListView::GetSelectedDataIndex() const
 	return Index;
 }
 
+int32 SPCGEditorGraphAttributeListView::GetSelectedDomainIndex() const
+{
+	int32 Index = INDEX_NONE;
+	if (const TSharedPtr<FString> SelectedItem = DomainsComboBox->GetSelectedItem())
+	{
+		DomainsComboBoxItems.Find(SelectedItem, Index);
+	}
+
+	return Index;
+}
+
 void SPCGEditorGraphAttributeListView::ToggleAllAttributes()
 {
 	const TArray<FName> HiddenColumns = ListViewHeader->GetHiddenColumnIds();
diff --git a/Engine/Plugins/PCG/Source/PCGEditor/Private/Widgets/SPCGEditorGraphAttributeListView.h b/Engine/Plugins/PCG/Source/PCGEditor/Private/Widgets/SPCGEditorGraphAttributeListView.h
--- a/Engine/Plugins/PCG/Source/PCGEditor/Private/Widgets/SPCGEditorGraphAttributeListView.h
+++ b/Engine/Plugins/PCG/Source/PCGEditor/Private/Widgets/SPCGEditorGraphAttributeListView.h
@@ -170,6 +170,7 @@ private:
 	void RefreshAttributeList();
 	void RefreshPinComboBox();
 	void RefreshDataComboBox();
+	void RefreshDomainComboBox();
 
 	void LaunchUpdateTask();
 
@@ -200,7 +201,9 @@ private:
 	TSharedRef<SWidget> OnGenerateDataWidget(TSharedPtr<FString> InItem) const;
 	void OnSelectionChanged(TSharedPtr<FString> Item, ESelectInfo::Type SelectInfo);
 	FText OnGenerateSelectedDataText() const;
+	FText OnGenerateSelectedDomainText() const;
 	int32 GetSelectedDataIndex() const;
+	int32 GetSelectedDomainIndex() const;
 
 	void ToggleAllAttributes();
 	void ToggleAttribute(FName InAttributeName);
@@ -268,6 +271,10 @@ private:
 	TSharedPtr<SComboBox<TSharedPtr<FString>>> DataComboBox;
 	TArray<TSharedPtr<FString>> DataComboBoxItems;
 
+	TSharedPtr<SComboBox<TSharedPtr<FString>>> DomainsComboBox;
+	TArray<TSharedPtr<FString>> DomainsComboBoxItems;
+	TArray<FPCGMetadataDomainID> DomainsComboBoxIds;
+
 	TSharedPtr<STextBlock> NodeNameTextBlock;
 	TSharedPtr<STextBlock> InfoTextBlock;
 
diff --git a/Engine/Plugins/PCG/Source/PCGEditor/Private/Widgets/SPCGEditorGraphFind.cpp b/Engine/Plugins/PCG/Source/PCGEditor/Private/Widgets/SPCGEditorGraphFind.cpp
--- a/Engine/Plugins/PCG/Source/PCGEditor/Private/Widgets/SPCGEditorGraphFind.cpp
+++ b/Engine/Plugins/PCG/Source/PCGEditor/Private/Widgets/SPCGEditorGraphFind.cpp
@@ -9,8 +9,10 @@
 #include "PCGNode.h"
 #include "PCGSubgraph.h"
 #include "PCGSubsystem.h"
+#include "Elements/PCGUserParameterGet.h"
 #include "Nodes/PCGEditorGraphNodeBase.h"
 
+#include "Algo/Contains.h"
 #include "EdGraph/EdGraphSchema.h"
 #include "Framework/Application/SlateApplication.h"
 #include "Layout/WidgetPath.h"
@@ -20,6 +22,12 @@
 
 #define LOCTEXT_NAMESPACE "PCGEditorGraphFind"
 
+namespace PCGEditorGraphFindResult
+{
+	// List of nodes that should output result for their pins.
+	static const TStaticArray<TSubclassOf<UPCGSettings>, 1> ExceptionNodesToShowPins = { UPCGUserParameterGetSettings::StaticClass() };
+}
+
 FPCGEditorGraphFindResult::FPCGEditorGraphFindResult(const FString& InValue)
 : Parent(nullptr)
 , Value(InValue)
@@ -911,7 +919,11 @@ void SPCGEditorGraphFind::MatchTokensInternal(const TArray<FString>& InTokens, U
 			GetOrCreateNodeResult()->bIsMatch = true;
 		}
 
-		if (bShowPinResults)
+		const UPCGEditorGraphNodeBase* PCGNode = Cast<const UPCGEditorGraphNodeBase>(Node);
+		const UPCGSettings* NodeSettings = PCGNode ? PCGNode->GetSettings() : nullptr;
+		TSubclassOf<UPCGSettings> NodeSettingsClass(NodeSettings ? NodeSettings->GetClass() : nullptr);
+		
+		if (bShowPinResults || (NodeSettingsClass && Algo::Contains(PCGEditorGraphFindResult::ExceptionNodesToShowPins, NodeSettingsClass)))
 		{
 			for (UEdGraphPin* Pin : Node->Pins)
 			{
diff --git a/Engine/Plugins/PCG/Source/PCGEditor/Public/DataVisualizations/PCGDataVisualizationHelpers.h b/Engine/Plugins/PCG/Source/PCGEditor/Public/DataVisualizations/PCGDataVisualizationHelpers.h
--- a/Engine/Plugins/PCG/Source/PCGEditor/Public/DataVisualizations/PCGDataVisualizationHelpers.h
+++ b/Engine/Plugins/PCG/Source/PCGEditor/Public/DataVisualizations/PCGDataVisualizationHelpers.h
@@ -3,11 +3,15 @@
 #pragma once
 
 #include "PCGDataVisualization.h"
+#include "Metadata/PCGMetadataCommon.h"
 
 class UPCGData;
 
 namespace PCGDataVisualizationHelpers
 {
 	FPCGTableVisualizerColumnInfo CreateColumnInfo(const UPCGData* Data, FName Id, const FText& Label, TFunction<TSharedPtr<const IPCGAttributeAccessor>()> CreateAccessorFunc = nullptr);
-	void CreateMetadataColumnInfos(const UPCGData* Data, TArray<FPCGTableVisualizerColumnInfo>& OutColumnInfos);
+
+	void CreateMetadataColumnInfos(const UPCGData* Data, TArray<FPCGTableVisualizerColumnInfo>& OutColumnInfos, const FPCGMetadataDomainID& DomainID = PCGMetadataDomainID::Default);
+
+	FPCGTableVisualizerInfo CreateDefaultMetadataColumnInfos(const UPCGData* Data, const FPCGMetadataDomainID& DomainID = PCGMetadataDomainID::Default);
 }
diff --git a/Engine/Plugins/PCG/Source/PCGEditor/Public/DataVisualizations/PCGParamDataVisualization.h b/Engine/Plugins/PCG/Source/PCGEditor/Public/DataVisualizations/PCGParamDataVisualization.h
--- a/Engine/Plugins/PCG/Source/PCGEditor/Public/DataVisualizations/PCGParamDataVisualization.h
+++ b/Engine/Plugins/PCG/Source/PCGEditor/Public/DataVisualizations/PCGParamDataVisualization.h
@@ -10,5 +10,5 @@ class PCGEDITOR_API IPCGParamDataVisualization : public IPCGDataVisualization
 {
 public:
 	virtual void ExecuteDebugDisplay(FPCGContext* Context, const UPCGSettingsInterface* SettingsInterface, const UPCGData* Data, AActor* TargetActor) const { /* Do nothing. */ }
-	virtual FPCGTableVisualizerInfo GetTableVisualizerInfo(const UPCGData* Data) const override;
+	virtual FPCGTableVisualizerInfo GetTableVisualizerInfoWithDomain(const UPCGData* Data, const FPCGMetadataDomainID& DomainID) const override;
 };
diff --git a/Engine/Plugins/PCG/Source/PCGEditor/Public/DataVisualizations/PCGSpatialDataVisualization.h b/Engine/Plugins/PCG/Source/PCGEditor/Public/DataVisualizations/PCGSpatialDataVisualization.h
--- a/Engine/Plugins/PCG/Source/PCGEditor/Public/DataVisualizations/PCGSpatialDataVisualization.h
+++ b/Engine/Plugins/PCG/Source/PCGEditor/Public/DataVisualizations/PCGSpatialDataVisualization.h
@@ -17,7 +17,12 @@ class PCGEDITOR_API IPCGSpatialDataVisualization : public IPCGDataVisualization
 public:
 	// ~Begin IPCGDataVisualization interface
 	virtual void ExecuteDebugDisplay(FPCGContext* Context, const UPCGSettingsInterface* SettingsInterface, const UPCGData* Data, AActor* TargetActor) const override;
-	virtual FPCGTableVisualizerInfo GetTableVisualizerInfo(const UPCGData* Data) const override;
+	virtual FPCGTableVisualizerInfo GetTableVisualizerInfoWithDomain(const UPCGData* Data, const FPCGMetadataDomainID& DomainID) const override;
+	
+	// Show the sampled points by default
+	virtual FPCGMetadataDomainID GetDefaultDomainForInspection(const UPCGData* Data) const override { return PCGMetadataDomainID::Elements; }
+	virtual FString GetDomainDisplayNameForInspection(const UPCGData* Data, const FPCGMetadataDomainID& DomainID) const override;
+	virtual TArray<FPCGMetadataDomainID> GetAllSupportedDomainsForInspection(const UPCGData* Data) const override;
 	// ~End IPCGDataVisualization interface
 
 	UE_DEPRECATED(5.6, "Implement CollapseToDebugBasePointData instead")
diff --git a/Engine/Plugins/PCG/Source/PCGEditor/Public/DataVisualizations/PCGSplineDataVisualization.h b/Engine/Plugins/PCG/Source/PCGEditor/Public/DataVisualizations/PCGSplineDataVisualization.h
--- a/Engine/Plugins/PCG/Source/PCGEditor/Public/DataVisualizations/PCGSplineDataVisualization.h
+++ b/Engine/Plugins/PCG/Source/PCGEditor/Public/DataVisualizations/PCGSplineDataVisualization.h
@@ -12,7 +12,7 @@ class PCGEDITOR_API IPCGSplineDataVisualization : public IPCGSpatialDataVisualiz
 {
 public:
 	// ~Begin IPCGDataVisualization interface
-	virtual FPCGTableVisualizerInfo GetTableVisualizerInfo(const UPCGData* Data) const override;
+	virtual FPCGTableVisualizerInfo GetTableVisualizerInfoWithDomain(const UPCGData* Data, const FPCGMetadataDomainID& DomainID) const override;
 	// ~End IPCGDataVisualization interface
 
 	// ~Begin IPCGSpatialDataVisualization interface
