[PCG] Add support for [REDACTED_HANDLE] attribute creation

* Add/Create attribute
* Metadata ops

#rnx

diff --git a/Engine/Plugins/PCG/Source/PCG/Private/Elements/Metadata/PCGMetadataOpElementBase.cpp b/Engine/Plugins/PCG/Source/PCG/Private/Elements/Metadata/PCGMetadataOpElementBase.cpp
--- a/Engine/Plugins/PCG/Source/PCG/Private/Elements/Metadata/PCGMetadataOpElementBase.cpp
+++ b/Engine/Plugins/PCG/Source/PCG/Private/Elements/Metadata/PCGMetadataOpElementBase.cpp
@@ -886,9 +886,7 @@ bool FPCGMetadataElementBase::PrepareDataInternal(FPCGContext* Context) const
 
 			if (OutputTarget.IsBasicAttribute())
 			{
-				const FName OutputName = OutputTarget.GetName();
-
-				FPCGMetadataAttributeBase* OutputAttribute = PCGMetadataElementCommon::ClearOrCreateAttribute<AttributeType>(OutputData->MutableMetadata(), OutputName);
+				FPCGMetadataAttributeBase* OutputAttribute = PCGMetadataElementCommon::ClearOrCreateAttribute<AttributeType>(OutputData->MutableMetadata(), OutputTarget);
 				if (!OutputAttribute)
 				{
 					return false;
diff --git a/Engine/Plugins/PCG/Source/PCG/Private/Elements/PCGCreateAttribute.cpp b/Engine/Plugins/PCG/Source/PCG/Private/Elements/PCGCreateAttribute.cpp
--- a/Engine/Plugins/PCG/Source/PCG/Private/Elements/PCGCreateAttribute.cpp
+++ b/Engine/Plugins/PCG/Source/PCG/Private/Elements/PCGCreateAttribute.cpp
@@ -12,7 +12,7 @@
 #include "Elements/Metadata/PCGMetadataElementCommon.h"
 #include "Helpers/PCGConversion.h"
 #include "Helpers/PCGPropertyHelpers.h"
-#include "Metadata/Accessors/PCGAttributeAccessorHelpers.h"
+#include "Utils/PCGLogErrors.h"
 
 #include UE_INLINE_GENERATED_CPP_BY_NAME(PCGCreateAttribute)
 
@@ -39,7 +39,7 @@ namespace PCGCreateAttributeConstants
 
 namespace PCGCreateAttribute
 {
-	FPCGMetadataAttributeBase* ClearOrCreateAttribute(const FPCGMetadataTypesConstantStruct& AttributeTypes, UPCGMetadata* Metadata, const FName OutputAttributeName)
+	FPCGMetadataAttributeBase* ClearOrCreateAttribute(const FPCGMetadataTypesConstantStruct& AttributeTypes, FPCGMetadataDomain* Metadata, const FName OutputAttributeName)
 	{
 		check(Metadata);
 
@@ -136,21 +136,21 @@ FString UPCGAddAttributeSettings::GetAdditionalTitleInformation() const
 		}
 		else
 		{
-			const FName SourceParamAttributeName = InputSource.GetName();
-			const FName OutputAttributeName = OutputTarget.CopyAndFixSource(&InputSource, nullptr).GetName();
-			if ((OutputAttributeName == NAME_None) && (SourceParamAttributeName == NAME_None))
+			const FString SourceParamAttributeName = InputSource.ToString();
+			const FString OutputAttributeName = OutputTarget.CopyAndFixSource(&InputSource, nullptr).ToString();
+			if (OutputAttributeName.IsEmpty() && SourceParamAttributeName.IsEmpty())
 			{
 				return PCGCreateAttributeConstants::NodeNameAddAttribute.ToString();
 			}
 			else
 			{
-				return ((OutputAttributeName == NAME_None) ? SourceParamAttributeName : OutputAttributeName).ToString();
+				return OutputAttributeName.IsEmpty() ? SourceParamAttributeName : OutputAttributeName;
 			}
 		}
 	}
 	else
 	{
-		return FString::Printf(TEXT("%s: %s"), *OutputTarget.GetName().ToString(), *AttributeTypes.ToString());
+		return FString::Printf(TEXT("%s: %s"), *OutputTarget.ToString(), *AttributeTypes.ToString());
 	}
 }
 
@@ -275,7 +275,7 @@ FString UPCGCreateAttributeSetSettings::GetAdditionalTitleInformation() const
 		return PCGCreateAttributeConstants::NodeNameAddAttribute.ToString();
 	}
 
-	return FString::Printf(TEXT("%s: %s"), *OutputTarget.GetName().ToString(), *AttributeTypes.ToString());
+	return FString::Printf(TEXT("%s: %s"), *OutputTarget.ToString(), *AttributeTypes.ToString());
 }
 
 void UPCGCreateAttributeSetSettings::ApplyPreconfiguredSettings(const FPCGPreConfiguredSettingsInfo& PreconfigureInfo)
@@ -399,7 +399,14 @@ bool FPCGAddAttributeElement::ExecuteInternal(FPCGContext* Context) const
 
 			UPCGData* OutputData = InData->DuplicateData(Context);
 			check(OutputData);
-			UPCGMetadata* OutputMetadata = OutputData->MutableMetadata();
+			FPCGMetadataDomain* OutputMetadata = OutputData->MutableMetadata()->GetMetadataDomainFromSelector(Settings->OutputTarget);
+
+			if (!OutputMetadata)
+			{
+				PCGLog::Metadata::LogInvalidMetadataDomain(Settings->OutputTarget, Context);
+				continue;
+			}
+
 			if (!PCGCreateAttribute::ClearOrCreateAttribute(Settings->AttributeTypes, OutputMetadata, OutputAttributeName))
 			{
 				PCGE_LOG(Error, GraphAndLog, FText::Format(PCGCreateAttributeConstants::ErrorCreatingAttributeMessage, FText::FromName(OutputAttributeName)));
@@ -489,11 +496,28 @@ bool FPCGCreateAttributeElement::ExecuteInternal(FPCGContext* Context) const
 
 	FName OutputAttributeName = Settings->OutputTarget.GetName();
 
+	const UPCGParamData* DefaultParam = GetDefault<UPCGParamData>();
+	const FPCGMetadataDomainID OutputMetadataDomainID = DefaultParam->GetMetadataDomainIDFromSelector(Settings->OutputTarget);
+	if (!DefaultParam->IsSupportedMetadataDomainID(OutputMetadataDomainID))
+	{
+		PCGLog::Metadata::LogInvalidMetadataDomain(Settings->OutputTarget, Context);
+		return true;
+	}
+
 	UPCGParamData* OutputData = FPCGContext::NewObject_AnyThread<UPCGParamData>(Context);
-	check(OutputData && OutputData->Metadata);
-	OutputData->Metadata->AddEntry();
+	check(OutputData);
+
+	FPCGMetadataDomain* OutputMetadata = OutputData->MutableMetadata()->GetMetadataDomainFromSelector(Settings->OutputTarget);
+
+	if (!OutputMetadata)
+	{
+		PCGLog::Metadata::LogInvalidMetadataDomain(Settings->OutputTarget, Context);
+		return true;
+	}
+
+	OutputMetadata->AddEntry();
 
-	if (!PCGCreateAttribute::ClearOrCreateAttribute(Settings->AttributeTypes, OutputData->Metadata, OutputAttributeName))
+	if (!PCGCreateAttribute::ClearOrCreateAttribute(Settings->AttributeTypes, OutputMetadata, OutputAttributeName))
 	{
 		PCGE_LOG(Error, GraphAndLog, FText::Format(PCGCreateAttributeConstants::ErrorCreatingAttributeMessage, FText::FromName(OutputAttributeName)));
 		return true;
diff --git a/Engine/Plugins/PCG/Source/PCG/Private/Helpers/PCGMetadataHelpers.cpp b/Engine/Plugins/PCG/Source/PCG/Private/Helpers/PCGMetadataHelpers.cpp
--- a/Engine/Plugins/PCG/Source/PCG/Private/Helpers/PCGMetadataHelpers.cpp
+++ b/Engine/Plugins/PCG/Source/PCG/Private/Helpers/PCGMetadataHelpers.cpp
@@ -126,6 +126,7 @@ namespace PCGMetadataHelpers
 	}
 
 
+	// @todo_pcg: Metadata -> This will copy only between default metadata. We need to handle cross domain, with restrictions?
 	bool CopyAttributes(UPCGData* TargetData, const UPCGData* SourceData, const TArray<TTuple<FPCGAttributePropertyInputSelector, FPCGAttributePropertyOutputSelector, EPCGMetadataTypes>>& AttributeSelectorsWithOutputType, bool bSameOrigin, FPCGContext* OptionalContext)
 	{
 		check(TargetData && SourceData);
diff --git a/Engine/Plugins/PCG/Source/PCG/Public/Elements/Metadata/PCGMetadataElementCommon.h b/Engine/Plugins/PCG/Source/PCG/Public/Elements/Metadata/PCGMetadataElementCommon.h
--- a/Engine/Plugins/PCG/Source/PCG/Public/Elements/Metadata/PCGMetadataElementCommon.h
+++ b/Engine/Plugins/PCG/Source/PCG/Public/Elements/Metadata/PCGMetadataElementCommon.h
@@ -29,19 +29,18 @@ namespace PCGMetadataElementCommon
 
 	/** Creates a new attribute, or clears the attribute if it already exists and is a 'T' type. If default value not provided, will take the zero value for that type. */
 	template<typename T>
-	FPCGMetadataAttribute<T>* ClearOrCreateAttribute(UPCGMetadata* Metadata, const FName& DestinationAttribute, T DefaultValue = PCG::Private::MetadataTraits<T>::ZeroValue())
+	FPCGMetadataAttribute<T>* ClearOrCreateAttribute(FPCGMetadataDomain* Metadata, const FName& DestinationAttribute, T DefaultValue = PCG::Private::MetadataTraits<T>::ZeroValue())
 	{
 		TRACE_CPUPROFILER_EVENT_SCOPE(PCGMetadataElementCommon::ClearOrCreateAttribute);
 
 		if (!Metadata)
 		{
-			UE_LOG(LogPCG, Error, TEXT("Failed to create metadata"));
 			return nullptr;
 		}
 
 		if (Metadata->HasAttribute(DestinationAttribute))
 		{
-			UE_LOG(LogPCG, Verbose, TEXT("Attribute %s already exists and has been overwritten"), *DestinationAttribute.ToString());
+			UE_LOG(LogPCG, Verbose, TEXT("Attribute %s already exists and will be overwritten"), *DestinationAttribute.ToString());
 			Metadata->DeleteAttribute(DestinationAttribute);
 		}
 
@@ -49,6 +48,37 @@ namespace PCGMetadataElementCommon
 
 		return static_cast<FPCGMetadataAttribute<T>*>(Metadata->GetMutableAttribute(DestinationAttribute));
 	}
+	
+	/** Creates a new attribute, or clears the attribute if it already exists and is a 'T' type. If default value not provided, will take the zero value for that type. */
+	template<typename T>
+	FPCGMetadataAttribute<T>* ClearOrCreateAttribute(UPCGMetadata* Metadata, const FName& DestinationAttribute, T DefaultValue = PCG::Private::MetadataTraits<T>::ZeroValue())
+	{
+		if (!Metadata)
+		{
+			return nullptr;
+		}
+		
+		return ClearOrCreateAttribute<T>(Metadata->GetDefaultMetadataDomain(), DestinationAttribute, std::move(DefaultValue));
+	}
+
+	/** Creates a new attribute, or clears the attribute if it already exists and is a 'T' type. If default value not provided, will take the zero value for that type. */
+	template<typename T>
+	FPCGMetadataAttribute<T>* ClearOrCreateAttribute(UPCGMetadata* Metadata, const FPCGAttributePropertySelector& DestinationAttribute, T DefaultValue = PCG::Private::MetadataTraits<T>::ZeroValue(), FPCGContext* InOptionalContext = nullptr)
+	{
+		if (!Metadata)
+		{
+			return nullptr;
+		}
+
+		FPCGMetadataDomain* MetadataDomain = Metadata->GetMetadataDomainFromSelector(DestinationAttribute);
+		if (!MetadataDomain)
+		{
+			PCGLog::Metadata::LogInvalidMetadataDomain(DestinationAttribute, InOptionalContext);
+			return nullptr;
+		}
+		
+		return ClearOrCreateAttribute<T>(MetadataDomain, DestinationAttribute.GetAttributeName(), std::move(DefaultValue));
+	}
 
 	constexpr int32 DefaultChunkSize = 256;
 
