[PCG] Renaming Get User Parameter nodes will rename the graph parameter

* Model (PCG) and View (Widget) nodes restructured
* Widgets created by [REDACTED_PERSON] pattern
* Renaming has been stripped from Base class and encapsulated into PCGGraphNode and subclasses
- Custom renaming logic removed in favor of underlying UE rename logic
+ Renaming a Get User Parameter node will rename the graph parameter

diff --git a/Engine/Plugins/PCG/Source/PCG/Private/PCGGraph.cpp b/Engine/Plugins/PCG/Source/PCG/Private/PCGGraph.cpp
--- a/Engine/Plugins/PCG/Source/PCG/Private/PCGGraph.cpp
+++ b/Engine/Plugins/PCG/Source/PCG/Private/PCGGraph.cpp
@@ -142,6 +142,28 @@ EPropertyBagResult UPCGGraphInterface::SetGraphParameter(const FName PropertyNam
 	return Result;
 }
 
+#if WITH_EDITOR
+EPropertyBagAlterationResult UPCGGraphInterface::RenameUserParameter(const FName CurrentName, const FName NewName)
+{
+	EPropertyBagAlterationResult Result = EPropertyBagAlterationResult::SourcePropertyNotFound;
+	if (FInstancedPropertyBag* UserParametersStruct = GetMutableUserParametersStruct())
+	{
+		if (UserParametersStruct->FindPropertyDescByName(CurrentName))
+		{
+			Modify();
+			Result = UserParametersStruct->RenameProperty(CurrentName, NewName);
+
+			if (Result == EPropertyBagAlterationResult::Success)
+			{
+				OnGraphParametersChanged(EPCGGraphParameterEvent::PropertyRenamed, NewName);
+			}
+		}
+	}
+
+	return Result;
+}
+#endif // WITH_EDITOR
+
 bool UPCGGraphInterface::UpdateArrayGraphParameter(const FName PropertyName, TFunctionRef<bool(FPropertyBagArrayRef& PropertyBagArrayRef)> Callback)
 {
 	FInstancedPropertyBag* UserParameters = GetMutableUserParametersStruct();
@@ -584,21 +606,20 @@ UPCGNode* UPCGGraph::FindNodeWithSettings(const UPCGSettingsInterface* InSetting
 	return NodeFound;
 }
 
-UPCGNode* UPCGGraph::FindNodeByTitleName(FName NodeTitle, bool bRecursive) const
+UPCGNode* UPCGGraph::FindNodeByTitleName(FName NodeTitle, bool bRecursive, const TSubclassOf<const UPCGSettings> OptionalClass) const
 {
 	UPCGNode* NodeFound = nullptr;
 
-	auto FindNode = [&NodeFound, NodeTitle](UPCGNode* InNode)
+	auto FindNode = [&NodeFound, NodeTitle, OptionalClass](UPCGNode* InNode)
 	{
-		if (InNode && InNode->NodeTitle == NodeTitle)
+		const UPCGSettings* Settings = InNode ? InNode->GetSettings() : nullptr;
+		if (Settings && (!OptionalClass || Settings->IsA(OptionalClass)) && InNode->NodeTitle == NodeTitle)
 		{
 			NodeFound = InNode;
 			return false; // stop execution
 		}
-		else
-		{
-			return true;
-		}
+
+		return true;
 	};
 
 	if (bRecursive)
@@ -1892,6 +1913,7 @@ void UPCGGraph::OnGraphParametersChanged(EPCGGraphParameterEvent InChangeType, F
 				{
 					const FName OldName = Settings->PropertyName;
 					Settings->UpdatePropertyName(PropertyDesc->Name);
+					Node->NodeTitle = PropertyDesc->Name;
 					// We make sure to keep the edges connected, by renaming the pin label
 					Node->RenameOutputPin(OldName, PropertyDesc->Name);
 				}
diff --git a/Engine/Plugins/PCG/Source/PCG/Public/Elements/PCGCollapseElement.h b/Engine/Plugins/PCG/Source/PCG/Public/Elements/PCGCollapseElement.h
--- a/Engine/Plugins/PCG/Source/PCG/Public/Elements/PCGCollapseElement.h
+++ b/Engine/Plugins/PCG/Source/PCG/Public/Elements/PCGCollapseElement.h
@@ -21,6 +21,7 @@ public:
 	virtual EPCGSettingsType GetType() const override { return EPCGSettingsType::Spatial; }
 	virtual bool ShouldDrawNodeCompact() const override { return true; }
 	virtual bool GetCompactNodeIcon(FName& OutCompactNodeIcon) const override;
+	virtual bool CanUserEditTitle() const override { return false; }
 #endif
 	virtual bool HasExecutionDependencyPin() const override { return false; }
 
diff --git a/Engine/Plugins/PCG/Source/PCG/Public/Elements/PCGFilterByType.h b/Engine/Plugins/PCG/Source/PCG/Public/Elements/PCGFilterByType.h
--- a/Engine/Plugins/PCG/Source/PCG/Public/Elements/PCGFilterByType.h
+++ b/Engine/Plugins/PCG/Source/PCG/Public/Elements/PCGFilterByType.h
@@ -20,6 +20,7 @@ public:
 	virtual FText GetNodeTooltipText() const override;
 	virtual bool ShouldDrawNodeCompact() const override { return true; }
 	virtual bool GetCompactNodeIcon(FName& OutCompactNodeIcon) const override;
+	virtual bool CanUserEditTitle() const override { return false; }
 #endif
 	virtual bool HasExecutionDependencyPin() const override { return false; }
 	
diff --git a/Engine/Plugins/PCG/Source/PCG/Public/Elements/PCGUserParameterGet.h b/Engine/Plugins/PCG/Source/PCG/Public/Elements/PCGUserParameterGet.h
--- a/Engine/Plugins/PCG/Source/PCG/Public/Elements/PCGUserParameterGet.h
+++ b/Engine/Plugins/PCG/Source/PCG/Public/Elements/PCGUserParameterGet.h
@@ -43,9 +43,6 @@ public:
 
 	//~Begin UPCGSettings interface
 	virtual bool RequiresDataFromPreTask() const override { return true; }
-#if WITH_EDITOR
-	virtual bool CanUserEditTitle() const override { return false; }
-#endif
 
 protected:
 	virtual TArray<FPCGPinProperties> InputPinProperties() const override { return {}; }
@@ -56,6 +53,8 @@ protected:
 #if WITH_EDITOR
 	virtual FName GetDefaultNodeName() const override { return FName("GetGraphParameter"); }
 	virtual FText GetDefaultNodeTitle() const override { return NSLOCTEXT("PCGUserParameterGetSettings", "NodeTitle", "Get Graph Parameter"); }
+	virtual bool ShouldDrawNodeCompact() const override { return true; }
+	virtual bool CanUserEditTitle() const override { return true; }
 	virtual EPCGSettingsType GetType() const override { return EPCGSettingsType::GraphParameters; }
 	virtual TArray<FPCGPreconfiguredInfo> GetConversionInfo() const override;
 #endif
diff --git a/Engine/Plugins/PCG/Source/PCG/Public/PCGGraph.h b/Engine/Plugins/PCG/Source/PCG/Public/PCGGraph.h
--- a/Engine/Plugins/PCG/Source/PCG/Public/PCGGraph.h
+++ b/Engine/Plugins/PCG/Source/PCG/Public/PCGGraph.h
@@ -211,6 +211,11 @@ public:
 
 	EPropertyBagResult SetGraphParameter(const FName PropertyName, const uint64 Value, const UEnum* Enum);
 
+#if WITH_EDITOR
+	/** Attempts a rename of a user parameter. */
+	EPropertyBagAlterationResult RenameUserParameter(FName CurrentName, FName NewName);
+#endif // WITH_EDITOR
+
 	/**
 	 * Allows to manipulate directly FPropertyBagArrayRef, while propagating changes to child instances. 
 	 * @param PropertyName Name of the property to access. Must be an Array.
@@ -361,7 +366,7 @@ public:
 	UPCGNode* FindNodeWithSettings(const UPCGSettingsInterface* InSettings, bool bRecursive = false) const;
 
 	/** Returns the first node that matches the given name in the graph, if any.*/
-	UPCGNode* FindNodeByTitleName(FName NodeTitle, bool bRecursive = false) const;
+	UPCGNode* FindNodeByTitleName(FName NodeTitle, bool bRecursive = false, TSubclassOf<const UPCGSettings> OptionalClass = nullptr) const;
 
 	const TArray<UPCGNode*>& GetNodes() const { return Nodes; }
 	void AddNode(UPCGNode* InNode);
diff --git a/Engine/Plugins/PCG/Source/PCG/Public/PCGInputOutputSettings.h b/Engine/Plugins/PCG/Source/PCG/Public/PCGInputOutputSettings.h
--- a/Engine/Plugins/PCG/Source/PCG/Public/PCGInputOutputSettings.h
+++ b/Engine/Plugins/PCG/Source/PCG/Public/PCGInputOutputSettings.h
@@ -44,6 +44,7 @@ public:
 	virtual FText GetDefaultNodeTitle() const override { return bIsInput ? NSLOCTEXT("PCGGraphInputOutputSettings", "InputNodeTitle", "Input Node") : NSLOCTEXT("PCGGraphInputOutputSettings", "OutputNodeTitle", "Output Node"); }
 	virtual EPCGSettingsType GetType() const override { return EPCGSettingsType::InputOutput; }
 	virtual void ApplyDeprecationBeforeUpdatePins(UPCGNode* InOutNode, TArray<TObjectPtr<UPCGPin>>& InputPins, TArray<TObjectPtr<UPCGPin>>& OutputPins) override;
+	virtual bool CanUserEditTitle() const override { return false; }
 #endif // WITH_EDITOR
 
 protected:
diff --git a/Engine/Plugins/PCG/Source/PCGEditor/Private/Nodes/PCGEditorGraphNode.cpp b/Engine/Plugins/PCG/Source/PCGEditor/Private/Nodes/PCGEditorGraphNode.cpp
--- a/Engine/Plugins/PCG/Source/PCGEditor/Private/Nodes/PCGEditorGraphNode.cpp
+++ b/Engine/Plugins/PCG/Source/PCGEditor/Private/Nodes/PCGEditorGraphNode.cpp
@@ -15,12 +15,17 @@
 #include "ToolMenuSection.h"
 #include "Framework/Commands/GenericCommands.h"
 
+#include "SPCGEditorGraphNode.h"
+#include "SPCGEditorGraphNodeCompact.h"
+
 #define LOCTEXT_NAMESPACE "PCGEditorGraphNode"
 
-UPCGEditorGraphNode::UPCGEditorGraphNode(const FObjectInitializer& ObjectInitializer)
-	: Super(ObjectInitializer)
+void UPCGEditorGraphNode::Construct(UPCGNode* InPCGNode)
 {
-	bCanRenameNode = true;
+	Super::Construct(InPCGNode);
+
+	const UPCGSettings* Settings = InPCGNode->GetSettings();
+	bCanRenameNode = Settings && Settings->CanUserEditTitle();
 }
 
 FText UPCGEditorGraphNode::GetNodeTitle(ENodeTitleType::Type TitleType) const
@@ -34,13 +39,17 @@ FText UPCGEditorGraphNode::GetNodeTitle(ENodeTitleType::Type TitleType) const
 	{
 		return PCGNode->GetNodeTitle(EPCGNodeTitleType::FullTitle);
 	}
+	else if (TitleType == ENodeTitleType::MenuTitle)
+	{
+		return FText::FromName(PCGNode->NodeTitle);
+	}
 	else
 	{
 		return PCGNode->GetNodeTitle(EPCGNodeTitleType::ListView);
 	}
 }
 
-void UPCGEditorGraphNode::GetNodeContextMenuActions(UToolMenu* Menu, class UGraphNodeContextMenuContext* Context) const
+void UPCGEditorGraphNode::GetNodeContextMenuActions(UToolMenu* Menu, UGraphNodeContextMenuContext* Context) const
 {
 	if (!Context->Node)
 	{
@@ -68,43 +77,51 @@ void UPCGEditorGraphNode::AllocateDefaultPins()
 	}
 }
 
-void UPCGEditorGraphNode::OnRenameNode(const FString& NewName)
+TSharedPtr<SGraphNode> UPCGEditorGraphNode::CreateVisualWidget()
 {
-	if (CanRenameNode(NewName))
-	{
-		CommitNodeRename(NewName);
-	}
+	return ShouldDrawCompact() ? SNew(SPCGEditorGraphNodeCompact, this) : SNew(SPCGEditorGraphNode, this);
 }
 
-bool UPCGEditorGraphNode::CanRenameNode(const FString& NewName) const
+bool UPCGEditorGraphNode::ShouldDrawCompact() const
+{
+	const UPCGSettings* Settings = GetSettings();
+	return Settings && Settings->ShouldDrawNodeCompact();
+}
+
+bool UPCGEditorGraphNode::GetCompactNodeIcon(FName& OutCompactNodeIcon) const
+{
+	const UPCGSettings* Settings = GetSettings();
+	return Settings && Settings->GetCompactNodeIcon(OutCompactNodeIcon);
+}
+
+void UPCGEditorGraphNode::OnRenameNode(const FString& NewName)
 {
-	if (!GetCanRenameNode())
+	check(PCGNode);
+
+	const FName Name = FName(NewName);
+	if (PCGNode->NodeTitle != Name)
 	{
-		return false;
+		Modify();
+		PCGNode->Modify();
+		PCGNode->NodeTitle = Name;
 	}
+}
 
-	if (NewName.Len() >= NAME_SIZE)
+bool UPCGEditorGraphNode::OnValidateNodeTitle(const FText& NewName, FText& OutErrorMessage)
+{
+	if (NewName.IsEmpty())
 	{
-		UE_LOG(LogPCGEditor, Error, TEXT("New name for PCG node is too long."));
+		OutErrorMessage = LOCTEXT("InvalidNodeTitleEmptyName", "Empty name");
 		return false;
 	}
 
-	if (!PCGNode)
+	if (NewName.ToString().Len() > MaxNodeNameCharacterCount)
 	{
+		OutErrorMessage = LOCTEXT("InvalidNodeTitleTooLong", "Name too long");
 		return false;
 	}
 
 	return true;
 }
 
-void UPCGEditorGraphNode::CommitNodeRename(const FString& NewName)
-{
-	if (PCGNode->GetAuthoredTitleLine().ToString() != NewName)
-	{
-		Modify();
-		PCGNode->Modify();
-		PCGNode->NodeTitle = FName(*NewName);
-	}
-}
-
-#undef LOCTEXT_NAMESPACE
\ No newline at end of file
+#undef LOCTEXT_NAMESPACE
diff --git a/Engine/Plugins/PCG/Source/PCGEditor/Private/Nodes/PCGEditorGraphNode.h b/Engine/Plugins/PCG/Source/PCGEditor/Private/Nodes/PCGEditorGraphNode.h
--- a/Engine/Plugins/PCG/Source/PCGEditor/Private/Nodes/PCGEditorGraphNode.h
+++ b/Engine/Plugins/PCG/Source/PCGEditor/Private/Nodes/PCGEditorGraphNode.h
@@ -16,15 +16,25 @@ class UPCGEditorGraphNode : public UPCGEditorGraphNodeBase
 	GENERATED_BODY()
 
 public:
-	UPCGEditorGraphNode(const FObjectInitializer& ObjectInitializer);
-	
+	static constexpr int32 MaxNodeNameCharacterCount = 128;
+	static constexpr float MaxNodeTitleWidth = 256.f;
+
+	virtual void Construct(UPCGNode* InPCGNode) override;
+
 	// ~Begin UEdGraphNode interface
 	virtual FText GetNodeTitle(ENodeTitleType::Type TitleType) const override;
 	virtual void GetNodeContextMenuActions(UToolMenu* Menu, class UGraphNodeContextMenuContext* Context) const override;
 	virtual void AllocateDefaultPins() override;
+	virtual TSharedPtr<SGraphNode> CreateVisualWidget() override;
 	virtual void OnRenameNode(const FString& NewName) override;
 	// ~End UEdGraphNode interface
 
-	virtual bool CanRenameNode(const FString& NewName) const;
-	virtual void CommitNodeRename(const FString& NewName);
+	/** Puts node title on node body, reducing overall node size */
+	bool ShouldDrawCompact() const;
+
+	/** Returns custom compact node icon if available */
+	bool GetCompactNodeIcon(FName& OutCompactNodeIcon) const;
+
+	/** Validate that this is an acceptable name to rename this node. */
+	virtual bool OnValidateNodeTitle(const FText& NewName, FText& OutErrorMessage);
 };
\ No newline at end of file
diff --git a/Engine/Plugins/PCG/Source/PCGEditor/Private/Nodes/PCGEditorGraphNodeBase.cpp b/Engine/Plugins/PCG/Source/PCGEditor/Private/Nodes/PCGEditorGraphNodeBase.cpp
--- a/Engine/Plugins/PCG/Source/PCGEditor/Private/Nodes/PCGEditorGraphNodeBase.cpp
+++ b/Engine/Plugins/PCG/Source/PCGEditor/Private/Nodes/PCGEditorGraphNodeBase.cpp
@@ -140,6 +140,12 @@ namespace PCGEditorGraphNodeBase
 	}
 }
 
+UPCGEditorGraphNodeBase::UPCGEditorGraphNodeBase(const FObjectInitializer& ObjectInitializer)
+	: Super(ObjectInitializer)
+{
+	bCanRenameNode = false;
+}
+
 void UPCGEditorGraphNodeBase::Construct(UPCGNode* InPCGNode)
 {
 	check(InPCGNode);
@@ -151,7 +157,6 @@ void UPCGEditorGraphNodeBase::Construct(UPCGNode* InPCGNode)
 	NodeComment = InPCGNode->NodeComment;
 	bCommentBubblePinned = InPCGNode->bCommentBubblePinned;
 	bCommentBubbleVisible = InPCGNode->bCommentBubbleVisible;
-	bCanRenameNode = false;
 
 	if (const UPCGSettingsInterface* PCGSettingsInterface = InPCGNode->GetSettingsInterface())
 	{
@@ -792,18 +797,6 @@ EPCGChangeType UPCGEditorGraphNodeBase::UpdateStructuralVisualization(UPCGCompon
 	return ChangeType;
 }
 
-bool UPCGEditorGraphNodeBase::ShouldDrawCompact() const
-{
-	UPCGSettings* Settings = PCGNode ? PCGNode->GetSettings() : nullptr;
-	return Settings && Settings->ShouldDrawNodeCompact();
-}
-
-bool UPCGEditorGraphNodeBase::GetCompactNodeIcon(FName& OutCompactNodeIcon) const
-{
-	UPCGSettings* Settings = PCGNode ? PCGNode->GetSettings() : nullptr;
-	return Settings && Settings->GetCompactNodeIcon(OutCompactNodeIcon);
-}
-
 bool UPCGEditorGraphNodeBase::CanUserAddRemoveDynamicInputPins() const
 {
 	const UPCGSettings* Settings = PCGNode ? PCGNode->GetSettings() : nullptr;
@@ -916,22 +909,6 @@ bool UPCGEditorGraphNodeBase::IsOutputPinActive(const UEdGraphPin* InOutputPin)
 	return bPinActive;
 }
 
-void UPCGEditorGraphNodeBase::EnterRenamingMode()
-{
-	bCanRenameNode = true;
-
-	// Notify the SPCGEditorGraphNode that the user is trying to rename it.
-	OnNodeRenameInitiatedDelegate.ExecuteIfBound();
-}
-
-void UPCGEditorGraphNodeBase::ExitRenamingMode()
-{
-	bCanRenameNode = false;
-
-	// Update so that the node renders with the new node title.
-	OnNodeChangedDelegate.ExecuteIfBound();
-}
-
 bool UPCGEditorGraphNodeBase::IsCompatible(const UPCGPin* InputPin, const UPCGPin* OutputPin, FText& OutReason) const
 {
 	if (PCGEditorGraphSwitches::CVarCheckConnectionCycles.GetValueOnAnyThread() && InputPin && OutputPin && InputPin->Node == PCGNode)
diff --git a/Engine/Plugins/PCG/Source/PCGEditor/Private/Nodes/PCGEditorGraphNodeBase.h b/Engine/Plugins/PCG/Source/PCGEditor/Private/Nodes/PCGEditorGraphNodeBase.h
--- a/Engine/Plugins/PCG/Source/PCGEditor/Private/Nodes/PCGEditorGraphNodeBase.h
+++ b/Engine/Plugins/PCG/Source/PCGEditor/Private/Nodes/PCGEditorGraphNodeBase.h
@@ -30,7 +30,8 @@ class UPCGEditorGraphNodeBase : public UEdGraphNode
 	GENERATED_BODY()
 
 public:
-	void Construct(UPCGNode* InPCGNode);
+	UPCGEditorGraphNodeBase(const FObjectInitializer& ObjectInitializer);
+	virtual void Construct(UPCGNode* InPCGNode);
 
 	// ~Begin UObject interface
 	virtual void BeginDestroy() override;
@@ -52,7 +53,6 @@ public:
 	virtual UObject* GetJumpTargetForDoubleClick() const override;
 	virtual void OnUpdateCommentText(const FString& NewComment) override;
 	virtual void OnCommentBubbleToggled(bool bInCommentBubbleVisible) override;
-	virtual bool GetCanRenameNode() const override { return bCanRenameNode; }
 	virtual void PinDefaultValueChanged(UEdGraphPin* Pin) override;
 	virtual void PinConnectionListChanged(UEdGraphPin* Pin) override;
 	// ~End UEdGraphNode interface
@@ -87,12 +87,6 @@ public:
 	/** If the currently inspected grid size is smaller than the grid size of this node, display transparent. */
 	EPCGChangeType UpdateStructuralVisualization(class UPCGComponent* InComponentBeingDebugged, const FPCGStack* InStackBeingInspected, bool bNewlyPlaced = false);
 
-	/** Puts node title on node body, reducing overall node size */
-	bool ShouldDrawCompact() const;
-
-	/** Returns custom compact node icon if available */
-	bool GetCompactNodeIcon(FName& OutCompactNodeIcon) const;
-
 	bool CanUserAddRemoveDynamicInputPins() const;
 
 	/** The settings support default value inline constants. */
@@ -137,12 +131,6 @@ public:
 	/** Whether the given output pin was active in the previous execution. */
 	bool IsOutputPinActive(const UEdGraphPin* InOutputPin) const;
 
-	/** Marks the node as re-nameable and provokes the node to update, placing the user in rename mode. */
-	void EnterRenamingMode();
-
-	/** Disables renaming, preventing the node from entering editing mode on subsequent updates. */
-	void ExitRenamingMode();
-
 	/** Whether this OutputPin can be connected to this InputPin */
 	virtual bool IsCompatible(const UPCGPin* InputPin, const UPCGPin* OutputPin, FText& OutReason) const;
 
@@ -152,9 +140,6 @@ public:
 	DECLARE_DELEGATE(FOnPCGEditorGraphNodeChanged);
 	FOnPCGEditorGraphNodeChanged OnNodeChangedDelegate;
 
-	DECLARE_DELEGATE(FOnPCGEditorGraphNodeRenameInitiated);
-	FOnPCGEditorGraphNodeRenameInitiated OnNodeRenameInitiatedDelegate;
-
 protected:
 	static FEdGraphPinType GetPinType(const UPCGPin* InPin);
 
diff --git a/Engine/Plugins/PCG/Source/PCGEditor/Private/Nodes/PCGEditorGraphNodeFactory.cpp b/Engine/Plugins/PCG/Source/PCGEditor/Private/Nodes/PCGEditorGraphNodeFactory.cpp
--- a/Engine/Plugins/PCG/Source/PCGEditor/Private/Nodes/PCGEditorGraphNodeFactory.cpp
+++ b/Engine/Plugins/PCG/Source/PCGEditor/Private/Nodes/PCGEditorGraphNodeFactory.cpp
@@ -3,62 +3,14 @@
 #include "PCGEditorGraphNodeFactory.h"
 
 #include "PCGEditorGraphNodeBase.h"
-#include "PCGEditorGraphNodeReroute.h"
-#include "PCGNode.h"
 #include "SPCGEditorGraphNode.h"
-#include "SPCGEditorGraphVarNode.h"
-#include "Elements/PCGReroute.h"
-#include "Elements/PCGUserParameterGet.h"
-
-#include "SGraphNodeKnot.h"
-
-class SPCGEditorGraphNodeKnot : public SGraphNodeKnot
-{
-public:
-	SLATE_BEGIN_ARGS(SPCGEditorGraphNodeKnot) {}
-	SLATE_END_ARGS();
-
-	void Construct(const FArguments& InArgs, UPCGEditorGraphNodeBase* InPCGGraphNode)
-	{
-		SGraphNodeKnot::Construct(SGraphNodeKnot::FArguments(), InPCGGraphNode);
-		InPCGGraphNode->OnNodeChangedDelegate.BindSP(this, &SPCGEditorGraphNodeKnot::OnNodeChanged);
-	}
-
-private:
-	void OnNodeChanged()
-	{
-		UpdateGraphNode();
-	}
-};
-
 
 TSharedPtr<SGraphNode> FPCGEditorGraphNodeFactory::CreateNode(UEdGraphNode* InNode) const
 {
-	if (UPCGEditorGraphNodeBase* GraphNode = Cast<UPCGEditorGraphNodeBase>(InNode))
+	if (TSharedPtr<SGraphNode> VisualNode = InNode->CreateVisualWidget())
 	{
-		TSharedPtr<SGraphNode> VisualNode;
-		
-		const UPCGNode* PCGNode = GraphNode->GetPCGNode();
-		if (PCGNode && Cast<UPCGNamedRerouteBaseSettings>(PCGNode->GetSettings()))
-		{
-			SAssignNew(VisualNode, SPCGEditorGraphVarNode, GraphNode);
-		}
-		else if (PCGNode && Cast<UPCGRerouteSettings>(PCGNode->GetSettings()))
-		{
-			SAssignNew(VisualNode, SPCGEditorGraphNodeKnot, GraphNode);
-		}
-		else if (PCGNode && Cast<UPCGUserParameterGetSettings>(PCGNode->GetSettings()))
-		{
-			SAssignNew(VisualNode, SPCGEditorGraphVarNode, GraphNode);
-		}
-		else
-		{
-			SAssignNew(VisualNode, SPCGEditorGraphNode, GraphNode);
-		}
-		
 		VisualNode->SlatePrepass();
-
-		return VisualNode.ToSharedRef();
+		return VisualNode;
 	}
 
 	return nullptr;
diff --git a/Engine/Plugins/PCG/Source/PCGEditor/Private/Nodes/PCGEditorGraphNodeGetUserParameter.cpp b/Engine/Plugins/PCG/Source/PCGEditor/Private/Nodes/PCGEditorGraphNodeGetUserParameter.cpp
new file mode 100644
--- /dev/null
+++ b/Engine/Plugins/PCG/Source/PCGEditor/Private/Nodes/PCGEditorGraphNodeGetUserParameter.cpp
@@ -0,0 +1,130 @@
+// Copyright Epic Games, Inc. All Rights Reserved.
+
+#include "PCGEditorGraphNodeGetUserParameter.h"
+
+#include "PCGEditorSettings.h"
+#include "PCGGraph.h"
+#include "PCGNode.h"
+#include "Elements/PCGReroute.h"
+#include "Elements/PCGUserParameterGet.h"
+
+#include "Misc/MessageDialog.h"
+#include "StructUtils/PropertyBag.h"
+
+#include "SPCGEditorGraphNodeCompact.h"
+
+enum class EPropertyBagAlterationResult : uint8;
+
+#define LOCTEXT_NAMESPACE "PCGEditorGraphGetUserParameter"
+
+void UPCGEditorGraphGetUserParameter::OnRenameNode(const FString& NewName)
+{
+	UPCGGraph* Graph = PCGNode->GetGraph();
+	const UPCGUserParameterGetSettings* Settings = Cast<UPCGUserParameterGetSettings>(GetSettings());
+	if (ensure(Graph && Settings) && NewName != Settings->PropertyName.ToString())
+	{
+		const FText DialogMessage = FText::Format(LOCTEXT("RenameGraphParameterConfirmationMessage", "Rename from '{0}' to '{1}'?"), FText::FromName(Settings->PropertyName), FText::FromString(NewName));
+		const FText DialogTitle = LOCTEXT("RenameGraphParameterConfirmationTitle", "Rename Graph Parameter");
+
+		if (GetDefault<UPCGEditorSettings>()->bConfirmLocalGraphParameterNameChanges)
+		{
+			if (EAppReturnType::Ok != FMessageDialog::Open(EAppMsgCategory::Info, EAppMsgType::OkCancel, std::move(DialogMessage), std::move(DialogTitle)))
+			{
+				return;
+			}
+		}
+
+		const EPropertyBagAlterationResult RenameResult = Graph->RenameUserParameter(Settings->PropertyName, FName(NewName));
+		if (RenameResult != EPropertyBagAlterationResult::Success)
+		{
+			FText ReasonText;
+			switch (RenameResult)
+			{
+				case EPropertyBagAlterationResult::PropertyNameEmpty:
+					ReasonText = LOCTEXT("RenameFailureEmptyName", "Empty property name.");
+					break;
+				case EPropertyBagAlterationResult::PropertyNameInvalidCharacters:
+					ReasonText = LOCTEXT("RenameFailureEmptyName", "Invalid characters in property name.");
+					break;
+				case EPropertyBagAlterationResult::SourcePropertyNotFound:
+					ReasonText = LOCTEXT("RenameFailureEmptyName", "Source property wasn't found.");
+					break;
+				case EPropertyBagAlterationResult::TargetPropertyAlreadyExists:
+					ReasonText = LOCTEXT("RenameFailureEmptyName", "Target property already exists.");
+					break;
+				default:
+					break;
+			}
+			const FText ErrorMessage = FText::Format(LOCTEXT("RenameGraphParameterFailure", "Failed to rename graph parameter: '{0}' to '{1}'. {2}"), FText::FromName(Settings->PropertyName), FText::FromString(NewName), std::move(ReasonText));
+
+			FMessageDialog::Open(EAppMsgType::Ok, std::move(ErrorMessage));
+		}
+
+		const FString NodeTitle = GetNodeTitle(ENodeTitleType::FullTitle).ToString();
+		PCGNode->NodeTitle = FName(NodeTitle);
+		Super::OnRenameNode(NodeTitle);
+	}
+}
+
+bool UPCGEditorGraphGetUserParameter::OnValidateNodeTitle(const FText& NewName, FText& OutErrorMessage)
+{
+	if (!Super::OnValidateNodeTitle(NewName, OutErrorMessage))
+	{
+		return false;
+	}
+
+	const FName Name(NewName.ToString());
+
+	if (!FInstancedPropertyBag::IsPropertyNameValid(Name))
+	{
+		OutErrorMessage = LOCTEXT("InvalidPropertyNameInvalidCharacters", "Invalid character(s)");
+		return false;
+	}
+
+	// Prevent name clashing with any existant Named Reroute or Graph Parameter node, to avoid confusion in the graph and graph context action search menu.
+	if (const UPCGGraph* PCGGraph = PCGNode ? PCGNode->GetGraph() : nullptr)
+	{
+		if (PCGGraph->FindNodeByTitleName(Name, /*bRecursive=*/false, UPCGUserParameterGetSettings::StaticClass()))
+		{
+			OutErrorMessage = LOCTEXT("NameAlreadyInUseUserParameterErrorMessage", "Name already in use (Graph Parameter)");
+			return false;
+		}
+
+		if (PCGGraph->FindNodeByTitleName(Name, /*bRecursive=*/false, UPCGNamedRerouteDeclarationSettings::StaticClass()))
+		{
+			OutErrorMessage = LOCTEXT("NameAlreadyInUseNamedRerouteErrorMessage", "Name already in use: (Named Reroute)");
+			return false;
+		}
+	}
+
+	return true;
+}
+
+TSharedPtr<SGraphNode> UPCGEditorGraphGetUserParameter::CreateVisualWidget()
+{
+	return SNew(SPCGEditorGraphNodeCompact, this);
+}
+
+FText UPCGEditorGraphGetUserParameter::GetNodeTitle(const ENodeTitleType::Type TitleType) const
+{
+	const UPCGUserParameterGetSettings* Settings = CastChecked<UPCGUserParameterGetSettings>(GetSettings());
+	if (ensure(Settings))
+	{
+		if (TitleType == ENodeTitleType::EditableTitle)
+		{
+			return FText::FromString(FName::NameToDisplayString(Settings->PropertyName.ToString(), /*bIsBool=*/false));
+		}
+		else if (TitleType == ENodeTitleType::MenuTitle)
+		{
+			return FText::FromName(Settings->PropertyName);
+		}
+		else if (TitleType == ENodeTitleType::ListView)
+		{
+			return FText::Format(LOCTEXT("GetUserParamListViewTitleFormat", "Graph Param: {0}"), FText::FromName(Settings->PropertyName));
+		}
+	}
+
+	return Super::GetNodeTitle(TitleType);
+}
+
+#undef LOCTEXT_NAMESPACE
diff --git a/Engine/Plugins/PCG/Source/PCGEditor/Private/Nodes/PCGEditorGraphNodeGetUserParameter.h b/Engine/Plugins/PCG/Source/PCGEditor/Private/Nodes/PCGEditorGraphNodeGetUserParameter.h
new file mode 100644
--- /dev/null
+++ b/Engine/Plugins/PCG/Source/PCGEditor/Private/Nodes/PCGEditorGraphNodeGetUserParameter.h
@@ -0,0 +1,28 @@
+// Copyright Epic Games, Inc. All Rights Reserved.
+
+#pragma once
+
+#include "PCGEditorGraphNodeBase.h"
+
+#include "PCGEditorGraphNode.h"
+
+#include "PCGEditorGraphNodeGetUserParameter.generated.h"
+
+// @todo_pcg: A 'graph parameter' icon or other symbology on the node would be more UX friendly
+
+UCLASS()
+class UPCGEditorGraphGetUserParameter : public UPCGEditorGraphNode
+{
+	GENERATED_BODY()
+
+public:
+	//~ Begin UPCGEditorGraphNode Interface
+	virtual void OnRenameNode(const FString& NewName) override;
+	virtual bool OnValidateNodeTitle(const FText& NewName, FText& OutErrorMessage) override;
+	virtual TSharedPtr<SGraphNode> CreateVisualWidget() override;
+	//~ End UPCGEditorGraphNode Interface
+
+	// ~Begin UEdGraphNode interface
+	virtual FText GetNodeTitle(ENodeTitleType::Type TitleType) const override;
+	// ~End UEdGraphNode interface
+};
diff --git a/Engine/Plugins/PCG/Source/PCGEditor/Private/Nodes/PCGEditorGraphNodeInput.h b/Engine/Plugins/PCG/Source/PCGEditor/Private/Nodes/PCGEditorGraphNodeInput.h
--- a/Engine/Plugins/PCG/Source/PCGEditor/Private/Nodes/PCGEditorGraphNodeInput.h
+++ b/Engine/Plugins/PCG/Source/PCGEditor/Private/Nodes/PCGEditorGraphNodeInput.h
@@ -2,18 +2,24 @@
 
 #pragma once
 
-#include "PCGEditorGraphNodeBase.h"
+#include "PCGEditorGraphNode.h"
 
 #include "PCGEditorGraphNodeInput.generated.h"
 
 namespace ENodeTitleType { enum Type : int; }
 
 UCLASS()
-class UPCGEditorGraphNodeInput : public UPCGEditorGraphNodeBase
+class UPCGEditorGraphNodeInput : public UPCGEditorGraphNode
 {
 	GENERATED_BODY()
 
 public:
+	UPCGEditorGraphNodeInput(const FObjectInitializer& ObjectInitializer)
+		: Super(ObjectInitializer)
+	{
+		bCanRenameNode = false;
+	}
+
 	// ~Begin UEdGraphNode interface
 	virtual FText GetNodeTitle(ENodeTitleType::Type TitleType) const override;
 	virtual void AllocateDefaultPins() override;
diff --git a/Engine/Plugins/PCG/Source/PCGEditor/Private/Nodes/PCGEditorGraphNodeOutput.h b/Engine/Plugins/PCG/Source/PCGEditor/Private/Nodes/PCGEditorGraphNodeOutput.h
--- a/Engine/Plugins/PCG/Source/PCGEditor/Private/Nodes/PCGEditorGraphNodeOutput.h
+++ b/Engine/Plugins/PCG/Source/PCGEditor/Private/Nodes/PCGEditorGraphNodeOutput.h
@@ -2,18 +2,24 @@
 
 #pragma once
 
-#include "PCGEditorGraphNodeBase.h"
+#include "PCGEditorGraphNode.h"
 
 #include "PCGEditorGraphNodeOutput.generated.h"
 
 namespace ENodeTitleType { enum Type : int; }
 
 UCLASS()
-class UPCGEditorGraphNodeOutput : public UPCGEditorGraphNodeBase
+class UPCGEditorGraphNodeOutput : public UPCGEditorGraphNode
 {
 	GENERATED_BODY()
 
 public:
+	UPCGEditorGraphNodeOutput(const FObjectInitializer& ObjectInitializer)
+		: Super(ObjectInitializer)
+	{
+		bCanRenameNode = false;
+	}
+
 	// ~Begin UEdGraphNode interface
 	virtual FText GetNodeTitle(ENodeTitleType::Type TitleType) const override;
 	virtual void AllocateDefaultPins() override;
diff --git a/Engine/Plugins/PCG/Source/PCGEditor/Private/Nodes/PCGEditorGraphNodeReroute.cpp b/Engine/Plugins/PCG/Source/PCGEditor/Private/Nodes/PCGEditorGraphNodeReroute.cpp
--- a/Engine/Plugins/PCG/Source/PCGEditor/Private/Nodes/PCGEditorGraphNodeReroute.cpp
+++ b/Engine/Plugins/PCG/Source/PCGEditor/Private/Nodes/PCGEditorGraphNodeReroute.cpp
@@ -6,9 +6,41 @@
 
 #include "PCGNode.h"
 #include "Elements/PCGReroute.h"
+#include "Elements/PCGUserParameterGet.h"
 
 #include "EdGraph/EdGraphPin.h"
 
+#include "SPCGEditorGraphNodeCompact.h"
+
+#include "SGraphNodeKnot.h"
+
+#define LOCTEXT_NAMESPACE "PCGEditorGraphNodeReroute"
+
+class SPCGEditorGraphNodeKnot : public SGraphNodeKnot
+{
+public:
+	SLATE_BEGIN_ARGS(SPCGEditorGraphNodeKnot) {}
+	SLATE_END_ARGS();
+
+	void Construct(const FArguments& InArgs, UPCGEditorGraphNodeBase* InPCGGraphNode)
+	{
+		SGraphNodeKnot::Construct(SGraphNodeKnot::FArguments(), InPCGGraphNode);
+		InPCGGraphNode->OnNodeChangedDelegate.BindSP(this, &SPCGEditorGraphNodeKnot::OnNodeChanged);
+	}
+
+private:
+	void OnNodeChanged()
+	{
+		UpdateGraphNode();
+	}
+};
+
+UPCGEditorGraphNodeReroute::UPCGEditorGraphNodeReroute(const FObjectInitializer& ObjectInitializer)
+	: Super(ObjectInitializer)
+{
+	bCanRenameNode = false;
+}
+
 FText UPCGEditorGraphNodeReroute::GetNodeTitle(ENodeTitleType::Type TitleType) const
 {
 	return NSLOCTEXT("PCGEditorGraphNodeReroute", "NodeTitle", "Reroute");
@@ -53,6 +85,11 @@ UEdGraphPin* UPCGEditorGraphNodeReroute::GetPassThroughPin(const UEdGraphPin* Fr
 	}
 }
 
+TSharedPtr<SGraphNode> UPCGEditorGraphNodeReroute::CreateVisualWidget()
+{
+	return SNew(SPCGEditorGraphNodeKnot, this);
+}
+
 UEdGraphPin* UPCGEditorGraphNodeReroute::GetInputPin() const
 {
 	for (UEdGraphPin* Pin : Pins)
@@ -85,11 +122,101 @@ FText UPCGEditorGraphNodeNamedRerouteBase::GetNodeTitle(ENodeTitleType::Type Tit
 	return Super::GetNodeTitle(ENodeTitleType::ListView);
 }
 
+TSharedPtr<SGraphNode> UPCGEditorGraphNodeNamedRerouteBase::CreateVisualWidget()
+{
+	return SNew(SPCGEditorGraphNodeCompact, this);
+}
+
+bool UPCGEditorGraphNodeNamedRerouteBase::OnValidateNodeTitle(const FText& NewName, FText& OutErrorMessage)
+{
+	if (!Super::OnValidateNodeTitle(NewName, OutErrorMessage))
+	{
+		return false;
+	}
+
+	const FName Name(NewName.ToString());
+
+	// Prevent name clashing with any existant Named Reroute or Graph Parameter node, to avoid confusion in the graph and graph context action search menu.
+	if (const UPCGGraph* PCGGraph = PCGNode ? PCGNode->GetGraph() : nullptr)
+	{
+		if (PCGGraph->FindNodeByTitleName(Name, /*bRecursive=*/false, UPCGNamedRerouteDeclarationSettings::StaticClass()))
+		{
+			OutErrorMessage = LOCTEXT("NameAlreadyInUseNamedRerouteErrorMessage", "Name already in use: (Named Reroute)");
+			return false;
+		}
+
+		if (PCGGraph->FindNodeByTitleName(Name, /*bRecursive=*/false, UPCGUserParameterGetSettings::StaticClass()))
+		{
+			OutErrorMessage = LOCTEXT("NameAlreadyInUseUserParameterErrorMessage", "Name already in use (Graph Parameter)");
+			return false;
+		}
+	}
+
+	return true;
+}
+
+void UPCGEditorGraphNodeNamedRerouteUsage::OnRenameNode(const FString& NewName)
+{
+	ApplyToDeclarationNode([&NewName](UPCGEditorGraphNodeNamedRerouteDeclaration* Declaration)
+	{
+		Declaration->OnRenameNode(NewName);
+	});
+}
+
+void UPCGEditorGraphNodeNamedRerouteUsage::InheritRename(const FString& NewName)
+{
+	UPCGEditorGraphNodeNamedRerouteBase::OnRenameNode(NewName);
+}
+
 FText UPCGEditorGraphNodeNamedRerouteUsage::GetPinFriendlyName(const UPCGPin* InPin) const
 {
 	return FText::FromString(" ");
 }
 
+void UPCGEditorGraphNodeNamedRerouteUsage::ApplyToDeclarationNode(TFunctionRef<void(UPCGEditorGraphNodeNamedRerouteDeclaration*)> Action) const
+{
+	if (const UPCGNamedRerouteUsageSettings* Settings = Cast<UPCGNamedRerouteUsageSettings>(GetSettings()))
+	{
+		if (const UPCGNamedRerouteDeclarationSettings* Declaration = Cast<UPCGNamedRerouteDeclarationSettings>(Settings->Declaration))
+		{
+			if (UPCGEditorGraph* EditorGraph = Cast<UPCGEditorGraph>(GetGraph()))
+			{
+				for (const TObjectPtr<UEdGraphNode>& EdGraphNode : EditorGraph->Nodes)
+				{
+					if (UPCGEditorGraphNodeNamedRerouteDeclaration* RerouteDeclaration = Cast<UPCGEditorGraphNodeNamedRerouteDeclaration>(EdGraphNode))
+					{
+						if (RerouteDeclaration->GetSettings() == Declaration)
+						{
+							Action(RerouteDeclaration);
+							break;
+						}
+					}
+				}
+			}
+		}
+	}
+}
+
+void UPCGEditorGraphNodeNamedRerouteDeclaration::OnRenameNode(const FString& NewName)
+{
+	check(PCGNode);
+	if (TOptional<FName> ModifiedName = UPCGRerouteSettings::GetCollisionFreeNodeName(PCGNode->GetGraph(), FName(NewName)))
+	{
+		// Propagate the name change to downstream usage nodes
+		ApplyToUsageNodes([ModifiedName = *ModifiedName](UPCGEditorGraphNodeNamedRerouteUsage* RerouteNode)
+		{
+			RerouteNode->InheritRename(ModifiedName.ToString());
+		});
+
+		Super::OnRenameNode(ModifiedName->ToString());
+		ReconstructNodeOnChange();
+	}
+	else
+	{
+		Super::OnRenameNode(NewName);
+	}
+}
+
 void UPCGEditorGraphNodeNamedRerouteUsage::RebuildEdgesFromPins_Internal()
 {
 	Super::RebuildEdgesFromPins_Internal();
@@ -189,28 +316,6 @@ FString UPCGEditorGraphNodeNamedRerouteDeclaration::GenerateNodeName(const UPCGN
 	}
 }
 
-void UPCGEditorGraphNodeNamedRerouteDeclaration::OnRenameNode(const FString& NewName)
-{
-	check(PCGNode);
-	if (TOptional<FName> ModifiedName = UPCGRerouteSettings::GetCollisionFreeNodeName(PCGNode->GetGraph(), FName(NewName)))
-	{
-		FString NameString = ModifiedName->ToString();
-		if (CanRenameNode(NameString))
-		{
-			CommitNodeRename(NameString);
-			// Propagate the name change to downstream usage nodes
-			ApplyToUsageNodes([&NameString](UPCGEditorGraphNodeNamedRerouteUsage* RerouteNode)
-			{
-				RerouteNode->bCanRenameNode = true;
-				RerouteNode->OnRenameNode(NameString);
-				RerouteNode->bCanRenameNode = false;
-			});
-
-			ReconstructNodeOnChange();
-		}
-	}
-}
-
 void UPCGEditorGraphNodeNamedRerouteDeclaration::OnColorPicked(FLinearColor NewColor)
 {
 	Super::OnColorPicked(NewColor);
@@ -232,7 +337,7 @@ void UPCGEditorGraphNodeNamedRerouteDeclaration::FixNodeNameCollision()
 {
 	const FString BaseName = GetNodeTitle(ENodeTitleType::ListView).ToString();
 
-	bCanRenameNode = true;
 	OnRenameNode(BaseName);
-	bCanRenameNode = false;
-}
\ No newline at end of file
+}
+
+#undef LOCTEXT_NAMESPACE
diff --git a/Engine/Plugins/PCG/Source/PCGEditor/Private/Nodes/PCGEditorGraphNodeReroute.h b/Engine/Plugins/PCG/Source/PCGEditor/Private/Nodes/PCGEditorGraphNodeReroute.h
--- a/Engine/Plugins/PCG/Source/PCGEditor/Private/Nodes/PCGEditorGraphNodeReroute.h
+++ b/Engine/Plugins/PCG/Source/PCGEditor/Private/Nodes/PCGEditorGraphNodeReroute.h
@@ -14,6 +14,8 @@ class UPCGEditorGraphNodeReroute : public UPCGEditorGraphNode
 	GENERATED_BODY()
 
 public:
+	UPCGEditorGraphNodeReroute(const FObjectInitializer& ObjectInitializer);
+
 	// ~Begin UEdGraphNode interface
 	virtual FText GetNodeTitle(ENodeTitleType::Type TitleType) const override;
 	virtual bool ShouldOverridePinNames() const override;
@@ -22,6 +24,7 @@ public:
 	virtual bool ShouldDrawNodeAsControlPointOnly(int32& OutInputPinIndex, int32& OutOutputPinIndex) const override;
 	virtual FText GetTooltipText() const override;
 	virtual UEdGraphPin* GetPassThroughPin(const UEdGraphPin* FromPin) const override;
+	virtual TSharedPtr<SGraphNode> CreateVisualWidget() override;
 	// ~End UEdGraphNode interface
 
 	UEdGraphPin* GetInputPin() const;
@@ -36,7 +39,12 @@ class UPCGEditorGraphNodeNamedRerouteBase : public UPCGEditorGraphNode
 public:
 	// ~Begin UEdGraphNode interface
 	virtual FText GetNodeTitle(ENodeTitleType::Type TitleType) const override;
+	virtual TSharedPtr<SGraphNode> CreateVisualWidget() override;
 	// ~End UEdGraphNode interface
+
+	// ~Begin UPCGEditorGraphNode interface
+	virtual bool OnValidateNodeTitle(const FText& NewName, FText& OutErrorMessage) override;
+	// ~End UPCGEditorGraphNode interface
 };
 
 UCLASS()
@@ -45,10 +53,17 @@ class UPCGEditorGraphNodeNamedRerouteUsage : public UPCGEditorGraphNodeNamedRero
 	GENERATED_BODY()
 	friend class UPCGEditorGraphNodeNamedRerouteDeclaration;
 
+public:
+	virtual void OnRenameNode(const FString& NewName) override;
+
+	void InheritRename(const FString& NewName);
+
 protected:
-	virtual void RebuildEdgesFromPins_Internal();
+	virtual void RebuildEdgesFromPins_Internal() override;
 	virtual bool CanPickColor() const override { return false; }
 	virtual FText GetPinFriendlyName(const UPCGPin* InPin) const override;
+
+	void ApplyToDeclarationNode(TFunctionRef<void(UPCGEditorGraphNodeNamedRerouteDeclaration*)> Action) const;
 };
 
 UCLASS()
@@ -57,8 +72,13 @@ class UPCGEditorGraphNodeNamedRerouteDeclaration : public UPCGEditorGraphNodeNam
 	GENERATED_BODY()
 
 public:
-	virtual void OnRenameNode(const FString& NewName) override;
+	// ~Begin UPCGEditorGraphNodeBase interface
 	virtual void PostPaste() override;
+	// ~End UPCGEditorGraphNodeBase interface
+
+	// ~Begin UPCGEditorGraphNode interface
+	virtual void OnRenameNode(const FString& NewName) override;
+	// ~End UPCGEditorGraphNode interface
 
 	FString GenerateNodeName(const UPCGNode* FromNode, FName FromPinName);
 	void FixNodeNameCollision();
diff --git a/Engine/Plugins/PCG/Source/PCGEditor/Private/Nodes/SPCGEditorGraphNode.cpp b/Engine/Plugins/PCG/Source/PCGEditor/Private/Nodes/SPCGEditorGraphNode.cpp
--- a/Engine/Plugins/PCG/Source/PCGEditor/Private/Nodes/SPCGEditorGraphNode.cpp
+++ b/Engine/Plugins/PCG/Source/PCGEditor/Private/Nodes/SPCGEditorGraphNode.cpp
@@ -18,17 +18,13 @@
 
 #include "GraphEditorSettings.h"
 #include "IDocumentation.h"
+#include "PCGEditorGraphNode.h"
 #include "SCommentBubble.h"
 #include "SGraphPin.h"
 #include "SPinTypeSelector.h"
-#include "TutorialMetaData.h"
-#include "Algo/AnyOf.h"
 
 #include "Widgets/SBoxPanel.h"
-#include "Widgets/SToolTip.h"
-#include "Widgets/Images/SImage.h"
 #include "Widgets/Input/SButton.h"
-#include "Widgets/Layout/SBox.h"
 #include "Widgets/Text/SInlineEditableTextBlock.h"
 
 #define LOCTEXT_NAMESPACE "SPCGEditorGraphNode"
@@ -41,7 +37,6 @@ void SPCGEditorGraphNode::Construct(const FArguments& InArgs, UPCGEditorGraphNod
 	if (InNode)
 	{
 		InNode->OnNodeChangedDelegate.BindSP(this, &SPCGEditorGraphNode::OnNodeChanged);
-		InNode->OnNodeRenameInitiatedDelegate.BindSP(this, &SPCGEditorGraphNode::OnNodeRenameInitiated);
 	}
 
 	UpdateGraphNode();
@@ -70,18 +65,11 @@ void SPCGEditorGraphNode::CreateAddPinButtonWidget()
 
 void SPCGEditorGraphNode::UpdateGraphNode()
 {
-	if (PCGEditorGraphNode && PCGEditorGraphNode->ShouldDrawCompact())
-	{
-		UpdateCompactNode();
-	}
-	else
-	{
-		SGraphNode::UpdateGraphNode();
+	SGraphNode::UpdateGraphNode();
 
-		if (PCGEditorGraphNode->CanUserAddRemoveDynamicInputPins())
-		{
-			CreateAddPinButtonWidget();
-		}
+	if (PCGEditorGraphNode->CanUserAddRemoveDynamicInputPins())
+	{
+		CreateAddPinButtonWidget();
 	}
 }
 
@@ -104,12 +92,16 @@ TSharedRef<SWidget> SPCGEditorGraphNode::CreateTitleWidget(TSharedPtr<SNodeTitle
 	const bool bIsInstanceNode = (PCGEditorGraphNode && PCGEditorGraphNode->GetPCGNode() && PCGEditorGraphNode->GetPCGNode()->IsInstance());
 
 	SAssignNew(InlineEditableText, SInlineEditableTextBlock)
-		.Style(FPCGEditorStyle::Get(), bIsInstanceNode ? "PCG.Node.InstancedNodeTitleInlineEditableText" : "PCG.Node.NodeTitleInlineEditableText")
-		.Text(InNodeTitle.Get(), &SNodeTitle::GetHeadTitle)
-		.OnVerifyTextChanged(this, &SPCGEditorGraphNode::OnVerifyNameTextChanged)
-		.OnTextCommitted(this, &SPCGEditorGraphNode::OnNameTextCommited)
-		.IsReadOnly(this, &SPCGEditorGraphNode::IsNameReadOnly)
-		.IsSelected(this, &SPCGEditorGraphNode::IsSelectedExclusively);
+	.Style(FPCGEditorStyle::Get(), bIsInstanceNode ? "PCG.Node.InstancedNodeTitleInlineEditableText" : "PCG.Node.NodeTitleInlineEditableText")
+	.Text(InNodeTitle.Get(), &SNodeTitle::GetHeadTitle)
+	.OnVerifyTextChanged(this, &SPCGEditorGraphNode::OnVerifyNameTextChanged)
+	.OnTextCommitted(this, &SPCGEditorGraphNode::OnNameTextCommited)
+	.IsReadOnly(this, &SPCGEditorGraphNode::IsNameReadOnly)
+	.IsSelected(this, &SPCGEditorGraphNode::IsSelectedExclusively)
+	.MultiLine(false)
+	.MaximumLength(UPCGEditorGraphNode::MaxNodeNameCharacterCount)
+	.OverflowPolicy(ETextOverflowPolicy::Ellipsis);
+
 	InlineEditableText->SetColorAndOpacity(TAttribute<FLinearColor>::Create(TAttribute<FLinearColor>::FGetter::CreateSP(this, &SPCGEditorGraphNode::GetNodeTitleTextColor)));
 
 	return InlineEditableText.ToSharedRef();
@@ -287,111 +279,14 @@ TArray<FOverlayWidgetInfo> SPCGEditorGraphNode::GetOverlayWidgets(bool bSelected
 {
 	TArray<FOverlayWidgetInfo> OverlayWidgets = SGraphNode::GetOverlayWidgets(bSelected, WidgetSize);
 
-	if (PCGEditorGraphNode && !PCGEditorGraphNode->ShouldDrawCompact())
+	if (UsesHiGenOverlay())
 	{
-		// Higen grid size overlay widget. All magic numbers below hand tweaked to match UI mockup.
-		const EPCGHiGenGrid InspectedGrid = PCGEditorGraphNode->GetInspectedGenerationGrid();
-		UPCGNode* PCGNode = PCGEditorGraphNode->GetPCGNode();
-
-		const bool bInspectingHigen = InspectedGrid != EPCGHiGenGrid::Uninitialized;
-		if (bInspectingHigen && PCGEditorGraphNode->IsNodeEnabled())
-		{
-			FText GenerationGridText;
-			const EPCGHiGenGrid Grid = PCGEditorGraphNode->GetGenerationGrid();
-
-			if (Grid == EPCGHiGenGrid::Unbounded)
-			{
-				GenerationGridText = FText::FromString(TEXT("UB"));
-			}
-			else
-			{
-				// Meters are easier on the eyes.
-				const uint32 GridSize = PCGHiGenGrid::GridToGridSize(Grid) / 100;
-				GenerationGridText = FText::AsNumber(GridSize, &FNumberFormattingOptions::DefaultNoGrouping());
-			}
-
-			FLinearColor Tint = FLinearColor::White;
-			if (Grid != EPCGHiGenGrid::Uninitialized)
-			{
-				Tint = GetGridLabelColor(Grid);
-			}
-			else if (PCGEditorGraphNode->IsDisplayAsDisabledForced())
-			{
-				Tint.A *= 0.35f;
-			}
-
-			// Create a border brush for each combination of grids, to workaround issue where the tint does not apply
-			// to the border element.
-			const FSlateBrush* BorderBrush = GetBorderBrush(InspectedGrid, Grid);
-
-			FLinearColor TextColor = FColor::White;
-			FLinearColor BackgroundColor = FColor::Black;
-			if (InspectedGrid == Grid)
-			{
-				// Flip colors for active grid to highlight them.
-				Swap(TextColor, BackgroundColor);
-			}
-
-			TSharedPtr<SWidget> GridSizeLabel =
-				SNew(SHorizontalBox)
-				.Visibility(EVisibility::Visible)
-				+SHorizontalBox::Slot()
-				[
-					SNew(SBorder)
-					.BorderImage(BorderBrush)
-					.Padding(FMargin(12, 3))
-					.ColorAndOpacity(Tint)
-					[
-						SNew(STextBlock)
-						.TextStyle(FAppStyle::Get(), "Graph.Node.NodeTitle")
-						.Text(GenerationGridText)
-						.Justification(ETextJustify::Center)
-						.ColorAndOpacity(TextColor)
-					]
-				];
-
-			FOverlayWidgetInfo GridSizeLabelInfo(GridSizeLabel);
-			GridSizeLabelInfo.OverlayOffset = FVector2D(GetDesiredSize().X - 30.0f, -9.0f);
-
-			OverlayWidgets.Add(GridSizeLabelInfo);
-		}
+		AddHiGenOverlayWidget(OverlayWidgets);
+	}
 
-		if (PCGEditorGraphNode->GetPCGNode() && PCGEditorGraphNode->GetPCGNode()->GetSettings() && PCGEditorGraphNode->GetPCGNode()->GetSettings()->ShouldExecuteOnGPU())
-		{
-			const float BorderRadius = 7.0f;
-			const float BorderStroke = 1.0f;
-			FText GPUText = FText::FromString(TEXT("GPU"));
-			const FLinearColor BorderColor(0.5f, 0.5f, 0.5f, 0.5f);
-			FLinearColor TextColor(0.5f, 0.5f, 0.5f, 0.8f);
-			
-			const FSlateBrush* BorderBrush = new FSlateRoundedBoxBrush(
-				FLinearColor::Transparent,
-				BorderRadius,
-				BorderColor,
-				BorderStroke);
-
-			TSharedPtr<SWidget> GPUUsageLabel =
-				SNew(SHorizontalBox)
-				.Visibility(EVisibility::Visible)
-				+SHorizontalBox::Slot()
-				[
-					SNew(SBorder)
-					.BorderImage(BorderBrush)
-					.Padding(FMargin(4, 3))
-					[
-						SNew(STextBlock)
-						.TextStyle(FPCGEditorStyle::Get(), "PCG.Node.AdditionalOverlayWidgetText")
-						.Text(GPUText)
-						.Justification(ETextJustify::Center)
-						.ColorAndOpacity(TextColor)
-					]
-				];
-
-			FOverlayWidgetInfo GPUUsageLabelInfo(GPUUsageLabel);
-			GPUUsageLabelInfo.OverlayOffset = FVector2D(GetDesiredSize().X - 34.0f, GetDesiredSize().Y + 5.0f);
-
-			OverlayWidgets.Add(GPUUsageLabelInfo);
-		}
+	if (UsesGPUOverlay())
+	{
+		AddGPUOverlayWidget(OverlayWidgets);
 	}
 
 	return OverlayWidgets;
@@ -449,233 +344,17 @@ void SPCGEditorGraphNode::OnNodeChanged()
 	else
 	{
 		ExecuteOnGameThread(UE_SOURCE_LOCATION, [this]() { UpdateGraphNode(); });
-	}}
-
-void SPCGEditorGraphNode::OnNodeRenameInitiated()
-{
-	if (InlineEditableText.IsValid())
-	{
-		InlineEditableText->EnterEditingMode();
 	}
 }
 
-void SPCGEditorGraphNode::UpdateCompactNode()
+bool SPCGEditorGraphNode::UsesHiGenOverlay() const
 {
-	// Based on SGraphNodeK2Base::UpdateCompactNode. Changes:
-	// * Removed creation of tooltip widget, it did not port across trivially and the usage is fairly obvious for the current
-	//   compact nodes, but this could be re-added - TODO
-	// * Changed title style - reduced font size substantially
-	// * Layout differentiation for "pure" vs "impure" K2 nodes removed
-
-	InputPins.Empty();
-	OutputPins.Empty();
-
-	// Error handling set-up
-	SetupErrorReporting();
-
-	// Reset variables that are going to be exposed, in case we are refreshing an already setup node.
-	RightNodeBox.Reset();
-	LeftNodeBox.Reset();
-
-	if (!SWidget::GetToolTip().IsValid())
-	{
-		TSharedRef<SToolTip> DefaultToolTip = IDocumentation::Get()->CreateToolTip(TAttribute< FText >(this, &SGraphNode::GetNodeTooltip), NULL, GraphNode->GetDocumentationLink(), GraphNode->GetDocumentationExcerptName());
-		SetToolTip(DefaultToolTip);
-	}
-
-	// Setup a meta tag for this node
-	FGraphNodeMetaData TagMeta(TEXT("Graphnode"));
-	PopulateMetaTag(&TagMeta);
-	
-	TSharedPtr<SNodeTitle> NodeTitle = SNew(SNodeTitle, GraphNode);
-	TSharedRef<SOverlay> NodeOverlay = SNew(SOverlay);
-	
-	IconColor = FLinearColor::White;
-
-	// Add optional node specific widget to the overlay:
-	TSharedPtr<SWidget> OverlayWidget = GraphNode->CreateNodeImage();
-	if(OverlayWidget.IsValid())
-	{
-		NodeOverlay->AddSlot()
-		.HAlign(HAlign_Center)
-		.VAlign(VAlign_Center)
-		[
-			SNew( SBox )
-			.WidthOverride( 70.f )
-			.HeightOverride( 70.f )
-			[
-				OverlayWidget.ToSharedRef()
-			]
-		];
-	}
-
-	FName CompactBodyIcon = NAME_None;
-	check(PCGEditorGraphNode);
-	if (PCGEditorGraphNode->GetCompactNodeIcon(CompactBodyIcon))
-	{
-		NodeOverlay->AddSlot()
-			.HAlign(HAlign_Center)
-			.VAlign(VAlign_Center)
-			.Padding(45.f, 0.f, 45.f, 0.f)
-			[
-				SNew(SImage)
-					.Image(FPCGEditorStyle::Get().GetBrush(CompactBodyIcon))
-					.ColorAndOpacity(this, &SGraphNode::GetNodeTitleIconColor)
-			];
-	}
-	else
-	{
-		NodeOverlay->AddSlot()
-			.HAlign(HAlign_Center)
-			.VAlign(VAlign_Center)
-			.Padding(45.f, 0.f, 45.f, 0.f)
-			[
-				// MIDDLE
-				SNew(SVerticalBox)
-				+SVerticalBox::Slot()
-				.HAlign(HAlign_Center)
-				.AutoHeight()
-				[
-					SNew(STextBlock)
-						.TextStyle(FAppStyle::Get(), "Graph.Node.NodeTitle" )
-						.Text( NodeTitle.Get(), &SNodeTitle::GetHeadTitle )
-						.WrapTextAt(128.0f)
-						.ColorAndOpacity(this, &SGraphNode::GetNodeTitleIconColor)
-				]
-				+SVerticalBox::Slot()
-				.AutoHeight()
-				[
-					NodeTitle.ToSharedRef()
-				]
-			];
-	}
-	
-	NodeOverlay->AddSlot()
-		.HAlign(HAlign_Left)
-		.VAlign(VAlign_Center)
-		.Padding(0.f, 0.f, 55.f, 0.f)
-		[
-			// LEFT
-			SAssignNew(LeftNodeBox, SVerticalBox)
-		];
-
-	NodeOverlay->AddSlot()
-		.HAlign(HAlign_Right)
-		.VAlign(VAlign_Center)
-		.Padding(55.f, 0.f, 0.f, 0.f)
-		[
-			// RIGHT
-			SAssignNew(RightNodeBox, SVerticalBox)
-		];
-
-	//
-	//             ______________________
-	//            | (>) L |      | R (>) |
-	//            | (>) E |      | I (>) |
-	//            | (>) F |   +  | G (>) |
-	//            | (>) T |      | H (>) |
-	//            |       |      | T (>) |
-	//            |_______|______|_______|
-	//
-	this->ContentScale.Bind( this, &SGraphNode::GetContentScale );
-	
-	TSharedRef<SVerticalBox> InnerVerticalBox =
-		SNew(SVerticalBox)
-		+SVerticalBox::Slot()
-		[
-			// NODE CONTENT AREA
-			SNew( SOverlay)
-			+SOverlay::Slot()
-			[
-				SNew(SImage)
-				.Image( FAppStyle::GetBrush("Graph.VarNode.Body") )
-				.ColorAndOpacity(this, &SGraphNode::GetNodeBodyColor)
-			]
-			+ SOverlay::Slot()
-			[
-				SNew(SImage)
-				.Image( FAppStyle::GetBrush("Graph.VarNode.Gloss") )
-				.ColorAndOpacity(this, &SGraphNode::GetNodeBodyColor)
-			]
-			+SOverlay::Slot()
-			.Padding( FMargin(0,3) )
-			[
-				NodeOverlay
-			]
-		];
-	
-	TSharedPtr<SWidget> EnabledStateWidget = GetEnabledStateWidget();
-	if (EnabledStateWidget.IsValid())
-	{
-		InnerVerticalBox->AddSlot()
-			.AutoHeight()
-			.HAlign(HAlign_Fill)
-			.VAlign(VAlign_Top)
-			.Padding(FMargin(2, 0))
-			[
-				EnabledStateWidget.ToSharedRef()
-			];
-	}
-
-	InnerVerticalBox->AddSlot()
-		.AutoHeight()
-		.Padding( FMargin(5.0f, 1.0f) )
-		[
-			ErrorReporting->AsWidget()
-		];
-
-	this->GetOrAddSlot( ENodeZone::Center )
-		.HAlign(HAlign_Center)
-		.VAlign(VAlign_Center)
-		[
-			InnerVerticalBox
-		];
-
-	CreatePinWidgets();
-
-	// Hide pin labels
-	for (auto InputPin: this->InputPins)
-	{
-		if (InputPin->GetPinObj()->ParentPin == nullptr)
-		{
-			InputPin->SetShowLabel(false);
-		}
-	}
-
-	for (auto OutputPin : this->OutputPins)
-	{
-		if (OutputPin->GetPinObj()->ParentPin == nullptr)
-		{
-			OutputPin->SetShowLabel(false);
-		}
-	}
-
-	// Create comment bubble
-	TSharedPtr<SCommentBubble> CommentBubble;
-	const FSlateColor CommentColor = GetDefault<UGraphEditorSettings>()->DefaultCommentNodeTitleColor;
-
-	SAssignNew( CommentBubble, SCommentBubble )
-		.GraphNode( GraphNode )
-		.Text( this, &SGraphNode::GetNodeComment )
-		.OnTextCommitted( this, &SGraphNode::OnCommentTextCommitted )
-		.ColorAndOpacity( CommentColor )
-		.AllowPinning( true )
-		.EnableTitleBarBubble( true )
-		.EnableBubbleCtrls( true )
-		.GraphLOD( this, &SGraphNode::GetCurrentLOD )
-		.IsGraphNodeHovered( this, &SGraphNode::IsHovered );
-
-	GetOrAddSlot( ENodeZone::TopCenter )
-		.SlotOffset(TAttribute<FVector2D>(CommentBubble.Get(), &SCommentBubble::GetOffset))
-		.SlotSize(TAttribute<FVector2D>(CommentBubble.Get(), &SCommentBubble::GetSize))
-		.AllowScaling(TAttribute<bool>(CommentBubble.Get(), &SCommentBubble::IsScalingAllowed))
-		.VAlign(VAlign_Top)
-		[
-			CommentBubble.ToSharedRef()
-		];
+	return PCGEditorGraphNode->GetInspectedGenerationGrid() != EPCGHiGenGrid::Uninitialized && PCGEditorGraphNode->IsNodeEnabled();
+}
 
-	CreateInputSideAddButton(LeftNodeBox);
-	CreateOutputSideAddButton(RightNodeBox);
+bool SPCGEditorGraphNode::UsesGPUOverlay() const
+{
+	return PCGEditorGraphNode->GetSettings() && PCGEditorGraphNode->GetSettings()->ShouldExecuteOnGPU();
 }
 
 FLinearColor SPCGEditorGraphNode::GetGridLabelColor(EPCGHiGenGrid NodeGrid)
@@ -722,6 +401,112 @@ FLinearColor SPCGEditorGraphNode::GetGridLabelColor(EPCGHiGenGrid NodeGrid)
 	}
 }
 
+// @todo_pcg: Should return a FOverlayWidgetInfo, rather than updating a passed in argument array
+void SPCGEditorGraphNode::AddHiGenOverlayWidget(TArray<FOverlayWidgetInfo>& OverlayWidgets) const
+{
+	check(PCGEditorGraphNode);
+	check(UsesHiGenOverlay());
+
+	// Higen grid size overlay widget. All magic numbers below hand tweaked to match UI mockup.
+	const EPCGHiGenGrid InspectedGrid = PCGEditorGraphNode->GetInspectedGenerationGrid();
+
+	FText GenerationGridText;
+	const EPCGHiGenGrid Grid = PCGEditorGraphNode->GetGenerationGrid();
+
+	if (Grid == EPCGHiGenGrid::Unbounded)
+	{
+		GenerationGridText = FText::FromString(TEXT("UB"));
+	}
+	else
+	{
+		// Meters are easier on the eyes.
+		const uint32 GridSize = PCGHiGenGrid::GridToGridSize(Grid) / 100;
+		GenerationGridText = FText::AsNumber(GridSize, &FNumberFormattingOptions::DefaultNoGrouping());
+	}
+
+	FLinearColor Tint = FLinearColor::White;
+	if (Grid != EPCGHiGenGrid::Uninitialized)
+	{
+		Tint = GetGridLabelColor(Grid);
+	}
+	else if (PCGEditorGraphNode->IsDisplayAsDisabledForced())
+	{
+		Tint.A *= 0.35f;
+	}
+
+	// Create a border brush for each combination of grids, to workaround issue where the tint does not apply
+	// to the border element.
+	const FSlateBrush* BorderBrush = GetBorderBrush(InspectedGrid, Grid);
+
+	FLinearColor TextColor = FColor::White;
+	FLinearColor BackgroundColor = FColor::Black;
+	if (InspectedGrid == Grid)
+	{
+		// Flip colors for active grid to highlight them.
+		Swap(TextColor, BackgroundColor);
+	}
+
+	TSharedPtr<SWidget> GridSizeLabel =
+		SNew(SHorizontalBox)
+		.Visibility(EVisibility::Visible)
+		+ SHorizontalBox::Slot()
+		[
+			SNew(SBorder)
+			.BorderImage(BorderBrush)
+			.Padding(FMargin(12, 3))
+			.ColorAndOpacity(Tint)
+			[
+				SNew(STextBlock)
+				.TextStyle(FAppStyle::Get(), "Graph.Node.NodeTitle")
+				.Text(GenerationGridText)
+				.Justification(ETextJustify::Center)
+				.ColorAndOpacity(TextColor)
+			]
+		];
+
+	FOverlayWidgetInfo GridSizeLabelInfo(GridSizeLabel);
+	GridSizeLabelInfo.OverlayOffset = FVector2D(GetDesiredSize().X - 30.0f, -9.0f);
+
+	OverlayWidgets.Add(GridSizeLabelInfo);
+}
+
+// @todo_pcg: Should return a FOverlayWidgetInfo, rather than updating a passed in argument array
+void SPCGEditorGraphNode::AddGPUOverlayWidget(TArray<FOverlayWidgetInfo>& OverlayWidgets) const
+{
+	check(PCGEditorGraphNode);
+	check(UsesGPUOverlay())
+
+	constexpr float BorderRadius = 7.0f;
+	constexpr float BorderStroke = 1.0f;
+	constexpr FLinearColor BorderColor(0.5f, 0.5f, 0.5f, 0.5f);
+	constexpr FLinearColor TextColor(0.5f, 0.5f, 0.5f, 0.8f);
+	const FText GPUText = LOCTEXT("GPULabel", "GPU");
+
+	const FSlateBrush* BorderBrush = new FSlateRoundedBoxBrush(FLinearColor::Transparent, BorderRadius, BorderColor, BorderStroke);
+
+	TSharedPtr<SWidget> GPUUsageLabel =
+		SNew(SHorizontalBox)
+		.Visibility(EVisibility::Visible)
+		+ SHorizontalBox::Slot()
+		[
+			SNew(SBorder)
+			.BorderImage(BorderBrush)
+			.Padding(FMargin(4, 3))
+			[
+				SNew(STextBlock)
+				.TextStyle(FPCGEditorStyle::Get(), "PCG.Node.AdditionalOverlayWidgetText")
+				.Text(std::move(GPUText))
+				.Justification(ETextJustify::Center)
+				.ColorAndOpacity(TextColor)
+			]
+		];
+
+	FOverlayWidgetInfo GPUUsageLabelInfo(GPUUsageLabel);
+	GPUUsageLabelInfo.OverlayOffset = FVector2D(GetDesiredSize().X - 34.0f, GetDesiredSize().Y + 5.0f);
+
+	OverlayWidgets.Add(GPUUsageLabelInfo);
+}
+
 const FSlateBrush* SPCGEditorGraphNode::GetBorderBrush(EPCGHiGenGrid InspectedGrid, EPCGHiGenGrid NodeGrid) const
 {
 	if (InspectedGrid == NodeGrid)
diff --git a/Engine/Plugins/PCG/Source/PCGEditor/Private/Nodes/SPCGEditorGraphNode.h b/Engine/Plugins/PCG/Source/PCGEditor/Private/Nodes/SPCGEditorGraphNode.h
--- a/Engine/Plugins/PCG/Source/PCGEditor/Private/Nodes/SPCGEditorGraphNode.h
+++ b/Engine/Plugins/PCG/Source/PCGEditor/Private/Nodes/SPCGEditorGraphNode.h
@@ -37,14 +37,21 @@ public:
 
 protected:
 	void OnNodeChanged();
-	void OnNodeRenameInitiated();
 
-	/** Set up node in 'compact' mode */
-	void UpdateCompactNode();
+	/** Will add the Hierarchical Generation overlay to the node. */
+	virtual bool UsesHiGenOverlay() const;
+	/** Will add the GPU icon overlay to the node. */
+	virtual bool UsesGPUOverlay() const;
 
 private:
 	static FLinearColor GetGridLabelColor(EPCGHiGenGrid NodeGrid);
 
+	/** Adds the Hierarchical Generation overlay to the array, displaying the HiGen grid size on the node. */
+	void AddHiGenOverlayWidget(TArray<FOverlayWidgetInfo>& OverlayWidgets) const;
+
+	/** Adds the "GPU" tag to the node, for . */
+	void AddGPUOverlayWidget(TArray<FOverlayWidgetInfo>& OverlayWidgets) const;
+
 	/**
 	* Get the border brush for the given combination of grid sizes and enabled state. All a big workaround for FSlateRoundedBoxBrush not respecting
 	* the tint colour.
diff --git a/Engine/Plugins/PCG/Source/PCGEditor/Private/Nodes/SPCGEditorGraphNodeCompact.cpp b/Engine/Plugins/PCG/Source/PCGEditor/Private/Nodes/SPCGEditorGraphNodeCompact.cpp
new file mode 100644
--- /dev/null
+++ b/Engine/Plugins/PCG/Source/PCGEditor/Private/Nodes/SPCGEditorGraphNodeCompact.cpp
@@ -0,0 +1,381 @@
+// Copyright Epic Games, Inc. All Rights Reserved.
+
+#include "SPCGEditorGraphNodeCompact.h"
+
+#include "PCGEditorGraphNode.h"
+#include "PCGEditorGraphNodeBase.h"
+#include "PCGEditorStyle.h"
+
+#include "GraphEditorSettings.h"
+#include "IDocumentation.h"
+#include "TutorialMetaData.h"
+#include "Fonts/FontMeasure.h"
+#include "Framework/Application/SlateApplication.h"
+
+#include "SCommentBubble.h"
+#include "SGraphPin.h"
+#include "Widgets/Images/SImage.h"
+#include "Widgets/Layout/SBox.h"
+#include "Widgets/Text/SInlineEditableTextBlock.h"
+
+#define LOCTEXT_NAMESPACE "SPCGEditorGraphNodeCompact"
+
+namespace PCGEditorGraphNodeCompact
+{
+	namespace Constants
+	{
+		static constexpr float CompactNodeSize = 36.f;
+		static constexpr float CompactNodeIconSize = 20.f;
+		static constexpr float IconTitleWidth = 45.f;
+		static constexpr float TitleTextExtraWidth = 5.f;
+		static constexpr float PinExtraPadding = 8.f;
+		static constexpr float NoPinPadding = 16.f;
+
+		static constexpr float SubduedSpillColorMultiplier = 0.6f;
+	}
+
+	namespace Helpers
+	{
+		static float GetTitleTextWidth(const FText& Text, const float ContentScale = 1.f)
+		{
+			static constexpr float FontScale = 0.07f;
+			const FInlineEditableTextBlockStyle& Style = FPCGEditorStyle::Get().GetWidgetStyle<FInlineEditableTextBlockStyle>("PCG.Node.CompactNodeTitle");
+			if (const TSharedPtr<FSlateFontMeasure> FontMeasure = FSlateApplication::Get().GetRenderer()->GetFontMeasureService())
+			{
+				return FontMeasure->Measure(Text, Style.EditableTextBoxStyle.TextStyle.Font, ContentScale).X * FontScale;
+			}
+			else // Guess based on an average font size of about 0.1f per character
+			{
+				return Text.ToString().Len() * FontScale * ContentScale;
+			}
+		}
+	}
+}
+
+void SPCGEditorGraphNodeCompact::Construct(const FArguments& InArgs, UPCGEditorGraphNodeBase* InNode)
+{
+	PCGEditorGraphNode = CastChecked<UPCGEditorGraphNode>(InNode);
+	SPCGEditorGraphNode::Construct(SPCGEditorGraphNode::FArguments{}, InNode);
+}
+
+void SPCGEditorGraphNodeCompact::UpdateGraphNode()
+{
+	using namespace PCGEditorGraphNodeCompact;
+	check(PCGEditorGraphNode);
+
+	// Based on SGraphNodeK2Base::UpdateCompactNode. Changes:
+	// * Removed creation of tooltip widget, it did not port across trivially and the usage is fairly obvious for the current
+	//   compact nodes, but this could be re-added - TODO
+	// * Changed title style - reduced font size substantially
+	// * Layout differentiation for "pure" vs "impure" K2 nodes removed
+	// * Title Widget created independently
+	InputPins.Empty();
+	OutputPins.Empty();
+
+	// Error handling set-up
+	SetupErrorReporting();
+
+	// Reset variables that are going to be exposed, in case we are refreshing an already set up node.
+	RightNodeBox.Reset();
+	LeftNodeBox.Reset();
+
+	if (!SWidget::GetToolTip().IsValid())
+	{
+		// @todo_pcg: Disabled temporarily to avoid new (Clang?) static analysis warning. Fix when upgrading to advanced tooltips.
+		// TSharedPtr<SToolTip> NewToolTip = IDocumentation::Get()->CreateToolTip(TAttribute<FText>(this, &SGraphNode::GetNodeTooltip), nullptr, GraphNode->GetDocumentationLink(), GraphNode->GetDocumentationExcerptName());
+		SetToolTipText(TAttribute<FText>(this, &SGraphNode::GetNodeTooltip));
+	}
+
+	// Set up a meta tag for this node
+	FGraphNodeMetaData TagMeta(TEXT("Graphnode"));
+	PopulateMetaTag(&TagMeta);
+
+	IconColor = FLinearColor::White;
+
+	check(GraphNode)
+
+	const TSharedPtr<SWidget> TitleWidget = CreateTitleWidget(SAssignNew(NodeTitle, SNodeTitle, GraphNode));
+	FName NodeIcon = NAME_None;
+	const bool bNeedsTitle = !PCGEditorGraphNode->GetCompactNodeIcon(NodeIcon);
+
+	check(InlineEditableText || NodeIcon != NAME_None);
+
+	const bool bHasInputPins = nullptr != GraphNode->FindPinByPredicate([](const UEdGraphPin* Pin) { return Pin->Direction == EGPD_Input; });
+	const bool bHasOutputPins = nullptr != GraphNode->FindPinByPredicate([](const UEdGraphPin* Pin) { return Pin->Direction == EGPD_Output; });
+
+	const float TitleWidth =
+		bNeedsTitle
+			? Constants::TitleTextExtraWidth + Helpers::GetTitleTextWidth(InlineEditableText->GetText(), InlineEditableText->GetContentScale().X)
+			: Constants::IconTitleWidth; // Default width for icon based compact nodes
+
+	const float PinPadding = Constants::PinExtraPadding + Settings->GetInputPinPadding().GetTotalSpaceAlong<Orient_Horizontal>();
+	const float InputSidePadding = bHasInputPins ? PinPadding : Constants::NoPinPadding;
+	const float OutputSidePadding = bHasOutputPins ? PinPadding : Constants::NoPinPadding;
+	FMargin ContentMargin(TitleWidth * 0.5, 0);
+	ContentMargin.Left += InputSidePadding;
+	ContentMargin.Right += OutputSidePadding;
+
+	const FVector2D NodeSize = ContentMargin.GetDesiredSize();
+
+	//
+	//             ______________________
+	//            | (<) L |   +  | R (>) |
+	//            |_______|______|_______|
+	//
+	const TSharedPtr<SOverlay> ContentOverlay =
+	SNew(SOverlay)
+	+ SOverlay::Slot()
+	.HAlign(HAlign_Center)
+	.VAlign(VAlign_Center)
+	.Padding(ContentMargin)
+	[
+		TitleWidget.ToSharedRef()
+	]
+	+ SOverlay::Slot()
+	.HAlign(HAlign_Left)
+	.VAlign(VAlign_Center)
+	.Padding(0, 0, /*Right=*/InputSidePadding, 0)
+	[
+		// LEFT
+		SAssignNew(LeftNodeBox, SVerticalBox)
+	]
+	+ SOverlay::Slot()
+	.HAlign(HAlign_Right)
+	.VAlign(VAlign_Center)
+	.Padding(/*Left=*/OutputSidePadding, 0, 0, 0)
+	[
+		// RIGHT
+		SAssignNew(RightNodeBox, SVerticalBox)
+	];
+
+	// Add optional node specific widget to the overlay:
+	const TSharedPtr<SWidget> CustomOverlay = GraphNode->CreateNodeImage();
+	if (CustomOverlay.IsValid())
+	{
+		ContentOverlay->AddSlot()
+		   .HAlign(HAlign_Center)
+		   .VAlign(VAlign_Center)
+		   [
+			   SNew(SBox)
+			   [
+				   CustomOverlay.ToSharedRef()
+			   ]
+		   ];
+	}
+
+	check(ContentOverlay);
+
+	this->ContentScale.Bind(this, &SGraphNode::GetContentScale);
+
+	// First add the body overlay
+	const TSharedRef<SOverlay> NodeContentOverlay =
+		SNew(SOverlay)
+		.AddMetaData<FGraphNodeMetaData>(TagMeta)
+		+ SOverlay::Slot()
+		.HAlign(HAlign_Fill)
+		.VAlign(VAlign_Fill)
+		.ZOrder(0)
+		[
+			SNew(SImage)
+			.Image(FAppStyle::GetBrush("Graph.VarNode.Body"))
+			.DesiredSizeOverride(NodeSize)
+			.ColorAndOpacity(this, &SGraphNode::GetNodeBodyColor)
+			.Visibility(EVisibility::SelfHitTestInvisible)
+		]
+		// Color spill is ZOrder 1 and will be inserted next
+		+ SOverlay::Slot()
+		.HAlign(HAlign_Fill)
+		.VAlign(VAlign_Fill)
+		.ZOrder(2)
+		[
+			SNew(SImage)
+			.Image(FAppStyle::GetBrush("Graph.VarNode.Gloss"))
+			.DesiredSizeOverride(NodeSize)
+			.ColorAndOpacity(FLinearColor::White)
+			.Visibility(EVisibility::SelfHitTestInvisible)
+		]
+		+ SOverlay::Slot()
+		.HAlign(HAlign_Fill)
+		.VAlign(VAlign_Fill)
+		.ZOrder(3)
+		[
+			ContentOverlay.ToSharedRef()
+		];
+
+	// Color spill not added to icon nodes, ex. filters, conversions
+	if (bNeedsTitle)
+	{
+		NodeContentOverlay->AddSlot()
+			.ZOrder(1)
+			[
+				SNew(SImage)
+				.Image(FAppStyle::GetBrush("Graph.VarNode.ColorSpill"))
+				.DesiredSizeOverride(NodeSize)
+				.ColorAndOpacity(this, &SPCGEditorGraphNodeCompact::GetSubduedSpillColor)
+				.Visibility(EVisibility::SelfHitTestInvisible)
+			];
+	}
+
+	const TSharedRef<SVerticalBox> InnerVerticalBox =
+		SNew(SVerticalBox)
+		+ SVerticalBox::Slot()
+		.MinHeight(Constants::CompactNodeSize)
+		.MaxHeight(Constants::CompactNodeSize)
+		[
+			NodeContentOverlay
+		];
+
+	// Enabled state bar
+	const TSharedPtr<SWidget> EnabledStateWidget = GetEnabledStateWidget();
+	if (EnabledStateWidget.IsValid())
+	{
+		InnerVerticalBox->AddSlot()
+			.AutoHeight()
+			.HAlign(HAlign_Fill)
+			.VAlign(VAlign_Top)
+			.Padding(FMargin(3, 0))
+			[
+				EnabledStateWidget.ToSharedRef()
+			];
+	}
+
+	InnerVerticalBox->AddSlot()
+		.AutoHeight()
+		.Padding(FMargin(5, 1))
+		[
+			ErrorReporting->AsWidget()
+		];
+
+	this->GetOrAddSlot(ENodeZone::Center)
+		.HAlign(HAlign_Center)
+		.VAlign(VAlign_Center)
+		[
+			InnerVerticalBox
+		];
+
+	CreatePinWidgets();
+
+	// Hide pin labels
+	for (const TSharedRef<SGraphPin>& InputPin : this->InputPins)
+	{
+		if (InputPin->GetPinObj()->ParentPin == nullptr)
+		{
+			InputPin->SetShowLabel(false);
+		}
+	}
+
+	for (const TSharedRef<SGraphPin>& OutputPin : this->OutputPins)
+	{
+		if (OutputPin->GetPinObj()->ParentPin == nullptr)
+		{
+			OutputPin->SetShowLabel(false);
+		}
+	}
+
+	// @todo_pcg: Likely this is the same as non-compact. Could be factored out.
+	// Create comment bubble
+	TSharedPtr<SCommentBubble> CommentBubble;
+	const FSlateColor CommentColor = GetDefault<UGraphEditorSettings>()->DefaultCommentNodeTitleColor;
+
+	SAssignNew(CommentBubble, SCommentBubble)
+	.GraphNode(GraphNode)
+	.Text(this, &SGraphNode::GetNodeComment)
+	.OnTextCommitted(this, &SGraphNode::OnCommentTextCommitted)
+	.ColorAndOpacity(CommentColor)
+	.AllowPinning(true)
+	.EnableTitleBarBubble(true)
+	.EnableBubbleCtrls(true)
+	.GraphLOD(this, &SGraphNode::GetCurrentLOD)
+	.IsGraphNodeHovered(this, &SGraphNode::IsHovered);
+
+	GetOrAddSlot(ENodeZone::TopCenter)
+		.SlotOffset(TAttribute<FVector2D>(CommentBubble.Get(), &SCommentBubble::GetOffset))
+		.SlotSize(TAttribute<FVector2D>(CommentBubble.Get(), &SCommentBubble::GetSize))
+		.AllowScaling(TAttribute<bool>(CommentBubble.Get(), &SCommentBubble::IsScalingAllowed))
+		.VAlign(VAlign_Top)
+		[
+			CommentBubble.ToSharedRef()
+		];
+
+	// @todo_pcg: Look into using these to replace the current implementation of Dynamic Input Pins
+	CreateInputSideAddButton(LeftNodeBox);
+	CreateOutputSideAddButton(RightNodeBox);
+}
+
+TSharedRef<SWidget> SPCGEditorGraphNodeCompact::CreateTitleWidget(TSharedPtr<SNodeTitle> InNodeTitle)
+{
+	check(PCGEditorGraphNode);
+
+	// Could be text or icon as the center of the graph
+	FName CompactBodyIcon = NAME_None;
+	if (PCGEditorGraphNode->GetCompactNodeIcon(CompactBodyIcon)) // Use the icon for the title
+	{
+		const FSlateBrush* ImageBrush = FPCGEditorStyle::Get().GetBrush(CompactBodyIcon);
+		return SNew(SBox)
+			.HAlign(HAlign_Center)
+			.VAlign(VAlign_Center)
+			.MaxDesiredWidth(PCGEditorGraphNodeCompact::Constants::CompactNodeIconSize)
+			.MaxDesiredHeight(PCGEditorGraphNodeCompact::Constants::CompactNodeIconSize)
+			[
+				SNew(SImage)
+				.Image(ImageBrush)
+				.ColorAndOpacity(this, &SPCGEditorGraphNode::GetNodeTitleIconColor)
+			];
+	}
+	else // Create a title widget with the node title
+	{
+		/**
+		 * Reimplementation of the SGraphNode::CreateTitleWidget
+		 * - Add a box to limit the width
+		 * - Control the style
+		 * - Control the inner text box for retrieving the property name
+		 */
+		SAssignNew(InlineEditableText, SInlineEditableTextBlock)
+		.Style(FPCGEditorStyle::Get(), "PCG.Node.CompactNodeTitle")
+		.Text(InNodeTitle.Get(), &SNodeTitle::GetHeadTitle)
+		.ColorAndOpacity(GetNodeTitleTextColor())
+		.OnEnterEditingMode_Lambda(
+			[this]
+			{
+				InlineEditableText->SetText(PCGEditorGraphNode->GetNodeTitle(ENodeTitleType::MenuTitle));
+			})
+		.OnVerifyTextChanged(this, &SPCGEditorGraphNodeCompact::OnVerifyNameTextChanged)
+		.OnTextCommitted(this, &SPCGEditorGraphNodeCompact::OnNameTextCommited)
+		.IsReadOnly(this, &SPCGEditorGraphNodeCompact::IsNameReadOnly)
+		.IsSelected(this, &SPCGEditorGraphNodeCompact::IsSelectedExclusively)
+		.MultiLine(false)
+		.MaximumLength(UPCGEditorGraphNode::MaxNodeNameCharacterCount)
+		.OverflowPolicy(ETextOverflowPolicy::Ellipsis);
+
+		InlineEditableText->SlatePrepass(); // Prepass to calculated desired size
+
+		return SNew(SBox)
+			.MaxDesiredWidth(UPCGEditorGraphNode::MaxNodeTitleWidth)
+			[
+				InlineEditableText.ToSharedRef()
+			];
+	}
+}
+
+const FSlateBrush* SPCGEditorGraphNodeCompact::GetShadowBrush(bool bSelected) const
+{
+	check(GraphNode)
+
+	if (GraphNode->GetDesiredEnabledState() == ENodeEnabledState::Enabled)
+	{
+		return bSelected ? FAppStyle::GetBrush(TEXT("Graph.VarNode.ShadowSelected")) : FAppStyle::GetBrush(TEXT("Graph.VarNode.Shadow"));
+	}
+	else // The disabled widget adds a disabled bar, forcing the widget into a rectangular box shape. Use default Shadow Brush.
+	{
+		return SGraphNode::GetShadowBrush(bSelected);
+	}
+}
+
+FSlateColor SPCGEditorGraphNodeCompact::GetSubduedSpillColor() const
+{
+	FLinearColor SpillColor(GetNodeTitleColor().GetSpecifiedColor());
+	SpillColor.A *= PCGEditorGraphNodeCompact::Constants::SubduedSpillColorMultiplier; // Subdue the color with alpha only to prevent darkening.
+	return SpillColor;
+}
+
+#undef LOCTEXT_NAMESPACE
diff --git a/Engine/Plugins/PCG/Source/PCGEditor/Private/Nodes/SPCGEditorGraphNodeCompact.h b/Engine/Plugins/PCG/Source/PCGEditor/Private/Nodes/SPCGEditorGraphNodeCompact.h
new file mode 100644
--- /dev/null
+++ b/Engine/Plugins/PCG/Source/PCGEditor/Private/Nodes/SPCGEditorGraphNodeCompact.h
@@ -0,0 +1,28 @@
+// Copyright Epic Games, Inc. All Rights Reserved.
+
+#pragma once
+
+#include "SPCGEditorGraphNode.h"
+
+class UPCGEditorGraphNode;
+
+class SPCGEditorGraphNodeCompact : public SPCGEditorGraphNode
+{
+public:
+	SLATE_BEGIN_ARGS(SPCGEditorGraphNodeCompact) {}
+	SLATE_END_ARGS()
+
+	void Construct(const FArguments& InArgs, UPCGEditorGraphNodeBase* InNode);
+
+	//~ Begin SGraphNode Interface
+	virtual void UpdateGraphNode() override;
+	virtual TSharedRef<SWidget> CreateTitleWidget(TSharedPtr<SNodeTitle> InNodeTitle) override;
+	virtual const FSlateBrush* GetShadowBrush(bool bSelected) const override;
+	//~ End SGraphNode Interface
+
+private:
+	FSlateColor GetSubduedSpillColor() const;
+
+	UPCGEditorGraphNode* PCGEditorGraphNode = nullptr;
+	TSharedPtr<SNodeTitle> NodeTitle = nullptr;
+};
diff --git a/Engine/Plugins/PCG/Source/PCGEditor/Private/Nodes/SPCGEditorGraphVarNode.cpp b/Engine/Plugins/PCG/Source/PCGEditor/Private/Nodes/SPCGEditorGraphVarNode.cpp
deleted file mode 100644
diff --git a/Engine/Plugins/PCG/Source/PCGEditor/Private/Nodes/SPCGEditorGraphVarNode.h b/Engine/Plugins/PCG/Source/PCGEditor/Private/Nodes/SPCGEditorGraphVarNode.h
deleted file mode 100644
diff --git a/Engine/Plugins/PCG/Source/PCGEditor/Private/PCGEditor.cpp b/Engine/Plugins/PCG/Source/PCGEditor/Private/PCGEditor.cpp
--- a/Engine/Plugins/PCG/Source/PCGEditor/Private/PCGEditor.cpp
+++ b/Engine/Plugins/PCG/Source/PCGEditor/Private/PCGEditor.cpp
@@ -340,6 +340,16 @@ void FPCGEditor::OnComponentGenerated(UPCGComponent* InComponent)
 	}
 }
 
+bool FPCGEditor::OnValidateNodeTitle(const FText& NewName, UEdGraphNode* GraphNode, FText& OutErrorMessage)
+{
+	if (UPCGEditorGraphNode* PCGGraphNode = Cast<UPCGEditorGraphNode>(GraphNode))
+	{
+		return PCGGraphNode->OnValidateNodeTitle(NewName, OutErrorMessage);
+	}
+
+	return false;
+}
+
 void FPCGEditor::UpdateAfterInspectedStackChanged()
 {
 	UPCGComponent* Component = PCGComponentBeingInspected.Get();
@@ -1340,8 +1350,11 @@ void FPCGEditor::OnRenameNode()
 		return;
 	}
 
-	UPCGEditorGraphNodeBase* Node = CastChecked<UPCGEditorGraphNodeBase>(*SelectedNodes.CreateConstIterator());
-	Node->EnterRenamingMode();
+	const UEdGraphNode* SelectedNode = Cast<UEdGraphNode>(*SelectedNodes.CreateConstIterator());
+	if (SelectedNode != nullptr && SelectedNode->GetCanRenameNode())
+	{
+		GraphEditorWidget->IsNodeTitleVisible(SelectedNode, true);
+	}
 }
 
 bool FPCGEditor::CanRenameNode() const
@@ -1357,9 +1370,9 @@ bool FPCGEditor::CanRenameNode() const
 		return false;
 	}
 
-	if (const UPCGEditorGraphNodeBase* SelectedNode = Cast<UPCGEditorGraphNodeBase>(*SelectedNodes.CreateConstIterator()))
+	if (const UPCGEditorGraphNode* SelectedNode = Cast<UPCGEditorGraphNode>(*SelectedNodes.CreateConstIterator()))
 	{
-		return !SelectedNode->GetPCGNode() || !SelectedNode->GetPCGNode()->GetSettings() || SelectedNode->GetPCGNode()->GetSettings()->CanUserEditTitle();
+		return SelectedNode->GetCanRenameNode();
 	}
 	else
 	{
@@ -2676,6 +2689,7 @@ TSharedRef<SGraphEditor> FPCGEditor::CreateGraphEditorWidget()
 
 	SGraphEditor::FGraphEditorEvents InEvents;
 	InEvents.OnSelectionChanged = SGraphEditor::FOnSelectionChanged::CreateSP(this, &FPCGEditor::OnSelectedNodesChanged);
+	InEvents.OnVerifyTextCommit = FOnNodeVerifyTextCommit::CreateSP(this, &FPCGEditor::OnValidateNodeTitle);
 	InEvents.OnTextCommitted = FOnNodeTextCommitted::CreateSP(this, &FPCGEditor::OnNodeTitleCommitted);
 	InEvents.OnNodeDoubleClicked = FSingleNodeEvent::CreateSP(this, &FPCGEditor::OnNodeDoubleClicked);
 	InEvents.OnSpawnNodeByShortcut = SGraphEditor::FOnSpawnNodeByShortcut::CreateSP(this, &FPCGEditor::OnSpawnNodeByShortcut, PCGEditorGraph);
@@ -2875,20 +2889,28 @@ void FPCGEditor::OnNodeTitleCommitted(const FText& NewText, ETextCommit::Type Co
 	{
 		if (CommitInfo == ETextCommit::OnEnter || CommitInfo == ETextCommit::OnUserMovedFocus)
 		{
-			const FScopedTransaction Transaction(*FPCGEditorCommon::ContextIdentifier, LOCTEXT("PCGEditorRenameNode", "PCG Editor: Rename Node"), nullptr);
-			NodeBeingChanged->OnRenameNode(NewText.ToString());
+			FText ErrorText;
+			if (OnValidateNodeTitle(NewText, NodeBeingChanged, ErrorText))
+			{
+				const FScopedTransaction Transaction(*FPCGEditorCommon::ContextIdentifier, LOCTEXT("PCGEditorRenameNode", "PCG Editor: Rename Node"), nullptr);
+				NodeBeingChanged->OnRenameNode(NewText.ToString());
 
-			// Implementation detail: In UPCGEditorGraphNode we only set the title under certain conditions, so it calls Modify() itself.
-			// However, UEdGraphNode does not call Modify() on its own, so we should still call it in this case.
-			if (!NodeBeingChanged->IsA<UPCGEditorGraphNode>())
+				// Implementation detail: In UPCGEditorGraphNode we only set the title under certain conditions, so it calls Modify() itself.
+				// However, UEdGraphNode does not call Modify() on its own, so we should still call it in this case.
+				if (!NodeBeingChanged->IsA<UPCGEditorGraphNode>())
+				{
+					NodeBeingChanged->Modify();
+				}
+			}
+			else
 			{
-				NodeBeingChanged->Modify();
+				UE_LOG(LogPCGEditor, Warning, TEXT("%s"), *FText::Format(LOCTEXT("UnableToRenameNode", "Unable to rename node {0}. Reason: {1}"), NodeBeingChanged->GetNodeTitle(ENodeTitleType::FullTitle), std::move(ErrorText)).ToString());
 			}
 		}
 
-		if (UPCGEditorGraphNodeBase* PCGEditorNode = Cast<UPCGEditorGraphNodeBase>(NodeBeingChanged))
+		if (const UPCGEditorGraphNodeBase* PCGEditorNode = Cast<UPCGEditorGraphNodeBase>(NodeBeingChanged))
 		{
-			PCGEditorNode->ExitRenamingMode();
+			PCGEditorNode->OnNodeChangedDelegate.ExecuteIfBound();
 		}
 	}
 }
diff --git a/Engine/Plugins/PCG/Source/PCGEditor/Private/PCGEditor.h b/Engine/Plugins/PCG/Source/PCGEditor/Private/PCGEditor.h
--- a/Engine/Plugins/PCG/Source/PCGEditor/Private/PCGEditor.h
+++ b/Engine/Plugins/PCG/Source/PCGEditor/Private/PCGEditor.h
@@ -358,6 +358,9 @@ private:
 	/** Called when the component inspected is generated/cleaned */
 	void OnComponentGenerated(UPCGComponent* InComponent);
 
+	/** Called to validate a node title change. */
+	bool OnValidateNodeTitle(const FText& NewName, UEdGraphNode* GraphNode, FText& OutErrorMessage);
+
 	/** Called when the title of a node is changed */
 	void OnNodeTitleCommitted(const FText& NewText, ETextCommit::Type CommitInfo, UEdGraphNode* NodeBeingChanged);
 
diff --git a/Engine/Plugins/PCG/Source/PCGEditor/Private/PCGEditorGraph.cpp b/Engine/Plugins/PCG/Source/PCGEditor/Private/PCGEditorGraph.cpp
--- a/Engine/Plugins/PCG/Source/PCGEditor/Private/PCGEditorGraph.cpp
+++ b/Engine/Plugins/PCG/Source/PCGEditor/Private/PCGEditorGraph.cpp
@@ -9,6 +9,7 @@
 #include "Elements/PCGReroute.h"
 #include "Elements/PCGUserParameterGet.h"
 #include "Nodes/PCGEditorGraphNode.h"
+#include "Nodes/PCGEditorGraphNodeGetUserParameter.h"
 #include "Nodes/PCGEditorGraphNodeInput.h"
 #include "Nodes/PCGEditorGraphNodeOutput.h"
 #include "Nodes/PCGEditorGraphNodeReroute.h"
@@ -70,21 +71,28 @@ void UPCGEditorGraph::ReconstructGraph()
 	}
 
 	TMap<UPCGNode*, UPCGEditorGraphNodeBase*> NodeLookup;
-	const bool bSelectNewNode = false;
-
-	UPCGNode* InputNode = PCGGraph->GetInputNode();
-	FGraphNodeCreator<UPCGEditorGraphNodeInput> InputNodeCreator(*this);
-	UPCGEditorGraphNodeInput* InputGraphNode = InputNodeCreator.CreateNode(bSelectNewNode);
-	InputGraphNode->Construct(InputNode);
-	InputNodeCreator.Finalize();
-	NodeLookup.Add(InputNode, InputGraphNode);
-
-	UPCGNode* OutputNode = PCGGraph->GetOutputNode();
-	FGraphNodeCreator<UPCGEditorGraphNodeOutput> OutputNodeCreator(*this);
-	UPCGEditorGraphNodeOutput* OutputGraphNode = OutputNodeCreator.CreateNode(bSelectNewNode);
-	OutputGraphNode->Construct(OutputNode);
-	OutputNodeCreator.Finalize();
-	NodeLookup.Add(OutputNode, OutputGraphNode);
+	constexpr bool bSelectNewNode = false;
+
+	// Create input and output nodes directly.
+	{
+		UPCGNode* InputNode = PCGGraph->GetInputNode();
+		check(InputNode);
+		FGraphNodeCreator<UPCGEditorGraphNodeInput> InputNodeCreator(*this);
+		UPCGEditorGraphNodeInput* InputGraphNode = InputNodeCreator.CreateNode(bSelectNewNode);
+		InputGraphNode->Construct(InputNode);
+		InputNodeCreator.Finalize();
+		NodeLookup.Add(InputNode, InputGraphNode);
+	}
+
+	{
+		UPCGNode* OutputNode = PCGGraph->GetOutputNode();
+		check(OutputNode);
+		FGraphNodeCreator<UPCGEditorGraphNodeOutput> OutputNodeCreator(*this);
+		UPCGEditorGraphNodeOutput* OutputGraphNode = OutputNodeCreator.CreateNode(bSelectNewNode);
+		OutputGraphNode->Construct(OutputNode);
+		OutputNodeCreator.Finalize();
+		NodeLookup.Add(OutputNode, OutputGraphNode);
+	}
 
 	for (UPCGNode* PCGNode : PCGGraph->GetNodes())
 	{
@@ -93,44 +101,24 @@ void UPCGEditorGraph::ReconstructGraph()
 			continue;
 		}
 
-		// TODO: replace this with a templated lambda because this is all very similar
-		if (Cast<UPCGNamedRerouteDeclarationSettings>(PCGNode->GetSettings()))
-		{
-			FGraphNodeCreator<UPCGEditorGraphNodeNamedRerouteDeclaration> NodeCreator(*this);
-			UPCGEditorGraphNodeNamedRerouteDeclaration* RerouteGraphNode = NodeCreator.CreateNode(bSelectNewNode);
-			RerouteGraphNode->Construct(PCGNode);
-			NodeCreator.Finalize();
-			NodeLookup.Add(PCGNode, RerouteGraphNode);
-		}
-		else if (Cast<UPCGNamedRerouteUsageSettings>(PCGNode->GetSettings()))
-		{
-			FGraphNodeCreator<UPCGEditorGraphNodeNamedRerouteUsage> NodeCreator(*this);
-			UPCGEditorGraphNodeNamedRerouteUsage* RerouteGraphNode = NodeCreator.CreateNode(bSelectNewNode);
-			RerouteGraphNode->Construct(PCGNode);
-			NodeCreator.Finalize();
-			NodeLookup.Add(PCGNode, RerouteGraphNode);
-		}
-		else if (Cast<UPCGRerouteSettings>(PCGNode->GetSettings()))
-		{
-			FGraphNodeCreator<UPCGEditorGraphNodeReroute> NodeCreator(*this);
-			UPCGEditorGraphNodeReroute* RerouteGraphNode = NodeCreator.CreateNode(bSelectNewNode);
-			RerouteGraphNode->Construct(PCGNode);
-			NodeCreator.Finalize();
-			NodeLookup.Add(PCGNode, RerouteGraphNode);
-		}
-		else
-		{
-			FGraphNodeCreator<UPCGEditorGraphNode> NodeCreator(*this);
-			UPCGEditorGraphNode* GraphNode = NodeCreator.CreateNode(bSelectNewNode);
-			GraphNode->Construct(PCGNode);
-			NodeCreator.Finalize();
-			NodeLookup.Add(PCGNode, GraphNode);
-		}
+		// Create other nodes based on settings.
+		const UPCGSettings* PCGSettings = PCGNode->GetSettings();
+		const TSubclassOf<UPCGEditorGraphNodeBase> PCGGraphNodeClass =
+			PCGSettings
+				? GetGraphNodeClassFromPCGSettings(PCGSettings)
+				: TSubclassOf<UPCGEditorGraphNodeBase>(UPCGEditorGraphNode::StaticClass());
+
+		FGraphNodeCreator<UPCGEditorGraphNodeBase> NodeCreator(*this);
+		check(PCGNode && PCGGraphNodeClass);
+		// No need to select, since its reconstruction.
+		UPCGEditorGraphNodeBase* GraphNode = NodeCreator.CreateNode(bSelectNewNode, PCGGraphNodeClass);
+		GraphNode->Construct(PCGNode);
+		NodeCreator.Finalize();
+		NodeLookup.Add(PCGNode, GraphNode);
 	}
 
 	for (const auto& NodeLookupIt : NodeLookup)
 	{
-		UPCGNode* PCGNode = NodeLookupIt.Key;
 		UPCGEditorGraphNodeBase* GraphNode = NodeLookupIt.Value;
 		CreateLinks(GraphNode, /*bCreateInbound=*/false, /*bCreateOutbound=*/true, NodeLookup);
 	}
@@ -140,8 +128,7 @@ void UPCGEditorGraph::ReconstructGraph()
 		if (const UEdGraphNode* ExtraGraphNode = Cast<UEdGraphNode>(ExtraNode))
 		{
 			UEdGraphNode* NewNode = DuplicateObject(ExtraGraphNode, /*Outer=*/this);
-			const bool bIsUserAction = false;
-			AddNode(NewNode, bIsUserAction, bSelectNewNode);
+			AddNode(NewNode, /*bIsUserAction=*/false, bSelectNewNode);
 		}
 	}
 
@@ -213,6 +200,30 @@ const UPCGEditorGraphNodeBase* UPCGEditorGraph::GetEditorNodeFromPCGNode(const U
 	return nullptr;
 }
 
+TSubclassOf<UPCGEditorGraphNodeBase> UPCGEditorGraph::GetGraphNodeClassFromPCGSettings(const UPCGSettings* Settings)
+{
+	if (Settings->IsA<UPCGNamedRerouteDeclarationSettings>())
+	{
+		return UPCGEditorGraphNodeNamedRerouteDeclaration::StaticClass();
+	}
+	else if (Settings->IsA<UPCGNamedRerouteUsageSettings>())
+	{
+		return UPCGEditorGraphNodeNamedRerouteUsage::StaticClass();
+	}
+	else if (Settings->IsA<UPCGRerouteSettings>())
+	{
+		return UPCGEditorGraphNodeReroute::StaticClass();
+	}
+	else if (Settings->IsA<UPCGUserParameterGetSettings>())
+	{
+		return UPCGEditorGraphGetUserParameter::StaticClass();
+	}
+	else // All other settings.
+	{
+		return UPCGEditorGraphNode::StaticClass();
+	}
+}
+
 void UPCGEditorGraph::CreateLinks(UPCGEditorGraphNodeBase* GraphNode, bool bCreateInbound, bool bCreateOutbound, const TMap<UPCGNode*, UPCGEditorGraphNodeBase*>& InPCGNodeToPCGEditorNodeMap)
 {
 	check(GraphNode);
diff --git a/Engine/Plugins/PCG/Source/PCGEditor/Private/PCGEditorGraph.h b/Engine/Plugins/PCG/Source/PCGEditor/Private/PCGEditorGraph.h
--- a/Engine/Plugins/PCG/Source/PCGEditor/Private/PCGEditorGraph.h
+++ b/Engine/Plugins/PCG/Source/PCGEditor/Private/PCGEditorGraph.h
@@ -52,6 +52,9 @@ public:
 	/** Returns the PCG editor graph node corresponding to the given PCG node. */
 	const UPCGEditorGraphNodeBase* GetEditorNodeFromPCGNode(const UPCGNode* InPCGNode) const;
 
+	/** Returns the PCG Editor Graph Node that should be spawned from any given PCG Settings. */
+	static TSubclassOf<UPCGEditorGraphNodeBase> GetGraphNodeClassFromPCGSettings(const UPCGSettings* Settings);
+
 protected:
 	void CreateLinks(UPCGEditorGraphNodeBase* InGraphNode, bool bCreateInbound, bool bCreateOutbound, const TMap<UPCGNode*, UPCGEditorGraphNodeBase*>& InPCGNodeToPCGEditorNodeMap);
 
diff --git a/Engine/Plugins/PCG/Source/PCGEditor/Private/PCGEditorGraphSchemaActions.cpp b/Engine/Plugins/PCG/Source/PCGEditor/Private/PCGEditorGraphSchemaActions.cpp
--- a/Engine/Plugins/PCG/Source/PCGEditor/Private/PCGEditorGraphSchemaActions.cpp
+++ b/Engine/Plugins/PCG/Source/PCGEditor/Private/PCGEditorGraphSchemaActions.cpp
@@ -70,7 +70,8 @@ UEdGraphNode* FPCGEditorGraphSchemaAction_NewNativeElement::PerformAction(UEdGra
 	PostCreation(NewPCGNode);
 
 	FGraphNodeCreator<UPCGEditorGraphNode> NodeCreator(*EditorGraph);
-	UPCGEditorGraphNode* NewNode = NodeCreator.CreateUserInvokedNode(bSelectNewNode);
+	const TSubclassOf<UPCGEditorGraphNode> GraphNodeClass = TSubclassOf<UPCGEditorGraphNode>(UPCGEditorGraph::GetGraphNodeClassFromPCGSettings(DefaultNodeSettings));
+	UPCGEditorGraphNode* NewNode = NodeCreator.CreateUserInvokedNode(bSelectNewNode, GraphNodeClass);
 	NewNode->Construct(NewPCGNode);
 	NewNode->NodePosX = Location.X;
 	NewNode->NodePosY = Location.Y;
@@ -95,6 +96,8 @@ void FPCGEditorGraphSchemaAction_NewGetParameterElement::PostCreation(UPCGNode*
 	Settings->PropertyGuid = PropertyGuid;
 	Settings->PropertyName = PropertyName;
 
+	NewNode->NodeTitle = Settings->PropertyName;
+
 	// We need to set the settings to update the pins.
 	NewNode->SetSettingsInterface(Settings);
 }
diff --git a/Engine/Plugins/PCG/Source/PCGEditor/Private/PCGEditorStyle.cpp b/Engine/Plugins/PCG/Source/PCGEditor/Private/PCGEditorStyle.cpp
--- a/Engine/Plugins/PCG/Source/PCGEditor/Private/PCGEditorStyle.cpp
+++ b/Engine/Plugins/PCG/Source/PCGEditor/Private/PCGEditorStyle.cpp
@@ -81,10 +81,14 @@ FPCGEditorStyle::FPCGEditorStyle() : FSlateStyleSet("PCGEditorStyle")
 	FInlineEditableTextBlockStyle NodeTitleStyle = FInlineEditableTextBlockStyle(FAppStyle::Get().GetWidgetStyle<FInlineEditableTextBlockStyle>("Graph.Node.NodeTitleInlineEditableText"));
 	FTextBlockStyle GraphNodeItalicTitle = FTextBlockStyle(FAppStyle::Get().GetWidgetStyle<FTextBlockStyle>("Graph.Node.NodeTitle"))
 		.SetFont(DEFAULT_FONT("BoldCondensedItalic", 10));
+	FTextBlockStyle CompactGraphNodeTitle = FTextBlockStyle(FAppStyle::Get().GetWidgetStyle<FTextBlockStyle>("Graph.Node.NodeTitle"))
+		.SetFont(DEFAULT_FONT("Regular", 10))
+		.SetShadowColorAndOpacity(FLinearColor::Black * 0.8f)
+		.SetShadowOffset(FVector2D::UnitVector);
 
 	Set("PCG.Node.NodeTitleInlineEditableText", NodeTitleStyle);
-	Set("PCG.Node.InstancedNodeTitleInlineEditableText", FInlineEditableTextBlockStyle(NodeTitleStyle)
-		.SetTextStyle(GraphNodeItalicTitle));
+	Set("PCG.Node.InstancedNodeTitleInlineEditableText", FInlineEditableTextBlockStyle(NodeTitleStyle).SetTextStyle(GraphNodeItalicTitle));
+	Set("PCG.Node.CompactNodeTitle", FInlineEditableTextBlockStyle(NodeTitleStyle).SetTextStyle(CompactGraphNodeTitle));
 
 	// Styles for higen grid size label. Hand-tweaked values to match UI target mockups.
 	Set(PCGEditorStyleConstants::Node_Overlay_GridSizeLabel_Active_Border, new FSlateRoundedBoxBrush(
diff --git a/Engine/Plugins/PCG/Source/PCGEditor/Public/PCGEditorSettings.h b/Engine/Plugins/PCG/Source/PCGEditor/Public/PCGEditorSettings.h
--- a/Engine/Plugins/PCG/Source/PCGEditor/Public/PCGEditorSettings.h
+++ b/Engine/Plugins/PCG/Source/PCGEditor/Public/PCGEditorSettings.h
@@ -220,6 +220,10 @@ public:
 	UPROPERTY(EditAnywhere, Config, Category = "Workflow")
 	EPCGEditorDoubleClickAction NodeDoubleClickAction = EPCGEditorDoubleClickAction::ToggleInspectNode;
 
+	/** Verify renaming graph parameters via their local getter nodes with a pop-up confirmation window. */
+	UPROPERTY(EditAnywhere, config, Category = Workflow)
+	bool bConfirmLocalGraphParameterNameChanges = true;
+
 	/** Specify if we want to disable CPU Throttling when a PCG Graph is executing, this will improve execution time when app is out of focus/minimized */
 	UPROPERTY(EditAnywhere, config, Category = "Editor Performance", meta = (DisplayName = "Disable CPU throttling during graph execution"))
 	bool bDisableCPUThrottlingDuringGraphExecution = false;
