[PCG] Add FPCGContext::StoreInCache/GetFromCache api to remove dependency on SourceComponent/Subsystem

#rnx

diff --git a/Engine/Plugins/PCG/Source/PCG/Private/PCGContext.cpp b/Engine/Plugins/PCG/Source/PCG/Private/PCGContext.cpp
--- a/Engine/Plugins/PCG/Source/PCG/Private/PCGContext.cpp
+++ b/Engine/Plugins/PCG/Source/PCG/Private/PCGContext.cpp
@@ -570,4 +570,24 @@ void FPCGContext::InitializeGraphExecutor(FPCGContext* InContext)
 	InContext->GraphExecutor = GraphExecutor;
 }
 
+void FPCGContext::StoreInCache(const FPCGStoreInCacheParams& Params, const FPCGDataCollection& InOutput)
+{
+	if (GraphExecutor)
+	{
+		GraphExecutor->GetCache().StoreInCache(Params, InOutput);
+	}
+}
+
+bool FPCGContext::GetFromCache(const FPCGGetFromCacheParams& Params, FPCGDataCollection& OutOutput) const
+{
+	if (GraphExecutor)
+	{
+		return GraphExecutor->GetCache().GetFromCache(Params, OutOutput);
+	}
+	else
+	{
+		return false;
+	}
+}
+
 #undef LOCTEXT_NAMESPACE
\ No newline at end of file
diff --git a/Engine/Plugins/PCG/Source/PCG/Private/PCGElement.cpp b/Engine/Plugins/PCG/Source/PCG/Private/PCGElement.cpp
--- a/Engine/Plugins/PCG/Source/PCG/Private/PCGElement.cpp
+++ b/Engine/Plugins/PCG/Source/PCG/Private/PCGElement.cpp
@@ -308,8 +308,7 @@ void IPCGElement::PreExecutePrimaryLoopElement(FPCGContext* Context, const UPCGS
 {
 	check(Context);
 
-	UPCGSubsystem* Subsystem = Context->SourceComponent.Get() ? Context->SourceComponent->GetSubsystem() : nullptr;
-	if (!Subsystem || !Settings)
+	if (!Settings)
 	{
 		return;
 	}
@@ -350,7 +349,7 @@ void IPCGElement::PreExecutePrimaryLoopElement(FPCGContext* Context, const UPCGS
 		GetDependenciesCrc(SubCollection, Settings, Context->SourceComponent.Get(), CacheParams.Crc);
 
 		FPCGDataCollection SubCollectionOutput;
-		if(Subsystem->GetCache()->GetFromCache(CacheParams, SubCollectionOutput))
+		if(Context->GetFromCache(CacheParams, SubCollectionOutput))
 		{
 			// Found a match in the cache, add it to the output, and remove the matching inputs.
 			// Note that in the input part we'll take only the data present in the PrimaryDataCollection here, as we will reintroduce only those then.
@@ -535,12 +534,8 @@ void IPCGElement::PostExecute(FPCGContext* Context) const
 
 			if (bCacheable)
 			{
-				UPCGSubsystem* Subsystem = Context->SourceComponent.Get() ? Context->SourceComponent->GetSubsystem() : nullptr;
-				if (Subsystem)
-				{
-					FPCGStoreInCacheParams Params = { .Element = this, .Crc = Context->DependenciesCrc };
-					Subsystem->GetCache()->StoreInCache(Params, Context->OutputData);
-				}
+				FPCGStoreInCacheParams Params = { .Element = this, .Crc = Context->DependenciesCrc };
+				Context->StoreInCache(Params, Context->OutputData);
 			}
 		}
 	}
@@ -601,10 +596,7 @@ void IPCGElement::PostExecutePrimaryLoopElement(FPCGContext* Context, const UPCG
 {
 	check(Context);
 
-	// In the case of primary pin-loops, write back individual results to the cache, and reinsert the cached results in the output as needed.
-	UPCGSubsystem* Subsystem = Context->SourceComponent.Get() ? Context->SourceComponent->GetSubsystem() : nullptr;
-
-	if (!Subsystem || !Settings)
+	if (!Settings)
 	{
 		return;
 	}
@@ -661,7 +653,7 @@ void IPCGElement::PostExecutePrimaryLoopElement(FPCGContext* Context, const UPCG
 				SubCollectionOutput.ComputeCrcs(bShouldComputeFullOutputDataCrc);
 
 				FPCGStoreInCacheParams Params = { .Element = this, .Crc = DependenciesCrc };
-				Subsystem->GetCache()->StoreInCache(Params, SubCollectionOutput);
+				Context->StoreInCache(Params, SubCollectionOutput);
 			}
 		}
 	}
diff --git a/Engine/Plugins/PCG/Source/PCG/Public/PCGContext.h b/Engine/Plugins/PCG/Source/PCG/Public/PCGContext.h
--- a/Engine/Plugins/PCG/Source/PCG/Public/PCGContext.h
+++ b/Engine/Plugins/PCG/Source/PCG/Public/PCGContext.h
@@ -18,8 +18,10 @@ class FPCGGraphExecutor;
 class UPCGGraphInterface;
 class UPCGSettingsInterface;
 class UPCGSpatialData;
+struct FPCGGetFromCacheParams;
 struct FPCGSettingsOverridableParam;
 struct FPCGStack;
+struct FPCGStoreInCacheParams;
 
 extern PCG_API TAutoConsoleVariable<bool> CVarPCGEnablePointArrayData;
 
@@ -241,6 +243,9 @@ struct PCG_API FPCGContext
 	bool GetOutputData(FPCGTaskId InTaskId, FPCGDataCollection& OutData);
 	void ClearOutputData(FPCGTaskId InTaskId);
 
+	void StoreInCache(const FPCGStoreInCacheParams& Params, const FPCGDataCollection& InOutput);
+	bool GetFromCache(const FPCGGetFromCacheParams& Params, FPCGDataCollection& OutOutput) const;
+
 protected:
 	virtual UObject* GetExternalContainerForOverridableParam(const FPCGSettingsOverridableParam& InParam) { return nullptr; }
 	virtual void* GetUnsafeExternalContainerForOverridableParam(const FPCGSettingsOverridableParam& InParam) { return nullptr; }
