Small optimization to prevent reentry when doing changes across long dynamic type chains in editing.
Brought back optimization in get input to forward type in a more robust manner.

#rnx

diff --git a/Engine/Plugins/PCG/Source/PCG/Private/Elements/Metadata/PCGMetadataOpElementBase.cpp b/Engine/Plugins/PCG/Source/PCG/Private/Elements/Metadata/PCGMetadataOpElementBase.cpp
--- a/Engine/Plugins/PCG/Source/PCG/Private/Elements/Metadata/PCGMetadataOpElementBase.cpp
+++ b/Engine/Plugins/PCG/Source/PCG/Private/Elements/Metadata/PCGMetadataOpElementBase.cpp
@@ -9,6 +9,7 @@
 #include "Data/PCGPointData.h"
 #include "Elements/Metadata/PCGMetadataElementCommon.h"
 #include "Helpers/PCGPropertyHelpers.h"
+#include "Helpers/PCGSettingsHelpers.h"
 #include "Metadata/PCGMetadataTypesConstantStruct.h"
 #include "Metadata/Accessors/PCGAttributeAccessorHelpers.h"
 
@@ -155,6 +156,9 @@ uint32 UPCGMetadataSettingsBase::GetInputPinToForward() const
 		}
 		else
 		{
+			// Implementation note: here we use the pin type scope helper so that all type queries are cached in the scope of this call.
+			PCGSettingsHelpers::FPinTypeScopeHelper PinTypeHelper;
+
 			for (uint32 InputPinIndex = 0; InputPinIndex < OperandNum; ++InputPinIndex)
 			{
 				const EPCGDataType PinType = GetInputPinType(InputPinIndex);
diff --git a/Engine/Plugins/PCG/Source/PCG/Private/Helpers/PCGSettingsHelpers.cpp b/Engine/Plugins/PCG/Source/PCG/Private/Helpers/PCGSettingsHelpers.cpp
--- a/Engine/Plugins/PCG/Source/PCG/Private/Helpers/PCGSettingsHelpers.cpp
+++ b/Engine/Plugins/PCG/Source/PCG/Private/Helpers/PCGSettingsHelpers.cpp
@@ -492,4 +492,51 @@ namespace PCGSettingsHelpers
 
 		return GetAllOverridableParams_Internal(InClass, InConfig, AlreadySeenProperties);
 	}
+
+	// Thread-local variables used for the FPinTypeScopeHelper implementation
+	static thread_local bool bVisitedPinsMapExists = false;
+	static thread_local TMap<const UPCGPin*, EPCGDataType, TInlineSetAllocator<32>> VisitedPinsMap;
+
+	FPinTypeScopeHelper::FPinTypeScopeHelper()
+	{
+		if (!bVisitedPinsMapExists)
+		{
+			bClearMapOnDestruction = true;
+			bVisitedPinsMapExists = true;
+		}
+	}
+
+	FPinTypeScopeHelper::~FPinTypeScopeHelper()
+	{
+		if (bClearMapOnDestruction)
+		{
+			VisitedPinsMap.Empty();
+			bVisitedPinsMapExists = false;
+		}
+	}
+
+	EPCGDataType FPinTypeScopeHelper::GetCurrentPinType(const UPCGPin* InPin)
+	{
+		// Hitting this ensure means break of contract (scope broken), this helper should never be dynamically allocated.
+		ensure(bVisitedPinsMapExists);
+		if (EPCGDataType* ExistingResult = bVisitedPinsMapExists ? VisitedPinsMap.Find(InPin) : nullptr)
+		{
+			return *ExistingResult;
+		}
+		else if(InPin)
+		{
+			EPCGDataType PinResult = InPin->GetCurrentTypes();
+
+			if (bVisitedPinsMapExists)
+			{
+				VisitedPinsMap.Add(InPin, PinResult);
+			}
+
+			return PinResult;
+		}
+		else
+		{
+			return EPCGDataType::None;
+		}
+	}
 }
diff --git a/Engine/Plugins/PCG/Source/PCG/Private/PCGNode.cpp b/Engine/Plugins/PCG/Source/PCG/Private/PCGNode.cpp
--- a/Engine/Plugins/PCG/Source/PCG/Private/PCGNode.cpp
+++ b/Engine/Plugins/PCG/Source/PCG/Private/PCGNode.cpp
@@ -1098,6 +1098,8 @@ EPCGChangeType UPCGNode::PropagateDynamicPinTypes(TSet<UPCGNode*>& TouchedNodes,
 	// ChangeType intentionally ignored here - using it breaks propagation of types across chained reroutes.
 	TouchedNodes.Add(this);
 
+	TArray<UPCGNode*, TInlineAllocator<8>> OtherNodesPropagatedTo;
+
 	for (UPCGPin* OutputPin : OutputPins)
 	{
 		for (UPCGEdge* Edge : OutputPin->Edges)
@@ -1109,16 +1111,19 @@ EPCGChangeType UPCGNode::PropagateDynamicPinTypes(TSet<UPCGNode*>& TouchedNodes,
 			}
 			
 			UPCGNode* OtherNode = OtherPin->Node;
-			if (!OtherNode || OtherNode == FromNode)
+			if (!OtherNode || OtherNode == FromNode || OtherNodesPropagatedTo.Contains(OtherNode))
 			{
 				continue;
 			}
 			
 			const UPCGSettings* OtherSettings = OtherNode->GetSettings();
-			if (Settings && OtherSettings->HasDynamicPins())
+			if (!OtherSettings || !OtherSettings->HasDynamicPins())
 			{
-				ChangeType |= OtherNode->PropagateDynamicPinTypes(TouchedNodes, this);
+				continue;
 			}
+
+			OtherNodesPropagatedTo.Add(OtherNode);
+			ChangeType |= OtherNode->PropagateDynamicPinTypes(TouchedNodes, this);
 		}
 	}
 	
diff --git a/Engine/Plugins/PCG/Source/PCG/Private/PCGSettings.cpp b/Engine/Plugins/PCG/Source/PCG/Private/PCGSettings.cpp
--- a/Engine/Plugins/PCG/Source/PCG/Private/PCGSettings.cpp
+++ b/Engine/Plugins/PCG/Source/PCG/Private/PCGSettings.cpp
@@ -547,12 +547,14 @@ EPCGDataType UPCGSettings::GetTypeUnionOfIncidentEdges(const FName& PinLabel) co
 	const UPCGPin* Pin = Node ? Node->GetInputPin(PinLabel) : nullptr;
 	if (Pin && Pin->EdgeCount() > 0)
 	{
+		PCGSettingsHelpers::FPinTypeScopeHelper PinTypeHelper;
+
 		for (const UPCGEdge* Edge : Pin->Edges)
 		{
 			const UPCGPin* OtherOutputPin = Edge ? Edge->GetOtherPin(Pin) : nullptr;
 			if (OtherOutputPin)
 			{
-				Result |= OtherOutputPin->GetCurrentTypes();
+				Result |= PinTypeHelper.GetCurrentPinType(OtherOutputPin);
 			}
 		}
 	}
diff --git a/Engine/Plugins/PCG/Source/PCG/Public/Helpers/PCGSettingsHelpers.h b/Engine/Plugins/PCG/Source/PCG/Public/Helpers/PCGSettingsHelpers.h
--- a/Engine/Plugins/PCG/Source/PCG/Public/Helpers/PCGSettingsHelpers.h
+++ b/Engine/Plugins/PCG/Source/PCG/Public/Helpers/PCGSettingsHelpers.h
@@ -138,6 +138,22 @@ namespace PCGSettingsHelpers
 	};
 
 	PCG_API TArray<FPCGSettingsOverridableParam> GetAllOverridableParams(const UStruct* InClass, const FPCGGetAllOverridableParamsConfig& InConfig);
+
+	/* Small helper that use thread local storage to build a map of pins to their current type. Meant to prevent reentrancy issues (combinatorics) in dynamic types.
+	* This struct is not meant to be allocated dynamically as the intention is to have a well-defined construction-destruction scope. */
+	struct PCG_API FPinTypeScopeHelper
+	{
+		FPinTypeScopeHelper();
+		~FPinTypeScopeHelper();
+		FPinTypeScopeHelper(const FPinTypeScopeHelper&) = delete;
+		FPinTypeScopeHelper(FPinTypeScopeHelper&&) = delete;
+		FPinTypeScopeHelper& operator=(const FPinTypeScopeHelper&) = delete;
+		FPinTypeScopeHelper& operator=(FPinTypeScopeHelper&&) = delete;
+
+		EPCGDataType GetCurrentPinType(const UPCGPin* InPin);
+	private:
+		bool bClearMapOnDestruction = false;
+	};
 }
 
 // Deprecated macro, not necessary anymore. Cf. GetValue
