[PCG] Fix custom hlsl nodes frequently recompiling shaders due to unstable object names being injected into source and used for temporary target file names.

#rnx

diff --git a/Engine/Plugins/PCG/Source/PCG/Private/Compute/Elements/PCGCustomHLSLKernel.cpp b/Engine/Plugins/PCG/Source/PCG/Private/Compute/Elements/PCGCustomHLSLKernel.cpp
--- a/Engine/Plugins/PCG/Source/PCG/Private/Compute/Elements/PCGCustomHLSLKernel.cpp
+++ b/Engine/Plugins/PCG/Source/PCG/Private/Compute/Elements/PCGCustomHLSLKernel.cpp
@@ -327,7 +327,7 @@ FString UPCGCustomHLSLKernel::GetCookedSource(const UPCGComputeGraph* InComputeG
 
 	ProcessShaderSource(InComputeGraph, Source, Functions);
 
-	FString ShaderPathName = GetPathName();
+	FString ShaderPathName = CustomHLSLSettings->GetPathName();
 	PCGCustomHLSLKernel::ConvertObjectPathToShaderFilePath(ShaderPathName);
 
 	//const bool bHasKernelKeyword = Source.Contains(TEXT("KERNEL"), ESearchCase::CaseSensitive);
@@ -594,8 +594,9 @@ void UPCGCustomHLSLKernel::InitEntryPoint()
 	// We append the PCG kernel name here because of a weird issue on Mac where if the last word of the name of a
 	// shader parameter matches the kernel name, the value for the parameter is always zero for unknown reasons
 	// (e.g. parameter named "XXXWeight" & kernel also named "Weight"). Adding this bit of "random" string after
-	// the user provided kernel name should greatly reduce the chance of things like that happening
-	EntryPoint += TEXT("_") + GetName();
+	// the user provided kernel name should greatly reduce the chance of things like that happening. We don't include
+	// the full fname with number as this will change across executions and cause DDC misses.
+	EntryPoint += TEXT("_") + GetFName().GetPlainNameString();
 
 	const TCHAR InvalidCharacters[] = {
 		TCHAR('\"'), TCHAR('\\'), TCHAR('\''), TCHAR('\"'), TCHAR(' '), TCHAR(','), TCHAR('.'), TCHAR('|'), TCHAR('&'), TCHAR('!'),
