[PCG] Missing success boolean set in CreateChainAccessor

* Fix missing tooltip in attribute list view

#rnx

diff --git a/Engine/Plugins/PCG/Source/PCG/Private/Metadata/Accessors/PCGAttributeAccessorHelpers.cpp b/Engine/Plugins/PCG/Source/PCG/Private/Metadata/Accessors/PCGAttributeAccessorHelpers.cpp
--- a/Engine/Plugins/PCG/Source/PCG/Private/Metadata/Accessors/PCGAttributeAccessorHelpers.cpp
+++ b/Engine/Plugins/PCG/Source/PCG/Private/Metadata/Accessors/PCGAttributeAccessorHelpers.cpp
@@ -69,11 +69,11 @@ namespace PCGAttributeAccessorHelpers
 			return {};
 		}
 		
+		bOutSuccess = true;
 		for (const FString& ExtraName : InSelector.GetExtraNames())
 		{
-			bool bSuccess = true;
-			InAccessor = CreateChainAccessor(std::move(InAccessor), FName(ExtraName), bSuccess);
-			if (!bSuccess)
+			InAccessor = CreateChainAccessor(std::move(InAccessor), FName(ExtraName), bOutSuccess);
+			if (!bOutSuccess)
 			{
 				if (!bQuiet)
 				{
