[PCG] Make sure param -> point conversion emits points, complying with the output pin type.

#rnx

diff --git a/Engine/Plugins/PCG/Source/PCG/Private/Elements/PCGCollapseElement.cpp b/Engine/Plugins/PCG/Source/PCG/Private/Elements/PCGCollapseElement.cpp
--- a/Engine/Plugins/PCG/Source/PCG/Private/Elements/PCGCollapseElement.cpp
+++ b/Engine/Plugins/PCG/Source/PCG/Private/Elements/PCGCollapseElement.cpp
@@ -7,8 +7,20 @@
 #include "Data/PCGBasePointData.h"
 #include "Data/PCGSpatialData.h"
 
+#include "UObject/FortniteMainBranchObjectVersion.h"
+
 #include UE_INLINE_GENERATED_CPP_BY_NAME(PCGCollapseElement)
 
+void UPCGCollapseSettings::PostLoad()
+{
+	Super::PostLoad();
+
+	if (GetLinkerCustomVersion(FFortniteMainBranchObjectVersion::GUID) < FFortniteMainBranchObjectVersion::PCGAttributeSetToPointAlwaysConverts)
+	{
+		bPassThroughEmptyAttributeSets = true;
+	}
+}
+
 #if WITH_EDITOR
 bool UPCGCollapseSettings::GetCompactNodeIcon(FName& OutCompactNodeIcon) const
 {
@@ -52,6 +64,9 @@ bool FPCGCollapseElement::ExecuteInternal(FPCGContext* InContext) const
 	check(InContext);
 	FPCGCollapseContext* Context = static_cast<FPCGCollapseContext*>(InContext);
 
+	const UPCGConvertToPointDataSettings* Settings = InContext->GetInputSettings<UPCGConvertToPointDataSettings>();
+	check(Settings);
+
 	TArray<FPCGTaggedData> Inputs = Context->InputData.GetInputsByPin(PCGPinConstants::DefaultInputLabel);
 	TArray<FPCGTaggedData>& Outputs = Context->OutputData.TaggedData;
 
@@ -79,7 +94,7 @@ bool FPCGCollapseElement::ExecuteInternal(FPCGContext* InContext) const
 			const UPCGMetadata* ParamMetadata = ParamData->Metadata;
 			const int64 ParamItemCount = ParamMetadata->GetLocalItemCount();
 
-			if (ParamItemCount == 0)
+			if (ParamItemCount == 0 && Settings->bPassThroughEmptyAttributeSets)
 			{
 				continue;
 			}
diff --git a/Engine/Plugins/PCG/Source/PCG/Public/Elements/PCGCollapseElement.h b/Engine/Plugins/PCG/Source/PCG/Public/Elements/PCGCollapseElement.h
--- a/Engine/Plugins/PCG/Source/PCG/Public/Elements/PCGCollapseElement.h
+++ b/Engine/Plugins/PCG/Source/PCG/Public/Elements/PCGCollapseElement.h
@@ -14,6 +14,10 @@ class UPCGCollapseSettings : public UPCGSettings
 	GENERATED_BODY()
 
 public:
+	//~Begin UObject interface
+	virtual void PostLoad() override;
+	//~End UObject interface
+
 	//~Begin UPCGSettings interface
 #if WITH_EDITOR
 	virtual FName GetDefaultNodeName() const override { return FName(TEXT("ToPoint")); }
@@ -30,6 +34,12 @@ protected:
 	virtual TArray<FPCGPinProperties> OutputPinProperties() const override { return Super::DefaultPointOutputPinProperties(); }
 	virtual FPCGElementPtr CreateElement() const override;
 	//~End UPCGSettings interface
+
+public:
+	/** Before version PCGAttributeSetToPointAlwaysConverts (5.6) we passed through empty attribute sets without conversion.
+	* This toggle preserves old behavior. */
+	UPROPERTY()
+	bool bPassThroughEmptyAttributeSets = false;
 };
 
 /** Converts attribute sets to point data */
diff --git a/Engine/Source/Runtime/Core/Public/UObject/FortniteMainBranchObjectVersions.inl b/Engine/Source/Runtime/Core/Public/UObject/FortniteMainBranchObjectVersions.inl
--- a/Engine/Source/Runtime/Core/Public/UObject/FortniteMainBranchObjectVersions.inl
+++ b/Engine/Source/Runtime/Core/Public/UObject/FortniteMainBranchObjectVersions.inl
@@ -595,5 +595,8 @@ FORTNITE_MAIN_VERSION(SoftObjectPathTrailingNULsMaintained, 193)
 // Water body components no longer need to maintain their own PhysicalMaterial property since they are primitive components. After this version, leverage that one instead.
 FORTNITE_MAIN_VERSION(WaterBodyPhysicalMaterialPropertyRemoval, 194)
 
+// PCG fixed attribute set -> point conversion passing through empty point data as-is and violating output pin type.
+FORTNITE_MAIN_VERSION(PCGAttributeSetToPointAlwaysConverts, 195)
+
 // -----<new versions can be added above this line>-------------------------------------------------
 #undef FORTNITE_MAIN_VERSION
\ No newline at end of file
