[PCG] Fix bug disallowing comment nodes from being renamable in PCG.

#rnx

diff --git a/Engine/Plugins/PCG/Source/PCGEditor/Private/PCGEditor.cpp b/Engine/Plugins/PCG/Source/PCGEditor/Private/PCGEditor.cpp
--- a/Engine/Plugins/PCG/Source/PCGEditor/Private/PCGEditor.cpp
+++ b/Engine/Plugins/PCG/Source/PCGEditor/Private/PCGEditor.cpp
@@ -346,6 +346,10 @@ bool FPCGEditor::OnValidateNodeTitle(const FText& NewName, UEdGraphNode* GraphNo
 	{
 		return PCGGraphNode->OnValidateNodeTitle(NewName, OutErrorMessage);
 	}
+	else if (GraphNode->IsA<UEdGraphNode_Comment>())
+	{
+		return true;
+	}
 
 	return false;
 }
@@ -1370,10 +1374,15 @@ bool FPCGEditor::CanRenameNode() const
 		return false;
 	}
 
-	if (const UPCGEditorGraphNode* SelectedNode = Cast<UPCGEditorGraphNode>(*SelectedNodes.CreateConstIterator()))
+	UObject* SelectedObject = *SelectedNodes.CreateConstIterator();
+	if (const UPCGEditorGraphNode* SelectedNode = Cast<UPCGEditorGraphNode>(SelectedObject))
 	{
 		return SelectedNode->GetCanRenameNode();
 	}
+	else if (SelectedObject->IsA<UEdGraphNode_Comment>())
+	{
+		return true;
+	}
 	else
 	{
 		return false;
