[PCG] Add null checks on comment node fix for SA

#rnx

diff --git a/Engine/Plugins/PCG/Source/PCGEditor/Private/PCGEditor.cpp b/Engine/Plugins/PCG/Source/PCGEditor/Private/PCGEditor.cpp
--- a/Engine/Plugins/PCG/Source/PCGEditor/Private/PCGEditor.cpp
+++ b/Engine/Plugins/PCG/Source/PCGEditor/Private/PCGEditor.cpp
@@ -346,7 +346,7 @@ bool FPCGEditor::OnValidateNodeTitle(const FText& NewName, UEdGraphNode* GraphNo
 	{
 		return PCGGraphNode->OnValidateNodeTitle(NewName, OutErrorMessage);
 	}
-	else if (GraphNode->IsA<UEdGraphNode_Comment>())
+	else if (GraphNode && GraphNode->IsA<UEdGraphNode_Comment>())
 	{
 		return true;
 	}
@@ -1379,7 +1379,7 @@ bool FPCGEditor::CanRenameNode() const
 	{
 		return SelectedNode->GetCanRenameNode();
 	}
-	else if (SelectedObject->IsA<UEdGraphNode_Comment>())
+	else if (SelectedObject && SelectedObject->IsA<UEdGraphNode_Comment>())
 	{
 		return true;
 	}
