[PCG] Fix check in CollapseElement

* It was cast in the wrong settings

#rnx

diff --git a/Engine/Plugins/PCG/Source/PCG/Private/Elements/PCGCollapseElement.cpp b/Engine/Plugins/PCG/Source/PCG/Private/Elements/PCGCollapseElement.cpp
--- a/Engine/Plugins/PCG/Source/PCG/Private/Elements/PCGCollapseElement.cpp
+++ b/Engine/Plugins/PCG/Source/PCG/Private/Elements/PCGCollapseElement.cpp
@@ -64,7 +64,7 @@ bool FPCGCollapseElement::ExecuteInternal(FPCGContext* InContext) const
 	check(InContext);
 	FPCGCollapseContext* Context = static_cast<FPCGCollapseContext*>(InContext);
 
-	const UPCGConvertToPointDataSettings* Settings = InContext->GetInputSettings<UPCGConvertToPointDataSettings>();
+	const UPCGCollapseSettings* Settings = InContext->GetInputSettings<UPCGCollapseSettings>();
 	check(Settings);
 
 	TArray<FPCGTaggedData> Inputs = Context->InputData.GetInputsByPin(PCGPinConstants::DefaultInputLabel);
