[PCG] Fix server/onethread deadlock issue: replace FPlatformMisc::NumberOfCoresIncludingHyperthreads() with LowLevelTasks::FScheduler::Get().GetNumWorkers()
- This fixes a server stall where GetNumWorkers returns 0 and so Launching tasks would execute them in the launching thread stack causing a recursive lock
- by checking GetNumWorkers instead we can avoid launching tasks completely if the result is 0
- Only wait on lock if we have a mainthread task pointer (instead of checking if we are inside the gamethread), this will avoid waiting on scheduling in a task if it is being executed on the main thread at any point
- check FApp::ShouldUseThreadingForPerformance() to know if we should be multithreading at all

diff --git a/Engine/Plugins/PCG/Source/PCG/Private/Graph/PCGGraphExecutor.cpp b/Engine/Plugins/PCG/Source/PCG/Private/Graph/PCGGraphExecutor.cpp
--- a/Engine/Plugins/PCG/Source/PCG/Private/Graph/PCGGraphExecutor.cpp
+++ b/Engine/Plugins/PCG/Source/PCG/Private/Graph/PCGGraphExecutor.cpp
@@ -1429,7 +1429,7 @@ bool FPCGGraphExecutor::ExecuteScheduling(double EndTime, TSharedPtr<FPCGGraphAc
 	const float MaxPercentageOfExecutingThreads = FMath::Clamp(CVarMaxPercentageOfExecutingThreads.GetValueOnAnyThread(), 0.0f, 0.8f);
 
 	// Number of threads to use across PCG
-	const int32 MaxNumThreads = FMath::Max(0, (int32)(FPlatformMisc::NumberOfCoresIncludingHyperthreads() * MaxPercentageOfThreadsToUse));
+	const int32 MaxNumThreads = FApp::ShouldUseThreadingForPerformance() ? FMath::Max(0, (int32)(LowLevelTasks::FScheduler::Get().GetNumWorkers() * MaxPercentageOfThreadsToUse)) : 0;
 
 	// Number of threads to use for PCG Element execution (outside of main thread tick)
 	const int32 MaxExecutingThreads = PCGGraphExecutor::CVarGraphMultithreading.GetValueOnAnyThread() ? FMath::Max(0, (int32)(MaxNumThreads * MaxPercentageOfExecutingThreads)) : 0;
@@ -1440,7 +1440,8 @@ bool FPCGGraphExecutor::ExecuteScheduling(double EndTime, TSharedPtr<FPCGGraphAc
 	TArray<FCachedResult*> CachedResults;
 	{
 		TRACE_CPUPROFILER_EVENT_SCOPE(FPCGGraphExecutor::ExecuteScheduling::ReadyTasks);
-		PCGGraphExecutor::TScopeLock ScopeLock(LiveTasksLock, !bIsInGameThread);
+		const bool bTryLock = (OutMainThreadTask == nullptr);
+		PCGGraphExecutor::TScopeLock ScopeLock(LiveTasksLock, bTryLock);
 		if (ScopeLock.IsLocked())
 		{
 			if (bNeedToCheckPausedTasks || bForceCheckPausedTasks)
