[PCG] Fix ensure in IPCGDataVisualization::GetTableVisualizerInfo

#rnx

diff --git a/Engine/Plugins/PCG/Source/PCG/Public/PCGDataVisualization.h b/Engine/Plugins/PCG/Source/PCG/Public/PCGDataVisualization.h
--- a/Engine/Plugins/PCG/Source/PCG/Public/PCGDataVisualization.h
+++ b/Engine/Plugins/PCG/Source/PCG/Public/PCGDataVisualization.h
@@ -58,7 +58,7 @@ public:
 	UE_DEPRECATED(5.6, "Use and implement GetTableVisualizerInfoWithDomain instead.")
 	virtual FPCGTableVisualizerInfo GetTableVisualizerInfo(const UPCGData* Data) const
 	{
-		ensure("Should never be called, child class need to override GetTableVisualizerInfoWithDomain");
+		ensureMsgf(false, TEXT("Should never be called, child class need to override GetTableVisualizerInfoWithDomain"));
 		return {};
 	}
 
