[PCG] Move somewhat-large inline parameter arrays to buffers after a validation warning about CBuffer size was reported.

#rnx

diff --git a/Engine/Plugins/PCG/Source/PCG/Private/Compute/DataInterfaces/Elements/PCGStaticMeshSpawnerDataInterface.cpp b/Engine/Plugins/PCG/Source/PCG/Private/Compute/DataInterfaces/Elements/PCGStaticMeshSpawnerDataInterface.cpp
--- a/Engine/Plugins/PCG/Source/PCG/Private/Compute/DataInterfaces/Elements/PCGStaticMeshSpawnerDataInterface.cpp
+++ b/Engine/Plugins/PCG/Source/PCG/Private/Compute/DataInterfaces/Elements/PCGStaticMeshSpawnerDataInterface.cpp
@@ -17,10 +17,13 @@
 #include "MeshSelectors/PCGMeshSelectorWeighted.h"
 
 #include "GlobalRenderResources.h"
+#include "RenderGraphBuilder.h"
+#include "RenderGraphResources.h"
 #include "RHIResources.h"
 #include "ShaderCompilerCore.h"
 #include "ShaderCore.h"
 #include "ShaderParameterMetadataBuilder.h"
+#include "SystemTextures.h"
 #include "Engine/StaticMesh.h"
 #include "ComputeFramework/ShaderParamTypeDefinition.h"
 
@@ -90,9 +93,9 @@ void UPCGStaticMeshSpawnerDataInterface::GetSupportedInputs(TArray<FShaderFuncti
 
 BEGIN_SHADER_PARAMETER_STRUCT(FPCGStaticMeshSpawnerDataInterfaceParameters,)
 	SHADER_PARAMETER_ARRAY(FUintVector4, AttributeIdOffsetStrides, [UPCGStaticMeshSpawnerDataInterface::MAX_ATTRIBUTES])
-	SHADER_PARAMETER_ARRAY(FIntVector4, PrimitiveStringKeys, [PCGComputeConstants::MAX_PRIMITIVE_COMPONENTS_PER_SPAWNER])
-	SHADER_PARAMETER_ARRAY(FVector4f, PrimitiveMeshBoundsMin, [PCGComputeConstants::MAX_PRIMITIVE_COMPONENTS_PER_SPAWNER])
-	SHADER_PARAMETER_ARRAY(FVector4f, PrimitiveMeshBoundsMax, [PCGComputeConstants::MAX_PRIMITIVE_COMPONENTS_PER_SPAWNER])
+	SHADER_PARAMETER_RDG_BUFFER_SRV(StructuredBuffer<int32>, PrimitiveStringKeys)
+	SHADER_PARAMETER_RDG_BUFFER_SRV(StructuredBuffer<FVector4f>, PrimitiveMeshBoundsMin)
+	SHADER_PARAMETER_RDG_BUFFER_SRV(StructuredBuffer<FVector4f>, PrimitiveMeshBoundsMax)
 	SHADER_PARAMETER_SCALAR_ARRAY(float, SelectionCDF, [PCGComputeConstants::MAX_PRIMITIVE_COMPONENTS_PER_SPAWNER])
 	SHADER_PARAMETER(uint32, NumAttributes)
 	SHADER_PARAMETER(uint32, NumPrimitives)
@@ -122,9 +125,9 @@ void UPCGStaticMeshSpawnerDataInterface::GetHLSL(FString& OutHLSL, FString const
 		"uint {DataInterfaceName}_ApplyBounds;\n"
 		"int {DataInterfaceName}_SelectedMeshAttributeId;\n"
 		"uint4 {DataInterfaceName}_AttributeIdOffsetStrides[{MaxAttributes}];\n"
-		"float4 {DataInterfaceName}_PrimitiveMeshBoundsMin[{MaxPrimitives}];\n"
-		"float4 {DataInterfaceName}_PrimitiveMeshBoundsMax[{MaxPrimitives}];\n"
-		"int4 {DataInterfaceName}_PrimitiveStringKeys[{MaxPrimitives}];\n"
+		"StructuredBuffer<float4> {DataInterfaceName}_PrimitiveMeshBoundsMin;\n"
+		"StructuredBuffer<float4> {DataInterfaceName}_PrimitiveMeshBoundsMax;\n"
+		"StructuredBuffer<int> {DataInterfaceName}_PrimitiveStringKeys;\n"
 		"DECLARE_SCALAR_ARRAY(float, {DataInterfaceName}_SelectionCDF, {MaxPrimitives});\n"
 		"\n"
 		"int SMSpawner_GetSelectorAttributeId_{DataInterfaceName}()\n"
@@ -164,14 +167,14 @@ void UPCGStaticMeshSpawnerDataInterface::GetHLSL(FString& OutHLSL, FString const
 		"\n"
 		"int SMSpawner_GetPrimitiveStringKey_{DataInterfaceName}(uint InPrimitiveIndex)\n"
 		"{\n"
-		"	return {DataInterfaceName}_PrimitiveStringKeys[InPrimitiveIndex][0];\n"
+		"	return {DataInterfaceName}_PrimitiveStringKeys[InPrimitiveIndex];\n"
 		"}\n"
 		"\n"
 		"uint SMSpawner_GetPrimitiveIndexFromStringKey_{DataInterfaceName}(int InMeshPathStringKey)\n"
 		"{\n"
 		"	for (uint Index = 0; Index < {DataInterfaceName}_NumPrimitives; ++Index)\n"
 		"	{\n"
-		"		if ({DataInterfaceName}_PrimitiveStringKeys[Index][0] == InMeshPathStringKey)\n"
+		"		if ({DataInterfaceName}_PrimitiveStringKeys[Index] == InMeshPathStringKey)\n"
 		"		{\n"
 		"			return Index;\n"
 		"		}\n"
@@ -743,6 +746,29 @@ bool UPCGStaticMeshSpawnerDataProvider::SetupPrimitives(FPCGContext* InContext,
 	return true;
 }
 
+FPCGStaticMeshSpawnerDataProviderProxy::FPCGStaticMeshSpawnerDataProviderProxy(
+	const TArray<FUintVector4>& InAttributeIdOffsetStrides,
+	int32 InSelectorAttributeId,
+	const TArray<int32>& InPrimitiveStringKeys,
+	TArray<float> InSelectionCDF,
+	int32 InSelectedMeshAttributeId,
+	const TArray<FBox>& InPrimitiveMeshBounds)
+	: AttributeIdOffsetStrides(InAttributeIdOffsetStrides)
+	, SelectionCDF(InSelectionCDF)
+	, SelectorAttributeId(InSelectorAttributeId)
+	, PrimitiveStringKeys(InPrimitiveStringKeys)
+	, SelectedMeshAttributeId(InSelectedMeshAttributeId)
+{
+	PrimitiveMeshBoundsMin.Reserve(InPrimitiveMeshBounds.Num());
+	PrimitiveMeshBoundsMax.Reserve(InPrimitiveMeshBounds.Num());
+
+	for (int Index = 0; Index < InPrimitiveMeshBounds.Num(); ++Index)
+	{
+		PrimitiveMeshBoundsMin[Index] = FVector4f(InPrimitiveMeshBounds[Index].Min.X, InPrimitiveMeshBounds[Index].Min.Y, InPrimitiveMeshBounds[Index].Min.Z, /*Unused*/0.0f);
+		PrimitiveMeshBoundsMax[Index] = FVector4f(InPrimitiveMeshBounds[Index].Max.X, InPrimitiveMeshBounds[Index].Max.Y, InPrimitiveMeshBounds[Index].Max.Z, /*Unused*/0.0f);
+	}
+}
+
 bool FPCGStaticMeshSpawnerDataProviderProxy::IsValid(FValidationData const& InValidationData) const
 {
 	if (InValidationData.ParameterStructSize != sizeof(FParameters))
@@ -753,6 +779,51 @@ bool FPCGStaticMeshSpawnerDataProviderProxy::IsValid(FValidationData const& InVa
 	return true;
 }
 
+void FPCGStaticMeshSpawnerDataProviderProxy::AllocateResources(FRDGBuilder& GraphBuilder, FAllocationData const& InAllocationData)
+{
+	if (!PrimitiveStringKeys.IsEmpty())
+	{
+		const FRDGBufferDesc BufferDesc = FRDGBufferDesc::CreateStructuredDesc(PrimitiveStringKeys.GetTypeSize(), PrimitiveStringKeys.Num());
+
+		FRDGBufferRef Buffer = GraphBuilder.CreateBuffer(BufferDesc, TEXT("PCGStaticMeshSpawner_PrimitiveStringKeys"));
+		PrimitiveStringKeysBufferSRV = GraphBuilder.CreateSRV(Buffer);
+
+		GraphBuilder.QueueBufferUpload(Buffer, MakeArrayView(PrimitiveStringKeys));
+	}
+	else
+	{
+		PrimitiveStringKeysBufferSRV = GraphBuilder.CreateSRV(FRDGBufferSRVDesc(GSystemTextures.GetDefaultStructuredBuffer(GraphBuilder, PrimitiveStringKeys.GetTypeSize())));
+	}
+
+	if (!PrimitiveMeshBoundsMin.IsEmpty())
+	{
+		const FRDGBufferDesc BufferDesc = FRDGBufferDesc::CreateStructuredDesc(PrimitiveMeshBoundsMin.GetTypeSize(), PrimitiveMeshBoundsMin.Num());
+
+		FRDGBufferRef Buffer = GraphBuilder.CreateBuffer(BufferDesc, TEXT("PCGStaticMeshSpawner_PrimitiveMeshBoundsMin"));
+		PrimitiveMeshBoundsMinBufferSRV = GraphBuilder.CreateSRV(Buffer);
+
+		GraphBuilder.QueueBufferUpload(Buffer, MakeArrayView(PrimitiveMeshBoundsMin));
+	}
+	else
+	{
+		PrimitiveMeshBoundsMinBufferSRV = GraphBuilder.CreateSRV(FRDGBufferSRVDesc(GSystemTextures.GetDefaultStructuredBuffer(GraphBuilder, PrimitiveMeshBoundsMin.GetTypeSize())));
+	}
+
+	if (!PrimitiveMeshBoundsMax.IsEmpty())
+	{
+		const FRDGBufferDesc BufferDesc = FRDGBufferDesc::CreateStructuredDesc(PrimitiveMeshBoundsMax.GetTypeSize(), PrimitiveMeshBoundsMax.Num());
+
+		FRDGBufferRef Buffer = GraphBuilder.CreateBuffer(BufferDesc, TEXT("PCGStaticMeshSpawner_PrimitiveMeshBoundsMax"));
+		PrimitiveMeshBoundsMaxBufferSRV = GraphBuilder.CreateSRV(Buffer);
+
+		GraphBuilder.QueueBufferUpload(Buffer, MakeArrayView(PrimitiveMeshBoundsMax));
+	}
+	else
+	{
+		PrimitiveMeshBoundsMaxBufferSRV = GraphBuilder.CreateSRV(FRDGBufferSRVDesc(GSystemTextures.GetDefaultStructuredBuffer(GraphBuilder, PrimitiveMeshBoundsMax.GetTypeSize())));
+	}
+}
+
 void FPCGStaticMeshSpawnerDataProviderProxy::GatherDispatchData(FDispatchData const& InDispatchData)
 {
 	const TStridedView<FParameters> ParameterArray = MakeStridedParameterView<FParameters>(InDispatchData);
@@ -770,23 +841,17 @@ void FPCGStaticMeshSpawnerDataProviderProxy::GatherDispatchData(FDispatchData co
 			Parameters.AttributeIdOffsetStrides[Index] = AttributeIdOffsetStrides[Index];
 		}
 
-		for (int32 Index = 0; Index < PrimitiveStringKeys.Num(); ++Index)
-		{
-			Parameters.PrimitiveStringKeys[Index] = FIntVector4(PrimitiveStringKeys[Index], /*Unused*/0, /*Unused*/0, /*Unused*/0);
-		}
+		Parameters.PrimitiveStringKeys = PrimitiveStringKeysBufferSRV;
 
 		for (int32 Index = 0; Index < SelectionCDF.Num(); ++Index)
 		{
 			GET_SCALAR_ARRAY_ELEMENT(Parameters.SelectionCDF, Index) = SelectionCDF[Index];
 		}
 
-		Parameters.ApplyBounds = PrimitiveMeshBounds.IsEmpty() ? 0 : 1;
+		Parameters.ApplyBounds = PrimitiveMeshBoundsMin.IsEmpty() ? 0 : 1;
 
-		for (int32 Index = 0; Index < PrimitiveMeshBounds.Num(); ++Index)
-		{
-			Parameters.PrimitiveMeshBoundsMin[Index] = FVector4f(PrimitiveMeshBounds[Index].Min.X, PrimitiveMeshBounds[Index].Min.Y, PrimitiveMeshBounds[Index].Min.Z, /*Unused*/0.0f);
-			Parameters.PrimitiveMeshBoundsMax[Index] = FVector4f(PrimitiveMeshBounds[Index].Max.X, PrimitiveMeshBounds[Index].Max.Y, PrimitiveMeshBounds[Index].Max.Z, /*Unused*/0.0f);
-		}
+		Parameters.PrimitiveMeshBoundsMin = PrimitiveMeshBoundsMinBufferSRV;
+		Parameters.PrimitiveMeshBoundsMax = PrimitiveMeshBoundsMaxBufferSRV;
 	}
 }
 
diff --git a/Engine/Plugins/PCG/Source/PCG/Private/Compute/DataInterfaces/Elements/PCGStaticMeshSpawnerDataInterface.h b/Engine/Plugins/PCG/Source/PCG/Private/Compute/DataInterfaces/Elements/PCGStaticMeshSpawnerDataInterface.h
--- a/Engine/Plugins/PCG/Source/PCG/Private/Compute/DataInterfaces/Elements/PCGStaticMeshSpawnerDataInterface.h
+++ b/Engine/Plugins/PCG/Source/PCG/Private/Compute/DataInterfaces/Elements/PCGStaticMeshSpawnerDataInterface.h
@@ -105,17 +105,11 @@ public:
 		const TArray<int32>& InPrimitiveStringKeys,
 		TArray<float> InSelectionCDF,
 		int32 InSelectedMeshAttributeId,
-		const TArray<FBox>& InPrimitiveMeshBounds)
-		: AttributeIdOffsetStrides(InAttributeIdOffsetStrides)
-		, SelectionCDF(InSelectionCDF)
-		, SelectorAttributeId(InSelectorAttributeId)
-		, PrimitiveStringKeys(InPrimitiveStringKeys)
-		, PrimitiveMeshBounds(InPrimitiveMeshBounds)
-		, SelectedMeshAttributeId(InSelectedMeshAttributeId)
-	{}
+		const TArray<FBox>& InPrimitiveMeshBounds);
 
 	//~ Begin FComputeDataProviderRenderProxy Interface
 	bool IsValid(FValidationData const& InValidationData) const override;
+	void AllocateResources(FRDGBuilder& GraphBuilder, FAllocationData const& InAllocationData) override;
 	void GatherDispatchData(FDispatchData const& InDispatchData) override;
 	//~ End FComputeDataProviderRenderProxy Interface
 
@@ -126,9 +120,15 @@ protected:
 	TArray<float> SelectionCDF;
 
 	int32 SelectorAttributeId = INDEX_NONE;
+
 	TArray<int32> PrimitiveStringKeys;
+	FRDGBufferSRVRef PrimitiveStringKeysBufferSRV = nullptr;
 
-	TArray<FBox> PrimitiveMeshBounds;
+	TArray<FVector4f> PrimitiveMeshBoundsMin;
+	FRDGBufferSRVRef PrimitiveMeshBoundsMinBufferSRV = nullptr;
+
+	TArray<FVector4f> PrimitiveMeshBoundsMax;
+	FRDGBufferSRVRef PrimitiveMeshBoundsMaxBufferSRV = nullptr;
 
 	int32 SelectedMeshAttributeId = INDEX_NONE;
 };
