[PCG] Fix "onethread" issue where we were executing/queuing RunnerFunc infinitely because RenderThread == GameThread
- Task execution was stuck in a loop where there was always a task to execute (RunnerFunc was never completing because Gamethread wasn't advancing)
- Check IsInActualRenderThread and use ExecuteOnGameThread to delay the call

diff --git a/Engine/Plugins/PCG/Source/PCG/Private/Compute/PCGDataForGPU.cpp b/Engine/Plugins/PCG/Source/PCG/Private/Compute/PCGDataForGPU.cpp
--- a/Engine/Plugins/PCG/Source/PCG/Private/Compute/PCGDataForGPU.cpp
+++ b/Engine/Plugins/PCG/Source/PCG/Private/Compute/PCGDataForGPU.cpp
@@ -1672,7 +1672,30 @@ bool FPCGProxyForGPUDataCollection::GetCPUData(FPCGContext* InContext, int32 InD
 
 			ReadbackRequest->EnqueueCopy(RHICmdList, Buffer->GetRHI());
 
-			auto RunnerFunc = [WeakThis, ReadbackRequest, BufferSizeBytes](auto&& RunnerFunc) -> void
+			auto ExecuteAsync = [](auto&& RunnerFunc) -> void
+			{
+				if (IsInActualRenderingThread())
+				{
+					AsyncTask(ENamedThreads::ActualRenderingThread, [RunnerFunc]()
+					{
+						RunnerFunc(RunnerFunc);
+					});
+				}
+				else
+				{
+					// In specific cases (Server, -onethread, etc) the RenderingThread is actually the same as the GameThread.
+					// When this happens we want to avoid calling AsyncTask which could put us in a infinite task execution loop. 
+					// The reason is that if we are running this callback through the task graph we might stay in an executing loop until it has no tasks to execute,
+					// since we are pushing a new task as long as our data isn't ready and we are not advancing the GameThread as we are already on the GameThread this causes a infinite task execution.
+					// Instead delay to GameThread with ExecuteOnGameThread
+					ExecuteOnGameThread(UE_SOURCE_LOCATION, [RunnerFunc]()
+					{
+						RunnerFunc(RunnerFunc);
+					});
+				}
+			};
+
+			auto RunnerFunc = [WeakThis, ReadbackRequest, BufferSizeBytes, ExecuteAsync](auto&& RunnerFunc) -> void
 			{
 				if (TSharedPtr<FPCGProxyForGPUDataCollection> This = WeakThis.Pin())
 				{
@@ -1692,18 +1715,12 @@ bool FPCGProxyForGPUDataCollection::GetCPUData(FPCGContext* InContext, int32 InD
 					}
 					else
 					{
-						AsyncTask(ENamedThreads::ActualRenderingThread, [RunnerFunc]()
-						{
-							RunnerFunc(RunnerFunc);
-						});
+						ExecuteAsync(RunnerFunc);
 					}
 				}
 			};
 
-			AsyncTask(ENamedThreads::ActualRenderingThread, [RunnerFunc]()
-			{
-				RunnerFunc(RunnerFunc);
-			});
+			ExecuteAsync(RunnerFunc);
 
 		});
 	}
diff --git a/Engine/Plugins/PCG/Source/PCGCompute/Private/PCGTextureReadback.cpp b/Engine/Plugins/PCG/Source/PCGCompute/Private/PCGTextureReadback.cpp
--- a/Engine/Plugins/PCG/Source/PCGCompute/Private/PCGTextureReadback.cpp
+++ b/Engine/Plugins/PCG/Source/PCGCompute/Private/PCGTextureReadback.cpp
@@ -113,7 +113,30 @@ void FPCGTextureReadbackInterface::Dispatch_RenderThread(FRHICommandListImmediat
 	FRHIGPUTextureReadback* GPUTextureReadback = new FRHIGPUTextureReadback(TEXT("PCGTextureReadbackCopy"));
 	GPUTextureReadback->EnqueueCopy(RHICmdList, PassParameters.OutputTexture->GetTexture());
 
-	auto RunnerFunc = [GPUTextureReadback, AsyncCallback](auto&& RunnerFunc) -> void
+	auto ExecuteAsync = [](auto&& RunnerFunc) -> void
+	{
+		if(IsInActualRenderingThread())
+		{
+			AsyncTask(ENamedThreads::ActualRenderingThread, [RunnerFunc]()
+			{
+				RunnerFunc(RunnerFunc);
+			});
+		}
+		else
+		{
+			// In specific cases (Server, -onethread, etc) the RenderingThread is actually the same as the GameThread.
+			// When this happens we want to avoid calling AsyncTask which could put us in a infinite task execution loop. 
+			// The reason is that if we are running this callback through the task graph we might stay in an executing loop until it has no tasks to execute,
+			// since we are pushing a new task as long as our data isn't ready and we are not advancing the GameThread as we are already on the GameThread this causes a infinite task execution.
+			// Instead delay to GameThread with ExecuteOnGameThread
+			ExecuteOnGameThread(UE_SOURCE_LOCATION, [RunnerFunc]()
+			{
+				RunnerFunc(RunnerFunc);
+			});
+		}
+	};
+
+	auto RunnerFunc = [GPUTextureReadback, AsyncCallback, ExecuteAsync](auto&& RunnerFunc) -> void
 	{
 		if (GPUTextureReadback->IsReady())
 		{
@@ -127,17 +150,11 @@ void FPCGTextureReadbackInterface::Dispatch_RenderThread(FRHICommandListImmediat
 		}
 		else
 		{
-			AsyncTask(ENamedThreads::ActualRenderingThread, [RunnerFunc]()
-			{
-				RunnerFunc(RunnerFunc);
-			});
+			ExecuteAsync(RunnerFunc);
 		}
 	};
 
-	AsyncTask(ENamedThreads::ActualRenderingThread, [RunnerFunc]()
-	{
-		RunnerFunc(RunnerFunc);
-	});
+	ExecuteAsync(RunnerFunc);
 }
 
 void FPCGTextureReadbackInterface::Dispatch_GameThread(const FPCGTextureReadbackDispatchParams& Params, const TFunction<void(void* OutBuffer, int32 ReadbackWidth, int32 ReadbackHeight)>& AsyncCallback)
