[PCG] Fix array index crash when mesh bounds are used.

#rnx

diff --git a/Engine/Plugins/PCG/Source/PCG/Private/Compute/DataInterfaces/Elements/PCGStaticMeshSpawnerDataInterface.cpp b/Engine/Plugins/PCG/Source/PCG/Private/Compute/DataInterfaces/Elements/PCGStaticMeshSpawnerDataInterface.cpp
--- a/Engine/Plugins/PCG/Source/PCG/Private/Compute/DataInterfaces/Elements/PCGStaticMeshSpawnerDataInterface.cpp
+++ b/Engine/Plugins/PCG/Source/PCG/Private/Compute/DataInterfaces/Elements/PCGStaticMeshSpawnerDataInterface.cpp
@@ -764,8 +764,8 @@ FPCGStaticMeshSpawnerDataProviderProxy::FPCGStaticMeshSpawnerDataProviderProxy(
 
 	for (int Index = 0; Index < InPrimitiveMeshBounds.Num(); ++Index)
 	{
-		PrimitiveMeshBoundsMin[Index] = FVector4f(InPrimitiveMeshBounds[Index].Min.X, InPrimitiveMeshBounds[Index].Min.Y, InPrimitiveMeshBounds[Index].Min.Z, /*Unused*/0.0f);
-		PrimitiveMeshBoundsMax[Index] = FVector4f(InPrimitiveMeshBounds[Index].Max.X, InPrimitiveMeshBounds[Index].Max.Y, InPrimitiveMeshBounds[Index].Max.Z, /*Unused*/0.0f);
+		PrimitiveMeshBoundsMin.Add(FVector4f(InPrimitiveMeshBounds[Index].Min.X, InPrimitiveMeshBounds[Index].Min.Y, InPrimitiveMeshBounds[Index].Min.Z, /*Unused*/0.0f));
+		PrimitiveMeshBoundsMax.Add(FVector4f(InPrimitiveMeshBounds[Index].Max.X, InPrimitiveMeshBounds[Index].Max.Y, InPrimitiveMeshBounds[Index].Max.Z, /*Unused*/0.0f));
 	}
 }
 
