[PCG] Real fix for crash in texture data interface - be defensive against incoming null textures.

Moving back to leaning on GSystemTextures because it's much nicer, and we have some other code leaning on it so we need to know if there's any issue with it.

#rnx

diff --git a/Engine/Plugins/PCG/Source/PCG/Private/Compute/DataInterfaces/PCGTextureDataInterface.cpp b/Engine/Plugins/PCG/Source/PCG/Private/Compute/DataInterfaces/PCGTextureDataInterface.cpp
--- a/Engine/Plugins/PCG/Source/PCG/Private/Compute/DataInterfaces/PCGTextureDataInterface.cpp
+++ b/Engine/Plugins/PCG/Source/PCG/Private/Compute/DataInterfaces/PCGTextureDataInterface.cpp
@@ -292,60 +292,40 @@ void FPCGTextureDataProviderProxy::AllocateResources(FRDGBuilder& GraphBuilder,
 
 void FPCGTextureDataProviderProxy::CreateTextureSRVs(FRDGBuilder& GraphBuilder)
 {
-	// Outside loop so we only create them when necessary, and then share across bindings.
-	FRDGTextureRef DummyTextureRDG = nullptr;
-	FRDGTextureRef DummyTextureArrayRDG = nullptr;
+	FRDGTextureRef DummyTextureRDG = GSystemTextures.GetDefaultTexture2D(GraphBuilder, EPixelFormat::PF_G8, 0.0f);
+	FRDGTextureRef DummyTextureArrayRDG = GSystemTextures.GetDefaultTexture(GraphBuilder, ETextureDimension::Texture2DArray, EPixelFormat::PF_G8, 0.0f);
 
 	for (int BindingIndex = 0; BindingIndex < PCGTextureDataInterfaceConstants::MAX_NUM_BINDINGS; ++BindingIndex)
 	{
 		// Default to blank dummy texture if no texture was provided.
-		FRDGTextureRef TextureRDG = nullptr;
-		FRDGTextureRef TextureArrayRDG = nullptr;
+		FRDGTextureRef TextureRDG = DummyTextureRDG;
+		FRDGTextureRef TextureArrayRDG = DummyTextureArrayRDG;
 
 		if (BindingIndex < BindingInfos.Num())
 		{
-			FRDGTextureRef TempRDG = nullptr;
+			FRDGTextureRef Texture = nullptr;
 
 			if (BindingInfos[BindingIndex].ResourceType == EPCGTextureResourceType::TextureObject && BindingInfos[BindingIndex].Texture)
 			{
-				TempRDG = GraphBuilder.RegisterExternalTexture(
+				Texture = GraphBuilder.RegisterExternalTexture(
 					CreateRenderTarget(BindingInfos[BindingIndex].Texture, *FText::Format(LOCTEXT("RenderTargetDebugName", "PCGDataInterfaceRenderTarget_{0}"), BindingIndex).ToString()));
 			}
 			else if (BindingInfos[BindingIndex].ResourceType == EPCGTextureResourceType::ExportedTexture && BindingInfos[BindingIndex].ExportedTexture)
 			{
-				TempRDG = GraphBuilder.RegisterExternalTexture(BindingInfos[BindingIndex].ExportedTexture);
+				Texture = GraphBuilder.RegisterExternalTexture(BindingInfos[BindingIndex].ExportedTexture);
 			}
 
-			if (BindingInfos[BindingIndex].bIsTextureArray)
+			if (Texture)
 			{
-				TextureArrayRDG = TempRDG;
-			}
-			else
-			{
-				TextureRDG = TempRDG;
-			}
-		}
-
-		if (TextureRDG == nullptr)
-		{
-			if (DummyTextureRDG == nullptr)
-			{
-				DummyTextureRDG = CreateDummyTexture(GraphBuilder);
-				check(DummyTextureRDG);
-			}
-
-			TextureRDG = DummyTextureRDG;
-		}
-
-		if (TextureArrayRDG == nullptr)
-		{
-			if (DummyTextureArrayRDG == nullptr)
-			{
-				DummyTextureArrayRDG = CreateDummyTextureArray(GraphBuilder);
-				check(DummyTextureArrayRDG);
+				if (BindingInfos[BindingIndex].bIsTextureArray)
+				{
+					TextureArrayRDG = Texture;
+				}
+				else
+				{
+					TextureRDG = Texture;
+				}
 			}
-
-			TextureArrayRDG = DummyTextureArrayRDG;
 		}
 
 		TextureSRV[BindingIndex] = GraphBuilder.CreateSRV(TextureRDG);
@@ -389,52 +369,5 @@ void FPCGTextureDataProviderProxy::PackTextureInfos(FRDGBuilder& GraphBuilder)
 	}
 }
 
-FRDGTextureRef FPCGTextureDataProviderProxy::CreateDummyTexture(FRDGBuilder& GraphBuilder)
-{
-	// Copied from GSystemTextures.GetBlackDummy(). We cannot use GSystemTextures.GetDefaultTexture2D because it
-	// can fail depending on the callstack leading to the compute graph execution.
-
-	const FRHITextureCreateDesc Desc = FRHITextureCreateDesc::Create2D(TEXT("PCGDummyTexture"), 1, 1, PF_G8)
-		.SetFlags(ETextureCreateFlags::ShaderResource);
-
-	FTextureRHIRef Texture = RHICreateTexture(Desc);
-
-	{
-		uint32 DestStride;
-		uint8* DestBuffer = (uint8*)RHILockTexture2D(Texture, 0, RLM_WriteOnly, DestStride, false);
-		*DestBuffer = 0;
-		RHIUnlockTexture2D(Texture, 0, false);
-	}
-
-	// Create a dummy texture to bind.
-	TRefCountPtr<IPooledRenderTarget> BlackDummyTarget = CreateRenderTarget(Texture, TEXT("PCGDummyRenderTarget"));
-	check(BlackDummyTarget.IsValid());
-
-	return GraphBuilder.RegisterExternalTexture(BlackDummyTarget);
-}
-
-FRDGTextureRef FPCGTextureDataProviderProxy::CreateDummyTextureArray(FRDGBuilder& GraphBuilder)
-{
-	// Copied from GSystemTextures.GetBlackDummy(). We cannot use GSystemTextures.GetDefaultTexture because it
-	// can fail depending on the callstack leading to the compute graph execution.
-
-	const FRHITextureCreateDesc Desc = FRHITextureCreateDesc::Create2DArray(TEXT("PCGDummyTextureArray"), 1, 1, 1, PF_G8)
-		.SetFlags(ETextureCreateFlags::ShaderResource);
-
-	FTextureRHIRef Texture = RHICreateTexture(Desc);
-
-	{
-		uint32 DestStride;
-		uint8* DestBuffer = (uint8*)RHILockTexture2D(Texture, 0, RLM_WriteOnly, DestStride, false);
-		*DestBuffer = 0;
-		RHIUnlockTexture2D(Texture, 0, false);
-	}
-
-	// Create a dummy texture to bind.
-	TRefCountPtr<IPooledRenderTarget> BlackDummyTarget = CreateRenderTarget(Texture, TEXT("PCGDummyArrayRenderTarget"));
-	check(BlackDummyTarget.IsValid());
-
-	return GraphBuilder.RegisterExternalTexture(BlackDummyTarget);
-}
-
 #undef LOCTEXT_NAMESPACE
+
diff --git a/Engine/Plugins/PCG/Source/PCG/Private/Compute/DataInterfaces/PCGTextureDataInterface.h b/Engine/Plugins/PCG/Source/PCG/Private/Compute/DataInterfaces/PCGTextureDataInterface.h
--- a/Engine/Plugins/PCG/Source/PCG/Private/Compute/DataInterfaces/PCGTextureDataInterface.h
+++ b/Engine/Plugins/PCG/Source/PCG/Private/Compute/DataInterfaces/PCGTextureDataInterface.h
@@ -103,10 +103,6 @@ public:
 	void CreateTextureSRVs(FRDGBuilder& GraphBuilder);
 	void PackTextureInfos(FRDGBuilder& GraphBuilder);
 
-protected:
-	FRDGTextureRef CreateDummyTexture(FRDGBuilder& GraphBuilder);
-	FRDGTextureRef CreateDummyTextureArray(FRDGBuilder& GraphBuilder);
-
 protected:
 	using FParameters = FPCGTextureDataInterfaceParameters;
 
