[PCG] Editor remove OnObjectsReplaced necessity by replacing all WeakPtrs by SoftObjectPtrs
- The only real case we need to support is BP Reinstancing in which case the path of the object doesn't change and so resolving should still work

#rnx

diff --git a/Engine/Plugins/PCG/Source/PCG/Private/Graph/PCGStackContext.cpp b/Engine/Plugins/PCG/Source/PCG/Private/Graph/PCGStackContext.cpp
--- a/Engine/Plugins/PCG/Source/PCG/Private/Graph/PCGStackContext.cpp
+++ b/Engine/Plugins/PCG/Source/PCG/Private/Graph/PCGStackContext.cpp
@@ -269,36 +269,6 @@ bool FPCGStack::HasObject(const UObject* InObject) const
 	});
 }
 
-bool FPCGStack::ReplaceRoot(const TMap<UObject*, UObject*>& ReplacementMap)
-{
-	if (!StackFrames.IsEmpty())
-	{
-		UObject* NewStackRoot = ReplacementMap.FindRef(StackFrames[0].GetObject_AnyThread());
-
-		// If the stack frame was marked as garbage, NewStackRoot will be nullptr, but we still match against the object ptr hash.
-		if (!NewStackRoot)
-		{
-			for (const TPair<UObject*, UObject*>& Pair : ReplacementMap)
-			{
-				// TODO: If we ever have a replacement where the name changed, this won't work.
-				// We might need to have a serialization pointer (soft) and an execution pointer (weak).
-				if (Pair.Key && Pair.Value && StackFrames[0].Object == Pair.Value)
-				{
-					NewStackRoot = Pair.Value;
-				}
-			}
-		}
-
-		if (NewStackRoot)
-		{
-			StackFrames[0].SetObject(NewStackRoot);
-			return true;
-		}
-	}
-
-	return false;
-}
-
 bool FPCGStack::operator==(const FPCGStack& Other) const
 {
 	// Stacks are the same if all stack frames are the same
diff --git a/Engine/Plugins/PCG/Source/PCG/Public/Graph/PCGStackContext.h b/Engine/Plugins/PCG/Source/PCG/Public/Graph/PCGStackContext.h
--- a/Engine/Plugins/PCG/Source/PCG/Public/Graph/PCGStackContext.h
+++ b/Engine/Plugins/PCG/Source/PCG/Public/Graph/PCGStackContext.h
@@ -177,8 +177,8 @@ public:
 	/** Stack has a frame corresponding to the given object. */
 	bool HasObject(const UObject* InObject) const;
 
-	/** Utility function to replace the root object, in case of re-instantiation. Returns true if a replacement happened. */
-	bool ReplaceRoot(const TMap<UObject*, UObject*>& ReplacementMap);
+	UE_DEPRECATED(5.6, "No longer supported")
+	bool ReplaceRoot(const TMap<UObject*, UObject*>& ReplacementMap) { return false; }
 
 	bool operator==(const FPCGStack& Other) const;
 	bool operator!=(const FPCGStack& Other) const { return !(*this == Other); }
diff --git a/Engine/Plugins/PCG/Source/PCGEditor/Private/PCGEditor.cpp b/Engine/Plugins/PCG/Source/PCGEditor/Private/PCGEditor.cpp
--- a/Engine/Plugins/PCG/Source/PCGEditor/Private/PCGEditor.cpp
+++ b/Engine/Plugins/PCG/Source/PCGEditor/Private/PCGEditor.cpp
@@ -127,7 +127,6 @@ UPCGEditorGraph* FPCGEditor::GetPCGEditorGraph(const UPCGSettings* InSettings)
 
 void FPCGEditor::Initialize(const EToolkitMode::Type InMode, const TSharedPtr<class IToolkitHost>& InToolkitHost, UPCGGraph* InPCGGraph)
 {
-	FCoreUObjectDelegates::OnObjectsReplaced.AddSP(this, &FPCGEditor::OnObjectsReplaced);
 	PCGGraphBeingEdited = InPCGGraph;
 
 	// Initializes the UPCGEditorGraph if needed
@@ -253,27 +252,6 @@ void FPCGEditor::Initialize(const EToolkitMode::Type InMode, const TSharedPtr<cl
 	}
 }
 
-void FPCGEditor::OnObjectsReplaced(const TMap<UObject*, UObject*>& ReplacementMap)
-{
-	if (UObject* NewComponentBeingInspected = ReplacementMap.FindRef(PCGComponentBeingInspected.GetEvenIfUnreachable()))
-	{
-		PCGComponentBeingInspected = Cast<UPCGComponent>(NewComponentBeingInspected);
-	}
-
-	if (UObject* NewLastValidPCGComponentBeingInspected = ReplacementMap.FindRef(LastValidPCGComponentBeingInspected.GetEvenIfUnreachable()))
-	{
-		LastValidPCGComponentBeingInspected = Cast<UPCGComponent>(NewLastValidPCGComponentBeingInspected);
-	}
-
-	StackBeingInspected.ReplaceRoot(ReplacementMap);
-
-	// Propagate object replacement to the debug object tree view too
-	if (DebugObjectTreeWidget)
-	{
-		DebugObjectTreeWidget->OnObjectsReplaced(ReplacementMap);
-	}
-}
-
 UPCGEditorGraph* FPCGEditor::GetPCGEditorGraph()
 {
 	return PCGEditorGraph;
@@ -2804,8 +2782,6 @@ void FPCGEditor::OnClose()
 		UnregisterDelegatesForWorld(GEditor->GetEditorWorldContext().World());
 		UnregisterDelegatesForWorld(GEditor->PlayWorld.Get());
 	}
-
-	FCoreUObjectDelegates::OnObjectsReplaced.RemoveAll(this);
 }
 
 void FPCGEditor::InitToolMenuContext(FToolMenuContext& MenuContext)
@@ -3020,7 +2996,7 @@ void FPCGEditor::OnComponentGenerationDone(UPCGSubsystem* Subsystem, UPCGCompone
 	bool bShouldRefresh = !Component || StackBeingInspected.GetRootComponent() == Component || Component->GetGraph() == PCGGraphBeingEdited;
 
 	// Additionally, if we are not inspecting but the component that's done executing contains this graph, then we should also update.
-	if (!bShouldRefresh && PCGComponentBeingInspected.IsExplicitlyNull() && PCGGraphBeingEdited && Subsystem)
+	if (!bShouldRefresh && PCGComponentBeingInspected.IsNull() && PCGGraphBeingEdited && Subsystem)
 	{
 		TArray<FPCGStackSharedPtr> ExecutedStacksContainingThisGraph = Subsystem->GetExecutedStacksPtrs(Component, PCGGraphBeingEdited);
 		bShouldRefresh |= !ExecutedStacksContainingThisGraph.IsEmpty();
diff --git a/Engine/Plugins/PCG/Source/PCGEditor/Private/PCGEditor.h b/Engine/Plugins/PCG/Source/PCGEditor/Private/PCGEditor.h
--- a/Engine/Plugins/PCG/Source/PCGEditor/Private/PCGEditor.h
+++ b/Engine/Plugins/PCG/Source/PCGEditor/Private/PCGEditor.h
@@ -157,9 +157,6 @@ private:
 	/** Bind commands to delegates */
 	void BindCommands();
 
-	/** Callback to process component changes */
-	void OnObjectsReplaced(const TMap<UObject*, UObject*>& ReplacementMap);
-
 	/** Bring up the find tab */
 	void OnFind();
 
@@ -434,8 +431,8 @@ private:
 	TObjectPtr<UPCGGraph> PCGGraphBeingEdited = nullptr;
 	UPCGEditorGraph* PCGEditorGraph = nullptr;
 
-	TWeakObjectPtr<UPCGComponent> PCGComponentBeingInspected;
+	TSoftObjectPtr<UPCGComponent> PCGComponentBeingInspected;
 	// Implementation note: we'll keep the last valid component inspected so we don't un-inspect on spurious selection changes
-	TWeakObjectPtr<UPCGComponent> LastValidPCGComponentBeingInspected;
+	TSoftObjectPtr<UPCGComponent> LastValidPCGComponentBeingInspected;
 	FPCGStack StackBeingInspected;
 };
diff --git a/Engine/Plugins/PCG/Source/PCGEditor/Private/Widgets/SPCGEditorGraphDebugObjectTree.cpp b/Engine/Plugins/PCG/Source/PCGEditor/Private/Widgets/SPCGEditorGraphDebugObjectTree.cpp
--- a/Engine/Plugins/PCG/Source/PCGEditor/Private/Widgets/SPCGEditorGraphDebugObjectTree.cpp
+++ b/Engine/Plugins/PCG/Source/PCGEditor/Private/Widgets/SPCGEditorGraphDebugObjectTree.cpp
@@ -181,14 +181,6 @@ void FPCGEditorGraphDebugObjectItem::SortChildren(bool bIsAscending, bool bIsRec
 	}
 }
 
-void FPCGEditorGraphDebugObjectItem::OnObjectsReplaced(const TMap<UObject*, UObject*>& ReplacementMap)
-{
-	if (FPCGStack* Stack = GetMutablePCGStack())
-	{
-		Stack->ReplaceRoot(ReplacementMap);
-	}
-}
-
 FString FPCGEditorGraphDebugObjectItem_Actor::GetLabel(bool bForSorting) const
 {
 	return Actor.IsValid() ? Actor->GetActorNameOrLabel() : FString();
@@ -196,7 +188,7 @@ FString FPCGEditorGraphDebugObjectItem_Actor::GetLabel(bool bForSorting) const
 
 const FSlateBrush* FPCGEditorGraphDebugObjectItem_Actor::GetIcon() const
 {
-	const UObject* Object = GetObject();
+	const UObject* Object = Actor.Get();
 	return FSlateIconFinder::FindIconBrushForClass(Object ? Object->GetClass() : AActor::StaticClass());
 }
 
@@ -210,16 +202,6 @@ FString FPCGEditorGraphDebugObjectItem_PCGComponent::GetLabel(bool bForSorting)
 	return FString();
 }
 
-void FPCGEditorGraphDebugObjectItem_PCGComponent::OnObjectsReplaced(const TMap<UObject*, UObject*>& ReplacementMap)
-{
-	FPCGEditorGraphDebugObjectItem::OnObjectsReplaced(ReplacementMap);
-
-	if (UObject* NewComponent = ReplacementMap.FindRef(PCGComponent.GetEvenIfUnreachable()))
-	{
-		PCGComponent = Cast<UPCGComponent>(NewComponent);
-	}
-}
-
 FString FPCGEditorGraphDebugObjectItem_PCGSubgraph::GetLabel(bool bForSorting) const
 {
 	if (PCGNode.IsValid() && PCGGraph.IsValid())
@@ -528,37 +510,6 @@ void SPCGEditorGraphDebugObjectTree::RequestRefresh()
 	bNeedsRefresh = true;
 }
 
-void SPCGEditorGraphDebugObjectTree::OnObjectsReplaced(const TMap<UObject*, UObject*>& ReplacementMap)
-{
-	for (FPCGEditorGraphDebugObjectItemPtr& GraphItem : AllGraphItems)
-	{
-		if (GraphItem)
-		{
-			GraphItem->OnObjectsReplaced(ReplacementMap);
-		}
-	}
-
-	// We need to replace the selected component, selected stacks and expanded stacks as they are used for selection matching.
-	if (UObject* NewComponent = ReplacementMap.FindRef(SelectedOriginalComponent.GetEvenIfUnreachable()))
-	{
-		SelectedOriginalComponent = Cast<UPCGComponent>(NewComponent);
-	}
-
-	SelectedStack.ReplaceRoot(ReplacementMap);
-
-	// For the expanded stacks, since it is a set, if we ever modify the root, we should probably re-hash it.
-	// We can't modify it in place so we'll just rebuild it.
-	TSet<FPCGStack> UpdatedExpandedStacks;
-	UpdatedExpandedStacks.Reserve(ExpandedStacks.Num());
-	for (FPCGStack& Stack : ExpandedStacks)
-	{
-		Stack.ReplaceRoot(ReplacementMap);
-		UpdatedExpandedStacks.Add(std::move(Stack));
-	}
-	
-	ExpandedStacks = std::move(UpdatedExpandedStacks);
-}
-
 void SPCGEditorGraphDebugObjectTree::SetNodeBeingInspected(const UPCGNode* InPCGNode)
 {
 	PCGNodeBeingInspected = IsValid(InPCGNode) ? InPCGNode : nullptr;
@@ -1327,7 +1278,7 @@ void SPCGEditorGraphDebugObjectTree::OnSelectionChanged(FPCGEditorGraphDebugObje
 			}
 			else
 			{
-				SelectedOriginalComponent = RootComponent;
+				SelectedOriginalComponent = const_cast<UPCGComponent*>(RootComponent);
 			}
 		}
 	}
diff --git a/Engine/Plugins/PCG/Source/PCGEditor/Private/Widgets/SPCGEditorGraphDebugObjectTree.h b/Engine/Plugins/PCG/Source/PCGEditor/Private/Widgets/SPCGEditorGraphDebugObjectTree.h
--- a/Engine/Plugins/PCG/Source/PCGEditor/Private/Widgets/SPCGEditorGraphDebugObjectTree.h
+++ b/Engine/Plugins/PCG/Source/PCGEditor/Private/Widgets/SPCGEditorGraphDebugObjectTree.h
@@ -52,15 +52,11 @@ public:
 	virtual bool IsDynamic() const { return false; }
 
 	virtual FString GetLabel(bool bForSorting = false) const = 0;
-	virtual UPCGComponent* GetPCGComponent() const = 0;
 	virtual const FPCGStack* GetPCGStack() const { return nullptr; }
-	virtual const UObject* GetObject() const  = 0;
 	virtual const UPCGGraph* GetPCGGraph() const { return nullptr; }
 	virtual bool IsLoopIteration() const { return false; }
 	virtual const FSlateBrush* GetIcon() const;
 
-	virtual void OnObjectsReplaced(const TMap<UObject*, UObject*>& ReplacementMap);
-
 protected:
 	virtual FPCGStack* GetMutablePCGStack() { return nullptr; }
 
@@ -82,8 +78,6 @@ public:
 	}
 
 	virtual FString GetLabel(bool bForSorting = false) const override;
-	virtual UPCGComponent* GetPCGComponent() const override { return nullptr; }
-	virtual const UObject* GetObject() const override { return Actor.Get(); };
 	virtual const FPCGStack* GetPCGStack() const override { return &PCGStack; }
 	virtual const FSlateBrush* GetIcon() const override;
 
@@ -99,8 +93,8 @@ class FPCGEditorGraphDebugObjectItem_PCGComponent : public FPCGEditorGraphDebugO
 {
 public:
 	explicit FPCGEditorGraphDebugObjectItem_PCGComponent(
-		TWeakObjectPtr<UPCGComponent> InPCGComponent,
-		TWeakObjectPtr<const UPCGGraph> InPCGGraph,
+		UPCGComponent* InPCGComponent,
+		const UPCGGraph* InPCGGraph,
 		const FPCGStack& InPCGStack,
 		bool bInIsDebuggable,
 		bool bInHasInspectionData)
@@ -112,19 +106,15 @@ public:
 	{}
 
 	virtual FString GetLabel(bool bForSorting = false) const override;
-	virtual UPCGComponent* GetPCGComponent() const override { return PCGComponent.Get(); }
-	virtual const UObject* GetObject() const override { return PCGComponent.Get(); };
 	virtual const FPCGStack* GetPCGStack() const override { return &PCGStack; }
 	virtual bool IsDebuggable() const override { return bIsDebuggable; }
 	virtual const UPCGGraph* GetPCGGraph() const override { return PCGGraph.Get(); }
 
-	virtual void OnObjectsReplaced(const TMap<UObject*, UObject*>& ReplacementMap) override;
-
 protected:
 	virtual FPCGStack* GetMutablePCGStack() override { return &PCGStack; }
 
-	TWeakObjectPtr<UPCGComponent> PCGComponent = nullptr;
-	TWeakObjectPtr<const UPCGGraph> PCGGraph = nullptr;
+	TSoftObjectPtr<UPCGComponent> PCGComponent;
+	TSoftObjectPtr<const UPCGGraph> PCGGraph;
 
 	FPCGStack PCGStack;
 	bool bIsDebuggable = false;
@@ -156,8 +146,6 @@ public:
 	{}
 
 	virtual FString GetLabel(bool bForSorting = false) const override;
-	virtual UPCGComponent* GetPCGComponent() const override { return GetParent() ? GetParent()->GetPCGComponent() : nullptr; }
-	virtual const UObject* GetObject() const override { return PCGNode.Get(); };
 	virtual const FPCGStack* GetPCGStack() const override { return &PCGStack; }
 	virtual bool IsDebuggable() const override { return bIsDebuggable; }
 	virtual bool IsDynamic() const override { return SubgraphType != EPCGEditorSubgraphNodeType::StaticSubgraph; }
@@ -194,8 +182,6 @@ public:
 	virtual int32 GetLoopIndex() const { return LoopIndex; }
 
 	virtual FString GetLabel(bool bForSorting = false) const override;
-	virtual UPCGComponent* GetPCGComponent() const override { return GetParent() ? GetParent()->GetPCGComponent() : nullptr; }
-	virtual const UObject* GetObject() const override { return nullptr; };
 	virtual const FPCGStack* GetPCGStack() const override { return &PCGStack; }
 	virtual int32 GetSortPriority() const override { return LoopIndex; }
 	virtual bool IsDebuggable() const override { return bIsDebuggable; }
@@ -269,8 +255,6 @@ public:
 
 	void SetNodeBeingInspected(const UPCGNode* InPCGNode);
 
-	void OnObjectsReplaced(const TMap<UObject*, UObject*>& ReplacementMap);
-
 	/** If the stack matches to an item, Expands the tree view to make that item visible and select it, and return true. Otherwise, return false. */
 	bool SetDebugObjectFromStackFromAnotherEditor(const FPCGStack& InStack);
 
@@ -352,10 +336,10 @@ private:
 	FPCGStack SelectedStack;
 
 	/** Used to retain item selection state across tree refreshes if the SelectedStack is invalidated (e.g. through BP reconstruction). */
-	TWeakObjectPtr<const UPCGGraph> SelectedGraph = nullptr;
-	TWeakObjectPtr<const AActor> SelectedOwner = nullptr;
+	TSoftObjectPtr<const UPCGGraph> SelectedGraph = nullptr;
+	TSoftObjectPtr<const AActor> SelectedOwner = nullptr;
 	uint32 SelectedGridSize = PCGHiGenGrid::UnboundedGridSize();
 	FIntVector SelectedGridCoord = FIntVector::ZeroValue;
-	TWeakObjectPtr<const UPCGComponent> SelectedOriginalComponent = nullptr;
-	TWeakObjectPtr<const UPCGNode> PCGNodeBeingInspected = nullptr;
+	TSoftObjectPtr<UPCGComponent> SelectedOriginalComponent = nullptr;
+	TSoftObjectPtr<const UPCGNode> PCGNodeBeingInspected = nullptr;
 };
