[PCG] Fix shader compile warning with narrowing of matrix type.

#rnx

diff --git a/Engine/Plugins/PCG/Shaders/Private/PCGDataCollectionDataInterface.ush b/Engine/Plugins/PCG/Shaders/Private/PCGDataCollectionDataInterface.ush
--- a/Engine/Plugins/PCG/Shaders/Private/PCGDataCollectionDataInterface.ush
+++ b/Engine/Plugins/PCG/Shaders/Private/PCGDataCollectionDataInterface.ush
@@ -521,7 +521,7 @@ void SetPointTransform_{DataInterfaceName}(uint InDataIndex, uint InElementIndex
 	SetScale_{DataInterfaceName}(InDataIndex, InElementIndex, Scale);
 
 	// Assumes Transform axes have the correct handedness. We could do more work to fix this up too if required.
-	float3x3 RotationMatix = transpose(Transform);
+	float3x3 RotationMatix = (float3x3)transpose(Transform);
 	// Required unfortunately, QuatFromMatrix expects orthonormalized.
 	RotationMatix[0] /= Scale.x;
 	RotationMatix[1] /= Scale.y;
