[PCG] Fix unpacked GPU data not creating metadata attributes when there are no elements

#rnx

diff --git a/Engine/Plugins/PCG/Source/PCG/Private/Compute/PCGDataForGPU.cpp b/Engine/Plugins/PCG/Source/PCG/Private/Compute/PCGDataForGPU.cpp
--- a/Engine/Plugins/PCG/Source/PCG/Private/Compute/PCGDataForGPU.cpp
+++ b/Engine/Plugins/PCG/Source/PCG/Private/Compute/PCGDataForGPU.cpp
@@ -1269,6 +1269,7 @@ EPCGUnpackDataCollectionResult FPCGDataCollectionDesc::UnpackDataCollection(FPCG
 	}
 
 	TArray<FPCGTaggedData>& OutData = OutDataCollection.TaggedData;
+	TStaticArray<FPCGMetadataAttributeBase*, MAX_NUM_ATTRS> MetadataAttributes;
 
 	for (uint32 DataIndex = 0; DataIndex < NumData; ++DataIndex)
 	{
@@ -1287,6 +1288,7 @@ EPCGUnpackDataCollectionResult FPCGDataCollectionDesc::UnpackDataCollection(FPCG
 		const FPCGDataDesc& DataDesc = DataDescs[DataIndex];
 		const TArray<FPCGKernelAttributeDesc>& AttributeDescs = DataDesc.AttributeDescs;
 		check(NumAttributes == AttributeDescs.Num());
+		check(AttributeDescs.Num() <= MAX_NUM_ATTRS);
 
 		if (TypeId == POINT_DATA_TYPE_ID)
 		{
@@ -1337,6 +1339,14 @@ EPCGUnpackDataCollectionResult FPCGDataCollectionDesc::UnpackDataCollection(FPCG
 				}
 			}
 
+			for (const FPCGKernelAttributeDesc& AttributeDesc : AttributeDescs)
+			{
+				if (AttributeDesc.Index >= NUM_RESERVED_ATTRS)
+				{
+					MetadataAttributes[AttributeDesc.Index] = PCGDataForGPUHelpers::CreateAttributeFromAttributeDesc(Metadata, AttributeDesc);
+				}
+			}
+
 			// No work to do if there are no elements to process.
 			if (NumElements == 0)
 			{
@@ -1472,7 +1482,7 @@ EPCGUnpackDataCollectionResult FPCGDataCollectionDesc::UnpackDataCollection(FPCG
 				{
 					TRACE_CPUPROFILER_EVENT_SCOPE(UnpackAttribute);
 
-					if (FPCGMetadataAttributeBase* AttributeBase = PCGDataForGPUHelpers::CreateAttributeFromAttributeDesc(Metadata, AttributeDesc))
+					if (FPCGMetadataAttributeBase* AttributeBase = MetadataAttributes[AttributeDesc.Index])
 					{
 						ensure(PCGDataForGPUHelpers::UnpackAttributeHelper(InContext, PackedData, AttributeDesc, InStringTable, AddressUints, NumElements, OutPointData));
 					}
@@ -1521,6 +1531,11 @@ EPCGUnpackDataCollectionResult FPCGDataCollectionDesc::UnpackDataCollection(FPCG
 			OutTaggedData.Data = OutParamData;
 			OutTaggedData.Pin = InPin;
 
+			for (const FPCGKernelAttributeDesc& AttributeDesc : AttributeDescs)
+			{
+				MetadataAttributes[AttributeDesc.Index] = PCGDataForGPUHelpers::CreateAttributeFromAttributeDesc(Metadata, AttributeDesc);
+			}
+
 			// No work to do if there are no elements to process.
 			if (NumElements == 0)
 			{
@@ -1532,7 +1547,7 @@ EPCGUnpackDataCollectionResult FPCGDataCollectionDesc::UnpackDataCollection(FPCG
 			{
 				TRACE_CPUPROFILER_EVENT_SCOPE(UnpackAttribute);
 
-				if (FPCGMetadataAttributeBase* AttributeBase = PCGDataForGPUHelpers::CreateAttributeFromAttributeDesc(Metadata, AttributeDesc))
+				if (FPCGMetadataAttributeBase* AttributeBase = MetadataAttributes[AttributeDesc.Index])
 				{
 					const uint32 AddressUints = GetElementDataStartAddressUints(DataAsUint, DataIndex, AttributeDesc.Index);
 					ensure(PCGDataForGPUHelpers::UnpackAttributeHelper(InContext, PackedData, AttributeDesc, InStringTable, AddressUints, NumElements, OutParamData));
