Fix CIS: fix warning about possible memory leak by storing into unique ptr and releasing to the lambda capture

#rnx

diff --git a/Engine/Plugins/PCG/Source/PCGCompute/Private/PCGTextureReadback.cpp b/Engine/Plugins/PCG/Source/PCGCompute/Private/PCGTextureReadback.cpp
--- a/Engine/Plugins/PCG/Source/PCGCompute/Private/PCGTextureReadback.cpp
+++ b/Engine/Plugins/PCG/Source/PCGCompute/Private/PCGTextureReadback.cpp
@@ -110,7 +110,7 @@ void FPCGTextureReadbackInterface::Dispatch_RenderThread(FRHICommandListImmediat
 	// Prepare OutputTexture to be copied
 	RHICmdList.Transition(FRHITransitionInfo(OutputTexture, ERHIAccess::UAVCompute, ERHIAccess::CopySrc));
 
-	FRHIGPUTextureReadback* GPUTextureReadback = new FRHIGPUTextureReadback(TEXT("PCGTextureReadbackCopy"));
+	TUniquePtr<FRHIGPUTextureReadback> GPUTextureReadback = MakeUnique<FRHIGPUTextureReadback>(TEXT("PCGTextureReadbackCopy"));
 	GPUTextureReadback->EnqueueCopy(RHICmdList, PassParameters.OutputTexture->GetTexture());
 
 	auto ExecuteAsync = [](auto&& RunnerFunc) -> void
@@ -136,17 +136,17 @@ void FPCGTextureReadbackInterface::Dispatch_RenderThread(FRHICommandListImmediat
 		}
 	};
 
-	auto RunnerFunc = [GPUTextureReadback, AsyncCallback, ExecuteAsync](auto&& RunnerFunc) -> void
+	auto RunnerFunc = [GPUTextureReadBackPtr = GPUTextureReadback.Release(), AsyncCallback, ExecuteAsync](auto&& RunnerFunc) -> void
 	{
-		if (GPUTextureReadback->IsReady())
+		if (GPUTextureReadBackPtr->IsReady())
 		{
 			int32 ReadbackWidth = 0, ReadbackHeight = 0;
-			void* OutBuffer = GPUTextureReadback->Lock(ReadbackWidth, &ReadbackHeight);
+			void* OutBuffer = GPUTextureReadBackPtr->Lock(ReadbackWidth, &ReadbackHeight);
 
 			AsyncCallback(OutBuffer, ReadbackWidth, ReadbackHeight);
 
-			GPUTextureReadback->Unlock();
-			delete GPUTextureReadback;
+			GPUTextureReadBackPtr->Unlock();
+			delete GPUTextureReadBackPtr;
 		}
 		else
 		{
