[PCG] Fixed signed vs unsigned issues triggering shader compilation errors

#rnx

diff --git a/Engine/Plugins/PCG/Shaders/Private/Elements/PCGCopyPointsAnalysis.usf b/Engine/Plugins/PCG/Shaders/Private/Elements/PCGCopyPointsAnalysis.usf
--- a/Engine/Plugins/PCG/Shaders/Private/Elements/PCGCopyPointsAnalysis.usf
+++ b/Engine/Plugins/PCG/Shaders/Private/Elements/PCGCopyPointsAnalysis.usf
@@ -52,8 +52,8 @@ void PCGCopyPointsAnalysisCS(uint3 GroupId : SV_GroupID, uint GroupIndex : SV_Gr
 	else
 	{
 		// Output is attribute set with max(SourceNumData, TargetNumData) bools. (S, T) can be (N, N), (1, N) or (N, 1).
-		SourceDataIndex = clamp(CopyOutputDataIndex, 0, -1 + (int)SourceNumData);
-		TargetDataIndex = clamp(CopyOutputDataIndex, 0, -1 + (int)TargetNumData);
+		SourceDataIndex = clamp(CopyOutputDataIndex, 0u, (uint)(SourceNumData - 1));
+		TargetDataIndex = clamp(CopyOutputDataIndex, 0u, (uint)(TargetNumData - 1));
 	}
 
 	if (Source_GetNumElements(SourceDataIndex) > 0 && Target_GetNumElements(TargetDataIndex) > 0)
diff --git a/Engine/Plugins/PCG/Source/PCG/Private/Compute/DataInterfaces/Elements/PCGCopyPointsDataInterface.cpp b/Engine/Plugins/PCG/Source/PCG/Private/Compute/DataInterfaces/Elements/PCGCopyPointsDataInterface.cpp
--- a/Engine/Plugins/PCG/Source/PCG/Private/Compute/DataInterfaces/Elements/PCGCopyPointsDataInterface.cpp
+++ b/Engine/Plugins/PCG/Source/PCG/Private/Compute/DataInterfaces/Elements/PCGCopyPointsDataInterface.cpp
@@ -129,8 +129,8 @@ void UPCGCopyPointsDataInterface::GetHLSL(FString& OutHLSL, FString const& InDat
 		"	}\n"
 		"	else\n"
 		"	{\n"
-		"		OutSourceIndex = clamp(InOutputDataIndex, 0, -1 + (int)InSourceDataCount);\n"
-		"		OutTargetIndex = clamp(InOutputDataIndex, 0, -1 + (int)InTargetDataCount);\n"
+		"		OutSourceIndex = clamp(InOutputDataIndex, 0u, (uint)(InSourceDataCount - 1));\n"
+		"		OutTargetIndex = clamp(InOutputDataIndex, 0u, (uint)(InTargetDataCount - 1));\n"
 		"	}\n"
 		"#endif\n"
 		"}\n"
