[PCG] Visual Logs: fix OnNodeChangedDelegate being called outside of GameThread triggering delegate race detector
- Also fix AddLog which was never calling OnNodeChangedDelegate outside of the GameThread

#rnx

diff --git a/Engine/Plugins/PCG/Source/PCG/Private/Utils/PCGNodeVisualLogs.cpp b/Engine/Plugins/PCG/Source/PCG/Private/Utils/PCGNodeVisualLogs.cpp
--- a/Engine/Plugins/PCG/Source/PCG/Private/Utils/PCGNodeVisualLogs.cpp
+++ b/Engine/Plugins/PCG/Source/PCG/Private/Utils/PCGNodeVisualLogs.cpp
@@ -6,6 +6,8 @@
 #include "PCGModule.h"
 
 #include "Algo/Find.h"
+#include "Algo/Transform.h"
+#include "Containers/Ticker.h"
 #include "GameFramework/Actor.h"
 
 #define LOCTEXT_NAMESPACE "PCGNode"
@@ -38,15 +40,42 @@ void FPCGNodeVisualLogs::Log(const FPCGStack& InPCGStack, ELogVerbosity::Type In
 	}
 
 	// Broadcast outside of write scope lock
-	if (bAdded && IsInGameThread() && !InPCGStack.GetStackFrames().IsEmpty())
+	if (bAdded && !InPCGStack.GetStackFrames().IsEmpty())
 	{
-		for (const FPCGStackFrame& Frame : InPCGStack.GetStackFrames())
+		const bool bIsInGameThread = IsInGameThread();
+
+		TArray<TWeakObjectPtr<const UPCGNode>> NodeWeakPtrs;
 		{
-			if (const UPCGNode* Node = Frame.GetObject_GameThread<UPCGNode>())
+			FGCScopeGuard Guard;
+			for (const FPCGStackFrame& Frame : InPCGStack.GetStackFrames())
 			{
-				Node->OnNodeChangedDelegate.Broadcast(const_cast<UPCGNode*>(Node), EPCGChangeType::Cosmetic);
+				if (const UPCGNode* Node = Frame.GetObject_NoGuard<UPCGNode>())
+				{
+					if (bIsInGameThread)
+					{
+						Node->OnNodeChangedDelegate.Broadcast(const_cast<UPCGNode*>(Node), EPCGChangeType::Cosmetic);
+					}
+					else
+					{
+						NodeWeakPtrs.Add(Node);
+					}
+				}
 			}
 		}
+
+		if(!NodeWeakPtrs.IsEmpty())
+		{
+			ExecuteOnGameThread(UE_SOURCE_LOCATION, [NodeWeakPtrs]()
+			{
+				for(const TWeakObjectPtr<const UPCGNode>& NodeWeakPtr : NodeWeakPtrs)
+				{
+					if (const UPCGNode* Node = NodeWeakPtr.Get())
+					{
+						Node->OnNodeChangedDelegate.Broadcast(const_cast<UPCGNode*>(Node), EPCGChangeType::Cosmetic);
+					}
+				}
+			});
+		}
 	}
 }
 
@@ -271,9 +300,32 @@ void FPCGNodeVisualLogs::ClearLogs(const FPCGStack& InPCGStack)
 	}
 
 	// Broadcast change notification outside of write scope lock
-	for (const UPCGNode* TouchedNode : TouchedNodes)
+	if (IsInGameThread())
+	{
+		for (const UPCGNode* Node : TouchedNodes)
+		{
+			if (Node)
+			{
+				Node->OnNodeChangedDelegate.Broadcast(const_cast<UPCGNode*>(Node), EPCGChangeType::Cosmetic);
+			}
+		}
+	}
+	else if(!TouchedNodes.IsEmpty())
 	{
-		TouchedNode->OnNodeChangedDelegate.Broadcast(const_cast<UPCGNode*>(TouchedNode), EPCGChangeType::Cosmetic);
+		TArray<TWeakObjectPtr<const UPCGNode>> NodeWeakPtrs;
+		NodeWeakPtrs.Reserve(TouchedNodes.Num());
+		Algo::Transform(TouchedNodes, NodeWeakPtrs, [](const UPCGNode* Node) { return Node; });
+
+		ExecuteOnGameThread(UE_SOURCE_LOCATION, [NodeWeakPtrs]()
+		{
+			for (const TWeakObjectPtr<const UPCGNode>& NodeWeakPtr : NodeWeakPtrs)
+			{
+				if (const UPCGNode* Node = NodeWeakPtr.Get())
+				{
+					Node->OnNodeChangedDelegate.Broadcast(const_cast<UPCGNode*>(Node), EPCGChangeType::Cosmetic);
+				}
+			}
+		});
 	}
 }
 
