[PCG] Fixed potential ensure with recursive subgraphs

#rnx

diff --git a/Engine/Plugins/PCG/Source/PCGEditor/Private/PCGEditorGraphSchema.cpp b/Engine/Plugins/PCG/Source/PCGEditor/Private/PCGEditorGraphSchema.cpp
--- a/Engine/Plugins/PCG/Source/PCGEditor/Private/PCGEditorGraphSchema.cpp
+++ b/Engine/Plugins/PCG/Source/PCGEditor/Private/PCGEditorGraphSchema.cpp
@@ -226,7 +226,7 @@ bool UPCGEditorGraphSchema::TryCreateConnectionInternal(UEdGraphPin* InA, UEdGra
 	PCGGraph->PrimeGraphCompilationCache();
 
 	// Creates a connection via an intermediate conversion node.
-	auto ConnectViaIntermediate = [this, PCGGraph, NodeA, NodeB, A, B](UPCGNode* IntermediateNode)
+	auto ConnectViaIntermediate = [this, NodeA, NodeB](UPCGNode* IntermediateNode, UEdGraphPin* A, UEdGraphPin* B)
 	{
 		UEdGraph* Graph = NodeA->GetGraph();
 		check(Graph);
@@ -294,26 +294,50 @@ bool UPCGEditorGraphSchema::TryCreateConnectionInternal(UEdGraphPin* InA, UEdGra
 		UPCGSettings* NodeSettings = nullptr;
 		UPCGNode* ConversionPCGNode = PCGGraph->AddNodeOfType(UPCGCollapseSettings::StaticClass(), NodeSettings);
 
-		return ConnectViaIntermediate(ConversionPCGNode);
+		return ConnectViaIntermediate(ConversionPCGNode, A, B);
 	}
 	else if (Conversion == EPCGTypeConversion::Filter)
 	{
 		UPCGSettings* NodeSettings = nullptr;
 		UPCGNode* ConversionPCGNode = PCGGraph->AddNodeOfType(UPCGFilterByTypeSettings::StaticClass(), NodeSettings);
 
+		FName PinAName = A->GetFName();
+		FName PinBName = B->GetFName();
+
 		// Setup the output pin based on the conversion target type, before new node is finalized.
 		UPCGFilterByTypeSettings* Settings = CastChecked<UPCGFilterByTypeSettings>(NodeSettings);
 		Settings->TargetType = PCGPinB->Properties.AllowedTypes;
 		ConversionPCGNode->UpdateAfterSettingsChangeDuringCreation();
 
-		return ConnectViaIntermediate(ConversionPCGNode);
+		// Requery A & B because the update might have changed the pins.
+		UEdGraphPin* const* NewA = NodeA->GetAllPins().FindByPredicate([&PinAName](const UEdGraphPin* InPin)
+		{
+			return InPin->Direction == EGPD_Output && InPin->GetName() == PinAName;
+		});
+
+		if (NewA)
+		{
+			A = *NewA;
+		}
+
+		UEdGraphPin* const* NewB = NodeB->GetAllPins().FindByPredicate([&PinBName](const UEdGraphPin* InPin)
+		{
+			return InPin->Direction == EGPD_Input && InPin->GetName() == PinBName;
+		});
+
+		if (NewB)
+		{
+			B = *NewB;
+		}
+
+		return ConnectViaIntermediate(ConversionPCGNode, A, B);
 	}
 	else if (Conversion == EPCGTypeConversion::MakeConcrete)
 	{
 		UPCGSettings* NodeSettings = nullptr;
 		UPCGNode* ConversionPCGNode = PCGGraph->AddNodeOfType(UPCGMakeConcreteSettings::StaticClass(), NodeSettings);
 
-		return ConnectViaIntermediate(ConversionPCGNode);
+		return ConnectViaIntermediate(ConversionPCGNode, A, B);
 	}
 	else if (Conversion == EPCGTypeConversion::SplineToSurface)
 	{
@@ -323,7 +347,7 @@ bool UPCGEditorGraphSchema::TryCreateConnectionInternal(UEdGraphPin* InA, UEdGra
 		UPCGCreateSurfaceFromSplineSettings* Settings = CastChecked<UPCGCreateSurfaceFromSplineSettings>(NodeSettings);
 		Settings->bShouldDrawNodeCompact = true;
 
-		return ConnectViaIntermediate(ConversionPCGNode);
+		return ConnectViaIntermediate(ConversionPCGNode, A, B);
 	}
 	else
 	{
