[PCG] Fixed determinism test that could crash if not all the data was comparable.

#rnx

diff --git a/Engine/Plugins/PCG/Source/PCG/Private/Tests/Determinism/PCGDeterminismTestsCommon.cpp b/Engine/Plugins/PCG/Source/PCG/Private/Tests/Determinism/PCGDeterminismTestsCommon.cpp
--- a/Engine/Plugins/PCG/Source/PCG/Private/Tests/Determinism/PCGDeterminismTestsCommon.cpp
+++ b/Engine/Plugins/PCG/Source/PCG/Private/Tests/Determinism/PCGDeterminismTestsCommon.cpp
@@ -248,8 +248,7 @@ namespace PCGDeterminismTests
 			// Check for consistent outputs
 			if (HighestDeterminismLevel == EDeterminismLevel::OrderConsistent)
 			{
-				TArray<int32> IndexOffsets;
-				if (!DataCollectionsMatch(FirstOutput, SecondOutput, IndexOffsets) && !DataCollectionsAreConsistent(FirstOutput, SecondOutput, IndexOffsets.Num()))
+				if (!DataCollectionsAreConsistent(FirstOutput, SecondOutput, FirstOutput.TaggedData.Num()))
 				{
 					// Not consistent, downgrade to orthogonal
 					HighestDeterminismLevel = EDeterminismLevel::OrderOrthogonal;
@@ -867,11 +866,15 @@ namespace PCGDeterminismTests
 		MatchedIndices.Reserve(FirstCollectionComparableIndices.Num());
 		OutIndexOffsets.SetNumUninitialized(FirstCollectionComparableIndices.Num());
 
-		for (int32 I : FirstCollectionComparableIndices)
+		for(int32 FirstIndex = 0; FirstIndex < FirstCollectionComparableIndices.Num(); ++FirstIndex)
 		{
+			int32 I = FirstCollectionComparableIndices[FirstIndex];
+
 			bool bFound = false;
-			for (int32 J : SecondCollectionComparableIndices)
+			for(int32 SecondIndex = 0; SecondIndex < SecondCollectionComparableIndices.Num(); ++SecondIndex)
 			{
+				int32 J = SecondCollectionComparableIndices[SecondIndex];
+
 				if (MatchedIndices.Contains(J))
 				{
 					continue;
@@ -887,7 +890,7 @@ namespace PCGDeterminismTests
 					if (CompareFunction(FirstCollection.TaggedData[I].Data, SecondCollection.TaggedData[J].Data))
 					{
 						MatchedIndices.Emplace(J);
-						OutIndexOffsets[I] = J - I;
+						OutIndexOffsets[FirstIndex] = J - I;
 						bFound = true;
 						break;
 					}
