[PCG] Fix type warnings, and SM5 cook failure by deduplicating samplers.

Each texture data interface instance was adding 8 samplers. At the cost of more branches this deduplicates the samplers down to the 2 we use.

#rnx

diff --git a/Engine/Plugins/PCG/Shaders/Private/PCGDataCollectionDataInterface.ush b/Engine/Plugins/PCG/Shaders/Private/PCGDataCollectionDataInterface.ush
--- a/Engine/Plugins/PCG/Shaders/Private/PCGDataCollectionDataInterface.ush
+++ b/Engine/Plugins/PCG/Shaders/Private/PCGDataCollectionDataInterface.ush
@@ -84,25 +84,25 @@ int AddToElementCounterInternal_{DataInterfaceName}(uint InDataIndex, uint InEle
 
 bool GetThreadData_{DataInterfaceName}(uint InThreadIndex, out uint OutDataIndex, out uint OutElementIndex)
 {
-	int ElementIndex = InThreadIndex;
+	int ElementIndex = (int)InThreadIndex;
 
 	const uint NumData = GetNumData_{DataInterfaceName}();
 	for (uint DataIndex = 0; DataIndex < NumData; ++DataIndex)
 	{
-		const uint ElemCount = GetNumElements_{DataInterfaceName}(DataIndex);
+		const int ElemCount = (int)GetNumElements_{DataInterfaceName}(DataIndex);
 
 		if (ElementIndex < ElemCount)
 		{
 			OutDataIndex = DataIndex;
-			OutElementIndex = ElementIndex;
+			OutElementIndex = (uint)ElementIndex;
 			return true;
 		}
 
 		ElementIndex -= ElemCount;
 	}
 
-	OutDataIndex = -1;
-	OutElementIndex = -1;
+	OutDataIndex = (uint)-1;
+	OutElementIndex = (uint)-1;
 	return false;
 }
 
@@ -261,7 +261,7 @@ int GetStringKey_{DataInterfaceName}(uint DataIndex, uint ElementIndex, int Attr
 uint2 GetName_{DataInterfaceName}(uint DataIndex, uint ElementIndex, int AttributeId)
 {
 	// Names are represented as ints.
-	return uint2((uint)GetInt_{DataInterfaceName}(DataIndex, ElementIndex, AttributeId), 0);
+	return uint2((uint)GetInt_{DataInterfaceName}(DataIndex, ElementIndex, AttributeId), 0u);
 }
 
 uint GetElementCountMultiplier_{DataInterfaceName}()
diff --git a/Engine/Plugins/PCG/Shaders/Private/PCGTextureDataInterface.ush b/Engine/Plugins/PCG/Shaders/Private/PCGTextureDataInterface.ush
--- a/Engine/Plugins/PCG/Shaders/Private/PCGTextureDataInterface.ush
+++ b/Engine/Plugins/PCG/Shaders/Private/PCGTextureDataInterface.ush
@@ -18,35 +18,30 @@ Texture2DArray {DataInterfaceName}_TextureArray_5;
 Texture2DArray {DataInterfaceName}_TextureArray_6;
 Texture2DArray {DataInterfaceName}_TextureArray_7;
 
-// TODO: We only need 2 samplers + an indirection array, de-duplicate these.
-SamplerState {DataInterfaceName}_Sampler_0;
-SamplerState {DataInterfaceName}_Sampler_1;
-SamplerState {DataInterfaceName}_Sampler_2;
-SamplerState {DataInterfaceName}_Sampler_3;
-SamplerState {DataInterfaceName}_Sampler_4;
-SamplerState {DataInterfaceName}_Sampler_5;
-SamplerState {DataInterfaceName}_Sampler_6;
-SamplerState {DataInterfaceName}_Sampler_7;
+SamplerState {DataInterfaceName}_Sampler_Linear;
+SamplerState {DataInterfaceName}_Sampler_Point;
 
 float4 {DataInterfaceName}_TextureBounds[MAX_NUM_BINDINGS];
 
 uint {DataInterfaceName}_NumTextureInfos;
-StructuredBuffer<uint3> {DataInterfaceName}_TextureInfos;
+StructuredBuffer<uint4> {DataInterfaceName}_TextureInfos;
 
-void GetTextureInfoInternal_{DataInterfaceName}(int InDataIndex, out uint OutBindingIndex, out uint OutSliceIndex, out uint OutIsTextureArray)
+void GetTextureInfoInternal_{DataInterfaceName}(uint InDataIndex, out uint OutBindingIndex, out uint OutSliceIndex, out uint OutIsTextureArray, out uint OutPointSample)
 {
-	if (InDataIndex >= 0 && InDataIndex < {DataInterfaceName}_NumTextureInfos)
+	if (InDataIndex >= 0u && InDataIndex < {DataInterfaceName}_NumTextureInfos)
 	{
-		const uint3 TextureInfo = {DataInterfaceName}_TextureInfos[InDataIndex];
+		const uint4 TextureInfo = {DataInterfaceName}_TextureInfos[InDataIndex];
 		OutBindingIndex = TextureInfo.x;
 		OutSliceIndex = TextureInfo.y;
 		OutIsTextureArray = TextureInfo.z;
+		OutPointSample = TextureInfo.w;
 	}
 	else
 	{
-		OutBindingIndex = 0;
-		OutSliceIndex = 0;
-		OutIsTextureArray = 0;
+		OutBindingIndex = 0u;
+		OutSliceIndex = 0u;
+		OutIsTextureArray = 0u;
+		OutPointSample = 0u;
 	}
 }
 
@@ -69,20 +64,35 @@ float2 GetTexCoords_{DataInterfaceName}(float2 WorldPos, float2 Min, float2 Max)
 	{ \
 		if (IsTextureArray == 0) \
 		{ \
-			return {DataInterfaceName}_Texture_##Index.SampleLevel({DataInterfaceName}_Sampler_##Index, InTextureUV, 0); \
+			if (PointSample == 0) \
+			{ \
+				return {DataInterfaceName}_Texture_##Index.SampleLevel({DataInterfaceName}_Sampler_Linear, InTextureUV, 0); \
+			} \
+			else \
+			{ \
+				return {DataInterfaceName}_Texture_##Index.SampleLevel({DataInterfaceName}_Sampler_Point, InTextureUV, 0); \
+			} \
 		} \
 		else \
 		{ \
-			return {DataInterfaceName}_TextureArray_##Index.SampleLevel({DataInterfaceName}_Sampler_##Index, float3(InTextureUV, SliceIndex), 0); \
+			if (PointSample == 0) \
+			{ \
+				return {DataInterfaceName}_TextureArray_##Index.SampleLevel({DataInterfaceName}_Sampler_Linear, float3(InTextureUV, SliceIndex), 0); \
+			} \
+			else \
+			{ \
+				return {DataInterfaceName}_TextureArray_##Index.SampleLevel({DataInterfaceName}_Sampler_Point, float3(InTextureUV, SliceIndex), 0); \
+			} \
 		} \
 	}
 
-float4 Sample_{DataInterfaceName}(int InDataIndex, float2 InTextureUV)
+float4 Sample_{DataInterfaceName}(uint InDataIndex, float2 InTextureUV)
 {
 	uint BindingIndex = 0;
 	uint SliceIndex = 0;
 	uint IsTextureArray = 0;
-	GetTextureInfoInternal_{DataInterfaceName}(InDataIndex, BindingIndex, SliceIndex, IsTextureArray);
+	uint PointSample = 0;
+	GetTextureInfoInternal_{DataInterfaceName}(InDataIndex, BindingIndex, SliceIndex, IsTextureArray, PointSample);
 
 #if ENABLE_MULTIPLE_TEXTURE_OBJECTS
 	PCG_CONDITIONAL_TEXTURE_SAMPLE(0)
@@ -99,11 +109,26 @@ float4 Sample_{DataInterfaceName}(int InDataIndex, float2 InTextureUV)
 #else
 	if (IsTextureArray == 0)
 	{
-		return {DataInterfaceName}_Texture_0.SampleLevel({DataInterfaceName}_Sampler_0, InTextureUV, 0);
+		if (PointSample == 0)
+		{
+			return {DataInterfaceName}_Texture_0.SampleLevel({DataInterfaceName}_Sampler_Linear, InTextureUV, 0);
+		}
+		else
+		{
+			return {DataInterfaceName}_Texture_0.SampleLevel({DataInterfaceName}_Sampler_Point, InTextureUV, 0);
+		}
 	}
 	else
 	{
-		return {DataInterfaceName}_TextureArray_0.SampleLevel({DataInterfaceName}_Sampler_0, float3(InTextureUV, SliceIndex), 0);
+		if (PointSample == 0)
+		{
+			return {DataInterfaceName}_TextureArray_0.SampleLevel({DataInterfaceName}_Sampler_Linear, float3(InTextureUV, SliceIndex), 0);
+		}
+		else
+		{
+			return {DataInterfaceName}_TextureArray_0.SampleLevel({DataInterfaceName}_Sampler_Point, float3(InTextureUV, SliceIndex), 0);
+		}
+
 	}
 #endif
 }
@@ -116,20 +141,35 @@ float4 Sample_{DataInterfaceName}(int InDataIndex, float2 InTextureUV)
 		const float2 TextureUV = GetTexCoords_{DataInterfaceName}(WorldPos, {DataInterfaceName}_TextureBounds[Index].xy, {DataInterfaceName}_TextureBounds[Index].zw); \
 		if (IsTextureArray == 0) \
 		{ \
-			return {DataInterfaceName}_Texture_##Index.SampleLevel({DataInterfaceName}_Sampler_##Index, TextureUV, 0); \
+			if (PointSample == 0) \
+			{ \
+				return {DataInterfaceName}_Texture_##Index.SampleLevel({DataInterfaceName}_Sampler_Linear, TextureUV, 0); \
+			} \
+			else \
+			{ \
+				return {DataInterfaceName}_Texture_##Index.SampleLevel({DataInterfaceName}_Sampler_Point, TextureUV, 0); \
+			} \
 		} \
 		else \
 		{ \
-			return {DataInterfaceName}_TextureArray_##Index.SampleLevel({DataInterfaceName}_Sampler_##Index, float3(TextureUV, SliceIndex), 0); \
+			if (PointSample == 0) \
+			{ \
+				return {DataInterfaceName}_TextureArray_##Index.SampleLevel({DataInterfaceName}_Sampler_Linear, float3(TextureUV, SliceIndex), 0); \
+			} \
+			else \
+			{ \
+				return {DataInterfaceName}_TextureArray_##Index.SampleLevel({DataInterfaceName}_Sampler_Point, float3(TextureUV, SliceIndex), 0); \
+			} \
 		} \
 	}
 
-float4 SampleWorldPos_{DataInterfaceName}(int InDataIndex, float2 WorldPos)
+float4 SampleWorldPos_{DataInterfaceName}(uint InDataIndex, float2 WorldPos)
 {
 	uint BindingIndex = 0;
 	uint SliceIndex = 0;
 	uint IsTextureArray = 0;
-	GetTextureInfoInternal_{DataInterfaceName}(InDataIndex, BindingIndex, SliceIndex, IsTextureArray);
+	uint PointSample = 0;
+	GetTextureInfoInternal_{DataInterfaceName}(InDataIndex, BindingIndex, SliceIndex, IsTextureArray, PointSample);
 
 #if ENABLE_MULTIPLE_TEXTURE_OBJECTS
 	PCG_CONDITIONAL_TEXTURE_SAMPLE_WORLD_POS(0)
@@ -148,11 +188,25 @@ float4 SampleWorldPos_{DataInterfaceName}(int InDataIndex, float2 WorldPos)
 
 	if (IsTextureArray == 0)
 	{
-		return {DataInterfaceName}_Texture_0.SampleLevel({DataInterfaceName}_Sampler_0, TextureUV, 0);
+		if (PointSample == 0)
+		{
+			return {DataInterfaceName}_Texture_0.SampleLevel({DataInterfaceName}_Sampler_Linear, TextureUV, 0);
+		}
+		else
+		{
+			return {DataInterfaceName}_Texture_0.SampleLevel({DataInterfaceName}_Sampler_Point, TextureUV, 0);
+		}
 	}
 	else
 	{
-		return {DataInterfaceName}_TextureArray_0.SampleLevel({DataInterfaceName}_Sampler_0, float3(TextureUV, SliceIndex), 0);
+		if (PointSample == 0)
+		{
+			return {DataInterfaceName}_TextureArray_0.SampleLevel({DataInterfaceName}_Sampler_Linear, float3(TextureUV, SliceIndex), 0);
+		}
+		else
+		{
+			return {DataInterfaceName}_TextureArray_0.SampleLevel({DataInterfaceName}_Sampler_Point, float3(TextureUV, SliceIndex), 0);
+		}
 	}
 #endif
 }
diff --git a/Engine/Plugins/PCG/Source/PCG/Private/Compute/DataInterfaces/PCGCustomKernelDataInterface.cpp b/Engine/Plugins/PCG/Source/PCG/Private/Compute/DataInterfaces/PCGCustomKernelDataInterface.cpp
--- a/Engine/Plugins/PCG/Source/PCG/Private/Compute/DataInterfaces/PCGCustomKernelDataInterface.cpp
+++ b/Engine/Plugins/PCG/Source/PCG/Private/Compute/DataInterfaces/PCGCustomKernelDataInterface.cpp
@@ -20,7 +20,7 @@ void UPCGCustomKernelDataInterface::GetSupportedInputs(TArray<FShaderFunctionDef
 {
 	OutFunctions.AddDefaulted_GetRef()
 		.SetName(TEXT("GetNumThreads"))
-		.AddReturnType(FShaderValueType::Get(EShaderFundamentalType::Int, 3));
+		.AddReturnType(FShaderValueType::Get(EShaderFundamentalType::Uint, 3));
 
 	OutFunctions.AddDefaulted_GetRef()
 		.SetName(TEXT("GetThreadCountMultiplier"))
@@ -49,7 +49,7 @@ void UPCGCustomKernelDataInterface::GetSupportedInputs(TArray<FShaderFunctionDef
 }
 
 BEGIN_SHADER_PARAMETER_STRUCT(FPCGCustomKernelDataInterfaceParameters, )
-	SHADER_PARAMETER(FIntVector3, NumThreads)
+	SHADER_PARAMETER(FUintVector3, NumThreads)
 	SHADER_PARAMETER(uint32, ThreadCountMultiplier)
 	SHADER_PARAMETER(uint32, Seed)
 	SHADER_PARAMETER(uint32, SeedSettings)
@@ -71,7 +71,7 @@ void UPCGCustomKernelDataInterface::GetHLSL(FString& OutHLSL, FString const& InD
 	};
 
 	OutHLSL += FString::Format(TEXT(
-		"int3 {DataInterfaceName}_NumThreads;\n"
+		"uint3 {DataInterfaceName}_NumThreads;\n"
 		"uint {DataInterfaceName}_ThreadCountMultiplier;\n"
 		"uint {DataInterfaceName}_Seed;\n"
 		"uint {DataInterfaceName}_SeedSettings;\n"
@@ -79,7 +79,7 @@ void UPCGCustomKernelDataInterface::GetHLSL(FString& OutHLSL, FString const& InD
 		"float3 {DataInterfaceName}_ComponentBoundsMin;\n"
 		"float3 {DataInterfaceName}_ComponentBoundsMax;\n"
 		"\n"
-		"int3 GetNumThreads_{DataInterfaceName}()\n{\n\treturn {DataInterfaceName}_NumThreads;\n}\n\n"
+		"uint3 GetNumThreads_{DataInterfaceName}()\n{\n\treturn {DataInterfaceName}_NumThreads;\n}\n\n"
 		"uint GetThreadCountMultiplier_{DataInterfaceName}()\n{\n\treturn {DataInterfaceName}_ThreadCountMultiplier;\n}\n\n"
 		"uint GetSeed_{DataInterfaceName}()\n{\n\treturn {DataInterfaceName}_Seed;\n}\n\n"
 		"uint GetSettingsSeed_{DataInterfaceName}()\n{\n\treturn {DataInterfaceName}_SeedSettings;\n}\n\n"
@@ -169,8 +169,8 @@ void FPCGCustomComputeKernelDataProviderProxy::GatherDispatchData(FDispatchData
 		// Thread count
 		// Note: If this ThreadCount is zero, the kernel will still execute one thread, but terminates early when comparing thread index against this value.
 		ensure(InDispatchData.bUnifiedDispatch || InDispatchData.NumInvocations == 1);
-		Parameters.NumThreads.X = ThreadCount;
-		Parameters.NumThreads.Y = Parameters.NumThreads.Z = 1;
+		Parameters.NumThreads.X = static_cast<uint32>(ThreadCount);
+		Parameters.NumThreads.Y = Parameters.NumThreads.Z = 1u;
 		Parameters.ThreadCountMultiplier = ThreadCountMultiplier;
 
 		// Seed for the node
diff --git a/Engine/Plugins/PCG/Source/PCG/Private/Compute/DataInterfaces/PCGTextureDataInterface.cpp b/Engine/Plugins/PCG/Source/PCG/Private/Compute/DataInterfaces/PCGTextureDataInterface.cpp
--- a/Engine/Plugins/PCG/Source/PCG/Private/Compute/DataInterfaces/PCGTextureDataInterface.cpp
+++ b/Engine/Plugins/PCG/Source/PCG/Private/Compute/DataInterfaces/PCGTextureDataInterface.cpp
@@ -35,8 +35,8 @@ FPCGTextureBindingInfo::FPCGTextureBindingInfo(const UPCGBaseTextureData* InText
 
 	ResourceType = InTextureData->GetTextureResourceType();
 	Texture = InTextureData->GetTextureRHI();
-	Filter = InTextureData->Filter;
 	TextureBounds = InTextureData->GetBounds();
+	bPointSample = InTextureData->Filter == EPCGTextureFilter::Point;
 
 	if (ResourceType == EPCGTextureResourceType::ExportedTexture)
 	{
@@ -55,7 +55,7 @@ bool FPCGTextureBindingInfo::operator==(const FPCGTextureBindingInfo& Other) con
 		&& Texture == Other.Texture
 		&& ExportedTexture == Other.ExportedTexture
 		&& bIsTextureArray == Other.bIsTextureArray
-		&& Filter == Other.Filter
+		&& bPointSample == bPointSample
 		&& TextureBounds == Other.TextureBounds;
 }
 
@@ -66,40 +66,41 @@ void UPCGTextureDataInterface::GetSupportedInputs(TArray<FShaderFunctionDefiniti
 	OutFunctions.AddDefaulted_GetRef()
 		.SetName(TEXT("GetTexCoords"))
 		.AddReturnType(EShaderFundamentalType::Float, 2)
-		.AddParam(EShaderFundamentalType::Float, 2)
-		.AddParam(EShaderFundamentalType::Float, 2)
-		.AddParam(EShaderFundamentalType::Float, 2);
+		.AddParam(EShaderFundamentalType::Float, 2) // WorldPos
+		.AddParam(EShaderFundamentalType::Float, 2) // Min
+		.AddParam(EShaderFundamentalType::Float, 2); // Max
 
 	OutFunctions.AddDefaulted_GetRef()
 		.SetName(TEXT("Sample"))
 		.AddReturnType(EShaderFundamentalType::Float, 4)
-		.AddParam(EShaderFundamentalType::Uint)
-		.AddParam(EShaderFundamentalType::Float, 2);
+		.AddParam(EShaderFundamentalType::Uint) // InDataIndex
+		.AddParam(EShaderFundamentalType::Float, 2); // InTextureUV
 
 	OutFunctions.AddDefaulted_GetRef()
 		.SetName(TEXT("Sample"))
-		.AddReturnType(EShaderFundamentalType::Float, 4)
-		.AddParam(EShaderFundamentalType::Float, 2);
+		.AddReturnType(EShaderFundamentalType::Float, 4) // InDataIndex
+		.AddParam(EShaderFundamentalType::Float, 2); // InTextureUV
 
 	OutFunctions.AddDefaulted_GetRef()
 		.SetName(TEXT("SampleWorldPos"))
 		.AddReturnType(EShaderFundamentalType::Float, 4)
-		.AddParam(EShaderFundamentalType::Uint)
-		.AddParam(EShaderFundamentalType::Float, 2);
+		.AddParam(EShaderFundamentalType::Uint) // InDataIndex
+		.AddParam(EShaderFundamentalType::Float, 2); // WorldPos
 
 	OutFunctions.AddDefaulted_GetRef()
 		.SetName(TEXT("SampleWorldPos"))
 		.AddReturnType(EShaderFundamentalType::Float, 4)
-		.AddParam(EShaderFundamentalType::Float, 2);
+		.AddParam(EShaderFundamentalType::Float, 2); // WorldPos
 }
 
 BEGIN_SHADER_PARAMETER_STRUCT(FPCGTextureDataInterfaceParameters,)
 	SHADER_PARAMETER_RDG_TEXTURE_SRV_ARRAY(Texture2D, Texture, [PCGTextureDataInterfaceConstants::MAX_NUM_BINDINGS])
 	SHADER_PARAMETER_RDG_TEXTURE_SRV_ARRAY(Texture2DArray, TextureArray, [PCGTextureDataInterfaceConstants::MAX_NUM_BINDINGS])
-	SHADER_PARAMETER_SAMPLER_ARRAY(SamplerState, Sampler, [PCGTextureDataInterfaceConstants::MAX_NUM_BINDINGS])
+	SHADER_PARAMETER_SAMPLER(SamplerState, Sampler_Linear)
+	SHADER_PARAMETER_SAMPLER(SamplerState, Sampler_Point)
 	SHADER_PARAMETER_ARRAY(FVector4f, TextureBounds, [PCGTextureDataInterfaceConstants::MAX_NUM_BINDINGS])
 	SHADER_PARAMETER(uint32, NumTextureInfos)
-	SHADER_PARAMETER_RDG_BUFFER_SRV(StructuredBuffer<FUintVector3>, TextureInfos)
+	SHADER_PARAMETER_RDG_BUFFER_SRV(StructuredBuffer<FUintVector4>, TextureInfos)
 END_SHADER_PARAMETER_STRUCT()
 
 void UPCGTextureDataInterface::GetShaderParameters(TCHAR const* UID, FShaderParametersMetadataBuilder& InOutBuilder, FShaderParametersMetadataAllocations& InOutAllocations) const
@@ -247,6 +248,9 @@ void FPCGTextureDataProviderProxy::GatherDispatchData(FDispatchData const& InDis
 	{
 		FParameters& Parameters = ParameterArray[InvocationIndex];
 
+		Parameters.Sampler_Linear = TStaticSamplerState<SF_Bilinear, AM_Clamp, AM_Clamp, AM_Clamp>::GetRHI();
+		Parameters.Sampler_Point = TStaticSamplerState<SF_Point, AM_Clamp, AM_Clamp, AM_Clamp>::GetRHI();
+
 		for (int BindingIndex = 0; BindingIndex < PCGTextureDataInterfaceConstants::MAX_NUM_BINDINGS; ++BindingIndex)
 		{
 			Parameters.Texture[BindingIndex] = TextureSRV[BindingIndex];
@@ -259,23 +263,10 @@ void FPCGTextureDataProviderProxy::GatherDispatchData(FDispatchData const& InDis
 					BindingInfos[BindingIndex].TextureBounds.Min.Y,
 					BindingInfos[BindingIndex].TextureBounds.Max.X,
 					BindingInfos[BindingIndex].TextureBounds.Max.Y);
-
-				switch (BindingInfos[BindingIndex].Filter)
-				{
-				case EPCGTextureFilter::Point:
-					Parameters.Sampler[BindingIndex] = TStaticSamplerState<SF_Point, AM_Clamp, AM_Clamp, AM_Clamp>::GetRHI();
-					break;
-				default:
-					ensureMsgf(false, TEXT("Unrecognized texture filter setting, falling back to bilinear."));
-				case EPCGTextureFilter::Bilinear: // Falls through
-					Parameters.Sampler[BindingIndex] = TStaticSamplerState<SF_Bilinear, AM_Clamp, AM_Clamp, AM_Clamp>::GetRHI();
-					break;
-				}
 			}
 			else
 			{
 				Parameters.TextureBounds[BindingIndex] = FVector4f::Zero();
-				Parameters.Sampler[BindingIndex] = TStaticSamplerState<SF_Bilinear, AM_Clamp, AM_Clamp, AM_Clamp>::GetRHI();
 			}
 
 			Parameters.NumTextureInfos = TextureInfos.Num();
@@ -346,28 +337,27 @@ void FPCGTextureDataProviderProxy::PackTextureInfos(FRDGBuilder& GraphBuilder)
 
 	if (!TextureInfos.IsEmpty())
 	{
-		const FRDGBufferDesc TextureInfosDesc = FRDGBufferDesc::CreateStructuredDesc(sizeof(FUintVector3), TextureInfos.Num());
+		const FRDGBufferDesc TextureInfosDesc = FRDGBufferDesc::CreateStructuredDesc(sizeof(FUintVector4), TextureInfos.Num());
 
 		FRDGBufferRef TextureInfosBuffer = GraphBuilder.CreateBuffer(TextureInfosDesc, TEXT("PCGTextureInfosBuffer"));
 		TextureInfosBufferSRV = GraphBuilder.CreateSRV(TextureInfosBuffer);
 
-		TArray<FUintVector3> TextureInfosPackedData;
+		TArray<FUintVector4> TextureInfosPackedData;
 		TextureInfosPackedData.SetNumUninitialized(TextureInfos.Num());
 
 		for (int TextureInfoIndex = 0; TextureInfoIndex < TextureInfos.Num(); ++TextureInfoIndex)
 		{
 			const FPCGTextureInfo& TextureInfo = TextureInfos[TextureInfoIndex];
 			const FPCGTextureBindingInfo& BindingInfo = BindingInfos[TextureInfo.BindingIndex];
-			TextureInfosPackedData[TextureInfoIndex] = FUintVector3(TextureInfo.BindingIndex, TextureInfo.SliceIndex, BindingInfo.bIsTextureArray);
+			TextureInfosPackedData[TextureInfoIndex] = FUintVector4(TextureInfo.BindingIndex, TextureInfo.SliceIndex, BindingInfo.bIsTextureArray, BindingInfo.bPointSample);
 		}
 
 		GraphBuilder.QueueBufferUpload(TextureInfosBuffer, MakeArrayView(TextureInfosPackedData));
 	}
 	else
 	{
-		TextureInfosBufferSRV = GraphBuilder.CreateSRV(FRDGBufferSRVDesc(GSystemTextures.GetDefaultStructuredBuffer(GraphBuilder, sizeof(FUintVector3))));
+		TextureInfosBufferSRV = GraphBuilder.CreateSRV(FRDGBufferSRVDesc(GSystemTextures.GetDefaultStructuredBuffer(GraphBuilder, sizeof(FUintVector4))));
 	}
 }
 
 #undef LOCTEXT_NAMESPACE
-
diff --git a/Engine/Plugins/PCG/Source/PCG/Private/Compute/DataInterfaces/PCGTextureDataInterface.h b/Engine/Plugins/PCG/Source/PCG/Private/Compute/DataInterfaces/PCGTextureDataInterface.h
--- a/Engine/Plugins/PCG/Source/PCG/Private/Compute/DataInterfaces/PCGTextureDataInterface.h
+++ b/Engine/Plugins/PCG/Source/PCG/Private/Compute/DataInterfaces/PCGTextureDataInterface.h
@@ -33,9 +33,9 @@ struct FPCGTextureBindingInfo
 	FTextureRHIRef Texture = nullptr;
 	TRefCountPtr<IPooledRenderTarget> ExportedTexture = nullptr;
 	bool bIsTextureArray = false;
+	bool bPointSample = false;
 
-	// @todo_pcg: Filter and TextureBounds could potentially be independent of the texture binding, and queried via lookup from the texture info buffer instead.
-	EPCGTextureFilter Filter = EPCGTextureFilter::Bilinear;
+	// @todo_pcg: TextureBounds could potentially be independent of the texture binding, and queried via lookup from the texture info buffer instead.
 	FBox TextureBounds = FBox(EForceInit::ForceInit);
 };
 
