[PCG] Fix SA warning in texture data interface

#rnx

diff --git a/Engine/Plugins/PCG/Source/PCG/Private/Compute/DataInterfaces/PCGTextureDataInterface.cpp b/Engine/Plugins/PCG/Source/PCG/Private/Compute/DataInterfaces/PCGTextureDataInterface.cpp
--- a/Engine/Plugins/PCG/Source/PCG/Private/Compute/DataInterfaces/PCGTextureDataInterface.cpp
+++ b/Engine/Plugins/PCG/Source/PCG/Private/Compute/DataInterfaces/PCGTextureDataInterface.cpp
@@ -55,7 +55,7 @@ bool FPCGTextureBindingInfo::operator==(const FPCGTextureBindingInfo& Other) con
 		&& Texture == Other.Texture
 		&& ExportedTexture == Other.ExportedTexture
 		&& bIsTextureArray == Other.bIsTextureArray
-		&& bPointSample == bPointSample
+		&& bPointSample == Other.bPointSample
 		&& TextureBounds == Other.TextureBounds;
 }
 
