[PCG] Fix a bug which prevented comment node renames from dirtying the graph or working with Undo/Redo

#rnx

diff --git a/Engine/Plugins/PCG/Source/PCGEditor/Private/PCGEditor.cpp b/Engine/Plugins/PCG/Source/PCGEditor/Private/PCGEditor.cpp
--- a/Engine/Plugins/PCG/Source/PCGEditor/Private/PCGEditor.cpp
+++ b/Engine/Plugins/PCG/Source/PCGEditor/Private/PCGEditor.cpp
@@ -2872,6 +2872,8 @@ void FPCGEditor::OnSelectedNodesChanged(const TSet<UObject*>& NewSelection)
 
 void FPCGEditor::OnNodeTitleCommitted(const FText& NewText, ETextCommit::Type CommitInfo, UEdGraphNode* NodeBeingChanged)
 {
+	check(PCGGraphBeingEdited);
+
 	if (NodeBeingChanged)
 	{
 		if (CommitInfo == ETextCommit::OnEnter || CommitInfo == ETextCommit::OnUserMovedFocus)
@@ -2880,14 +2882,17 @@ void FPCGEditor::OnNodeTitleCommitted(const FText& NewText, ETextCommit::Type Co
 			if (OnValidateNodeTitle(NewText, NodeBeingChanged, ErrorText))
 			{
 				const FScopedTransaction Transaction(*FPCGEditorCommon::ContextIdentifier, LOCTEXT("PCGEditorRenameNode", "PCG Editor: Rename Node"), nullptr);
-				NodeBeingChanged->OnRenameNode(NewText.ToString());
 
 				// Implementation detail: In UPCGEditorGraphNode we only set the title under certain conditions, so it calls Modify() itself.
 				// However, UEdGraphNode does not call Modify() on its own, so we should still call it in this case.
 				if (!NodeBeingChanged->IsA<UPCGEditorGraphNode>())
 				{
 					NodeBeingChanged->Modify();
+					// Modify the graph as well, as non-pcg editor nodes (like the comment nodes) are serialized in UPCGGraph::ExtraEditorNodes.
+					PCGGraphBeingEdited->Modify();
 				}
+
+				NodeBeingChanged->OnRenameNode(NewText.ToString());
 			}
 			else
 			{
