[PCG] Added a general purpose wait node

diff --git a/Engine/Plugins/PCG/Source/PCG/Private/Elements/ControlFlow/PCGWait.cpp b/Engine/Plugins/PCG/Source/PCG/Private/Elements/ControlFlow/PCGWait.cpp
new file mode 100644
--- /dev/null
+++ b/Engine/Plugins/PCG/Source/PCG/Private/Elements/ControlFlow/PCGWait.cpp
@@ -0,0 +1,72 @@
+// Copyright Epic Games, Inc. All Rights Reserved.
+
+#include "Elements/ControlFlow/PCGWait.h"
+
+#include "PCGComponent.h"
+#include "PCGSubsystem.h"
+
+FPCGElementPtr UPCGWaitSettings::CreateElement() const
+{
+	return MakeShared<FPCGWaitElement>();
+}
+
+TArray<FPCGPinProperties> UPCGWaitSettings::InputPinProperties() const
+{
+	return {};
+}
+
+TArray<FPCGPinProperties> UPCGWaitSettings::OutputPinProperties() const
+{
+	TArray<FPCGPinProperties> PinProperties;
+	FPCGPinProperties& DependencyPin = PinProperties.Emplace_GetRef(PCGPinConstants::DefaultExecutionDependencyLabel, EPCGDataType::Any, /*bInAllowMultipleConnections=*/true, /*bAllowMultipleData=*/true);
+	DependencyPin.Usage = EPCGPinUsage::DependencyOnly;
+
+	return PinProperties;
+}
+
+bool FPCGWaitElement::ExecuteInternal(FPCGContext* InContext) const
+{
+	TRACE_CPUPROFILER_EVENT_SCOPE(FPCGWaitElement::Execute);
+	check(InContext);
+	FPCGWaitContext* Context = static_cast<FPCGWaitContext*>(InContext);
+	const UPCGWaitSettings* Settings = Context->GetInputSettings<UPCGWaitSettings>();
+	check(Settings);
+
+	const double CurrentTime = FPlatformTime::Seconds();
+	const int64 CurrentEngineFrame = GFrameCounter;
+	const int64 CurrentRenderFrame = GFrameCounterRenderThread;
+
+	if (!Context->bQueriedTimers)
+	{
+		Context->StartTime = CurrentTime;
+		Context->StartEngineFrame = CurrentEngineFrame;
+		Context->StartRenderFrame = CurrentRenderFrame;
+		Context->bQueriedTimers = true;
+	}
+
+	const bool bTimeDone = (CurrentTime - Context->StartTime) >= Settings->WaitTimeInSeconds;
+	const bool bEngineFramesDone = (CurrentEngineFrame - Context->StartEngineFrame) >= (uint64)Settings->WaitTimeInEngineFrames;
+	const bool bRenderFramesDone = (CurrentRenderFrame - Context->StartRenderFrame) >= (uint64)Settings->WaitTimeInRenderFrames;
+
+	if((Settings->bEndWaitWhenAllConditionsAreMet && bTimeDone && bEngineFramesDone && bRenderFramesDone) ||
+		(!Settings->bEndWaitWhenAllConditionsAreMet && (bTimeDone || bEngineFramesDone || bRenderFramesDone)))
+	{
+		Context->OutputData = Context->InputData;
+
+		return true;
+	}
+	else
+	{
+		Context->bIsPaused = true;
+		Context->SourceComponent->GetSubsystem()->RegisterBeginTickAction([ContextHandle = Context->GetOrCreateHandle()]()
+		{
+			FPCGContext::FSharedContext<FPCGWaitContext> SharedContext(ContextHandle);
+			if (FPCGWaitContext* ContextPtr = SharedContext.Get())
+			{
+				ContextPtr->bIsPaused = false;
+			}
+		});
+
+		return false;
+	}
+}
\ No newline at end of file
diff --git a/Engine/Plugins/PCG/Source/PCG/Public/Elements/ControlFlow/PCGWait.h b/Engine/Plugins/PCG/Source/PCG/Public/Elements/ControlFlow/PCGWait.h
new file mode 100644
--- /dev/null
+++ b/Engine/Plugins/PCG/Source/PCG/Public/Elements/ControlFlow/PCGWait.h
@@ -0,0 +1,65 @@
+// Copyright Epic Games, Inc. All Rights Reserved.
+
+#pragma once
+
+#include "PCGContext.h"
+#include "PCGElement.h"
+#include "PCGSettings.h"
+
+#include "PCGWait.generated.h"
+
+/**
+* Simple node to wait some amount of time and/or some amount of frames. Simply forwards inputs.
+*/
+UCLASS(BlueprintType, ClassGroup=(Procedural))
+class UPCGWaitSettings : public UPCGSettings
+{
+	GENERATED_BODY()
+
+public:
+	//~Begin UPCGSettings interface
+#if WITH_EDITOR
+	virtual FName GetDefaultNodeName() const override { return FName(TEXT("Wait")); }
+	virtual FText GetDefaultNodeTitle() const override { return NSLOCTEXT("PCGWaitSettings", "NodeTitle", "Wait"); }
+	virtual FText GetNodeTooltipText() const override { return NSLOCTEXT("PCGWaitSettings", "NodeTooltip", "Waits some time and/or frames. Not a node that should be used in production except in very special cases."); }
+	virtual EPCGSettingsType GetType() const override { return EPCGSettingsType::Generic; }
+#endif
+
+protected:
+	virtual TArray<FPCGPinProperties> InputPinProperties() const override;
+	virtual TArray<FPCGPinProperties> OutputPinProperties() const override;
+	virtual FPCGElementPtr CreateElement() const override;
+	//~End UPCGSettings interface
+
+public:
+	UPROPERTY(BlueprintReadWrite, EditAnywhere, Category = Settings, meta = (ClampMin="0.0", PCG_Overridable))
+	double WaitTimeInSeconds = 1.0;
+
+	UPROPERTY(BlueprintReadWrite, EditAnywhere, Category = Settings, meta = (ClampMin="0", PCG_Overridable))
+	int64 WaitTimeInEngineFrames = 0;
+
+	UPROPERTY(BlueprintReadWrite, EditAnywhere, Category = Settings, meta = (ClampMin = "0", PCG_Overridable))
+	int64 WaitTimeInRenderFrames = 0;
+
+	// Controls whether all conditions are needed or any condition is sufficient to end the wait.
+	UPROPERTY(BlueprintReadWrite, EditAnywhere, Category = Settings, meta = (PCG_Overridable))
+	bool bEndWaitWhenAllConditionsAreMet = true;
+};
+
+struct FPCGWaitContext : public FPCGContext
+{
+	double StartTime = -1.0;
+	uint64 StartEngineFrame = 0;
+	uint64 StartRenderFrame = 0;
+	bool bQueriedTimers = false;
+};
+
+class FPCGWaitElement : public IPCGElementWithCustomContext<FPCGWaitContext>
+{
+public:
+	virtual bool CanExecuteOnlyOnMainThread(FPCGContext* Context) const override { return true; }
+	virtual bool IsCacheable(const UPCGSettings* InSettings) const override { return false; }
+
+protected:
+	virtual bool ExecuteInternal(FPCGContext* Context) const override;
+};
\ No newline at end of file
