[PCG] Fixed collapse element test

#rnx

diff --git a/Engine/Plugins/PCG/Source/PCG/Private/Tests/PCGTestsCommon.cpp b/Engine/Plugins/PCG/Source/PCG/Private/Tests/PCGTestsCommon.cpp
--- a/Engine/Plugins/PCG/Source/PCG/Private/Tests/PCGTestsCommon.cpp
+++ b/Engine/Plugins/PCG/Source/PCG/Private/Tests/PCGTestsCommon.cpp
@@ -294,6 +294,8 @@ bool FPCGTestBaseClass::SmokeTestAnyValidInput(UPCGSettings* InSettings, TFuncti
 	{
 		// Prepare input
 		FPCGDataCollection InputData;
+		InputData.TaggedData.Emplace_GetRef().Data = InSettings;
+
 		for (int32 PinIndex = 0; PinIndex < InputIndices.Num(); ++PinIndex)
 		{
 			InputData.TaggedData.Append(InputsPerProperties[PinIndex][InputIndices[PinIndex]].TaggedData);
