[PCG] Fix for texture sampler not working on non-square textures.

#11996

#rnx

diff --git a/Engine/Plugins/PCG/Source/PCG/Private/Data/PCGTextureData.cpp b/Engine/Plugins/PCG/Source/PCG/Private/Data/PCGTextureData.cpp
--- a/Engine/Plugins/PCG/Source/PCG/Private/Data/PCGTextureData.cpp
+++ b/Engine/Plugins/PCG/Source/PCG/Private/Data/PCGTextureData.cpp
@@ -320,7 +320,7 @@ const UPCGPointData* UPCGBaseTextureData::CreatePointData(FPCGContext* Context)
 	FPCGAsync::AsyncPointProcessing(Context, PointCount, Points, [this, XCount, YCount, &Surface](int32 Index, FPCGPoint& OutPoint)
 	{
 		const int X = (Index % XCount);
-		const int Y = (Index / YCount);
+		const int Y = (Index / XCount);
 
 		// TODO: we should have a 0.5 bias here
 		FVector2D LocalCoordinate((2.0 * X + 0.5) / XCount - 1.0, (2.0 * Y + 0.5) / YCount - 1.0);
