[PCG] Add more resilience around UserDefinedStruct property extraction.

* Property names in UserDefinedStruct can contains invalid characters like spaces
* This make the call to GetAuthoredNameForField not returning the same thing in editor or runtime
* In order to be more resilient, if we detect that the name we are looking for contains invalid characters, we also look for the sanitize version of the name.
* Also add the option, on for all new nodes, to output sanitized name as attribute name, to make sure our attributes are valid and the same in editor and runtime.

#rnx

diff --git a/Engine/Plugins/PCG/Source/PCG/Private/Elements/PCGGetActorProperty.cpp b/Engine/Plugins/PCG/Source/PCG/Private/Elements/PCGGetActorProperty.cpp
--- a/Engine/Plugins/PCG/Source/PCG/Private/Elements/PCGGetActorProperty.cpp
+++ b/Engine/Plugins/PCG/Source/PCG/Private/Elements/PCGGetActorProperty.cpp
@@ -14,6 +14,8 @@
 #include "Helpers/PCGHelpers.h"
 #include "Metadata/PCGAttributePropertySelector.h"
 
+#include "UObject/FortniteReleaseBranchCustomObjectVersion.h"
+
 #define LOCTEXT_NAMESPACE "PCGPropertyToParamDataElement"
 
 UPCGGetActorPropertySettings::UPCGGetActorPropertySettings(const FObjectInitializer& ObjectInitializer)
@@ -102,6 +104,12 @@ void UPCGGetActorPropertySettings::PostLoad()
 	{
 		ActorSelector.ActorSelectionClass = TSubclassOf<AActor>();
 	}
+
+	if (GetLinkerCustomVersion(FFortniteReleaseBranchCustomObjectVersion::GUID) < FFortniteReleaseBranchCustomObjectVersion::OptionSanitizeOutputAttributeNamesPCG)
+	{
+		// For all previous nodes, we'll force this option to false. for retro-compatibility
+		bSanitizeOutputAttributeName = false;
+	}
 #endif // WITH_EDITOR
 }
 
@@ -338,6 +346,7 @@ bool FPCGGetActorPropertyElement::ExecuteInternal(FPCGContext* InContext) const
 			const FName OutputAttributeName = Settings->OutputAttributeName.GetName();
 
 			PCGPropertyHelpers::FExtractorParameters Parameters(ObjectToInspect, ObjectToInspect->GetClass(), Settings->PropertyName.ToString(), OutputAttributeName, Settings->bForceObjectAndStructExtraction, /*bPropertyNeedsToBeVisible=*/true);
+			Parameters.bStrictSanitizeOutputAttributeNames = Settings->bSanitizeOutputAttributeName;
 
 			// Don't care for object traversed in non-editor build, since it is only useful for tracking.
 			TSet<FSoftObjectPath>* ObjectTraversedPtr = nullptr;
diff --git a/Engine/Plugins/PCG/Source/PCG/Private/Elements/PCGGetPropertyFromObjectPath.cpp b/Engine/Plugins/PCG/Source/PCG/Private/Elements/PCGGetPropertyFromObjectPath.cpp
--- a/Engine/Plugins/PCG/Source/PCG/Private/Elements/PCGGetPropertyFromObjectPath.cpp
+++ b/Engine/Plugins/PCG/Source/PCG/Private/Elements/PCGGetPropertyFromObjectPath.cpp
@@ -8,6 +8,8 @@
 #include "Helpers/PCGHelpers.h"
 #include "Helpers/PCGPropertyHelpers.h"
 
+#include "UObject/FortniteReleaseBranchCustomObjectVersion.h"
+
 #define LOCTEXT_NAMESPACE "PCGGetPropertyFromObjectPathElement"
 
 UPCGGetPropertyFromObjectPathSettings::UPCGGetPropertyFromObjectPathSettings()
@@ -84,6 +86,19 @@ bool UPCGGetPropertyFromObjectPathSettings::CanEditChange(const FProperty* InPro
 }
 #endif // WITH_EDITOR
 
+void UPCGGetPropertyFromObjectPathSettings::PostLoad()
+{
+	Super::PostLoad();
+
+#if WITH_EDITOR
+	if (GetLinkerCustomVersion(FFortniteReleaseBranchCustomObjectVersion::GUID) < FFortniteReleaseBranchCustomObjectVersion::OptionSanitizeOutputAttributeNamesPCG)
+	{
+		// For all previous nodes, we'll force this option to false. for retro-compatibility
+		bSanitizeOutputAttributeName = false;
+	}
+#endif // WITH_EDITOR
+}
+
 FString UPCGGetPropertyFromObjectPathSettings::GetAdditionalTitleInformation() const
 {
 	FString Path;
@@ -221,6 +236,7 @@ bool FPCGGetPropertyFromObjectPathElement::ExecuteInternal(FPCGContext* Context)
 		const FName OutputAttributeName = Settings->OutputAttributeName.GetName();
 
 		PCGPropertyHelpers::FExtractorParameters Parameters(Object, Object->GetClass(), Settings->PropertyName.ToString(), OutputAttributeName, Settings->bForceObjectAndStructExtraction, /*bPropertyNeedsToBeVisible=*/true);
+		Parameters.bStrictSanitizeOutputAttributeNames = Settings->bSanitizeOutputAttributeName;
 
 		if (UPCGParamData* ParamData = PCGPropertyHelpers::ExtractPropertyAsAttributeSet(Parameters, Context))
 		{
diff --git a/Engine/Plugins/PCG/Source/PCG/Private/Elements/PCGUserParameterGet.cpp b/Engine/Plugins/PCG/Source/PCG/Private/Elements/PCGUserParameterGet.cpp
--- a/Engine/Plugins/PCG/Source/PCG/Private/Elements/PCGUserParameterGet.cpp
+++ b/Engine/Plugins/PCG/Source/PCG/Private/Elements/PCGUserParameterGet.cpp
@@ -12,11 +12,13 @@
 #include "Elements/PCGCreateAttribute.h"
 #include "Helpers/PCGConversion.h"
 #include "Helpers/PCGDynamicTrackingHelpers.h"
+#include "Helpers/PCGHelpers.h"
 #include "Helpers/PCGPropertyHelpers.h"
 #include "Metadata/PCGAttributePropertySelector.h"
 
 #include "StructUtils/PropertyBag.h"
 #include "StructUtils/StructView.h"
+#include "UObject/FortniteReleaseBranchCustomObjectVersion.h"
 
 #define LOCTEXT_NAMESPACE "PCGUserParameterGetElement"
 
@@ -192,6 +194,19 @@ void UPCGUserParameterGetSettings::UpdatePropertyName(FName InNewName)
 	}
 }
 
+void UPCGUserParameterGetSettings::PostLoad()
+{
+	Super::PostLoad();
+	
+#if WITH_EDITOR
+	if (GetLinkerCustomVersion(FFortniteReleaseBranchCustomObjectVersion::GUID) < FFortniteReleaseBranchCustomObjectVersion::OptionSanitizeOutputAttributeNamesPCG)
+	{
+		// For all previous nodes, we'll force this option to false. for retro-compatibility
+		bSanitizeOutputAttributeName = false;
+	}
+#endif // WITH_EDITOR
+}
+
 FPCGElementPtr UPCGUserParameterGetSettings::CreateElement() const
 {
 	return MakeShared<FPCGUserParameterGetElement>();
@@ -276,6 +291,8 @@ bool UPCGUserParameterGetSettings::ConvertNode(const FPCGPreconfiguredInfo& Conv
 	return false;
 }
 
+//////////////////////////////////////////////
+
 TArray<FPCGPinProperties> UPCGGenericUserParameterGetSettings::OutputPinProperties() const
 {
 	TArray<FPCGPinProperties> PinProperties;
@@ -284,6 +301,19 @@ TArray<FPCGPinProperties> UPCGGenericUserParameterGetSettings::OutputPinProperti
 	return PinProperties;
 }
 
+void UPCGGenericUserParameterGetSettings::PostLoad()
+{
+	Super::PostLoad();
+
+#if WITH_EDITOR
+	if (GetLinkerCustomVersion(FFortniteReleaseBranchCustomObjectVersion::GUID) < FFortniteReleaseBranchCustomObjectVersion::OptionSanitizeOutputAttributeNamesPCG)
+	{
+		// For all previous nodes, we'll force this option to false. for retro-compatibility
+		bSanitizeOutputAttributeName = false;
+	}
+#endif // WITH_EDITOR
+}
+
 FString UPCGGenericUserParameterGetSettings::GetAdditionalTitleInformation() const
 {
 	return PropertyPath;
@@ -416,6 +446,8 @@ bool UPCGGenericUserParameterGetSettings::ConvertNode(const FPCGPreconfiguredInf
 	return false;
 }
 
+//////////////////////////////////////////////
+
 bool FPCGUserParameterGetElement::ExecuteInternal(FPCGContext* Context) const
 {
 	TRACE_CPUPROFILER_EVENT_SCOPE(FPCGUserParameterGetElement::Execute);
@@ -471,6 +503,8 @@ bool FPCGUserParameterGetElement::ExecuteInternal(FPCGContext* Context) const
 				Settings->PropertyName,
 				Settings->bForceObjectAndStructExtraction,
 				/*bPropertyNeedsToBeVisible=*/false);
+
+			ExtractorParameters.bStrictSanitizeOutputAttributeNames = Settings->bSanitizeOutputAttributeName;
 		}
 		else if (GenericSettings)
 		{
@@ -480,6 +514,8 @@ bool FPCGUserParameterGetElement::ExecuteInternal(FPCGContext* Context) const
 				GenericSettings->OutputAttributeName,
 				GenericSettings->bForceObjectAndStructExtraction,
 				/*bPropertyNeedsToBeVisible=*/false);
+
+			ExtractorParameters.bStrictSanitizeOutputAttributeNames = GenericSettings->bSanitizeOutputAttributeName;
 		}
 
 		// Don't care for object traversed in non-editor build, since it is only useful for tracking.
diff --git a/Engine/Plugins/PCG/Source/PCG/Private/Helpers/PCGPropertyHelpers.cpp b/Engine/Plugins/PCG/Source/PCG/Private/Helpers/PCGPropertyHelpers.cpp
--- a/Engine/Plugins/PCG/Source/PCG/Private/Helpers/PCGPropertyHelpers.cpp
+++ b/Engine/Plugins/PCG/Source/PCG/Private/Helpers/PCGPropertyHelpers.cpp
@@ -78,18 +78,12 @@ namespace PCGPropertyHelpers
 		check(CurrentClass);
 
 		const FProperty* Property = nullptr;
-		// Try to get the property. If it is coming from a user struct, we need to iterate on all properties because the property name is mangled
+		// Try to get the property. If it is coming from a user struct, we need to iterate on all properties because the property name is mangled.
+		// There is also a difference between runtime and editor name, as property name in editor can contain invalid characters (like spaces).
+		// So if the current name is invalid, also try with sanitized name.
 		if (const UUserDefinedStruct* UserDefinedStruct = Cast<UUserDefinedStruct>(CurrentClass))
 		{
-			for (TFieldIterator<const FProperty> PropIt(UserDefinedStruct, EFieldIterationFlags::IncludeSuper); PropIt; ++PropIt)
-			{
-				const FName PropertyName = *UserDefinedStruct->GetAuthoredNameForField(*PropIt);
-				if (PropertyName == CurrentName)
-				{
-					Property = *PropIt;
-					break;
-				}
-			}
+			Property = FindPropertyInUserDefinedStruct(UserDefinedStruct, CurrentName);
 		}
 		else
 		{
@@ -335,6 +329,8 @@ UPCGParamData* PCGPropertyHelpers::ExtractPropertyAsAttributeSet(const PCGProper
 			}
 		}
 
+		const FProperty* OriginalProperty = Property;
+
 		// If the property is an array/set, we will work on the underlying property, and extract each element as an entry in the param data
 		const FArrayProperty* ArrayProperty = CastField<FArrayProperty>(Property);
 		const FSetProperty* SetProperty = CastField<FSetProperty>(Property);
@@ -417,9 +413,27 @@ UPCGParamData* PCGPropertyHelpers::ExtractPropertyAsAttributeSet(const PCGProper
 		{
 			// For non struct/object, there is just a single property to extract with no shenanigans for address indirection.
 			const bool bIsSourceName = Parameters.OutputAttributeName == PCGMetadataAttributeConstants::SourceNameAttributeName || Parameters.OutputAttributeName == PCGMetadataAttributeConstants::SourceAttributeName;
-			
-			const FName AttributeName = (bIgnoreOutputAttributeName || bIsSourceName) ? Property->GetFName() : Parameters.OutputAttributeName;
-			ExtractableProperties.Emplace(AttributeName.ToString(), Property);
+			FString AttributeName;
+			if (bIgnoreOutputAttributeName || bIsSourceName)
+			{
+				// Make sure that the name is the authored name
+				if (const UStruct* StructOwner = Property->GetOwnerStruct())
+				{
+					AttributeName = StructOwner->GetAuthoredNameForField(OriginalProperty);
+				}
+
+				if (AttributeName.IsEmpty())
+				{
+					AttributeName = Property->GetName();
+				}
+			}
+			else
+			{
+				AttributeName = Parameters.OutputAttributeName.ToString();
+			}
+
+			ExtractableProperties.Emplace(std::move(AttributeName), Property);
+
 			// Identity
 			AddressFunc = [](const void* InAddress) { return InAddress; };
 		}
@@ -488,13 +502,22 @@ UPCGParamData* PCGPropertyHelpers::ExtractPropertyAsAttributeSet(const PCGProper
 
 			for (ExtractablePropertyTuple& ExtractableProperty : ExtractableProperties)
 			{
-				FString AttributeName = ExtractableProperty.Get<0>();
+				FString AttributeNameStr = ExtractableProperty.Get<0>();
+				FName AttributeName = NAME_None;
 				const FProperty* FinalProperty = ExtractableProperty.Get<1>();
 
 				// Make sure the Attribute name is sanitized, to prevent cases where property names have unsupported characters.
-				FPCGMetadataAttributeBase::SanitizeName(AttributeName);
+				if (Parameters.bStrictSanitizeOutputAttributeNames)
+				{
+					AttributeName = MakeObjectNameFromDisplayLabel(AttributeNameStr, NAME_None);
+				}
+				else
+				{
+					FPCGMetadataAttributeBase::SanitizeName(AttributeNameStr);
+					AttributeName = *AttributeNameStr;
+				}
 
-				if (!Metadata->SetAttributeFromDataProperty(FName(AttributeName), EntryKey, ContainerPtr, FinalProperty, /*bCreate=*/ true))
+				if (!Metadata->SetAttributeFromDataProperty(AttributeName, EntryKey, ContainerPtr, FinalProperty, /*bCreate=*/ true))
 				{
 					if (!bQuiet)
 					{
@@ -525,4 +548,32 @@ UPCGParamData* PCGPropertyHelpers::ExtractPropertyAsAttributeSet(const PCGProper
 	return bValidOperation ? ParamData : nullptr;
 }
 
+const FProperty* PCGPropertyHelpers::FindPropertyInUserDefinedStruct(const UUserDefinedStruct* InStruct, const FName InName)
+{
+	if (!InStruct)
+	{
+		return nullptr;
+	}
+
+	TArray<FName, TInlineAllocator<2>> NamesToLookFor = { InName };
+
+	if (!InName.IsValidXName())
+	{
+		NamesToLookFor.Add(MakeObjectNameFromDisplayLabel(InName.ToString(), NAME_None));
+	}
+
+	for (TFieldIterator<const FProperty> PropIt(InStruct, EFieldIterationFlags::IncludeSuper); PropIt; ++PropIt)
+	{
+		const FString PropertyNameStr = InStruct->GetAuthoredNameForField(*PropIt);
+		const FName PropertyName = *PropertyNameStr;
+		if (NamesToLookFor.Contains(PropertyName)
+			|| (!PropertyName.IsValidXName() && NamesToLookFor.Contains(MakeObjectNameFromDisplayLabel(PropertyNameStr, NAME_None))))
+		{
+			return *PropIt;
+		}
+	}
+
+	return nullptr;
+}
+
 #undef LOCTEXT_NAMESPACE
\ No newline at end of file
diff --git a/Engine/Plugins/PCG/Source/PCG/Private/Metadata/Accessors/PCGAttributeAccessorHelpers.cpp b/Engine/Plugins/PCG/Source/PCG/Private/Metadata/Accessors/PCGAttributeAccessorHelpers.cpp
--- a/Engine/Plugins/PCG/Source/PCG/Private/Metadata/Accessors/PCGAttributeAccessorHelpers.cpp
+++ b/Engine/Plugins/PCG/Source/PCG/Private/Metadata/Accessors/PCGAttributeAccessorHelpers.cpp
@@ -7,6 +7,7 @@
 #include "PCGSettings.h"
 #include "Data/PCGPointData.h"
 #include "Data/PCGSpatialData.h"
+#include "Helpers/PCGPropertyHelpers.h"
 #include "Metadata/PCGAttributePropertySelector.h"
 #include "Metadata/PCGMetadata.h"
 #include "Metadata/PCGMetadataAttributeTpl.h"
@@ -236,18 +237,10 @@ namespace PCGAttributeAccessorHelpers
 			const FProperty* Property = nullptr;
 			const FName PropertyName = InPropertyNames[i];
 
-			// Try to get the property. If it is coming from a user struct, we need to iterate on all properties because the property name is mangled
+			// Try to get the property. If it is coming from a user struct, we need to have a special search because the property name might be mangled.
 			if (const UUserDefinedStruct* UserDefinedStruct = Cast<UUserDefinedStruct>(CurrentStruct))
 			{
-				for (TFieldIterator<const FProperty> PropIt(UserDefinedStruct, EFieldIterationFlags::IncludeSuper); PropIt; ++PropIt)
-				{
-					const FName TempPropertyName = *UserDefinedStruct->GetAuthoredNameForField(*PropIt);
-					if (TempPropertyName == PropertyName)
-					{
-						Property = *PropIt;
-						break;
-					}
-				}
+				Property = PCGPropertyHelpers::FindPropertyInUserDefinedStruct(UserDefinedStruct, PropertyName);
 			}
 			else
 			{
diff --git a/Engine/Plugins/PCG/Source/PCG/Private/PCGSettings.cpp b/Engine/Plugins/PCG/Source/PCG/Private/PCGSettings.cpp
--- a/Engine/Plugins/PCG/Source/PCG/Private/PCGSettings.cpp
+++ b/Engine/Plugins/PCG/Source/PCG/Private/PCGSettings.cpp
@@ -23,8 +23,9 @@
 
 #include "Misc/Crc.h"
 #include "Serialization/ArchiveObjectCrc32.h"
-#include "UObject/ObjectSaveContext.h"
 #include "UObject/FortniteMainBranchObjectVersion.h"
+#include "UObject/FortniteReleaseBranchCustomObjectVersion.h"
+#include "UObject/ObjectSaveContext.h"
 
 #include UE_INLINE_GENERATED_CPP_BY_NAME(PCGSettings)
 
@@ -346,6 +347,7 @@ void UPCGSettings::Serialize(FArchive& Ar)
 
 	Ar.UsingCustomVersion(FPCGCustomVersion::GUID);
 	Ar.UsingCustomVersion(FFortniteMainBranchObjectVersion::GUID);
+	Ar.UsingCustomVersion(FFortniteReleaseBranchCustomObjectVersion::GUID);
 
 #if WITH_EDITOR
 	if (Ar.IsLoading())
diff --git a/Engine/Plugins/PCG/Source/PCG/Public/Elements/PCGGetActorProperty.h b/Engine/Plugins/PCG/Source/PCG/Public/Elements/PCGGetActorProperty.h
--- a/Engine/Plugins/PCG/Source/PCG/Public/Elements/PCGGetActorProperty.h
+++ b/Engine/Plugins/PCG/Source/PCG/Public/Elements/PCGGetActorProperty.h
@@ -83,6 +83,10 @@ public:
 	UPROPERTY(BlueprintReadWrite, EditAnywhere, Category = Settings, meta = (PCG_Overridable, EditCondition = "!bForceObjectAndStructExtraction", EditConditionHides, PCG_DiscardPropertySelection, PCG_DiscardExtraSelection))
 	FPCGAttributePropertyOutputSelector OutputAttributeName;
 
+	/** If the output attribute name has special characters, remove them. */
+	UPROPERTY(BlueprintReadWrite, EditAnywhere, Category = Settings)
+	bool bSanitizeOutputAttributeName = true;
+
 	/** Controls whether an actor reference attribute will be added to the result */
 	UPROPERTY(BlueprintReadWrite, EditAnywhere, Category = Settings, meta = (PCG_Overridable))
 	bool bOutputActorReference = false;
diff --git a/Engine/Plugins/PCG/Source/PCG/Public/Elements/PCGGetPropertyFromObjectPath.h b/Engine/Plugins/PCG/Source/PCG/Public/Elements/PCGGetPropertyFromObjectPath.h
--- a/Engine/Plugins/PCG/Source/PCG/Public/Elements/PCGGetPropertyFromObjectPath.h
+++ b/Engine/Plugins/PCG/Source/PCG/Public/Elements/PCGGetPropertyFromObjectPath.h
@@ -18,6 +18,10 @@ class UPCGGetPropertyFromObjectPathSettings : public UPCGSettings
 public:
 	UPCGGetPropertyFromObjectPathSettings();
 
+	//~Begin UObject interface
+	virtual void PostLoad() override;
+	//~End UObject interface
+
 	//~Begin UPCGSettings interface
 #if WITH_EDITOR
 	virtual FName GetDefaultNodeName() const override;
@@ -61,6 +65,10 @@ public:
 	UPROPERTY(BlueprintReadWrite, EditAnywhere, Category = Settings, meta = (PCG_Overridable, EditCondition = "!bForceObjectAndStructExtraction", EditConditionHides, PCG_DiscardPropertySelection, PCG_DiscardExtraSelection))
 	FPCGAttributePropertyOutputSelector OutputAttributeName;
 
+	/** If the output attribute name has special characters, remove them. */
+	UPROPERTY(BlueprintReadWrite, EditAnywhere, Category = Settings)
+	bool bSanitizeOutputAttributeName = true;
+
 	/** By default, object loading is asynchronous, can force it synchronous if needed. */
 	UPROPERTY(BlueprintReadWrite, EditAnywhere, Category = "Settings|Debug")
 	bool bSynchronousLoad = false;
diff --git a/Engine/Plugins/PCG/Source/PCG/Public/Elements/PCGUserParameterGet.h b/Engine/Plugins/PCG/Source/PCG/Public/Elements/PCGUserParameterGet.h
--- a/Engine/Plugins/PCG/Source/PCG/Public/Elements/PCGUserParameterGet.h
+++ b/Engine/Plugins/PCG/Source/PCG/Public/Elements/PCGUserParameterGet.h
@@ -38,9 +38,17 @@ public:
 	/** If the property is a struct/object supported by metadata, this option can be toggled to force extracting all (compatible) properties contained in this property. Automatically true if unsupported by metadata. For now, only supports direct child properties (and not deeper). */
 	UPROPERTY(BlueprintReadWrite, EditAnywhere, Category = Settings)
 	bool bForceObjectAndStructExtraction = false;
+
+	/** If the output attribute name has special characters, remove them. */
+	UPROPERTY(BlueprintReadWrite, EditAnywhere, Category = Settings)
+	bool bSanitizeOutputAttributeName = true;
 	
 	void UpdatePropertyName(FName InNewName);
 
+	//~Begin UObject interface
+	virtual void PostLoad() override;
+	//~End UObject interface
+
 	//~Begin UPCGSettings interface
 	virtual bool RequiresDataFromPreTask() const override { return true; }
 
@@ -81,6 +89,10 @@ public:
 	UPROPERTY(BlueprintReadWrite, EditAnywhere, Category = Settings, meta = (PCG_Overridable))
 	bool bForceObjectAndStructExtraction = false;
 
+	/** If the output attribute name has special characters, remove them. */
+	UPROPERTY(BlueprintReadWrite, EditAnywhere, Category = Settings)
+	bool bSanitizeOutputAttributeName = true;
+
 	UPROPERTY(BlueprintReadWrite, EditAnywhere, Category = Settings, meta = (PCG_Overridable))
 	FName OutputAttributeName = NAME_None;
 
@@ -90,6 +102,10 @@ public:
 	UPROPERTY(BlueprintReadWrite, EditAnywhere, Category = Settings)
 	bool bQuiet = false;
 
+	//~Begin UObject interface
+	virtual void PostLoad() override;
+	//~End UObject interface
+
 	//~Begin UPCGSettings interface
 	virtual bool RequiresDataFromPreTask() const override { return true; }
 	virtual FString GetAdditionalTitleInformation() const override;
diff --git a/Engine/Plugins/PCG/Source/PCG/Public/Helpers/PCGPropertyHelpers.h b/Engine/Plugins/PCG/Source/PCG/Public/Helpers/PCGPropertyHelpers.h
--- a/Engine/Plugins/PCG/Source/PCG/Public/Helpers/PCGPropertyHelpers.h
+++ b/Engine/Plugins/PCG/Source/PCG/Public/Helpers/PCGPropertyHelpers.h
@@ -16,6 +16,7 @@ enum class EPCGMetadataTypes : uint8;
 struct FPCGContext;
 class UPCGData;
 class UPCGParamData;
+class UUserDefinedStruct;
 
 namespace PCGPropertyHelpers
 {
@@ -91,6 +92,11 @@ namespace PCGPropertyHelpers
 
 		// Extra flag if the property needs to be visible to be extractable.
 		bool bPropertyNeedsToBeVisible = false;
+
+		// Output attribute name are already sanitized by default but it keeps some characters like spaces.
+		// This is a stricter option and if the name is not following the object name rule (pretty much only alphanumerical and "_"), it will be sanitized.
+		// Necessary if we have cases with spaces, which doesn't play nicely with PCG attributes.
+		bool bStrictSanitizeOutputAttributeNames = false;
 	};
 
 	/**
@@ -109,6 +115,12 @@ namespace PCGPropertyHelpers
 	*/
 	template <typename T>
 	TArray<T> ExtractAttributeSetAsArrayOfStructs(const UPCGParamData* InParamData, const TMap<FName, TTuple<FName, bool>>* OptionalNameMapping = nullptr, FPCGContext* OptionalContext = nullptr);
+
+	/**
+	* More permissive search for a property name in a UserDefinedStruct.
+	* Since the names in editor and at runtime can be different for the properties defined in a UserDefinedStruct, the search will also look for sanitized version of the name if it is not a valid name.
+	*/
+	const FProperty* FindPropertyInUserDefinedStruct(const UUserDefinedStruct* InStruct, const FName InName);
 }
 
 //////
diff --git a/Engine/Source/Runtime/Core/Public/UObject/FortniteReleaseBranchCustomObjectVersions.inl b/Engine/Source/Runtime/Core/Public/UObject/FortniteReleaseBranchCustomObjectVersions.inl
--- a/Engine/Source/Runtime/Core/Public/UObject/FortniteReleaseBranchCustomObjectVersions.inl
+++ b/Engine/Source/Runtime/Core/Public/UObject/FortniteReleaseBranchCustomObjectVersions.inl
@@ -54,5 +54,8 @@ FORTNITE_RELEASE_VERSION(NavMovementComponentMovingPropertiesToStruct, 14)
 // Add bone serialization for dynamic mesh attributes
 FORTNITE_RELEASE_VERSION(DynamicMeshAttributesSerializeBones, 15)
 
+// Add option for sanitizing output attribute names for all PCG data getters
+FORTNITE_RELEASE_VERSION(OptionSanitizeOutputAttributeNamesPCG, 16)
+
 // -----<new versions can be added above this line>-------------------------------------------------
 #undef FORTNITE_RELEASE_VERSION
\ No newline at end of file
