[PCG] Fix Spline Direction algorithm to work with all splines

* Was using the tangents but it's not working for curves
* Use the Shoelace formula instead

#rnx

diff --git a/Engine/Plugins/PCG/Source/PCG/Private/Elements/PCGSplineDirection.cpp b/Engine/Plugins/PCG/Source/PCG/Private/Elements/PCGSplineDirection.cpp
--- a/Engine/Plugins/PCG/Source/PCG/Private/Elements/PCGSplineDirection.cpp
+++ b/Engine/Plugins/PCG/Source/PCG/Private/Elements/PCGSplineDirection.cpp
@@ -13,26 +13,42 @@ namespace PCGSplineDirection
 	{
 		check(InputSplineData);
 
-		double CumulativeAngle = 0.0;
-		int NumPoints = InputSplineData->SplineStruct.GetSplinePointsPosition().Points.Num();
-		if (!InputSplineData->IsClosed())
+		double CumulativeArea = 0.0;
+		const int NumPoints = InputSplineData->SplineStruct.GetSplinePointsPosition().Points.Num();
+
+		if (NumPoints <= 1)
 		{
-			NumPoints--;
+			return true;
 		}
 
-		for (int32 Index = 0; Index < NumPoints; ++Index)
+		if (NumPoints == 2)
 		{
-			FVector ArriveTangent{}, LeaveTangent{};
-
-			InputSplineData->GetTangentsAtSegmentStart(Index, ArriveTangent, LeaveTangent);
+			FVector ArriveTangent{}, LeaveTangent{}, DummyTangent;
+			InputSplineData->GetTangentsAtSegmentStart(0, DummyTangent, LeaveTangent);
+			InputSplineData->GetTangentsAtSegmentStart(1, ArriveTangent, DummyTangent);
 			const double CrossProduct = (ArriveTangent.X * LeaveTangent.Y - ArriveTangent.Y * LeaveTangent.X);
-			const double AbsAngle = FMath::Acos(ArriveTangent.CosineAngle2D(LeaveTangent));
-			const double Angle = FMath::Sign(CrossProduct) * AbsAngle;
+			return CrossProduct <= 0;
+		}
+
+		int NumIter = NumPoints;
+    
+		if (InputSplineData->IsClosed())
+		{
+			NumIter++;
+		}
+
+		// Shoelace formula - https://en.wikipedia.org/wiki/Shoelace_formula.
+		for (int32 Index = 0; Index < NumIter; ++Index)
+		{
+			const FVector& ThisPoint = InputSplineData->SplineStruct.GetSplinePointsPosition().Points[Index % NumPoints].OutVal;
+			const FVector& NextPoint = InputSplineData->SplineStruct.GetSplinePointsPosition().Points[(Index + 1) % NumPoints].OutVal;
+
+			const double CrossProduct = (NextPoint.X * ThisPoint.Y - ThisPoint.X * NextPoint.Y);
 
-			CumulativeAngle += Angle;
+			CumulativeArea += CrossProduct;
 		}
 
-		return CumulativeAngle <= 0;
+		return CumulativeArea <= 0;
 	}
 
 	UPCGSplineData* Reverse(const UPCGSplineData* InputSplineData, FPCGContext* Context)
