[PCG] Added warning case in Match&Set with keep unmatched to prevent data loss for attributes that wouldn't have the same type.

#rnx

diff --git a/Engine/Plugins/PCG/Source/PCG/Private/Elements/PCGMatchAndSetAttributes.cpp b/Engine/Plugins/PCG/Source/PCG/Private/Elements/PCGMatchAndSetAttributes.cpp
--- a/Engine/Plugins/PCG/Source/PCG/Private/Elements/PCGMatchAndSetAttributes.cpp
+++ b/Engine/Plugins/PCG/Source/PCG/Private/Elements/PCGMatchAndSetAttributes.cpp
@@ -677,10 +677,36 @@ public:
 			FPCGMetadataAttributeBase* PointAttribute = OutMetadata->GetMutableAttribute(AttributeName);
 			if (PointAttribute && PointAttribute->GetTypeId() != ParamAttribute->GetTypeId())
 			{
-				OutMetadata->DeleteAttribute(AttributeName);
-				PointAttribute = nullptr;
+				// If we're in the unmatched case, we'll have to cast the attribute if possible.
+				// Otherwise, log a warning, since we don't want data loss here.
+				if (Settings->bMatchAttributes && Settings->bKeepUnmatched)
+				{
+					bool bCouldCast = false;
+
+					if (PCG::Private::IsBroadcastableOrConstructible(PointAttribute->GetTypeId(), ParamAttribute->GetTypeId()))
+					{
+						PointAttribute = nullptr; // ChangeAttributeType will delete the attribute PointAttribute points to.
+						bCouldCast = OutMetadata->ChangeAttributeType(AttributeName, ParamAttribute->GetTypeId());
+					}
+
+					if (bCouldCast)
+					{
+						PointAttribute = OutMetadata->GetMutableAttribute(AttributeName);
+					}
+					else
+					{
+						PCGE_LOG_C(Warning, GraphAndLog, Context, FText::Format(LOCTEXT("CannotReuseDifferentTypedAttribute", "Unable to reuse attribute '{0}' on output data due to its different type vs. the match data and to the usage of 'Keep Unmatched'."), FText::FromName(AttributeName)));
+						return false;
+					}
+				}
+				else
+				{
+					OutMetadata->DeleteAttribute(AttributeName);
+					PointAttribute = nullptr;
+				}
 			}
 
+			// In most cases, the attribute didn't exist or we're not keeping the unmatched entries, so just create a new attribute.
 			if (!PointAttribute)
 			{
 				PointAttribute = OutMetadata->CopyAttribute(ParamAttribute, AttributeName, /*bKeepParent=*/false, /*bCopyEntries=*/false, /*bCopyValues=*/false);
