[PCG] Used proper SM transform in PCG level exporter.

#12658

diff --git a/Engine/Plugins/PCG/Source/PCGEditor/Private/PCGLevelToAsset.cpp b/Engine/Plugins/PCG/Source/PCGEditor/Private/PCGLevelToAsset.cpp
--- a/Engine/Plugins/PCG/Source/PCGEditor/Private/PCGLevelToAsset.cpp
+++ b/Engine/Plugins/PCG/Source/PCGEditor/Private/PCGLevelToAsset.cpp
@@ -479,7 +479,7 @@ bool UPCGLevelToAsset::BP_ExportWorld_Implementation(UWorld* World, const FStrin
 			else
 			{
 				PointData->SetNumPoints(PointData->GetNumPoints() + 1);
-				MakePoint(FirstIndex, ActorTransform, MeshPath, MeshBounds, Materials);
+				MakePoint(FirstIndex, SMC->GetComponentTransform(), MeshPath, MeshBounds, Materials);
 			}
 		}
 
