[PCG] Update and refresh validation of proxies.

#rnx

diff --git a/Engine/Plugins/PCG/Source/PCG/Private/Compute/DataInterfaces/Elements/PCGCopyPointsAnalysisDataInterface.cpp b/Engine/Plugins/PCG/Source/PCG/Private/Compute/DataInterfaces/Elements/PCGCopyPointsAnalysisDataInterface.cpp
--- a/Engine/Plugins/PCG/Source/PCG/Private/Compute/DataInterfaces/Elements/PCGCopyPointsAnalysisDataInterface.cpp
+++ b/Engine/Plugins/PCG/Source/PCG/Private/Compute/DataInterfaces/Elements/PCGCopyPointsAnalysisDataInterface.cpp
@@ -101,7 +101,9 @@ FComputeDataProviderRenderProxy* UPCGCopyPointsAnalysisDataProvider::GetRenderPr
 
 bool FPCGCopyPointsAnalysisDataProviderProxy::IsValid(FValidationData const& InValidationData) const
 {
-	return InValidationData.ParameterStructSize == sizeof(FParameters);
+	return InValidationData.ParameterStructSize == sizeof(FParameters)
+		&& Data.MatchAttributeId != INDEX_NONE
+		&& Data.SelectedFlagAttributeId != INDEX_NONE;
 }
 
 void FPCGCopyPointsAnalysisDataProviderProxy::GatherDispatchData(FDispatchData const& InDispatchData)
diff --git a/Engine/Plugins/PCG/Source/PCG/Private/Compute/DataInterfaces/Elements/PCGCopyPointsAnalysisDataInterface.h b/Engine/Plugins/PCG/Source/PCG/Private/Compute/DataInterfaces/Elements/PCGCopyPointsAnalysisDataInterface.h
--- a/Engine/Plugins/PCG/Source/PCG/Private/Compute/DataInterfaces/Elements/PCGCopyPointsAnalysisDataInterface.h
+++ b/Engine/Plugins/PCG/Source/PCG/Private/Compute/DataInterfaces/Elements/PCGCopyPointsAnalysisDataInterface.h
@@ -56,8 +56,8 @@ class FPCGCopyPointsAnalysisDataProviderProxy : public FComputeDataProviderRende
 public:
 	struct FCopyPointsAnalysisData_RenderThread
 	{
-		int32 MatchAttributeId = 0;
-		int32 SelectedFlagAttributeId = 0;
+		int32 MatchAttributeId = INDEX_NONE;
+		int32 SelectedFlagAttributeId = INDEX_NONE;
 		bool bCopyEachSourceOnEveryTarget = true;
 	};
 
diff --git a/Engine/Plugins/PCG/Source/PCG/Private/Compute/DataInterfaces/Elements/PCGCopyPointsDataInterface.cpp b/Engine/Plugins/PCG/Source/PCG/Private/Compute/DataInterfaces/Elements/PCGCopyPointsDataInterface.cpp
--- a/Engine/Plugins/PCG/Source/PCG/Private/Compute/DataInterfaces/Elements/PCGCopyPointsDataInterface.cpp
+++ b/Engine/Plugins/PCG/Source/PCG/Private/Compute/DataInterfaces/Elements/PCGCopyPointsDataInterface.cpp
@@ -276,12 +276,7 @@ void UPCGCopyPointsDataProvider::Reset()
 
 bool FPCGCopyPointsDataProviderProxy::IsValid(FValidationData const& InValidationData) const
 {
-	if (InValidationData.ParameterStructSize != sizeof(FParameters))
-	{
-		return false;
-	}
-
-	return true;
+	return InValidationData.ParameterStructSize == sizeof(FParameters);
 }
 
 struct FCopyPointsDataInterfacePermutationIds
diff --git a/Engine/Plugins/PCG/Source/PCG/Private/Compute/DataInterfaces/Elements/PCGMetadataPartitionDataInterface.cpp b/Engine/Plugins/PCG/Source/PCG/Private/Compute/DataInterfaces/Elements/PCGMetadataPartitionDataInterface.cpp
--- a/Engine/Plugins/PCG/Source/PCG/Private/Compute/DataInterfaces/Elements/PCGMetadataPartitionDataInterface.cpp
+++ b/Engine/Plugins/PCG/Source/PCG/Private/Compute/DataInterfaces/Elements/PCGMetadataPartitionDataInterface.cpp
@@ -152,7 +152,8 @@ void UPCGMetaDataPartitionDataProvider::Reset()
 
 bool FPCGMetaDataPartitionProviderProxy::IsValid(FValidationData const& InValidationData) const
 {
-	return (InValidationData.ParameterStructSize == sizeof(FParameters));
+	return InValidationData.ParameterStructSize == sizeof(FParameters)
+		&& !UniqueStringKeyValues.IsEmpty();
 }
 
 void FPCGMetaDataPartitionProviderProxy::AllocateResources(FRDGBuilder& GraphBuilder, FAllocationData const& InAllocationData)
diff --git a/Engine/Plugins/PCG/Source/PCG/Private/Compute/DataInterfaces/Elements/PCGStaticMeshSpawnerDataInterface.cpp b/Engine/Plugins/PCG/Source/PCG/Private/Compute/DataInterfaces/Elements/PCGStaticMeshSpawnerDataInterface.cpp
--- a/Engine/Plugins/PCG/Source/PCG/Private/Compute/DataInterfaces/Elements/PCGStaticMeshSpawnerDataInterface.cpp
+++ b/Engine/Plugins/PCG/Source/PCG/Private/Compute/DataInterfaces/Elements/PCGStaticMeshSpawnerDataInterface.cpp
@@ -771,12 +771,7 @@ FPCGStaticMeshSpawnerDataProviderProxy::FPCGStaticMeshSpawnerDataProviderProxy(
 
 bool FPCGStaticMeshSpawnerDataProviderProxy::IsValid(FValidationData const& InValidationData) const
 {
-	if (InValidationData.ParameterStructSize != sizeof(FParameters))
-	{
-		return false;
-	}
-
-	return true;
+	return InValidationData.ParameterStructSize == sizeof(FParameters);
 }
 
 void FPCGStaticMeshSpawnerDataProviderProxy::AllocateResources(FRDGBuilder& GraphBuilder, FAllocationData const& InAllocationData)
diff --git a/Engine/Plugins/PCG/Source/PCG/Private/Compute/DataInterfaces/Elements/PCGStaticMeshSpawnerDataInterface.h b/Engine/Plugins/PCG/Source/PCG/Private/Compute/DataInterfaces/Elements/PCGStaticMeshSpawnerDataInterface.h
--- a/Engine/Plugins/PCG/Source/PCG/Private/Compute/DataInterfaces/Elements/PCGStaticMeshSpawnerDataInterface.h
+++ b/Engine/Plugins/PCG/Source/PCG/Private/Compute/DataInterfaces/Elements/PCGStaticMeshSpawnerDataInterface.h
@@ -56,6 +56,7 @@ public:
 	UPROPERTY()
 	TObjectPtr<const UPCGStaticMeshSpawnerKernel> ProducerKernel;
 
+	/** Attributes to use for writing per-instance custom floats. */
 	UPROPERTY()
 	TArray<FUintVector4> AttributeIdOffsetStrides;
 
diff --git a/Engine/Plugins/PCG/Source/PCG/Private/Compute/DataInterfaces/PCGCustomKernelDataInterface.cpp b/Engine/Plugins/PCG/Source/PCG/Private/Compute/DataInterfaces/PCGCustomKernelDataInterface.cpp
--- a/Engine/Plugins/PCG/Source/PCG/Private/Compute/DataInterfaces/PCGCustomKernelDataInterface.cpp
+++ b/Engine/Plugins/PCG/Source/PCG/Private/Compute/DataInterfaces/PCGCustomKernelDataInterface.cpp
@@ -143,13 +143,8 @@ void UPCGCustomComputeKernelDataProvider::Reset()
 
 bool FPCGCustomComputeKernelDataProviderProxy::IsValid(FValidationData const& InValidationData) const
 {
-	if (InValidationData.ParameterStructSize != sizeof(FParameters))
-	{
-		UE_LOG(LogPCG, Error, TEXT("Proxy invalid due to mismatching ParameterStructSize."));
-		return false;
-	}
-
-	return true;
+	return InValidationData.ParameterStructSize == sizeof(FParameters)
+		&& ThreadCount >= 0;
 }
 
 int32 FPCGCustomComputeKernelDataProviderProxy::GetDispatchThreadCount(TArray<FIntVector>& InOutThreadCounts) const
diff --git a/Engine/Plugins/PCG/Source/PCG/Private/Compute/DataInterfaces/PCGDataCollectionUploadDataInterface.cpp b/Engine/Plugins/PCG/Source/PCG/Private/Compute/DataInterfaces/PCGDataCollectionUploadDataInterface.cpp
--- a/Engine/Plugins/PCG/Source/PCG/Private/Compute/DataInterfaces/PCGDataCollectionUploadDataInterface.cpp
+++ b/Engine/Plugins/PCG/Source/PCG/Private/Compute/DataInterfaces/PCGDataCollectionUploadDataInterface.cpp
@@ -71,7 +71,6 @@ FComputeDataProviderRenderProxy* UPCGDataProviderDataCollectionUpload::GetRender
 		DataAdaptor,
 		this,
 		PinDesc,
-		DownstreamInputPinLabels,
 		OutputPinLabel,
 		OutputPinLabelAlias,
 		ExportMode,
@@ -105,7 +104,6 @@ FPCGDataProviderDataCollectionUploadProxy::FPCGDataProviderDataCollectionUploadP
 	TSharedPtr<FPCGDataUploadAdaptor> InDataAdaptor,
 	TWeakObjectPtr<UPCGDataCollectionDataProvider> InDataProvider,
 	const FPCGDataCollectionDesc& InPinDesc,
-	const TArray<FName>& InDownstreamInputPinLabels,
 	FName InOutputPinLabel,
 	FName InOutputPinLabelAlias,
 	EPCGExportMode InExportMode,
@@ -114,10 +112,15 @@ FPCGDataProviderDataCollectionUploadProxy::FPCGDataProviderDataCollectionUploadP
 	bool bInExposeElementCounters)
 	: FPCGDataCollectionDataProviderProxy(InDataProvider, InPinDesc, InExportMode, bInZeroInitialize, InElementCountMultiplier, bInExposeElementCounters, InOutputPinLabel, InOutputPinLabelAlias)
 	, DataAdaptor(InDataAdaptor)
-	, DownstreamInputPinLabels(InDownstreamInputPinLabels)
 {
 }
 
+bool FPCGDataProviderDataCollectionUploadProxy::IsValid(FValidationData const& InValidationData) const
+{
+	return FPCGDataCollectionDataProviderProxy::IsValid(InValidationData)
+		&& DataAdaptor.IsValid();
+}
+
 void FPCGDataProviderDataCollectionUploadProxy::AllocateResources(FRDGBuilder& GraphBuilder, FAllocationData const& InAllocationData)
 {
 	TRACE_CPUPROFILER_EVENT_SCOPE(FPCGDataProviderDataCollectionUploadProxy::AllocateResources);
diff --git a/Engine/Plugins/PCG/Source/PCG/Private/Compute/DataInterfaces/PCGDataCollectionUploadDataInterface.h b/Engine/Plugins/PCG/Source/PCG/Private/Compute/DataInterfaces/PCGDataCollectionUploadDataInterface.h
--- a/Engine/Plugins/PCG/Source/PCG/Private/Compute/DataInterfaces/PCGDataCollectionUploadDataInterface.h
+++ b/Engine/Plugins/PCG/Source/PCG/Private/Compute/DataInterfaces/PCGDataCollectionUploadDataInterface.h
@@ -53,7 +53,6 @@ public:
 		TSharedPtr<FPCGDataUploadAdaptor> InDataAdaptor,
 		TWeakObjectPtr<UPCGDataCollectionDataProvider> InDataProvider,
 		const FPCGDataCollectionDesc& InPinDesc,
-		const TArray<FName>& InDownstreamInputPinLabels,
 		FName InOutputPinLabel,
 		FName InOutputPinLabelAlias,
 		EPCGExportMode InExportMode,
@@ -62,13 +61,12 @@ public:
 		bool bInExposeElementCounters);
 
 	//~ Begin FComputeDataProviderRenderProxy Interface
+	bool IsValid(FValidationData const& InValidationData) const override;
 	void AllocateResources(FRDGBuilder& GraphBuilder, FAllocationData const& InAllocationData) override;
 	//~ End FComputeDataProviderRenderProxy Interface
 
 protected:
 	TSharedPtr<FPCGDataUploadAdaptor> DataAdaptor;
 
-	TArray<FName> DownstreamInputPinLabels;
-
 	TRefCountPtr<FRDGPooledBuffer> ExternalBufferForReuse;
 };
diff --git a/Engine/Plugins/PCG/Source/PCG/Private/Compute/DataInterfaces/PCGDataLabelResolverDataInterface.cpp b/Engine/Plugins/PCG/Source/PCG/Private/Compute/DataInterfaces/PCGDataLabelResolverDataInterface.cpp
--- a/Engine/Plugins/PCG/Source/PCG/Private/Compute/DataInterfaces/PCGDataLabelResolverDataInterface.cpp
+++ b/Engine/Plugins/PCG/Source/PCG/Private/Compute/DataInterfaces/PCGDataLabelResolverDataInterface.cpp
@@ -13,6 +13,7 @@
 #include "ShaderCompilerCore.h"
 #include "ShaderParameterMetadataBuilder.h"
 #include "SystemTextures.h"
+#include "Algo/AnyOf.h"
 #include "ComputeFramework/ComputeMetadataBuilder.h"
 #include "ComputeFramework/ShaderParamTypeDefinition.h"
 
@@ -158,7 +159,7 @@ void UPCGDataLabelResolverDataProvider::Reset()
 
 bool FPCGDataLabelResolverDataProviderProxy::IsValid(FValidationData const& InValidationData) const
 {
-	return InValidationData.ParameterStructSize == sizeof(FParameters);
+	return InValidationData.ParameterStructSize == sizeof(FParameters) && !Algo::AnyOf(DataIdToDataIndexMap, [](int32 InDataIndex) { return InDataIndex < 0; });
 }
 
 void FPCGDataLabelResolverDataProviderProxy::GatherDispatchData(FDispatchData const& InDispatchData)
diff --git a/Engine/Plugins/PCG/Source/PCG/Private/Compute/DataInterfaces/PCGDebugDataInterface.cpp b/Engine/Plugins/PCG/Source/PCG/Private/Compute/DataInterfaces/PCGDebugDataInterface.cpp
--- a/Engine/Plugins/PCG/Source/PCG/Private/Compute/DataInterfaces/PCGDebugDataInterface.cpp
+++ b/Engine/Plugins/PCG/Source/PCG/Private/Compute/DataInterfaces/PCGDebugDataInterface.cpp
@@ -105,13 +105,7 @@ void UPCGDebugDataProvider::Reset()
 
 bool FPCGDebugDataProviderProxy::IsValid(FValidationData const& InValidationData) const
 {
-	if (InValidationData.ParameterStructSize != sizeof(FParameters))
-	{
-		UE_LOG(LogPCG, Error, TEXT("Proxy invalid due to mismatching ParameterStructSize."));
-		return false;
-	}
-
-	return true;
+	return InValidationData.ParameterStructSize == sizeof(FParameters);
 }
 
 void FPCGDebugDataProviderProxy::GatherDispatchData(FDispatchData const& InDispatchData)
diff --git a/Engine/Plugins/PCG/Source/PCG/Private/Compute/DataInterfaces/PCGInstanceDataInterface.cpp b/Engine/Plugins/PCG/Source/PCG/Private/Compute/DataInterfaces/PCGInstanceDataInterface.cpp
--- a/Engine/Plugins/PCG/Source/PCG/Private/Compute/DataInterfaces/PCGInstanceDataInterface.cpp
+++ b/Engine/Plugins/PCG/Source/PCG/Private/Compute/DataInterfaces/PCGInstanceDataInterface.cpp
@@ -29,6 +29,7 @@
 #include "ShaderParameterMetadataBuilder.h"
 #include "ShaderParameterStruct.h"
 #include "TextureResource.h"
+#include "Algo/AnyOf.h"
 #include "Components/PrimitiveComponent.h"
 #include "ComputeFramework/ComputeKernelPermutationSet.h"
 #include "ComputeFramework/ShaderParamTypeDefinition.h"
@@ -207,7 +208,16 @@ FPCGInstanceDataProviderProxy::FPCGInstanceDataProviderProxy(TArray<FPrimitiveSc
 
 bool FPCGInstanceDataProviderProxy::IsValid(FValidationData const& InValidationData) const
 {
-	return InValidationData.ParameterStructSize == sizeof(FParameters);
+	for (const FPrimitiveSceneProxy* Proxy : PrimitiveSceneProxies)
+	{
+		if (!Proxy)
+		{
+			return false;
+		}
+	}
+
+	return InValidationData.ParameterStructSize == sizeof(FParameters)
+		&& !Algo::AnyOf(PrimitiveSceneProxies, [](const FPrimitiveSceneProxy* InProxy) { return !InProxy; });
 }
 
 void FPCGInstanceDataProviderProxy::AllocateResources(FRDGBuilder& GraphBuilder, FAllocationData const& InAllocationData)
@@ -339,7 +349,10 @@ void FPCGInstanceDataProviderProxy::PostSubmit(FRDGBuilder& InGraphBuilder) cons
 				// We're using a custom compute shader here rather than using a compute graph kernel, because ultimately the scene update will incorporate the
 				// new instance data. In the future we will not write directly to the scene here.
 				TShaderMapRef<FPCGSceneWriterCS> Shader(GetGlobalShaderMap(GMaxRHIFeatureLevel));
-				const int GroupCount = FMath::DivideAndRoundUp<int>(NumInstancesAllocatedInGPUScene, FPCGSceneWriterCS::NumThreadsPerGroup);
+				int GroupCount = FMath::DivideAndRoundUp<int>(NumInstancesAllocatedInGPUScene, FPCGSceneWriterCS::NumThreadsPerGroup);
+				ensureAlways(GroupCount <= GRHIMaxDispatchThreadGroupsPerDimension.X);
+				GroupCount = FMath::Min(GRHIMaxDispatchThreadGroupsPerDimension.X, GroupCount);
+
 				FComputeShaderUtils::AddPass(GraphBuilder, RDG_EVENT_NAME("PCGWriteInstanceData"), ERDGPassFlags::Compute | ERDGPassFlags::NeverCull, Shader, Parameters, FIntVector(GroupCount, 1, 1));
 			}));
 
diff --git a/Engine/Plugins/PCG/Source/PCG/Private/Compute/DataInterfaces/PCGInstanceDataInterface.h b/Engine/Plugins/PCG/Source/PCG/Private/Compute/DataInterfaces/PCGInstanceDataInterface.h
--- a/Engine/Plugins/PCG/Source/PCG/Private/Compute/DataInterfaces/PCGInstanceDataInterface.h
+++ b/Engine/Plugins/PCG/Source/PCG/Private/Compute/DataInterfaces/PCGInstanceDataInterface.h
@@ -82,7 +82,6 @@ protected:
 	using FParameters = FPCGInstanceDataInterfaceParameters;
 
 	TArray<FPrimitiveSceneProxy*> PrimitiveSceneProxies;
-	TArray<uint32> PrimitiveCounts;
 
 	/** Instance transforms of all instances across all primitives. Three float4s per instance. */
 	FRDGBufferRef InstanceData = nullptr;
diff --git a/Engine/Plugins/PCG/Source/PCG/Private/Compute/DataInterfaces/PCGLandscapeDataInterface.cpp b/Engine/Plugins/PCG/Source/PCG/Private/Compute/DataInterfaces/PCGLandscapeDataInterface.cpp
--- a/Engine/Plugins/PCG/Source/PCG/Private/Compute/DataInterfaces/PCGLandscapeDataInterface.cpp
+++ b/Engine/Plugins/PCG/Source/PCG/Private/Compute/DataInterfaces/PCGLandscapeDataInterface.cpp
@@ -680,12 +680,8 @@ FPCGLandscapeDataProviderProxy::FPCGLandscapeDataProviderProxy(
 
 bool FPCGLandscapeDataProviderProxy::IsValid(FValidationData const& InValidationData) const
 {
-	if (InValidationData.ParameterStructSize != sizeof(FParameters))
-	{
-		return false;
-	}
-
-	return true;
+	return InValidationData.ParameterStructSize == sizeof(FParameters)
+		&& Resource.IsValid();
 }
 
 void FPCGLandscapeDataProviderProxy::GatherDispatchData(FDispatchData const& InDispatchData)
diff --git a/Engine/Plugins/PCG/Source/PCG/Private/Compute/DataInterfaces/PCGStaticMeshDataInterface.cpp b/Engine/Plugins/PCG/Source/PCG/Private/Compute/DataInterfaces/PCGStaticMeshDataInterface.cpp
--- a/Engine/Plugins/PCG/Source/PCG/Private/Compute/DataInterfaces/PCGStaticMeshDataInterface.cpp
+++ b/Engine/Plugins/PCG/Source/PCG/Private/Compute/DataInterfaces/PCGStaticMeshDataInterface.cpp
@@ -236,17 +236,8 @@ FPCGStaticMeshDataProviderProxy::~FPCGStaticMeshDataProviderProxy()
 
 bool FPCGStaticMeshDataProviderProxy::IsValid(FValidationData const& InValidationData) const
 {
-	if (InValidationData.ParameterStructSize != sizeof(FParameters))
-	{
-		return false;
-	}
-
-	if (!LODResources)
-	{
-		return false;
-	}
-
-	return true;
+	return InValidationData.ParameterStructSize == sizeof(FParameters)
+		&& LODResources;
 }
 
 void FPCGStaticMeshDataProviderProxy::GatherDispatchData(FDispatchData const& InDispatchData)
diff --git a/Engine/Plugins/PCG/Source/PCG/Private/Compute/DataInterfaces/PCGTextureDataInterface.cpp b/Engine/Plugins/PCG/Source/PCG/Private/Compute/DataInterfaces/PCGTextureDataInterface.cpp
--- a/Engine/Plugins/PCG/Source/PCG/Private/Compute/DataInterfaces/PCGTextureDataInterface.cpp
+++ b/Engine/Plugins/PCG/Source/PCG/Private/Compute/DataInterfaces/PCGTextureDataInterface.cpp
@@ -212,12 +212,24 @@ void UPCGTextureDataProvider::Reset()
 
 bool FPCGTextureDataProviderProxy::IsValid(FValidationData const& InValidationData) const
 {
-	if (InValidationData.ParameterStructSize != sizeof(FParameters))
+	for (const FPCGTextureBindingInfo& BindingInfo : BindingInfos)
 	{
-		return false;
+		if (BindingInfo.ResourceType == EPCGTextureResourceType::Invalid
+			|| (!BindingInfo.Texture && !BindingInfo.ExportedTexture))
+		{
+			return false;
+		}
+	}
+
+	for (const FPCGTextureInfo& TextureInfo : TextureInfos)
+	{
+		if (!BindingInfos.IsValidIndex(TextureInfo.BindingIndex) || TextureInfo.SliceIndex < 0)
+		{
+			return false;
+		}
 	}
 
-	return true;
+	return InValidationData.ParameterStructSize == sizeof(FParameters);
 }
 
 struct FPCGTextureDataInterfacePermutationIds
diff --git a/Engine/Plugins/PCG/Source/PCG/Private/Compute/Elements/PCGComputeGraphElement.cpp b/Engine/Plugins/PCG/Source/PCG/Private/Compute/Elements/PCGComputeGraphElement.cpp
--- a/Engine/Plugins/PCG/Source/PCG/Private/Compute/Elements/PCGComputeGraphElement.cpp
+++ b/Engine/Plugins/PCG/Source/PCG/Private/Compute/Elements/PCGComputeGraphElement.cpp
@@ -427,6 +427,7 @@ bool FPCGComputeGraphElement::ExecuteInternal(FPCGContext* InContext) const
 
 		if (Context->bGraphSubmitFailed)
 		{
+			UE_LOG(LogPCG, Warning, TEXT("Submit of compute graph '%s' failed."), *Context->ComputeGraph->GetName());
 			return true;
 		}
 
