[PCG] Remove option to enable node renaming via left click + delay

Note: Dependent on 40701722

diff --git a/Engine/Plugins/PCG/Source/PCGEditor/Private/Nodes/SPCGEditorGraphNode.cpp b/Engine/Plugins/PCG/Source/PCGEditor/Private/Nodes/SPCGEditorGraphNode.cpp
--- a/Engine/Plugins/PCG/Source/PCGEditor/Private/Nodes/SPCGEditorGraphNode.cpp
+++ b/Engine/Plugins/PCG/Source/PCGEditor/Private/Nodes/SPCGEditorGraphNode.cpp
@@ -100,7 +100,8 @@ TSharedRef<SWidget> SPCGEditorGraphNode::CreateTitleWidget(TSharedPtr<SNodeTitle
 	.IsSelected(this, &SPCGEditorGraphNode::IsSelectedExclusively)
 	.MultiLine(false)
 	.MaximumLength(UPCGEditorGraphNode::MaxNodeNameCharacterCount)
-	.OverflowPolicy(ETextOverflowPolicy::Ellipsis);
+	.OverflowPolicy(ETextOverflowPolicy::Ellipsis)
+	.DelayedLeftClickEntersEditMode(false);
 
 	InlineEditableText->SetColorAndOpacity(TAttribute<FLinearColor>::Create(TAttribute<FLinearColor>::FGetter::CreateSP(this, &SPCGEditorGraphNode::GetNodeTitleTextColor)));
 
diff --git a/Engine/Plugins/PCG/Source/PCGEditor/Private/Nodes/SPCGEditorGraphNodeCompact.cpp b/Engine/Plugins/PCG/Source/PCGEditor/Private/Nodes/SPCGEditorGraphNodeCompact.cpp
--- a/Engine/Plugins/PCG/Source/PCGEditor/Private/Nodes/SPCGEditorGraphNodeCompact.cpp
+++ b/Engine/Plugins/PCG/Source/PCGEditor/Private/Nodes/SPCGEditorGraphNodeCompact.cpp
@@ -345,7 +345,8 @@ TSharedRef<SWidget> SPCGEditorGraphNodeCompact::CreateTitleWidget(TSharedPtr<SNo
 		.IsSelected(this, &SPCGEditorGraphNodeCompact::IsSelectedExclusively)
 		.MultiLine(false)
 		.MaximumLength(UPCGEditorGraphNode::MaxNodeNameCharacterCount)
-		.OverflowPolicy(ETextOverflowPolicy::Ellipsis);
+		.OverflowPolicy(ETextOverflowPolicy::Ellipsis)
+		.DelayedLeftClickEntersEditMode(false);
 
 		InlineEditableText->SlatePrepass(); // Prepass to calculated desired size
 
