[PCG] Propagated error to actor lines in the debug object tree view

diff --git a/Engine/Plugins/PCG/Source/PCGEditor/Private/Widgets/SPCGEditorGraphDebugObjectTree.cpp b/Engine/Plugins/PCG/Source/PCGEditor/Private/Widgets/SPCGEditorGraphDebugObjectTree.cpp
--- a/Engine/Plugins/PCG/Source/PCGEditor/Private/Widgets/SPCGEditorGraphDebugObjectTree.cpp
+++ b/Engine/Plugins/PCG/Source/PCGEditor/Private/Widgets/SPCGEditorGraphDebugObjectTree.cpp
@@ -36,43 +36,67 @@ namespace PCGEditorGraphDebugObjectTree
 		OutColorAndOpacity = Item->IsDebuggable() ? DefaultGraphColor : FLinearColor(1.0f, 1.0f, 1.0f, 0.5f);
 	}
 
-	void GetErrorInfo(FPCGEditorGraphDebugObjectItemPtr Item, EVisibility& OutIconVisibility, FText& OutIconTooltipText, FLinearColor& OutIconColorAndOpacity, FLinearColor& OutRowTextColorAndOpacity)
+	void GetErrorInfoInternal(FPCGEditorGraphDebugObjectItemPtr Item, ELogVerbosity::Type& OutMinVerbosity, ELogVerbosity::Type& OutMinLocalVerbosity, FPCGPerNodeVisualLogs& OutLogs)
 	{
-		OutIconVisibility = EVisibility::Hidden;
-		OutIconTooltipText = FText();
-		OutIconColorAndOpacity = FLinearColor::White;
-		OutRowTextColorAndOpacity = FLinearColor::White;
-
-		if (!ensure(Item->GetPCGStack()))
+		if (!Item || !ensure(Item->GetPCGStack()))
 		{
 			return;
 		}
 
-		ELogVerbosity::Type MinLocalVerbosity = ELogVerbosity::All;
-		ELogVerbosity::Type MinVerbosity = ELogVerbosity::All;
-
 		if (const UPCGSubsystem* Subsystem = UPCGSubsystem::GetActiveEditorInstance())
 		{
-			FPCGPerNodeVisualLogs Logs;
-			Subsystem->GetNodeVisualLogs().ForAllMatchingLogs(*Item->GetPCGStack(), [&Item, &Logs, &MinLocalVerbosity](const FPCGStack& InStack, const FPCGPerNodeVisualLogs& InLogs)
+			Subsystem->GetNodeVisualLogs().ForAllMatchingLogs(*Item->GetPCGStack(), [&Item, &OutLogs, &OutMinVerbosity, &OutMinLocalVerbosity](const FPCGStack& InStack, const FPCGPerNodeVisualLogs& InLogs)
 			{
 				// Need to verify this, and clarify this when it's a loop item.
 				// i.e. should get the graph + loop index and compare
-				if (Item->GetPCGStack()->GetGraphForCurrentFrame() == InStack.GetGraphForCurrentFrame())
+				const bool bIsLocal = Item->GetPCGStack()->GetGraphForCurrentFrame() == InStack.GetGraphForCurrentFrame();
+
+				for (const FPCGNodeLogEntry& LogEntry : InLogs)
 				{
-					for (const FPCGNodeLogEntry& LogEntry : InLogs)
+					OutMinVerbosity = FMath::Min(OutMinVerbosity, LogEntry.Verbosity);
+
+					if (bIsLocal)
 					{
-						MinLocalVerbosity = FMath::Min(MinLocalVerbosity, LogEntry.Verbosity);
+						OutMinLocalVerbosity = FMath::Min(OutMinLocalVerbosity, LogEntry.Verbosity);
 					}
 				}
 
-				Logs.Append(InLogs);
+				OutLogs.Append(InLogs);
 				return true;
 			});
+		}
+	}
+
+	void GetErrorInfo(FPCGEditorGraphDebugObjectItemPtr Item, EVisibility& OutIconVisibility, FText& OutIconTooltipText, FLinearColor& OutIconColorAndOpacity, FLinearColor& OutRowTextColorAndOpacity)
+	{
+		OutIconVisibility = EVisibility::Hidden;
+		OutIconTooltipText = FText();
+		OutIconColorAndOpacity = FLinearColor::White;
+		OutRowTextColorAndOpacity = FLinearColor::White;
 
-			OutIconTooltipText = Subsystem->GetNodeVisualLogs().GetSummaryText(Logs, nullptr, &MinVerbosity);
+		ELogVerbosity::Type MinLocalVerbosity = ELogVerbosity::All;
+		ELogVerbosity::Type MinVerbosity = ELogVerbosity::All;
+		FPCGPerNodeVisualLogs Logs;
+
+		// Exception: for actor items, since they can have multiple PCG components, we need to forward the query to the individual items below.
+		if (Item && Item->IsRootGenerationItem())
+		{
+			for (FPCGEditorGraphDebugObjectItemPtr Child : Item->GetChildren())
+			{
+				ELogVerbosity::Type DummyLocalVerbosity = ELogVerbosity::All;
+				GetErrorInfoInternal(Child, MinVerbosity, DummyLocalVerbosity, Logs);
+			}
+		}
+		else // Otherwise -> normal call
+		{
+			GetErrorInfoInternal(Item, MinVerbosity, MinLocalVerbosity, Logs);
 		}
 
+		if (const UPCGSubsystem* Subsystem = UPCGSubsystem::GetActiveEditorInstance())
+		{
+			OutIconTooltipText = Subsystem->GetNodeVisualLogs().GetSummaryText(Logs, nullptr);
+		}
+		
 		OutIconVisibility = OutIconTooltipText.IsEmpty() ? EVisibility::Hidden : EVisibility::Visible;
 
 		constexpr FLinearColor WarningColor(1.0f, 0.75f, 0.0f, 0.9f);
diff --git a/Engine/Plugins/PCG/Source/PCGEditor/Private/Widgets/SPCGEditorGraphDebugObjectTree.h b/Engine/Plugins/PCG/Source/PCGEditor/Private/Widgets/SPCGEditorGraphDebugObjectTree.h
--- a/Engine/Plugins/PCG/Source/PCGEditor/Private/Widgets/SPCGEditorGraphDebugObjectTree.h
+++ b/Engine/Plugins/PCG/Source/PCGEditor/Private/Widgets/SPCGEditorGraphDebugObjectTree.h
@@ -50,6 +50,8 @@ public:
 	virtual bool IsDebuggable() const { return false; }
 	/** Whether this item represents a dynamically executed element (as in dynamic subgraph). */
 	virtual bool IsDynamic() const { return false; }
+	/** Whether this is an actor/root item, which doesn't have stack information to its components. */
+	virtual bool IsRootGenerationItem() const { return false; }
 
 	virtual FString GetLabel(bool bForSorting = false) const = 0;
 	virtual const FPCGStack* GetPCGStack() const { return nullptr; }
@@ -80,6 +82,7 @@ public:
 	virtual FString GetLabel(bool bForSorting = false) const override;
 	virtual const FPCGStack* GetPCGStack() const override { return &PCGStack; }
 	virtual const FSlateBrush* GetIcon() const override;
+	virtual bool IsRootGenerationItem() const override { return true; }
 
 protected:
 	virtual FPCGStack* GetMutablePCGStack() override { return &PCGStack; }
