[PCG] Clarified behavior when casting in match & set node and allowed processing even if a warning happens (prevents data breakage).

diff --git a/Engine/Plugins/PCG/Source/PCG/Private/Elements/PCGMatchAndSetAttributes.cpp b/Engine/Plugins/PCG/Source/PCG/Private/Elements/PCGMatchAndSetAttributes.cpp
--- a/Engine/Plugins/PCG/Source/PCG/Private/Elements/PCGMatchAndSetAttributes.cpp
+++ b/Engine/Plugins/PCG/Source/PCG/Private/Elements/PCGMatchAndSetAttributes.cpp
@@ -677,29 +677,34 @@ public:
 			FPCGMetadataAttributeBase* PointAttribute = OutMetadata->GetMutableAttribute(AttributeName);
 			if (PointAttribute && PointAttribute->GetTypeId() != ParamAttribute->GetTypeId())
 			{
+				bool bShouldDeleteAttribute = true;
+
 				// If we're in the unmatched case, we'll have to cast the attribute if possible.
 				// Otherwise, log a warning, since we don't want data loss here.
 				if (Settings->bMatchAttributes && Settings->bKeepUnmatched)
 				{
-					bool bCouldCast = false;
-
 					if (PCG::Private::IsBroadcastableOrConstructible(PointAttribute->GetTypeId(), ParamAttribute->GetTypeId()))
 					{
-						PointAttribute = nullptr; // ChangeAttributeType will delete the attribute PointAttribute points to.
-						bCouldCast = OutMetadata->ChangeAttributeType(AttributeName, ParamAttribute->GetTypeId());
-					}
+						if (OutMetadata->ChangeAttributeType(AttributeName, ParamAttribute->GetTypeId()))
+						{
+							bShouldDeleteAttribute = false;
 
-					if (bCouldCast)
-					{
-						PointAttribute = OutMetadata->GetMutableAttribute(AttributeName);
+							// ChangeAttributeType will delete the attribute PointAttribute points to so we need to requery it.
+							PointAttribute = OutMetadata->GetMutableAttribute(AttributeName);
+
+							if (Settings->bWarnOnAttributeCast)
+							{
+								PCGE_LOG_C(Warning, GraphAndLog, Context, FText::Format(LOCTEXT("OperationRequiredAttributeCast", "Attribute '{0}' was cast to the match data type because of the 'Keep Unmatched' option. This warning can be turned off in the settings."), FText::FromName(AttributeName)));
+							}
+						}
 					}
 					else
 					{
-						PCGE_LOG_C(Warning, GraphAndLog, Context, FText::Format(LOCTEXT("CannotReuseDifferentTypedAttribute", "Unable to reuse attribute '{0}' on output data due to its different type vs. the match data and to the usage of 'Keep Unmatched'."), FText::FromName(AttributeName)));
-						return false;
+						PCGE_LOG_C(Warning, GraphAndLog, Context, FText::Format(LOCTEXT("CannotReuseDifferentTypedAttribute", "Attribute '{0}' on output data has a different type vs. the match data, and will reset the values of 'Keep Unmatched' data."), FText::FromName(AttributeName)));
 					}
 				}
-				else
+				
+				if(bShouldDeleteAttribute)
 				{
 					OutMetadata->DeleteAttribute(AttributeName);
 					PointAttribute = nullptr;
diff --git a/Engine/Plugins/PCG/Source/PCG/Public/Elements/PCGMatchAndSetAttributes.h b/Engine/Plugins/PCG/Source/PCG/Public/Elements/PCGMatchAndSetAttributes.h
--- a/Engine/Plugins/PCG/Source/PCG/Public/Elements/PCGMatchAndSetAttributes.h
+++ b/Engine/Plugins/PCG/Source/PCG/Public/Elements/PCGMatchAndSetAttributes.h
@@ -118,8 +118,11 @@ public:
 	FPCGAttributePropertyInputSelector WeightAttribute;
 
 	/** Controls whether we will emit a warning and return nothing if there is no provided attribute set. */
-	UPROPERTY(BlueprintReadWrite, EditAnywhere, Category = Settings)
+	UPROPERTY(BlueprintReadWrite, EditAnywhere, Category = "Settings|Advanced")
 	bool bWarnIfNoMatchData = true;
+
+	UPROPERTY(BlueprintReadWrite, EditAnywhere, Category = "Settings|Advanced")
+	bool bWarnOnAttributeCast = true;
 };
 
 class FPCGMatchAndSetAttributesElement : public TPCGTimeSlicedElementBase<FPCGMatchAndSetAttributesExecutionState, FPCGMatchAndSetAttributesIterationState>
