[PCG] DynamicDependencies : Fix case where task could get stuck sleeping because one of its dynamic dependencies (but not all of them) was already completed prior to its PostTaskExecute
- In this case DynamicDependencies was left with a dangling task and the other dependencies completion would not unpause it because of that
- Erase the completed DynamicDependencies in PostTaskExecute

#rnx

diff --git a/Engine/Plugins/PCG/Source/PCG/Private/Graph/PCGGraphExecutor.cpp b/Engine/Plugins/PCG/Source/PCG/Private/Graph/PCGGraphExecutor.cpp
--- a/Engine/Plugins/PCG/Source/PCG/Private/Graph/PCGGraphExecutor.cpp
+++ b/Engine/Plugins/PCG/Source/PCG/Private/Graph/PCGGraphExecutor.cpp
@@ -1023,13 +1023,20 @@ void FPCGGraphExecutor::PostTaskExecute(TSharedPtr<FPCGGraphActiveTask> ActiveTa
 				{
 					PCGGraphExecutor::TScopeLock PauseLock(PausedTaskSuccessorsLock);
 					PCGGraphExecutor::TScopeLock OutputLock(TaskOutputsLock);
-					for (FPCGTaskId DependencyId : ActiveTask.Context->DynamicDependencies)
+					for (TSet<FPCGTaskId>::TIterator It = ActiveTask.Context->DynamicDependencies.CreateIterator(); It; ++It)
 					{
+						FPCGTaskId DependencyId = *It;
+
 						if (!TaskOutputs.Contains(DependencyId))
 						{
 							PausedTaskSuccessors.FindOrAdd(DependencyId).Add(ActiveTaskPtr);
 							bMoveFromActiveToPaused = true;
 						}
+						else
+						{
+							// Make sure to remove DynamicDependency since it has already completed
+							It.RemoveCurrent();
+						}
 					}
 				}
 
