[PCG] Expose ray tracing support for GPU-instanced primitives.

#rnx

diff --git a/Engine/Plugins/PCG/Source/PCG/Private/Components/PCGProceduralISMComponent.cpp b/Engine/Plugins/PCG/Source/PCG/Private/Components/PCGProceduralISMComponent.cpp
--- a/Engine/Plugins/PCG/Source/PCG/Private/Components/PCGProceduralISMComponent.cpp
+++ b/Engine/Plugins/PCG/Source/PCG/Private/Components/PCGProceduralISMComponent.cpp
@@ -149,6 +149,9 @@ FPCGProceduralISMComponentDescriptor& FPCGProceduralISMComponentDescriptor::oper
 	WorldPositionOffsetDisableDistance = Other.WorldPositionOffsetDisableDistance;
 	ShadowCacheInvalidationBehavior = Other.ShadowCacheInvalidationBehavior;
 	DetailMode = Other.DetailMode;
+	bVisibleInRayTracing = Other.bVisibleInRayTracing;
+	RayTracingGroupId = Other.RayTracingGroupId;
+	RayTracingGroupCullingPriority = Other.RayTracingGroupCullingPriority;
 
 	Algo::Transform(Other.OverrideMaterials, OverrideMaterials, [](const TSoftObjectPtr<UMaterialInterface>& OverrideMaterial)
 	{
@@ -197,6 +200,9 @@ void FPCGProceduralISMComponentDescriptor::InitComponent(UPCGProceduralISMCompon
 	InProceduralISMComponent->WorldPositionOffsetDisableDistance = WorldPositionOffsetDisableDistance;
 	InProceduralISMComponent->ShadowCacheInvalidationBehavior = ShadowCacheInvalidationBehavior;
 	InProceduralISMComponent->DetailMode = DetailMode;
+	InProceduralISMComponent->bVisibleInRayTracing = bVisibleInRayTracing;
+	InProceduralISMComponent->RayTracingGroupId = RayTracingGroupId;
+	InProceduralISMComponent->RayTracingGroupCullingPriority = RayTracingGroupCullingPriority;
 }
 
 void FPCGProceduralISMComponentDescriptor::InitFrom(UPCGProceduralISMComponent* InProceduralISMComponent)
@@ -233,6 +239,9 @@ void FPCGProceduralISMComponentDescriptor::InitFrom(UPCGProceduralISMComponent*
 	WorldPositionOffsetDisableDistance = InProceduralISMComponent->WorldPositionOffsetDisableDistance;
 	ShadowCacheInvalidationBehavior = InProceduralISMComponent->ShadowCacheInvalidationBehavior;
 	DetailMode = InProceduralISMComponent->DetailMode;
+	bVisibleInRayTracing = InProceduralISMComponent->bVisibleInRayTracing;
+	RayTracingGroupId = InProceduralISMComponent->RayTracingGroupId;
+	RayTracingGroupCullingPriority = InProceduralISMComponent->RayTracingGroupCullingPriority;
 }
 
 uint32 FPCGProceduralISMComponentDescriptor::ComputeHash() const
@@ -285,7 +294,10 @@ bool FPCGProceduralISMComponentDescriptor::operator==(const FPCGProceduralISMCom
 		&& bReverseCulling == Other.bReverseCulling
 		&& WorldPositionOffsetDisableDistance == Other.WorldPositionOffsetDisableDistance
 		&& ShadowCacheInvalidationBehavior == Other.ShadowCacheInvalidationBehavior
-		&& DetailMode == Other.DetailMode;
+		&& DetailMode == Other.DetailMode
+		&& bVisibleInRayTracing == Other.bVisibleInRayTracing
+		&& RayTracingGroupId == Other.RayTracingGroupId
+		&& RayTracingGroupCullingPriority == Other.RayTracingGroupCullingPriority;
 }
 
 UPCGProceduralISMComponent::UPCGProceduralISMComponent(const FObjectInitializer& ObjectInitializer)
@@ -294,7 +306,6 @@ UPCGProceduralISMComponent::UPCGProceduralISMComponent(const FObjectInitializer&
 	// Disable unsupported rendering features (currently require instance data on CPU).
 	bAffectDynamicIndirectLighting = false;
 	bAffectDistanceFieldLighting = false;
-	bVisibleInRayTracing = false;
 
 	BodyInstance.bSimulatePhysics = false;
 	SetGenerateOverlapEvents(false);
@@ -503,12 +514,6 @@ void UPCGProceduralISMComponent::ValidateComponentSetup()
 		UE_LOG(LogPCG, Warning, TEXT("Affecting distance field lighting is not currently supported by PCGProceduralISMComponent, disabling."));
 		bAffectDistanceFieldLighting = false;
 	}
-
-	if (bVisibleInRayTracing)
-	{
-		UE_LOG(LogPCG, Warning, TEXT("Ray tracing is not currently supported by PCGProceduralISMComponent, disabling."));
-		bVisibleInRayTracing = false;
-	}
 }
 
 #if WITH_EDITOR
diff --git a/Engine/Plugins/PCG/Source/PCG/Private/Components/PCGProceduralISMComponent.h b/Engine/Plugins/PCG/Source/PCG/Private/Components/PCGProceduralISMComponent.h
--- a/Engine/Plugins/PCG/Source/PCG/Private/Components/PCGProceduralISMComponent.h
+++ b/Engine/Plugins/PCG/Source/PCG/Private/Components/PCGProceduralISMComponent.h
@@ -246,6 +246,15 @@ public:
 	UPROPERTY()
 	int32 CustomDepthStencilValue;
 
+	UPROPERTY()
+	uint8 bVisibleInRayTracing : 1;
+
+	UPROPERTY()
+	int32 RayTracingGroupId;
+
+	UPROPERTY()
+	ERayTracingGroupCullingPriority RayTracingGroupCullingPriority;
+
 	UPROPERTY()
 	uint8 bCastShadow : 1;
 
