[PCG] Add number of points set to "Final Centroids"  output to the cluster node

+ Additional attribute output for the count of each

diff --git a/Engine/Plugins/PCG/Source/PCG/Private/Elements/PCGClusterElement.cpp b/Engine/Plugins/PCG/Source/PCG/Private/Elements/PCGClusterElement.cpp
--- a/Engine/Plugins/PCG/Source/PCG/Private/Elements/PCGClusterElement.cpp
+++ b/Engine/Plugins/PCG/Source/PCG/Private/Elements/PCGClusterElement.cpp
@@ -5,6 +5,7 @@
 #include "PCGContext.h"
 #include "Algo/Accumulate.h"
 #include "Data/PCGPointData.h"
+#include "Elements/Metadata/PCGMetadataElementCommon.h"
 #include "Helpers/PCGAsync.h"
 #include "Helpers/PCGHelpers.h"
 #include "Metadata/Accessors/PCGAttributeAccessorHelpers.h"
@@ -17,6 +18,8 @@ namespace PCGClusterElement
 	{
 		const FName InputCentroidsLabel = TEXT("Initial Centroids");
 		const FName FinalCentroidsLabel = TEXT("Final Centroids");
+		const FName ClusterIDAttributeDefaultName = TEXT("ClusterID");
+		const FName FinalCentroidElementCountAttributeDefaultName = TEXT("ClusterElementCount");
 	}
 
 	namespace Helpers
@@ -366,6 +369,13 @@ namespace PCGClusterElement
 } // namespace PCGClusterElement
 
 #if WITH_EDITOR
+
+UPCGClusterSettings::UPCGClusterSettings()
+{
+	ClusterAttribute.SetAttributeName(PCGClusterElement::Constants::ClusterIDAttributeDefaultName);
+	FinalCentroidElementCountAttribute.SetAttributeName(PCGClusterElement::Constants::FinalCentroidElementCountAttributeDefaultName);
+}
+
 FText UPCGClusterSettings::GetNodeTooltipText() const
 {
 	return LOCTEXT("NodeTooltip", "Given a desired number of clusters (categories), find the best fit cluster for each point by distance, using one of various clustering algorithms.");
@@ -492,17 +502,17 @@ bool FPCGClusterElement::PrepareDataInternal(FPCGContext* InContext) const
 		UPCGPointData* OutPointData = FPCGContext::NewObject_AnyThread<UPCGPointData>(Context);
 		OutPointData->InitializeFromData(InputPointData);
 		OutPointData->GetMutablePoints() = InputPointData->GetPointsCopy();
+		UPCGMetadata* OutPointMetadata = OutPointData->MutableMetadata();
+		check(OutPointMetadata);
 
-		if (!OutPointData->Metadata->FindOrCreateAttribute<int32>(Settings->ClusterAttribute, 0, /*bAllowsInterpolation=*/false))
+		const FPCGAttributePropertySelector& Selector = Settings->ClusterAttribute;
+		if (Selector.IsBasicAttribute())
 		{
-			PCGLog::Metadata::LogFailToCreateAttributeError<int32>(Settings->ClusterAttribute, Context);
-			return EPCGTimeSliceInitResult::AbortExecution;
+			PCGMetadataElementCommon::ClearOrCreateAttribute(OutPointMetadata, Selector, 0, Context);
 		}
 
 		OutState.OutPointData = OutPointData;
 
-		FPCGAttributePropertySelector Selector;
-		Selector.SetAttributeName(Settings->ClusterAttribute);
 		OutState.Accessor = PCGAttributeAccessorHelpers::CreateAccessor(OutPointData, Selector);
 		OutState.Keys = PCGAttributeAccessorHelpers::CreateKeys(OutPointData, Selector);
 		if (!OutState.Accessor || !OutState.Keys)
@@ -700,6 +710,35 @@ bool FPCGClusterElement::ExecuteInternal(FPCGContext* InContext) const
 			FPCGTaggedData& OutputCentroidData = Context->OutputData.TaggedData.Emplace_GetRef();
 			OutputCentroidData.Pin = PCGClusterElement::Constants::FinalCentroidsLabel;
 			OutputCentroidData.Data = OutputCentroidPointData;
+
+			if (Settings->bOutputFinalCentroidElementCount)
+			{
+				check(OutputCentroidPointData->MutableMetadata());
+
+				const FPCGAttributePropertyOutputNoSourceSelector& Selector = Settings->FinalCentroidElementCountAttribute;
+				if (Selector.IsBasicAttribute())
+				{
+					PCGMetadataElementCommon::ClearOrCreateAttribute(OutputCentroidPointData->MutableMetadata(), Selector, 0, Context);
+				}
+
+				const TUniquePtr<IPCGAttributeAccessor> Accessor = PCGAttributeAccessorHelpers::CreateAccessor(OutputCentroidPointData, Selector);
+				const TUniquePtr<IPCGAttributeAccessorKeys> Keys = PCGAttributeAccessorHelpers::CreateKeys(OutputCentroidPointData, Selector);
+				if (!Accessor || !Keys)
+				{
+					PCGLog::Metadata::LogFailToCreateAccessorError(Selector, Context);
+					return true;
+				}
+
+				TArray<int32, TInlineAllocator<16>> CentroidElementCount;
+				CentroidElementCount.SetNumZeroed(OutputCentroidPoints.Num());
+				// Tabulate the count of each ClusterID
+				for (int32 PointIndex = 0; PointIndex < ClusteringData.PointToClusterAssignments.Num(); ++PointIndex)
+				{
+					++CentroidElementCount[ClusteringData.PointToClusterAssignments[PointIndex]];
+				}
+
+				Accessor->SetRange(MakeConstArrayView(CentroidElementCount), 0, *Keys);
+			}
 		}
 
 		return true;
diff --git a/Engine/Plugins/PCG/Source/PCG/Private/Metadata/PCGAttributePropertySelector.cpp b/Engine/Plugins/PCG/Source/PCG/Private/Metadata/PCGAttributePropertySelector.cpp
--- a/Engine/Plugins/PCG/Source/PCG/Private/Metadata/PCGAttributePropertySelector.cpp
+++ b/Engine/Plugins/PCG/Source/PCG/Private/Metadata/PCGAttributePropertySelector.cpp
@@ -553,6 +553,24 @@ void FPCGAttributePropertyInputSelector::ApplyDeprecation(int32 InPCGCustomVersi
 
 ///////////////////////////////////////////////////////////////////////
 
+bool FPCGAttributePropertyOutputNoSourceSelector::SerializeFromMismatchedTag(const struct FPropertyTag& Tag, FStructuredArchive::FSlot Slot)
+{
+	if (Tag.GetType().IsStruct(StaticStruct()->GetFName()))
+	{
+		StaticStruct()->SerializeItem(Slot, this, nullptr);
+		return true;
+	}
+	else if (Tag.GetType().GetName() == NAME_NameProperty)
+	{
+		FName Value;
+		Slot << Value;
+		SetAttributeName(Value);
+		return true;
+	}
+
+	return false;
+}
+
 FPCGAttributePropertyOutputSelector::FPCGAttributePropertyOutputSelector()
 {
 	AttributeName = PCGMetadataAttributeConstants::SourceAttributeName;
@@ -592,24 +610,6 @@ FPCGAttributePropertyOutputSelector FPCGAttributePropertyOutputSelector::CopyAnd
 	return *this;
 }
 
-bool FPCGAttributePropertyOutputSelector::SerializeFromMismatchedTag(const struct FPropertyTag& Tag, FStructuredArchive::FSlot Slot)
-{
-	if (Tag.GetType().IsStruct(FPCGAttributePropertySelector::StaticStruct()->GetFName()))
-	{
-		FPCGAttributePropertyOutputSelector::StaticStruct()->SerializeItem(Slot, this, nullptr);
-		return true;
-	}
-	else if (Tag.GetType().GetName() == NAME_NameProperty)
-	{
-		FName Value;
-		Slot << Value;
-		SetAttributeName(Value);
-		return true;
-	}
-
-	return false;
-}
-
 ///////////////////////////////////////////////////////////////////////
 
 bool UPCGAttributePropertySelectorBlueprintHelpers::SetPointProperty(FPCGAttributePropertySelector& Selector, EPCGPointProperties InPointProperty, bool bResetExtraNames)
diff --git a/Engine/Plugins/PCG/Source/PCG/Public/Elements/PCGClusterElement.h b/Engine/Plugins/PCG/Source/PCG/Public/Elements/PCGClusterElement.h
--- a/Engine/Plugins/PCG/Source/PCG/Public/Elements/PCGClusterElement.h
+++ b/Engine/Plugins/PCG/Source/PCG/Public/Elements/PCGClusterElement.h
@@ -25,6 +25,8 @@ class UPCGClusterSettings : public UPCGSettings
 	GENERATED_BODY()
 
 public:
+	UPCGClusterSettings();
+
 	//~Begin UPCGSettings interface
 #if WITH_EDITOR
 	virtual FName GetDefaultNodeName() const override { return FName(TEXT("ClusterElement")); }
@@ -49,9 +51,9 @@ public:
 	UPROPERTY(EditAnywhere, BlueprintReadWrite, Category = Settings, DisplayName = "Clusters", meta = (ClampMin = "1", PCG_Overridable))
 	int32 NumClusters = 3;
 
-	/** Cluster IDs will be written to this attribute on the output. */
-	UPROPERTY(EditAnywhere, BlueprintReadWrite, Category = Settings)
-	FName ClusterAttribute = TEXT("ClusterID");
+	/** Cluster IDs will be written to this attribute on the main output. */
+	UPROPERTY(EditAnywhere, BlueprintReadWrite, Category = Settings, meta = (PCG_Overridable))
+	FPCGAttributePropertyOutputNoSourceSelector ClusterAttribute;
 
 	/** A limit on the number of iterations to run on each algorithm, if it doesn't otherwise converge. A higher number can offer more accuracy at the cost of processing time. */
 	UPROPERTY(EditAnywhere, BlueprintReadWrite, Category = "Settings|Advanced", meta = (ClampMin = "1", PCG_Overridable))
@@ -66,6 +68,14 @@ public:
 	/** Output the final location of the centroids or gaussians. */
 	UPROPERTY(EditAnywhere, BlueprintReadWrite, Category = "Settings|Advanced")
 	bool bOutputFinalCentroids = false;
+
+	/** Output the element count for each of the final centroids. */
+	UPROPERTY(EditAnywhere, BlueprintReadWrite, Category = "Settings|Advanced", meta = (EditCondition = "bOutputFinalCentroids", EditConditionHides, PCG_Overridable))
+	bool bOutputFinalCentroidElementCount = false;
+
+	/** Final centroid element count will be written to this attribute on the Final Centroid output. */
+	UPROPERTY(EditAnywhere, BlueprintReadWrite, Category = "Settings|Advanced", meta = (EditCondition = "bOutputFinalCentroidElementCount", EditConditionHides, PCG_Overridable))
+	FPCGAttributePropertyOutputNoSourceSelector FinalCentroidElementCountAttribute;
 };
 
 namespace PCGClusterElement
diff --git a/Engine/Plugins/PCG/Source/PCG/Public/Metadata/PCGAttributePropertySelector.h b/Engine/Plugins/PCG/Source/PCG/Public/Metadata/PCGAttributePropertySelector.h
--- a/Engine/Plugins/PCG/Source/PCG/Public/Metadata/PCGAttributePropertySelector.h
+++ b/Engine/Plugins/PCG/Source/PCG/Public/Metadata/PCGAttributePropertySelector.h
@@ -240,26 +240,26 @@ struct PCG_API FPCGAttributePropertyInputSelector : public FPCGAttributeProperty
 	void ApplyDeprecation(int32 InPCGCustomVersion);
 };
 
-/** Struct that will default on @Source. */
+// Version where it doesn't make sense to have @Source, alias for FPCGAttributePropertySelector
 USTRUCT(BlueprintType)
-struct PCG_API FPCGAttributePropertyOutputSelector : public FPCGAttributePropertySelector
+struct PCG_API FPCGAttributePropertyOutputNoSourceSelector : public FPCGAttributePropertySelector
 {
 	GENERATED_BODY()
 
-	FPCGAttributePropertyOutputSelector();
-
-	/** Get a copy of the selector, with @Source replaced by the right selector. Can add extra data for specific deprecation cases. */
-	FPCGAttributePropertyOutputSelector CopyAndFixSource(const FPCGAttributePropertyInputSelector* InSourceSelector, const UPCGData* InOptionalData = nullptr) const;
-
 	/** To support previously saved nodes, that used FPCGAttributePropertySelector, we need to define this function to de-serialize the new class using the old. And add a trait (see below). */
-	bool SerializeFromMismatchedTag(const struct FPropertyTag& Tag, FStructuredArchive::FSlot Slot);
+	virtual bool SerializeFromMismatchedTag(const struct FPropertyTag& Tag, FStructuredArchive::FSlot Slot);
 };
 
-// Version where it doesn't make sense to have @Source, alias for FPCGAttributePropertySelector
+/** Struct that will default on @Source. */
 USTRUCT(BlueprintType)
-struct PCG_API FPCGAttributePropertyOutputNoSourceSelector : public FPCGAttributePropertySelector
+struct PCG_API FPCGAttributePropertyOutputSelector : public FPCGAttributePropertyOutputNoSourceSelector
 {
 	GENERATED_BODY()
+
+	FPCGAttributePropertyOutputSelector();
+
+	/** Get a copy of the selector, with @Source replaced by the right selector. Can add extra data for specific deprecation cases. */
+	FPCGAttributePropertyOutputSelector CopyAndFixSource(const FPCGAttributePropertyInputSelector* InSourceSelector, const UPCGData* InOptionalData = nullptr) const;
 };
 
 
@@ -277,7 +277,13 @@ struct TStructOpsTypeTraits<FPCGAttributePropertySelector> : public TStructOpsTy
 	};
 };
 
-template<> struct TStructOpsTypeTraits<FPCGAttributePropertyOutputNoSourceSelector> : public TStructOpsTypeTraits<FPCGAttributePropertySelector>{};
+template<> struct TStructOpsTypeTraits<FPCGAttributePropertyOutputNoSourceSelector> : public TStructOpsTypeTraits<FPCGAttributePropertySelector>
+{
+	enum
+	{
+		WithStructuredSerializeFromMismatchedTag = true
+	};
+};
 
 // Also extra trait to specify to the deserializer that those two classes can be deserialized using the old class.
 template<> 
@@ -290,7 +296,7 @@ struct TStructOpsTypeTraits<FPCGAttributePropertyInputSelector> : public TStruct
 };
 
 template<>
-struct TStructOpsTypeTraits<FPCGAttributePropertyOutputSelector> : public TStructOpsTypeTraits<FPCGAttributePropertyInputSelector> {};
+struct TStructOpsTypeTraits<FPCGAttributePropertyOutputSelector> : public TStructOpsTypeTraits<FPCGAttributePropertyOutputNoSourceSelector> {};
 
 /**
 * Helper class to allow the BP to call the custom setters and getters on FPCGAttributePropertySelector.
