[PCG] Add a new Start->Goal mapping modes the Pathfinder node

* Each Start to Each Goal
* Each Start to Nearest Goal
* Each Start to Pairwise Goal

* Also fixed an oversight where users were forced into 1:N, N:1, or N:N, when those cases were valid. I.E. 2 start points to 5 goal points.

diff --git a/Engine/Plugins/PCG/Source/PCG/Private/Elements/PCGPathfindingElement.cpp b/Engine/Plugins/PCG/Source/PCG/Private/Elements/PCGPathfindingElement.cpp
--- a/Engine/Plugins/PCG/Source/PCG/Private/Elements/PCGPathfindingElement.cpp
+++ b/Engine/Plugins/PCG/Source/PCG/Private/Elements/PCGPathfindingElement.cpp
@@ -274,9 +274,15 @@ bool FPCGPathfindingElement::PrepareDataInternal(FPCGContext* InContext) const
 			return EPCGTimeSliceInitResult::NoOperation;
 		}
 
-		if (OutState.StartPoints.Num() != OutState.GoalPoints.Num() && OutState.StartPoints.Num() != 1 && OutState.GoalPoints.Num() != 1)
+		OutState.IterationCount = Settings->GoalMappingMode == EPCGPathfindingGoalMappingMode::EachStartToEachGoal
+			? OutState.StartPoints.Num() * OutState.GoalPoints.Num()
+			: OutState.StartPoints.Num();
+
+		// Validate that if in N:N -> Start:Goal mode, the inputs are the correct cardinality. All other cases were validated in PrepareData and are acceptable.
+		if (Settings->GoalMappingMode == EPCGPathfindingGoalMappingMode::EachStartToPairwiseGoal && OutState.StartPoints.Num() != OutState.GoalPoints.Num())
 		{
-			PCGLog::InputOutput::LogInvalidCardinalityError(PCGPathfindingElement::Constants::StartLocationsInputPinLabel, PCGPathfindingElement::Constants::GoalLocationsInputPinLabel, SlicedContext);
+			PCGLog::LogErrorOnGraph(FText::Format(
+		LOCTEXT("EachStartToPairwiseGoalInvalidMapping", "For 'Each Start To Pairwise Goal' pathfinding mode, there must be a one-to-one mapping between start locations and goal locations. The current input was {0}:{1} Start->Goal locations."), OutState.StartPoints.Num(), OutState.GoalPoints.Num()), SlicedContext);
 			return EPCGTimeSliceInitResult::AbortExecution;
 		}
 
@@ -456,7 +462,9 @@ bool FPCGPathfindingElement::ExecuteInternal(FPCGContext* InContext) const
 			return true;
 		}
 
-		while (IterationState.PathIterationIndex < ExecutionState.StartPoints.Num())
+		const bool bCartesianMapping = Settings->GoalMappingMode == EPCGPathfindingGoalMappingMode::EachStartToEachGoal;
+
+		while (IterationState.PathIterationIndex < ExecutionState.IterationCount)
 		{
 			using PCGSpatialAlgo::AStar::ESearchResult;
 			TArray<FPCGPoint> FinalPath;
@@ -464,14 +472,31 @@ bool FPCGPathfindingElement::ExecuteInternal(FPCGContext* InContext) const
 			// Initialize the next iteration.
 			if (IterationState.PathIterationIndex != IterationState.LastPathIterationIndex)
 			{
-				PCGSpatialAlgo::AStar::Initialize(&ExecutionState.StartPoints[IterationState.PathIterationIndex], IterationState.SearchState);
+				// Per start, per goal. GoalPoints was validated as not empty (0) in PrepareData.
+				const int32 StartPointIndex = bCartesianMapping
+					? IterationState.PathIterationIndex / ExecutionState.GoalPoints.Num()
+					: IterationState.PathIterationIndex;
+				PCGSpatialAlgo::AStar::Initialize(&ExecutionState.StartPoints[StartPointIndex], IterationState.SearchState);
 				IterationState.LastPathIterationIndex = IterationState.PathIterationIndex;
 			}
 
+			TConstArrayView<FPCGPoint> GoalPointsView;
+			if (Settings->GoalMappingMode == EPCGPathfindingGoalMappingMode::EachStartToNearestGoal)
+			{
+				GoalPointsView = MakeConstArrayView<FPCGPoint>(ExecutionState.GoalPoints.GetData(), ExecutionState.GoalPoints.Num());
+			}
+			else // Other mapping modes will do a one-to-one comparison of a single start and single goal.
+			{
+				const int32 GoalPointIndex = bCartesianMapping
+					? IterationState.PathIterationIndex % ExecutionState.GoalPoints.Num()
+					: IterationState.PathIterationIndex;
+				GoalPointsView = MakeConstArrayView<FPCGPoint>(ExecutionState.GoalPoints.GetData() + GoalPointIndex, /*Size=*/1);
+			}
+
 			ESearchResult SearchResult;
 			do
 			{
-				SearchResult = ExecuteSearchIteration(ExecutionState.Settings, IterationState.SearchState, ExecutionState.GoalPoints, FinalPath);
+				SearchResult = ExecuteSearchIteration(ExecutionState.Settings, IterationState.SearchState, GoalPointsView, FinalPath);
 
 				if (SearchResult == ESearchResult::Processing && Context->ShouldStop())
 				{
diff --git a/Engine/Plugins/PCG/Source/PCG/Private/SpatialAlgo/PCGAStar.cpp b/Engine/Plugins/PCG/Source/PCG/Private/SpatialAlgo/PCGAStar.cpp
--- a/Engine/Plugins/PCG/Source/PCG/Private/SpatialAlgo/PCGAStar.cpp
+++ b/Engine/Plugins/PCG/Source/PCG/Private/SpatialAlgo/PCGAStar.cpp
@@ -41,7 +41,7 @@ namespace Helpers
 		Algo::Reverse(OutPath);
 	}
 
-	TTuple<const FPCGPoint*, double> GetNearestGoalToLocation(const FVector& Location, const TArray<FPCGPoint>& GoalPoints)
+	TTuple<const FPCGPoint*, double> GetNearestGoalToLocation(const FVector& Location, const TConstArrayView<FPCGPoint>& GoalPoints)
 	{
 		check(!GoalPoints.IsEmpty());
 
@@ -95,7 +95,7 @@ void Initialize(const FPCGPoint* const StartPoint, FSearchState& OutSearchState)
 }
 
 /** Runs a single iteration of the A* algorithm. Intended to be called multiple times, whereupon it will return true when the algorithm is finished. */
-ESearchResult ExecuteSearchIteration(const FSearchSettings& SearchSettings, FSearchState& SearchState, const TArray<FPCGPoint>& GoalPoints, TArray<FPCGPoint>& OutPath)
+ESearchResult ExecuteSearchIteration(const FSearchSettings& SearchSettings, FSearchState& SearchState, const TConstArrayView<FPCGPoint> GoalPoints, TArray<FPCGPoint>& OutPath)
 {
 	TRACE_CPUPROFILER_EVENT_SCOPE(FPCGPathfindingElement::Algorithm::ExecuteSearchIteration);
 	check(SearchState.OriginatingPointData);
diff --git a/Engine/Plugins/PCG/Source/PCG/Public/Elements/PCGPathfindingElement.h b/Engine/Plugins/PCG/Source/PCG/Public/Elements/PCGPathfindingElement.h
--- a/Engine/Plugins/PCG/Source/PCG/Public/Elements/PCGPathfindingElement.h
+++ b/Engine/Plugins/PCG/Source/PCG/Public/Elements/PCGPathfindingElement.h
@@ -27,6 +27,14 @@ enum class EPCGPathfindingCostFunctionMode : uint8
 	CostMultiplier UMETA(Tooltip = "Pathfinding cost will be the distance multiplied by the provided factor. Note that multipliers below 1 will be clamped to 1.")
 };
 
+UENUM(BlueprintType)
+enum class EPCGPathfindingGoalMappingMode : uint8
+{
+	EachStartToNearestGoal UMETA(Tooltip = "For each starting location, find the optimal path to any of the goal locations. There will be one attempted path from each starting location."),
+	EachStartToEachGoal UMETA(Tooltip = "For each starting location, find the optimal path to each of the goal locations. There will be an attempted path from each starting location to every goal. Ex. S1->G1, S1->G2, S2->G1, S2->G2."),
+	EachStartToPairwiseGoal UMETA(Tooltip = "Map each starting location to a consecutively corresponding goal location as a pair, and find the optimal path between them. Input count must match one-to-one.")
+};
+
 /** Finds the optimal path across the points of a given point cloud--should one exist--when provided a start and goal
  * location, and a maximum jump distance between points. Can return a partial path.
  */
@@ -80,6 +88,10 @@ public:
 	UPROPERTY(EditAnywhere, BlueprintReadWrite, Category = Settings, meta = (EditCondition = "!bGoalLocationsAsInput", EditConditionHides, PCG_Overridable))
 	FVector Goal = FVector::ZeroVector;
 
+	/** How each goal location correlates to each start location. Only relevant when using multiple start and goal locations as input. */
+	UPROPERTY(EditAnywhere, BlueprintReadWrite, Category = Settings, meta = (PCG_Overridable))
+	EPCGPathfindingGoalMappingMode GoalMappingMode = EPCGPathfindingGoalMappingMode::EachStartToNearestGoal;
+
 	/** The heuristic estimates a faster path to speed up processing. A higher than 1 heuristic weight can be faster, but it may cease being the optimal path. A weight of 0 is essentially flood fill. */
 	UPROPERTY(EditAnywhere, BlueprintReadWrite, Category = Settings, meta = (ClampMin = "0.0", PCG_Overridable))
 	double HeuristicWeight = 1.0;
@@ -128,6 +140,8 @@ namespace PCGPathfindingElement
 		TArray<FPCGPoint> StartPoints;
 		/** The goal points of the paths to search for. If more than one goal exists, the heuristic will be ignored. */
 		TArray<FPCGPoint> GoalPoints;
+		/** The total number of paths to find. */
+		int32 IterationCount = 0;
 		PCGSpatialAlgo::AStar::FSearchSettings Settings;
 	};
 
diff --git a/Engine/Plugins/PCG/Source/PCG/Public/SpatialAlgo/PCGAStar.h b/Engine/Plugins/PCG/Source/PCG/Public/SpatialAlgo/PCGAStar.h
--- a/Engine/Plugins/PCG/Source/PCG/Public/SpatialAlgo/PCGAStar.h
+++ b/Engine/Plugins/PCG/Source/PCG/Public/SpatialAlgo/PCGAStar.h
@@ -87,5 +87,5 @@ enum class ESearchResult
 void Initialize(const FPCGPoint* const StartPoint, FSearchState& OutSearchState);
 
 /** Runs a single iteration of the A* algorithm. Intended to be called multiple times, whereupon it will return true when the algorithm is finished. */
-ESearchResult ExecuteSearchIteration(const FSearchSettings& SearchSettings, FSearchState& SearchState, const TArray<FPCGPoint>& GoalPoints, TArray<FPCGPoint>& OutPath);
+ESearchResult ExecuteSearchIteration(const FSearchSettings& SearchSettings, FSearchState& SearchState, TConstArrayView<FPCGPoint> GoalPoints, TArray<FPCGPoint>& OutPath);
 }
