[PCG] Fix PerfB issue with new Cluster node constructor.

#rnx

diff --git a/Engine/Plugins/PCG/Source/PCG/Private/Elements/PCGClusterElement.cpp b/Engine/Plugins/PCG/Source/PCG/Private/Elements/PCGClusterElement.cpp
--- a/Engine/Plugins/PCG/Source/PCG/Private/Elements/PCGClusterElement.cpp
+++ b/Engine/Plugins/PCG/Source/PCG/Private/Elements/PCGClusterElement.cpp
@@ -370,7 +370,8 @@ namespace PCGClusterElement
 
 #if WITH_EDITOR
 
-UPCGClusterSettings::UPCGClusterSettings()
+UPCGClusterSettings::UPCGClusterSettings(const FObjectInitializer& ObjectInitializer)
+	: Super(ObjectInitializer)
 {
 	ClusterAttribute.SetAttributeName(PCGClusterElement::Constants::ClusterIDAttributeDefaultName);
 	FinalCentroidElementCountAttribute.SetAttributeName(PCGClusterElement::Constants::FinalCentroidElementCountAttributeDefaultName);
diff --git a/Engine/Plugins/PCG/Source/PCG/Public/Elements/PCGClusterElement.h b/Engine/Plugins/PCG/Source/PCG/Public/Elements/PCGClusterElement.h
--- a/Engine/Plugins/PCG/Source/PCG/Public/Elements/PCGClusterElement.h
+++ b/Engine/Plugins/PCG/Source/PCG/Public/Elements/PCGClusterElement.h
@@ -25,7 +25,7 @@ class UPCGClusterSettings : public UPCGSettings
 	GENERATED_BODY()
 
 public:
-	UPCGClusterSettings();
+	UPCGClusterSettings(const FObjectInitializer& ObjectInitializer);
 
 	//~Begin UPCGSettings interface
 #if WITH_EDITOR
