[PCG] Move NodeVisualLogs / ExecutedStacks  to PCGEditorModule, Move RegisterBeginTickAction to PCGModule::ExecuteNextTick
- Add new IPCGEditorModule::Get for easier access
- Reduce dependency on PCGSubsystem which is a World system (try to be as world agnostic as possible to allow running out of world graphs)

#rnx

diff --git a/Engine/Plugins/PCG/Source/PCG/Private/Compute/Elements/PCGComputeGraphElement.cpp b/Engine/Plugins/PCG/Source/PCG/Private/Compute/Elements/PCGComputeGraphElement.cpp
--- a/Engine/Plugins/PCG/Source/PCG/Private/Compute/Elements/PCGComputeGraphElement.cpp
+++ b/Engine/Plugins/PCG/Source/PCG/Private/Compute/Elements/PCGComputeGraphElement.cpp
@@ -11,6 +11,7 @@
 #include "Compute/Data/PCGProxyForGPUData.h"
 #include "Compute/DataInterfaces/PCGDataCollectionDataInterface.h"
 #include "Elements/PCGStaticMeshSpawnerKernel.h"
+#include "Editor/IPCGEditorModule.h"
 #include "Graph/PCGGraphCompiler.h"
 
 #include "ComputeWorkerInterface.h"
@@ -86,7 +87,7 @@ bool FPCGComputeGraphElement::ExecuteInternal(FPCGContext* InContext) const
 	auto SleepUntilNextFrame = [Context]()
 	{
 		Context->bIsPaused = true;
-		Context->SourceComponent->GetSubsystem()->RegisterBeginTickAction([ContextHandle = Context->GetOrCreateHandle()]()
+		FPCGModule::GetPCGModuleChecked().ExecuteNextTick([ContextHandle = Context->GetOrCreateHandle()]()
 		{
 			if (TSharedPtr<FPCGContextHandle> SharedHandle = ContextHandle.Pin())
 			{
@@ -667,7 +668,7 @@ void FPCGComputeGraphElement::LogCompilationMessages(FPCGComputeGraphContext* In
 
 				if (Verbosity < ELogVerbosity::Log)
 				{
-					if (UPCGSubsystem* Subsystem = UPCGSubsystem::GetInstance(InContext->SourceComponent->GetWorld()))
+					if (IPCGEditorModule* PCGEditorModule = IPCGEditorModule::Get())
 					{
 						FPCGStack StackWithNode = *InContext->Stack;
 						StackWithNode.PushFrame(NodeAndCompileMessages.Get<0>().ResolveObjectPtr());
@@ -690,7 +691,7 @@ void FPCGComputeGraphElement::LogCompilationMessages(FPCGComputeGraphContext* In
 							LogText = FText::FromString(Message.Text);
 						}
 
-						Subsystem->GetNodeVisualLogsMutable().Log(StackWithNode, Verbosity, LogText);
+						PCGEditorModule->GetNodeVisualLogsMutable().Log(StackWithNode, Verbosity, LogText);
 					}
 				}
 			}
diff --git a/Engine/Plugins/PCG/Source/PCG/Private/Compute/PCGComputeCommon.cpp b/Engine/Plugins/PCG/Source/PCG/Private/Compute/PCGComputeCommon.cpp
--- a/Engine/Plugins/PCG/Source/PCG/Private/Compute/PCGComputeCommon.cpp
+++ b/Engine/Plugins/PCG/Source/PCG/Private/Compute/PCGComputeCommon.cpp
@@ -10,6 +10,7 @@
 #include "Compute/PCGPinPropertiesGPU.h"
 #include "Compute/Data/PCGProxyForGPUData.h"
 #include "Data/PCGPointData.h"
+#include "Editor/IPCGEditorModule.h"
 
 #include "DynamicRHI.h"
 #include "RHIStats.h"
@@ -62,14 +63,14 @@ namespace PCGComputeHelpers
 	void LogKernelWarning(const FPCGContext* Context, const UPCGSettings* Settings, const FText& InText)
 	{
 #if WITH_EDITOR
-		if (Context && ensure(Context->SourceComponent.IsValid() && Context->SourceComponent.Get()) && Settings)
+		if (Context && Settings)
 		{
-			if (UPCGSubsystem* Subsystem = Context->SourceComponent->GetSubsystem())
+			if (IPCGEditorModule* PCGEditorModule = IPCGEditorModule::Get())
 			{
 				FPCGStack StackWithNode = Context->Stack ? *Context->Stack : FPCGStack();
 				StackWithNode.PushFrame(Settings->GetOuter());
 
-				Subsystem->GetNodeVisualLogsMutable().Log(StackWithNode, ELogVerbosity::Warning, InText);
+				PCGEditorModule->GetNodeVisualLogsMutable().Log(StackWithNode, ELogVerbosity::Warning, InText);
 			}
 		}
 #endif
@@ -79,14 +80,14 @@ namespace PCGComputeHelpers
 	void LogKernelError(const FPCGContext* Context, const UPCGSettings* Settings, const FText& InText)
 	{
 #if WITH_EDITOR
-		if (Context && ensure(Context->SourceComponent.IsValid() && Context->SourceComponent.Get()) && Settings)
+		if (Context && Settings)
 		{
-			if (UPCGSubsystem* Subsystem = Context->SourceComponent->GetSubsystem())
+			if (IPCGEditorModule* PCGEditorModule = IPCGEditorModule::Get())
 			{
 				FPCGStack StackWithNode = Context->Stack ? *Context->Stack : FPCGStack();
 				StackWithNode.PushFrame(Settings->GetOuter());
 
-				Subsystem->GetNodeVisualLogsMutable().Log(StackWithNode, ELogVerbosity::Error, InText);
+				PCGEditorModule->GetNodeVisualLogsMutable().Log(StackWithNode, ELogVerbosity::Error, InText);
 			}
 		}
 #endif
diff --git a/Engine/Plugins/PCG/Source/PCG/Private/Editor/IPCGEditorModule.cpp b/Engine/Plugins/PCG/Source/PCG/Private/Editor/IPCGEditorModule.cpp
new file mode 100644
--- /dev/null
+++ b/Engine/Plugins/PCG/Source/PCG/Private/Editor/IPCGEditorModule.cpp
@@ -0,0 +1,16 @@
+// Copyright Epic Games, Inc. All Rights Reserved.
+
+#include "Editor/IPCGEditorModule.h"
+
+static IPCGEditorModule* PCGEditorModulePtr = nullptr;
+
+IPCGEditorModule* IPCGEditorModule::Get()
+{
+	return PCGEditorModulePtr;
+}
+
+void IPCGEditorModule::SetEditorModule(IPCGEditorModule* InModule)
+{
+	check(PCGEditorModulePtr == nullptr || InModule == nullptr);
+	PCGEditorModulePtr = InModule;
+}
\ No newline at end of file
diff --git a/Engine/Plugins/PCG/Source/PCG/Private/Elements/ControlFlow/PCGWait.cpp b/Engine/Plugins/PCG/Source/PCG/Private/Elements/ControlFlow/PCGWait.cpp
--- a/Engine/Plugins/PCG/Source/PCG/Private/Elements/ControlFlow/PCGWait.cpp
+++ b/Engine/Plugins/PCG/Source/PCG/Private/Elements/ControlFlow/PCGWait.cpp
@@ -3,6 +3,7 @@
 #include "Elements/ControlFlow/PCGWait.h"
 
 #include "PCGComponent.h"
+#include "PCGModule.h"
 #include "PCGSubsystem.h"
 
 FPCGElementPtr UPCGWaitSettings::CreateElement() const
@@ -58,7 +59,7 @@ bool FPCGWaitElement::ExecuteInternal(FPCGContext* InContext) const
 	else
 	{
 		Context->bIsPaused = true;
-		Context->SourceComponent->GetSubsystem()->RegisterBeginTickAction([ContextHandle = Context->GetOrCreateHandle()]()
+		FPCGModule::GetPCGModuleChecked().ExecuteNextTick([ContextHandle = Context->GetOrCreateHandle()]()
 		{
 			FPCGContext::FSharedContext<FPCGWaitContext> SharedContext(ContextHandle);
 			if (FPCGWaitContext* ContextPtr = SharedContext.Get())
diff --git a/Engine/Plugins/PCG/Source/PCG/Private/Elements/PCGGenerateGrassMaps.cpp b/Engine/Plugins/PCG/Source/PCG/Private/Elements/PCGGenerateGrassMaps.cpp
--- a/Engine/Plugins/PCG/Source/PCG/Private/Elements/PCGGenerateGrassMaps.cpp
+++ b/Engine/Plugins/PCG/Source/PCG/Private/Elements/PCGGenerateGrassMaps.cpp
@@ -4,6 +4,7 @@
 
 #include "PCGComponent.h"
 #include "PCGGrassMapUnpackerCS.h"
+#include "PCGModule.h"
 #include "PCGSubsystem.h"
 #include "Compute/PCGComputeCommon.h"
 #include "Data/PCGLandscapeData.h"
@@ -361,7 +362,7 @@ bool FPCGGenerateGrassMapsElement::ExecuteInternal(FPCGContext* InContext) const
 		{
 			// Sleep until next frame, no use spinning on this.
 			Context->bIsPaused = true;
-			Context->SourceComponent->GetSubsystem()->RegisterBeginTickAction([ContextHandle = InContext->GetOrCreateHandle()]()
+			FPCGModule::GetPCGModuleChecked().ExecuteNextTick([ContextHandle = InContext->GetOrCreateHandle()]()
 			{
 				if (TSharedPtr<FPCGContextHandle> SharedHandle = ContextHandle.Pin())
 				{
@@ -562,7 +563,7 @@ bool FPCGGenerateGrassMapsElement::ExecuteInternal(FPCGContext* InContext) const
 		// Initialization not complete. Could be waiting on async texture processing or for GPU readback. Sleep until next frame.
 		// TODO: Ideally we do lazy readback on texture data in the future so we don't have to read it back to CPU unless it's needed.
 		Context->bIsPaused = true;
-		Context->SourceComponent->GetSubsystem()->RegisterBeginTickAction([ContextHandle = Context->GetOrCreateHandle()]()
+		FPCGModule::GetPCGModuleChecked().ExecuteNextTick([ContextHandle = Context->GetOrCreateHandle()]()
 		{
 			FPCGContext::FSharedContext<FPCGGenerateGrassMapsContext> SharedContext(ContextHandle);
 			if (FPCGGenerateGrassMapsContext* ContextPtr = SharedContext.Get())
diff --git a/Engine/Plugins/PCG/Source/PCG/Private/Elements/PCGTextureSampler.cpp b/Engine/Plugins/PCG/Source/PCG/Private/Elements/PCGTextureSampler.cpp
--- a/Engine/Plugins/PCG/Source/PCG/Private/Elements/PCGTextureSampler.cpp
+++ b/Engine/Plugins/PCG/Source/PCG/Private/Elements/PCGTextureSampler.cpp
@@ -5,6 +5,7 @@
 #include "PCGComponent.h"
 #include "PCGContext.h"
 #include "PCGCrc.h"
+#include "PCGModule.h"
 #include "PCGSubsystem.h"
 #include "Compute/PCGComputeCommon.h"
 #include "Data/PCGRenderTargetData.h"
@@ -260,7 +261,7 @@ bool FPCGTextureSamplerElement::ExecuteInternal(FPCGContext* InContext) const
 			{
 				// Initialization not complete. Could be waiting on async texture processing or for GPU readback. Sleep until next frame.
 				Context->bIsPaused = true;
-				Context->SourceComponent->GetSubsystem()->RegisterBeginTickAction([ContextHandle = Context->GetOrCreateHandle()]()
+				FPCGModule::GetPCGModuleChecked().ExecuteNextTick([ContextHandle = Context->GetOrCreateHandle()]()
 					{
 						FPCGContext::FSharedContext<FPCGTextureSamplerContext> SharedContext(ContextHandle);
 						if (FPCGTextureSamplerContext* ContextPtr = SharedContext.Get())
diff --git a/Engine/Plugins/PCG/Source/PCG/Private/Graph/PCGGraphExecutor.cpp b/Engine/Plugins/PCG/Source/PCG/Private/Graph/PCGGraphExecutor.cpp
--- a/Engine/Plugins/PCG/Source/PCG/Private/Graph/PCGGraphExecutor.cpp
+++ b/Engine/Plugins/PCG/Source/PCG/Private/Graph/PCGGraphExecutor.cpp
@@ -286,6 +286,11 @@ FPCGTaskId FPCGGraphExecutor::ScheduleGraph(const FPCGScheduleGraphParams& InPar
 	}
 
 #if WITH_EDITOR
+	if (IPCGEditorModule* PCGEditorModule = IPCGEditorModule::Get())
+	{
+		PCGEditorModule->OnScheduleGraph(*StackContextPtr);
+	}
+
 	if (UPCGSubsystem* Subsystem = InParams.SourceComponent->GetSubsystem())
 	{
 		Subsystem->OnScheduleGraph(*StackContextPtr);
@@ -1107,9 +1112,9 @@ void FPCGGraphExecutor::PostTaskExecute(TSharedPtr<FPCGGraphActiveTask> ActiveTa
 				CollectGCReferenceTasks.Add(ActiveTaskPtr);
 			}
 
-			PCGHelpers::ExecuteOnGameThread(UE_SOURCE_LOCATION, [WeakHandler = GameThreadHandler->AsWeak(), ActiveTaskPtr, WorldWeak = MakeWeakObjectPtr(World)]()
+			PCGHelpers::ExecuteOnGameThread(UE_SOURCE_LOCATION, [WeakHandler = GameThreadHandler->AsWeak(), ActiveTaskPtr]()
 			{
-				DebugDisplayTask(WeakHandler, ActiveTaskPtr, WorldWeak);
+				DebugDisplayTask(WeakHandler, ActiveTaskPtr);
 			});
 		}
 		
@@ -1378,7 +1383,7 @@ void FPCGGraphExecutor::PrepareForExecute(FPCGGraphTask& Task, FCachedResult*& O
 }
 
 #if WITH_EDITOR
-void FPCGGraphExecutor::DebugDisplayTask(TWeakPtr<FPCGGraphExecutor::FGameThreadHandler> WeakHandler, TSharedPtr<FPCGGraphActiveTask> ActiveTaskPtr, TWeakObjectPtr<UWorld> WeakWorld)
+void FPCGGraphExecutor::DebugDisplayTask(TWeakPtr<FPCGGraphExecutor::FGameThreadHandler> WeakHandler, TSharedPtr<FPCGGraphActiveTask> ActiveTaskPtr)
 {
 	TSharedPtr<FPCGGraphExecutor::FGameThreadHandler> Handler = WeakHandler.Pin();
 	if (!Handler)
@@ -1388,16 +1393,13 @@ void FPCGGraphExecutor::DebugDisplayTask(TWeakPtr<FPCGGraphExecutor::FGameThread
 
 	if (!ActiveTaskPtr->Element->DebugDisplay(ActiveTaskPtr->Context.Get()))
 	{
-		if (UPCGSubsystem* PCGSubsystem = UWorld::GetSubsystem<UPCGSubsystem>(WeakWorld.Get()))
+		// Not done, run on next game thread next tick instead. This is safe to call even if running inside a begin tick action.
+		FPCGModule::GetPCGModuleChecked().ExecuteNextTick([WeakHandler, ActiveTaskPtr]()
 		{
-			// Not done, run on next game thread next tick instead. This is safe to call even if running inside a begin tick action.
-			PCGSubsystem->RegisterBeginTickAction([WeakHandler, ActiveTaskPtr, WeakWorld]()
-			{
-				DebugDisplayTask(WeakHandler, ActiveTaskPtr, WeakWorld);
-			});
+			DebugDisplayTask(WeakHandler, ActiveTaskPtr);
+		});
 
-			return;
-		}
+		return;
 	}
 
 	// If we fall through to here, debug display completed. Cleanup.
@@ -2614,7 +2616,7 @@ void FPCGGraphExecutor::UpdateGenerationNotification()
 		{
 			GenerationProgressNotificationStartTime = 0.0;
 
-			IPCGEditorModule* EditorModule = FModuleManager::GetModulePtr<IPCGEditorModule>("PCGEditor");
+			IPCGEditorModule* EditorModule = IPCGEditorModule::Get();
 			if (!EditorModule)
 			{
 				return;
@@ -2641,7 +2643,7 @@ void FPCGGraphExecutor::ReleaseGenerationNotification()
 	if (GenerationProgressNotification.IsValid())
 	{
 		GenerationProgressNotification.Pin()->OnCancelTasks().RemoveAll(this);
-		if (IPCGEditorModule* EditorModule = FModuleManager::GetModulePtr<IPCGEditorModule>("PCGEditor"))
+		if (IPCGEditorModule* EditorModule = IPCGEditorModule::Get())
 		{
 			EditorModule->ReleaseProgressNotification(GenerationProgressNotification);
 		}
@@ -2886,13 +2888,13 @@ namespace PCGGraphExecutor
 		if (FromGridSize < ToGridSize)
 		{
 #if WITH_EDITOR
-			if (UPCGSubsystem* Subsystem = UPCGSubsystem::GetInstance(InContext->SourceComponent->GetWorld()))
+			if (IPCGEditorModule* PCGEditorModule = IPCGEditorModule::Get())
 			{
 				// Using the low level logging call because we have only a node pointer for the downstream node. Note that InContext
 				// is the context for the linkage element/task, which is not represented on the graph and cannot receive graph warnings/errors.
 				if (ToGridSize == PCGHiGenGrid::UnboundedGridSize())
 				{
-					Subsystem->GetNodeVisualLogsMutable().Log(
+					PCGEditorModule->GetNodeVisualLogsMutable().Log(
 						*InContext->Stack,
 						ELogVerbosity::Error,
 						FText::Format(
@@ -2902,7 +2904,7 @@ namespace PCGGraphExecutor
 				}
 				else
 				{
-					Subsystem->GetNodeVisualLogsMutable().Log(
+					PCGEditorModule->GetNodeVisualLogsMutable().Log(
 						*InContext->Stack,
 						ELogVerbosity::Error,
 						FText::Format(
diff --git a/Engine/Plugins/PCG/Source/PCG/Private/Graph/PCGGraphExecutor.h b/Engine/Plugins/PCG/Source/PCG/Private/Graph/PCGGraphExecutor.h
--- a/Engine/Plugins/PCG/Source/PCG/Private/Graph/PCGGraphExecutor.h
+++ b/Engine/Plugins/PCG/Source/PCG/Private/Graph/PCGGraphExecutor.h
@@ -232,7 +232,7 @@ private:
 	void PrepareForExecute(FPCGGraphTask& Task, FCachedResult*& OutCachedResult, bool bLiveTasksLockAlreadyLocked);
 
 #if WITH_EDITOR
-	static void DebugDisplayTask(TWeakPtr<FGameThreadHandler> GameThreadHandler, TSharedPtr<FPCGGraphActiveTask> ActiveTaskPtr, TWeakObjectPtr<UWorld> WeakWorld);
+	static void DebugDisplayTask(TWeakPtr<FGameThreadHandler> GameThreadHandler, TSharedPtr<FPCGGraphActiveTask> ActiveTaskPtr);
 #endif
 
 	/** Store cache results and Queue next tasks */
diff --git a/Engine/Plugins/PCG/Source/PCG/Private/Graph/PCGGraphTask.cpp b/Engine/Plugins/PCG/Source/PCG/Private/Graph/PCGGraphTask.cpp
--- a/Engine/Plugins/PCG/Source/PCG/Private/Graph/PCGGraphTask.cpp
+++ b/Engine/Plugins/PCG/Source/PCG/Private/Graph/PCGGraphTask.cpp
@@ -5,6 +5,7 @@
 #include "PCGComponent.h"
 #include "PCGSubsystem.h"
 #include "Compute/Elements/PCGComputeGraphElement.h"
+#include "Editor/IPCGEditorModule.h"
 #include "Graph/PCGGraphExecutor.h"
 #include "Helpers/PCGHelpers.h"
 
@@ -30,17 +31,12 @@ bool FPCGGraphTask::CanExecuteOnlyOnMainThread() const
 #if WITH_EDITOR
 void FPCGGraphTask::LogVisual(ELogVerbosity::Type InVerbosity, const FText& InMessage) const
 {
-	if (!SourceComponent.IsValid())
-	{
-		return;
-	}
-
-	if (UPCGSubsystem* Subsystem = UPCGSubsystem::GetInstance(SourceComponent->GetWorld()))
+	if (IPCGEditorModule* PCGEditorModule = IPCGEditorModule::Get())
 	{
 		const FPCGStack* TaskStack = GetStack();
 		FPCGStack StackWithNode = TaskStack ? FPCGStack(*TaskStack) : FPCGStack();
 		StackWithNode.PushFrame(Node);
-		Subsystem->GetNodeVisualLogsMutable().Log(StackWithNode, InVerbosity, InMessage);
+		PCGEditorModule->GetNodeVisualLogsMutable().Log(StackWithNode, InVerbosity, InMessage);
 	}
 }
 
diff --git a/Engine/Plugins/PCG/Source/PCG/Private/Helpers/PCGActorHelpers.cpp b/Engine/Plugins/PCG/Source/PCG/Private/Helpers/PCGActorHelpers.cpp
--- a/Engine/Plugins/PCG/Source/PCG/Private/Helpers/PCGActorHelpers.cpp
+++ b/Engine/Plugins/PCG/Source/PCG/Private/Helpers/PCGActorHelpers.cpp
@@ -586,7 +586,7 @@ bool UPCGActorHelpers::DeleteActors(UWorld* World, const TArray<TSoftObjectPtr<A
 	{
 #if WITH_EDITOR
 
-		if (IPCGEditorModule* EditorModule = FModuleManager::GetModulePtr<IPCGEditorModule>("PCGEditor"))
+		if (IPCGEditorModule* EditorModule = IPCGEditorModule::Get())
 		{
 			EditorModule->SetOutlinerUIRefreshDelay(UE::PCGActorHelpers::Local::OutlinerUIRefreshDelay);
 		}
@@ -726,7 +726,7 @@ AActor* UPCGActorHelpers::SpawnDefaultActor(const FSpawnDefaultActorParams& Para
 	// Specify EDL we want to use if any for spawning this actor
 	FScopedOverrideSpawningLevelMountPointObject EDLScope(ExternalDataLayerInstance ? ExternalDataLayerInstance->GetExternalDataLayerAsset() : nullptr);
 
-	if (IPCGEditorModule* EditorModule = FModuleManager::GetModulePtr<IPCGEditorModule>("PCGEditor"))
+	if (IPCGEditorModule* EditorModule = IPCGEditorModule::Get())
 	{
 		EditorModule->SetOutlinerUIRefreshDelay(UE::PCGActorHelpers::Local::OutlinerUIRefreshDelay);
 	}
diff --git a/Engine/Plugins/PCG/Source/PCG/Private/PCGContext.cpp b/Engine/Plugins/PCG/Source/PCG/Private/PCGContext.cpp
--- a/Engine/Plugins/PCG/Source/PCG/Private/PCGContext.cpp
+++ b/Engine/Plugins/PCG/Source/PCG/Private/PCGContext.cpp
@@ -10,6 +10,7 @@
 #include "Data/PCGSpatialData.h"
 #include "Data/PCGPointData.h"
 #include "Data/PCGPointArrayData.h"
+#include "Editor/IPCGEditorModule.h"
 #include "Graph/PCGGraphExecutor.h"
 #include "Helpers/PCGHelpers.h"
 #include "Metadata/PCGAttributePropertySelector.h"
@@ -470,31 +471,21 @@ bool FPCGContext::IsValueOverriden(const FName PropertyName)
 #if WITH_EDITOR
 void FPCGContext::LogVisual(ELogVerbosity::Type InVerbosity, const FText& InMessage) const
 {
-	if (!SourceComponent.IsValid())
-	{
-		return;
-	}
-
-	if (UPCGSubsystem* Subsystem = UPCGSubsystem::GetInstance(SourceComponent->GetWorld()))
+	if (IPCGEditorModule* PCGEditorModule = IPCGEditorModule::Get())
 	{
 		FPCGStack StackWithNode = Stack ? FPCGStack(*Stack) : FPCGStack();
 		StackWithNode.PushFrame(Node);
-		Subsystem->GetNodeVisualLogsMutable().Log(StackWithNode, InVerbosity, InMessage);
+		PCGEditorModule->GetNodeVisualLogsMutable().Log(StackWithNode, InVerbosity, InMessage);
 	}
 }
 
 bool FPCGContext::HasVisualLogs() const
 {
-	if (!SourceComponent.IsValid())
-	{
-		return false;
-	}
-
-	if (const UPCGSubsystem* Subsystem = UPCGSubsystem::GetInstance(SourceComponent->GetWorld()))
+	if (IPCGEditorModule* PCGEditorModule = IPCGEditorModule::Get())
 	{
 		FPCGStack StackWithNode = Stack ? FPCGStack(*Stack) : FPCGStack();
 		StackWithNode.PushFrame(Node);
-		return Subsystem->GetNodeVisualLogs().HasLogs(StackWithNode);
+		return PCGEditorModule->GetNodeVisualLogs().HasLogs(StackWithNode);
 	}
 
 	return false;
diff --git a/Engine/Plugins/PCG/Source/PCG/Private/PCGGraph.cpp b/Engine/Plugins/PCG/Source/PCG/Private/PCGGraph.cpp
--- a/Engine/Plugins/PCG/Source/PCG/Private/PCGGraph.cpp
+++ b/Engine/Plugins/PCG/Source/PCG/Private/PCGGraph.cpp
@@ -11,6 +11,7 @@
 #include "PCGPin.h"
 #include "PCGSubgraph.h"
 #include "PCGSubsystem.h"
+#include "Editor/IPCGEditorModule.h"
 #include "Elements/PCGHiGenGridSize.h"
 #include "Elements/PCGUserParameterGet.h"
 #include "Elements/ControlFlow/PCGQualityBranch.h"
@@ -742,10 +743,17 @@ void UPCGGraph::BeginDestroy()
 		InputNode->OnNodeChangedDelegate.RemoveAll(this);
 	}
 
+	EPCGChangeType ChangeTypes = EPCGChangeType::Structural | EPCGChangeType::GenerationGrid;
+
+	if (IPCGEditorModule* PCGEditorModule = IPCGEditorModule::Get())
+	{
+		PCGEditorModule->NotifyGraphChanged(this, ChangeTypes);
+	}
+
 	// Notify the compiler to remove this graph from its cache
 	if (UPCGSubsystem* PCGSubsystem = UPCGSubsystem::GetActiveEditorInstance())
 	{
-		PCGSubsystem->NotifyGraphChanged(this, EPCGChangeType::Structural | EPCGChangeType::GenerationGrid);
+		PCGSubsystem->NotifyGraphChanged(this, ChangeTypes);
 	}
 
 #endif
@@ -1607,6 +1615,11 @@ void UPCGGraph::NotifyGraphChanged(EPCGChangeType ChangeType)
 	const bool bNotifySubsystem = ((ChangeType & (EPCGChangeType::Structural | EPCGChangeType::Edge)) != EPCGChangeType::None);
 	if (bNotifySubsystem && GEditor)
 	{
+		if (IPCGEditorModule* PCGEditorModule = IPCGEditorModule::Get())
+		{
+			PCGEditorModule->NotifyGraphChanged(this, ChangeType);
+		}
+
 		if (GEditor->PlayWorld)
 		{
 			if (UPCGSubsystem* PCGPIESubsystem = UPCGSubsystem::GetInstance(GEditor->PlayWorld.Get()))
diff --git a/Engine/Plugins/PCG/Source/PCG/Private/PCGModule.cpp b/Engine/Plugins/PCG/Source/PCG/Private/PCGModule.cpp
--- a/Engine/Plugins/PCG/Source/PCG/Private/PCGModule.cpp
+++ b/Engine/Plugins/PCG/Source/PCG/Private/PCGModule.cpp
@@ -66,10 +66,14 @@ void FPCGModule::StartupModule()
 
 	// Register onto the PreExit, because we need the class to be still valid to remove them from the mapping
 	FCoreDelegates::OnPreExit.AddRaw(this, &FPCGModule::PreExit);
+
+	TickDelegateHandle = FTSTicker::GetCoreTicker().AddTicker(FTickerDelegate::CreateRaw(this, &FPCGModule::Tick));
 }
 
 void FPCGModule::ShutdownModule()
 {
+	FTSTicker::GetCoreTicker().RemoveTicker(TickDelegateHandle);
+
 	FCoreDelegates::OnPreExit.RemoveAll(this);
 
 	PCGModule::PCGModulePtr = nullptr;
@@ -90,6 +94,29 @@ void FPCGModule::PreExit()
 #endif // WITH_EDITOR
 }
 
+void FPCGModule::ExecuteNextTick(TFunction<void()> TickFunction)
+{
+	FScopeLock Lock(&ExecuteNextTickLock);
+	ExecuteNextTicks.Add(TickFunction);
+}
+
+bool FPCGModule::Tick(float DeltaTime)
+{
+	TArray<TFunction<void()>> LocalExecuteNextTicks;
+
+	{
+		FScopeLock Lock(&ExecuteNextTickLock);
+		LocalExecuteNextTicks = MoveTemp(ExecuteNextTicks);
+	}
+
+	for (TFunction<void()>& LocalExecuteNextTick : LocalExecuteNextTicks)
+	{
+		LocalExecuteNextTick();
+	}
+
+	return true;
+}
+
 #if WITH_EDITOR
 void FPCGModule::RegisterNativeElementDeterminismTests()
 {
diff --git a/Engine/Plugins/PCG/Source/PCG/Private/PCGSubsystem.cpp b/Engine/Plugins/PCG/Source/PCG/Private/PCGSubsystem.cpp
--- a/Engine/Plugins/PCG/Source/PCG/Private/PCGSubsystem.cpp
+++ b/Engine/Plugins/PCG/Source/PCG/Private/PCGSubsystem.cpp
@@ -23,6 +23,8 @@
 #include UE_INLINE_GENERATED_CPP_BY_NAME(PCGSubsystem)
 
 #if WITH_EDITOR
+#include "Editor/IPCGEditorModule.h"
+
 #include "Editor.h"
 #include "PackageSourceControlHelper.h"
 #include "ObjectTools.h"
@@ -533,56 +535,17 @@ void UPCGSubsystem::OnOriginalComponentUnregistered(UPCGComponent* InComponent)
 	}
 
 #if WITH_EDITOR
-	ClearExecutionMetadata(InComponent);
+	if (IPCGEditorModule* PCGEditorModule = IPCGEditorModule::Get())
+	{
+		PCGEditorModule->ClearExecutionMetadata(InComponent);
+	}
 #endif
 }
 
 #if WITH_EDITOR
 void UPCGSubsystem::OnScheduleGraph(const FPCGStackContext& StackContext)
 {
-	TRACE_CPUPROFILER_EVENT_SCOPE(UPCGSubsystem::OnScheduleGraph);
-	// Always clear any possibly related warnings/errors on schedule.
-	for (const FPCGStack& Stack : StackContext.GetStacks())
-	{
-		GetNodeVisualLogsMutable().ClearLogs(Stack);
-	}
-
-	UPCGComponent* RootComponent = nullptr;
-
-	// Flush out all stacks that begin from the component / top graph as the existing dynamic
-	// stacks may not be occur during the next execution.
-	if (const FPCGStack* BaseStack = StackContext.GetStack(0))
-	{
-		RootComponent = const_cast<UPCGComponent*>(BaseStack->GetRootComponent());
-
-		if (BaseStack->IsCurrentFrameInRootGraph())
-		{
-			ClearExecutionMetadata(RootComponent);
-		}
-	}
-
-	// Record executed stacks.
-	{
-		TRACE_CPUPROFILER_EVENT_SCOPE(UPCGSubsystem::OnScheduleGraph::RecordExecutedStacks);
-		FWriteScopeLock Lock(ExecutedStacksLock);
-
-		TSoftObjectPtr<UPCGComponent> RootComponentSoftPtr(RootComponent);
-
-		for (const FPCGStack& ExecutedStack : StackContext.GetStacks())
-		{
-			FPCGStackSharedPtr SharedStack(MakeShared<FPCGStack>(ExecutedStack));
-
-			bool bIsAlreadyInSet = false;
-
-			// Push to the set
-			ExecutedStacks.Add(SharedStack, &bIsAlreadyInSet);
-
-			if (!bIsAlreadyInSet && RootComponent)
-			{
-				ExecutedStacksPerComponent.FindOrAdd(RootComponentSoftPtr).Add(SharedStack);
-			}
-		}
-	}
+	// nothing to do for now
 }
 #endif
 
@@ -1754,13 +1717,6 @@ void UPCGSubsystem::DeleteSerializedPartitionActors(bool bOnlyDeleteUnused, bool
 
 void UPCGSubsystem::NotifyGraphChanged(UPCGGraph* InGraph, EPCGChangeType ChangeType)
 {
-	if (!!(ChangeType & (EPCGChangeType::Structural | EPCGChangeType::GenerationGrid)))
-	{
-		// If change was deep enough, clear out all executed stacks, and let them refresh upon next generate. Fixes
-		// cases where stale stacks never got flushed.
-		ClearExecutedStacks(InGraph);
-	}
-
 	if (GraphExecutor)
 	{
 		GraphExecutor->NotifyGraphChanged(InGraph, ChangeType);
@@ -1865,6 +1821,34 @@ uint32 UPCGSubsystem::GetGraphCacheEntryCount(IPCGElement* InElement) const
 	return GraphExecutor ? GraphExecutor->GetGraphCacheEntryCount(InElement) : 0;
 }
 
+const FPCGNodeVisualLogs& UPCGSubsystem::GetNodeVisualLogs() const
+{
+	if (IPCGEditorModule* PCGEditorModule = IPCGEditorModule::Get())
+	{
+		return PCGEditorModule->GetNodeVisualLogs();
+	}
+	else
+	{
+		// Not meant to be used but to respect the existing API return a placeholder log
+		static FPCGNodeVisualLogs Placeholder;
+		return Placeholder;
+	}
+}
+
+FPCGNodeVisualLogs& UPCGSubsystem::GetNodeVisualLogsMutable()
+{
+	if (IPCGEditorModule* PCGEditorModule = IPCGEditorModule::Get())
+	{
+		return PCGEditorModule->GetNodeVisualLogsMutable();
+	}
+	else
+	{
+		// Not meant to be used but to respect the existing API return a placeholder log
+		static FPCGNodeVisualLogs Placeholder;
+		return Placeholder;
+	}
+}
+
 void UPCGSubsystem::NotifyLandscapeEditModeExited()
 {
 	ActorAndComponentMapping.NotifyLandscapeEditModeExited();
@@ -1873,74 +1857,51 @@ void UPCGSubsystem::NotifyLandscapeEditModeExited()
 /* UE_DEPRECATED(5.6) */
 void UPCGSubsystem::ClearExecutionMetadata(const FPCGStack& BaseStack)
 {
-	if (UPCGComponent* Component = const_cast<UPCGComponent*>(BaseStack.GetRootComponent()))
+	if (IPCGEditorModule* PCGEditorModule = IPCGEditorModule::Get())
 	{
-		ClearExecutionMetadata(Component);
+		if (UPCGComponent* Component = const_cast<UPCGComponent*>(BaseStack.GetRootComponent()))
+		{
+			PCGEditorModule->ClearExecutionMetadata(Component);
+		}
 	}
 }
 
 void UPCGSubsystem::ClearExecutionMetadata(UPCGComponent* InComponent)
 {
-	if (InComponent)
+	if (IPCGEditorModule* PCGEditorModule = IPCGEditorModule::Get())
 	{
-		ClearExecutedStacks(InComponent);
-		InComponent->ClearInspectionData();
-		NodeVisualLogs.ClearLogs(InComponent);
+		PCGEditorModule->ClearExecutionMetadata(InComponent);
 	}
 }
 
 /* UE_DEPRECATED(5.6) */
 TArray<FPCGStack> UPCGSubsystem::GetExecutedStacks(const UPCGComponent* InComponent, const UPCGGraph* InSubgraph, bool bOnlyWithSubgraphAsCurrentFrame)
 {
-	TArray<FPCGStackSharedPtr> StackPtrs = GetExecutedStacksPtrs(InComponent, InSubgraph, bOnlyWithSubgraphAsCurrentFrame);
-	TArray<FPCGStack> Stacks;
-	Stacks.Reserve(StackPtrs.Num());
-	Algo::Transform(StackPtrs, Stacks, [](const FPCGStackSharedPtr& StackPtr) { return *StackPtr; });
-
-	return Stacks;
-}
-
-TArray<FPCGStackSharedPtr> UPCGSubsystem::GetExecutedStacksPtrs(const UPCGComponent* InComponent, const UPCGGraph* InSubgraph, bool bOnlyWithSubgraphAsCurrentFrame)
-{
-	TRACE_CPUPROFILER_EVENT_SCOPE(UPCGSubsystem::GetExecutedStacksPtrs);
-	TArray<FPCGStackSharedPtr> MatchingStacks;
-	FReadScopeLock Lock(ExecutedStacksLock);
-
-	TSet<FPCGStackSharedPtr>* ComponentSet = nullptr;
-	TSet<FPCGStackSharedPtr>* GraphSet = nullptr;
-
-	if (InComponent)
-	{
-		ComponentSet = ExecutedStacksPerComponent.Find(InComponent);
-	}
-	
-	if(!ComponentSet)
+	if (IPCGEditorModule* PCGEditorModule = IPCGEditorModule::Get())
 	{
-		ComponentSet = &ExecutedStacks;
-	}
-
-	auto StackMatches = [bOnlyWithSubgraphAsCurrentFrame, InSubgraph](const FPCGStackSharedPtr& Stack) -> bool
-	{
-		return ((bOnlyWithSubgraphAsCurrentFrame && Stack->GetGraphForCurrentFrame() == InSubgraph) ||
-			(!bOnlyWithSubgraphAsCurrentFrame && Stack->HasObject(InSubgraph)));
-	};
+		TArray<FPCGStackSharedPtr> StackPtrs = PCGEditorModule->GetExecutedStacksPtrs(InComponent, InSubgraph, bOnlyWithSubgraphAsCurrentFrame);
+		TArray<FPCGStack> Stacks;
+		Stacks.Reserve(StackPtrs.Num());
+		Algo::Transform(StackPtrs, Stacks, [](const FPCGStackSharedPtr& StackPtr) { return *StackPtr; });
 
-	for (const FPCGStackSharedPtr& Stack : *ComponentSet)
+		return Stacks;
+	}
+	else
 	{
-		if (StackMatches(Stack))
-		{
-			MatchingStacks.Add(Stack);
-		}
+		return {};
 	}
-
-	return MatchingStacks;
 }
 
 /* UE_DEPRECATED(5.6) */
 TArray<FPCGStack> UPCGSubsystem::GetExecutedStacks(const FPCGStack& BeginningWithStack)
 {
-	TRACE_CPUPROFILER_EVENT_SCOPE(UPCGSubsystem::GetExecutedStacks);
-	TArray<FPCGStackSharedPtr> StackPtrs = GetExecutedStacksPtrs(BeginningWithStack);
+	IPCGEditorModule* PCGEditorModule = IPCGEditorModule::Get();
+	if (!PCGEditorModule)
+	{
+		return {};
+	}
+
+	TArray<FPCGStackSharedPtr> StackPtrs = PCGEditorModule->GetExecutedStacksPtrs(BeginningWithStack);
 	TArray<FPCGStack> Stacks;
 
 	Stacks.Reserve(StackPtrs.Num());
@@ -1953,78 +1914,29 @@ TArray<FPCGStack> UPCGSubsystem::GetExecutedStacks(const FPCGStack& BeginningWit
 	return Stacks;
 }
 
-TArray<FPCGStackSharedPtr> UPCGSubsystem::GetExecutedStacksPtrs(const FPCGStack& BeginningWithStack)
-{
-	TRACE_CPUPROFILER_EVENT_SCOPE(UPCGSubsystem::GetExecutedStacksPtrsWithBeginning);
-	TArray<FPCGStackSharedPtr> MatchingStacks;
-
-	FReadScopeLock Lock(ExecutedStacksLock);
-	TSet<FPCGStackSharedPtr>* StacksToTest = ExecutedStacksPerComponent.Find(BeginningWithStack.GetRootComponent());
-
-	// Stack doesn't have a root component (unlikely), we'll check everything then
-	if (!StacksToTest)
-	{
-		StacksToTest = &ExecutedStacks;
-	}
-
-	for (const FPCGStackSharedPtr& Stack : *StacksToTest)
-	{
-		if (Stack->BeginsWith(BeginningWithStack))
-		{
-			MatchingStacks.Add(Stack);
-		}
-	}
-
-	return MatchingStacks;
-}
-
+/* UE_DEPRECATED(5.6) */
 void UPCGSubsystem::ClearExecutedStacks(const FPCGStack& BeginningWithStack)
 {
+PRAGMA_DISABLE_DEPRECATION_WARNINGS
 	return ClearExecutedStacks(BeginningWithStack.GetRootComponent());
+PRAGMA_ENABLE_DEPRECATION_WARNINGS
 }
 
+/* UE_DEPRECATED(5.6) */
 void UPCGSubsystem::ClearExecutedStacks(const UPCGComponent* InRootComponent)
 {
-	TRACE_CPUPROFILER_EVENT_SCOPE(UPCGSubsystem::ClearExecutedStacksWithComponent);
-	if (InRootComponent)
+	if (IPCGEditorModule* PCGEditorModule = IPCGEditorModule::Get())
 	{
-		FWriteScopeLock Lock(ExecutedStacksLock);
-
-		if (TSet<FPCGStackSharedPtr>* AllStackForComponent = ExecutedStacksPerComponent.Find(InRootComponent))
-		{
-			for (const FPCGStackSharedPtr& Stack : *AllStackForComponent)
-			{
-				ExecutedStacks.Remove(Stack);
-			}
-
-			ExecutedStacksPerComponent.Remove(InRootComponent);
-		}
+		PCGEditorModule->ClearExecutedStacks(InRootComponent);
 	}
 }
 
+/* UE_DEPRECATED(5.6) */
 void UPCGSubsystem::ClearExecutedStacks(const UPCGGraph* InContainingGraph)
 {
-	// TODO - it's fine to take more time here since it'll happen in rare(r) occurrences.
-	TRACE_CPUPROFILER_EVENT_SCOPE(UPCGSubsystem::ClearExecutedStacksWithGraph);
-	if (!InContainingGraph)
-	{
-		return;
-	}
-
-	FWriteScopeLock Lock(ExecutedStacksLock);
-	for(auto Iterator = ExecutedStacks.CreateIterator(); Iterator; ++Iterator)
+	if (IPCGEditorModule* PCGEditorModule = IPCGEditorModule::Get())
 	{
-		FPCGStackSharedPtr StackPtr = *Iterator;
-
-		if (StackPtr->HasObject(InContainingGraph))
-		{
-			ExecutedStacks.Remove(StackPtr);
-
-			if (TSet<FPCGStackSharedPtr>* StacksForComponent = ExecutedStacksPerComponent.Find(StackPtr->GetRootComponent()))
-			{
-				StacksForComponent->Remove(StackPtr);
-			}
-		}
+		PCGEditorModule->ClearExecutedStacks(InContainingGraph);
 	}
 }
 
diff --git a/Engine/Plugins/PCG/Source/PCG/Public/Editor/IPCGEditorModule.h b/Engine/Plugins/PCG/Source/PCG/Public/Editor/IPCGEditorModule.h
--- a/Engine/Plugins/PCG/Source/PCG/Public/Editor/IPCGEditorModule.h
+++ b/Engine/Plugins/PCG/Source/PCG/Public/Editor/IPCGEditorModule.h
@@ -3,14 +3,49 @@
 #pragma once
 
 #include "Editor/IPCGEditorProgressNotification.h"
+#include "Utils/PCGNodeVisualLogs.h"
 
 #include "Internationalization/Text.h"
 #include "Modules/ModuleInterface.h"
 
+struct FPCGStackContext;
+class UPCGComponent;
+
 class IPCGEditorModule : public IModuleInterface
 {
 public:
+	PCG_API static IPCGEditorModule* Get();
+
 	virtual TWeakPtr<IPCGEditorProgressNotification> CreateProgressNotification(const FTextFormat& TextFormat, bool bCanCancel) = 0;
 	virtual void ReleaseProgressNotification(TWeakPtr<IPCGEditorProgressNotification> InNotification) = 0;
 	virtual void SetOutlinerUIRefreshDelay(float InDelay) = 0;
+	
+#if WITH_EDITOR
+	virtual const FPCGNodeVisualLogs& GetNodeVisualLogs() const = 0;
+	virtual FPCGNodeVisualLogs& GetNodeVisualLogsMutable() = 0;
+#endif
+
+protected:
+	friend class FPCGEditor;
+	friend class FPCGGraphExecutor;
+	friend class SPCGEditorGraphDebugObjectTree;
+	friend class SPCGEditorGraphFind;
+	friend class UPCGGraph;
+	friend class UPCGSubsystem;
+
+	virtual void OnScheduleGraph(const FPCGStackContext& StackContext) = 0;
+
+	/** Clear any data collected during execution, normally called prior to generating the component. */
+	virtual void ClearExecutionMetadata(UPCGComponent* InComponent) = 0;
+
+	virtual void ClearExecutedStacks(const UPCGComponent* InRootComponent) = 0;
+	virtual void ClearExecutedStacks(const UPCGGraph* InContainingGraph) = 0;
+
+	/** Get a list of stacks that were executed during the last execution. */
+	virtual	TArray<FPCGStackSharedPtr> GetExecutedStacksPtrs(const FPCGStack& BeginningWithStack) = 0;
+	virtual TArray<FPCGStackSharedPtr> GetExecutedStacksPtrs(const UPCGComponent* InComponent, const UPCGGraph* InSubgraph, bool bOnlyWithSubgraphAsCurrentFrame = true) = 0;
+
+	virtual void NotifyGraphChanged(UPCGGraph* InGraph, EPCGChangeType ChangeType) = 0;
+
+	PCG_API static void SetEditorModule(IPCGEditorModule* InModule);
 };
\ No newline at end of file
diff --git a/Engine/Plugins/PCG/Source/PCG/Public/PCGModule.h b/Engine/Plugins/PCG/Source/PCG/Public/PCGModule.h
--- a/Engine/Plugins/PCG/Source/PCG/Public/PCGModule.h
+++ b/Engine/Plugins/PCG/Source/PCG/Public/PCGModule.h
@@ -7,6 +7,7 @@
 #include "Metadata/Accessors/PCGAttributeAccessorFactory.h"
 #include "Utils/PCGLogErrors.h"
 
+#include "Containers/Ticker.h"
 #include "Modules/ModuleInterface.h"
 
 // Logs
@@ -45,10 +46,19 @@ public:
 	
 	PCG_API static bool IsPCGModuleLoaded();
 	
+	PCG_API void ExecuteNextTick(TFunction<void()> TickFunction);
+
 private:
+	bool Tick(float DeltaTime);
+	
 	FPCGGetDataFunctionRegistry GetDataFunctionRegistry;
 	FPCGAttributeAccessorFactory AttributeAccessorFactory;
 
+	FTSTicker::FDelegateHandle	TickDelegateHandle;
+
+	FCriticalSection ExecuteNextTickLock;
+	TArray<TFunction<void()>> ExecuteNextTicks;
+
 #if WITH_EDITOR
 private:
 	void RegisterNativeElementDeterminismTests();
diff --git a/Engine/Plugins/PCG/Source/PCG/Public/PCGSubsystem.h b/Engine/Plugins/PCG/Source/PCG/Public/PCGSubsystem.h
--- a/Engine/Plugins/PCG/Source/PCG/Public/PCGSubsystem.h
+++ b/Engine/Plugins/PCG/Source/PCG/Public/PCGSubsystem.h
@@ -86,6 +86,8 @@ public:
 
 	/** Adds an action that will be executed once at the beginning of this subsystem's next Tick(). */
 	using FTickAction = TFunction<void()>;
+
+	UE_DEPRECATED(5.6, "Use FPCGModule::ExecuteNextTick instead")
 	void RegisterBeginTickAction(FTickAction&& Action);
 
 #if WITH_EDITOR
@@ -343,32 +345,34 @@ public:
 	uint32 GetGraphCacheEntryCount(IPCGElement* InElement) const;
 
 	/** Get graph warnings and errors for all nodes. */
-	const FPCGNodeVisualLogs& GetNodeVisualLogs() const { return NodeVisualLogs; }
-	FPCGNodeVisualLogs& GetNodeVisualLogsMutable() { return NodeVisualLogs; }
+	UE_DEPRECATED(5.6, "Use IPCGEditorModule::GetNodeVisualLogs instead")
+	const FPCGNodeVisualLogs& GetNodeVisualLogs() const;
+	
+	UE_DEPRECATED(5.6, "Use IPCGEditorModule::GetNodeVisualLogsMutable instead")
+	FPCGNodeVisualLogs& GetNodeVisualLogsMutable();
 
 	/** Notify that we exited the Landscape edit mode. */
 	void NotifyLandscapeEditModeExited();
 
-	/** Clear any data collected during execution, normally called prior to generating the component. */
+	UE_DEPRECATED(5.6, "No longer supported")
 	void ClearExecutionMetadata(UPCGComponent* InComponent);
 
-	UE_DEPRECATED(5.6, "Deprecated in favor of the version with the component pointer")
+	UE_DEPRECATED(5.6, "No longer supported")
 	void ClearExecutionMetadata(const FPCGStack& BaseStack);
 
-	/** Get a list of stacks that were executed during the last execution. */
-	TArray<FPCGStackSharedPtr> GetExecutedStacksPtrs(const UPCGComponent* InComponent, const UPCGGraph* InSubgraph, bool bOnlyWithSubgraphAsCurrentFrame = true);
-	TArray<FPCGStackSharedPtr> GetExecutedStacksPtrs(const FPCGStack& BeginningWithStack);
-	
+	UE_DEPRECATED(5.6, "No longer supported")
 	void ClearExecutedStacks(const UPCGComponent* InRootComponent);
+
+	UE_DEPRECATED(5.6, "No longer supported")
 	void ClearExecutedStacks(const UPCGGraph* InContainingGraph);
 
-	UE_DEPRECATED(5.6, "Deprecated in favor of GetExecutedStacksPtrs")
+	UE_DEPRECATED(5.6, "No longer supported")
 	TArray<FPCGStack> GetExecutedStacks(const UPCGComponent* InComponent, const UPCGGraph* InSubgraph, bool bOnlyWithSubgraphAsCurrentFrame = true);
 
-	UE_DEPRECATED(5.6, "Deprecated in favor of GetExecutedStacksPtrs")
+	UE_DEPRECATED(5.6, "No longer supported")
 	TArray<FPCGStack> GetExecutedStacks(const FPCGStack& BeginningWithStack);
 
-	UE_DEPRECATED(5.6, "Deprecated in favor of ClearExecutedStacks with the component pointer")
+	UE_DEPRECATED(5.6, "No longer supported")
 	void ClearExecutedStacks(const FPCGStack& BeginningWithStack);
 
 	UE_DEPRECATED(5.5, "Deprecated in favor of OnPCGComponentUnregistered, will not be notified anymore")
@@ -396,8 +400,6 @@ private:
 
 	void CreatePartitionActorsWithinBounds(UPCGComponent* InComponent, const FBox& InBounds, const PCGHiGenGrid::FSizeArray& InGridSizes);
 	void UpdateMappingPCGComponentPartitionActor(UPCGComponent* InComponent);
-
-	FPCGNodeVisualLogs NodeVisualLogs;
 #endif // WITH_EDITOR
 	
 private:
@@ -425,11 +427,6 @@ private:
 	bool bHasTickedOnce = false;
 	FPCGActorAndComponentMapping ActorAndComponentMapping;
 
-	/** A record of stacks that were executed. Used to populate debugging tool UIs. */
-	TSet<FPCGStackSharedPtr> ExecutedStacks;
-	TMap<TSoftObjectPtr<UPCGComponent>, TSet<FPCGStackSharedPtr>> ExecutedStacksPerComponent;
-	mutable FRWLock ExecutedStacksLock;
-
 	/** Functions will be executed at the beginning of the tick and then removed from this array. */
 	TArray<FTickAction> BeginTickActions;
 
diff --git a/Engine/Plugins/PCG/Source/PCGEditor/Private/Nodes/PCGEditorGraphNodeBase.cpp b/Engine/Plugins/PCG/Source/PCGEditor/Private/Nodes/PCGEditorGraphNodeBase.cpp
--- a/Engine/Plugins/PCG/Source/PCGEditor/Private/Nodes/PCGEditorGraphNodeBase.cpp
+++ b/Engine/Plugins/PCG/Source/PCGEditor/Private/Nodes/PCGEditorGraphNodeBase.cpp
@@ -10,6 +10,7 @@
 #include "PCGSettingsWithDynamicInputs.h"
 #include "PCGSubsystem.h"
 #include "PCGWorldActor.h"
+#include "Editor/IPCGEditorModule.h"
 #include "Elements/PCGHiGenGridSize.h"
 #include "Elements/PCGReroute.h"
 #include "Helpers/PCGMetadataHelpers.h"
@@ -624,8 +625,8 @@ void UPCGEditorGraphNodeBase::OnNodeChanged(UPCGNode* InNode, EPCGChangeType Cha
 
 EPCGChangeType UPCGEditorGraphNodeBase::UpdateErrorsAndWarnings()
 {
-	const UPCGSubsystem* Subsystem = UPCGSubsystem::GetActiveEditorInstance();
-	if (!PCGNode || !Subsystem)
+	IPCGEditorModule* PCGEditorModule = IPCGEditorModule::Get();
+	if (!PCGNode || !PCGEditorModule)
 	{
 		return EPCGChangeType::None;
 	}
@@ -646,13 +647,13 @@ EPCGChangeType UPCGEditorGraphNodeBase::UpdateErrorsAndWarnings()
 		// Get errors/warnings for the inspected stack.
 		FPCGStack StackWithNode = *InspectedStack;
 		StackWithNode.PushFrame(PCGNode);
-		bHasCompilerMessage = Subsystem->GetNodeVisualLogs().HasLogs(StackWithNode);
+		bHasCompilerMessage = PCGEditorModule->GetNodeVisualLogs().HasLogs(StackWithNode);
 
 		if (bHasCompilerMessage)
 		{
-			ErrorMsg = Subsystem->GetNodeVisualLogs().GetLogsSummaryText(StackWithNode).ToString();
+			ErrorMsg = PCGEditorModule->GetNodeVisualLogs().GetLogsSummaryText(StackWithNode).ToString();
 
-			const bool bHasErrors = Subsystem->GetNodeVisualLogs().HasLogsOfVerbosity(StackWithNode, ELogVerbosity::Error);
+			const bool bHasErrors = PCGEditorModule->GetNodeVisualLogs().HasLogsOfVerbosity(StackWithNode, ELogVerbosity::Error);
 			ErrorType = bHasErrors ? EMessageSeverity::Error : EMessageSeverity::Warning;
 		}
 		else
@@ -665,7 +666,7 @@ EPCGChangeType UPCGEditorGraphNodeBase::UpdateErrorsAndWarnings()
 	{
 		// Collect all errors/warnings for this node.
 		ELogVerbosity::Type MinimumVerbosity;
-		ErrorMsg = Subsystem->GetNodeVisualLogs().GetLogsSummaryText(PCGNode.Get(), MinimumVerbosity).ToString();
+		ErrorMsg = PCGEditorModule->GetNodeVisualLogs().GetLogsSummaryText(PCGNode.Get(), MinimumVerbosity).ToString();
 
 		bHasCompilerMessage = !ErrorMsg.IsEmpty();
 
diff --git a/Engine/Plugins/PCG/Source/PCGEditor/Private/PCGEditor.cpp b/Engine/Plugins/PCG/Source/PCGEditor/Private/PCGEditor.cpp
--- a/Engine/Plugins/PCG/Source/PCGEditor/Private/PCGEditor.cpp
+++ b/Engine/Plugins/PCG/Source/PCGEditor/Private/PCGEditor.cpp
@@ -9,6 +9,7 @@
 #include "PCGGraphFactory.h"
 #include "PCGPin.h"
 #include "PCGSubsystem.h"
+#include "Editor/IPCGEditorModule.h"
 #include "Elements/PCGReroute.h"
 #include "Helpers/PCGSubgraphHelpers.h"
 #include "Rendering/SlateRenderer.h"
@@ -2989,9 +2990,9 @@ void FPCGEditor::OnComponentUnregistered(UPCGComponent* Component)
 		DebugObjectTreeWidget->RequestRefresh();
 	}
 
-	if (UPCGSubsystem* Subsystem = Component ? Component->GetSubsystem() : nullptr)
+	if (IPCGEditorModule* PCGEditorModule = IPCGEditorModule::Get())
 	{
-		Subsystem->GetNodeVisualLogsMutable().ClearLogs(Component);
+		PCGEditorModule->GetNodeVisualLogsMutable().ClearLogs(Component);
 	}
 }
 
@@ -3003,9 +3004,10 @@ void FPCGEditor::OnComponentGenerationDone(UPCGSubsystem* Subsystem, UPCGCompone
 	bool bShouldRefresh = !Component || StackBeingInspected.GetRootComponent() == Component || Component->GetGraph() == PCGGraphBeingEdited;
 
 	// Additionally, if we are not inspecting but the component that's done executing contains this graph, then we should also update.
-	if (!bShouldRefresh && PCGComponentBeingInspected.IsNull() && PCGGraphBeingEdited && Subsystem)
+	IPCGEditorModule* PCGEditorModule = IPCGEditorModule::Get();
+	if (!bShouldRefresh && PCGComponentBeingInspected.IsNull() && PCGGraphBeingEdited && PCGEditorModule)
 	{
-		TArray<FPCGStackSharedPtr> ExecutedStacksContainingThisGraph = Subsystem->GetExecutedStacksPtrs(Component, PCGGraphBeingEdited);
+		TArray<FPCGStackSharedPtr> ExecutedStacksContainingThisGraph = PCGEditorModule->GetExecutedStacksPtrs(Component, PCGGraphBeingEdited);
 		bShouldRefresh |= !ExecutedStacksContainingThisGraph.IsEmpty();
 	}
 
diff --git a/Engine/Plugins/PCG/Source/PCGEditor/Private/PCGEditorModule.cpp b/Engine/Plugins/PCG/Source/PCGEditor/Private/PCGEditorModule.cpp
--- a/Engine/Plugins/PCG/Source/PCGEditor/Private/PCGEditorModule.cpp
+++ b/Engine/Plugins/PCG/Source/PCGEditor/Private/PCGEditorModule.cpp
@@ -99,10 +99,14 @@ void FPCGEditorModule::StartupModule()
 		/*MenuEntryLabel*/LOCTEXT("BuildPCG", "Build PCG"),
 		/*MenuSectionLabel*/LOCTEXT("PCG", "PCG"),
 		/*bExternalProcess*/true);
+
+	IPCGEditorModule::SetEditorModule(this);
 }
 
 void FPCGEditorModule::ShutdownModule()
 {
+	IPCGEditorModule::SetEditorModule(nullptr);
+
 	UnregisterPCGDataVisualizations();
 	UnregisterSettings();
 	UnregisterDetailsCustomizations();
@@ -733,6 +737,180 @@ void FPCGEditorModule::UnregisterPCGDataVisualizations()
 	FPCGModule::GetMutablePCGDataVisualizationRegistry().InternalRegistry.Empty();
 }
 
+void FPCGEditorModule::OnScheduleGraph(const FPCGStackContext& StackContext)
+{
+	TRACE_CPUPROFILER_EVENT_SCOPE(UPCGSubsystem::OnScheduleGraph);
+	// Always clear any possibly related warnings/errors on schedule.
+	for (const FPCGStack& Stack : StackContext.GetStacks())
+	{
+		GetNodeVisualLogsMutable().ClearLogs(Stack);
+	}
+	
+	UPCGComponent* RootComponent = nullptr;
+
+	// Flush out all stacks that begin from the component / top graph as the existing dynamic
+	// stacks may not be occur during the next execution.
+	if (const FPCGStack* BaseStack = StackContext.GetStack(0))
+	{
+		RootComponent = const_cast<UPCGComponent*>(BaseStack->GetRootComponent());
+
+		if (BaseStack->IsCurrentFrameInRootGraph())
+		{
+			ClearExecutionMetadata(RootComponent);
+		}
+	}
+
+	// Record executed stacks.
+	{
+		TRACE_CPUPROFILER_EVENT_SCOPE(UPCGSubsystem::OnScheduleGraph::RecordExecutedStacks);
+		FWriteScopeLock Lock(ExecutedStacksLock);
+
+		TSoftObjectPtr<UPCGComponent> RootComponentSoftPtr(RootComponent);
+
+		for (const FPCGStack& ExecutedStack : StackContext.GetStacks())
+		{
+			FPCGStackSharedPtr SharedStack(MakeShared<FPCGStack>(ExecutedStack));
+
+			bool bIsAlreadyInSet = false;
+
+			// Push to the set
+			ExecutedStacks.Add(SharedStack, &bIsAlreadyInSet);
+
+			if (!bIsAlreadyInSet && RootComponent)
+			{
+				ExecutedStacksPerComponent.FindOrAdd(RootComponentSoftPtr).Add(SharedStack);
+			}
+		}
+	}
+}
+
+void FPCGEditorModule::ClearExecutionMetadata(UPCGComponent* InComponent)
+{
+	if (InComponent)
+	{
+		ClearExecutedStacks(InComponent);
+		InComponent->ClearInspectionData();
+
+		GetNodeVisualLogsMutable().ClearLogs(InComponent);
+	}
+}
+
+void FPCGEditorModule::ClearExecutedStacks(const UPCGComponent* InRootComponent)
+{
+	TRACE_CPUPROFILER_EVENT_SCOPE(FPCGEditorModule::ClearExecutedStacksWithComponent);
+	if (InRootComponent)
+	{
+		FWriteScopeLock Lock(ExecutedStacksLock);
+
+		if (TSet<FPCGStackSharedPtr>* AllStackForComponent = ExecutedStacksPerComponent.Find(InRootComponent))
+		{
+			for (const FPCGStackSharedPtr& Stack : *AllStackForComponent)
+			{
+				ExecutedStacks.Remove(Stack);
+			}
+
+			ExecutedStacksPerComponent.Remove(InRootComponent);
+		}
+	}
+}
+
+void FPCGEditorModule::ClearExecutedStacks(const UPCGGraph* InContainingGraph)
+{
+	// TODO - it's fine to take more time here since it'll happen in rare(r) occurrences.
+	TRACE_CPUPROFILER_EVENT_SCOPE(FPCGEditorModule::ClearExecutedStacksWithGraph);
+	if (!InContainingGraph)
+	{
+		return;
+	}
+
+	FWriteScopeLock Lock(ExecutedStacksLock);
+	for (auto Iterator = ExecutedStacks.CreateIterator(); Iterator; ++Iterator)
+	{
+		FPCGStackSharedPtr StackPtr = *Iterator;
+
+		if (StackPtr->HasObject(InContainingGraph))
+		{
+			ExecutedStacks.Remove(StackPtr);
+
+			if (TSet<FPCGStackSharedPtr>* StacksForComponent = ExecutedStacksPerComponent.Find(StackPtr->GetRootComponent()))
+			{
+				StacksForComponent->Remove(StackPtr);
+			}
+		}
+	}
+}
+
+TArray<FPCGStackSharedPtr> FPCGEditorModule::GetExecutedStacksPtrs(const FPCGStack& BeginningWithStack)
+{
+	TRACE_CPUPROFILER_EVENT_SCOPE(FPCGEditorModule::GetExecutedStacksPtrsWithBeginning);
+	TArray<FPCGStackSharedPtr> MatchingStacks;
+
+	FReadScopeLock Lock(ExecutedStacksLock);
+	TSet<FPCGStackSharedPtr>* StacksToTest = ExecutedStacksPerComponent.Find(BeginningWithStack.GetRootComponent());
+
+	// Stack doesn't have a root component (unlikely), we'll check everything then
+	if (!StacksToTest)
+	{
+		StacksToTest = &ExecutedStacks;
+	}
+
+	for (const FPCGStackSharedPtr& Stack : *StacksToTest)
+	{
+		if (Stack->BeginsWith(BeginningWithStack))
+		{
+			MatchingStacks.Add(Stack);
+		}
+	}
+
+	return MatchingStacks;
+}
+
+TArray<FPCGStackSharedPtr> FPCGEditorModule::GetExecutedStacksPtrs(const UPCGComponent* InComponent, const UPCGGraph* InSubgraph, bool bOnlyWithSubgraphAsCurrentFrame)
+{
+	TRACE_CPUPROFILER_EVENT_SCOPE(FPCGEditorModule::GetExecutedStacksPtrs);
+	TArray<FPCGStackSharedPtr> MatchingStacks;
+	FReadScopeLock Lock(ExecutedStacksLock);
+
+	TSet<FPCGStackSharedPtr>* ComponentSet = nullptr;
+	TSet<FPCGStackSharedPtr>* GraphSet = nullptr;
+
+	if (InComponent)
+	{
+		ComponentSet = ExecutedStacksPerComponent.Find(InComponent);
+	}
+
+	if (!ComponentSet)
+	{
+		ComponentSet = &ExecutedStacks;
+	}
+
+	auto StackMatches = [bOnlyWithSubgraphAsCurrentFrame, InSubgraph](const FPCGStackSharedPtr& Stack) -> bool
+	{
+		return ((bOnlyWithSubgraphAsCurrentFrame && Stack->GetGraphForCurrentFrame() == InSubgraph) ||
+			(!bOnlyWithSubgraphAsCurrentFrame && Stack->HasObject(InSubgraph)));
+	};
+
+	for (const FPCGStackSharedPtr& Stack : *ComponentSet)
+	{
+		if (StackMatches(Stack))
+		{
+			MatchingStacks.Add(Stack);
+		}
+	}
+
+	return MatchingStacks;
+}
+
+void FPCGEditorModule::NotifyGraphChanged(UPCGGraph* InGraph, EPCGChangeType ChangeType)
+{
+	if (!!(ChangeType & (EPCGChangeType::Structural | EPCGChangeType::GenerationGrid)))
+	{
+		// If change was deep enough, clear out all executed stacks, and let them refresh upon next generate. Fixes
+		// cases where stale stacks never got flushed.
+		ClearExecutedStacks(InGraph);
+	}
+}
+
 IMPLEMENT_MODULE(FPCGEditorModule, PCGEditor);
 
 DEFINE_LOG_CATEGORY(LogPCGEditor);
diff --git a/Engine/Plugins/PCG/Source/PCGEditor/Private/Widgets/SPCGEditorGraphDebugObjectTree.cpp b/Engine/Plugins/PCG/Source/PCGEditor/Private/Widgets/SPCGEditorGraphDebugObjectTree.cpp
--- a/Engine/Plugins/PCG/Source/PCGEditor/Private/Widgets/SPCGEditorGraphDebugObjectTree.cpp
+++ b/Engine/Plugins/PCG/Source/PCGEditor/Private/Widgets/SPCGEditorGraphDebugObjectTree.cpp
@@ -5,6 +5,7 @@
 #include "PCGComponent.h"
 #include "PCGGraph.h"
 #include "PCGSubsystem.h"
+#include "Editor/IPCGEditorModule.h"
 #include "Grid/PCGPartitionActor.h"
 #include "Helpers/PCGHelpers.h"
 
@@ -43,9 +44,9 @@ namespace PCGEditorGraphDebugObjectTree
 			return;
 		}
 
-		if (const UPCGSubsystem* Subsystem = UPCGSubsystem::GetActiveEditorInstance())
+		if (IPCGEditorModule* PCGEditorModule = IPCGEditorModule::Get())
 		{
-			Subsystem->GetNodeVisualLogs().ForAllMatchingLogs(*Item->GetPCGStack(), [&Item, &OutLogs, &OutMinVerbosity, &OutMinLocalVerbosity](const FPCGStack& InStack, const FPCGPerNodeVisualLogs& InLogs)
+			PCGEditorModule->GetNodeVisualLogs().ForAllMatchingLogs(*Item->GetPCGStack(), [&Item, &OutLogs, &OutMinVerbosity, &OutMinLocalVerbosity](const FPCGStack& InStack, const FPCGPerNodeVisualLogs& InLogs)
 			{
 				// Need to verify this, and clarify this when it's a loop item.
 				// i.e. should get the graph + loop index and compare
@@ -92,9 +93,9 @@ namespace PCGEditorGraphDebugObjectTree
 			GetErrorInfoInternal(Item, MinVerbosity, MinLocalVerbosity, Logs);
 		}
 
-		if (const UPCGSubsystem* Subsystem = UPCGSubsystem::GetActiveEditorInstance())
+		if (IPCGEditorModule* PCGEditorModule = IPCGEditorModule::Get())
 		{
-			OutIconTooltipText = Subsystem->GetNodeVisualLogs().GetSummaryText(Logs, nullptr);
+			OutIconTooltipText = PCGEditorModule->GetNodeVisualLogs().GetSummaryText(Logs, nullptr);
 		}
 		
 		OutIconVisibility = OutIconTooltipText.IsEmpty() ? EVisibility::Hidden : EVisibility::Visible;
@@ -712,7 +713,7 @@ FPCGEditorGraphDebugObjectItemPtr SPCGEditorGraphDebugObjectTree::GetFirstDebugO
 		return nullptr;
 	}
 
-	UPCGSubsystem* Subsystem = PCGEditor.Pin() ? PCGEditor.Pin()->GetSubsystem() : nullptr;
+	IPCGEditorModule* PCGEditorModule = IPCGEditorModule::Get();
 
 	for (FSelectionIterator It(*SelectedActors); It; ++It)
 	{
@@ -746,9 +747,9 @@ FPCGEditorGraphDebugObjectItemPtr SPCGEditorGraphDebugObjectTree::GetFirstDebugO
 			}
 
 			// Look for graph in dynamic stacks.
-			if (Subsystem)
+			if (PCGEditorModule)
 			{
-				TArray<FPCGStackSharedPtr> ExecutedStacks = Subsystem->GetExecutedStacksPtrs(PCGComponent, PCGGraph);
+				TArray<FPCGStackSharedPtr> ExecutedStacks = PCGEditorModule->GetExecutedStacksPtrs(PCGComponent, PCGGraph);
 				for (const FPCGStackSharedPtr& Stack : ExecutedStacks)
 				{
 					if (FPCGEditorGraphDebugObjectItemPtr Item = GetItemFromStack(*Stack))
@@ -1087,9 +1088,10 @@ void SPCGEditorGraphDebugObjectTree::RefreshTree()
 
 		// Process stacks encountered during execution so far, which will include dynamic subgraphs & loop subgraphs.
 		// There will be overlaps with the static stacks but only unique entries will be added to the tree.
+		if(IPCGEditorModule* PCGEditorModule = IPCGEditorModule::Get())
 		{
 			TRACE_CPUPROFILER_EVENT_SCOPE(SPCGEditorGraphDebugObjectTree::RefreshTree::DynamicStacks);
-			TArray<FPCGStackSharedPtr> GraphStacks = Subsystem->GetExecutedStacksPtrs(nullptr, PCGGraph, /*bOnlyWithSubgraphAsCurrentFrame=*/false);
+			TArray<FPCGStackSharedPtr> GraphStacks = PCGEditorModule->GetExecutedStacksPtrs(nullptr, PCGGraph, /*bOnlyWithSubgraphAsCurrentFrame=*/false);
 			// TODO: do per component filtering?
 			AddStacksToTree(GraphStacks, ActorItems, StackToItem);
 		}
diff --git a/Engine/Plugins/PCG/Source/PCGEditor/Private/Widgets/SPCGEditorGraphFind.cpp b/Engine/Plugins/PCG/Source/PCGEditor/Private/Widgets/SPCGEditorGraphFind.cpp
--- a/Engine/Plugins/PCG/Source/PCGEditor/Private/Widgets/SPCGEditorGraphFind.cpp
+++ b/Engine/Plugins/PCG/Source/PCGEditor/Private/Widgets/SPCGEditorGraphFind.cpp
@@ -9,6 +9,7 @@
 #include "PCGNode.h"
 #include "PCGSubgraph.h"
 #include "PCGSubsystem.h"
+#include "Editor/IPCGEditorModule.h"
 #include "Elements/PCGUserParameterGet.h"
 #include "Nodes/PCGEditorGraphNodeBase.h"
 
@@ -679,10 +680,11 @@ void SPCGEditorGraphFind::MatchTokens(const TArray<FString>& InTokens)
 	VisitAllNodes(PCGEditorGraph, GraphRoot, VisitAllNodes);
 
 	// Then gather all stacks starting from the stack being inspected, and link them to their matching subgraph node.
-	if (PCGEditor->GetStackBeingInspected() && PCGEditorPtr.Pin()->GetSubsystem())
+	IPCGEditorModule* PCGEditorModule = IPCGEditorModule::Get();
+	if (PCGEditor->GetStackBeingInspected() && PCGEditorModule)
 	{
 		const FPCGStack& CurrentStack = *PCGEditor->GetStackBeingInspected();
-		TArray<FPCGStackSharedPtr> Substacks = PCGEditorPtr.Pin()->GetSubsystem()->GetExecutedStacksPtrs(CurrentStack);
+		TArray<FPCGStackSharedPtr> Substacks = PCGEditorModule->GetExecutedStacksPtrs(CurrentStack);
 		// For each stack, find the first subgraph node from the start, then get all graphs under it.
 		for (const FPCGStackSharedPtr& SubstackPtr : Substacks)
 		{
diff --git a/Engine/Plugins/PCG/Source/PCGEditor/Public/PCGEditorModule.h b/Engine/Plugins/PCG/Source/PCGEditor/Public/PCGEditorModule.h
--- a/Engine/Plugins/PCG/Source/PCGEditor/Public/PCGEditorModule.h
+++ b/Engine/Plugins/PCG/Source/PCGEditor/Public/PCGEditorModule.h
@@ -3,8 +3,10 @@
 #pragma once
 
 #include "Editor/IPCGEditorModule.h"
+#include "Utils/PCGNodeVisualLogs.h"
 
 #include "AssetTypeCategories.h"
+#include "Containers/Ticker.h"
 #include "Modules/ModuleInterface.h"
 #include "Toolkits/AssetEditorToolkit.h"
 
@@ -29,9 +31,19 @@ public:
 	virtual TWeakPtr<IPCGEditorProgressNotification> CreateProgressNotification(const FTextFormat& TextFormat, bool bCanCancel) override;
 	virtual void ReleaseProgressNotification(TWeakPtr<IPCGEditorProgressNotification> InNotification) override;
 	virtual void SetOutlinerUIRefreshDelay(float InDelay) override;
-	// ~End IPCGEditorModule implementation
+	virtual const FPCGNodeVisualLogs& GetNodeVisualLogs() const { return NodeVisualLogs; }
+	virtual FPCGNodeVisualLogs& GetNodeVisualLogsMutable() { return NodeVisualLogs; }
 
 protected:
+	virtual void OnScheduleGraph(const FPCGStackContext& StackContext);
+	virtual void ClearExecutionMetadata(UPCGComponent* InComponent);
+	virtual void ClearExecutedStacks(const UPCGComponent* InRootComponent);
+	virtual void ClearExecutedStacks(const UPCGGraph* InContainingGraph);
+	virtual	TArray<FPCGStackSharedPtr> GetExecutedStacksPtrs(const FPCGStack& BeginningWithStack);
+	virtual TArray<FPCGStackSharedPtr> GetExecutedStacksPtrs(const UPCGComponent* InComponent, const UPCGGraph* InSubgraph, bool bOnlyWithSubgraphAsCurrentFrame);
+	void NotifyGraphChanged(UPCGGraph* InGraph, EPCGChangeType ChangeType);
+	// ~End IPCGEditorModule implementation
+
 	void RegisterDetailsCustomizations();
 	void UnregisterDetailsCustomizations();
 	void RegisterMenuExtensions();
@@ -64,4 +76,11 @@ protected:
 	FDelegateHandle ShouldDisableCPUThrottlingDelegateHandle;
 
 	TSet<TSharedPtr<IPCGEditorProgressNotification>> ActiveNotifications;
+
+	FPCGNodeVisualLogs NodeVisualLogs;
+
+	/** A record of stacks that were executed. Used to populate debugging tool UIs. */
+	TSet<FPCGStackSharedPtr> ExecutedStacks;
+	TMap<TSoftObjectPtr<UPCGComponent>, TSet<FPCGStackSharedPtr>> ExecutedStacksPerComponent;
+	mutable FRWLock ExecutedStacksLock;
 };
