[PCG] Fix some issues with per-data crc caching related to the order in which we pushed results back to the output.
Also fixed an issue where the post-execute push to cache would be missing additional data if any.
Also added some robustness wrt to ordering to maximize caching chances.

diff --git a/Engine/Plugins/PCG/Source/PCG/Private/PCGElement.cpp b/Engine/Plugins/PCG/Source/PCG/Private/PCGElement.cpp
--- a/Engine/Plugins/PCG/Source/PCG/Private/PCGElement.cpp
+++ b/Engine/Plugins/PCG/Source/PCG/Private/PCGElement.cpp
@@ -343,6 +343,7 @@ void IPCGElement::PreExecutePrimaryLoopElement(FPCGContext* Context, const UPCGS
 		FPCGDataCollection SubCollection = PrimaryDataCollection;
 		SubCollection.TaggedData.Append(OtherData.TaggedData);
 
+		SubCollection.TaggedData.Sort([](const FPCGTaggedData& A, const FPCGTaggedData& B) { return A.OriginalIndex < B.OriginalIndex; });
 		SubCollection.ComputeCrcs(bShouldComputeFullOutputDataCrc);
 
 		FPCGGetFromCacheParams CacheParams = { .Node = Context->Node, .Element = this, .Component = Context->SourceComponent.Get() };
@@ -645,6 +646,9 @@ void IPCGElement::PostExecutePrimaryLoopElement(FPCGContext* Context, const UPCG
 			if (bCacheable)
 			{
 				FPCGDataCollection SubCollection = PrimaryDataCollection;
+				SubCollection.TaggedData.Append(OtherData.TaggedData);
+
+				SubCollection.TaggedData.Sort([](const FPCGTaggedData& A, const FPCGTaggedData& B) { return A.OriginalIndex < B.OriginalIndex; });
 				SubCollection.ComputeCrcs(bShouldComputeFullOutputDataCrc);
 
 				FPCGCrc DependenciesCrc;
@@ -664,9 +668,10 @@ void IPCGElement::PostExecutePrimaryLoopElement(FPCGContext* Context, const UPCG
 		// Push cached results back to the final output data, from the last to the first, at the right place.
 		for (int CachedCollectionIndex = Context->CachedInputToOutputInternalResults.Num() - 1; CachedCollectionIndex >= 0; --CachedCollectionIndex)
 		{
+			// Since we've assigned original indices up front, we should be able to respect this here.
 			const FPCGDataCollection& CachedInputData = Context->CachedInputToOutputInternalResults[CachedCollectionIndex].Key;
-			int32 CacheInputOriginalIndex = (CachedInputData.TaggedData.IsEmpty() ? INDEX_NONE : CachedInputData.TaggedData[0].OriginalIndex);
-			int32 InsertInputIndex = CacheInputOriginalIndex != INDEX_NONE ? Context->InputData.TaggedData.IndexOfByPredicate([CacheInputOriginalIndex](const FPCGTaggedData& TaggedData) { return TaggedData.OriginalIndex > CacheInputOriginalIndex; }) : INDEX_NONE;
+			const int32 CacheInputOriginalIndex = (CachedInputData.TaggedData.IsEmpty() ? INDEX_NONE : CachedInputData.TaggedData[0].OriginalIndex);
+			const int32 InsertInputIndex = CacheInputOriginalIndex != INDEX_NONE ? Context->InputData.TaggedData.IndexOfByPredicate([CacheInputOriginalIndex](const FPCGTaggedData& TaggedData) { return TaggedData.OriginalIndex > CacheInputOriginalIndex; }) : INDEX_NONE;
 
 			if (InsertInputIndex != INDEX_NONE)
 			{
@@ -679,9 +684,9 @@ void IPCGElement::PostExecutePrimaryLoopElement(FPCGContext* Context, const UPCG
 
 			const FPCGDataCollection& CachedOutputData = Context->CachedInputToOutputInternalResults[CachedCollectionIndex].Value;
 
-			// This assumes we have hinted properly on the output data; we'll look at the hinted value on the first data THEN insert before the next hint value
-			int32 CacheOutputOriginalIndex = (CachedOutputData.TaggedData.IsEmpty() ? INDEX_NONE : CachedOutputData.TaggedData[0].OriginalIndex);
-			int32 InsertOutputIndex = CacheOutputOriginalIndex != INDEX_NONE ? Context->OutputData.TaggedData.IndexOfByPredicate([CacheOutputOriginalIndex](const FPCGTaggedData& TaggedData) { return TaggedData.OriginalIndex > CacheOutputOriginalIndex; }) : INDEX_NONE;
+			// Since this data comes from the cache, we can't rely on its original index; however, it should be used as-if it was keyed on the input original index.
+			const int32 CacheOutputOriginalIndex = CacheInputOriginalIndex;
+			const int32 InsertOutputIndex = CacheOutputOriginalIndex != INDEX_NONE ? Context->OutputData.TaggedData.IndexOfByPredicate([CacheOutputOriginalIndex](const FPCGTaggedData& TaggedData) { return TaggedData.OriginalIndex > CacheOutputOriginalIndex; }) : INDEX_NONE;
 
 			if (InsertOutputIndex != INDEX_NONE)
 			{
