[PCG] Improve error message for non-normalized transform attributes to provide attribute name.

#rnx

diff --git a/Engine/Plugins/PCG/Source/PCG/Private/Compute/PCGDataForGPU.cpp b/Engine/Plugins/PCG/Source/PCG/Private/Compute/PCGDataForGPU.cpp
--- a/Engine/Plugins/PCG/Source/PCG/Private/Compute/PCGDataForGPU.cpp
+++ b/Engine/Plugins/PCG/Source/PCG/Private/Compute/PCGDataForGPU.cpp
@@ -221,8 +221,8 @@ namespace PCGDataForGPUHelpers
 			const bool bIsRotationNormalized = Transform.IsRotationNormalized();
 			if (!bIsRotationNormalized)
 			{
-				UE_LOG(LogPCG, Error, TEXT("Tried to pack transform for GPU data collection, but the transform's rotation is not normalized (%f, %f, %f, %f). Using identity instead."),
-					Transform.GetRotation().X, Transform.GetRotation().Y, Transform.GetRotation().Z, Transform.GetRotation().W);
+				UE_LOG(LogPCG, Error, TEXT("Tried to pack attribute '%s' of type Transform for GPU data collection, but the transform's rotation is not normalized (%f, %f, %f, %f). Using identity instead."),
+					*InAttributeBase->Name.ToString(), Transform.GetRotation().X, Transform.GetRotation().Y, Transform.GetRotation().Z, Transform.GetRotation().W);
 			}
 
 			// Note: ToMatrixWithScale() crashes if the transform is not normalized.
