[PCG] Add main thread requirement for World Ray Hit query node in case of landscape query

#rnx

diff --git a/Engine/Plugins/PCG/Source/PCG/Private/Elements/PCGWorldQuery.cpp b/Engine/Plugins/PCG/Source/PCG/Private/Elements/PCGWorldQuery.cpp
--- a/Engine/Plugins/PCG/Source/PCG/Private/Elements/PCGWorldQuery.cpp
+++ b/Engine/Plugins/PCG/Source/PCG/Private/Elements/PCGWorldQuery.cpp
@@ -155,6 +155,14 @@ EPCGChangeType UPCGWorldRayHitSettings::GetChangeTypeForProperty(FPropertyChange
 }
 #endif // WITH_EDITOR
 
+bool FPCGWorldRayHitQueryElement::CanExecuteOnlyOnMainThread(FPCGContext* Context) const
+{
+	const UPCGWorldRayHitSettings* Settings = Context->GetInputSettings<UPCGWorldRayHitSettings>();
+
+	// When getting landscape metadata we do a find actors that requires to be on the main thread.
+	return !Settings || Settings->QueryParams.bApplyMetadataFromLandscape;
+}
+
 bool FPCGWorldRayHitQueryElement::ExecuteInternal(FPCGContext* Context) const
 {
 	TRACE_CPUPROFILER_EVENT_SCOPE(FPCGWorldRayHitQueryElement::Execute);
diff --git a/Engine/Plugins/PCG/Source/PCG/Public/Elements/PCGWorldQuery.h b/Engine/Plugins/PCG/Source/PCG/Public/Elements/PCGWorldQuery.h
--- a/Engine/Plugins/PCG/Source/PCG/Public/Elements/PCGWorldQuery.h
+++ b/Engine/Plugins/PCG/Source/PCG/Public/Elements/PCGWorldQuery.h
@@ -84,6 +84,7 @@ class FPCGWorldRayHitQueryElement : public IPCGElement
 {
 public:
 	virtual bool IsCacheable(const UPCGSettings* InSettings) const override { return false; }
+	virtual bool CanExecuteOnlyOnMainThread(FPCGContext* Context) const override;
 protected:
 	virtual bool ExecuteInternal(FPCGContext* Context) const override;
 };
