[PCG] Loc key fix

#rnx

diff --git a/Engine/Plugins/PCG/Source/PCGEditor/Private/Nodes/PCGEditorGraphNodeGetUserParameter.cpp b/Engine/Plugins/PCG/Source/PCGEditor/Private/Nodes/PCGEditorGraphNodeGetUserParameter.cpp
--- a/Engine/Plugins/PCG/Source/PCGEditor/Private/Nodes/PCGEditorGraphNodeGetUserParameter.cpp
+++ b/Engine/Plugins/PCG/Source/PCGEditor/Private/Nodes/PCGEditorGraphNodeGetUserParameter.cpp
@@ -44,13 +44,13 @@ void UPCGEditorGraphGetUserParameter::OnRenameNode(const FString& NewName)
 					ReasonText = LOCTEXT("RenameFailureEmptyName", "Empty property name.");
 					break;
 				case EPropertyBagAlterationResult::PropertyNameInvalidCharacters:
-					ReasonText = LOCTEXT("RenameFailureEmptyName", "Invalid characters in property name.");
+					ReasonText = LOCTEXT("RenameFailureInvalidCharacters", "Invalid characters in property name.");
 					break;
 				case EPropertyBagAlterationResult::SourcePropertyNotFound:
-					ReasonText = LOCTEXT("RenameFailureEmptyName", "Source property wasn't found.");
+					ReasonText = LOCTEXT("RenameFailureSourcePropertyNotFound", "Source property wasn't found.");
 					break;
 				case EPropertyBagAlterationResult::TargetPropertyAlreadyExists:
-					ReasonText = LOCTEXT("RenameFailureEmptyName", "Target property already exists.");
+					ReasonText = LOCTEXT("RenameFailureTargetPropertyAlreadyExists", "Target property already exists.");
 					break;
 				default:
 					break;
