[PCG] Fix crash in World Ray hit query element

#rnx

diff --git a/Engine/Plugins/PCG/Source/PCG/Private/Elements/PCGWorldQuery.cpp b/Engine/Plugins/PCG/Source/PCG/Private/Elements/PCGWorldQuery.cpp
--- a/Engine/Plugins/PCG/Source/PCG/Private/Elements/PCGWorldQuery.cpp
+++ b/Engine/Plugins/PCG/Source/PCG/Private/Elements/PCGWorldQuery.cpp
@@ -157,7 +157,7 @@ EPCGChangeType UPCGWorldRayHitSettings::GetChangeTypeForProperty(FPropertyChange
 
 bool FPCGWorldRayHitQueryElement::CanExecuteOnlyOnMainThread(FPCGContext* Context) const
 {
-	const UPCGWorldRayHitSettings* Settings = Context->GetInputSettings<UPCGWorldRayHitSettings>();
+	const UPCGWorldRayHitSettings* Settings = Context ? Context->GetInputSettings<UPCGWorldRayHitSettings>() : nullptr;
 
 	// When getting landscape metadata we do a find actors that requires to be on the main thread.
 	return !Settings || Settings->QueryParams.bApplyMetadataFromLandscape;
