[PCG] Fixed missing 'disabled' banner on disabled node when enable flag is toggled on the details view.

#rnx

diff --git a/Engine/Plugins/PCG/Source/PCGEditor/Private/Nodes/PCGEditorGraphNodeBase.cpp b/Engine/Plugins/PCG/Source/PCGEditor/Private/Nodes/PCGEditorGraphNodeBase.cpp
--- a/Engine/Plugins/PCG/Source/PCGEditor/Private/Nodes/PCGEditorGraphNodeBase.cpp
+++ b/Engine/Plugins/PCG/Source/PCGEditor/Private/Nodes/PCGEditorGraphNodeBase.cpp
@@ -592,6 +592,8 @@ void UPCGEditorGraphNodeBase::OnNodeChanged(UPCGNode* InNode, EPCGChangeType Cha
 {
 	if (InNode == PCGNode)
 	{
+		bool bRequiresReconstructNode = false;
+
 		if (!!(ChangeType & EPCGChangeType::Settings))
 		{
 			if (const UPCGSettingsInterface* PCGSettingsInterface = InNode->GetSettingsInterface())
@@ -600,23 +602,15 @@ void UPCGEditorGraphNodeBase::OnNodeChanged(UPCGNode* InNode, EPCGChangeType Cha
 				if (NewEnabledState != GetDesiredEnabledState())
 				{
 					SetEnabledState(NewEnabledState);
+					bRequiresReconstructNode = true;
 				}
 			}
 		}
 
 		ChangeType |= UpdateErrorsAndWarnings();
 
-		UPCGComponent* ComponentBeingDebugged = nullptr;
-		const FPCGStack* StackBeingDebugged = nullptr;
-
-		{
-			const UPCGEditorGraph* EditorGraph = CastChecked<UPCGEditorGraph>(GetGraph());
-			const FPCGEditor* Editor = EditorGraph ? EditorGraph->GetEditor().Pin().Get() : nullptr;
-			ComponentBeingDebugged = Editor ? Editor->GetPCGComponentBeingInspected() : nullptr;
-			StackBeingDebugged = Editor ? Editor->GetStackBeingInspected() : nullptr;
-		}
-
-		if (!!(ChangeType & (EPCGChangeType::Structural | EPCGChangeType::Node | EPCGChangeType::Edge | EPCGChangeType::Cosmetic)))
+		if (!!(ChangeType & (EPCGChangeType::Structural | EPCGChangeType::Node | EPCGChangeType::Edge | EPCGChangeType::Cosmetic)) ||
+			bRequiresReconstructNode)
 		{
 			ReconstructNodeOnChange();
 		}
