[PCG] Made the landscape cache object instanced on the PCG World Actor to make it easier to find.
Also improved error message when the landscape cache is not enabled, and appropriately added validation in the prepare spatial query call.

diff --git a/Engine/Plugins/PCG/Source/PCG/Private/Data/PCGLandscapeData.cpp b/Engine/Plugins/PCG/Source/PCG/Private/Data/PCGLandscapeData.cpp
--- a/Engine/Plugins/PCG/Source/PCG/Private/Data/PCGLandscapeData.cpp
+++ b/Engine/Plugins/PCG/Source/PCG/Private/Data/PCGLandscapeData.cpp
@@ -12,6 +12,7 @@
 #include "Grid/PCGLandscapeCache.h"
 #include "Helpers/PCGHelpers.h"
 #include "Helpers/PCGWorldQueryHelpers.h"
+#include "Utils/PCGLogErrors.h"
 
 #include "ChaosInterfaceWrapperCore.h"
 #include "Landscape.h"
@@ -98,7 +99,7 @@ void UPCGLandscapeData::Initialize(const TArray<TWeakObjectPtr<ALandscapeProxy>>
 	}
 	else
 	{
-		UE_LOG(LogPCG, Error, TEXT("Landscape is unable to access the landscape cache"));
+		UE_LOG(LogPCG, Error, TEXT("Landscape data is unable to access the landscape cache (Note: Will happen if there is no PCG world actor."));
 	}
 
 	// Create secondary attributes as we need them
@@ -482,7 +483,7 @@ const UPCGBasePointData* UPCGLandscapeData::CreateBasePointData(FPCGContext* Con
 {
 	if (LandscapeInfos.IsEmpty() || !LandscapeCache)
 	{
-		UE_LOG(LogPCG, Error, TEXT("PCG Landscape cache or Landscape info are not initialized"));
+		PCGLog::LogErrorOnGraph(NSLOCTEXT("PCGLandscapeData", "RequiredDataNotInitialized", "PCG Landscape cache or Landscape info are not initialized"), Context);
 		return nullptr;
 	}
 
@@ -662,6 +663,12 @@ TArray<FPCGTaskId> UPCGLandscapeData::PrepareForSpatialQuery(FPCGContext* InCont
 		return {};
 	}
 
+	if (LandscapeCache->SerializationMode == EPCGLandscapeCacheSerializationMode::NeverSerialize && PCGHelpers::IsRuntimeOrPIE())
+	{
+		PCGLog::Landscape::LogLandscapeCacheNotAvailableError(InContext);
+		return {};
+	}
+
 	const ULandscapeInfo* LandscapeInfo = GetLandscapeInfo(InBounds.GetCenter());
 	ALandscapeProxy* LandscapeProxy = LandscapeInfo ? LandscapeInfo->GetLandscapeProxy() : nullptr;
 	if (!LandscapeProxy)
diff --git a/Engine/Plugins/PCG/Source/PCG/Private/Grid/PCGLandscapeCache.cpp b/Engine/Plugins/PCG/Source/PCG/Private/Grid/PCGLandscapeCache.cpp
--- a/Engine/Plugins/PCG/Source/PCG/Private/Grid/PCGLandscapeCache.cpp
+++ b/Engine/Plugins/PCG/Source/PCG/Private/Grid/PCGLandscapeCache.cpp
@@ -7,6 +7,7 @@
 #include "Helpers/PCGBlueprintHelpers.h"
 #include "Helpers/PCGHelpers.h"
 #include "Metadata/PCGMetadata.h"
+#include "Utils/PCGLogErrors.h"
 
 #include "LandscapeComponent.h"
 #include "LandscapeDataAccess.h"
@@ -876,7 +877,7 @@ const FPCGLandscapeCacheEntry* UPCGLandscapeCache::GetOrCreateCacheEntryInternal
 	{
 		if (!bLoggedNoCacheError)
 		{
-			UE_LOG(LogPCG, Error, TEXT("PCG Landscape cache (on the PCG World Actor) is not set to be serialized and will not work in non-editor modes."));
+			PCGLog::Landscape::LogLandscapeCacheNotAvailableError();
 			bLoggedNoCacheError = true;
 		}
 
diff --git a/Engine/Plugins/PCG/Source/PCG/Private/Utils/PCGLogErrors.cpp b/Engine/Plugins/PCG/Source/PCG/Private/Utils/PCGLogErrors.cpp
--- a/Engine/Plugins/PCG/Source/PCG/Private/Utils/PCGLogErrors.cpp
+++ b/Engine/Plugins/PCG/Source/PCG/Private/Utils/PCGLogErrors.cpp
@@ -154,4 +154,17 @@ namespace PCGLog::Settings
 	}
 }
 
+namespace PCGLog::Landscape
+{
+	namespace Format
+	{
+		const FText LandscapeCacheNotAvailable = LOCTEXT("LandscapeCacheNotAvailableInPIEOrCookedBuilds", "PCG Landscape cache (on the PCG World Actor) is not set to be serialized and will not work in non-editor modes.");
+	}
+
+	void LogLandscapeCacheNotAvailableError(const FPCGContext* InContext)
+	{
+		LogErrorOnGraph(Format::LandscapeCacheNotAvailable, InContext);
+	}
+}
+
 #undef LOCTEXT_NAMESPACE
diff --git a/Engine/Plugins/PCG/Source/PCG/Public/Grid/PCGLandscapeCache.h b/Engine/Plugins/PCG/Source/PCG/Public/Grid/PCGLandscapeCache.h
--- a/Engine/Plugins/PCG/Source/PCG/Public/Grid/PCGLandscapeCache.h
+++ b/Engine/Plugins/PCG/Source/PCG/Public/Grid/PCGLandscapeCache.h
@@ -132,10 +132,10 @@ public:
 	/** Initialize cache. Can be safely called multiple times. */
 	void Initialize();
 
-	UFUNCTION(CallInEditor, Category = "Cache")
+	UFUNCTION(CallInEditor, Category = "Landscape Cache")
 	void PrimeCache();
 
-	UFUNCTION(CallInEditor, Category = "Cache")
+	UFUNCTION(CallInEditor, Category = "Landscape Cache")
 	void ClearCache();
 
 	void TakeOwnership(UPCGLandscapeCache* InLandscapeCache);
@@ -155,14 +155,14 @@ public:
 	/** Convenience method to get metadata from the landscape for a given pair of landscape and position */
 	void SampleMetadataOnPoint(ALandscapeProxy* Landscape, FPCGPoint& OutPoint, UPCGMetadata* OutMetadata);
 
-	UPROPERTY(BlueprintReadWrite, EditAnywhere, Category = "Cache")
+	UPROPERTY(BlueprintReadWrite, EditAnywhere, Category = "Landscape Cache")
 	EPCGLandscapeCacheSerializationMode SerializationMode = EPCGLandscapeCacheSerializationMode::NeverSerialize;
 
-	UPROPERTY(BlueprintReadWrite, EditAnywhere, Category = "Cache")
+	UPROPERTY(BlueprintReadWrite, EditAnywhere, Category = "Landscape Cache")
 	EPCGLandscapeCacheSerializationContents CookedSerializedContents = EPCGLandscapeCacheSerializationContents::SerializeAll;
 
 #if WITH_EDITORONLY_DATA
-	UPROPERTY(Transient, BlueprintReadOnly, VisibleAnywhere, Category = "Cache")
+	UPROPERTY(Transient, BlueprintReadOnly, VisibleAnywhere, Category = "Landscape Cache")
 	int32 CacheEntryCount = 0;
 #endif
 
@@ -227,7 +227,7 @@ private:
 	TMap<CacheMapKey, FPCGLandscapeCacheEntry*> CachedData;
 
 	//TODO: separate by landscape
-	UPROPERTY(VisibleAnywhere, Category = "Cache")
+	UPROPERTY(VisibleAnywhere, Category = "Landscape Cache")
 	TSet<FName> CachedLayerNames;
 
 	std::atomic<int32> CacheMemorySize = 0;
diff --git a/Engine/Plugins/PCG/Source/PCG/Public/PCGWorldActor.h b/Engine/Plugins/PCG/Source/PCG/Public/PCGWorldActor.h
--- a/Engine/Plugins/PCG/Source/PCG/Public/PCGWorldActor.h
+++ b/Engine/Plugins/PCG/Source/PCG/Public/PCGWorldActor.h
@@ -62,7 +62,7 @@ public:
 	uint32 PartitionGridSize;
 
 	/** Contains all the PCG data required to query the landscape complete. Serialized in cooked builds only */
-	UPROPERTY(BlueprintReadOnly, VisibleAnywhere, Category = CachedData, meta = (NoResetToDefault, DisplayName="Landscape Cache"))
+	UPROPERTY(BlueprintReadOnly, VisibleAnywhere, Instanced, Category = CachedData, meta = (NoResetToDefault))
 	TObjectPtr<UPCGLandscapeCache> LandscapeCacheObject = nullptr;
 
 #if WITH_EDITORONLY_DATA
diff --git a/Engine/Plugins/PCG/Source/PCG/Public/Utils/PCGLogErrors.h b/Engine/Plugins/PCG/Source/PCG/Public/Utils/PCGLogErrors.h
--- a/Engine/Plugins/PCG/Source/PCG/Public/Utils/PCGLogErrors.h
+++ b/Engine/Plugins/PCG/Source/PCG/Public/Utils/PCGLogErrors.h
@@ -139,4 +139,9 @@ namespace PCGLog
 		PCG_API void LogInvalidPreconfigurationWarning(int32 PreconfigurationIndex, const FText& NodeTitle);
 		PCG_API void LogInvalidConversionError(int32 PreconfigurationIndex, const FText& NodeTitle, const FText& Reason);
 	}
+
+	namespace Landscape
+	{
+		PCG_API void LogLandscapeCacheNotAvailableError(const FPCGContext* InContext = nullptr);
+	}
 }
