Removed per-data caching on filter by index, since by definition it can't work.

#rnx

diff --git a/Engine/Plugins/PCG/Source/PCG/Public/Elements/PCGFilterByIndex.h b/Engine/Plugins/PCG/Source/PCG/Public/Elements/PCGFilterByIndex.h
--- a/Engine/Plugins/PCG/Source/PCG/Public/Elements/PCGFilterByIndex.h
+++ b/Engine/Plugins/PCG/Source/PCG/Public/Elements/PCGFilterByIndex.h
@@ -45,6 +45,5 @@ class FPCGFilterByIndexElement : public IPCGElement
 {
 protected:
 	virtual bool ExecuteInternal(FPCGContext* Context) const override;
-	virtual EPCGElementExecutionLoopMode ExecutionLoopMode(const UPCGSettings* Settings) const override { return EPCGElementExecutionLoopMode::SinglePrimaryPin; }
 	virtual bool SupportsBasePointDataInputs(FPCGContext* InContext) const override { return true; }
 };
