[PCG] Add behavior to open or focus Details panel when clicking Graph Settings button in the graph

#rnx

diff --git a/Engine/Plugins/PCG/Source/PCGEditor/Private/PCGEditor.cpp b/Engine/Plugins/PCG/Source/PCGEditor/Private/PCGEditor.cpp
--- a/Engine/Plugins/PCG/Source/PCGEditor/Private/PCGEditor.cpp
+++ b/Engine/Plugins/PCG/Source/PCGEditor/Private/PCGEditor.cpp
@@ -1241,22 +1241,36 @@ void FPCGEditor::OnDeterminismGraphTest()
 	}
 }
 
-void FPCGEditor::OnEditGraphSettings() const
+void FPCGEditor::OnEditGraphSettings()
 {
 	for (TSharedPtr<SPCGEditorGraphDetailsView> PropertyDetailsWidget : PropertyDetailsWidgets)
 	{
 		PropertyDetailsWidget->SetObject(PCGGraphBeingEdited);
 	}
+
+	OpenDetailsView();
 }
 
 bool FPCGEditor::IsEditGraphSettingsToggled() const
 {
-	for (TSharedPtr<SPCGEditorGraphDetailsView> PropertyDetailsWidget : PropertyDetailsWidgets)
+	if (!TabManager.IsValid())
+	{
+		return false;
+	}
+
+	for (int32 WidgetIndex = 0; WidgetIndex < PropertyDetailsWidgets.Num(); ++WidgetIndex)
 	{
-		const TArray<TWeakObjectPtr<UObject>>& SelectedObjects = PropertyDetailsWidget->GetSelectedObjects();
+		const TArray<TWeakObjectPtr<UObject>>& SelectedObjects = PropertyDetailsWidgets[WidgetIndex]->GetSelectedObjects();
+		// The only object selected should be the graph. If there is no details view panel open, leave it disabled.
 		if (SelectedObjects.Num() == 1 && SelectedObjects[0] == PCGGraphBeingEdited.Get())
 		{
-			return true;
+			if (const TSharedPtr<SDockTab>& Tab = TabManager->FindExistingLiveTab(FPCGEditor_private::PropertyDetailsID[WidgetIndex]))
+			{
+				if (Tab.IsValid() && Tab->IsForeground())
+				{
+					return true;
+				}
+			}
 		}
 	}
 
diff --git a/Engine/Plugins/PCG/Source/PCGEditor/Private/PCGEditor.h b/Engine/Plugins/PCG/Source/PCGEditor/Private/PCGEditor.h
--- a/Engine/Plugins/PCG/Source/PCGEditor/Private/PCGEditor.h
+++ b/Engine/Plugins/PCG/Source/PCGEditor/Private/PCGEditor.h
@@ -229,7 +229,7 @@ private:
 	void OnDeterminismGraphTest();
 
 	/** Open details view for the PCG object being edited */
-	void OnEditGraphSettings() const;
+	void OnEditGraphSettings();
 	/** Whether the PCG object being edited is opened in details view or not */
 	bool IsEditGraphSettingsToggled() const;
 
