[PCG] Add support for string matching filtering when initializing metadata by copy

* Was doing it for normal initialize but not for the copy

#rnx

diff --git a/Engine/Plugins/PCG/Source/PCG/Private/Metadata/PCGMetadata.cpp b/Engine/Plugins/PCG/Source/PCG/Private/Metadata/PCGMetadata.cpp
--- a/Engine/Plugins/PCG/Source/PCG/Private/Metadata/PCGMetadata.cpp
+++ b/Engine/Plugins/PCG/Source/PCG/Private/Metadata/PCGMetadata.cpp
@@ -394,12 +394,12 @@ void UPCGMetadata::InitializeAsCopy(const UPCGMetadata* InMetadataToCopy, const
 	InitializeAsCopyWithAttributeFilter(InMetadataToCopy, TSet<FName>(), EPCGMetadataFilterMode::ExcludeAttributes, EntriesToCopy);
 }
 
-void UPCGMetadata::InitializeAsCopyWithAttributeFilter(const UPCGMetadata* InMetadataToCopy, const TSet<FName>& InFilteredAttributes, EPCGMetadataFilterMode InFilterMode, const TArray<PCGMetadataEntryKey>* EntriesToCopy)
+void UPCGMetadata::InitializeAsCopyWithAttributeFilter(const UPCGMetadata* InMetadataToCopy, const TSet<FName>& InFilteredAttributes, EPCGMetadataFilterMode InFilterMode, const TArray<PCGMetadataEntryKey>* EntriesToCopy, EPCGStringMatchingOperator InMatchOperator)
 {
-	InitializeAsCopyWithAttributeFilter(InMetadataToCopy, FPCGAttributeIdentifier::TransformNameSet(InFilteredAttributes), InFilterMode, EntriesToCopy);
+	InitializeAsCopyWithAttributeFilter(InMetadataToCopy, FPCGAttributeIdentifier::TransformNameSet(InFilteredAttributes), InFilterMode, EntriesToCopy, InMatchOperator);
 }
 
-void UPCGMetadata::InitializeAsCopyWithAttributeFilter(const UPCGMetadata* InMetadataToCopy, const TSet<FPCGAttributeIdentifier>& InFilteredAttributes, EPCGMetadataFilterMode InFilterMode, const TArray<PCGMetadataEntryKey>* EntriesToCopy)
+void UPCGMetadata::InitializeAsCopyWithAttributeFilter(const UPCGMetadata* InMetadataToCopy, const TSet<FPCGAttributeIdentifier>& InFilteredAttributes, EPCGMetadataFilterMode InFilterMode, const TArray<PCGMetadataEntryKey>* EntriesToCopy, EPCGStringMatchingOperator InMatchOperator)
 {
 	if (!InMetadataToCopy)
 	{
@@ -421,11 +421,11 @@ void UPCGMetadata::InitializeAsCopyWithAttributeFilter(const UPCGMetadata* InMet
 	SetupDomainsFromOtherMetadataIfNeeded(InMetadataToCopy);
 		
 	TMap<FPCGMetadataDomainID, TSet<FName>> FilteredAttributeSplit = FPCGAttributeIdentifier::SplitByDomain(InFilteredAttributes);
-	InMetadataToCopy->ForEachValidUniqueConstMetadataDomain([this, &FilteredAttributeSplit, InFilterMode, &EntriesToCopy, &bShouldParent, InMetadataToCopy](const FPCGMetadataDomainID& DomainID, const FPCGMetadataDomain& OtherMetadataDomain)
+	InMetadataToCopy->ForEachValidUniqueConstMetadataDomain([this, &FilteredAttributeSplit, InFilterMode, &EntriesToCopy, &bShouldParent, InMetadataToCopy, InMatchOperator](const FPCGMetadataDomainID& DomainID, const FPCGMetadataDomain& OtherMetadataDomain)
 	{
 		TSet<FName>* FilteredAttributesPtr = FilteredAttributeSplit.Find(DomainID);
 		FPCGMetadataDomain* CurrentMetadataDomain = FindOrCreateMetadataDomain(DomainID);
-		CurrentMetadataDomain->InitializeAsCopyWithAttributeFilter(&OtherMetadataDomain, (FilteredAttributesPtr ? *FilteredAttributesPtr : TSet<FName>{}), InFilterMode, EntriesToCopy);
+		CurrentMetadataDomain->InitializeAsCopyWithAttributeFilter(&OtherMetadataDomain, (FilteredAttributesPtr ? *FilteredAttributesPtr : TSet<FName>{}), InFilterMode, EntriesToCopy, InMatchOperator);
 
 		// We still validate that the parent was set correctly and it matches the metadata to copy.
 		if (CurrentMetadataDomain->Parent && CurrentMetadataDomain->Parent->TopMetadata == InMetadataToCopy)
diff --git a/Engine/Plugins/PCG/Source/PCG/Private/Metadata/PCGMetadataDomain.cpp b/Engine/Plugins/PCG/Source/PCG/Private/Metadata/PCGMetadataDomain.cpp
--- a/Engine/Plugins/PCG/Source/PCG/Private/Metadata/PCGMetadataDomain.cpp
+++ b/Engine/Plugins/PCG/Source/PCG/Private/Metadata/PCGMetadataDomain.cpp
@@ -90,6 +90,57 @@ namespace PCGMetadata
 
 		return PCGPropertyHelpers::GetPropertyValueWithCallback(DataPtr, InProperty, CreateAttributeAndSet);
 	}
+
+	/** Utility structure, should not be kept around as we keep a ref on `InNames` */
+	struct FPCGMetadataAttributeNameFilter
+	{
+		FPCGMetadataAttributeNameFilter(const TSet<FName>& InNames, EPCGMetadataFilterMode InFilterMode, EPCGStringMatchingOperator InMatchOperator)
+			: Names(InNames)
+			, FilterMode(InFilterMode)
+			, MatchOperator(InMatchOperator)
+		{
+			if (InMatchOperator != EPCGStringMatchingOperator::Equal)
+			{
+				Algo::Transform(InNames, NameStrings, [](const FName& InFilteredAttribute) {return InFilteredAttribute.ToString(); });
+			}
+		}
+
+		/** Returns true if InName should be excluded. */
+		bool operator()(const FName InName)
+		{
+			bool Result = false;
+
+			if (MatchOperator == EPCGStringMatchingOperator::Equal)
+			{
+				Result = Names.Contains(InName);
+			}
+			else
+			{
+				const FString OtherAttributeString = InName.ToString();
+				if (MatchOperator == EPCGStringMatchingOperator::Substring)
+				{
+					Result = Algo::AnyOf(NameStrings, [&OtherAttributeString](const FString& InAttribute) { return OtherAttributeString.Contains(InAttribute); });
+				}
+				else if (MatchOperator == EPCGStringMatchingOperator::Matches)
+				{
+					Result = Algo::AnyOf(NameStrings, [&OtherAttributeString](const FString& InAttribute) { return OtherAttributeString.MatchesWildcard(InAttribute); });
+				}
+				else
+				{
+					checkNoEntry();
+					return false;
+				}
+			}
+
+			return FilterMode == EPCGMetadataFilterMode::ExcludeAttributes ? Result : !Result;
+		}
+
+	private:
+		const TSet<FName>& Names;
+		TArray<FString> NameStrings;
+		EPCGMetadataFilterMode FilterMode;
+		EPCGStringMatchingOperator MatchOperator;
+	};
 }
 
 
@@ -279,7 +330,7 @@ void FPCGMetadataDomain::InitializeAsCopy(const FPCGMetadataDomain* InMetadataTo
 	InitializeAsCopyWithAttributeFilter(InMetadataToCopy, TSet<FName>(), EPCGMetadataFilterMode::ExcludeAttributes, EntriesToCopy);
 }
 
-void FPCGMetadataDomain::InitializeAsCopyWithAttributeFilter(const FPCGMetadataDomain* InMetadataToCopy, const TSet<FName>& InFilteredAttributes, EPCGMetadataFilterMode InFilterMode, const TArray<PCGMetadataEntryKey>* EntriesToCopy)
+void FPCGMetadataDomain::InitializeAsCopyWithAttributeFilter(const FPCGMetadataDomain* InMetadataToCopy, const TSet<FName>& InFilteredAttributes, EPCGMetadataFilterMode InFilterMode, const TArray<PCGMetadataEntryKey>* EntriesToCopy, EPCGStringMatchingOperator InMatchOperator)
 {
 	if (!InMetadataToCopy)
 	{
@@ -293,7 +344,7 @@ void FPCGMetadataDomain::InitializeAsCopyWithAttributeFilter(const FPCGMetadataD
 		return;
 	}
 
-	const bool bSkipAttributesInFilterList = (InFilterMode == EPCGMetadataFilterMode::ExcludeAttributes);
+	PCGMetadata::FPCGMetadataAttributeNameFilter ShouldSkipAttribute(InFilteredAttributes, InFilterMode, InMatchOperator);
 
 	// If we have a partial copy, it will flatten the metadata, so we don't need a parent.
 	// Otherwise, we keep the parent hierarchy.
@@ -324,10 +375,7 @@ void FPCGMetadataDomain::InitializeAsCopyWithAttributeFilter(const FPCGMetadataD
 	// Copy attributes
 	for (const TPair<FName, FPCGMetadataAttributeBase*>& OtherAttribute : InMetadataToCopy->Attributes)
 	{
-		const bool bAttributeInFilterList = InFilteredAttributes.Contains(OtherAttribute.Key);
-		const bool bSkipThisAttribute = (bSkipAttributesInFilterList == bAttributeInFilterList);
-
-		if (!bSkipThisAttribute)
+		if (!ShouldSkipAttribute(OtherAttribute.Key))
 		{
 			// Don't copy entries if we have a partial copy, we will set them all after.
 			FPCGMetadataAttributeBase* Attribute = CopyAttribute(OtherAttribute.Value, OtherAttribute.Key, /*bKeepParent=*/false, /*bCopyEntries=*/!bPartialCopy, /*bCopyValues=*/true);
@@ -348,47 +396,13 @@ bool FPCGMetadataDomain::AddAttributesFiltered(const FPCGMetadataDomain* InOther
 		return false;
 	}
 
-	TArray<FString> InFilteredAttributesStrings;
-	if (InMatchOperator != EPCGStringMatchingOperator::Equal)
-	{
-		Algo::Transform(InFilteredAttributes, InFilteredAttributesStrings, [](const FName& InFilteredAttribute) {return InFilteredAttribute.ToString(); });
-	}
-
-	auto IsAttributeInFilterList = [&InFilteredAttributes, &InFilteredAttributesStrings, InMatchOperator](FName OtherAttributeName) -> bool
-	{
-		if (InMatchOperator == EPCGStringMatchingOperator::Equal)
-		{
-			return InFilteredAttributes.Contains(OtherAttributeName);
-		}
-		else
-		{
-			const FString OtherAttributeString = OtherAttributeName.ToString();
-			if (InMatchOperator == EPCGStringMatchingOperator::Substring)
-			{
-				return Algo::AnyOf(InFilteredAttributesStrings, [&OtherAttributeString](const FString& InAttribute) { return OtherAttributeString.Contains(InAttribute); });
-			}
-			else if (InMatchOperator == EPCGStringMatchingOperator::Matches)
-			{
-				return Algo::AnyOf(InFilteredAttributesStrings, [&OtherAttributeString](const FString& InAttribute) { return OtherAttributeString.MatchesWildcard(InAttribute); });
-			}
-			else
-			{
-				checkNoEntry();
-				return false;
-			}
-		}
-	};
+	PCGMetadata::FPCGMetadataAttributeNameFilter ShouldSkipAttribute(InFilteredAttributes, InFilterMode, InMatchOperator);
 
 	bool bAttributeAdded = false;
 
 	for (const TPair<FName, FPCGMetadataAttributeBase*>& OtherAttribute : InOther->Attributes)
 	{
-		// Skip this attribute if it is in an exclude list, or if it is not in an include list
-		const bool bAttributeInFilterList = IsAttributeInFilterList(OtherAttribute.Key);
-		const bool bSkipAttributesInFilterList = InFilterMode == EPCGMetadataFilterMode::ExcludeAttributes;
-		const bool bSkipThisAttribute = bSkipAttributesInFilterList == bAttributeInFilterList;
-
-		if(bSkipThisAttribute || !OtherAttribute.Value)
+		if (ShouldSkipAttribute(OtherAttribute.Key) || !OtherAttribute.Value)
 		{
 			continue;
 		}
@@ -400,7 +414,7 @@ bool FPCGMetadataDomain::AddAttributesFiltered(const FPCGMetadataDomain* InOther
 			const FPCGMetadataAttributeBase* Attribute = GetConstAttribute(OtherAttribute.Key);
 			check(Attribute);
 
-			if(Attribute->GetTypeId() != OtherAttribute.Value->GetTypeId() && 
+			if (Attribute->GetTypeId() != OtherAttribute.Value->GetTypeId() &&
 				!PCG::Private::IsBroadcastable(OtherAttribute.Value->GetTypeId(), Attribute->GetTypeId()) &&
 				PCG::Private::IsBroadcastable(Attribute->GetTypeId(), OtherAttribute.Value->GetTypeId()))
 			{
diff --git a/Engine/Plugins/PCG/Source/PCG/Public/Metadata/PCGMetadata.h b/Engine/Plugins/PCG/Source/PCG/Public/Metadata/PCGMetadata.h
--- a/Engine/Plugins/PCG/Source/PCG/Public/Metadata/PCGMetadata.h
+++ b/Engine/Plugins/PCG/Source/PCG/Public/Metadata/PCGMetadata.h
@@ -47,9 +47,10 @@ public:
 
 	/**
 	 * Initializes the metadata from a parent metadata. Copies attributes and values.
-	 * @param InParent The parent metadata to use as a template, if any (can be null).
-	 * @param InFilteredAttributes Optional list of attributes to exclude or include when adding the attributes from the parent.
-	 * @param InFilterMode Defines attribute filter operation.
+	 * @param InParent              The parent metadata to use as a template, if any (can be null).
+	 * @param InFilteredAttributes  Optional list of attributes to exclude or include when adding the attributes from the parent.
+	 * @param InFilterMode          Defines attribute filter operation.
+	 * @param InMatchOperator       How to match the names for the filtering
 	 */
 	UFUNCTION(BlueprintCallable, Category = "PCG|Metadata")
 	void InitializeWithAttributeFilter(const UPCGMetadata* InParent, const TSet<FName>& InFilteredAttributes, EPCGMetadataFilterMode InFilterMode = EPCGMetadataFilterMode::ExcludeAttributes, EPCGStringMatchingOperator InMatchOperator = EPCGStringMatchingOperator::Equal);
@@ -57,19 +58,20 @@ public:
 	void InitializeWithAttributeFilter(const UPCGMetadata* InParent, const TSet<FPCGAttributeIdentifier>& InFilteredAttributes, EPCGMetadataFilterMode InFilterMode = EPCGMetadataFilterMode::ExcludeAttributes, EPCGStringMatchingOperator InMatchOperator = EPCGStringMatchingOperator::Equal, const TMap<FPCGMetadataDomainID, FPCGMetadataDomainID>& DomainMapping = {});
 
 	/** Initializes the metadata from a parent metadata by copying all attributes to it.
-	* @param InMetadataToCopy Metadata to copy from
+	* @param InMetadataToCopy        Metadata to copy from
 	* @param InOptionalEntriesToCopy Optional array that contains the keys to copy over. This array order will be respected, so it can also be used to re-order entries. Can be null.
 	*/
 	void InitializeAsCopy(const UPCGMetadata* InMetadataToCopy, const TArray<PCGMetadataEntryKey>* InOptionalEntriesToCopy = nullptr);
 
 	/** Initializes the metadata from a parent metadata by copy filtered attributes only to it
-	* @param InMetadataToCopy Metadata to copy from
-	* @param InFilteredAttributes Attributes to keep/exclude, can be empty.
-	* @param InFilterMode Filter to know if we should keep or exclude InFilteredAttributes.
-	* @param InOptionalEntriesToCopy Optional array that contains the keys to copy over. This array order will be respected, so it can also be used to re-order entries. Can be null.
+	* @param InMetadataToCopy         Metadata to copy from
+	* @param InFilteredAttributes     Attributes to keep/exclude, can be empty.
+	* @param InFilterMode             Filter to know if we should keep or exclude InFilteredAttributes.
+	* @param InOptionalEntriesToCopy  Optional array that contains the keys to copy over. This array order will be respected, so it can also be used to re-order entries. Can be null.
+	* @param InMatchOperator          How to match the names for the filtering
 	*/
-	void InitializeAsCopyWithAttributeFilter(const UPCGMetadata* InMetadataToCopy, const TSet<FName>& InFilteredAttributes, EPCGMetadataFilterMode InFilterMode = EPCGMetadataFilterMode::ExcludeAttributes, const TArray<PCGMetadataEntryKey>* InOptionalEntriesToCopy = nullptr);
-	void InitializeAsCopyWithAttributeFilter(const UPCGMetadata* InMetadataToCopy, const TSet<FPCGAttributeIdentifier>& InFilteredAttributes, EPCGMetadataFilterMode InFilterMode = EPCGMetadataFilterMode::ExcludeAttributes, const TArray<PCGMetadataEntryKey>* InOptionalEntriesToCopy = nullptr);
+	void InitializeAsCopyWithAttributeFilter(const UPCGMetadata* InMetadataToCopy, const TSet<FName>& InFilteredAttributes, EPCGMetadataFilterMode InFilterMode = EPCGMetadataFilterMode::ExcludeAttributes, const TArray<PCGMetadataEntryKey>* InOptionalEntriesToCopy = nullptr, EPCGStringMatchingOperator InMatchOperator = EPCGStringMatchingOperator::Equal);
+	void InitializeAsCopyWithAttributeFilter(const UPCGMetadata* InMetadataToCopy, const TSet<FPCGAttributeIdentifier>& InFilteredAttributes, EPCGMetadataFilterMode InFilterMode = EPCGMetadataFilterMode::ExcludeAttributes, const TArray<PCGMetadataEntryKey>* InOptionalEntriesToCopy = nullptr, EPCGStringMatchingOperator InMatchOperator = EPCGStringMatchingOperator::Equal);
 
 	/** Creates missing attributes from another metadata if they are not currently present - note that this does not copy values */
 	UFUNCTION(BlueprintCallable, Category = "PCG|Metadata")
@@ -77,9 +79,10 @@ public:
 
 	/**
 	 * Creates missing attributes from another metadata if they are not currently present - note that this does not copy values.
-	 * @param InOther The other metadata to obtain a list of attributes from.
-	 * @param InFilteredAttributes Optional list of attributes to exclude or include when adding the attributes.
-	 * @param InFilterMode Defines attribute filter operation.
+	 * @param InOther               The other metadata to obtain a list of attributes from.
+	 * @param InFilteredAttributes  Optional list of attributes to exclude or include when adding the attributes.
+	 * @param InFilterMode          Defines attribute filter operation.
+	 * @param InMatchOperator       How to match the names for the filtering
 	 */
 	UFUNCTION(BlueprintCallable, Category = "PCG|Metadata")
 	void AddAttributesFiltered(const UPCGMetadata* InOther, const TSet<FName>& InFilteredAttributes, EPCGMetadataFilterMode InFilterMode = EPCGMetadataFilterMode::ExcludeAttributes, EPCGStringMatchingOperator InMatchOperator = EPCGStringMatchingOperator::Equal);
@@ -377,15 +380,16 @@ public:
 	void K2_InitializeAsCopy(const UPCGMetadata* InMetadataToCopy, const TArray<int64>& InOptionalEntriesToCopy) { return InitializeAsCopy(InMetadataToCopy, !InOptionalEntriesToCopy.IsEmpty() ? &InOptionalEntriesToCopy : nullptr); }
 
 	/** Initializes the metadata from a parent metadata by copy filtered attributes only to it
-	* @param InMetadataToCopy Metadata to copy from
-	* @param InFilteredAttributes Attributes to keep/exclude, can be empty.
-	* @param InOptionalEntriesToCopy Optional array that contains the keys to copy over. This array order will be respected, so it can also be used to re-order entries. If empty, copy them all.
-	* @param InFilterMode Filter to know if we should keep or exclude InFilteredAttributes.
+	* @param InMetadataToCopy         Metadata to copy from
+	* @param InFilteredAttributes     Attributes to keep/exclude, can be empty.
+	* @param InOptionalEntriesToCopy  Optional array that contains the keys to copy over. This array order will be respected, so it can also be used to re-order entries. If empty, copy them all.
+	* @param InFilterMode             Filter to know if we should keep or exclude InFilteredAttributes.
+	* @param InMatchOperator          How to match the names for the filtering
 	*/
 	UFUNCTION(BlueprintCallable, Category = "PCG|Metadata", meta = (DisplayName = "Initialize As Copy With Attribute Filter", AutoCreateRefTerm = "InFilteredAttributes,InOptionalEntriesToCopy"))
-	void K2_InitializeAsCopyWithAttributeFilter(const UPCGMetadata* InMetadataToCopy, const TSet<FName>& InFilteredAttributes, const TArray<int64>& InOptionalEntriesToCopy, EPCGMetadataFilterMode InFilterMode = EPCGMetadataFilterMode::ExcludeAttributes)
+	void K2_InitializeAsCopyWithAttributeFilter(const UPCGMetadata* InMetadataToCopy, const TSet<FName>& InFilteredAttributes, const TArray<int64>& InOptionalEntriesToCopy, EPCGMetadataFilterMode InFilterMode = EPCGMetadataFilterMode::ExcludeAttributes, EPCGStringMatchingOperator InMatchOperator = EPCGStringMatchingOperator::Equal)
 	{
-		return InitializeAsCopyWithAttributeFilter(InMetadataToCopy, FPCGAttributeIdentifier::TransformNameSet(InFilteredAttributes), InFilterMode, !InOptionalEntriesToCopy.IsEmpty() ? &InOptionalEntriesToCopy : nullptr);
+		return InitializeAsCopyWithAttributeFilter(InMetadataToCopy, FPCGAttributeIdentifier::TransformNameSet(InFilteredAttributes), InFilterMode, !InOptionalEntriesToCopy.IsEmpty() ? &InOptionalEntriesToCopy : nullptr, InMatchOperator);
 	}
 
 	/** Computes Crc from all attributes & keys from outer's data. */
diff --git a/Engine/Plugins/PCG/Source/PCG/Public/Metadata/PCGMetadataDomain.h b/Engine/Plugins/PCG/Source/PCG/Public/Metadata/PCGMetadataDomain.h
--- a/Engine/Plugins/PCG/Source/PCG/Public/Metadata/PCGMetadataDomain.h
+++ b/Engine/Plugins/PCG/Source/PCG/Public/Metadata/PCGMetadataDomain.h
@@ -44,7 +44,7 @@ public:
 	* @param InFilterMode Filter to know if we should keep or exclude InFilteredAttributes.
 	* @param InOptionalEntriesToCopy Optional array that contains the keys to copy over. This array order will be respected, so it can also be used to re-order entries. Can be null.
 	*/
-	void InitializeAsCopyWithAttributeFilter(const FPCGMetadataDomain* InMetadataToCopy, const TSet<FName>& InFilteredAttributes, EPCGMetadataFilterMode InFilterMode = EPCGMetadataFilterMode::ExcludeAttributes, const TArray<PCGMetadataEntryKey>* InOptionalEntriesToCopy = nullptr);
+	void InitializeAsCopyWithAttributeFilter(const FPCGMetadataDomain* InMetadataToCopy, const TSet<FName>& InFilteredAttributes, EPCGMetadataFilterMode InFilterMode = EPCGMetadataFilterMode::ExcludeAttributes, const TArray<PCGMetadataEntryKey>* InOptionalEntriesToCopy = nullptr, EPCGStringMatchingOperator InMatchOperator = EPCGStringMatchingOperator::Equal);
 
 	/** Creates missing attributes from another metadata if they are not currently present - note that this does not copy values */
 	bool AddAttributes(const FPCGMetadataDomain* InOther) { return AddAttributesFiltered(InOther, TSet<FName>()); }
