[PCG] Harden runtime gen scheduler code against invalid OriginalComponent pointers after a related crash was reported.

#rnx

diff --git a/Engine/Plugins/PCG/Source/PCG/Private/RuntimeGen/PCGRuntimeGenScheduler.cpp b/Engine/Plugins/PCG/Source/PCG/Private/RuntimeGen/PCGRuntimeGenScheduler.cpp
--- a/Engine/Plugins/PCG/Source/PCG/Private/RuntimeGen/PCGRuntimeGenScheduler.cpp
+++ b/Engine/Plugins/PCG/Source/PCG/Private/RuntimeGen/PCGRuntimeGenScheduler.cpp
@@ -472,10 +472,6 @@ void FPCGRuntimeGenScheduler::TickCleanup(const TSet<IPCGGenSourceBase*>& InGenS
 
 	check(ActorAndComponentMapping && InPCGWorldActor);
 
-	// Generated Entry Key, Local/Generated Component
-	using PCGComponentToClean = TTuple<FGridGenerationKey, UPCGComponent*>;
-	TArray<PCGComponentToClean> ComponentsToClean;
-
 	auto CheckIfAllGenSourcesWantToCleanup = [&InGenSources](const UPCGSchedulingPolicyBase* Policy, const FPCGGridDescriptor& GridDescriptor, const FBox& GridBounds, const double& CleanupRadiusSquared)
 	{
 		bool bAllGenSourcesWantToCleanup = true;
@@ -516,11 +512,23 @@ void FPCGRuntimeGenScheduler::TickCleanup(const TSet<IPCGGenSourceBase*>& InGenS
 		return bAllGenSourcesWantToCleanup;
 	};
 
+	// Generated Entry Key, Local/Generated Component
+	using PCGComponentToClean = TTuple<FGridGenerationKey, UPCGComponent*>;
+	TArray<PCGComponentToClean, TInlineAllocator<16>> ComponentsToClean;
+
+	TArray<FGridGenerationKey, TInlineAllocator<16>> InvalidKeys;
+
 	// Find any generated components which should be cleaned.
-	for (FGridGenerationKey GenerationKey : GeneratedComponents)
+	for (const FGridGenerationKey& GenerationKey : GeneratedComponents)
 	{
 		TRACE_CPUPROFILER_EVENT_SCOPE(FPCGRuntimeGenScheduler::SelectComponentsToCleanup);
 
+		if (!GenerationKey.IsValid())
+		{
+			InvalidKeys.Add(GenerationKey);
+			continue;
+		}
+
 		const FPCGGridDescriptor GridDescriptor = GenerationKey.GetGridDescriptor();
 		const EPCGHiGenGrid Grid = GridDescriptor.GetHiGenGrid();
 		const FIntVector& GridCoords = GenerationKey.GetGridCoords();
@@ -572,6 +580,11 @@ void FPCGRuntimeGenScheduler::TickCleanup(const TSet<IPCGGenSourceBase*>& InGenS
 		}
 	}
 
+	for (const FGridGenerationKey& InvalidKey : InvalidKeys)
+	{
+		GeneratedComponents.Remove(InvalidKey);
+	}
+
 	for (const PCGComponentToClean& ComponentToClean : ComponentsToClean)
 	{
 		CleanupComponent(ComponentToClean.Get<0>(), ComponentToClean.Get<1>());
@@ -786,8 +799,16 @@ void FPCGRuntimeGenScheduler::OnOriginalComponentUnregistered(UPCGComponent* InO
 		}
 	}
 
-	for (FGridGenerationKey GenerationKey : KeysToCleanup)
+	TArray<FGridGenerationKey, TInlineAllocator<16>> InvalidKeys;
+
+	for (const FGridGenerationKey& GenerationKey : KeysToCleanup)
 	{
+		if (!GenerationKey.IsValid())
+		{
+			InvalidKeys.Add(GenerationKey);
+			continue;
+		}
+
 		const FPCGGridDescriptor GridDescriptor = GenerationKey.GetGridDescriptor();
 		const bool bIsOriginalComponent = GridDescriptor.GetGridSize() == PCGHiGenGrid::UnboundedGridSize();
 
@@ -801,6 +822,11 @@ void FPCGRuntimeGenScheduler::OnOriginalComponentUnregistered(UPCGComponent* InO
 		}
 	}
 
+	for (const FGridGenerationKey& InvalidKey : InvalidKeys)
+	{
+		GeneratedComponents.Remove(InvalidKey);
+	}
+
 	CleanupRemainingComponents(InOriginalComponent);
 }
 
@@ -830,12 +856,22 @@ void FPCGRuntimeGenScheduler::CleanupLocalComponents(const APCGWorldActor* InPCG
 
 	// Generated Entry Key, Local/Generated Component
 	using PCGComponentToClean = TTuple<FGridGenerationKey, UPCGComponent*>;
-	TArray<PCGComponentToClean> ComponentsToClean;
+	TArray<PCGComponentToClean, TInlineAllocator<16>> ComponentsToClean;
+
 	TSet<UPCGComponent*> OriginalComponents;
+	OriginalComponents.Reserve(16);
+
+	TArray<FGridGenerationKey, TInlineAllocator<16>> InvalidKeys;
 
 	// Find all generated local components.
-	for (FGridGenerationKey GenerationKey : GeneratedComponents)
+	for (const FGridGenerationKey& GenerationKey : GeneratedComponents)
 	{
+		if (!GenerationKey.IsValid())
+		{
+			InvalidKeys.Add(GenerationKey);
+			continue;
+		}
+
 		const FPCGGridDescriptor GridDescriptor = GenerationKey.GetGridDescriptor();
 		const EPCGHiGenGrid Grid = GridDescriptor.GetHiGenGrid();
 		const FIntVector& GridCoords = GenerationKey.GetGridCoords();
@@ -851,6 +887,11 @@ void FPCGRuntimeGenScheduler::CleanupLocalComponents(const APCGWorldActor* InPCG
 		}
 	}
 
+	for (const FGridGenerationKey& InvalidKey : InvalidKeys)
+	{
+		GeneratedComponents.Remove(InvalidKey);
+	}
+
 	for (const PCGComponentToClean& ComponentToClean : ComponentsToClean)
 	{
 		CleanupComponent(ComponentToClean.Get<0>(), ComponentToClean.Get<1>());
@@ -980,6 +1021,11 @@ void FPCGRuntimeGenScheduler::CleanupDelayedRefreshComponents()
 	// If it is not, it would be leaked instead of refreshed, so we should force a full cleanup.
 	for (const FGridGenerationKey& GenerationKey : GeneratedComponentsToRemove)
 	{
+		if (!GenerationKey.IsValid())
+		{
+			continue;
+		}
+
 		const FPCGGridDescriptor GridDescriptor = GenerationKey.GetGridDescriptor();
 		const EPCGHiGenGrid Grid = GridDescriptor.GetHiGenGrid();
 		UPCGComponent* OriginalComponent = GenerationKey.GetOriginalComponent();
diff --git a/Engine/Plugins/PCG/Source/PCG/Public/RuntimeGen/PCGRuntimeGenScheduler.h b/Engine/Plugins/PCG/Source/PCG/Public/RuntimeGen/PCGRuntimeGenScheduler.h
--- a/Engine/Plugins/PCG/Source/PCG/Public/RuntimeGen/PCGRuntimeGenScheduler.h
+++ b/Engine/Plugins/PCG/Source/PCG/Public/RuntimeGen/PCGRuntimeGenScheduler.h
@@ -3,6 +3,7 @@
 #pragma once
 
 #include "PCGCommon.h"
+#include "PCGComponent.h"
 #include "Grid/PCGGridDescriptor.h"
 #include "UObject/WeakObjectPtr.h"
 
@@ -49,14 +50,16 @@ public:
 	void FlushAllGeneratedActors() { bActorFlushRequested = true; }
 
 protected:
-	struct FGridGenerationKey : TTuple<uint32, FIntVector, UPCGComponent*>
+	struct FGridGenerationKey : TTuple<uint32, FIntVector, TWeakObjectPtr<UPCGComponent>>
 	{
 		FGridGenerationKey(uint32 InGridSize, const FIntVector& InGridCoords, UPCGComponent* InOriginalComponent)
-			: TTuple<uint32, FIntVector, UPCGComponent*>(InGridSize, InGridCoords, InOriginalComponent) {}
+			: TTuple<uint32, FIntVector, TWeakObjectPtr<UPCGComponent>>(InGridSize, InGridCoords, InOriginalComponent) {}
+
+		bool IsValid() const { return Get<2>().IsValid(); }
 
 		uint32 GetGridSize() const { return Get<0>(); }
 		FIntVector GetGridCoords() const { return Get<1>(); }
-		UPCGComponent* GetOriginalComponent() const { return Get<2>(); }
+		UPCGComponent* GetOriginalComponent() const { return Get<2>().Get(); }
 
 		FPCGGridDescriptor GetGridDescriptor() const;
 	};
