[PCG] Rein back data interface validation further after finding a few conditions hit by a use case.

Reverts state of this validation to what it was before we recently extended it.

#rnx

diff --git a/Engine/Plugins/PCG/Source/PCG/Private/Compute/DataInterfaces/PCGDataCollectionUploadDataInterface.cpp b/Engine/Plugins/PCG/Source/PCG/Private/Compute/DataInterfaces/PCGDataCollectionUploadDataInterface.cpp
--- a/Engine/Plugins/PCG/Source/PCG/Private/Compute/DataInterfaces/PCGDataCollectionUploadDataInterface.cpp
+++ b/Engine/Plugins/PCG/Source/PCG/Private/Compute/DataInterfaces/PCGDataCollectionUploadDataInterface.cpp
@@ -117,8 +117,13 @@ FPCGDataProviderDataCollectionUploadProxy::FPCGDataProviderDataCollectionUploadP
 
 bool FPCGDataProviderDataCollectionUploadProxy::IsValid(FValidationData const& InValidationData) const
 {
-	return FPCGDataCollectionDataProviderProxy::IsValid(InValidationData)
-		&& DataAdaptor.IsValid();
+	if (!DataAdaptor.IsValid())
+	{
+		UE_LOG(LogPCG, Warning, TEXT("FPCGDataProviderDataCollectionUploadProxy invalid due to null DataAdaptor."));
+		return false;
+	}
+
+	return FPCGDataCollectionDataProviderProxy::IsValid(InValidationData);
 }
 
 void FPCGDataProviderDataCollectionUploadProxy::AllocateResources(FRDGBuilder& GraphBuilder, FAllocationData const& InAllocationData)
diff --git a/Engine/Plugins/PCG/Source/PCG/Private/Compute/DataInterfaces/PCGLandscapeDataInterface.cpp b/Engine/Plugins/PCG/Source/PCG/Private/Compute/DataInterfaces/PCGLandscapeDataInterface.cpp
--- a/Engine/Plugins/PCG/Source/PCG/Private/Compute/DataInterfaces/PCGLandscapeDataInterface.cpp
+++ b/Engine/Plugins/PCG/Source/PCG/Private/Compute/DataInterfaces/PCGLandscapeDataInterface.cpp
@@ -350,7 +350,7 @@ void UPCGLandscapeDataProvider::InitFromLandscape(ALandscape* InLandscape, const
 	// If the height can't be accessed via RVT, build a physics based height texture instead.
 	if (!HeightVirtualTexture)
 	{
-		InitializeLandscapeResource(InLandscape, Bounds);
+		InitializeFromLandscapeCollision(InLandscape, Bounds);
 	}
 
 	LandscapeGridSize = FVector2D(InLandscape->GetTransform().GetScale3D());
@@ -431,9 +431,9 @@ void UPCGLandscapeDataProvider::InitializeRuntimeVirtualTextures(ALandscape* InL
 		: nullptr;
 }
 
-void UPCGLandscapeDataProvider::InitializeLandscapeResource(ALandscape* InLandscape, const FBox& Bounds)
+void UPCGLandscapeDataProvider::InitializeFromLandscapeCollision(ALandscape* InLandscape, const FBox& Bounds)
 {
-	TRACE_CPUPROFILER_EVENT_SCOPE(UPCGLandscapeDataProvider::InitializeLandscapeResource);
+	TRACE_CPUPROFILER_EVENT_SCOPE(UPCGLandscapeDataProvider::InitializeFromLandscapeCollision);
 	check(InLandscape);
 
 	const ULandscapeInfo* LandscapeInfo = InLandscape->GetLandscapeInfo();
@@ -680,8 +680,7 @@ FPCGLandscapeDataProviderProxy::FPCGLandscapeDataProviderProxy(
 
 bool FPCGLandscapeDataProviderProxy::IsValid(FValidationData const& InValidationData) const
 {
-	return InValidationData.ParameterStructSize == sizeof(FParameters)
-		&& Resource.IsValid();
+	return InValidationData.ParameterStructSize == sizeof(FParameters);
 }
 
 void FPCGLandscapeDataProviderProxy::GatherDispatchData(FDispatchData const& InDispatchData)
diff --git a/Engine/Plugins/PCG/Source/PCG/Private/Compute/DataInterfaces/PCGLandscapeDataInterface.h b/Engine/Plugins/PCG/Source/PCG/Private/Compute/DataInterfaces/PCGLandscapeDataInterface.h
--- a/Engine/Plugins/PCG/Source/PCG/Private/Compute/DataInterfaces/PCGLandscapeDataInterface.h
+++ b/Engine/Plugins/PCG/Source/PCG/Private/Compute/DataInterfaces/PCGLandscapeDataInterface.h
@@ -89,7 +89,7 @@ public:
 private:
 	void InitFromLandscape(ALandscape* InLandscape, const FBox& Bounds);
 	void InitializeRuntimeVirtualTextures(ALandscape* InLandscape);
-	void InitializeLandscapeResource(ALandscape* InLandscape, const FBox& Bounds);
+	void InitializeFromLandscapeCollision(ALandscape* InLandscape, const FBox& Bounds);
 
 private:
 	UPROPERTY()
diff --git a/Engine/Plugins/PCG/Source/PCG/Private/Compute/DataInterfaces/PCGTextureDataInterface.cpp b/Engine/Plugins/PCG/Source/PCG/Private/Compute/DataInterfaces/PCGTextureDataInterface.cpp
--- a/Engine/Plugins/PCG/Source/PCG/Private/Compute/DataInterfaces/PCGTextureDataInterface.cpp
+++ b/Engine/Plugins/PCG/Source/PCG/Private/Compute/DataInterfaces/PCGTextureDataInterface.cpp
@@ -214,9 +214,9 @@ bool FPCGTextureDataProviderProxy::IsValid(FValidationData const& InValidationDa
 {
 	for (const FPCGTextureBindingInfo& BindingInfo : BindingInfos)
 	{
-		if (BindingInfo.ResourceType == EPCGTextureResourceType::Invalid
-			|| (!BindingInfo.Texture && !BindingInfo.ExportedTexture))
+		if (BindingInfo.ResourceType == EPCGTextureResourceType::Invalid)
 		{
+			UE_LOG(LogPCG, Warning, TEXT("FPCGTextureDataProviderProxy invalid due to invalid resource type."));
 			return false;
 		}
 	}
@@ -225,6 +225,7 @@ bool FPCGTextureDataProviderProxy::IsValid(FValidationData const& InValidationDa
 	{
 		if (!BindingInfos.IsValidIndex(TextureInfo.BindingIndex) || TextureInfo.SliceIndex < 0)
 		{
+			UE_LOG(LogPCG, Warning, TEXT("FPCGTextureDataProviderProxy invalid due to invalid binding indices."));
 			return false;
 		}
 	}
