[PCG] Fix dangling pointer in RDG buffer names.

#rnx

diff --git a/Engine/Plugins/PCG/Source/PCG/Private/Compute/DataInterfaces/PCGDataCollectionDataInterface.cpp b/Engine/Plugins/PCG/Source/PCG/Private/Compute/DataInterfaces/PCGDataCollectionDataInterface.cpp
--- a/Engine/Plugins/PCG/Source/PCG/Private/Compute/DataInterfaces/PCGDataCollectionDataInterface.cpp
+++ b/Engine/Plugins/PCG/Source/PCG/Private/Compute/DataInterfaces/PCGDataCollectionDataInterface.cpp
@@ -528,20 +528,6 @@ FComputeDataProviderRenderProxy* UPCGDataCollectionDataProvider::GetRenderProxy(
 
 	FPCGDataCollectionDataProviderProxy* Proxy = new FPCGDataCollectionDataProviderProxy(this, PinDesc, ExportMode, bRequiresZeroInitialization, ElementCountMultiplier, bExposeElementCounters, OutputPinLabel, OutputPinLabelAlias);
 
-#if WITH_EDITOR
-	if (const UPCGNode* Node = Cast<UPCGNode>(ProducerSettings ? ProducerSettings->GetOuter() : nullptr))
-	{
-		const FString NodeTitle = Node->GetNodeTitle(EPCGNodeTitleType::ListView).ToString();
-		Proxy->SetBufferName(FString::Format(TEXT("PCG_DC_{0}_{1}"), { NodeTitle, OutputPinLabel.ToString() }));
-		Proxy->SetCountersBufferName(FString::Format(TEXT("PCG_DC_{0}_{1}_Counters"), { NodeTitle, OutputPinLabel.ToString()}));
-	}
-	else
-	{
-		Proxy->SetBufferName(TEXT("PCG_DC"));
-		Proxy->SetCountersBufferName(TEXT("PCG_DC_Counters"));
-	}
-#endif
-
 	if (bExposeElementCounters)
 	{
 		const uint64 GenerationCount = GenerationCounter;
@@ -675,20 +661,13 @@ void FPCGDataCollectionDataProviderProxy::AllocateResources(FRDGBuilder& GraphBu
 	check(SizeBytes > 0);
 
 	{
-		const TCHAR* DebugNameCounters =
-#if WITH_EDITOR
-			*CountersBufferName;
-#else
-			TEXT("PCGDataCollection_Counters");
-#endif
-
 		// Ensure we always allocate at least 1 element so we have something to bind.
 		const uint32 NumCounters = bExposeElementCounters ? FMath::Max(PinDesc.DataDescs.Num(), 1) : 1;
 
 		FRDGBufferDesc CountersDesc = FRDGBufferDesc::CreateStructuredDesc(sizeof(uint32), NumCounters);
 		CountersDesc.Usage |= BUF_SourceCopy;
 
-		Counters = GraphBuilder.CreateBuffer(CountersDesc, DebugNameCounters);
+		Counters = GraphBuilder.CreateBuffer(CountersDesc, TEXT("PCGDataCollection_Counters"));
 		CountersUAV = GraphBuilder.CreateUAV(Counters);
 
 		// A single counter, or a small number of counters, will be a very common case.
@@ -705,13 +684,6 @@ void FPCGDataCollectionDataProviderProxy::AllocateResources(FRDGBuilder& GraphBu
 		}
 	}
 
-	const TCHAR* DebugName =
-#if WITH_EDITOR
-		*BufferName;
-#else
-		TEXT("PCGDataCollection");
-#endif
-
 	FRDGBufferDesc Desc = FRDGBufferDesc::CreateByteAddressDesc(SizeBytes);
 	if (ExportMode != EPCGExportMode::NoExport)
 	{
@@ -719,7 +691,7 @@ void FPCGDataCollectionDataProviderProxy::AllocateResources(FRDGBuilder& GraphBu
 		Desc.Usage |= BUF_SourceCopy;
 	}
 
-	Buffer = GraphBuilder.CreateBuffer(Desc, DebugName);
+	Buffer = GraphBuilder.CreateBuffer(Desc, TEXT("PCGDataCollection"));
 	BufferUAV = GraphBuilder.CreateUAV(Buffer);
 
 	// Initialize with an empty data collection. The kernel may not run, for example if indirect dispatch args end up being 0. This ensures
@@ -798,15 +770,4 @@ void FPCGDataCollectionDataProviderProxy::GetReadbackData(TArray<FReadbackData>&
 	}
 }
 
-#if WITH_EDITOR
-void FPCGDataCollectionDataProviderProxy::SetBufferName(const FString& InBufferName)
-{
-	BufferName = InBufferName;
-}
-void FPCGDataCollectionDataProviderProxy::SetCountersBufferName(const FString& InCountersBufferName)
-{
-	CountersBufferName = InCountersBufferName;
-}
-#endif
-
 #undef LOCTEXT_NAMESPACE
diff --git a/Engine/Plugins/PCG/Source/PCG/Private/Compute/DataInterfaces/PCGDataCollectionDataInterface.h b/Engine/Plugins/PCG/Source/PCG/Private/Compute/DataInterfaces/PCGDataCollectionDataInterface.h
--- a/Engine/Plugins/PCG/Source/PCG/Private/Compute/DataInterfaces/PCGDataCollectionDataInterface.h
+++ b/Engine/Plugins/PCG/Source/PCG/Private/Compute/DataInterfaces/PCGDataCollectionDataInterface.h
@@ -161,11 +161,6 @@ public:
 	void GetReadbackData(TArray<FReadbackData>& OutReadbackData) const override;
 	//~ End FComputeDataProviderRenderProxy Interface
 
-#if WITH_EDITOR
-	void SetBufferName(const FString& InBufferName);
-	void SetCountersBufferName(const FString& InCountersBufferName);
-#endif
-
 public:
 	FReadbackCallback AsyncReadbackCallback_RenderThread;
 
@@ -198,9 +193,4 @@ protected:
 
 	FName OutputPinLabel;
 	FName OutputPinLabelAlias;
-
-#if WITH_EDITOR
-	FString BufferName;
-	FString CountersBufferName;
-#endif
 };
diff --git a/Engine/Plugins/PCG/Source/PCG/Private/Compute/DataInterfaces/PCGDataCollectionUploadDataInterface.cpp b/Engine/Plugins/PCG/Source/PCG/Private/Compute/DataInterfaces/PCGDataCollectionUploadDataInterface.cpp
--- a/Engine/Plugins/PCG/Source/PCG/Private/Compute/DataInterfaces/PCGDataCollectionUploadDataInterface.cpp
+++ b/Engine/Plugins/PCG/Source/PCG/Private/Compute/DataInterfaces/PCGDataCollectionUploadDataInterface.cpp
@@ -67,7 +67,7 @@ FComputeDataProviderRenderProxy* UPCGDataProviderDataCollectionUpload::GetRender
 {
 	TRACE_CPUPROFILER_EVENT_SCOPE(UPCGDataProviderDataCollectionUpload::GetRenderProxy);
 
-	FPCGDataProviderDataCollectionUploadProxy* Proxy = new FPCGDataProviderDataCollectionUploadProxy(
+	return new FPCGDataProviderDataCollectionUploadProxy(
 		DataAdaptor,
 		this,
 		PinDesc,
@@ -77,19 +77,6 @@ FComputeDataProviderRenderProxy* UPCGDataProviderDataCollectionUpload::GetRender
 		bRequiresZeroInitialization,
 		ElementCountMultiplier,
 		bExposeElementCounters);
-
-#if WITH_EDITOR
-	if (const UPCGNode* Node = Cast<UPCGNode>(ProducerSettings ? ProducerSettings->GetOuter() : nullptr))
-	{
-		Proxy->SetBufferName(FString::Format(TEXT("PCG_UDC_{0}_{1}"), { Node->GetNodeTitle(EPCGNodeTitleType::ListView).ToString(), OutputPinLabel.ToString() }));
-	}
-	else
-	{
-		Proxy->SetBufferName(TEXT("PCG_UDC"));
-	}
-#endif
-
-	return Proxy;
 }
 
 void UPCGDataProviderDataCollectionUpload::Reset()
@@ -131,32 +118,18 @@ void FPCGDataProviderDataCollectionUploadProxy::AllocateResources(FRDGBuilder& G
 	TRACE_CPUPROFILER_EVENT_SCOPE(FPCGDataProviderDataCollectionUploadProxy::AllocateResources);
 
 	{
-		const TCHAR* DebugName =
-#if WITH_EDITOR
-			*BufferName;
-#else
-			TEXT("PCGDataCollectionUpload");
-#endif
-
-		Buffer = DataAdaptor->GetBuffer_RenderThread(GraphBuilder, ExportMode, DebugName);
+		Buffer = DataAdaptor->GetBuffer_RenderThread(GraphBuilder, ExportMode);
 		BufferUAV = GraphBuilder.CreateUAV(Buffer);
 	}
 
 	{
-		const TCHAR* DebugName =
-#if WITH_EDITOR
-			*CountersBufferName;
-#else
-			TEXT("PCGDataCollectionUpload_Counters");
-#endif
-
 		// Ensure we always allocate at least 1 element so we have something to bind.
 		const uint32 NumCounters = bExposeElementCounters ? FMath::Max(PinDesc.DataDescs.Num(), 1) : 1;
 
 		FRDGBufferDesc CountersDesc = FRDGBufferDesc::CreateStructuredDesc(sizeof(uint32), NumCounters);
 		CountersDesc.Usage |= BUF_SourceCopy;
 
-		Counters = GraphBuilder.CreateBuffer(CountersDesc, DebugName);
+		Counters = GraphBuilder.CreateBuffer(CountersDesc, TEXT("PCGDataCollectionUpload_Counters"));
 		CountersUAV = GraphBuilder.CreateUAV(Counters);
 
 		// A single counter, or a small number of counters, will be a very common case.
diff --git a/Engine/Plugins/PCG/Source/PCG/Private/Compute/PCGDataUploadAdaptor.cpp b/Engine/Plugins/PCG/Source/PCG/Private/Compute/PCGDataUploadAdaptor.cpp
--- a/Engine/Plugins/PCG/Source/PCG/Private/Compute/PCGDataUploadAdaptor.cpp
+++ b/Engine/Plugins/PCG/Source/PCG/Private/Compute/PCGDataUploadAdaptor.cpp
@@ -221,7 +221,7 @@ FRDGBufferSRVRef FPCGDataUploadAdaptor::GetDataElementCountsBufferSRV(FRDGBuilde
 }
 
 /** Gets the buffer that can then be used to read the data in kernels. */
-FRDGBufferRef FPCGDataUploadAdaptor::GetBuffer_RenderThread(FRDGBuilder& InGraphBuilder, EPCGExportMode InExportMode, const TCHAR* InBufferDebugName)
+FRDGBufferRef FPCGDataUploadAdaptor::GetBuffer_RenderThread(FRDGBuilder& InGraphBuilder, EPCGExportMode InExportMode)
 {
 	if (!ExternalBufferForReuse)
 	{
@@ -233,7 +233,7 @@ FRDGBufferRef FPCGDataUploadAdaptor::GetBuffer_RenderThread(FRDGBuilder& InGraph
 			Desc.Usage |= BUF_SourceCopy;
 		}
 
-		FRDGBufferRef Buffer = InGraphBuilder.CreateBuffer(Desc, InBufferDebugName);
+		FRDGBufferRef Buffer = InGraphBuilder.CreateBuffer(Desc, TEXT("PCGDataCollectionUpload"));
 
 		InGraphBuilder.QueueBufferUpload(Buffer, PackedDataCollection.GetData(), PackedDataCollection.Num() * PackedDataCollection.GetTypeSize());
 
diff --git a/Engine/Plugins/PCG/Source/PCG/Private/Compute/PCGDataUploadAdaptor.h b/Engine/Plugins/PCG/Source/PCG/Private/Compute/PCGDataUploadAdaptor.h
--- a/Engine/Plugins/PCG/Source/PCG/Private/Compute/PCGDataUploadAdaptor.h
+++ b/Engine/Plugins/PCG/Source/PCG/Private/Compute/PCGDataUploadAdaptor.h
@@ -26,7 +26,7 @@ public:
 	FRDGBufferSRVRef GetDataElementCountsBufferSRV(FRDGBuilder& InGraphBuilder);
 
 	/** Gets the buffer that can then be used to read the data in kernels. */
-	FRDGBufferRef GetBuffer_RenderThread(FRDGBuilder& InGraphBuilder, EPCGExportMode InExportMode, const TCHAR* InBufferDebugName);
+	FRDGBufferRef GetBuffer_RenderThread(FRDGBuilder& InGraphBuilder, EPCGExportMode InExportMode);
 
 private:
 	TRefCountPtr<FRDGPooledBuffer> ExternalBufferForReuse;
