[PCG] Added small tie-in for preconfigured settings to consume metadata to use for extra action search info

+ Added math symbol aliases to Metadata Ops

diff --git a/Engine/Plugins/PCG/Source/PCG/Public/Elements/Metadata/PCGMetadataMathsOpElement.h b/Engine/Plugins/PCG/Source/PCG/Public/Elements/Metadata/PCGMetadataMathsOpElement.h
--- a/Engine/Plugins/PCG/Source/PCG/Public/Elements/Metadata/PCGMetadataMathsOpElement.h
+++ b/Engine/Plugins/PCG/Source/PCG/Public/Elements/Metadata/PCGMetadataMathsOpElement.h
@@ -15,35 +15,35 @@ enum class EPCGMetadataMathsOperation : uint16
 	Frac,
 	Truncate,
 	Round,
-	Sqrt,
-	Abs,
+	Sqrt UMETA(SearchHints = "square root"),
+	Abs UMETA(SearchHints = "|"),
 	Floor,
 	Ceil,
-	OneMinus UMETA(Tooltip = "1 - X operation"),
-	Inc UMETA(Tooltip = "X + 1 operation"),
-	Dec UMETA(Tooltip = "X - 1 operation"),
-	Negate UMETA(Tooltip = "-X operation"),
+	OneMinus UMETA(Tooltip = "1 - X operation", SearchHints = "1-"),
+	Inc UMETA(Tooltip = "X + 1 operation", SearchHints = "++"),
+	Dec UMETA(Tooltip = "X - 1 operation", SearchHints = "--"),
+	Negate UMETA(Tooltip = "-X operation", SearchHints = "-"),
 
 	// Binary op
 	BinaryOp = 1 << 11 UMETA(Hidden),
-	Add,
-	Subtract,
-	Multiply,
-	Divide,
+	Add UMETA(SearchHints = "+"),
+	Subtract UMETA(SearchHints = "-"),
+	Multiply UMETA(SearchHints = "*"),
+	Divide UMETA(SearchHints = "/"),
 	Max,
 	Min,
-	Pow,
+	Pow UMETA(SearchHints = "**"),
 	// ClampMin is just Max and ClampMax is just Min, so hide them.
 	ClampMin UMETA(Hidden),
 	ClampMax UMETA(Hidden),
-	Modulo,
-	Set,
+	Modulo UMETA(SearchHints = "%"),
+	Set UMETA(SearchHints = "="),
 
 	// Ternary op
 	TernaryOp = 1 << 12 UMETA(Hidden),
 	Clamp,
 	Lerp,
-	MulAdd UMETA(Tooltip = "Multiply Add (A + B * C)")
+	MulAdd UMETA(Tooltip = "Multiply Add (A + B * C)", SearchHints = "*+")
 };
 ENUM_CLASS_FLAGS(EPCGMetadataMathsOperation);
 
diff --git a/Engine/Plugins/PCG/Source/PCG/Public/Utils/PCGPreconfiguration.h b/Engine/Plugins/PCG/Source/PCG/Public/Utils/PCGPreconfiguration.h
--- a/Engine/Plugins/PCG/Source/PCG/Public/Utils/PCGPreconfiguration.h
+++ b/Engine/Plugins/PCG/Source/PCG/Public/Utils/PCGPreconfiguration.h
@@ -32,9 +32,14 @@ struct PCG_API FPCGPreconfiguredInfo
 	/**
 	* Automatically fill all preconfigured settings depending on the provided enum.
 	* Can also specify explicitly values that should not be included, i.e. cases that may not be available in non-editor builds.
+	* Can pass in an optional function or lambda for filtering or finer control over the results, for things like metadata.
+	* @param InValuesToSkip A set of enum values that should be skipped during population, like counts or hidden values.
+	* @param InOptionalPrefix An optional string prefix to the name of the action as it appears in the contextual search.
+	* @param ProcessFunc A post-process or filtering callback for each enum value. Returns false if the enum should be filtered.
+	* @returns An array of preconfigured info for schema actions or conversion, etc.
 	*/
 	template <typename EnumType, typename SubclassType = FPCGPreconfiguredInfo>
-	static TArray<SubclassType> PopulateFromEnum(const TSet<EnumType>& InValuesToSkip = {}, const FText& InOptionalPrefix = FText())
+	static TArray<SubclassType> PopulateFromEnum(const TSet<EnumType>& InValuesToSkip = {}, const FText& InOptionalPrefix = FText(), TFunction<bool(SubclassType& InOutInfo, const UEnum* EnumPtr, int32 ValueIndex)> ProcessFunc = nullptr)
 	{
 		static_assert(std::is_base_of_v<FPCGPreconfiguredInfo, SubclassType>, "SubclassType must be a subclass of FPCGPreconfiguredInfo");
 
@@ -52,7 +57,11 @@ struct PCG_API FPCGPreconfiguredInfo
 					FText DisplayName = EnumPtr->GetDisplayNameTextByValue(Value);
 					if (!DisplayName.IsEmpty())
 					{
-						PreconfiguredInfo.Emplace(Value, FText::Format(FText::FromString(TEXT("{0}{1}")), InOptionalPrefix, std::move(DisplayName)));
+						SubclassType Info(Value, FText::Format(FText::FromString(TEXT("{0}{1}")), InOptionalPrefix, std::move(DisplayName)));
+						if (!ProcessFunc || ProcessFunc(Info, EnumPtr, i))
+						{
+							PreconfiguredInfo.Emplace(std::move(Info));
+						}
 					}
 				}
 			}
@@ -85,6 +94,8 @@ struct PCG_API FPCGPreConfiguredSettingsInfo : public FPCGPreconfiguredInfo
 {
 	GENERATED_BODY()
 
+	static constexpr TCHAR SearchHintMetadataKey[] = TEXT("SearchHints");
+
 	FPCGPreConfiguredSettingsInfo() = default;
 
 	explicit FPCGPreConfiguredSettingsInfo(int32 InIndex, FText InLabel = FText{})
@@ -92,14 +103,32 @@ struct PCG_API FPCGPreConfiguredSettingsInfo : public FPCGPreconfiguredInfo
 	{}
 
 #if WITH_EDITOR
-	FPCGPreConfiguredSettingsInfo(int32 InIndex, FText InLabel, FText InTooltip)
+	FPCGPreConfiguredSettingsInfo(int32 InIndex, FText InLabel, FText InTooltip, FText InSearchHints = {})
 		: FPCGPreconfiguredInfo(InIndex, InLabel, InTooltip)
+		, SearchHints(InSearchHints)
 	{}
 #endif // WITH_EDITOR
 
 	template <typename EnumType>
 	static TArray<FPCGPreConfiguredSettingsInfo> PopulateFromEnum(const TSet<EnumType>& InValuesToSkip = {}, const FText& InOptionalPrefix = FText())
 	{
+#if WITH_EDITOR
+		return FPCGPreconfiguredInfo::PopulateFromEnum<EnumType, FPCGPreConfiguredSettingsInfo>(InValuesToSkip, InOptionalPrefix, [](FPCGPreConfiguredSettingsInfo& InOutInfo, const UEnum* EnumPtr, const int32 ValueIndex) -> bool
+		{
+			if (EnumPtr && EnumPtr->HasMetaData(SearchHintMetadataKey, ValueIndex))
+			{
+				InOutInfo.SearchHints = FText::FromString(EnumPtr->GetMetaData(SearchHintMetadataKey, ValueIndex));
+			}
+
+			return true;
+		});
+#else
 		return FPCGPreconfiguredInfo::PopulateFromEnum<EnumType, FPCGPreConfiguredSettingsInfo>(InValuesToSkip, InOptionalPrefix);
+#endif // WITH_EDITOR
 	}
+
+#if WITH_EDITORONLY_DATA
+	UPROPERTY(EditAnywhere, Category = "")
+	FText SearchHints;
+#endif // WITH_EDITORONLY_DATA
 };
\ No newline at end of file
diff --git a/Engine/Plugins/PCG/Source/PCGEditor/Private/PCGEditorGraphSchema.cpp b/Engine/Plugins/PCG/Source/PCGEditor/Private/PCGEditorGraphSchema.cpp
--- a/Engine/Plugins/PCG/Source/PCGEditor/Private/PCGEditorGraphSchema.cpp
+++ b/Engine/Plugins/PCG/Source/PCGEditor/Private/PCGEditorGraphSchema.cpp
@@ -463,7 +463,7 @@ void UPCGEditorGraphSchema::GetNativeElementActions(FGraphActionMenuBuilder& Act
 
 				for (FPCGPreConfiguredSettingsInfo PreconfiguredInfo : AllPreconfiguredInfo)
 				{
-					TSharedPtr<FPCGEditorGraphSchemaAction_NewNativeElement> NewPreconfiguredAction(new FPCGEditorGraphSchemaAction_NewNativeElement(NewCategory, PreconfiguredInfo.Label, PreconfiguredInfo.Tooltip.IsEmpty() ? Description : PreconfiguredInfo.Tooltip, 0));
+					TSharedPtr<FPCGEditorGraphSchemaAction_NewNativeElement> NewPreconfiguredAction(new FPCGEditorGraphSchemaAction_NewNativeElement(NewCategory, PreconfiguredInfo.Label, PreconfiguredInfo.Tooltip.IsEmpty() ? Description : PreconfiguredInfo.Tooltip,0, PreconfiguredInfo.SearchHints));
 					NewPreconfiguredAction->SettingsClass = SettingsClass;
 					NewPreconfiguredAction->PreconfiguredInfo = std::move(PreconfiguredInfo);
 					ActionMenuBuilder.AddAction(NewPreconfiguredAction);
