[PCG] Some updates to the dynamic tracking

* Don't do PostCleanupGraph at the beginning of the generation, this is reseting the dynamic tracking and if we are cancel we lost all tracking
* Add change origin enum to track where changes came from
* Do not refresh the owner on actor delete (was already the case) and actor added

#rnx

diff --git a/Engine/Plugins/PCG/Source/PCG/Private/PCGActorAndComponentMapping.cpp b/Engine/Plugins/PCG/Source/PCG/Private/PCGActorAndComponentMapping.cpp
--- a/Engine/Plugins/PCG/Source/PCG/Private/PCGActorAndComponentMapping.cpp
+++ b/Engine/Plugins/PCG/Source/PCG/Private/PCGActorAndComponentMapping.cpp
@@ -1351,18 +1351,20 @@ void FPCGActorAndComponentMapping::ProcessDelayedEvents()
 		return;
 	}
 			
-	TMap<TObjectKey<AActor>, TObjectKey<UObject>> LocalDelayedChangedActors = MoveTemp(DelayedChangedActors);
+	auto LocalDelayedChangedActors = MoveTemp(DelayedChangedActors);
 
-	for (const TPair<TObjectKey<AActor>, TObjectKey<UObject>>& ActorToObject : LocalDelayedChangedActors)
+	for (const auto& ActorToObject : LocalDelayedChangedActors)
 	{
 		AActor* Actor = ActorToObject.Key.ResolveObjectPtr();
-		UObject* Object = ActorToObject.Value.ResolveObjectPtr();
+		UObject* Object = ActorToObject.Value.Get<0>().ResolveObjectPtr();
+		const EChangedObjectReason Reason = ActorToObject.Value.Get<1>();
+		
 		if (!Actor)
 		{
 			continue;
 		}
 
-		OnActorChanged_Internal(Actor, Object);
+		OnActorChanged_Internal(Actor, Object, Reason);
 	}
 }
 
@@ -1516,17 +1518,17 @@ void FPCGActorAndComponentMapping::BuildPartitionActorRecords()
 	}
 }
 
-void FPCGActorAndComponentMapping::OnActorChanged_Recursive(AActor* InActor, UObject* InOriginatingChangeObject)
+void FPCGActorAndComponentMapping::OnActorChanged_Recursive(AActor* InActor, UObject* InOriginatingChangeObject, const EChangedObjectReason InChangeReason)
 {
 	check(!DelayedChangedActors.Contains(InActor));
 
 	// Process actors recursively
-	auto OnActorChanged = [this, InOriginatingChangeObject](AActor* InActor, auto&& OnActorChangedRecursive) -> void
+	auto OnActorChanged = [this, InOriginatingChangeObject, InChangeReason](AActor* InActor, auto&& OnActorChangedRecursive) -> void
 	{
 		if (InActor)
 		{
 			FActorPreviousData* PreviousData = ActorToPreviousDataMap.Find(InActor);
-			OnObjectChanged(InActor, PreviousData, InOriginatingChangeObject);
+			OnObjectChanged(InActor, InChangeReason, PreviousData, InOriginatingChangeObject);
 			if (PreviousData)
 			{
 				ActorToPreviousDataMap.Remove(InActor);
@@ -1564,10 +1566,10 @@ void FPCGActorAndComponentMapping::OnActorAdded_Internal(AActor* InActor, bool b
 
 	// We delay adding because OnActorAdded fires before an actor's properties are set,
 	// so the actor is not ready for processing until the next tick.
-	DelayedChangedActors.Add(InActor, nullptr);
+	DelayedChangedActors.Emplace(InActor, {nullptr, EChangedObjectReason::ActorAdded});
 }
 
-void FPCGActorAndComponentMapping::OnActorChanged_Internal(AActor* InActor, UObject* InOriginatingChangeObject)
+void FPCGActorAndComponentMapping::OnActorChanged_Internal(AActor* InActor, UObject* InOriginatingChangeObject, const EChangedObjectReason InChangeReason)
 {
 	// Currently delayed, return
 	if (DelayedChangedActors.Contains(InActor))
@@ -1578,12 +1580,12 @@ void FPCGActorAndComponentMapping::OnActorChanged_Internal(AActor* InActor, UObj
 	// Should still be delayed, add to delayed actors, return
 	if (ShouldDelayActor(InActor))
 	{
-		DelayedChangedActors.Add(InActor, InOriginatingChangeObject);
+		DelayedChangedActors.Emplace(InActor, {InOriginatingChangeObject, InChangeReason});
 		return;
 	}
 	
 	// Call Recursive method
-	OnActorChanged_Recursive(InActor, InOriginatingChangeObject);
+	OnActorChanged_Recursive(InActor, InOriginatingChangeObject, InChangeReason);
 }
 
 void FPCGActorAndComponentMapping::OnActorUnloaded(AActor& InActor)
@@ -1628,7 +1630,7 @@ void FPCGActorAndComponentMapping::OnActorDeleted_Internal(AActor* InActor, bool
 #endif // WITH_EDITOR
 
 	// Notify all components that the actor has changed (was removed), but the Refresh will only happen AFTER the actor was actually removed from the world (because of delayed refresh).
-	OnObjectChanged(InActor, /*InPreviousData=*/nullptr, /*InOriginatingObject=*/nullptr, /*bNoRefreshOnOwner=*/true);
+	OnObjectChanged(InActor, EChangedObjectReason::ActorDeleted, /*InPreviousData=*/nullptr, /*InOriginatingObject=*/nullptr);
 }
 
 void FPCGActorAndComponentMapping::OnObjectModified(UObject* InObject)
@@ -1702,7 +1704,7 @@ void FPCGActorAndComponentMapping::OnObjectSaved(UObject* InObject, FObjectPreSa
 	if (!InObjectSaveContext.IsProceduralSave() && (InObject->IsA<UDataTable>() || InObject->IsA<UPCGDataAsset>() || InObject->IsA<UStaticMesh>()))
 	{
 		FPropertyChangedEvent Event{ nullptr };
-		OnObjectChanged(InObject, /*InPreviousData=*/nullptr, /*InOriginatingChangeObject=*/ InObject);
+		OnObjectChanged(InObject, EChangedObjectReason::ObjectSaved, /*InPreviousData=*/nullptr, /*InOriginatingChangeObject=*/ InObject);
 	}
 }
 
@@ -1778,18 +1780,23 @@ void FPCGActorAndComponentMapping::OnObjectPropertyChanged(UObject* InObject, FP
 
 	if (Actor)
 	{
-		OnActorChanged_Internal(Actor, /*InOriginatingChangeObject=*/ InObject);
+		OnActorChanged_Internal(Actor, /*InOriginatingChangeObject=*/ InObject, EChangedObjectReason::ActorChanged);
 	}
 	else
 	{
-		OnObjectChanged(InObject, /*InPreviousData=*/nullptr, /*InOriginatingChangeObject=*/ InObject);
+		OnObjectChanged(InObject, EChangedObjectReason::ObjectChanged, /*InPreviousData=*/nullptr, /*InOriginatingChangeObject=*/ InObject);
 	}
 }
 
-void FPCGActorAndComponentMapping::OnObjectChanged(UObject* InObject, const FActorPreviousData* InPreviousData, const UObject* InOriginatingChangeObject, bool bNoRefreshOwner)
+void FPCGActorAndComponentMapping::OnObjectChanged(UObject* InObject, const EChangedObjectReason InChangeReason, const FActorPreviousData* InPreviousData, const UObject* InOriginatingChangeObject)
 {
 	TRACE_CPUPROFILER_EVENT_SCOPE(FPCGActorAndComponentMapping::OnObjectChanged);
 
+	// Discard everything related to self for actor added/delete and changes coming from a PCG Component.
+	const bool bNoRefreshOwner = InChangeReason == EChangedObjectReason::ActorAdded
+		|| InChangeReason == EChangedObjectReason::ActorDeleted
+		|| InChangeReason == EChangedObjectReason::PCGComponentGenerated;
+
 	// Nothing to do if we track nothing or there is no object.
 	if (!InObject || (CulledTrackedKeysToComponentsMap.IsEmpty() && AlwaysTrackedKeysToComponentsMap.IsEmpty()))
 	{
@@ -2144,8 +2151,8 @@ void FPCGActorAndComponentMapping::ApplyLandscapeChanges(ALandscapeProxy* InLand
 	{
 		return;
 	}
-	
-	OnObjectChanged(InLandscape, /*InPreviousData=*/nullptr, InLandscape);
+
+	OnObjectChanged(InLandscape, EChangedObjectReason::ActorChanged, /*InPreviousData=*/nullptr, InLandscape);
 }
 
 void FPCGActorAndComponentMapping::OnPCGGraphStartsGenerating(UPCGComponent* InComponent)
@@ -2247,7 +2254,7 @@ void FPCGActorAndComponentMapping::OnPCGGraphGeneratedOrCleaned(UPCGComponent* I
 		return;
 	}
 
-	OnObjectChanged(InComponent->GetOwner(), /*InPreviousData=*/nullptr, InComponent);
+	OnObjectChanged(InComponent->GetOwner(), EChangedObjectReason::PCGComponentGenerated, /*InPreviousData=*/nullptr, InComponent);
 }
 
 bool FPCGActorAndComponentMapping::ClearCacheForKeys(const TArray<FPCGSelectionKey>& InKeys, const UPCGComponent* InComponent, const bool bIntersect, const UObject* InOriginatingChange) const
diff --git a/Engine/Plugins/PCG/Source/PCG/Private/PCGComponent.cpp b/Engine/Plugins/PCG/Source/PCG/Private/PCGComponent.cpp
--- a/Engine/Plugins/PCG/Source/PCG/Private/PCGComponent.cpp
+++ b/Engine/Plugins/PCG/Source/PCG/Private/PCGComponent.cpp
@@ -866,11 +866,17 @@ void UPCGComponent::CallPostGenerateFunctions(FPCGContext* Context) const
 	}
 }
 
-void UPCGComponent::PostCleanupGraph()
+void UPCGComponent::PostCleanupGraph(bool bRemoveComponents)
 {
-	bGenerated = false;
 	CurrentCleanupTask = InvalidPCGTaskId;
+	
+	if (!bRemoveComponents)
+	{
+		// If we didn't remove components, it's a shallow cleanup before generating, so early out here.
+		return;
+	}
 
+	bGenerated = false;
 	const bool bHadGeneratedGraphOutput = GeneratedGraphOutput.TaggedData.Num() > 0;
 
 	ClearGraphGeneratedOutput();
@@ -1500,12 +1506,7 @@ void UPCGComponent::CleanupLocalImmediate(bool bRemoveComponents, bool bCleanupL
 		UPCGActorHelpers::DeleteActors(GetWorld(), ActorsToDelete.Array());
 	}
 
-	// If bRemoveComponents is true, it means we are in a "real" cleanup, not a pre-cleanup before a generate.
-	// So call PostCleanup in this case.
-	if (bRemoveComponents)
-	{
-		PostCleanupGraph();
-	}
+	PostCleanupGraph(bRemoveComponents);
 
 	// If the component is partitioned, we will forward the calls to its local components.
 	if (Subsystem && bCleanupLocalComponents && IsPartitioned())
diff --git a/Engine/Plugins/PCG/Source/PCG/Private/PCGSubsystem.cpp b/Engine/Plugins/PCG/Source/PCG/Private/PCGSubsystem.cpp
--- a/Engine/Plugins/PCG/Source/PCG/Private/PCGSubsystem.cpp
+++ b/Engine/Plugins/PCG/Source/PCG/Private/PCGSubsystem.cpp
@@ -756,7 +756,7 @@ FPCGTaskId UPCGSubsystem::ScheduleCleanup(UPCGComponent* PCGComponent, bool bRem
 				return true;
 			}
 
-			Component->PostCleanupGraph();
+			Component->PostCleanupGraph(bRemoveComponents);
 
 			// Remove the local component mappings if requested and the component is partitioned. If 'bRemoveComponents' is false, that indicates we are doing a refresh, so destroying
 			// the component mappings is counterproductive.
diff --git a/Engine/Plugins/PCG/Source/PCG/Public/PCGActorAndComponentMapping.h b/Engine/Plugins/PCG/Source/PCG/Public/PCGActorAndComponentMapping.h
--- a/Engine/Plugins/PCG/Source/PCG/Public/PCGActorAndComponentMapping.h
+++ b/Engine/Plugins/PCG/Source/PCG/Public/PCGActorAndComponentMapping.h
@@ -104,6 +104,18 @@ public:
 #endif
 private:
 
+	/** Used to track where a change comes from. */
+	enum class EChangedObjectReason
+	{
+		ActorAdded,
+		ActorChanged,
+		ActorDeleted,
+		ObjectChanged,
+		ObjectSaved,
+		PCGComponentGenerated,
+		Undefined
+	};
+
 #if WITH_EDITOR
 	void RegisterDelegates();
 	void UnregisterDelegates();
@@ -180,8 +192,8 @@ private:
 	void OnActorAdded(AActor* InActor);
 	void OnActorLoaded(AActor& InActor);
 	void OnActorAdded_Internal(AActor* InActor, bool bShouldDirty);
-	void OnActorChanged_Internal(AActor* InActor, UObject* InOriginatingChangeObject);
-	void OnActorChanged_Recursive(AActor* InActor, UObject* InOriginatingChangeObject);
+	void OnActorChanged_Internal(AActor* InActor, UObject* InOriginatingChangeObject, const EChangedObjectReason InChangeReason);
+	void OnActorChanged_Recursive(AActor* InActor, UObject* InOriginatingChangeObject, const EChangedObjectReason InChangeReason);
 	void OnActorDeleted(AActor* InActor);
 	void OnActorUnloaded(AActor& InActor);
 	void OnActorDeleted_Internal(AActor* InActor, bool bShouldDirty);
@@ -210,7 +222,7 @@ private:
 	* Can also specify an optional object, originating the change, to avoid re-dirtying a component if it was the origin.
 	* Another option when an actor is deleted/unload, don't refresh their components.
 	*/
-	void OnObjectChanged(UObject* InObject, const FActorPreviousData* InPreviousData = nullptr, const UObject* InOriginatingChangeObject = nullptr, bool bNoRefreshOnOwner = false);
+	void OnObjectChanged(UObject* InObject, const EChangedObjectReason InChangeReason, const FActorPreviousData* InPreviousData = nullptr, const UObject* InOriginatingChangeObject = nullptr);
 
 	/** Gather all settings from a given component that track the key, and clear the cache for them. Returns true if we should dirty afterwards (aka at least one settings was cleared and/or landscape changed). */
 	bool ClearCacheForKeys(const TArray<FPCGSelectionKey>& InKeys, const UPCGComponent* InComponent, const bool bIntersect, const UObject* InOriginatingChange) const;
@@ -254,7 +266,7 @@ private:
 	mutable FRWLock TrackedComponentsLock;
 
 	// Keep track of actors that aren't yet ready (or if the subsystem is not yet ready), whether we should dirty them in next tick.
-	TMap<TObjectKey<AActor>, TObjectKey<UObject>> DelayedChangedActors;
+	TMap<TObjectKey<AActor>, TTuple<TObjectKey<UObject>, EChangedObjectReason>> DelayedChangedActors;
 	
 	/** Transient map of actors and their previous data, it's set in the pre object change to be able to track changes (such as tags or positions) */
 	TMap<TObjectKey<AActor>, FActorPreviousData> ActorToPreviousDataMap;
diff --git a/Engine/Plugins/PCG/Source/PCG/Public/PCGComponent.h b/Engine/Plugins/PCG/Source/PCG/Public/PCGComponent.h
--- a/Engine/Plugins/PCG/Source/PCG/Public/PCGComponent.h
+++ b/Engine/Plugins/PCG/Source/PCG/Public/PCGComponent.h
@@ -581,7 +581,7 @@ private:
 
 	void PostProcessGraph(const FBox& InNewBounds, bool bInGenerated, FPCGContext* Context);
 	void CallPostGenerateFunctions(FPCGContext* Context) const;
-	void PostCleanupGraph();
+	void PostCleanupGraph(bool bRemoveComponents);
 	void OnProcessGraphAborted(bool bQuiet = false, bool bCleanupUnusedResources = true);
 	void CleanupUnusedManagedResources();
 	void ClearGraphGeneratedOutput();
