[PCG] Rename Create Attribute to 'Create Constant'

diff --git a/Engine/Plugins/PCG/Source/PCG/Private/Elements/PCGCreateAttribute.cpp b/Engine/Plugins/PCG/Source/PCG/Private/Elements/PCGCreateAttribute.cpp
--- a/Engine/Plugins/PCG/Source/PCG/Private/Elements/PCGCreateAttribute.cpp
+++ b/Engine/Plugins/PCG/Source/PCG/Private/Elements/PCGCreateAttribute.cpp
@@ -22,8 +22,10 @@ namespace PCGCreateAttributeConstants
 {
 	const FName NodeNameAddAttribute = TEXT("AddAttribute");
 	const FText NodeTitleAddAttribute = LOCTEXT("NodeTitleAddAttribute", "Add Attribute");
-	const FName NodeNameCreateAttribute = TEXT("CreateAttribute");
-	const FText NodeTitleCreateAttribute = LOCTEXT("NodeTitleCreateAttribute", "Create Attribute");
+	const FName NodeNameCreateConstant = TEXT("CreateConstant");
+	const FText NodeTitleCreateConstant = LOCTEXT("NodeTitleCreateConstant", "Create Constant");
+	const FText NodeAliasCreateConstant = LOCTEXT("NodeAliasCreateConstant", "Create Attribute");
+	const FText NodeTooltipFormatCreateConstant = LOCTEXT("NodeTooltipFormatCreateConstant", "Outputs an attribute set containing a constant '{0}' value:  {1}");
 	const FName AttributesLabel = TEXT("Attributes");
 	const FText AttributesTooltip = LOCTEXT("AttributesTooltip", "Optional Attribute Set to create from. Not used if not connected.");
 	const FText ErrorCreatingAttributeMessage = LOCTEXT("ErrorCreatingAttribute", "Error while creating attribute '{0}'");
@@ -96,8 +98,8 @@ bool UPCGAddAttributeSettings::CanEditChange(const FProperty* InProperty) const
 void UPCGAddAttributeSettings::ApplyStructuralDeprecation(UPCGNode* InOutNode)
 {
 	check(InOutNode);
-	// Arbitrary version that approximatively matches the time when Add/Create attributes changed.
-	// It will convert any add attributes that have nothing connected to it to a create attribute.
+	// Arbitrary version that approximately matches the time when Add/Create attributes changed.
+	// It will convert any add attributes that have nothing connected to it to a Create Constant.
 	if (DataVersion < FPCGCustomVersion::SupportPartitionedComponentsInNonPartitionedLevels)
 	{
 		if (!InOutNode->IsInputPinConnected(PCGPinConstants::DefaultInputLabel))
@@ -246,12 +248,22 @@ void UPCGCreateAttributeSetSettings::PostLoad()
 #if WITH_EDITOR
 FName UPCGCreateAttributeSetSettings::GetDefaultNodeName() const
 {
-	return PCGCreateAttributeConstants::NodeNameCreateAttribute;
+	return PCGCreateAttributeConstants::NodeNameCreateConstant;
 }
 
 FText UPCGCreateAttributeSetSettings::GetDefaultNodeTitle() const
 {
-	return PCGCreateAttributeConstants::NodeTitleCreateAttribute;
+	return PCGCreateAttributeConstants::NodeTitleCreateConstant;
+}
+
+TArray<FText> UPCGCreateAttributeSetSettings::GetNodeTitleAliases() const
+{
+	return {PCGCreateAttributeConstants::NodeAliasCreateConstant};
+}
+
+FText UPCGCreateAttributeSetSettings::GetNodeTooltipText() const
+{
+	return FText::Format(PCGCreateAttributeConstants::NodeTooltipFormatCreateConstant, FText::FromString(AttributeTypes.TypeToString()), FText::FromString(AttributeTypes.ToString()));
 }
 #endif // WITH_EDITOR
 
@@ -314,7 +326,7 @@ bool UPCGCreateAttributeSetSettings::ConvertNode(const FPCGPreconfiguredInfo& Co
 			FPCGSingleNodeConverter NodeConverter(Node, UPCGUserParameterGetSettings::StaticClass());
 			if (!NodeConverter.IsGraphInitialized() || !NodeConverter.IsSourceInitialized())
 			{
-				PCGLog::Settings::LogInvalidPreconfigurationWarning(ConversionInfo.PreconfiguredIndex, PCGCreateAttributeConstants::NodeTitleCreateAttribute);
+				PCGLog::Settings::LogInvalidPreconfigurationWarning(ConversionInfo.PreconfiguredIndex, PCGCreateAttributeConstants::NodeTitleCreateConstant);
 				return false;
 			}
 
@@ -435,7 +447,7 @@ bool FPCGAddAttributeElement::ExecuteInternal(FPCGContext* Context) const
 		return true;
 	}
 
-	// If we copy all attributes, support to have multiple source param. Overwise, add a warning
+	// If we copy all attributes, support to have multiple source param. Otherwise, add a warning
 	if (SourceParams.Num() > 1 && !Settings->bCopyAllAttributes)
 	{
 		PCGLog::LogWarningOnGraph(LOCTEXT("MultiAttributeWhenNoCopyAll", "Multiple source param detected in the Attributes pin, but we do not copy all attributes. We will only look into the first source param."), Context);
diff --git a/Engine/Plugins/PCG/Source/PCG/Private/Elements/PCGUserParameterGet.cpp b/Engine/Plugins/PCG/Source/PCG/Private/Elements/PCGUserParameterGet.cpp
--- a/Engine/Plugins/PCG/Source/PCG/Private/Elements/PCGUserParameterGet.cpp
+++ b/Engine/Plugins/PCG/Source/PCG/Private/Elements/PCGUserParameterGet.cpp
@@ -107,7 +107,7 @@ namespace PCGUserParameterGet
 
 		const FText SpecificToGenericTooltip = LOCTEXT("SpecificToGenericTooltip", "Convert to a Get User Parameter (Generic) node.");
 		const FText GenericToSpecificTooltip = LOCTEXT("GenericToSpecificTooltip", "Convert to a specified Get User Parameter node.");
-		const FText ToCreateAttributeTooltip = LOCTEXT("ToConstantTooltip", "Convert to a Create Attribute node.");
+		const FText ToCreateConstantTooltip = LOCTEXT("ToConstantTooltip", "Convert to a Create Constant node.");
 
 		const FText SpecificGetUserParameterNodeTitle = LOCTEXT("SpecificGetUserParameterNodeTitle", "Get Graph Parameter");
 		const FText GenericGetUserParameterNodeTitle = LOCTEXT("GenericGetUserParameterNodeTitle", "Get Graph Parameter (Generic)");
@@ -146,7 +146,7 @@ namespace PCGUserParameterGet
 			OutTypeStruct.Type = static_cast<EPCGMetadataTypes>(PropertyAccessor->GetUnderlyingType());
 			if (OutTypeStruct.Type == EPCGMetadataTypes::Unknown)
 			{
-				OutFailReason = LOCTEXT("ConversionAssignValueUnknownType", "Parameter type is not supported by Create Attribute.");
+				OutFailReason = LOCTEXT("ConversionAssignValueUnknownType", "Parameter type is not supported by Create Constant.");
 				return false;
 			}
 
@@ -226,7 +226,7 @@ TArray<FPCGPreconfiguredInfo> UPCGUserParameterGetSettings::GetConversionInfo()
 	ConversionInfo.Emplace(
 		SpecificToConstantIndex,
 		PCGConversion::Helpers::GetDefaultNodeTitle(UPCGCreateAttributeSetSettings::StaticClass()),
-		ToCreateAttributeTooltip);
+		ToCreateConstantTooltip);
 
 	return ConversionInfo;
 }
@@ -352,7 +352,7 @@ TArray<FPCGPreconfiguredInfo> UPCGGenericUserParameterGetSettings::GetConversion
 	ConversionInfo.Emplace(
 		GenericToConstantIndex,
 		PCGConversion::Helpers::GetDefaultNodeTitle(UPCGCreateAttributeSetSettings::StaticClass()),
-		ToCreateAttributeTooltip);
+		ToCreateConstantTooltip);
 
 	return ConversionInfo;
 }
diff --git a/Engine/Plugins/PCG/Source/PCG/Private/Metadata/PCGMetadataTypesConstantStruct.cpp b/Engine/Plugins/PCG/Source/PCG/Private/Metadata/PCGMetadataTypesConstantStruct.cpp
--- a/Engine/Plugins/PCG/Source/PCG/Private/Metadata/PCGMetadataTypesConstantStruct.cpp
+++ b/Engine/Plugins/PCG/Source/PCG/Private/Metadata/PCGMetadataTypesConstantStruct.cpp
@@ -43,6 +43,12 @@ FString FPCGMetadataTypesConstantStruct::ToString() const
 	return FString();
 }
 
+FString FPCGMetadataTypesConstantStruct::TypeToString() const
+{
+	const UEnum* EnumPtr = StaticEnum<EPCGMetadataTypes>();
+	return EnumPtr ? EnumPtr->GetNameStringByValue(static_cast<int64>(Type)) : FString{};
+}
+
 #if WITH_EDITOR
 void FPCGMetadataTypesConstantStruct::OnPostLoad()
 {
diff --git a/Engine/Plugins/PCG/Source/PCG/Public/Elements/PCGCreateAttribute.h b/Engine/Plugins/PCG/Source/PCG/Public/Elements/PCGCreateAttribute.h
--- a/Engine/Plugins/PCG/Source/PCG/Public/Elements/PCGCreateAttribute.h
+++ b/Engine/Plugins/PCG/Source/PCG/Public/Elements/PCGCreateAttribute.h
@@ -92,6 +92,8 @@ public:
 	virtual FName GetDefaultNodeName() const override;
 	virtual FText GetDefaultNodeTitle() const override;
 	virtual EPCGSettingsType GetType() const override { return EPCGSettingsType::Param; }
+	virtual TArray<FText> GetNodeTitleAliases() const override;
+	virtual FText GetNodeTooltipText() const override;
 #endif
 	virtual FString GetAdditionalTitleInformation() const override;
 	virtual void ApplyPreconfiguredSettings(const FPCGPreConfiguredSettingsInfo& PreconfigureInfo) override;
diff --git a/Engine/Plugins/PCG/Source/PCG/Public/Metadata/PCGMetadataTypesConstantStruct.h b/Engine/Plugins/PCG/Source/PCG/Public/Metadata/PCGMetadataTypesConstantStruct.h
--- a/Engine/Plugins/PCG/Source/PCG/Public/Metadata/PCGMetadataTypesConstantStruct.h
+++ b/Engine/Plugins/PCG/Source/PCG/Public/Metadata/PCGMetadataTypesConstantStruct.h
@@ -32,6 +32,7 @@ public:
 	decltype(auto) Dispatcher(Func Callback) const;
 
 	FString ToString() const;
+	FString TypeToString() const;
 
 #if WITH_EDITOR
 	/** This is not called by the engine because this is a USTRUCT, but it can be called from usages to perform deprecation. */
diff --git a/Engine/Plugins/PCG/Source/PCG/Public/PCGNode.h b/Engine/Plugins/PCG/Source/PCG/Public/PCGNode.h
--- a/Engine/Plugins/PCG/Source/PCG/Public/PCGNode.h
+++ b/Engine/Plugins/PCG/Source/PCG/Public/PCGNode.h
@@ -74,7 +74,7 @@ public:
 	/** Title to use if no title is authored. */
 	FText GetDefaultTitle() const;
 
-	/** Authored part of node title (like "Create Attribute 1"). */
+	/** Authored part of node title (like "Create Constant 1"). */
 	FText GetAuthoredTitleLine() const;
 
 	/** Authored node title as raw name or none if no title authored. */
