[PCG] Fixed gray-out mechanism in debug object tree view

diff --git a/Engine/Plugins/PCG/Source/PCGEditor/Private/Widgets/SPCGEditorGraphDebugObjectTree.cpp b/Engine/Plugins/PCG/Source/PCGEditor/Private/Widgets/SPCGEditorGraphDebugObjectTree.cpp
--- a/Engine/Plugins/PCG/Source/PCGEditor/Private/Widgets/SPCGEditorGraphDebugObjectTree.cpp
+++ b/Engine/Plugins/PCG/Source/PCGEditor/Private/Widgets/SPCGEditorGraphDebugObjectTree.cpp
@@ -866,6 +866,7 @@ void SPCGEditorGraphDebugObjectTree::AddStackToTree(const FPCGStack& Stack,
 		if (TSharedPtr<FPCGEditorGraphDebugObjectItem_Actor>* FoundActorItem = InOutActorItems.Find(Actor))
 		{
 			ActorItem = *FoundActorItem;
+			ActorItem->UpdateGrayedOut(bDisplayGrayedOut);
 		}
 		else
 		{
@@ -901,7 +902,11 @@ void SPCGEditorGraphDebugObjectTree::AddStackToTree(const FPCGStack& Stack,
 			FPCGStack GraphStack = Stack;
 			GraphStack.GetStackFramesMutable().SetNum(FrameCutoff, EAllowShrinking::No);
 
-			if (!InOutStackToItem.Contains(GraphStack))
+			if (FPCGEditorGraphDebugObjectItemPtr* ExistingGraphItem = InOutStackToItem.Find(GraphStack))
+			{
+				(*ExistingGraphItem)->UpdateGrayedOut(bInDisplayGrayedOut);
+			}
+			else
 			{
 				FPCGEditorGraphDebugObjectItemPtr GraphItem = InOutStackToItem.Emplace(
 					GraphStack,
@@ -938,6 +943,24 @@ void SPCGEditorGraphDebugObjectTree::AddStackToTree(const FPCGStack& Stack,
 					UnrollFrames.Pop(EAllowShrinking::No);
 				}
 			}
+
+			// If the data we're adding would not gray out the node, then we need to propagate upwards that information.
+			if (PCGNodeBeingInspected.IsValid() && !bDisplayGrayedOut)
+			{
+				while (!UnrollFrames.IsEmpty())
+				{
+					if (FPCGEditorGraphDebugObjectItemPtr* ExistingItem = InOutStackToItem.Find(UnrollStack))
+					{
+						// If we update something upstream and it was already not grayed out, then we can rest assured that anything upstream
+						// is already not grayed out
+						if (!(*ExistingItem)->UpdateGrayedOut(bDisplayGrayedOut))
+						{
+							break;
+						}
+					}
+					UnrollFrames.Pop(EAllowShrinking::No);
+				}
+			}
 		}
 
 		for(; FrameIndex < StackFrames.Num(); ++FrameIndex)
@@ -958,7 +981,11 @@ void SPCGEditorGraphDebugObjectTree::AddStackToTree(const FPCGStack& Stack,
 					FPCGStack GraphStack = Stack;
 					GraphStack.GetStackFramesMutable().SetNum(FrameIndex + 1);
 
-					if (!InOutStackToItem.Contains(GraphStack))
+					if (FPCGEditorGraphDebugObjectItemPtr* ExistingTopGraphItem = InOutStackToItem.Find(GraphStack))
+					{
+						(*ExistingTopGraphItem)->UpdateGrayedOut(bDisplayGrayedOut);
+					}
+					else
 					{
 						FPCGEditorGraphDebugObjectItemPtr TopGraphItem = InOutStackToItem.Emplace(
 							GraphStack,
@@ -987,7 +1014,11 @@ void SPCGEditorGraphDebugObjectTree::AddStackToTree(const FPCGStack& Stack,
 						FPCGStack LoopIterationStack = Stack;
 						LoopIterationStack.GetStackFramesMutable().SetNum(FrameIndex + 1);
 
-						if (!InOutStackToItem.Contains(LoopIterationStack))
+						if (FPCGEditorGraphDebugObjectItemPtr* ExistingLoopIterationItem = InOutStackToItem.Find(LoopIterationStack))
+						{
+							(*ExistingLoopIterationItem)->UpdateGrayedOut(bDisplayGrayedOut);
+						}
+						else
 						{
 							FPCGEditorGraphDebugObjectItemPtr LoopIterationItem = InOutStackToItem.Emplace(
 								LoopIterationStack,
diff --git a/Engine/Plugins/PCG/Source/PCGEditor/Private/Widgets/SPCGEditorGraphDebugObjectTree.h b/Engine/Plugins/PCG/Source/PCGEditor/Private/Widgets/SPCGEditorGraphDebugObjectTree.h
--- a/Engine/Plugins/PCG/Source/PCGEditor/Private/Widgets/SPCGEditorGraphDebugObjectTree.h
+++ b/Engine/Plugins/PCG/Source/PCGEditor/Private/Widgets/SPCGEditorGraphDebugObjectTree.h
@@ -39,6 +39,12 @@ public:
 	void SetExpanded(bool bInIsExpanded) { bIsExpanded = bInIsExpanded; }
 
 	bool IsGrayedOut() const { return bGrayedOut; }
+	bool UpdateGrayedOut(bool bInGrayedOut) 
+	{ 
+		const bool bWasGrayedOut = bGrayedOut;
+		bGrayedOut &= bInGrayedOut; 
+		return bWasGrayedOut;
+	}
 
 	bool IsSelected() const { return bSelected; }
 	void SetSelected(bool bInSelected) { bSelected = bInSelected; }
