[PCG] Fix check() when iterating over managed resources that were inaccessible.

#rnx

diff --git a/Engine/Plugins/PCG/Source/PCG/Private/Components/PCGProceduralISMComponent.cpp b/Engine/Plugins/PCG/Source/PCG/Private/Components/PCGProceduralISMComponent.cpp
--- a/Engine/Plugins/PCG/Source/PCG/Private/Components/PCGProceduralISMComponent.cpp
+++ b/Engine/Plugins/PCG/Source/PCG/Private/Components/PCGProceduralISMComponent.cpp
@@ -538,7 +538,7 @@ void UPCGProceduralISMComponent::OnRenderStateDirty(UActorComponent& InComponent
 	{
 		bool bManagesThis = false;
 
-		if (InComponent->bGenerated && InComponent->AreProceduralInstancesInUse())
+		if (InComponent->bGenerated && InComponent->AreProceduralInstancesInUse() && InComponent->AreManagedResourcesAccessible())
 		{
 			InComponent->ForEachManagedResource([this, InComponent, &bManagesThis](UPCGManagedResource* InResource)
 			{
diff --git a/Engine/Plugins/PCG/Source/PCG/Public/PCGComponent.h b/Engine/Plugins/PCG/Source/PCG/Public/PCGComponent.h
--- a/Engine/Plugins/PCG/Source/PCG/Public/PCGComponent.h
+++ b/Engine/Plugins/PCG/Source/PCG/Public/PCGComponent.h
@@ -167,6 +167,7 @@ public:
 	void AddActorsToManagedResources(const TArray<AActor*>& InActors);
 
 public:
+	bool AreManagedResourcesAccessible() const { return !GeneratedResourcesInaccessible; }
 	void ForEachManagedResource(TFunctionRef<void(UPCGManagedResource*)> InFunction);
 	void ForEachConstManagedResource(TFunctionRef<void(const UPCGManagedResource*)> InFunction) const;
 
