[PCG] Encapsulate PCGComponent inspection into FPCGGraphExecutionInspection so that this class can be reused

#rnx

diff --git a/Engine/Plugins/PCG/Source/PCG/Private/PCGComponent.cpp b/Engine/Plugins/PCG/Source/PCG/Private/PCGComponent.cpp
--- a/Engine/Plugins/PCG/Source/PCG/Private/PCGComponent.cpp
+++ b/Engine/Plugins/PCG/Source/PCG/Private/PCGComponent.cpp
@@ -1901,6 +1901,17 @@ void UPCGComponent::OnComponentDestroyed(bool bDestroyingHierarchy)
 	Super::OnComponentDestroyed(bDestroyingHierarchy);
 }
 
+#if WITH_EDITOR
+void UPCGComponent::AddReferencedObjects(UObject* InThis, FReferenceCollector& Collector)
+{
+	UPCGComponent* This = CastChecked<UPCGComponent>(InThis);
+
+	This->ExecutionInspection.AddReferencedObjects(Collector);
+
+	Super::AddReferencedObjects(This, Collector);
+}
+#endif // WITH_EDITOR
+
 void UPCGComponent::Serialize(FArchive& Ar)
 {
 #if WITH_EDITOR
@@ -2617,235 +2628,89 @@ void UPCGComponent::ResetLastGeneratedBounds()
 	LastGeneratedBounds = FBox(EForceInit::ForceInit);
 }
 
+// UE_DEPRECATED(5.6)
 bool UPCGComponent::IsInspecting() const
 {
-	return InspectionCounter > 0;
+	return ExecutionInspection.IsInspecting();
 }
 
+// UE_DEPRECATED(5.6)
 void UPCGComponent::EnableInspection()
 {
-	if (!ensure(InspectionCounter >= 0))
-	{
-		InspectionCounter = 0;
-	}
-	
-	InspectionCounter++;
+	ExecutionInspection.EnableInspection();
 }
 
+// UE_DEPRECATED(5.6)
 void UPCGComponent::DisableInspection()
 {
-	if (ensure(InspectionCounter > 0))
-	{
-		InspectionCounter--;
-	}
-	
-	if (InspectionCounter == 0)
-	{
-		ClearInspectionData(/*bClearPerNodeExecutionData=*/false);
-	}
+	ExecutionInspection.DisableInspection();
 };
 
+// UE_DEPRECATED(5.6)
 void UPCGComponent::NotifyNodeExecuted(const UPCGNode* InNode, const FPCGStack* InStack, const PCGUtils::FCallTime* InTimer, bool bNodeUsedCache)
 {
-	if (!ensure(InStack && InNode))
-	{
-		return;
-	}
-
-	// Reset timer information if taken from cache to provide good info in the profiling window
-	PCGUtils::FCallTime Timer;
-	if (InTimer && !bNodeUsedCache)
-	{
-		Timer = *InTimer;
-	}
-
-	FWriteScopeLock Lock(NodeToStacksInWhichNodeExecutedLock);
-	NodeToStacksInWhichNodeExecuted.FindOrAdd(InNode).Add(UPCGComponent::NodeExecutedNotificationData(*InStack, MoveTemp(Timer)));
+	ExecutionInspection.NotifyNodeExecuted(InNode, InStack, InTimer, bNodeUsedCache);
 }
 
+// UE_DEPRECATED(5.6)
+PRAGMA_DISABLE_DEPRECATION_WARNINGS
 TMap<TObjectKey<const UPCGNode>, TSet<UPCGComponent::NodeExecutedNotificationData>> UPCGComponent::GetExecutedNodeStacks() const
 {
-	FReadScopeLock Lock(NodeToStacksInWhichNodeExecutedLock);
-	return NodeToStacksInWhichNodeExecuted;
-}
-
-uint64 UPCGComponent::GetNodeInactivePinMask(const UPCGNode* InNode, const FPCGStack& Stack) const
-{
-	FReadScopeLock Lock(NodeToStackToInactivePinMaskLock);
-
-	if (const TMap<const FPCGStack, uint64>* StackToMask = NodeToStackToInactivePinMask.Find(InNode))
+	TMap<TObjectKey<const UPCGNode>, TSet<FPCGGraphExecutionInspection::FNodeExecutedNotificationData>> ExecutedNodeStacks = ExecutionInspection.GetExecutedNodeStacks();
+	TMap<TObjectKey<const UPCGNode>, TSet<UPCGComponent::NodeExecutedNotificationData>> DeprecatedExecutedNodeStacks;
+	for(auto& Pair : ExecutedNodeStacks)
 	{
-		if (const uint64* Mask = StackToMask->Find(Stack))
+		TSet<UPCGComponent::NodeExecutedNotificationData>& DeprecatedNotifications = DeprecatedExecutedNodeStacks.Add(Pair.Key);
+		Algo::Transform(Pair.Value, DeprecatedNotifications, [](const FPCGGraphExecutionInspection::FNodeExecutedNotificationData& Notification)
 		{
-			return *Mask;
-		}
+			return UPCGComponent::NodeExecutedNotificationData(Notification.Stack, Notification.Timer);
+		});
 	}
 
-	return 0;
+	return DeprecatedExecutedNodeStacks;
 }
+PRAGMA_ENABLE_DEPRECATION_WARNINGS
 
-void UPCGComponent::NotifyNodeDynamicInactivePins(const UPCGNode* InNode, const FPCGStack* InStack, uint64 InactivePinBitmask) const
+// UE_DEPRECATED(5.6)
+uint64 UPCGComponent::GetNodeInactivePinMask(const UPCGNode* InNode, const FPCGStack& Stack) const
 {
-	if (!ensure(InStack && InNode))
-	{
-		return;
-	}
+	return ExecutionInspection.GetNodeInactivePinMask(InNode, Stack);
+}
 
-	FWriteScopeLock Lock(NodeToStackToInactivePinMaskLock);
-	TMap<const FPCGStack, uint64>& StackToInactivePinMask = NodeToStackToInactivePinMask.FindOrAdd(InNode);
-	StackToInactivePinMask.FindOrAdd(*InStack) = InactivePinBitmask;
+// UE_DEPRECATED(5.6)
+void UPCGComponent::NotifyNodeDynamicInactivePins(const UPCGNode* InNode, const FPCGStack* InStack, uint64 InactivePinBitmask) const
+{
+	ExecutionInspection.NotifyNodeDynamicInactivePins(InNode, InStack, InactivePinBitmask);
 }
 
+// UE_DEPRECATED(5.6)
 bool UPCGComponent::WasNodeExecuted(const UPCGNode* InNode, const FPCGStack& Stack) const
 {
-	FReadScopeLock Lock(NodeToStacksInWhichNodeExecutedLock);
-	const TSet<UPCGComponent::NodeExecutedNotificationData>* FoundNotifications = NodeToStacksInWhichNodeExecuted.Find(InNode);
-
-	// Since the operator== & hash functions don't rely on the timer, we can just build a stub from the stack.
-	UPCGComponent::NodeExecutedNotificationData NotificationStub(Stack, PCGUtils::FCallTime());
-	return FoundNotifications && FoundNotifications->Contains(NotificationStub);
+	return ExecutionInspection.WasNodeExecuted(InNode, Stack);
 }
 
+// UE_DEPRECATED(5.6)
 void UPCGComponent::StoreInspectionData(const FPCGStack* InStack, const UPCGNode* InNode, const PCGUtils::FCallTime* InTimer, const FPCGDataCollection& InInputData, const FPCGDataCollection& InOutputData, bool bUsedCache)
 {
-	if (!InNode || !ensure(InStack))
-	{
-		return;
-	}
-
-	InInputData.MarkUsage(EPCGDataUsage::ComponentInspectionData);
-	InOutputData.MarkUsage(EPCGDataUsage::ComponentInspectionData);
-
-	// Notify component that this task executed. Useful for editor visualization.
-	NotifyNodeExecuted(InNode, InStack, InTimer, bUsedCache);
-
-	if (!InOutputData.TaggedData.IsEmpty())
-	{
-		FWriteScopeLock Lock(NodeToStacksThatProducedDataLock);
-
-		NodeToStacksThatProducedData.FindOrAdd(InNode).Add(*InStack);
-	}
-	else
-	{
-		FWriteScopeLock Lock(NodeToStacksThatProducedDataLock);
-
-		if (TSet<FPCGStack>* Stacks = NodeToStacksThatProducedData.Find(InNode))
-		{
-			Stacks->Remove(*InStack);
-		}
-	}
-
-	if (IsInspecting())
-	{
-		auto StorePinInspectionDataFromNode = [](const FPCGStack* InStack, const TArray<TObjectPtr<UPCGPin>>& InPins, const FPCGDataCollection& InData, TMap<FPCGStack, FPCGDataCollection>& InOutInspectionCache)
-		{
-			for (const UPCGPin* Pin : InPins)
-			{
-				FPCGStack Stack = *InStack;
-
-				// Append the Node and Pin to the current Stack to uniquely identify each DataCollection
-				TArray<FPCGStackFrame>& StackFrames = Stack.GetStackFramesMutable();
-				StackFrames.Emplace(Pin);
-
-				FPCGDataCollection PinDataCollection;
-				InData.GetInputsAndCrcsByPin(Pin->Properties.Label, PinDataCollection.TaggedData, PinDataCollection.DataCrcs);
-
-				// Implementation note: since static subgraphs actually are visited twice and the second time the input doesn't match the input pins, we don't clear the data.
-				if (!PinDataCollection.TaggedData.IsEmpty())
-				{
-					if (FPCGDataCollection* CollectionInCache = InOutInspectionCache.Find(Stack))
-					{
-						CollectionInCache->TaggedData.Append(PinDataCollection.TaggedData);
-					}
-					else
-					{
-						InOutInspectionCache.Add(Stack, PinDataCollection);
-					}
-				}
-			}
-		};
-
-		auto StorePinInspectionData = [InStack, InNode, &StorePinInspectionDataFromNode](const TArray<TObjectPtr<UPCGPin>>& InPins, const FPCGDataCollection& InData, TMap<FPCGStack, FPCGDataCollection>& InOutInspectionCache)
-		{
-			FPCGStack Stack = *InStack;
-
-			// Append the Node (here) and Pin (in call) to the current Stack to uniquely identify each DataCollection
-			TArray<FPCGStackFrame>& StackFrames = Stack.GetStackFramesMutable();
-			StackFrames.Reserve(StackFrames.Num() + 2);
-			StackFrames.Emplace(InNode);
-
-			StorePinInspectionDataFromNode(&Stack, InPins, InData, InOutInspectionCache);
-		};
-
-		FWriteScopeLock Lock(InspectionCacheLock);
-
-		// Special case: if we have a static (embedded) subgraph, then the actual data inputs (not params) of the subgraph will be on the input node.
-		// Considering we don't allow inspection on input pins of the input node, then we can move that data up the chain.
-		if (InNode->GetSettings()->IsA<UPCGGraphInputOutputSettings>() && InStack->GetStackFrames().Num() > 2)
-		{
-			// We're expecting the last frame to be the graph
-			// Then, if the graph was statically dispatched, it will be the subgraph node.
-			// In the case of a dynamic subgraph or loop, it will be the loop index instead.
-			FPCGStack StackToSubgraphNode = *InStack;
-			TArray<FPCGStackFrame>& StackFrames = StackToSubgraphNode.GetStackFramesMutable();
-			StackFrames.Pop();
-
-			if (const UPCGSubgraphNode* Node = StackFrames.Last().GetObject_AnyThread<UPCGSubgraphNode>())
-			{
-				StorePinInspectionDataFromNode(&StackToSubgraphNode, Node->GetInputPins(), InInputData, InspectionCache);
-			}
-		}
-
-		StorePinInspectionData(InNode->GetInputPins(), InInputData, InspectionCache);
-		StorePinInspectionData(InNode->GetOutputPins(), InOutputData, InspectionCache);
-	}
+	ExecutionInspection.StoreInspectionData(InStack, InNode, InTimer, InInputData, InOutputData, bUsedCache);
 }
 
+// UE_DEPRECATED(5.6)
 const FPCGDataCollection* UPCGComponent::GetInspectionData(const FPCGStack& InStack) const
 {
-	FReadScopeLock Lock(InspectionCacheLock);
-	return InspectionCache.Find(InStack);
+	return ExecutionInspection.GetInspectionData(InStack);
 }
 
+// UE_DEPRECATED(5.6)
 void UPCGComponent::ClearInspectionData(bool bClearPerNodeExecutionData)
 {
-	{
-		FWriteScopeLock Lock(InspectionCacheLock);
-
-		for (TPair<FPCGStack, FPCGDataCollection>& Entry : InspectionCache)
-		{
-			Entry.Value.ClearUsage(EPCGDataUsage::ComponentInspectionData);
-		}
-
-		InspectionCache.Reset();
-	}
-
-	if (bClearPerNodeExecutionData)
-	{
-		{
-			FWriteScopeLock Lock(NodeToStacksThatProducedDataLock);
-			NodeToStacksThatProducedData.Reset();
-		}
-
-		{
-			FWriteScopeLock Lock(NodeToStacksInWhichNodeExecutedLock);
-			NodeToStacksInWhichNodeExecuted.Reset();
-		}
-
-		{
-			FWriteScopeLock Lock(NodeToStackToInactivePinMaskLock);
-			NodeToStackToInactivePinMask.Reset();
-		}
-	}
+	ExecutionInspection.ClearInspectionData(bClearPerNodeExecutionData);
 }
 
+// UE_DEPRECATED(5.6)
 bool UPCGComponent::HasNodeProducedData(const UPCGNode* InNode, const FPCGStack& Stack) const
 {
-	FReadScopeLock Lock(NodeToStacksThatProducedDataLock);
-
-	const TSet<FPCGStack>* StacksThatProducedData = NodeToStacksThatProducedData.Find(InNode);
-
-	return StacksThatProducedData && StacksThatProducedData->Contains(Stack);
+	return ExecutionInspection.HasNodeProducedData(InNode, Stack);
 }
 
 void UPCGComponent::Refresh(EPCGChangeType ChangeType, bool bCancelExistingRefresh)
@@ -4133,7 +3998,8 @@ void FPCGComponentInstanceData::ApplyToComponent(UActorComponent* Component, con
 			}
 
 			PCGComponent->bWasGeneratedThisSession = ConstructionSourceComponent->bWasGeneratedThisSession;
-			PCGComponent->InspectionCounter = ConstructionSourceComponent->InspectionCounter;
+
+			PCGComponent->ExecutionInspection.InspectionCounter = ConstructionSourceComponent->ExecutionInspection.InspectionCounter;
 #endif
 		}
 
diff --git a/Engine/Plugins/PCG/Source/PCG/Private/PCGGraphExecutionInspection.cpp b/Engine/Plugins/PCG/Source/PCG/Private/PCGGraphExecutionInspection.cpp
new file mode 100644
--- /dev/null
+++ b/Engine/Plugins/PCG/Source/PCG/Private/PCGGraphExecutionInspection.cpp
@@ -0,0 +1,251 @@
+// Copyright Epic Games, Inc. All Rights Reserved.
+
+#include "PCGGraphExecutionInspection.h"
+
+#if WITH_EDITOR
+
+#include "PCGInputOutputSettings.h"
+#include "PCGPin.h"
+#include "PCGSettings.h"
+#include "PCGSubgraph.h"
+
+bool FPCGGraphExecutionInspection::IsInspecting() const
+{
+	return InspectionCounter > 0;
+}
+
+void FPCGGraphExecutionInspection::EnableInspection()
+{
+	if (!ensure(InspectionCounter >= 0))
+	{
+		InspectionCounter = 0;
+	}
+
+	InspectionCounter++;
+}
+
+void FPCGGraphExecutionInspection::DisableInspection()
+{
+	if (ensure(InspectionCounter > 0))
+	{
+		InspectionCounter--;
+	}
+
+	if (InspectionCounter == 0)
+	{
+		ClearInspectionData(/*bClearPerNodeExecutionData=*/false);
+	}
+};
+
+void FPCGGraphExecutionInspection::NotifyNodeExecuted(const UPCGNode* InNode, const FPCGStack* InStack, const PCGUtils::FCallTime* InTimer, bool bNodeUsedCache)
+{
+	if (!ensure(InStack && InNode))
+	{
+		return;
+	}
+
+	// Reset timer information if taken from cache to provide good info in the profiling window
+	PCGUtils::FCallTime Timer;
+	if (InTimer && !bNodeUsedCache)
+	{
+		Timer = *InTimer;
+	}
+
+	FWriteScopeLock Lock(NodeToStacksInWhichNodeExecutedLock);
+	NodeToStacksInWhichNodeExecuted.FindOrAdd(InNode).Add(FPCGGraphExecutionInspection::FNodeExecutedNotificationData(*InStack, MoveTemp(Timer)));
+}
+
+TMap<TObjectKey<const UPCGNode>, TSet<FPCGGraphExecutionInspection::FNodeExecutedNotificationData>> FPCGGraphExecutionInspection::GetExecutedNodeStacks() const
+{
+	FReadScopeLock Lock(NodeToStacksInWhichNodeExecutedLock);
+	return NodeToStacksInWhichNodeExecuted;
+}
+
+uint64 FPCGGraphExecutionInspection::GetNodeInactivePinMask(const UPCGNode* InNode, const FPCGStack& Stack) const
+{
+	FReadScopeLock Lock(NodeToStackToInactivePinMaskLock);
+
+	if (const TMap<const FPCGStack, uint64>* StackToMask = NodeToStackToInactivePinMask.Find(InNode))
+	{
+		if (const uint64* Mask = StackToMask->Find(Stack))
+		{
+			return *Mask;
+		}
+	}
+
+	return 0;
+}
+
+void FPCGGraphExecutionInspection::NotifyNodeDynamicInactivePins(const UPCGNode* InNode, const FPCGStack* InStack, uint64 InactivePinBitmask) const
+{
+	if (!ensure(InStack && InNode))
+	{
+		return;
+	}
+
+	FWriteScopeLock Lock(NodeToStackToInactivePinMaskLock);
+	TMap<const FPCGStack, uint64>& StackToInactivePinMask = NodeToStackToInactivePinMask.FindOrAdd(InNode);
+	StackToInactivePinMask.FindOrAdd(*InStack) = InactivePinBitmask;
+}
+
+bool FPCGGraphExecutionInspection::WasNodeExecuted(const UPCGNode* InNode, const FPCGStack& Stack) const
+{
+	FReadScopeLock Lock(NodeToStacksInWhichNodeExecutedLock);
+	const TSet<FPCGGraphExecutionInspection::FNodeExecutedNotificationData>* FoundNotifications = NodeToStacksInWhichNodeExecuted.Find(InNode);
+
+	// Since the operator== & hash functions don't rely on the timer, we can just build a stub from the stack.
+	FPCGGraphExecutionInspection::FNodeExecutedNotificationData NotificationStub(Stack, PCGUtils::FCallTime());
+	return FoundNotifications && FoundNotifications->Contains(NotificationStub);
+}
+
+void FPCGGraphExecutionInspection::StoreInspectionData(const FPCGStack* InStack, const UPCGNode* InNode, const PCGUtils::FCallTime* InTimer, const FPCGDataCollection& InInputData, const FPCGDataCollection& InOutputData, bool bUsedCache)
+{
+	if (!InNode || !ensure(InStack))
+	{
+		return;
+	}
+
+	InInputData.MarkUsage(EPCGDataUsage::ComponentInspectionData);
+	InOutputData.MarkUsage(EPCGDataUsage::ComponentInspectionData);
+
+	// Notify component that this task executed. Useful for editor visualization.
+	NotifyNodeExecuted(InNode, InStack, InTimer, bUsedCache);
+
+	if (!InOutputData.TaggedData.IsEmpty())
+	{
+		FWriteScopeLock Lock(NodeToStacksThatProducedDataLock);
+
+		NodeToStacksThatProducedData.FindOrAdd(InNode).Add(*InStack);
+	}
+	else
+	{
+		FWriteScopeLock Lock(NodeToStacksThatProducedDataLock);
+
+		if (TSet<FPCGStack>* Stacks = NodeToStacksThatProducedData.Find(InNode))
+		{
+			Stacks->Remove(*InStack);
+		}
+	}
+
+	if (IsInspecting())
+	{
+		auto StorePinInspectionDataFromNode = [](const FPCGStack* InStack, const TArray<TObjectPtr<UPCGPin>>& InPins, const FPCGDataCollection& InData, TMap<FPCGStack, FPCGDataCollection>& InOutInspectionCache)
+		{
+			for (const UPCGPin* Pin : InPins)
+			{
+				FPCGStack Stack = *InStack;
+
+				// Append the Node and Pin to the current Stack to uniquely identify each DataCollection
+				TArray<FPCGStackFrame>& StackFrames = Stack.GetStackFramesMutable();
+				StackFrames.Emplace(Pin);
+
+				FPCGDataCollection PinDataCollection;
+				InData.GetInputsAndCrcsByPin(Pin->Properties.Label, PinDataCollection.TaggedData, PinDataCollection.DataCrcs);
+
+				// Implementation note: since static subgraphs actually are visited twice and the second time the input doesn't match the input pins, we don't clear the data.
+				if (!PinDataCollection.TaggedData.IsEmpty())
+				{
+					if (FPCGDataCollection* CollectionInCache = InOutInspectionCache.Find(Stack))
+					{
+						CollectionInCache->TaggedData.Append(PinDataCollection.TaggedData);
+					}
+					else
+					{
+						InOutInspectionCache.Add(Stack, PinDataCollection);
+					}
+				}
+			}
+		};
+
+		auto StorePinInspectionData = [InStack, InNode, &StorePinInspectionDataFromNode](const TArray<TObjectPtr<UPCGPin>>& InPins, const FPCGDataCollection& InData, TMap<FPCGStack, FPCGDataCollection>& InOutInspectionCache)
+		{
+			FPCGStack Stack = *InStack;
+
+			// Append the Node (here) and Pin (in call) to the current Stack to uniquely identify each DataCollection
+			TArray<FPCGStackFrame>& StackFrames = Stack.GetStackFramesMutable();
+			StackFrames.Reserve(StackFrames.Num() + 2);
+			StackFrames.Emplace(InNode);
+
+			StorePinInspectionDataFromNode(&Stack, InPins, InData, InOutInspectionCache);
+		};
+
+		FWriteScopeLock Lock(InspectionCacheLock);
+
+		// Special case: if we have a static (embedded) subgraph, then the actual data inputs (not params) of the subgraph will be on the input node.
+		// Considering we don't allow inspection on input pins of the input node, then we can move that data up the chain.
+		if (InNode->GetSettings()->IsA<UPCGGraphInputOutputSettings>() && InStack->GetStackFrames().Num() > 2)
+		{
+			// We're expecting the last frame to be the graph
+			// Then, if the graph was statically dispatched, it will be the subgraph node.
+			// In the case of a dynamic subgraph or loop, it will be the loop index instead.
+			FPCGStack StackToSubgraphNode = *InStack;
+			TArray<FPCGStackFrame>& StackFrames = StackToSubgraphNode.GetStackFramesMutable();
+			StackFrames.Pop();
+
+			if (const UPCGSubgraphNode* Node = StackFrames.Last().GetObject_AnyThread<UPCGSubgraphNode>())
+			{
+				StorePinInspectionDataFromNode(&StackToSubgraphNode, Node->GetInputPins(), InInputData, InspectionCache);
+			}
+		}
+
+		StorePinInspectionData(InNode->GetInputPins(), InInputData, InspectionCache);
+		StorePinInspectionData(InNode->GetOutputPins(), InOutputData, InspectionCache);
+	}
+}
+
+const FPCGDataCollection* FPCGGraphExecutionInspection::GetInspectionData(const FPCGStack& InStack) const
+{
+	FReadScopeLock Lock(InspectionCacheLock);
+	return InspectionCache.Find(InStack);
+}
+
+void FPCGGraphExecutionInspection::ClearInspectionData(bool bClearPerNodeExecutionData)
+{
+	{
+		FWriteScopeLock Lock(InspectionCacheLock);
+
+		for (TPair<FPCGStack, FPCGDataCollection>& Entry : InspectionCache)
+		{
+			Entry.Value.ClearUsage(EPCGDataUsage::ComponentInspectionData);
+		}
+
+		InspectionCache.Reset();
+	}
+
+	if (bClearPerNodeExecutionData)
+	{
+		{
+			FWriteScopeLock Lock(NodeToStacksThatProducedDataLock);
+			NodeToStacksThatProducedData.Reset();
+		}
+
+		{
+			FWriteScopeLock Lock(NodeToStacksInWhichNodeExecutedLock);
+			NodeToStacksInWhichNodeExecuted.Reset();
+		}
+
+		{
+			FWriteScopeLock Lock(NodeToStackToInactivePinMaskLock);
+			NodeToStackToInactivePinMask.Reset();
+		}
+	}
+}
+
+bool FPCGGraphExecutionInspection::HasNodeProducedData(const UPCGNode* InNode, const FPCGStack& Stack) const
+{
+	FReadScopeLock Lock(NodeToStacksThatProducedDataLock);
+
+	const TSet<FPCGStack>* StacksThatProducedData = NodeToStacksThatProducedData.Find(InNode);
+
+	return StacksThatProducedData && StacksThatProducedData->Contains(Stack);
+}
+
+void FPCGGraphExecutionInspection::AddReferencedObjects(FReferenceCollector& Collector)
+{
+	for (auto& Pair : InspectionCache)
+	{
+		Pair.Value.AddReferences(Collector);
+	}
+}
+
+#endif // WITH_EDITOR
\ No newline at end of file
diff --git a/Engine/Plugins/PCG/Source/PCG/Private/Utils/PCGExtraCapture.cpp b/Engine/Plugins/PCG/Source/PCG/Private/Utils/PCGExtraCapture.cpp
--- a/Engine/Plugins/PCG/Source/PCG/Private/Utils/PCGExtraCapture.cpp
+++ b/Engine/Plugins/PCG/Source/PCG/Private/Utils/PCGExtraCapture.cpp
@@ -150,7 +150,7 @@ PCGUtils::FCallTreeInfo PCGUtils::FExtraCapture::CalculateCallTreeInfo(const UPC
 	FCallTreeInfo RootInfo;
 
 	// Basically, what we want is - visit all entries in the "NodeToStacksInWhichNodeExecuted" and build our information from there.
-	TMap<TObjectKey<const UPCGNode>, TSet<UPCGComponent::NodeExecutedNotificationData>> NodeToStacksInWhichNodeExecuted = Component->GetExecutedNodeStacks();
+	TMap<TObjectKey<const UPCGNode>, TSet<FPCGGraphExecutionInspection::FNodeExecutedNotificationData>> NodeToStacksInWhichNodeExecuted = Component->GetExecutionInspection().GetExecutedNodeStacks();
 
 	TArray<const UPCGNode*> NodePath;
 	TArray<int32> NodePathLoop;
@@ -224,9 +224,9 @@ PCGUtils::FCallTreeInfo PCGUtils::FExtraCapture::CalculateCallTreeInfo(const UPC
 	for (const auto& NodeToStacks : NodeToStacksInWhichNodeExecuted)
 	{
 		const UPCGNode* Node = NodeToStacks.Key.ResolveObjectPtr();
-		const TSet<UPCGComponent::NodeExecutedNotificationData>& NodeExecutions = NodeToStacks.Value;
+		const TSet<FPCGGraphExecutionInspection::FNodeExecutedNotificationData>& NodeExecutions = NodeToStacks.Value;
 
-		for(const UPCGComponent::NodeExecutedNotificationData& NodeExecution : NodeExecutions)
+		for(const FPCGGraphExecutionInspection::FNodeExecutedNotificationData& NodeExecution : NodeExecutions)
 		{
 			const FPCGStack& Stack = NodeExecution.Stack;
 
diff --git a/Engine/Plugins/PCG/Source/PCG/Public/PCGComponent.h b/Engine/Plugins/PCG/Source/PCG/Public/PCGComponent.h
--- a/Engine/Plugins/PCG/Source/PCG/Public/PCGComponent.h
+++ b/Engine/Plugins/PCG/Source/PCG/Public/PCGComponent.h
@@ -3,6 +3,7 @@
 #pragma once
 
 #include "PCGCommon.h"
+#include "PCGGraphExecutionInspection.h"
 #include "PCGNode.h"
 #include "PCGSettings.h"
 #include "Graph/PCGStackContext.h"
@@ -113,6 +114,8 @@ public:
 	virtual bool CanEditChange(const FProperty* InProperty) const override;
 	virtual void PostEditImport() override;
 	virtual void PreSave(FObjectPreSaveContext ObjectSaveContext) override;
+
+	static void AddReferencedObjects(UObject* InThis, FReferenceCollector& Collector);
 #endif
 	/** ~End UObject interface */
 
@@ -378,32 +381,6 @@ public:
 
 	/** Functions for managing the node inspection cache */
 	bool WasGeneratedThisSession() const { return bWasGeneratedThisSession; }
-	bool IsInspecting() const;
-	void EnableInspection();
-	void DisableInspection();
-	void StoreInspectionData(const FPCGStack* InStack, const UPCGNode* InNode, const PCGUtils::FCallTime* InTimer, const FPCGDataCollection& InInputData, const FPCGDataCollection& InOutputData, bool bUsedCache);
-	const FPCGDataCollection* GetInspectionData(const FPCGStack& InStack) const;
-	void ClearInspectionData(bool bClearPerNodeExecutionData = true);
-
-	/** Whether a task for the given node and stack was executed during the last execution. */
-	bool WasNodeExecuted(const UPCGNode* InNode, const FPCGStack& Stack) const;
-
-	/** Called at execution time each time a node has been executed. */
-	void NotifyNodeExecuted(const UPCGNode* InNode, const FPCGStack* InStack, const PCGUtils::FCallTime* InTimer, bool bNodeUsedCache);
-
-	/* Retrieves the executed nodes information */
-	struct NodeExecutedNotificationData
-	{
-		NodeExecutedNotificationData(const FPCGStack& InStack, const PCGUtils::FCallTime& InTimer) : Stack(InStack), Timer(InTimer) {}
-		// Important implementation note: some logic in WasNodeExecuted relies on the fact we don't use the timer for the operator== and hash functions.
-		friend PCG_API uint32 GetTypeHash(const NodeExecutedNotificationData& NotifData) { return GetTypeHash(NotifData.Stack); }
-		bool operator==(const NodeExecutedNotificationData& OtherNotifData) const { return Stack == OtherNotifData.Stack; }
-
-		FPCGStack Stack;
-		PCGUtils::FCallTime Timer;
-	};
-
-	TMap<TObjectKey<const UPCGNode>, TSet<NodeExecutedNotificationData>> GetExecutedNodeStacks() const;
 
 	/** Retrieve the inactive pin bitmask for the given node and stack in the last execution. */
 	uint64 GetNodeInactivePinMask(const UPCGNode* InNode, const FPCGStack& Stack) const;
@@ -514,6 +491,53 @@ public:
 	/** Builds the canonical PCG data from a given actor and its PCG component if any. */
 	static UPCGData* CreateActorPCGData(AActor* Actor, const UPCGComponent* Component, bool bParseActor = true);
 
+#if WITH_EDITOR
+
+	const FPCGGraphExecutionInspection& GetExecutionInspection() const { return ExecutionInspection; }
+
+	// todo_pcg: uncomment deprecation once this is moved to upcoming execution state
+	//UE_DEPRECATED(5.6, "Use FPCGGraphExecutionInspection::IsInspecting() instead")
+	bool IsInspecting() const;
+
+	//UE_DEPRECATED(5.6, "Use FPCGGraphExecutionInspection::EnableInspection() instead")
+	void EnableInspection();
+
+	//UE_DEPRECATED(5.6, "Use FPCGGraphExecutionInspection::DisableInspection() instead")
+	void DisableInspection();
+
+	//UE_DEPRECATED(5.6, "Use FPCGGraphExecutionInspection::StoreInspectionData() instead")
+	void StoreInspectionData(const FPCGStack* InStack, const UPCGNode* InNode, const PCGUtils::FCallTime* InTimer, const FPCGDataCollection& InInputData, const FPCGDataCollection& InOutputData, bool bUsedCache);
+
+	//UE_DEPRECATED(5.6, "Use FPCGGraphExecutionInspection::GetInspectionData() instead")
+	const FPCGDataCollection* GetInspectionData(const FPCGStack& InStack) const;
+
+	//UE_DEPRECATED(5.6, "Use FPCGGraphExecutionInspection::ClearInspectionData() instead")
+	void ClearInspectionData(bool bClearPerNodeExecutionData = true);
+
+	//UE_DEPRECATED(5.6, "Use FPCGGraphExecutionInspection::WasNodeExecuted() instead")
+	bool WasNodeExecuted(const UPCGNode* InNode, const FPCGStack& Stack) const;
+
+	//UE_DEPRECATED(5.6, "Use FPCGGraphExecutionInspection::NotifyNodeExecuted() instead")
+	void NotifyNodeExecuted(const UPCGNode* InNode, const FPCGStack* InStack, const PCGUtils::FCallTime* InTimer, bool bNodeUsedCache);
+		
+	struct UE_DEPRECATED(5.6, "Use FPCGGraphExecutionInspection::FNodeExecutionNotificationData instead") NodeExecutedNotificationData
+	{
+		NodeExecutedNotificationData(const FPCGStack& InStack, const PCGUtils::FCallTime& InTimer) : Stack(InStack), Timer(InTimer) {}
+		// Important implementation note: some logic in WasNodeExecuted relies on the fact we don't use the timer for the operator== and hash functions.
+		friend PCG_API uint32 GetTypeHash(const NodeExecutedNotificationData& NotifData) { return GetTypeHash(NotifData.Stack); }
+		bool operator==(const NodeExecutedNotificationData& OtherNotifData) const { return Stack == OtherNotifData.Stack; }
+
+		FPCGStack Stack;
+		PCGUtils::FCallTime Timer;
+	};
+
+PRAGMA_DISABLE_DEPRECATION_WARNINGS
+	UE_DEPRECATED(5.6, "Use FPCGGraphExecutionInspection::GetExecutedNodeStacks() instead")
+	TMap<TObjectKey<const UPCGNode>, TSet<NodeExecutedNotificationData>> GetExecutedNodeStacks() const;
+PRAGMA_ENABLE_DEPRECATION_WARNINGS
+
+#endif // WITH_EDITOR
+
 protected:
 	FPCGGridDescriptor GetGridDescriptorInternal(uint32 GridSize, bool bRuntimeHashUpdate) const;
 	
@@ -695,7 +719,6 @@ private:
 
 #if WITH_EDITOR
 	bool bWasGeneratedThisSession = false;
-	int32 InspectionCounter = 0;
 	FBox LastGeneratedBoundsPriorToUndo = FBox(EForceInit::ForceInit);
 #endif
 
@@ -711,25 +734,10 @@ private:
 	// Need to keep a reference to all tracked settings to still react to changes after a map load (since the component won't have been executed).
 	// Serialization will be done in the Serialize function
 	FPCGSelectionKeyToSettingsMap DynamicallyTrackedKeysToSettings;
-
-	UPROPERTY(Transient, NonTransactional)
-	TMap<FPCGStack, FPCGDataCollection> InspectionCache;
 #endif
 
 #if WITH_EDITOR
-	mutable FRWLock InspectionCacheLock;
-
-	/** Map from nodes to all stacks for which the node produced at least one data item. */
-	TMap<TObjectKey<const UPCGNode>, TSet<FPCGStack>> NodeToStacksThatProducedData;
-	mutable FRWLock NodeToStacksThatProducedDataLock;
-
-	/** Map from nodes to all stacks for which a task for the node was executed. */
-	TMap<TObjectKey<const UPCGNode>, TSet<NodeExecutedNotificationData>> NodeToStacksInWhichNodeExecuted;
-	mutable FRWLock NodeToStacksInWhichNodeExecutedLock;
-
-	/** Map from nodes to stacks to mask of output pins that were deactivated during execution. */
-	mutable TMap<TObjectKey<const UPCGNode>, TMap<const FPCGStack, uint64>> NodeToStackToInactivePinMask;
-	mutable FRWLock NodeToStackToInactivePinMaskLock;
+	FPCGGraphExecutionInspection ExecutionInspection;
 
 	/** The tracking system will not trigger a generation on this component for these change origins. Populated within the scope
 	* of an element. Entries removed when counter is decremented to 0, so empty map means no active ignores.
diff --git a/Engine/Plugins/PCG/Source/PCG/Public/PCGGraphExecutionInspection.h b/Engine/Plugins/PCG/Source/PCG/Public/PCGGraphExecutionInspection.h
new file mode 100644
--- /dev/null
+++ b/Engine/Plugins/PCG/Source/PCG/Public/PCGGraphExecutionInspection.h
@@ -0,0 +1,77 @@
+// Copyright Epic Games, Inc. All Rights Reserved.
+
+#pragma once
+
+#include "PCGCommon.h"
+#include "PCGData.h"
+#include "PCGNode.h"
+#include "Graph/PCGStackContext.h"
+#include "Utils/PCGExtraCapture.h"
+
+#include "UObject/ObjectKey.h"
+
+#if WITH_EDITOR
+class FPCGGraphExecutionInspection
+{
+public:
+	friend struct FPCGComponentInstanceData;
+	friend class UPCGComponent;
+
+	/* Retrieves the executed nodes information */
+	struct FNodeExecutedNotificationData
+	{
+		FNodeExecutedNotificationData(const FPCGStack& InStack, const PCGUtils::FCallTime& InTimer) : Stack(InStack), Timer(InTimer) {}
+		// Important implementation note: some logic in WasNodeExecuted relies on the fact we don't use the timer for the operator== and hash functions.
+		friend PCG_API uint32 GetTypeHash(const FNodeExecutedNotificationData& NotifData) { return GetTypeHash(NotifData.Stack); }
+		bool operator==(const FNodeExecutedNotificationData& OtherNotifData) const { return Stack == OtherNotifData.Stack; }
+
+		FPCGStack Stack;
+		PCGUtils::FCallTime Timer;
+	};
+
+	bool IsInspecting() const;
+	void EnableInspection();
+	void DisableInspection();
+	void StoreInspectionData(const FPCGStack* InStack, const UPCGNode* InNode, const PCGUtils::FCallTime* InTimer, const FPCGDataCollection& InInputData, const FPCGDataCollection& InOutputData, bool bUsedCache);
+	const FPCGDataCollection* GetInspectionData(const FPCGStack& InStack) const;
+	void ClearInspectionData(bool bClearPerNodeExecutionData = true);
+
+	/** Whether a task for the given node and stack was executed during the last execution. */
+	bool WasNodeExecuted(const UPCGNode* InNode, const FPCGStack& Stack) const;
+
+	/** Called at execution time each time a node has been executed. */
+	void NotifyNodeExecuted(const UPCGNode* InNode, const FPCGStack* InStack, const PCGUtils::FCallTime* InTimer, bool bNodeUsedCache);
+		
+	TMap<TObjectKey<const UPCGNode>, TSet<FNodeExecutedNotificationData>> GetExecutedNodeStacks() const;
+
+	/** Retrieve the inactive pin bitmask for the given node and stack in the last execution. */
+	uint64 GetNodeInactivePinMask(const UPCGNode* InNode, const FPCGStack& Stack) const;
+
+	/** Whether the given node was culled by a dynamic branch in the given stack. */
+	void NotifyNodeDynamicInactivePins(const UPCGNode* InNode, const FPCGStack* InStack, uint64 InactivePinBitmask) const;
+
+	/** Did the given node produce one or more data items in the given stack during the last execution. */
+	bool HasNodeProducedData(const UPCGNode* InNode, const FPCGStack& Stack) const;
+
+private:
+	void AddReferencedObjects(FReferenceCollector& Collector);
+
+	int32 InspectionCounter = 0;
+
+	TMap<FPCGStack, FPCGDataCollection> InspectionCache;
+
+	mutable FRWLock InspectionCacheLock;
+
+	/** Map from nodes to all stacks for which the node produced at least one data item. */
+	TMap<TObjectKey<const UPCGNode>, TSet<FPCGStack>> NodeToStacksThatProducedData;
+	mutable FRWLock NodeToStacksThatProducedDataLock;
+
+	/** Map from nodes to all stacks for which a task for the node was executed. */
+	TMap<TObjectKey<const UPCGNode>, TSet<FNodeExecutedNotificationData>> NodeToStacksInWhichNodeExecuted;
+	mutable FRWLock NodeToStacksInWhichNodeExecutedLock;
+
+	/** Map from nodes to stacks to mask of output pins that were deactivated during execution. */
+	mutable TMap<TObjectKey<const UPCGNode>, TMap<const FPCGStack, uint64>> NodeToStackToInactivePinMask;
+	mutable FRWLock NodeToStackToInactivePinMaskLock;
+};
+#endif
\ No newline at end of file
