[PCG] Remove delegate bind/copy code not needed in FPCGComponentInstanceData::ApplyToComponent
- UBlueprintGeneratedClass::BindDynamicDelegates gets called prior to InstanceDataCache->ApplyToActor (which calls FPCGComponentInstanceData::ApplyToComponent)

#rnx

diff --git a/Engine/Plugins/PCG/Source/PCG/Private/PCGComponent.cpp b/Engine/Plugins/PCG/Source/PCG/Private/PCGComponent.cpp
--- a/Engine/Plugins/PCG/Source/PCG/Private/PCGComponent.cpp
+++ b/Engine/Plugins/PCG/Source/PCG/Private/PCGComponent.cpp
@@ -589,10 +589,9 @@ FPCGTaskId UPCGComponent::GenerateInternal(bool bForce, EPCGHiGenGrid Grid, EPCG
 #if WITH_EDITOR
 		// Notify Subsystem first
 		GetSubsystem()->OnPCGGraphStartGenerating(this);
-		
+#endif // WITH_EDITOR
 		// Notify Delegate next
 		OnPCGGraphStartGeneratingDelegate.Broadcast(this);
-#endif // WITH_EDITOR
 
 		PCGComponent::BroadcastDynamicDelegate(OnPCGGraphStartGeneratingExternal, this);
 	}
@@ -775,10 +774,12 @@ void UPCGComponent::PostProcessGraph(const FBox& InNewBounds, bool bInGenerated,
 
 		// Notify Subsystem first
 		GetSubsystem()->OnPCGGraphGenerated(this);
+#endif
 
 		// Notify Delegate next
 		OnPCGGraphGeneratedDelegate.Broadcast(this);
-		
+
+#if WITH_EDITOR
 		UpdateDynamicTracking();
 #endif // WITH_EDITOR
 
@@ -889,9 +890,12 @@ void UPCGComponent::PostCleanupGraph(bool bRemoveComponents)
 	{
 		Subsystem->OnPCGGraphCleaned(this);
 	}
+#endif
 
 	// Notify Delegate next
 	OnPCGGraphCleanedDelegate.Broadcast(this);
+
+#if WITH_EDITOR
 	bDirtyGenerated = false;
 
 	if (bHadGeneratedGraphOutput)
@@ -936,10 +940,12 @@ void UPCGComponent::OnProcessGraphAborted(bool bQuiet, bool bCleanupUnusedResour
 	// Notify Subsystem first
 	UPCGSubsystem* Subsystem = GetSubsystem();
 	Subsystem->OnPCGGraphCancelled(this);
+#endif
 
 	// Notify Delegate next
 	OnPCGGraphCancelledDelegate.Broadcast(this);
 
+#if WITH_EDITOR
 	Subsystem->OnPCGComponentGenerationDone.Broadcast(Subsystem, this, EPCGGenerationStatus::Aborted);
 #endif
 
@@ -3982,21 +3988,6 @@ void FPCGComponentInstanceData::ApplyToComponent(UActorComponent* Component, con
 				PCGComponent->LoadedPreviewResources = DuplicateLoadedPreviewResources;
 			}
 
-			// Move over invocation lists for dynamic delegates
-			PCGComponent->OnPCGGraphStartGeneratingExternal = ConstructionSourceComponent->OnPCGGraphStartGeneratingExternal;
-			PCGComponent->OnPCGGraphCancelledExternal = ConstructionSourceComponent->OnPCGGraphCancelledExternal;
-			PCGComponent->OnPCGGraphGeneratedExternal = ConstructionSourceComponent->OnPCGGraphGeneratedExternal;
-			PCGComponent->OnPCGGraphCleanedExternal = ConstructionSourceComponent->OnPCGGraphCleanedExternal;
-
-			// But if the owner has changed and BP had dynamic delegates, we need to restore them.
-			if (AActor* Owner = PCGComponent->GetOwner())
-			{
-				if (Owner != ConstructionSourceComponent->GetOwner())
-				{
-					UBlueprintGeneratedClass::BindDynamicDelegates(Owner->GetClass(), Owner);
-				}
-			}
-
 			PCGComponent->bWasGeneratedThisSession = ConstructionSourceComponent->bWasGeneratedThisSession;
 
 			PCGComponent->ExecutionInspection.InspectionCounter = ConstructionSourceComponent->ExecutionInspection.InspectionCounter;
diff --git a/Engine/Plugins/PCG/Source/PCG/Public/PCGComponent.h b/Engine/Plugins/PCG/Source/PCG/Public/PCGComponent.h
--- a/Engine/Plugins/PCG/Source/PCG/Public/PCGComponent.h
+++ b/Engine/Plugins/PCG/Source/PCG/Public/PCGComponent.h
@@ -41,12 +41,10 @@ namespace UE
 	class FPCGInternalHelpers;
 }
 
-#if WITH_EDITOR
-	DECLARE_MULTICAST_DELEGATE_OneParam(FOnPCGGraphStartGenerating, UPCGComponent*);
-	DECLARE_MULTICAST_DELEGATE_OneParam(FOnPCGGraphCancelled, UPCGComponent*);
-	DECLARE_MULTICAST_DELEGATE_OneParam(FOnPCGGraphGenerated, UPCGComponent*);
-	DECLARE_MULTICAST_DELEGATE_OneParam(FOnPCGGraphCleaned, UPCGComponent*);
-#endif
+DECLARE_MULTICAST_DELEGATE_OneParam(FOnPCGGraphStartGenerating, UPCGComponent*);
+DECLARE_MULTICAST_DELEGATE_OneParam(FOnPCGGraphCancelled, UPCGComponent*);
+DECLARE_MULTICAST_DELEGATE_OneParam(FOnPCGGraphGenerated, UPCGComponent*);
+DECLARE_MULTICAST_DELEGATE_OneParam(FOnPCGGraphCleaned, UPCGComponent*);
 
 DECLARE_DYNAMIC_MULTICAST_DELEGATE_OneParam(FOnPCGGraphStartGeneratingExternal, UPCGComponent*, PCGComponent);
 DECLARE_DYNAMIC_MULTICAST_DELEGATE_OneParam(FOnPCGGraphCancelledExternal, UPCGComponent*, PCGComponent);
@@ -319,11 +317,12 @@ public:
 	UPROPERTY()
 	bool bForceGenerateOnBPAddedToWorld = false;
 
+#endif // WITH_EDITORONLY_DATA
+
 	FOnPCGGraphStartGenerating OnPCGGraphStartGeneratingDelegate;
 	FOnPCGGraphCancelled OnPCGGraphCancelledDelegate;
 	FOnPCGGraphGenerated OnPCGGraphGeneratedDelegate;
 	FOnPCGGraphCleaned OnPCGGraphCleanedDelegate;
-#endif // WITH_EDITORONLY_DATA
 
 	/** Event dispatched when a graph begins generation on this component. */
 	UPROPERTY(BlueprintAssignable, meta = (DisplayName = "On Graph Started Generating"))
