[PCG] Support adaptive RVT in landscape data interface

diff --git a/Engine/Plugins/PCG/Shaders/Private/PCGLandscapeDataInterface.ush b/Engine/Plugins/PCG/Shaders/Private/PCGLandscapeDataInterface.ush
--- a/Engine/Plugins/PCG/Shaders/Private/PCGLandscapeDataInterface.ush
+++ b/Engine/Plugins/PCG/Shaders/Private/PCGLandscapeDataInterface.ush
@@ -16,6 +16,8 @@
 // Height RVT
 Texture2D			{DataInterfaceName}_HeightVirtualTexture;
 Texture2D<uint4>	{DataInterfaceName}_HeightVirtualTexturePageTable;
+Texture2D<uint>		{DataInterfaceName}_HeightVirtualTexturePageTableIndirection;
+uint				{DataInterfaceName}_HeightVirtualTextureAdaptive;
 SamplerState		{DataInterfaceName}_HeightVirtualTextureSampler;
 float3				{DataInterfaceName}_HeightVirtualTextureLWCTile;
 float4x4			{DataInterfaceName}_HeightVirtualTextureWorldToUvTransform;
@@ -28,6 +30,8 @@ uint4				{DataInterfaceName}_HeightVirtualTextureUniforms;
 Texture2D			{DataInterfaceName}_NormalVirtualTexture0;
 Texture2D			{DataInterfaceName}_NormalVirtualTexture1;
 Texture2D<uint4>	{DataInterfaceName}_NormalVirtualTexturePageTable;
+Texture2D<uint>		{DataInterfaceName}_NormalVirtualTexturePageTableIndirection;
+uint				{DataInterfaceName}_NormalVirtualTextureAdaptive;
 SamplerState		{DataInterfaceName}_NormalVirtualTexture0Sampler;
 SamplerState		{DataInterfaceName}_NormalVirtualTexture1Sampler;
 float3				{DataInterfaceName}_NormalVirtualTextureLWCTile;
@@ -42,6 +46,8 @@ uint				{DataInterfaceName}_NormalVirtualTextureEnabled;
 // Base Color RVT
 Texture2D			{DataInterfaceName}_BaseColorVirtualTexture;
 Texture2D<uint4>	{DataInterfaceName}_BaseColorVirtualTexturePageTable;
+Texture2D<uint>		{DataInterfaceName}_BaseColorVirtualTexturePageTableIndirection;
+uint				{DataInterfaceName}_BaseColorVirtualTextureAdaptive;
 SamplerState		{DataInterfaceName}_BaseColorVirtualTextureSampler;
 float3				{DataInterfaceName}_BaseColorVirtualTextureLWCTile;
 float4x4			{DataInterfaceName}_BaseColorVirtualTextureWorldToUvTransform;
@@ -92,13 +98,29 @@ float GetHeight_{DataInterfaceName}(float3 InWorldPos)
 		const FDFVector3 LwcHeightOrigin = DFFromTileOffset_Hack(MakeLWCVector3({DataInterfaceName}_HeightVirtualTextureLWCTile, UnpackHeightOrigin));
 		const float2 SampleUv = VirtualTextureWorldToUV(LwcWorldPos, LwcHeightOrigin, UnpackHeightU, UnpackHeightV);
 
-		VTPageTableResult PageTable = TextureLoadVirtualPageTableLevel(
-			{DataInterfaceName}_HeightVirtualTexturePageTable,
-			VTPageTableUniform_Unpack({DataInterfaceName}_HeightVirtualTexturePackedUniform0, {DataInterfaceName}_HeightVirtualTexturePackedUniform1),
-			SampleUv,
-			VTADDRESSMODE_CLAMP,
-			VTADDRESSMODE_CLAMP,
-			0.0f);
+		VTPageTableResult PageTable;
+
+		if ({DataInterfaceName}_HeightVirtualTextureAdaptive == 0)
+		{
+			PageTable = TextureLoadVirtualPageTableLevel(
+				{DataInterfaceName}_HeightVirtualTexturePageTable,
+				VTPageTableUniform_Unpack({DataInterfaceName}_HeightVirtualTexturePackedUniform0, {DataInterfaceName}_HeightVirtualTexturePackedUniform1),
+				SampleUv,
+				VTADDRESSMODE_CLAMP,
+				VTADDRESSMODE_CLAMP,
+				0.0f);
+		}
+		else
+		{
+			PageTable = TextureLoadVirtualPageTableAdaptiveLevel(
+				{DataInterfaceName}_HeightVirtualTexturePageTable,
+				{DataInterfaceName}_HeightVirtualTexturePageTableIndirection,
+				VTPageTableUniform_Unpack({DataInterfaceName}_HeightVirtualTexturePackedUniform0, {DataInterfaceName}_HeightVirtualTexturePackedUniform1),
+				SampleUv,
+				VTADDRESSMODE_CLAMP,
+				VTADDRESSMODE_CLAMP,
+				0.0f);
+		}
 
 		const float4 PackedValue = TextureVirtualSample(
 			{DataInterfaceName}_HeightVirtualTexture,
@@ -131,13 +153,29 @@ float3 GetNormal_{DataInterfaceName}(float3 InWorldPos)
 		const FDFVector3 LwcNormalOrigin = DFFromTileOffset_Hack(MakeLWCVector3({DataInterfaceName}_NormalVirtualTextureLWCTile, UnpackNormalOrigin));
 		const float2 SampleUv = VirtualTextureWorldToUV(LwcWorldPos, LwcNormalOrigin, UnpackNormalU, UnpackNormalV);
 
-		VTPageTableResult PageTable = TextureLoadVirtualPageTableLevel(
-			{DataInterfaceName}_NormalVirtualTexturePageTable,
-			VTPageTableUniform_Unpack({DataInterfaceName}_NormalVirtualTexturePackedUniform0, {DataInterfaceName}_NormalVirtualTexturePackedUniform1),
-			SampleUv,
-			VTADDRESSMODE_CLAMP,
-			VTADDRESSMODE_CLAMP,
-			0.0f);
+		VTPageTableResult PageTable;
+
+		if ({DataInterfaceName}_NormalVirtualTextureAdaptive == 0)
+		{
+			PageTable = TextureLoadVirtualPageTableLevel(
+				{DataInterfaceName}_NormalVirtualTexturePageTable,
+				VTPageTableUniform_Unpack({DataInterfaceName}_NormalVirtualTexturePackedUniform0, {DataInterfaceName}_NormalVirtualTexturePackedUniform1),
+				SampleUv,
+				VTADDRESSMODE_CLAMP,
+				VTADDRESSMODE_CLAMP,
+				0.0f);
+		}
+		else
+		{
+			PageTable = TextureLoadVirtualPageTableAdaptiveLevel(
+				{DataInterfaceName}_NormalVirtualTexturePageTable,
+				{DataInterfaceName}_NormalVirtualTexturePageTableIndirection,
+				VTPageTableUniform_Unpack({DataInterfaceName}_NormalVirtualTexturePackedUniform0, {DataInterfaceName}_NormalVirtualTexturePackedUniform1),
+				SampleUv,
+				VTADDRESSMODE_CLAMP,
+				VTADDRESSMODE_CLAMP,
+				0.0f);
+		}
 
 		const float4 PackedValue0 = TextureVirtualSample(
 			{DataInterfaceName}_NormalVirtualTexture0,
@@ -195,13 +233,29 @@ float3 GetBaseColor_{DataInterfaceName}(float3 InWorldPos)
 		const FDFVector3 LwcBaseColorOrigin = DFFromTileOffset_Hack(MakeLWCVector3({DataInterfaceName}_BaseColorVirtualTextureLWCTile, UnpackBaseColorOrigin));
 		const float2 SampleUv = VirtualTextureWorldToUV(LwcWorldPos, LwcBaseColorOrigin, UnpackBaseColorU, UnpackBaseColorV);
 
-		VTPageTableResult PageTable = TextureLoadVirtualPageTableLevel(
-			{DataInterfaceName}_BaseColorVirtualTexturePageTable,
-			VTPageTableUniform_Unpack({DataInterfaceName}_BaseColorVirtualTexturePackedUniform0, {DataInterfaceName}_BaseColorVirtualTexturePackedUniform1),
-			SampleUv,
-			VTADDRESSMODE_CLAMP,
-			VTADDRESSMODE_CLAMP,
-			0.0f);
+		VTPageTableResult PageTable;
+
+		if ({DataInterfaceName}_NormalVirtualTextureAdaptive == 0)
+		{
+			PageTable = TextureLoadVirtualPageTableLevel(
+				{DataInterfaceName}_BaseColorVirtualTexturePageTable,
+				VTPageTableUniform_Unpack({DataInterfaceName}_BaseColorVirtualTexturePackedUniform0, {DataInterfaceName}_BaseColorVirtualTexturePackedUniform1),
+				SampleUv,
+				VTADDRESSMODE_CLAMP,
+				VTADDRESSMODE_CLAMP,
+				0.0f);
+		}
+		else
+		{
+			PageTable = TextureLoadVirtualPageTableAdaptiveLevel(
+				{DataInterfaceName}_BaseColorVirtualTexturePageTable,
+				{DataInterfaceName}_BaseColorVirtualTexturePageTableIndirection,
+				VTPageTableUniform_Unpack({DataInterfaceName}_BaseColorVirtualTexturePackedUniform0, {DataInterfaceName}_BaseColorVirtualTexturePackedUniform1),
+				SampleUv,
+				VTADDRESSMODE_CLAMP,
+				VTADDRESSMODE_CLAMP,
+				0.0f);
+		}
 
 		const float4 PackedValue = TextureVirtualSample(
 			{DataInterfaceName}_BaseColorVirtualTexture,
diff --git a/Engine/Plugins/PCG/Source/PCG/Private/Compute/DataInterfaces/PCGLandscapeDataInterface.cpp b/Engine/Plugins/PCG/Source/PCG/Private/Compute/DataInterfaces/PCGLandscapeDataInterface.cpp
--- a/Engine/Plugins/PCG/Source/PCG/Private/Compute/DataInterfaces/PCGLandscapeDataInterface.cpp
+++ b/Engine/Plugins/PCG/Source/PCG/Private/Compute/DataInterfaces/PCGLandscapeDataInterface.cpp
@@ -184,6 +184,8 @@ BEGIN_SHADER_PARAMETER_STRUCT(FPCGLandscapeDataInterfaceParameters,)
 	// Height RVT
 	SHADER_PARAMETER_SRV(Texture2D,				HeightVirtualTexture)
 	SHADER_PARAMETER_TEXTURE(Texture2D<uint4>,	HeightVirtualTexturePageTable)
+	SHADER_PARAMETER_TEXTURE(Texture2D<uint>,	HeightVirtualTexturePageTableIndirection)
+	SHADER_PARAMETER(uint32,					HeightVirtualTextureAdaptive)
 	SHADER_PARAMETER_SAMPLER(SamplerState,		HeightVirtualTextureSampler)
 	SHADER_PARAMETER(FVector3f,					HeightVirtualTextureLWCTile)
 	SHADER_PARAMETER(FMatrix44f,				HeightVirtualTextureWorldToUvTransform)
@@ -196,6 +198,8 @@ BEGIN_SHADER_PARAMETER_STRUCT(FPCGLandscapeDataInterfaceParameters,)
 	SHADER_PARAMETER_SRV(Texture2D,				NormalVirtualTexture0)
 	SHADER_PARAMETER_SRV(Texture2D,				NormalVirtualTexture1)
 	SHADER_PARAMETER_TEXTURE(Texture2D<uint4>,	NormalVirtualTexturePageTable)
+	SHADER_PARAMETER_TEXTURE(Texture2D<uint>,	NormalVirtualTexturePageTableIndirection)
+	SHADER_PARAMETER(uint32,					NormalVirtualTextureAdaptive)
 	SHADER_PARAMETER_SAMPLER(SamplerState,		NormalVirtualTexture0Sampler)
 	SHADER_PARAMETER_SAMPLER(SamplerState,		NormalVirtualTexture1Sampler)
 	SHADER_PARAMETER(FVector3f,					NormalVirtualTextureLWCTile)
@@ -210,6 +214,8 @@ BEGIN_SHADER_PARAMETER_STRUCT(FPCGLandscapeDataInterfaceParameters,)
 	// Base Color RVT
 	SHADER_PARAMETER_SRV(Texture2D,				BaseColorVirtualTexture)
 	SHADER_PARAMETER_TEXTURE(Texture2D<uint4>,	BaseColorVirtualTexturePageTable)
+	SHADER_PARAMETER_TEXTURE(Texture2D<uint>,	BaseColorVirtualTexturePageTableIndirection)
+	SHADER_PARAMETER(uint32,					BaseColorVirtualTextureAdaptive)
 	SHADER_PARAMETER_SAMPLER(SamplerState,		BaseColorVirtualTextureSampler)
 	SHADER_PARAMETER(FVector3f,					BaseColorVirtualTextureLWCTile)
 	SHADER_PARAMETER(FMatrix44f,				BaseColorVirtualTextureWorldToUvTransform)
@@ -593,6 +599,9 @@ void FPCGLandscapeDataProviderProxy::FVirtualTexturePageTable::Initialize(const
 		if (IAllocatedVirtualTexture* AllocatedTexture = VirtualTexture->GetAllocatedVirtualTexture())
 		{
 			PageTableRef = AllocatedTexture->GetPageTableTexture(PageTableIndex);
+			PageTableIndirectionRef = AllocatedTexture->GetPageTableIndirectionTexture();
+			bIsAdaptive = VirtualTexture->GetAdaptivePageTable();
+
 			if (PageTableRef.IsValid())
 			{
 				AllocatedTexture->GetPackedPageTableUniform(PageTableUniforms);
@@ -616,6 +625,8 @@ void FPCGLandscapeDataProviderProxy::FVirtualTexturePageTable::Initialize(const
 void FPCGLandscapeDataProviderProxy::FVirtualTexturePageTable::Reset()
 {
 	PageTableRef = nullptr;
+	PageTableIndirectionRef = nullptr;
+	bIsAdaptive = false;
 	PageTableUniforms[0] = FUintVector4(0, 0, 0, 0);
 	PageTableUniforms[1] = FUintVector4(0, 0, 0, 0);
 	WorldToUvParameters[0] = FVector4::Zero();
@@ -626,6 +637,12 @@ void FPCGLandscapeDataProviderProxy::FVirtualTexturePageTable::Reset()
 
 bool FPCGLandscapeDataProviderProxy::FVirtualTexturePageTable::IsValid() const
 {
+	// If this is an adaptive RVT then we require the indirection page table.
+	if (bIsAdaptive && !PageTableIndirectionRef.IsValid())
+	{
+		return false;
+	}
+
 	return PageTableRef.IsValid();
 }
 
@@ -743,6 +760,8 @@ bool FPCGLandscapeDataProviderProxy::SetBaseColorVirtualTextureParameters(FParam
 
 	OutShaderParameters.BaseColorVirtualTexture = BaseColorVirtualLayer.TextureSRV;
 	OutShaderParameters.BaseColorVirtualTexturePageTable = BaseColorVirtualPage.PageTableRef;
+	OutShaderParameters.BaseColorVirtualTexturePageTableIndirection = BaseColorVirtualPage.PageTableIndirectionRef;
+	OutShaderParameters.BaseColorVirtualTextureAdaptive = BaseColorVirtualPage.bIsAdaptive;
 
 	FLargeWorldRenderPosition BaseColorVirtualTextureOrigin(BaseColorVirtualPage.WorldToUvParameters[0]);
 	OutShaderParameters.BaseColorVirtualTextureLWCTile = BaseColorVirtualTextureOrigin.GetTile();
@@ -771,6 +790,8 @@ bool FPCGLandscapeDataProviderProxy::SetHeightVirtualTextureParameters(FParamete
 
 	OutShaderParameters.HeightVirtualTexture = HeightVirtualLayer.TextureSRV;
 	OutShaderParameters.HeightVirtualTexturePageTable = HeightVirtualPage.PageTableRef;
+	OutShaderParameters.HeightVirtualTexturePageTableIndirection = HeightVirtualPage.PageTableIndirectionRef;
+	OutShaderParameters.HeightVirtualTextureAdaptive = HeightVirtualPage.bIsAdaptive;
 
 	FLargeWorldRenderPosition HeightVirtualTextureOrigin(HeightVirtualPage.WorldToUvParameters[0]);
 	OutShaderParameters.HeightVirtualTextureLWCTile = HeightVirtualTextureOrigin.GetTile();
@@ -799,6 +820,8 @@ bool FPCGLandscapeDataProviderProxy::SetNormalVirtualTextureParameters(FParamete
 	OutShaderParameters.NormalVirtualTexture0 = NormalVirtualLayer0.TextureSRV;
 	OutShaderParameters.NormalVirtualTexture1 = NormalVirtualLayer1.TextureSRV;
 	OutShaderParameters.NormalVirtualTexturePageTable = NormalVirtualPage.PageTableRef;
+	OutShaderParameters.NormalVirtualTexturePageTableIndirection = NormalVirtualPage.PageTableIndirectionRef;
+	OutShaderParameters.NormalVirtualTextureAdaptive = NormalVirtualPage.bIsAdaptive;
 
 	FLargeWorldRenderPosition NormalVirtualTextureOrigin(NormalVirtualPage.WorldToUvParameters[0]);
 	OutShaderParameters.NormalVirtualTextureLWCTile = NormalVirtualTextureOrigin.GetTile();
@@ -847,6 +870,8 @@ void FPCGLandscapeDataProviderProxy::SetBaseColorVirtualTextureParameters_Defaul
 
 	OutShaderParameters.BaseColorVirtualTexture = GBlackTextureWithSRV->ShaderResourceViewRHI;
 	OutShaderParameters.BaseColorVirtualTexturePageTable = GBlackUintTexture->TextureRHI;
+	OutShaderParameters.BaseColorVirtualTexturePageTableIndirection = GBlackUintTexture->TextureRHI;
+	OutShaderParameters.BaseColorVirtualTextureAdaptive = 0;
 	OutShaderParameters.BaseColorVirtualTextureWorldToUvTransform = FMatrix44f::Identity;
 	OutShaderParameters.BaseColorVirtualTextureUnpackSRGB = 0;
 	OutShaderParameters.BaseColorVirtualTextureUnpackYCoCg = 0;
@@ -862,6 +887,8 @@ void FPCGLandscapeDataProviderProxy::SetHeightVirtualTextureParameters_Default(F
 
 	OutShaderParameters.HeightVirtualTexture = GBlackTextureWithSRV->ShaderResourceViewRHI;
 	OutShaderParameters.HeightVirtualTexturePageTable = GBlackUintTexture->TextureRHI;
+	OutShaderParameters.HeightVirtualTexturePageTableIndirection = GBlackUintTexture->TextureRHI;
+	OutShaderParameters.HeightVirtualTextureAdaptive = 0;
 	OutShaderParameters.HeightVirtualTextureWorldToUvTransform = FMatrix44f::Identity;
 	OutShaderParameters.HeightVirtualTextureEnabled = 0;
 	OutShaderParameters.HeightVirtualTexturePackedUniform0 = DummyUint4;
@@ -876,6 +903,8 @@ void FPCGLandscapeDataProviderProxy::SetNormalVirtualTextureParameters_Default(F
 	OutShaderParameters.NormalVirtualTexture0 = GBlackTextureWithSRV->ShaderResourceViewRHI;
 	OutShaderParameters.NormalVirtualTexture1 = GBlackTextureWithSRV->ShaderResourceViewRHI;
 	OutShaderParameters.NormalVirtualTexturePageTable = GBlackUintTexture->TextureRHI;
+	OutShaderParameters.NormalVirtualTexturePageTableIndirection = GBlackUintTexture->TextureRHI;
+	OutShaderParameters.NormalVirtualTextureAdaptive = 0;
 	OutShaderParameters.NormalVirtualTextureWorldToUvTransform = FMatrix44f::Identity;
 	OutShaderParameters.NormalVirtualTextureEnabled = 0;
 	OutShaderParameters.NormalVirtualTextureUnpackMode = (int32)ENormalUnpackType::None;
diff --git a/Engine/Plugins/PCG/Source/PCG/Private/Compute/DataInterfaces/PCGLandscapeDataInterface.h b/Engine/Plugins/PCG/Source/PCG/Private/Compute/DataInterfaces/PCGLandscapeDataInterface.h
--- a/Engine/Plugins/PCG/Source/PCG/Private/Compute/DataInterfaces/PCGLandscapeDataInterface.h
+++ b/Engine/Plugins/PCG/Source/PCG/Private/Compute/DataInterfaces/PCGLandscapeDataInterface.h
@@ -154,6 +154,8 @@ protected:
 		bool IsValid() const;
 	
 		FTextureRHIRef PageTableRef;
+		FTextureRHIRef PageTableIndirectionRef;
+		bool bIsAdaptive = false;
 		FUintVector4 PageTableUniforms[2];
 		FVector4 WorldToUvParameters[4];
 	};
