[PCG] Fixed some recursive issues with dynamic subgraphs. Added cvar to force dynamic dispatch for testing purposes.

#rnx

diff --git a/Engine/Plugins/PCG/Source/PCG/Private/Graph/PCGGraphCompiler.cpp b/Engine/Plugins/PCG/Source/PCG/Private/Graph/PCGGraphCompiler.cpp
--- a/Engine/Plugins/PCG/Source/PCG/Private/Graph/PCGGraphCompiler.cpp
+++ b/Engine/Plugins/PCG/Source/PCG/Private/Graph/PCGGraphCompiler.cpp
@@ -181,11 +181,15 @@ TArray<FPCGGraphTask> FPCGGraphCompiler::CompileGraph(UPCGGraph* InGraph, FPCGTa
 
 		// Note that recursive graphs must be dynamic by definition, as we aren't able to 'finish' emitting compiled tasks during compilation otherwise.
 		// Implementation note: it is very important that the predicate here is symmetrical with the one in the subgraph node otherwise some tasks could be missing.
-		const bool bIsRecursiveGraph = Subgraph && Subgraph->Contains(InGraph);
-		const bool bIsNonDynamic = SubgraphSettings && !SubgraphSettings->IsDynamicGraph();
+		const bool bIsRecursiveGraph = Subgraph && (InGraph == Subgraph || Subgraph->Contains(InGraph));
+		bool bIsDynamic = (!SubgraphSettings || SubgraphSettings->IsDynamicGraph());
 		const bool bIsNotDisabled = SubgraphSettings && SubgraphSettings->bEnabled;
 
-		if(Subgraph && bIsNonDynamic && bIsNotDisabled && !bIsRecursiveGraph)
+#if WITH_EDITOR
+		bIsDynamic |= PCGSystemSwitches::CVarForceDynamicGraphDispatch.GetValueOnAnyThread();
+#endif
+
+		if(Subgraph && !bIsDynamic && bIsNotDisabled && !bIsRecursiveGraph)
 		{
 			const FPCGTaskId PreId = NextId++;
 
diff --git a/Engine/Plugins/PCG/Source/PCG/Private/PCGCommon.cpp b/Engine/Plugins/PCG/Source/PCG/Private/PCGCommon.cpp
--- a/Engine/Plugins/PCG/Source/PCG/Private/PCGCommon.cpp
+++ b/Engine/Plugins/PCG/Source/PCG/Private/PCGCommon.cpp
@@ -79,6 +79,11 @@ namespace PCGSystemSwitches
 		TEXT("pcg.DirtyLoadAsPreviewOnLoad"),
 		false,
 		TEXT("Enables dirtying on load for load as preview components.\nTurning off this option will require to force generate or apply a change before this component is regenerated."));
+
+	TAutoConsoleVariable<bool> CVarForceDynamicGraphDispatch(
+		TEXT("pcg.GraphCompilation.ForceDynamicDispatch"),
+		false,
+		TEXT("Forces all subgraph executions to be dynamic. Performance warning, also requires to flush all compiled graphs."));
 #endif
 }
 
diff --git a/Engine/Plugins/PCG/Source/PCG/Private/PCGSubgraph.cpp b/Engine/Plugins/PCG/Source/PCG/Private/PCGSubgraph.cpp
--- a/Engine/Plugins/PCG/Source/PCG/Private/PCGSubgraph.cpp
+++ b/Engine/Plugins/PCG/Source/PCG/Private/PCGSubgraph.cpp
@@ -2,6 +2,7 @@
 
 #include "PCGSubgraph.h"
 
+#include "PCGCommon.h"
 #include "PCGComponent.h"
 #include "PCGPin.h"
 #include "Data/PCGUserParametersData.h"
@@ -717,7 +718,11 @@ bool FPCGSubgraphElement::ExecuteInternal(FPCGContext* InContext) const
 	const UPCGSubgraphSettings* Settings = Context->GetInputSettings<UPCGSubgraphSettings>();
 	check(Settings);
 
-	const bool bIsDynamic = Settings->IsDynamicGraph();
+	bool bIsDynamic = Settings->IsDynamicGraph();
+
+#if WITH_EDITOR
+	bIsDynamic |= PCGSystemSwitches::CVarForceDynamicGraphDispatch.GetValueOnAnyThread();
+#endif
 
 	if (bIsDynamic && !Context->bScheduledSubgraph && Settings->SubgraphOverride)
 	{
@@ -738,11 +743,12 @@ bool FPCGSubgraphElement::ExecuteInternal(FPCGContext* InContext) const
 	{
 		if (Context->Stack)
 		{
-			bIsRecursive = Subgraph->Contains(Context->Stack->GetGraphForCurrentFrame());
+			bIsRecursive = (bIsDynamic && Context->Stack->HasObject(Subgraph)) || Subgraph->Contains(Context->Stack->GetGraphForCurrentFrame());
 		}
 		else if(Context->SourceComponent.Get())
 		{
-			bIsRecursive = Subgraph->Contains(Context->SourceComponent->GetGraph());
+			const UPCGGraph* CurrentComponentGraph = Context->SourceComponent->GetGraph();
+			bIsRecursive = (Subgraph == CurrentComponentGraph || Subgraph->Contains(CurrentComponentGraph));
 		}
 	}
 
diff --git a/Engine/Plugins/PCG/Source/PCG/Public/PCGCommon.h b/Engine/Plugins/PCG/Source/PCG/Public/PCGCommon.h
--- a/Engine/Plugins/PCG/Source/PCG/Public/PCGCommon.h
+++ b/Engine/Plugins/PCG/Source/PCG/Public/PCGCommon.h
@@ -270,6 +270,7 @@ namespace PCGSystemSwitches
 	extern PCG_API TAutoConsoleVariable<bool> CVarPausePCGExecution;
 	extern TAutoConsoleVariable<bool> CVarGlobalDisableRefresh;
 	extern TAutoConsoleVariable<bool> CVarDirtyLoadAsPreviewOnLoad;
+	extern TAutoConsoleVariable<bool> CVarForceDynamicGraphDispatch;
 #endif
 }
 
