[PCG] Avoid creating pin labels for data resolver data interfaces
* Resolver DIs were getting assigned prefixes like "In_Resolver", which were used as pin labels even when they're not
* Refactor how data interface prefixes and labels are managed during compilation

#rnx

diff --git a/Engine/Plugins/PCG/Source/PCG/Private/Graph/PCGGraphCompilerGPU.cpp b/Engine/Plugins/PCG/Source/PCG/Private/Graph/PCGGraphCompilerGPU.cpp
--- a/Engine/Plugins/PCG/Source/PCG/Private/Graph/PCGGraphCompilerGPU.cpp
+++ b/Engine/Plugins/PCG/Source/PCG/Private/Graph/PCGGraphCompilerGPU.cpp
@@ -1133,6 +1133,7 @@ void FPCGGraphCompilerGPU::CompileComputeGraph(
 			const UComputeDataInterface* DataInterface = nullptr;
 			int32 DataInterfaceBindingIndex = INDEX_NONE;
 			FName Prefix = NAME_None;
+			FName PinLabel = NAME_None;
 			FString BindingFunctionName;
 		};
 
@@ -1181,20 +1182,31 @@ void FPCGGraphCompilerGPU::CompileComputeGraph(
 			GetDataLabels(Settings->AllOutputPinProperties());
 		}
 
-		// A given data interface may be connected to multiple downstream pins. Because each entry in these arrays produces its own bindings and graph edges,
-		// we also need to specify the prefix to uniquely describe each edge.
-		TArray<TPair</*DataInterfaceIndex*/int, /*DataInterfacePrefix*/FName>> InputDataInterfaceIndexAndPrefix;
-		TArray<TPair</*DataInterfaceIndex*/int, /*DataInterfacePrefix*/FName>> OutputDataInterfaceIndexAndPrefix;
-		InputDataInterfaceIndexAndPrefix.Reserve(Settings->InputPinProperties().Num());
-		OutputDataInterfaceIndexAndPrefix.Reserve(Settings->OutputPinProperties().Num());
+		struct FDataInterfaceInfo
+		{
+			FDataInterfaceInfo(int InIndex, FName InPrefix, FName InPinLabel)
+				: Index(InIndex)
+				, Prefix(InPrefix)
+				, PinLabel(InPinLabel)
+			{}
+
+			int Index = INDEX_NONE;
+			FName Prefix = NAME_None;
+			FName PinLabel = NAME_None;
+		};
+
+		const int NumInputPins = Settings->InputPinProperties().Num();
+		const int NumOutputPins = Settings->OutputPinProperties().Num();
+
+		TArray<FDataInterfaceInfo> InputDataInterfaces;
+		TArray<FDataInterfaceInfo> OutputDataInterfaces;
+		InputDataInterfaces.Reserve(NumInputPins);
+		OutputDataInterfaces.Reserve(NumOutputPins);
 
 		TMap</*OriginalPinLabel*/FName, /*VirtualPinLabel*/FName> OriginalToVirtualInputPin;
-		OriginalToVirtualInputPin.Reserve(InputDataInterfaceIndexAndPrefix.Num());
+		OriginalToVirtualInputPin.Reserve(NumInputPins);
 
-		auto CreateDataResolverDataInterface = [&InOutContext, Kernel = KernelWithBindings.Kernel, InOutComputeGraph](
-			TArray<TPair</*DataInterfaceIndex=*/int, /*DataInterfacePrefix=*/FName>>& InOutDataInterfaceIndexAndPrefix,
-			FName PinLabel,
-			bool bIsInput)
+		auto CreateDataResolverDataInterface = [&InOutContext, Kernel = KernelWithBindings.Kernel, InOutComputeGraph](TArray<FDataInterfaceInfo>& InOutDataInterfaces, FName PinLabel, bool bIsInput)
 		{
 			// Add data label resolver DI if this pin uses any data labels.
 			const FPCGPinDataLabels* PinDataLabels = InOutComputeGraph->GetStaticDataLabelsTable().Find(Kernel->GetKernelIndex());
@@ -1208,7 +1220,7 @@ void FPCGGraphCompilerGPU::CompileComputeGraph(
 				Resolver->bIsInput = bIsInput;
 
 				const int ResolverIndex = InOutComputeGraph->DataInterfaces.Add(Resolver);
-				InOutDataInterfaceIndexAndPrefix.Emplace(ResolverIndex, PCGComputeHelpers::GetDataLabelResolverName(PinLabel));
+				InOutDataInterfaces.Emplace(ResolverIndex, FName(*PCGComputeHelpers::GetDataLabelResolverName(PinLabel)), /*PinLabel=*/NAME_None);
 			}
 		};
 
@@ -1245,14 +1257,14 @@ void FPCGGraphCompilerGPU::CompileComputeGraph(
 				continue;
 			}
 
-			InputDataInterfaceIndexAndPrefix.Emplace(Index, Input.DownstreamPin->Label);
+			InputDataInterfaces.Emplace(Index, /*Prefix=*/Input.DownstreamPin->Label, /*PinLabel=*/Input.DownstreamPin->Label);
 
 			if (const FName* VirtualPin = InOriginalToVirtualPin.Find(FNodePin(TaskId, Input.DownstreamPin->Label, /*IsInput*/true)))
 			{
 				OriginalToVirtualInputPin.Add(Input.DownstreamPin->Label, *VirtualPin);
 			}
 
-			CreateDataResolverDataInterface(InputDataInterfaceIndexAndPrefix, Input.DownstreamPin->Label, /*bIsInput=*/true);
+			CreateDataResolverDataInterface(InputDataInterfaces, Input.DownstreamPin->Label, /*bIsInput=*/true);
 		}
 
 		// Setup output pin DIs. Always bind a DI to every output pin, so kernel always has something to write to.
@@ -1277,8 +1289,8 @@ void FPCGGraphCompilerGPU::CompileComputeGraph(
 				continue;
 			}
 
-			OutputDataInterfaceIndexAndPrefix.Emplace(Index, OutputPinProperties.Label);
-			CreateDataResolverDataInterface(OutputDataInterfaceIndexAndPrefix, OutputPinProperties.Label, /*bIsInput=*/false);
+			OutputDataInterfaces.Emplace(Index, /*Prefix=*/OutputPinProperties.Label, /*PinLabel=*/OutputPinProperties.Label);
+			CreateDataResolverDataInterface(OutputDataInterfaces, OutputPinProperties.Label, /*bIsInput=*/false);
 		}
 
 		// Kernel data interface.
@@ -1290,8 +1302,8 @@ void FPCGGraphCompilerGPU::CompileComputeGraph(
 			const int32 KernelDIIndex = InOutComputeGraph->DataInterfaces.Num();
 			InOutComputeGraph->DataInterfaces.Add(KernelDI);
 
-			// Kernel DI functions don't need a prefix.
-			InputDataInterfaceIndexAndPrefix.Emplace(KernelDIIndex, NAME_None);
+			// @todo_pcg: CustomKernel DI should probably have a prefix, e.g. "Kernel_"
+			InputDataInterfaces.Emplace(KernelDIIndex, /*Prefix=*/NAME_None, /*PinLabel=*/NAME_None);
 		}
 
 		// Additional DIs created by settings.
@@ -1320,7 +1332,7 @@ void FPCGGraphCompilerGPU::CompileComputeGraph(
 
 				if (DataInterfaceIndex != INDEX_NONE)
 				{
-					InputDataInterfaceIndexAndPrefix.Emplace(DataInterfaceIndex, NAME_None);
+					InputDataInterfaces.Emplace(DataInterfaceIndex, /*Prefix=*/NAME_None, /*PinLabel=*/NAME_None);
 				}
 			}
 
@@ -1331,13 +1343,12 @@ void FPCGGraphCompilerGPU::CompileComputeGraph(
 
 				if (DataInterfaceIndex != INDEX_NONE)
 				{
-					OutputDataInterfaceIndexAndPrefix.Emplace(DataInterfaceIndex, NAME_None);
+					OutputDataInterfaces.Emplace(DataInterfaceIndex, /*Prefix=*/NAME_None, /*PinLabel=*/NAME_None);
 				}
 			}
 		}
 
 		// Debug data interface.
-		// TODO skip depending on target platform.
 		if (Settings->bPrintShaderDebugValues)
 		{
 			UPCGDebugDataInterface* DebugDI = InOutContext.NewObject_AnyThread<UPCGDebugDataInterface>(InOutComputeGraph);
@@ -1347,7 +1358,7 @@ void FPCGGraphCompilerGPU::CompileComputeGraph(
 			InOutComputeGraph->DataInterfaces.Add(DebugDI);
 
 			// TODO: Maybe debug DI should be prefixed with 'Debug', e.g. 'Debug_WriteValue()'.
-			OutputDataInterfaceIndexAndPrefix.Emplace(DebugDIIndex, NAME_None);
+			OutputDataInterfaces.Emplace(DebugDIIndex, /*Prefix=*/NAME_None, /*PinLabel=*/NAME_None);
 		}
 
 		InOutComputeGraph->bLogDataDescriptions |= Settings->bDumpDataDescriptions;
@@ -1357,9 +1368,9 @@ void FPCGGraphCompilerGPU::CompileComputeGraph(
 
 		InOutComputeGraph->KernelToNode.Add(Node);
 
-		auto SetupAllInputBindings = [&KernelWithBindings, InOutComputeGraph](const TPair<int, FName>& InDataInterfaceIndex)
+		auto SetupAllInputBindings = [&KernelWithBindings, InOutComputeGraph](const FDataInterfaceInfo& InDataInterfaceInfo)
 		{
-			const UComputeDataInterface* DataInterface = InOutComputeGraph->DataInterfaces[InDataInterfaceIndex.Key];
+			const UComputeDataInterface* DataInterface = InOutComputeGraph->DataInterfaces[InDataInterfaceInfo.Index];
 			TArray<FShaderFunctionDefinition> Functions;
 			DataInterface->GetSupportedInputs(Functions);
 
@@ -1367,15 +1378,16 @@ void FPCGGraphCompilerGPU::CompileComputeGraph(
 			{
 				FInterfaceBinding& Binding = KernelWithBindings.InputDataBindings.Emplace_GetRef();
 				Binding.DataInterface = DataInterface;
-				Binding.Prefix = InDataInterfaceIndex.Value;
+				Binding.Prefix = InDataInterfaceInfo.Prefix;
+				Binding.PinLabel = InDataInterfaceInfo.PinLabel;
 				Binding.BindingFunctionName = Functions[FuncIndex].Name;
 				Binding.DataInterfaceBindingIndex = FuncIndex;
 			}
 		};
 
-		auto SetupAllOutputBindings = [&KernelWithBindings, InOutComputeGraph](const TPair<int, FName>& InDataInterfaceIndex)
+		auto SetupAllOutputBindings = [&KernelWithBindings, InOutComputeGraph](const FDataInterfaceInfo& InDataInterfaceInfo)
 		{
-			const UComputeDataInterface* DataInterface = InOutComputeGraph->DataInterfaces[InDataInterfaceIndex.Key];
+			const UComputeDataInterface* DataInterface = InOutComputeGraph->DataInterfaces[InDataInterfaceInfo.Index];
 			TArray<FShaderFunctionDefinition> Functions;
 			DataInterface->GetSupportedOutputs(Functions);
 
@@ -1383,21 +1395,22 @@ void FPCGGraphCompilerGPU::CompileComputeGraph(
 			{
 				FInterfaceBinding& Binding = KernelWithBindings.OutputDataBindings.Emplace_GetRef();
 				Binding.DataInterface = DataInterface;
-				Binding.Prefix = InDataInterfaceIndex.Value;
+				Binding.Prefix = InDataInterfaceInfo.Prefix;
+				Binding.PinLabel = InDataInterfaceInfo.PinLabel;
 				Binding.BindingFunctionName = Functions[FuncIndex].Name;
 				Binding.DataInterfaceBindingIndex = FuncIndex;
 			}
 		};
 
 		// Bind data interfaces.
-		for (const TPair<int, FName>& InputDataInterfaceIndex : InputDataInterfaceIndexAndPrefix)
+		for (const FDataInterfaceInfo& DataInterfaceInfo : InputDataInterfaces)
 		{
-			SetupAllInputBindings(InputDataInterfaceIndex);
+			SetupAllInputBindings(DataInterfaceInfo);
 		}
 
-		for (const TPair<int, FName>& OutputDataInterfaceIndex : OutputDataInterfaceIndexAndPrefix)
+		for (const FDataInterfaceInfo& DataInterfaceInfo : OutputDataInterfaces)
 		{
-			SetupAllOutputBindings(OutputDataInterfaceIndex);
+			SetupAllOutputBindings(DataInterfaceInfo);
 		}
 
 		{
@@ -1499,21 +1512,24 @@ void FPCGGraphCompilerGPU::CompileComputeGraph(
 						TEXT("{0}_{1}"),
 						{ Binding.Prefix.ToString(), DataInterfaceFunctions[Edge.DataInterfaceBindingIndex].Name }
 					);
+				}
 
+				if (Binding.PinLabel != NAME_None)
+				{
 					// A binding corresponds to a single API within a data interface, like GetNumData() for example. There are multiple bindings
 					// per PCG graph edge and we only need to create our mappings for the first binding.
 					if (Edge.DataInterfaceBindingIndex == 0)
 					{
-						const FPCGKernelPin KernelPin(InKernelIndex, Binding.Prefix, bInEdgesAreInputs);
+						const FPCGKernelPin KernelPin(InKernelIndex, Binding.PinLabel, bInEdgesAreInputs);
 
-						InOutComputeGraph->KernelBindingToPinLabel.Add(GraphBindingIndex, Binding.Prefix);
+						InOutComputeGraph->KernelBindingToPinLabel.Add(GraphBindingIndex, Binding.PinLabel);
 
 						if (!InOutComputeGraph->KernelPinToFirstBinding.Contains(KernelPin))
 						{
 							InOutComputeGraph->KernelPinToFirstBinding.Add(KernelPin, GraphBindingIndex);
 						}
 
-						if (const FName* VirtualPin = bInEdgesAreInputs ? OriginalToVirtualInputPin.Find(Binding.Prefix) : nullptr)
+						if (const FName* VirtualPin = bInEdgesAreInputs ? OriginalToVirtualInputPin.Find(Binding.PinLabel) : nullptr)
 						{
 							InOutComputeGraph->CPUDataBindingToVirtualPinLabel.Add(GraphBindingIndex, *VirtualPin);
 						}
