[PCG] Add a new pass 'SchedulePrimeDataDescriptionCache' to the compute graph element to compute all necessary data descriptions in advance and off the game thread.
* 'SchedulePrimeDataDescriptionCache' takes place after the 'PreExecuteReadbacks' phase because we may need to compute and cache data descriptions based on data that has been readback (e.g. for analysis kernels).
* Defer graph data validation until after the data description cache is primed and move it into its own phase of execution.
* Split IsGraphValid() into IsGraphSettingsValid() and IsGraphDataValid() so that we can still terminate execution early when there are static validation errors.

diff --git a/Engine/Plugins/PCG/Source/PCG/Internal/Compute/PCGComputeCommon.h b/Engine/Plugins/PCG/Source/PCG/Internal/Compute/PCGComputeCommon.h
--- a/Engine/Plugins/PCG/Source/PCG/Internal/Compute/PCGComputeCommon.h
+++ b/Engine/Plugins/PCG/Source/PCG/Internal/Compute/PCGComputeCommon.h
@@ -121,8 +121,6 @@ namespace PCGComputeHelpers
 	/** Returns true if the given buffer size is dangerously large. Optionally emits error log. */
 	bool IsBufferSizeTooLarge(uint64 InBufferSizeBytes, bool bInLogError = true);
 
-	bool ComputeDataDescription(const FPCGTaggedData& InTaggedData, const UPCGDataBinding* InBinding, FPCGDataDesc& OutDescription);
-
 	int32 GetAttributeIdFromMetadataAttributeIndex(int32 InAttributeIndex);
 	int32 GetMetadataAttributeIndexFromAttributeId(int32 InAttributeId);
 
diff --git a/Engine/Plugins/PCG/Source/PCG/Internal/Compute/PCGDataForGPU.h b/Engine/Plugins/PCG/Source/PCG/Internal/Compute/PCGDataForGPU.h
--- a/Engine/Plugins/PCG/Source/PCG/Internal/Compute/PCGDataForGPU.h
+++ b/Engine/Plugins/PCG/Source/PCG/Internal/Compute/PCGDataForGPU.h
@@ -176,11 +176,6 @@ struct FPCGDataCollectionDesc
 		const FPCGDataCollection& InDataCollection,
 		const UPCGDataBinding* InBinding);
 
-	static FPCGDataCollectionDesc BuildFromInputDataCollectionAndInputPinLabel(
-		const FPCGDataCollection& InDataCollection,
-		FName InputPinLabel,
-		const UPCGDataBinding* InBinding);
-
 	/** Computes the size (in bytes) of the header portion of the packed data collection buffer. */
 	uint32 ComputePackedHeaderSizeBytes() const;
 
@@ -272,7 +267,7 @@ namespace PCGDataForGPUHelpers
 		const UPCGSettings* InSettings,
 		TArray<FName>& InAttributeNames,
 		const UPCGDataBinding* InBinding,
-		const FPCGDataCollectionDesc& InDataCollectionDescription,
+		const FPCGDataCollectionDesc* InDataCollectionDescription,
 		uint32& OutCustomFloatCount,
 		TArray<FUint32Vector4>& OutAttributeIdOffsetStrides);
 }
diff --git a/Engine/Plugins/PCG/Source/PCG/Private/Compute/BuiltInKernels/PCGCountUniqueAttributeValuesKernel.cpp b/Engine/Plugins/PCG/Source/PCG/Private/Compute/BuiltInKernels/PCGCountUniqueAttributeValuesKernel.cpp
--- a/Engine/Plugins/PCG/Source/PCG/Private/Compute/BuiltInKernels/PCGCountUniqueAttributeValuesKernel.cpp
+++ b/Engine/Plugins/PCG/Source/PCG/Private/Compute/BuiltInKernels/PCGCountUniqueAttributeValuesKernel.cpp
@@ -31,11 +31,11 @@ void UPCGCountUniqueAttributeValuesKernel::Initialize(const UPCGSettings* InSett
 }
 #endif
 
-bool UPCGCountUniqueAttributeValuesKernel::IsKernelValid(FPCGContext* InContext) const
+bool UPCGCountUniqueAttributeValuesKernel::IsKernelDataValid(FPCGContext* InContext) const
 {
-	TRACE_CPUPROFILER_EVENT_SCOPE(UPCGCountUniqueAttributeValuesKernel::IsKernelValid);
+	TRACE_CPUPROFILER_EVENT_SCOPE(UPCGCountUniqueAttributeValuesKernel::IsKernelDataValid);
 
-	if (!Super::IsKernelValid(InContext))
+	if (!Super::IsKernelDataValid(InContext))
 	{
 		return false;
 	}
@@ -49,15 +49,19 @@ bool UPCGCountUniqueAttributeValuesKernel::IsKernelValid(FPCGContext* InContext)
 	const UPCGDataBinding* DataBinding = Context->DataBinding.Get();
 	if (DataBinding)
 	{
-		FPCGDataCollectionDesc InputDataDesc;
-		ensure(DataBinding->GetKernelPinDataDesc_GameThread(this, PCGPinConstants::DefaultInputLabel, /*bIsInputPin=*/true, InputDataDesc));
+		const FPCGDataCollectionDesc* InputDataDesc = DataBinding->GetCachedKernelPinDataDesc(this, PCGPinConstants::DefaultInputLabel, /*bIsInputPin=*/true);
+
+		if (!ensure(InputDataDesc))
+		{
+			return false;
+		}
 
 		FPCGKernelAttributeDesc AttributeDesc;
 		bool bConflictingTypesFound = false;
 
-		if (!InputDataDesc.GetAttributeDesc(AttributeName, AttributeDesc, bConflictingTypesFound))
+		if (!InputDataDesc->GetAttributeDesc(AttributeName, AttributeDesc, bConflictingTypesFound))
 		{
-			if (!InputDataDesc.DataDescs.IsEmpty())
+			if (!InputDataDesc->DataDescs.IsEmpty())
 			{
 				PCG_KERNEL_VALIDATION_ERR(InContext, Settings, FText::Format(
 					LOCTEXT("AttributeMissing", "Count attribute '{0}' not found, this attribute must be present on all input data, and be of type String Key."),
@@ -90,7 +94,7 @@ bool UPCGCountUniqueAttributeValuesKernel::IsKernelValid(FPCGContext* InContext)
 	return true;
 }
 
-bool UPCGCountUniqueAttributeValuesKernel::ComputeOutputBindingDataDesc(const UPCGComputeGraph* InGraph, FName InOutputPinLabel, const UPCGDataBinding* InBinding, FPCGDataCollectionDesc& OutDataDesc) const
+bool UPCGCountUniqueAttributeValuesKernel::ComputeOutputBindingDataDesc(const UPCGComputeGraph* InGraph, FName InOutputPinLabel, UPCGDataBinding* InBinding, FPCGDataCollectionDesc& OutDataDesc) const
 {
 	check(InGraph);
 	check(InBinding);
@@ -102,7 +106,8 @@ bool UPCGCountUniqueAttributeValuesKernel::ComputeOutputBindingDataDesc(const UP
 	}
 
 	FPCGDataCollectionDesc InputDesc;
-	ensure(InBinding->GetKernelPinDataDesc_GameThread(this, PCGPinConstants::DefaultInputLabel, /*bIsInput=*/true, InputDesc));
+	const FPCGKernelPin InputKernelPin(KernelIndex, PCGPinConstants::DefaultInputLabel, /*bIsInput=*/true);
+	ensure(InBinding->ComputeKernelPinDataDesc(InputKernelPin, &InputDesc));
 
 	const int32 AttributeId = InBinding->GetAttributeId(AttributeName, EPCGKernelAttributeType::StringKey);
 
@@ -127,10 +132,11 @@ int UPCGCountUniqueAttributeValuesKernel::ComputeThreadCount(const UPCGDataBindi
 	int ThreadCount = 0;
 
 	// Execute a thread for every element in the input data.
-	FPCGDataCollectionDesc OutputPinDesc;
-	if (ensure(InBinding->GetKernelPinDataDesc_GameThread(this, PCGPinConstants::DefaultInputLabel, /*bIsInput=*/true, OutputPinDesc)))
+	const FPCGDataCollectionDesc* OutputPinDesc = InBinding->GetCachedKernelPinDataDesc(this, PCGPinConstants::DefaultInputLabel, /*bIsInput=*/true);
+
+	if (ensure(OutputPinDesc))
 	{
-		for (const FPCGDataDesc& DataDesc : OutputPinDesc.DataDescs)
+		for (const FPCGDataDesc& DataDesc : OutputPinDesc->DataDescs)
 		{
 			ThreadCount += DataDesc.ElementCount;
 		}
diff --git a/Engine/Plugins/PCG/Source/PCG/Private/Compute/BuiltInKernels/PCGCountUniqueAttributeValuesKernel.h b/Engine/Plugins/PCG/Source/PCG/Private/Compute/BuiltInKernels/PCGCountUniqueAttributeValuesKernel.h
--- a/Engine/Plugins/PCG/Source/PCG/Private/Compute/BuiltInKernels/PCGCountUniqueAttributeValuesKernel.h
+++ b/Engine/Plugins/PCG/Source/PCG/Private/Compute/BuiltInKernels/PCGCountUniqueAttributeValuesKernel.h
@@ -22,8 +22,8 @@ public:
 	virtual void Initialize(const UPCGSettings* InSettings, bool bLogDescriptions) override;
 #endif
 
-	virtual bool IsKernelValid(FPCGContext* InContext) const override;
-	virtual bool ComputeOutputBindingDataDesc(const UPCGComputeGraph* InGraph, FName InOutputPinLabel, const UPCGDataBinding* InBinding, FPCGDataCollectionDesc& OutDataDesc) const override;
+	virtual bool IsKernelDataValid(FPCGContext* InContext) const override;
+	virtual bool ComputeOutputBindingDataDesc(const UPCGComputeGraph* InGraph, FName InOutputPinLabel, UPCGDataBinding* InBinding, FPCGDataCollectionDesc& OutDataDesc) const override;
 	virtual int ComputeThreadCount(const UPCGDataBinding* Binding) const override;
 	virtual bool DoesOutputPinRequireZeroInitialization(FName InOutputPinLabel) const override;
 #if WITH_EDITOR
diff --git a/Engine/Plugins/PCG/Source/PCG/Private/Compute/DataInterfaces/BuiltInKernels/PCGCountUniqueAttributeValuesDataInterface.cpp b/Engine/Plugins/PCG/Source/PCG/Private/Compute/DataInterfaces/BuiltInKernels/PCGCountUniqueAttributeValuesDataInterface.cpp
--- a/Engine/Plugins/PCG/Source/PCG/Private/Compute/DataInterfaces/BuiltInKernels/PCGCountUniqueAttributeValuesDataInterface.cpp
+++ b/Engine/Plugins/PCG/Source/PCG/Private/Compute/DataInterfaces/BuiltInKernels/PCGCountUniqueAttributeValuesDataInterface.cpp
@@ -116,14 +116,18 @@ bool UPCGCountUniqueAttributeValuesDataProvider::PrepareForExecute_GameThread(UP
 	OutputCountAttributeId = InBinding->GetAttributeId(PCGCountUniqueAttributeValuesConstants::ValueCountAttributeName, EPCGKernelAttributeType::Int);
 	check(OutputCountAttributeId != -1);
 
-	FPCGDataCollectionDesc InputDataDesc;
-	ensure(InBinding->GetKernelPinDataDesc_GameThread(Kernel, PCGPinConstants::DefaultInputLabel, /*bIsInput*/true, InputDataDesc));
+	const FPCGDataCollectionDesc* InputDataDesc = InBinding->GetCachedKernelPinDataDesc(Kernel, PCGPinConstants::DefaultInputLabel, /*bIsInput*/true);
+
+	if (!ensure(InputDataDesc))
+	{
+		return true;
+	}
 
 	AttributeToCountId = InBinding->GetAttributeId(AttributeToCountName, EPCGKernelAttributeType::StringKey);
 
 	if (AttributeToCountId != INDEX_NONE)
 	{
-		InputDataDesc.GetUniqueStringKeyValues(AttributeToCountId, UniqueStringKeyValues);
+		InputDataDesc->GetUniqueStringKeyValues(AttributeToCountId, UniqueStringKeyValues);
 	}
 
 	MaxAttributeValue = INDEX_NONE;
diff --git a/Engine/Plugins/PCG/Source/PCG/Private/Compute/DataInterfaces/Elements/PCGCopyPointsDataInterface.cpp b/Engine/Plugins/PCG/Source/PCG/Private/Compute/DataInterfaces/Elements/PCGCopyPointsDataInterface.cpp
--- a/Engine/Plugins/PCG/Source/PCG/Private/Compute/DataInterfaces/Elements/PCGCopyPointsDataInterface.cpp
+++ b/Engine/Plugins/PCG/Source/PCG/Private/Compute/DataInterfaces/Elements/PCGCopyPointsDataInterface.cpp
@@ -172,7 +172,7 @@ bool UPCGCopyPointsDataProvider::PerformPreExecuteReadbacks_GameThread(UPCGDataB
 	const int32 AnalysisDataIndex = ProducerKernel ? InBinding->GetFirstInputDataIndex(ProducerKernel, PCGCopyPointsConstants::SelectedFlagsPinLabel) : INDEX_NONE;
 	if (AnalysisDataIndex == INDEX_NONE)
 	{
-		ensure(false);
+		// Analysis data was not produced, nothing to read back.
 		return true;
 	}
 
@@ -200,12 +200,15 @@ bool UPCGCopyPointsDataProvider::PrepareForExecute_GameThread(UPCGDataBinding* I
 			return true;
 		}
 
-		FPCGDataCollectionDesc SourcePinDesc;
-		FPCGDataCollectionDesc TargetPinDesc;
-		ensure(InBinding->GetKernelPinDataDesc_GameThread(ProducerKernel, PCGCopyPointsConstants::SourcePointsLabel, /*bIsInput=*/true, SourcePinDesc));
-		ensure(InBinding->GetKernelPinDataDesc_GameThread(ProducerKernel, PCGCopyPointsConstants::TargetPointsLabel, /*bIsInput=*/true, TargetPinDesc));
+		const FPCGDataCollectionDesc* SourcePinDesc = InBinding->GetCachedKernelPinDataDesc(ProducerKernel, PCGCopyPointsConstants::SourcePointsLabel, /*bIsInput=*/true);
+		const FPCGDataCollectionDesc* TargetPinDesc = InBinding->GetCachedKernelPinDataDesc(ProducerKernel, PCGCopyPointsConstants::TargetPointsLabel, /*bIsInput=*/true);
 
-		const int32 NumOutputs = Settings->bCopyEachSourceOnEveryTarget ? SourcePinDesc.DataDescs.Num() * TargetPinDesc.DataDescs.Num() : FMath::Max(SourcePinDesc.DataDescs.Num(), TargetPinDesc.DataDescs.Num());
+		if (!ensure(SourcePinDesc) || !ensure(TargetPinDesc))
+		{
+			return true;
+		}
+
+		const int32 NumOutputs = Settings->bCopyEachSourceOnEveryTarget ? SourcePinDesc->DataDescs.Num() * TargetPinDesc->DataDescs.Num() : FMath::Max(SourcePinDesc->DataDescs.Num(), TargetPinDesc->DataDescs.Num());
 		SourceAndTargetDataIndices.Reserve(NumOutputs);
 
 		const UPCGParamData* AnalysisResultsData = Cast<UPCGParamData>(InBinding->InputDataCollection.TaggedData[AnalysisDataIndex].Data);
@@ -218,11 +221,11 @@ bool UPCGCopyPointsDataProvider::PrepareForExecute_GameThread(UPCGDataBinding* I
 		{
 			if (Settings->bCopyEachSourceOnEveryTarget)
 			{
-				for (int32 S = 0; S < SourcePinDesc.DataDescs.Num(); ++S)
+				for (int32 S = 0; S < SourcePinDesc->DataDescs.Num(); ++S)
 				{
-					for (int32 T = 0; T < TargetPinDesc.DataDescs.Num(); ++T)
+					for (int32 T = 0; T < TargetPinDesc->DataDescs.Num(); ++T)
 					{
-						if (SelectedFlagsAttribute->GetValue(S * TargetPinDesc.DataDescs.Num() + T))
+						if (SelectedFlagsAttribute->GetValue(S * TargetPinDesc->DataDescs.Num() + T))
 						{
 							SourceAndTargetDataIndices.Emplace(S, T);
 						}
@@ -233,8 +236,8 @@ bool UPCGCopyPointsDataProvider::PrepareForExecute_GameThread(UPCGDataBinding* I
 			{
 				for (int32 OutputIndex = 0; OutputIndex < NumOutputs; ++OutputIndex)
 				{
-					const int32 S = FMath::Clamp(OutputIndex, 0, SourcePinDesc.DataDescs.Num() - 1);
-					const int32 T = FMath::Clamp(OutputIndex, 0, TargetPinDesc.DataDescs.Num() - 1);
+					const int32 S = FMath::Clamp(OutputIndex, 0, SourcePinDesc->DataDescs.Num() - 1);
+					const int32 T = FMath::Clamp(OutputIndex, 0, TargetPinDesc->DataDescs.Num() - 1);
 
 					if (SelectedFlagsAttribute->GetValue(OutputIndex))
 					{
diff --git a/Engine/Plugins/PCG/Source/PCG/Private/Compute/DataInterfaces/Elements/PCGMetadataPartitionDataInterface.cpp b/Engine/Plugins/PCG/Source/PCG/Private/Compute/DataInterfaces/Elements/PCGMetadataPartitionDataInterface.cpp
--- a/Engine/Plugins/PCG/Source/PCG/Private/Compute/DataInterfaces/Elements/PCGMetadataPartitionDataInterface.cpp
+++ b/Engine/Plugins/PCG/Source/PCG/Private/Compute/DataInterfaces/Elements/PCGMetadataPartitionDataInterface.cpp
@@ -111,13 +111,17 @@ bool UPCGMetaDataPartitionDataProvider::PrepareForExecute_GameThread(UPCGDataBin
 
 	if (ensure(Settings->PartitionAttributeSelectors.Num() == 1 && Settings->PartitionAttributeSelectors[0].IsBasicAttribute()))
 	{
-		FPCGDataCollectionDesc InputDataDesc;
-		ensure(InBinding->GetKernelPinDataDesc_GameThread(Kernel, PCGPinConstants::DefaultInputLabel, /*bIsInput*/true, InputDataDesc));
+		const FPCGDataCollectionDesc* InputDataDesc = InBinding->GetCachedKernelPinDataDesc(Kernel, PCGPinConstants::DefaultInputLabel, /*bIsInput*/true);
+
+		if (!ensure(InputDataDesc))
+		{
+			return true;
+		}
 
 		PartitionAttributeId = InBinding->GetAttributeId(Settings->PartitionAttributeSelectors[0].GetAttributeName(), EPCGKernelAttributeType::StringKey);
 		if (PartitionAttributeId != INDEX_NONE)
 		{
-			InputDataDesc.GetUniqueStringKeyValues(PartitionAttributeId, UniqueStringKeyValues);
+			InputDataDesc->GetUniqueStringKeyValues(PartitionAttributeId, UniqueStringKeyValues);
 		}
 
 		MaxAttributeValue = INDEX_NONE;
@@ -126,7 +130,7 @@ bool UPCGMetaDataPartitionDataProvider::PrepareForExecute_GameThread(UPCGDataBin
 			MaxAttributeValue = FMath::Max(MaxAttributeValue, Key);
 		}
 
-		NumInputData = InputDataDesc.DataDescs.Num();
+		NumInputData = InputDataDesc->DataDescs.Num();
 		NumPartitions = UniqueStringKeyValues.Num();
 	}
 
diff --git a/Engine/Plugins/PCG/Source/PCG/Private/Compute/DataInterfaces/Elements/PCGStaticMeshSpawnerDataInterface.cpp b/Engine/Plugins/PCG/Source/PCG/Private/Compute/DataInterfaces/Elements/PCGStaticMeshSpawnerDataInterface.cpp
--- a/Engine/Plugins/PCG/Source/PCG/Private/Compute/DataInterfaces/Elements/PCGStaticMeshSpawnerDataInterface.cpp
+++ b/Engine/Plugins/PCG/Source/PCG/Private/Compute/DataInterfaces/Elements/PCGStaticMeshSpawnerDataInterface.cpp
@@ -285,7 +285,7 @@ bool UPCGStaticMeshSpawnerDataProvider::PerformPreExecuteReadbacks_GameThread(UP
 
 bool UPCGStaticMeshSpawnerDataProvider::PrepareForExecute_GameThread(UPCGDataBinding* InBinding)
 {
-	TRACE_CPUPROFILER_EVENT_SCOPE(UPCGStaticMeshSpawnerDataProvider::PerformPreExecuteReadbacks_GameThread);
+	TRACE_CPUPROFILER_EVENT_SCOPE(UPCGStaticMeshSpawnerDataProvider::PrepareForExecute_GameThread);
 	check(InBinding);
 	const UPCGStaticMeshSpawnerSettings* Settings = CastChecked<UPCGStaticMeshSpawnerSettings>(ProducerKernel->GetSettings());
 
@@ -305,11 +305,14 @@ bool UPCGStaticMeshSpawnerDataProvider::PrepareForExecute_GameThread(UPCGDataBin
 	if (SelectorAttributeId == INDEX_NONE && SelectorByAttribute)
 	{
 		const FName SelectorName = SelectorByAttribute->AttributeName;
+		const FPCGDataCollectionDesc* InputDataDesc = InBinding->GetCachedKernelPinDataDesc(ProducerKernel, PCGPinConstants::DefaultInputLabel, /*bIsInput=*/true);
 
-		FPCGDataCollectionDesc InputDataDesc;
-		InBinding->GetKernelPinDataDesc_GameThread(ProducerKernel, PCGPinConstants::DefaultInputLabel, /*bIsInput=*/true, InputDataDesc);
+		if (!ensure(InputDataDesc))
+		{
+			return true;
+		}
 
-		for (const FPCGDataDesc& Desc : InputDataDesc.DataDescs)
+		for (const FPCGDataDesc& Desc : InputDataDesc->DataDescs)
 		{
 			if (Desc.ElementCount <= 0)
 			{
@@ -333,7 +336,7 @@ bool UPCGStaticMeshSpawnerDataProvider::PrepareForExecute_GameThread(UPCGDataBin
 
 		if (SelectorAttributeId == INDEX_NONE)
 		{
-			if (!InputDataDesc.DataDescs.IsEmpty())
+			if (!InputDataDesc->DataDescs.IsEmpty())
 			{
 				PCG_KERNEL_VALIDATION_ERR(Context, Settings, FText::Format(
 					LOCTEXT("MeshSelectorAttributeNotFound", "Mesh selector attribute '{0}' not found."),
@@ -446,10 +449,14 @@ void UPCGStaticMeshSpawnerDataProvider::CreatePrimitiveDescriptors(FPCGContext*
 		return;
 	}
 
-	FPCGDataCollectionDesc InputDataDesc;
-	InBinding->GetKernelPinDataDesc_GameThread(ProducerKernel, PCGPinConstants::DefaultInputLabel, /*bIsInput=*/true, InputDataDesc);
+	const FPCGDataCollectionDesc* InputDataDesc = InBinding->GetCachedKernelPinDataDesc(ProducerKernel, PCGPinConstants::DefaultInputLabel, /*bIsInput=*/true);
+
+	if (!ensure(InputDataDesc))
+	{
+		return;
+	}
 
-	const int32 TotalInputPointCount = InputDataDesc.ComputeDataElementCount(EPCGDataType::Point);
+	const int32 TotalInputPointCount = InputDataDesc->ComputeDataElementCount(EPCGDataType::Point);
 	if (TotalInputPointCount == 0)
 	{
 		return;
@@ -472,7 +479,7 @@ void UPCGStaticMeshSpawnerDataProvider::CreatePrimitiveDescriptors(FPCGContext*
 				continue;
 			}
 
-			if (!InputDataDesc.ContainsAttributeOnAnyData(AttributeName))
+			if (!InputDataDesc->ContainsAttributeOnAnyData(AttributeName))
 			{
 				PCG_KERNEL_VALIDATION_ERR(InContext, Settings,
 					FText::Format(LOCTEXT("InstanceDataAttributeNotFound", "Instance data attribute '{0}' not found."), FText::FromName(AttributeName)));
@@ -491,7 +498,7 @@ void UPCGStaticMeshSpawnerDataProvider::CreatePrimitiveDescriptors(FPCGContext*
 		// Compute how many instances we expect for each mesh, if we do not already have the answer from analysis.
 		if (StringKeyToInstanceCount.IsEmpty())
 		{
-			for (const FPCGDataDesc& Desc : InputDataDesc.DataDescs)
+			for (const FPCGDataDesc& Desc : InputDataDesc->DataDescs)
 			{
 				if (Desc.ElementCount <= 0)
 				{
diff --git a/Engine/Plugins/PCG/Source/PCG/Private/Compute/DataInterfaces/PCGDataCollectionDataInterface.cpp b/Engine/Plugins/PCG/Source/PCG/Private/Compute/DataInterfaces/PCGDataCollectionDataInterface.cpp
--- a/Engine/Plugins/PCG/Source/PCG/Private/Compute/DataInterfaces/PCGDataCollectionDataInterface.cpp
+++ b/Engine/Plugins/PCG/Source/PCG/Private/Compute/DataInterfaces/PCGDataCollectionDataInterface.cpp
@@ -512,11 +512,11 @@ bool UPCGDataCollectionDataProvider::PrepareForExecute_GameThread(UPCGDataBindin
 		return false;
 	}
 
-	ensure(Binding->GetKernelPinDataDesc_GameThread(GraphBindingIndex, PinDesc));
+	PinDesc = Binding->GetCachedKernelPinDataDesc(GraphBindingIndex);
 
-	if (PCGComputeHelpers::IsBufferSizeTooLarge(PinDesc.ComputePackedSizeBytes()))
+	if (ensure(PinDesc) && PCGComputeHelpers::IsBufferSizeTooLarge(PinDesc->ComputePackedSizeBytes()))
 	{
-		PinDesc = {};
+		PinDesc = nullptr;
 	}
 
 	return true;
@@ -592,7 +592,7 @@ void UPCGDataCollectionDataProvider::Reset()
 	bExposeElementCounters = false;
 	OutputPinLabel = NAME_None;
 	OutputPinLabelAlias = NAME_None;
-	PinDesc = {};
+	PinDesc = nullptr;
 	bReadbackComplete = false;
 	ReadbackCounters.Empty();
 	bRequiresZeroInitialization = false;
@@ -602,7 +602,7 @@ void UPCGDataCollectionDataProvider::Reset()
 
 FPCGDataCollectionDataProviderProxy::FPCGDataCollectionDataProviderProxy(
 	TWeakObjectPtr<UPCGDataCollectionDataProvider> InDataProvider,
-	const FPCGDataCollectionDesc& InPinDesc,
+	const FPCGDataCollectionDesc* InPinDesc,
 	EPCGExportMode InExportMode,
 	bool bInZeroInitialize,
 	uint32 InElementCountMultiplier,
@@ -613,13 +613,20 @@ FPCGDataCollectionDataProviderProxy::FPCGDataCollectionDataProviderProxy(
 	, bZeroInitialize(bInZeroInitialize)
 	, ElementCountMultiplier(InElementCountMultiplier)
 	, bExposeElementCounters(bInExposeElementCounters)
-	, PinDesc(InPinDesc)
 	, DataProviderWeakPtr(InDataProvider)
 	, OutputPinLabel(InOutputPinLabel)
 	, OutputPinLabelAlias(InOutputPinLabelAlias)
 {
 	OriginatingGenerationCount = InDataProvider->GenerationCounter;
 
+	if (InPinDesc)
+	{
+		TRACE_CPUPROFILER_EVENT_SCOPE(CopyDataDescription);
+
+		// @todo_pcg: Avoid costly copy of data description.
+		PinDesc = *InPinDesc;
+	}
+
 	SizeBytes = PinDesc.ComputePackedSizeBytes();
 }
 
diff --git a/Engine/Plugins/PCG/Source/PCG/Private/Compute/DataInterfaces/PCGDataCollectionDataInterface.h b/Engine/Plugins/PCG/Source/PCG/Private/Compute/DataInterfaces/PCGDataCollectionDataInterface.h
--- a/Engine/Plugins/PCG/Source/PCG/Private/Compute/DataInterfaces/PCGDataCollectionDataInterface.h
+++ b/Engine/Plugins/PCG/Source/PCG/Private/Compute/DataInterfaces/PCGDataCollectionDataInterface.h
@@ -129,7 +129,7 @@ public:
 	std::atomic<uint64> GenerationCounter = 0;
 
 protected:
-	FPCGDataCollectionDesc PinDesc;
+	const FPCGDataCollectionDesc* PinDesc = nullptr;
 
 private:
 	FOnDataExported OnDataExported;
@@ -146,7 +146,7 @@ class FPCGDataCollectionDataProviderProxy : public FComputeDataProviderRenderPro
 public:
 	FPCGDataCollectionDataProviderProxy(
 		TWeakObjectPtr<UPCGDataCollectionDataProvider> InDataProvider,
-		const FPCGDataCollectionDesc& InPinDesc,
+		const FPCGDataCollectionDesc* InPinDesc,
 		EPCGExportMode InExportMode,
 		bool bInZeroInitialize,
 		uint32 InElementCountMultiplier,
diff --git a/Engine/Plugins/PCG/Source/PCG/Private/Compute/DataInterfaces/PCGDataCollectionUploadDataInterface.cpp b/Engine/Plugins/PCG/Source/PCG/Private/Compute/DataInterfaces/PCGDataCollectionUploadDataInterface.cpp
--- a/Engine/Plugins/PCG/Source/PCG/Private/Compute/DataInterfaces/PCGDataCollectionUploadDataInterface.cpp
+++ b/Engine/Plugins/PCG/Source/PCG/Private/Compute/DataInterfaces/PCGDataCollectionUploadDataInterface.cpp
@@ -45,10 +45,13 @@ bool UPCGDataProviderDataCollectionUpload::PrepareForExecute_GameThread(UPCGData
 
 	if (!DataAdaptor)
 	{
-		ensure(InBinding->GetKernelPinDataDesc_GameThread(GraphBindingIndex, PinDesc));
+		PinDesc = Binding->GetCachedKernelPinDataDesc(GraphBindingIndex);
 
-		check(!DownstreamInputPinLabels.IsEmpty());
-		DataAdaptor = MakeShared<FPCGDataUploadAdaptor>(InBinding, PinDesc, DownstreamInputPinLabels[0]);
+		if (ensure(PinDesc))
+		{
+			check(!DownstreamInputPinLabels.IsEmpty());
+			DataAdaptor = MakeShared<FPCGDataUploadAdaptor>(InBinding, PinDesc, DownstreamInputPinLabels[0]);
+		}
 	}
 
 	return !DataAdaptor || DataAdaptor->PrepareData_GameThread();
@@ -90,7 +93,7 @@ void UPCGDataProviderDataCollectionUpload::Reset()
 FPCGDataProviderDataCollectionUploadProxy::FPCGDataProviderDataCollectionUploadProxy(
 	TSharedPtr<FPCGDataUploadAdaptor> InDataAdaptor,
 	TWeakObjectPtr<UPCGDataCollectionDataProvider> InDataProvider,
-	const FPCGDataCollectionDesc& InPinDesc,
+	const FPCGDataCollectionDesc* InPinDesc,
 	FName InOutputPinLabel,
 	FName InOutputPinLabelAlias,
 	EPCGExportMode InExportMode,
diff --git a/Engine/Plugins/PCG/Source/PCG/Private/Compute/DataInterfaces/PCGDataCollectionUploadDataInterface.h b/Engine/Plugins/PCG/Source/PCG/Private/Compute/DataInterfaces/PCGDataCollectionUploadDataInterface.h
--- a/Engine/Plugins/PCG/Source/PCG/Private/Compute/DataInterfaces/PCGDataCollectionUploadDataInterface.h
+++ b/Engine/Plugins/PCG/Source/PCG/Private/Compute/DataInterfaces/PCGDataCollectionUploadDataInterface.h
@@ -52,7 +52,7 @@ public:
 	FPCGDataProviderDataCollectionUploadProxy(
 		TSharedPtr<FPCGDataUploadAdaptor> InDataAdaptor,
 		TWeakObjectPtr<UPCGDataCollectionDataProvider> InDataProvider,
-		const FPCGDataCollectionDesc& InPinDesc,
+		const FPCGDataCollectionDesc* InPinDesc,
 		FName InOutputPinLabel,
 		FName InOutputPinLabelAlias,
 		EPCGExportMode InExportMode,
diff --git a/Engine/Plugins/PCG/Source/PCG/Private/Compute/DataInterfaces/PCGDataLabelResolverDataInterface.cpp b/Engine/Plugins/PCG/Source/PCG/Private/Compute/DataInterfaces/PCGDataLabelResolverDataInterface.cpp
--- a/Engine/Plugins/PCG/Source/PCG/Private/Compute/DataInterfaces/PCGDataLabelResolverDataInterface.cpp
+++ b/Engine/Plugins/PCG/Source/PCG/Private/Compute/DataInterfaces/PCGDataLabelResolverDataInterface.cpp
@@ -99,8 +99,12 @@ bool UPCGDataLabelResolverDataProvider::PrepareForExecute_GameThread(UPCGDataBin
 
 	const TArray<FString>& StringTable = InBinding->GetStringTable();
 
-	FPCGDataCollectionDesc PinDataDesc;
-	ensure(InBinding->GetKernelPinDataDesc_GameThread(Kernel, PinLabel, bIsInput, PinDataDesc));
+	const FPCGDataCollectionDesc* PinDataDesc = InBinding->GetCachedKernelPinDataDesc(Kernel, PinLabel, bIsInput);
+
+	if (!ensure(PinDataDesc))
+	{
+		return true;
+	}
 
 	const FPCGPinDataLabels* FoundPinDataLabels = ComputeGraph->GetStaticDataLabelsTable().Find(Kernel->GetKernelIndex());
 	if (const FPCGDataLabels* FoundLabels = FoundPinDataLabels ? FoundPinDataLabels->PinToDataLabels.Find(PinLabel) : nullptr)
@@ -115,9 +119,9 @@ bool UPCGDataLabelResolverDataProvider::PrepareForExecute_GameThread(UPCGDataBin
 			// Initialize with invalid indices. The data IDs are not ordered.
 			DataIdToDataIndexMap[LabelIndex] = INDEX_NONE;
 
-			for (int32 InputDataIndex = 0; InputDataIndex < PinDataDesc.DataDescs.Num(); ++InputDataIndex)
+			for (int32 InputDataIndex = 0; InputDataIndex < PinDataDesc->DataDescs.Num(); ++InputDataIndex)
 			{
-				for (const int32 TagStringKey : PinDataDesc.DataDescs[InputDataIndex].TagStringKeys)
+				for (const int32 TagStringKey : PinDataDesc->DataDescs[InputDataIndex].TagStringKeys)
 				{
 					if (StringTable[TagStringKey] == PrefixedLabel)
 					{
diff --git a/Engine/Plugins/PCG/Source/PCG/Private/Compute/Elements/PCGComputeGraphElement.cpp b/Engine/Plugins/PCG/Source/PCG/Private/Compute/Elements/PCGComputeGraphElement.cpp
--- a/Engine/Plugins/PCG/Source/PCG/Private/Compute/Elements/PCGComputeGraphElement.cpp
+++ b/Engine/Plugins/PCG/Source/PCG/Private/Compute/Elements/PCGComputeGraphElement.cpp
@@ -137,6 +137,11 @@ bool FPCGComputeGraphElement::ExecuteInternal(FPCGContext* InContext) const
 
 		Context->ComputeGraphIndex = ComputeGraphIndex;
 
+		if (!Context->ComputeGraph->AreGraphSettingsValid(Context))
+		{
+			return true;
+		}
+
 		Context->ExecutionSubPhase = EPCGComputeGraphExecutionPhase::InitializeDataBinding;
 
 		if (Context->AsyncState.ShouldStop())
@@ -154,12 +159,6 @@ bool FPCGComputeGraphElement::ExecuteInternal(FPCGContext* InContext) const
 
 		Context->DataBinding->Initialize(Context->ComputeGraph.Get(), Context);
 
-		// Perform validation after input data is initialized.
-		if (!Context->ComputeGraph->IsGraphValid(InContext))
-		{
-			return true;
-		}
-
 		Context->ExecutionSubPhase = EPCGComputeGraphExecutionPhase::InitializeComputeGraphInstance;
 
 		if (Context->AsyncState.ShouldStop())
@@ -282,7 +281,7 @@ bool FPCGComputeGraphElement::ExecuteInternal(FPCGContext* InContext) const
 
 		if (bAllReady)
 		{
-			Context->ExecutionSubPhase = EPCGComputeGraphExecutionPhase::PrepareForExecute;
+			Context->ExecutionSubPhase = EPCGComputeGraphExecutionPhase::SchedulePrimeDataDescriptionCache;
 		}
 		else
 		{
@@ -295,6 +294,72 @@ bool FPCGComputeGraphElement::ExecuteInternal(FPCGContext* InContext) const
 			return false;
 		}
 	}
+	case EPCGComputeGraphExecutionPhase::SchedulePrimeDataDescriptionCache: // Fallthrough
+	{
+		// Note: Priming the data description cache must take place after pre-execute readbacks, as data descriptions
+		// may rely on the readback data, e.g. analysis kernels.
+
+		if (!Context->bPrimeDataDescCacheScheduled)
+		{
+			TRACE_CPUPROFILER_EVENT_SCOPE(SchedulePrimeDataDescriptionCache);
+
+			Context->bPrimeDataDescCacheScheduled = true;
+
+			// @todo_pcg: In the future perhaps we can just use Context->ScheduleGeneric(), but it's unclear at the moment if there is any guarantee this task
+			// would be kicked off this tick. To avoid the unknown, we'll just launch a task directly here, but it should be investigated in the future.
+			Context->PrimeDataDescCacheTask = UE::Tasks::Launch(UE_SOURCE_LOCATION, [ContextHandle = Context->GetOrCreateHandle()]()
+			{
+				FPCGContext::FSharedContext<FPCGComputeGraphContext> SharedContext(ContextHandle);
+				FPCGComputeGraphContext* Context = SharedContext.Get();
+				if (!Context)
+				{
+					return;
+				}
+
+				Context->DataBinding->PrimeDataDescriptionCache();
+			});
+
+			return false;
+		}
+
+		if (!ensure(Context->PrimeDataDescCacheTask.IsValid()))
+		{
+			return true;
+		}
+
+		if (!Context->PrimeDataDescCacheTask.IsCompleted())
+		{
+			SleepUntilNextFrame();
+			return false;
+		}
+
+		ensure(Context->DataBinding->IsDataDescriptionCachePrimed());
+		Context->DataBinding->DebugLogDataDescriptions();
+		
+		Context->ExecutionSubPhase = EPCGComputeGraphExecutionPhase::PerformDataValidation;
+
+		if (Context->AsyncState.ShouldStop())
+		{
+			return false;
+		}
+	}
+	case EPCGComputeGraphExecutionPhase::PerformDataValidation: // Fallthrough
+	{
+		TRACE_CPUPROFILER_EVENT_SCOPE(PerformDataValidation);
+
+		// Perform validation after input data is initialized and any pre-execute data has been read back.
+		if (!Context->ComputeGraph->IsGraphDataValid(InContext))
+		{
+			return true;
+		}
+
+		Context->ExecutionSubPhase = EPCGComputeGraphExecutionPhase::PrepareForExecute;
+
+		if (Context->AsyncState.ShouldStop())
+		{
+			return false;
+		}
+	}
 	case EPCGComputeGraphExecutionPhase::PrepareForExecute: // Fallthrough
 	{
 		TRACE_CPUPROFILER_EVENT_SCOPE(PrepareForExecute);
diff --git a/Engine/Plugins/PCG/Source/PCG/Private/Compute/Elements/PCGComputeGraphElement.h b/Engine/Plugins/PCG/Source/PCG/Private/Compute/Elements/PCGComputeGraphElement.h
--- a/Engine/Plugins/PCG/Source/PCG/Private/Compute/Elements/PCGComputeGraphElement.h
+++ b/Engine/Plugins/PCG/Source/PCG/Private/Compute/Elements/PCGComputeGraphElement.h
@@ -9,6 +9,7 @@
 #include "Compute/PCGDataBinding.h"
 #include "ComputeFramework/ComputeGraphInstance.h"
 
+#include "Tasks/Task.h"
 #include "UObject/StrongObjectPtr.h"
 
 #include "PCGComputeGraphElement.generated.h"
@@ -21,6 +22,8 @@ enum class EPCGComputeGraphExecutionPhase : uint8
 	InitializeDataBinding,
 	InitializeComputeGraphInstance,
 	PreExecuteReadbacks,
+	SchedulePrimeDataDescriptionCache,
+	PerformDataValidation,
 	PrepareForExecute,
 	ValidateComputeGraphCompilation,
 	ScheduleComputeGraph,
@@ -50,6 +53,10 @@ public:
 	/** Whether the compute graph instance has been initialized (data providers created, etc). */
 	bool bComputeGraphInstanceInitialized = false;
 
+	/** Whether the task to prime the data description cache has been scheduled. */
+	bool bPrimeDataDescCacheScheduled = false;
+	UE::Tasks::TTask<void> PrimeDataDescCacheTask;
+
 	/** Graph was enqueued but was invalid for some reason and the work could not be submitted. */
 	bool bGraphSubmitFailed = false;
 
diff --git a/Engine/Plugins/PCG/Source/PCG/Private/Compute/Elements/PCGCustomHLSLKernel.cpp b/Engine/Plugins/PCG/Source/PCG/Private/Compute/Elements/PCGCustomHLSLKernel.cpp
--- a/Engine/Plugins/PCG/Source/PCG/Private/Compute/Elements/PCGCustomHLSLKernel.cpp
+++ b/Engine/Plugins/PCG/Source/PCG/Private/Compute/Elements/PCGCustomHLSLKernel.cpp
@@ -105,11 +105,11 @@ const UPCGSettings* UPCGCustomHLSLKernel::GetSettings() const
 	return Settings;
 }
 
-bool UPCGCustomHLSLKernel::IsKernelValid(FPCGContext* InContext) const
+bool UPCGCustomHLSLKernel::IsKernelDataValid(FPCGContext* InContext) const
 {
-	TRACE_CPUPROFILER_EVENT_SCOPE(UPCGCustomHLSLKernel::IsKernelValid);
+	TRACE_CPUPROFILER_EVENT_SCOPE(UPCGCustomHLSLKernel::IsKernelDataValid);
 
-	if (!Super::IsKernelValid(InContext))
+	if (!Super::IsKernelDataValid(InContext))
 	{
 		return false;
 	}
@@ -135,7 +135,7 @@ bool UPCGCustomHLSLKernel::IsKernelValid(FPCGContext* InContext) const
 	return true;
 }
 
-bool UPCGCustomHLSLKernel::ComputeOutputBindingDataDesc(const UPCGComputeGraph* InGraph, FName InOutputPinLabel, const UPCGDataBinding* InBinding, FPCGDataCollectionDesc& OutDataDesc) const
+bool UPCGCustomHLSLKernel::ComputeOutputBindingDataDesc(const UPCGComputeGraph* InGraph, FName InOutputPinLabel, UPCGDataBinding* InBinding, FPCGDataCollectionDesc& OutDataDesc) const
 {
 	check(InGraph);
 	check(InBinding);
@@ -160,7 +160,9 @@ bool UPCGCustomHLSLKernel::ComputeOutputBindingDataDesc(const UPCGComputeGraph*
 
 		if (const FPCGPinProperties* FirstInputPinProps = GetFirstPointInputPin())
 		{
-			if (!InBinding->GetKernelPinDataDesc_GameThread(this, FirstInputPinProps->Label, /*bIsInputPin=*/true, InputDesc))
+			const FPCGKernelPin FirstKernelPin(KernelIndex, FirstInputPinProps->Label, /*bIsInput=*/true);
+
+			if (!InBinding->ComputeKernelPinDataDesc(FirstKernelPin, &InputDesc))
 			{
 				ensure(false);
 				return false;
@@ -231,7 +233,8 @@ bool UPCGCustomHLSLKernel::ComputeOutputBindingDataDesc(const UPCGComputeGraph*
 		for(const FPCGPinProperties& PinProps : CustomHLSLSettings->InputPins)
 		{
 			FPCGDataCollectionDesc InputPinDesc;
-			ensure(InBinding->GetKernelPinDataDesc_GameThread(this, PinProps.Label, /*bIsInputPin=*/true, InputPinDesc));
+			const FPCGKernelPin InputKernelPin(KernelIndex, PinProps.Label, /*bIsInput=*/true);
+			ensure(InBinding->ComputeKernelPinDataDesc(InputKernelPin, &InputPinDesc));
 
 			bool bFoundStringKeyAttribute = false;
 
@@ -483,11 +486,8 @@ int UPCGCustomHLSLKernel::ComputeThreadCount(const UPCGDataBinding* InBinding) c
 		{
 			if (const FPCGPinPropertiesGPU* OutputPin = CustomHLSLSettings->OutputPins.IsEmpty() ? nullptr : &CustomHLSLSettings->OutputPins[0])
 			{
-				FPCGDataCollectionDesc Desc;
-				if (ensure(ComputeOutputBindingDataDesc(InBinding->Graph, OutputPin->Label, InBinding, Desc)))
-				{
-					ThreadCount = OutputPin->GetElementCountMultiplier() * Desc.ComputeDataElementCount(EPCGDataType::Any);
-				}
+				const FPCGDataCollectionDesc* Desc = InBinding->GetCachedKernelPinDataDesc(this, OutputPin->Label, /*bIsInput=*/false);
+				ThreadCount = Desc ? OutputPin->GetElementCountMultiplier() * Desc->ComputeDataElementCount(EPCGDataType::Any) : 0;
 			}
 		}
 		else if (CustomHLSLSettings->DispatchThreadCount == EPCGDispatchThreadCount::FromProductOfInputPins)
@@ -1157,24 +1157,30 @@ bool UPCGCustomHLSLKernel::AreAttributesValid(FPCGContext* InContext, FText* Out
 
 	if (DataBinding)
 	{
-		TMap<FName, FPCGDataCollectionDesc> InputPinDescs;
-		TMap<FName, FPCGDataCollectionDesc> OutputPinDescs;
+		TMap<FName, const FPCGDataCollectionDesc*> InputPinDescs;
+		TMap<FName, const FPCGDataCollectionDesc*> OutputPinDescs;
 
 		{
-			TRACE_CPUPROFILER_EVENT_SCOPE(UPCGCustomHLSLKernel::AreAttributesValid::GetPinDataDescs);
+			TRACE_CPUPROFILER_EVENT_SCOPE(UPCGCustomHLSLKernel::AreAttributesValid::GetCachedPinDataDescs);
 
 			for (const FPCGPinProperties& Pin : CustomHLSLSettings->InputPins)
 			{
-				FPCGDataCollectionDesc DataDesc;
-				ensure(DataBinding->GetKernelPinDataDesc_GameThread(this, Pin.Label, /*bIsInputPin=*/true, DataDesc));
-				InputPinDescs.Add(Pin.Label, MoveTemp(DataDesc));
+				const FPCGDataCollectionDesc* DataDesc = DataBinding->GetCachedKernelPinDataDesc(this, Pin.Label, /*bIsInputPin=*/true);
+
+				if (ensure(DataDesc))
+				{
+					InputPinDescs.Add(Pin.Label, DataDesc);
+				}
 			}
 
 			for (const FPCGPinProperties& Pin : CustomHLSLSettings->OutputPins)
 			{
-				FPCGDataCollectionDesc DataDesc;
-				ensure(DataBinding->GetKernelPinDataDesc_GameThread(this, Pin.Label, /*bIsInputPin=*/false, DataDesc));
-				OutputPinDescs.Add(Pin.Label, MoveTemp(DataDesc));
+				const FPCGDataCollectionDesc* DataDesc = DataBinding->GetCachedKernelPinDataDesc(this, Pin.Label, /*bIsInputPin=*/false);
+
+				if (ensure(DataDesc))
+				{
+					OutputPinDescs.Add(Pin.Label, DataDesc);
+				}
 			}
 		}
 
@@ -1200,7 +1206,10 @@ bool UPCGCustomHLSLKernel::AreAttributesValid(FPCGContext* InContext, FText* Out
 
 				if (FunctionName == PCGCustomHLSLKernel::AttributeFunctionSet)
 				{
-					PinDesc = OutputPinDescs.Find(PinLabel);
+					if (const FPCGDataCollectionDesc** PinDescPtr = OutputPinDescs.Find(PinLabel))
+					{
+						PinDesc = *PinDescPtr;
+					}
 
 					if (!PinDesc && InputPinDescs.Find(PinLabel))
 					{
@@ -1219,7 +1228,10 @@ bool UPCGCustomHLSLKernel::AreAttributesValid(FPCGContext* InContext, FText* Out
 				}
 				else if (ensure(FunctionName == PCGCustomHLSLKernel::AttributeFunctionGet))
 				{
-					PinDesc = InputPinDescs.Find(PinLabel);
+					if (const FPCGDataCollectionDesc** PinDescPtr = InputPinDescs.Find(PinLabel))
+					{
+						PinDesc = *PinDescPtr;
+					}
 				}
 
 				if (!PinDesc)
@@ -1339,14 +1351,9 @@ int UPCGCustomHLSLKernel::GetElementCountForInputPin(const FPCGPinProperties& In
 	check(InBinding);
 	check(InBinding->Graph);
 
-	FPCGDataCollectionDesc InputDesc;
-	if (!InBinding->GetKernelPinDataDesc_GameThread(this, InInputPinProps.Label, /*bIsInputPin=*/true, InputDesc))
-	{
-		ensure(false);
-		return false;
-	}
+	const FPCGDataCollectionDesc* InputDesc = InBinding->GetCachedKernelPinDataDesc(this, InInputPinProps.Label, /*bIsInputPin=*/true);
 
-	return InputDesc.ComputeDataElementCount(InInputPinProps.AllowedTypes);
+	return ensure(InputDesc) ? InputDesc->ComputeDataElementCount(InInputPinProps.AllowedTypes) : 0;
 }
 
 #undef LOCTEXT_NAMESPACE
diff --git a/Engine/Plugins/PCG/Source/PCG/Private/Compute/Elements/PCGCustomHLSLKernel.h b/Engine/Plugins/PCG/Source/PCG/Private/Compute/Elements/PCGCustomHLSLKernel.h
--- a/Engine/Plugins/PCG/Source/PCG/Private/Compute/Elements/PCGCustomHLSLKernel.h
+++ b/Engine/Plugins/PCG/Source/PCG/Private/Compute/Elements/PCGCustomHLSLKernel.h
@@ -71,8 +71,8 @@ public:
 #endif
 
 	virtual const UPCGSettings* GetSettings() const override;
-	virtual bool IsKernelValid(FPCGContext* InContext) const override;
-	virtual bool ComputeOutputBindingDataDesc(const UPCGComputeGraph* InGraph, FName InOutputPinLabel, const UPCGDataBinding* InBinding, FPCGDataCollectionDesc& OutDataDesc) const override;
+	virtual bool IsKernelDataValid(FPCGContext* InContext) const override;
+	virtual bool ComputeOutputBindingDataDesc(const UPCGComputeGraph* InGraph, FName InOutputPinLabel, UPCGDataBinding* InBinding, FPCGDataCollectionDesc& OutDataDesc) const override;
 #if WITH_EDITOR
 	virtual FString GetCookedSource(const UPCGComputeGraph* InComputeGraph) const override;
 	virtual FString GetEntryPoint() const override { return EntryPoint; }
diff --git a/Engine/Plugins/PCG/Source/PCG/Private/Compute/PCGComputeCommon.cpp b/Engine/Plugins/PCG/Source/PCG/Private/Compute/PCGComputeCommon.cpp
--- a/Engine/Plugins/PCG/Source/PCG/Private/Compute/PCGComputeCommon.cpp
+++ b/Engine/Plugins/PCG/Source/PCG/Private/Compute/PCGComputeCommon.cpp
@@ -121,42 +121,6 @@ namespace PCGComputeHelpers
 		return bBufferTooLarge;
 	}
 
-	bool ComputeDataDescription(const FPCGTaggedData& InTaggedData, const UPCGDataBinding* InBinding, FPCGDataDesc& OutDescription)
-	{
-		if (const UPCGProxyForGPUData* Proxy = Cast<UPCGProxyForGPUData>(InTaggedData.Data))
-		{
-			if (!Proxy->GetDescription(OutDescription))
-			{
-				return false;
-			}
-
-			if (TSharedPtr<const FPCGProxyForGPUDataCollection> DataGPU = Proxy->GetGPUInfo())
-			{
-				for (FPCGKernelAttributeDesc& AttrDesc : OutDescription.AttributeDescs)
-				{
-					const int32 Index = PCGComputeHelpers::GetMetadataAttributeIndexFromAttributeId(AttrDesc.Index);
-					if (Index == INDEX_NONE)
-					{
-						continue;
-					}
-
-					if (const int32* AttributeIndex = InBinding->GetAttributeIdRemap().Find(UPCGDataBinding::FSourceBufferAndAttributeIndex(DataGPU, Index)))
-					{
-						const int32 NewId = PCGComputeHelpers::GetAttributeIdFromMetadataAttributeIndex(*AttributeIndex);
-						//UE_LOG(LogPCG, Warning, TEXT("Remapped attribute '%s' ID: %d -> %d"), *AttrDesc.Name.ToString(), AttrDesc.Index, NewId);
-						AttrDesc.Index = NewId;
-					}
-				}
-			}
-		}
-		else
-		{
-			OutDescription = FPCGDataDesc(InTaggedData, InBinding);
-		}
-
-		return true;
-	}
-
 	int32 GetAttributeIdFromMetadataAttributeIndex(int32 InAttributeIndex)
 	{
 		return InAttributeIndex >= 0 ? (InAttributeIndex + PCGComputeConstants::NUM_RESERVED_ATTRS) : INDEX_NONE;
diff --git a/Engine/Plugins/PCG/Source/PCG/Private/Compute/PCGComputeGraph.cpp b/Engine/Plugins/PCG/Source/PCG/Private/Compute/PCGComputeGraph.cpp
--- a/Engine/Plugins/PCG/Source/PCG/Private/Compute/PCGComputeGraph.cpp
+++ b/Engine/Plugins/PCG/Source/PCG/Private/Compute/PCGComputeGraph.cpp
@@ -9,6 +9,7 @@
 #include "Compute/PCGComputeKernel.h"
 #include "Compute/PCGDataBinding.h"
 #include "Compute/PCGPinPropertiesGPU.h"
+#include "Compute/Data/PCGProxyForGPUData.h"
 #include "Metadata/PCGMetadata.h"
 
 #include "ComputeFramework/ComputeKernelCompileResult.h"
@@ -97,21 +98,7 @@ void UPCGComputeGraph::OnKernelCompilationComplete(int32 InKernelIndex, FCompute
 	}
 }
 
-bool UPCGComputeGraph::ComputeKernelBindingDataDesc(int32 InKernelIndex, FName InPinLabel, bool bInIsInputPin, const UPCGDataBinding* InBinding, FPCGDataCollectionDesc& OutDataDesc) const
-{
-	const int32 BindingIndex = GetBindingIndex(InKernelIndex, InPinLabel, bInIsInputPin);
-	if (BindingIndex != INDEX_NONE)
-	{
-		return ComputeKernelBindingDataDesc(BindingIndex, InBinding, OutDataDesc);
-	}
-	else
-	{
-		ensure(false);
-		return false;
-	}
-}
-
-bool UPCGComputeGraph::ComputeKernelBindingDataDesc(int32 InBindingIndex, const UPCGDataBinding* InBinding, FPCGDataCollectionDesc& OutDataDesc) const
+bool UPCGComputeGraph::ComputeKernelBindingDataDesc(int32 InBindingIndex, UPCGDataBinding* InBinding, FPCGDataCollectionDesc& OutDataDesc) const
 {
 	TRACE_CPUPROFILER_EVENT_SCOPE(UPCGComputeGraph::ComputeKernelBindingDataDesc);
 
@@ -130,42 +117,32 @@ bool UPCGComputeGraph::ComputeKernelBindingDataDesc(int32 InBindingIndex, const
 	if (const FName* VirtualLabel = CPUDataBindingToVirtualPinLabel.Find(InBindingIndex))
 	{
 		// Create description from compute graph element input data.
-		OutDataDesc = FPCGDataCollectionDesc::BuildFromInputDataCollectionAndInputPinLabel(InBinding->InputDataCollection, *VirtualLabel, InBinding);
+		ComputeExternalPinDesc(*VirtualLabel, InBinding, OutDataDesc);
+
 		return true;
 	}
 	else if (GraphEdge.bKernelInput)
 	{
-		if (const int32* UpstreamBindingIndex = DownstreamToUpstreamBinding.Find(InBindingIndex))
-		{
-			return ComputeKernelBindingDataDesc(*UpstreamBindingIndex, InBinding, OutDataDesc);
-		}
-		else
-		{
-			ensure(false);
-			return {};
-		}
+		const int32* UpstreamBindingIndex = DownstreamToUpstreamBinding.Find(InBindingIndex);
+		return ensure(UpstreamBindingIndex) && ComputeKernelBindingDataDesc(*UpstreamBindingIndex, InBinding, OutDataDesc);
 	}
 	else
 	{
 		// Compute description of data from output binding. Data doesn't come from CPU but via a kernel. Consult the kernel for its data description.
 		check(KernelInvocations.IsValidIndex(GraphEdge.KernelIndex));
+
 		const UPCGComputeKernel* Kernel = CastChecked<UPCGComputeKernel>(KernelInvocations[GraphEdge.KernelIndex]);
+		const FName* OutputPinLabel = KernelBindingToPinLabel.Find(InBindingIndex);
 
-		if (const FName* OutputPinLabel = KernelBindingToPinLabel.Find(InBindingIndex))
-		{
-			return Kernel->ComputeOutputBindingDataDesc(this, *OutputPinLabel, InBinding, OutDataDesc);
-		}
-		else
-		{
-			ensure(false);
-			return {};
-		}
+		return ensure(OutputPinLabel) && Kernel->ComputeOutputBindingDataDesc(this, *OutputPinLabel, InBinding, OutDataDesc);
 	}
 }
 
-int32 UPCGComputeGraph::GetBindingIndex(int32 InKernelIndex, FName InPinLabel, bool bInIsInput) const
+int32 UPCGComputeGraph::GetBindingIndex(const FPCGKernelPin& InKernelPin) const
 {
-	if (const int32* BindingIndex = KernelPinToFirstBinding.Find(FPCGKernelPin(InKernelIndex, InPinLabel, bInIsInput)))
+	// Bindings in compute graphs roughly map to a single API on a Data Interface, like GetNumData(). There can be multiple bindings per PCG edge.
+	// We choose the first binding index to represent the kernel pin.
+	if (const int32* BindingIndex = KernelPinToFirstBinding.Find(InKernelPin))
 	{
 		return *BindingIndex;
 	}
@@ -176,14 +153,36 @@ int32 UPCGComputeGraph::GetBindingIndex(int32 InKernelIndex, FName InPinLabel, b
 	}
 }
 
-bool UPCGComputeGraph::IsGraphValid(FPCGContext* InContext) const
+void UPCGComputeGraph::GetKernelPins(TArray<FPCGKernelPin>& OutKernelPins) const
+{
+	KernelPinToFirstBinding.GenerateKeyArray(OutKernelPins);
+}
+
+bool UPCGComputeGraph::AreGraphSettingsValid(FPCGContext* InContext) const
+{
+	TRACE_CPUPROFILER_EVENT_SCOPE(UPCGComputeGraph::AreGraphSettingsValid);
+
+	bool bAllKernelsValid = true;
+
+	// Run validation on all kernels.
+	for (const UComputeKernel* Kernel : KernelInvocations)
+	{
+		bAllKernelsValid &= CastChecked<UPCGComputeKernel>(Kernel)->AreKernelSettingsValid(InContext);
+	}
+
+	return bAllKernelsValid;
+}
+
+bool UPCGComputeGraph::IsGraphDataValid(FPCGContext* InContext) const
 {
+	TRACE_CPUPROFILER_EVENT_SCOPE(UPCGComputeGraph::IsGraphDataValid);
+
 	bool bAllKernelsValid = true;
 
 	// Run validation on all kernels.
 	for (const UComputeKernel* Kernel : KernelInvocations)
 	{
-		bAllKernelsValid &= CastChecked<UPCGComputeKernel>(Kernel)->IsKernelValid(InContext);
+		bAllKernelsValid &= CastChecked<UPCGComputeKernel>(Kernel)->IsKernelDataValid(InContext);
 	}
 
 	return bAllKernelsValid;
@@ -196,11 +195,16 @@ void UPCGComputeGraph::DebugLogDataDescriptions(const UPCGDataBinding* InBinding
 	const UEnum* PCGKernelAttributeTypeEnum = StaticEnum<EPCGKernelAttributeType>();
 	check(PCGDataTypeEnum && PCGKernelAttributeTypeEnum);
 
-	auto LogDataDescription = [this, PCGDataTypeEnum, PCGKernelAttributeTypeEnum](const FPCGDataCollectionDesc& InDataDesc)
+	auto LogDataDescription = [this, PCGDataTypeEnum, PCGKernelAttributeTypeEnum](const FPCGDataCollectionDesc* InDataDesc)
 	{
-		for (int32 DataIndex = 0; DataIndex < InDataDesc.DataDescs.Num(); ++DataIndex)
+		if (!InDataDesc)
 		{
-			const FPCGDataDesc& DataDesc = InDataDesc.DataDescs[DataIndex];
+			return;
+		}
+
+		for (int32 DataIndex = 0; DataIndex < InDataDesc->DataDescs.Num(); ++DataIndex)
+		{
+			const FPCGDataDesc& DataDesc = InDataDesc->DataDescs[DataIndex];
 
 			UE_LOG(LogPCG, Display, TEXT("\t\tData Index %d"), DataIndex);
 			UE_LOG(LogPCG, Display, TEXT("\t\t\tType: %s"), *PCGDataTypeEnum->GetDisplayNameTextByValue(static_cast<int64>(DataDesc.Type)).ToString());
@@ -250,10 +254,11 @@ void UPCGComputeGraph::DebugLogDataDescriptions(const UPCGDataBinding* InBinding
 			{
 				if (Pin.Label != NAME_None)
 				{
-					FPCGDataCollectionDesc DataDesc;
-					ensure(InBinding->GetKernelPinDataDesc_GameThread(PCGKernel, Pin.Label, true, DataDesc));
-					UE_LOG(LogPCG, Display, TEXT("\tInput Pin: %s (%d data)"), *Pin.Label.ToString(), DataDesc.DataDescs.Num());
-					LogDataDescription(DataDesc);
+					if (const FPCGDataCollectionDesc* DataDesc = InBinding->GetCachedKernelPinDataDesc(PCGKernel, Pin.Label, true))
+					{
+						UE_LOG(LogPCG, Display, TEXT("\tInput Pin: %s (%d data)"), *Pin.Label.ToString(), DataDesc->DataDescs.Num());
+						LogDataDescription(DataDesc);
+					}
 				}
 			}
 
@@ -263,10 +268,11 @@ void UPCGComputeGraph::DebugLogDataDescriptions(const UPCGDataBinding* InBinding
 			{
 				if (Pin.Label != NAME_None)
 				{
-					FPCGDataCollectionDesc DataDesc;
-					ensure(InBinding->GetKernelPinDataDesc_GameThread(PCGKernel, Pin.Label, false, DataDesc));
-					UE_LOG(LogPCG, Display, TEXT("\tOutput Pin: %s (%d data)"), *Pin.Label.ToString(), DataDesc.DataDescs.Num());
-					LogDataDescription(DataDesc);
+					if (const FPCGDataCollectionDesc* DataDesc = InBinding->GetCachedKernelPinDataDesc(PCGKernel, Pin.Label, false))
+					{
+						UE_LOG(LogPCG, Display, TEXT("\tOutput Pin: %s (%d data)"), *Pin.Label.ToString(), DataDesc->DataDescs.Num());
+						LogDataDescription(DataDesc);
+					}
 				}
 			}
 		}
@@ -274,4 +280,61 @@ void UPCGComputeGraph::DebugLogDataDescriptions(const UPCGDataBinding* InBinding
 #endif
 }
 
+void UPCGComputeGraph::ComputeExternalPinDesc(FName InVirtualLabel, UPCGDataBinding* InBinding, FPCGDataCollectionDesc& OutDataDesc) const
+{
+	check(InBinding);
+
+	const TArray<FPCGTaggedData> TaggedDatas = InBinding->InputDataCollection.GetInputsByPin(InVirtualLabel);
+
+	for (const FPCGTaggedData& TaggedData : TaggedDatas)
+	{
+		if (!TaggedData.Data || !PCGComputeHelpers::IsTypeAllowedAsInput(TaggedData.Data->GetDataType()))
+		{
+			continue;
+		}
+
+		FPCGDataDesc DataDesc;
+		if (ensure(ComputeTaggedDataPinDesc(TaggedData, InBinding, DataDesc)))
+		{
+			OutDataDesc.DataDescs.Add(MoveTemp(DataDesc));
+		}
+	}
+}
+
+bool UPCGComputeGraph::ComputeTaggedDataPinDesc(const FPCGTaggedData& InTaggedData, const UPCGDataBinding* InBinding, FPCGDataDesc& OutDescription) const
+{
+	if (const UPCGProxyForGPUData* Proxy = Cast<UPCGProxyForGPUData>(InTaggedData.Data))
+	{
+		if (!Proxy->GetDescription(OutDescription))
+		{
+			return false;
+		}
+
+		if (TSharedPtr<const FPCGProxyForGPUDataCollection> DataGPU = Proxy->GetGPUInfo())
+		{
+			for (FPCGKernelAttributeDesc& AttrDesc : OutDescription.AttributeDescs)
+			{
+				const int32 Index = PCGComputeHelpers::GetMetadataAttributeIndexFromAttributeId(AttrDesc.Index);
+				if (Index == INDEX_NONE)
+				{
+					continue;
+				}
+
+				if (const int32* AttributeIndex = InBinding->GetAttributeIdRemap().Find(UPCGDataBinding::FSourceBufferAndAttributeIndex(DataGPU, Index)))
+				{
+					const int32 NewId = PCGComputeHelpers::GetAttributeIdFromMetadataAttributeIndex(*AttributeIndex);
+					//UE_LOG(LogPCG, Warning, TEXT("Remapped attribute '%s' ID: %d -> %d"), *AttrDesc.Name.ToString(), AttrDesc.Index, NewId);
+					AttrDesc.Index = NewId;
+				}
+			}
+		}
+	}
+	else
+	{
+		OutDescription = FPCGDataDesc(InTaggedData, InBinding);
+	}
+
+	return true;
+}
+
 #undef PCG_DATA_DESCRIPTION_LOGGING 
diff --git a/Engine/Plugins/PCG/Source/PCG/Private/Compute/PCGComputeGraph.h b/Engine/Plugins/PCG/Source/PCG/Private/Compute/PCGComputeGraph.h
--- a/Engine/Plugins/PCG/Source/PCG/Private/Compute/PCGComputeGraph.h
+++ b/Engine/Plugins/PCG/Source/PCG/Private/Compute/PCGComputeGraph.h
@@ -88,22 +88,41 @@ public:
 	/** Get the strings collated statically during compilation. */
 	const TArray<FString>& GetStaticStringTable() const { return StringTable; }
 
-	// Get description of data produced or consumed by a pin on a kernel.
-	bool ComputeKernelBindingDataDesc(int32 InKernelIndex, FName InPinLabel, bool bInIsInputPin, const UPCGDataBinding* InBinding, FPCGDataCollectionDesc& OutDataDesc) const;
-
-	// Get description of data produced or consumed by a binding in the graph. In compute graph a binding is a single function that is exposed to a kernel.
-	// in PCG a pin is a data interface (multiple incident edges on an input pin will be merged using a DI).
-	bool ComputeKernelBindingDataDesc(int32 InBindingIndex, const UPCGDataBinding* InBinding, FPCGDataCollectionDesc& OutDataDesc) const;
+	/**
+	 * Computes a description of the data produced/consumed by a binding in the graph.
+	 *
+	 * In compute graph, a binding is a single function that is exposed to a kernel.
+	 * In PCG, a pin is a data interface (multiple incident edges on an input pin will be merged using a DI).
+	 *
+	 * A pin will have many bindings, but we arbitrarily use the first binding index to uniquely represent a pin.
+	 */
+	bool ComputeKernelBindingDataDesc(int32 InBindingIndex, UPCGDataBinding* InBinding, FPCGDataCollectionDesc& OutDataDesc) const;
+
+	/** Get the first binding index, which uniquely identifies a kernel pin to the compute graph. */
+	int32 GetBindingIndex(const FPCGKernelPin& InKernelPin) const;
+
+	/** Get the virtual label associated with a binding index/kernel pin. This only exists if the kernel pin is an input pin on the border of the compute graph. */
+	const FName* GetVirtualPinLabelFromBindingIndex(int32 InBindingIndex) const { return CPUDataBindingToVirtualPinLabel.Find(InBindingIndex); }
 
-	/** Bindings in compute graphs roughly map to a single API on a Data Interface, like GetNumData(). There can be multiple bindings per PCG edge. */
-	int32 GetBindingIndex(int32 InKernelIndex, FName InPinLabel, bool bInIsInput) const;
+	/** Populates an array with all of the kernel pins in the compute graph. */
+	void GetKernelPins(TArray<FPCGKernelPin>& OutKernelPins) const;
 
-	const FName* GetVirtualPinLabelFromBindingIndex(int32 InBindingIndex) const { return CPUDataBindingToVirtualPinLabel.Find(InBindingIndex); }
+	/** Validates and logs information about the graph known at compile time. */
+	bool AreGraphSettingsValid(FPCGContext* InContext) const;
 
-	bool IsGraphValid(FPCGContext* InContext) const;
+	/** Validates and logs information about incoming data to the compute graph. */
+	bool IsGraphDataValid(FPCGContext* InContext) const;
 
+	/** Logs detailed data descriptions for all the kernels in the compute graph. */
 	void DebugLogDataDescriptions(const UPCGDataBinding* InBinding) const;
 
+private:
+	/** Compute the data description on a pin external to the compute graph. */
+	void ComputeExternalPinDesc(FName InVirtualLabel, UPCGDataBinding* InBinding, FPCGDataCollectionDesc& OutDataDesc) const;
+
+	/** Compute the data description of a tagged data. */
+	bool ComputeTaggedDataPinDesc(const FPCGTaggedData& InTaggedData, const UPCGDataBinding* InBinding, FPCGDataDesc& OutDescription) const;
+
 public:
 	TMap<TObjectKey<const UPCGNode>, TArray<FComputeKernelCompileMessage>> KernelToCompileMessages;
 	
@@ -137,10 +156,11 @@ protected:
 	UPROPERTY()
 	TArray<FString> StringTable;
 
+	/** Table to look up a kernel pin's first binding index. */
 	UPROPERTY()
 	TMap<FPCGKernelPin, /*FirstBindingIndex*/int32> KernelPinToFirstBinding;
 
-	/** Binding index to virtual label for bindings that are that receive data from nodes executing on the CPU or in separate compute graphs. */
+	/** Binding index to virtual label for bindings that are that receive data from external nodes (executing on the CPU or in separate compute graphs). */
 	UPROPERTY()
 	TMap</*BindingIndex*/int32, /*VirtualPinLabel*/FName> CPUDataBindingToVirtualPinLabel;
 
diff --git a/Engine/Plugins/PCG/Source/PCG/Private/Compute/PCGComputeKernel.cpp b/Engine/Plugins/PCG/Source/PCG/Private/Compute/PCGComputeKernel.cpp
--- a/Engine/Plugins/PCG/Source/PCG/Private/Compute/PCGComputeKernel.cpp
+++ b/Engine/Plugins/PCG/Source/PCG/Private/Compute/PCGComputeKernel.cpp
@@ -34,7 +34,7 @@ void UPCGComputeKernel::Initialize(const UPCGSettings* InSettings, bool bLogDesc
 }
 #endif
 
-bool UPCGComputeKernel::IsKernelValid(FPCGContext* InContext) const
+bool UPCGComputeKernel::AreKernelSettingsValid(FPCGContext* InContext) const
 {
 	if (!bInitialized)
 	{
@@ -59,7 +59,7 @@ bool UPCGComputeKernel::IsKernelValid(FPCGContext* InContext) const
 	return !bHasStaticValidationErrors;
 }
 
-void UPCGComputeKernel::ComputeDataDescFromPinProperties(const FPCGPinPropertiesGPU& OutputPinProps, const TArrayView<const FPCGPinProperties>& InInputPinProps, const UPCGDataBinding* InBinding, FPCGDataCollectionDesc& OutPinDesc) const
+void UPCGComputeKernel::ComputeDataDescFromPinProperties(const FPCGPinPropertiesGPU& OutputPinProps, const TArrayView<const FPCGPinProperties>& InInputPinProps, UPCGDataBinding* InBinding, FPCGDataCollectionDesc& OutPinDesc) const
 {
 	check(InBinding);
 
@@ -77,7 +77,8 @@ void UPCGComputeKernel::ComputeDataDescFromPinProperties(const FPCGPinProperties
 				InitPins.Emplace(PinToInitFrom);
 
 				FPCGDataCollectionDesc& InputDesc = InputDescs.AddDefaulted_GetRef();
-				ensure(InBinding->GetKernelPinDataDesc_GameThread(this, PinToInitFrom, /*bIsInputPin=*/true, InputDesc));
+				const FPCGKernelPin KernelPin(KernelIndex, PinToInitFrom, /*bIsInput=*/true);
+				ensure(InBinding->ComputeKernelPinDataDesc(KernelPin, &InputDesc));
 			}
 		}
 
diff --git a/Engine/Plugins/PCG/Source/PCG/Private/Compute/PCGComputeKernel.h b/Engine/Plugins/PCG/Source/PCG/Private/Compute/PCGComputeKernel.h
--- a/Engine/Plugins/PCG/Source/PCG/Private/Compute/PCGComputeKernel.h
+++ b/Engine/Plugins/PCG/Source/PCG/Private/Compute/PCGComputeKernel.h
@@ -65,8 +65,11 @@ public:
 	/** Gets settings for node associated with this kernel, if any. */
 	virtual const UPCGSettings* GetSettings() const { return Settings; }
 
+	/** Performs settings validation and returns true if this node is suitable for execution. */
+	virtual bool AreKernelSettingsValid(FPCGContext* InContext) const;
+
 	/** Performs data validation and returns true if this node is suitable for deployment to the GPU. */
-	virtual bool IsKernelValid(FPCGContext* InContext) const;
+	virtual bool IsKernelDataValid(FPCGContext* InContext) const { return true; }
 
 #if WITH_EDITOR
 	/** Produces the node specific portion of kernel shader source text, including the main entry point. */
@@ -86,7 +89,7 @@ public:
 #endif
 
 	/** Compute a description of data that will be output from pin InOutputPinLabel. */
-	virtual bool ComputeOutputBindingDataDesc(const UPCGComputeGraph* InGraph, FName InOutputPinLabel, const UPCGDataBinding* InBinding, FPCGDataCollectionDesc& OutDataDesc) const PURE_VIRTUAL(UPCGComputeKernel::ComputeOutputBindingDataDesc, return {};);
+	virtual bool ComputeOutputBindingDataDesc(const UPCGComputeGraph* InGraph, FName InOutputPinLabel, UPCGDataBinding* InBinding, FPCGDataCollectionDesc& OutDataDesc) const PURE_VIRTUAL(UPCGComputeKernel::ComputeOutputBindingDataDesc, return {};);
 
 	/** Compute how many threads should be dispatched to execute this node on the GPU. */
 	virtual int ComputeThreadCount(const UPCGDataBinding* Binding) const PURE_VIRTUAL(UPCGComputeKernel::GetCookedSource, return 0;);
@@ -116,7 +119,7 @@ public:
 	virtual uint32 GetElementCountMultiplier(FName InOutputPinLabel) const { return 1u; }
 
 	/** Helper to compute data descriptions when data is entirely determined by pin properties. */
-	void ComputeDataDescFromPinProperties(const FPCGPinPropertiesGPU& OutputPinProps, const TArrayView<const FPCGPinProperties>& InInputPinProps, const UPCGDataBinding* InBinding, FPCGDataCollectionDesc& OutPinDesc) const;
+	void ComputeDataDescFromPinProperties(const FPCGPinPropertiesGPU& OutputPinProps, const TArrayView<const FPCGPinProperties>& InInputPinProps, UPCGDataBinding* InBinding, FPCGDataCollectionDesc& OutPinDesc) const;
 
 	virtual void GetInputPins(TArray<FPCGPinProperties>& OutPins) const {}
 	virtual void GetOutputPins(TArray<FPCGPinPropertiesGPU>& OutPins) const {}
@@ -153,7 +156,7 @@ protected:
 	bool bHasStaticValidationErrors = false;
 
 	// @todo_pcg: These are a bit of a hack to avoid supporting static/compilation logging in PCGNodeVisualLogs.h
-	/** Log entries created and cached on initialization. Are logged during IsKernelValid(), unless PCG_KERNEL_LOGGING_ENABLED is false. */
+	/** Log entries created and cached on initialization. Are logged during AreKernelSettingsValid(), unless PCG_KERNEL_LOGGING_ENABLED is false. */
 	UPROPERTY()
 	TArray<FPCGKernelLogEntry> StaticLogEntries;
 };
diff --git a/Engine/Plugins/PCG/Source/PCG/Private/Compute/PCGDataBinding.cpp b/Engine/Plugins/PCG/Source/PCG/Private/Compute/PCGDataBinding.cpp
--- a/Engine/Plugins/PCG/Source/PCG/Private/Compute/PCGDataBinding.cpp
+++ b/Engine/Plugins/PCG/Source/PCG/Private/Compute/PCGDataBinding.cpp
@@ -30,13 +30,6 @@ void UPCGDataBinding::Initialize(const UPCGComputeGraph* InComputeGraph, FPCGCon
 	StringTable.Append(Graph->GetStaticStringTable());
 
 	InitializeInputData(InContext->InputData);
-
-#if PCG_KERNEL_LOGGING_ENABLED
-	if (Graph->bLogDataDescriptions)
-	{
-		DebugLogDataDescriptions();
-	}
-#endif
 }
 
 void UPCGDataBinding::ReleaseTransientResources()
@@ -50,36 +43,98 @@ void UPCGDataBinding::ReleaseTransientResources()
 	ContextHandle = nullptr;
 }
 
-bool UPCGDataBinding::GetKernelPinDataDesc_GameThread(const UPCGComputeKernel* InKernel, FName InPinLabel, bool bIsInput, FPCGDataCollectionDesc& OutDescription) const
+void UPCGDataBinding::PrimeDataDescriptionCache()
 {
-	return ensure(Graph && InKernel) && GetKernelPinDataDesc_GameThread(Graph->GetBindingIndex(InKernel->GetKernelIndex(), InPinLabel, bIsInput), OutDescription);
+	TRACE_CPUPROFILER_EVENT_SCOPE(UPCGDataBinding::PrimeDataDescriptionCache);
+
+	if (bIsDataDescriptionCachePrimed || !DataDescriptionCache.IsEmpty())
+	{
+		ensureMsgf(false, TEXT("Attempted to prime the data description cache, but it was already populated."));
+		return;
+	}
+
+	if (ensure(Graph))
+	{
+		TArray<FPCGKernelPin> KernelPins;
+		Graph->GetKernelPins(KernelPins);
+
+		// Compute and cache a description for every kernel pin in the graph (including both input and output pins).
+		// Note: It is fine for some kernel pins to fail to compute a description, e.g. DataLabelResolver pins, as they will never need a data description.
+		// @todo_pcg: Could add a ParallelFor if this becomes a bottle neck.
+		for (const FPCGKernelPin& KernelPin : KernelPins)
+		{
+			ComputeKernelPinDataDesc(KernelPin, /*OutDesc=*/nullptr);
+		}
+	}
+
+	bIsDataDescriptionCachePrimed = true;
 }
 
-bool UPCGDataBinding::GetKernelPinDataDesc_GameThread(int32 InGraphBindingIndex, FPCGDataCollectionDesc& OutDescription) const
+bool UPCGDataBinding::ComputeKernelPinDataDesc(const FPCGKernelPin& InKernelPin, FPCGDataCollectionDesc* OutDesc)
 {
-	ensure(IsInGameThread());
+	// Uncomment to be notified when a kernel pin description is computed late:
+	//ensureAlwaysMsgf(!IsInGameThread(), TEXT("Data descriptions should only be computed when priming the cache, which takes place on a worker thread."));
 
-	if (InGraphBindingIndex == INDEX_NONE)
+	if (!ensure(Graph))
 	{
-		ensure(false);
 		return false;
 	}
 
-	if (const FPCGDataCollectionDesc* FoundDescription = OutputPinDataDescriptionCache.Find(InGraphBindingIndex))
+	const int32 GraphBindingIndex = Graph->GetBindingIndex(InKernelPin);
+
+	if (FPCGDataCollectionDesc* FoundDesc = DataDescriptionCache.Find(GraphBindingIndex))
 	{
-		OutDescription = *FoundDescription;
+		if (OutDesc)
+		{
+			*OutDesc = *FoundDesc;
+		}
+
 		return true;
 	}
 
-	if (Graph->ComputeKernelBindingDataDesc(InGraphBindingIndex, this, OutDescription))
+	FPCGDataCollectionDesc ComputedDesc;
+	
+	if (Graph->ComputeKernelBindingDataDesc(GraphBindingIndex, this, ComputedDesc))
 	{
-		OutputPinDataDescriptionCache.Add(InGraphBindingIndex, OutDescription);
+		DataDescriptionCache.Add(GraphBindingIndex, ComputedDesc);
+
+		if (OutDesc)
+		{
+			*OutDesc = MoveTemp(ComputedDesc);
+		}
+
 		return true;
 	}
-	else
+
+	return false;
+}
+
+const FPCGDataCollectionDesc* UPCGDataBinding::GetCachedKernelPinDataDesc(const UPCGComputeKernel* InKernel, FName InPinLabel, bool bIsInput) const
+{
+	if (!ensure(Graph) || !ensure(InKernel))
 	{
-		return false;
+		return nullptr;
+	}
+
+	return GetCachedKernelPinDataDesc(Graph->GetBindingIndex(FPCGKernelPin(InKernel->GetKernelIndex(), InPinLabel, bIsInput)));
+}
+
+const FPCGDataCollectionDesc* UPCGDataBinding::GetCachedKernelPinDataDesc(int32 InGraphBindingIndex) const
+{
+	TRACE_CPUPROFILER_EVENT_SCOPE(UPCGDataBinding::GetCachedKernelPinDataDesc);
+
+	if (!ensure(bIsDataDescriptionCachePrimed))
+	{
+		UE_LOG(LogPCG, Error, TEXT("Tried to get a kernel pin data description before the data description cache was primed."));
+		return nullptr;
 	}
+
+	const FPCGDataCollectionDesc* FoundDescription = DataDescriptionCache.Find(InGraphBindingIndex);
+
+	// We should never land here, all data descriptions should have been computed at this point.
+	ensureMsgf(FoundDescription, TEXT("Cache miss when trying to look up a kernel pin data description."));
+
+	return FoundDescription;
 }
 
 void UPCGDataBinding::ReceiveDataFromGPU_GameThread(
@@ -164,7 +219,7 @@ int32 UPCGDataBinding::GetFirstInputDataIndex(const UPCGComputeKernel* InKernel,
 {
 	int32 FirstDataIndex = INDEX_NONE;
 
-	const int32 BindingIndex = Graph ? Graph->GetBindingIndex(InKernel->GetKernelIndex(), InPinLabel, /*bIsInput=*/true) : INDEX_NONE;
+	const int32 BindingIndex = Graph ? Graph->GetBindingIndex(FPCGKernelPin(InKernel->GetKernelIndex(), InPinLabel, /*bIsInput=*/true)) : INDEX_NONE;
 
 	if (ensure(BindingIndex != INDEX_NONE))
 	{
@@ -408,17 +463,20 @@ void UPCGDataBinding::AddInputDataTagsToTable()
 void UPCGDataBinding::DebugLogDataDescriptions() const
 {
 #if PCG_KERNEL_LOGGING_ENABLED
-	UE_LOG(LogPCG, Display, TEXT("\n### METADATA ATTRIBUTE TABLE ###"));
-	AttributeTable.DebugLog();
-
-	UE_LOG(LogPCG, Display, TEXT("\n### STRING TABLE ###"));
-	for (int32 i = 0; i < StringTable.Num(); ++i)
+	if (Graph && Graph->bLogDataDescriptions)
 	{
-		UE_LOG(LogPCG, Display, TEXT("\t%d: %s"), i, *StringTable[i]);
-	}
+		UE_LOG(LogPCG, Display, TEXT("\n### METADATA ATTRIBUTE TABLE ###"));
+		AttributeTable.DebugLog();
 
-	// TODO - we don't have explicit pin lists for kernels because we don't need them. Decide how to bring this back or pivot the feature.
-	UE_LOG(LogPCG, Display, TEXT("\n### INPUT PIN DATA DESCRIPTIONS ###"));
-	Graph->DebugLogDataDescriptions(this);
+		UE_LOG(LogPCG, Display, TEXT("\n### STRING TABLE ###"));
+		for (int32 i = 0; i < StringTable.Num(); ++i)
+		{
+			UE_LOG(LogPCG, Display, TEXT("\t%d: %s"), i, *StringTable[i]);
+		}
+
+		// @todo_pcg: We don't have explicit pin lists for kernels because we don't need them. Decide how to bring this back or pivot the feature.
+		UE_LOG(LogPCG, Display, TEXT("\n### INPUT PIN DATA DESCRIPTIONS ###"));
+		Graph->DebugLogDataDescriptions(this);
+	}
 #endif
 }
diff --git a/Engine/Plugins/PCG/Source/PCG/Private/Compute/PCGDataBinding.h b/Engine/Plugins/PCG/Source/PCG/Private/Compute/PCGDataBinding.h
--- a/Engine/Plugins/PCG/Source/PCG/Private/Compute/PCGDataBinding.h
+++ b/Engine/Plugins/PCG/Source/PCG/Private/Compute/PCGDataBinding.h
@@ -14,6 +14,7 @@ class UPCGProxyForGPUData;
 class UPCGStaticMeshSpawnerKernel;
 struct FPCGComputeGraphContext;
 struct FPCGContextHandle;
+struct FPCGKernelPin;
 
 USTRUCT()
 struct FPCGSpawnerPrimitives
@@ -73,9 +74,16 @@ public:
 	/** Get the set of all unique strings entering the compute graph. */
 	const TArray<FString>& GetStringTable() const { return StringTable; }
 
-	/** Get description of data produced by an output pin. */
-	bool GetKernelPinDataDesc_GameThread(const UPCGComputeKernel* InKernel, FName InPinLabel, bool bIsInput, FPCGDataCollectionDesc& OutDescription) const;
-	bool GetKernelPinDataDesc_GameThread(int32 InGraphBindingIndex, FPCGDataCollectionDesc& OutDescription) const;
+	/** Computes a description of data for every pin in the compute graph and caches it. */
+	void PrimeDataDescriptionCache();
+	bool IsDataDescriptionCachePrimed() const { return bIsDataDescriptionCachePrimed; }
+
+	/** Computes a description of data for a single pin in the compute graph and caches it. */
+	bool ComputeKernelPinDataDesc(const FPCGKernelPin& InKernelPin, FPCGDataCollectionDesc* OutDesc);
+
+	/** Get description of data produced by a pin. */
+	const FPCGDataCollectionDesc* GetCachedKernelPinDataDesc(const UPCGComputeKernel* InKernel, FName InPinLabel, bool bIsInput) const;
+	const FPCGDataCollectionDesc* GetCachedKernelPinDataDesc(int32 InGraphBindingIndex) const;
 
 	void ReceiveDataFromGPU_GameThread(UPCGProxyForGPUData* InData, const UPCGSettings* InProducerSettings, EPCGExportMode InExportMode, FName InPinLabel, FName InPinLabelAlias);
 
@@ -97,18 +105,16 @@ public:
 	*/
 	bool ReadbackInputDataToCPU(int32 InInputDataIndex);
 
+	void DebugLogDataDescriptions() const;
+
 private:
 	void InitializeInputData(const FPCGDataCollection& InComputeGraphElementInputData);
 
 	/** Loop over all metadata attributes in all input data ensure all attributes are registered in attribute table. */
 	void AddInputDataAttributesToTable();
-
 	void AddInputDataStringsToTable();
-
 	void AddInputDataTagsToTable();
 
-	void DebugLogDataDescriptions() const;
-
 public:
 	/** Compute graph element output data. Data items are labelled with unique virtual output pin labels so that the can be routed correctly by the graph
 	 * executor to downstream nodes. */
@@ -145,6 +151,7 @@ protected:
 	/** Remaps attribute IDs to avoid collisions with this compute graph's attribute IDs. */
 	TMap<FSourceBufferAndAttributeIndex, /*Attribute ID*/int32> AttributeIdRemap;
 
-	/** Cache of data descriptions computed for output pins to amortized cost of building data descriptions. */
-	mutable TMap</*Binding index*/int32, FPCGDataCollectionDesc> OutputPinDataDescriptionCache;
+	/** Cache of data descriptions to amortize cost of computing them at runtime. */
+	TMap</*Binding index*/int32, FPCGDataCollectionDesc> DataDescriptionCache;
+	std::atomic<bool> bIsDataDescriptionCachePrimed = false;
 };
diff --git a/Engine/Plugins/PCG/Source/PCG/Private/Compute/PCGDataForGPU.cpp b/Engine/Plugins/PCG/Source/PCG/Private/Compute/PCGDataForGPU.cpp
--- a/Engine/Plugins/PCG/Source/PCG/Private/Compute/PCGDataForGPU.cpp
+++ b/Engine/Plugins/PCG/Source/PCG/Private/Compute/PCGDataForGPU.cpp
@@ -560,12 +560,13 @@ namespace PCGDataForGPUHelpers
 		const UPCGSettings* InSettings,
 		TArray<FName>& InAttributeNames,
 		const UPCGDataBinding* InBinding,
-		const FPCGDataCollectionDesc& InDataCollectionDescription,
+		const FPCGDataCollectionDesc* InDataCollectionDescription,
 		uint32& OutCustomFloatCount,
 		TArray<FUint32Vector4>& OutAttributeIdOffsetStrides)
 	{
 		check(InBinding);
 		check(InBinding->Graph);
+		check(InDataCollectionDescription);
 
 		uint32 OffsetFloats = 0;
 
@@ -575,7 +576,7 @@ namespace PCGDataForGPUHelpers
 			// find in the input data similar to the CPU SM Spawner.
 			FPCGKernelAttributeDesc AttributeDesc;
 			bool bConflictingTypesInData = false;
-			InDataCollectionDescription.GetAttributeDesc(AttributeName, AttributeDesc, bConflictingTypesInData);
+			InDataCollectionDescription->GetAttributeDesc(AttributeName, AttributeDesc, bConflictingTypesInData);
 
 			if (bConflictingTypesInData)
 			{
@@ -920,31 +921,6 @@ FPCGDataCollectionDesc FPCGDataCollectionDesc::BuildFromDataCollection(
 	return CollectionDesc;
 }
 
-FPCGDataCollectionDesc FPCGDataCollectionDesc::BuildFromInputDataCollectionAndInputPinLabel(
-	const FPCGDataCollection& InDataCollection,
-	FName InputPinLabel,
-	const UPCGDataBinding* InBinding)
-{
-	FPCGDataCollectionDesc CollectionDesc;
-	TArray<FPCGTaggedData> DataForPin = InDataCollection.GetInputsByPin(InputPinLabel);
-
-	for (const FPCGTaggedData& Data : DataForPin)
-	{
-		if (!Data.Data || !PCGComputeHelpers::IsTypeAllowedAsInput(Data.Data->GetDataType()))
-		{
-			continue;
-		}
-
-		FPCGDataDesc DataDesc;
-		if (ensure(PCGComputeHelpers::ComputeDataDescription(Data, InBinding, DataDesc)))
-		{
-			CollectionDesc.DataDescs.Add(MoveTemp(DataDesc));
-		}
-	}
-
-	return CollectionDesc;
-}
-
 uint32 FPCGDataCollectionDesc::ComputePackedHeaderSizeBytes() const
 {
 	return PCGComputeConstants::DATA_COLLECTION_HEADER_SIZE_BYTES + PCGComputeConstants::DATA_HEADER_SIZE_BYTES * DataDescs.Num();
diff --git a/Engine/Plugins/PCG/Source/PCG/Private/Compute/PCGDataUploadAdaptor.cpp b/Engine/Plugins/PCG/Source/PCG/Private/Compute/PCGDataUploadAdaptor.cpp
--- a/Engine/Plugins/PCG/Source/PCG/Private/Compute/PCGDataUploadAdaptor.cpp
+++ b/Engine/Plugins/PCG/Source/PCG/Private/Compute/PCGDataUploadAdaptor.cpp
@@ -25,11 +25,17 @@ namespace PCGDataUploadAdaptor
 #endif // !UE_BUILD_SHIPPING
 }
 
-FPCGDataUploadAdaptor::FPCGDataUploadAdaptor(UPCGDataBinding* InDataBinding, const FPCGDataCollectionDesc& InTargetDataCollectionDesc, FName InInputPinLabel)
+FPCGDataUploadAdaptor::FPCGDataUploadAdaptor(UPCGDataBinding* InDataBinding, const FPCGDataCollectionDesc* InTargetDataCollectionDesc, FName InInputPinLabel)
 {
 	check(IsInGameThread());
+	check(InTargetDataCollectionDesc);
 
-	TargetDataCollectionDesc = InTargetDataCollectionDesc;
+	{
+		TRACE_CPUPROFILER_EVENT_SCOPE(CopyDataDescription);
+
+		// @todo_pcg: Avoid costly copy of data description.
+		TargetDataCollectionDesc = *InTargetDataCollectionDesc;
+	}
 	InputPinLabel = InInputPinLabel;
 	DataBinding = InDataBinding;
 
diff --git a/Engine/Plugins/PCG/Source/PCG/Private/Compute/PCGDataUploadAdaptor.h b/Engine/Plugins/PCG/Source/PCG/Private/Compute/PCGDataUploadAdaptor.h
--- a/Engine/Plugins/PCG/Source/PCG/Private/Compute/PCGDataUploadAdaptor.h
+++ b/Engine/Plugins/PCG/Source/PCG/Private/Compute/PCGDataUploadAdaptor.h
@@ -15,7 +15,7 @@ enum class EPCGExportMode : uint8;
 class FPCGDataUploadAdaptor
 {
 public:
-	FPCGDataUploadAdaptor(UPCGDataBinding* InDataBinding, const FPCGDataCollectionDesc& InTargetDataCollectionDesc, FName InInputPinLabel);
+	FPCGDataUploadAdaptor(UPCGDataBinding* InDataBinding, const FPCGDataCollectionDesc* InTargetDataCollectionDesc, FName InInputPinLabel);
 
 	/** Do any preparation work such as data readbacks. Returns true when preparation is complete. */
 	bool PrepareData_GameThread();
diff --git a/Engine/Plugins/PCG/Source/PCG/Private/Elements/Metadata/PCGMetadataPartitionKernel.cpp b/Engine/Plugins/PCG/Source/PCG/Private/Elements/Metadata/PCGMetadataPartitionKernel.cpp
--- a/Engine/Plugins/PCG/Source/PCG/Private/Elements/Metadata/PCGMetadataPartitionKernel.cpp
+++ b/Engine/Plugins/PCG/Source/PCG/Private/Elements/Metadata/PCGMetadataPartitionKernel.cpp
@@ -23,11 +23,11 @@ const UPCGSettings* UPCGMetadataPartitionKernel::GetSettings() const
 	return Settings;
 }
 
-bool UPCGMetadataPartitionKernel::IsKernelValid(FPCGContext* InContext) const
+bool UPCGMetadataPartitionKernel::IsKernelDataValid(FPCGContext* InContext) const
 {
-	TRACE_CPUPROFILER_EVENT_SCOPE(UPCGMetadataPartitionKernel::IsKernelValid);
+	TRACE_CPUPROFILER_EVENT_SCOPE(UPCGMetadataPartitionKernel::IsKernelDataValid);
 
-	if (!Super::IsKernelValid(InContext))
+	if (!Super::IsKernelDataValid(InContext))
 	{
 		return false;
 	}
@@ -39,16 +39,20 @@ bool UPCGMetadataPartitionKernel::IsKernelValid(FPCGContext* InContext) const
 		const UPCGMetadataPartitionSettings* MPSettings = CastChecked<UPCGMetadataPartitionSettings>(Settings);
 		const FName AttributeName = MPSettings->PartitionAttributeSelectors[0].GetAttributeName();
 
-		FPCGDataCollectionDesc InputDesc;
-		ensure(Context->DataBinding->GetKernelPinDataDesc_GameThread(this, PCGPinConstants::DefaultInputLabel, /*bIsInput=*/true, InputDesc));
+		const FPCGDataCollectionDesc* InputDesc = Context->DataBinding->GetCachedKernelPinDataDesc(this, PCGPinConstants::DefaultInputLabel, /*bIsInput=*/true);
+
+		if (!ensure(InputDesc))
+		{
+			return false;
+		}
 
 		FPCGKernelAttributeDesc AttributeDesc;
 		bool bConflictingTypesInData = false;
-		const bool bAttributePresent = InputDesc.GetAttributeDesc(AttributeName, AttributeDesc, bConflictingTypesInData);
+		const bool bAttributePresent = InputDesc->GetAttributeDesc(AttributeName, AttributeDesc, bConflictingTypesInData);
 
 		if (!bAttributePresent)
 		{
-			if (!InputDesc.DataDescs.IsEmpty())
+			if (!InputDesc->DataDescs.IsEmpty())
 			{
 				PCG_KERNEL_VALIDATION_ERR(InContext, Settings, FText::Format(
 					LOCTEXT("PartitionAttributeMissing", "Partition attribute was not present in GPU data."),
@@ -78,7 +82,7 @@ bool UPCGMetadataPartitionKernel::IsKernelValid(FPCGContext* InContext) const
 	return true;
 }
 
-bool UPCGMetadataPartitionKernel::ComputeOutputBindingDataDesc(const UPCGComputeGraph* InGraph, FName InOutputPinLabel, const UPCGDataBinding* InBinding, FPCGDataCollectionDesc& OutDataDesc) const
+bool UPCGMetadataPartitionKernel::ComputeOutputBindingDataDesc(const UPCGComputeGraph* InGraph, FName InOutputPinLabel, UPCGDataBinding* InBinding, FPCGDataCollectionDesc& OutDataDesc) const
 {
 	check(InGraph);
 	check(InBinding);
@@ -92,7 +96,8 @@ bool UPCGMetadataPartitionKernel::ComputeOutputBindingDataDesc(const UPCGCompute
 	}
 
 	FPCGDataCollectionDesc InputPinDesc;
-	ensure(InBinding->GetKernelPinDataDesc_GameThread(this, PCGPinConstants::DefaultInputLabel, /*bIsInput=*/true, InputPinDesc));
+	const FPCGKernelPin InputKernelPin(KernelIndex, PCGPinConstants::DefaultInputLabel, /*bIsInput=*/true);
+	ensure(InBinding->ComputeKernelPinDataDesc(InputKernelPin, &InputPinDesc));
 
 	const int32 PartitionAttributeId = InBinding->GetAttributeId(MPSettings->PartitionAttributeSelectors[0].GetAttributeName(), EPCGKernelAttributeType::StringKey);
 
@@ -136,10 +141,8 @@ bool UPCGMetadataPartitionKernel::ComputeOutputBindingDataDesc(const UPCGCompute
 
 int UPCGMetadataPartitionKernel::ComputeThreadCount(const UPCGDataBinding* InBinding) const
 {
-	FPCGDataCollectionDesc InputPinDesc;
-	ensure(InBinding->GetKernelPinDataDesc_GameThread(this, PCGPinConstants::DefaultInputLabel, /*bIsInput=*/true, InputPinDesc));
-
-	return InputPinDesc.ComputeDataElementCount(EPCGDataType::Point);
+	const FPCGDataCollectionDesc* InputPinDesc = InBinding->GetCachedKernelPinDataDesc(this, PCGPinConstants::DefaultInputLabel, /*bIsInput=*/true);
+	return ensure(InputPinDesc) ? InputPinDesc->ComputeDataElementCount(EPCGDataType::Point) : 0;
 }
 
 #if WITH_EDITOR
diff --git a/Engine/Plugins/PCG/Source/PCG/Private/Elements/PCGCopyPointsAnalysisKernel.cpp b/Engine/Plugins/PCG/Source/PCG/Private/Elements/PCGCopyPointsAnalysisKernel.cpp
--- a/Engine/Plugins/PCG/Source/PCG/Private/Elements/PCGCopyPointsAnalysisKernel.cpp
+++ b/Engine/Plugins/PCG/Source/PCG/Private/Elements/PCGCopyPointsAnalysisKernel.cpp
@@ -26,11 +26,11 @@ const UPCGSettings* UPCGCopyPointsAnalysisKernel::GetSettings() const
 	return Settings;
 }
 
-bool UPCGCopyPointsAnalysisKernel::IsKernelValid(FPCGContext* InContext) const
+bool UPCGCopyPointsAnalysisKernel::IsKernelDataValid(FPCGContext* InContext) const
 {
-	TRACE_CPUPROFILER_EVENT_SCOPE(UPCGCopyPointsAnalysisKernel::IsKernelValid);
+	TRACE_CPUPROFILER_EVENT_SCOPE(UPCGCopyPointsAnalysisKernel::IsKernelDataValid);
 
-	if (!Super::IsKernelValid(InContext))
+	if (!Super::IsKernelDataValid(InContext))
 	{
 		return false;
 	}
@@ -40,10 +40,10 @@ bool UPCGCopyPointsAnalysisKernel::IsKernelValid(FPCGContext* InContext) const
 		return true;
 	}
 
-	return PCGCopyPointsKernel::IsKernelValid(this, CastChecked<UPCGCopyPointsSettings>(Settings), static_cast<FPCGComputeGraphContext*>(InContext));
+	return PCGCopyPointsKernel::IsKernelDataValid(this, CastChecked<UPCGCopyPointsSettings>(Settings), static_cast<FPCGComputeGraphContext*>(InContext));
 }
 
-bool UPCGCopyPointsAnalysisKernel::ComputeOutputBindingDataDesc(const UPCGComputeGraph* InGraph, FName InOutputPinLabel, const UPCGDataBinding* InBinding, FPCGDataCollectionDesc& OutDataDesc) const
+bool UPCGCopyPointsAnalysisKernel::ComputeOutputBindingDataDesc(const UPCGComputeGraph* InGraph, FName InOutputPinLabel, UPCGDataBinding* InBinding, FPCGDataCollectionDesc& OutDataDesc) const
 {
 	check(InGraph);
 	check(InBinding);
@@ -55,8 +55,12 @@ bool UPCGCopyPointsAnalysisKernel::ComputeOutputBindingDataDesc(const UPCGComput
 	}
 
 	FPCGDataCollectionDesc SourcesDesc, TargetsDesc;
-	ensure(InBinding->GetKernelPinDataDesc_GameThread(this, PCGCopyPointsConstants::SourcePointsLabel, /*bIsInput=*/true, SourcesDesc));
-	ensure(InBinding->GetKernelPinDataDesc_GameThread(this, PCGCopyPointsConstants::TargetPointsLabel, /*bIsInput=*/true, TargetsDesc));
+
+	const FPCGKernelPin SourceKernelPin(KernelIndex, PCGCopyPointsConstants::SourcePointsLabel, /*bIsInput=*/true);
+	const FPCGKernelPin TargetKernelPin(KernelIndex, PCGCopyPointsConstants::TargetPointsLabel, /*bIsInput=*/true);
+
+	ensure(InBinding->ComputeKernelPinDataDesc(SourceKernelPin, &SourcesDesc));
+	ensure(InBinding->ComputeKernelPinDataDesc(TargetKernelPin, &TargetsDesc));
 
 	const uint32 NumSourceDatas = SourcesDesc.DataDescs.Num();
 	const uint32 NumTargetDatas = TargetsDesc.DataDescs.Num();
@@ -91,10 +95,11 @@ int UPCGCopyPointsAnalysisKernel::ComputeThreadCount(const UPCGDataBinding* InBi
 {
 	int ThreadCount = 0;
 
-	FPCGDataCollectionDesc OutputPinDesc;
-	if (ensure(InBinding->GetKernelPinDataDesc_GameThread(this, PCGPinConstants::DefaultOutputLabel, /*bIsInput=*/false, OutputPinDesc)))
+	const FPCGDataCollectionDesc* OutputPinDesc = InBinding->GetCachedKernelPinDataDesc(this, PCGPinConstants::DefaultOutputLabel, /*bIsInput=*/false);
+
+	if (ensure(OutputPinDesc))
 	{
-		for (const FPCGDataDesc& DataDesc : OutputPinDesc.DataDescs)
+		for (const FPCGDataDesc& DataDesc : OutputPinDesc->DataDescs)
 		{
 			ThreadCount += DataDesc.ElementCount;
 		}
diff --git a/Engine/Plugins/PCG/Source/PCG/Private/Elements/PCGCopyPointsKernel.cpp b/Engine/Plugins/PCG/Source/PCG/Private/Elements/PCGCopyPointsKernel.cpp
--- a/Engine/Plugins/PCG/Source/PCG/Private/Elements/PCGCopyPointsKernel.cpp
+++ b/Engine/Plugins/PCG/Source/PCG/Private/Elements/PCGCopyPointsKernel.cpp
@@ -27,11 +27,11 @@ const UPCGSettings* UPCGCopyPointsKernel::GetSettings() const
 	return Settings;
 }
 
-bool UPCGCopyPointsKernel::IsKernelValid(FPCGContext* InContext) const
+bool UPCGCopyPointsKernel::IsKernelDataValid(FPCGContext* InContext) const
 {
-	TRACE_CPUPROFILER_EVENT_SCOPE(UPCGCopyPointsKernel::IsKernelValid);
+	TRACE_CPUPROFILER_EVENT_SCOPE(UPCGCopyPointsKernel::IsKernelDataValid);
 
-	if (!Super::IsKernelValid(InContext))
+	if (!Super::IsKernelDataValid(InContext))
 	{
 		return false;
 	}
@@ -41,10 +41,10 @@ bool UPCGCopyPointsKernel::IsKernelValid(FPCGContext* InContext) const
 		return true;
 	}
 
-	return PCGCopyPointsKernel::IsKernelValid(this, CastChecked<UPCGCopyPointsSettings>(Settings), static_cast<FPCGComputeGraphContext*>(InContext));
+	return PCGCopyPointsKernel::IsKernelDataValid(this, CastChecked<UPCGCopyPointsSettings>(Settings), static_cast<FPCGComputeGraphContext*>(InContext));
 }
 
-bool UPCGCopyPointsKernel::ComputeOutputBindingDataDesc(const UPCGComputeGraph* InGraph, FName InOutputPinLabel, const UPCGDataBinding* InBinding, FPCGDataCollectionDesc& OutDataDesc) const
+bool UPCGCopyPointsKernel::ComputeOutputBindingDataDesc(const UPCGComputeGraph* InGraph, FName InOutputPinLabel, UPCGDataBinding* InBinding, FPCGDataCollectionDesc& OutDataDesc) const
 {
 	check(InGraph);
 	check(InBinding);
@@ -102,8 +102,12 @@ bool UPCGCopyPointsKernel::ComputeOutputBindingDataDesc(const UPCGComputeGraph*
 
 	FPCGDataCollectionDesc SourcePinDesc;
 	FPCGDataCollectionDesc TargetPinDesc;
-	ensure(InBinding->GetKernelPinDataDesc_GameThread(this, PCGCopyPointsConstants::SourcePointsLabel, /*bIsInput=*/true, SourcePinDesc));
-	ensure(InBinding->GetKernelPinDataDesc_GameThread(this, PCGCopyPointsConstants::TargetPointsLabel, /*bIsInput=*/true, TargetPinDesc));
+
+	const FPCGKernelPin SourceKernelPin(KernelIndex, PCGCopyPointsConstants::SourcePointsLabel, /*bIsInput=*/true);
+	const FPCGKernelPin TargetKernelPin(KernelIndex, PCGCopyPointsConstants::TargetPointsLabel, /*bIsInput=*/true);
+
+	ensure(InBinding->ComputeKernelPinDataDesc(SourceKernelPin, &SourcePinDesc));
+	ensure(InBinding->ComputeKernelPinDataDesc(TargetKernelPin, &TargetPinDesc));
 
 	FPCGDataCollectionDesc PinDesc;
 
@@ -181,10 +185,11 @@ int UPCGCopyPointsKernel::ComputeThreadCount(const UPCGDataBinding* InBinding) c
 {
 	int ThreadCount = 0;
 
-	FPCGDataCollectionDesc OutputPinDesc;
-	if (ensure(InBinding->GetKernelPinDataDesc_GameThread(this, PCGPinConstants::DefaultOutputLabel, /*bIsInput=*/false, OutputPinDesc)))
+	const FPCGDataCollectionDesc* OutputPinDesc = InBinding->GetCachedKernelPinDataDesc(this, PCGPinConstants::DefaultOutputLabel, /*bIsInput=*/false);
+
+	if (ensure(OutputPinDesc))
 	{
-		for (const FPCGDataDesc& DataDesc : OutputPinDesc.DataDescs)
+		for (const FPCGDataDesc& DataDesc : OutputPinDesc->DataDescs)
 		{
 			ThreadCount += DataDesc.ElementCount;
 		}
diff --git a/Engine/Plugins/PCG/Source/PCG/Private/Elements/PCGCopyPointsKernelShared.cpp b/Engine/Plugins/PCG/Source/PCG/Private/Elements/PCGCopyPointsKernelShared.cpp
--- a/Engine/Plugins/PCG/Source/PCG/Private/Elements/PCGCopyPointsKernelShared.cpp
+++ b/Engine/Plugins/PCG/Source/PCG/Private/Elements/PCGCopyPointsKernelShared.cpp
@@ -10,7 +10,7 @@
 
 #define LOCTEXT_NAMESPACE "PCGCopyPointsKernel"
 
-bool PCGCopyPointsKernel::IsKernelValid(const UPCGComputeKernel* InKernel, const UPCGCopyPointsSettings* InCopyPointSettings, const FPCGComputeGraphContext* InContext)
+bool PCGCopyPointsKernel::IsKernelDataValid(const UPCGComputeKernel* InKernel, const UPCGCopyPointsSettings* InCopyPointSettings, const FPCGComputeGraphContext* InContext)
 {
 	check(InCopyPointSettings);
 	check(InContext);
@@ -21,15 +21,19 @@ bool PCGCopyPointsKernel::IsKernelValid(const UPCGComputeKernel* InKernel, const
 		{
 			auto ValidateAttributeExists = [InContext, DataBinding, InKernel, InCopyPointSettings](FName InputPin) -> bool
 			{
-				FPCGDataCollectionDesc PinDataDesc;
-				ensure(DataBinding->GetKernelPinDataDesc_GameThread(InKernel, InputPin, /*bIsInputPin=*/true, PinDataDesc));
+				const FPCGDataCollectionDesc* PinDataDesc = DataBinding->GetCachedKernelPinDataDesc(InKernel, InputPin, /*bIsInputPin=*/true);
+
+				if (!ensure(PinDataDesc))
+				{
+					return false;
+				}
 
 				const FName MatchAttributeName = InCopyPointSettings->MatchAttribute;
 
 				// todo_pcg: Can generalize this to any type?
 				constexpr EPCGKernelAttributeType MatchAttributeType = EPCGKernelAttributeType::Int;
 
-				for (const FPCGDataDesc& DataDesc : PinDataDesc.DataDescs)
+				for (const FPCGDataDesc& DataDesc : PinDataDesc->DataDescs)
 				{
 					if (!DataDesc.ContainsAttribute(MatchAttributeName, MatchAttributeType))
 					{
@@ -42,7 +46,7 @@ bool PCGCopyPointsKernel::IsKernelValid(const UPCGComputeKernel* InKernel, const
 				}
 
 				// Valid for execution if we have some data to process.
-				return !PinDataDesc.DataDescs.IsEmpty();
+				return !PinDataDesc->DataDescs.IsEmpty();
 			};
 
 			if (!ValidateAttributeExists(PCGCopyPointsConstants::SourcePointsLabel) || !ValidateAttributeExists(PCGCopyPointsConstants::TargetPointsLabel))
diff --git a/Engine/Plugins/PCG/Source/PCG/Private/Elements/PCGStaticMeshSpawnerKernel.cpp b/Engine/Plugins/PCG/Source/PCG/Private/Elements/PCGStaticMeshSpawnerKernel.cpp
--- a/Engine/Plugins/PCG/Source/PCG/Private/Elements/PCGStaticMeshSpawnerKernel.cpp
+++ b/Engine/Plugins/PCG/Source/PCG/Private/Elements/PCGStaticMeshSpawnerKernel.cpp
@@ -37,7 +37,7 @@ const UPCGSettings* UPCGStaticMeshSpawnerKernel::GetSettings() const
 	return Settings;
 }
 
-bool UPCGStaticMeshSpawnerKernel::ComputeOutputBindingDataDesc(const UPCGComputeGraph* InGraph, FName InOutputPinLabel, const UPCGDataBinding* InBinding, FPCGDataCollectionDesc& OutDataDesc) const
+bool UPCGStaticMeshSpawnerKernel::ComputeOutputBindingDataDesc(const UPCGComputeGraph* InGraph, FName InOutputPinLabel, UPCGDataBinding* InBinding, FPCGDataCollectionDesc& OutDataDesc) const
 {
 	check(InGraph);
 	check(InBinding);
@@ -50,7 +50,8 @@ bool UPCGStaticMeshSpawnerKernel::ComputeOutputBindingDataDesc(const UPCGCompute
 
 	// Forward data from In to Out.
 	FPCGDataCollectionDesc PinDesc;
-	ensure(InBinding->GetKernelPinDataDesc_GameThread(this, PCGPinConstants::DefaultInputLabel, /*bIsInput=*/true, PinDesc));
+	const FPCGKernelPin InputKernelPin(KernelIndex, PCGPinConstants::DefaultInputLabel, /*bIsInput=*/true);
+	ensure(InBinding->ComputeKernelPinDataDesc(InputKernelPin, &PinDesc));
 
 	// Add output attribute (selected mesh).
 	{
@@ -110,10 +111,8 @@ bool UPCGStaticMeshSpawnerKernel::ComputeOutputBindingDataDesc(const UPCGCompute
 
 int UPCGStaticMeshSpawnerKernel::ComputeThreadCount(const UPCGDataBinding* InBinding) const
 {
-	FPCGDataCollectionDesc InputPinDesc;
-	ensure(InBinding->GetKernelPinDataDesc_GameThread(this, PCGPinConstants::DefaultInputLabel, /*bIsInput=*/true, InputPinDesc));
-
-	return InputPinDesc.ComputeDataElementCount(EPCGDataType::Point);
+	const FPCGDataCollectionDesc* InputPinDesc = InBinding->GetCachedKernelPinDataDesc(this, PCGPinConstants::DefaultInputLabel, /*bIsInput=*/true);
+	return ensure(InputPinDesc) ? InputPinDesc->ComputeDataElementCount(EPCGDataType::Point) : 0;
 }
 
 #if WITH_EDITOR
diff --git a/Engine/Plugins/PCG/Source/PCG/Public/Elements/Metadata/PCGMetadataPartitionKernel.h b/Engine/Plugins/PCG/Source/PCG/Public/Elements/Metadata/PCGMetadataPartitionKernel.h
--- a/Engine/Plugins/PCG/Source/PCG/Public/Elements/Metadata/PCGMetadataPartitionKernel.h
+++ b/Engine/Plugins/PCG/Source/PCG/Public/Elements/Metadata/PCGMetadataPartitionKernel.h
@@ -15,8 +15,8 @@ class UPCGMetadataPartitionKernel : public UPCGComputeKernel
 
 public:
 	virtual const UPCGSettings* GetSettings() const override;
-	virtual bool IsKernelValid(FPCGContext* InContext) const override;
-	virtual bool ComputeOutputBindingDataDesc(const UPCGComputeGraph* InGraph, FName InOutputPinLabel, const UPCGDataBinding* InBinding, FPCGDataCollectionDesc& OutDataDesc) const override;
+	virtual bool IsKernelDataValid(FPCGContext* InContext) const override;
+	virtual bool ComputeOutputBindingDataDesc(const UPCGComputeGraph* InGraph, FName InOutputPinLabel, UPCGDataBinding* InBinding, FPCGDataCollectionDesc& OutDataDesc) const override;
 	virtual int ComputeThreadCount(const UPCGDataBinding* Binding) const override;
 #if WITH_EDITOR
 	virtual FString GetCookedSource(const UPCGComputeGraph* InComputeGraph) const override;
diff --git a/Engine/Plugins/PCG/Source/PCG/Public/Elements/PCGCopyPointsAnalysisKernel.h b/Engine/Plugins/PCG/Source/PCG/Public/Elements/PCGCopyPointsAnalysisKernel.h
--- a/Engine/Plugins/PCG/Source/PCG/Public/Elements/PCGCopyPointsAnalysisKernel.h
+++ b/Engine/Plugins/PCG/Source/PCG/Public/Elements/PCGCopyPointsAnalysisKernel.h
@@ -15,8 +15,8 @@ class UPCGCopyPointsAnalysisKernel : public UPCGComputeKernel
 
 public:
 	virtual const UPCGSettings* GetSettings() const override;
-	virtual bool IsKernelValid(FPCGContext* InContext) const override;
-	virtual bool ComputeOutputBindingDataDesc(const UPCGComputeGraph* InGraph, FName InOutputPinLabel, const UPCGDataBinding* InBinding, FPCGDataCollectionDesc& OutDataDesc) const override;
+	virtual bool IsKernelDataValid(FPCGContext* InContext) const override;
+	virtual bool ComputeOutputBindingDataDesc(const UPCGComputeGraph* InGraph, FName InOutputPinLabel, UPCGDataBinding* InBinding, FPCGDataCollectionDesc& OutDataDesc) const override;
 	virtual int ComputeThreadCount(const UPCGDataBinding* Binding) const override;
 #if WITH_EDITOR
 	virtual FString GetCookedSource(const UPCGComputeGraph* InComputeGraph) const override;
diff --git a/Engine/Plugins/PCG/Source/PCG/Public/Elements/PCGCopyPointsKernel.h b/Engine/Plugins/PCG/Source/PCG/Public/Elements/PCGCopyPointsKernel.h
--- a/Engine/Plugins/PCG/Source/PCG/Public/Elements/PCGCopyPointsKernel.h
+++ b/Engine/Plugins/PCG/Source/PCG/Public/Elements/PCGCopyPointsKernel.h
@@ -17,8 +17,8 @@ class UPCGCopyPointsKernel : public UPCGComputeKernel
 
 public:
 	virtual const UPCGSettings* GetSettings() const override;
-	virtual bool IsKernelValid(FPCGContext* InContext) const override;
-	virtual bool ComputeOutputBindingDataDesc(const UPCGComputeGraph* InGraph, FName InOutputPinLabel, const UPCGDataBinding* InBinding, FPCGDataCollectionDesc& OutDataDesc) const override;
+	virtual bool IsKernelDataValid(FPCGContext* InContext) const override;
+	virtual bool ComputeOutputBindingDataDesc(const UPCGComputeGraph* InGraph, FName InOutputPinLabel, UPCGDataBinding* InBinding, FPCGDataCollectionDesc& OutDataDesc) const override;
 	virtual int ComputeThreadCount(const UPCGDataBinding* Binding) const override;
 #if WITH_EDITOR
 	virtual FString GetCookedSource(const UPCGComputeGraph* InComputeGraph) const override;
diff --git a/Engine/Plugins/PCG/Source/PCG/Public/Elements/PCGCopyPointsKernelShared.h b/Engine/Plugins/PCG/Source/PCG/Public/Elements/PCGCopyPointsKernelShared.h
--- a/Engine/Plugins/PCG/Source/PCG/Public/Elements/PCGCopyPointsKernelShared.h
+++ b/Engine/Plugins/PCG/Source/PCG/Public/Elements/PCGCopyPointsKernelShared.h
@@ -9,6 +9,6 @@ struct FPCGContext;
 
 namespace PCGCopyPointsKernel
 {
-	/** Performs validation common to all copy points kernels. */
-	bool IsKernelValid(const UPCGComputeKernel* InKernel, const UPCGCopyPointsSettings* InCopyPointSettings, const FPCGComputeGraphContext* InContext);
+	/** Performs data validation common to all copy points kernels. */
+	bool IsKernelDataValid(const UPCGComputeKernel* InKernel, const UPCGCopyPointsSettings* InCopyPointSettings, const FPCGComputeGraphContext* InContext);
 }
diff --git a/Engine/Plugins/PCG/Source/PCG/Public/Elements/PCGStaticMeshSpawnerKernel.h b/Engine/Plugins/PCG/Source/PCG/Public/Elements/PCGStaticMeshSpawnerKernel.h
--- a/Engine/Plugins/PCG/Source/PCG/Public/Elements/PCGStaticMeshSpawnerKernel.h
+++ b/Engine/Plugins/PCG/Source/PCG/Public/Elements/PCGStaticMeshSpawnerKernel.h
@@ -20,7 +20,7 @@ class UPCGStaticMeshSpawnerKernel : public UPCGComputeKernel
 
 public:
 	virtual const UPCGSettings* GetSettings() const override;
-	virtual bool ComputeOutputBindingDataDesc(const UPCGComputeGraph* InGraph, FName InOutputPinLabel, const UPCGDataBinding* InBinding, FPCGDataCollectionDesc& OutDataDesc) const override;
+	virtual bool ComputeOutputBindingDataDesc(const UPCGComputeGraph* InGraph, FName InOutputPinLabel, UPCGDataBinding* InBinding, FPCGDataCollectionDesc& OutDataDesc) const override;
 	virtual int ComputeThreadCount(const UPCGDataBinding* Binding) const override;
 #if WITH_EDITOR
 	virtual FString GetCookedSource(const UPCGComputeGraph* InComputeGraph) const override;
