[PCG] Improve runtime scheduler priority-based scheme.

diff --git a/Engine/Plugins/PCG/Source/PCG/Private/RuntimeGen/PCGRuntimeGenScheduler.cpp b/Engine/Plugins/PCG/Source/PCG/Private/RuntimeGen/PCGRuntimeGenScheduler.cpp
--- a/Engine/Plugins/PCG/Source/PCG/Private/RuntimeGen/PCGRuntimeGenScheduler.cpp
+++ b/Engine/Plugins/PCG/Source/PCG/Private/RuntimeGen/PCGRuntimeGenScheduler.cpp
@@ -36,10 +36,10 @@ namespace PCGRuntimeGenSchedulerHelpers
 		true,
 		TEXT("Enable the RuntimeGeneration system."));
 
-	static TAutoConsoleVariable<int32> CVarFramesBetweenGraphSchedules(
-		TEXT("pcg.RuntimeGeneration.FramesBetweenGraphSchedules"),
-		0,
-		TEXT("Defines the minimum number of frames/ticks between any two Generation schedules in the RuntimeGenScheduler."));
+	static TAutoConsoleVariable<int32> CVarNumGeneratingComponentsAtSameTime(
+		TEXT("pcg.RuntimeGeneration.NumGeneratingComponents"),
+		16,
+		TEXT("Defines the maximum number of runtime components that can generate at the same time."));
 
 	static TAutoConsoleVariable<bool> CVarRuntimeGenerationEnableDebugging(
 		TEXT("pcg.RuntimeGeneration.EnableDebugging"),
@@ -161,20 +161,12 @@ void FPCGRuntimeGenScheduler::Tick(APCGWorldActor* InPCGWorldActor)
 	}
 
 	// 3. Schedule generation on components in priority order.
-
-	if (ScheduleFrameCounter <= 0)
+	if (!ComponentsToGenerate.IsEmpty())
 	{
-		if (!ComponentsToGenerate.IsEmpty())
-		{
-			// Sort components by priority (will be generated in descending order).
-			ComponentsToGenerate.ValueSort([](double PrioA, double PrioB)->bool { return PrioA > PrioB; });
+		// Sort components by priority (will be generated in descending order).
+		ComponentsToGenerate.ValueSort([](double PrioA, double PrioB)->bool { return PrioA > PrioB; });
 
-			TickScheduleGeneration(ComponentsToGenerate);
-		}
-	}
-	else
-	{
-		--ScheduleFrameCounter;
+		TickScheduleGeneration(ComponentsToGenerate);
 	}
 }
 
@@ -597,8 +589,40 @@ void FPCGRuntimeGenScheduler::TickScheduleGeneration(TMap<FGridGenerationKey, do
 
 	check(Subsystem && ActorAndComponentMapping);
 
+	// Count number of currently generating components
+	int NumGenerating = 0;
+
+	{
+		TRACE_CPUPROFILER_EVENT_SCOPE(FPCGRuntimeGenScheduler::TickScheduleGeneration::CountingCurrentlyGenerating);
+		for (const FGridGenerationKey& Key : GeneratedComponents)
+		{
+			if (!Key.IsValid())
+			{
+				continue;
+			}
+
+			UPCGComponent* Component = Key.GetOriginalComponent();
+			if (Key.GetGridDescriptor().GetHiGenGrid() != EPCGHiGenGrid::Unbounded)
+			{
+				Component = ActorAndComponentMapping->GetLocalComponent(Key.GetGridDescriptor(), Key.GetGridCoords(), Component);
+			}
+
+			if (Component && Component->IsGenerating())
+			{
+				++NumGenerating;
+			}
+		}
+	}
+
+	const int MaxNumGenerating = FMath::Max(1, PCGRuntimeGenSchedulerHelpers::CVarNumGeneratingComponentsAtSameTime.GetValueOnAnyThread());
+
 	for (const TPair<FGridGenerationKey, double>& GenerateEntry : ComponentsToGenerate)
 	{
+		if (NumGenerating >= MaxNumGenerating)
+		{
+			break;
+		}
+
 		const FGridGenerationKey& Key = GenerateEntry.Key;
 		const double Priority = GenerateEntry.Value;
 
@@ -720,12 +744,7 @@ void FPCGRuntimeGenScheduler::TickScheduleGeneration(TMap<FGridGenerationKey, do
 		}
 
 		GeneratedComponents.Add(Key);
-
-		ScheduleFrameCounter += PCGRuntimeGenSchedulerHelpers::CVarFramesBetweenGraphSchedules.GetValueOnAnyThread();
-		if (ScheduleFrameCounter > 0)
-		{
-			break;
-		}
+		++NumGenerating;
 	}
 }
 
diff --git a/Engine/Plugins/PCG/Source/PCG/Public/RuntimeGen/PCGRuntimeGenScheduler.h b/Engine/Plugins/PCG/Source/PCG/Public/RuntimeGen/PCGRuntimeGenScheduler.h
--- a/Engine/Plugins/PCG/Source/PCG/Public/RuntimeGen/PCGRuntimeGenScheduler.h
+++ b/Engine/Plugins/PCG/Source/PCG/Public/RuntimeGen/PCGRuntimeGenScheduler.h
@@ -129,9 +129,6 @@ private:
 	/** PartitionActorPoolSize represents the current maximum capacity of the PartitionActorPool. */
 	int32 PartitionActorPoolSize = 0;
 
-	/** Used to track the number of frames until we can schedule another generation. */
-	uint32 ScheduleFrameCounter = 0;
-
 	FPCGActorAndComponentMapping* ActorAndComponentMapping = nullptr;
 	FPCGGenSourceManager* GenSourceManager = nullptr;
 	UPCGSubsystem* Subsystem = nullptr;
