[PCG] Invisible pins should not be inspectable

#rnx

diff --git a/Engine/Plugins/PCG/Source/PCGEditor/Private/Widgets/SPCGEditorGraphAttributeListView.cpp b/Engine/Plugins/PCG/Source/PCGEditor/Private/Widgets/SPCGEditorGraphAttributeListView.cpp
--- a/Engine/Plugins/PCG/Source/PCGEditor/Private/Widgets/SPCGEditorGraphAttributeListView.cpp
+++ b/Engine/Plugins/PCG/Source/PCGEditor/Private/Widgets/SPCGEditorGraphAttributeListView.cpp
@@ -865,7 +865,7 @@ void SPCGEditorGraphAttributeListView::RefreshPinComboBox()
 			const UPCGPin* PCGPin = InPins[PinIndex];
 			const bool bIsOutputPin = PCGPin->IsOutputPin();
 			// Pin is included in list if it is connected, or if it is an output pin.
-			if (PCGPin->IsConnected() || bIsOutputPin)
+			if ((PCGPin->IsConnected() || bIsOutputPin) && !PCGPin->Properties.bInvisiblePin && !PCGPin->Properties.IsDatalessPin())
 			{
 				FString ItemName = FString::Format(*InFormatText, { PCGPin->Properties.Label.ToString() });
 				InOutItems.Add(MakeShared<FPinComboBoxItem>(FName(ItemName), PinIndex, bIsOutputPin));
