[PCG] Alternate short-cut for "debug only selected" for Mac users

#rnx

diff --git a/Engine/Plugins/PCG/Source/PCGEditor/Private/PCGEditorCommands.cpp b/Engine/Plugins/PCG/Source/PCGEditor/Private/PCGEditorCommands.cpp
--- a/Engine/Plugins/PCG/Source/PCGEditor/Private/PCGEditorCommands.cpp
+++ b/Engine/Plugins/PCG/Source/PCGEditor/Private/PCGEditorCommands.cpp
@@ -36,13 +36,19 @@ void FPCGEditorCommands::RegisterCommands()
 	UI_COMMAND(CancelExecution, "Cancel Execution", "Cancels the execution of the current graph", EUserInterfaceActionType::Button, FInputChord(EModifierKey::Control, EKeys::Escape));
 	UI_COMMAND(ToggleEnabled, "Toggle Enabled", "Toggle node enabled state for selected nodes.", EUserInterfaceActionType::ToggleButton, FInputChord(EKeys::E));
 	UI_COMMAND(ToggleDebug, "Toggle Debug", "Toggle node debug state for selected nodes", EUserInterfaceActionType::ToggleButton, FInputChord(EKeys::D));
-	UI_COMMAND(DebugOnlySelected, "Debug Only Selected", "Enable node debug state for selected nodes and disable debug state for the others", EUserInterfaceActionType::Button, FInputChord(EModifierKey::Control | EModifierKey::Alt, EKeys::D));
 	UI_COMMAND(DisableDebugOnAllNodes, "Disable Debug on all nodes", "Disable debug state for all nodes", EUserInterfaceActionType::Button, FInputChord(EModifierKey::Alt, EKeys::D));
 	UI_COMMAND(ToggleInspect, "Toggle Inspection", "Toggle node inspection for selected node", EUserInterfaceActionType::ToggleButton, FInputChord(EKeys::A));
 	UI_COMMAND(AddSourcePin, "Add Source Pin", "Add new source pin to the current node", EUserInterfaceActionType::Button, FInputChord());
 	UI_COMMAND(RenameNode, "Rename Node", "Rename the selected node", EUserInterfaceActionType::Button, FInputChord(EKeys::F2));
 	UI_COMMAND(SelectNamedRerouteUsages, "Select Named Reroute Usages", "Selects all usages of this Named Reroute Declaration", EUserInterfaceActionType::Button, FInputChord());
 	UI_COMMAND(SelectNamedRerouteDeclaration, "Select Named Reroute Declaration", "Selects the associated Named Reroute Declaration matching this Named Reroute Usage", EUserInterfaceActionType::Button, FInputChord());
+
+	//~ PLATFORM SPECIFIC ~//
+#if PLATFORM_MAC
+	UI_COMMAND(DebugOnlySelected, "Debug Only Selected", "Enable node debug state for selected nodes and disable debug state for the others", EUserInterfaceActionType::Button, FInputChord(EModifierKey::Control | EModifierKey::Alt, EKeys::D));
+#else
+	UI_COMMAND(DebugOnlySelected, "Debug Only Selected", "Enable node debug state for selected nodes and disable debug state for the others", EUserInterfaceActionType::Button, FInputChord(EModifierKey::Command | EModifierKey::Shift, EKeys::D));
+#endif
 }
 
 TSharedPtr<FEdGraphSchemaAction> FPCGSpawnNodeCommandInfo::GetAction() const
