[PCG] Add restrictions on MultiEntries and Parenting for Metadata domains

#rnx

diff --git a/Engine/Plugins/PCG/Source/PCG/Private/Data/PCGBasePointData.cpp b/Engine/Plugins/PCG/Source/PCG/Private/Data/PCGBasePointData.cpp
--- a/Engine/Plugins/PCG/Source/PCG/Private/Data/PCGBasePointData.cpp
+++ b/Engine/Plugins/PCG/Source/PCG/Private/Data/PCGBasePointData.cpp
@@ -443,16 +443,32 @@ FPCGAttributeAccessorMethods UPCGBasePointData::GetPointAccessorMethods()
 		return PCGBasePointData::CreateStaticAccessor(CastChecked<UPCGBasePointData>(InData), InSelector, bQuiet, true);
 	};
 
-	auto CreateAccessorKeysFunc = [](UPCGData* InData, const FPCGAttributePropertySelector& InSelector, bool bQuiet)
+	auto CreateAccessorKeysFunc = [](UPCGData* InData, const FPCGAttributePropertySelector& InSelector, bool bQuiet) -> TUniquePtr<IPCGAttributeAccessorKeys>
 	{
-		UPCGBasePointData* PointData = CastChecked<UPCGBasePointData>(InData);
-		return MakeUnique<FPCGAttributeAccessorKeysPointIndices>(PointData);
+		const FPCGMetadataDomainID DomainID = InData->GetMetadataDomainIDFromSelector(InSelector);
+		if (DomainID.IsDefault() || DomainID == PCGMetadataDomainID::Elements)
+		{
+			UPCGBasePointData* PointData = CastChecked<UPCGBasePointData>(InData);
+			return MakeUnique<FPCGAttributeAccessorKeysPointIndices>(PointData);
+		}
+		else
+		{
+			return {};
+		}
 	};
 
 	auto CreateConstAccessorKeysFunc = [](const UPCGData* InData, const FPCGAttributePropertySelector& InSelector, bool bQuiet) -> TUniquePtr<const IPCGAttributeAccessorKeys>
 	{
-		const UPCGBasePointData* PointData = CastChecked<UPCGBasePointData>(InData);
-		return MakeUnique<const FPCGAttributeAccessorKeysPointIndices>(PointData);
+		const FPCGMetadataDomainID DomainID = InData->GetMetadataDomainIDFromSelector(InSelector);
+		if (DomainID.IsDefault() || DomainID == PCGMetadataDomainID::Elements)
+		{
+			const UPCGBasePointData* PointData = CastChecked<UPCGBasePointData>(InData);
+			return MakeUnique<const FPCGAttributeAccessorKeysPointIndices>(PointData);
+		}
+		else
+		{
+			return {};
+		}
 	};
 
 	FPCGAttributeAccessorMethods Methods
diff --git a/Engine/Plugins/PCG/Source/PCG/Private/Data/PCGPointData.cpp b/Engine/Plugins/PCG/Source/PCG/Private/Data/PCGPointData.cpp
--- a/Engine/Plugins/PCG/Source/PCG/Private/Data/PCGPointData.cpp
+++ b/Engine/Plugins/PCG/Source/PCG/Private/Data/PCGPointData.cpp
@@ -197,17 +197,33 @@ FPCGAttributeAccessorMethods UPCGPointData::GetPointAccessorMethods()
 		return CreateStaticAccessor(InSelector, bQuiet);
 	};
 
-	auto CreateAccessorKeysFunc = [](UPCGData* InData, const FPCGAttributePropertySelector& InSelector, bool bQuiet)
+	auto CreateAccessorKeysFunc = [](UPCGData* InData, const FPCGAttributePropertySelector& InSelector, bool bQuiet) -> TUniquePtr<IPCGAttributeAccessorKeys>
 	{
-		UPCGPointData* PointData = CastChecked<UPCGPointData>(InData);
-		TArrayView<FPCGPoint> View(PointData->GetMutablePoints());
-		return MakeUnique<FPCGAttributeAccessorKeysPoints>(View);
+		const FPCGMetadataDomainID DomainID = InData->GetMetadataDomainIDFromSelector(InSelector);
+		if (DomainID.IsDefault() || DomainID == PCGMetadataDomainID::Elements)
+		{
+			UPCGPointData* PointData = CastChecked<UPCGPointData>(InData);
+			TArrayView<FPCGPoint> View(PointData->GetMutablePoints());
+			return MakeUnique<FPCGAttributeAccessorKeysPoints>(View);
+		}
+		else
+		{
+			return {};
+		}
 	};
 	
 	auto CreateConstAccessorKeysFunc = [](const UPCGData* InData, const FPCGAttributePropertySelector& InSelector, bool bQuiet) -> TUniquePtr<const IPCGAttributeAccessorKeys>
 	{
-		const UPCGPointData* PointData = CastChecked<UPCGPointData>(InData);
-		return MakeUnique<const FPCGAttributeAccessorKeysPoints>(PointData->GetPoints());
+		const FPCGMetadataDomainID DomainID = InData->GetMetadataDomainIDFromSelector(InSelector);
+		if (DomainID.IsDefault() || DomainID == PCGMetadataDomainID::Elements)
+		{
+			const UPCGPointData* PointData = CastChecked<UPCGPointData>(InData);
+			return MakeUnique<const FPCGAttributeAccessorKeysPoints>(PointData->GetPoints());
+		}
+		else
+		{
+			return {};
+		}
 	};
 	
 	FPCGAttributeAccessorMethods Methods
diff --git a/Engine/Plugins/PCG/Source/PCG/Private/Data/PCGSpatialData.cpp b/Engine/Plugins/PCG/Source/PCG/Private/Data/PCGSpatialData.cpp
--- a/Engine/Plugins/PCG/Source/PCG/Private/Data/PCGSpatialData.cpp
+++ b/Engine/Plugins/PCG/Source/PCG/Private/Data/PCGSpatialData.cpp
@@ -333,7 +333,7 @@ void UPCGSpatialData::InitializeOtherMetadataInternal(const FPCGInitializeFromDa
 	check(InParams.bInheritMetadata);
 	check(MetadataToInitialize);
 
-	MetadataToInitialize->Initialize(Metadata, InParams.bInheritAttributes);
+	MetadataToInitialize->Initialize(Metadata, InParams.bInheritAttributes, InParams.DomainMapping);
 }
 
 void UPCGSpatialData::InitializeSpatialDataInternal(const FPCGInitializeFromDataParams& InParams)
diff --git a/Engine/Plugins/PCG/Source/PCG/Private/Elements/Metadata/PCGMetadataOpElementBase.cpp b/Engine/Plugins/PCG/Source/PCG/Private/Elements/Metadata/PCGMetadataOpElementBase.cpp
--- a/Engine/Plugins/PCG/Source/PCG/Private/Elements/Metadata/PCGMetadataOpElementBase.cpp
+++ b/Engine/Plugins/PCG/Source/PCG/Private/Elements/Metadata/PCGMetadataOpElementBase.cpp
@@ -900,6 +900,14 @@ bool FPCGMetadataElementBase::PrepareDataInternal(FPCGContext* Context) const
 				return false;
 			}
 
+			if (OperationData.OutputAccessors[OutputIndex]->IsAttribute()
+				&& OperationData.NumberOfElementsToProcess > 1
+				&& !OutputData->ConstMetadata()->MetadataDomainSupportsMultiEntries(OutputData->GetMetadataDomainIDFromSelector(OutputTarget)))
+			{
+				PCGE_LOG(Error, GraphAndLog, FText::Format(LOCTEXT("OutputAccessorIsNotSupportingMultiEntries", "Output attribute '{0}' is on a domain that doesn't support multi entries, but we try to process multiple elements ({1}). It's invalid."), OutputTarget.GetDisplayText(), OperationData.NumberOfElementsToProcess));
+				return false;
+			}
+
 			OperationData.OutputKeys[OutputIndex] = PCGAttributeAccessorHelpers::CreateKeys(OutputData, OutputTarget);
 
 			return OperationData.OutputKeys[OutputIndex].IsValid();
diff --git a/Engine/Plugins/PCG/Source/PCG/Private/Metadata/PCGMetadata.cpp b/Engine/Plugins/PCG/Source/PCG/Private/Metadata/PCGMetadata.cpp
--- a/Engine/Plugins/PCG/Source/PCG/Private/Metadata/PCGMetadata.cpp
+++ b/Engine/Plugins/PCG/Source/PCG/Private/Metadata/PCGMetadata.cpp
@@ -1038,6 +1038,30 @@ void UPCGMetadata::SetLastCachedSelectorOnOwner(FName AttributeName, FPCGMetadat
 	}
 }
 
+bool UPCGMetadata::MetadataDomainSupportsMultiEntries(const FPCGMetadataDomainID& InDomainID) const
+{
+	const UPCGData* Data = Cast<UPCGData>(GetOuter());
+	if (!Data)
+	{
+		Data = GetDefault<UPCGData>();
+	}
+
+	check(Data);
+	return Data->MetadataDomainSupportsMultiEntries(InDomainID);
+}
+
+bool UPCGMetadata::MetadataDomainSupportsParenting(const FPCGMetadataDomainID& InDomainID) const
+{
+	const UPCGData* Data = Cast<UPCGData>(GetOuter());
+	if (!Data)
+	{
+		Data = GetDefault<UPCGData>();
+	}
+
+	check(Data);
+	return Data->MetadataDomainSupportsParenting(InDomainID);
+}
+
 void UPCGMetadata::SetupDomainsFromOtherMetadataIfNeeded(const UPCGMetadata* OtherMetadata)
 {
 	// Only do this if our outer is not a PCG data and we are not setuped.
diff --git a/Engine/Plugins/PCG/Source/PCG/Private/Metadata/PCGMetadataDomain.cpp b/Engine/Plugins/PCG/Source/PCG/Private/Metadata/PCGMetadataDomain.cpp
--- a/Engine/Plugins/PCG/Source/PCG/Private/Metadata/PCGMetadataDomain.cpp
+++ b/Engine/Plugins/PCG/Source/PCG/Private/Metadata/PCGMetadataDomain.cpp
@@ -283,6 +283,7 @@ void FPCGMetadataDomain::AddToCrc(FArchiveCrc32& Ar, const UPCGData* Data, bool
 FPCGMetadataDomain::FPCGMetadataDomain(UPCGMetadata* InTopMetadata, FPCGMetadataDomainID InMetadataDomainID)
 	: TopMetadata(InTopMetadata)
 	, DomainID(InMetadataDomainID)
+	, bSupportMultiEntries(InTopMetadata && InTopMetadata->MetadataDomainSupportsMultiEntries(DomainID))
 {
 	check(InTopMetadata);
 }
@@ -314,6 +315,13 @@ void FPCGMetadataDomain::InitializeWithAttributeFilter(const FPCGMetadataDomain*
 		return;
 	}
 
+	// If we don't have a top metadata (ill-formed), or we don't support parenting, force the copy
+	if (!ensure(TopMetadata) || !TopMetadata->MetadataDomainSupportsParenting(DomainID))
+	{
+		InitializeAsCopyWithAttributeFilter(InParent, InFilteredAttributes, InFilterMode, nullptr, InMatchOperator);
+		return;
+	}
+
 	Parent = ((InParent != this) ? InParent : nullptr);
 	ItemKeyOffset = Parent ? Parent->GetItemCountForChild() : 0;
 
@@ -921,49 +929,148 @@ int64 FPCGMetadataDomain::GetLocalItemCount() const
 int64 FPCGMetadataDomain::AddEntry(int64 ParentEntry)
 {
 	FWriteScopeLock ScopeLock(ItemLock);
-	return ParentKeys.Add(ParentEntry) + ItemKeyOffset;
+	if (bSupportMultiEntries)
+	{
+		return ParentKeys.Add(ParentEntry) + ItemKeyOffset;
+	}
+	else
+	{
+		if (ParentKeys.IsEmpty() && ItemKeyOffset == 0)
+		{
+			ParentKeys.Add(ParentEntry);
+		}
+		else
+		{
+			UE_LOG(LogPCG, Warning, TEXT("Try to add an entry to a domain (%ls) that doesn't support multi entries. Will always return 0."), *DomainID.DebugName.ToString());
+		}
+		
+		return 0;
+	}
 }
 
 TArray<int64> FPCGMetadataDomain::AddEntries(TArrayView<const int64> ParentEntryKeys)
 {
+	if (ParentEntryKeys.IsEmpty())
+	{
+		return {};
+	}
+	
 	TArray<int64> Result;
 	Result.Reserve(ParentEntryKeys.Num());
 
 	FWriteScopeLock ScopeLock(ItemLock);
-	ParentKeys.Reserve(ParentKeys.Num() + ParentEntryKeys.Num());
-	for (const int64 ParentEntry : ParentEntryKeys)
+	if (bSupportMultiEntries)
 	{
-		Result.Add(ParentKeys.Add(ParentEntry) + ItemKeyOffset);
+		ParentKeys.Reserve(ParentKeys.Num() + ParentEntryKeys.Num());
+		for (const int64 ParentEntry : ParentEntryKeys)
+		{
+			Result.Add(ParentKeys.Add(ParentEntry) + ItemKeyOffset);
+		}
 	}
+	else
+	{
+		if (ParentEntryKeys.Num() > 1 || !ParentKeys.IsEmpty() || ItemKeyOffset != 0)
+		{
+			UE_LOG(LogPCG, Warning, TEXT("Try to add multiple entries to a metadata domain that don't support it (%ls). Will always return 0."), *DomainID.DebugName.ToString());
+		}
+		
+		if (ParentKeys.IsEmpty() && ItemKeyOffset == 0)
+		{
+			ParentKeys.Add(ParentEntryKeys[0]);
+		}
 
+		// The function expect to return the same number of keys, so return an array of 0;
+		Result.SetNumZeroed(ParentEntryKeys.Num());
+	}
+	
 	return Result;
 }
 
 void FPCGMetadataDomain::AddEntriesInPlace(TArrayView<int64*> ParentEntryKeys)
 {
+	if (ParentEntryKeys.IsEmpty())
+	{
+		return;
+	}
+	
 	FWriteScopeLock ScopeLock(ItemLock);
-	ParentKeys.Reserve(ParentKeys.Num() + ParentEntryKeys.Num());
-	for (int64* ParentEntry : ParentEntryKeys)
+	if (bSupportMultiEntries)
 	{
-		*ParentEntry = ParentKeys.Add(*ParentEntry) + ItemKeyOffset;
+		ParentKeys.Reserve(ParentKeys.Num() + ParentEntryKeys.Num());
+		for (int64* ParentEntry : ParentEntryKeys)
+		{
+			*ParentEntry = ParentKeys.Add(*ParentEntry) + ItemKeyOffset;
+		}
+	}
+	else
+	{
+		if (ParentEntryKeys.Num() > 1 || !ParentKeys.IsEmpty() || ItemKeyOffset != 0)
+		{
+			UE_LOG(LogPCG, Warning, TEXT("Try to add multiple entries to a metadata domain that don't support it (%ls). Will always return 0."), *DomainID.DebugName.ToString());
+		}
+		
+		if (ParentKeys.IsEmpty() && ItemKeyOffset == 0)
+		{
+			ParentKeys.Add(*ParentEntryKeys[0]);
+		}
+
+		// The function expect to return the same number of keys, so return an array of 0
+		for (int64* ParentEntry : ParentEntryKeys)
+		{
+			*ParentEntry = 0;
+		}
 	}
 }
 
 int64 FPCGMetadataDomain::AddEntryPlaceholder()
 {
 	FReadScopeLock ScopeLock(ItemLock);
-	return ParentKeys.Num() + DelayedEntriesIndex.IncrementExchange() + ItemKeyOffset;
+	if (bSupportMultiEntries)
+	{
+		return ParentKeys.Num() + DelayedEntriesIndex.IncrementExchange() + ItemKeyOffset;
+	}
+	else
+	{
+		const int64 DelayedEntry = DelayedEntriesIndex.IncrementExchange();
+		if (DelayedEntry != 0 || !ParentKeys.IsEmpty() || ItemKeyOffset != 0)
+		{
+			UE_LOG(LogPCG, Warning, TEXT("Try to add an entry to a domain (%ls) that doesn't support multi entries. Will always return 0."), *DomainID.DebugName.ToString());
+		}
+		
+		return 0;
+	}
 }
 
 void FPCGMetadataDomain::AddDelayedEntries(const TArray<TTuple<int64, int64>>& AllEntries)
 {
+	if (AllEntries.IsEmpty())
+	{
+		return;
+	}
+	
 	FWriteScopeLock ScopeLock(ItemLock);
-	ParentKeys.AddUninitialized(AllEntries.Num());
-	for (const TTuple<int64, int64>& Entry : AllEntries)
+		
+	if (bSupportMultiEntries)
 	{
-		int64 Index = Entry.Get<0>() - ItemKeyOffset;
-		check(Index < ParentKeys.Num());
-		ParentKeys[Index] = Entry.Get<1>();
+		ParentKeys.AddUninitialized(AllEntries.Num());
+		for (const TTuple<int64, int64>& Entry : AllEntries)
+		{
+			int64 Index = Entry.Get<0>() - ItemKeyOffset;
+			check(Index < ParentKeys.Num());
+			ParentKeys[Index] = Entry.Get<1>();
+		}
+	}
+	else
+	{
+		if (AllEntries.Num() > 1 || !ParentKeys.IsEmpty() || ItemKeyOffset != 0)
+		{
+			UE_LOG(LogPCG, Warning, TEXT("Try to add multiple entries to a metadata domain that don't support it (%ls). Will always return 0."), *DomainID.DebugName.ToString());
+		}
+
+		if (ParentKeys.IsEmpty() && ItemKeyOffset == 0)
+		{
+			ParentKeys.Add(AllEntries[0].Get<1>());
+		}
 	}
 
 	DelayedEntriesIndex.Exchange(0);
diff --git a/Engine/Plugins/PCG/Source/PCG/Private/PCGData.cpp b/Engine/Plugins/PCG/Source/PCG/Private/PCGData.cpp
--- a/Engine/Plugins/PCG/Source/PCG/Private/PCGData.cpp
+++ b/Engine/Plugins/PCG/Source/PCG/Private/PCGData.cpp
@@ -90,7 +90,7 @@ bool UPCGData::IsSupportedMetadataDomainID(const FPCGMetadataDomainID& InDomainI
 FPCGMetadataDomainID UPCGData::GetMetadataDomainIDFromSelector(const FPCGAttributePropertySelector& InSelector) const
 {
 	const FName DomainName = InSelector.GetDomainName();
-	if (DomainName == NAME_None)
+	if (DomainName == NAME_None || DomainName == PCGDataConstants::DefaultDomainName)
 	{
 		return PCGMetadataDomainID::Default;
 	}
@@ -122,6 +122,18 @@ bool UPCGData::SetDomainFromDomainID(const FPCGMetadataDomainID& InDomainID, FPC
 	}
 }
 
+bool UPCGData::MetadataDomainSupportsMultiEntries(const FPCGMetadataDomainID& InDomainID) const
+{
+	// By default, global does not support multi entries
+	return InDomainID.IsValid() && InDomainID != PCGMetadataDomainID::Data;
+}
+
+bool UPCGData::MetadataDomainSupportsParenting(const FPCGMetadataDomainID& InDomainID) const
+{
+	// By default parenting is supported
+	return true;
+}
+
 FPCGCrc UPCGData::ComputeCrc(bool bFullDataCrc) const
 {
 	FArchiveCrc32 Ar;
diff --git a/Engine/Plugins/PCG/Source/PCG/Private/PCGParamData.cpp b/Engine/Plugins/PCG/Source/PCG/Private/PCGParamData.cpp
--- a/Engine/Plugins/PCG/Source/PCG/Private/PCGParamData.cpp
+++ b/Engine/Plugins/PCG/Source/PCG/Private/PCGParamData.cpp
@@ -170,6 +170,12 @@ FPCGMetadataDomainID UPCGParamData::GetMetadataDomainIDFromSelector(const FPCGAt
 	}
 }
 
+bool UPCGParamData::MetadataDomainSupportsParenting(const FPCGMetadataDomainID& InDomainID) const
+{
+	// Param data doesn't support parenting.
+	return false;
+}
+
 bool UPCGParamData::SetDomainFromDomainID(const FPCGMetadataDomainID& InDomainID, FPCGAttributePropertySelector& InOutSelector) const
 {
 	if (InDomainID == PCGMetadataDomainID::Elements)
@@ -181,4 +187,4 @@ bool UPCGParamData::SetDomainFromDomainID(const FPCGMetadataDomainID& InDomainID
 	{
 		return Super::SetDomainFromDomainID(InDomainID, InOutSelector);
 	}
-}
\ No newline at end of file
+}
diff --git a/Engine/Plugins/PCG/Source/PCG/Private/Tests/Metadata/PCGMetadataTest.cpp b/Engine/Plugins/PCG/Source/PCG/Private/Tests/Metadata/PCGMetadataTest.cpp
new file mode 100644
--- /dev/null
+++ b/Engine/Plugins/PCG/Source/PCG/Private/Tests/Metadata/PCGMetadataTest.cpp
@@ -0,0 +1,322 @@
+// Copyright Epic Games, Inc. All Rights Reserved.
+
+#include "PCGParamData.h"
+#if WITH_EDITOR
+
+#include "Tests/PCGTestsCommon.h"
+
+#include "Data/PCGPointData.h"
+#include "Metadata/PCGMetadata.h"
+#include "Metadata/PCGMetadataAttributeTpl.h"
+
+#include "Engine/Engine.h"
+
+IMPLEMENT_CUSTOM_SIMPLE_AUTOMATION_TEST(FPCGMetadataTest_Inherit, FPCGTestBaseClass, "Plugins.PCG.Metadata.Inherit", PCGTestsCommon::TestFlags)
+IMPLEMENT_CUSTOM_SIMPLE_AUTOMATION_TEST(FPCGMetadataTest_InheritCopy, FPCGTestBaseClass, "Plugins.PCG.Metadata.InheritCopy", PCGTestsCommon::TestFlags)
+IMPLEMENT_CUSTOM_SIMPLE_AUTOMATION_TEST(FPCGMetadataTest_InheritWithNoParenting, FPCGTestBaseClass, "Plugins.PCG.Metadata.InheritWithNoParenting", PCGTestsCommon::TestFlags)
+IMPLEMENT_CUSTOM_SIMPLE_AUTOMATION_TEST(FPCGMetadataTest_TwoAddEntryToGlobal, FPCGTestBaseClass, "Plugins.PCG.Metadata.TwoAddEntryToGlobal", PCGTestsCommon::TestFlags)
+IMPLEMENT_CUSTOM_SIMPLE_AUTOMATION_TEST(FPCGMetadataTest_AddEntriesToGlobal, FPCGTestBaseClass, "Plugins.PCG.Metadata.AddEntriesToGlobal", PCGTestsCommon::TestFlags)
+IMPLEMENT_CUSTOM_SIMPLE_AUTOMATION_TEST(FPCGMetadataTest_AddEntriesInPlaceToGlobal, FPCGTestBaseClass, "Plugins.PCG.Metadata.AddEntriesToGlobal", PCGTestsCommon::TestFlags)
+IMPLEMENT_CUSTOM_SIMPLE_AUTOMATION_TEST(FPCGMetadataTest_DelayedAddEntriesToGlobal, FPCGTestBaseClass, "Plugins.PCG.Metadata.DelayedAddEntriesToGlobal", PCGTestsCommon::TestFlags)
+
+bool FPCGMetadataTest_Inherit::RunTest(const FString& Parameters)
+{
+	const FName AttributeName = TEXT("MyAttr");
+	
+	UPCGPointData* PointData3 = nullptr;
+
+	{
+		TObjectPtr<UPCGPointData> PointData = NewObject<UPCGPointData>();
+		TArray<FPCGPoint>& Points = PointData->GetMutablePoints();
+		Points.SetNum(5);
+		UPCGMetadata* PointMetadata = PointData->MutableMetadata();
+		FPCGMetadataAttribute<int32>* Attribute = PointMetadata->CreateAttribute<int32>(AttributeName, 0, true, false);
+	
+		for (int32 i = 0 ; i < Points.Num(); ++i)
+		{
+			PointMetadata->InitializeOnSet(Points[i].MetadataEntry);
+			Attribute->SetValue(Points[i].MetadataEntry, i);
+		}
+	
+		TObjectPtr<UPCGPointData> PointData2 = NewObject<UPCGPointData>();
+		UPCGMetadata* PointMetadata2 = PointData2->MutableMetadata();
+		PointMetadata2->Initialize(PointMetadata);
+
+		// Verify that the initialization inherited
+		UTEST_EQUAL("Initialization inherited", PointMetadata2->GetLocalItemCount(), 0);
+
+		FPCGMetadataAttribute<int32>* Attribute2 = PointMetadata2->GetMutableTypedAttribute<int32>(AttributeName);
+
+		TArray<FPCGPoint>& Points2 = PointData2->GetMutablePoints();
+		Points2 = Points;
+		Points2.SetNum(10);
+	
+		for (int32 i = Points.Num() ; i < Points2.Num(); ++i)
+		{
+			PointMetadata2->InitializeOnSet(Points2[i].MetadataEntry);
+			Attribute2->SetValue(Points2[i].MetadataEntry, i);
+		}
+		
+		PointData3 = NewObject<UPCGPointData>();
+		UPCGMetadata* PointMetadata3 = PointData3->MutableMetadata();
+		PointMetadata3->Initialize(PointMetadata2);
+
+		UTEST_EQUAL("Initialization inherited", PointMetadata3->GetLocalItemCount(), 0);
+		
+		PointData3->GetMutablePoints() = Points2;
+	}
+
+	// Force a GC run
+	if (ensure(GEngine))
+	{
+		GEngine->ForceGarbageCollection(/*bFullPurge=*/true);
+	}
+
+	// Verify that all the values are readable from the metadata
+	check(PointData3);
+	const TArray<FPCGPoint>& Points3 = PointData3->GetPoints();
+	const FPCGMetadataAttribute<int32>* Attribute = PointData3->ConstMetadata()->GetConstTypedAttribute<int32>(AttributeName);
+	UTEST_NOT_NULL("Attribute exists", Attribute);
+	check(Attribute);
+	
+	for (int32 i = 0; i < Points3.Num(); ++i)
+	{
+		UTEST_EQUAL("Values are the same", Attribute->GetValueFromItemKey(Points3[i].MetadataEntry), i);
+	}
+	
+	return true;
+}
+
+bool FPCGMetadataTest_InheritCopy::RunTest(const FString& Parameters)
+{
+	const FName AttributeName = TEXT("MyAttr");
+	
+	UPCGPointData* PointData3 = nullptr;
+
+	{
+		TObjectPtr<UPCGPointData> PointData = NewObject<UPCGPointData>();
+		TArray<FPCGPoint>& Points = PointData->GetMutablePoints();
+		Points.SetNum(5);
+		UPCGMetadata* PointMetadata = PointData->MutableMetadata();
+		FPCGMetadataAttribute<int32>* Attribute = PointMetadata->CreateAttribute<int32>(AttributeName, 0, true, false);
+	
+		for (int32 i = 0 ; i < Points.Num(); ++i)
+		{
+			PointMetadata->InitializeOnSet(Points[i].MetadataEntry);
+			Attribute->SetValue(Points[i].MetadataEntry, i);
+		}
+	
+		TObjectPtr<UPCGPointData> PointData2 = NewObject<UPCGPointData>();
+		UPCGMetadata* PointMetadata2 = PointData2->MutableMetadata();
+		PointMetadata2->Initialize(PointMetadata);
+
+		// Verify that the initialization inherited
+		UTEST_EQUAL("Initialization inherited", PointMetadata2->GetLocalItemCount(), 0);
+
+		FPCGMetadataAttribute<int32>* Attribute2 = PointMetadata2->GetMutableTypedAttribute<int32>(AttributeName);
+
+		TArray<FPCGPoint>& Points2 = PointData2->GetMutablePoints();
+		Points2 = Points;
+		Points2.SetNum(10);
+	
+		for (int32 i = Points.Num() ; i < Points2.Num(); ++i)
+		{
+			PointMetadata2->InitializeOnSet(Points2[i].MetadataEntry);
+			Attribute2->SetValue(Points2[i].MetadataEntry, i);
+		}
+		
+		PointData3 = NewObject<UPCGPointData>();
+		UPCGMetadata* PointMetadata3 = PointData3->MutableMetadata();
+		PointMetadata3->InitializeAsCopy(PointMetadata2);
+
+		// Verify that the initialization made a copy
+		UTEST_EQUAL("Initialization made a copy", PointMetadata3->GetLocalItemCount(), 5);
+		
+		PointData3->GetMutablePoints() = Points2;
+	}
+
+	// Force a GC run
+	if (ensure(GEngine))
+	{
+		GEngine->ForceGarbageCollection(/*bFullPurge=*/true);
+	}
+
+	// Verify that all the values are readable from the metadata
+	check(PointData3);
+	const TArray<FPCGPoint>& Points3 = PointData3->GetPoints();
+	const FPCGMetadataAttribute<int32>* Attribute = PointData3->ConstMetadata()->GetConstTypedAttribute<int32>(AttributeName);
+	UTEST_NOT_NULL("Attribute exists", Attribute);
+	check(Attribute);
+	
+	for (int32 i = 0; i < Points3.Num(); ++i)
+	{
+		UTEST_EQUAL("Values are the same", Attribute->GetValueFromItemKey(Points3[i].MetadataEntry), i);
+	}
+	
+	return true;
+}
+
+bool FPCGMetadataTest_InheritWithNoParenting::RunTest(const FString& Parameters)
+{
+	const FName AttributeName = TEXT("MyAttr");
+	
+	UPCGParamData* ParamData3 = nullptr;
+
+	{
+		TObjectPtr<UPCGParamData> ParamData = NewObject<UPCGParamData>();
+		UPCGMetadata* ParamMetadata = ParamData->MutableMetadata();
+		FPCGMetadataAttribute<int32>* Attribute = ParamMetadata->CreateAttribute<int32>(AttributeName, 0, true, false);
+	
+		for (int32 i = 0 ; i < 5; ++i)
+		{
+			PCGMetadataEntryKey EntryKey = PCGInvalidEntryKey;
+			ParamMetadata->InitializeOnSet(EntryKey);
+			Attribute->SetValue(EntryKey, i);
+		}
+	
+		TObjectPtr<UPCGParamData> ParamData2 = NewObject<UPCGParamData>();
+		UPCGMetadata* ParamMetadata2 = ParamData2->MutableMetadata();
+		ParamMetadata2->Initialize(ParamMetadata);
+
+		// Verify that the initialization made a copy
+		UTEST_EQUAL("Initialization made a copy", ParamMetadata2->GetLocalItemCount(), 5);
+
+		FPCGMetadataAttribute<int32>* Attribute2 = ParamMetadata2->GetMutableTypedAttribute<int32>(AttributeName);
+	
+		for (int32 i = 5 ; i < 10; ++i)
+		{
+			PCGMetadataEntryKey EntryKey = PCGInvalidEntryKey;
+			ParamMetadata2->InitializeOnSet(EntryKey);
+			Attribute2->SetValue(EntryKey, i);
+		}
+		
+		ParamData3 = NewObject<UPCGParamData>();
+		UPCGMetadata* ParamMetadata3 = ParamData3->MutableMetadata();
+		ParamMetadata3->Initialize(ParamMetadata2);
+
+		// Verify that the initialization made a copy
+		UTEST_EQUAL("Initialization made a copy", ParamMetadata3->GetLocalItemCount(), 10);
+	}
+
+	// Force a GC run
+	if (ensure(GEngine))
+	{
+		GEngine->ForceGarbageCollection(/*bFullPurge=*/true);
+	}
+
+	// Verify that all the values are readable from the metadata
+	check(ParamData3);
+	const FPCGMetadataAttribute<int32>* Attribute = ParamData3->ConstMetadata()->GetConstTypedAttribute<int32>(AttributeName);
+	UTEST_NOT_NULL("Attribute exists", Attribute);
+	check(Attribute);
+	
+	for (int32 i = 0; i < 10; ++i)
+	{
+		UTEST_EQUAL("Values are the same", Attribute->GetValueFromItemKey(i), i);
+	}
+	
+	return true;
+}
+
+bool FPCGMetadataTest_TwoAddEntryToGlobal::RunTest(const FString& Parameters)
+{
+	const FName AttributeName = TEXT("MyAttr");
+	
+	TObjectPtr<UPCGParamData> ParamData = NewObject<UPCGParamData>();
+	UPCGMetadata* ParamMetadata = ParamData->MutableMetadata();
+	FPCGMetadataDomain* ParamMetadataDomain = ParamMetadata->GetMetadataDomain(PCGMetadataDomainID::Data);
+	FPCGMetadataAttribute<int32>* Attribute = ParamMetadataDomain->CreateAttribute<int32>(AttributeName, 0, false, false);
+
+	AddExpectedError(TEXT("Try to add an entry to a domain (Data) that doesn't support multi entries. Will always return 0."), EAutomationExpectedErrorFlags::MatchType::Exact, 1, /*bIsRegex=*/false);
+
+	for (int32 i = 0 ; i < 2; ++i)
+	{
+		PCGMetadataEntryKey EntryKey = PCGInvalidEntryKey;
+		ParamMetadataDomain->InitializeOnSet(EntryKey);
+		Attribute->SetValue(EntryKey, i);
+	}
+
+	UTEST_EQUAL("Just have a single entry", ParamMetadataDomain->GetItemCountForChild(), 1);
+	UTEST_EQUAL("Entry has the latest value", Attribute->GetValueFromItemKey(PCGMetadataEntryKey(0)), 1);
+	
+	return true;
+}
+
+bool FPCGMetadataTest_AddEntriesToGlobal::RunTest(const FString& Parameters)
+{
+	const FName AttributeName = TEXT("MyAttr");
+	
+	TObjectPtr<UPCGParamData> ParamData = NewObject<UPCGParamData>();
+	UPCGMetadata* ParamMetadata = ParamData->MutableMetadata();
+	FPCGMetadataDomain* ParamMetadataDomain = ParamMetadata->GetMetadataDomain(PCGMetadataDomainID::Data);
+	FPCGMetadataAttribute<int32>* Attribute = ParamMetadataDomain->CreateAttribute<int32>(AttributeName, 0, false, false);
+
+	AddExpectedError(TEXT("Try to add multiple entries to a metadata domain that don't support it (Data). Will always return 0."), EAutomationExpectedErrorFlags::MatchType::Exact, 1, /*bIsRegex=*/false);
+
+	TStaticArray<PCGMetadataEntryKey, 2> ParentKeys = {PCGInvalidEntryKey, PCGInvalidEntryKey};
+	TArray<PCGMetadataEntryKey> NewKeys = ParamMetadataDomain->AddEntries(ParentKeys);
+	TStaticArray<int32, 2> Values = {5, 6};
+	Attribute->SetValues(NewKeys, Values);
+
+	UTEST_EQUAL("Just have a single entry", ParamMetadataDomain->GetItemCountForChild(), 1);
+	UTEST_EQUAL("Entry has the latest value", Attribute->GetValueFromItemKey(PCGMetadataEntryKey(0)), Values[1]);
+	
+	return true;
+}
+
+bool FPCGMetadataTest_AddEntriesInPlaceToGlobal::RunTest(const FString& Parameters)
+{
+	const FName AttributeName = TEXT("MyAttr");
+	
+	TObjectPtr<UPCGParamData> ParamData = NewObject<UPCGParamData>();
+	UPCGMetadata* ParamMetadata = ParamData->MutableMetadata();
+	FPCGMetadataDomain* ParamMetadataDomain = ParamMetadata->GetMetadataDomain(PCGMetadataDomainID::Data);
+	FPCGMetadataAttribute<int32>* Attribute = ParamMetadataDomain->CreateAttribute<int32>(AttributeName, 0, false, false);
+
+	AddExpectedError(TEXT("Try to add multiple entries to a metadata domain that don't support it (Data). Will always return 0."), EAutomationExpectedErrorFlags::MatchType::Exact, 1, /*bIsRegex=*/false);
+
+	TStaticArray<PCGMetadataEntryKey, 2> Keys = {PCGInvalidEntryKey, PCGInvalidEntryKey};
+	TStaticArray<PCGMetadataEntryKey*, 2> KeysPtr = {&Keys[0], &Keys[1]};
+	
+	ParamMetadataDomain->AddEntriesInPlace(KeysPtr);
+	TStaticArray<int32, 2> Values = {5, 6};
+	Attribute->SetValues(Keys, Values);
+
+	UTEST_EQUAL("Just have a single entry", ParamMetadataDomain->GetItemCountForChild(), 1);
+	UTEST_EQUAL("Entry has the latest value", Attribute->GetValueFromItemKey(PCGMetadataEntryKey(0)), Values[1]);
+	
+	return true;
+}
+
+bool FPCGMetadataTest_DelayedAddEntriesToGlobal::RunTest(const FString& Parameters)
+{
+	const FName AttributeName = TEXT("MyAttr");
+	
+	TObjectPtr<UPCGParamData> ParamData = NewObject<UPCGParamData>();
+	UPCGMetadata* ParamMetadata = ParamData->MutableMetadata();
+	FPCGMetadataDomain* ParamMetadataDomain = ParamMetadata->GetMetadataDomain(PCGMetadataDomainID::Data);
+	FPCGMetadataAttribute<int32>* Attribute = ParamMetadataDomain->CreateAttribute<int32>(AttributeName, 0, false, false);
+
+	AddExpectedError(TEXT("Try to add an entry to a domain (Data) that doesn't support multi entries. Will always return 0."), EAutomationExpectedErrorFlags::MatchType::Exact, 1, /*bIsRegex=*/false);
+
+	TArray<TTuple<PCGMetadataEntryKey, PCGMetadataEntryKey>> Mapping;
+	TStaticArray<PCGMetadataEntryKey, 2> Keys;
+
+	Keys[0] = ParamMetadataDomain->AddEntryPlaceholder();
+	Mapping.Emplace(PCGInvalidEntryKey, Keys[0]);
+	Keys[1] = ParamMetadataDomain->AddEntryPlaceholder();
+	Mapping.Emplace(PCGInvalidEntryKey, Keys[1]);
+
+	TStaticArray<int32, 2> Values = {5, 6};
+	Attribute->SetValues(Keys, Values);
+	
+	AddExpectedError(TEXT("Try to add multiple entries to a metadata domain that don't support it (Data). Will always return 0."), EAutomationExpectedErrorFlags::MatchType::Exact, 1, /*bIsRegex=*/false);
+	ParamMetadataDomain->AddDelayedEntries(Mapping);
+
+	UTEST_EQUAL("Just have a single entry", ParamMetadataDomain->GetItemCountForChild(), 1);
+	UTEST_EQUAL("Entry has the latest value", Attribute->GetValueFromItemKey(PCGMetadataEntryKey(0)), Values[1]);
+	
+	return true;
+}
+
+#endif // WITH_EDITOR
\ No newline at end of file
diff --git a/Engine/Plugins/PCG/Source/PCG/Public/Metadata/Accessors/PCGAttributeAccessor.h b/Engine/Plugins/PCG/Source/PCG/Public/Metadata/Accessors/PCGAttributeAccessor.h
--- a/Engine/Plugins/PCG/Source/PCG/Public/Metadata/Accessors/PCGAttributeAccessor.h
+++ b/Engine/Plugins/PCG/Source/PCG/Public/Metadata/Accessors/PCGAttributeAccessor.h
@@ -114,6 +114,13 @@ public:
 	{
 		//TRACE_CPUPROFILER_EVENT_SCOPE(FPCGAttributeAccessor::Prepare);
 
+		const bool bSupportsMultiEntries = MetadataDomain->SupportsMultiEntries();
+		if (!bSupportsMultiEntries && Count > 1)
+		{
+			UE_LOG(LogPCG, Warning, TEXT("Try to prepare multiple values to set in an attribute from a domain that doesn't support multi entries (%ls)"), *MetadataDomain->GetDomainID().DebugName.ToString());
+			return false;
+		}
+
 		check(!OutEntryKeys || OutEntryKeys->Num() == Count);
 
 		TArray<PCGMetadataEntryKey*, TInlineAllocator<512>> EntryKeys;
@@ -149,7 +156,14 @@ public:
 		if (!EntriesToSet.IsEmpty())
 		{
 			TRACE_CPUPROFILER_EVENT_SCOPE(FPCGAttributeAccessor::Prepare::AddEntriesInPlace);
-			MetadataDomain->AddEntriesInPlace(EntriesToSet);
+			if (bSupportsMultiEntries || (MetadataDomain->GetItemCountForChild() == 0 && ensure(EntriesToSet.Num() == 1)))
+			{
+				MetadataDomain->AddEntriesInPlace(EntriesToSet);
+			}
+			else if (!bSupportsMultiEntries && ensure(EntriesToSet.Num() == 1))
+			{
+				*(EntriesToSet[0]) = PCGMetadataEntryKey(0);
+			}
 		}
 
 		if (bPreallocateValues && ensure(!bWasPrepared))
diff --git a/Engine/Plugins/PCG/Source/PCG/Public/Metadata/PCGMetadata.h b/Engine/Plugins/PCG/Source/PCG/Public/Metadata/PCGMetadata.h
--- a/Engine/Plugins/PCG/Source/PCG/Public/Metadata/PCGMetadata.h
+++ b/Engine/Plugins/PCG/Source/PCG/Public/Metadata/PCGMetadata.h
@@ -406,6 +406,10 @@ public:
 	FPCGMetadataDomain* GetDefaultMetadataDomain() { return GetMetadataDomain(PCGMetadataDomainID::Default); }
 	const FPCGMetadataDomain* GetConstDefaultMetadataDomain() const { return GetConstMetadataDomain(PCGMetadataDomainID::Default); }
 
+	/** Mirror functions to be called on the outer data. If there is no outer, we will use the default object for a UPCGData. */
+	bool MetadataDomainSupportsMultiEntries(const FPCGMetadataDomainID& InDomainID) const;
+	bool MetadataDomainSupportsParenting(const FPCGMetadataDomainID& InDomainID) const;
+
 private:
 	// Utility functions to forward calls to FPCGMetadataDomain.
 	template <typename Func, typename ...Args>
diff --git a/Engine/Plugins/PCG/Source/PCG/Public/Metadata/PCGMetadataDomain.h b/Engine/Plugins/PCG/Source/PCG/Public/Metadata/PCGMetadataDomain.h
--- a/Engine/Plugins/PCG/Source/PCG/Public/Metadata/PCGMetadataDomain.h
+++ b/Engine/Plugins/PCG/Source/PCG/Public/Metadata/PCGMetadataDomain.h
@@ -231,6 +231,8 @@ public:
 
 	FPCGMetadataDomainID GetDomainID() const { return DomainID; }
 
+	bool SupportsMultiEntries() const { return bSupportMultiEntries; }
+
 protected:
 	FPCGMetadataAttributeBase* CopyAttribute(FName AttributeToCopy, FName NewAttributeName, bool bKeepParent, bool bCopyEntries, bool bCopyValues);
 
@@ -245,6 +247,9 @@ protected:
 	FPCGMetadataDomainID DomainID;
 	const FPCGMetadataDomain* Parent = nullptr;
 
+	// Cached value on construction to know if we support multi entries
+	const bool bSupportMultiEntries = true;
+
 	TMap<FName, FPCGMetadataAttributeBase*> Attributes;
 	PCGMetadataAttributeKey NextAttributeId = 0;
 
diff --git a/Engine/Plugins/PCG/Source/PCG/Public/PCGData.h b/Engine/Plugins/PCG/Source/PCG/Public/PCGData.h
--- a/Engine/Plugins/PCG/Source/PCG/Public/PCGData.h
+++ b/Engine/Plugins/PCG/Source/PCG/Public/PCGData.h
@@ -35,6 +35,7 @@ ENUM_CLASS_FLAGS(EPCGDataUsage);
 namespace PCGDataConstants
 {
 	const FName GlobalDomainName = "Global";
+	const FName DefaultDomainName = "Default";
 }
 
 /**
@@ -107,6 +108,8 @@ public:
 
 	virtual void Flatten();
 
+	/************** Metadata domain specific functions **************/
+
 	/** Returns the default domain for this data. Any data that has other default domain should override this method. */
 	virtual FPCGMetadataDomainID GetDefaultMetadataDomainID() const { return PCGMetadataDomainID::Data; }
 
@@ -122,6 +125,12 @@ public:
 	/** Write the domain name into the selector associated with this domain ID. Each data is responsible to override this for every supported domain. */
 	virtual bool SetDomainFromDomainID(const FPCGMetadataDomainID& InDomainID, FPCGAttributePropertySelector& InOutSelector) const;
 
+	/** Check if the domain supports multi entries. If not, an add entry won't work. */
+	virtual bool MetadataDomainSupportsMultiEntries(const FPCGMetadataDomainID& InDomainID) const;
+
+	/** Check if parenting is possible for this domain, if not, any initialization will result to a copy. */
+	virtual bool MetadataDomainSupportsParenting(const FPCGMetadataDomainID& InDomainID) const;
+
 	// Not accessible through blueprint to make sure the constness is preserved.
 	UPROPERTY()
 	TObjectPtr<UPCGMetadata> Metadata = nullptr;
diff --git a/Engine/Plugins/PCG/Source/PCG/Public/PCGParamData.h b/Engine/Plugins/PCG/Source/PCG/Public/PCGParamData.h
--- a/Engine/Plugins/PCG/Source/PCG/Public/PCGParamData.h
+++ b/Engine/Plugins/PCG/Source/PCG/Public/PCGParamData.h
@@ -32,12 +32,13 @@ public:
 	virtual void SetLastSelector(const FPCGAttributePropertySelector& InSelector) override;
 
 	virtual UPCGParamData* DuplicateData(FPCGContext* Context, bool bInitializeMetadata = true) const override;
-	// ~End UPCGData interface
 
 	virtual FPCGMetadataDomainID GetDefaultMetadataDomainID() const override { return PCGMetadataDomainID::Elements; }
 	virtual TArray<FPCGMetadataDomainID> GetAllSupportedMetadataDomainIDs() const override { return {PCGMetadataDomainID::Data, PCGMetadataDomainID::Elements}; }
 	virtual bool SetDomainFromDomainID(const FPCGMetadataDomainID& InDomainID, FPCGAttributePropertySelector& InOutSelector) const override;
 	virtual FPCGMetadataDomainID GetMetadataDomainIDFromSelector(const FPCGAttributePropertySelector& InSelector) const override;
+	virtual bool MetadataDomainSupportsParenting(const FPCGMetadataDomainID& InDomainID) const override;
+	// ~End UPCGData interface
 
 	/** Returns the entry for the given name */
 	UFUNCTION(BlueprintCallable, Category = Params)
