[PCG] Harden a couple of code paths against transforms with non-normalized rotations to fix a check().

#rnx

diff --git a/Engine/Plugins/PCG/Source/PCG/Private/Data/PCGSplineInteriorSurfaceData.cpp b/Engine/Plugins/PCG/Source/PCG/Private/Data/PCGSplineInteriorSurfaceData.cpp
--- a/Engine/Plugins/PCG/Source/PCG/Private/Data/PCGSplineInteriorSurfaceData.cpp
+++ b/Engine/Plugins/PCG/Source/PCG/Private/Data/PCGSplineInteriorSurfaceData.cpp
@@ -184,7 +184,17 @@ bool UPCGSplineInteriorSurfaceData::PointInsidePolygon(const FTransform& InTrans
 	check(!bNeedsToCache);
 #endif
 
-	const FBox TransformedBounds = InBounds.TransformBy(InTransform);
+	FBox TransformedBounds;
+	if (InTransform.IsRotationNormalized())
+	{
+		TransformedBounds = InBounds.TransformBy(InTransform);
+	}
+	else
+	{
+		FTransform TranslationAndScale = InTransform;
+		TranslationAndScale.SetRotation(FQuat::Identity);
+		TransformedBounds = InBounds.TransformBy(TranslationAndScale);
+	}
 
 	// Test point bounds against the spline bounds.
 	if (CachedBounds.ComputeSquaredDistanceToBox(TransformedBounds) > 0.0)
diff --git a/Engine/Plugins/PCG/Source/PCG/Public/Helpers/PCGPointHelpers.h b/Engine/Plugins/PCG/Source/PCG/Public/Helpers/PCGPointHelpers.h
--- a/Engine/Plugins/PCG/Source/PCG/Public/Helpers/PCGPointHelpers.h
+++ b/Engine/Plugins/PCG/Source/PCG/Public/Helpers/PCGPointHelpers.h
@@ -37,7 +37,16 @@ namespace PCGPointHelpers
 
 	FORCEINLINE FBoxSphereBounds GetDensityBounds(const FTransform& InTransform, float InSteepness, const FVector& InBoundsMin, const FVector& InBoundsMax)
 	{
-		return FBoxSphereBounds(PCGPointHelpers::GetLocalDensityBounds(InSteepness, InBoundsMin, InBoundsMax).TransformBy(InTransform));
+		if (InTransform.IsRotationNormalized())
+		{
+			return FBoxSphereBounds(PCGPointHelpers::GetLocalDensityBounds(InSteepness, InBoundsMin, InBoundsMax).TransformBy(InTransform));
+		}
+		else
+		{
+			FTransform TranslationAndScale = InTransform;
+			TranslationAndScale.SetRotation(FQuat::Identity);
+			return FBoxSphereBounds(PCGPointHelpers::GetLocalDensityBounds(InSteepness, InBoundsMin, InBoundsMax).TransformBy(TranslationAndScale));
+		}
 	}
 	
 	FORCEINLINE void SetExtents(const FVector& InExtents, FVector& InOutBoundsMin, FVector& InOutBoundsMax)
