[PCG] Automatically select a debug object when toggling on inspection if no object was selected before.

diff --git a/Engine/Plugins/PCG/Source/PCGEditor/Private/Widgets/SPCGEditorGraphDebugObjectTree.cpp b/Engine/Plugins/PCG/Source/PCGEditor/Private/Widgets/SPCGEditorGraphDebugObjectTree.cpp
--- a/Engine/Plugins/PCG/Source/PCGEditor/Private/Widgets/SPCGEditorGraphDebugObjectTree.cpp
+++ b/Engine/Plugins/PCG/Source/PCGEditor/Private/Widgets/SPCGEditorGraphDebugObjectTree.cpp
@@ -509,6 +509,7 @@ void SPCGEditorGraphDebugObjectTree::Tick(const FGeometry& AllottedGeometry, con
 {
 	SCompoundWidget::Tick(AllottedGeometry, InCurrentTime, InDeltaTime);
 
+	bool bRefreshDone = false;
 	if (bNeedsRefresh && InCurrentTime >= NextRefreshTime)
 	{
 		// Updating the tree while the inspected component is generating can be bad as the selection can
@@ -521,6 +522,7 @@ void SPCGEditorGraphDebugObjectTree::Tick(const FGeometry& AllottedGeometry, con
 
 			constexpr double RefreshCooldownTime = 0.25f;
 			NextRefreshTime = InCurrentTime + RefreshCooldownTime;
+			bRefreshDone = true;
 		}
 	}
 
@@ -528,6 +530,31 @@ void SPCGEditorGraphDebugObjectTree::Tick(const FGeometry& AllottedGeometry, con
 	{
 		UpdateIsSetDebugObjectFromSelectionEnabled();
 	}
+
+	if (bRefreshDone && bShouldSelectStackOnNextRefresh)
+	{
+		if (SelectedStack.GetStackFrames().IsEmpty())
+		{
+			if (IsSetDebugObjectFromSelectionEnabled.IsSet() && IsSetDebugObjectFromSelectionEnabled.GetValue())
+			{
+				SetDebugObjectFromSelection_OnClicked();
+			}
+			else // Select first occurrence otherwise
+			{
+				const UPCGGraph* CurrentGraph = PCGEditor.Pin()->GetPCGGraph();
+				for (const FPCGEditorGraphDebugObjectItemPtr& Item : AllGraphItems)
+				{
+					if (Item->GetPCGStack() && Item->GetPCGStack()->GetGraphForCurrentFrame() == CurrentGraph)
+					{
+						ExpandAndSelectDebugObject(Item);
+						break;
+					}
+				}
+			}
+		}
+
+		bShouldSelectStackOnNextRefresh = false;
+	}
 }
 
 void SPCGEditorGraphDebugObjectTree::RequestRefresh()
@@ -539,6 +566,11 @@ void SPCGEditorGraphDebugObjectTree::SetNodeBeingInspected(const UPCGNode* InPCG
 {
 	PCGNodeBeingInspected = IsValid(InPCGNode) ? InPCGNode : nullptr;
 
+	if (PCGNodeBeingInspected && SelectedStack.GetStackFrames().IsEmpty())
+	{
+		bShouldSelectStackOnNextRefresh = true;
+	}
+
 	RequestRefresh();
 }
 
diff --git a/Engine/Plugins/PCG/Source/PCGEditor/Private/Widgets/SPCGEditorGraphDebugObjectTree.h b/Engine/Plugins/PCG/Source/PCGEditor/Private/Widgets/SPCGEditorGraphDebugObjectTree.h
--- a/Engine/Plugins/PCG/Source/PCGEditor/Private/Widgets/SPCGEditorGraphDebugObjectTree.h
+++ b/Engine/Plugins/PCG/Source/PCGEditor/Private/Widgets/SPCGEditorGraphDebugObjectTree.h
@@ -332,6 +332,7 @@ private:
 	TArray<FPCGEditorGraphDebugObjectItemPtr> AllGraphItems;
 
 	bool bNeedsRefresh = false;
+	bool bShouldSelectStackOnNextRefresh = false;
 
 	double NextRefreshTime = 0.0;
 
