[PCG] Delegate more compute graph dispatch overheads to async tasks.

#rnx

diff --git a/Engine/Plugins/PCG/Source/PCG/Private/Compute/DataInterfaces/Elements/PCGCopyPointsAnalysisDataInterface.cpp b/Engine/Plugins/PCG/Source/PCG/Private/Compute/DataInterfaces/Elements/PCGCopyPointsAnalysisDataInterface.cpp
--- a/Engine/Plugins/PCG/Source/PCG/Private/Compute/DataInterfaces/Elements/PCGCopyPointsAnalysisDataInterface.cpp
+++ b/Engine/Plugins/PCG/Source/PCG/Private/Compute/DataInterfaces/Elements/PCGCopyPointsAnalysisDataInterface.cpp
@@ -78,15 +78,25 @@ void UPCGCopyPointsAnalysisDataProvider::Initialize(const UComputeDataInterface*
 
 	Settings = CastChecked<UPCGCopyPointsSettings>(DataInterface->ProducerKernel->GetSettings());
 
+	bCopyEachSourceOnEveryTarget = Settings->bCopyEachSourceOnEveryTarget;
+}
+
+bool UPCGCopyPointsAnalysisDataProvider::PrepareForExecute_GameThread(UPCGDataBinding* InBinding)
+{
+	if (!Super::PrepareForExecute_GameThread(InBinding))
+	{
+		return false;
+	}
+
 	// Resolve attribute IDs from incoming data.
 	const FName MatchAttributeName = Settings->MatchAttribute;
-	MatchAttributeId = Binding->GetAttributeId(MatchAttributeName, EPCGKernelAttributeType::Int);
+	MatchAttributeId = InBinding->GetAttributeId(MatchAttributeName, EPCGKernelAttributeType::Int);
 	ensure(MatchAttributeId != INDEX_NONE);
 
-	SelectedFlagAttributeId = Binding->GetAttributeId(PCGCopyPointsConstants::SelectedFlagAttributeName, EPCGKernelAttributeType::Bool);
+	SelectedFlagAttributeId = InBinding->GetAttributeId(PCGCopyPointsConstants::SelectedFlagAttributeName, EPCGKernelAttributeType::Bool);
 	ensure(SelectedFlagAttributeId != INDEX_NONE);
 
-	bCopyEachSourceOnEveryTarget = Settings->bCopyEachSourceOnEveryTarget;
+	return true;
 }
 
 FComputeDataProviderRenderProxy* UPCGCopyPointsAnalysisDataProvider::GetRenderProxy()
diff --git a/Engine/Plugins/PCG/Source/PCG/Private/Compute/DataInterfaces/Elements/PCGCopyPointsAnalysisDataInterface.h b/Engine/Plugins/PCG/Source/PCG/Private/Compute/DataInterfaces/Elements/PCGCopyPointsAnalysisDataInterface.h
--- a/Engine/Plugins/PCG/Source/PCG/Private/Compute/DataInterfaces/Elements/PCGCopyPointsAnalysisDataInterface.h
+++ b/Engine/Plugins/PCG/Source/PCG/Private/Compute/DataInterfaces/Elements/PCGCopyPointsAnalysisDataInterface.h
@@ -34,6 +34,7 @@ class UPCGCopyPointsAnalysisDataProvider : public UPCGComputeDataProvider
 public:
 	//~ Begin UComputeDataProvider Interface
 	virtual void Initialize(const UComputeDataInterface* InDataInterface, UObject* InBinding, uint64 InInputMask, uint64 InOutputMask) override;
+	virtual bool PrepareForExecute_GameThread(UPCGDataBinding* InBinding) override;
 	FComputeDataProviderRenderProxy* GetRenderProxy() override;
 	//~ End UComputeDataProvider Interface
 
diff --git a/Engine/Plugins/PCG/Source/PCG/Private/Compute/Elements/PCGComputeGraphElement.cpp b/Engine/Plugins/PCG/Source/PCG/Private/Compute/Elements/PCGComputeGraphElement.cpp
--- a/Engine/Plugins/PCG/Source/PCG/Private/Compute/Elements/PCGComputeGraphElement.cpp
+++ b/Engine/Plugins/PCG/Source/PCG/Private/Compute/Elements/PCGComputeGraphElement.cpp
@@ -142,37 +142,41 @@ bool FPCGComputeGraphElement::ExecuteInternal(FPCGContext* InContext) const
 			return true;
 		}
 
-		Context->ExecutionSubPhase = EPCGComputeGraphExecutionPhase::InitializeDataBinding;
+		Context->ExecutionSubPhase = EPCGComputeGraphExecutionPhase::InitializeDataBindingAndComputeGraph;
 
 		if (Context->AsyncState.ShouldStop())
 		{
 			return false;
 		}
 	}
-	case EPCGComputeGraphExecutionPhase::InitializeDataBinding: // Fallthrough
+	case EPCGComputeGraphExecutionPhase::InitializeDataBindingAndComputeGraph: // Fallthrough
 	{
-		TRACE_CPUPROFILER_EVENT_SCOPE(InitializeDataBinding);
+		TRACE_CPUPROFILER_EVENT_SCOPE(InitializeDataBindingAndComputeGraph);
 		ensure(!Context->DataBinding);
 
 		UPCGDataBinding* DataBinding = FPCGContext::NewObject_AnyThread<UPCGDataBinding>(Context);
 		Context->DataBinding.Reset(DataBinding);
 
-		Context->DataBinding->Initialize(Context->ComputeGraph.Get(), Context);
+		DataBinding->Initialize(Context->ComputeGraph.Get(), Context);
 
-		Context->ExecutionSubPhase = EPCGComputeGraphExecutionPhase::InitializeComputeGraphInstance;
-
-		if (Context->AsyncState.ShouldStop())
+		// Start data binding full initialization task which can run concurrently with initializing data providers.
+		UE::Tasks::TTask<void> InitializeBindingTask = UE::Tasks::Launch(UE_SOURCE_LOCATION, [ContextHandle = Context->GetOrCreateHandle()]()
 		{
-			return false;
-		}
-	}
-	case EPCGComputeGraphExecutionPhase::InitializeComputeGraphInstance: // Fallthrough
-	{
-		TRACE_CPUPROFILER_EVENT_SCOPE(InitializeComputeGraphInstance);
+			FPCGContext::FSharedContext<FPCGComputeGraphContext> SharedContext(ContextHandle);
+			FPCGComputeGraphContext* Context = SharedContext.Get();
+			if (!Context)
+			{
+				return;
+			}
+
+			Context->DataBinding->InitializeTables(Context);
+		});
 
 		const UPCGGraph* Graph = Context->Stack ? Context->Stack->GetNearestNonInlinedGraphForCurrentFrame() : nullptr;
 		const UPCGGraph::FComputeGraphInstanceKey Key = { Context->SourceComponent->GetGenerationGridSize(), ComputeGraphIndex };
 
+		// The data provider initialization must not depend on the data binding being set up, although they still output of the PreInitialize such
+		// as the SourceComponent.
 		if (Graph && Graph->RetrieveComputeGraphInstanceFromPool(Key, Context->ComputeGraphInstance))
 		{
 			Context->ComputeGraphInstance.InitializeDataProviders(Context->ComputeGraph.Get(), 0, Context->DataBinding.Get());
@@ -255,6 +259,8 @@ bool FPCGComputeGraphElement::ExecuteInternal(FPCGContext* InContext) const
 			}
 		}
 
+		InitializeBindingTask.Wait();
+
 		Context->ExecutionSubPhase = EPCGComputeGraphExecutionPhase::PreExecuteReadbacks;
 
 		if (Context->AsyncState.ShouldStop())
@@ -281,7 +287,7 @@ bool FPCGComputeGraphElement::ExecuteInternal(FPCGContext* InContext) const
 
 		if (bAllReady)
 		{
-			Context->ExecutionSubPhase = EPCGComputeGraphExecutionPhase::SchedulePrimeDataDescriptionCache;
+			Context->ExecutionSubPhase = EPCGComputeGraphExecutionPhase::PrimeDataDescriptionsAndValidateData;
 		}
 		else
 		{
@@ -294,21 +300,22 @@ bool FPCGComputeGraphElement::ExecuteInternal(FPCGContext* InContext) const
 			return false;
 		}
 	}
-	case EPCGComputeGraphExecutionPhase::SchedulePrimeDataDescriptionCache: // Fallthrough
+	case EPCGComputeGraphExecutionPhase::PrimeDataDescriptionsAndValidateData: // Fallthrough
 	{
 		// Note: Priming the data description cache must take place after pre-execute readbacks, as data descriptions
 		// may rely on the readback data, e.g. analysis kernels.
-
-		if (!Context->bPrimeDataDescCacheScheduled)
+		if (!Context->bDataDescrPrimeAndValidateScheduled)
 		{
-			TRACE_CPUPROFILER_EVENT_SCOPE(SchedulePrimeDataDescriptionCache);
+			TRACE_CPUPROFILER_EVENT_SCOPE(PrimeDataDescriptionsAndValidateData);
 
-			Context->bPrimeDataDescCacheScheduled = true;
+			Context->bDataDescrPrimeAndValidateScheduled = true;
 
 			// @todo_pcg: In the future perhaps we can just use Context->ScheduleGeneric(), but it's unclear at the moment if there is any guarantee this task
 			// would be kicked off this tick. To avoid the unknown, we'll just launch a task directly here, but it should be investigated in the future.
-			Context->PrimeDataDescCacheTask = UE::Tasks::Launch(UE_SOURCE_LOCATION, [ContextHandle = Context->GetOrCreateHandle()]()
+			Context->DataDescrPrimeAndValidateTask = UE::Tasks::Launch(UE_SOURCE_LOCATION, [ContextHandle = Context->GetOrCreateHandle()]()
 			{
+				TRACE_CPUPROFILER_EVENT_SCOPE(DataDescrPrimeAndValidateTask);
+
 				FPCGContext::FSharedContext<FPCGComputeGraphContext> SharedContext(ContextHandle);
 				FPCGComputeGraphContext* Context = SharedContext.Get();
 				if (!Context)
@@ -317,38 +324,28 @@ bool FPCGComputeGraphElement::ExecuteInternal(FPCGContext* InContext) const
 				}
 
 				Context->DataBinding->PrimeDataDescriptionCache();
+
+				// Graph data validation is currently allowed to depend on data descriptions, hence done serially here.
+				Context->bGraphValid = Context->ComputeGraph->IsGraphDataValid(Context);
 			});
 
 			return false;
 		}
 
-		if (!ensure(Context->PrimeDataDescCacheTask.IsValid()))
+		if (!ensure(Context->DataDescrPrimeAndValidateTask.IsValid()))
 		{
 			return true;
 		}
 
-		if (!Context->PrimeDataDescCacheTask.IsCompleted())
+		if (!Context->DataDescrPrimeAndValidateTask.IsCompleted())
 		{
 			SleepUntilNextFrame();
 			return false;
 		}
 
-		ensure(Context->DataBinding->IsDataDescriptionCachePrimed());
 		Context->DataBinding->DebugLogDataDescriptions();
-		
-		Context->ExecutionSubPhase = EPCGComputeGraphExecutionPhase::PerformDataValidation;
-
-		if (Context->AsyncState.ShouldStop())
-		{
-			return false;
-		}
-	}
-	case EPCGComputeGraphExecutionPhase::PerformDataValidation: // Fallthrough
-	{
-		TRACE_CPUPROFILER_EVENT_SCOPE(PerformDataValidation);
 
-		// Perform validation after input data is initialized and any pre-execute data has been read back.
-		if (!Context->ComputeGraph->IsGraphDataValid(InContext))
+		if (!Context->bGraphValid)
 		{
 			return true;
 		}
diff --git a/Engine/Plugins/PCG/Source/PCG/Private/Compute/Elements/PCGComputeGraphElement.h b/Engine/Plugins/PCG/Source/PCG/Private/Compute/Elements/PCGComputeGraphElement.h
--- a/Engine/Plugins/PCG/Source/PCG/Private/Compute/Elements/PCGComputeGraphElement.h
+++ b/Engine/Plugins/PCG/Source/PCG/Private/Compute/Elements/PCGComputeGraphElement.h
@@ -19,11 +19,9 @@ enum class EPCGComputeGraphExecutionPhase : uint8
 {
 	None,
 	GetComputeGraph,
-	InitializeDataBinding,
-	InitializeComputeGraphInstance,
+	InitializeDataBindingAndComputeGraph,
 	PreExecuteReadbacks,
-	SchedulePrimeDataDescriptionCache,
-	PerformDataValidation,
+	PrimeDataDescriptionsAndValidateData,
 	PrepareForExecute,
 	ValidateComputeGraphCompilation,
 	ScheduleComputeGraph,
@@ -53,9 +51,9 @@ public:
 	/** Whether the compute graph instance has been initialized (data providers created, etc). */
 	bool bComputeGraphInstanceInitialized = false;
 
-	/** Whether the task to prime the data description cache has been scheduled. */
-	bool bPrimeDataDescCacheScheduled = false;
-	UE::Tasks::TTask<void> PrimeDataDescCacheTask;
+	/** Whether the task to prime the data description cache and validate the input data has been scheduled. */
+	bool bDataDescrPrimeAndValidateScheduled = false;
+	UE::Tasks::TTask<void> DataDescrPrimeAndValidateTask;
 
 	/** Graph was enqueued but was invalid for some reason and the work could not be submitted. */
 	bool bGraphSubmitFailed = false;
@@ -71,6 +69,8 @@ public:
 	TSet<TObjectPtr<UComputeDataProvider>> ProvidersWithReadbacks;
 
 	EPCGComputeGraphExecutionPhase ExecutionSubPhase = EPCGComputeGraphExecutionPhase::GetComputeGraph;
+
+	std::atomic<bool> bGraphValid = false;
 };
 
 /** Executes a CF graph. Created by the compiler when collapsing GPU nodes rather than by a settings/node. */
diff --git a/Engine/Plugins/PCG/Source/PCG/Private/Compute/PCGDataBinding.cpp b/Engine/Plugins/PCG/Source/PCG/Private/Compute/PCGDataBinding.cpp
--- a/Engine/Plugins/PCG/Source/PCG/Private/Compute/PCGDataBinding.cpp
+++ b/Engine/Plugins/PCG/Source/PCG/Private/Compute/PCGDataBinding.cpp
@@ -19,17 +19,42 @@ void UPCGDataBinding::Initialize(const UPCGComputeGraph* InComputeGraph, FPCGCon
 	check(InComputeGraph);
 	check(InContext);
 
+	Graph = InComputeGraph;
+
 	ContextHandle = InContext->GetOrCreateHandle();
 
-	Graph = InComputeGraph;
+	// Add compatible data objects from input collection.
+	InputDataCollection.TaggedData.Reserve(InContext->InputData.TaggedData.Num());
+	for (const FPCGTaggedData& InputData : InContext->InputData.TaggedData)
+	{
+		if (InputData.Data && PCGComputeHelpers::IsTypeAllowedAsInput(InputData.Data->GetDataType()))
+		{
+			InputDataCollection.TaggedData.Add(InputData);
+		}
+		else if (InputData.Data)
+		{
+			UE_LOG(LogPCG, Warning, TEXT("Stripped input data that is not currently supported by GPU execution: %s"), *InputData.Data->GetName());
+		}
+	}
+}
 
-	AttributeTable = InComputeGraph->GetStaticAttributeTable();
+void UPCGDataBinding::InitializeTables(FPCGContext* InContext)
+{
+	TRACE_CPUPROFILER_EVENT_SCOPE(UPCGDataBinding::InitializeTables);
+	check(InContext);
+
+	AttributeTable = Graph->GetStaticAttributeTable();
 
 	// String table always contains empty string in index 0 (and string key attributes are 0-initialized), and then any strings known statically at compilation time.
 	StringTable = { FString() };
 	StringTable.Append(Graph->GetStaticStringTable());
 
-	InitializeInputData(InContext->InputData);
+	// Augment static tables with data incoming from CPU.
+	AddInputDataAttributesToTable();
+	AddInputDataStringsToTable();
+	AddInputDataTagsToTable();
+
+	bTablesInitialized = true;
 }
 
 void UPCGDataBinding::ReleaseTransientResources()
@@ -74,6 +99,8 @@ bool UPCGDataBinding::ComputeKernelPinDataDesc(const FPCGKernelPin& InKernelPin,
 {
 	// Uncomment to be notified when a kernel pin description is computed late:
 	//ensureAlwaysMsgf(!IsInGameThread(), TEXT("Data descriptions should only be computed when priming the cache, which takes place on a worker thread."));
+	ensure(!bIsDataDescriptionCachePrimed);
+	ensure(bTablesInitialized);
 
 	if (!ensure(Graph))
 	{
@@ -264,30 +291,6 @@ bool UPCGDataBinding::ReadbackInputDataToCPU(int32 InInputDataIndex)
 	return true;
 }
 
-void UPCGDataBinding::InitializeInputData(const FPCGDataCollection& InComputeGraphElementInputData)
-{
-	TRACE_CPUPROFILER_EVENT_SCOPE(UPCGDataBinding::InitializeInputData);
-
-	// Add compatible data objects from input collection.
-	InputDataCollection.TaggedData.Reserve(InComputeGraphElementInputData.TaggedData.Num());
-	for (const FPCGTaggedData& InputData : InComputeGraphElementInputData.TaggedData)
-	{
-		if (InputData.Data && PCGComputeHelpers::IsTypeAllowedAsInput(InputData.Data->GetDataType()))
-		{
-			InputDataCollection.TaggedData.Add(InputData);
-		}
-		else if (InputData.Data)
-		{
-			UE_LOG(LogPCG, Warning, TEXT("Stripped input data that is not currently supported by GPU execution: %s"), *InputData.Data->GetName());
-		}
-	}
-
-	// Augment static tables with data incoming from CPU.
-	AddInputDataAttributesToTable();
-	AddInputDataStringsToTable();
-	AddInputDataTagsToTable();
-}
-
 void UPCGDataBinding::AddInputDataAttributesToTable()
 {
 	for (const FPCGTaggedData& Data : InputDataCollection.TaggedData)
@@ -463,6 +466,8 @@ void UPCGDataBinding::AddInputDataTagsToTable()
 void UPCGDataBinding::DebugLogDataDescriptions() const
 {
 #if PCG_KERNEL_LOGGING_ENABLED
+	ensure(IsDataDescriptionCachePrimed());
+
 	if (Graph && Graph->bLogDataDescriptions)
 	{
 		UE_LOG(LogPCG, Display, TEXT("\n### METADATA ATTRIBUTE TABLE ###"));
diff --git a/Engine/Plugins/PCG/Source/PCG/Private/Compute/PCGDataBinding.h b/Engine/Plugins/PCG/Source/PCG/Private/Compute/PCGDataBinding.h
--- a/Engine/Plugins/PCG/Source/PCG/Private/Compute/PCGDataBinding.h
+++ b/Engine/Plugins/PCG/Source/PCG/Private/Compute/PCGDataBinding.h
@@ -63,16 +63,20 @@ class UPCGDataBinding : public UObject
 	GENERATED_BODY()
 
 public:
+	/** Pre-initialization, set up pointers and store input data. */
 	void Initialize(const UPCGComputeGraph* InComputeGraph, FPCGContext* InContext);
 
+	/** Initialization look up tables from input data. */
+	void InitializeTables(FPCGContext* InContext);
+
 	/** Clear state and release any handles to resources such as GPU buffers. */
 	void ReleaseTransientResources();
 
-	int32 GetAttributeId(const FPCGKernelAttributeKey& InAttribute) const { return AttributeTable.GetAttributeId(InAttribute); }
-	int32 GetAttributeId(FName InName, EPCGKernelAttributeType InType) const { return AttributeTable.GetAttributeId(InName, InType); }
+	int32 GetAttributeId(const FPCGKernelAttributeKey& InAttribute) const { ensure(bTablesInitialized); return AttributeTable.GetAttributeId(InAttribute); }
+	int32 GetAttributeId(FName InName, EPCGKernelAttributeType InType) const { ensure(bTablesInitialized); return AttributeTable.GetAttributeId(InName, InType); }
 
 	/** Get the set of all unique strings entering the compute graph. */
-	const TArray<FString>& GetStringTable() const { return StringTable; }
+	const TArray<FString>& GetStringTable() const { ensure(bTablesInitialized); return StringTable; }
 
 	/** Computes a description of data for every pin in the compute graph and caches it. */
 	void PrimeDataDescriptionCache();
@@ -108,8 +112,6 @@ public:
 	void DebugLogDataDescriptions() const;
 
 private:
-	void InitializeInputData(const FPCGDataCollection& InComputeGraphElementInputData);
-
 	/** Loop over all metadata attributes in all input data ensure all attributes are registered in attribute table. */
 	void AddInputDataAttributesToTable();
 	void AddInputDataStringsToTable();
@@ -154,4 +156,6 @@ protected:
 	/** Cache of data descriptions to amortize cost of computing them at runtime. */
 	TMap</*Binding index*/int32, FPCGDataCollectionDesc> DataDescriptionCache;
 	std::atomic<bool> bIsDataDescriptionCachePrimed = false;
+
+	std::atomic<bool> bTablesInitialized = false;
 };
