[PCG] Temporary fix for Point to AttributeSet

* Now that the normal initialize is not working anymore for ParamData, we need to do the setup manually
* Will have a followup to make it cleaner

#rnx

diff --git a/Engine/Plugins/PCG/Source/PCG/Private/Elements/PCGConvertToAttributeSet.cpp b/Engine/Plugins/PCG/Source/PCG/Private/Elements/PCGConvertToAttributeSet.cpp
--- a/Engine/Plugins/PCG/Source/PCG/Private/Elements/PCGConvertToAttributeSet.cpp
+++ b/Engine/Plugins/PCG/Source/PCG/Private/Elements/PCGConvertToAttributeSet.cpp
@@ -63,22 +63,34 @@ bool FPCGConvertToAttributeSetElement::ExecuteInternal(FPCGContext* Context) con
 		}
 
 		UPCGParamData* ParamData = FPCGContext::NewObject_AnyThread<UPCGParamData>(Context);
-		check(ParamData->Metadata);
-		ParamData->Metadata->Initialize(SourceMetadata);
+		UPCGMetadata* Metadata = ParamData->MutableMetadata();
 
-		// Create new entry that will map (or not) to the original parent
-		TArray<PCGMetadataEntryKey> EntryKeys;
+		check(Metadata);
 
-		const TConstPCGValueRange<int64> MetadataEntryRange = PointData->GetConstMetadataEntryValueRange();
-		check(PointData->GetNumPoints() == MetadataEntryRange.Num());
-		EntryKeys.Reserve(MetadataEntryRange.Num());
-		for (int32 Index = 0; Index < MetadataEntryRange.Num(); ++Index)
+		// Copy all the global metadata, if it exists.
+		if (const FPCGMetadataDomain* SourceGlobalDomain = SourceMetadata->GetConstMetadataDomain(PCGMetadataDomainID::Data))
 		{
-			EntryKeys.Add(ParamData->Metadata->AddEntry(MetadataEntryRange[Index]));
+			FPCGMetadataDomain* GlobalDomain = Metadata->GetMetadataDomain(PCGMetadataDomainID::Data);
+			check(GlobalDomain);
+			GlobalDomain->InitializeAsCopy(SourceGlobalDomain);
 		}
 
-		ParamData->Metadata->FlattenAndCompress(EntryKeys);
+		// Then copy all the elements, if it exists.
+		if (const FPCGMetadataDomain* SourceElementsDomain = SourceMetadata->GetConstMetadataDomain(PCGMetadataDomainID::Elements))
+		{
+			// Gather all the keys to copy over.
+			TArray<PCGMetadataEntryKey> EntryKeys;
+
+			const TConstPCGValueRange<int64> MetadataEntryRange = PointData->GetConstMetadataEntryValueRange();
+			check(PointData->GetNumPoints() == MetadataEntryRange.Num());
+			EntryKeys.Reserve(MetadataEntryRange.Num());
+			Algo::Transform(MetadataEntryRange, EntryKeys, [](int64 EntryKey) -> PCGMetadataEntryKey { return EntryKey; });
 
+			FPCGMetadataDomain* ElementsDomain = Metadata->GetMetadataDomain(PCGMetadataDomainID::Elements);
+			check(ElementsDomain);
+			ElementsDomain->InitializeAsCopy(SourceElementsDomain, &EntryKeys);
+		}
+		
 		FPCGTaggedData& Output = Outputs.Emplace_GetRef(Input);
 		Output.Data = ParamData;
 	}
