[PCG] Swap Mac platform ifdef for Debug Only Selected

#rnx

diff --git a/Engine/Plugins/PCG/Source/PCGEditor/Private/PCGEditorCommands.cpp b/Engine/Plugins/PCG/Source/PCGEditor/Private/PCGEditorCommands.cpp
--- a/Engine/Plugins/PCG/Source/PCGEditor/Private/PCGEditorCommands.cpp
+++ b/Engine/Plugins/PCG/Source/PCGEditor/Private/PCGEditorCommands.cpp
@@ -45,9 +45,9 @@ void FPCGEditorCommands::RegisterCommands()
 
 	//~ PLATFORM SPECIFIC ~//
 #if PLATFORM_MAC
-	UI_COMMAND(DebugOnlySelected, "Debug Only Selected", "Enable node debug state for selected nodes and disable debug state for the others", EUserInterfaceActionType::Button, FInputChord(EModifierKey::Control | EModifierKey::Alt, EKeys::D));
-#else
 	UI_COMMAND(DebugOnlySelected, "Debug Only Selected", "Enable node debug state for selected nodes and disable debug state for the others", EUserInterfaceActionType::Button, FInputChord(EModifierKey::Command | EModifierKey::Shift, EKeys::D));
+#else
+	UI_COMMAND(DebugOnlySelected, "Debug Only Selected", "Enable node debug state for selected nodes and disable debug state for the others", EUserInterfaceActionType::Button, FInputChord(EModifierKey::Control | EModifierKey::Alt, EKeys::D));
 #endif
 }
 
