[PGC] Attribute List View: Preserve selection of Data and Domain combobox

diff --git a/Engine/Plugins/PCG/Source/PCGEditor/Private/Widgets/SPCGEditorGraphAttributeListView.cpp b/Engine/Plugins/PCG/Source/PCGEditor/Private/Widgets/SPCGEditorGraphAttributeListView.cpp
--- a/Engine/Plugins/PCG/Source/PCGEditor/Private/Widgets/SPCGEditorGraphAttributeListView.cpp
+++ b/Engine/Plugins/PCG/Source/PCGEditor/Private/Widgets/SPCGEditorGraphAttributeListView.cpp
@@ -584,8 +584,15 @@ void SPCGEditorGraphAttributeListView::Tick(const FGeometry& AllottedGeometry, c
 	{
 		bNeedsRefresh = false;
 
-		RefreshPinComboBox();
-		RefreshDataComboBox();
+		bool bKeepSelection = !bPCGEditorGraphNodeChanged;
+		bPCGEditorGraphNodeChanged = false;
+
+		bool bSelectionChanged = false;
+		RefreshPinComboBox(bKeepSelection, bSelectionChanged);
+
+		bKeepSelection &= !bSelectionChanged;
+
+		RefreshDataComboBox(bKeepSelection);
 		RefreshAttributeList();
 	}
 
@@ -664,6 +671,7 @@ void SPCGEditorGraphAttributeListView::SetNodeBeingInspected(UPCGEditorGraphNode
 	// Always unlock when changing the node, to make sure we unlock when removing the inspected node
 	bIsLocked = false;
 
+	bPCGEditorGraphNodeChanged = true;
 	RequestRefresh();
 }
 
@@ -840,8 +848,10 @@ void SPCGEditorGraphAttributeListView::RefreshAttributeList()
 	LaunchUpdateTask();
 }
 
-void SPCGEditorGraphAttributeListView::RefreshPinComboBox()
+void SPCGEditorGraphAttributeListView::RefreshPinComboBox(bool bKeepSelection, bool& OutSelectionChanged)
 {
+	OutSelectionChanged = !bKeepSelection;
+	const int32 PinComboBoxItemSelectedIndex = bKeepSelection ? PinComboBoxItems.IndexOfByKey(PinComboBox->GetSelectedItem()) : INDEX_NONE;
 	PinComboBoxItems.Empty();
 	PinComboBox->ClearSelection();
 	PinComboBox->RefreshOptions();
@@ -885,14 +895,31 @@ void SPCGEditorGraphAttributeListView::RefreshPinComboBox()
 	PopulatePins(PCGNode->GetOutputPins(), TEXT("Output: {0}"), PinComboBoxItems, &FirstConnectedItemIndex);
 	PopulatePins(PCGNode->GetInputPins(), TEXT("Input: {0}"), PinComboBoxItems, nullptr);
 
-	if (PinComboBoxItems.Num() > 0 && ensure(PinComboBoxItems.IsValidIndex(FirstConnectedItemIndex)))
+	bool bResetDataComboBoxItemsSelectedIndex = true;
+	int32 SelectedIndex = PinComboBoxItems.IsValidIndex(PinComboBoxItemSelectedIndex) ? PinComboBoxItemSelectedIndex : FirstConnectedItemIndex;
+	if (PinComboBoxItems.Num() > 0 && ensure(PinComboBoxItems.IsValidIndex(SelectedIndex)))
 	{
-		PinComboBox->SetSelectedItem(PinComboBoxItems[FirstConnectedItemIndex]);
+		PinComboBox->SetSelectedItem(PinComboBoxItems[SelectedIndex]);
+
+		OutSelectionChanged = SelectedIndex != PinComboBoxItemSelectedIndex;
 	}
 }
 
-void SPCGEditorGraphAttributeListView::RefreshDataComboBox()
+void SPCGEditorGraphAttributeListView::RefreshDataComboBox(bool bKeepSelection)
 {
+	// Cache previous selection only if we have items because this code gets called twice.
+	// First call early outs on GetInspectionData(), when we get called the second time the combo box is empty but
+	// we want to try and restore the cached element.
+	if (!DataComboBoxItems.IsEmpty())
+	{
+		DataComboBoxItemsSelectedIndex = GetSelectedDataIndex();
+	}
+	
+	if(!bKeepSelection)
+	{
+		DataComboBoxItemsSelectedIndex = INDEX_NONE;
+	}
+
 	DataComboBoxItems.Empty();
 	DataComboBox->ClearSelection();
 	DataComboBox->RefreshOptions();
@@ -900,7 +927,7 @@ void SPCGEditorGraphAttributeListView::RefreshDataComboBox()
 	const FPCGDataCollection* InspectionData = GetInspectionData();
 	if (!InspectionData)
 	{
-		RefreshDomainComboBox();
+		RefreshDomainComboBox(bKeepSelection);
 		return;
 	}
 
@@ -920,20 +947,39 @@ void SPCGEditorGraphAttributeListView::RefreshDataComboBox()
 		DataComboBoxItems.Add(MakeShared<FString>(ItemName));
 	}
 
-	if (DataComboBoxItems.Num() > 0)
+	bool bKeepDomainComboBox = false;
+
+	if (DataComboBoxItems.IsValidIndex(DataComboBoxItemsSelectedIndex))
+	{
+		DataComboBox->SetSelectedItem(DataComboBoxItems[DataComboBoxItemsSelectedIndex]);
+		bKeepDomainComboBox = true;
+	}
+	else if (DataComboBoxItems.Num() > 0)
 	{
 		DataComboBox->SetSelectedItem(DataComboBoxItems[0]);
 	}
 
-	RefreshDomainComboBox();
+	DataComboBoxItemsSelectedIndex = INDEX_NONE;
+
+	RefreshDomainComboBox(bKeepDomainComboBox);
 }
 
-void SPCGEditorGraphAttributeListView::RefreshDomainComboBox()
+void SPCGEditorGraphAttributeListView::RefreshDomainComboBox(bool bKeepSelection)
 {
-	// Try to stick to the same metadata domain if possible
-	const int32 DomainIndex = GetSelectedDomainIndex();
-	const FPCGMetadataDomainID PreviousDomain = DomainsComboBoxIds.IsValidIndex(DomainIndex) ? DomainsComboBoxIds[DomainIndex] : PCGMetadataDomainID::Invalid;
-	
+	// Cache previous selection only if we have items because this code gets called twice.
+	// First call early outs on GetInspectionData(), when we get called the second time the combo box is empty but
+	// we want to try and restore the cached element.
+	if (!DomainsComboBoxItems.IsEmpty())
+	{
+		const int32 SelectedDomainIndex = GetSelectedDomainIndex();
+		DomainsComboBoxItemsSelectedDomain = DomainsComboBoxIds.IsValidIndex(SelectedDomainIndex) ? DomainsComboBoxIds[SelectedDomainIndex] : PCGMetadataDomainID::Invalid;
+	}
+
+	if (!bKeepSelection)
+	{
+		DomainsComboBoxItemsSelectedDomain = PCGMetadataDomainID::Invalid;
+	}
+
 	DomainsComboBoxItems.Empty();
 	DomainsComboBoxIds.Empty();
 	DomainsComboBox->ClearSelection();
@@ -972,7 +1018,7 @@ void SPCGEditorGraphAttributeListView::RefreshDomainComboBox()
 	int32 PreviousDomainIndex = INDEX_NONE;
 	for (const FPCGMetadataDomainID& DomainID : DataVisualization->GetAllSupportedDomainsForInspection(PCGData))
 	{
-		if (PreviousDomain == DomainID)
+		if (DomainsComboBoxItemsSelectedDomain == DomainID)
 		{
 			PreviousDomainIndex = DomainsComboBoxItems.Num();
 		}
@@ -991,6 +1037,8 @@ void SPCGEditorGraphAttributeListView::RefreshDomainComboBox()
 		const int32 Index = PreviousDomainIndex != INDEX_NONE ? PreviousDomainIndex : DefaultIndex; 
 		DomainsComboBox->SetSelectedItem(DomainsComboBoxItems[Index]);
 	}
+
+	DomainsComboBoxItemsSelectedDomain = PCGMetadataDomainID::Invalid;
 }
 
 void SPCGEditorGraphAttributeListView::LaunchUpdateTask()
@@ -1180,7 +1228,7 @@ FText SPCGEditorGraphAttributeListView::OnGenerateSelectedPinText() const
 
 void SPCGEditorGraphAttributeListView::OnSelectionChangedPin(TSharedPtr<FPinComboBoxItem> InItem, ESelectInfo::Type InSelectInfo)
 {
-	RefreshDataComboBox();
+	RefreshDataComboBox(/*bKeepSelection=*/InSelectInfo == ESelectInfo::Direct);
 
 	if (InSelectInfo != ESelectInfo::Direct)
 	{
diff --git a/Engine/Plugins/PCG/Source/PCGEditor/Private/Widgets/SPCGEditorGraphAttributeListView.h b/Engine/Plugins/PCG/Source/PCGEditor/Private/Widgets/SPCGEditorGraphAttributeListView.h
--- a/Engine/Plugins/PCG/Source/PCGEditor/Private/Widgets/SPCGEditorGraphAttributeListView.h
+++ b/Engine/Plugins/PCG/Source/PCGEditor/Private/Widgets/SPCGEditorGraphAttributeListView.h
@@ -168,9 +168,9 @@ private:
 	void RestoreColumnWidthVisibility();
 
 	void RefreshAttributeList();
-	void RefreshPinComboBox();
-	void RefreshDataComboBox();
-	void RefreshDomainComboBox();
+	void RefreshPinComboBox(bool bKeepSelection, bool& bOutSelectionChanged);
+	void RefreshDataComboBox(bool bKeepSelection);
+	void RefreshDomainComboBox(bool bKeepSelection);
 
 	void LaunchUpdateTask();
 
@@ -254,6 +254,7 @@ private:
 
 	/** Cached PCGGraphNode being viewed */
 	TWeakObjectPtr<UPCGEditorGraphNodeBase> PCGEditorGraphNode;
+	bool bPCGEditorGraphNodeChanged = false;
 
 	TSharedPtr<FUICommandList> ListViewCommands;
 
@@ -271,10 +272,16 @@ private:
 	TSharedPtr<SComboBox<TSharedPtr<FString>>> DataComboBox;
 	TArray<TSharedPtr<FString>> DataComboBoxItems;
 
+	/** Cached Selected Data ComboBox Item */
+	int32 DataComboBoxItemsSelectedIndex = INDEX_NONE;
+
 	TSharedPtr<SComboBox<TSharedPtr<FString>>> DomainsComboBox;
 	TArray<TSharedPtr<FString>> DomainsComboBoxItems;
 	TArray<FPCGMetadataDomainID> DomainsComboBoxIds;
 
+	/** Cached Selected Domain */
+	FPCGMetadataDomainID DomainsComboBoxItemsSelectedDomain = PCGMetadataDomainID::Invalid;
+
 	TSharedPtr<STextBlock> NodeNameTextBlock;
 	TSharedPtr<STextBlock> InfoTextBlock;
 
