[PCG] Add editor preference to make PCG Partition Actors unselectable

diff --git a/Engine/Plugins/PCG/Source/PCG/Private/Grid/PCGPartitionActor.cpp b/Engine/Plugins/PCG/Source/PCG/Private/Grid/PCGPartitionActor.cpp
--- a/Engine/Plugins/PCG/Source/PCG/Private/Grid/PCGPartitionActor.cpp
+++ b/Engine/Plugins/PCG/Source/PCG/Private/Grid/PCGPartitionActor.cpp
@@ -6,6 +6,7 @@
 #include "PCGModule.h"
 #include "PCGSubsystem.h"
 #include "PCGWorldActor.h"
+#include "Editor/IPCGEditorModule.h"
 #include "Helpers/PCGActorHelpers.h"
 #include "Helpers/PCGHelpers.h"
 
@@ -638,6 +639,11 @@ bool APCGPartitionActor::CanEditChange(const FProperty* InProperty) const
 
 	return true;
 }
+
+bool APCGPartitionActor::IsSelectable() const
+{
+	return IPCGEditorModule::Get() ? IPCGEditorModule::Get()->CanSelectPartitionActors() : true;
+}
 #endif // WITH_EDITOR
 
 void APCGPartitionActor::PostCreation(const FPCGGridDescriptor& GridDescriptor)
diff --git a/Engine/Plugins/PCG/Source/PCG/Public/Editor/IPCGEditorModule.h b/Engine/Plugins/PCG/Source/PCG/Public/Editor/IPCGEditorModule.h
--- a/Engine/Plugins/PCG/Source/PCG/Public/Editor/IPCGEditorModule.h
+++ b/Engine/Plugins/PCG/Source/PCG/Public/Editor/IPCGEditorModule.h
@@ -23,6 +23,7 @@ public:
 #if WITH_EDITOR
 	virtual const FPCGNodeVisualLogs& GetNodeVisualLogs() const = 0;
 	virtual FPCGNodeVisualLogs& GetNodeVisualLogsMutable() = 0;
+	virtual bool CanSelectPartitionActors() const = 0;
 #endif
 
 protected:
diff --git a/Engine/Plugins/PCG/Source/PCG/Public/Grid/PCGPartitionActor.h b/Engine/Plugins/PCG/Source/PCG/Public/Grid/PCGPartitionActor.h
--- a/Engine/Plugins/PCG/Source/PCG/Public/Grid/PCGPartitionActor.h
+++ b/Engine/Plugins/PCG/Source/PCG/Public/Grid/PCGPartitionActor.h
@@ -44,6 +44,7 @@ public:
 	virtual bool ShouldExport() override { return false; }
 	virtual bool ShouldImport(FStringView ActorPropString, bool IsMovingLevel) override { return false; }
 	virtual bool CanEditChange(const FProperty* InProperty) const override;
+	virtual bool IsSelectable() const override;
 #endif
 	//~End AActor Interface
 
diff --git a/Engine/Plugins/PCG/Source/PCGEditor/Private/PCGEditorModule.cpp b/Engine/Plugins/PCG/Source/PCGEditor/Private/PCGEditorModule.cpp
--- a/Engine/Plugins/PCG/Source/PCGEditor/Private/PCGEditorModule.cpp
+++ b/Engine/Plugins/PCG/Source/PCGEditor/Private/PCGEditorModule.cpp
@@ -249,6 +249,11 @@ void FPCGEditorModule::SetOutlinerUIRefreshDelay(float InDelay)
 	}
 }
 
+bool FPCGEditorModule::CanSelectPartitionActors() const
+{
+	return GetDefault<UPCGEditorSettings>()->bCanSelectPartitionActors;
+}
+
 void FPCGEditorModule::RegisterOnEditorModeChange()
 {
 	// Have a callback that catches changes in the Editor modes, to catch when we exit the landscape edit mode.
diff --git a/Engine/Plugins/PCG/Source/PCGEditor/Public/PCGEditorModule.h b/Engine/Plugins/PCG/Source/PCGEditor/Public/PCGEditorModule.h
--- a/Engine/Plugins/PCG/Source/PCGEditor/Public/PCGEditorModule.h
+++ b/Engine/Plugins/PCG/Source/PCGEditor/Public/PCGEditorModule.h
@@ -33,6 +33,7 @@ public:
 	virtual void SetOutlinerUIRefreshDelay(float InDelay) override;
 	virtual const FPCGNodeVisualLogs& GetNodeVisualLogs() const { return NodeVisualLogs; }
 	virtual FPCGNodeVisualLogs& GetNodeVisualLogsMutable() { return NodeVisualLogs; }
+	virtual bool CanSelectPartitionActors() const override;
 
 protected:
 	virtual void OnScheduleGraph(const FPCGStackContext& StackContext);
diff --git a/Engine/Plugins/PCG/Source/PCGEditor/Public/PCGEditorSettings.h b/Engine/Plugins/PCG/Source/PCGEditor/Public/PCGEditorSettings.h
--- a/Engine/Plugins/PCG/Source/PCGEditor/Public/PCGEditorSettings.h
+++ b/Engine/Plugins/PCG/Source/PCGEditor/Public/PCGEditorSettings.h
@@ -224,6 +224,10 @@ public:
 	UPROPERTY(EditAnywhere, config, Category = Workflow)
 	bool bConfirmLocalGraphParameterNameChanges = true;
 
+	/** Allow selecting PCG partition actors in the viewport/outliner. */
+	UPROPERTY(EditAnywhere, config, Category = "Workflow")
+	bool bCanSelectPartitionActors = true;
+
 	/** Specify if we want to disable CPU Throttling when a PCG Graph is executing, this will improve execution time when app is out of focus/minimized */
 	UPROPERTY(EditAnywhere, config, Category = "Editor Performance", meta = (DisplayName = "Disable CPU throttling during graph execution"))
 	bool bDisableCPUThrottlingDuringGraphExecution = false;
