[PCG] Fix crash when returning stale compute graph instance to the pool

diff --git a/Engine/Plugins/PCG/Source/PCG/Private/Compute/Elements/PCGComputeGraphElement.cpp b/Engine/Plugins/PCG/Source/PCG/Private/Compute/Elements/PCGComputeGraphElement.cpp
--- a/Engine/Plugins/PCG/Source/PCG/Private/Compute/Elements/PCGComputeGraphElement.cpp
+++ b/Engine/Plugins/PCG/Source/PCG/Private/Compute/Elements/PCGComputeGraphElement.cpp
@@ -30,23 +30,26 @@ FPCGComputeGraphContext::~FPCGComputeGraphContext()
 		DataBinding->ReleaseTransientResources();
 	}
 
-	for (TObjectPtr<UComputeDataProvider> Provider : ComputeGraphInstance.GetDataProviders())
+	if (ComputeGraphInstance.IsValid())
 	{
-		if (UPCGComputeDataProvider* PCGDataProvider = Cast<UPCGComputeDataProvider>(Provider.Get()))
+		for (TObjectPtr<UComputeDataProvider>& Provider : ComputeGraphInstance->GetDataProviders())
 		{
-			PCGDataProvider->ReleaseTransientResources(TEXT("~FPCGComputeGraphContext"));
+			if (UPCGComputeDataProvider* PCGDataProvider = Cast<UPCGComputeDataProvider>(Provider.Get()))
+			{
+				PCGDataProvider->ReleaseTransientResources(TEXT("~FPCGComputeGraphContext"));
+			}
 		}
-	}
 
-	if (bComputeGraphInstanceInitialized)
-	{
-		ComputeGraphInstance.ResetDataProviders();
-
-		// Return instance to pool.
-		if (const UPCGGraph* Graph = Stack ? Stack->GetNearestNonInlinedGraphForCurrentFrame() : nullptr; ensure(Graph))
+		if (bComputeGraphInstanceInitialized)
 		{
-			const UPCGGraph::FComputeGraphInstanceKey Key = { SourceComponent->GetGenerationGridSize(), ComputeGraphIndex };
-			Graph->ReturnComputeGraphInstanceToPool(Key, MoveTemp(ComputeGraphInstance));
+			ComputeGraphInstance->ResetDataProviders();
+
+			// Return instance to pool.
+			if (const UPCGGraph* Graph = Stack ? Stack->GetNearestNonInlinedGraphForCurrentFrame() : nullptr; ensure(Graph))
+			{
+				const UPCGGraph::FComputeGraphInstanceKey Key = { SourceComponent->GetGenerationGridSize(), ComputeGraphIndex };
+				Graph->ReturnComputeGraphInstanceToPool(Key, ComputeGraphInstance);
+			}
 		}
 	}
 }
@@ -60,7 +63,7 @@ void FPCGComputeGraphContext::AddExtraStructReferencedObjects(FReferenceCollecto
 {
 	if (bComputeGraphInstanceInitialized)
 	{
-		Collector.AddPropertyReferences(FComputeGraphInstance::StaticStruct(), &ComputeGraphInstance);
+		Collector.AddPropertyReferences(FComputeGraphInstance::StaticStruct(), ComputeGraphInstance.Get());
 	}
 
 	for (TObjectPtr<UComputeDataProvider> Provider : ProvidersWithBufferExports)
@@ -175,21 +178,29 @@ bool FPCGComputeGraphElement::ExecuteInternal(FPCGContext* InContext) const
 		const UPCGGraph* Graph = Context->Stack ? Context->Stack->GetNearestNonInlinedGraphForCurrentFrame() : nullptr;
 		const UPCGGraph::FComputeGraphInstanceKey Key = { Context->SourceComponent->GetGenerationGridSize(), ComputeGraphIndex };
 
+		if (!ensure(Graph))
+		{
+			return true;
+		}
+
+		bool bNewInstance = false;
+		Context->ComputeGraphInstance = Graph->RetrieveComputeGraphInstanceFromPool(Key, bNewInstance);
+
 		// The data provider initialization must not depend on the data binding being set up, although they still output of the PreInitialize such
 		// as the SourceComponent.
-		if (Graph && Graph->RetrieveComputeGraphInstanceFromPool(Key, Context->ComputeGraphInstance))
+		if (bNewInstance)
 		{
-			Context->ComputeGraphInstance.InitializeDataProviders(Context->ComputeGraph.Get(), 0, Context->DataBinding.Get());
+			Context->ComputeGraphInstance->CreateDataProviders(Context->ComputeGraph.Get(), 0, Context->DataBinding.Get());
 		}
 		else
 		{
-			Context->ComputeGraphInstance.CreateDataProviders(Context->ComputeGraph.Get(), 0, Context->DataBinding.Get());
+			Context->ComputeGraphInstance->InitializeDataProviders(Context->ComputeGraph.Get(), 0, Context->DataBinding.Get());
 		}
 
 		Context->bComputeGraphInstanceInitialized = true;
 
 		// Register all providers running async operations.
-		for (UComputeDataProvider* ComputeDataProvider : Context->ComputeGraphInstance.GetDataProviders())
+		for (UComputeDataProvider* ComputeDataProvider : Context->ComputeGraphInstance->GetDataProviders())
 		{
 			UPCGDataCollectionDataProvider* DataProvider = Cast<UPCGDataCollectionDataProvider>(ComputeDataProvider);
 			if (!DataProvider)
@@ -274,7 +285,7 @@ bool FPCGComputeGraphElement::ExecuteInternal(FPCGContext* InContext) const
 
 		bool bAllReady = true;
 
-		for (UComputeDataProvider* ComputeDataProvider : Context->ComputeGraphInstance.GetDataProviders())
+		for (UComputeDataProvider* ComputeDataProvider : Context->ComputeGraphInstance->GetDataProviders())
 		{
 			if (UPCGComputeDataProvider* DataProvider = Cast<UPCGComputeDataProvider>(ComputeDataProvider))
 			{
@@ -363,7 +374,7 @@ bool FPCGComputeGraphElement::ExecuteInternal(FPCGContext* InContext) const
 
 		bool bAllReady = true;
 
-		for (UComputeDataProvider* ComputeDataProvider : Context->ComputeGraphInstance.GetDataProviders())
+		for (UComputeDataProvider* ComputeDataProvider : Context->ComputeGraphInstance->GetDataProviders())
 		{
 			if (UPCGComputeDataProvider* DataProvider = Cast<UPCGComputeDataProvider>(ComputeDataProvider))
 			{
@@ -442,7 +453,7 @@ bool FPCGComputeGraphElement::ExecuteInternal(FPCGContext* InContext) const
 	{
 		TRACE_CPUPROFILER_EVENT_SCOPE(ScheduleComputeGraph);
 
-		const bool bGraphEnqueued = Context->ComputeGraphInstance.EnqueueWork(
+		const bool bGraphEnqueued = Context->ComputeGraphInstance->EnqueueWork(
 			Context->ComputeGraph.Get(),
 			InContext->SourceComponent->GetScene(),
 			ComputeTaskExecutionGroup::EndOfFrameUpdate,
diff --git a/Engine/Plugins/PCG/Source/PCG/Private/Compute/Elements/PCGComputeGraphElement.h b/Engine/Plugins/PCG/Source/PCG/Private/Compute/Elements/PCGComputeGraphElement.h
--- a/Engine/Plugins/PCG/Source/PCG/Private/Compute/Elements/PCGComputeGraphElement.h
+++ b/Engine/Plugins/PCG/Source/PCG/Private/Compute/Elements/PCGComputeGraphElement.h
@@ -46,7 +46,7 @@ public:
 	int32 ComputeGraphIndex = INDEX_NONE;
 
 	/** Data providers created from data interfaces and data bindings. */
-	FComputeGraphInstance ComputeGraphInstance;
+	TSharedPtr<FComputeGraphInstance> ComputeGraphInstance = nullptr;
 
 	/** Whether the compute graph instance has been initialized (data providers created, etc). */
 	bool bComputeGraphInstanceInitialized = false;
diff --git a/Engine/Plugins/PCG/Source/PCG/Private/PCGGraph.cpp b/Engine/Plugins/PCG/Source/PCG/Private/PCGGraph.cpp
--- a/Engine/Plugins/PCG/Source/PCG/Private/PCGGraph.cpp
+++ b/Engine/Plugins/PCG/Source/PCG/Private/PCGGraph.cpp
@@ -765,11 +765,11 @@ void UPCGGraph::AddReferencedObjects(UObject* InThis, FReferenceCollector& Colle
 {
 	UPCGGraph* This = CastChecked<UPCGGraph>(InThis);
 
-	for (TPair<FComputeGraphInstanceKey, TArray<FComputeGraphInstance>>& Entry : This->ComputeGraphInstancePool)
+	for (TPair<FComputeGraphInstanceKey, TArray<TSharedPtr<FComputeGraphInstance>>>& Entry : This->AllComputeGraphInstances)
 	{
-		for (FComputeGraphInstance& Instance : Entry.Value)
+		for (TSharedPtr<FComputeGraphInstance> Instance : Entry.Value)
 		{
-			Collector.AddPropertyReferences(FComputeGraphInstance::StaticStruct(), &Instance);
+			Collector.AddPropertyReferences(FComputeGraphInstance::StaticStruct(), Instance.Get());
 		}
 	}
 
@@ -786,34 +786,59 @@ uint32 UPCGGraph::GetDefaultGridSize() const
 	return PCGHiGenGrid::IsValidGrid(HiGenGridSize) ? (PCGHiGenGrid::GridToGridSize(HiGenGridSize) * (1 << HiGenExponential)) : PCGHiGenGrid::UnboundedGridSize();
 }
 
-bool UPCGGraph::RetrieveComputeGraphInstanceFromPool(const FComputeGraphInstanceKey& InKey, FComputeGraphInstance& OutInstance) const
+TSharedPtr<FComputeGraphInstance> UPCGGraph::RetrieveComputeGraphInstanceFromPool(const FComputeGraphInstanceKey& InKey, bool& bOutNewInstance) const
 {
 	check(IsInGameThread());
 
+	TSharedPtr<FComputeGraphInstance> Instance = nullptr;
+
 	if (PCGGraph::CVarEnableComputeGraphInstancePooling.GetValueOnGameThread())
 	{
-		TArray<FComputeGraphInstance>* Instances = ComputeGraphInstancePool.Find(InKey);
-		if (Instances && !Instances->IsEmpty())
+		TArray<TSharedPtr<FComputeGraphInstance>>& Instances = AvailableComputeGraphInstances.FindOrAdd(InKey);
+
+		if (!Instances.IsEmpty())
+		{
+			Instance = Instances.Pop(EAllowShrinking::No);
+			bOutNewInstance = false;
+		}
+		else
 		{
-			OutInstance = MoveTemp(Instances->Last());
-			Instances->RemoveAt(Instances->Num() - 1, EAllowShrinking::No);
+			// If we didn't find an instance, create a new one.
+			Instance = MakeShared<FComputeGraphInstance>();
+			bOutNewInstance = true;
 
-			return true;
+			AllComputeGraphInstances.FindOrAdd(InKey).Add(Instance);
 		}
 	}
-
-	return false;
+	else
+	{
+		Instance = MakeShared<FComputeGraphInstance>();
+		bOutNewInstance = true;
+	}
+	
+	return Instance;
 }
 
-void UPCGGraph::ReturnComputeGraphInstanceToPool(const FComputeGraphInstanceKey& InKey, FComputeGraphInstance InInstance) const
+void UPCGGraph::ReturnComputeGraphInstanceToPool(const FComputeGraphInstanceKey& InKey, TSharedPtr<FComputeGraphInstance>& InInstance) const
 {
 	check(IsInGameThread());
 
 	if (PCGGraph::CVarEnableComputeGraphInstancePooling.GetValueOnGameThread())
 	{
-		TArray<FComputeGraphInstance>& Instances = ComputeGraphInstancePool.FindOrAdd(InKey);
-		Instances.Add(MoveTemp(InInstance));
+		TArray<TSharedPtr<FComputeGraphInstance>>* AllInstances = AllComputeGraphInstances.Find(InKey);
+
+		// An instance can only be considered valid if it still exists in the pool of all instances. If it's no longer in the pool of all instances, then it
+		// must have been flushed and should now be considered stale.
+		const bool bValidInstance = AllInstances && AllInstances->Contains(InInstance);
+
+		if (bValidInstance)
+		{
+			TArray<TSharedPtr<FComputeGraphInstance>>& Instances = AvailableComputeGraphInstances.FindOrAdd(InKey);
+			Instances.Add(InInstance);
+		}
 	}
+
+	InInstance = nullptr;
 }
 
 UPCGNode* UPCGGraph::AddNodeOfType(TSubclassOf<class UPCGSettings> InSettingsClass, UPCGSettings*& OutDefaultNodeSettings)
@@ -1653,7 +1678,8 @@ void UPCGGraph::NotifyGraphChanged(EPCGChangeType ChangeType)
 
 		// Any non-trivial change to graph may change compute graphs.
 		check(IsInGameThread());
-		ComputeGraphInstancePool.Empty(ComputeGraphInstancePool.Num());
+		AvailableComputeGraphInstances.Empty(AvailableComputeGraphInstances.Num());
+		AllComputeGraphInstances.Empty(AllComputeGraphInstances.Num());
 	}
 
 	OnGraphChangedDelegate.Broadcast(this, ChangeType);
diff --git a/Engine/Plugins/PCG/Source/PCG/Public/PCGGraph.h b/Engine/Plugins/PCG/Source/PCG/Public/PCGGraph.h
--- a/Engine/Plugins/PCG/Source/PCG/Public/PCGGraph.h
+++ b/Engine/Plugins/PCG/Source/PCG/Public/PCGGraph.h
@@ -282,13 +282,13 @@ public:
 	bool Use2DGrid() const { return bUse2DGrid; }
 
 	using FComputeGraphInstanceKey = TPair</*GridSize*/uint32, /*Compute graph index*/int32>;
-	using FComputeGraphInstancePool = TMap<FComputeGraphInstanceKey, TArray<FComputeGraphInstance>>;
+	using FComputeGraphInstancePool = TMap<FComputeGraphInstanceKey, TArray<TSharedPtr<FComputeGraphInstance>>>;
 
-	/** Attempt to retrieve a pooled compute graph instance for the given key. Returns true on success. */
-	bool RetrieveComputeGraphInstanceFromPool(const FComputeGraphInstanceKey& InKey, FComputeGraphInstance& OutInstance) const;
+	/** Attempt to retrieve a pooled compute graph instance for the given key. Creates a new instance if one could not be found. */
+	TSharedPtr<FComputeGraphInstance> RetrieveComputeGraphInstanceFromPool(const FComputeGraphInstanceKey& InKey, bool& bOutNewInstance) const;
 
 	/** Places given compute graph instance into pool. No-ops if pooling is disabled. */
-	void ReturnComputeGraphInstanceToPool(const FComputeGraphInstanceKey& InKey, FComputeGraphInstance InInstance) const;
+	void ReturnComputeGraphInstanceToPool(const FComputeGraphInstanceKey& InKey, TSharedPtr<FComputeGraphInstance>& InInstance) const;
 
 #if WITH_EDITORONLY_DATA
 	UPROPERTY(BlueprintReadWrite, EditAnywhere, Category = AssetInfo, AssetRegistrySearchable)
@@ -542,7 +542,11 @@ public:
 	void UpdateUserParametersStruct(TFunctionRef<void(FInstancedPropertyBag&)> Callback);
 
 private:
-	mutable FComputeGraphInstancePool ComputeGraphInstancePool;
+	/** Pool of compute graph instances available for use. Grows when new instances are requested and when instances are returned to the pool. */
+	mutable FComputeGraphInstancePool AvailableComputeGraphInstances;
+
+	/** Used to track all valid compute graph instances that are alive for this graph. Grows monotonically until flushed. */
+	mutable FComputeGraphInstancePool AllComputeGraphInstances;
 
 #if WITH_EDITOR
 	/** Sends a change notification. Demotes change if the compiled tasks are not significantly changed. */
