[PCG] Re-added clearing the selection set when clicking the Graph Settings button

#rnx

diff --git a/Engine/Plugins/PCG/Source/PCGEditor/Private/PCGEditor.cpp b/Engine/Plugins/PCG/Source/PCGEditor/Private/PCGEditor.cpp
--- a/Engine/Plugins/PCG/Source/PCGEditor/Private/PCGEditor.cpp
+++ b/Engine/Plugins/PCG/Source/PCGEditor/Private/PCGEditor.cpp
@@ -1243,6 +1243,11 @@ void FPCGEditor::OnDeterminismGraphTest()
 
 void FPCGEditor::OnEditGraphSettings()
 {
+	check(GraphEditorWidget)
+
+	// Clear any selected nodes.
+	GraphEditorWidget->ClearSelectionSet();
+
 	for (TSharedPtr<SPCGEditorGraphDetailsView> PropertyDetailsWidget : PropertyDetailsWidgets)
 	{
 		PropertyDetailsWidget->SetObject(PCGGraphBeingEdited);
