[PCG] Fix crash in runtime gen when source component is gone.

#rnx

diff --git a/Engine/Plugins/PCG/Source/PCG/Private/Compute/Elements/PCGComputeGraphElement.cpp b/Engine/Plugins/PCG/Source/PCG/Private/Compute/Elements/PCGComputeGraphElement.cpp
--- a/Engine/Plugins/PCG/Source/PCG/Private/Compute/Elements/PCGComputeGraphElement.cpp
+++ b/Engine/Plugins/PCG/Source/PCG/Private/Compute/Elements/PCGComputeGraphElement.cpp
@@ -45,9 +45,10 @@ FPCGComputeGraphContext::~FPCGComputeGraphContext()
 			ComputeGraphInstance->ResetDataProviders();
 
 			// Return instance to pool.
-			if (const UPCGGraph* Graph = Stack ? Stack->GetNearestNonInlinedGraphForCurrentFrame() : nullptr; ensure(Graph))
+			const UPCGGraph* Graph = Stack ? Stack->GetNearestNonInlinedGraphForCurrentFrame() : nullptr;
+			if (ensure(Graph) && GenerationGridSize != PCGHiGenGrid::UninitializedGridSize())
 			{
-				const UPCGGraph::FComputeGraphInstanceKey Key = { SourceComponent->GetGenerationGridSize(), ComputeGraphIndex };
+				const UPCGGraph::FComputeGraphInstanceKey Key = { GenerationGridSize, ComputeGraphIndex };
 				Graph->ReturnComputeGraphInstanceToPool(Key, ComputeGraphInstance);
 			}
 		}
@@ -115,8 +116,6 @@ bool FPCGComputeGraphElement::ExecuteInternal(FPCGContext* InContext) const
 			return true;
 		}
 
-		uint32 GenerationGridSize = PCGHiGenGrid::UninitializedGridSize();
-
 		// Higen is always disabled within dynamic subgraphs - will retrieve Uninitialized tasks (which are always cooked).
 		const UPCGGraph* DynamicSubgraph = InContext->Stack ? InContext->Stack->GetNearestDynamicSubgraphForCurrentFrame() : nullptr;
 
@@ -124,13 +123,13 @@ bool FPCGComputeGraphElement::ExecuteInternal(FPCGContext* InContext) const
 		{
 			if (Component->IsLocalComponent() || Component->IsPartitioned())
 			{
-				GenerationGridSize = Component->GetGenerationGridSize();
+				Context->GenerationGridSize = Component->GetGenerationGridSize();
 			}
 		}
 
 		if (UPCGSubsystem* Subsystem = Context->SourceComponent->GetSubsystem())
 		{
-			Context->ComputeGraph.Reset(Subsystem->GetComputeGraph(DynamicSubgraph ? DynamicSubgraph : TopGraph, GenerationGridSize, ComputeGraphIndex));
+			Context->ComputeGraph.Reset(Subsystem->GetComputeGraph(DynamicSubgraph ? DynamicSubgraph : TopGraph, Context->GenerationGridSize, ComputeGraphIndex));
 		}
 
 		if (!ensure(Context->ComputeGraph))
@@ -176,7 +175,7 @@ bool FPCGComputeGraphElement::ExecuteInternal(FPCGContext* InContext) const
 		});
 
 		const UPCGGraph* Graph = Context->Stack ? Context->Stack->GetNearestNonInlinedGraphForCurrentFrame() : nullptr;
-		const UPCGGraph::FComputeGraphInstanceKey Key = { Context->SourceComponent->GetGenerationGridSize(), ComputeGraphIndex };
+		const UPCGGraph::FComputeGraphInstanceKey Key = { Context->GenerationGridSize, ComputeGraphIndex };
 
 		if (!ensure(Graph))
 		{
diff --git a/Engine/Plugins/PCG/Source/PCG/Private/Compute/Elements/PCGComputeGraphElement.h b/Engine/Plugins/PCG/Source/PCG/Private/Compute/Elements/PCGComputeGraphElement.h
--- a/Engine/Plugins/PCG/Source/PCG/Private/Compute/Elements/PCGComputeGraphElement.h
+++ b/Engine/Plugins/PCG/Source/PCG/Private/Compute/Elements/PCGComputeGraphElement.h
@@ -43,6 +43,8 @@ public:
 	TStrongObjectPtr<UPCGDataBinding> DataBinding = nullptr;
 	TStrongObjectPtr<UPCGComputeGraph> ComputeGraph = nullptr;
 
+	uint32 GenerationGridSize = PCGHiGenGrid::UninitializedGridSize();
+
 	int32 ComputeGraphIndex = INDEX_NONE;
 
 	/** Data providers created from data interfaces and data bindings. */
