[PCG] Fix crash when compute graph element context ultimately released from one of the async tasks, destroying it from a worker thread.

#rnx

diff --git a/Engine/Plugins/PCG/Source/PCG/Private/Compute/Elements/PCGComputeGraphElement.cpp b/Engine/Plugins/PCG/Source/PCG/Private/Compute/Elements/PCGComputeGraphElement.cpp
--- a/Engine/Plugins/PCG/Source/PCG/Private/Compute/Elements/PCGComputeGraphElement.cpp
+++ b/Engine/Plugins/PCG/Source/PCG/Private/Compute/Elements/PCGComputeGraphElement.cpp
@@ -13,6 +13,7 @@
 #include "Elements/PCGStaticMeshSpawnerKernel.h"
 #include "Editor/IPCGEditorModule.h"
 #include "Graph/PCGGraphCompiler.h"
+#include "Helpers/PCGHelpers.h"
 
 #include "ComputeWorkerInterface.h"
 #include "ComputeFramework/ComputeFramework.h"
@@ -25,34 +26,40 @@
 
 FPCGComputeGraphContext::~FPCGComputeGraphContext()
 {
-	if (DataBinding)
-	{
-		DataBinding->ReleaseTransientResources();
-	}
+	TWeakObjectPtr<const UPCGGraph> GraphWeak = Stack ? Stack->GetNearestNonInlinedGraphForCurrentFrame() : nullptr;
 
-	if (ComputeGraphInstance.IsValid())
+	// Can be here on a worker thread, after unpinning the context in one of the async tasks.
+	// Will execute immediately if on game thread, otherwise on next game thread frame.
+	PCGHelpers::ExecuteOnGameThread(TEXT("~FPCGComputeGraphContext"), [DataBinding=DataBinding, GraphInstance=ComputeGraphInstance, bInstanceInitialized=bComputeGraphInstanceInitialized, GraphWeak, GridSize=GenerationGridSize, GraphIndex=ComputeGraphIndex]
 	{
-		for (TObjectPtr<UComputeDataProvider>& Provider : ComputeGraphInstance->GetDataProviders())
+		if (DataBinding)
 		{
-			if (UPCGComputeDataProvider* PCGDataProvider = Cast<UPCGComputeDataProvider>(Provider.Get()))
-			{
-				PCGDataProvider->ReleaseTransientResources(TEXT("~FPCGComputeGraphContext"));
-			}
+			DataBinding->ReleaseTransientResources();
 		}
 
-		if (bComputeGraphInstanceInitialized)
+		if (GraphInstance.IsValid())
 		{
-			ComputeGraphInstance->ResetDataProviders();
+			for (TObjectPtr<UComputeDataProvider>& Provider : GraphInstance->GetDataProviders())
+			{
+				if (UPCGComputeDataProvider* PCGDataProvider = Cast<UPCGComputeDataProvider>(Provider.Get()))
+				{
+					PCGDataProvider->ReleaseTransientResources(TEXT("~FPCGComputeGraphContext"));
+				}
+			}
 
-			// Return instance to pool.
-			const UPCGGraph* Graph = Stack ? Stack->GetNearestNonInlinedGraphForCurrentFrame() : nullptr;
-			if (ensure(Graph) && GenerationGridSize != PCGHiGenGrid::UninitializedGridSize())
+			if (bInstanceInitialized)
 			{
-				const UPCGGraph::FComputeGraphInstanceKey Key = { GenerationGridSize, ComputeGraphIndex };
-				Graph->ReturnComputeGraphInstanceToPool(Key, ComputeGraphInstance);
+				GraphInstance->ResetDataProviders();
+
+				// Return instance to pool.
+				if (ensure(GraphWeak.IsValid()) && GridSize != PCGHiGenGrid::UninitializedGridSize())
+				{
+					const UPCGGraph::FComputeGraphInstanceKey Key = { GridSize, GraphIndex };
+					GraphWeak->ReturnComputeGraphInstanceToPool(Key, GraphInstance);
+				}
 			}
 		}
-	}
+	});
 }
 
 bool FPCGComputeGraphContext::HasPendingAsyncOperations() const
diff --git a/Engine/Plugins/PCG/Source/PCG/Private/PCGGraph.cpp b/Engine/Plugins/PCG/Source/PCG/Private/PCGGraph.cpp
--- a/Engine/Plugins/PCG/Source/PCG/Private/PCGGraph.cpp
+++ b/Engine/Plugins/PCG/Source/PCG/Private/PCGGraph.cpp
@@ -819,7 +819,7 @@ TSharedPtr<FComputeGraphInstance> UPCGGraph::RetrieveComputeGraphInstanceFromPoo
 	return Instance;
 }
 
-void UPCGGraph::ReturnComputeGraphInstanceToPool(const FComputeGraphInstanceKey& InKey, TSharedPtr<FComputeGraphInstance>& InInstance) const
+void UPCGGraph::ReturnComputeGraphInstanceToPool(const FComputeGraphInstanceKey& InKey, TSharedPtr<FComputeGraphInstance> InInstance) const
 {
 	check(IsInGameThread());
 
@@ -837,8 +837,6 @@ void UPCGGraph::ReturnComputeGraphInstanceToPool(const FComputeGraphInstanceKey&
 			Instances.Add(InInstance);
 		}
 	}
-
-	InInstance = nullptr;
 }
 
 UPCGNode* UPCGGraph::AddNodeOfType(TSubclassOf<class UPCGSettings> InSettingsClass, UPCGSettings*& OutDefaultNodeSettings)
diff --git a/Engine/Plugins/PCG/Source/PCG/Public/PCGGraph.h b/Engine/Plugins/PCG/Source/PCG/Public/PCGGraph.h
--- a/Engine/Plugins/PCG/Source/PCG/Public/PCGGraph.h
+++ b/Engine/Plugins/PCG/Source/PCG/Public/PCGGraph.h
@@ -288,7 +288,7 @@ public:
 	TSharedPtr<FComputeGraphInstance> RetrieveComputeGraphInstanceFromPool(const FComputeGraphInstanceKey& InKey, bool& bOutNewInstance) const;
 
 	/** Places given compute graph instance into pool. No-ops if pooling is disabled. */
-	void ReturnComputeGraphInstanceToPool(const FComputeGraphInstanceKey& InKey, TSharedPtr<FComputeGraphInstance>& InInstance) const;
+	void ReturnComputeGraphInstanceToPool(const FComputeGraphInstanceKey& InKey, TSharedPtr<FComputeGraphInstance> InInstance) const;
 
 #if WITH_EDITORONLY_DATA
 	UPROPERTY(BlueprintReadWrite, EditAnywhere, Category = AssetInfo, AssetRegistrySearchable)
