[PCG] Add facility to be able to release user-created render targets from within PCG more efficiently. Partially WIP.

diff --git a/Engine/Plugins/PCG/Source/PCG/Private/Data/PCGRenderTargetData.cpp b/Engine/Plugins/PCG/Source/PCG/Private/Data/PCGRenderTargetData.cpp
--- a/Engine/Plugins/PCG/Source/PCG/Private/Data/PCGRenderTargetData.cpp
+++ b/Engine/Plugins/PCG/Source/PCG/Private/Data/PCGRenderTargetData.cpp
@@ -8,11 +8,12 @@
 
 #include UE_INLINE_GENERATED_CPP_BY_NAME(PCGRenderTargetData)
 
-void UPCGRenderTargetData::Initialize(UTextureRenderTarget2D* InRenderTarget, const FTransform& InTransform, bool bInSkipReadbackToCPU)
+void UPCGRenderTargetData::Initialize(UTextureRenderTarget2D* InRenderTarget, const FTransform& InTransform, bool bInSkipReadbackToCPU, bool bInTakeOwnershipOfRenderTarget)
 {
 	RenderTarget = InRenderTarget;
 	Transform = InTransform;
 	bSkipReadbackToCPU = bInSkipReadbackToCPU;
+	bOwnsRenderTarget = bInTakeOwnershipOfRenderTarget;
 
 	ColorData.Reset();
 
@@ -31,6 +32,12 @@ void UPCGRenderTargetData::Initialize(UTextureRenderTarget2D* InRenderTarget, co
 		}
 	}
 
+	// Never take resource ownership on assets.
+	if (RenderTarget && RenderTarget->IsAsset())
+	{
+		bOwnsRenderTarget = false;
+	}
+
 	Width = !ColorData.IsEmpty() ? RenderTarget->SizeX : 0;
 	Height = !ColorData.IsEmpty() ? RenderTarget->SizeY : 0;
 
@@ -40,6 +47,16 @@ void UPCGRenderTargetData::Initialize(UTextureRenderTarget2D* InRenderTarget, co
 	Bounds = Bounds.TransformBy(Transform);
 }
 
+void UPCGRenderTargetData::ReleaseTransientResources(const TCHAR* InReason)
+{
+	if (bOwnsRenderTarget && RenderTarget)
+	{
+		RenderTarget->ReleaseResource();
+		RenderTarget = nullptr;
+		bOwnsRenderTarget = false;
+	}
+}
+
 void UPCGRenderTargetData::AddToCrc(FArchiveCrc32& Ar, bool bFullDataCrc) const
 {
 	Super::AddToCrc(Ar, bFullDataCrc);
@@ -62,7 +79,14 @@ UPCGSpatialData* UPCGRenderTargetData::CopyInternal(FPCGContext* Context) const
 
 	CopyBaseTextureData(NewRenderTargetData);
 
+	// TODO: We can't really support copying owned things at this point, so we will relinquish assumed ownership.
+	if (!ensure(!bOwnsRenderTarget))
+	{
+		const_cast<UPCGRenderTargetData*>(this)->bOwnsRenderTarget = false;
+	}
+
 	NewRenderTargetData->RenderTarget = RenderTarget;
+	NewRenderTargetData->bOwnsRenderTarget = false;
 
 	return NewRenderTargetData;
 }
diff --git a/Engine/Plugins/PCG/Source/PCG/Public/Data/PCGRenderTargetData.h b/Engine/Plugins/PCG/Source/PCG/Public/Data/PCGRenderTargetData.h
--- a/Engine/Plugins/PCG/Source/PCG/Public/Data/PCGRenderTargetData.h
+++ b/Engine/Plugins/PCG/Source/PCG/Public/Data/PCGRenderTargetData.h
@@ -21,6 +21,9 @@ public:
 	// ~Begin UPCGData interface
 	virtual EPCGDataType GetDataType() const override { return EPCGDataType::RenderTarget; }
 	virtual void AddToCrc(FArchiveCrc32& Ar, bool bFullDataCrc) const override;
+	virtual bool HoldsTransientResources() const override { return bOwnsRenderTarget; }
+	virtual bool IsCacheable() const { return Super::IsCacheable() && !bOwnsRenderTarget; }
+	virtual void ReleaseTransientResources(const TCHAR* InReason = nullptr) override;
 	// ~End UPCGData interface
 
 	//~Begin UPCGBaseTextureData interface
@@ -35,8 +38,11 @@ protected:
 
 public:
 	UFUNCTION(BlueprintCallable, Category = RenderTarget)
-	PCG_API void Initialize(UTextureRenderTarget2D* InRenderTarget, const FTransform& InTransform, bool bInSkipReadbackToCPU = false);
+	PCG_API void Initialize(UTextureRenderTarget2D* InRenderTarget, const FTransform& InTransform, bool bInSkipReadbackToCPU = false, bool bInTakeOwnershipOfRenderTarget = false);
 
 	UPROPERTY(BlueprintReadOnly, VisibleAnywhere, Category = Properties)
 	TObjectPtr<UTextureRenderTarget2D> RenderTarget = nullptr;
+
+	UPROPERTY(BlueprintReadOnly, VisibleAnywhere, Category = Properties)
+	bool bOwnsRenderTarget = false;
 };
