[PCG] Disable toast notification if we are only generating runtime gen components.

Fixes toast getting stuck on screen for large periods of time during streaming etc, and also reduces disturbing hitches when testing runtime gen.

#rnx

diff --git a/Engine/Plugins/PCG/Source/PCG/Private/Graph/PCGGraphExecutor.cpp b/Engine/Plugins/PCG/Source/PCG/Private/Graph/PCGGraphExecutor.cpp
--- a/Engine/Plugins/PCG/Source/PCG/Private/Graph/PCGGraphExecutor.cpp
+++ b/Engine/Plugins/PCG/Source/PCG/Private/Graph/PCGGraphExecutor.cpp
@@ -309,6 +309,10 @@ FPCGTaskId FPCGGraphExecutor::ScheduleGraph(const FPCGScheduleGraphParams& InPar
 	{ 
 		check(CompiledTasks[0].Node == InParams.Graph->GetInputNode());
 
+#if WITH_EDITOR
+		bAnyNonRuntimeGenComponentScheduled = bAnyNonRuntimeGenComponentScheduled || InParams.SourceComponent->GenerationTrigger != EPCGComponentGenerationTrigger::GenerateAtRuntime;
+#endif
+
 		// Setup fetch task on input node
 		CompiledTasks[0].Element = InParams.InputElement;
 
@@ -755,6 +759,10 @@ FPCGTaskId FPCGGraphExecutor::ScheduleGenericWithContext(TFunction<bool(FPCGCont
 
 FPCGTaskId FPCGGraphExecutor::ScheduleSingleElement(TSharedRef<IPCGElement> Element, bool bHasAbortCallbacks, UPCGComponent* InSourceComponent, const TArray<FPCGTaskId>& TaskExecutionDependencies, const TArray<FPCGTaskId>& TaskDataDependencies)
 {
+#if WITH_EDITOR
+	bAnyNonRuntimeGenComponentScheduled = bAnyNonRuntimeGenComponentScheduled || (InSourceComponent ? (InSourceComponent->GenerationTrigger != EPCGComponentGenerationTrigger::GenerateAtRuntime) : true);
+#endif
+
 	// Build task & element to hold the operation to perform
 	FPCGGraphTask Task;
 
@@ -966,6 +974,11 @@ void FPCGGraphExecutor::ExecuteTasksEnded()
 #endif
 
 			PCGGraphExecutionLogging::LogGraphExecuteFrameFinished();
+
+#if WITH_EDITOR
+			// Reset when we go idle. Will be set again next time any non-runtime-gen component is scheduled.
+			bAnyNonRuntimeGenComponentScheduled = false;
+#endif
 		}
 	}
 
@@ -2603,7 +2616,8 @@ void FPCGGraphExecutor::UpdateGenerationNotification()
 		return;
 	}
 
-	if (!GenerationProgressNotification.IsValid())
+	// Start timer for notification, if any non-runtime-gen component is generating.
+	if (!GenerationProgressNotification.IsValid() && bAnyNonRuntimeGenComponentScheduled)
 	{
 		// If we didn't have tasks on the prior frame, start a timer to create the new one to cause a delay and avoid spamming on and off for short tasks
 		if (GenerationProgressLastTaskNum == 0)
diff --git a/Engine/Plugins/PCG/Source/PCG/Private/Graph/PCGGraphExecutor.h b/Engine/Plugins/PCG/Source/PCG/Private/Graph/PCGGraphExecutor.h
--- a/Engine/Plugins/PCG/Source/PCG/Private/Graph/PCGGraphExecutor.h
+++ b/Engine/Plugins/PCG/Source/PCG/Private/Graph/PCGGraphExecutor.h
@@ -359,6 +359,9 @@ private:
 	TWeakPtr<IPCGEditorProgressNotification> GenerationProgressNotification;
 	double GenerationProgressNotificationStartTime = 0.0;
 	int32 GenerationProgressLastTaskNum = 0;
+
+	/** Set true when a non-runtime-gen component is scheduled and reset to false when all tasks are done executing. */
+	std::atomic<bool> bAnyNonRuntimeGenComponentScheduled = false;
 #endif
 
 	TObjectPtr<UWorld> World = nullptr;
