[PCG] Add deprecation for old clockwise algorithm in Spline Direction

* Algorithm changed for detecting if the spline is clockwise. The old one was working only on linear splines and was giving the inverse result.
* New algorithm works with all splines and gives the right result.
* Since it changes the behavior, old nodes flips the result of the new algorithm and throw a warning.
* To update, create a new node (don't copy or duplicate it)

diff --git a/Engine/Plugins/PCG/Source/PCG/Private/Elements/PCGSplineDirection.cpp b/Engine/Plugins/PCG/Source/PCG/Private/Elements/PCGSplineDirection.cpp
--- a/Engine/Plugins/PCG/Source/PCG/Private/Elements/PCGSplineDirection.cpp
+++ b/Engine/Plugins/PCG/Source/PCG/Private/Elements/PCGSplineDirection.cpp
@@ -5,6 +5,8 @@
 #include "PCGContext.h"
 #include "Data/PCGSplineData.h"
 
+#include "UObject/FortniteMainBranchObjectVersion.h"
+
 #define LOCTEXT_NAMESPACE "PCGSplineDirectionElement"
 
 namespace PCGSplineDirection
@@ -30,14 +32,11 @@ namespace PCGSplineDirection
 			return CrossProduct <= 0;
 		}
 
+		// Shoelace formula - https://en.wikipedia.org/wiki/Shoelace_formula.
+		// It's not clear the algorithm gives the right answer if we consider the spline not closed, so always consider the spline is
+		// closed for this.
 		int NumIter = NumPoints;
-    
-		if (InputSplineData->IsClosed())
-		{
-			NumIter++;
-		}
 
-		// Shoelace formula - https://en.wikipedia.org/wiki/Shoelace_formula.
 		for (int32 Index = 0; Index < NumIter; ++Index)
 		{
 			const FVector& ThisPoint = InputSplineData->SplineStruct.GetSplinePointsPosition().Points[Index % NumPoints].OutVal;
@@ -89,6 +88,19 @@ namespace PCGSplineDirection
 	}
 }
 
+void UPCGReverseSplineSettings::PostLoad()
+{
+	Super::PostLoad();
+	
+#if WITH_EDITOR
+	// UsingCustomVersion(FFortniteMainBranchObjectVersion::GUID) is located in parent class UPCGSettings::Serialize
+	if (GetLinkerCustomVersion(FFortniteMainBranchObjectVersion::GUID) < FFortniteMainBranchObjectVersion::PCGSplineDirectionClockwiseFix)
+	{
+		bFlipClockwiseComputationResult = true;
+	}
+#endif // WITH_EDITOR
+}
+
 #if WITH_EDITOR
 FName UPCGReverseSplineSettings::GetDefaultNodeName() const
 {
@@ -150,6 +162,9 @@ bool FPCGSplineDirectionElement::ExecuteInternal(FPCGContext* InContext) const
 	const UPCGReverseSplineSettings* Settings = InContext->GetInputSettings<UPCGReverseSplineSettings>();
 	check(Settings);
 
+	// Only warn for the deprecated algorithm once.
+	bool bHasWarned = false;
+
 	for (const FPCGTaggedData& InputData : InContext->InputData.GetInputsByPin(PCGPinConstants::DefaultInputLabel))
 	{
 		FPCGTaggedData& Output = InContext->OutputData.TaggedData.Emplace_GetRef(InputData);
@@ -161,18 +176,22 @@ bool FPCGSplineDirectionElement::ExecuteInternal(FPCGContext* InContext) const
 		}
 
 		bool bShouldReverse = true;
-		switch (Settings->Operation)
+
+		if (Settings->Operation == EPCGReverseSplineOperation::ForceClockwise || Settings->Operation == EPCGReverseSplineOperation::ForceCounterClockwise)
+		{
+			bShouldReverse = PCGSplineDirection::IsClockwiseXY(InputSplineData) != (Settings->Operation == EPCGReverseSplineOperation::ForceClockwise);
+			if (Settings->bFlipClockwiseComputationResult)
+			{
+				bShouldReverse = !bShouldReverse;
+				if (!bHasWarned)
+				{
+					PCGLog::LogWarningOnGraph(LOCTEXT("WarningUpdateClockwiseAlgorithm", "The clockwise detecting algorithm has been updated. Replace with a new copy of the node to remove this warning."), InContext);
+					bHasWarned = true;
+				}
+			}
+		}
+		else if (Settings->Operation != EPCGReverseSplineOperation::Reverse)
 		{
-		case EPCGReverseSplineOperation::ForceClockwise:
-			bShouldReverse = !PCGSplineDirection::IsClockwiseXY(InputSplineData);
-			break;
-		case EPCGReverseSplineOperation::ForceCounterClockwise:
-			bShouldReverse = PCGSplineDirection::IsClockwiseXY(InputSplineData);
-			break;
-		case EPCGReverseSplineOperation::Reverse:
-			bShouldReverse = true;
-			break;
-		default:
 			PCGLog::LogErrorOnGraph(LOCTEXT("InvalidOperation", "Invalid operation enum value"), InContext);
 			return true;
 		}
diff --git a/Engine/Plugins/PCG/Source/PCG/Public/Elements/PCGSplineDirection.h b/Engine/Plugins/PCG/Source/PCG/Public/Elements/PCGSplineDirection.h
--- a/Engine/Plugins/PCG/Source/PCG/Public/Elements/PCGSplineDirection.h
+++ b/Engine/Plugins/PCG/Source/PCG/Public/Elements/PCGSplineDirection.h
@@ -30,6 +30,8 @@ enum class EPCGReverseSplineOperation
 	ForceCounterClockwise
 };
 
+class FPCGSplineDirectionElement;
+
 /**
 * Direct the order of a spline's control points.
 * This can be conditional to force a given orientation (clockwise or counter clockwise).
@@ -37,9 +39,15 @@ enum class EPCGReverseSplineOperation
 UCLASS(MinimalAPI, BlueprintType, ClassGroup = (Procedural))
 class UPCGReverseSplineSettings : public UPCGSettings
 {
+	friend FPCGSplineDirectionElement;
+
 	GENERATED_BODY()
 
 public:
+	//~Begin UObject interface
+	virtual void PostLoad() override;
+	//~End UObject interface
+	
 	//~Begin UPCGSettings interface
 #if WITH_EDITOR
 	virtual FName GetDefaultNodeName() const override;
@@ -59,6 +67,12 @@ protected:
 public:
 	UPROPERTY(EditAnywhere, BlueprintReadWrite, Category = "Settings")
 	EPCGReverseSplineOperation Operation = EPCGReverseSplineOperation::Reverse;
+	
+private:
+	// UE_DEPRECATED(5.6, "Deprecation to be removed at some point")
+	/* Prior to 5.6, the CW/CCW were reversed. This is added to make sure we don't break data existing prior (e.g. from 5.5), but will throw a warning to be fixed. */
+	UPROPERTY()
+	bool bFlipClockwiseComputationResult = false;
 };
 
 class FPCGSplineDirectionElement : public IPCGElement
diff --git a/Engine/Source/Runtime/Core/Public/UObject/FortniteMainBranchObjectVersions.inl b/Engine/Source/Runtime/Core/Public/UObject/FortniteMainBranchObjectVersions.inl
--- a/Engine/Source/Runtime/Core/Public/UObject/FortniteMainBranchObjectVersions.inl
+++ b/Engine/Source/Runtime/Core/Public/UObject/FortniteMainBranchObjectVersions.inl
@@ -616,5 +616,8 @@ FORTNITE_MAIN_VERSION(AdditionalGameThreadDataSupportInChaosVisualDebugger, 200)
 // Upgrade UMG widget blueprints using legacy animation API
 FORTNITE_MAIN_VERSION(UpgradeWidgetBlueprintLegacySequencePlayer, 201)
 
+// Changed clockwise detection algorithm for PCGSplineDirection node with the correct one, but add a version to not break previous nodes.
+FORTNITE_MAIN_VERSION(PCGSplineDirectionClockwiseFix, 202)
+
 // -----<new versions can be added above this line>-------------------------------------------------
 #undef FORTNITE_MAIN_VERSION
\ No newline at end of file
