[PCG] Generate grass maps waits for landscape weight layer textures to be streamed in before generating.

Matches logic in the landscape grass system.

Also fix some missing references for TObjectPtrs.

#rnx

diff --git a/Engine/Plugins/PCG/Source/PCG/Private/Compute/Elements/PCGComputeGraphElement.cpp b/Engine/Plugins/PCG/Source/PCG/Private/Compute/Elements/PCGComputeGraphElement.cpp
--- a/Engine/Plugins/PCG/Source/PCG/Private/Compute/Elements/PCGComputeGraphElement.cpp
+++ b/Engine/Plugins/PCG/Source/PCG/Private/Compute/Elements/PCGComputeGraphElement.cpp
@@ -74,7 +74,7 @@ void FPCGComputeGraphContext::AddExtraStructReferencedObjects(FReferenceCollecto
 		Collector.AddPropertyReferences(FComputeGraphInstance::StaticStruct(), ComputeGraphInstance.Get());
 	}
 
-	for (TObjectPtr<UComputeDataProvider> Provider : ProvidersWithBufferExports)
+	for (TObjectPtr<UComputeDataProvider>& Provider : ProvidersWithBufferExports)
 	{
 		Collector.AddReferencedObject(Provider);
 	}
diff --git a/Engine/Plugins/PCG/Source/PCG/Private/Elements/PCGGenerateGrassMaps.cpp b/Engine/Plugins/PCG/Source/PCG/Private/Elements/PCGGenerateGrassMaps.cpp
--- a/Engine/Plugins/PCG/Source/PCG/Private/Elements/PCGGenerateGrassMaps.cpp
+++ b/Engine/Plugins/PCG/Source/PCG/Private/Elements/PCGGenerateGrassMaps.cpp
@@ -46,6 +46,16 @@ namespace PCGGenerateGrassMaps
 		false,
 		TEXT("Draws debug boxes around landscapes for which grass maps are generated, colored by the current task ID."));
 #endif
+
+	static bool IsTextureFullyStreamedIn(UTexture* InTexture)
+	{
+		const bool bCheckForLODTransition = true;
+		return InTexture &&
+#if WITH_EDITOR
+			!InTexture->IsDefaultTexture() &&
+#endif // WITH_EDITOR
+			!InTexture->HasPendingInitOrStreaming(bCheckForLODTransition) && InTexture->IsFullyStreamedIn();
+	}
 }
 
 TArray<FPCGPinProperties> UPCGGenerateGrassMapsSettings::InputPinProperties() const
@@ -109,13 +119,21 @@ FPCGGenerateGrassMapsContext::~FPCGGenerateGrassMapsContext()
 
 void FPCGGenerateGrassMapsContext::AddExtraStructReferencedObjects(FReferenceCollector& Collector)
 {
-	for (TObjectPtr<UPCGTextureData> Data : TextureDatas)
+	for (TObjectPtr<UPCGTextureData>& Data : TextureDatas)
 	{
 		if (Data)
 		{
 			Collector.AddReferencedObject(Data);
 		}
 	}
+
+	for (TObjectPtr<UTexture>& Texture : TexturesToStream)
+	{
+		if (Texture)
+		{
+			Collector.AddReferencedObject(Texture);
+		}
+	}
 }
 
 void FPCGGenerateGrassMapsElement::GetDependenciesCrc(const FPCGDataCollection& InInput, const UPCGSettings* InSettings, UPCGComponent* InComponent, FPCGCrc& OutCrc) const
@@ -145,7 +163,9 @@ bool FPCGGenerateGrassMapsElement::ExecuteInternal(FPCGContext* InContext) const
 	check(Context);
 
 	const UPCGGenerateGrassMapsSettings* Settings = Context->GetInputSettings<UPCGGenerateGrassMapsSettings>();
+	const UPCGComponent* SourceComponent = Context->SourceComponent.IsValid() ? Context->SourceComponent.Get() : nullptr;
 	check(Settings);
+	check(SourceComponent);
 
 	// 1. Select landscape components that overlap the given bounds.
 	if (!Context->bLandscapeComponentsFiltered)
@@ -172,9 +192,6 @@ bool FPCGGenerateGrassMapsElement::ExecuteInternal(FPCGContext* InContext) const
 			return true;
 		}
 
-		const UPCGComponent* SourceComponent = Context->SourceComponent.IsValid() ? Context->SourceComponent.Get() : nullptr;
-		check(SourceComponent);
-
 		if (!UE::Landscape::IsRuntimeGrassMapGenerationSupported() && (!SourceComponent->GetWorld() || SourceComponent->GetWorld()->IsGameWorld()))
 		{
 			UE_LOG(LogPCG, Warning, TEXT("Grass map generation is disabled outside of editor worlds. Try enabling the CVar 'grass.GrassMap.AlwaysBuildRuntimeGenerationResources'."));
@@ -342,6 +359,40 @@ bool FPCGGenerateGrassMapsElement::ExecuteInternal(FPCGContext* InContext) const
 	}
 
 	// 2. Wait for landscape components to be ready for grass map rendering.
+	if (!Context->bTextureStreamingRequested)
+	{
+		if (const UWorld* World = SourceComponent->GetWorld())
+		{
+			for (TWeakObjectPtr<ULandscapeComponent> LandscapeComponentWeak : Context->LandscapeComponents)
+			{
+				if (ULandscapeComponent* LandscapeComponent = LandscapeComponentWeak.Get())
+				{
+					// Make list of textures to stream before generating.
+					if (UTexture* HeightMap = LandscapeComponent->GetHeightmap())
+					{
+						Context->TexturesToStream.Add(HeightMap);
+					}
+
+					const ERHIFeatureLevel::Type FeatureLevel = World->GetFeatureLevel();
+					for (UTexture2D* WeightmapTexture : LandscapeComponent->GetRenderedWeightmapTexturesForFeatureLevel(FeatureLevel))
+					{
+						if (WeightmapTexture)
+						{
+							Context->TexturesToStream.Add(WeightmapTexture);
+						}
+					}
+				}
+			}
+
+			for (UTexture* Texture : Context->TexturesToStream)
+			{
+				Texture->bForceMiplevelsToBeResident = true;
+			}
+		}
+
+		Context->bTextureStreamingRequested = true;
+	}
+
 	if (!Context->bReadyToRender)
 	{
 		bool bAllReady = true;
@@ -358,6 +409,21 @@ bool FPCGGenerateGrassMapsElement::ExecuteInternal(FPCGContext* InContext) const
 			bAllReady &= UE::Landscape::CanRenderGrassMap(LandscapeComponent);
 		}
 
+		if (bAllReady)
+		{
+			for (UTexture* Texture : Context->TexturesToStream)
+			{
+				const bool bStreamedIn = PCGGenerateGrassMaps::IsTextureFullyStreamedIn(Texture);
+				bAllReady &= bStreamedIn;
+
+				if (!bStreamedIn)
+				{
+					UE_LOG(LogPCG, Verbose, TEXT("Waiting for landscape texture '%s' to stream in."), *Texture->GetName());
+					break;
+				}
+			}
+		}
+		
 		if (!bAllReady)
 		{
 			// Sleep until next frame, no use spinning on this.
diff --git a/Engine/Plugins/PCG/Source/PCG/Private/Graph/PCGGraphCompiler.cpp b/Engine/Plugins/PCG/Source/PCG/Private/Graph/PCGGraphCompiler.cpp
--- a/Engine/Plugins/PCG/Source/PCG/Private/Graph/PCGGraphCompiler.cpp
+++ b/Engine/Plugins/PCG/Source/PCG/Private/Graph/PCGGraphCompiler.cpp
@@ -982,11 +982,11 @@ void FPCGGraphCompiler::AddReferencedObjects(FReferenceCollector& Collector)
 
 	FReadScopeLock Lock(Cache.GraphToTaskMapLock);
 
-	for (const TPair<UPCGGraph*, TMap<uint32, TArray<TObjectPtr<UPCGComputeGraph>>>>& GraphToGridSizes : Cache.TopGraphToComputeGraphMap)
+	for (TPair<UPCGGraph*, TMap<uint32, TArray<TObjectPtr<UPCGComputeGraph>>>>& GraphToGridSizes : Cache.TopGraphToComputeGraphMap)
 	{
-		for (const TPair<uint32, TArray<TObjectPtr<UPCGComputeGraph>>>& GridSizeToComputeGraphs : GraphToGridSizes.Value)
+		for (TPair<uint32, TArray<TObjectPtr<UPCGComputeGraph>>>& GridSizeToComputeGraphs : GraphToGridSizes.Value)
 		{
-			for (TObjectPtr<UPCGComputeGraph> ComputeGraph : GridSizeToComputeGraphs.Value)
+			for (TObjectPtr<UPCGComputeGraph>& ComputeGraph : GridSizeToComputeGraphs.Value)
 			{
 				Collector.AddReferencedObject(ComputeGraph);
 			}
diff --git a/Engine/Plugins/PCG/Source/PCG/Public/Elements/PCGGenerateGrassMaps.h b/Engine/Plugins/PCG/Source/PCG/Public/Elements/PCGGenerateGrassMaps.h
--- a/Engine/Plugins/PCG/Source/PCG/Public/Elements/PCGGenerateGrassMaps.h
+++ b/Engine/Plugins/PCG/Source/PCG/Public/Elements/PCGGenerateGrassMaps.h
@@ -9,11 +9,12 @@
 
 #include "PCGGenerateGrassMaps.generated.h"
 
-class UPCGTextureData;
 class ALandscapeProxy;
 class FLandscapeGrassWeightExporter;
 class ULandscapeComponent;
 class ULandscapeGrassType;
+class UPCGTextureData;
+class UTexture;
 struct IPooledRenderTarget;
 
 UCLASS(MinimalAPI, BlueprintType, ClassGroup = (Procedural))
@@ -98,6 +99,12 @@ public:
 	/** True when the landscape components are ready for rendering. */
 	bool bReadyToRender = false;
 
+	/** Textures that we wait to be streamed before generating the grass maps. */
+	TArray<TObjectPtr<UTexture>> TexturesToStream;
+
+	/** True when streaming has been requested on landscape textures. */
+	bool bTextureStreamingRequested = false;
+
 	/** True when grass map generation has been scheduled on the render thread. */
 	bool bGenerationScheduled = false;
 
diff --git a/Engine/Source/Runtime/Landscape/Classes/LandscapeComponent.h b/Engine/Source/Runtime/Landscape/Classes/LandscapeComponent.h
--- a/Engine/Source/Runtime/Landscape/Classes/LandscapeComponent.h
+++ b/Engine/Source/Runtime/Landscape/Classes/LandscapeComponent.h
@@ -787,7 +787,7 @@ public:
 	LANDSCAPE_API const TArray<UTexture2D*>& GetWeightmapTextures(bool InReturnEditingWeightmap = false) const;
 	LANDSCAPE_API TArray<TObjectPtr<UTexture2D>>& GetWeightmapTextures(const FGuid& InLayerGuid);
 	LANDSCAPE_API const TArray<UTexture2D*>& GetWeightmapTextures(const FGuid& InLayerGuid) const;
-	const TArray<UTexture2D*>& GetRenderedWeightmapTexturesForFeatureLevel(ERHIFeatureLevel::Type FeatureLevel) const;
+	LANDSCAPE_API const TArray<UTexture2D*>& GetRenderedWeightmapTexturesForFeatureLevel(ERHIFeatureLevel::Type FeatureLevel) const;
 
 	LANDSCAPE_API TArray<FWeightmapLayerAllocationInfo>& GetWeightmapLayerAllocations(bool InReturnEditingWeightmap = false);
 	LANDSCAPE_API const TArray<FWeightmapLayerAllocationInfo>& GetWeightmapLayerAllocations(bool InReturnEditingWeightmap = false) const;
