[PCG] Async load assets used in GPU instancing.

#rnx

diff --git a/Engine/Plugins/PCG/Source/PCG/Private/Components/PCGProceduralISMComponent.cpp b/Engine/Plugins/PCG/Source/PCG/Private/Components/PCGProceduralISMComponent.cpp
--- a/Engine/Plugins/PCG/Source/PCG/Private/Components/PCGProceduralISMComponent.cpp
+++ b/Engine/Plugins/PCG/Source/PCG/Private/Components/PCGProceduralISMComponent.cpp
@@ -22,8 +22,8 @@ UPCGManagedProceduralISMComponent* PCGManagedProceduralISMComponent::GetOrCreate
 {
 	check(InTargetActor && InSourceComponent);
 
-	const UStaticMesh* StaticMesh = InParams.Descriptor.StaticMesh;
-	if (!StaticMesh)
+	const UStaticMesh* StaticMesh = InParams.Descriptor.StaticMesh.Get();
+	if (!ensure(StaticMesh))
 	{
 		return nullptr;
 	}
@@ -125,7 +125,7 @@ FPCGProceduralISMComponentDescriptor& FPCGProceduralISMComponentDescriptor::oper
 	InstanceStartCullDistance = Other.InstanceStartCullDistance;
 	InstanceEndCullDistance = Other.InstanceEndCullDistance;
 	OverlayMaterial = Other.OverlayMaterial.LoadSynchronous();
-	StaticMesh = Other.StaticMesh.LoadSynchronous();
+	StaticMesh = Other.StaticMesh;
 	ComponentTags = Other.ComponentTags;
 	Mobility = Other.Mobility;
 	VirtualTextureRenderPassType = Other.VirtualTextureRenderPassType;
@@ -168,7 +168,7 @@ FPCGProceduralISMComponentDescriptor& FPCGProceduralISMComponentDescriptor::oper
 
 void FPCGProceduralISMComponentDescriptor::InitComponent(UPCGProceduralISMComponent* InProceduralISMComponent) const
 {
-	InProceduralISMComponent->SetStaticMesh(StaticMesh);
+	InProceduralISMComponent->SetStaticMesh(StaticMesh.Get());
 	InProceduralISMComponent->OverrideMaterials = OverrideMaterials;
 	InProceduralISMComponent->OverlayMaterial = OverlayMaterial;
 	InProceduralISMComponent->RuntimeVirtualTextures = RuntimeVirtualTextures;
diff --git a/Engine/Plugins/PCG/Source/PCG/Private/Components/PCGProceduralISMComponent.h b/Engine/Plugins/PCG/Source/PCG/Private/Components/PCGProceduralISMComponent.h
--- a/Engine/Plugins/PCG/Source/PCG/Private/Components/PCGProceduralISMComponent.h
+++ b/Engine/Plugins/PCG/Source/PCG/Private/Components/PCGProceduralISMComponent.h
@@ -193,7 +193,7 @@ public:
 	mutable uint32 Hash = 0;
 
 	UPROPERTY()
-	TObjectPtr<UStaticMesh> StaticMesh = nullptr;
+	TSoftObjectPtr<UStaticMesh> StaticMesh = nullptr;
 
 	UPROPERTY()
 	TArray<TObjectPtr<UMaterialInterface>> OverrideMaterials;
diff --git a/Engine/Plugins/PCG/Source/PCG/Private/Compute/DataInterfaces/Elements/PCGStaticMeshSpawnerDataInterface.cpp b/Engine/Plugins/PCG/Source/PCG/Private/Compute/DataInterfaces/Elements/PCGStaticMeshSpawnerDataInterface.cpp
--- a/Engine/Plugins/PCG/Source/PCG/Private/Compute/DataInterfaces/Elements/PCGStaticMeshSpawnerDataInterface.cpp
+++ b/Engine/Plugins/PCG/Source/PCG/Private/Compute/DataInterfaces/Elements/PCGStaticMeshSpawnerDataInterface.cpp
@@ -356,6 +356,13 @@ bool UPCGStaticMeshSpawnerDataProvider::PrepareForExecute_GameThread(UPCGDataBin
 	{
 		CreatePrimitiveDescriptors(Context, InBinding);
 
+		for (const FPCGProceduralISMComponentDescriptor& PrimitiveDescriptor : PrimitiveDescriptors)
+		{
+			UE_LOG(LogPCG, Verbose, TEXT("Request '%s' to load."), *PrimitiveDescriptor.StaticMesh.ToString());
+
+			PrimitiveDescriptor.StaticMesh.LoadAsync({});
+		}
+
 		bPrimitiveDescriptorsCreated = true;
 
 		if (Context->AsyncState.ShouldStop())
@@ -364,6 +371,21 @@ bool UPCGStaticMeshSpawnerDataProvider::PrepareForExecute_GameThread(UPCGDataBin
 		}
 	}
 
+	if (!bStaticMeshesLoaded)
+	{
+		for (const FPCGProceduralISMComponentDescriptor& PrimitiveDescriptor : PrimitiveDescriptors)
+		{
+			if (PrimitiveDescriptor.StaticMesh.IsPending())
+			{
+				UE_LOG(LogPCG, Verbose, TEXT("Waiting for '%s' to load."), *PrimitiveDescriptor.StaticMesh.ToString());
+
+				return false;
+			}
+		}
+
+		bStaticMeshesLoaded = true;
+	}
+
 	if (!bPrimitivesSetUp)
 	{
 		if (!SetupPrimitives(Context, InBinding))
@@ -562,21 +584,12 @@ void UPCGStaticMeshSpawnerDataProvider::CreatePrimitiveDescriptors(FPCGContext*
 				break;
 			}
 
-			UStaticMesh* StaticMesh = Cast<UStaticMesh>(FSoftObjectPath(MeshPathString).TryLoad());
-			if (!StaticMesh)
-			{
-				PCG_KERNEL_VALIDATION_ERR(InContext, Settings,
-					FText::Format(PCGStaticMeshSpawnerDataInterface::CouldNotLoadStaticMeshFormat, FText::FromString(MeshPathString)));
-
-				continue;
-			}
-
 			FPCGProceduralISMComponentDescriptor Descriptor;
 			Descriptor = SelectorByAttribute->TemplateDescriptor;
 			Descriptor.NumInstances = StringKeyAndInstanceCount.Value;
 			Descriptor.LocalBounds = LocalBounds;
 			Descriptor.NumCustomFloats = CustomFloatCount;
-			Descriptor.StaticMesh = StaticMesh;
+			Descriptor.StaticMesh = FSoftObjectPath(MeshPathString);
 
 			// Sanity check instance count.
 			if (!ensure(Descriptor.NumInstances <= TotalInputPointCount))
@@ -586,11 +599,6 @@ void UPCGStaticMeshSpawnerDataProvider::CreatePrimitiveDescriptors(FPCGContext*
 
 			PrimitiveStringKeys.Add(StringKeyAndInstanceCount.Key);
 			PrimitiveDescriptors.Add(MoveTemp(Descriptor));
-
-			if (Settings->bApplyMeshBoundsToPoints)
-			{
-				PrimitiveMeshBounds.Add(Descriptor.StaticMesh->GetBoundingBox());
-			}
 		}
 
 		PrimitiveSelectionCDF.SetNumZeroed(PrimitiveDescriptors.Num());
@@ -635,7 +643,7 @@ void UPCGStaticMeshSpawnerDataProvider::CreatePrimitiveDescriptors(FPCGContext*
 				break;
 			}
 
-			if (UStaticMesh* StaticMesh = Entry.Descriptor.StaticMesh.LoadSynchronous())
+			if (!Entry.Descriptor.StaticMesh.IsNull())
 			{
 				const float Weight = float(Entry.Weight) / TotalWeight;
 				CumulativeWeight += Weight;
@@ -646,7 +654,7 @@ void UPCGStaticMeshSpawnerDataProvider::CreatePrimitiveDescriptors(FPCGContext*
 				Descriptor = Entry.Descriptor;
 				Descriptor.LocalBounds = LocalBounds;
 				Descriptor.NumCustomFloats = CustomFloatCount;
-				Descriptor.StaticMesh = StaticMesh;
+				Descriptor.StaticMesh = Entry.Descriptor.StaticMesh;
 
 				int32 InstanceCount = FMath::CeilToInt(TotalInputPointCount * Weight);
 
@@ -668,11 +676,6 @@ void UPCGStaticMeshSpawnerDataProvider::CreatePrimitiveDescriptors(FPCGContext*
 				Descriptor.NumInstances = FMath::Min(InstanceCount, TotalInputPointCount);
 
 				PrimitiveDescriptors.Add(MoveTemp(Descriptor));
-
-				if (Settings->bApplyMeshBoundsToPoints)
-				{
-					PrimitiveMeshBounds.Add(Descriptor.StaticMesh->GetBoundingBox());
-				}
 			}
 			else
 			{
@@ -731,6 +734,11 @@ bool UPCGStaticMeshSpawnerDataProvider::SetupPrimitives(FPCGContext* InContext,
 	{
 		const FPCGProceduralISMComponentDescriptor& Desc = PrimitiveDescriptors[NumPrimitivesSetup];
 
+		if (Settings->bApplyMeshBoundsToPoints)
+		{
+			PrimitiveMeshBounds.Add(Desc.StaticMesh->GetBoundingBox());
+		}
+
 		FPCGProceduralISMCBuilderParameters Params;
 		Params.Descriptor = Desc;
 		Params.bAllowDescriptorChanges = false;
diff --git a/Engine/Plugins/PCG/Source/PCG/Private/Compute/DataInterfaces/Elements/PCGStaticMeshSpawnerDataInterface.h b/Engine/Plugins/PCG/Source/PCG/Private/Compute/DataInterfaces/Elements/PCGStaticMeshSpawnerDataInterface.h
--- a/Engine/Plugins/PCG/Source/PCG/Private/Compute/DataInterfaces/Elements/PCGStaticMeshSpawnerDataInterface.h
+++ b/Engine/Plugins/PCG/Source/PCG/Private/Compute/DataInterfaces/Elements/PCGStaticMeshSpawnerDataInterface.h
@@ -91,6 +91,8 @@ public:
 
 	uint32 CustomFloatCount = 0;
 
+	bool bStaticMeshesLoaded = false;
+
 	bool bPrimitivesSetUp = false;
 
 	/** Number of primitives created during primitive setup. */
