[PCG] Fix crashes in new async load logic found in standalone build.

#rnx

diff --git a/Engine/Plugins/PCG/Source/PCG/Private/Compute/DataInterfaces/Elements/PCGStaticMeshSpawnerDataInterface.cpp b/Engine/Plugins/PCG/Source/PCG/Private/Compute/DataInterfaces/Elements/PCGStaticMeshSpawnerDataInterface.cpp
--- a/Engine/Plugins/PCG/Source/PCG/Private/Compute/DataInterfaces/Elements/PCGStaticMeshSpawnerDataInterface.cpp
+++ b/Engine/Plugins/PCG/Source/PCG/Private/Compute/DataInterfaces/Elements/PCGStaticMeshSpawnerDataInterface.cpp
@@ -454,7 +454,9 @@ void UPCGStaticMeshSpawnerDataProvider::Reset()
 	PrimitiveDescriptors.Empty();
 	CustomFloatCount = 0;
 	bPrimitivesSetUp = false;
+	NumPrimitivesProcessed = 0;
 	NumPrimitivesSetup = 0;
+	bStaticMeshesLoaded = false;
 
 	Super::Reset();
 }
@@ -708,7 +710,6 @@ void UPCGStaticMeshSpawnerDataProvider::CreatePrimitiveDescriptors(FPCGContext*
 	Primitives.SelectionCDF = PrimitiveSelectionCDF;
 	Primitives.PrimitiveStringKeys = PrimitiveStringKeys;
 	Primitives.PrimitiveMeshBounds = PrimitiveMeshBounds;
-	Primitives.NumInstancesAllPrimitives = static_cast<uint32>(TotalInstanceCount);
 
 	// Useful for debugging instance counts.
 	//UE_LOG(LogPCG, Log, TEXT("Input points %d, total instance count %d (+ %.2f%%)"), TotalInputPointCount, Primitives.NumInstancesAllPrimitives, (float(Primitives.NumInstancesAllPrimitives - TotalInputPointCount) / TotalInputPointCount) * 100.0f);
@@ -730,14 +731,11 @@ bool UPCGStaticMeshSpawnerDataProvider::SetupPrimitives(FPCGContext* InContext,
 		return true;
 	}
 
-	while (NumPrimitivesSetup < PrimitiveDescriptors.Num())
-	{
-		const FPCGProceduralISMComponentDescriptor& Desc = PrimitiveDescriptors[NumPrimitivesSetup];
+	Primitives->NumInstancesAllPrimitives = 0;
 
-		if (Settings->bApplyMeshBoundsToPoints)
-		{
-			PrimitiveMeshBounds.Add(Desc.StaticMesh->GetBoundingBox());
-		}
+	while (NumPrimitivesProcessed < PrimitiveDescriptors.Num())
+	{
+		const FPCGProceduralISMComponentDescriptor& Desc = PrimitiveDescriptors[NumPrimitivesProcessed];
 
 		FPCGProceduralISMCBuilderParameters Params;
 		Params.Descriptor = Desc;
@@ -745,17 +743,28 @@ bool UPCGStaticMeshSpawnerDataProvider::SetupPrimitives(FPCGContext* InContext,
 
 		UPCGManagedProceduralISMComponent* MISMC = PCGManagedProceduralISMComponent::GetOrCreateManagedProceduralISMC(TargetActor, InContext->SourceComponent.Get(), Settings->UID, Params);
 
-		check(MISMC);
-		MISMC->SetCrc(InContext->DependenciesCrc);
+		if (MISMC)
+		{
+			MISMC->SetCrc(InContext->DependenciesCrc);
+
+			// Don't bother registering the resource change as we're transient anyway.
+			//InContext->TouchedResources.Emplace(MISMC);
 
-		// Don't bother registering the resource change as we're transient anyway.
-		//InContext->TouchedResources.Emplace(MISMC);
+			Primitives->Primitives.Add(MISMC->GetComponent());
 
-		Primitives->Primitives.Add(MISMC->GetComponent());
+			if (Settings->bApplyMeshBoundsToPoints)
+			{
+				PrimitiveMeshBounds.Add(Desc.StaticMesh->GetBoundingBox());
+			}
+
+			Primitives->NumInstancesAllPrimitives += Desc.NumInstances;
+
+			++NumPrimitivesSetup;
+		}
 
-		++NumPrimitivesSetup;
+		++NumPrimitivesProcessed;
 
-		if (NumPrimitivesSetup < PrimitiveDescriptors.Num() && InContext->AsyncState.ShouldStop())
+		if (NumPrimitivesProcessed < PrimitiveDescriptors.Num() && InContext->AsyncState.ShouldStop())
 		{
 			return false;
 		}
diff --git a/Engine/Plugins/PCG/Source/PCG/Private/Compute/DataInterfaces/Elements/PCGStaticMeshSpawnerDataInterface.h b/Engine/Plugins/PCG/Source/PCG/Private/Compute/DataInterfaces/Elements/PCGStaticMeshSpawnerDataInterface.h
--- a/Engine/Plugins/PCG/Source/PCG/Private/Compute/DataInterfaces/Elements/PCGStaticMeshSpawnerDataInterface.h
+++ b/Engine/Plugins/PCG/Source/PCG/Private/Compute/DataInterfaces/Elements/PCGStaticMeshSpawnerDataInterface.h
@@ -96,6 +96,7 @@ public:
 	bool bPrimitivesSetUp = false;
 
 	/** Number of primitives created during primitive setup. */
+	int32 NumPrimitivesProcessed = 0;
 	int32 NumPrimitivesSetup = 0;
 };
 
