[LevelInstance] [PCG] Add param to LevelInstanceEditCancelled to inform if any changes were discarded
- PCG avoid reacting to even if new flag is false toi avoid refresh of PCG if user Edits/Cancels without doing any editing (any saved edits will already have been processed)

#rnx

diff --git a/Engine/Plugins/PCG/Source/PCG/Private/PCGActorAndComponentMapping.cpp b/Engine/Plugins/PCG/Source/PCG/Private/PCGActorAndComponentMapping.cpp
--- a/Engine/Plugins/PCG/Source/PCG/Private/PCGActorAndComponentMapping.cpp
+++ b/Engine/Plugins/PCG/Source/PCG/Private/PCGActorAndComponentMapping.cpp
@@ -1376,9 +1376,12 @@ void FPCGActorAndComponentMapping::OnLevelInstancesUpdated(const TArray<ILevelIn
 	}
 }
 
-void FPCGActorAndComponentMapping::OnLevelInstanceEditCancelled(ILevelInstanceInterface* InLevelInstance)
+void FPCGActorAndComponentMapping::OnLevelInstanceEditCancelled(ILevelInstanceInterface* InLevelInstance, bool bInHasDiscardedChanges)
 {
-	OnActorAdded(Cast<AActor>(InLevelInstance));
+	if (bInHasDiscardedChanges)
+	{
+		OnActorAdded(Cast<AActor>(InLevelInstance));
+	}
 }
 
 void FPCGActorAndComponentMapping::OnActorLoaded(AActor& InActor)
diff --git a/Engine/Plugins/PCG/Source/PCG/Public/PCGActorAndComponentMapping.h b/Engine/Plugins/PCG/Source/PCG/Public/PCGActorAndComponentMapping.h
--- a/Engine/Plugins/PCG/Source/PCG/Public/PCGActorAndComponentMapping.h
+++ b/Engine/Plugins/PCG/Source/PCG/Public/PCGActorAndComponentMapping.h
@@ -199,7 +199,7 @@ private:
 	void OnActorDeleted_Internal(AActor* InActor, bool bShouldDirty);
 	void OnLandscapeChanged(ALandscapeProxy* InLandscape, const FLandscapeProxyComponentDataChangedParams& InChangeParams);
 	void OnLevelInstancesUpdated(const TArray<ILevelInstanceInterface*>& InLevelInstances);
-	void OnLevelInstanceEditCancelled(ILevelInstanceInterface* InLevelInstance);
+	void OnLevelInstanceEditCancelled(ILevelInstanceInterface* InLevelInstance, bool bInHasDiscardedChanges);
 	void ApplyLandscapeChanges(ALandscapeProxy* InLandscape);
 	void OnObjectModified(UObject* InObject);
 	void OnObjectPropertyChanged(UObject* InObject, FPropertyChangedEvent& InEvent);
diff --git a/Engine/Source/Runtime/Engine/Private/LevelInstance/LevelInstanceSubsystem.cpp b/Engine/Source/Runtime/Engine/Private/LevelInstance/LevelInstanceSubsystem.cpp
--- a/Engine/Source/Runtime/Engine/Private/LevelInstance/LevelInstanceSubsystem.cpp
+++ b/Engine/Source/Runtime/Engine/Private/LevelInstance/LevelInstanceSubsystem.cpp
@@ -2769,7 +2769,7 @@ bool ULevelInstanceSubsystem::CommitLevelInstanceInternal(TUniquePtr<FLevelInsta
 	else
 	{
 		LevelInstance = GetLevelInstance(LevelInstanceID);
-		LevelInstanceEditCancelled.Broadcast(LevelInstance);
+		LevelInstanceEditCancelled.Broadcast(LevelInstance, /*bHasDiscardedChanges=*/!PackagesToSave.IsEmpty());
 	}
 
 	GEngine->BroadcastLevelActorListChanged();
diff --git a/Engine/Source/Runtime/Engine/Public/LevelInstance/LevelInstanceSubsystem.h b/Engine/Source/Runtime/Engine/Public/LevelInstance/LevelInstanceSubsystem.h
--- a/Engine/Source/Runtime/Engine/Public/LevelInstance/LevelInstanceSubsystem.h
+++ b/Engine/Source/Runtime/Engine/Public/LevelInstance/LevelInstanceSubsystem.h
@@ -169,7 +169,7 @@ public:
 	DECLARE_EVENT_OneParam(ULevelInstanceSubsystem, FLevelInstancesUpdated, const TArray<ILevelInstanceInterface*>&);
 	FLevelInstancesUpdated& OnLevelInstancesUpdated() { return LevelInstancesUpdatedEvent; }
 
-	DECLARE_EVENT_OneParam(ULevelInstanceSubsystem, FLevelInstanceEditCancelled, ILevelInstanceInterface*);
+	DECLARE_EVENT_TwoParams(ULevelInstanceSubsystem, FLevelInstanceEditCancelled, ILevelInstanceInterface*, bool);
 	FLevelInstanceEditCancelled& OnLevelInstanceEditCancelled() { return LevelInstanceEditCancelled; }
 
 	static ENGINE_API void ResetLoadersForWorldAsset(const FString& WorldAsset);
