[PCG] PCG Level to Asset add 'PCGLevelToAssetExclude' tag support will exclude actor and all its children (level instance, attachments)

diff --git a/Engine/Plugins/PCG/Source/PCG/Public/PCGDataAsset.h b/Engine/Plugins/PCG/Source/PCG/Public/PCGDataAsset.h
--- a/Engine/Plugins/PCG/Source/PCG/Public/PCGDataAsset.h
+++ b/Engine/Plugins/PCG/Source/PCG/Public/PCGDataAsset.h
@@ -17,6 +17,7 @@ namespace PCGLevelToAssetConstants
 	const FName RelativeTransformAttributeName = TEXT("RelativeTransform");
 	const FName IgnoreParentRotationAttributeName = TEXT("IgnoreParentRotation");
 	const FName IgnoreParentScaleAttributeName = TEXT("IgnoreParentScale");
+	const FName ExcludedActorTag = TEXT("PCGLevelToAssetExclude");
 }
 
 /** Container for PCG data exported as standalone objects */
diff --git a/Engine/Plugins/PCG/Source/PCGEditor/Private/PCGLevelToAsset.cpp b/Engine/Plugins/PCG/Source/PCGEditor/Private/PCGLevelToAsset.cpp
--- a/Engine/Plugins/PCG/Source/PCGEditor/Private/PCGLevelToAsset.cpp
+++ b/Engine/Plugins/PCG/Source/PCGEditor/Private/PCGLevelToAsset.cpp
@@ -305,6 +305,48 @@ bool UPCGLevelToAsset::BP_ExportWorld_Implementation(UWorld* World, const FStrin
 		}
 	};
 
+	auto IsActorExcluded = [LevelInstanceSubsystem, &GetAttachParentActor](AActor* InActor)
+	{
+		if (InActor->ActorHasTag(PCGLevelToAssetConstants::ExcludedActorTag))
+		{
+			return true;
+		}
+
+		// Check Level Instance Hierarchy first
+		if (InActor->IsInLevelInstance())
+		{
+			bool bExcluded = false;
+			LevelInstanceSubsystem->ForEachLevelInstanceAncestors(InActor, [&bExcluded](ILevelInstanceInterface* Ancestor)
+			{
+				AActor* AncestorActor = CastChecked<AActor>(Ancestor);
+				if (AncestorActor->ActorHasTag(PCGLevelToAssetConstants::ExcludedActorTag))
+				{
+					bExcluded = true;
+					return false;
+				}
+				return true;
+			});
+
+			if (bExcluded)
+			{
+				return true;
+			}
+		}
+
+		// Check Attachment next (which will go passed Level Instance Hierarchy)
+		AActor* AttachParent = GetAttachParentActor(InActor);
+		while (AttachParent)
+		{
+			if (AttachParent->ActorHasTag(PCGLevelToAssetConstants::ExcludedActorTag))
+			{
+				return true;
+			}
+			AttachParent = GetAttachParentActor(AttachParent);
+		}
+
+		return false;
+	};
+
 	// Parent tags are meant to be at the end of the array so they are applied last in the MakePoint lambda (we want the parent tag values to override the childrens)
 	auto GetActorTags = [LevelInstanceSubsystem](AActor* InActor, TArray<FName>& OutTags)
 	{
@@ -323,9 +365,17 @@ bool UPCGLevelToAsset::BP_ExportWorld_Implementation(UWorld* World, const FStrin
 
 	// Build actor-index map
 	TMap<AActor*, FActorEntry> ActorIndexMap;
+	TSet<AActor*> ExcludedActors;
+	
 	int LastActorIndex = 1; // Since the root is the "first" point we'll have, we'll have the map start from 1.
-	UPCGActorHelpers::ForEachActorInWorld(World, AActor::StaticClass(), [&ActorIndexMap, &LastActorIndex, &GetAttachParentActor, &GetActorTags](AActor* Actor)
+	UPCGActorHelpers::ForEachActorInWorld(World, AActor::StaticClass(), [&ActorIndexMap, &ExcludedActors, &LastActorIndex, &GetAttachParentActor, &GetActorTags, &IsActorExcluded](AActor* Actor)
 	{
+		if(IsActorExcluded(Actor))
+		{
+			ExcludedActors.Add(Actor);
+			return true;
+		}
+
 		FActorEntry ActorEntry(LastActorIndex++);
 		
 		ActorEntry.AttachmentParent = GetAttachParentActor(Actor);
@@ -341,6 +391,11 @@ bool UPCGLevelToAsset::BP_ExportWorld_Implementation(UWorld* World, const FStrin
 	{
 		// TODO Actor-level decisions if any; if the actor is "consumed" at this step, make sure to update AllActorBounds as well.
 
+		if (ExcludedActors.Contains(Actor))
+		{
+			return true;
+		}
+
 		// Otherwise, parse "known" actor components
 		TArray<UStaticMeshComponent*> SMCs;
 		Actor->GetComponents(SMCs);
