Fix proxy for UPCGProceduralISMComponent not being created when Nanite Fallback Mesh is stripped.
- don't check HasValidRenderData() and instead rely on the checks are already done in UStaticMeshComponent::CreateSceneProxy().

diff --git a/Engine/Plugins/PCG/Source/PCG/Private/Components/PCGProceduralISMComponent.cpp b/Engine/Plugins/PCG/Source/PCG/Private/Components/PCGProceduralISMComponent.cpp
--- a/Engine/Plugins/PCG/Source/PCG/Private/Components/PCGProceduralISMComponent.cpp
+++ b/Engine/Plugins/PCG/Source/PCG/Private/Components/PCGProceduralISMComponent.cpp
@@ -407,14 +407,7 @@ FPrimitiveSceneProxy* UPCGProceduralISMComponent::CreateStaticMeshSceneProxy(Nan
 
 FPrimitiveSceneProxy* UPCGProceduralISMComponent::CreateSceneProxy()
 {
-	// Verify that both mesh and instance data is valid before using it.
-	const bool bIsMeshAndInstanceDataValid =
-		NumInstances > 0 &&
-		GetStaticMesh() &&
-		GetStaticMesh()->IsCompiling() == false &&
-		GetStaticMesh()->HasValidRenderData();
-
-	if (!bIsMeshAndInstanceDataValid)
+	if (NumInstances <= 0)
 	{
 		return nullptr;
 	}
