[PCG] Level to Asset: add exception to use parent material if material instance is transient

#rnx

diff --git a/Engine/Plugins/PCG/Source/PCGEditor/Private/PCGLevelToAsset.cpp b/Engine/Plugins/PCG/Source/PCGEditor/Private/PCGLevelToAsset.cpp
--- a/Engine/Plugins/PCG/Source/PCGEditor/Private/PCGLevelToAsset.cpp
+++ b/Engine/Plugins/PCG/Source/PCGEditor/Private/PCGLevelToAsset.cpp
@@ -23,6 +23,7 @@
 #include "LevelInstance/LevelInstanceEditorInstanceActor.h"
 #include "LevelInstance/LevelInstanceInterface.h"
 #include "LevelInstance/LevelInstanceSubsystem.h"
+#include "Materials/MaterialInstance.h"
 #include "UObject/UObjectGlobals.h"
 #include "UObject/Package.h"
 #include "WorldPartition/WorldPartition.h"
@@ -478,7 +479,16 @@ bool UPCGLevelToAsset::BP_ExportWorld_Implementation(UWorld* World, const FStrin
 
 			if (!MeshMaterials.IsEmpty())
 			{
-				MaterialAttribute->SetValue(MetadataEntry, FSoftObjectPath(MeshMaterials[0]));
+				// Avoid references to transient materials, in this case we reference the first non transient parent
+				UMaterialInterface* Material = MeshMaterials[0];
+				UMaterialInstance* MaterialInstance = Cast<UMaterialInstance>(Material);
+				while (MaterialInstance && MaterialInstance->GetPackage() == GetTransientPackage())
+				{
+					Material = MaterialInstance->Parent;
+					MaterialInstance = Cast<UMaterialInstance>(Material);
+				}
+
+				MaterialAttribute->SetValue(MetadataEntry, FSoftObjectPath(Material));
 			}
 
 			ActorIndexAttribute->SetValue(MetadataEntry, ActorIndex);
