Revert async load feature, which fails in an odd way (some instances appear missing when pcg budget is set low like 1ms).

#rnx

diff --git a/Engine/Plugins/PCG/Source/PCG/Private/Components/PCGProceduralISMComponent.cpp b/Engine/Plugins/PCG/Source/PCG/Private/Components/PCGProceduralISMComponent.cpp
--- a/Engine/Plugins/PCG/Source/PCG/Private/Components/PCGProceduralISMComponent.cpp
+++ b/Engine/Plugins/PCG/Source/PCG/Private/Components/PCGProceduralISMComponent.cpp
@@ -22,8 +22,8 @@ UPCGManagedProceduralISMComponent* PCGManagedProceduralISMComponent::GetOrCreate
 {
 	check(InTargetActor && InSourceComponent);
 
-	const UStaticMesh* StaticMesh = InParams.Descriptor.StaticMesh.Get();
-	if (!ensure(StaticMesh))
+	const UStaticMesh* StaticMesh = InParams.Descriptor.StaticMesh;
+	if (!StaticMesh)
 	{
 		return nullptr;
 	}
@@ -125,7 +125,7 @@ FPCGProceduralISMComponentDescriptor& FPCGProceduralISMComponentDescriptor::oper
 	InstanceStartCullDistance = Other.InstanceStartCullDistance;
 	InstanceEndCullDistance = Other.InstanceEndCullDistance;
 	OverlayMaterial = Other.OverlayMaterial.LoadSynchronous();
-	StaticMesh = Other.StaticMesh;
+	StaticMesh = Other.StaticMesh.LoadSynchronous();
 	ComponentTags = Other.ComponentTags;
 	Mobility = Other.Mobility;
 	VirtualTextureRenderPassType = Other.VirtualTextureRenderPassType;
@@ -168,7 +168,7 @@ FPCGProceduralISMComponentDescriptor& FPCGProceduralISMComponentDescriptor::oper
 
 void FPCGProceduralISMComponentDescriptor::InitComponent(UPCGProceduralISMComponent* InProceduralISMComponent) const
 {
-	InProceduralISMComponent->SetStaticMesh(StaticMesh.Get());
+	InProceduralISMComponent->SetStaticMesh(StaticMesh);
 	InProceduralISMComponent->OverrideMaterials = OverrideMaterials;
 	InProceduralISMComponent->OverlayMaterial = OverlayMaterial;
 	InProceduralISMComponent->RuntimeVirtualTextures = RuntimeVirtualTextures;
diff --git a/Engine/Plugins/PCG/Source/PCG/Private/Components/PCGProceduralISMComponent.h b/Engine/Plugins/PCG/Source/PCG/Private/Components/PCGProceduralISMComponent.h
--- a/Engine/Plugins/PCG/Source/PCG/Private/Components/PCGProceduralISMComponent.h
+++ b/Engine/Plugins/PCG/Source/PCG/Private/Components/PCGProceduralISMComponent.h
@@ -193,7 +193,7 @@ public:
 	mutable uint32 Hash = 0;
 
 	UPROPERTY()
-	TSoftObjectPtr<UStaticMesh> StaticMesh = nullptr;
+	TObjectPtr<UStaticMesh> StaticMesh = nullptr;
 
 	UPROPERTY()
 	TArray<TObjectPtr<UMaterialInterface>> OverrideMaterials;
diff --git a/Engine/Plugins/PCG/Source/PCG/Private/Compute/DataInterfaces/Elements/PCGStaticMeshSpawnerDataInterface.cpp b/Engine/Plugins/PCG/Source/PCG/Private/Compute/DataInterfaces/Elements/PCGStaticMeshSpawnerDataInterface.cpp
--- a/Engine/Plugins/PCG/Source/PCG/Private/Compute/DataInterfaces/Elements/PCGStaticMeshSpawnerDataInterface.cpp
+++ b/Engine/Plugins/PCG/Source/PCG/Private/Compute/DataInterfaces/Elements/PCGStaticMeshSpawnerDataInterface.cpp
@@ -356,13 +356,6 @@ bool UPCGStaticMeshSpawnerDataProvider::PrepareForExecute_GameThread(UPCGDataBin
 	{
 		CreatePrimitiveDescriptors(Context, InBinding);
 
-		for (const FPCGProceduralISMComponentDescriptor& PrimitiveDescriptor : PrimitiveDescriptors)
-		{
-			UE_LOG(LogPCG, Verbose, TEXT("Request '%s' to load."), *PrimitiveDescriptor.StaticMesh.ToString());
-
-			PrimitiveDescriptor.StaticMesh.LoadAsync({});
-		}
-
 		bPrimitiveDescriptorsCreated = true;
 
 		if (Context->AsyncState.ShouldStop())
@@ -371,21 +364,6 @@ bool UPCGStaticMeshSpawnerDataProvider::PrepareForExecute_GameThread(UPCGDataBin
 		}
 	}
 
-	if (!bStaticMeshesLoaded)
-	{
-		for (const FPCGProceduralISMComponentDescriptor& PrimitiveDescriptor : PrimitiveDescriptors)
-		{
-			if (PrimitiveDescriptor.StaticMesh.IsPending())
-			{
-				UE_LOG(LogPCG, Verbose, TEXT("Waiting for '%s' to load."), *PrimitiveDescriptor.StaticMesh.ToString());
-
-				return false;
-			}
-		}
-
-		bStaticMeshesLoaded = true;
-	}
-
 	if (!bPrimitivesSetUp)
 	{
 		if (!SetupPrimitives(Context, InBinding))
@@ -454,9 +432,7 @@ void UPCGStaticMeshSpawnerDataProvider::Reset()
 	PrimitiveDescriptors.Empty();
 	CustomFloatCount = 0;
 	bPrimitivesSetUp = false;
-	NumPrimitivesProcessed = 0;
 	NumPrimitivesSetup = 0;
-	bStaticMeshesLoaded = false;
 
 	Super::Reset();
 }
@@ -586,12 +562,21 @@ void UPCGStaticMeshSpawnerDataProvider::CreatePrimitiveDescriptors(FPCGContext*
 				break;
 			}
 
+			UStaticMesh* StaticMesh = Cast<UStaticMesh>(FSoftObjectPath(MeshPathString).TryLoad());
+			if (!StaticMesh)
+			{
+				PCG_KERNEL_VALIDATION_ERR(InContext, Settings,
+					FText::Format(PCGStaticMeshSpawnerDataInterface::CouldNotLoadStaticMeshFormat, FText::FromString(MeshPathString)));
+
+				continue;
+			}
+
 			FPCGProceduralISMComponentDescriptor Descriptor;
 			Descriptor = SelectorByAttribute->TemplateDescriptor;
 			Descriptor.NumInstances = StringKeyAndInstanceCount.Value;
 			Descriptor.LocalBounds = LocalBounds;
 			Descriptor.NumCustomFloats = CustomFloatCount;
-			Descriptor.StaticMesh = FSoftObjectPath(MeshPathString);
+			Descriptor.StaticMesh = StaticMesh;
 
 			// Sanity check instance count.
 			if (!ensure(Descriptor.NumInstances <= TotalInputPointCount))
@@ -601,6 +586,11 @@ void UPCGStaticMeshSpawnerDataProvider::CreatePrimitiveDescriptors(FPCGContext*
 
 			PrimitiveStringKeys.Add(StringKeyAndInstanceCount.Key);
 			PrimitiveDescriptors.Add(MoveTemp(Descriptor));
+
+			if (Settings->bApplyMeshBoundsToPoints)
+			{
+				PrimitiveMeshBounds.Add(Descriptor.StaticMesh->GetBoundingBox());
+			}
 		}
 
 		PrimitiveSelectionCDF.SetNumZeroed(PrimitiveDescriptors.Num());
@@ -645,7 +635,7 @@ void UPCGStaticMeshSpawnerDataProvider::CreatePrimitiveDescriptors(FPCGContext*
 				break;
 			}
 
-			if (!Entry.Descriptor.StaticMesh.IsNull())
+			if (UStaticMesh* StaticMesh = Entry.Descriptor.StaticMesh.LoadSynchronous())
 			{
 				const float Weight = float(Entry.Weight) / TotalWeight;
 				CumulativeWeight += Weight;
@@ -656,7 +646,7 @@ void UPCGStaticMeshSpawnerDataProvider::CreatePrimitiveDescriptors(FPCGContext*
 				Descriptor = Entry.Descriptor;
 				Descriptor.LocalBounds = LocalBounds;
 				Descriptor.NumCustomFloats = CustomFloatCount;
-				Descriptor.StaticMesh = Entry.Descriptor.StaticMesh;
+				Descriptor.StaticMesh = StaticMesh;
 
 				int32 InstanceCount = FMath::CeilToInt(TotalInputPointCount * Weight);
 
@@ -678,6 +668,11 @@ void UPCGStaticMeshSpawnerDataProvider::CreatePrimitiveDescriptors(FPCGContext*
 				Descriptor.NumInstances = FMath::Min(InstanceCount, TotalInputPointCount);
 
 				PrimitiveDescriptors.Add(MoveTemp(Descriptor));
+
+				if (Settings->bApplyMeshBoundsToPoints)
+				{
+					PrimitiveMeshBounds.Add(Descriptor.StaticMesh->GetBoundingBox());
+				}
 			}
 			else
 			{
@@ -710,6 +705,7 @@ void UPCGStaticMeshSpawnerDataProvider::CreatePrimitiveDescriptors(FPCGContext*
 	Primitives.SelectionCDF = PrimitiveSelectionCDF;
 	Primitives.PrimitiveStringKeys = PrimitiveStringKeys;
 	Primitives.PrimitiveMeshBounds = PrimitiveMeshBounds;
+	Primitives.NumInstancesAllPrimitives = static_cast<uint32>(TotalInstanceCount);
 
 	// Useful for debugging instance counts.
 	//UE_LOG(LogPCG, Log, TEXT("Input points %d, total instance count %d (+ %.2f%%)"), TotalInputPointCount, Primitives.NumInstancesAllPrimitives, (float(Primitives.NumInstancesAllPrimitives - TotalInputPointCount) / TotalInputPointCount) * 100.0f);
@@ -731,11 +727,9 @@ bool UPCGStaticMeshSpawnerDataProvider::SetupPrimitives(FPCGContext* InContext,
 		return true;
 	}
 
-	Primitives->NumInstancesAllPrimitives = 0;
-
-	while (NumPrimitivesProcessed < PrimitiveDescriptors.Num())
+	while (NumPrimitivesSetup < PrimitiveDescriptors.Num())
 	{
-		const FPCGProceduralISMComponentDescriptor& Desc = PrimitiveDescriptors[NumPrimitivesProcessed];
+		const FPCGProceduralISMComponentDescriptor& Desc = PrimitiveDescriptors[NumPrimitivesSetup];
 
 		FPCGProceduralISMCBuilderParameters Params;
 		Params.Descriptor = Desc;
@@ -743,28 +737,17 @@ bool UPCGStaticMeshSpawnerDataProvider::SetupPrimitives(FPCGContext* InContext,
 
 		UPCGManagedProceduralISMComponent* MISMC = PCGManagedProceduralISMComponent::GetOrCreateManagedProceduralISMC(TargetActor, InContext->SourceComponent.Get(), Settings->UID, Params);
 
-		if (MISMC)
-		{
-			MISMC->SetCrc(InContext->DependenciesCrc);
-
-			// Don't bother registering the resource change as we're transient anyway.
-			//InContext->TouchedResources.Emplace(MISMC);
+		check(MISMC);
+		MISMC->SetCrc(InContext->DependenciesCrc);
 
-			Primitives->Primitives.Add(MISMC->GetComponent());
+		// Don't bother registering the resource change as we're transient anyway.
+		//InContext->TouchedResources.Emplace(MISMC);
 
-			if (Settings->bApplyMeshBoundsToPoints)
-			{
-				PrimitiveMeshBounds.Add(Desc.StaticMesh->GetBoundingBox());
-			}
-
-			Primitives->NumInstancesAllPrimitives += Desc.NumInstances;
-
-			++NumPrimitivesSetup;
-		}
+		Primitives->Primitives.Add(MISMC->GetComponent());
 
-		++NumPrimitivesProcessed;
+		++NumPrimitivesSetup;
 
-		if (NumPrimitivesProcessed < PrimitiveDescriptors.Num() && InContext->AsyncState.ShouldStop())
+		if (NumPrimitivesSetup < PrimitiveDescriptors.Num() && InContext->AsyncState.ShouldStop())
 		{
 			return false;
 		}
diff --git a/Engine/Plugins/PCG/Source/PCG/Private/Compute/DataInterfaces/Elements/PCGStaticMeshSpawnerDataInterface.h b/Engine/Plugins/PCG/Source/PCG/Private/Compute/DataInterfaces/Elements/PCGStaticMeshSpawnerDataInterface.h
--- a/Engine/Plugins/PCG/Source/PCG/Private/Compute/DataInterfaces/Elements/PCGStaticMeshSpawnerDataInterface.h
+++ b/Engine/Plugins/PCG/Source/PCG/Private/Compute/DataInterfaces/Elements/PCGStaticMeshSpawnerDataInterface.h
@@ -91,12 +91,9 @@ public:
 
 	uint32 CustomFloatCount = 0;
 
-	bool bStaticMeshesLoaded = false;
-
 	bool bPrimitivesSetUp = false;
 
 	/** Number of primitives created during primitive setup. */
-	int32 NumPrimitivesProcessed = 0;
 	int32 NumPrimitivesSetup = 0;
 };
 
