[PCG] Refactor Initialize Metadata and domain

* Can now specify a specific mapping for attribute initialization
* Create new structures for Metadata and Metadata domain for parameters
* Deprecated some Metadata function
* No need for deprecation in Metadata Domain as it is not yet released.
* Added some support for converting array/arrayviews to pcg value range and use them in Attribute getter functions
* Cleaned-up the the Point to Attribute set temporary fix

#rnx

diff --git a/Engine/Plugins/PCG/Source/PCG/Private/Data/PCGSpatialData.cpp b/Engine/Plugins/PCG/Source/PCG/Private/Data/PCGSpatialData.cpp
--- a/Engine/Plugins/PCG/Source/PCG/Private/Data/PCGSpatialData.cpp
+++ b/Engine/Plugins/PCG/Source/PCG/Private/Data/PCGSpatialData.cpp
@@ -333,7 +333,11 @@ void UPCGSpatialData::InitializeOtherMetadataInternal(const FPCGInitializeFromDa
 	check(InParams.bInheritMetadata);
 	check(MetadataToInitialize);
 
-	MetadataToInitialize->Initialize(Metadata, InParams.bInheritAttributes, InParams.DomainMapping);
+	FPCGMetadataInitializeParams InitializeParams(Metadata, TSet<FName>(), InParams.bInheritAttributes ? EPCGMetadataFilterMode::ExcludeAttributes : EPCGMetadataFilterMode::IncludeAttributes);
+	InitializeParams.DomainMapping.Reserve(InParams.DomainMapping.Num());
+	Algo::Transform(InParams.DomainMapping, InitializeParams.DomainMapping, [](const auto& It) { return It; });
+
+	MetadataToInitialize->Initialize(InitializeParams);
 }
 
 void UPCGSpatialData::InitializeSpatialDataInternal(const FPCGInitializeFromDataParams& InParams)
@@ -391,6 +395,7 @@ UPCGSpatialData* UPCGSpatialData::DuplicateData(FPCGContext* Context, bool bInit
 	if (bInitializeMetadata)
 	{
 		FPCGInitializeFromDataParams Params(this);
+		Params.bIsDuplicatingData = true;
 		InitializeOtherMetadata(Params, NewSpatialData->Metadata);
 	}
 
diff --git a/Engine/Plugins/PCG/Source/PCG/Private/Elements/PCGConvertToAttributeSet.cpp b/Engine/Plugins/PCG/Source/PCG/Private/Elements/PCGConvertToAttributeSet.cpp
--- a/Engine/Plugins/PCG/Source/PCG/Private/Elements/PCGConvertToAttributeSet.cpp
+++ b/Engine/Plugins/PCG/Source/PCG/Private/Elements/PCGConvertToAttributeSet.cpp
@@ -64,33 +64,18 @@ bool FPCGConvertToAttributeSetElement::ExecuteInternal(FPCGContext* Context) con
 
 		UPCGParamData* ParamData = FPCGContext::NewObject_AnyThread<UPCGParamData>(Context);
 		UPCGMetadata* Metadata = ParamData->MutableMetadata();
-
 		check(Metadata);
 
-		// Copy all the global metadata, if it exists.
-		if (const FPCGMetadataDomain* SourceGlobalDomain = SourceMetadata->GetConstMetadataDomain(PCGMetadataDomainID::Data))
-		{
-			FPCGMetadataDomain* GlobalDomain = Metadata->GetMetadataDomain(PCGMetadataDomainID::Data);
-			check(GlobalDomain);
-			GlobalDomain->InitializeAsCopy(SourceGlobalDomain);
-		}
+		FPCGMetadataInitializeParams Params(SourceMetadata);
+		FPCGMetadataDomainInitializeParams ElementsDomainParams(SourceMetadata->GetConstMetadataDomain(PCGMetadataDomainID::Elements));
 
-		// Then copy all the elements, if it exists.
-		if (const FPCGMetadataDomain* SourceElementsDomain = SourceMetadata->GetConstMetadataDomain(PCGMetadataDomainID::Elements))
-		{
-			// Gather all the keys to copy over.
-			TArray<PCGMetadataEntryKey> EntryKeys;
+		const TConstPCGValueRange<int64> MetadataEntryRange = PointData->GetConstMetadataEntryValueRange();
+		check(PointData->GetNumPoints() == MetadataEntryRange.Num());
+		ElementsDomainParams.OptionalEntriesToCopy = MetadataEntryRange;
 
-			const TConstPCGValueRange<int64> MetadataEntryRange = PointData->GetConstMetadataEntryValueRange();
-			check(PointData->GetNumPoints() == MetadataEntryRange.Num());
-			EntryKeys.Reserve(MetadataEntryRange.Num());
-			Algo::Transform(MetadataEntryRange, EntryKeys, [](int64 EntryKey) -> PCGMetadataEntryKey { return EntryKey; });
+		Params.DomainInitializeParams.Emplace(PCGMetadataDomainID::Elements, std::move(ElementsDomainParams));
+		Metadata->InitializeAsCopy(Params);
 
-			FPCGMetadataDomain* ElementsDomain = Metadata->GetMetadataDomain(PCGMetadataDomainID::Elements);
-			check(ElementsDomain);
-			ElementsDomain->InitializeAsCopy(SourceElementsDomain, &EntryKeys);
-		}
-		
 		FPCGTaggedData& Output = Outputs.Emplace_GetRef(Input);
 		Output.Data = ParamData;
 	}
diff --git a/Engine/Plugins/PCG/Source/PCG/Private/Elements/PCGDeleteAttributesElement.cpp b/Engine/Plugins/PCG/Source/PCG/Private/Elements/PCGDeleteAttributesElement.cpp
--- a/Engine/Plugins/PCG/Source/PCG/Private/Elements/PCGDeleteAttributesElement.cpp
+++ b/Engine/Plugins/PCG/Source/PCG/Private/Elements/PCGDeleteAttributesElement.cpp
@@ -183,7 +183,7 @@ bool FPCGDeleteAttributesElement::ExecuteInternal(FPCGContext* Context) const
 
 			UPCGParamData* NewParamData = FPCGContext::NewObject_AnyThread<UPCGParamData>(Context);
 			Metadata = NewParamData->Metadata;
-			Metadata->InitializeAsCopyWithAttributeFilter(ParentMetadata, AttributesToFilter, FilterMode);
+			Metadata->InitializeAsCopy(FPCGMetadataInitializeParams(ParentMetadata, AttributesToFilter, FilterMode));
 
 			OutputData = NewParamData;
 		}
diff --git a/Engine/Plugins/PCG/Source/PCG/Private/Elements/PCGFilterElementsByIndex.cpp b/Engine/Plugins/PCG/Source/PCG/Private/Elements/PCGFilterElementsByIndex.cpp
--- a/Engine/Plugins/PCG/Source/PCG/Private/Elements/PCGFilterElementsByIndex.cpp
+++ b/Engine/Plugins/PCG/Source/PCG/Private/Elements/PCGFilterElementsByIndex.cpp
@@ -365,7 +365,7 @@ bool FPCGFilterElementsByIndexElement::ExecuteInternal(FPCGContext* InContext) c
 
 			UPCGMetadata* InFilterMetadata = InFilterParamData->Metadata;
 			check(InFilterMetadata);
-			InFilterMetadata->InitializeAsCopy(InputMetadata, &InEntryKeys);
+			InFilterMetadata->InitializeAsCopy(FPCGMetadataInitializeParams(InputMetadata, &InEntryKeys));
 
 			if (Settings->bOutputDiscardedElements)
 			{
@@ -373,7 +373,7 @@ bool FPCGFilterElementsByIndexElement::ExecuteInternal(FPCGContext* InContext) c
 				UPCGMetadata* OutFilterMetadata = OutFilterParamData->Metadata;
 				check(OutFilterMetadata);
 				AddNewOutput(OutFilterParamData, Constants::OutFilterOutputPinLabel);
-				OutFilterMetadata->InitializeAsCopy(InputMetadata, &OutEntryKeys);
+				OutFilterMetadata->InitializeAsCopy(FPCGMetadataInitializeParams(InputMetadata, &OutEntryKeys));
 			}
 		}
 		else
diff --git a/Engine/Plugins/PCG/Source/PCG/Private/Elements/PCGMergeAttributes.cpp b/Engine/Plugins/PCG/Source/PCG/Private/Elements/PCGMergeAttributes.cpp
--- a/Engine/Plugins/PCG/Source/PCG/Private/Elements/PCGMergeAttributes.cpp
+++ b/Engine/Plugins/PCG/Source/PCG/Private/Elements/PCGMergeAttributes.cpp
@@ -108,7 +108,7 @@ bool FPCGMergeAttributesElement::ExecuteInternal(FPCGContext* Context) const
 		{
 			MergedAttributeSet = FPCGContext::NewObject_AnyThread<UPCGParamData>(Context);
 			check(MergedOutput);
-			MergedAttributeSet->Metadata->InitializeAsCopy(CastChecked<const UPCGParamData>(MergedOutput->Data)->Metadata);
+			MergedAttributeSet->Metadata->InitializeAsCopy(FPCGMetadataInitializeParams(MergedOutput->Data->ConstMetadata()));
 			MergedOutput->Data = MergedAttributeSet;
 		}
 
diff --git a/Engine/Plugins/PCG/Source/PCG/Private/Elements/PCGRandomChoice.cpp b/Engine/Plugins/PCG/Source/PCG/Private/Elements/PCGRandomChoice.cpp
--- a/Engine/Plugins/PCG/Source/PCG/Private/Elements/PCGRandomChoice.cpp
+++ b/Engine/Plugins/PCG/Source/PCG/Private/Elements/PCGRandomChoice.cpp
@@ -49,7 +49,7 @@ namespace PCGRandomChoice
 		TArray<PCGMetadataEntryKey> SelectedEntries;
 		Algo::Transform(InIndexes, SelectedEntries, [](int32 In) -> PCGMetadataEntryKey { return In; });
 
-		OutParamData->Metadata->InitializeAsCopy(InParamData->Metadata, &SelectedEntries);
+		OutParamData->Metadata->InitializeAsCopy(FPCGMetadataInitializeParams(InParamData->Metadata, &SelectedEntries));
 
 		return OutParamData;
 	}
diff --git a/Engine/Plugins/PCG/Source/PCG/Private/Elements/PCGSortAttributes.cpp b/Engine/Plugins/PCG/Source/PCG/Private/Elements/PCGSortAttributes.cpp
--- a/Engine/Plugins/PCG/Source/PCG/Private/Elements/PCGSortAttributes.cpp
+++ b/Engine/Plugins/PCG/Source/PCG/Private/Elements/PCGSortAttributes.cpp
@@ -121,7 +121,7 @@ bool FPCGSortAttributesElement::ExecuteInternal(FPCGContext* Context) const
 
 				PCGAttributeAccessorHelpers::SortByAttribute(*Accessor, *Keys, Entries, Settings->SortMethod == EPCGSortMethod::Ascending);
 
-				OutMetadata->InitializeAsCopy(InputData->ConstMetadata(), &Entries);
+				OutMetadata->InitializeAsCopy(FPCGMetadataInitializeParams(InputData->ConstMetadata(), &Entries));
 			}
 
 			FPCGTaggedData& Output = Outputs.Add_GetRef(Input);
diff --git a/Engine/Plugins/PCG/Source/PCG/Private/Metadata/PCGMetadata.cpp b/Engine/Plugins/PCG/Source/PCG/Private/Metadata/PCGMetadata.cpp
--- a/Engine/Plugins/PCG/Source/PCG/Private/Metadata/PCGMetadata.cpp
+++ b/Engine/Plugins/PCG/Source/PCG/Private/Metadata/PCGMetadata.cpp
@@ -16,6 +16,67 @@
 // UPCGMetadata
 ////////////////////////////
 
+template <typename Func>
+void UPCGMetadata::FindFixOrCreateDomainInitializeParams(const FPCGMetadataInitializeParams& InParams, const FPCGMetadataDomainID DomainID, const FPCGMetadataDomain& OtherMetadataDomain, Func&& InFunc)
+{
+	// Will look for the DomainID in the "Mapping" array (that is expected to be an array of tuple with the first element being a DomainID)
+	// If it finds it, it returns a pointer on the second element of the tuple, nullptr otherwise (like a Find function on a TMap)
+	auto FindForDomain = []<typename T>(const FPCGMetadataDomainID& DomainID, const TArray<TTuple<FPCGMetadataDomainID, T>>& Mapping) -> const T*
+	{
+		if (const auto* It = Mapping.FindByPredicate([&DomainID](const auto& Item) { return Item.template Get<0>() == DomainID;}))
+		{
+			return &It->template Get<1>();
+		}
+		else
+		{
+			return nullptr;
+		}
+	};
+
+	const FPCGMetadataDomainInitializeParams* Params = nullptr;
+	FPCGMetadataDomainInitializeParams TempParams{nullptr};
+	
+	Params = FindForDomain(DomainID, InParams.DomainInitializeParams);
+
+	// If we didn't find it, and it is the default domain, retry with default identifier.
+	if (!Params && InParams.Parent->DefaultDomain == DomainID)
+	{
+		Params = FindForDomain(PCGMetadataDomainID::Default, InParams.DomainInitializeParams);
+	}
+
+	// If we didn't find it, and domain ID is default, retry with the default domain.
+	if (!Params && DomainID.IsDefault())
+	{
+		Params = FindForDomain(InParams.Parent->DefaultDomain, InParams.DomainInitializeParams);
+	}
+
+	// Validate that the parenting is valid, if not update it
+	if (!Params || !Params->Parent || Params->Parent != &OtherMetadataDomain)
+	{
+		if (Params)
+		{
+			TempParams = *Params;
+			TempParams.Parent = &OtherMetadataDomain;
+		}
+		else
+		{
+			TempParams = FPCGMetadataDomainInitializeParams(&OtherMetadataDomain);
+		}
+		
+		Params = &TempParams;
+	}
+
+	// If we have a mapping for domains, find it there.
+	FPCGMetadataDomainID OtherDomain = DomainID;
+	if (const FPCGMetadataDomainID* It = FindForDomain(DomainID, InParams.DomainMapping))
+	{
+		OtherDomain = *It;
+	}
+
+	check(Params);
+	InFunc(FindOrCreateMetadataDomain(OtherDomain), *Params);
+}
+
 template <typename Func, typename ...Args>
 decltype(auto) UPCGMetadata::WithMetadataDomain(const FPCGMetadataDomainID& InMetadataDomainID, Func InFunc, Args&& ...InArgs)
 {
@@ -200,6 +261,54 @@ FPCGMetadataDomain* UPCGMetadata::CreateMetadataDomain_Unsafe(const FPCGMetadata
 	return MetadataDomains[InMetadataDomainID].Get();
 }
 
+FPCGMetadataInitializeParams::FPCGMetadataInitializeParams(const UPCGMetadata* InParent, const TArray<PCGMetadataEntryKey>* InOptionalEntriesToCopy)
+	: Parent(InParent)
+{
+	FPCGMetadataDomainInitializeParams& DefaultDomainParams = DomainInitializeParams.Emplace_GetRef(PCGMetadataDomainID::Default, FPCGMetadataDomainInitializeParams(InParent->GetConstDefaultMetadataDomain())).Get<1>();
+	if (InOptionalEntriesToCopy)
+	{
+		DefaultDomainParams.OptionalEntriesToCopy.Emplace(PCGValueRangeHelpers::MakeConstValueRange(*InOptionalEntriesToCopy));
+	}
+}
+
+FPCGMetadataInitializeParams::FPCGMetadataInitializeParams(const UPCGMetadata* InParent,
+	const TSet<FName>& InFilteredAttributes,
+	EPCGMetadataFilterMode InFilterMode,
+	EPCGStringMatchingOperator InMatchOperator,
+	const TArray<PCGMetadataEntryKey>* InOptionalEntriesToCopy)
+		: FPCGMetadataInitializeParams(InParent, InOptionalEntriesToCopy)
+{
+	// We added the default in the ctor just above. It will be the first element.
+	FPCGMetadataDomainInitializeParams& DefaultDomainParams = DomainInitializeParams[0].Get<1>();
+	DefaultDomainParams.FilterMode = InFilterMode;
+	DefaultDomainParams.MatchOperator = InMatchOperator;
+	if (!InFilteredAttributes.IsEmpty())
+	{
+		DefaultDomainParams.FilteredAttributes = InFilteredAttributes;
+	}
+}
+
+void FPCGMetadataInitializeParams::PopulateDomainInitializeParamsFromParent()
+{
+	if (!Parent)
+	{
+		return;
+	}
+
+	// -1 to remove the default
+	DomainInitializeParams.Empty(Parent->MetadataDomains.Num() - 1);
+
+	for (const auto& It : Parent->MetadataDomains)
+	{
+		if (It.Key.IsDefault() || !It.Value)
+		{
+			continue;
+		}
+		
+		DomainInitializeParams.Emplace(It.Key, FPCGMetadataDomainInitializeParams(It.Value.Get()));
+	}
+}
+
 UPCGMetadata::UPCGMetadata(const FObjectInitializer& ObjectInitializer)
 	: Super(ObjectInitializer)
 {
@@ -313,6 +422,20 @@ void UPCGMetadata::Serialize(FArchive& InArchive)
 	}
 }
 
+void UPCGMetadata::K2_InitializeAsCopy(const UPCGMetadata* InMetadataToCopy, const TArray<int64>& InOptionalEntriesToCopy)
+{
+	return InitializeAsCopy(FPCGMetadataInitializeParams(InMetadataToCopy, !InOptionalEntriesToCopy.IsEmpty() ? &InOptionalEntriesToCopy : nullptr));
+}
+
+void UPCGMetadata::K2_InitializeAsCopyWithAttributeFilter(const UPCGMetadata* InMetadataToCopy,
+	const TSet<FName>& InFilteredAttributes,
+	const TArray<int64>& InOptionalEntriesToCopy,
+	EPCGMetadataFilterMode InFilterMode,
+	EPCGStringMatchingOperator InMatchOperator)
+{
+	return InitializeAsCopy(FPCGMetadataInitializeParams(InMetadataToCopy, InFilteredAttributes, InFilterMode, InMatchOperator, !InOptionalEntriesToCopy.IsEmpty() ? &InOptionalEntriesToCopy : nullptr));
+}
+
 void UPCGMetadata::AddToCrc(FArchiveCrc32& Ar, bool bFullDataCrc) const
 {
 	TRACE_CPUPROFILER_EVENT_SCOPE(UPCGMetadata::AddToCrc);
@@ -341,73 +464,73 @@ void UPCGMetadata::AddToCrc(FArchiveCrc32& Ar, bool bFullDataCrc) const
 	}
 }
 
-void UPCGMetadata::Initialize(const UPCGMetadata* InParent, bool bAddAttributesFromParent, const TMap<FPCGMetadataDomainID, FPCGMetadataDomainID>& DomainMapping)
+void UPCGMetadata::Initialize(const UPCGMetadata* InParent)
 {
+	Initialize(FPCGMetadataInitializeParams(InParent));
+}
+
+void UPCGMetadata::Initialize(const UPCGMetadata* InParent, bool bAddAttributesFromParent)
+{
+	if (!InParent)
+	{
+		return;
+	}
+	
 	// If we are adding attributes from parent, then we use exclude filter with empty list so
 	// that all parameters added. Otherwise use include filter with empty list so none are added.
-	const EPCGMetadataFilterMode bFilter = bAddAttributesFromParent ? EPCGMetadataFilterMode::ExcludeAttributes : EPCGMetadataFilterMode::IncludeAttributes;
-	InitializeWithAttributeFilter(InParent, TSet<FPCGAttributeIdentifier>(), bFilter, EPCGStringMatchingOperator::Equal, DomainMapping);
+	const FPCGMetadataInitializeParams Params(InParent, TSet<FName>(), bAddAttributesFromParent ? EPCGMetadataFilterMode::ExcludeAttributes : EPCGMetadataFilterMode::IncludeAttributes);
+	Initialize(Params);
 }
 
 void UPCGMetadata::InitializeWithAttributeFilter(const UPCGMetadata* InParent, const TSet<FName>& InFilteredAttributes, EPCGMetadataFilterMode InFilterMode, EPCGStringMatchingOperator InMatchOperator)
 {
-	return InitializeWithAttributeFilter(InParent, FPCGAttributeIdentifier::TransformNameSet(InFilteredAttributes), InFilterMode, InMatchOperator);
+	Initialize(FPCGMetadataInitializeParams(InParent, InFilteredAttributes, InFilterMode, InMatchOperator));
 }
 
-void UPCGMetadata::InitializeWithAttributeFilter(const UPCGMetadata* InParent, const TSet<FPCGAttributeIdentifier>& InFilteredAttributes, EPCGMetadataFilterMode InFilterMode, EPCGStringMatchingOperator InMatchOperator, const TMap<FPCGMetadataDomainID, FPCGMetadataDomainID>& DomainMapping)
+void UPCGMetadata::Initialize(const FPCGMetadataInitializeParams& InParams)
 {
-	if (Parent)
+	if (Parent || InParams.Parent == this || !InParams.Parent)
 	{
-		// Already initialized; note that while that might be constructed as a warning, there are legit cases where this is correct
+		// Already initialized, or invalid parent; note that while that might be constructed as a warning, there are legit cases where this is correct
 		return;
 	}
+	
+	// If we have a domain mapping, we can't have "dynamic" domains
+	if (InParams.DomainMapping.IsEmpty())
+	{
+		SetupDomainsFromOtherMetadataIfNeeded(InParams.Parent);
+	}
 
-	if (InParent != this && InParent)
+	Parent = InParams.Parent;
+	
+	Parent->ForEachValidUniqueConstMetadataDomain([this, &InParams](const FPCGMetadataDomainID& DomainID, const FPCGMetadataDomain& OtherMetadataDomain)
 	{
-		// If we have a domain mapping, we can't have "dynamic" domains
-		if (DomainMapping.IsEmpty())
+		auto InitializeFunc = [](FPCGMetadataDomain* CurrentMetadataDomain, const FPCGMetadataDomainInitializeParams& DomainParams)
 		{
-			SetupDomainsFromOtherMetadataIfNeeded(InParent);
-		}
-
-		Parent = InParent;
-		
-		TMap<FPCGMetadataDomainID, TSet<FName>> FilteredAttributeSplit = FPCGAttributeIdentifier::SplitByDomain(InFilteredAttributes);
+			CurrentMetadataDomain->Initialize(DomainParams);
+		};
 		
-		InParent->ForEachValidUniqueConstMetadataDomain([this, &FilteredAttributeSplit, InFilterMode, InMatchOperator, &DomainMapping](const FPCGMetadataDomainID& DomainID, const FPCGMetadataDomain& OtherMetadataDomain)
-		{
-			FPCGMetadataDomainID OtherDomain = DomainID;
-			if (const FPCGMetadataDomainID* MappedDomain = DomainMapping.Find(DomainID))
-			{
-				OtherDomain = *MappedDomain;
-			}
-
-			TSet<FName>* FilteredAttributesPtr = FilteredAttributeSplit.Find(DomainID);
-			FPCGMetadataDomain* CurrentMetadataDomain = FindOrCreateMetadataDomain(OtherDomain);
-			CurrentMetadataDomain->InitializeWithAttributeFilter(&OtherMetadataDomain, (FilteredAttributesPtr ? *FilteredAttributesPtr : TSet<FName>{}), InFilterMode, InMatchOperator);
-		});
-	}
+		FindFixOrCreateDomainInitializeParams(InParams, DomainID, OtherMetadataDomain, std::move(InitializeFunc));
+	});
 }
 
 void UPCGMetadata::InitializeAsCopy(const UPCGMetadata* InMetadataToCopy, const TArray<PCGMetadataEntryKey>* EntriesToCopy)
 {
-	InitializeAsCopyWithAttributeFilter(InMetadataToCopy, TSet<FName>(), EPCGMetadataFilterMode::ExcludeAttributes, EntriesToCopy);
+	InitializeAsCopy(FPCGMetadataInitializeParams(InMetadataToCopy, EntriesToCopy));
 }
 
 void UPCGMetadata::InitializeAsCopyWithAttributeFilter(const UPCGMetadata* InMetadataToCopy, const TSet<FName>& InFilteredAttributes, EPCGMetadataFilterMode InFilterMode, const TArray<PCGMetadataEntryKey>* EntriesToCopy, EPCGStringMatchingOperator InMatchOperator)
 {
-	InitializeAsCopyWithAttributeFilter(InMetadataToCopy, FPCGAttributeIdentifier::TransformNameSet(InFilteredAttributes), InFilterMode, EntriesToCopy, InMatchOperator);
+	InitializeAsCopy(FPCGMetadataInitializeParams(InMetadataToCopy, InFilteredAttributes, InFilterMode, InMatchOperator, EntriesToCopy));
 }
 
-void UPCGMetadata::InitializeAsCopyWithAttributeFilter(const UPCGMetadata* InMetadataToCopy, const TSet<FPCGAttributeIdentifier>& InFilteredAttributes, EPCGMetadataFilterMode InFilterMode, const TArray<PCGMetadataEntryKey>* EntriesToCopy, EPCGStringMatchingOperator InMatchOperator)
+void UPCGMetadata::InitializeAsCopy(const FPCGMetadataInitializeParams& InParams)
 {
-	if (!InMetadataToCopy)
+	if (!InParams.Parent)
 	{
 		return;
 	}
 	
-	check(InMetadataToCopy);
-	
 	if (Parent)
 	{
 		UE_LOG(LogPCG, Error, TEXT("Metadata has already been initialized."));
@@ -418,56 +541,64 @@ void UPCGMetadata::InitializeAsCopyWithAttributeFilter(const UPCGMetadata* InMet
 	// But in case of partial copy, the parenting is not necessary.
 	bool bShouldParent = false;
 
-	SetupDomainsFromOtherMetadataIfNeeded(InMetadataToCopy);
+	SetupDomainsFromOtherMetadataIfNeeded(InParams.Parent);
 		
-	TMap<FPCGMetadataDomainID, TSet<FName>> FilteredAttributeSplit = FPCGAttributeIdentifier::SplitByDomain(InFilteredAttributes);
-	InMetadataToCopy->ForEachValidUniqueConstMetadataDomain([this, &FilteredAttributeSplit, InFilterMode, &EntriesToCopy, &bShouldParent, InMetadataToCopy, InMatchOperator](const FPCGMetadataDomainID& DomainID, const FPCGMetadataDomain& OtherMetadataDomain)
+	InParams.Parent->ForEachValidUniqueConstMetadataDomain([this, &bShouldParent, &InParams](const FPCGMetadataDomainID& DomainID, const FPCGMetadataDomain& OtherMetadataDomain)
 	{
-		TSet<FName>* FilteredAttributesPtr = FilteredAttributeSplit.Find(DomainID);
-		FPCGMetadataDomain* CurrentMetadataDomain = FindOrCreateMetadataDomain(DomainID);
-		CurrentMetadataDomain->InitializeAsCopyWithAttributeFilter(&OtherMetadataDomain, (FilteredAttributesPtr ? *FilteredAttributesPtr : TSet<FName>{}), InFilterMode, EntriesToCopy, InMatchOperator);
-
-		// We still validate that the parent was set correctly and it matches the metadata to copy.
-		if (CurrentMetadataDomain->Parent && CurrentMetadataDomain->Parent->TopMetadata == InMetadataToCopy)
+		auto InitializeFunc = [&bShouldParent, &InParams](FPCGMetadataDomain* CurrentMetadataDomain, const FPCGMetadataDomainInitializeParams& DomainParams)
 		{
-			bShouldParent = true;
-		}
+			CurrentMetadataDomain->InitializeAsCopy(DomainParams);
+
+			// We still validate that the parent was set correctly and it matches the metadata to copy.
+			if (CurrentMetadataDomain->Parent && CurrentMetadataDomain->Parent->TopMetadata == InParams.Parent)
+			{
+				bShouldParent = true;
+			}
+		};
+		
+		FindFixOrCreateDomainInitializeParams(InParams, DomainID, OtherMetadataDomain, std::move(InitializeFunc));
 	});
 
 	if (bShouldParent)
 	{
-		Parent = InMetadataToCopy;
-		OtherParents = InMetadataToCopy->OtherParents;
+		Parent = InParams.Parent;
+		OtherParents = InParams.Parent->OtherParents;
 	}
 }
 
+void UPCGMetadata::AddAttributes(const UPCGMetadata* InOther)
+{
+	AddAttributes(FPCGMetadataInitializeParams(InOther));
+}
+
 void UPCGMetadata::AddAttributesFiltered(const UPCGMetadata* InOther, const TSet<FName>& InFilteredAttributes, EPCGMetadataFilterMode InFilterMode, EPCGStringMatchingOperator InMatchOperator)
 {
-	AddAttributesFiltered(InOther, FPCGAttributeIdentifier::TransformNameSet(InFilteredAttributes), InFilterMode, InMatchOperator);
+	AddAttributes(FPCGMetadataInitializeParams(InOther, InFilteredAttributes, InFilterMode, InMatchOperator));
 }
 
-void UPCGMetadata::AddAttributesFiltered(const UPCGMetadata* InOther, const TSet<FPCGAttributeIdentifier>& InFilteredAttributes, EPCGMetadataFilterMode InFilterMode, EPCGStringMatchingOperator InMatchOperator)
+void UPCGMetadata::AddAttributes(const FPCGMetadataInitializeParams& InParams)
 {
-	if (!InOther)
+	if (!InParams.Parent)
 	{
 		return;
 	}
-	
-	TMap<FPCGMetadataDomainID, TSet<FName>> FilteredAttributeSplit = FPCGAttributeIdentifier::SplitByDomain(InFilteredAttributes);
 
-	InOther->ForEachValidUniqueConstMetadataDomain([this, &FilteredAttributeSplit, InFilterMode, InMatchOperator, InOther](const FPCGMetadataDomainID& DomainID, const FPCGMetadataDomain& OtherMetadataDomain)
+	bool bAddSucceeded = false;
+
+	InParams.Parent->ForEachValidUniqueConstMetadataDomain([this, &InParams, &bAddSucceeded](const FPCGMetadataDomainID& DomainID, const FPCGMetadataDomain& OtherMetadataDomain)
 	{
-		TSet<FName>* FilteredAttributesPtr = FilteredAttributeSplit.Find(DomainID);
-		FPCGMetadataDomain* CurrentMetadataDomain = FindOrCreateMetadataDomain(DomainID);
-				
-		if (CurrentMetadataDomain->AddAttributesFiltered(&OtherMetadataDomain, (FilteredAttributesPtr ? *FilteredAttributesPtr : TSet<FName>{}), InFilterMode, InMatchOperator))
+		auto InitializeFunc = [&bAddSucceeded](FPCGMetadataDomain* CurrentMetadataDomain, const FPCGMetadataDomainInitializeParams& DomainParams)
 		{
-			if (InOther != Parent)
-			{
-				OtherParents.Add(InOther);
-			}
-		}
+			bAddSucceeded |= CurrentMetadataDomain->AddAttributes(DomainParams);
+		};
+
+		FindFixOrCreateDomainInitializeParams(InParams, DomainID, OtherMetadataDomain, std::move(InitializeFunc));
 	});
+
+	if (bAddSucceeded && InParams.Parent != Parent)
+	{
+		OtherParents.Add(InParams.Parent);
+	}
 }
 
 void UPCGMetadata::BP_AddAttribute(const UPCGMetadata* InOther, FName AttributeName)
diff --git a/Engine/Plugins/PCG/Source/PCG/Private/Metadata/PCGMetadataAttribute.cpp b/Engine/Plugins/PCG/Source/PCG/Private/Metadata/PCGMetadataAttribute.cpp
--- a/Engine/Plugins/PCG/Source/PCG/Private/Metadata/PCGMetadataAttribute.cpp
+++ b/Engine/Plugins/PCG/Source/PCG/Private/Metadata/PCGMetadataAttribute.cpp
@@ -169,6 +169,11 @@ PCGMetadataValueKey FPCGMetadataAttributeBase::GetValueKey(PCGMetadataEntryKey E
 }
 
 void FPCGMetadataAttributeBase::GetValueKeys(const TArrayView<const PCGMetadataEntryKey>& EntryKeys, TArray<PCGMetadataValueKey>& OutValueKeys) const
+{
+	GetValueKeys(PCGValueRangeHelpers::MakeConstValueRange<PCGMetadataEntryKey>(EntryKeys), OutValueKeys);
+}
+
+void FPCGMetadataAttributeBase::GetValueKeys(TConstPCGValueRange<PCGMetadataEntryKey> EntryKeys, TArray<PCGMetadataValueKey>& OutValueKeys) const
 {
 	if (EntryKeys.IsEmpty())
 	{
@@ -183,6 +188,11 @@ void FPCGMetadataAttributeBase::GetValueKeys(const TArrayView<const PCGMetadataE
 }
 
 void FPCGMetadataAttributeBase::GetValueKeys(TArrayView<PCGMetadataEntryKey> EntryKeys, TArray<PCGMetadataValueKey>& OutValueKeys) const
+{
+	GetValueKeys(PCGValueRangeHelpers::MakeValueRange<PCGMetadataEntryKey>(EntryKeys), OutValueKeys);
+}
+
+void FPCGMetadataAttributeBase::GetValueKeys(TPCGValueRange<PCGMetadataEntryKey> EntryKeys, TArray<PCGMetadataValueKey>& OutValueKeys) const
 {
 	if (EntryKeys.IsEmpty())
 	{
@@ -193,11 +203,10 @@ void FPCGMetadataAttributeBase::GetValueKeys(TArrayView<PCGMetadataEntryKey> Ent
 	// Bitset with all unset values. If we have any unset value, we will ask the parent for those.
 	TBitArray<> UnsetValues(true, EntryKeys.Num());
 
-	// We have ownership on the memory, so we can skip the copy of entries.
-	GetValueKeys_Internal(EntryKeys, OutValueKeys, UnsetValues, /*bOwnerOfEntryKeysView=*/true);
+	GetValueKeys_Internal(PCGValueRangeHelpers::MakeConstValueRange(EntryKeys), OutValueKeys, UnsetValues, /*bOwnerOfEntryKeysView=*/true);
 }
 
-void FPCGMetadataAttributeBase::GetValueKeys_Internal(const TArrayView<const PCGMetadataEntryKey> EntryKeys, TArrayView<PCGMetadataValueKey> OutValueKeys, TBitArray<>& UnsetValues, bool bOwnerOfEntryKeysView) const
+void FPCGMetadataAttributeBase::GetValueKeys_Internal(TConstPCGValueRange<PCGMetadataEntryKey> EntryKeys, TArrayView<PCGMetadataValueKey> OutValueKeys, TBitArray<>& UnsetValues, bool bOwnerOfEntryKeysView) const
 {
 	check(EntryKeys.Num() == OutValueKeys.Num() && OutValueKeys.Num() == UnsetValues.Num());
 
@@ -245,11 +254,11 @@ void FPCGMetadataAttributeBase::GetValueKeys_Internal(const TArrayView<const PCG
 
 	if (Parent && !bFoundAllKeys)
 	{
-		auto ParentCall = [this, &OutValueKeys, &UnsetValues](const TArrayView<const PCGMetadataEntryKey> CurrentEntryKeys)
+		auto ParentCall = [this, &OutValueKeys, &UnsetValues](TConstPCGValueRange<PCGMetadataEntryKey> CurrentEntryKeys)
 		{
 			// Before querying the parent, we need to update all our entry keys to get them in the parent referential.
 			// At that point, we are owner of our memory, so it is safe to cast
-			Metadata->GetParentKeys(MakeArrayView(const_cast<PCGMetadataEntryKey*>(CurrentEntryKeys.GetData()), CurrentEntryKeys.Num()), &UnsetValues);
+			Metadata->GetParentKeysWithRange(PCGValueRangeHelpers::MakeValueRange_Unsafe(CurrentEntryKeys), &UnsetValues);
 			Parent->GetValueKeys_Internal(CurrentEntryKeys, OutValueKeys, UnsetValues, /*bOwnerOfEntryKeysView=*/true);
 		};
 
@@ -258,8 +267,9 @@ void FPCGMetadataAttributeBase::GetValueKeys_Internal(const TArrayView<const PCG
 		if (!bOwnerOfEntryKeysView)
 		{
 			TArray<PCGMetadataEntryKey, TInlineAllocator<256>> CopiedKeys;
-			CopiedKeys = EntryKeys;
-			ParentCall(CopiedKeys);
+			CopiedKeys.Reserve(EntryKeys.Num());
+			Algo::Transform(EntryKeys, CopiedKeys, [](const PCGMetadataEntryKey& It) { return It; });
+			ParentCall(PCGValueRangeHelpers::MakeConstValueRange(CopiedKeys));
 		}
 		else
 		{
diff --git a/Engine/Plugins/PCG/Source/PCG/Private/Metadata/PCGMetadataDomain.cpp b/Engine/Plugins/PCG/Source/PCG/Private/Metadata/PCGMetadataDomain.cpp
--- a/Engine/Plugins/PCG/Source/PCG/Private/Metadata/PCGMetadataDomain.cpp
+++ b/Engine/Plugins/PCG/Source/PCG/Private/Metadata/PCGMetadataDomain.cpp
@@ -91,17 +91,15 @@ namespace PCGMetadata
 		return PCGPropertyHelpers::GetPropertyValueWithCallback(DataPtr, InProperty, CreateAttributeAndSet);
 	}
 
-	/** Utility structure, should not be kept around as we keep a ref on `InNames` */
+	/** Utility structure to filter attributes when adding them. Must not be kept around as we hold a const ref to Params.*/
 	struct FPCGMetadataAttributeNameFilter
 	{
-		FPCGMetadataAttributeNameFilter(const TSet<FName>& InNames, EPCGMetadataFilterMode InFilterMode, EPCGStringMatchingOperator InMatchOperator)
-			: Names(InNames)
-			, FilterMode(InFilterMode)
-			, MatchOperator(InMatchOperator)
+		FPCGMetadataAttributeNameFilter(const FPCGMetadataDomainInitializeParams& InParams)
+			: Params(InParams)
 		{
-			if (InMatchOperator != EPCGStringMatchingOperator::Equal)
+			if (Params.MatchOperator != EPCGStringMatchingOperator::Equal && Params.FilteredAttributes)
 			{
-				Algo::Transform(InNames, NameStrings, [](const FName& InFilteredAttribute) {return InFilteredAttribute.ToString(); });
+				Algo::Transform(*Params.FilteredAttributes, NameStrings, [](const FName& InFilteredAttribute) {return InFilteredAttribute.ToString(); });
 			}
 		}
 
@@ -110,18 +108,18 @@ namespace PCGMetadata
 		{
 			bool Result = false;
 
-			if (MatchOperator == EPCGStringMatchingOperator::Equal)
+			if (Params.MatchOperator == EPCGStringMatchingOperator::Equal)
 			{
-				Result = Names.Contains(InName);
+				Result = Params.FilteredAttributes && Params.FilteredAttributes->Contains(InName);
 			}
 			else
 			{
 				const FString OtherAttributeString = InName.ToString();
-				if (MatchOperator == EPCGStringMatchingOperator::Substring)
+				if (Params.MatchOperator == EPCGStringMatchingOperator::Substring)
 				{
 					Result = Algo::AnyOf(NameStrings, [&OtherAttributeString](const FString& InAttribute) { return OtherAttributeString.Contains(InAttribute); });
 				}
-				else if (MatchOperator == EPCGStringMatchingOperator::Matches)
+				else if (Params.MatchOperator == EPCGStringMatchingOperator::Matches)
 				{
 					Result = Algo::AnyOf(NameStrings, [&OtherAttributeString](const FString& InAttribute) { return OtherAttributeString.MatchesWildcard(InAttribute); });
 				}
@@ -132,14 +130,12 @@ namespace PCGMetadata
 				}
 			}
 
-			return FilterMode == EPCGMetadataFilterMode::ExcludeAttributes ? Result : !Result;
+			return Params.FilterMode == EPCGMetadataFilterMode::ExcludeAttributes ? Result : !Result;
 		}
 
 	private:
-		const TSet<FName>& Names;
+		const FPCGMetadataDomainInitializeParams& Params;
 		TArray<FString> NameStrings;
-		EPCGMetadataFilterMode FilterMode;
-		EPCGStringMatchingOperator MatchOperator;
 	};
 }
 
@@ -299,15 +295,12 @@ FPCGMetadataDomain::~FPCGMetadataDomain()
 	Attributes.Reset();
 }
 
-void FPCGMetadataDomain::Initialize(const FPCGMetadataDomain* InParent, bool bAddAttributesFromParent)
+void FPCGMetadataDomain::Initialize(const FPCGMetadataDomain* InParent)
 {
-	// If we are adding attributes from parent, then we use exclude filter with empty list so
-	// that all parameters added. Otherwise use include filter with empty list so none are added.
-	const EPCGMetadataFilterMode bFilter = bAddAttributesFromParent ? EPCGMetadataFilterMode::ExcludeAttributes : EPCGMetadataFilterMode::IncludeAttributes;
-	InitializeWithAttributeFilter(InParent, TSet<FName>(), bFilter);
+	Initialize(FPCGMetadataDomainInitializeParams(InParent));
 }
 
-void FPCGMetadataDomain::InitializeWithAttributeFilter(const FPCGMetadataDomain* InParent, const TSet<FName>& InFilteredAttributes, EPCGMetadataFilterMode InFilterMode, EPCGStringMatchingOperator InMatchOperator)
+void FPCGMetadataDomain::Initialize(const FPCGMetadataDomainInitializeParams& InParams)
 {
 	if (Parent || Attributes.Num() != 0)
 	{
@@ -318,50 +311,60 @@ void FPCGMetadataDomain::InitializeWithAttributeFilter(const FPCGMetadataDomain*
 	// If we don't have a top metadata (ill-formed), or we don't support parenting, force the copy
 	if (!ensure(TopMetadata) || !TopMetadata->MetadataDomainSupportsParenting(DomainID))
 	{
-		InitializeAsCopyWithAttributeFilter(InParent, InFilteredAttributes, InFilterMode, nullptr, InMatchOperator);
+		// Make sure we have nothing in the attribute to copy
+		if (InParams.OptionalEntriesToCopy)
+		{
+			FPCGMetadataDomainInitializeParams CopyParams = InParams;
+			CopyParams.OptionalEntriesToCopy.Reset();
+			InitializeAsCopy(CopyParams);
+		}
+		else
+		{
+			InitializeAsCopy(InParams);
+		}
+
 		return;
 	}
 
-	Parent = ((InParent != this) ? InParent : nullptr);
+	Parent = ((InParams.Parent != this) ? InParams.Parent : nullptr);
 	ItemKeyOffset = Parent ? Parent->GetItemCountForChild() : 0;
 
 	// If we have been given an include list which is empty, then don't bother adding any attributes
-	const bool bSkipAddingAttributesFromParent = (InFilterMode == EPCGMetadataFilterMode::IncludeAttributes) && (InFilteredAttributes.Num() == 0);
+	const bool bSkipAddingAttributesFromParent = (InParams.FilterMode == EPCGMetadataFilterMode::IncludeAttributes) && (!InParams.FilteredAttributes || InParams.FilteredAttributes->IsEmpty());
 	if (!bSkipAddingAttributesFromParent)
 	{
-		AddAttributesFiltered(InParent, InFilteredAttributes, InFilterMode, InMatchOperator);
+		AddAttributes(InParams);
 	}
 }
 
-void FPCGMetadataDomain::InitializeAsCopy(const FPCGMetadataDomain* InMetadataToCopy, const TArray<PCGMetadataEntryKey>* EntriesToCopy)
+void FPCGMetadataDomain::InitializeAsCopy(const FPCGMetadataDomain* InMetadataToCopy)
 {
-	InitializeAsCopyWithAttributeFilter(InMetadataToCopy, TSet<FName>(), EPCGMetadataFilterMode::ExcludeAttributes, EntriesToCopy);
+	InitializeAsCopy(FPCGMetadataDomainInitializeParams(InMetadataToCopy));
 }
 
-void FPCGMetadataDomain::InitializeAsCopyWithAttributeFilter(const FPCGMetadataDomain* InMetadataToCopy, const TSet<FName>& InFilteredAttributes, EPCGMetadataFilterMode InFilterMode, const TArray<PCGMetadataEntryKey>* EntriesToCopy, EPCGStringMatchingOperator InMatchOperator)
+void FPCGMetadataDomain::InitializeAsCopy(const FPCGMetadataDomainInitializeParams& InParams)
 {
-	if (!InMetadataToCopy)
+	if (!InParams.Parent)
 	{
 		return;
 	}
 
-	check(InMetadataToCopy);
 	if (Parent || Attributes.Num() != 0)
 	{
 		UE_LOG(LogPCG, Error, TEXT("Metadata has already been initialized or already contains attributes"));
 		return;
 	}
 
-	PCGMetadata::FPCGMetadataAttributeNameFilter ShouldSkipAttribute(InFilteredAttributes, InFilterMode, InMatchOperator);
+	PCGMetadata::FPCGMetadataAttributeNameFilter ShouldSkipAttribute(InParams);
 
 	// If we have a partial copy, it will flatten the metadata, so we don't need a parent.
 	// Otherwise, we keep the parent hierarchy.
-	const bool bPartialCopy = EntriesToCopy && EntriesToCopy->Num() <= InMetadataToCopy->GetItemCountForChild();
+	const bool bPartialCopy = InParams.OptionalEntriesToCopy && InParams.OptionalEntriesToCopy->Num() <= InParams.Parent->GetItemCountForChild();
 	TArray<PCGMetadataEntryKey> NewEntryKeys;
 	TArray<PCGMetadataValueKey> NewValueKeys;
 	if (bPartialCopy)
 	{
-		const int32 Count = EntriesToCopy->Num();
+		const int32 Count = InParams.OptionalEntriesToCopy->Num();
 		NewEntryKeys.SetNumUninitialized(Count);
 		NewValueKeys.SetNumUninitialized(Count);
 		ParentKeys.SetNumUninitialized(Count);
@@ -375,13 +378,13 @@ void FPCGMetadataDomain::InitializeAsCopyWithAttributeFilter(const FPCGMetadataD
 	}
 	else
 	{
-		ParentKeys = InMetadataToCopy->ParentKeys;
-		ItemKeyOffset = InMetadataToCopy->ItemKeyOffset;
-		Parent = InMetadataToCopy->Parent;
+		ParentKeys = InParams.Parent->ParentKeys;
+		ItemKeyOffset = InParams.Parent->ItemKeyOffset;
+		Parent = InParams.Parent->Parent;
 	}
 
 	// Copy attributes
-	for (const TPair<FName, FPCGMetadataAttributeBase*>& OtherAttribute : InMetadataToCopy->Attributes)
+	for (const TPair<FName, FPCGMetadataAttributeBase*>& OtherAttribute : InParams.Parent->Attributes)
 	{
 		if (!ShouldSkipAttribute(OtherAttribute.Key))
 		{
@@ -390,25 +393,25 @@ void FPCGMetadataDomain::InitializeAsCopyWithAttributeFilter(const FPCGMetadataD
 
 			if (bPartialCopy && OtherAttribute.Value && Attribute)
 			{
-				OtherAttribute.Value->GetValueKeys(*EntriesToCopy, NewValueKeys);
+				OtherAttribute.Value->GetValueKeys(*InParams.OptionalEntriesToCopy, NewValueKeys);
 				Attribute->SetValuesFromValueKeys(NewEntryKeys, NewValueKeys);
 			}
 		}
 	}
 }
 
-bool FPCGMetadataDomain::AddAttributesFiltered(const FPCGMetadataDomain* InOther, const TSet<FName>& InFilteredAttributes, EPCGMetadataFilterMode InFilterMode, EPCGStringMatchingOperator InMatchOperator)
+bool FPCGMetadataDomain::AddAttributes(const FPCGMetadataDomainInitializeParams& InParams)
 {
-	if (!InOther)
+	if (!InParams.Parent)
 	{
 		return false;
 	}
 
-	PCGMetadata::FPCGMetadataAttributeNameFilter ShouldSkipAttribute(InFilteredAttributes, InFilterMode, InMatchOperator);
+	PCGMetadata::FPCGMetadataAttributeNameFilter ShouldSkipAttribute(InParams);
 
 	bool bAttributeAdded = false;
 
-	for (const TPair<FName, FPCGMetadataAttributeBase*>& OtherAttribute : InOther->Attributes)
+	for (const TPair<FName, FPCGMetadataAttributeBase*>& OtherAttribute : InParams.Parent->Attributes)
 	{
 		if (ShouldSkipAttribute(OtherAttribute.Key) || !OtherAttribute.Value)
 		{
@@ -429,7 +432,7 @@ bool FPCGMetadataDomain::AddAttributesFiltered(const FPCGMetadataDomain* InOther
 				ChangeAttributeType(OtherAttribute.Key, OtherAttribute.Value->GetTypeId());
 			}
 		}
-		else if (CopyAttribute(OtherAttribute.Value, OtherAttribute.Key, /*bKeepParent=*/InOther == Parent, /*bCopyEntries=*/false, /*bCopyValues=*/false))
+		else if (CopyAttribute(OtherAttribute.Value, OtherAttribute.Key, /*bKeepParent=*/InParams.Parent == Parent, /*bCopyEntries=*/false, /*bCopyValues=*/false))
 		{
 			bAttributeAdded = true;
 		}
@@ -1130,6 +1133,11 @@ PCGMetadataEntryKey FPCGMetadataDomain::GetParentKey(PCGMetadataEntryKey LocalIt
 }
 
 void FPCGMetadataDomain::GetParentKeys(TArrayView<PCGMetadataEntryKey> LocalItemKeys, const TBitArray<>* Mask) const
+{
+	GetParentKeysWithRange(PCGValueRangeHelpers::MakeValueRange(LocalItemKeys), Mask);
+}
+
+void FPCGMetadataDomain::GetParentKeysWithRange(TPCGValueRange<PCGMetadataEntryKey> LocalItemKeys, const TBitArray<>* Mask) const
 {
 	auto GetParentKey_Unsafe = [this](PCGMetadataEntryKey& LocalItemKey) -> void
 	{
diff --git a/Engine/Plugins/PCG/Source/PCG/Private/PCGParamData.cpp b/Engine/Plugins/PCG/Source/PCG/Private/PCGParamData.cpp
--- a/Engine/Plugins/PCG/Source/PCG/Private/PCGParamData.cpp
+++ b/Engine/Plugins/PCG/Source/PCG/Private/PCGParamData.cpp
@@ -151,7 +151,7 @@ UPCGParamData* UPCGParamData::DuplicateData(FPCGContext* Context, bool bInitiali
 	if (bInitializeMetadata)
 	{
 		check(NewParamData && NewParamData->Metadata);
-		NewParamData->Metadata->InitializeAsCopy(Metadata);
+		NewParamData->Metadata->InitializeAsCopy(FPCGMetadataInitializeParams(Metadata));
 	}
 
 	return NewParamData;
diff --git a/Engine/Plugins/PCG/Source/PCG/Private/Tests/Metadata/PCGMetadataTest.cpp b/Engine/Plugins/PCG/Source/PCG/Private/Tests/Metadata/PCGMetadataTest.cpp
--- a/Engine/Plugins/PCG/Source/PCG/Private/Tests/Metadata/PCGMetadataTest.cpp
+++ b/Engine/Plugins/PCG/Source/PCG/Private/Tests/Metadata/PCGMetadataTest.cpp
@@ -127,7 +127,7 @@ bool FPCGMetadataTest_InheritCopy::RunTest(const FString& Parameters)
 		
 		PointData3 = NewObject<UPCGPointData>();
 		UPCGMetadata* PointMetadata3 = PointData3->MutableMetadata();
-		PointMetadata3->InitializeAsCopy(PointMetadata2);
+		PointMetadata3->InitializeAsCopy(FPCGMetadataInitializeParams(PointMetadata2));
 
 		// Verify that the initialization made a copy
 		UTEST_EQUAL("Initialization made a copy", PointMetadata3->GetLocalItemCount(), 5);
diff --git a/Engine/Plugins/PCG/Source/PCG/Public/Data/PCGSpatialData.h b/Engine/Plugins/PCG/Source/PCG/Public/Data/PCGSpatialData.h
--- a/Engine/Plugins/PCG/Source/PCG/Public/Data/PCGSpatialData.h
+++ b/Engine/Plugins/PCG/Source/PCG/Public/Data/PCGSpatialData.h
@@ -51,6 +51,9 @@ struct PCG_API FPCGInitializeFromDataParams
 	UPROPERTY(BlueprintReadWrite, EditAnywhere, Category = Default)
 	bool bInheritSpatialData = true;
 
+	// Special boolean to be set to true when we duplicate data.
+	bool bIsDuplicatingData = false;
+
 	/** When initializing metadata from different types, can provide an extra mapping to know how to initialize metadata across domains. If empty, will match with the same domains. */
 	UPROPERTY(BlueprintReadWrite, EditAnywhere, Category = Default)
 	TMap<FPCGMetadataDomainID, FPCGMetadataDomainID> DomainMapping;
diff --git a/Engine/Plugins/PCG/Source/PCG/Public/Metadata/PCGMetadata.h b/Engine/Plugins/PCG/Source/PCG/Public/Metadata/PCGMetadata.h
--- a/Engine/Plugins/PCG/Source/PCG/Public/Metadata/PCGMetadata.h
+++ b/Engine/Plugins/PCG/Source/PCG/Public/Metadata/PCGMetadata.h
@@ -15,10 +15,43 @@
 struct FPCGPoint;
 struct FPCGContext;
 
+// @todo_pcg: Support blueprint
+struct FPCGMetadataInitializeParams
+{
+	explicit FPCGMetadataInitializeParams(const UPCGMetadata* InParent)
+		: Parent(InParent)
+	{}
+
+	// Will create a new DomainInitializeParams with the Default metadata domain, to ease deprecation.
+	PCG_API explicit FPCGMetadataInitializeParams(const UPCGMetadata* InParent, const TArray<PCGMetadataEntryKey>* InOptionalEntriesToCopy);
+
+	// Will create a new DomainInitializeParams with the Default metadata domain, to ease deprecation.
+	PCG_API explicit FPCGMetadataInitializeParams(const UPCGMetadata* InParent,
+		const TSet<FName>& InFilteredAttributes,
+		EPCGMetadataFilterMode InFilterMode = EPCGMetadataFilterMode::ExcludeAttributes,
+		EPCGStringMatchingOperator InMatchOperator = EPCGStringMatchingOperator::Equal,
+		const TArray<PCGMetadataEntryKey>* InOptionalEntriesToCopy = nullptr);
+
+	PCG_API void PopulateDomainInitializeParamsFromParent();
+
+	/** Parent to initialize from */
+	const UPCGMetadata* Parent = nullptr;
+
+	/** Optional mapping for cross domain. In an array since we'll probably never have a lot of them. */
+	TArray<TTuple<FPCGMetadataDomainID, FPCGMetadataDomainID>> DomainMapping;
+
+	/**
+	 * Optional mapping for each domain initialization. In an array since we'll probably never have a lot of them.
+	 * If not set explicitly, it will use the default FPCGMetadataDomainInitializeParams, with the right parent.
+	 */
+	TArray<TTuple<FPCGMetadataDomainID, FPCGMetadataDomainInitializeParams>> DomainInitializeParams;
+};
+
 UCLASS(BlueprintType)
 class PCG_API UPCGMetadata : public UObject
 {
 	friend FPCGMetadataDomain;
+	friend FPCGMetadataInitializeParams;
 	
 	GENERATED_BODY()
 
@@ -40,10 +73,12 @@ public:
 
 	/** Initializes the metadata from a parent metadata, if any (can be null). Copies attributes and values. */
 	UFUNCTION(BlueprintCallable, Category = "PCG|Metadata")
-	void Initialize(const UPCGMetadata* InParent) { InitializeWithAttributeFilter(InParent, TSet<FPCGAttributeIdentifier>()); }
+	void Initialize(const UPCGMetadata* InParent);
 
 	/** Initializes the metadata from a parent metadata, if any (can be null) with the option to not add attributes from the parent. */
-	void Initialize(const UPCGMetadata* InParent, bool bAddAttributesFromParent, const TMap<FPCGMetadataDomainID, FPCGMetadataDomainID>& DomainMapping = {});
+	void Initialize(const UPCGMetadata* InParent, bool bAddAttributesFromParent);
+
+	void Initialize(const FPCGMetadataInitializeParams& InParams);
 
 	/**
 	 * Initializes the metadata from a parent metadata. Copies attributes and values.
@@ -54,28 +89,19 @@ public:
 	 */
 	UFUNCTION(BlueprintCallable, Category = "PCG|Metadata")
 	void InitializeWithAttributeFilter(const UPCGMetadata* InParent, const TSet<FName>& InFilteredAttributes, EPCGMetadataFilterMode InFilterMode = EPCGMetadataFilterMode::ExcludeAttributes, EPCGStringMatchingOperator InMatchOperator = EPCGStringMatchingOperator::Equal);
-	
-	void InitializeWithAttributeFilter(const UPCGMetadata* InParent, const TSet<FPCGAttributeIdentifier>& InFilteredAttributes, EPCGMetadataFilterMode InFilterMode = EPCGMetadataFilterMode::ExcludeAttributes, EPCGStringMatchingOperator InMatchOperator = EPCGStringMatchingOperator::Equal, const TMap<FPCGMetadataDomainID, FPCGMetadataDomainID>& DomainMapping = {});
 
-	/** Initializes the metadata from a parent metadata by copying all attributes to it.
-	* @param InMetadataToCopy        Metadata to copy from
-	* @param InOptionalEntriesToCopy Optional array that contains the keys to copy over. This array order will be respected, so it can also be used to re-order entries. Can be null.
-	*/
+	UE_DEPRECATED(5.6, "Use the FPCGMetadataInitializeParams version")
 	void InitializeAsCopy(const UPCGMetadata* InMetadataToCopy, const TArray<PCGMetadataEntryKey>* InOptionalEntriesToCopy = nullptr);
 
-	/** Initializes the metadata from a parent metadata by copy filtered attributes only to it
-	* @param InMetadataToCopy         Metadata to copy from
-	* @param InFilteredAttributes     Attributes to keep/exclude, can be empty.
-	* @param InFilterMode             Filter to know if we should keep or exclude InFilteredAttributes.
-	* @param InOptionalEntriesToCopy  Optional array that contains the keys to copy over. This array order will be respected, so it can also be used to re-order entries. Can be null.
-	* @param InMatchOperator          How to match the names for the filtering
-	*/
-	void InitializeAsCopyWithAttributeFilter(const UPCGMetadata* InMetadataToCopy, const TSet<FName>& InFilteredAttributes, EPCGMetadataFilterMode InFilterMode = EPCGMetadataFilterMode::ExcludeAttributes, const TArray<PCGMetadataEntryKey>* InOptionalEntriesToCopy = nullptr, EPCGStringMatchingOperator InMatchOperator = EPCGStringMatchingOperator::Equal);
-	void InitializeAsCopyWithAttributeFilter(const UPCGMetadata* InMetadataToCopy, const TSet<FPCGAttributeIdentifier>& InFilteredAttributes, EPCGMetadataFilterMode InFilterMode = EPCGMetadataFilterMode::ExcludeAttributes, const TArray<PCGMetadataEntryKey>* InOptionalEntriesToCopy = nullptr, EPCGStringMatchingOperator InMatchOperator = EPCGStringMatchingOperator::Equal);
+	/** Initializes the metadata from a parent metadata by copying all attributes to it. */
+	void InitializeAsCopy(const FPCGMetadataInitializeParams& InParams);
+
+	UE_DEPRECATED(5.6, "Use InitializeAsCopy with FPCGMetadataInitializeParams version")
+	void InitializeAsCopyWithAttributeFilter(const UPCGMetadata* InMetadataToCopy, const TSet<FName>& InFilteredAttributes, EPCGMetadataFilterMode InFilterMode = EPCGMetadataFilterMode::ExcludeAttributes, const TArray<PCGMetadataEntryKey>* InOptionalEntriesToCopy = nullptr, EPCGStringMatchingOperator InMatchOperator = EPCGStringMatchingOperator::Equal);	
 
 	/** Creates missing attributes from another metadata if they are not currently present - note that this does not copy values */
 	UFUNCTION(BlueprintCallable, Category = "PCG|Metadata")
-	void AddAttributes(const UPCGMetadata* InOther) { AddAttributesFiltered(InOther, TSet<FPCGAttributeIdentifier>()); }
+	void AddAttributes(const UPCGMetadata* InOther);
 
 	/**
 	 * Creates missing attributes from another metadata if they are not currently present - note that this does not copy values.
@@ -87,7 +113,7 @@ public:
 	UFUNCTION(BlueprintCallable, Category = "PCG|Metadata")
 	void AddAttributesFiltered(const UPCGMetadata* InOther, const TSet<FName>& InFilteredAttributes, EPCGMetadataFilterMode InFilterMode = EPCGMetadataFilterMode::ExcludeAttributes, EPCGStringMatchingOperator InMatchOperator = EPCGStringMatchingOperator::Equal);
 
-	void AddAttributesFiltered(const UPCGMetadata* InOther, const TSet<FPCGAttributeIdentifier>& InFilteredAttributes, EPCGMetadataFilterMode InFilterMode = EPCGMetadataFilterMode::ExcludeAttributes, EPCGStringMatchingOperator InMatchOperator = EPCGStringMatchingOperator::Equal);
+	void AddAttributes(const FPCGMetadataInitializeParams& InParams);
 	
 	/** Creates missing attribute from another metadata if it is not currently present - note that this does not copy values */
 	UFUNCTION(BlueprintCallable, Category = "PCG|Metadata", meta = (DisplayName = "Add Attribute"))
@@ -377,7 +403,7 @@ public:
 	* @param InOptionalEntriesToCopy Optional array that contains the keys to copy over. This array order will be respected, so it can also be used to re-order entries. If empty, copy them all.
 	*/
 	UFUNCTION(BlueprintCallable, Category = "PCG|Metadata", meta = (DisplayName = "Initialize As Copy", AutoCreateRefTerm = "InOptionalEntriesToCopy"))
-	void K2_InitializeAsCopy(const UPCGMetadata* InMetadataToCopy, const TArray<int64>& InOptionalEntriesToCopy) { return InitializeAsCopy(InMetadataToCopy, !InOptionalEntriesToCopy.IsEmpty() ? &InOptionalEntriesToCopy : nullptr); }
+	void K2_InitializeAsCopy(const UPCGMetadata* InMetadataToCopy, const TArray<int64>& InOptionalEntriesToCopy);
 
 	/** Initializes the metadata from a parent metadata by copy filtered attributes only to it
 	* @param InMetadataToCopy         Metadata to copy from
@@ -387,10 +413,7 @@ public:
 	* @param InMatchOperator          How to match the names for the filtering
 	*/
 	UFUNCTION(BlueprintCallable, Category = "PCG|Metadata", meta = (DisplayName = "Initialize As Copy With Attribute Filter", AutoCreateRefTerm = "InFilteredAttributes,InOptionalEntriesToCopy"))
-	void K2_InitializeAsCopyWithAttributeFilter(const UPCGMetadata* InMetadataToCopy, const TSet<FName>& InFilteredAttributes, const TArray<int64>& InOptionalEntriesToCopy, EPCGMetadataFilterMode InFilterMode = EPCGMetadataFilterMode::ExcludeAttributes, EPCGStringMatchingOperator InMatchOperator = EPCGStringMatchingOperator::Equal)
-	{
-		return InitializeAsCopyWithAttributeFilter(InMetadataToCopy, FPCGAttributeIdentifier::TransformNameSet(InFilteredAttributes), InFilterMode, !InOptionalEntriesToCopy.IsEmpty() ? &InOptionalEntriesToCopy : nullptr, InMatchOperator);
-	}
+	void K2_InitializeAsCopyWithAttributeFilter(const UPCGMetadata* InMetadataToCopy, const TSet<FName>& InFilteredAttributes, const TArray<int64>& InOptionalEntriesToCopy, EPCGMetadataFilterMode InFilterMode = EPCGMetadataFilterMode::ExcludeAttributes, EPCGStringMatchingOperator InMatchOperator = EPCGStringMatchingOperator::Equal);
 
 	/** Computes Crc from all attributes & keys from outer's data. */
 	void AddToCrc(FArchiveCrc32& Ar, bool bFullDataCrc) const;
@@ -427,6 +450,9 @@ private:
 	template <typename Func>
 	void ForEachValidUniqueMetadataDomain(Func InFunc);
 
+	template <typename Func>
+	void FindFixOrCreateDomainInitializeParams(const FPCGMetadataInitializeParams& InParams, const FPCGMetadataDomainID DomainID, const FPCGMetadataDomain& OtherMetadataDomain, Func&& InFunc);
+
 protected:
 	FPCGMetadataDomain* FindOrCreateMetadataDomain(const FPCGMetadataDomainID& InMetadataDomainID);
 	
diff --git a/Engine/Plugins/PCG/Source/PCG/Public/Metadata/PCGMetadataAttribute.h b/Engine/Plugins/PCG/Source/PCG/Public/Metadata/PCGMetadataAttribute.h
--- a/Engine/Plugins/PCG/Source/PCG/Public/Metadata/PCGMetadataAttribute.h
+++ b/Engine/Plugins/PCG/Source/PCG/Public/Metadata/PCGMetadataAttribute.h
@@ -4,7 +4,7 @@
 
 #include "PCGMetadataCommon.h"
 
-#include "UObject/ObjectPtr.h" // IWYU pragma: keep
+#include "Utils/PCGValueRange.h"
 
 class FPCGMetadataDomain;
 class UPCGMetadata;
@@ -75,9 +75,15 @@ public:
 	/** Bulk getter, to lock in read only once per parent. */
 	void GetValueKeys(const TArrayView<const PCGMetadataEntryKey>& EntryKeys, TArray<PCGMetadataValueKey>& OutValueKeys) const;
 
+	/** Bulk getter, to lock in read only once per parent. */
+	void GetValueKeys(TConstPCGValueRange<PCGMetadataEntryKey> EntryKeys, TArray<PCGMetadataValueKey>& OutValueKeys) const;
+
 	/** Optimized version that take ownership on the Entries passed.*/
 	void GetValueKeys(TArrayView<PCGMetadataEntryKey> EntryKeys, TArray<PCGMetadataValueKey>& OutValueKeys) const;
 
+	/** Optimized version that take ownership on the Entries passed.*/
+	void GetValueKeys(TPCGValueRange<PCGMetadataEntryKey> EntryKeys, TArray<PCGMetadataValueKey>& OutValueKeys) const;
+
 	/** Bulk setter to lock in write only once. */
 	void SetValuesFromValueKeys(const TArrayView<const TTuple<PCGMetadataEntryKey, PCGMetadataValueKey>>& EntryValuePairs, bool bResetValueOnDefaultValueKey = true);
 
@@ -110,7 +116,7 @@ private:
 	// Because we need to update the entry keys if we look for the value keys in the parent, but because the EntryKeys are coming from the outside, we can't modify them.
 	// (as entry keys need to be put into the parent referential when looking for value keys). 
 	// So we will copy internally the EntryKeys to modify them (and only once) if we are not owner of the memory.
-	void GetValueKeys_Internal(const TArrayView<const PCGMetadataEntryKey> EntryKeys, TArrayView<PCGMetadataValueKey> OutValueKeys, TBitArray<>& UnsetValues, bool bOwnerOfEntryKeysView = false) const;
+	void GetValueKeys_Internal(TConstPCGValueRange<PCGMetadataEntryKey> EntryKeys, TArrayView<PCGMetadataValueKey> OutValueKeys, TBitArray<>& UnsetValues, bool bOwnerOfEntryKeysView = false) const;
 
 protected:
 	TMap<PCGMetadataEntryKey, PCGMetadataValueKey> EntryToValueKeyMap;
diff --git a/Engine/Plugins/PCG/Source/PCG/Public/Metadata/PCGMetadataDomain.h b/Engine/Plugins/PCG/Source/PCG/Public/Metadata/PCGMetadataDomain.h
--- a/Engine/Plugins/PCG/Source/PCG/Public/Metadata/PCGMetadataDomain.h
+++ b/Engine/Plugins/PCG/Source/PCG/Public/Metadata/PCGMetadataDomain.h
@@ -5,9 +5,38 @@
 #include "PCGCommon.h"
 #include "Metadata/PCGMetadataCommon.h"
 #include "Metadata/PCGMetadataAttributeTpl.h"
+#include "Utils/PCGValueRange.h"
 
 class UPCGMetadata;
 
+// @todo_pcg: Support blueprint
+struct FPCGMetadataDomainInitializeParams
+{
+	explicit FPCGMetadataDomainInitializeParams(const FPCGMetadataDomain* InParent, const TSet<FName>* InFilteredAttributes = nullptr)
+		: Parent(InParent)
+	{
+		if (InFilteredAttributes)
+		{
+			FilteredAttributes = *InFilteredAttributes;
+		}
+	}
+	
+	/** The parent metadata to use as a template, if any (can be null). */
+	const FPCGMetadataDomain* Parent = nullptr;
+	
+	/** Optional list of attributes to exclude or include when adding the attributes from the parent. */
+	TOptional<TSet<FName>> FilteredAttributes;
+	
+	/** Defines attribute filter operation. */
+	EPCGMetadataFilterMode FilterMode = EPCGMetadataFilterMode::ExcludeAttributes;
+	
+	/** Defines attribute filter operator. */
+	EPCGStringMatchingOperator MatchOperator = EPCGStringMatchingOperator::Equal;
+
+	/** Optional keys to copy over, in case of copy operation. Can be constructed from an array. */
+	TOptional<TConstPCGValueRange<PCGMetadataEntryKey>> OptionalEntriesToCopy;
+};
+
 class PCG_API FPCGMetadataDomain
 {
 	friend UPCGMetadata;
@@ -19,43 +48,22 @@ public:
 	virtual void Serialize(FArchive& InArchive);
 	
 	/** Initializes the sub metadata from a parent sub metadata, if any (can be null). Copies attributes and values. */
-	void Initialize(const FPCGMetadataDomain* InParent) { InitializeWithAttributeFilter(InParent, TSet<FName>()); }
+	void Initialize(const FPCGMetadataDomain* InParent);
 
-	/** Initializes the sub metadata from a parent sub metadata, if any (can be null) with the option to not add attributes from the parent. */
-	void Initialize(const FPCGMetadataDomain* InParent, bool bAddAttributesFromParent);
+	/** Initializes the sub metadata from a parent sub metadata. Copies attributes and values. */
+	void Initialize(const FPCGMetadataDomainInitializeParams& InParams);
 
-	/**
-	 * Initializes the sub metadata from a parent sub metadata. Copies attributes and values.
-	 * @param InParent The parent metadata to use as a template, if any (can be null).
-	 * @param InFilteredAttributes Optional list of attributes to exclude or include when adding the attributes from the parent.
-	 * @param InFilterMode Defines attribute filter operation.
-	 */
-	void InitializeWithAttributeFilter(const FPCGMetadataDomain* InParent, const TSet<FName>& InFilteredAttributes, EPCGMetadataFilterMode InFilterMode = EPCGMetadataFilterMode::ExcludeAttributes, EPCGStringMatchingOperator InMatchOperator = EPCGStringMatchingOperator::Equal);
-
-	/** Initializes the sub metadata from a parent sub metadata by copying all attributes to it.
-	* @param InMetadataToCopy Metadata to copy from
-	* @param InOptionalEntriesToCopy Optional array that contains the keys to copy over. This array order will be respected, so it can also be used to re-order entries. Can be null.
-	*/
-	void InitializeAsCopy(const FPCGMetadataDomain* InMetadataToCopy, const TArray<PCGMetadataEntryKey>* InOptionalEntriesToCopy = nullptr);
+	/** Initializes the sub metadata from a parent sub metadata by copying all attributes to it. */
+	void InitializeAsCopy(const FPCGMetadataDomain* InMetadataToCopy);
 
-	/** Initializes the metadata from a parent metadata by copy filtered attributes only to it
-	* @param InMetadataToCopy Metadata to copy from
-	* @param InFilteredAttributes Attributes to keep/exclude, can be empty.
-	* @param InFilterMode Filter to know if we should keep or exclude InFilteredAttributes.
-	* @param InOptionalEntriesToCopy Optional array that contains the keys to copy over. This array order will be respected, so it can also be used to re-order entries. Can be null.
-	*/
-	void InitializeAsCopyWithAttributeFilter(const FPCGMetadataDomain* InMetadataToCopy, const TSet<FName>& InFilteredAttributes, EPCGMetadataFilterMode InFilterMode = EPCGMetadataFilterMode::ExcludeAttributes, const TArray<PCGMetadataEntryKey>* InOptionalEntriesToCopy = nullptr, EPCGStringMatchingOperator InMatchOperator = EPCGStringMatchingOperator::Equal);
+	/** Initializes the metadata from a parent metadata by copy filtered attributes only to it */
+	void InitializeAsCopy(const FPCGMetadataDomainInitializeParams& InParams);
 
 	/** Creates missing attributes from another metadata if they are not currently present - note that this does not copy values */
-	bool AddAttributes(const FPCGMetadataDomain* InOther) { return AddAttributesFiltered(InOther, TSet<FName>()); }
+	bool AddAttributes(const FPCGMetadataDomain* InOther);
 
-	/**
-	 * Creates missing attributes from another metadata if they are not currently present - note that this does not copy values.
-	 * @param InOther The other metadata to obtain a list of attributes from.
-	 * @param InFilteredAttributes Optional list of attributes to exclude or include when adding the attributes.
-	 * @param InFilterMode Defines attribute filter operation.
-	 */
-	bool AddAttributesFiltered(const FPCGMetadataDomain* InOther, const TSet<FName>& InFilteredAttributes, EPCGMetadataFilterMode InFilterMode = EPCGMetadataFilterMode::ExcludeAttributes, EPCGStringMatchingOperator InMatchOperator = EPCGStringMatchingOperator::Equal);
+	/** Creates missing attributes from another metadata if they are not currently present - note that this does not copy values. */
+	bool AddAttributes(const FPCGMetadataDomainInitializeParams& InParams);
 
 	/** Creates missing attribute from another metadata if it is not currently present - note that this does not copy values */
 	bool AddAttribute(const FPCGMetadataDomain* InOther, FName AttributeName);
@@ -187,9 +195,12 @@ public:
 	/** Metadata chaining mechanism */
 	PCGMetadataEntryKey GetParentKey(PCGMetadataEntryKey LocalItemKey) const;
 
-	/** Metadata chaining mechanism for bulk version. Can provide a mask to only update only a subset of the passed keys. */
+	/** Metadata chaining mechanism for bulk version. Can provide a mask to update only a subset of the passed keys. */
 	void GetParentKeys(TArrayView<PCGMetadataEntryKey> LocalItemKeys, const TBitArray<>* Mask = nullptr) const;
 
+	/** Metadata chaining mechanism for bulk version. Can provide a mask to update only a subset of the passed keys. */
+	void GetParentKeysWithRange(TPCGValueRange<PCGMetadataEntryKey> LocalItemKeys, const TBitArray<>* Mask = nullptr) const;
+
 	/** Attributes operations */
 	void MergeAttributes(PCGMetadataEntryKey InKeyA, const FPCGMetadataDomain* InMetadataA, PCGMetadataEntryKey InKeyB, const FPCGMetadataDomain* InMetadataB, PCGMetadataEntryKey& OutKey, EPCGMetadataOp Op);
 	void MergeAttributesSubset(PCGMetadataEntryKey InKeyA, const FPCGMetadataDomain* InMetadataA, const FPCGMetadataDomain* InMetadataSubetA, PCGMetadataEntryKey InKeyB, const FPCGMetadataDomain* InMetadataB, const FPCGMetadataDomain* InMetadataSubsetb, PCGMetadataEntryKey& OutKey, EPCGMetadataOp Op);
diff --git a/Engine/Plugins/PCG/Source/PCG/Public/Utils/PCGValueRange.h b/Engine/Plugins/PCG/Source/PCG/Public/Utils/PCGValueRange.h
--- a/Engine/Plugins/PCG/Source/PCG/Public/Utils/PCGValueRange.h
+++ b/Engine/Plugins/PCG/Source/PCG/Public/Utils/PCGValueRange.h
@@ -8,6 +8,8 @@
 
 #include "HAL/PlatformString.h" // for INT64_FMT
 
+struct PCGValueRangeHelpers;
+
 /*
 * TPCGValueRange is used to represent a list of InNumElements values but those
 * values are provided by an underlying TStridedView which can have a different number of values.
@@ -40,6 +42,8 @@
 template<typename ElementType, typename ViewType = TStridedView<ElementType>>
 class TPCGValueRange
 {
+	friend PCGValueRangeHelpers;
+
 public:
 	using SizeType = typename ViewType::SizeType;
 	using NonConstElementType = std::remove_const_t<ElementType>;
@@ -174,4 +178,56 @@ private:
 	mutable NonConstElementType StubValue;
 };
 
-template<typename ElementType> using TConstPCGValueRange = TPCGValueRange<const ElementType, TConstStridedView<ElementType>>;
\ No newline at end of file
+template<typename ElementType> using TConstPCGValueRange = TPCGValueRange<const ElementType, TConstStridedView<ElementType>>;
+
+struct PCGValueRangeHelpers
+{
+	// Not instantiatable
+	~PCGValueRangeHelpers() = delete;
+	
+	/** For Array views */
+	template<typename ElementType>
+	static TPCGValueRange<ElementType> MakeValueRange(TArrayView<ElementType> InView)
+	{
+		return TPCGValueRange<ElementType>(MakeStridedView(InView));
+	}
+
+	template<typename ElementType>
+	static TConstPCGValueRange<ElementType> MakeConstValueRange(TArrayView<const ElementType> InView)
+	{
+		return TConstPCGValueRange<ElementType>(MakeConstStridedView(InView));
+	}
+
+	/** For Array */
+	template<typename ElementType, typename ...Args>
+	static TPCGValueRange<ElementType> MakeValueRange(const TArray<ElementType, Args...>& InArray)
+	{
+		return TPCGValueRange<ElementType>(MakeStridedView(InArray));
+	}
+
+	template<typename ElementType, typename ...Args>
+	static TConstPCGValueRange<ElementType> MakeConstValueRange(const TArray<const ElementType, Args...>& InArray)
+	{
+		return TConstPCGValueRange<ElementType>(MakeConstStridedView(InArray));
+	}
+
+	template<typename ElementType, typename ...Args>
+	static TConstPCGValueRange<ElementType> MakeConstValueRange(const TArray<ElementType, Args...>& InArray)
+	{
+		return TConstPCGValueRange<ElementType>(MakeConstStridedView(InArray));
+	}
+
+	/** For TPCGValueRange */
+	template<typename ElementType>
+	static TConstPCGValueRange<ElementType> MakeConstValueRange(TPCGValueRange<ElementType> InView)
+	{
+		return TConstPCGValueRange<ElementType>(MakeConstStridedView(InView.ElementView.GetStride(), &InView.ElementView.GetUnsafe(0), InView.Num()));
+	}
+
+	// Const -> Non-const, use it at your own risk.
+	template<typename ElementType>
+	static TPCGValueRange<ElementType> MakeValueRange_Unsafe(TConstPCGValueRange<ElementType> InView)
+	{
+		return TPCGValueRange<ElementType>(MakeStridedView(InView.ElementView.GetStride(), const_cast<ElementType*>(&InView.ElementView.GetUnsafe(0)), InView.Num()));
+	}
+};
\ No newline at end of file
