[PCG] Add special initialization for composite data's metadata when we duplicate it.

#rnx

diff --git a/Engine/Plugins/PCG/Source/PCG/Private/Data/PCGDifferenceData.cpp b/Engine/Plugins/PCG/Source/PCG/Private/Data/PCGDifferenceData.cpp
--- a/Engine/Plugins/PCG/Source/PCG/Private/Data/PCGDifferenceData.cpp
+++ b/Engine/Plugins/PCG/Source/PCG/Private/Data/PCGDifferenceData.cpp
@@ -395,6 +395,13 @@ void UPCGDifferenceData::InitializeOtherMetadataInternal(const FPCGInitializeFro
 	check(InParams.bInheritMetadata);
 	check(MetadataToInitialize);
 
+	// Duplicate data case, call the spatial base method
+	if (InParams.bIsDuplicatingData)
+	{
+		UPCGSpatialData::InitializeOtherMetadataInternal(InParams, MetadataToInitialize);
+		return;
+	}
+
 	// In the case of the difference, we initialize the data from the source, not the difference itself. The difference will add its own
 	// attributes afterward.
 	if (Source)
diff --git a/Engine/Plugins/PCG/Source/PCG/Private/Data/PCGIntersectionData.cpp b/Engine/Plugins/PCG/Source/PCG/Private/Data/PCGIntersectionData.cpp
--- a/Engine/Plugins/PCG/Source/PCG/Private/Data/PCGIntersectionData.cpp
+++ b/Engine/Plugins/PCG/Source/PCG/Private/Data/PCGIntersectionData.cpp
@@ -350,6 +350,13 @@ void UPCGIntersectionData::InitializeOtherMetadataInternal(const FPCGInitializeF
 	check(GetA());
 	check(GetB());
 
+	// Duplicate data case, call the spatial base method
+	if (InParams.bIsDuplicatingData)
+	{
+		UPCGSpatialData::InitializeOtherMetadataInternal(InParams, MetadataToInitialize);
+		return;
+	}
+
 	// In the case of the intersection, we initialize the data from either A or B, depending on the dimension, not the intersection itself. The intersection will add its own
 	// attributes afterward.
 	FPCGInitializeFromDataParams ParamsCopy = InParams;
diff --git a/Engine/Plugins/PCG/Source/PCG/Private/Data/PCGProjectionData.cpp b/Engine/Plugins/PCG/Source/PCG/Private/Data/PCGProjectionData.cpp
--- a/Engine/Plugins/PCG/Source/PCG/Private/Data/PCGProjectionData.cpp
+++ b/Engine/Plugins/PCG/Source/PCG/Private/Data/PCGProjectionData.cpp
@@ -456,6 +456,13 @@ void UPCGProjectionData::InitializeOtherMetadataInternal(const FPCGInitializeFro
 {
 	check(InParams.bInheritMetadata);
 	check(MetadataToInitialize);
+
+	// Duplicate data case, call the spatial base method
+	if (InParams.bIsDuplicatingData)
+	{
+		UPCGSpatialData::InitializeOtherMetadataInternal(InParams, MetadataToInitialize);
+		return;
+	}
 	
 	if (Source)
 	{
diff --git a/Engine/Plugins/PCG/Source/PCG/Private/Data/PCGUnionData.cpp b/Engine/Plugins/PCG/Source/PCG/Private/Data/PCGUnionData.cpp
--- a/Engine/Plugins/PCG/Source/PCG/Private/Data/PCGUnionData.cpp
+++ b/Engine/Plugins/PCG/Source/PCG/Private/Data/PCGUnionData.cpp
@@ -575,6 +575,13 @@ void UPCGUnionData::InitializeOtherMetadataInternal(const FPCGInitializeFromData
 	check(InParams.bInheritMetadata);
 	check(MetadataToInitialize);
 
+	// Duplicate data case, call the spatial base method
+	if (InParams.bIsDuplicatingData)
+	{
+		UPCGSpatialData::InitializeOtherMetadataInternal(InParams, MetadataToInitialize);
+		return;
+	}
+
 	// In the case of the Union, we initialize the data from all the sources. The union will add its own
 	// attributes afterward.
 	FPCGInitializeFromDataParams CopyParams = InParams;
