[PCG] New IPCGGraphExecutionSource interface aimed at replacing UPCGComponent dependency in graph execution so that PCG graphs can be executed in different context (SceneGraph, Editor tools and more)
- IPCGGraphExecutionSource returns a IPCGGraphExecutionState which should contain the APIs needed for graph execution
- Some nodes have a hard dependency on PCGComponents / Actors, those have been modified to cast the FPCGContext::ExecutionSource
- New signatures for IPCGElement::Initialize and IPCGElement::GetDependenciesCrc with Params

diff --git a/Engine/Plugins/Experimental/PCGInterops/PCGNiagaraInterop/Source/PCGNiagaraInterop/Private/Elements/PCGWriteToNiagaraDataChannel.cpp b/Engine/Plugins/Experimental/PCGInterops/PCGNiagaraInterop/Source/PCGNiagaraInterop/Private/Elements/PCGWriteToNiagaraDataChannel.cpp
--- a/Engine/Plugins/Experimental/PCGInterops/PCGNiagaraInterop/Source/PCGNiagaraInterop/Private/Elements/PCGWriteToNiagaraDataChannel.cpp
+++ b/Engine/Plugins/Experimental/PCGInterops/PCGNiagaraInterop/Source/PCGNiagaraInterop/Private/Elements/PCGWriteToNiagaraDataChannel.cpp
@@ -103,7 +103,7 @@ bool FPCGWriteToNiagaraDataChannelElement::ExecuteInternal(FPCGContext* InContex
 
 	int InputNum = 0;
 
-	UWorld* World = InContext->SourceComponent.IsValid() ? InContext->SourceComponent->GetWorld() : nullptr;
+	UWorld* World = InContext->ExecutionSource.IsValid() ? InContext->ExecutionSource->GetExecutionState().GetWorld() : nullptr;
 	if (!World)
 	{
 		return true;
diff --git a/Engine/Plugins/PCG/Source/PCG/Private/Compute/DataInterfaces/Elements/PCGStaticMeshSpawnerDataInterface.cpp b/Engine/Plugins/PCG/Source/PCG/Private/Compute/DataInterfaces/Elements/PCGStaticMeshSpawnerDataInterface.cpp
--- a/Engine/Plugins/PCG/Source/PCG/Private/Compute/DataInterfaces/Elements/PCGStaticMeshSpawnerDataInterface.cpp
+++ b/Engine/Plugins/PCG/Source/PCG/Private/Compute/DataInterfaces/Elements/PCGStaticMeshSpawnerDataInterface.cpp
@@ -301,6 +301,12 @@ bool UPCGStaticMeshSpawnerDataProvider::PrepareForExecute_GameThread(UPCGDataBin
 		return true;
 	}
 
+	UPCGComponent* SourceComponent = Cast<UPCGComponent>(Context->ExecutionSource.Get());
+	if (!ensure(SourceComponent))
+	{
+		return true;
+	}
+
 	const UPCGMeshSelectorByAttribute* SelectorByAttribute = Cast<UPCGMeshSelectorByAttribute>(Settings->MeshSelectorParameters);
 	if (SelectorAttributeId == INDEX_NONE && SelectorByAttribute)
 	{
@@ -375,7 +381,7 @@ bool UPCGStaticMeshSpawnerDataProvider::PrepareForExecute_GameThread(UPCGDataBin
 
 		if (NumPrimitivesSetup > 0)
 		{
-			Context->SourceComponent->NotifyProceduralInstancesInUse();
+			SourceComponent->NotifyProceduralInstancesInUse();
 
 			// Signal not finished. Compute graph element will wait a frame, giving time for GPU Scene to be updated.
 			return false;
@@ -462,7 +468,13 @@ void UPCGStaticMeshSpawnerDataProvider::CreatePrimitiveDescriptors(FPCGContext*
 		return;
 	}
 
-	const FBox LocalBounds = InContext->SourceComponent->GetGridBounds().ShiftBy(-InContext->SourceComponent->GetOwner()->GetActorLocation());
+	UPCGComponent* SourceComponent = Cast<UPCGComponent>(InContext->ExecutionSource.Get());
+	if (!SourceComponent)
+	{
+		return;
+	}
+
+	const FBox LocalBounds = SourceComponent->GetGridBounds().ShiftBy(-SourceComponent->GetOwner()->GetActorLocation());
 
 	if (Settings->InstanceDataPackerParameters && Settings->InstanceDataPackerParameters->GetAttributeNames(/*OutNames=*/nullptr))
 	{
@@ -735,7 +747,10 @@ bool UPCGStaticMeshSpawnerDataProvider::SetupPrimitives(FPCGContext* InContext,
 		Params.Descriptor = Desc;
 		Params.bAllowDescriptorChanges = false;
 
-		UPCGManagedProceduralISMComponent* MISMC = PCGManagedProceduralISMComponent::GetOrCreateManagedProceduralISMC(TargetActor, InContext->SourceComponent.Get(), Settings->UID, Params);
+		UPCGComponent* SourceComponent = Cast<UPCGComponent>(InContext->ExecutionSource.Get());
+		check(SourceComponent);
+
+		UPCGManagedProceduralISMComponent* MISMC = PCGManagedProceduralISMComponent::GetOrCreateManagedProceduralISMC(TargetActor, SourceComponent, Settings->UID, Params);
 
 		check(MISMC);
 		MISMC->SetCrc(InContext->DependenciesCrc);
diff --git a/Engine/Plugins/PCG/Source/PCG/Private/Compute/DataInterfaces/PCGCustomKernelDataInterface.cpp b/Engine/Plugins/PCG/Source/PCG/Private/Compute/DataInterfaces/PCGCustomKernelDataInterface.cpp
--- a/Engine/Plugins/PCG/Source/PCG/Private/Compute/DataInterfaces/PCGCustomKernelDataInterface.cpp
+++ b/Engine/Plugins/PCG/Source/PCG/Private/Compute/DataInterfaces/PCGCustomKernelDataInterface.cpp
@@ -103,18 +103,18 @@ void UPCGCustomComputeKernelDataProvider::Initialize(const UComputeDataInterface
 	const UPCGCustomKernelDataInterface* DataInterface = CastChecked<UPCGCustomKernelDataInterface>(InDataInterface);
 
 	UPCGDataBinding* Binding = CastChecked<UPCGDataBinding>(InBinding);
-	const UPCGComponent* SourceComponent = Binding->GetSourceComponent().Get();
-	check(SourceComponent);
+	const IPCGGraphExecutionSource* ExecutionSource = Binding->GetExecutionSource();
+	check(ExecutionSource);
 	check(DataInterface->Settings);
 	check(DataInterface->Kernel);
 
 	Kernel = DataInterface->Kernel;
 	ThreadCountMultiplier = DataInterface->Kernel->GetThreadCountMultiplier();
-	Seed = static_cast<uint32>(DataInterface->Settings->GetSeed(SourceComponent));
+	Seed = static_cast<uint32>(DataInterface->Settings->GetSeed(ExecutionSource));
 	SeedSettings = static_cast<uint32>(DataInterface->Settings->Seed);
-	SeedComponent = static_cast<uint32>(SourceComponent->Seed);
+	SeedComponent = static_cast<uint32>(ExecutionSource->GetExecutionState().GetSeed());
 
-	SourceComponentBounds = SourceComponent->GetGridBounds();
+	SourceComponentBounds = ExecutionSource->GetExecutionState().GetBounds();
 }
 
 bool UPCGCustomComputeKernelDataProvider::PrepareForExecute_GameThread(UPCGDataBinding* InBinding)
diff --git a/Engine/Plugins/PCG/Source/PCG/Private/Compute/DataInterfaces/PCGDataCollectionDataInterface.cpp b/Engine/Plugins/PCG/Source/PCG/Private/Compute/DataInterfaces/PCGDataCollectionDataInterface.cpp
--- a/Engine/Plugins/PCG/Source/PCG/Private/Compute/DataInterfaces/PCGDataCollectionDataInterface.cpp
+++ b/Engine/Plugins/PCG/Source/PCG/Private/Compute/DataInterfaces/PCGDataCollectionDataInterface.cpp
@@ -483,8 +483,8 @@ void UPCGDataCollectionDataProvider::Initialize(const UComputeDataInterface* InD
 	bExposeElementCounters = DataInterface->bExposeElementCounters;
 
 #if WITH_EDITOR
-	const UPCGComponent* SourceComponent = Binding->GetSourceComponent().Get();
-	if (ProducerSettings->bIsInspecting && SourceComponent && SourceComponent->IsInspecting())
+	const IPCGGraphExecutionSource* ExecutionSource = Binding->GetExecutionSource();
+	if (ProducerSettings->bIsInspecting && ExecutionSource && ExecutionSource->GetExecutionState().GetInspection().IsInspecting())
 	{
 		ExportMode |= EPCGExportMode::Inspection;
 	}
diff --git a/Engine/Plugins/PCG/Source/PCG/Private/Compute/DataInterfaces/PCGLandscapeDataInterface.cpp b/Engine/Plugins/PCG/Source/PCG/Private/Compute/DataInterfaces/PCGLandscapeDataInterface.cpp
--- a/Engine/Plugins/PCG/Source/PCG/Private/Compute/DataInterfaces/PCGLandscapeDataInterface.cpp
+++ b/Engine/Plugins/PCG/Source/PCG/Private/Compute/DataInterfaces/PCGLandscapeDataInterface.cpp
@@ -284,11 +284,11 @@ void UPCGLandscapeDataProvider::Initialize(const UComputeDataInterface* InDataIn
 
 	const UPCGLandscapeDataInterface* DataInterface = CastChecked<UPCGLandscapeDataInterface>(InDataInterface);
 	UPCGDataBinding* Binding = CastChecked<UPCGDataBinding>(InBinding);
-	const UPCGComponent* SourceComponent = Binding->GetSourceComponent().Get();
-	check(SourceComponent);
+	const IPCGGraphExecutionSource* ExecutionSource = Binding->GetExecutionSource();
+	check(ExecutionSource);
 
-	const UWorld* World = SourceComponent->GetWorld();
-	const FBox ComponentBounds = SourceComponent->GetGridBounds();
+	const UWorld* World = ExecutionSource->GetExecutionState().GetWorld();
+	const FBox ComponentBounds = ExecutionSource->GetExecutionState().GetBounds();
 
 	auto TestLandscape = [World, &ComponentBounds](const ALandscape* InLandscape)
 	{
diff --git a/Engine/Plugins/PCG/Source/PCG/Private/Compute/Elements/PCGComputeGraphElement.cpp b/Engine/Plugins/PCG/Source/PCG/Private/Compute/Elements/PCGComputeGraphElement.cpp
--- a/Engine/Plugins/PCG/Source/PCG/Private/Compute/Elements/PCGComputeGraphElement.cpp
+++ b/Engine/Plugins/PCG/Source/PCG/Private/Compute/Elements/PCGComputeGraphElement.cpp
@@ -18,6 +18,7 @@
 #include "ComputeWorkerInterface.h"
 #include "ComputeFramework/ComputeFramework.h"
 #include "ComputeFramework/ComputeKernelCompileResult.h"
+#include "Engine/World.h"
 #include "Logging/LogVerbosity.h"
 
 #include UE_INLINE_GENERATED_CPP_BY_NAME(PCGComputeGraphElement)
@@ -91,7 +92,7 @@ bool FPCGComputeGraphElement::ExecuteInternal(FPCGContext* InContext) const
 {
 	TRACE_CPUPROFILER_EVENT_SCOPE(FPCGComputeGraphElement::ExecuteInternal);
 	check(InContext);
-	check(InContext->SourceComponent.IsValid());
+	check(InContext->ExecutionSource.IsValid());
 
 	FPCGComputeGraphContext* Context = static_cast<FPCGComputeGraphContext*>(InContext);
 
@@ -116,8 +117,8 @@ bool FPCGComputeGraphElement::ExecuteInternal(FPCGContext* InContext) const
 	{
 		TRACE_CPUPROFILER_EVENT_SCOPE(GetComputeGraph);
 
-		const UPCGComponent* Component = InContext->SourceComponent.Get();
-		const UPCGGraph* TopGraph = Component ? Component->GetGraph() : nullptr;
+		const IPCGGraphExecutionSource* Source = InContext->ExecutionSource.Get();
+		const UPCGGraph* TopGraph = Source ? Source->GetExecutionState().GetGraph() : nullptr;
 		if (!ensure(TopGraph))
 		{
 			return true;
@@ -128,13 +129,16 @@ bool FPCGComputeGraphElement::ExecuteInternal(FPCGContext* InContext) const
 
 		if (TopGraph->IsHierarchicalGenerationEnabled() && !DynamicSubgraph)
 		{
-			if (Component->IsLocalComponent() || Component->IsPartitioned())
+			if (const UPCGComponent* Component = Cast<UPCGComponent>(Source))
 			{
-				Context->GenerationGridSize = Component->GetGenerationGridSize();
+				if (Component->IsLocalComponent() || Component->IsPartitioned())
+				{
+					Context->GenerationGridSize = Component->GetGenerationGridSize();
+				}
 			}
 		}
 
-		if (UPCGSubsystem* Subsystem = Context->SourceComponent->GetSubsystem())
+		if (UPCGSubsystem* Subsystem = UWorld::GetSubsystem<UPCGSubsystem>(Source->GetExecutionState().GetWorld()))
 		{
 			Context->ComputeGraph.Reset(Subsystem->GetComputeGraph(DynamicSubgraph ? DynamicSubgraph : TopGraph, Context->GenerationGridSize, ComputeGraphIndex));
 		}
@@ -458,12 +462,14 @@ bool FPCGComputeGraphElement::ExecuteInternal(FPCGContext* InContext) const
 	case EPCGComputeGraphExecutionPhase::ScheduleComputeGraph: // Fallthrough
 	{
 		TRACE_CPUPROFILER_EVENT_SCOPE(ScheduleComputeGraph);
+		IPCGGraphExecutionSource* ExecutionSource = Context->ExecutionSource.Get();
+		check(ExecutionSource&& ExecutionSource->GetExecutionState().GetWorld());
 
 		const bool bGraphEnqueued = Context->ComputeGraphInstance->EnqueueWork(
 			Context->ComputeGraph.Get(),
-			InContext->SourceComponent->GetScene(),
+			InContext->ExecutionSource->GetExecutionState().GetWorld()->Scene,
 			ComputeTaskExecutionGroup::EndOfFrameUpdate,
-			InContext->SourceComponent->GetOwner()->GetFName(),
+			FName(*InContext->ExecutionSource->GetExecutionState().GetDebugName()),
 			/*InFallbackDelegate=*/FSimpleDelegate::CreateLambda([ContextHandle = Context->GetOrCreateHandle()]()
 			{
 				// This render thread delegate will be executed if SubmitWork fails at any stage.
@@ -567,10 +573,10 @@ void FPCGComputeGraphElement::PostExecuteInternal(FPCGContext* InContext) const
 		for (TSoftObjectPtr<const UPCGNode> NodePtr : Context->DataBinding->Graph->KernelToNode)
 		{
 			const UPCGNode* Node = NodePtr.Get();
-			UPCGComponent* Component = Context->SourceComponent.Get();
-			if (Component && Context->Stack && Node)
+			IPCGGraphExecutionSource* ExecutionSource = Context->ExecutionSource.Get();
+			if (ExecutionSource && Context->Stack && Node)
 			{
-				Component->NotifyNodeExecuted(Node, Context->Stack, /*InTimer=*/nullptr, /*bUsedCache*/false);
+				ExecutionSource->GetExecutionState().GetInspection().NotifyNodeExecuted(Node, Context->Stack, /*InTimer=*/nullptr, /*bUsedCache*/false);
 			}
 		}
 	}
@@ -650,7 +656,7 @@ bool FPCGComputeGraphElement::ExecuteDebugDraw(FPCGComputeGraphContext* InContex
 
 bool FPCGComputeGraphElement::StoreDataForInspection(FPCGComputeGraphContext* InContext) const
 {
-	if (!InContext || !InContext->DataBinding || InContext->DataBinding->DataToInspect.IsEmpty() || !InContext->SourceComponent.IsValid())
+	if (!InContext || !InContext->DataBinding || InContext->DataBinding->DataToInspect.IsEmpty() || !InContext->ExecutionSource.IsValid())
 	{
 		return true;
 	}
@@ -709,7 +715,7 @@ bool FPCGComputeGraphElement::StoreDataForInspection(FPCGComputeGraphContext* In
 		const UPCGNode* Node = Cast<UPCGNode>(SettingsAndData.Key->GetOuter());
 
 		// TODO: Input data not yet supported.
-		InContext->SourceComponent->StoreInspectionData(InContext->Stack, Node, /*InTimer=*/nullptr, /*InInputData=*/{}, SettingsAndData.Get<1>(), /*bUsedCache*/false);
+		InContext->ExecutionSource->GetExecutionState().GetInspection().StoreInspectionData(InContext->Stack, Node, /*InTimer=*/nullptr, /*InInputData=*/{}, SettingsAndData.Get<1>(), /*bUsedCache*/false);
 	}
 
 	InContext->DataBinding->DataToInspect.Empty();
@@ -719,7 +725,7 @@ bool FPCGComputeGraphElement::StoreDataForInspection(FPCGComputeGraphContext* In
 
 void FPCGComputeGraphElement::LogCompilationMessages(FPCGComputeGraphContext* InContext) const
 {
-	if (InContext->SourceComponent.IsValid() && InContext->Stack)
+	if (InContext->ExecutionSource.IsValid() && InContext->Stack)
 	{
 		using FNodeAndCompileMessages = const TPair<TObjectKey<const UPCGNode>, TArray<FComputeKernelCompileMessage>>;
 		for (FNodeAndCompileMessages& NodeAndCompileMessages : InContext->ComputeGraph->KernelToCompileMessages)
diff --git a/Engine/Plugins/PCG/Source/PCG/Private/Compute/PCGDataBinding.cpp b/Engine/Plugins/PCG/Source/PCG/Private/Compute/PCGDataBinding.cpp
--- a/Engine/Plugins/PCG/Source/PCG/Private/Compute/PCGDataBinding.cpp
+++ b/Engine/Plugins/PCG/Source/PCG/Private/Compute/PCGDataBinding.cpp
@@ -229,13 +229,13 @@ void UPCGDataBinding::ApplyElementCounts_GameThread()
 	UpdateCountsOnDebugData(DataToDebug);
 }
 
-TWeakObjectPtr<UPCGComponent> UPCGDataBinding::GetSourceComponent() const
+IPCGGraphExecutionSource* UPCGDataBinding::GetExecutionSource() const
 {
 	if (TSharedPtr<FPCGContextHandle> SharedHandle = ContextHandle.Pin())
 	{
 		if (FPCGContext* ContextPtr = SharedHandle->GetContext())
 		{
-			return ContextPtr->SourceComponent;
+			return ContextPtr->ExecutionSource.Get();
 		}
 	}
 
diff --git a/Engine/Plugins/PCG/Source/PCG/Private/Compute/PCGDataBinding.h b/Engine/Plugins/PCG/Source/PCG/Private/Compute/PCGDataBinding.h
--- a/Engine/Plugins/PCG/Source/PCG/Private/Compute/PCGDataBinding.h
+++ b/Engine/Plugins/PCG/Source/PCG/Private/Compute/PCGDataBinding.h
@@ -99,7 +99,7 @@ public:
 
 	const TMap<FSourceBufferAndAttributeIndex, /*Attribute Index*/int32>& GetAttributeIdRemap() const { return AttributeIdRemap; }
 
-	TWeakObjectPtr<UPCGComponent> GetSourceComponent() const;
+	IPCGGraphExecutionSource* GetExecutionSource() const;
 
 	/** Helper to get index in input data collection of the first data item for the given kernel and input pin label. */
 	int32 GetFirstInputDataIndex(const UPCGComputeKernel* InKernel, FName InPinLabel) const;
diff --git a/Engine/Plugins/PCG/Source/PCG/Private/Data/PCGLandscapeData.cpp b/Engine/Plugins/PCG/Source/PCG/Private/Data/PCGLandscapeData.cpp
--- a/Engine/Plugins/PCG/Source/PCG/Private/Data/PCGLandscapeData.cpp
+++ b/Engine/Plugins/PCG/Source/PCG/Private/Data/PCGLandscapeData.cpp
@@ -703,7 +703,7 @@ TArray<FPCGTaskId> UPCGLandscapeData::PrepareForSpatialQuery(FPCGContext* InCont
 				
 		return true;
 	},
-	InContext->SourceComponent.Get());
+	InContext->ExecutionSource.Get());
 	Params.bCanExecuteOnlyOnMainThread = false;
 
 	return { InContext->ScheduleGeneric(Params) };
diff --git a/Engine/Plugins/PCG/Source/PCG/Private/Elements/ControlFlow/PCGQualityBranch.cpp b/Engine/Plugins/PCG/Source/PCG/Private/Elements/ControlFlow/PCGQualityBranch.cpp
--- a/Engine/Plugins/PCG/Source/PCG/Private/Elements/ControlFlow/PCGQualityBranch.cpp
+++ b/Engine/Plugins/PCG/Source/PCG/Private/Elements/ControlFlow/PCGQualityBranch.cpp
@@ -9,6 +9,7 @@
 #include "PCGPin.h"
 #include "PCGSubsystem.h"
 #include "Elements/PCGGather.h"
+#include "Helpers/PCGHelpers.h"
 
 #define LOCTEXT_NAMESPACE "FPCGQualityBranchElement"
 
@@ -137,10 +138,10 @@ FPCGElementPtr UPCGQualityBranchSettings::CreateElement() const
 	return MakeShared<FPCGQualityBranchElement>();
 }
 
-void FPCGQualityBranchElement::GetDependenciesCrc(const FPCGDataCollection& InInput, const UPCGSettings* InSettings, UPCGComponent* InComponent, FPCGCrc& OutCrc) const
+void FPCGQualityBranchElement::GetDependenciesCrc(const FPCGGetDependenciesCrcParams& InParams, FPCGCrc& OutCrc) const
 {
 	FPCGCrc Crc;
-	IPCGElement::GetDependenciesCrc(InInput, InSettings, InComponent, Crc);
+	IPCGElement::GetDependenciesCrc(InParams, Crc);
 
 	// Add the quality level to CRC so that we can fully refresh when quality level changes.
 	Crc.Combine(UPCGSubsystem::GetPCGQualityLevel());
@@ -155,7 +156,7 @@ bool FPCGQualityBranchElement::ExecuteInternal(FPCGContext* Context) const
 	const UPCGQualityBranchSettings* Settings = Context->GetInputSettings<UPCGQualityBranchSettings>();
 	check(Settings);
 
-	if (Context->SourceComponent.IsValid() && !Context->SourceComponent->IsManagedByRuntimeGenSystem())
+	if (!PCGHelpers::IsRuntimeGeneration(Context->ExecutionSource.Get()))
 	{
 		PCGLog::LogWarningOnGraph(LOCTEXT("UsedOutsideOfRuntimeGen", "Using Runtime Quality Branch node in a non-RuntimeGen execution domain. Set your PCG component's 'Generation Trigger' to 'Generate At Runtime'."), Context);
 	}
diff --git a/Engine/Plugins/PCG/Source/PCG/Private/Elements/ControlFlow/PCGQualitySelect.cpp b/Engine/Plugins/PCG/Source/PCG/Private/Elements/ControlFlow/PCGQualitySelect.cpp
--- a/Engine/Plugins/PCG/Source/PCG/Private/Elements/ControlFlow/PCGQualitySelect.cpp
+++ b/Engine/Plugins/PCG/Source/PCG/Private/Elements/ControlFlow/PCGQualitySelect.cpp
@@ -9,6 +9,7 @@
 #include "PCGPin.h"
 #include "PCGSubsystem.h"
 #include "Elements/PCGGather.h"
+#include "Helpers/PCGHelpers.h"
 
 #define LOCTEXT_NAMESPACE "FPCGQualitySelectElement"
 
@@ -114,10 +115,10 @@ FPCGElementPtr UPCGQualitySelectSettings::CreateElement() const
 	return MakeShared<FPCGQualitySelectElement>();
 }
 
-void FPCGQualitySelectElement::GetDependenciesCrc(const FPCGDataCollection& InInput, const UPCGSettings* InSettings, UPCGComponent* InComponent, FPCGCrc& OutCrc) const
+void FPCGQualitySelectElement::GetDependenciesCrc(const FPCGGetDependenciesCrcParams& InParams, FPCGCrc& OutCrc) const
 {
 	FPCGCrc Crc;
-	IPCGElement::GetDependenciesCrc(InInput, InSettings, InComponent, Crc);
+	IPCGElement::GetDependenciesCrc(InParams, Crc);
 
 	// Add the quality level to CRC so that we can fully refresh when quality level changes.
 	Crc.Combine(UPCGSubsystem::GetPCGQualityLevel());
@@ -132,7 +133,7 @@ bool FPCGQualitySelectElement::ExecuteInternal(FPCGContext* Context) const
 	const UPCGQualitySelectSettings* Settings = Context->GetInputSettings<UPCGQualitySelectSettings>();
 	check(Settings);
 
-	if (Context->SourceComponent.IsValid() && !Context->SourceComponent->IsManagedByRuntimeGenSystem())
+	if (!PCGHelpers::IsRuntimeGeneration(Context->ExecutionSource.Get()))
 	{
 		PCGLog::LogWarningOnGraph(LOCTEXT("UsedOutsideOfRuntimeGen", "Using Runtime Quality Select node in a non-RuntimeGen execution domain. Set your PCG component's 'Generation Trigger' to 'Generate At Runtime'."), Context);
 	}
diff --git a/Engine/Plugins/PCG/Source/PCG/Private/Elements/Landscape/PCGWaitLandscapeReady.cpp b/Engine/Plugins/PCG/Source/PCG/Private/Elements/Landscape/PCGWaitLandscapeReady.cpp
--- a/Engine/Plugins/PCG/Source/PCG/Private/Elements/Landscape/PCGWaitLandscapeReady.cpp
+++ b/Engine/Plugins/PCG/Source/PCG/Private/Elements/Landscape/PCGWaitLandscapeReady.cpp
@@ -54,7 +54,8 @@ bool FPCGWaitLandscapeReadyElement::ExecuteInternal(FPCGContext* InContext) cons
 		auto BoundsCheck = [](const AActor*) -> bool { return true; };
 		auto SelfIgnoreCheck = [](const AActor*) -> bool { return true; };
 
-		TArray<AActor*> FoundLandscapeProxies = PCGActorSelector::FindActors(ActorSelector, Context->SourceComponent.Get(), BoundsCheck, SelfIgnoreCheck);
+		const UPCGComponent* SourceComponent = Cast<UPCGComponent>(Context->ExecutionSource.Get());
+		TArray<AActor*> FoundLandscapeProxies = PCGActorSelector::FindActors(ActorSelector, SourceComponent, BoundsCheck, SelfIgnoreCheck);
 
 		for (AActor* Proxy : FoundLandscapeProxies)
 		{
@@ -96,39 +97,30 @@ bool FPCGWaitLandscapeReadyElement::ExecuteInternal(FPCGContext* InContext) cons
 	}
 
 	// At this point, we need to go to sleep for at least the remainder of the frame
-	UPCGSubsystem* Subsystem = Context->SourceComponent.IsValid() ? Context->SourceComponent->GetSubsystem() : nullptr;
-	if (ensure(Subsystem))
-	{
-		Context->bIsPaused = true;
+	Context->bIsPaused = true;
 
-		Subsystem->ScheduleGeneric(
-			[ContextHandle = Context->GetOrCreateHandle()]() // Normal execution: wake up the current task
+	Context->ScheduleGeneric(FPCGScheduleGenericParams(
+		[ContextHandle = Context->GetOrCreateHandle()](FPCGContext*) // Normal execution: wake up the current task
+		{
+			FPCGContext::FSharedContext<FPCGWaitLandscapeReadyElementContext> SharedContext(ContextHandle);
+			if (FPCGWaitLandscapeReadyElementContext* ContextPtr = SharedContext.Get())
 			{
-				FPCGContext::FSharedContext<FPCGWaitLandscapeReadyElementContext> SharedContext(ContextHandle);
-				if (FPCGWaitLandscapeReadyElementContext* ContextPtr = SharedContext.Get())
-				{
-					ContextPtr->bIsPaused = false;
-				}
-				return true;
-			},
-			[ContextHandle = Context->GetOrCreateHandle()]() // On abort: wakeup and cancel
+				ContextPtr->bIsPaused = false;
+			}
+			return true;
+		},
+		[ContextHandle = Context->GetOrCreateHandle()](FPCGContext*) // On abort: wakeup and cancel
+		{
+			FPCGContext::FSharedContext<FPCGWaitLandscapeReadyElementContext> SharedContext(ContextHandle);
+			if (FPCGWaitLandscapeReadyElementContext* ContextPtr = SharedContext.Get())
 			{
-				FPCGContext::FSharedContext<FPCGWaitLandscapeReadyElementContext> SharedContext(ContextHandle);
-				if (FPCGWaitLandscapeReadyElementContext* ContextPtr = SharedContext.Get())
-				{
-					ContextPtr->bIsPaused = false;
-					ContextPtr->OutputData.bCancelExecution = true;
-				}
-				return true;
-			},
-			Context->SourceComponent.Get(),
-			{});
-
-		return false;
-	}
-	else
-	{
-		return true;
-	}
+				ContextPtr->bIsPaused = false;
+				ContextPtr->OutputData.bCancelExecution = true;
+			}
+		},
+		Context->ExecutionSource.Get(),
+		{}));
+
+	return false;
 }
 
diff --git a/Engine/Plugins/PCG/Source/PCG/Private/Elements/Metadata/PCGMetadataOpElementBase.cpp b/Engine/Plugins/PCG/Source/PCG/Private/Elements/Metadata/PCGMetadataOpElementBase.cpp
--- a/Engine/Plugins/PCG/Source/PCG/Private/Elements/Metadata/PCGMetadataOpElementBase.cpp
+++ b/Engine/Plugins/PCG/Source/PCG/Private/Elements/Metadata/PCGMetadataOpElementBase.cpp
@@ -988,13 +988,13 @@ bool FPCGMetadataElementBase::ExecuteInternal(FPCGContext* Context) const
 	});
 }
 
-void FPCGMetadataElementBase::GetDependenciesCrc(const FPCGDataCollection& InInput, const UPCGSettings* InSettings, UPCGComponent* InComponent, FPCGCrc& OutCrc) const
+void FPCGMetadataElementBase::GetDependenciesCrc(const FPCGGetDependenciesCrcParams& InParams, FPCGCrc& OutCrc) const
 {
 	FPCGCrc Crc;
-	IPCGElement::GetDependenciesCrc(InInput, InSettings, InComponent, Crc);
+	IPCGElement::GetDependenciesCrc(InParams, Crc);
 
 	// Add the default values to the crc
-	if (const UPCGMetadataSettingsBase* Settings = Cast<UPCGMetadataSettingsBase>(InSettings))
+	if (const UPCGMetadataSettingsBase* Settings = Cast<UPCGMetadataSettingsBase>(InParams.Settings))
 	{
 		FArchiveCrc32 Crc32;
 		/*
diff --git a/Engine/Plugins/PCG/Source/PCG/Private/Elements/PCGAddComponent.cpp b/Engine/Plugins/PCG/Source/PCG/Private/Elements/PCGAddComponent.cpp
--- a/Engine/Plugins/PCG/Source/PCG/Private/Elements/PCGAddComponent.cpp
+++ b/Engine/Plugins/PCG/Source/PCG/Private/Elements/PCGAddComponent.cpp
@@ -233,7 +233,8 @@ bool FPCGAddComponentElement::ExecuteInternal(FPCGContext* InContext) const
 	TArray<FPCGTaggedData> TargetActors = Context->InputData.GetInputsByPin(PCGAddComponentConstants::TargetActorPinLabel);
 
 	// Early outs & validation
-	if (Inputs.Num() == 0 || !Context->SourceComponent.Get())
+	const UPCGComponent* SourceComponent = Cast<UPCGComponent>(Context->ExecutionSource.Get());
+	if (Inputs.Num() == 0 || !SourceComponent)
 	{
 		return true;
 	}
@@ -340,7 +341,7 @@ bool FPCGAddComponentElement::ExecuteInternal(FPCGContext* InContext) const
 		else
 		{
 			// Wrap on default actor namely the owner of the PCG component.
-			TargetActorAccessor = MakeUnique<FPCGConstantValueAccessor<FSoftObjectPath>>(FSoftObjectPath(Context->SourceComponent->GetOwner()));
+			TargetActorAccessor = MakeUnique<FPCGConstantValueAccessor<FSoftObjectPath>>(FSoftObjectPath(SourceComponent->GetOwner()));
 			TargetActorKeys = MakeUnique<FPCGAttributeAccessorKeysSingleObjectPtr<void>>();
 		}
 
@@ -404,7 +405,7 @@ void FPCGAddComponentElement::AddComponents(FPCGAddComponentContext* Context,
 	check(OutputComponentAccessor && OutputComponentKeys);
 
 	// TODO: implement component reuse
-	UPCGComponent* SourceComponent = Context->SourceComponent.Get();
+	UPCGComponent* SourceComponent = Cast<UPCGComponent>(Context->ExecutionSource.Get());
 	check(SourceComponent);
 	const bool bIsInPreviewMode = SourceComponent->IsInPreviewMode();
 	const EObjectFlags ObjectFlags = bIsInPreviewMode ? RF_Transient : RF_NoFlags;
diff --git a/Engine/Plugins/PCG/Source/PCG/Private/Elements/PCGCreatePoints.cpp b/Engine/Plugins/PCG/Source/PCG/Private/Elements/PCGCreatePoints.cpp
--- a/Engine/Plugins/PCG/Source/PCG/Private/Elements/PCGCreatePoints.cpp
+++ b/Engine/Plugins/PCG/Source/PCG/Private/Elements/PCGCreatePoints.cpp
@@ -68,13 +68,13 @@ bool FPCGCreatePointsElement::ExecuteInternal(FPCGContext* Context) const
 {
 	TRACE_CPUPROFILER_EVENT_SCOPE(FPCGCreatePointsElement::Execute);
 	
-	check(Context && Context->SourceComponent.Get());
+	check(Context && Context->ExecutionSource.Get());
 
 	const UPCGCreatePointsSettings* Settings = Context->GetInputSettings<UPCGCreatePointsSettings>();
 	check(Settings);
 
 	// Used for culling, regardless of generation coordinate space
-	const UPCGSpatialData* CullingShape = Settings->bCullPointsOutsideVolume ? Cast<UPCGSpatialData>(Context->SourceComponent->GetActorPCGData()) : nullptr;
+	const UPCGSpatialData* CullingShape = Settings->bCullPointsOutsideVolume ? Cast<UPCGSpatialData>(Context->ExecutionSource->GetExecutionState().GetSelfData()) : nullptr;
 
 	// Early out if the culling shape isn't valid
 	if (!CullingShape && Settings->bCullPointsOutsideVolume)
@@ -83,17 +83,16 @@ bool FPCGCreatePointsElement::ExecuteInternal(FPCGContext* Context) const
 		return true;
 	}
 
-	FTransform LocalTransform = FTransform::Identity;
+	FTransform LocalTransform = Settings->CoordinateSpace == EPCGCoordinateSpace::World ? FTransform::Identity : Context->ExecutionSource->GetExecutionState().GetTransform();
 
-	if(Settings->CoordinateSpace == EPCGCoordinateSpace::OriginalComponent)
+	if (Settings->CoordinateSpace == EPCGCoordinateSpace::OriginalComponent)
 	{
-		check(Context->SourceComponent->GetOriginalComponent() && Context->SourceComponent->GetOriginalComponent()->GetOwner());
-		LocalTransform = Context->SourceComponent->GetOriginalComponent()->GetOwner()->GetActorTransform();
-	}
-	else if (Settings->CoordinateSpace == EPCGCoordinateSpace::LocalComponent)
-	{
-		check(Context->SourceComponent->GetOwner());
-		LocalTransform = Context->SourceComponent->GetOwner()->GetActorTransform();
+		UPCGComponent* SourceComponent = Cast<UPCGComponent>(Context->ExecutionSource.Get());
+		if (SourceComponent)
+		{
+			check(SourceComponent->GetOriginalComponent() && SourceComponent->GetOriginalComponent()->GetOwner());
+			LocalTransform = SourceComponent->GetOriginalComponent()->GetOwner()->GetActorTransform();
+		}
 	}
 
 	// Reset scale as we are not going to derive the points size from it
@@ -212,4 +211,56 @@ bool FPCGCreatePointsElement::IsCacheable(const UPCGSettings* InSettings) const
 	return Settings && Settings->CoordinateSpace == EPCGCoordinateSpace::World;
 }
 
+namespace PCGCreatePointsHelper
+{
+	UPCGData* GetDependenciesData(IPCGGraphExecutionSource* InExecutionSource, EPCGCoordinateSpace InCoordinateSpace, bool bInCullPointsOutsideVolume)
+	{
+		if (InExecutionSource)
+		{
+			if (bInCullPointsOutsideVolume || InCoordinateSpace == EPCGCoordinateSpace::LocalComponent)
+			{
+				return InExecutionSource->GetExecutionState().GetSelfData();
+			}
+			else if (InCoordinateSpace == EPCGCoordinateSpace::OriginalComponent)
+			{
+				const UPCGComponent* SourceComponent = Cast<UPCGComponent>(InExecutionSource);
+				if (SourceComponent)
+				{
+					return SourceComponent->GetOriginalActorPCGData();
+				}
+				else
+				{
+					return InExecutionSource->GetExecutionState().GetSelfData();
+				}
+			}
+		}
+
+		return nullptr;
+	}
+}
+
+void FPCGCreatePointsElement::GetDependenciesCrc(const FPCGGetDependenciesCrcParams& InParams, FPCGCrc& OutCrc) const
+{
+	FPCGCrc Crc;
+	IPCGElement::GetDependenciesCrc(InParams, Crc);
+
+	if (const UPCGCreatePointsSettings* Settings = Cast<UPCGCreatePointsSettings>(InParams.Settings))
+	{
+		int CoordinateSpace = static_cast<int>(EPCGCoordinateSpace::World);
+		bool bCullPointsOutsideVolume = false;
+		PCGSettingsHelpers::GetOverrideValue(*InParams.InputData, Settings, GET_MEMBER_NAME_CHECKED(UPCGCreatePointsSettings, CoordinateSpace), static_cast<int>(Settings->CoordinateSpace), CoordinateSpace);
+		PCGSettingsHelpers::GetOverrideValue(*InParams.InputData, Settings, GET_MEMBER_NAME_CHECKED(UPCGCreatePointsSettings, bCullPointsOutsideVolume), Settings->bCullPointsOutsideVolume, bCullPointsOutsideVolume);
+
+		// We're using the bounds of the pcg volume, so we extract the actor data here
+		EPCGCoordinateSpace EnumCoordinateSpace = static_cast<EPCGCoordinateSpace>(CoordinateSpace);
+
+		if (UPCGData* Data = PCGCreatePointsHelper::GetDependenciesData(InParams.ExecutionSource, EnumCoordinateSpace, bCullPointsOutsideVolume))
+		{
+			Crc.Combine(Data->GetOrComputeCrc(/*bFullDataCrc=*/false));
+		}
+	}
+
+	OutCrc = Crc;
+}
+
 #undef LOCTEXT_NAMESPACE
diff --git a/Engine/Plugins/PCG/Source/PCG/Private/Elements/PCGCreatePointsGrid.cpp b/Engine/Plugins/PCG/Source/PCG/Private/Elements/PCGCreatePointsGrid.cpp
--- a/Engine/Plugins/PCG/Source/PCG/Private/Elements/PCGCreatePointsGrid.cpp
+++ b/Engine/Plugins/PCG/Source/PCG/Private/Elements/PCGCreatePointsGrid.cpp
@@ -5,6 +5,7 @@
 #include "PCGComponent.h"
 #include "PCGContext.h"
 #include "Data/PCGPointArrayData.h"
+#include "Elements/PCGCreatePoints.h"
 #include "Helpers/PCGAsync.h"
 #include "Helpers/PCGHelpers.h"
 #include "Helpers/PCGSettingsHelpers.h"
@@ -51,13 +52,13 @@ bool FPCGCreatePointsGridElement::ExecuteInternal(FPCGContext* Context) const
 {
 	TRACE_CPUPROFILER_EVENT_SCOPE(FPCGCreatePointsGridElement::Execute);
 	
-	check(Context && Context->SourceComponent.Get());
+	check(Context && Context->ExecutionSource.Get());
 
 	const UPCGCreatePointsGridSettings* Settings = Context->GetInputSettings<UPCGCreatePointsGridSettings>();
 	check(Settings);
 
 	// Used for culling, regardless of generation coordinate space
-	const UPCGSpatialData* CullingShape = Settings->bCullPointsOutsideVolume ? Cast<UPCGSpatialData>(Context->SourceComponent->GetActorPCGData()) : nullptr;
+	const UPCGSpatialData* CullingShape = Settings->bCullPointsOutsideVolume ? Cast<UPCGSpatialData>(Context->ExecutionSource->GetExecutionState().GetSelfData()) : nullptr;
 
 	// Early out if the culling shape isn't valid
 	if (!CullingShape && Settings->bCullPointsOutsideVolume)
@@ -66,17 +67,16 @@ bool FPCGCreatePointsGridElement::ExecuteInternal(FPCGContext* Context) const
 		return true;
 	}
 
-	FTransform LocalTransform = FTransform::Identity;
+	FTransform LocalTransform = Settings->CoordinateSpace == EPCGCoordinateSpace::World ? FTransform::Identity : Context->ExecutionSource->GetExecutionState().GetTransform();
 
 	if (Settings->CoordinateSpace == EPCGCoordinateSpace::OriginalComponent)
 	{
-		check(Context->SourceComponent->GetOriginalComponent() && Context->SourceComponent->GetOriginalComponent()->GetOwner());
-		LocalTransform = Context->SourceComponent->GetOriginalComponent()->GetOwner()->GetActorTransform();
-	}
-	else if (Settings->CoordinateSpace == EPCGCoordinateSpace::LocalComponent)
-	{
-		check(Context->SourceComponent->GetOwner());
-		LocalTransform = Context->SourceComponent->GetOwner()->GetActorTransform();
+		UPCGComponent* SourceComponent = Cast<UPCGComponent>(Context->ExecutionSource.Get());
+		if (SourceComponent)
+		{
+			check(SourceComponent->GetOriginalComponent() && SourceComponent->GetOriginalComponent()->GetOwner());
+			LocalTransform = SourceComponent->GetOriginalComponent()->GetOwner()->GetActorTransform();
+		}
 	}
 
 	// Reset the scale on the local transform since we don't want to derive the points scale from the referential
@@ -305,27 +305,24 @@ bool FPCGCreatePointsGridElement::ExecuteInternal(FPCGContext* Context) const
 	return true;
 }
 
-void FPCGCreatePointsGridElement::GetDependenciesCrc(const FPCGDataCollection& InInput, const UPCGSettings* InSettings, UPCGComponent* InComponent, FPCGCrc& OutCrc) const
+void FPCGCreatePointsGridElement::GetDependenciesCrc(const FPCGGetDependenciesCrcParams& InParams, FPCGCrc& OutCrc) const
 {
 	FPCGCrc Crc;
-	IPCGElement::GetDependenciesCrc(InInput, InSettings, InComponent, Crc);
+	IPCGElement::GetDependenciesCrc(InParams, Crc);
 	
-	if (const UPCGCreatePointsGridSettings* Settings = Cast<UPCGCreatePointsGridSettings>(InSettings))
+	if (const UPCGCreatePointsGridSettings* Settings = Cast<UPCGCreatePointsGridSettings>(InParams.Settings))
 	{
 		int CoordinateSpace = static_cast<int>(EPCGCoordinateSpace::World);
 		bool bCullPointsOutsideVolume = false;
-		PCGSettingsHelpers::GetOverrideValue(InInput, Settings, GET_MEMBER_NAME_CHECKED(UPCGCreatePointsGridSettings, CoordinateSpace), static_cast<int>(Settings->CoordinateSpace), CoordinateSpace);
-		PCGSettingsHelpers::GetOverrideValue(InInput, Settings, GET_MEMBER_NAME_CHECKED(UPCGCreatePointsGridSettings, bCullPointsOutsideVolume), Settings->bCullPointsOutsideVolume, bCullPointsOutsideVolume);
+		PCGSettingsHelpers::GetOverrideValue(*InParams.InputData, Settings, GET_MEMBER_NAME_CHECKED(UPCGCreatePointsGridSettings, CoordinateSpace), static_cast<int>(Settings->CoordinateSpace), CoordinateSpace);
+		PCGSettingsHelpers::GetOverrideValue(*InParams.InputData, Settings, GET_MEMBER_NAME_CHECKED(UPCGCreatePointsGridSettings, bCullPointsOutsideVolume), Settings->bCullPointsOutsideVolume, bCullPointsOutsideVolume);
 		
 		// We're using the bounds of the pcg volume, so we extract the actor data here
 		EPCGCoordinateSpace EnumCoordinateSpace = static_cast<EPCGCoordinateSpace>(CoordinateSpace);
 
-		if ((EnumCoordinateSpace == EPCGCoordinateSpace::OriginalComponent || EnumCoordinateSpace == EPCGCoordinateSpace::LocalComponent || bCullPointsOutsideVolume) && InComponent)
+		if (UPCGData* Data = PCGCreatePointsHelper::GetDependenciesData(InParams.ExecutionSource, EnumCoordinateSpace, bCullPointsOutsideVolume))
 		{
-			if (const UPCGData* Data = InComponent->GetActorPCGData())
-			{
-				Crc.Combine(Data->GetOrComputeCrc(/*bFullDataCrc=*/false));
-			}
+			Crc.Combine(Data->GetOrComputeCrc(/*bFullDataCrc=*/false));
 		}
 	}
 
diff --git a/Engine/Plugins/PCG/Source/PCG/Private/Elements/PCGCreatePointsSphere.cpp b/Engine/Plugins/PCG/Source/PCG/Private/Elements/PCGCreatePointsSphere.cpp
--- a/Engine/Plugins/PCG/Source/PCG/Private/Elements/PCGCreatePointsSphere.cpp
+++ b/Engine/Plugins/PCG/Source/PCG/Private/Elements/PCGCreatePointsSphere.cpp
@@ -4,10 +4,10 @@
 
 #include "PCGComponent.h"
 #include "Data/PCGPointData.h"
+#include "Elements/PCGCreatePoints.h"
 #include "Helpers/PCGAsync.h"
 #include "Helpers/PCGHelpers.h"
 #include "Helpers/PCGSettingsHelpers.h"
-
 #include "GameFramework/Actor.h"
 
 #include UE_INLINE_GENERATED_CPP_BY_NAME(PCGCreatePointsSphere)
@@ -375,7 +375,7 @@ bool FPCGCreatePointsSphereElement::ExecuteInternal(FPCGContext* InContext) cons
 	const UPCGCreatePointsSphereSettings* Settings = InContext->GetInputSettings<UPCGCreatePointsSphereSettings>();
 
 	// Used for culling, regardless of generation coordinate space
-	const UPCGSpatialData* CullingShape = Settings->bCullPointsOutsideVolume ? Cast<UPCGSpatialData>(InContext->SourceComponent->GetActorPCGData()) : nullptr;
+	const UPCGSpatialData* CullingShape = Settings->bCullPointsOutsideVolume ? Cast<UPCGSpatialData>(InContext->ExecutionSource->GetExecutionState().GetSelfData()) : nullptr;
 
 	// Early out if the culling shape isn't valid
 	if (!CullingShape && Settings->bCullPointsOutsideVolume)
@@ -441,17 +441,16 @@ bool FPCGCreatePointsSphereElement::ExecuteInternal(FPCGContext* InContext) cons
 	TArray<FPCGPoint>& OutputPoints = PointData->GetMutablePoints();
 	Output.Data = PointData;
 
-	FTransform LocalTransform = FTransform::Identity;
+	FTransform LocalTransform = Settings->CoordinateSpace == EPCGCoordinateSpace::World ? FTransform::Identity : InContext->ExecutionSource->GetExecutionState().GetTransform();
 
 	if (Settings->CoordinateSpace == EPCGCoordinateSpace::OriginalComponent)
 	{
-		check(InContext->SourceComponent->GetOriginalComponent() && InContext->SourceComponent->GetOriginalComponent()->GetOwner());
-		LocalTransform = InContext->SourceComponent->GetOriginalComponent()->GetOwner()->GetActorTransform();
-	}
-	else if (Settings->CoordinateSpace == EPCGCoordinateSpace::LocalComponent)
-	{
-		check(InContext->SourceComponent->GetOwner());
-		LocalTransform = InContext->SourceComponent->GetOwner()->GetActorTransform();
+		UPCGComponent* SourceComponent = Cast<UPCGComponent>(InContext->ExecutionSource.Get());
+		if (SourceComponent)
+		{
+			check(SourceComponent->GetOriginalComponent() && SourceComponent->GetOriginalComponent()->GetOwner());
+			LocalTransform = SourceComponent->GetOriginalComponent()->GetOwner()->GetActorTransform();
+		}
 	}
 
 	// Point scale should not be relative to the coordinate space
@@ -500,37 +499,23 @@ bool FPCGCreatePointsSphereElement::ExecuteInternal(FPCGContext* InContext) cons
 	return true;
 }
 
-void FPCGCreatePointsSphereElement::GetDependenciesCrc(const FPCGDataCollection& InInput, const UPCGSettings* InSettings, UPCGComponent* InComponent, FPCGCrc& OutCrc) const
+void FPCGCreatePointsSphereElement::GetDependenciesCrc(const FPCGGetDependenciesCrcParams& InParams, FPCGCrc& OutCrc) const
 {
 	FPCGCrc Crc;
-	IPCGElement::GetDependenciesCrc(InInput, InSettings, InComponent, Crc);
+	IPCGElement::GetDependenciesCrc(InParams, Crc);
 
-	if (const UPCGCreatePointsSphereSettings* Settings = Cast<UPCGCreatePointsSphereSettings>(InSettings))
+	if (const UPCGCreatePointsSphereSettings* Settings = Cast<UPCGCreatePointsSphereSettings>(InParams.Settings))
 	{
 		int CoordinateSpace = static_cast<int>(EPCGCoordinateSpace::World);
 		bool bCullPointsOutsideVolume = false;
-		PCGSettingsHelpers::GetOverrideValue(InInput, Settings, GET_MEMBER_NAME_CHECKED(UPCGCreatePointsSphereSettings, CoordinateSpace), static_cast<int>(Settings->CoordinateSpace), CoordinateSpace);
-		PCGSettingsHelpers::GetOverrideValue(InInput, Settings, GET_MEMBER_NAME_CHECKED(UPCGCreatePointsSphereSettings, bCullPointsOutsideVolume), Settings->bCullPointsOutsideVolume, bCullPointsOutsideVolume);
+		PCGSettingsHelpers::GetOverrideValue(*InParams.InputData, Settings, GET_MEMBER_NAME_CHECKED(UPCGCreatePointsSphereSettings, CoordinateSpace), static_cast<int>(Settings->CoordinateSpace), CoordinateSpace);
+		PCGSettingsHelpers::GetOverrideValue(*InParams.InputData, Settings, GET_MEMBER_NAME_CHECKED(UPCGCreatePointsSphereSettings, bCullPointsOutsideVolume), Settings->bCullPointsOutsideVolume, bCullPointsOutsideVolume);
 
 		EPCGCoordinateSpace EnumCoordinateSpace = static_cast<EPCGCoordinateSpace>(CoordinateSpace);
 
-		if (InComponent)
+		if (UPCGData* Data = PCGCreatePointsHelper::GetDependenciesData(InParams.ExecutionSource, EnumCoordinateSpace, bCullPointsOutsideVolume))
 		{
-			const UPCGData* Data = nullptr;
-
-			if (bCullPointsOutsideVolume || EnumCoordinateSpace == EPCGCoordinateSpace::LocalComponent)
-			{
-				Data = InComponent->GetActorPCGData();
-			}
-			else if (EnumCoordinateSpace == EPCGCoordinateSpace::OriginalComponent)
-			{
-				Data = InComponent->GetOriginalActorPCGData();
-			}
-
-			if (Data)
-			{
-				Crc.Combine(Data->GetOrComputeCrc(/*bFullDataCrc=*/false));
-			}
+			Crc.Combine(Data->GetOrComputeCrc(/*bFullDataCrc=*/false));
 		}
 	}
 
diff --git a/Engine/Plugins/PCG/Source/PCG/Private/Elements/PCGCreateSpline.cpp b/Engine/Plugins/PCG/Source/PCG/Private/Elements/PCGCreateSpline.cpp
--- a/Engine/Plugins/PCG/Source/PCG/Private/Elements/PCGCreateSpline.cpp
+++ b/Engine/Plugins/PCG/Source/PCG/Private/Elements/PCGCreateSpline.cpp
@@ -176,23 +176,27 @@ bool FPCGCreateSplineElement::ExecuteInternal(FPCGContext* Context) const
 		if (Settings->Mode != EPCGCreateSplineMode::CreateDataOnly)
 		{
 			check(IsInGameThread());
-			SplineComponent = NewObject<USplineComponent>(SplineActor);
-			SplineComponent->ComponentTags.Add(Context->SourceComponent.Get()->GetFName());
-			SplineComponent->ComponentTags.Add(PCGHelpers::DefaultPCGTag);
+			if (UPCGComponent* SourceComponent = CastChecked<UPCGComponent>(Context->ExecutionSource.Get()))
+			{
+				SplineComponent = NewObject<USplineComponent>(SplineActor);
+				SplineComponent->ComponentTags.Add(SourceComponent->GetFName());
+				SplineComponent->ComponentTags.Add(PCGHelpers::DefaultPCGTag);
 
-			SplineComponent->RegisterComponent();
-			SplineActor->AddInstanceComponent(SplineComponent);
+				SplineComponent->RegisterComponent();
+				SplineActor->AddInstanceComponent(SplineComponent);
 
-			if (!SplineComponent->AttachToComponent(SplineActor->GetRootComponent(), FAttachmentTransformRules(EAttachmentRule::KeepRelative, EAttachmentRule::KeepWorld, EAttachmentRule::KeepWorld, false)))
-			{
-				PCGLog::Component::LogComponentAttachmentFailedWarning(Context);
-			}
+				if (!SplineComponent->AttachToComponent(SplineActor->GetRootComponent(), FAttachmentTransformRules(EAttachmentRule::KeepRelative, EAttachmentRule::KeepWorld, EAttachmentRule::KeepWorld, false)))
+				{
+					PCGLog::Component::LogComponentAttachmentFailedWarning(Context);
+				}
 
-			SplineData->ApplyTo(SplineComponent);
+				SplineData->ApplyTo(SplineComponent);
 
-			UPCGManagedComponent* ManagedComponent = NewObject<UPCGManagedComponent>(Context->SourceComponent.Get());
-			ManagedComponent->GeneratedComponent = SplineComponent;
-			Context->SourceComponent->AddToManagedResources(ManagedComponent);
+				UPCGManagedComponent* ManagedComponent = NewObject<UPCGManagedComponent>(SourceComponent);
+				ManagedComponent->GeneratedComponent = SplineComponent;
+
+				SourceComponent->AddToManagedResources(ManagedComponent);
+			}
 		}
 
 		// Execute PostProcess Functions
diff --git a/Engine/Plugins/PCG/Source/PCG/Private/Elements/PCGCreateTargetActor.cpp b/Engine/Plugins/PCG/Source/PCG/Private/Elements/PCGCreateTargetActor.cpp
--- a/Engine/Plugins/PCG/Source/PCG/Private/Elements/PCGCreateTargetActor.cpp
+++ b/Engine/Plugins/PCG/Source/PCG/Private/Elements/PCGCreateTargetActor.cpp
@@ -361,7 +361,8 @@ bool FPCGCreateTargetActorElement::ExecuteInternal(FPCGContext* Context) const
 		return true;
 	}
 
-	const bool bHasAuthority = !Context->SourceComponent.IsValid() || (Context->SourceComponent->GetOwner() && Context->SourceComponent->GetOwner()->HasAuthority());
+	UPCGComponent* SourceComponent = Cast<UPCGComponent>(Context->ExecutionSource.Get());
+	const bool bHasAuthority = !SourceComponent || (SourceComponent->GetOwner() && SourceComponent->GetOwner()->HasAuthority());
 	const bool bSpawnedActorRequiresAuthority = CastChecked<AActor>(Settings->TemplateActorClass->GetDefaultObject())->GetIsReplicated();
 
 	if (!bHasAuthority && bSpawnedActorRequiresAuthority)
@@ -377,7 +378,7 @@ bool FPCGCreateTargetActorElement::ExecuteInternal(FPCGContext* Context) const
 	SpawnParams.SpawnCollisionHandlingOverride = ESpawnActorCollisionHandlingMethod::AlwaysSpawn;
 	SpawnParams.OverrideLevel = TargetActor->GetLevel();
 
-	if (PCGHelpers::IsRuntimeOrPIE() || (Context->SourceComponent.IsValid() && Context->SourceComponent->IsInPreviewMode()))
+	if (PCGHelpers::IsRuntimeOrPIE() || (SourceComponent && SourceComponent->IsInPreviewMode()))
 	{
 		SpawnParams.ObjectFlags |= RF_Transient | RF_NonPIEDuplicateTransient;
 	}
@@ -410,7 +411,7 @@ bool FPCGCreateTargetActorElement::ExecuteInternal(FPCGContext* Context) const
 			// since CreateTargetActor settings are already overriden by the inputdata, it doesn't matter here and we can just ignore it
 			// but for future reference if we want to have better reusing and stable generation across, we need to fix this.
 			FPCGDataCollection EmptyCollection;
-			GetDependenciesCrc(EmptyCollection, Settings, Context->SourceComponent.Get(), Context->DependenciesCrc);
+			GetDependenciesCrc(FPCGGetDependenciesCrcParams(&EmptyCollection, Settings, Context->ExecutionSource.Get()), Context->DependenciesCrc);
 
 #if WITH_EDITOR
 			if (DataLayerCrc != 0)
@@ -435,21 +436,24 @@ bool FPCGCreateTargetActorElement::ExecuteInternal(FPCGContext* Context) const
 			const FPCGCrc AdditionalCRC = PCGCreateTargetActor::GetAdditionalDependenciesCrc(SpawnDefaultActorParams, TargetActor);
 			ResourceCrc.Combine(AdditionalCRC);
 
-			Context->SourceComponent->ForEachManagedResource([&ReusableResource, ResourceCrc, &Context](UPCGManagedResource* InResource)
+			if (SourceComponent)
 			{
-				if (ReusableResource)
+				SourceComponent->ForEachManagedResource([&ReusableResource, ResourceCrc, &Context](UPCGManagedResource* InResource)
 				{
-					return;
-				}
+					if (ReusableResource)
+					{
+						return;
+					}
 
-				if (UPCGManagedActors* Resource = Cast<UPCGManagedActors>(InResource))
-				{
-					if (Resource->GetCrc().IsValid() && Resource->GetCrc() == ResourceCrc && Resource->GetConstGeneratedActors().Num() == 1 && Resource->GetConstGeneratedActors()[0].IsValid())
+					if (UPCGManagedActors* Resource = Cast<UPCGManagedActors>(InResource))
 					{
-						ReusableResource = Resource;
+						if (Resource->GetCrc().IsValid() && Resource->GetCrc() == ResourceCrc && Resource->GetConstGeneratedActors().Num() == 1 && Resource->GetConstGeneratedActors()[0].IsValid())
+						{
+							ReusableResource = Resource;
+						}
 					}
-				}
-			});
+				});
+			}
 		}
 	}
 		
@@ -493,7 +497,7 @@ bool FPCGCreateTargetActorElement::ExecuteInternal(FPCGContext* Context) const
 	// Create Resource if it isn't reused
 	if (!ReusableResource)
 	{
-		if (UPCGComponent* SourceComponent = Context->SourceComponent.Get())
+		if (SourceComponent)
 		{
 			UPCGManagedActors* ManagedActors = NewObject<UPCGManagedActors>(SourceComponent);
 			if (ResourceCrc.IsValid())
diff --git a/Engine/Plugins/PCG/Source/PCG/Private/Elements/PCGCullPointsOutsideActorBounds.cpp b/Engine/Plugins/PCG/Source/PCG/Private/Elements/PCGCullPointsOutsideActorBounds.cpp
--- a/Engine/Plugins/PCG/Source/PCG/Private/Elements/PCGCullPointsOutsideActorBounds.cpp
+++ b/Engine/Plugins/PCG/Source/PCG/Private/Elements/PCGCullPointsOutsideActorBounds.cpp
@@ -42,7 +42,7 @@ bool FPCGCullPointsOutsideActorBoundsElement::ExecuteInternal(FPCGContext* Conte
 {
 	TRACE_CPUPROFILER_EVENT_SCOPE(FPCGCullPointsOutsideActorBoundsElement::Execute);
 
-	if (!Context->SourceComponent.IsValid())
+	if (!Context->ExecutionSource.IsValid())
 	{
 		return true;
 	}
@@ -51,7 +51,7 @@ bool FPCGCullPointsOutsideActorBoundsElement::ExecuteInternal(FPCGContext* Conte
 	check(Settings);
 
 	// Initialize directly. Could also have gone through PCGComponent::CreateActorPCGDataCollection.
-	const FBox BoundsBox = Context->SourceComponent->GetGridBounds().ExpandBy(Settings->BoundsExpansion);
+	const FBox BoundsBox = Context->ExecutionSource->GetExecutionState().GetBounds().ExpandBy(Settings->BoundsExpansion);
 
 	const TArray<FPCGTaggedData>& Inputs = Context->InputData.TaggedData;
 	for (const FPCGTaggedData& Input : Inputs)
@@ -174,16 +174,16 @@ bool FPCGCullPointsOutsideActorBoundsElement::ExecuteInternal(FPCGContext* Conte
 	return true;
 }
 
-void FPCGCullPointsOutsideActorBoundsElement::GetDependenciesCrc(const FPCGDataCollection& InInput, const UPCGSettings* InSettings, UPCGComponent* InComponent, FPCGCrc& OutCrc) const
+void FPCGCullPointsOutsideActorBoundsElement::GetDependenciesCrc(const FPCGGetDependenciesCrcParams& InParams, FPCGCrc& OutCrc) const
 {
 	FPCGCrc Crc;
-	IPCGElement::GetDependenciesCrc(InInput, InSettings, InComponent, Crc);
+	IPCGElement::GetDependenciesCrc(InParams, Crc);
 
 	// The culling volume depends on the component transform.
-	const UPCGData* ActorData = InComponent ? InComponent->GetActorPCGData() : nullptr;
-	if (ActorData)
+	const UPCGData* SelfData = InParams.ExecutionSource ? InParams.ExecutionSource->GetExecutionState().GetSelfData() : nullptr;
+	if (SelfData)
 	{
-		Crc.Combine(ActorData->GetOrComputeCrc(/*bFullDataCrc=*/false));
+		Crc.Combine(SelfData->GetOrComputeCrc(/*bFullDataCrc=*/false));
 	}
 
 	OutCrc = Crc;
diff --git a/Engine/Plugins/PCG/Source/PCG/Private/Elements/PCGDataFromActor.cpp b/Engine/Plugins/PCG/Source/PCG/Private/Elements/PCGDataFromActor.cpp
--- a/Engine/Plugins/PCG/Source/PCG/Private/Elements/PCGDataFromActor.cpp
+++ b/Engine/Plugins/PCG/Source/PCG/Private/Elements/PCGDataFromActor.cpp
@@ -353,12 +353,13 @@ bool FPCGDataFromActorElement::ExecuteInternal(FPCGContext* InContext) const
 	const UPCGDataFromActorSettings* Settings = Context->GetInputSettings<UPCGDataFromActorSettings>();
 	check(Settings);
 
+	UPCGComponent* PCGComponent = Cast<UPCGComponent>(Context->ExecutionSource.Get());
+
 	if (!Context->bPerformedQuery)
 	{
 		Context->ComponentSelector = Settings->ComponentSelector;
 
 		TFunction<bool(const AActor*)> BoundsCheck = [](const AActor*) -> bool { return true; };
-		const UPCGComponent* PCGComponent = Context->SourceComponent.IsValid() ? Context->SourceComponent.Get() : nullptr;
 		const AActor* Self = PCGComponent ? PCGComponent->GetOwner() : nullptr;
 		if (Self && Settings->ActorSelector.bMustOverlapSelf)
 		{
@@ -408,7 +409,7 @@ bool FPCGDataFromActorElement::ExecuteInternal(FPCGContext* InContext) const
 		// or just gather all registered components (which is way faster than going through all actors in the world).
 		if (Settings->Mode == EPCGGetDataFromActorMode::GetDataFromPCGComponent && Settings->ActorSelector.ActorFilter == EPCGActorFilter::AllWorldActors)
 		{
-			UPCGSubsystem* Subsystem = Context->SourceComponent.IsValid() ? Context->SourceComponent->GetSubsystem() : nullptr;
+			UPCGSubsystem* Subsystem = PCGComponent ? PCGComponent->GetSubsystem() : nullptr;
 			if (Subsystem)
 			{
 				TRACE_CPUPROFILER_EVENT_SCOPE(FPCGDataFromActorElement::Execute::FindPCGComponents);
@@ -473,7 +474,7 @@ bool FPCGDataFromActorElement::ExecuteInternal(FPCGContext* InContext) const
 
 			Context->ComponentSelector.ComponentList = std::move(ComponentsFromInput);
 
-			Context->FoundActors = PCGActorSelector::FindActors(&Settings->ActorSelector, &Context->ComponentSelector, Context->SourceComponent.Get(), BoundsCheck, SelfIgnoreCheck, ActorsFromInput);
+			Context->FoundActors = PCGActorSelector::FindActors(&Settings->ActorSelector, &Context->ComponentSelector, PCGComponent, BoundsCheck, SelfIgnoreCheck, ActorsFromInput);
 			Context->bPerformedQuery = true;
 
 #if WITH_EDITOR
@@ -531,42 +532,33 @@ bool FPCGDataFromActorElement::ExecuteInternal(FPCGContext* InContext) const
 
 			if (!WaitOnTaskIds.IsEmpty())
 			{
-				UPCGSubsystem* Subsystem = Context->SourceComponent.IsValid() ? Context->SourceComponent->GetSubsystem() : nullptr;
-				if (Subsystem)
-				{
-					// Add a trivial task after these generations that wakes up this task
-					Context->bIsPaused = true;
+				// Add a trivial task after these generations that wakes up this task
+				Context->bIsPaused = true;
 
-					Subsystem->ScheduleGeneric(
-						[ContextHandle = Context->GetOrCreateHandle()]() // Normal execution: Wake up the current task
+				Context->ScheduleGeneric(FPCGScheduleGenericParams(
+					[ContextHandle = Context->GetOrCreateHandle()](FPCGContext*) // Normal execution: Wake up the current task
+					{
+						FPCGContext::FSharedContext<FPCGDataFromActorContext> SharedContext(ContextHandle);
+						if (FPCGDataFromActorContext* ContextPtr = SharedContext.Get())
 						{
-							FPCGContext::FSharedContext<FPCGDataFromActorContext> SharedContext(ContextHandle);
-							if (FPCGDataFromActorContext* ContextPtr = SharedContext.Get())
-							{
-								ContextPtr->bIsPaused = false;
-							}
-							return true;
-						}, 
-						[ContextHandle = Context->GetOrCreateHandle()]() // On Abort: Wake up on abort, clear all results and mark as cancelled
+							ContextPtr->bIsPaused = false;
+						}
+						return true;
+					}, 
+					[ContextHandle = Context->GetOrCreateHandle()](FPCGContext*) // On Abort: Wake up on abort, clear all results and mark as cancelled
+					{
+						FPCGContext::FSharedContext<FPCGDataFromActorContext> SharedContext(ContextHandle);
+						if (FPCGDataFromActorContext* ContextPtr = SharedContext.Get())
 						{
-							FPCGContext::FSharedContext<FPCGDataFromActorContext> SharedContext(ContextHandle);
-							if (FPCGDataFromActorContext* ContextPtr = SharedContext.Get())
-							{
-								ContextPtr->bIsPaused = false;
-								ContextPtr->FoundActors.Reset();
-								ContextPtr->OutputData.bCancelExecution = true;
-							}
-							return true;
-						},
-						Context->SourceComponent.Get(),
-						WaitOnTaskIds);
+							ContextPtr->bIsPaused = false;
+							ContextPtr->FoundActors.Reset();
+							ContextPtr->OutputData.bCancelExecution = true;
+						}
+					},
+					Context->ExecutionSource.Get(),
+					WaitOnTaskIds));
 
-					return false;
-				}
-				else
-				{
-					PCGE_LOG(Error, GraphAndLog, LOCTEXT("UnableToWaitForGenerationTasks", "Unable to wait for end of generation tasks"));
-				}
+				return false;
 			}
 		}
 	}
@@ -575,7 +567,7 @@ bool FPCGDataFromActorElement::ExecuteInternal(FPCGContext* InContext) const
 	{
 #if WITH_EDITOR
 		// Remove ignored change origins now that we've completed the wait tasks.
-		UPCGComponent* OriginalComponent = Context->SourceComponent->GetOriginalComponent();
+		UPCGComponent* OriginalComponent = PCGComponent ? PCGComponent->GetOriginalComponent() : nullptr;
 		if (ensure(OriginalComponent))
 		{
 			for (TObjectKey<UObject>& IgnoredChangeOriginKey : Context->IgnoredChangeOrigins)
@@ -623,7 +615,7 @@ void FPCGDataFromActorElement::GatherWaitTasks(AActor* FoundActor, FPCGContext*
 	const UPCGDataFromActorSettings* Settings = Context->GetInputSettings<UPCGDataFromActorSettings>();
 	check(Settings);
 
-	UPCGComponent* SourceComponent = Context->SourceComponent.IsValid() ? Context->SourceComponent.Get() : nullptr;
+	UPCGComponent* SourceComponent = Cast<UPCGComponent>(Context->ExecutionSource.Get());
 	UPCGComponent* SourceOriginalComponent = SourceComponent ? SourceComponent->GetOriginalComponent() : nullptr;
 
 	if (!SourceOriginalComponent)
@@ -875,7 +867,7 @@ void FPCGDataFromActorElement::ProcessActor(FPCGContext* InContext, const UPCGDa
 	FPCGDataFromActorContext* Context = static_cast<FPCGDataFromActorContext*>(InContext);
 	check(Settings);
 
-	UPCGComponent* SourceComponent = Context->SourceComponent.IsValid() ? Context->SourceComponent.Get() : nullptr;
+	UPCGComponent* SourceComponent = CastChecked<UPCGComponent>(Context->ExecutionSource.Get());
 	const UPCGComponent* SourceOriginalComponent = SourceComponent ? SourceComponent->GetOriginalComponent() : nullptr;
 
 	if (!FoundActor || !IsValid(FoundActor) || !SourceOriginalComponent)
@@ -1038,13 +1030,13 @@ bool FPCGDataFromActorElement::IsCacheable(const UPCGSettings* InSettings) const
 	}
 }
 
-void FPCGDataFromActorElement::GetDependenciesCrc(const FPCGDataCollection& InInput, const UPCGSettings* InSettings, UPCGComponent* InComponent, FPCGCrc& OutCrc) const
+void FPCGDataFromActorElement::GetDependenciesCrc(const FPCGGetDependenciesCrcParams& InParams, FPCGCrc& OutCrc) const
 {
 	FPCGCrc Crc;
-	IPCGElement::GetDependenciesCrc(InInput, InSettings, InComponent, Crc);
+	IPCGElement::GetDependenciesCrc(InParams, Crc);
 
 	// If we track self or original, we are dependent on the actor data
-	if (const UPCGDataFromActorSettings* Settings = Cast<const UPCGDataFromActorSettings>(InSettings))
+	if (const UPCGDataFromActorSettings* Settings = Cast<const UPCGDataFromActorSettings>(InParams.Settings))
 	{
 		if (Settings->bAlwaysRequeryActors)
 		{
@@ -1054,9 +1046,10 @@ void FPCGDataFromActorElement::GetDependenciesCrc(const FPCGDataCollection& InIn
 		const bool bDependsOnSelfOrHierarchy = (Settings->ActorSelector.ActorFilter == EPCGActorFilter::Self || Settings->ActorSelector.ActorFilter == EPCGActorFilter::Original);
 		const bool bDependsOnSelfBounds = Settings->ActorSelector.bMustOverlapSelf;
 
-		if (InComponent && (bDependsOnSelfOrHierarchy || bDependsOnSelfBounds))
+		UPCGComponent* PCGComponent = Cast<UPCGComponent>(InParams.ExecutionSource);
+		if (PCGComponent && (bDependsOnSelfOrHierarchy || bDependsOnSelfBounds))
 		{
-			UPCGComponent* ComponentToCheck = (Settings->ActorSelector.ActorFilter == EPCGActorFilter::Original) ? InComponent->GetOriginalComponent() : InComponent;
+			UPCGComponent* ComponentToCheck = (Settings->ActorSelector.ActorFilter == EPCGActorFilter::Original) ? PCGComponent->GetOriginalComponent() : PCGComponent;
 			const UPCGData* ActorData = ComponentToCheck ? ComponentToCheck->GetActorPCGData() : nullptr;
 
 			if (ActorData)
@@ -1083,9 +1076,9 @@ void FPCGDataFromActorElement::GetDependenciesCrc(const FPCGDataCollection& InIn
 		const bool bDependsOnComponentData = Settings->Mode == EPCGGetDataFromActorMode::GetDataFromPCGComponent || Settings->Mode == EPCGGetDataFromActorMode::GetDataFromPCGComponentOrParseComponents;
 		const bool bDependsOnLocalComponentBounds = Settings->bComponentsMustOverlapSelf || !Settings->bGetDataOnAllGrids;
 
-		if (InComponent && bDependsOnComponentData && bDependsOnLocalComponentBounds)
+		if (PCGComponent && bDependsOnComponentData && bDependsOnLocalComponentBounds)
 		{
-			if (const UPCGData* LocalActorData = InComponent->GetActorPCGData())
+			if (const UPCGData* LocalActorData = PCGComponent->GetActorPCGData())
 			{
 				Crc.Combine(LocalActorData->GetOrComputeCrc(/*bFullDataCrc=*/false));
 			}
diff --git a/Engine/Plugins/PCG/Source/PCG/Private/Elements/PCGDebugElement.cpp b/Engine/Plugins/PCG/Source/PCG/Private/Elements/PCGDebugElement.cpp
--- a/Engine/Plugins/PCG/Source/PCG/Private/Elements/PCGDebugElement.cpp
+++ b/Engine/Plugins/PCG/Source/PCG/Private/Elements/PCGDebugElement.cpp
@@ -37,7 +37,7 @@ bool FPCGDebugElement::ExecuteInternal(FPCGContext* Context) const
 
 #if WITH_EDITOR
 	// Early validation: if we don't have a valid PCG component, we're not going to add the debug display info.
-	if (!Context->SourceComponent.IsValid())
+	if (!Context->ExecutionSource.IsValid())
 	{
 		return true;
 	}
diff --git a/Engine/Plugins/PCG/Source/PCG/Private/Elements/PCGElevationIsolines.cpp b/Engine/Plugins/PCG/Source/PCG/Private/Elements/PCGElevationIsolines.cpp
--- a/Engine/Plugins/PCG/Source/PCG/Private/Elements/PCGElevationIsolines.cpp
+++ b/Engine/Plugins/PCG/Source/PCG/Private/Elements/PCGElevationIsolines.cpp
@@ -67,10 +67,10 @@ bool FPCGElevationIsolinesElement::ExecuteInternal(FPCGContext* InContext) const
 	const UPCGSpatialData* BoundingShape = InContext->InputData.GetSpatialUnionOfInputsByPin(InContext, TEXT("BoundingShape"), bOutUnionDataCreated);
 
 	// Fallback to getting bounds from actor
-	if (!BoundingShape && InContext->SourceComponent.IsValid())
+	if (!BoundingShape && InContext->ExecutionSource.IsValid())
 	{
 		check(!bOutUnionDataCreated);
-		BoundingShape = Cast<UPCGSpatialData>(InContext->SourceComponent->GetActorPCGData());
+		BoundingShape = Cast<UPCGSpatialData>(InContext->ExecutionSource->GetExecutionState().GetSelfData());
 	}
 
 	if (!BoundingShape)
diff --git a/Engine/Plugins/PCG/Source/PCG/Private/Elements/PCGExecuteBlueprint.cpp b/Engine/Plugins/PCG/Source/PCG/Private/Elements/PCGExecuteBlueprint.cpp
--- a/Engine/Plugins/PCG/Source/PCG/Private/Elements/PCGExecuteBlueprint.cpp
+++ b/Engine/Plugins/PCG/Source/PCG/Private/Elements/PCGExecuteBlueprint.cpp
@@ -1087,13 +1087,11 @@ void FPCGBlueprintExecutionContext::AddExtraStructReferencedObjects(FReferenceCo
 	}
 }
 
-FPCGContext* FPCGExecuteBlueprintElement::Initialize(const FPCGDataCollection& InputData, TWeakObjectPtr<UPCGComponent> SourceComponent, const UPCGNode* Node)
+FPCGContext* FPCGExecuteBlueprintElement::Initialize(const FPCGInitializeElementParams& InParams)
 {
 	TRACE_CPUPROFILER_EVENT_SCOPE(FPCGExecuteBlueprintElement::Initialize);
 	FPCGBlueprintExecutionContext* Context = new FPCGBlueprintExecutionContext();
-	Context->InputData = InputData;
-	Context->SourceComponent = SourceComponent;
-	Context->Node = Node;
+	Context->InitFromParams(InParams);
 
 	const UPCGBlueprintSettings* Settings = Context->GetInputSettings<UPCGBlueprintSettings>();
 	if (Settings && Settings->BlueprintElementInstance)
@@ -1101,9 +1099,9 @@ FPCGContext* FPCGExecuteBlueprintElement::Initialize(const FPCGDataCollection& I
 		Context->BlueprintElementInstance = CastChecked<UPCGBlueprintElement>(StaticDuplicateObject(Settings->BlueprintElementInstance, GetTransientPackage(), FName()));
 
 #if !WITH_EDITOR
-		if (SourceComponent.IsValid() && SourceComponent->GetOwner())
+		if (Context->ExecutionSource.IsValid())
 		{
-			Context->BlueprintElementInstance->SetInstanceWorld(SourceComponent->GetOwner()->GetWorld());
+			Context->BlueprintElementInstance->SetInstanceWorld(Context->ExecutionSource->GetExecutionState().GetWorld());
 		}
 #endif
 	}
diff --git a/Engine/Plugins/PCG/Source/PCG/Private/Elements/PCGGenerateGrassMaps.cpp b/Engine/Plugins/PCG/Source/PCG/Private/Elements/PCGGenerateGrassMaps.cpp
--- a/Engine/Plugins/PCG/Source/PCG/Private/Elements/PCGGenerateGrassMaps.cpp
+++ b/Engine/Plugins/PCG/Source/PCG/Private/Elements/PCGGenerateGrassMaps.cpp
@@ -136,13 +136,13 @@ void FPCGGenerateGrassMapsContext::AddExtraStructReferencedObjects(FReferenceCol
 	}
 }
 
-void FPCGGenerateGrassMapsElement::GetDependenciesCrc(const FPCGDataCollection& InInput, const UPCGSettings* InSettings, UPCGComponent* InComponent, FPCGCrc& OutCrc) const
+void FPCGGenerateGrassMapsElement::GetDependenciesCrc(const FPCGGetDependenciesCrcParams& InParams, FPCGCrc& OutCrc) const
 {
 	FPCGCrc Crc;
-	IPCGElement::GetDependenciesCrc(InInput, InSettings, InComponent, Crc);
+	IPCGElement::GetDependenciesCrc(InParams, Crc);
 
 	// todo_pcg: Technically this could be fancier and hash the set of landscape components that overlap with the generation volume.
-	if (const UPCGData* Data = InComponent ? InComponent->GetActorPCGData() : nullptr)
+	if (const UPCGData* Data = InParams.ExecutionSource ? InParams.ExecutionSource->GetExecutionState().GetSelfData() : nullptr)
 	{
 		Crc.Combine(Data->GetOrComputeCrc(/*bFullDataCrc=*/false));
 	}
@@ -163,7 +163,7 @@ bool FPCGGenerateGrassMapsElement::ExecuteInternal(FPCGContext* InContext) const
 	check(Context);
 
 	const UPCGGenerateGrassMapsSettings* Settings = Context->GetInputSettings<UPCGGenerateGrassMapsSettings>();
-	const UPCGComponent* SourceComponent = Context->SourceComponent.IsValid() ? Context->SourceComponent.Get() : nullptr;
+	const UPCGComponent* SourceComponent = Cast<UPCGComponent>(Context->ExecutionSource.Get());
 	check(Settings);
 	check(SourceComponent);
 
diff --git a/Engine/Plugins/PCG/Source/PCG/Private/Elements/PCGGetActorProperty.cpp b/Engine/Plugins/PCG/Source/PCG/Private/Elements/PCGGetActorProperty.cpp
--- a/Engine/Plugins/PCG/Source/PCG/Private/Elements/PCGGetActorProperty.cpp
+++ b/Engine/Plugins/PCG/Source/PCG/Private/Elements/PCGGetActorProperty.cpp
@@ -173,30 +173,33 @@ bool FPCGGetActorPropertyElement::IsCacheable(const UPCGSettings* InSettings) co
 	}
 }
 
-void FPCGGetActorPropertyElement::GetDependenciesCrc(const FPCGDataCollection& InInput, const UPCGSettings* InSettings, UPCGComponent* InComponent,
-	FPCGCrc& OutCrc) const
+void FPCGGetActorPropertyElement::GetDependenciesCrc(const FPCGGetDependenciesCrcParams& InParams, FPCGCrc& OutCrc) const
 {
 	FPCGCrc Crc;
-	IPCGElement::GetDependenciesCrc(InInput, InSettings, InComponent, Crc);
+	IPCGElement::GetDependenciesCrc(InParams, Crc);
 
 	// If we track self or original, we are dependent on the actor data
-	if (const UPCGGetActorPropertySettings* Settings = Cast<const UPCGGetActorPropertySettings>(InSettings))
+	if (const UPCGGetActorPropertySettings* Settings = Cast<const UPCGGetActorPropertySettings>(InParams.Settings))
 	{
 		if (Settings->bAlwaysRequeryActors)
 		{
 			return;
 		}
 
-		const bool bDependsOnSelfOrHierarchy = (Settings->ActorSelector.ActorFilter == EPCGActorFilter::Self || Settings->ActorSelector.ActorFilter == EPCGActorFilter::Original);
-		UPCGComponent* ComponentToCheck = (Settings->ActorSelector.ActorFilter == EPCGActorFilter::Original) ? InComponent->GetOriginalComponent() : InComponent;
-
-		// If it is depending on Self, we need to CRC the tags (as they will be passed to the output data)
-		// and also the owner path.
-		if (bDependsOnSelfOrHierarchy && ComponentToCheck)
+		UPCGComponent* PCGComponent = Cast<UPCGComponent>(InParams.ExecutionSource);
+		if (PCGComponent)
 		{
-			if (const AActor* Owner = ComponentToCheck->GetOwner())
+			const bool bDependsOnSelfOrHierarchy = (Settings->ActorSelector.ActorFilter == EPCGActorFilter::Self || Settings->ActorSelector.ActorFilter == EPCGActorFilter::Original);
+			UPCGComponent* ComponentToCheck = (Settings->ActorSelector.ActorFilter == EPCGActorFilter::Original) ? PCGComponent->GetOriginalComponent() : PCGComponent;
+
+			// If it is depending on Self, we need to CRC the tags (as they will be passed to the output data)
+			// and also the owner path.
+			if (bDependsOnSelfOrHierarchy && ComponentToCheck)
 			{
-				Crc.Combine(UPCGActorHelpers::ComputeHashFromActorTagsAndReference(Owner, /*bIncludeTags=*/true, /*bIncludeActorReference=*/true));
+				if (const AActor* Owner = ComponentToCheck->GetOwner())
+				{
+					Crc.Combine(UPCGActorHelpers::ComputeHashFromActorTagsAndReference(Owner, /*bIncludeTags=*/true, /*bIncludeActorReference=*/true));
+				}
 			}
 		}
 	}
diff --git a/Engine/Plugins/PCG/Source/PCG/Private/Elements/PCGGetPropertyFromObjectPath.cpp b/Engine/Plugins/PCG/Source/PCG/Private/Elements/PCGGetPropertyFromObjectPath.cpp
--- a/Engine/Plugins/PCG/Source/PCG/Private/Elements/PCGGetPropertyFromObjectPath.cpp
+++ b/Engine/Plugins/PCG/Source/PCG/Private/Elements/PCGGetPropertyFromObjectPath.cpp
@@ -197,7 +197,7 @@ bool FPCGGetPropertyFromObjectPathElement::ExecuteInternal(FPCGContext* Context)
 
 #if WITH_EDITOR
 	FPCGDynamicTrackingHelper DynamicTrackingHelper;
-	UPCGComponent* SourceComponent = Context->SourceComponent.Get();
+	UPCGComponent* SourceComponent = Cast<UPCGComponent>(Context->ExecutionSource.Get());
 	const bool IsDynamicallyTracking = (SourceComponent && Context->Node) ? Context->Node->IsInputPinConnected(PCGPinConstants::DefaultInputLabel) : false;
 	if (IsDynamicallyTracking)
 	{
diff --git a/Engine/Plugins/PCG/Source/PCG/Private/Elements/PCGHiGenGridSize.cpp b/Engine/Plugins/PCG/Source/PCG/Private/Elements/PCGHiGenGridSize.cpp
--- a/Engine/Plugins/PCG/Source/PCG/Private/Elements/PCGHiGenGridSize.cpp
+++ b/Engine/Plugins/PCG/Source/PCG/Private/Elements/PCGHiGenGridSize.cpp
@@ -228,7 +228,7 @@ bool FPCGHiGenGridSizeElement::ExecuteInternal(FPCGContext* Context) const
 
 	UPCGVolumeData* VolumeData = FPCGContext::NewObject_AnyThread<UPCGVolumeData>(Context);
 	check(VolumeData);
-	VolumeData->Initialize(Context->SourceComponent->GetGridBounds());
+	VolumeData->Initialize(Context->ExecutionSource->GetExecutionState().GetBounds());
 
 	FPCGTaggedData& VolumeOutput = Context->OutputData.TaggedData.Emplace_GetRef();
 	VolumeOutput.Data = VolumeData;
@@ -237,14 +237,14 @@ bool FPCGHiGenGridSizeElement::ExecuteInternal(FPCGContext* Context) const
 	return true;
 }
 
-void FPCGHiGenGridSizeElement::GetDependenciesCrc(const FPCGDataCollection& InInput, const UPCGSettings* InSettings, UPCGComponent* InComponent, FPCGCrc& OutCrc) const
+void FPCGHiGenGridSizeElement::GetDependenciesCrc(const FPCGGetDependenciesCrcParams& InParams, FPCGCrc& OutCrc) const
 {
 	FPCGCrc Crc;
-	IPCGElement::GetDependenciesCrc(InInput, InSettings, InComponent, Crc);
+	IPCGElement::GetDependenciesCrc(InParams, Crc);
 
 	// The grid cell volume output depends on the component transform.
 	// NOTE: It might be interesting to only incorporate the transform if the cell volume output pin is connected (and if we have a node obviously).
-	const UPCGData* ActorData = InComponent ? InComponent->GetActorPCGData() : nullptr;
+	const UPCGData* ActorData = InParams.ExecutionSource ? InParams.ExecutionSource->GetExecutionState().GetSelfData() : nullptr;
 	if (ActorData)
 	{
 		Crc.Combine(ActorData->GetOrComputeCrc(/*bFullDataCrc=*/false));
diff --git a/Engine/Plugins/PCG/Source/PCG/Private/Elements/PCGIndirectionElement.cpp b/Engine/Plugins/PCG/Source/PCG/Private/Elements/PCGIndirectionElement.cpp
--- a/Engine/Plugins/PCG/Source/PCG/Private/Elements/PCGIndirectionElement.cpp
+++ b/Engine/Plugins/PCG/Source/PCG/Private/Elements/PCGIndirectionElement.cpp
@@ -215,7 +215,7 @@ void FPCGIndirectionContext::SetupInnerContext()
 
 	// TODO: there are some types of settings that might require a bit more information to be able to do their processing correctly
 	// namely (dynamic) subgraphs, so YMMV.
-	InnerContext = InnerElement->Initialize(InnerInput, SourceComponent, nullptr);
+	InnerContext = InnerElement->Initialize(FPCGInitializeElementParams(&InnerInput, ExecutionSource.Get(), nullptr));
 	check(InnerContext);
 	InitializeGraphExecutor(InnerContext);
 	InnerContext->InitializeSettings();
diff --git a/Engine/Plugins/PCG/Source/PCG/Private/Elements/PCGLoopElement.cpp b/Engine/Plugins/PCG/Source/PCG/Private/Elements/PCGLoopElement.cpp
--- a/Engine/Plugins/PCG/Source/PCG/Private/Elements/PCGLoopElement.cpp
+++ b/Engine/Plugins/PCG/Source/PCG/Private/Elements/PCGLoopElement.cpp
@@ -295,7 +295,7 @@ bool FPCGLoopElement::ExecuteInternal(FPCGContext* InContext) const
 			FPCGElementPtr InputElement = MakeShared<FPCGLoopInputForwardingElement>(InputDataCollection, PreviousTaskId, FeedbackPinNames);
 			FPCGTaskId SubgraphTaskId = Context->ScheduleGraph(FPCGScheduleGraphParams(
 				Subgraph,
-				Context->SourceComponent.Get(),
+				Context->ExecutionSource.Get(),
 				PreGraphElement,
 				InputElement,
 				Dependencies,
@@ -353,7 +353,7 @@ bool FPCGLoopElement::ExecuteInternal(FPCGContext* InContext) const
 							ContextPtr->OutputData.bCancelExecution = true;
 						}
 					},
-					Context->SourceComponent.Get(),
+					Context->ExecutionSource.Get(),
 					Context->SubgraphTaskIds));
 			}
 
@@ -450,7 +450,7 @@ bool FPCGLoopInputForwardingElement::ExecuteInternal(FPCGContext* Context) const
 	// Retrieve data from previous task if needed
 	if (PreviousIterationTaskId != InvalidPCGTaskId && !FeedbackPinNames.IsEmpty())
 	{
-		check(Context && Context->SourceComponent.IsValid());
+		check(Context && Context->ExecutionSource.IsValid());
 		FPCGDataCollection PreviousTaskOutput;
 		if (Context->GetOutputData(PreviousIterationTaskId, PreviousTaskOutput))
 		{
diff --git a/Engine/Plugins/PCG/Source/PCG/Private/Elements/PCGMatchAndSetAttributes.cpp b/Engine/Plugins/PCG/Source/PCG/Private/Elements/PCGMatchAndSetAttributes.cpp
--- a/Engine/Plugins/PCG/Source/PCG/Private/Elements/PCGMatchAndSetAttributes.cpp
+++ b/Engine/Plugins/PCG/Source/PCG/Private/Elements/PCGMatchAndSetAttributes.cpp
@@ -556,12 +556,12 @@ private:
 class FPCGMatchAndSetPartition : public FPCGDataPartitionBase<FPCGMatchAndSetPartition, PCGMetadataValueKey>
 {
 public:
-	FPCGMatchAndSetPartition(FPCGContext* InContext, const UPCGMatchAndSetAttributesSettings* InSettings, const UPCGComponent* InSourceComponent, const UPCGParamData* InParamData) 
+	FPCGMatchAndSetPartition(FPCGContext* InContext, const UPCGMatchAndSetAttributesSettings* InSettings, const IPCGGraphExecutionSource* InExecutionSource, const UPCGParamData* InParamData) 
 		: FPCGDataPartitionBase<FPCGMatchAndSetPartition, PCGMetadataValueKey>()
 		, Context(InContext)
 		, Settings(InSettings)
 		, ParamData(InParamData)
-		, SourceComponent(InSourceComponent)
+		, ExecutionSource(InExecutionSource)
 	{
 	}
 
@@ -757,12 +757,12 @@ public:
 			// Generate a random value from the seed
 			for (int32 Index = 0; Index < SeedRange.Num(); ++Index)
 			{
-				Weights.Add(PCGHelpers::GetRandomStreamFromSeed(SeedRange[Index], Settings, SourceComponent).FRand());
+				Weights.Add(PCGHelpers::GetRandomStreamFromSeed(SeedRange[Index], Settings, ExecutionSource).FRand());
 			}
 		}
 		else
 		{
-			FRandomStream RandomStream = PCGHelpers::GetRandomStreamFromSeed(0, Settings, SourceComponent);
+			FRandomStream RandomStream = PCGHelpers::GetRandomStreamFromSeed(0, Settings, ExecutionSource);
 			for (int32 i = 0; i < NumElements; ++i)
 			{
 				Weights.Add(RandomStream.FRand());
@@ -893,7 +893,7 @@ private:
 	const UPCGMatchAndSetAttributesSettings* Settings = nullptr;
 	const UPCGParamData* ParamData = nullptr;
 	FPCGAttributeSetPartition AttributeSetPartition;
-	const UPCGComponent* SourceComponent = nullptr;
+	const IPCGGraphExecutionSource* ExecutionSource = nullptr;
 
 	// Per point data iteration data
 	TUniquePtr<const IPCGAttributeAccessor> InputAttributeAccessor;
@@ -954,7 +954,7 @@ bool FPCGMatchAndSetAttributesElement::PrepareDataInternal(FPCGContext* InContex
 			return EPCGTimeSliceInitResult::NoOperation;
 		}
 
-		FPCGMatchAndSetPartition* Partition = new FPCGMatchAndSetPartition(Context, Settings, Context->SourceComponent.Get(), ParamData);
+		FPCGMatchAndSetPartition* Partition = new FPCGMatchAndSetPartition(Context, Settings, Context->ExecutionSource.Get(), ParamData);
 		check(Partition);
 
 		// If there are provided max distance entries, we should have either 1 or the same cardinality as the inputs
diff --git a/Engine/Plugins/PCG/Source/PCG/Private/Elements/PCGPathfindingElement.cpp b/Engine/Plugins/PCG/Source/PCG/Private/Elements/PCGPathfindingElement.cpp
--- a/Engine/Plugins/PCG/Source/PCG/Private/Elements/PCGPathfindingElement.cpp
+++ b/Engine/Plugins/PCG/Source/PCG/Private/Elements/PCGPathfindingElement.cpp
@@ -347,12 +347,12 @@ bool FPCGPathfindingElement::PrepareDataInternal(FPCGContext* InContext) const
 				return EPCGTimeSliceInitResult::NoOperation;
 			}
 			
-			if (UWorld* World = Context->SourceComponent.Get() ? Context->SourceComponent->GetWorld() : nullptr)
+			if (UWorld* World = Context->ExecutionSource.Get() ? Context->ExecutionSource->GetExecutionState().GetWorld() : nullptr)
 			{
 				FPCGWorldRaycastQueryParams PathTraceParams = Settings->PathTraceParams;
 				PathTraceParams.Initialize();
 
-				TWeakObjectPtr<UPCGComponent> OriginatingComponent = Context->SourceComponent;
+				TWeakObjectPtr<UPCGComponent> OriginatingComponent = Cast<UPCGComponent>(Context->ExecutionSource.Get());
 				FCollisionObjectQueryParams ObjectQueryParams(PathTraceParams.CollisionChannel);
 				FCollisionQueryParams Params;
 				Params.bTraceComplex = PathTraceParams.bTraceComplex;
diff --git a/Engine/Plugins/PCG/Source/PCG/Private/Elements/PCGPrintElement.cpp b/Engine/Plugins/PCG/Source/PCG/Private/Elements/PCGPrintElement.cpp
--- a/Engine/Plugins/PCG/Source/PCG/Private/Elements/PCGPrintElement.cpp
+++ b/Engine/Plugins/PCG/Source/PCG/Private/Elements/PCGPrintElement.cpp
@@ -82,7 +82,8 @@ bool FPCGPrintElement::ExecuteInternal(FPCGContext* Context) const
 	TArray<FString> Prefixes;
 
 	// Build the message prefixes
-	const UPCGComponent* Component = Context->SourceComponent->GetOriginalComponent();
+	UPCGComponent* SourceComponent = Cast<UPCGComponent>(Context->ExecutionSource.Get());
+	const UPCGComponent* Component = SourceComponent ? SourceComponent->GetOriginalComponent() : nullptr;
 	if (Settings->bPrintPerComponent && Component)
 	{
 		if (Settings->bPrefixWithOwner)
@@ -154,7 +155,7 @@ bool FPCGPrintElement::ExecuteInternal(FPCGContext* Context) const
 	}
 
 #if WITH_EDITOR
-	if (Settings->bPrintToScreen && GEngine)
+	if (Settings->bPrintToScreen && GEngine && SourceComponent)
 	{
 		if (UEditorEngine* Editor = static_cast<UEditorEngine*>(GEngine))
 		{
@@ -168,11 +169,11 @@ bool FPCGPrintElement::ExecuteInternal(FPCGContext* Context) const
 			// Add the stack hash to include the loop indices
 			HashKey32 = HashCombine(HashKey32, GetTypeHash(Context->Stack));
 
-			check(Context->SourceComponent.IsValid());
-			UPCGManagedDebugStringMessageKey* ManagedMessageKey = FPCGContext::NewObject_AnyThread<UPCGManagedDebugStringMessageKey>(Context, Context->SourceComponent.Get());
+			check(SourceComponent);
+			UPCGManagedDebugStringMessageKey* ManagedMessageKey = FPCGContext::NewObject_AnyThread<UPCGManagedDebugStringMessageKey>(Context, SourceComponent);
 			ManagedMessageKey->HashKey = static_cast<uint64>(HashKey32);
 
-			Context->SourceComponent->AddToManagedResources(ManagedMessageKey);
+			SourceComponent->AddToManagedResources(ManagedMessageKey);
 
 			// Count 0 as infinite duration to match other debug features in PCG
 			const double Duration = Settings->PrintToScreenDuration > 0.0 ? Settings->PrintToScreenDuration : std::numeric_limits<double>::max();
diff --git a/Engine/Plugins/PCG/Source/PCG/Private/Elements/PCGSanityCheckPointData.cpp b/Engine/Plugins/PCG/Source/PCG/Private/Elements/PCGSanityCheckPointData.cpp
--- a/Engine/Plugins/PCG/Source/PCG/Private/Elements/PCGSanityCheckPointData.cpp
+++ b/Engine/Plugins/PCG/Source/PCG/Private/Elements/PCGSanityCheckPointData.cpp
@@ -61,18 +61,18 @@ bool FPCGSanityCheckPointDataElement::ExecuteInternal(FPCGContext* Context) cons
 	{
 		PCGE_LOG(Error, GraphAndLog, FText::Format(LOCTEXT("NotEnoughPoints", "Expected at least {0} Points, found {1}"), Settings->MinPointCount, InputPointCount));
 
-		if (Context->SourceComponent.IsValid())
+		if (Context->ExecutionSource.IsValid())
 		{
-			Context->SourceComponent->CancelGeneration();
+			Context->ExecutionSource->GetExecutionState().Cancel();
 		}
 	}
 	else if (InputPointCount > Settings->MaxPointCount)
 	{
 		PCGE_LOG(Error, GraphAndLog, FText::Format(LOCTEXT("TooManyPoints", "Expected less than {0} Points, found {1}"), Settings->MaxPointCount, InputPointCount));
 
-		if (Context->SourceComponent.IsValid())
+		if (Context->ExecutionSource.IsValid())
 		{
-			Context->SourceComponent->CancelGeneration();
+			Context->ExecutionSource->GetExecutionState().Cancel();
 		}
 	}
 	else
diff --git a/Engine/Plugins/PCG/Source/PCG/Private/Elements/PCGSkinnedMeshSpawner.cpp b/Engine/Plugins/PCG/Source/PCG/Private/Elements/PCGSkinnedMeshSpawner.cpp
--- a/Engine/Plugins/PCG/Source/PCG/Private/Elements/PCGSkinnedMeshSpawner.cpp
+++ b/Engine/Plugins/PCG/Source/PCG/Private/Elements/PCGSkinnedMeshSpawner.cpp
@@ -83,7 +83,8 @@ bool FPCGSkinnedMeshSpawnerElement::PrepareDataInternal(FPCGContext* InContext)
 		return true;
 	}
 
-	if (!Context->SourceComponent.Get())
+	UPCGComponent* SourceComponent = Cast<UPCGComponent>(Context->ExecutionSource.Get());
+	if (!SourceComponent)
 	{
 		return true;
 	}
@@ -103,13 +104,13 @@ bool FPCGSkinnedMeshSpawnerElement::PrepareDataInternal(FPCGContext* InContext)
 		// Compute CRC if it has not been computed (it likely isn't, but this is to futureproof this)
 		if (!Context->DependenciesCrc.IsValid())
 		{
-			GetDependenciesCrc(Context->InputData, Settings, Context->SourceComponent.Get(), Context->DependenciesCrc);
+			GetDependenciesCrc(FPCGGetDependenciesCrcParams(&Context->InputData, Settings, Context->ExecutionSource.Get()), Context->DependenciesCrc);
 		}
 		
 		if (Context->DependenciesCrc.IsValid())
 		{
 			TArray<UPCGManagedISKMComponent*> MISKMCs;
-			Context->SourceComponent->ForEachManagedResource([&MISKMCs, &Context, Settings](UPCGManagedResource* InResource)
+			SourceComponent->ForEachManagedResource([&MISKMCs, &Context, Settings](UPCGManagedResource* InResource)
 			{
 				if (UPCGManagedISKMComponent* Resource = Cast<UPCGManagedISKMComponent>(InResource))
 				{
@@ -430,7 +431,8 @@ void FPCGSkinnedMeshSpawnerElement::SpawnSkinnedMeshInstances(FPCGSkinnedMeshSpa
 	Params.SettingsCrc = Settings->GetSettingsCrc();
 	ensure(Params.SettingsCrc.IsValid());
 
-	UPCGManagedISKMComponent* MISKMC = UPCGActorHelpers::GetOrCreateManagedABMC(TargetActor, Context->SourceComponent.Get(), Params, Context);
+	UPCGComponent* SourceComponent = Cast<UPCGComponent>(Context->ExecutionSource.Get());
+	UPCGManagedISKMComponent* MISKMC = UPCGActorHelpers::GetOrCreateManagedABMC(TargetActor, SourceComponent, Params, Context);
 
 	check(MISKMC);
 	MISKMC->SetCrc(Context->DependenciesCrc);
diff --git a/Engine/Plugins/PCG/Source/PCG/Private/Elements/PCGSpatialNoise.cpp b/Engine/Plugins/PCG/Source/PCG/Private/Elements/PCGSpatialNoise.cpp
--- a/Engine/Plugins/PCG/Source/PCG/Private/Elements/PCGSpatialNoise.cpp
+++ b/Engine/Plugins/PCG/Source/PCG/Private/Elements/PCGSpatialNoise.cpp
@@ -689,14 +689,15 @@ bool FPCGSpatialNoise::ExecuteInternal(FPCGContext* Context) const
 	PCGSpatialNoise::FSharedParams SharedParams;
 	SharedParams.Context = Context;
 
-	if (!Context->SourceComponent.IsValid())
+	UPCGComponent* SourceComponent = Cast<UPCGComponent>(Context->ExecutionSource.Get());
+	if (!SourceComponent)
 	{
 		PCGE_LOG(Error, GraphAndLog, LOCTEXT("NoSourceComponent", "No Source Component."));
 		return true;		
 	}
 
 	{
-		AActor* Actor = Context->SourceComponent->GetOwner();
+		AActor* Actor = SourceComponent->GetOwner();
 		if (!Actor)
 		{
 			PCGE_LOG(Error, GraphAndLog, LOCTEXT("NoActor", "Source Component has no actor"));
diff --git a/Engine/Plugins/PCG/Source/PCG/Private/Elements/PCGSpawnActor.cpp b/Engine/Plugins/PCG/Source/PCG/Private/Elements/PCGSpawnActor.cpp
--- a/Engine/Plugins/PCG/Source/PCG/Private/Elements/PCGSpawnActor.cpp
+++ b/Engine/Plugins/PCG/Source/PCG/Private/Elements/PCGSpawnActor.cpp
@@ -480,8 +480,12 @@ bool FPCGSpawnActorElement::SpawnAndPrepareSubgraphs(FPCGSubgraphContext* Contex
 			return true;
 		}
 	}
-
-	UPCGSubsystem* Subsystem = Context->SourceComponent.IsValid() ? Context->SourceComponent->GetSubsystem() : nullptr;
+	UPCGComponent* SourceComponent = Cast<UPCGComponent>(Context->ExecutionSource.Get());
+	if (!SourceComponent)
+	{
+		PCGE_LOG(Error, GraphAndLog, LOCTEXT("InvalidSourceComponent", "Invalid SourceComponent."));
+		return true;
+	}
 
 	// Check if we can reuse existing resources - note that this is done on a per-settings basis when collapsed,
 	// Otherwise we'll check against merged crc 
@@ -492,7 +496,7 @@ bool FPCGSpawnActorElement::SpawnAndPrepareSubgraphs(FPCGSubgraphContext* Contex
 		// Compute CRC if it has not been computed (it likely isn't, but this is to futureproof this)
 		if (!Context->DependenciesCrc.IsValid())
 		{
-			GetDependenciesCrc(Context->InputData, Settings, Context->SourceComponent.Get(), Context->DependenciesCrc);
+			GetDependenciesCrc(FPCGGetDependenciesCrcParams(&Context->InputData, Settings, Context->ExecutionSource.Get()), Context->DependenciesCrc);
 		}
 
 		if (Context->DependenciesCrc.IsValid())
@@ -500,7 +504,7 @@ bool FPCGSpawnActorElement::SpawnAndPrepareSubgraphs(FPCGSubgraphContext* Contex
 			if (Settings->Option == EPCGSpawnActorOption::CollapseActors)
 			{
 				TArray<UPCGManagedISMComponent*> MISMCs;
-				Context->SourceComponent->ForEachManagedResource([&MISMCs, &Context](UPCGManagedResource* InResource)
+				SourceComponent->ForEachManagedResource([&MISMCs, &Context](UPCGManagedResource* InResource)
 				{
 					if (UPCGManagedISMComponent* Resource = Cast<UPCGManagedISMComponent>(InResource))
 					{
@@ -539,7 +543,7 @@ bool FPCGSpawnActorElement::SpawnAndPrepareSubgraphs(FPCGSubgraphContext* Contex
 	const bool bGenerateOutputsWithActorReference = (Settings->Option != EPCGSpawnActorOption::CollapseActors) && Context->Node && Context->Node->IsOutputPinConnected(PCGPinConstants::DefaultOutputLabel);
 #endif
 
-	const bool bHasAuthority = !Context->SourceComponent.IsValid() || (Context->SourceComponent->GetOwner() && Context->SourceComponent->GetOwner()->HasAuthority());
+	const bool bHasAuthority = SourceComponent->GetOwner() && SourceComponent->GetOwner()->HasAuthority();
 
 	TArray<FPCGTaggedData> Inputs = Context->InputData.GetInputsByPin(PCGPinConstants::DefaultInputLabel);
 	for (const FPCGTaggedData& Input : Inputs)
@@ -622,7 +626,7 @@ bool FPCGSpawnActorElement::SpawnAndPrepareSubgraphs(FPCGSubgraphContext* Contex
 					// Exception case here: if we've spawned actors but are merging the PCG inputs,
 					// normally this node is taken as a subgraph node (e.g. no need to do anything more than forwarding the inputs)
 					// However, if we`re in the spawn by attribute case, we need to dispatch it here.
-					if (Settings->Option != EPCGSpawnActorOption::NoMerging && Subsystem)
+					if (Settings->Option != EPCGSpawnActorOption::NoMerging)
 					{
 						// TODO: maybe consider a version that would support multi PCG
 						if (UPCGGraphInterface* GraphInterface = UPCGSpawnActorSettings::GetGraphInterfaceFromActorSubclass(Element.Key))
@@ -637,13 +641,14 @@ bool FPCGSpawnActorElement::SpawnAndPrepareSubgraphs(FPCGSubgraphContext* Contex
 
 							UPCGGraph* Graph = GraphInterface->GetGraph();
 
-							FPCGTaskId SubgraphTaskId = Subsystem->ScheduleGraph(Graph,
-								Context->SourceComponent.Get(),
+							FPCGTaskId SubgraphTaskId = Context->ScheduleGraph(FPCGScheduleGraphParams(
+								Graph,
+								Context->ExecutionSource.Get(),
 								MakeShared<FPCGTrivialElement>(),// TODO: prepare user parameters like in subgraph/loop
 								MakeShared<FPCGInputForwardingElement>(SubgraphInputData),
 								/*Dependencies=*/{},
 								&InvocationStack,
-								/*bAllowHierarchicalGeneration=*/false);
+								/*bAllowHierarchicalGeneration=*/false));
 
 							if (SubgraphTaskId != InvalidPCGTaskId)
 							{
@@ -686,8 +691,8 @@ bool FPCGSpawnActorElement::SpawnAndPrepareSubgraphs(FPCGSubgraphContext* Contex
 		}
 		else
 		{
-			Subsystem->ScheduleGeneric(
-				[ContextHandle = Context->GetOrCreateHandle()]() // Normal execution: Wake up the current task
+			Context->ScheduleGeneric(FPCGScheduleGenericParams(
+				[ContextHandle = Context->GetOrCreateHandle()](FPCGContext*) // Normal execution: Wake up the current task
 				{
 					FPCGContext::FSharedContext<FPCGSubgraphContext> SharedContext(ContextHandle);
 					if (FPCGSubgraphContext* ContextPtr = SharedContext.Get())
@@ -696,7 +701,7 @@ bool FPCGSpawnActorElement::SpawnAndPrepareSubgraphs(FPCGSubgraphContext* Contex
 					}
 					return true;
 				},
-				[ContextHandle = Context->GetOrCreateHandle()]() // On Abort: Wake up & cancel
+				[ContextHandle = Context->GetOrCreateHandle()](FPCGContext*) // On Abort: Wake up & cancel
 				{
 					FPCGContext::FSharedContext<FPCGSubgraphContext> SharedContext(ContextHandle);
 					if (FPCGSubgraphContext* ContextPtr = SharedContext.Get())
@@ -705,8 +710,8 @@ bool FPCGSpawnActorElement::SpawnAndPrepareSubgraphs(FPCGSubgraphContext* Contex
 						ContextPtr->OutputData.bCancelExecution = true;
 					}
 				},
-				Context->SourceComponent.Get(),
-				Context->SubgraphTaskIds);
+				Context->ExecutionSource.Get(),
+				Context->SubgraphTaskIds));
 		}
 
 		return false;
@@ -773,7 +778,13 @@ void FPCGSpawnActorElement::CollapseIntoTargetActor(FPCGSubgraphContext* Context
 	{
 		const FPCGISMComponentBuilderParams& ISMCParams = ISMCBuilderTransforms.Key;
 
-		UPCGManagedISMComponent* MISMC = UPCGActorHelpers::GetOrCreateManagedISMC(TargetActor, Context->SourceComponent.Get(), ISMCParams, Context);
+		UPCGComponent* SourceComponent = Cast<UPCGComponent>(Context->ExecutionSource.Get());
+		if (!SourceComponent)
+		{
+			continue;
+		}
+
+		UPCGManagedISMComponent* MISMC = UPCGActorHelpers::GetOrCreateManagedISMC(TargetActor, SourceComponent, ISMCParams, Context);
 		if (!MISMC)
 		{
 			continue;
@@ -867,7 +878,8 @@ void FPCGSpawnActorElement::SpawnActors(FPCGSubgraphContext* Context, AActor* Ta
 	SpawnParams.SpawnCollisionHandlingOverride = ESpawnActorCollisionHandlingMethod::AlwaysSpawn;
 	SpawnParams.OverrideLevel = TargetActor->GetLevel();
 
-	if (PCGHelpers::IsRuntimeOrPIE() || (Context->SourceComponent.IsValid() && Context->SourceComponent->IsInPreviewMode()))
+	UPCGComponent* SourceComponent = Cast<UPCGComponent>(Context->ExecutionSource.Get());
+	if (PCGHelpers::IsRuntimeOrPIE() || (SourceComponent && SourceComponent->IsInPreviewMode()))
 	{
 		SpawnParams.ObjectFlags |= RF_Transient | RF_NonPIEDuplicateTransient;
 	}
@@ -890,7 +902,7 @@ void FPCGSpawnActorElement::SpawnActors(FPCGSubgraphContext* Context, AActor* Ta
 	const bool bOnLoadCallGenerate = (Settings->GenerationTrigger == EPCGSpawnActorGenerationTrigger::Default ||
 		Settings->GenerationTrigger == EPCGSpawnActorGenerationTrigger::DoNotGenerateInEditor);
 #endif
-	UPCGSubsystem* Subsystem = (Context->SourceComponent.Get() ? Context->SourceComponent->GetSubsystem() : nullptr);
+	UPCGSubsystem* Subsystem = UWorld::GetSubsystem<UPCGSubsystem>(SourceComponent->GetWorld());
 
 	// Try to reuse actors if they are preexisting
 	UPCGManagedActors* ReusedManagedActorsResource = nullptr;
@@ -903,7 +915,7 @@ void FPCGSpawnActorElement::SpawnActors(FPCGSubgraphContext* Context, AActor* Ta
 		// Since it is spawning by attribute, all point data will be different.
 		SingleInputCollection.ComputeCrcs(/*bFullDataCrc=*/true);
 
-		GetDependenciesCrc(SingleInputCollection, Settings, Context->SourceComponent.Get(), InputDependenciesCrc);
+		GetDependenciesCrc(FPCGGetDependenciesCrcParams(&SingleInputCollection, Settings, Context->ExecutionSource.Get()), InputDependenciesCrc);
 
 #if WITH_EDITOR
 		if (DataLayerCrc != 0)
@@ -919,7 +931,7 @@ void FPCGSpawnActorElement::SpawnActors(FPCGSubgraphContext* Context, AActor* Ta
 
 		if (InputDependenciesCrc.IsValid())
 		{
-			Context->SourceComponent->ForEachManagedResource([&ReusedManagedActorsResource, &InputDependenciesCrc, &Context, NumPoints = PointData->GetNumPoints()](UPCGManagedResource* InResource)
+			SourceComponent->ForEachManagedResource([&ReusedManagedActorsResource, &InputDependenciesCrc, &Context, NumPoints = PointData->GetNumPoints()](UPCGManagedResource* InResource)
 			{
 				if (ReusedManagedActorsResource)
 				{
@@ -980,7 +992,7 @@ void FPCGSpawnActorElement::SpawnActors(FPCGSubgraphContext* Context, AActor* Ta
 		TArray<FName> NewActorTags = GetNewActorTags(Context, TargetActor, Settings->bInheritActorTags, Settings->TagsToAddOnActors);
 
 		// Create managed resource for actor tracking
-		UPCGManagedActors* ManagedActors = NewObject<UPCGManagedActors>(Context->SourceComponent.Get());
+		UPCGManagedActors* ManagedActors = NewObject<UPCGManagedActors>(SourceComponent);
 		ManagedActors->SetCrc(InputDependenciesCrc);
 		ManagedActors->bSupportsReset = !Settings->bDeleteActorsBeforeGeneration;
 
@@ -1057,7 +1069,7 @@ void FPCGSpawnActorElement::SpawnActors(FPCGSubgraphContext* Context, AActor* Ta
 			PCGE_LOG(Error, GraphAndLog, LOCTEXT("ActorOverridesFailed", "At least one actor property override failed."));
 		}
 
-		Context->SourceComponent->AddToManagedResources(ManagedActors);
+		SourceComponent->AddToManagedResources(ManagedActors);
 
 		PCGE_LOG(Verbose, LogOnly, FText::Format(LOCTEXT("GenerationInfo", "Generated {0} actors"), PointData->GetNumPoints()));
 	}
@@ -1072,10 +1084,10 @@ void FPCGSpawnActorElement::SpawnActors(FPCGSubgraphContext* Context, AActor* Ta
 		{
 #if WITH_EDITOR
 			// For both pre-existing and new actors, we need to make sure we're inline with loading/generation as needed
-			if (PCGComponent->GetEditingMode() != Context->SourceComponent->GetEditingMode())
+			if (PCGComponent->GetEditingMode() != SourceComponent->GetEditingMode())
 			{
-				PCGComponent->SetEditingMode(/*CurrentEditingMode=*/Context->SourceComponent->GetEditingMode(), /*SerializedEditingMode=*/Context->SourceComponent->GetEditingMode());
-				PCGComponent->ChangeTransientState(Context->SourceComponent->GetEditingMode());
+				PCGComponent->SetEditingMode(/*CurrentEditingMode=*/SourceComponent->GetEditingMode(), /*SerializedEditingMode=*/SourceComponent->GetEditingMode());
+				PCGComponent->ChangeTransientState(SourceComponent->GetEditingMode());
 			}
 #endif // WITH_EDITOR
 
@@ -1119,10 +1131,13 @@ TArray<FName> FPCGSpawnActorElement::GetNewActorTags(FPCGContext* Context, AActo
 		// and find the original actor tags
 		if (APCGPartitionActor* PartitionActor = Cast<APCGPartitionActor>(TargetActor))
 		{
-			if (UPCGComponent* OriginalComponent = PartitionActor->GetOriginalComponent(Context->SourceComponent.Get()))
+			if (UPCGComponent* SourceComponent = Cast<UPCGComponent>(Context->ExecutionSource.Get()))
 			{
-				check(OriginalComponent->GetOwner());
-				NewActorTags = OriginalComponent->GetOwner()->Tags;
+				if (UPCGComponent* OriginalComponent = PartitionActor->GetOriginalComponent(SourceComponent))
+				{
+					check(OriginalComponent->GetOwner());
+					NewActorTags = OriginalComponent->GetOwner()->Tags;
+				}
 			}
 		}
 		else
diff --git a/Engine/Plugins/PCG/Source/PCG/Private/Elements/PCGSpawnSpline.cpp b/Engine/Plugins/PCG/Source/PCG/Private/Elements/PCGSpawnSpline.cpp
--- a/Engine/Plugins/PCG/Source/PCG/Private/Elements/PCGSpawnSpline.cpp
+++ b/Engine/Plugins/PCG/Source/PCG/Private/Elements/PCGSpawnSpline.cpp
@@ -51,7 +51,8 @@ bool FPCGSpawnSplineElement::ExecuteInternal(FPCGContext* InContext) const
 	const UPCGSpawnSplineSettings* Settings = InContext->GetInputSettings<UPCGSpawnSplineSettings>();
 	check(Settings);
 
-	if (!*Settings->SplineComponent || !InContext->SourceComponent.IsValid())
+	UPCGComponent* SourceComponent = Cast<UPCGComponent>(InContext->ExecutionSource.Get());
+	if (!*Settings->SplineComponent || !SourceComponent)
 	{
 		return true;
 	}
@@ -103,7 +104,7 @@ bool FPCGSpawnSplineElement::ExecuteInternal(FPCGContext* InContext) const
 
 		check(IsInGameThread());
 		USplineComponent* SplineComponent = NewObject<USplineComponent>(TargetActor, *SplineComponentToSpawn);
-		SplineComponent->ComponentTags.Add(InContext->SourceComponent.Get()->GetFName());
+		SplineComponent->ComponentTags.Add(SourceComponent->GetFName());
 		SplineComponent->ComponentTags.Add(PCGHelpers::DefaultPCGTag);
 
 		// Also propagate the tags on the tagged data to the component.
@@ -119,9 +120,9 @@ bool FPCGSpawnSplineElement::ExecuteInternal(FPCGContext* InContext) const
 
 		InSplineData->ApplyTo(SplineComponent);
 
-		UPCGManagedComponent* ManagedComponent = NewObject<UPCGManagedComponent>(InContext->SourceComponent.Get());
+		UPCGManagedComponent* ManagedComponent = NewObject<UPCGManagedComponent>(SourceComponent);
 		ManagedComponent->GeneratedComponent = SplineComponent;
-		InContext->SourceComponent->AddToManagedResources(ManagedComponent);
+		SourceComponent->AddToManagedResources(ManagedComponent);
 
 		if (!Settings->PropertyOverrideDescriptions.IsEmpty())
 		{
diff --git a/Engine/Plugins/PCG/Source/PCG/Private/Elements/PCGSpawnSplineMesh.cpp b/Engine/Plugins/PCG/Source/PCG/Private/Elements/PCGSpawnSplineMesh.cpp
--- a/Engine/Plugins/PCG/Source/PCG/Private/Elements/PCGSpawnSplineMesh.cpp
+++ b/Engine/Plugins/PCG/Source/PCG/Private/Elements/PCGSpawnSplineMesh.cpp
@@ -159,7 +159,13 @@ bool FPCGSpawnSplineMeshElement::ExecuteInternal(FPCGContext* InContext) const
 		return true;
 	}
 
-	return ExecuteSlice(Context, [this, Settings](ContextType* Context, const ExecStateType& ExecState, IterStateType& IterState, const uint32 IterIndex)
+	UPCGComponent* SourceComponent = Cast<UPCGComponent>(Context->ExecutionSource.Get());
+	if (!SourceComponent)
+	{
+		return true;
+	}
+
+	return ExecuteSlice(Context, [this, Settings, SourceComponent](ContextType* Context, const ExecStateType& ExecState, IterStateType& IterState, const uint32 IterIndex)
 	{
 		if (Context->GetIterationStateResult(IterIndex) != EPCGTimeSliceInitResult::Success)
 		{
@@ -268,7 +274,7 @@ bool FPCGSpawnSplineMeshElement::ExecuteInternal(FPCGContext* InContext) const
 			IterState.SMCBuilderParams.SettingsCrc = Settings->GetSettingsCrc();
 			ensure(IterState.SMCBuilderParams.SettingsCrc.IsValid());
 
-			USplineMeshComponent* SplineMeshComponent = UPCGActorHelpers::GetOrCreateSplineMeshComponent(ExecState.TargetActor, Context->SourceComponent.Get(), IterState.SMCBuilderParams, Context);
+			USplineMeshComponent* SplineMeshComponent = UPCGActorHelpers::GetOrCreateSplineMeshComponent(ExecState.TargetActor, SourceComponent, IterState.SMCBuilderParams, Context);
 
 			// TODO: Write out the geometry to a dynamic mesh type.
 			//SplineMeshComponent->BodySetup->TriMeshGeometries
diff --git a/Engine/Plugins/PCG/Source/PCG/Private/Elements/PCGSplineSampler.cpp b/Engine/Plugins/PCG/Source/PCG/Private/Elements/PCGSplineSampler.cpp
--- a/Engine/Plugins/PCG/Source/PCG/Private/Elements/PCGSplineSampler.cpp
+++ b/Engine/Plugins/PCG/Source/PCG/Private/Elements/PCGSplineSampler.cpp
@@ -1503,21 +1503,21 @@ FPCGElementPtr UPCGSplineSamplerSettings::CreateElement() const
 	return MakeShared<FPCGSplineSamplerElement>();
 }
 
-void FPCGSplineSamplerElement::GetDependenciesCrc(const FPCGDataCollection& InInput, const UPCGSettings* InSettings, UPCGComponent* InComponent, FPCGCrc& OutCrc) const
+void FPCGSplineSamplerElement::GetDependenciesCrc(const FPCGGetDependenciesCrcParams& InParams, FPCGCrc& OutCrc) const
 {
 	FPCGCrc Crc;
-	IPCGElement::GetDependenciesCrc(InInput, InSettings, InComponent, Crc);
+	IPCGElement::GetDependenciesCrc(InParams, Crc);
 
-	if (const UPCGSplineSamplerSettings* Settings = Cast<UPCGSplineSamplerSettings>(InSettings))
+	if (const UPCGSplineSamplerSettings* Settings = Cast<UPCGSplineSamplerSettings>(InParams.Settings))
 	{
 		bool bUnbounded;
-		PCGSettingsHelpers::GetOverrideValue(InInput, Settings, GET_MEMBER_NAME_CHECKED(FPCGSplineSamplerParams, bUnbounded), Settings->SamplerParams.bUnbounded, bUnbounded);
-		const bool bBoundsConnected = InInput.GetInputsByPin(PCGSplineSamplerConstants::BoundingShapeLabel).Num() > 0;
+		PCGSettingsHelpers::GetOverrideValue(*InParams.InputData, Settings, GET_MEMBER_NAME_CHECKED(FPCGSplineSamplerParams, bUnbounded), Settings->SamplerParams.bUnbounded, bUnbounded);
+		const bool bBoundsConnected = InParams.InputData->GetInputsByPin(PCGSplineSamplerConstants::BoundingShapeLabel).Num() > 0;
 
 		// If we're operating in bounded mode then we'll use actor bounds, and therefore take dependency on actor data.
-		if (!bUnbounded && !bBoundsConnected && InComponent)
+		if (!bUnbounded && !bBoundsConnected && InParams.ExecutionSource)
 		{
-			if (const UPCGData* Data = InComponent->GetActorPCGData())
+			if (const UPCGData* Data = InParams.ExecutionSource->GetExecutionState().GetSelfData())
 			{
 				Crc.Combine(Data->GetOrComputeCrc(/*bFullDataCrc=*/false));
 			}
@@ -1550,9 +1550,9 @@ bool FPCGSplineSamplerElement::ExecuteInternal(FPCGContext* Context) const
 		BoundingShape = Context->InputData.GetSpatialUnionOfInputsByPin(Context, PCGSplineSamplerConstants::BoundingShapeLabel, bUnionCreated);
 
 		// Fallback to getting bounds from actor
-		if (!BoundingShape && Context->SourceComponent.IsValid())
+		if (!BoundingShape && Context->ExecutionSource.IsValid())
 		{
-			BoundingShape = Cast<UPCGSpatialData>(Context->SourceComponent->GetActorPCGData());
+			BoundingShape = Cast<UPCGSpatialData>(Context->ExecutionSource->GetExecutionState().GetSelfData());
 		}
 	}
 	else if (BoundingShapeInputs.Num() > 0)
diff --git a/Engine/Plugins/PCG/Source/PCG/Private/Elements/PCGStaticMeshSpawner.cpp b/Engine/Plugins/PCG/Source/PCG/Private/Elements/PCGStaticMeshSpawner.cpp
--- a/Engine/Plugins/PCG/Source/PCG/Private/Elements/PCGStaticMeshSpawner.cpp
+++ b/Engine/Plugins/PCG/Source/PCG/Private/Elements/PCGStaticMeshSpawner.cpp
@@ -125,7 +125,8 @@ bool FPCGStaticMeshSpawnerElement::PrepareDataInternal(FPCGContext* InContext) c
 		return true;
 	}
 
-	if (!Context->SourceComponent.Get())
+	UPCGComponent* SourceComponent = Cast<UPCGComponent>(Context->ExecutionSource.Get());
+	if (!SourceComponent)
 	{
 		return true;
 	}
@@ -145,13 +146,13 @@ bool FPCGStaticMeshSpawnerElement::PrepareDataInternal(FPCGContext* InContext) c
 		// Compute CRC if it has not been computed (it likely isn't, but this is to futureproof this)
 		if (!Context->DependenciesCrc.IsValid())
 		{
-			GetDependenciesCrc(Context->InputData, Settings, Context->SourceComponent.Get(), Context->DependenciesCrc);
+			GetDependenciesCrc(FPCGGetDependenciesCrcParams(&Context->InputData, Settings, Context->ExecutionSource.Get()), Context->DependenciesCrc);
 		}
 		
 		if (Context->DependenciesCrc.IsValid())
 		{
 			TArray<UPCGManagedISMComponent*> MISMCs;
-			Context->SourceComponent->ForEachManagedResource([&MISMCs, &Context, Settings](UPCGManagedResource* InResource)
+			SourceComponent->ForEachManagedResource([&MISMCs, &Context, Settings](UPCGManagedResource* InResource)
 			{
 				if (UPCGManagedISMComponent* Resource = Cast<UPCGManagedISMComponent>(InResource))
 				{
@@ -238,11 +239,11 @@ bool FPCGStaticMeshSpawnerElement::PrepareDataInternal(FPCGContext* InContext) c
 					SubCollection.TaggedData.Add(Input);
 					SubCollection.ComputeCrcs(ShouldComputeFullOutputDataCrc(Context));
 
-					GetDependenciesCrc(SubCollection, Settings, Context->SourceComponent.Get(), DataCrc);
+					GetDependenciesCrc(FPCGGetDependenciesCrcParams(&SubCollection, Settings, Context->ExecutionSource.Get()), DataCrc);
 				}
 
 				TArray<UPCGManagedISMComponent*> MISMCs;
-				Context->SourceComponent->ForEachManagedResource([&MISMCs, &SettingsCrc, &DataCrc, &Context](UPCGManagedResource* InResource)
+				SourceComponent->ForEachManagedResource([&MISMCs, &SettingsCrc, &DataCrc, &Context](UPCGManagedResource* InResource)
 				{
 					if (UPCGManagedISMComponent* Resource = Cast<UPCGManagedISMComponent>(InResource))
 					{
@@ -575,7 +576,13 @@ void FPCGStaticMeshSpawnerElement::SpawnStaticMeshInstances(FPCGStaticMeshSpawne
 
 	Params.bAllowDescriptorChanges = Settings->bAllowDescriptorChanges;
 
-	UPCGManagedISMComponent* MISMC = UPCGActorHelpers::GetOrCreateManagedISMC(TargetActor, Context->SourceComponent.Get(), Params, Context);
+	UPCGComponent* SourceComponent = Cast<UPCGComponent>(Context->ExecutionSource.Get());
+	if (!ensure(SourceComponent))
+	{
+		return;
+	}
+
+	UPCGManagedISMComponent* MISMC = UPCGActorHelpers::GetOrCreateManagedISMC(TargetActor, SourceComponent, Params, Context);
 
 	check(MISMC);
 	MISMC->SetCrc(Context->DependenciesCrc);
diff --git a/Engine/Plugins/PCG/Source/PCG/Private/Elements/PCGSurfaceSampler.cpp b/Engine/Plugins/PCG/Source/PCG/Private/Elements/PCGSurfaceSampler.cpp
--- a/Engine/Plugins/PCG/Source/PCG/Private/Elements/PCGSurfaceSampler.cpp
+++ b/Engine/Plugins/PCG/Source/PCG/Private/Elements/PCGSurfaceSampler.cpp
@@ -660,22 +660,22 @@ bool FPCGSurfaceSamplerElement::ExecuteInternal(FPCGContext* InContext) const
 	});
 }
 
-void FPCGSurfaceSamplerElement::GetDependenciesCrc(const FPCGDataCollection& InInput, const UPCGSettings* InSettings, UPCGComponent* InComponent, FPCGCrc& OutCrc) const
+void FPCGSurfaceSamplerElement::GetDependenciesCrc(const FPCGGetDependenciesCrcParams& InParams, FPCGCrc& OutCrc) const
 {
 	FPCGCrc Crc;
-	IPCGElement::GetDependenciesCrc(InInput, InSettings, InComponent, Crc);
+	IPCGElement::GetDependenciesCrc(InParams, Crc);
 
-	if (const UPCGSurfaceSamplerSettings* Settings = Cast<UPCGSurfaceSamplerSettings>(InSettings))
+	if (const UPCGSurfaceSamplerSettings* Settings = Cast<UPCGSurfaceSamplerSettings>(InParams.Settings))
 	{
 		bool bUnbounded;
-		PCGSettingsHelpers::GetOverrideValue(InInput, Settings, GET_MEMBER_NAME_CHECKED(UPCGSurfaceSamplerSettings, bUnbounded), Settings->bUnbounded, bUnbounded);
-		const bool bBoundsConnected = InInput.GetInputsByPin(PCGSurfaceSamplerConstants::BoundingShapeLabel).Num() > 0;
+		PCGSettingsHelpers::GetOverrideValue(*InParams.InputData, Settings, GET_MEMBER_NAME_CHECKED(UPCGSurfaceSamplerSettings, bUnbounded), Settings->bUnbounded, bUnbounded);
+		const bool bBoundsConnected = InParams.InputData->GetInputsByPin(PCGSurfaceSamplerConstants::BoundingShapeLabel).Num() > 0;
 
 		// If we're operating in bounded mode and there is no bounding shape connected then we'll use actor bounds, and therefore take
 		// dependency on actor data.
-		if (!bUnbounded && !bBoundsConnected && InComponent)
+		if (!bUnbounded && !bBoundsConnected && InParams.ExecutionSource)
 		{
-			if (const UPCGData* Data = InComponent->GetActorPCGData())
+			if (const UPCGData* Data = InParams.ExecutionSource->GetExecutionState().GetSelfData())
 			{
 				Crc.Combine(Data->GetOrComputeCrc(/*bFullDataCrc=*/false));
 			}
diff --git a/Engine/Plugins/PCG/Source/PCG/Private/Elements/PCGTextureSampler.cpp b/Engine/Plugins/PCG/Source/PCG/Private/Elements/PCGTextureSampler.cpp
--- a/Engine/Plugins/PCG/Source/PCG/Private/Elements/PCGTextureSampler.cpp
+++ b/Engine/Plugins/PCG/Source/PCG/Private/Elements/PCGTextureSampler.cpp
@@ -330,19 +330,19 @@ void UPCGTextureSamplerSettings::SetTexture(TSoftObjectPtr<UTexture> InTexture)
 #endif // WITH_EDITOR
 }
 
-void FPCGTextureSamplerElement::GetDependenciesCrc(const FPCGDataCollection& InInput, const UPCGSettings* InSettings, UPCGComponent* InComponent, FPCGCrc& OutCrc) const
+void FPCGTextureSamplerElement::GetDependenciesCrc(const FPCGGetDependenciesCrcParams& InParams, FPCGCrc& OutCrc) const
 {
 	FPCGCrc Crc;
-	IPCGElement::GetDependenciesCrc(InInput, InSettings, InComponent, Crc);
+	IPCGElement::GetDependenciesCrc(InParams, Crc);
 
-	if (const UPCGTextureSamplerSettings* Settings = Cast<UPCGTextureSamplerSettings>(InSettings))
+	if (const UPCGTextureSamplerSettings* Settings = Cast<UPCGTextureSamplerSettings>(InParams.Settings))
 	{
 		// If not using absolute transform, depend on actor transform and bounds, and therefore take dependency on actor data.
 		bool bUseAbsoluteTransform;
-		PCGSettingsHelpers::GetOverrideValue(InInput, Settings, GET_MEMBER_NAME_CHECKED(UPCGTextureSamplerSettings, bUseAbsoluteTransform), Settings->bUseAbsoluteTransform, bUseAbsoluteTransform);
-		if (!bUseAbsoluteTransform && InComponent)
+		PCGSettingsHelpers::GetOverrideValue(*InParams.InputData, Settings, GET_MEMBER_NAME_CHECKED(UPCGTextureSamplerSettings, bUseAbsoluteTransform), Settings->bUseAbsoluteTransform, bUseAbsoluteTransform);
+		if (!bUseAbsoluteTransform && InParams.ExecutionSource)
 		{
-			if (const UPCGData* Data = InComponent->GetActorPCGData())
+			if (const UPCGData* Data = InParams.ExecutionSource->GetExecutionState().GetSelfData())
 			{
 				Crc.Combine(Data->GetOrComputeCrc(/*bFullDataCrc=*/false));
 			}
diff --git a/Engine/Plugins/PCG/Source/PCG/Private/Elements/PCGUserParameterGet.cpp b/Engine/Plugins/PCG/Source/PCG/Private/Elements/PCGUserParameterGet.cpp
--- a/Engine/Plugins/PCG/Source/PCG/Private/Elements/PCGUserParameterGet.cpp
+++ b/Engine/Plugins/PCG/Source/PCG/Private/Elements/PCGUserParameterGet.cpp
@@ -70,8 +70,8 @@ namespace PCGUserParameterGet
 			FConstStructView GraphFromNodeParameters = GraphFromNode->GetUserParametersStruct() ? GraphFromNode->GetUserParametersStruct()->GetValue() : FConstStructView{};
 
 			// If we don't have a graph instance, we just use the user parameters from the node graph owner.
-			const UPCGComponent* SourceComponent = InContext.SourceComponent.Get();
-			const UPCGGraphInstance* GraphInstance = SourceComponent ? SourceComponent->GetGraphInstance() : nullptr;
+			const IPCGGraphExecutionSource* ExecutionSource = InContext.ExecutionSource.Get();
+			const UPCGGraphInstance* GraphInstance = ExecutionSource ? ExecutionSource->GetExecutionState().GetGraphInstance() : nullptr;
 
 			// If there is no user parameters, just use the one from the graph.
 			// Go down the graph instance chain and take all layouts. Revert back to original graph node params if none were found
diff --git a/Engine/Plugins/PCG/Source/PCG/Private/Elements/PCGVisualizeAttribute.cpp b/Engine/Plugins/PCG/Source/PCG/Private/Elements/PCGVisualizeAttribute.cpp
--- a/Engine/Plugins/PCG/Source/PCG/Private/Elements/PCGVisualizeAttribute.cpp
+++ b/Engine/Plugins/PCG/Source/PCG/Private/Elements/PCGVisualizeAttribute.cpp
@@ -31,7 +31,8 @@ bool FPCGVisualizeAttribute::ExecuteInternal(FPCGContext* Context) const
 		return true;
 	}
 
-	if (!ensure(Context->SourceComponent.IsValid()))
+	UPCGComponent* SourceComponent = Cast<UPCGComponent>(Context->ExecutionSource.Get());
+	if (!ensure(SourceComponent))
 	{
 		return true;
 	}
@@ -104,9 +105,9 @@ bool FPCGVisualizeAttribute::ExecuteInternal(FPCGContext* Context) const
 			TargetActor->Modify(/*bAlwaysMarkDirty=*/false);
 			DebugDrawComponent->RegisterComponent();
 			TargetActor->AddInstanceComponent(DebugDrawComponent);
-			UPCGManagedDebugDrawComponent* ManagedComponent = NewObject<UPCGManagedDebugDrawComponent>(Context->SourceComponent.Get());
+			UPCGManagedDebugDrawComponent* ManagedComponent = NewObject<UPCGManagedDebugDrawComponent>(SourceComponent);
 			ManagedComponent->GeneratedComponent = DebugDrawComponent;
-			Context->SourceComponent->AddToManagedResources(ManagedComponent);
+			SourceComponent->AddToManagedResources(ManagedComponent);
 		}
 
 		const TArray<FPCGPoint>& Points = PointData->GetPoints();
diff --git a/Engine/Plugins/PCG/Source/PCG/Private/Elements/PCGVolumeSampler.cpp b/Engine/Plugins/PCG/Source/PCG/Private/Elements/PCGVolumeSampler.cpp
--- a/Engine/Plugins/PCG/Source/PCG/Private/Elements/PCGVolumeSampler.cpp
+++ b/Engine/Plugins/PCG/Source/PCG/Private/Elements/PCGVolumeSampler.cpp
@@ -252,10 +252,10 @@ namespace PCGVolumeSamplerHelpers
 				Context->TrackObject(OutState.BoundingShape);
 			}
 
-			if (!OutState.BoundingShape && Context->SourceComponent.IsValid())
+			if (!OutState.BoundingShape && Context->ExecutionSource.IsValid())
 			{
 				// Create a bounding shape from the actor data
-				OutState.BoundingShape = Cast<UPCGSpatialData>(Context->SourceComponent->GetActorPCGData());
+				OutState.BoundingShape = Cast<UPCGSpatialData>(Context->ExecutionSource->GetExecutionState().GetSelfData());
 				bUsedDefaultBoundingShape = true;
 			}
 		}
@@ -441,22 +441,22 @@ bool FPCGVolumeSamplerElement::ExecuteInternal(FPCGContext* Context) const
 	});
 }
 
-void FPCGVolumeSamplerElement::GetDependenciesCrc(const FPCGDataCollection& InInput, const UPCGSettings* InSettings, UPCGComponent* InComponent, FPCGCrc& OutCrc) const
+void FPCGVolumeSamplerElement::GetDependenciesCrc(const FPCGGetDependenciesCrcParams& InParams, FPCGCrc& OutCrc) const
 {
 	FPCGCrc Crc;
-	IPCGElement::GetDependenciesCrc(InInput, InSettings, InComponent, Crc);
+	IPCGElement::GetDependenciesCrc(InParams, Crc);
 
-	if (const UPCGVolumeSamplerSettings* Settings = Cast<UPCGVolumeSamplerSettings>(InSettings))
+	if (const UPCGVolumeSamplerSettings* Settings = Cast<UPCGVolumeSamplerSettings>(InParams.Settings))
 	{
 		bool bUnbounded;
-		PCGSettingsHelpers::GetOverrideValue(InInput, Settings, GET_MEMBER_NAME_CHECKED(UPCGVolumeSamplerSettings, bUnbounded), Settings->bUnbounded, bUnbounded);
-		const bool bBoundsConnected = InInput.GetInputsByPin(PCGVolumeSamplerConstants::BoundingShapeLabel).Num() > 0;
+		PCGSettingsHelpers::GetOverrideValue(*InParams.InputData, Settings, GET_MEMBER_NAME_CHECKED(UPCGVolumeSamplerSettings, bUnbounded), Settings->bUnbounded, bUnbounded);
+		const bool bBoundsConnected = InParams.InputData->GetInputsByPin(PCGVolumeSamplerConstants::BoundingShapeLabel).Num() > 0;
 
 		// If we're operating in bounded mode and there is no bounding shape connected then we'll use actor bounds, and therefore take
 		// dependency on actor data.
-		if (!bUnbounded && !bBoundsConnected && InComponent)
+		if (!bUnbounded && !bBoundsConnected && InParams.ExecutionSource)
 		{
-			if (const UPCGData* Data = InComponent->GetActorPCGData())
+			if (const UPCGData* Data = InParams.ExecutionSource->GetExecutionState().GetSelfData())
 			{
 				Crc.Combine(Data->GetOrComputeCrc(/*bFullDataCrc=*/false));
 			}
diff --git a/Engine/Plugins/PCG/Source/PCG/Private/Elements/PCGWorldQuery.cpp b/Engine/Plugins/PCG/Source/PCG/Private/Elements/PCGWorldQuery.cpp
--- a/Engine/Plugins/PCG/Source/PCG/Private/Elements/PCGWorldQuery.cpp
+++ b/Engine/Plugins/PCG/Source/PCG/Private/Elements/PCGWorldQuery.cpp
@@ -95,14 +95,14 @@ bool FPCGWorldVolumetricQueryElement::ExecuteInternal(FPCGContext* Context) cons
 	
 	// TODO: Add params pin + Apply param data overrides
 
-	check(Context->SourceComponent.IsValid());
-	UWorld* World = Context->SourceComponent->GetWorld();
+	check(Context->ExecutionSource.IsValid());
+	UWorld* World = Context->ExecutionSource->GetExecutionState().GetWorld();
 
 	UPCGWorldVolumetricData* Data = FPCGContext::NewObject_AnyThread<UPCGWorldVolumetricData>(Context);
 	Data->Initialize(World);
 	Data->QueryParams = QueryParams;
 	Data->QueryParams.Initialize();
-	Data->OriginatingComponent = Context->SourceComponent;
+	Data->OriginatingComponent = Cast<UPCGComponent>(Context->ExecutionSource.Get());
 
 	PCGWorldQuery::ExtractActorFiltersIfNeeded(Settings->QueryParams, Context, Data);
 
@@ -177,7 +177,7 @@ bool FPCGWorldRayHitQueryElement::ExecuteInternal(FPCGContext* Context) const
 	// TODO: Might want to revisit this for 3D partitioning. The reasoning here is that ray origin should be the same
 	// if we are partitioned or non-partitioned. But with 3D partitioning we might want something different.
 	// Since it is not yet widely used, we'll stick with same behavior for all cases.
-	UPCGComponent* SourceComponent = Context->SourceComponent.Get();
+	UPCGComponent* SourceComponent = Cast<UPCGComponent>(Context->ExecutionSource.Get());
 	UPCGComponent* OriginalComponent = SourceComponent ? SourceComponent->GetOriginalComponent() : nullptr;
 	AActor* Owner = OriginalComponent ? OriginalComponent->GetOwner() : nullptr;
 	FTransform Transform = Owner ? Owner->GetTransform() : FTransform::Identity;
@@ -210,13 +210,13 @@ bool FPCGWorldRayHitQueryElement::ExecuteInternal(FPCGContext* Context) const
 		}
 	}
 
-	check(Context->SourceComponent.IsValid());
-	UWorld* World = Context->SourceComponent->GetWorld();
+	check(SourceComponent);
+	UWorld* World = SourceComponent->GetWorld();
 
 	UPCGWorldRayHitData* Data = FPCGContext::NewObject_AnyThread<UPCGWorldRayHitData>(Context);
 	Data->QueryParams = QueryParams;
 	Data->QueryParams.Initialize();
-	Data->OriginatingComponent = Context->SourceComponent;
+	Data->OriginatingComponent = SourceComponent;
 	Data->CollisionShape = Settings->CollisionShape;
 	Data->Initialize(World, Transform, /*InBounds=*/FBox(EForceInit::ForceInit), LocalBounds);
 
@@ -239,14 +239,14 @@ bool FPCGWorldRayHitQueryElement::ExecuteInternal(FPCGContext* Context) const
 
 			if (Data->Bounds.IsValid)
 			{
-				BoundsCheck = [Data, Component=Context->SourceComponent.Get()](const AActor* OtherActor) -> bool
+				BoundsCheck = [Data, SourceComponent](const AActor* OtherActor) -> bool
 				{
-					const FBox OtherActorBounds = OtherActor ? PCGHelpers::GetGridBounds(OtherActor, Component) : FBox(EForceInit::ForceInit);
+					const FBox OtherActorBounds = OtherActor ? PCGHelpers::GetGridBounds(OtherActor, SourceComponent) : FBox(EForceInit::ForceInit);
 					return OtherActorBounds.Intersect(Data->Bounds);
 				};
 			}
 
-			TArray<AActor*> LandscapeActors = PCGActorSelector::FindActors(ActorSelector, Context->SourceComponent.Get(), BoundsCheck, SelfIgnoreCheck);
+			TArray<AActor*> LandscapeActors = PCGActorSelector::FindActors(ActorSelector, SourceComponent, BoundsCheck, SelfIgnoreCheck);
 			for (AActor* Landscape : LandscapeActors)
 			{
 				if (ALandscapeProxy* LandscapeProxy = Cast<ALandscapeProxy>(Landscape))
diff --git a/Engine/Plugins/PCG/Source/PCG/Private/Elements/PCGWorldRaycast.cpp b/Engine/Plugins/PCG/Source/PCG/Private/Elements/PCGWorldRaycast.cpp
--- a/Engine/Plugins/PCG/Source/PCG/Private/Elements/PCGWorldRaycast.cpp
+++ b/Engine/Plugins/PCG/Source/PCG/Private/Elements/PCGWorldRaycast.cpp
@@ -402,7 +402,10 @@ bool FPCGWorldRaycastElement::ExecuteInternal(FPCGContext* InContext) const
 		return true;
 	}
 
-	return ExecuteSlice(Context, [](ContextType* Context, const ExecStateType& ExecState, IterStateType& IterState, const uint32 IterIndex)
+	UWorld* World = Context->ExecutionSource.IsValid() ? Context->ExecutionSource->GetExecutionState().GetWorld() : nullptr;
+	check(World);
+
+	return ExecuteSlice(Context, [World](ContextType* Context, const ExecStateType& ExecState, IterStateType& IterState, const uint32 IterIndex)
 	{
 		if (Context->GetIterationStateResult(IterIndex) == EPCGTimeSliceInitResult::NoOperation)
 		{
@@ -412,9 +415,6 @@ bool FPCGWorldRaycastElement::ExecuteInternal(FPCGContext* InContext) const
 		const UPCGWorldRaycastElementSettings* Settings = Context->GetInputSettings<UPCGWorldRaycastElementSettings>();
 		check(Settings);
 
-		UWorld* World = Context->SourceComponent->GetWorld();
-		check(World);
-
 		const TArray<FPCGTaggedData> OriginInputs = Context->InputData.GetInputsByPin(PCGWorldRaycastElement::Constants::Labels::OriginsInput);
 		const FPCGTaggedData& OriginInputData = OriginInputs[IterIndex];
 		const UPCGPointData* OriginInputPointData = Cast<UPCGPointData>(OriginInputData.Data);
@@ -480,7 +480,7 @@ bool FPCGWorldRaycastElement::ExecuteInternal(FPCGContext* InContext) const
 				ExecState.CollisionObjectQueryParams,
 				Settings->CollisionShape.ToCollisionShape(),
 				ExecState.CollisionQueryParams);
-			const TOptional<FHitResult>& HitResult = PCGWorldQueryHelpers::FilterRayHitResults(&WorldQueryParams, Context->SourceComponent, Hits, IterState.CachedFilterActors);
+			const TOptional<FHitResult>& HitResult = PCGWorldQueryHelpers::FilterRayHitResults(&WorldQueryParams, Cast<UPCGComponent>(Context->ExecutionSource.Get()), Hits, IterState.CachedFilterActors);
 
 			// If there's a hit, and it's within bounds, keep the point
 			if (HitResult.IsSet() && (Settings->bUnbounded || PCGHelpers::IsInsideBounds(ExecState.Bounds, HitResult.GetValue().ImpactPoint)))
diff --git a/Engine/Plugins/PCG/Source/PCG/Private/Graph/PCGGraphCache.cpp b/Engine/Plugins/PCG/Source/PCG/Private/Graph/PCGGraphCache.cpp
--- a/Engine/Plugins/PCG/Source/PCG/Private/Graph/PCGGraphCache.cpp
+++ b/Engine/Plugins/PCG/Source/PCG/Private/Graph/PCGGraphCache.cpp
@@ -61,7 +61,7 @@ bool FPCGGraphCache::GetFromCache(const FPCGGetFromCacheParams& Params, FPCGData
 
 	const UPCGNode* InNode = Params.Node;
 	const IPCGElement* InElement = Params.Element;
-	const UPCGComponent* InComponent = Params.Component;
+	const IPCGGraphExecutionSource* InExecutionSource = Params.ExecutionSource;
 	const FPCGCrc& InDependenciesCrc = Params.Crc;
 
 	if(!InDependenciesCrc.IsValid())
@@ -70,7 +70,7 @@ bool FPCGGraphCache::GetFromCache(const FPCGGetFromCacheParams& Params, FPCGData
 		return false;
 	}
 
-	const bool bDebuggingEnabled = IsDebuggingEnabled() && InComponent && InComponent->GetOwner() && InNode;
+	const bool bDebuggingEnabled = IsDebuggingEnabled() && InExecutionSource && InNode;
 
 	{
 		TRACE_CPUPROFILER_EVENT_SCOPE(FPCGGraphCache::GetFromCache);
@@ -82,7 +82,7 @@ bool FPCGGraphCache::GetFromCache(const FPCGGetFromCacheParams& Params, FPCGData
 			if (bDebuggingEnabled)
 			{
 				// Leading spaces to align log content with warnings below - helps readability a lot.
-				UE_LOG(LogPCG, Log, TEXT("         [%s] %s\t\tCACHE HIT %u"), *InComponent->GetOwner()->GetName(), *InNode->GetNodeTitle(EPCGNodeTitleType::ListView).ToString(), InDependenciesCrc.GetValue());
+				UE_LOG(LogPCG, Log, TEXT("         [%s] %s\t\tCACHE HIT %u"), *InExecutionSource->GetExecutionState().GetDebugName(), *InNode->GetNodeTitle(EPCGNodeTitleType::ListView).ToString(), InDependenciesCrc.GetValue());
 			}
 
 			OutOutput = *Value;
@@ -93,7 +93,7 @@ bool FPCGGraphCache::GetFromCache(const FPCGGetFromCacheParams& Params, FPCGData
 		{
 			if (bDebuggingEnabled)
 			{
-				UE_LOG(LogPCG, Warning, TEXT("[%s] %s\t\tCACHE MISS %u"), *InComponent->GetOwner()->GetName(), *InNode->GetNodeTitle(EPCGNodeTitleType::ListView).ToString(), InDependenciesCrc.GetValue());
+				UE_LOG(LogPCG, Warning, TEXT("[%s] %s\t\tCACHE MISS %u"), *InExecutionSource->GetExecutionState().GetDebugName(), *InNode->GetNodeTitle(EPCGNodeTitleType::ListView).ToString(), InDependenciesCrc.GetValue());
 			}
 
 			return false;
diff --git a/Engine/Plugins/PCG/Source/PCG/Private/Graph/PCGGraphCompiler.cpp b/Engine/Plugins/PCG/Source/PCG/Private/Graph/PCGGraphCompiler.cpp
--- a/Engine/Plugins/PCG/Source/PCG/Private/Graph/PCGGraphCompiler.cpp
+++ b/Engine/Plugins/PCG/Source/PCG/Private/Graph/PCGGraphCompiler.cpp
@@ -590,7 +590,7 @@ void FPCGGraphCompiler::CreateGridLinkages(
 						static_cast<FPCGGridLinkageContext*>(InContext));
 				};
 
-				FPCGGenericElement::FContextAllocator ContextAllocator = [](const FPCGDataCollection&, TWeakObjectPtr<UPCGComponent>, const UPCGNode*)
+				FPCGGenericElement::FContextAllocator ContextAllocator = [](const FPCGInitializeElementParams&)
 				{
 					return new FPCGGridLinkageContext();
 				};
diff --git a/Engine/Plugins/PCG/Source/PCG/Private/Graph/PCGGraphExecutor.cpp b/Engine/Plugins/PCG/Source/PCG/Private/Graph/PCGGraphExecutor.cpp
--- a/Engine/Plugins/PCG/Source/PCG/Private/Graph/PCGGraphExecutor.cpp
+++ b/Engine/Plugins/PCG/Source/PCG/Private/Graph/PCGGraphExecutor.cpp
@@ -232,40 +232,43 @@ void FPCGGraphExecutor::Compile(UPCGGraph* Graph)
 	GraphCompiler.Compile(Graph);
 }
 
-FPCGTaskId FPCGGraphExecutor::Schedule(UPCGComponent* Component, const TArray<FPCGTaskId>& ExternalDependencies, const FPCGStack* InFromStack)
+FPCGTaskId FPCGGraphExecutor::Schedule(IPCGGraphExecutionSource* ExecutionSource, const TArray<FPCGTaskId>& ExternalDependencies, const FPCGStack* InFromStack)
 {
-	check(Component);
-	UPCGGraph* Graph = Component->GetGraph();
+	check(ExecutionSource);
+	UPCGGraph* Graph = ExecutionSource->GetExecutionState().GetGraph();
 
-	return Schedule(Graph, Component, GetPreGraphElement(), GetFetchInputElement(), ExternalDependencies, InFromStack, /*bAllowHierarchicalGeneration=*/true);
+	return Schedule(Graph, ExecutionSource, GetPreGraphElement(), GetFetchInputElement(), ExternalDependencies, InFromStack, /*bAllowHierarchicalGeneration=*/true);
 }
 
 FPCGTaskId FPCGGraphExecutor::Schedule(
 	UPCGGraph* Graph,
-	UPCGComponent* SourceComponent,
+	IPCGGraphExecutionSource* ExecutionSource,
 	FPCGElementPtr PreGraphElement,
 	FPCGElementPtr InputElement,
 	const TArray<FPCGTaskId>& ExternalDependencies,
 	const FPCGStack* InFromStack,
 	bool bAllowHierarchicalGeneration)
 {
-	return ScheduleGraph(FPCGScheduleGraphParams(Graph, SourceComponent, PreGraphElement, InputElement, ExternalDependencies, InFromStack, bAllowHierarchicalGeneration));
+	return ScheduleGraph(FPCGScheduleGraphParams(Graph, ExecutionSource, PreGraphElement, InputElement, ExternalDependencies, InFromStack, bAllowHierarchicalGeneration));
 }
 
 FPCGTaskId FPCGGraphExecutor::ScheduleGraph(const FPCGScheduleGraphParams& InParams)
 {
-	check(InParams.SourceComponent);
+	check(InParams.ExecutionSource);
 
-	PCGGraphExecutionLogging::LogGraphSchedule(InParams.SourceComponent, InParams.Graph);
+	PCGGraphExecutionLogging::LogGraphSchedule(InParams.ExecutionSource, InParams.Graph);
 	
 	FPCGTaskId ScheduledId = InvalidPCGTaskId;
 
 	uint32 GenerationGridSize = PCGHiGenGrid::UninitializedGridSize();
 	if (InParams.bAllowHierarchicalGeneration && InParams.Graph->IsHierarchicalGenerationEnabled())
 	{
-		if (InParams.SourceComponent->IsLocalComponent() || InParams.SourceComponent->IsPartitioned())
+		if (const UPCGComponent* PCGComponent = Cast<UPCGComponent>(InParams.ExecutionSource))
 		{
-			GenerationGridSize = InParams.SourceComponent->GetGenerationGridSize();
+			if (PCGComponent->IsLocalComponent() || PCGComponent->IsPartitioned())
+			{
+				GenerationGridSize = PCGComponent->GetGenerationGridSize();
+			}
 		}
 	}
 
@@ -281,7 +284,7 @@ FPCGTaskId FPCGGraphExecutor::ScheduleGraph(const FPCGScheduleGraphParams& InPar
 	else
 	{
 		FPCGStack ComponentStack;
-		ComponentStack.PushFrame(InParams.SourceComponent);
+		ComponentStack.PushFrame(Cast<UObject>(InParams.ExecutionSource));
 		StackContextPtr->PrependParentStack(&ComponentStack);		
 	}
 
@@ -291,7 +294,7 @@ FPCGTaskId FPCGGraphExecutor::ScheduleGraph(const FPCGScheduleGraphParams& InPar
 		PCGEditorModule->OnScheduleGraph(*StackContextPtr);
 	}
 
-	if (UPCGSubsystem* Subsystem = InParams.SourceComponent->GetSubsystem())
+	if (UPCGSubsystem* Subsystem = UWorld::GetSubsystem<UPCGSubsystem>(InParams.ExecutionSource->GetExecutionState().GetWorld()))
 	{
 		Subsystem->OnScheduleGraph(*StackContextPtr);
 	}
@@ -300,7 +303,7 @@ FPCGTaskId FPCGGraphExecutor::ScheduleGraph(const FPCGScheduleGraphParams& InPar
 	// Assign this component to the tasks
 	for (FPCGGraphTask& Task : CompiledTasks)
 	{
-		Task.SourceComponent = InParams.SourceComponent;
+		Task.ExecutionSource = InParams.ExecutionSource;
 		Task.StackContext = StackContextPtr;
 	}
 
@@ -310,7 +313,7 @@ FPCGTaskId FPCGGraphExecutor::ScheduleGraph(const FPCGScheduleGraphParams& InPar
 		check(CompiledTasks[0].Node == InParams.Graph->GetInputNode());
 
 #if WITH_EDITOR
-		bAnyNonRuntimeGenComponentScheduled = bAnyNonRuntimeGenComponentScheduled || InParams.SourceComponent->GenerationTrigger != EPCGComponentGenerationTrigger::GenerateAtRuntime;
+		bAnyNonRuntimeGenComponentScheduled = bAnyNonRuntimeGenComponentScheduled || !PCGHelpers::IsRuntimeGeneration(InParams.ExecutionSource);
 #endif
 
 		// Setup fetch task on input node
@@ -320,7 +323,7 @@ FPCGTaskId FPCGGraphExecutor::ScheduleGraph(const FPCGScheduleGraphParams& InPar
 
 		FPCGGraphScheduleTask& ScheduledTask = ScheduledTasks.Emplace_GetRef();
 		ScheduledTask.Tasks = MoveTemp(CompiledTasks);
-		ScheduledTask.SourceComponent = InParams.SourceComponent;
+		ScheduledTask.ExecutionSource = InParams.ExecutionSource;
 
 		// Offset task node ids
 		FPCGGraphCompiler::OffsetNodeIds(ScheduledTask.Tasks, NextTaskId, InvalidPCGTaskId);
@@ -353,40 +356,40 @@ FPCGTaskId FPCGGraphExecutor::ScheduleGraph(const FPCGScheduleGraphParams& InPar
 	return ScheduledId;
 }
 
-TArray<UPCGComponent*> FPCGGraphExecutor::Cancel(UPCGComponent* InComponent)
+TArray<IPCGGraphExecutionSource*> FPCGGraphExecutor::Cancel(IPCGGraphExecutionSource* InExecutionSource)
 {
-	auto CancelIfSameComponent = [InComponent](TWeakObjectPtr<UPCGComponent> Component)
+	auto CancelIfSameExecutionSource = [InExecutionSource](TWeakInterfacePtr<IPCGGraphExecutionSource> ExecutionSource)
 	{
-		return Component.IsValid() && InComponent == Component;
+		return ExecutionSource.IsValid() && InExecutionSource == ExecutionSource.Get();
 	};
 
-	return Cancel(CancelIfSameComponent).Array();
+	return Cancel(CancelIfSameExecutionSource).Array();
 }
 
-TArray<UPCGComponent*> FPCGGraphExecutor::Cancel(UPCGGraph* InGraph)
+TArray<IPCGGraphExecutionSource*> FPCGGraphExecutor::Cancel(UPCGGraph* InGraph)
 {
-	auto CancelIfComponentUsesGraph = [InGraph](TWeakObjectPtr<UPCGComponent> Component)
+	auto CancelIfComponentUsesGraph = [InGraph](TWeakInterfacePtr<IPCGGraphExecutionSource> ExecutionSource)
 	{
-		return (Component.IsValid() && Component->GetGraph() == InGraph);
+		return (ExecutionSource.IsValid() && ExecutionSource->GetExecutionState().GetGraph() == InGraph);
 	};
 
 	return Cancel(CancelIfComponentUsesGraph).Array();
 }
 
-TArray<UPCGComponent*> FPCGGraphExecutor::CancelAll()
+TArray<IPCGGraphExecutionSource*> FPCGGraphExecutor::CancelAll()
 {
-	auto CancelAllGeneration = [](TWeakObjectPtr<UPCGComponent> Component)
+	auto CancelAllGeneration = [](TWeakInterfacePtr<IPCGGraphExecutionSource> ExecutionSource)
 	{
-		return Component.IsValid();
+		return ExecutionSource.IsValid();
 	};
 
 	return Cancel(CancelAllGeneration).Array();
 }
 
-TSet<UPCGComponent*> FPCGGraphExecutor::Cancel(TFunctionRef<bool(TWeakObjectPtr<UPCGComponent>)> CancelFilter)
+TSet<IPCGGraphExecutionSource*> FPCGGraphExecutor::Cancel(TFunctionRef<bool(TWeakInterfacePtr<IPCGGraphExecutionSource>)> CancelFilter)
 {
 	check(IsInGameThread());
-	TSet<UPCGComponent*> CancelledComponents;
+	TSet<IPCGGraphExecutionSource*> CancelledExecutionSources;
 	
 	// Prevent any scheduling while cancelling
 	SchedulingDisabled = true;
@@ -400,9 +403,9 @@ TSet<UPCGComponent*> FPCGGraphExecutor::Cancel(TFunctionRef<bool(TWeakObjectPtr<
 		PCGGraphExecutor::TScopeLock ScopeLock(ScheduleLock);
 		for (const FPCGGraphScheduleTask& ScheduledTask : ScheduledTasks)
 		{
-			if (CancelFilter(ScheduledTask.SourceComponent))
+			if (CancelFilter(ScheduledTask.ExecutionSource))
 			{
-				CancelledComponents.Add(ScheduledTask.SourceComponent.Get());
+				CancelledExecutionSources.Add(ScheduledTask.ExecutionSource.Get());
 			}
 		}
 	}
@@ -412,41 +415,41 @@ TSet<UPCGComponent*> FPCGGraphExecutor::Cancel(TFunctionRef<bool(TWeakObjectPtr<
 		PCGGraphExecutor::TScopeLock ScopeLock(LiveTasksLock);
 		for (const FPCGGraphTask& Task : ReadyTasks)
 		{
-			if (CancelFilter(Task.SourceComponent))
+			if (CancelFilter(Task.ExecutionSource))
 			{
-				CancelledComponents.Add(Task.SourceComponent.Get());
+				CancelledExecutionSources.Add(Task.ExecutionSource.Get());
 			}
 		}
 
-		auto GatherCancelledComponentsForActiveTasks = [&CancelFilter, &CancelledComponents](const TArray<TSharedPtr<FPCGGraphActiveTask>>& InActiveTasks)
+		auto GatherCancelledExecutionSourcesForActiveTasks = [&CancelFilter, &CancelledExecutionSources](const TArray<TSharedPtr<FPCGGraphActiveTask>>& InActiveTasks)
 		{
 			for (const TSharedPtr<FPCGGraphActiveTask>& Task : InActiveTasks)
 			{
-				if (Task->Context && CancelFilter(Task->Context->SourceComponent))
+				if (Task->Context && CancelFilter(Task->Context->ExecutionSource))
 				{
-					CancelledComponents.Add(Task->Context->SourceComponent.Get());
+					CancelledExecutionSources.Add(Task->Context->ExecutionSource.Get());
 				}
 			}
 		};
 
-		GatherCancelledComponentsForActiveTasks(ActiveTasks);
-		GatherCancelledComponentsForActiveTasks(ActiveTasksGameThreadOnly);
-		GatherCancelledComponentsForActiveTasks(PausedTasks);
+		GatherCancelledExecutionSourcesForActiveTasks(ActiveTasks);
+		GatherCancelledExecutionSourcesForActiveTasks(ActiveTasksGameThreadOnly);
+		GatherCancelledExecutionSourcesForActiveTasks(PausedTasks);
 	}
 
 	// In one instance this function was observed to return nullptr in the CancelledComponents set.
 	// All the cancel filter lambdas check the ptr is valid, so it's not clear why. It seems perhaps
 	// the SourceComponent weak ptr became nullptr between calling CancelFilter and adding the component
 	// to the set.
-	ensure(CancelledComponents.Remove(nullptr) == 0);
+	ensure(CancelledExecutionSources.Remove(nullptr) == 0);
 
 	// Early out - nothing to cancel
-	if (CancelledComponents.IsEmpty())
+	if (CancelledExecutionSources.IsEmpty())
 	{
-		return CancelledComponents;
+		return CancelledExecutionSources;
 	}
 
-	PCGGraphExecutionLogging::LogComponentCancellation(CancelledComponents);
+	PCGGraphExecutionLogging::LogExecutionSourceCancellation(CancelledExecutionSources);
 
 	auto TryAbortScheduledTasks = [](FPCGGraphScheduleTask& ScheduledTask)
 	{
@@ -475,7 +478,7 @@ TSet<UPCGComponent*> FPCGGraphExecutor::Cancel(TFunctionRef<bool(TWeakObjectPtr<
 			for (int32 ScheduledTaskIndex = ScheduledTasks.Num() - 1; ScheduledTaskIndex >= 0; --ScheduledTaskIndex)
 			{
 				FPCGGraphScheduleTask& ScheduledTask = ScheduledTasks[ScheduledTaskIndex];
-				if (CancelledComponents.Contains(ScheduledTask.SourceComponent.Get()))
+				if (CancelledExecutionSources.Contains(ScheduledTask.ExecutionSource.Get()))
 				{
 					CancelledScheduledTasks.Add(ScheduledTask.Tasks[ScheduledTask.LastTaskIndex].NodeId);
 
@@ -495,10 +498,10 @@ TSet<UPCGComponent*> FPCGGraphExecutor::Cancel(TFunctionRef<bool(TWeakObjectPtr<
 
 					if (bContainsDependency)
 					{
-						UPCGComponent* TaskComponent = ScheduledTask.SourceComponent.Get();
-						if (TaskComponent && !CancelledComponents.Contains(TaskComponent))
+						IPCGGraphExecutionSource* TaskExecutionSource = ScheduledTask.ExecutionSource.Get();
+						if (TaskExecutionSource && !CancelledExecutionSources.Contains(TaskExecutionSource))
 						{
-							CancelledComponents.Add(TaskComponent);
+							CancelledExecutionSources.Add(TaskExecutionSource);
 							bStableCancellationSet = false;
 						}
 
@@ -517,12 +520,12 @@ TSet<UPCGComponent*> FPCGGraphExecutor::Cancel(TFunctionRef<bool(TWeakObjectPtr<
 		{
 			TArray<TTuple<TSharedPtr<FPCGGraphActiveTask>, UE::Tasks::TTask<bool>>> CancelledActiveTasks;
 			{
-				auto CancelActiveTasks = [&CancelledComponents, &CancelledActiveTasks](TArray<TSharedPtr<FPCGGraphActiveTask>>& InActiveTasks)
+				auto CancelActiveTasks = [&CancelledExecutionSources, &CancelledActiveTasks](TArray<TSharedPtr<FPCGGraphActiveTask>>& InActiveTasks)
 				{
 					for (int32 ActiveTaskIndex = InActiveTasks.Num() - 1; ActiveTaskIndex >= 0; --ActiveTaskIndex)
 					{
 						TSharedPtr<FPCGGraphActiveTask> ActiveTask = InActiveTasks[ActiveTaskIndex];
-						if (ActiveTask->Context && CancelledComponents.Contains(ActiveTask->Context->SourceComponent.Get()))
+						if (ActiveTask->Context && CancelledExecutionSources.Contains(ActiveTask->Context->ExecutionSource.Get()))
 						{
 							// While we have lock Task can't complete, but we can't wait on this task with the lock neither so we capture it here 
 							// and are going to wait on it outside of the LiveTaskLock
@@ -553,7 +556,7 @@ TSet<UPCGComponent*> FPCGGraphExecutor::Cancel(TFunctionRef<bool(TWeakObjectPtr<
 				check(ActiveTask->Element);
 				ActiveTask->Element->Abort(ActiveTask->Context.Get());
 												
-				bStableCancellationSet &= !CancelNextTasks(ActiveTask->NodeId, CancelledComponents);
+				bStableCancellationSet &= !CancelNextTasks(ActiveTask->NodeId, CancelledExecutionSources);
 			}
 		}
 
@@ -566,7 +569,7 @@ TSet<UPCGComponent*> FPCGGraphExecutor::Cancel(TFunctionRef<bool(TWeakObjectPtr<
 				{
 					FPCGGraphTask& Task = ReadyTasks[ReadyTaskIndex];
 
-					if (CancelledComponents.Contains(Task.SourceComponent.Get()))
+					if (CancelledExecutionSources.Contains(Task.ExecutionSource.Get()))
 					{
 						CancelledReadyTasks.Add(MoveTemp(Task));
 						ReadyTasks.RemoveAtSwap(ReadyTaskIndex);
@@ -585,7 +588,7 @@ TSet<UPCGComponent*> FPCGGraphExecutor::Cancel(TFunctionRef<bool(TWeakObjectPtr<
 				RemoveTaskFromInputSuccessors(CancelledTaskId, CancelledReadyTask.Inputs);
 
 				FPCGContext::Release(CancelledReadyTask.Context);
-				bStableCancellationSet &= !CancelNextTasks(CancelledTaskId, CancelledComponents);
+				bStableCancellationSet &= !CancelNextTasks(CancelledTaskId, CancelledExecutionSources);
 			}
 		}
 
@@ -598,7 +601,7 @@ TSet<UPCGComponent*> FPCGGraphExecutor::Cancel(TFunctionRef<bool(TWeakObjectPtr<
 				for (int32 PausedTaskIndex = PausedTasks.Num() - 1; PausedTaskIndex >= 0; --PausedTaskIndex)
 				{
 					TSharedPtr<FPCGGraphActiveTask> Task = PausedTasks[PausedTaskIndex];
-					if (Task->Context && CancelledComponents.Contains(Task->Context->SourceComponent.Get()))
+					if (Task->Context && CancelledExecutionSources.Contains(Task->Context->ExecutionSource.Get()))
 					{
 						check(Task->TaskIndex == PausedTaskIndex);
 						CancelledPausedTasks.Add(Task);
@@ -636,7 +639,7 @@ TSet<UPCGComponent*> FPCGGraphExecutor::Cancel(TFunctionRef<bool(TWeakObjectPtr<
 				FPCGTaskId CancelledTaskId = CancelledPausedTask->NodeId;
 				RemoveTaskFromInputSuccessors(CancelledTaskId, CancelledPausedTask->Inputs);
 								
-				bStableCancellationSet &= !CancelNextTasks(CancelledTaskId, CancelledComponents);
+				bStableCancellationSet &= !CancelNextTasks(CancelledTaskId, CancelledExecutionSources);
 			}
 		}
 	}
@@ -646,7 +649,7 @@ TSet<UPCGComponent*> FPCGGraphExecutor::Cancel(TFunctionRef<bool(TWeakObjectPtr<
 	UpdateGenerationNotification();
 #endif
 
-	return CancelledComponents;
+	return CancelledExecutionSources;
 }
 
 bool FPCGGraphExecutor::IsGraphCurrentlyExecuting(const UPCGGraph* InGraph)
@@ -654,9 +657,9 @@ bool FPCGGraphExecutor::IsGraphCurrentlyExecuting(const UPCGGraph* InGraph)
 	bool bAnyPresent = false;
 
 	// This makes use of the Cancel function which runs over all tasks, but it always returns false so no tasks are cancelled.
-	auto CheckIfAnyGraphTasksPresent = [InGraph, &bAnyPresent](TWeakObjectPtr<UPCGComponent> Component)
+	auto CheckIfAnyGraphTasksPresent = [InGraph, &bAnyPresent](TWeakInterfacePtr<IPCGGraphExecutionSource> ExecutionSource)
 	{
-		bAnyPresent |= Component.IsValid() && Component->GetGraph() == InGraph;
+		bAnyPresent |= ExecutionSource.IsValid() && ExecutionSource->GetExecutionState().GetGraph() == InGraph;
 		return false;
 	};
 
@@ -708,19 +711,19 @@ void FPCGGraphExecutor::SetExecutionCacheData(FPCGTaskId InGraphExecutionTaskId,
 FPCGTaskId FPCGGraphExecutor::ScheduleGeneric(const FPCGScheduleGenericParams& Params)
 {
 	return ScheduleSingleElement(MakeShared<FPCGGenericElement>(Params.Operation, Params.AbortOperation, Params.bSupportBasePointDataInput, Params.bCanExecuteOnlyOnMainThread),
-		!!Params.AbortOperation, Params.SourceComponent, Params.ExecutionDependencies, Params.DataDependencies);
+		!!Params.AbortOperation, Params.ExecutionSource, Params.ExecutionDependencies, Params.DataDependencies);
 }
 
-FPCGTaskId FPCGGraphExecutor::ScheduleGeneric(TFunction<bool()> InOperation, UPCGComponent* InSourceComponent, const TArray<FPCGTaskId>& TaskExecutionDependencies)
+FPCGTaskId FPCGGraphExecutor::ScheduleGeneric(TFunction<bool()> InOperation, IPCGGraphExecutionSource* InExecutionSource, const TArray<FPCGTaskId>& TaskExecutionDependencies)
 {
 	return ScheduleGeneric(
 		InOperation,
 		TFunction<void()>(),
-		InSourceComponent,
+		InExecutionSource,
 		TaskExecutionDependencies);
 }
 
-FPCGTaskId FPCGGraphExecutor::ScheduleGeneric(TFunction<bool()> InOperation, TFunction<void()> InAbortOperation, UPCGComponent* InSourceComponent, const TArray<FPCGTaskId>& TaskExecutionDependencies)
+FPCGTaskId FPCGGraphExecutor::ScheduleGeneric(TFunction<bool()> InOperation, TFunction<void()> InAbortOperation, IPCGGraphExecutionSource* InExecutionSource, const TArray<FPCGTaskId>& TaskExecutionDependencies)
 {
 	// Since we have no context, the generic task will consume no input (no data dependencies).
 	return ScheduleGenericWithContext(
@@ -735,32 +738,32 @@ FPCGTaskId FPCGGraphExecutor::ScheduleGeneric(TFunction<bool()> InOperation, TFu
 				AbortOperation();
 			}
 		},
-		InSourceComponent,
+		InExecutionSource,
 		TaskExecutionDependencies,
 		/*TaskDataDependencies=*/{});
 }
 
-FPCGTaskId FPCGGraphExecutor::ScheduleGenericWithContext(TFunction<bool(FPCGContext*)> InOperation, UPCGComponent* InSourceComponent, const TArray<FPCGTaskId>& TaskExecutionDependencies, const TArray<FPCGTaskId>& TaskDataDependencies, bool bSupportBasePointDataInput)
+FPCGTaskId FPCGGraphExecutor::ScheduleGenericWithContext(TFunction<bool(FPCGContext*)> InOperation, IPCGGraphExecutionSource* InExecutionSource, const TArray<FPCGTaskId>& TaskExecutionDependencies, const TArray<FPCGTaskId>& TaskDataDependencies, bool bSupportBasePointDataInput)
 {
 	return ScheduleGenericWithContext(
 		InOperation,
 		TFunction<void(FPCGContext*)>(),
-		InSourceComponent,
+		InExecutionSource,
 		TaskExecutionDependencies,
 		TaskDataDependencies,
 		bSupportBasePointDataInput);
 }
 
-FPCGTaskId FPCGGraphExecutor::ScheduleGenericWithContext(TFunction<bool(FPCGContext*)> InOperation, TFunction<void(FPCGContext*)> InAbortOperation, UPCGComponent* InSourceComponent, const TArray<FPCGTaskId>& TaskExecutionDependencies, const TArray<FPCGTaskId>& TaskDataDependencies, bool bSupportBasePointDataInput)
+FPCGTaskId FPCGGraphExecutor::ScheduleGenericWithContext(TFunction<bool(FPCGContext*)> InOperation, TFunction<void(FPCGContext*)> InAbortOperation, IPCGGraphExecutionSource* InExecutionSource, const TArray<FPCGTaskId>& TaskExecutionDependencies, const TArray<FPCGTaskId>& TaskDataDependencies, bool bSupportBasePointDataInput)
 {
 	return ScheduleSingleElement(MakeShared<FPCGGenericElement>(InOperation, InAbortOperation, bSupportBasePointDataInput),
-		!!InAbortOperation, InSourceComponent, TaskExecutionDependencies, TaskDataDependencies);
+		!!InAbortOperation, InExecutionSource, TaskExecutionDependencies, TaskDataDependencies);
 }
 
-FPCGTaskId FPCGGraphExecutor::ScheduleSingleElement(TSharedRef<IPCGElement> Element, bool bHasAbortCallbacks, UPCGComponent* InSourceComponent, const TArray<FPCGTaskId>& TaskExecutionDependencies, const TArray<FPCGTaskId>& TaskDataDependencies)
+FPCGTaskId FPCGGraphExecutor::ScheduleSingleElement(TSharedRef<IPCGElement> Element, bool bHasAbortCallbacks, IPCGGraphExecutionSource* InExecutionSource, const TArray<FPCGTaskId>& TaskExecutionDependencies, const TArray<FPCGTaskId>& TaskDataDependencies)
 {
 #if WITH_EDITOR
-	bAnyNonRuntimeGenComponentScheduled = bAnyNonRuntimeGenComponentScheduled || (InSourceComponent ? (InSourceComponent->GenerationTrigger != EPCGComponentGenerationTrigger::GenerateAtRuntime) : true);
+	bAnyNonRuntimeGenComponentScheduled = bAnyNonRuntimeGenComponentScheduled || !PCGHelpers::IsRuntimeGeneration(InExecutionSource);
 #endif
 
 	// Build task & element to hold the operation to perform
@@ -780,7 +783,7 @@ FPCGTaskId FPCGGraphExecutor::ScheduleSingleElement(TSharedRef<IPCGElement> Elem
 		TaskInput.bIsUsedMultipleTimes = false;
 	}
 
-	Task.SourceComponent = InSourceComponent;
+	Task.ExecutionSource = InExecutionSource;
 	Task.Element = Element;
 
 	{
@@ -791,7 +794,7 @@ FPCGTaskId FPCGGraphExecutor::ScheduleSingleElement(TSharedRef<IPCGElement> Elem
 
 		FPCGGraphScheduleTask& ScheduledTask = ScheduledTasks.Emplace_GetRef();
 		ScheduledTask.Tasks.Add(Task);
-		ScheduledTask.SourceComponent = InSourceComponent;
+		ScheduledTask.ExecutionSource = InExecutionSource;
 		ScheduledTask.bHasAbortCallbacks = bHasAbortCallbacks;
 	}
 	return Task.NodeId;
@@ -1131,12 +1134,12 @@ void FPCGGraphExecutor::PostTaskExecute(TSharedPtr<FPCGGraphActiveTask> ActiveTa
 			});
 		}
 		
-		if (UPCGComponent* SourceComponent = ActiveTask.Context->SourceComponent.Get())
+		if (IPCGGraphExecutionSource* ExecutionSource = ActiveTask.Context->ExecutionSource.Get())
 		{
 			if (ActiveTask.StackIndex != INDEX_NONE)
 			{
 				const FPCGStack* Stack = ActiveTask.StackContext->GetStack(ActiveTask.StackIndex);
-				SourceComponent->StoreInspectionData(Stack, ActiveTask.Context->Node, &ActiveTask.Context->Timer, ActiveTask.Context->InputData, ActiveTask.Context->OutputData, /*bUsedCache=*/false);
+				ExecutionSource->GetExecutionState().GetInspection().StoreInspectionData(Stack, ActiveTask.Context->Node, &ActiveTask.Context->Timer, ActiveTask.Context->InputData, ActiveTask.Context->OutputData, /*bUsedCache=*/false);
 			}
 		}
 #endif // WITH_EDITOR
@@ -1303,7 +1306,7 @@ void FPCGGraphExecutor::PrepareForExecute(FPCGGraphTask& Task, FCachedResult*& O
 	{
 		PCGGraphExecutor::TScopeLock ScopeLock(CachingResultsLock);
 		FCachedResult LocalCachedResult;
-		CacheStatus = Task.Element->RetrieveResultsFromCache(&GraphCache, Task.Node, Task.TaskInput, Task.SourceComponent.Get(), LocalCachedResult.Output, &DependenciesCrc);
+		CacheStatus = Task.Element->RetrieveResultsFromCache(&GraphCache, Task.Node, Task.TaskInput, Task.ExecutionSource.Get(), LocalCachedResult.Output, &DependenciesCrc);
 
 		bResultAlreadyInCache = (CacheStatus == EPCGCachingStatus::Cached);
 		if (bResultAlreadyInCache)
@@ -1332,12 +1335,12 @@ void FPCGGraphExecutor::PrepareForExecute(FPCGGraphTask& Task, FCachedResult*& O
 		check(OutCachedResult);
 #if WITH_EDITOR
 		// Doing this now since we're about to modify ReadyTasks potentially reallocating while Task is a reference. 
-		if (UPCGComponent* SourceComponent = Task.SourceComponent.Get())
+		if (IPCGGraphExecutionSource* ExecutionSource = Task.ExecutionSource.Get())
 		{
 			if (Task.StackIndex != INDEX_NONE)
 			{
 				const FPCGStack* Stack = Task.GetStack();
-				SourceComponent->StoreInspectionData(Stack, Task.Node, nullptr, Task.TaskInput, OutCachedResult->Output, /*bUsedCache=*/true);
+				ExecutionSource->GetExecutionState().GetInspection().StoreInspectionData(Stack, Task.Node, nullptr, Task.TaskInput, OutCachedResult->Output, /*bUsedCache=*/true);
 			}
 		}
 #endif
@@ -1361,7 +1364,7 @@ void FPCGGraphExecutor::PrepareForExecute(FPCGGraphTask& Task, FCachedResult*& O
 	check(!Task.Context);
 	
 	{
-		Task.Context = Task.Element->Initialize(Task.TaskInput, Task.SourceComponent, Task.Node);
+		Task.Context = Task.Element->Initialize(FPCGInitializeElementParams(&Task.TaskInput, Task.ExecutionSource.Get(), Task.Node));
 		Task.Context->GraphExecutor = this;
 
 		// Skip PostLoad only when not on GameThread
@@ -1900,11 +1903,11 @@ TArray<FPCGTaskId> FPCGGraphExecutor::ProcessCachedResultsInternal(TArray<FCache
 	return NextTasks;
 }
 
-bool FPCGGraphExecutor::CancelNextTasks(FPCGTaskId CancelledTask, TSet<UPCGComponent*>& OutCancelledComponents)
+bool FPCGGraphExecutor::CancelNextTasks(FPCGTaskId CancelledTask, TSet<IPCGGraphExecutionSource*>& OutCancelledExecutionSources)
 {
 	check(IsInGameThread());
 
-	bool bAddedComponents = false;
+	bool bNewCancellations = false;
 
 	TSet<FPCGTaskId> LocalSuccessors;
 	{
@@ -1918,10 +1921,10 @@ bool FPCGGraphExecutor::CancelNextTasks(FPCGTaskId CancelledTask, TSet<UPCGCompo
 			{
 				if (FPCGGraphTask* Task = Tasks.Find(Successor))
 				{
-					if (!OutCancelledComponents.Contains(Task->SourceComponent.Get()))
+					if (!OutCancelledExecutionSources.Contains(Task->ExecutionSource.Get()))
 					{
-						OutCancelledComponents.Add(Task->SourceComponent.Get());
-						bAddedComponents = true;
+						OutCancelledExecutionSources.Add(Task->ExecutionSource.Get());
+						bNewCancellations = true;
 					}
 
 					if (Task->Element)
@@ -1938,7 +1941,7 @@ bool FPCGGraphExecutor::CancelNextTasks(FPCGTaskId CancelledTask, TSet<UPCGCompo
 
 	for (FPCGTaskId Successor : LocalSuccessors)
 	{
-		bAddedComponents |= CancelNextTasks(Successor, OutCancelledComponents);
+		bNewCancellations |= CancelNextTasks(Successor, OutCancelledExecutionSources);
 	}
 
 	{
@@ -1950,10 +1953,10 @@ bool FPCGGraphExecutor::CancelNextTasks(FPCGTaskId CancelledTask, TSet<UPCGCompo
 		{
 			for (const TSharedPtr<FPCGGraphActiveTask>& ActiveTask : ActiveTaskSuccessors)
 			{
-				if (!OutCancelledComponents.Contains(ActiveTask->Context->SourceComponent.Get()))
+				if (!OutCancelledExecutionSources.Contains(ActiveTask->Context->ExecutionSource.Get()))
 				{
-					OutCancelledComponents.Add(ActiveTask->Context->SourceComponent.Get());
-					bAddedComponents = true;
+					OutCancelledExecutionSources.Add(ActiveTask->Context->ExecutionSource.Get());
+					bNewCancellations = true;
 				}
 
 				// Do not call abort as this is an ActiveTask and will be processed by the Cancellation of it's component
@@ -1966,15 +1969,15 @@ bool FPCGGraphExecutor::CancelNextTasks(FPCGTaskId CancelledTask, TSet<UPCGCompo
 		PCGGraphExecutor::TScopeLock ScopeLock(ScheduleLock);
 		for (FPCGGraphScheduleTask& ScheduledTask : ScheduledTasks)
 		{
-			if (!OutCancelledComponents.Contains(ScheduledTask.SourceComponent.Get()) &&
+			if (!OutCancelledExecutionSources.Contains(ScheduledTask.ExecutionSource.Get()) &&
 				Algo::AnyOf(ScheduledTask.Tasks[ScheduledTask.FirstTaskIndex].Inputs, [CancelledTask](const FPCGGraphTaskInput& Input) { return Input.TaskId == CancelledTask; }))
 			{
-				OutCancelledComponents.Add(ScheduledTask.SourceComponent.Get());
-				bAddedComponents = true;
+				OutCancelledExecutionSources.Add(ScheduledTask.ExecutionSource.Get());
+				bNewCancellations = true;
 			}
 		}
 	}
-	return bAddedComponents;
+	return bNewCancellations;
 }
 
 void FPCGGraphExecutor::RemoveTaskFromInputSuccessors(FPCGTaskId CancelledTask, const TArray<FPCGGraphTaskInput>& CancelledTaskInputs)
@@ -2014,7 +2017,7 @@ void FPCGGraphExecutor::BuildTaskInput(FPCGGraphTask& Task, TArray<FPCGTaskId>&
 
 		const FString Message = FString::Printf(
 			TEXT("[%s] %s - BuildTaskInput - too many data items arriving on single data pin '%s', only first data item will be used"),
-			(Task.SourceComponent.Get() && Task.SourceComponent->GetOwner()) ? *Task.SourceComponent->GetOwner()->GetName() : TEXT("MissingComponent"),
+			Task.ExecutionSource.Get() ? *Task.ExecutionSource->GetExecutionState().GetDebugName() : TEXT("Missing Execution Source"),
 			Task.Node ? *Task.Node->GetNodeTitle(EPCGNodeTitleType::ListView).ToString() : TEXT("MissingNode"),
 			DownstreamPin.IsSet() ? *DownstreamPin.GetValue().Label.ToString() : TEXT("NoPin"));
 
@@ -2688,7 +2691,7 @@ bool FPCGPreGraphElement::ExecuteInternal(FPCGContext* Context) const
 	check(Context);
 
 	check(IsInGameThread());
-	UPCGComponent* Component = Context->SourceComponent.Get();
+	UPCGComponent* Component = Cast<UPCGComponent>(Context->ExecutionSource.Get());
 
 	// Early out if the component has been deleted/is invalid
 	if (!Component)
@@ -2721,7 +2724,7 @@ bool FPCGFetchInputElement::ExecuteInternal(FPCGContext* Context) const
 	Context->OutputData = Context->InputData;
 
 	// Second: fetch the inputs provided by the component
-	UPCGComponent* Component = Context->SourceComponent.Get();
+	UPCGComponent* Component = Cast<UPCGComponent>(Context->ExecutionSource.Get());
 	
 	// Early out if the component has been deleted/is invalid
 	if (!Component)
@@ -2830,13 +2833,11 @@ FPCGGenericElement::FPCGGenericElement(TFunction<bool(FPCGContext*)> InOperation
 {
 }
 
-FPCGContext* FPCGGenericElement::Initialize(const FPCGDataCollection& InputData, TWeakObjectPtr<UPCGComponent> SourceComponent, const UPCGNode* Node)
+FPCGContext* FPCGGenericElement::Initialize(const FPCGInitializeElementParams& InParams)
 {
 	check(ContextAllocator);
-	FPCGContext* Context = ContextAllocator(InputData, SourceComponent, Node);
-	Context->InputData = InputData;
-	Context->SourceComponent = SourceComponent;
-	Context->Node = Node;
+	FPCGContext* Context = ContextAllocator(InParams);
+	Context->InitFromParams(InParams);
 
 	return Context;
 }
@@ -2881,14 +2882,15 @@ namespace PCGGraphExecutor
 		check(InContext);
 		TRACE_CPUPROFILER_EVENT_SCOPE(PCGGraphExecutor::ExecuteGridLinkage);
 
-		if (!InContext->SourceComponent.IsValid())
+		UPCGComponent* SourceComponent = Cast<UPCGComponent>(InContext->ExecutionSource.Get());
+		if (!SourceComponent)
 		{
 			// Source no longer exists, nothing to be done.
 			return true;
 		}
 
 		// Non-hierarchical generation - no linkage required - data should just pass through.
-		if (!InContext->SourceComponent->GetGraph()->IsHierarchicalGenerationEnabled()
+		if (!InContext->ExecutionSource->GetExecutionState().GetGraph()->IsHierarchicalGenerationEnabled()
 			|| !ensure(PCGHiGenGrid::IsValidGrid(InFromGrid) || InFromGrid == EPCGHiGenGrid::Unbounded))
 		{
 			InContext->OutputData = InContext->InputData;
@@ -2941,13 +2943,13 @@ namespace PCGGraphExecutor
 
 			PCGGraphExecutionLogging::LogGridLinkageTaskExecuteStore(InContext, InGenerationGrid, FromGridSize, ToGridSize, InResourceKey, Data.TaggedData.Num());
 			
-			InContext->SourceComponent->StoreOutputDataForPin(InResourceKey, Data);
+			SourceComponent->StoreOutputDataForPin(InResourceKey, Data);
 		}
 		else if (InToGrid == InGenerationGrid && FromGridSize != ToGridSize)
 		{
 			PCGGraphExecutionLogging::LogGridLinkageTaskExecuteRetrieve(InContext, InGenerationGrid, FromGridSize, ToGridSize, InResourceKey);
 
-			UPCGSubsystem* Subsystem = UPCGSubsystem::GetInstance(InContext->SourceComponent->GetWorld());
+			UPCGSubsystem* Subsystem = UPCGSubsystem::GetInstance(SourceComponent->GetWorld());
 			if (!ensure(Subsystem))
 			{
 				return false;
@@ -2957,15 +2959,15 @@ namespace PCGGraphExecutor
 			bool bComponentWithDataIsOriginalComponent = false;
 			if (FromGridSize == PCGHiGenGrid::UnboundedGridSize())
 			{
-				ComponentWithData = InContext->SourceComponent->GetOriginalComponent();
+				ComponentWithData = SourceComponent->GetOriginalComponent();
 				bComponentWithDataIsOriginalComponent = true;
 			}
 			else
 			{
-				const AActor* ComponentActor = InContext->SourceComponent->GetOwner();
+				const AActor* ComponentActor = SourceComponent->GetOwner();
 				if (ComponentActor)
 				{
-					const UPCGComponent* OriginalComponent = InContext->SourceComponent->GetOriginalComponent();
+					const UPCGComponent* OriginalComponent = SourceComponent->GetOriginalComponent();
 					if (!ensure(OriginalComponent))
 					{
 						return false;
@@ -2985,7 +2987,7 @@ namespace PCGGraphExecutor
 				// Nothing we can do currently if PCG component not present. One idea is to schedule an artifact-less execution but that
 				// comes with complications - artifacts/side effects are an integral part of execution. Most likely we'll do a cleanup
 				// pass of any unwanted artifacts/local-components later.
-				PCGGraphExecutionLogging::LogGridLinkageTaskExecuteRetrieveNoLocalComponent(InContext, InResourceKey);
+				PCGGraphExecutionLogging::LogGridLinkageTaskExecuteRetrieveNoLocalSource(InContext, InResourceKey);
 				return true;
 			}
 
@@ -3042,7 +3044,7 @@ namespace PCGGraphExecutor
 					}
 					else
 					{
-						FPCGScheduleGenericParams Params(MoveTemp(WakeUpLambda), InContext->SourceComponent.Get());
+						FPCGScheduleGenericParams Params(MoveTemp(WakeUpLambda), InContext->ExecutionSource.Get());
 						Params.bCanExecuteOnlyOnMainThread = false;
 						Params.ExecutionDependencies = { ComponentWithData->GetGenerationTaskId() };
 
@@ -3059,7 +3061,7 @@ namespace PCGGraphExecutor
 			{
 				PCGGraphExecutionLogging::LogGridLinkageTaskExecuteRetrieveScheduleGraph(InContext, ComponentWithData, InResourceKey);
 
-				EPCGComponentGenerationTrigger GenTrigger = (InContext->SourceComponent->IsManagedByRuntimeGenSystem()) ?
+				EPCGComponentGenerationTrigger GenTrigger = PCGHelpers::IsRuntimeGeneration(InContext->ExecutionSource.Get()) ?
 					EPCGComponentGenerationTrigger::GenerateAtRuntime : EPCGComponentGenerationTrigger::GenerateOnDemand;
 
 				// Wake up this task after graph has generated.
@@ -3071,7 +3073,7 @@ namespace PCGGraphExecutor
 				}
 				else
 				{
-					FPCGScheduleGenericParams Params(MoveTemp(WakeUpLambda), InContext->SourceComponent.Get());
+					FPCGScheduleGenericParams Params(MoveTemp(WakeUpLambda), InContext->ExecutionSource.Get());
 					Params.bCanExecuteOnlyOnMainThread = false;
 					Params.ExecutionDependencies = { GraphTaskId };
 
@@ -3120,7 +3122,7 @@ FPCGElementPtr UPCGGridLinkageSettings::CreateElement() const
 			static_cast<FPCGGridLinkageContext*>(InContext));
 	};
 
-	FPCGGenericElement::FContextAllocator ContextAllocator = [](const FPCGDataCollection&, TWeakObjectPtr<UPCGComponent>, const UPCGNode*)
+	FPCGGenericElement::FContextAllocator ContextAllocator = [](const FPCGInitializeElementParams&)
 	{
 		return new FPCGGridLinkageContext();
 	};
diff --git a/Engine/Plugins/PCG/Source/PCG/Private/Graph/PCGGraphExecutor.h b/Engine/Plugins/PCG/Source/PCG/Private/Graph/PCGGraphExecutor.h
--- a/Engine/Plugins/PCG/Source/PCG/Private/Graph/PCGGraphExecutor.h
+++ b/Engine/Plugins/PCG/Source/PCG/Private/Graph/PCGGraphExecutor.h
@@ -54,14 +54,14 @@ public:
 	void Compile(UPCGGraph* InGraph);
 
 	/** Schedules the execution of a given graph with specified inputs. This call is threadsafe */
-	FPCGTaskId Schedule(UPCGComponent* InComponent, const TArray<FPCGTaskId>& TaskDependency = TArray<FPCGTaskId>(), const FPCGStack* InFromStack = nullptr);
+	FPCGTaskId Schedule(IPCGGraphExecutionSource* InExecutionSource, const TArray<FPCGTaskId>& TaskDependency = TArray<FPCGTaskId>(), const FPCGStack* InFromStack = nullptr);
 
 	/** Schedules the execution of a given graph with specified inputs. This call is threadsafe */
 	FPCGTaskId ScheduleGraph(const FPCGScheduleGraphParams& InParams);
 
 	FPCGTaskId Schedule(
 		UPCGGraph* Graph,
-		UPCGComponent* InSourceComponent,
+		IPCGGraphExecutionSource* InExecutionSource,
 		FPCGElementPtr PreGraphElement,
 		FPCGElementPtr InputElement,
 		const TArray<FPCGTaskId>& TaskDependency,
@@ -69,13 +69,13 @@ public:
 		bool bAllowHierarchicalGeneration);
 
 	/** Cancels all tasks originating from the given component */
-	TArray<UPCGComponent*> Cancel(UPCGComponent* InComponent);
+	TArray<IPCGGraphExecutionSource*> Cancel(IPCGGraphExecutionSource* InExecutionSource);
 	
 	/** Cancels all tasks running a given graph */
-	TArray<UPCGComponent*> Cancel(UPCGGraph* InGraph);
+	TArray<IPCGGraphExecutionSource*> Cancel(UPCGGraph* InGraph);
 
 	/** Cancels all tasks */
-	TArray<UPCGComponent*> CancelAll();
+	TArray<IPCGGraphExecutionSource*> CancelAll();
 
 	/** Returns true if any task is scheduled or executing for the given graph. */
 	bool IsGraphCurrentlyExecuting(const UPCGGraph* InGraph);
@@ -87,28 +87,28 @@ public:
 	FPCGTaskId ScheduleGeneric(const FPCGScheduleGenericParams& Params);
 
 	// Back compatibility function. Use ScheduleGenericWithContext or ScheduleGeneric taking a FPCGScheduleGenericParams
-	FPCGTaskId ScheduleGeneric(TFunction<bool()> InOperation, UPCGComponent* InSourceComponent, const TArray<FPCGTaskId>& TaskExecutionDependencies);
-	FPCGTaskId ScheduleGeneric(TFunction<bool()> InOperation, TFunction<void()> InAbortOperation, UPCGComponent* InSourceComponent, const TArray<FPCGTaskId>& TaskExecutionDependencies);
+	FPCGTaskId ScheduleGeneric(TFunction<bool()> InOperation, IPCGGraphExecutionSource* InExecutionSource, const TArray<FPCGTaskId>& TaskExecutionDependencies);
+	FPCGTaskId ScheduleGeneric(TFunction<bool()> InOperation, TFunction<void()> InAbortOperation, IPCGGraphExecutionSource* InExecutionSource, const TArray<FPCGTaskId>& TaskExecutionDependencies);
 
 	/** General job scheduling
 	*  @param InOperation:                Callback that takes a Context as argument and returns true if the task is done, false otherwise
-	*  @param InSourceComponent:          PCG component associated with this task. Can be null.
+	*  @param InExecutionSource:          Execution source associated with this task. Can be null.
 	*  @param TaskExecutionDependencies:  Task will wait on these tasks to execute and won't take their output data as input.
 	*  @param TaskDataDependencies:       Task will wait on these tasks to execute and will take their output data as input.
 	*  @param bSupportBasePointDataInput: When true, generic element will not convert input to UPCGPointData, this is false by default to preserver backward compatibility.
 	* 
 	*/
-	FPCGTaskId ScheduleGenericWithContext(TFunction<bool(FPCGContext*)> InOperation, UPCGComponent* InSourceComponent, const TArray<FPCGTaskId>& TaskExecutionDependencies, const TArray<FPCGTaskId>& TaskDataDependencies, bool bSupportBasePointDataInput = false);
+	FPCGTaskId ScheduleGenericWithContext(TFunction<bool(FPCGContext*)> InOperation, IPCGGraphExecutionSource* InExecutionSource, const TArray<FPCGTaskId>& TaskExecutionDependencies, const TArray<FPCGTaskId>& TaskDataDependencies, bool bSupportBasePointDataInput = false);
 
 	/** General job scheduling
 	*  @param InOperation:                Callback that takes a Context as argument and returns true if the task is done, false otherwise
 	*  @param InAbortOperation:           Callback that is called if the task is aborted (cancelled) before fully executed.
-	*  @param InSourceComponent:          PCG component associated with this task. Can be null.
+	*  @param InExecutionSource:          Execution source associated with this task. Can be null.
 	*  @param TaskExecutionDependencies:  Task will wait on these tasks to execute and won't take their output data as input.
 	*  @param TaskDataDependencies:       Task will wait on these tasks to execute and will take their output data as input.
 	*  @param bSupportBasePointDataInput: When true, generic element will not convert input to UPCGPointData, this is false by default to preserver backward compatibility.
 	*/
-	FPCGTaskId ScheduleGenericWithContext(TFunction<bool(FPCGContext*)> InOperation, TFunction<void(FPCGContext*)> InAbortOperation, UPCGComponent* InSourceComponent, const TArray<FPCGTaskId>& TaskExecutionDependencies, const TArray<FPCGTaskId>& TaskDataDependencies, bool bSupportBasePointDataInput = false);
+	FPCGTaskId ScheduleGenericWithContext(TFunction<bool(FPCGContext*)> InOperation, TFunction<void(FPCGContext*)> InAbortOperation, IPCGGraphExecutionSource* InExecutionSource, const TArray<FPCGTaskId>& TaskExecutionDependencies, const TArray<FPCGTaskId>& TaskDataDependencies, bool bSupportBasePointDataInput = false);
 
 	/** Gets data in the output results. Returns false if data is not ready. */
 	bool GetOutputData(FPCGTaskId InTaskId, FPCGDataCollection& OutData);
@@ -205,7 +205,7 @@ private:
 		FPCGGraphExecutor* Executor = nullptr;
 	};
 
-	FPCGTaskId ScheduleSingleElement(TSharedRef<IPCGElement> Element, bool bHasAbortCallbacks, UPCGComponent* InSourceComponent,
+	FPCGTaskId ScheduleSingleElement(TSharedRef<IPCGElement> Element, bool bHasAbortCallbacks, IPCGGraphExecutionSource* InExecutionSource,
 		const TArray<FPCGTaskId>& TaskExecutionDependencies, const TArray<FPCGTaskId>& TaskDataDependencies);
 
 	void PostTaskExecute(TSharedPtr<FPCGGraphActiveTask> ActiveTask, bool bIsDone);
@@ -213,12 +213,12 @@ private:
 	void ExecuteTasksEnded();
 	bool ExecuteScheduling(double EndTime, TSharedPtr<FPCGGraphActiveTask>* OutMainThreadTask = nullptr, bool bForceCheckPausedTasks = false);
 
-	TSet<UPCGComponent*> Cancel(TFunctionRef<bool(TWeakObjectPtr<UPCGComponent>)> CancelFilter);
+	TSet<IPCGGraphExecutionSource*> Cancel(TFunctionRef<bool(TWeakInterfacePtr<IPCGGraphExecutionSource>)> CancelFilter);
 	void ClearAllTasks();
 	void QueueNextTasks(FPCGTaskId FinishedTask);
 	TArray<FCachedResult*> QueueNextTasksInternal(FPCGTaskId FinishedTask);
 
-	bool CancelNextTasks(FPCGTaskId CancelledTask, TSet<UPCGComponent*>& OutCancelledComponents);
+	bool CancelNextTasks(FPCGTaskId CancelledTask, TSet<IPCGGraphExecutionSource*>& OutCancelledExecutionSources);
 	void RemoveTaskFromInputSuccessors(FPCGTaskId CancelledTask, const TArray<FPCGGraphTaskInput>& CancelledTaskInputs);
 	void RemoveTaskFromInputSuccessorsNoLock(FPCGTaskId CancelledTask, const TArray<FPCGGraphTaskInput>& CancelledTaskInputs);
 	
@@ -397,11 +397,11 @@ protected:
 class FPCGGenericElement : public IPCGElement
 {
 public:
-	using FContextAllocator = TFunction<FPCGContext*(const FPCGDataCollection&, TWeakObjectPtr<UPCGComponent>, const UPCGNode*)>;
+	using FContextAllocator = TFunction<FPCGContext*(const FPCGInitializeElementParams&)>;
 
 	FPCGGenericElement(
 		TFunction<bool(FPCGContext*)> InOperation,
-		const FContextAllocator& InContextAllocator = (FContextAllocator)[](const FPCGDataCollection&, TWeakObjectPtr<UPCGComponent>, const UPCGNode*)
+		const FContextAllocator& InContextAllocator = (FContextAllocator)[](const FPCGInitializeElementParams&)
 	{
 		return new FPCGContext();
 	});
@@ -409,7 +409,7 @@ public:
 	FPCGGenericElement(
 		TFunction<bool(FPCGContext*)> InOperation,
 		TFunction<void(FPCGContext*)> InAbortOperation,
-		const FContextAllocator& InContextAllocator = (FContextAllocator)[](const FPCGDataCollection&, TWeakObjectPtr<UPCGComponent>, const UPCGNode*)
+		const FContextAllocator& InContextAllocator = (FContextAllocator)[](const FPCGInitializeElementParams&)
 	{
 		return new FPCGContext();
 	});
@@ -418,7 +418,7 @@ public:
 		TFunction<bool(FPCGContext*)> InOperation,
 		TFunction<void(FPCGContext*)> InAbortOperation,
 		bool bInSupportBasePointDataInput,
-		const FContextAllocator& InContextAllocator = (FContextAllocator)[](const FPCGDataCollection&, TWeakObjectPtr<UPCGComponent>, const UPCGNode*)
+		const FContextAllocator& InContextAllocator = (FContextAllocator)[](const FPCGInitializeElementParams&)
 	{
 		return new FPCGContext();
 	});
@@ -428,12 +428,12 @@ public:
 		TFunction<void(FPCGContext*)> InAbortOperation,
 		bool bInSupportBasePointDataInput,
 		bool bInCanExecuteOnlyOnMainThread,
-		const FContextAllocator& InContextAllocator = (FContextAllocator)[](const FPCGDataCollection&, TWeakObjectPtr<UPCGComponent>, const UPCGNode*)
+		const FContextAllocator& InContextAllocator = (FContextAllocator)[](const FPCGInitializeElementParams&)
 	{
 		return new FPCGContext();
 	});
 	
-	virtual FPCGContext* Initialize(const FPCGDataCollection& InputData, TWeakObjectPtr<UPCGComponent> SourceComponent, const UPCGNode* Node) override;
+	virtual FPCGContext* Initialize(const FPCGInitializeElementParams& InParams) override;
 
 	virtual bool IsCacheable(const UPCGSettings* InSettings) const override { return false; }
 	virtual bool CanExecuteOnlyOnMainThread(FPCGContext* Context) const override { return bCanExecuteOnlyOnMainThread; }
diff --git a/Engine/Plugins/PCG/Source/PCG/Private/Graph/PCGGraphTask.cpp b/Engine/Plugins/PCG/Source/PCG/Private/Graph/PCGGraphTask.cpp
--- a/Engine/Plugins/PCG/Source/PCG/Private/Graph/PCGGraphTask.cpp
+++ b/Engine/Plugins/PCG/Source/PCG/Private/Graph/PCGGraphTask.cpp
@@ -75,7 +75,7 @@ bool FPCGGraphTask::IsApproximatelyEqual(const FPCGGraphTask& Other) const
 
 	return (Inputs == Other.Inputs)
 		&& (Node == Other.Node)
-		&& (SourceComponent == Other.SourceComponent)
+		&& (ExecutionSource == Other.ExecutionSource)
 		&& bElementsMatch
 		&& (Context == Other.Context)
 		&& (NodeId == Other.NodeId)
diff --git a/Engine/Plugins/PCG/Source/PCG/Private/Graph/PCGGraphTask.h b/Engine/Plugins/PCG/Source/PCG/Private/Graph/PCGGraphTask.h
--- a/Engine/Plugins/PCG/Source/PCG/Private/Graph/PCGGraphTask.h
+++ b/Engine/Plugins/PCG/Source/PCG/Private/Graph/PCGGraphTask.h
@@ -135,7 +135,9 @@ struct FPCGGraphTask
 	int32 StackIndex = INDEX_NONE;
 
 	FPCGElementPtr Element; // Added to have tasks that aren't node-bound
-	TWeakObjectPtr<UPCGComponent> SourceComponent = nullptr;
+
+	TWeakInterfacePtr<IPCGGraphExecutionSource> ExecutionSource = nullptr;
+
 	const UPCGNode* Node = nullptr;
 	FPCGContext* Context = nullptr;
 	TSharedPtr<const FPCGStackContext> StackContext;
@@ -162,7 +164,8 @@ struct FPCGGraphTask
 struct FPCGGraphScheduleTask
 {
 	TArray<FPCGGraphTask> Tasks;
-	TWeakObjectPtr<UPCGComponent> SourceComponent = nullptr;
+
+	TWeakInterfacePtr<IPCGGraphExecutionSource> ExecutionSource = nullptr;
 	int32 FirstTaskIndex = 0;
 	int32 LastTaskIndex = 0;
 	bool bHasAbortCallbacks = false;
diff --git a/Engine/Plugins/PCG/Source/PCG/Private/Helpers/PCGBlueprintHelpers.cpp b/Engine/Plugins/PCG/Source/PCG/Private/Helpers/PCGBlueprintHelpers.cpp
--- a/Engine/Plugins/PCG/Source/PCG/Private/Helpers/PCGBlueprintHelpers.cpp
+++ b/Engine/Plugins/PCG/Source/PCG/Private/Helpers/PCGBlueprintHelpers.cpp
@@ -58,30 +58,34 @@ const UPCGSettings* UPCGBlueprintHelpers::GetSettings(FPCGContext& Context)
 
 UPCGData* UPCGBlueprintHelpers::GetActorData(FPCGContext& Context)
 {
-	return Context.SourceComponent.IsValid() ? Context.SourceComponent->GetActorPCGData() : nullptr;
+	UPCGComponent* PCGComponent = GetComponent(Context);
+	return PCGComponent ? PCGComponent->GetActorPCGData() : nullptr;
 }
 
 UPCGData* UPCGBlueprintHelpers::GetInputData(FPCGContext& Context)
 {
-	return Context.SourceComponent.IsValid() ? Context.SourceComponent->GetInputPCGData() : nullptr;
+	UPCGComponent* PCGComponent = GetComponent(Context);
+	return PCGComponent ? PCGComponent->GetInputPCGData() : nullptr;
 }
 
 UPCGComponent* UPCGBlueprintHelpers::GetComponent(FPCGContext& Context)
 {
-	return Context.SourceComponent.Get();
+	return Cast<UPCGComponent>(Context.ExecutionSource.Get());
 }
 
 UPCGComponent* UPCGBlueprintHelpers::GetOriginalComponent(FPCGContext& Context)
 {
-	if (Context.SourceComponent.IsValid() &&
-		Cast<APCGPartitionActor>(Context.SourceComponent->GetOwner()) &&
-		Cast<APCGPartitionActor>(Context.SourceComponent->GetOwner())->GetOriginalComponent(Context.SourceComponent.Get()))
+	UPCGComponent* SourceComponent = Cast<UPCGComponent>(Context.ExecutionSource.Get());
+	APCGPartitionActor* PartitionActor = SourceComponent ? Cast<APCGPartitionActor>(SourceComponent->GetOwner()) : nullptr;
+	UPCGComponent* OriginalComponent = PartitionActor ? PartitionActor->GetOriginalComponent(SourceComponent) : nullptr;
+
+	if (OriginalComponent)
 	{
-		return Cast<APCGPartitionActor>(Context.SourceComponent->GetOwner())->GetOriginalComponent(Context.SourceComponent.Get());
+		return OriginalComponent;
 	}
 	else
 	{
-		return Context.SourceComponent.Get();
+		return SourceComponent;
 	}
 }
 
diff --git a/Engine/Plugins/PCG/Source/PCG/Private/Helpers/PCGDynamicTrackingHelpers.cpp b/Engine/Plugins/PCG/Source/PCG/Private/Helpers/PCGDynamicTrackingHelpers.cpp
--- a/Engine/Plugins/PCG/Source/PCG/Private/Helpers/PCGDynamicTrackingHelpers.cpp
+++ b/Engine/Plugins/PCG/Source/PCG/Private/Helpers/PCGDynamicTrackingHelpers.cpp
@@ -4,17 +4,18 @@
 
 #if WITH_EDITOR
 
-#include "PCGComponent.h"
 #include "PCGContext.h"
+#include "PCGGraphExecutionStateInterface.h"
+#include "PCGSettings.h"
 
 void FPCGDynamicTrackingHelper::EnableAndInitialize(const FPCGContext* InContext, int32 OptionalNumElements)
 {
-	if (InContext && InContext->SourceComponent.IsValid())
+	if (InContext && InContext->ExecutionSource.IsValid())
 	{
 		const UPCGSettings* Settings = InContext->GetOriginalSettings<UPCGSettings>();
 		if (Settings && Settings->CanDynamicallyTrackKeys())
 		{
-			CachedComponent = InContext->SourceComponent;
+			CachedExecutionSource = InContext->ExecutionSource;
 			bDynamicallyTracked = true;
 			DynamicallyTrackedKeysAndCulling.Reserve(OptionalNumElements);
 		}
@@ -31,11 +32,11 @@ void FPCGDynamicTrackingHelper::AddToTracking(FPCGSelectionKey&& InKey, bool bIs
 
 void FPCGDynamicTrackingHelper::Finalize(const FPCGContext* InContext)
 {
-	if (InContext && bDynamicallyTracked && CachedComponent == InContext->SourceComponent)
+	if (InContext && bDynamicallyTracked && CachedExecutionSource == InContext->ExecutionSource)
 	{
-		if (UPCGComponent* SourceComponent = CachedComponent.Get())
+		if (IPCGGraphExecutionSource* ExecutionSource = CachedExecutionSource.Get())
 		{
-			SourceComponent->RegisterDynamicTracking(InContext->GetOriginalSettings<UPCGSettings>(), DynamicallyTrackedKeysAndCulling);
+			ExecutionSource->GetExecutionState().RegisterDynamicTracking(InContext->GetOriginalSettings<UPCGSettings>(), DynamicallyTrackedKeysAndCulling);
 		}
 	}
 }
@@ -44,10 +45,10 @@ void FPCGDynamicTrackingHelper::AddSingleDynamicTrackingKey(FPCGContext* InConte
 {
 	if (InContext)
 	{
-		if (UPCGComponent* SourceComponent = InContext->SourceComponent.Get())
+		if (IPCGGraphExecutionSource* ExecutionSource = InContext->ExecutionSource.Get())
 		{
 			TPair<FPCGSelectionKey, bool> NewPair(std::forward<FPCGSelectionKey>(InKey), bIsCulled);
-			SourceComponent->RegisterDynamicTracking(InContext->GetOriginalSettings<UPCGSettings>(), MakeArrayView(&NewPair, 1));
+			ExecutionSource->GetExecutionState().RegisterDynamicTracking(InContext->GetOriginalSettings<UPCGSettings>(), MakeArrayView(&NewPair, 1));
 		}
 	}
 }
diff --git a/Engine/Plugins/PCG/Source/PCG/Private/Helpers/PCGHelpers.cpp b/Engine/Plugins/PCG/Source/PCG/Private/Helpers/PCGHelpers.cpp
--- a/Engine/Plugins/PCG/Source/PCG/Private/Helpers/PCGHelpers.cpp
+++ b/Engine/Plugins/PCG/Source/PCG/Private/Helpers/PCGHelpers.cpp
@@ -46,41 +46,41 @@ namespace PCGHelpers
 		return ComputeSeed(static_cast<int>(InPosition.X), static_cast<int>(InPosition.Y), static_cast<int>(InPosition.Z));
 	}
 
-	FRandomStream GetRandomStreamFromSeed(int32 InSeed, const UPCGSettings* OptionalSettings, const UPCGComponent* OptionalComponent)
+	FRandomStream GetRandomStreamFromSeed(int32 InSeed, const UPCGSettings* OptionalSettings, const IPCGGraphExecutionSource* OptionalExecutionSource)
 	{
 		int Seed = InSeed;
 
-		if (OptionalSettings && OptionalComponent)
+		if (OptionalSettings && OptionalExecutionSource)
 		{
-			Seed = PCGHelpers::ComputeSeed(Seed, OptionalSettings->Seed, OptionalComponent->Seed);
+			Seed = PCGHelpers::ComputeSeed(Seed, OptionalSettings->Seed, OptionalExecutionSource->GetExecutionState().GetSeed());
 		}
 		else if (OptionalSettings)
 		{
 			Seed = PCGHelpers::ComputeSeed(Seed, OptionalSettings->Seed);
 		}
-		else if (OptionalComponent)
+		else if (OptionalExecutionSource)
 		{
-			Seed = PCGHelpers::ComputeSeed(Seed, OptionalComponent->Seed);
+			Seed = PCGHelpers::ComputeSeed(Seed, OptionalExecutionSource->GetExecutionState().GetSeed());
 		}
 
 		return FRandomStream(Seed);
 	}
 
-	FRandomStream GetRandomStreamFromTwoSeeds(int32 SeedA, int32 SeedB, const UPCGSettings* OptionalSettings, const UPCGComponent* OptionalComponent)
+	FRandomStream GetRandomStreamFromTwoSeeds(int32 SeedA, int32 SeedB, const UPCGSettings* OptionalSettings, const IPCGGraphExecutionSource* OptionalExecutionSource)
 	{
 		int Seed = PCGHelpers::ComputeSeed(SeedA, SeedB);
 
-		if (OptionalSettings && OptionalComponent)
+		if (OptionalSettings && OptionalExecutionSource)
 		{
-			Seed = PCGHelpers::ComputeSeed(Seed, OptionalSettings->Seed, OptionalComponent->Seed);
+			Seed = PCGHelpers::ComputeSeed(Seed, OptionalSettings->Seed, OptionalExecutionSource->GetExecutionState().GetSeed());
 		}
 		else if (OptionalSettings)
 		{
 			Seed = PCGHelpers::ComputeSeed(Seed, OptionalSettings->Seed);
 		}
-		else if (OptionalComponent)
+		else if (OptionalExecutionSource)
 		{
-			Seed = PCGHelpers::ComputeSeed(Seed, OptionalComponent->Seed);
+			Seed = PCGHelpers::ComputeSeed(Seed, OptionalExecutionSource->GetExecutionState().GetSeed());
 		}
 
 		return FRandomStream(Seed);
@@ -551,6 +551,30 @@ namespace PCGHelpers
 		}
 	}
 
+	int32 GetGenerationGridSize(const IPCGGraphExecutionSource* InExecutionSource)
+	{
+		if (const UPCGComponent* PCGComponent = Cast<UPCGComponent>(InExecutionSource))
+		{
+			return PCGComponent->GetGenerationGridSize();
+		}
+		else
+		{
+			return PCGHiGenGrid::UninitializedGridSize();
+		}
+	}
+
+	bool IsRuntimeGeneration(const IPCGGraphExecutionSource* InExecutionSource)
+	{
+		if (const UPCGComponent* PCGComponent = Cast<UPCGComponent>(InExecutionSource))
+		{
+			return PCGComponent->IsManagedByRuntimeGenSystem();
+		}
+		else
+		{
+			return false;
+		}
+	}
+
 #if WITH_EDITOR
 	void GetGeneratedActorsFolderPath(const AActor* InTargetActor, FString& OutFolderPath)
 	{
diff --git a/Engine/Plugins/PCG/Source/PCG/Private/MatchAndSet/PCGMatchAndSetWeighted.cpp b/Engine/Plugins/PCG/Source/PCG/Private/MatchAndSet/PCGMatchAndSetWeighted.cpp
--- a/Engine/Plugins/PCG/Source/PCG/Private/MatchAndSet/PCGMatchAndSetWeighted.cpp
+++ b/Engine/Plugins/PCG/Source/PCG/Private/MatchAndSet/PCGMatchAndSetWeighted.cpp
@@ -6,7 +6,7 @@
 #include "PCGContext.h"
 #include "Data/PCGPointData.h"
 #include "Elements/PCGPointMatchAndSet.h"
-#include "Helpers/PCGBlueprintHelpers.h"
+#include "Helpers/PCGHelpers.h"
 #include "Metadata/Accessors/IPCGAttributeAccessorTpl.h"
 #include "Metadata/Accessors/PCGAttributeAccessorHelpers.h"
 #include "Metadata/Accessors/PCGCustomAccessor.h"
@@ -115,7 +115,7 @@ void UPCGMatchAndSetWeighted::MatchAndSet_Implementation(
 	}
 
 	// TODO: implement async loop?
-	const UPCGComponent* SourceComponent = Context.SourceComponent.Get();
+	const IPCGGraphExecutionSource* ExecutionSource = Context.ExecutionSource.Get();
 	const TArray<FPCGPoint>& InPoints = InPointData->GetPoints();
 
 	TUniquePtr<IPCGAttributeAccessor> SetTargetAccessor = PCGAttributeAccessorHelpers::CreateAccessor(OutPointData, SetTarget);
@@ -133,11 +133,11 @@ void UPCGMatchAndSetWeighted::MatchAndSet_Implementation(
 		return;
 	}
 
-	auto GetValueIndexFromElementIndex = [&InPoints, InSettings, SourceComponent, &CumulativeWeights, TotalWeight](int32 Index) -> int32
+	auto GetValueIndexFromElementIndex = [&InPoints, InSettings, ExecutionSource, &CumulativeWeights, TotalWeight](int32 Index) -> int32
 	{
 		const FPCGPoint& Point = InPoints[Index];
 
-		FRandomStream RandomSource = UPCGBlueprintHelpers::GetRandomStreamFromPoint(Point, InSettings, SourceComponent);
+		FRandomStream RandomSource = PCGHelpers::GetRandomStreamFromSeed(Point.Seed, InSettings, ExecutionSource);
 		int RandomWeightedPick = RandomSource.RandRange(0, TotalWeight - 1);
 
 		int RandomPick = 0;
diff --git a/Engine/Plugins/PCG/Source/PCG/Private/MatchAndSet/PCGMatchAndSetWeightedByCategory.cpp b/Engine/Plugins/PCG/Source/PCG/Private/MatchAndSet/PCGMatchAndSetWeightedByCategory.cpp
--- a/Engine/Plugins/PCG/Source/PCG/Private/MatchAndSet/PCGMatchAndSetWeightedByCategory.cpp
+++ b/Engine/Plugins/PCG/Source/PCG/Private/MatchAndSet/PCGMatchAndSetWeightedByCategory.cpp
@@ -6,7 +6,7 @@
 #include "PCGContext.h"
 #include "Data/PCGPointData.h"
 #include "Elements/PCGPointMatchAndSet.h"
-#include "Helpers/PCGBlueprintHelpers.h"
+#include "Helpers/PCGHelpers.h"
 #include "Math/RandomStream.h"
 #include "Metadata/Accessors/PCGAttributeAccessorHelpers.h"
 #include "Metadata/Accessors/IPCGAttributeAccessorTpl.h"
@@ -200,7 +200,7 @@ void UPCGMatchAndSetWeightedByCategory::MatchAndSet_Implementation(
 	}
 
 	//TODO: implement async loop?
-	const UPCGComponent* SourceComponent = Context.SourceComponent.Get();
+	const IPCGGraphExecutionSource* ExecutionSource = Context.ExecutionSource.Get();
 	const TArray<FPCGPoint>& InPoints = InPointData->GetPoints();
 
 	TUniquePtr<const IPCGAttributeAccessor> InputAccessor = PCGAttributeAccessorHelpers::CreateConstAccessor(InPointData, InputSource);
@@ -227,7 +227,7 @@ void UPCGMatchAndSetWeightedByCategory::MatchAndSet_Implementation(
 		return;
 	}
 
-	auto GetValueIndex = [this, InSettings, SourceComponent, &InPoints](int32 CategoryIndex, int32 PointIndex) -> int32
+	auto GetValueIndex = [this, InSettings, ExecutionSource, &InPoints](int32 CategoryIndex, int32 PointIndex) -> int32
 	{
 		const FPCGMatchAndSetWeightedByCategoryEntryList& Category = Categories[CategoryIndex];
 		const int32 TotalWeight = Category.GetTotalWeight();
@@ -239,7 +239,7 @@ void UPCGMatchAndSetWeightedByCategory::MatchAndSet_Implementation(
 
 		const FPCGPoint& Point = InPoints[PointIndex];
 
-		FRandomStream RandomSource = UPCGBlueprintHelpers::GetRandomStreamFromPoint(Point, InSettings, SourceComponent);
+		FRandomStream RandomSource = PCGHelpers::GetRandomStreamFromSeed(Point.Seed, InSettings, ExecutionSource);
 		int RandomWeightedPick = RandomSource.RandRange(0, TotalWeight - 1);
 
 		int RandomPick = 0;
diff --git a/Engine/Plugins/PCG/Source/PCG/Private/MeshSelectors/PCGMeshSelectorWeighted.cpp b/Engine/Plugins/PCG/Source/PCG/Private/MeshSelectors/PCGMeshSelectorWeighted.cpp
--- a/Engine/Plugins/PCG/Source/PCG/Private/MeshSelectors/PCGMeshSelectorWeighted.cpp
+++ b/Engine/Plugins/PCG/Source/PCG/Private/MeshSelectors/PCGMeshSelectorWeighted.cpp
@@ -219,7 +219,7 @@ bool UPCGMeshSelectorWeighted::SelectMeshInstances(
 			const int64& InMetadataEntry = InRanges.MetadataEntryRange[InPointIndex];
 			const int32& InSeed = InRanges.SeedRange[InPointIndex];
 			
-			FRandomStream RandomSource = PCGHelpers::GetRandomStreamFromSeed(InSeed, Settings, Context.SourceComponent.Get());
+			FRandomStream RandomSource = PCGHelpers::GetRandomStreamFromSeed(InSeed, Settings, Context.ExecutionSource.Get());
 			const int RandomWeightedPick = RandomSource.RandRange(0, TotalWeight - 1);
 
 			int RandomPick = 0;
diff --git a/Engine/Plugins/PCG/Source/PCG/Private/MeshSelectors/PCGMeshSelectorWeightedByCategory.cpp b/Engine/Plugins/PCG/Source/PCG/Private/MeshSelectors/PCGMeshSelectorWeightedByCategory.cpp
--- a/Engine/Plugins/PCG/Source/PCG/Private/MeshSelectors/PCGMeshSelectorWeightedByCategory.cpp
+++ b/Engine/Plugins/PCG/Source/PCG/Private/MeshSelectors/PCGMeshSelectorWeightedByCategory.cpp
@@ -223,7 +223,7 @@ bool UPCGMeshSelectorWeightedByCategory::SelectMeshInstances(
 
 		const int TotalWeight = InstancesAndWeights->CumulativeWeights.Last();
 
-		FRandomStream RandomSource = PCGHelpers::GetRandomStreamFromSeed(InSeed, Settings, Context.SourceComponent.Get());
+		FRandomStream RandomSource = PCGHelpers::GetRandomStreamFromSeed(InSeed, Settings, Context.ExecutionSource.Get());
 		const int RandomWeightedPick = RandomSource.RandRange(0, TotalWeight - 1);
 
 		int RandomPick = 0;
diff --git a/Engine/Plugins/PCG/Source/PCG/Private/PCGCommon.cpp b/Engine/Plugins/PCG/Source/PCG/Private/PCGCommon.cpp
--- a/Engine/Plugins/PCG/Source/PCG/Private/PCGCommon.cpp
+++ b/Engine/Plugins/PCG/Source/PCG/Private/PCGCommon.cpp
@@ -27,6 +27,23 @@ namespace PCGLog
 
 		return InDefaultName;
 	}
+
+	FString GetExecutionSourceName(const IPCGGraphExecutionSource* InExecutionSource, bool bUseLabel, FString InDefaultSource)
+	{
+		const UPCGComponent* PCGComponent = Cast<UPCGComponent>(InExecutionSource);
+		if (PCGComponent)
+		{
+			return GetComponentOwnerName(PCGComponent, bUseLabel, InDefaultSource);
+		}
+		else if (InExecutionSource)
+		{
+			return InExecutionSource->GetExecutionState().GetDebugName();
+		}
+		else
+		{
+			return InDefaultSource;
+		}
+	}
 }
 
 namespace PCGFeatureSwitches
diff --git a/Engine/Plugins/PCG/Source/PCG/Private/PCGComponent.cpp b/Engine/Plugins/PCG/Source/PCG/Private/PCGComponent.cpp
--- a/Engine/Plugins/PCG/Source/PCG/Private/PCGComponent.cpp
+++ b/Engine/Plugins/PCG/Source/PCG/Private/PCGComponent.cpp
@@ -144,6 +144,7 @@ namespace PCGComponent
 UPCGComponent::UPCGComponent(const FObjectInitializer& InObjectInitializer)
 	: Super(InObjectInitializer)
 {
+	ExecutionState.Component = this;
 	GraphInstance = InObjectInitializer.CreateDefaultSubobject<UPCGGraphInstance>(this, TEXT("PCGGraphInstance"));
 	SchedulingPolicyClass = UPCGSchedulingPolicyDistanceAndDirection::StaticClass();
 
@@ -2215,7 +2216,7 @@ void UPCGComponent::RefreshAfterGraphChanged(UPCGGraphInterface* InGraph, EPCGCh
 			CleanupLocal(/*bRemoveComponents=*/true);
 		}
 
-		ClearInspectionData();
+		ExecutionState.GetInspection().ClearInspectionData();
 
 		return;
 	}
@@ -3085,15 +3086,15 @@ UPCGData* UPCGComponent::GetOriginalActorPCGData() const
 	return Data;
 }
 
-UPCGComponent* UPCGComponent::GetOriginalComponent()
+UPCGComponent* UPCGComponent::GetOriginalComponent() const
 {
 	if (!IsLocalComponent())
 	{
-		return this;
+		return const_cast<UPCGComponent*>(this);
 	}
 
-	APCGPartitionActor* PartitionActor = Cast<APCGPartitionActor>(GetOwner());
-	return ensure(PartitionActor) ? PartitionActor->GetOriginalComponent(this) : this;
+	const APCGPartitionActor* PartitionActor = Cast<APCGPartitionActor>(GetOwner());
+	return ensure(PartitionActor) ? PartitionActor->GetOriginalComponent(this) : const_cast<UPCGComponent*>(this);
 }
 
 const UPCGComponent* UPCGComponent::GetConstOriginalComponent() const
@@ -4048,4 +4049,97 @@ void FPCGComponentInstanceData::AddReferencedObjects(FReferenceCollector& Collec
 	Collector.AddReferencedObject(SourceComponent);
 }
 
+UPCGData* FPCGComponentExecutionState::GetSelfData() const
+{
+	check(Component);
+	return Component->GetActorPCGData();
+}
+
+int32 FPCGComponentExecutionState::GetSeed() const
+{
+	check(Component);
+	return Component->Seed;
+}
+
+FString FPCGComponentExecutionState::GetDebugName() const
+{
+	check(Component && Component->GetOwner());
+	return Component->GetOwner()->GetName();
+}
+
+FTransform FPCGComponentExecutionState::GetTransform() const
+{
+	check(Component && Component->GetOwner());
+	return Component->GetOwner()->GetTransform();
+}
+
+UWorld* FPCGComponentExecutionState::GetWorld() const
+{
+	check(Component);
+	return Component->GetWorld();
+}
+
+FBox FPCGComponentExecutionState::GetBounds() const
+{
+	check(Component);
+	return Component->GetGridBounds();
+}
+
+UPCGGraph* FPCGComponentExecutionState::GetGraph() const
+{
+	check(Component);
+	return Component->GetGraph();
+}
+
+UPCGGraphInstance* FPCGComponentExecutionState::GetGraphInstance() const
+{
+	check(Component);
+	return Component->GetGraphInstance();
+}
+
+void FPCGComponentExecutionState::OnGraphExecutionAborted(bool bQuiet, bool bCleanupUnusedResources)
+{
+	check(Component);
+	Component->OnProcessGraphAborted(bQuiet, bCleanupUnusedResources);
+}
+
+void FPCGComponentExecutionState::Cancel()
+{
+	check(Component);
+	Component->CancelGeneration();
+}
+
+#if WITH_EDITOR
+const PCGUtils::FExtraCapture& FPCGComponentExecutionState::GetExtraCapture() const
+{
+	check(Component);
+	return Component->ExtraCapture;
+}
+
+PCGUtils::FExtraCapture& FPCGComponentExecutionState::GetExtraCapture()
+{
+	check(Component);
+	return Component->ExtraCapture;
+}
+
+const FPCGGraphExecutionInspection& FPCGComponentExecutionState::GetInspection() const
+{
+	check(Component);
+	return Component->ExecutionInspection;
+}
+
+FPCGGraphExecutionInspection& FPCGComponentExecutionState::GetInspection()
+{
+	check(Component);
+	return Component->ExecutionInspection;
+}
+
+void FPCGComponentExecutionState::RegisterDynamicTracking(const UPCGSettings* InSettings, const TArrayView<TPair<FPCGSelectionKey, bool>>& InDynamicKeysAndCulling)
+{
+	check(Component);
+	Component->RegisterDynamicTracking(InSettings, InDynamicKeysAndCulling);
+}
+
+#endif // WITH_EDITOR
+
 #undef LOCTEXT_NAMESPACE
diff --git a/Engine/Plugins/PCG/Source/PCG/Private/PCGContext.cpp b/Engine/Plugins/PCG/Source/PCG/Private/PCGContext.cpp
--- a/Engine/Plugins/PCG/Source/PCG/Private/PCGContext.cpp
+++ b/Engine/Plugins/PCG/Source/PCG/Private/PCGContext.cpp
@@ -65,6 +65,17 @@ FPCGContextHandle::~FPCGContextHandle()
 	}
 }
 
+void FPCGContext::InitFromParams(const FPCGInitializeElementParams& InParams)
+{
+	ExecutionSource = InParams.ExecutionSource;
+	InputData = *InParams.InputData;
+	Node = InParams.Node;
+
+PRAGMA_DISABLE_DEPRECATION_WARNINGS
+	SourceComponent = Cast<UPCGComponent>(ExecutionSource.Get());
+PRAGMA_ENABLE_DEPRECATION_WARNINGS
+}
+
 UPCGBasePointData* FPCGContext::NewPointData_AnyThread(FPCGContext* Context)
 {
 	if (CVarPCGEnablePointArrayData.GetValueOnAnyThread())
@@ -110,11 +121,11 @@ int FPCGContext::GetSeed() const
 {
 	if (const UPCGSettings* Settings = GetInputSettings<UPCGSettings>())
 	{
-		return Settings->GetSeed(SourceComponent.Get());
+		return Settings->GetSeed(ExecutionSource.Get());
 	}
-	else if (SourceComponent.IsValid())
+	else if (ExecutionSource.IsValid())
 	{
-		return SourceComponent->Seed;
+		return ExecutionSource->GetExecutionState().GetSeed();
 	}
 	else
 	{
@@ -124,7 +135,12 @@ int FPCGContext::GetSeed() const
 
 FString FPCGContext::GetComponentName() const
 {
-	return PCGLog::GetComponentOwnerName(SourceComponent.Get(), /*bUseLabel=*/false, /*InDefaultName=*/TEXT("Non-PCG Component"));
+	return GetExecutionSourceName();
+}
+
+FString FPCGContext::GetExecutionSourceName() const
+{
+	return PCGLog::GetExecutionSourceName(ExecutionSource.Get(), /*bUseLabel=*/false, /*InDefaultName=*/TEXT("Non-PCG Component"));
 }
 
 AActor* FPCGContext::GetTargetActor(const UPCGSpatialData* InSpatialData) const
@@ -133,9 +149,9 @@ AActor* FPCGContext::GetTargetActor(const UPCGSpatialData* InSpatialData) const
 	{
 		return InSpatialData->TargetActor.Get();
 	}
-	else if (SourceComponent.IsValid() && SourceComponent->GetOwner())
+	else if (UPCGComponent* PCGComponent = Cast<UPCGComponent>(ExecutionSource.Get()))
 	{
-		return SourceComponent->GetOwner();
+		return PCGComponent->GetOwner();
 	}
 	else
 	{
@@ -513,7 +529,7 @@ void FPCGContext::AddStructReferencedObjects(FReferenceCollector& Collector)
 
 FPCGTaskId FPCGContext::ScheduleGraph(const FPCGScheduleGraphParams& InParams)
 {
-	if (GraphExecutor && InParams.SourceComponent)
+	if (GraphExecutor && InParams.ExecutionSource)
 	{
 		return GraphExecutor->ScheduleGraph(InParams);
 	}
diff --git a/Engine/Plugins/PCG/Source/PCG/Private/PCGElement.cpp b/Engine/Plugins/PCG/Source/PCG/Private/PCGElement.cpp
--- a/Engine/Plugins/PCG/Source/PCG/Private/PCGElement.cpp
+++ b/Engine/Plugins/PCG/Source/PCG/Private/PCGElement.cpp
@@ -247,7 +247,7 @@ void IPCGElement::PreExecute(FPCGContext* Context) const
 	check(Context);
 	// Check for early outs (task cancelled + node disabled)
 	// Early out to stop execution
-	if (Context->InputData.bCancelExecution || (!Context->SourceComponent.IsExplicitlyNull() && !Context->SourceComponent.IsValid()))
+	if (Context->InputData.bCancelExecution || (!Context->ExecutionSource.GetWeakObjectPtr().IsExplicitlyNull() && !Context->ExecutionSource.IsValid()))
 	{
 		Context->OutputData.bCancelExecution = true;
 
@@ -346,8 +346,8 @@ void IPCGElement::PreExecutePrimaryLoopElement(FPCGContext* Context, const UPCGS
 		SubCollection.TaggedData.Sort([](const FPCGTaggedData& A, const FPCGTaggedData& B) { return A.OriginalIndex < B.OriginalIndex; });
 		SubCollection.ComputeCrcs(bShouldComputeFullOutputDataCrc);
 
-		FPCGGetFromCacheParams CacheParams = { .Node = Context->Node, .Element = this, .Component = Context->SourceComponent.Get() };
-		GetDependenciesCrc(SubCollection, Settings, Context->SourceComponent.Get(), CacheParams.Crc);
+		FPCGGetFromCacheParams CacheParams = { .Node = Context->Node, .Element = this, .ExecutionSource = Context->ExecutionSource.Get() };
+		GetDependenciesCrc(FPCGGetDependenciesCrcParams(&SubCollection, Settings, Context->ExecutionSource.Get()), CacheParams.Crc);
 
 		FPCGDataCollection SubCollectionOutput;
 		if(Context->GetFromCache(CacheParams, SubCollectionOutput))
@@ -583,10 +583,10 @@ void IPCGElement::PostExecute(FPCGContext* Context) const
 	
 #if WITH_EDITOR
 	// Register the element to the component indicating the element has run and can have dynamic tracked keys.
-	if (Settings && Settings->CanDynamicallyTrackKeys() && Context->SourceComponent.IsValid())
+	if (Settings && Settings->CanDynamicallyTrackKeys() && Context->ExecutionSource.IsValid())
 	{
 		TRACE_CPUPROFILER_EVENT_SCOPE(IPCGElement::PostExecute::RegisterDynamicTracking);
-		Context->SourceComponent->RegisterDynamicTracking(Settings, {});
+		Context->ExecutionSource->GetExecutionState().RegisterDynamicTracking(Settings, {});
 	}
 #endif // WITH_EDITOR
 
@@ -652,7 +652,7 @@ void IPCGElement::PostExecutePrimaryLoopElement(FPCGContext* Context, const UPCG
 				SubCollection.ComputeCrcs(bShouldComputeFullOutputDataCrc);
 
 				FPCGCrc DependenciesCrc;
-				GetDependenciesCrc(SubCollection, Settings, Context->SourceComponent.Get(), DependenciesCrc);
+				GetDependenciesCrc(FPCGGetDependenciesCrcParams(&SubCollection, Settings, Context->ExecutionSource.Get()), DependenciesCrc);
 
 				SubCollectionOutput.ComputeCrcs(bShouldComputeFullOutputDataCrc);
 
@@ -818,8 +818,9 @@ bool IPCGElement::DebugDisplay(FPCGContext* Context) const
 
 	// If the graph is not being inspected, or the current component is being inspected, then we know we should display
 	// debug, if not then we do further checks.
-	UPCGGraph* Graph = Context->SourceComponent.Get() ? Context->SourceComponent->GetGraph() : nullptr;
-	if (Graph && Graph->IsInspecting() && !Context->SourceComponent->IsInspecting() && Graph->DebugFlagAppliesToIndividualComponents())
+	UPCGComponent* SourceComponent = Cast<UPCGComponent>(Context->ExecutionSource.Get());
+	UPCGGraph* Graph = SourceComponent ? SourceComponent->GetGraph() : nullptr;
+	if (Graph && Graph->IsInspecting() && !SourceComponent->GetExecutionState().GetInspection().IsInspecting() && Graph->DebugFlagAppliesToIndividualComponents())
 	{
 		// If we're no doing HiGen, or if the current component is not a local component (and therefore will not have children),
 		// then do not display debug.
@@ -830,14 +831,14 @@ bool IPCGElement::DebugDisplay(FPCGContext* Context) const
 
 		// If a child of this component is being inspected (a local component on smaller grid and overlapping) then we still show debug,
 		// because this data is available to that child for use.
-		if (UPCGSubsystem* Subsystem = UPCGSubsystem::GetInstance(Context->SourceComponent->GetWorld()))
+		if (UPCGSubsystem* Subsystem = UPCGSubsystem::GetInstance(SourceComponent->GetExecutionState().GetWorld()))
 		{
-			const uint32 ThisGenerationGridSize = Context->SourceComponent->GetGenerationGridSize();
+			const uint32 ThisGenerationGridSize = SourceComponent->GetGenerationGridSize();
 
 			bool bFoundInspectedChildComponent = false;
-			Subsystem->ForAllOverlappingComponentsInHierarchy(Context->SourceComponent.Get(), [ThisGenerationGridSize, &bFoundInspectedChildComponent](UPCGComponent* InLocalComponent)
+			Subsystem->ForAllOverlappingComponentsInHierarchy(SourceComponent, [ThisGenerationGridSize, &bFoundInspectedChildComponent](UPCGComponent* InLocalComponent)
 			{
-				if (InLocalComponent->GetGenerationGridSize() < ThisGenerationGridSize && InLocalComponent->IsInspecting())
+				if (InLocalComponent->GetGenerationGridSize() < ThisGenerationGridSize && InLocalComponent->GetExecutionState().GetInspection().IsInspecting())
 				{
 					bFoundInspectedChildComponent = true;
 				}
@@ -1013,10 +1014,30 @@ void IPCGElement::CleanupAndValidateOutput(FPCGContext* Context) const
 
 FPCGContext* IPCGElement::Initialize(const FPCGDataCollection& InputData, TWeakObjectPtr<UPCGComponent> SourceComponent, const UPCGNode* Node)
 {
-	FPCGContext* Context = CreateContext();
-	Context->InputData = InputData;
-	Context->SourceComponent = SourceComponent;
-	Context->Node = Node;
+	return nullptr;
+}
+
+FPCGContext* IPCGElement::Initialize(const FPCGInitializeElementParams& InParams)
+{
+	check(InParams.InputData);
+	// For backward compatibility (call the old signature)
+PRAGMA_DISABLE_DEPRECATION_WARNINGS
+	FPCGContext* Context = Initialize(*InParams.InputData, Cast<UPCGComponent>(InParams.ExecutionSource.GetObject()), InParams.Node);
+PRAGMA_ENABLE_DEPRECATION_WARNINGS
+	
+	if (!Context)
+	{
+		Context = CreateContext();
+		Context->InitFromParams(InParams);
+	}
+	else
+	{
+		Context->ExecutionSource = InParams.ExecutionSource;
+
+		const UPCGSettings* Settings = InParams.Node ? InParams.Node->GetSettings() : nullptr;
+		const FString SettingsName = Settings ? Settings->GetName() : TEXT("Unknown");
+		UE_LOG(LogPCG, Warning, TEXT("This node '%s' implements a deprecated version of Initialize. Please implement version with FPCGInitializeElementParams parameter."), *SettingsName);
+	}
 
 	return Context;
 }
@@ -1052,35 +1073,55 @@ bool IPCGElement::ShouldVerifyIfOutputsAreUsedMultipleTimes(const UPCGSettings*
 
 void IPCGElement::GetDependenciesCrc(const FPCGDataCollection& InInput, const UPCGSettings* InSettings, UPCGComponent* InComponent, FPCGCrc& OutCrc) const
 {
-	TRACE_CPUPROFILER_EVENT_SCOPE_TEXT(*FString::Printf(TEXT("IPCGElement::GetDependenciesCrc (%s)"), InSettings ? *InSettings->GetFName().ToString() : TEXT("")));
-
-	// Start from a random prime.
-	FPCGCrc Crc(1000003);
+	// do nothing
+}
 
-	// The cached data CRCs are computed in FPCGGraphExecutor::BuildTaskInput and incorporate data CRC, tags, output pin label and input pin label.
-	for (const FPCGCrc& DataCrc : InInput.DataCrcs)
+void IPCGElement::GetDependenciesCrc(const FPCGGetDependenciesCrcParams& InParams, FPCGCrc& OutCrc) const
+{
+	TRACE_CPUPROFILER_EVENT_SCOPE_TEXT(*FString::Printf(TEXT("IPCGElement::GetDependenciesCrc (%s)"), InParams.Settings ? *InParams.Settings->GetName() : TEXT("")));
+	
+	check(InParams.InputData);
+	FPCGCrc CopyCrc(OutCrc);
+PRAGMA_DISABLE_DEPRECATION_WARNINGS
+	GetDependenciesCrc(*InParams.InputData, InParams.Settings, Cast<UPCGComponent>(InParams.ExecutionSource), OutCrc);
+PRAGMA_ENABLE_DEPRECATION_WARNINGS
+
+	// Call to deprecated method didn't yield a different Crc so we calculate it here
+	if (!OutCrc.IsValid() || CopyCrc == OutCrc)
 	{
-		Crc.Combine(DataCrc);
-	}
+		// Start from a random prime.
+		FPCGCrc Crc(1000003);
 
-	if (InSettings)
-	{
-		FPCGCrc SettingsCrc = InSettings->GetSettingsCrc();
-		if (ensure(SettingsCrc.IsValid()))
+		// The cached data CRCs are computed in FPCGGraphExecutor::BuildTaskInput and incorporate data CRC, tags, output pin label and input pin label.
+		for (const FPCGCrc& DataCrc : InParams.InputData->DataCrcs)
 		{
-			Crc.Combine(SettingsCrc);
+			Crc.Combine(DataCrc);
+		}
+
+		if (InParams.Settings)
+		{
+			FPCGCrc SettingsCrc = InParams.Settings->GetSettingsCrc();
+			if (ensure(SettingsCrc.IsValid()))
+			{
+				Crc.Combine(SettingsCrc);
+			}
+		}
+
+		if (InParams.ExecutionSource && (!InParams.Settings || InParams.Settings->UseSeed()))
+		{
+			Crc.Combine(InParams.ExecutionSource->GetExecutionState().GetSeed());
 		}
-	}
 
-	if (InComponent && (!InSettings || InSettings->UseSeed()))
+		OutCrc = Crc;
+	}
+	else
 	{
-		Crc.Combine(InComponent->Seed);
+		const FString SettingsName = InParams.Settings ? InParams.Settings->GetName() : TEXT("Unknown");
+		UE_LOG(LogPCG, Warning, TEXT("This node '%s' implements a deprecated version of GetDependenciesCrc. Please implement version with FPCGGetDependenciesCrcParams parameter."), *SettingsName);
 	}
-
-	OutCrc = Crc;
 }
 
-EPCGCachingStatus IPCGElement::RetrieveResultsFromCache(IPCGGraphCache* Cache, const UPCGNode* Node, const FPCGDataCollection& Input, UPCGComponent* Component, FPCGDataCollection& Output, FPCGCrc* OutCrc) const
+EPCGCachingStatus IPCGElement::RetrieveResultsFromCache(IPCGGraphCache* Cache, const UPCGNode* Node, const FPCGDataCollection& Input, IPCGGraphExecutionSource* ExecutionSource, FPCGDataCollection& Output, FPCGCrc* OutCrc) const
 {
 	if (!Cache)
 	{
@@ -1091,11 +1132,11 @@ EPCGCachingStatus IPCGElement::RetrieveResultsFromCache(IPCGGraphCache* Cache, c
 	const UPCGSettings* Settings = SettingsInterface ? SettingsInterface->GetSettings() : nullptr;
 	const bool bCacheable = IsCacheableInstance(SettingsInterface);
 
-	FPCGGetFromCacheParams Params = { .Node = Node, .Element = this, .Component = Component };
+	FPCGGetFromCacheParams Params = { .Node = Node, .Element = this, .ExecutionSource = ExecutionSource };
 
 	if (Settings && bCacheable)
 	{
-		GetDependenciesCrc(Input, Settings, Component, Params.Crc);
+		GetDependenciesCrc(FPCGGetDependenciesCrcParams(&Input, Settings, ExecutionSource), Params.Crc);
 
 		if (OutCrc)
 		{
diff --git a/Engine/Plugins/PCG/Source/PCG/Private/PCGGraphExecutionStateInterface.cpp b/Engine/Plugins/PCG/Source/PCG/Private/PCGGraphExecutionStateInterface.cpp
new file mode 100644
--- /dev/null
+++ b/Engine/Plugins/PCG/Source/PCG/Private/PCGGraphExecutionStateInterface.cpp
@@ -0,0 +1,8 @@
+// Copyright Epic Games, Inc. All Rights Reserved.
+
+#include "PCGGraphExecutionStateInterface.h"
+
+UPCGGraphExecutionSource::UPCGGraphExecutionSource(const FObjectInitializer& ObjectInitializer)
+	: Super(ObjectInitializer)
+{
+}
\ No newline at end of file
diff --git a/Engine/Plugins/PCG/Source/PCG/Private/PCGSettings.cpp b/Engine/Plugins/PCG/Source/PCG/Private/PCGSettings.cpp
--- a/Engine/Plugins/PCG/Source/PCG/Private/PCGSettings.cpp
+++ b/Engine/Plugins/PCG/Source/PCG/Private/PCGSettings.cpp
@@ -606,9 +606,9 @@ UPCGNode* UPCGSettings::CreateNode() const
 	return NewObject<UPCGNode>();
 }
 
-int UPCGSettings::GetSeed(const UPCGComponent* InSourceComponent) const
+int UPCGSettings::GetSeed(const IPCGGraphExecutionSource* InExecutionSource) const
 {
-	return !UseSeed() ? PCGValueConstants::DefaultSeed : (InSourceComponent ? PCGHelpers::ComputeSeed(Seed, InSourceComponent->Seed) : Seed);
+	return !UseSeed() ? PCGValueConstants::DefaultSeed : (InExecutionSource ? PCGHelpers::ComputeSeed(Seed, InExecutionSource->GetExecutionState().GetSeed()) : Seed);
 }
 
 #if WITH_EDITOR
diff --git a/Engine/Plugins/PCG/Source/PCG/Private/PCGSubgraph.cpp b/Engine/Plugins/PCG/Source/PCG/Private/PCGSubgraph.cpp
--- a/Engine/Plugins/PCG/Source/PCG/Private/PCGSubgraph.cpp
+++ b/Engine/Plugins/PCG/Source/PCG/Private/PCGSubgraph.cpp
@@ -613,13 +613,10 @@ void FPCGSubgraphContext::AddExtraStructReferencedObjects(FReferenceCollector& C
 	Collector.AddReferencedObjects(ReferencedObjects);
 }
 
-FPCGContext* FPCGSubgraphElement::Initialize(const FPCGDataCollection& InputData, TWeakObjectPtr<UPCGComponent> SourceComponent, const UPCGNode* Node)
+FPCGContext* FPCGSubgraphElement::Initialize(const FPCGInitializeElementParams& InParams)
 {
 	FPCGSubgraphContext* Context = new FPCGSubgraphContext();
-	Context->InputData = InputData;
-	Context->SourceComponent = SourceComponent;
-	Context->Node = Node;
-
+	Context->InitFromParams(InParams);
 	Context->InitializeUserParametersStruct();
 
 	return Context;
@@ -745,10 +742,10 @@ bool FPCGSubgraphElement::ExecuteInternal(FPCGContext* InContext) const
 		{
 			bIsRecursive = (bIsDynamic && Context->Stack->HasObject(Subgraph)) || Subgraph->Contains(Context->Stack->GetGraphForCurrentFrame());
 		}
-		else if(Context->SourceComponent.Get())
+		else if(Context->ExecutionSource.Get())
 		{
-			const UPCGGraph* CurrentComponentGraph = Context->SourceComponent->GetGraph();
-			bIsRecursive = (Subgraph == CurrentComponentGraph || Subgraph->Contains(CurrentComponentGraph));
+			const UPCGGraph* CurrentSourceGraph = Context->ExecutionSource->GetExecutionState().GetGraph();
+			bIsRecursive = (Subgraph == CurrentSourceGraph || Subgraph->Contains(CurrentSourceGraph));
 		}
 	}
 
@@ -784,7 +781,7 @@ bool FPCGSubgraphElement::ExecuteInternal(FPCGContext* InContext) const
 				// Higen is not allowed in dynamic subgraphs, entire subgraph is executed on the same grid as this subgraph node.
 				FPCGTaskId SubgraphTaskId = Context->ScheduleGraph(FPCGScheduleGraphParams(
 					Subgraph,
-					Context->SourceComponent.Get(),
+					Context->ExecutionSource.Get(),
 					MakeShared<FPCGInputForwardingElement>(PreSubgraphInputData),
 					MakeShared<FPCGInputForwardingElement>(SubgraphInputData),
 					/*Dependencies=*/{},
@@ -824,7 +821,7 @@ bool FPCGSubgraphElement::ExecuteInternal(FPCGContext* InContext) const
 									ContextPtr->SubgraphTaskIds.Reset();
 								}
 							},
-							Context->SourceComponent.Get(),
+							Context->ExecutionSource.Get(),
 							Context->SubgraphTaskIds));
 					}
 
diff --git a/Engine/Plugins/PCG/Source/PCG/Private/PCGSubsystem.cpp b/Engine/Plugins/PCG/Source/PCG/Private/PCGSubsystem.cpp
--- a/Engine/Plugins/PCG/Source/PCG/Private/PCGSubsystem.cpp
+++ b/Engine/Plugins/PCG/Source/PCG/Private/PCGSubsystem.cpp
@@ -811,7 +811,7 @@ FPCGTaskId UPCGSubsystem::ScheduleGraph(
 
 FPCGTaskId UPCGSubsystem::ScheduleGraph(const FPCGScheduleGraphParams& InParams)
 {
-	if (InParams.SourceComponent)
+	if (InParams.ExecutionSource)
 	{
 		return GraphExecutor->ScheduleGraph(InParams);
 	}
@@ -889,12 +889,12 @@ void UPCGSubsystem::CancelGeneration(UPCGComponent* Component, bool bCleanupUnus
 		ActorAndComponentMapping.DispatchToRegisteredLocalComponents(Component, LocalCancel);
 	}
 
-	TArray<UPCGComponent*> CancelledComponents = GraphExecutor->Cancel(Component);
-	for (UPCGComponent* CancelledComponent : CancelledComponents)
+	TArray<IPCGGraphExecutionSource*> CancelledExecutionSources = GraphExecutor->Cancel(Component);
+	for (IPCGGraphExecutionSource* CancelledExecutionSource : CancelledExecutionSources)
 	{
-		if (CancelledComponent)
+		if (CancelledExecutionSource)
 		{
-			CancelledComponent->OnProcessGraphAborted(/*bQuiet=*/true, bCleanupUnusedResources);
+			CancelledExecutionSource->GetExecutionState().OnGraphExecutionAborted(/*bQuiet=*/true, bCleanupUnusedResources);
 		}
 	}	
 }
@@ -908,12 +908,12 @@ void UPCGSubsystem::CancelGeneration(UPCGGraph* Graph)
 		return;
 	}
 
-	TArray<UPCGComponent*> CancelledComponents = GraphExecutor->Cancel(Graph);
-	for (UPCGComponent* CancelledComponent : CancelledComponents)
+	TArray<IPCGGraphExecutionSource*> CancelledExecutionSources = GraphExecutor->Cancel(Graph);
+	for (IPCGGraphExecutionSource* CancelledExecutionSource : CancelledExecutionSources)
 	{
-		if (ensure(CancelledComponent))
+		if (ensure(CancelledExecutionSource))
 		{
-			CancelledComponent->OnProcessGraphAborted(/*bQuiet=*/true);
+			CancelledExecutionSource->GetExecutionState().OnGraphExecutionAborted(/*bQuiet=*/true);
 		}
 	}
 }
@@ -922,12 +922,12 @@ void UPCGSubsystem::CancelAllGeneration()
 {
 	check(GraphExecutor);
 
-	TArray<UPCGComponent*> CancelledComponents = GraphExecutor->CancelAll();
-	for (UPCGComponent* CancelledComponent : CancelledComponents)
+	TArray<IPCGGraphExecutionSource*> CancelledExecutionSources = GraphExecutor->CancelAll();
+	for (IPCGGraphExecutionSource* CancelledExecutionSource : CancelledExecutionSources)
 	{
-		if (ensure(CancelledComponent))
+		if (ensure(CancelledExecutionSource))
 		{
-			CancelledComponent->OnProcessGraphAborted(/*bQuiet=*/true);
+			CancelledExecutionSource->GetExecutionState().OnGraphExecutionAborted(/*bQuiet=*/true);
 		}
 	}
 }
diff --git a/Engine/Plugins/PCG/Source/PCG/Private/Tests/PCGTestsCommon.cpp b/Engine/Plugins/PCG/Source/PCG/Private/Tests/PCGTestsCommon.cpp
--- a/Engine/Plugins/PCG/Source/PCG/Private/Tests/PCGTestsCommon.cpp
+++ b/Engine/Plugins/PCG/Source/PCG/Private/Tests/PCGTestsCommon.cpp
@@ -111,7 +111,7 @@ namespace PCGTestsCommon
 	TUniquePtr<FPCGContext> InitializeTestContext(IPCGElement* InElement, const FPCGDataCollection& InputData, UPCGComponent* InSourceComponent, const UPCGNode* InNode)
 	{
 		check(InElement);
-		TUniquePtr<FPCGContext> Context{ InElement->Initialize(InputData, InSourceComponent, InNode) };
+		TUniquePtr<FPCGContext> Context{ InElement->Initialize(FPCGInitializeElementParams(&InputData, InSourceComponent, InNode)) };
 		Context->InitializeSettings();
 		Context->AsyncState.NumAvailableTasks = 1;
 		return Context;
diff --git a/Engine/Plugins/PCG/Source/PCG/Private/Utils/PCGExtraCapture.cpp b/Engine/Plugins/PCG/Source/PCG/Private/Utils/PCGExtraCapture.cpp
--- a/Engine/Plugins/PCG/Source/PCG/Private/Utils/PCGExtraCapture.cpp
+++ b/Engine/Plugins/PCG/Source/PCG/Private/Utils/PCGExtraCapture.cpp
@@ -145,12 +145,12 @@ namespace PCGUtils
 	}
 }
 
-PCGUtils::FCallTreeInfo PCGUtils::FExtraCapture::CalculateCallTreeInfo(const UPCGComponent* Component, const FPCGStack& RootStack) const
+PCGUtils::FCallTreeInfo PCGUtils::FExtraCapture::CalculateCallTreeInfo(const IPCGGraphExecutionSource* ExecutionSource, const FPCGStack& RootStack) const
 {
 	FCallTreeInfo RootInfo;
 
 	// Basically, what we want is - visit all entries in the "NodeToStacksInWhichNodeExecuted" and build our information from there.
-	TMap<TObjectKey<const UPCGNode>, TSet<FPCGGraphExecutionInspection::FNodeExecutedNotificationData>> NodeToStacksInWhichNodeExecuted = Component->GetExecutionInspection().GetExecutedNodeStacks();
+	TMap<TObjectKey<const UPCGNode>, TSet<FPCGGraphExecutionInspection::FNodeExecutedNotificationData>> NodeToStacksInWhichNodeExecuted = ExecutionSource->GetExecutionState().GetInspection().GetExecutedNodeStacks();
 
 	TArray<const UPCGNode*> NodePath;
 	TArray<int32> NodePathLoop;
@@ -275,9 +275,9 @@ PCGUtils::FScopedCall::FScopedCall(const IPCGElement& InOwner, FPCGContext* InCo
 
 PCGUtils::FScopedCall::~FScopedCall()
 {
-	if (Context && Context->SourceComponent.IsValid())
+	if (Context && Context->ExecutionSource.IsValid())
 	{
-		Context->SourceComponent->ExtraCapture.Update(*this);
+		Context->ExecutionSource->GetExecutionState().GetExtraCapture().Update(*this);
 	}
 	OutputDevice.ScopedCall = nullptr;
 }
diff --git a/Engine/Plugins/PCG/Source/PCG/Private/Utils/PCGGraphExecutionLogging.cpp b/Engine/Plugins/PCG/Source/PCG/Private/Utils/PCGGraphExecutionLogging.cpp
--- a/Engine/Plugins/PCG/Source/PCG/Private/Utils/PCGGraphExecutionLogging.cpp
+++ b/Engine/Plugins/PCG/Source/PCG/Private/Utils/PCGGraphExecutionLogging.cpp
@@ -133,7 +133,7 @@ namespace PCGGraphExecutionLogging
 #endif // !(UE_BUILD_SHIPPING || UE_BUILD_TEST) || USE_LOGGING_IN_SHIPPING
 	}
 
-	void LogGraphSchedule(const UPCGComponent* SourceComponent, const UPCGGraph* InScheduledGraph)
+	void LogGraphSchedule(const IPCGGraphExecutionSource* InExecutionSource, const UPCGGraph* InScheduledGraph)
 	{
 #if !(UE_BUILD_SHIPPING || UE_BUILD_TEST) || USE_LOGGING_IN_SHIPPING
 		if (!LogEnabled() && !IsRunningCommandlet())
@@ -142,13 +142,13 @@ namespace PCGGraphExecutionLogging
 		}
 
 		UE_LOG(LogPCG, Display, TEXT("[%s/%s] --- SCHEDULE GRAPH %s ---"),
-			*PCGLog::GetComponentOwnerName(SourceComponent),
-			(SourceComponent && SourceComponent->GetGraph()) ? *SourceComponent->GetGraph()->GetName() : TEXT("MISSINGGRAPH"),
+			*PCGLog::GetExecutionSourceName(InExecutionSource),
+			(InExecutionSource && InExecutionSource->GetExecutionState().GetGraph()) ? *InExecutionSource->GetExecutionState().GetGraph()->GetName() : TEXT("MISSINGGRAPH"),
 			InScheduledGraph ? *InScheduledGraph->GetName() : TEXT("MISSINGGRAPH)"));
 #endif // !(UE_BUILD_SHIPPING || UE_BUILD_TEST) || USE_LOGGING_IN_SHIPPING
 	}
 
-	void LogGraphScheduleDependency(const UPCGComponent* InComponent, const FPCGStack* InFromStack)
+	void LogGraphScheduleDependency(const IPCGGraphExecutionSource* InExecutionSource, const FPCGStack* InFromStack)
 	{
 #if !(UE_BUILD_SHIPPING || UE_BUILD_TEST) || USE_LOGGING_IN_SHIPPING
 		if (!LogEnabled() && !IsRunningCommandlet())
@@ -163,13 +163,13 @@ namespace PCGGraphExecutionLogging
 		}
 
 		UE_LOG(LogPCG, Display, TEXT("[%s/%s] --- SCHEDULE GRAPH FOR DEPENDENCY, from stack: %s"),
-			*PCGLog::GetComponentOwnerName(InComponent),
-			(InComponent && InComponent->GetGraph()) ? *InComponent->GetGraph()->GetName() : TEXT("MISSINGGRAPH"),
+			*PCGLog::GetExecutionSourceName(InExecutionSource),
+			(InExecutionSource && InExecutionSource->GetExecutionState().GetGraph()) ? *InExecutionSource->GetExecutionState().GetGraph()->GetName() : TEXT("MISSINGGRAPH"),
 			*FromStackString);
 #endif // !(UE_BUILD_SHIPPING || UE_BUILD_TEST) || USE_LOGGING_IN_SHIPPING
 	}
 
-	void LogGraphScheduleDependencyFailed(const UPCGComponent* InComponent, const FPCGStack* InFromStack)
+	void LogGraphScheduleDependencyFailed(const IPCGGraphExecutionSource* InExecutionSource, const FPCGStack* InFromStack)
 	{
 #if !(UE_BUILD_SHIPPING || UE_BUILD_TEST) || USE_LOGGING_IN_SHIPPING
 		if (!LogEnabled() && !IsRunningCommandlet())
@@ -184,8 +184,8 @@ namespace PCGGraphExecutionLogging
 		}
 
 		UE_LOG(LogPCG, Warning, TEXT("[%s/%s] Failed to schedule dependency, from stack: %s"),
-			*PCGLog::GetComponentOwnerName(InComponent),
-			(InComponent && InComponent->GetGraph()) ? *InComponent->GetGraph()->GetName() : TEXT("MISSINGGRAPH"),
+			*PCGLog::GetExecutionSourceName(InExecutionSource),
+			(InExecutionSource && InExecutionSource->GetExecutionState().GetGraph()) ? *InExecutionSource->GetExecutionState().GetGraph()->GetName() : TEXT("MISSINGGRAPH"),
 			*FromStackString);
 #endif // !(UE_BUILD_SHIPPING || UE_BUILD_TEST) || USE_LOGGING_IN_SHIPPING
 	}
@@ -204,7 +204,7 @@ namespace PCGGraphExecutionLogging
 #endif // !(UE_BUILD_SHIPPING || UE_BUILD_TEST) || USE_LOGGING_IN_SHIPPING
 	}
 
-	void LogPostProcessGraph(const UPCGComponent* InSourceComponent)
+	void LogPostProcessGraph(const IPCGGraphExecutionSource* InExecutionSource)
 	{
 #if !(UE_BUILD_SHIPPING || UE_BUILD_TEST) || USE_LOGGING_IN_SHIPPING
 		if (!LogEnabled() && !IsRunningCommandlet())
@@ -212,13 +212,13 @@ namespace PCGGraphExecutionLogging
 			return;
 		}
 
-		UE_LOG(LogPCG, Display, TEXT("[%s/%s] UPCGComponent::PostProcessGraph"),
-			*PCGLog::GetComponentOwnerName(InSourceComponent),
-			(InSourceComponent && InSourceComponent->GetGraph()) ? *InSourceComponent->GetGraph()->GetName() : TEXT("MISSINGGRAPH"));
+		UE_LOG(LogPCG, Display, TEXT("[%s/%s] IPCGGraphExecutionSource::PostProcessGraph"),
+			*PCGLog::GetExecutionSourceName(InExecutionSource),
+			(InExecutionSource && InExecutionSource->GetExecutionState().GetGraph()) ? *InExecutionSource->GetExecutionState().GetGraph()->GetName() : TEXT("MISSINGGRAPH"));
 #endif // !(UE_BUILD_SHIPPING || UE_BUILD_TEST) || USE_LOGGING_IN_SHIPPING
 	}
 
-	void LogComponentCancellation(const TSet<UPCGComponent*>& CancelledComponents)
+	void LogExecutionSourceCancellation(const TSet<IPCGGraphExecutionSource*>& CancelledExecutionSources)
 	{
 #if !(UE_BUILD_SHIPPING || UE_BUILD_TEST) || USE_LOGGING_IN_SHIPPING
 		if (!LogEnabled() && !IsRunningCommandlet())
@@ -226,16 +226,16 @@ namespace PCGGraphExecutionLogging
 			return;
 		}
 
-		for (const UPCGComponent* Component : CancelledComponents)
+		for (const IPCGGraphExecutionSource* ExecutionSource : CancelledExecutionSources)
 		{
-			UE_LOG(LogPCG, Display, TEXT("[%s/%s] Component cancelled"),
-				*PCGLog::GetComponentOwnerName(Component),
-				(Component && Component->GetGraph()) ? *Component->GetGraph()->GetName() : TEXT("MISSINGGRAPH"));
+			UE_LOG(LogPCG, Display, TEXT("[%s/%s] ExecutionSource cancelled"),
+				*PCGLog::GetExecutionSourceName(ExecutionSource),
+				(ExecutionSource && ExecutionSource->GetExecutionState().GetGraph()) ? *ExecutionSource->GetExecutionState().GetGraph()->GetName() : TEXT("MISSINGGRAPH"));
 		}
 #endif // !(UE_BUILD_SHIPPING || UE_BUILD_TEST) || USE_LOGGING_IN_SHIPPING
 	}
 
-	void LogChangeOriginIgnoredForComponent(const UObject* InObject, const UPCGComponent* InComponent)
+	void LogChangeOriginIgnoredForComponent(const UObject* InObject, const IPCGGraphExecutionSource* InExecutionSource)
 	{
 #if !(UE_BUILD_SHIPPING || UE_BUILD_TEST) || USE_LOGGING_IN_SHIPPING
 		if (!LogEnabled())
@@ -244,8 +244,8 @@ namespace PCGGraphExecutionLogging
 		}
 
 		UE_LOG(LogPCG, Log, TEXT("[%s/%s] Change origin ignored: '%s'"),
-			*PCGLog::GetComponentOwnerName(InComponent),
-			(InComponent && InComponent->GetGraph()) ? *InComponent->GetGraph()->GetName() : TEXT("MISSINGGRAPH"),
+			*PCGLog::GetExecutionSourceName(InExecutionSource),
+			(InExecutionSource && InExecutionSource->GetExecutionState().GetGraph()) ? *InExecutionSource->GetExecutionState().GetGraph()->GetName() : TEXT("MISSINGGRAPH"),
 			InObject ? *InObject->GetName() : TEXT("MISSINGOBJECT"));
 #endif // !(UE_BUILD_SHIPPING || UE_BUILD_TEST) || USE_LOGGING_IN_SHIPPING
 	}
@@ -283,14 +283,20 @@ namespace PCGGraphExecutionLogging
 	void LogTaskExecute(const FPCGGraphTask& Task)
 	{
 #if !(UE_BUILD_SHIPPING || UE_BUILD_TEST) || USE_LOGGING_IN_SHIPPING
-		if (!LogEnabled() || !Task.SourceComponent.Get())
+		if (!LogEnabled())
+		{
+			return;
+		}
+
+		IPCGGraphExecutionSource* ExecutionSource = Task.ExecutionSource.Get();
+		if (!ExecutionSource)
 		{
 			return;
 		}
 
 		UE_LOG(LogPCG, Log, TEXT("         [%s/%s] %s\t\tEXECUTE"),
-			*PCGLog::GetComponentOwnerName(Task.SourceComponent.Get()),
-			Task.SourceComponent->GetGraph() ? *Task.SourceComponent->GetGraph()->GetName() : TEXT("MISSINGGRAPH"),
+			*PCGLog::GetExecutionSourceName(ExecutionSource),
+			ExecutionSource->GetExecutionState().GetGraph() ? *ExecutionSource->GetExecutionState().GetGraph()->GetName() : TEXT("MISSINGGRAPH"),
 			*FString::Printf(TEXT("%" UINT64_FMT "'%s'"), Task.NodeId, Task.Node ? *Task.Node->GetNodeTitle(EPCGNodeTitleType::ListView).ToString() : TEXT("")));
 #endif // !(UE_BUILD_SHIPPING || UE_BUILD_TEST) || USE_LOGGING_IN_SHIPPING
 	}
@@ -298,14 +304,20 @@ namespace PCGGraphExecutionLogging
 	void LogTaskExecuteCachingDisabled(const FPCGGraphTask& Task)
 	{
 #if !(UE_BUILD_SHIPPING || UE_BUILD_TEST) || USE_LOGGING_IN_SHIPPING
-		if (!LogEnabled() || !Task.SourceComponent.Get())
+		if (!LogEnabled())
+		{
+			return;
+		}
+
+		IPCGGraphExecutionSource* ExecutionSource = Task.ExecutionSource.Get();
+		if (!ExecutionSource)
 		{
 			return;
 		}
 
 		UE_LOG(LogPCG, Warning, TEXT("[%s/%s] %s\t\tCACHING DISABLED"),
-			*PCGLog::GetComponentOwnerName(Task.SourceComponent.Get()),
-			Task.SourceComponent->GetGraph() ? *Task.SourceComponent->GetGraph()->GetName() : TEXT("MISSINGGRAPH"),
+			*PCGLog::GetExecutionSourceName(ExecutionSource),
+			ExecutionSource->GetExecutionState().GetGraph() ? *ExecutionSource->GetExecutionState().GetGraph()->GetName() : TEXT("MISSINGGRAPH"),
 			*FString::Printf(TEXT("%" UINT64_FMT "'%s'"), Task.NodeId, Task.Node ? *Task.Node->GetNodeTitle(EPCGNodeTitleType::ListView).ToString() : TEXT("")));
 #endif // !(UE_BUILD_SHIPPING || UE_BUILD_TEST) || USE_LOGGING_IN_SHIPPING
 	}
@@ -369,7 +381,7 @@ namespace PCGGraphExecutionLogging
 		check(InContext);
 
 		UE_LOG(LogPCG, Log, TEXT("[GRIDLINKING] [%s] STORE. GenerationGridSize=%u, FromGridSize=%u, ToGridSize=%u, Path=%s, DataItems=%d"),
-			*PCGLog::GetComponentOwnerName(InContext->SourceComponent.Get(), /*bUseLabel=*/true),
+			*PCGLog::GetExecutionSourceName(InContext->ExecutionSource.Get(), /*bUseLabel=*/true),
 			PCGHiGenGrid::IsValidGrid(InGenerationGrid) ? PCGHiGenGrid::GridToGridSize(InGenerationGrid) : PCGHiGenGrid::UnboundedGridSize(),
 			InFromGridSize,
 			InToGridSize,
@@ -388,7 +400,7 @@ namespace PCGGraphExecutionLogging
 		check(InContext);
 
 		UE_LOG(LogPCG, Log, TEXT("[GRIDLINKING] [%s] RETRIEVE. GenerationGridSize=%u, FromGridSize=%u, ToGridSize=%u, Path=%s"),
-			*PCGLog::GetComponentOwnerName(InContext->SourceComponent.Get(), /*bUseLabel=*/true),
+			*PCGLog::GetExecutionSourceName(InContext->ExecutionSource.Get(), /*bUseLabel=*/true),
 			PCGHiGenGrid::IsValidGrid(InGenerationGrid) ? PCGHiGenGrid::GridToGridSize(InGenerationGrid) : PCGHiGenGrid::UnboundedGridSize(),
 			InFromGridSize,
 			InToGridSize,
@@ -396,7 +408,7 @@ namespace PCGGraphExecutionLogging
 #endif // !(UE_BUILD_SHIPPING || UE_BUILD_TEST) || USE_LOGGING_IN_SHIPPING
 	}
 
-	void LogGridLinkageTaskExecuteRetrieveSuccess(const FPCGContext* InContext, const UPCGComponent* InComponent, const FString& InResourcePath, int32 InDataItemCount)
+	void LogGridLinkageTaskExecuteRetrieveSuccess(const FPCGContext* InContext, const IPCGGraphExecutionSource* InExecutionSource, const FString& InResourcePath, int32 InDataItemCount)
 	{
 #if !(UE_BUILD_SHIPPING || UE_BUILD_TEST) || USE_LOGGING_IN_SHIPPING
 		if (!LogEnabled())
@@ -406,13 +418,13 @@ namespace PCGGraphExecutionLogging
 		check(InContext);
 
 		UE_LOG(LogPCG, Log, TEXT("[GRIDLINKING] [%s] RETRIEVE: SUCCESS. Path=%s, DataItems=%d"),
-			*PCGLog::GetComponentOwnerName(InContext->SourceComponent.Get(), /*bUseLabel=*/true),
+			*PCGLog::GetExecutionSourceName(InContext->ExecutionSource.Get(), /*bUseLabel=*/true),
 			*InResourcePath,
 			InDataItemCount);
 #endif // !(UE_BUILD_SHIPPING || UE_BUILD_TEST) || USE_LOGGING_IN_SHIPPING
 	}
 
-	void LogGridLinkageTaskExecuteRetrieveScheduleGraph(const FPCGContext* InContext, const UPCGComponent* InScheduledComponent, const FString& InResourcePath)
+	void LogGridLinkageTaskExecuteRetrieveScheduleGraph(const FPCGContext* InContext, const IPCGGraphExecutionSource* InScheduledSource, const FString& InResourcePath)
 	{
 #if !(UE_BUILD_SHIPPING || UE_BUILD_TEST) || USE_LOGGING_IN_SHIPPING
 		if (!LogEnabled())
@@ -421,14 +433,14 @@ namespace PCGGraphExecutionLogging
 		}
 		check(InContext);
 
-		UE_LOG(LogPCG, Log, TEXT("[GRIDLINKING] [%s] RETRIEVE: SCHEDULE GRAPH. Component=%s, Path=%s"),
-			*PCGLog::GetComponentOwnerName(InContext->SourceComponent.Get(), /*bUseLabel=*/true),
-			*PCGLog::GetComponentOwnerName(InScheduledComponent, /*bUseLabel=*/true),
+		UE_LOG(LogPCG, Log, TEXT("[GRIDLINKING] [%s] RETRIEVE: SCHEDULE GRAPH. Source=%s, Path=%s"),
+			*PCGLog::GetExecutionSourceName(InContext->ExecutionSource.Get(), /*bUseLabel=*/true),
+			*PCGLog::GetExecutionSourceName(InScheduledSource, /*bUseLabel=*/true),
 			*InResourcePath);
 #endif // !(UE_BUILD_SHIPPING || UE_BUILD_TEST) || USE_LOGGING_IN_SHIPPING
 	}
 
-	void LogGridLinkageTaskExecuteRetrieveWaitOnScheduledGraph(const FPCGContext* InContext, const UPCGComponent* InWaitOnComponent, const FString& InResourcePath)
+	void LogGridLinkageTaskExecuteRetrieveWaitOnScheduledGraph(const FPCGContext* InContext, const IPCGGraphExecutionSource* InWaitOnSource, const FString& InResourcePath)
 	{
 #if !(UE_BUILD_SHIPPING || UE_BUILD_TEST) || USE_LOGGING_IN_SHIPPING
 		if (!LogEnabled())
@@ -437,14 +449,14 @@ namespace PCGGraphExecutionLogging
 		}
 		check(InContext);
 
-		UE_LOG(LogPCG, Log, TEXT("[GRIDLINKING] [%s] RETRIEVE: WAIT FOR SCHEDULED GRAPH. Component=%s, Path=%s"),
-			*PCGLog::GetComponentOwnerName(InContext->SourceComponent.Get(), /*bUseLabel=*/true),
-			*PCGLog::GetComponentOwnerName(InWaitOnComponent, /*bUseLabel=*/true),
+		UE_LOG(LogPCG, Log, TEXT("[GRIDLINKING] [%s] RETRIEVE: WAIT FOR SCHEDULED GRAPH. Source=%s, Path=%s"),
+			*PCGLog::GetExecutionSourceName(InContext->ExecutionSource.Get(), /*bUseLabel=*/true),
+			*PCGLog::GetExecutionSourceName(InWaitOnSource, /*bUseLabel=*/true),
 			*InResourcePath);
 #endif // !(UE_BUILD_SHIPPING || UE_BUILD_TEST) || USE_LOGGING_IN_SHIPPING
 	}
 	
-	void LogGridLinkageTaskExecuteRetrieveWakeUp(const FPCGContext* InContext, const UPCGComponent* InWokenBy)
+	void LogGridLinkageTaskExecuteRetrieveWakeUp(const FPCGContext* InContext, const IPCGGraphExecutionSource* InWokenBySource)
 	{
 #if !(UE_BUILD_SHIPPING || UE_BUILD_TEST) || USE_LOGGING_IN_SHIPPING
 		if (!LogEnabled())
@@ -453,13 +465,13 @@ namespace PCGGraphExecutionLogging
 		}
 		check(InContext);
 
-		UE_LOG(LogPCG, Log, TEXT("[GRIDLINKING] [%s] RETRIEVE: WOKEN BY Component=%s"),
-			*PCGLog::GetComponentOwnerName(InContext->SourceComponent.Get(), /*bUseLabel=*/true),
-			*PCGLog::GetComponentOwnerName(InWokenBy, /*bUseLabel=*/true));
+		UE_LOG(LogPCG, Log, TEXT("[GRIDLINKING] [%s] RETRIEVE: WOKEN BY Source=%s"),
+			*PCGLog::GetExecutionSourceName(InContext->ExecutionSource.Get(), /*bUseLabel=*/true),
+			*PCGLog::GetExecutionSourceName(InWokenBySource, /*bUseLabel=*/true));
 #endif // !(UE_BUILD_SHIPPING || UE_BUILD_TEST) || USE_LOGGING_IN_SHIPPING
 	}
 
-	void LogGridLinkageTaskExecuteRetrieveNoLocalComponent(const FPCGContext* InContext, const FString& InResourcePath)
+	void LogGridLinkageTaskExecuteRetrieveNoLocalSource(const FPCGContext* InContext, const FString& InResourcePath)
 	{
 #if !(UE_BUILD_SHIPPING || UE_BUILD_TEST) || USE_LOGGING_IN_SHIPPING
 		if (!LogEnabled())
@@ -468,13 +480,13 @@ namespace PCGGraphExecutionLogging
 		}
 		check(InContext);
 
-		UE_LOG(LogPCG, Warning, TEXT("[GRIDLINKING] [%s] RETRIEVE: FAILED: No overlapping local component found. This may be expected. Path=%s"),
-			*PCGLog::GetComponentOwnerName(InContext->SourceComponent.Get(), /*bUseLabel=*/true),
+		UE_LOG(LogPCG, Warning, TEXT("[GRIDLINKING] [%s] RETRIEVE: FAILED: No overlapping local source found. This may be expected. Path=%s"),
+			*PCGLog::GetExecutionSourceName(InContext->ExecutionSource.Get(), /*bUseLabel=*/true),
 			*InResourcePath);
 #endif // !(UE_BUILD_SHIPPING || UE_BUILD_TEST) || USE_LOGGING_IN_SHIPPING
 	}
 
-	void LogGridLinkageTaskExecuteRetrieveNoData(const FPCGContext* InContext, const UPCGComponent* InComponent, const FString& InResourcePath)
+	void LogGridLinkageTaskExecuteRetrieveNoData(const FPCGContext* InContext, const IPCGGraphExecutionSource* InExecutionSource, const FString& InResourcePath)
 	{
 #if !(UE_BUILD_SHIPPING || UE_BUILD_TEST) || USE_LOGGING_IN_SHIPPING
 		if (!LogEnabled())
@@ -483,9 +495,9 @@ namespace PCGGraphExecutionLogging
 		}
 		check(InContext);
 
-		UE_LOG(LogPCG, Warning, TEXT("[GRIDLINKING] [%s] RETRIEVE: FAILED: No data found on local component. Component=%s, Path=%s"),
-			*PCGLog::GetComponentOwnerName(InContext->SourceComponent.Get(), /*bUseLabel=*/true),
-			*PCGLog::GetComponentOwnerName(InComponent, /*bUseLabel=*/true),
+		UE_LOG(LogPCG, Warning, TEXT("[GRIDLINKING] [%s] RETRIEVE: FAILED: No data found on local source. Source=%s, Path=%s"),
+			*PCGLog::GetExecutionSourceName(InContext->ExecutionSource.Get(), /*bUseLabel=*/true),
+			*PCGLog::GetExecutionSourceName(InExecutionSource, /*bUseLabel=*/true),
 			*InResourcePath);
 #endif // !(UE_BUILD_SHIPPING || UE_BUILD_TEST) || USE_LOGGING_IN_SHIPPING
 	}
diff --git a/Engine/Plugins/PCG/Source/PCG/Private/Utils/PCGGraphExecutionLogging.h b/Engine/Plugins/PCG/Source/PCG/Private/Utils/PCGGraphExecutionLogging.h
--- a/Engine/Plugins/PCG/Source/PCG/Private/Utils/PCGGraphExecutionLogging.h
+++ b/Engine/Plugins/PCG/Source/PCG/Private/Utils/PCGGraphExecutionLogging.h
@@ -7,6 +7,7 @@
 #include "Containers/UnrealString.h"
 #include "HAL/Platform.h"
 
+class IPCGGraphExecutionSource;
 class UPCGComponent;
 class UPCGGraph;
 enum class EPCGHiGenGrid : uint32;
@@ -25,16 +26,16 @@ namespace PCGGraphExecutionLogging
 	void LogGraphTasks(const TMap<FPCGTaskId, FPCGGraphTask>& Tasks, const TMap<FPCGTaskId, TSet<FPCGTaskId>>* TaskSuccessors = nullptr);
 	void LogGraphTasks(const TArray<FPCGGraphTask>& Tasks);
 
-	void LogGraphSchedule(const UPCGComponent* SourceComponent, const UPCGGraph* InScheduledGraph);
-	void LogGraphScheduleDependency(const UPCGComponent* InComponent, const FPCGStack* InFromStack);
-	void LogGraphScheduleDependencyFailed(const UPCGComponent* InComponent, const FPCGStack* InFromStack);
+	void LogGraphSchedule(const IPCGGraphExecutionSource* InExecutionSource, const UPCGGraph* InScheduledGraph);
+	void LogGraphScheduleDependency(const IPCGGraphExecutionSource* InExecutionSource, const FPCGStack* InFromStack);
+	void LogGraphScheduleDependencyFailed(const IPCGGraphExecutionSource* InExecutionSource, const FPCGStack* InFromStack);
 
 	void LogGraphPostSchedule(const TMap<FPCGTaskId, FPCGGraphTask>& Tasks, const TMap<FPCGTaskId, TSet<FPCGTaskId>>& TaskSuccessors);
-	void LogPostProcessGraph(const UPCGComponent* InSourceComponent);
+	void LogPostProcessGraph(const IPCGGraphExecutionSource* InExecutionSource);
 
-	void LogComponentCancellation(const TSet<UPCGComponent*>& CancelledComponents);
+	void LogExecutionSourceCancellation(const TSet<IPCGGraphExecutionSource*>& CancelledComponents);
 
-	void LogChangeOriginIgnoredForComponent(const UObject* InObject, const UPCGComponent* InComponent);
+	void LogChangeOriginIgnoredForComponent(const UObject* InObject, const IPCGGraphExecutionSource* InExecutionSource);
 
 	void LogGraphExecuteFrameFinished();
 
@@ -47,10 +48,10 @@ namespace PCGGraphExecutionLogging
 
 	void LogGridLinkageTaskExecuteStore(const FPCGContext* InContext, EPCGHiGenGrid InGenerationGrid, int32 InFromLOD, int32 InToLOD, const FString& InResourcePath, int32 InDataItemCount);
 	void LogGridLinkageTaskExecuteRetrieve(const FPCGContext* InContext, EPCGHiGenGrid InGenerationGrid, int32 InFromLOD, int32 InToLOD, const FString& InResourcePath);
-	void LogGridLinkageTaskExecuteRetrieveSuccess(const FPCGContext* InContext, const UPCGComponent* InComponent, const FString& InResourcePath, int32 InDataItemCount);
-	void LogGridLinkageTaskExecuteRetrieveScheduleGraph(const FPCGContext* InContext, const UPCGComponent* InScheduledComponent, const FString& InResourcePath);
-	void LogGridLinkageTaskExecuteRetrieveWaitOnScheduledGraph(const FPCGContext* InContext, const UPCGComponent* InWaitOnComponent, const FString& InResourcePath);
-	void LogGridLinkageTaskExecuteRetrieveWakeUp(const FPCGContext* InContext, const UPCGComponent* InWokenBy);
-	void LogGridLinkageTaskExecuteRetrieveNoLocalComponent(const FPCGContext* InContext, const FString& InResourcePath);
-	void LogGridLinkageTaskExecuteRetrieveNoData(const FPCGContext* InContext, const UPCGComponent* InComponent, const FString& InResourcePath);
+	void LogGridLinkageTaskExecuteRetrieveSuccess(const FPCGContext* InContext, const IPCGGraphExecutionSource* InExecutionSource, const FString& InResourcePath, int32 InDataItemCount);
+	void LogGridLinkageTaskExecuteRetrieveScheduleGraph(const FPCGContext* InContext, const IPCGGraphExecutionSource* InScheduledSource, const FString& InResourcePath);
+	void LogGridLinkageTaskExecuteRetrieveWaitOnScheduledGraph(const FPCGContext* InContext, const IPCGGraphExecutionSource* InWaitOnSource, const FString& InResourcePath);
+	void LogGridLinkageTaskExecuteRetrieveWakeUp(const FPCGContext* InContext, const IPCGGraphExecutionSource* InWokenBySource);
+	void LogGridLinkageTaskExecuteRetrieveNoLocalSource(const FPCGContext* InContext, const FString& InResourcePath);
+	void LogGridLinkageTaskExecuteRetrieveNoData(const FPCGContext* InContext, const IPCGGraphExecutionSource* InExecutionSource, const FString& InResourcePath);
 }
diff --git a/Engine/Plugins/PCG/Source/PCG/Public/Elements/ControlFlow/PCGQualityBranch.h b/Engine/Plugins/PCG/Source/PCG/Public/Elements/ControlFlow/PCGQualityBranch.h
--- a/Engine/Plugins/PCG/Source/PCG/Public/Elements/ControlFlow/PCGQualityBranch.h
+++ b/Engine/Plugins/PCG/Source/PCG/Public/Elements/ControlFlow/PCGQualityBranch.h
@@ -52,7 +52,7 @@ public:
 class FPCGQualityBranchElement : public IPCGElement
 {
 public:
-	virtual void GetDependenciesCrc(const FPCGDataCollection& InInput, const UPCGSettings* InSettings, UPCGComponent* InComponent, FPCGCrc& OutCrc) const override;
+	virtual void GetDependenciesCrc(const FPCGGetDependenciesCrcParams& InParams, FPCGCrc& OutCrc) const override;
 
 protected:
 	virtual bool ExecuteInternal(FPCGContext* Context) const override;
diff --git a/Engine/Plugins/PCG/Source/PCG/Public/Elements/ControlFlow/PCGQualitySelect.h b/Engine/Plugins/PCG/Source/PCG/Public/Elements/ControlFlow/PCGQualitySelect.h
--- a/Engine/Plugins/PCG/Source/PCG/Public/Elements/ControlFlow/PCGQualitySelect.h
+++ b/Engine/Plugins/PCG/Source/PCG/Public/Elements/ControlFlow/PCGQualitySelect.h
@@ -52,7 +52,7 @@ public:
 class FPCGQualitySelectElement : public IPCGElement
 {
 public:
-	virtual void GetDependenciesCrc(const FPCGDataCollection& InInput, const UPCGSettings* InSettings, UPCGComponent* InComponent, FPCGCrc& OutCrc) const override;
+	virtual void GetDependenciesCrc(const FPCGGetDependenciesCrcParams& InParams, FPCGCrc& OutCrc) const override;
 
 protected:
 	virtual bool ExecuteInternal(FPCGContext* Context) const override;
diff --git a/Engine/Plugins/PCG/Source/PCG/Public/Elements/Metadata/PCGMetadataOpElementBase.h b/Engine/Plugins/PCG/Source/PCG/Public/Elements/Metadata/PCGMetadataOpElementBase.h
--- a/Engine/Plugins/PCG/Source/PCG/Public/Elements/Metadata/PCGMetadataOpElementBase.h
+++ b/Engine/Plugins/PCG/Source/PCG/Public/Elements/Metadata/PCGMetadataOpElementBase.h
@@ -253,7 +253,7 @@ protected:
 	virtual bool ExecuteInternal(FPCGContext* Context) const override;
 	virtual EPCGElementExecutionLoopMode ExecutionLoopMode(const UPCGSettings* Settings) const override { return EPCGElementExecutionLoopMode::SinglePrimaryPin; }
 	virtual bool SupportsBasePointDataInputs(FPCGContext* InContext) const override { return true; }
-	virtual void GetDependenciesCrc(const FPCGDataCollection& InInput, const UPCGSettings* InSettings, UPCGComponent* InComponent, FPCGCrc& OutCrc) const override;
+	virtual void GetDependenciesCrc(const FPCGGetDependenciesCrcParams& InParams, FPCGCrc& OutCrc) const override;
 
 	virtual bool DoOperation(PCGMetadataOps::FOperationData& InOperationData) const = 0;
 
diff --git a/Engine/Plugins/PCG/Source/PCG/Public/Elements/PCGCreatePoints.h b/Engine/Plugins/PCG/Source/PCG/Public/Elements/PCGCreatePoints.h
--- a/Engine/Plugins/PCG/Source/PCG/Public/Elements/PCGCreatePoints.h
+++ b/Engine/Plugins/PCG/Source/PCG/Public/Elements/PCGCreatePoints.h
@@ -69,9 +69,15 @@ PRAGMA_ENABLE_DEPRECATION_WARNINGS
 	bool bCullPointsOutsideVolume = false;
 };
 
+namespace PCGCreatePointsHelper
+{
+	UPCGData* GetDependenciesData(IPCGGraphExecutionSource* InExecutionSource, EPCGCoordinateSpace InCoordinateSpace, bool bInCullPointsOutsideVolume);
+}
+
 class FPCGCreatePointsElement : public IPCGElement
 {
 public:
+	virtual void GetDependenciesCrc(const FPCGGetDependenciesCrcParams& InParams, FPCGCrc& OutCrc) const override;
 	virtual bool IsCacheable(const UPCGSettings* InSettings) const override;
 	virtual bool ShouldComputeFullOutputDataCrc(FPCGContext* Context) const override { return true; }
 
diff --git a/Engine/Plugins/PCG/Source/PCG/Public/Elements/PCGCreatePointsGrid.h b/Engine/Plugins/PCG/Source/PCG/Public/Elements/PCGCreatePointsGrid.h
--- a/Engine/Plugins/PCG/Source/PCG/Public/Elements/PCGCreatePointsGrid.h
+++ b/Engine/Plugins/PCG/Source/PCG/Public/Elements/PCGCreatePointsGrid.h
@@ -92,7 +92,7 @@ PRAGMA_ENABLE_DEPRECATION_WARNINGS
 class FPCGCreatePointsGridElement : public IPCGElement
 {
 public:
-	virtual void GetDependenciesCrc(const FPCGDataCollection& InInput, const UPCGSettings* InSettings, UPCGComponent* InComponent, FPCGCrc& OutCrc) const override;
+	virtual void GetDependenciesCrc(const FPCGGetDependenciesCrcParams& InParams, FPCGCrc& OutCrc) const override;
 	virtual bool SupportsBasePointDataInputs(FPCGContext* InContext) const override { return true; }
 
 protected:
diff --git a/Engine/Plugins/PCG/Source/PCG/Public/Elements/PCGCreatePointsSphere.h b/Engine/Plugins/PCG/Source/PCG/Public/Elements/PCGCreatePointsSphere.h
--- a/Engine/Plugins/PCG/Source/PCG/Public/Elements/PCGCreatePointsSphere.h
+++ b/Engine/Plugins/PCG/Source/PCG/Public/Elements/PCGCreatePointsSphere.h
@@ -143,7 +143,7 @@ public:
 class FPCGCreatePointsSphereElement : public IPCGElement
 {
 public:
-	virtual void GetDependenciesCrc(const FPCGDataCollection& InInput, const UPCGSettings* InSettings, UPCGComponent* InComponent, FPCGCrc& OutCrc) const override;
+	virtual void GetDependenciesCrc(const FPCGGetDependenciesCrcParams& InParams, FPCGCrc& OutCrc) const override;
 
 protected:
 	virtual bool ExecuteInternal(FPCGContext* InContext) const override;
diff --git a/Engine/Plugins/PCG/Source/PCG/Public/Elements/PCGCullPointsOutsideActorBounds.h b/Engine/Plugins/PCG/Source/PCG/Public/Elements/PCGCullPointsOutsideActorBounds.h
--- a/Engine/Plugins/PCG/Source/PCG/Public/Elements/PCGCullPointsOutsideActorBounds.h
+++ b/Engine/Plugins/PCG/Source/PCG/Public/Elements/PCGCullPointsOutsideActorBounds.h
@@ -48,7 +48,7 @@ public:
 class FPCGCullPointsOutsideActorBoundsElement : public IPCGElement
 {
 public:
-	virtual void GetDependenciesCrc(const FPCGDataCollection& InInput, const UPCGSettings* InSettings, UPCGComponent* InComponent, FPCGCrc& OutCrc) const override;
+	virtual void GetDependenciesCrc(const FPCGGetDependenciesCrcParams& InParams, FPCGCrc& OutCrc) const override;
 
 protected:
 	virtual bool ExecuteInternal(FPCGContext* Context) const override;
diff --git a/Engine/Plugins/PCG/Source/PCG/Public/Elements/PCGDataFromActor.h b/Engine/Plugins/PCG/Source/PCG/Public/Elements/PCGDataFromActor.h
--- a/Engine/Plugins/PCG/Source/PCG/Public/Elements/PCGDataFromActor.h
+++ b/Engine/Plugins/PCG/Source/PCG/Public/Elements/PCGDataFromActor.h
@@ -161,7 +161,7 @@ class PCG_API FPCGDataFromActorElement : public IPCGElementWithCustomContext<FPC
 public:
 	virtual bool IsCacheable(const UPCGSettings* InSettings) const override;
 	virtual bool CanExecuteOnlyOnMainThread(FPCGContext* Context) const override { return true; }
-	virtual void GetDependenciesCrc(const FPCGDataCollection& InInput, const UPCGSettings* InSettings, UPCGComponent* InComponent, FPCGCrc& OutCrc) const override;
+	virtual void GetDependenciesCrc(const FPCGGetDependenciesCrcParams& InParams, FPCGCrc& OutCrc) const override;
 	virtual bool ShouldComputeFullOutputDataCrc(FPCGContext* Context) const override { return true; }
 
 protected:
diff --git a/Engine/Plugins/PCG/Source/PCG/Public/Elements/PCGExecuteBlueprint.h b/Engine/Plugins/PCG/Source/PCG/Public/Elements/PCGExecuteBlueprint.h
--- a/Engine/Plugins/PCG/Source/PCG/Public/Elements/PCGExecuteBlueprint.h
+++ b/Engine/Plugins/PCG/Source/PCG/Public/Elements/PCGExecuteBlueprint.h
@@ -420,5 +420,5 @@ public:
 protected:
 	virtual bool ExecuteInternal(FPCGContext* Context) const override;
 	virtual void PostExecuteInternal(FPCGContext* Context) const override;
-	virtual FPCGContext* Initialize(const FPCGDataCollection& InputData, TWeakObjectPtr<UPCGComponent> SourceComponent, const UPCGNode* Node) override;	
+	virtual FPCGContext* Initialize(const FPCGInitializeElementParams& InParams) override;	
 };
diff --git a/Engine/Plugins/PCG/Source/PCG/Public/Elements/PCGGenerateGrassMaps.h b/Engine/Plugins/PCG/Source/PCG/Public/Elements/PCGGenerateGrassMaps.h
--- a/Engine/Plugins/PCG/Source/PCG/Public/Elements/PCGGenerateGrassMaps.h
+++ b/Engine/Plugins/PCG/Source/PCG/Public/Elements/PCGGenerateGrassMaps.h
@@ -115,7 +115,7 @@ protected:
 class FPCGGenerateGrassMapsElement : public IPCGElement
 {
 public:
-	virtual void GetDependenciesCrc(const FPCGDataCollection& InInput, const UPCGSettings* InSettings, UPCGComponent* InComponent, FPCGCrc& OutCrc) const override;
+	virtual void GetDependenciesCrc(const FPCGGetDependenciesCrcParams& InParams, FPCGCrc& OutCrc) const override;
 	/** FLandscapeGrassWeightExporter expects to exist in scope only on the game thread. */
 	virtual bool CanExecuteOnlyOnMainThread(FPCGContext* Context) const { return true; }
 
diff --git a/Engine/Plugins/PCG/Source/PCG/Public/Elements/PCGGetActorProperty.h b/Engine/Plugins/PCG/Source/PCG/Public/Elements/PCGGetActorProperty.h
--- a/Engine/Plugins/PCG/Source/PCG/Public/Elements/PCGGetActorProperty.h
+++ b/Engine/Plugins/PCG/Source/PCG/Public/Elements/PCGGetActorProperty.h
@@ -128,7 +128,7 @@ class FPCGGetActorPropertyElement : public IPCGElementWithCustomContext<FPCGLoad
 public:
 	virtual bool CanExecuteOnlyOnMainThread(FPCGContext* Context) const override { return true; }
 	virtual bool IsCacheable(const UPCGSettings* InSettings) const override;
-	virtual void GetDependenciesCrc(const FPCGDataCollection& InInput, const UPCGSettings* InSettings, UPCGComponent* InComponent, FPCGCrc& OutCrc) const override;
+	virtual void GetDependenciesCrc(const FPCGGetDependenciesCrcParams& InParams, FPCGCrc& OutCrc) const override;
 protected:
 	virtual bool PrepareDataInternal(FPCGContext* Context) const override;
 	virtual bool ExecuteInternal(FPCGContext* Context) const override;
diff --git a/Engine/Plugins/PCG/Source/PCG/Public/Elements/PCGHiGenGridSize.h b/Engine/Plugins/PCG/Source/PCG/Public/Elements/PCGHiGenGridSize.h
--- a/Engine/Plugins/PCG/Source/PCG/Public/Elements/PCGHiGenGridSize.h
+++ b/Engine/Plugins/PCG/Source/PCG/Public/Elements/PCGHiGenGridSize.h
@@ -55,7 +55,7 @@ protected:
 class FPCGHiGenGridSizeElement : public IPCGElement
 {
 public:
-	virtual void GetDependenciesCrc(const FPCGDataCollection& InInput, const UPCGSettings* InSettings, UPCGComponent* InComponent, FPCGCrc& OutCrc) const override;
+	virtual void GetDependenciesCrc(const FPCGGetDependenciesCrcParams& InParams, FPCGCrc& OutCrc) const override;
 
 protected:
 	virtual bool ExecuteInternal(FPCGContext* Context) const override;
diff --git a/Engine/Plugins/PCG/Source/PCG/Public/Elements/PCGSplineSampler.h b/Engine/Plugins/PCG/Source/PCG/Public/Elements/PCGSplineSampler.h
--- a/Engine/Plugins/PCG/Source/PCG/Public/Elements/PCGSplineSampler.h
+++ b/Engine/Plugins/PCG/Source/PCG/Public/Elements/PCGSplineSampler.h
@@ -282,7 +282,7 @@ public:
 class FPCGSplineSamplerElement : public IPCGElement
 {
 public:
-	virtual void GetDependenciesCrc(const FPCGDataCollection& InInput, const UPCGSettings* InSettings, UPCGComponent* InComponent, FPCGCrc& OutCrc) const override;
+	virtual void GetDependenciesCrc(const FPCGGetDependenciesCrcParams& InParams, FPCGCrc& OutCrc) const override;
 	// Worth computing a full CRC in case we can halt change propagation/re-executions
 	virtual bool ShouldComputeFullOutputDataCrc(FPCGContext* Context) const override { return true; }
 
diff --git a/Engine/Plugins/PCG/Source/PCG/Public/Elements/PCGSurfaceSampler.h b/Engine/Plugins/PCG/Source/PCG/Public/Elements/PCGSurfaceSampler.h
--- a/Engine/Plugins/PCG/Source/PCG/Public/Elements/PCGSurfaceSampler.h
+++ b/Engine/Plugins/PCG/Source/PCG/Public/Elements/PCGSurfaceSampler.h
@@ -166,7 +166,7 @@ public:
 class FPCGSurfaceSamplerElement : public TPCGTimeSlicedElementBase<PCGSurfaceSampler::FSurfaceSamplerExecutionState, PCGSurfaceSampler::FSurfaceSamplerIterationState>
 {
 public:
-	virtual void GetDependenciesCrc(const FPCGDataCollection& InInput, const UPCGSettings* InSettings, UPCGComponent* InComponent, FPCGCrc& OutCrc) const override;
+	virtual void GetDependenciesCrc(const FPCGGetDependenciesCrcParams&, FPCGCrc& OutCrc) const override;
 	// Might be sampling landscape or other external data, worth computing a full CRC in case we can halt change propagation/re-executions
 	virtual bool ShouldComputeFullOutputDataCrc(FPCGContext* Context) const override { return true; }
 	virtual bool SupportsBasePointDataInputs(FPCGContext* InContext) const override { return true; }
diff --git a/Engine/Plugins/PCG/Source/PCG/Public/Elements/PCGTextureSampler.h b/Engine/Plugins/PCG/Source/PCG/Public/Elements/PCGTextureSampler.h
--- a/Engine/Plugins/PCG/Source/PCG/Public/Elements/PCGTextureSampler.h
+++ b/Engine/Plugins/PCG/Source/PCG/Public/Elements/PCGTextureSampler.h
@@ -159,7 +159,7 @@ protected:
 class FPCGTextureSamplerElement : public IPCGElement
 {
 public:
-	virtual void GetDependenciesCrc(const FPCGDataCollection& InInput, const UPCGSettings* InSettings, UPCGComponent* InComponent, FPCGCrc& OutCrc) const override;
+	virtual void GetDependenciesCrc(const FPCGGetDependenciesCrcParams& InParams, FPCGCrc& OutCrc) const override;
 	// Loading needs to be done on the main thread and accessing objects outside of PCG might not be thread safe, so taking the safe approach
 	virtual bool CanExecuteOnlyOnMainThread(FPCGContext* Context) const override { return true; }
 
diff --git a/Engine/Plugins/PCG/Source/PCG/Public/Elements/PCGVolumeSampler.h b/Engine/Plugins/PCG/Source/PCG/Public/Elements/PCGVolumeSampler.h
--- a/Engine/Plugins/PCG/Source/PCG/Public/Elements/PCGVolumeSampler.h
+++ b/Engine/Plugins/PCG/Source/PCG/Public/Elements/PCGVolumeSampler.h
@@ -107,7 +107,7 @@ protected:
 class FPCGVolumeSamplerElement : public TPCGTimeSlicedElementBase<PCGVolumeSampler::FVolumeSamplerExecutionState, PCGVolumeSampler::FVolumeSamplerIterationState>
 {
 public:
-	virtual void GetDependenciesCrc(const FPCGDataCollection& InInput, const UPCGSettings* InSettings, UPCGComponent* InComponent, FPCGCrc& OutCrc) const override;
+	virtual void GetDependenciesCrc(const FPCGGetDependenciesCrcParams& InParams, FPCGCrc& OutCrc) const override;
 	// Might be sampling external data like brush, worth computing a full CRC in case we can halt change propagation/re-executions
 	virtual bool ShouldComputeFullOutputDataCrc(FPCGContext* Context) const override { return true; }
 
diff --git a/Engine/Plugins/PCG/Source/PCG/Public/Graph/IPCGGraphCache.h b/Engine/Plugins/PCG/Source/PCG/Public/Graph/IPCGGraphCache.h
--- a/Engine/Plugins/PCG/Source/PCG/Public/Graph/IPCGGraphCache.h
+++ b/Engine/Plugins/PCG/Source/PCG/Public/Graph/IPCGGraphCache.h
@@ -5,6 +5,7 @@
 #include "PCGCrc.h"
 
 class IPCGElement;
+class IPCGGraphExecutionSource;
 class UPCGComponent;
 class UPCGNode;
 struct FPCGDataCollection;
@@ -13,7 +14,11 @@ struct FPCGGetFromCacheParams
 {
 	const UPCGNode* Node = nullptr;
 	const IPCGElement* Element = nullptr;
+
+	UE_DEPRECATED(5.6, "Use ExecutionSource instead")
 	const UPCGComponent* Component = nullptr;
+
+	const IPCGGraphExecutionSource* ExecutionSource = nullptr;
 	FPCGCrc Crc;
 };
 
diff --git a/Engine/Plugins/PCG/Source/PCG/Public/Helpers/PCGDynamicTrackingHelpers.h b/Engine/Plugins/PCG/Source/PCG/Public/Helpers/PCGDynamicTrackingHelpers.h
--- a/Engine/Plugins/PCG/Source/PCG/Public/Helpers/PCGDynamicTrackingHelpers.h
+++ b/Engine/Plugins/PCG/Source/PCG/Public/Helpers/PCGDynamicTrackingHelpers.h
@@ -6,9 +6,10 @@
 
 #include "Elements/PCGActorSelector.h"
 #include "Containers/Array.h"
+#include "UObject/WeakInterfacePtr.h"
 
 struct FPCGContext;
-class UPCGComponent;
+class IPCGGraphExecutionSource;
 
 /**
 * Simple helper class to factorize the logic for gathering dynamic tracking keys and pushing them to the component.
@@ -34,7 +35,7 @@ public:
 
 private:
 	bool bDynamicallyTracked = false;
-	TWeakObjectPtr<UPCGComponent> CachedComponent;
+	TWeakInterfacePtr<IPCGGraphExecutionSource> CachedExecutionSource;
 	TArray<TPair<FPCGSelectionKey, bool>, TInlineAllocator<16>> DynamicallyTrackedKeysAndCulling;
 };
 
diff --git a/Engine/Plugins/PCG/Source/PCG/Public/Helpers/PCGHelpers.h b/Engine/Plugins/PCG/Source/PCG/Public/Helpers/PCGHelpers.h
--- a/Engine/Plugins/PCG/Source/PCG/Public/Helpers/PCGHelpers.h
+++ b/Engine/Plugins/PCG/Source/PCG/Public/Helpers/PCGHelpers.h
@@ -16,6 +16,7 @@ class AActor;
 class APCGWorldActor;
 class ALandscape;
 class ALandscapeProxy;
+class IPCGGraphExecutionSource;
 class UPCGComponent;
 class UPCGGraph;
 class UPCGMetadata;
@@ -36,8 +37,8 @@ namespace PCGHelpers
 	PCG_API int ComputeSeed(int A, int B);
 	PCG_API int ComputeSeed(int A, int B, int C);
 	PCG_API int ComputeSeedFromPosition(const FVector& InPosition);
-	PCG_API FRandomStream GetRandomStreamFromSeed(int32 Seed, const UPCGSettings* OptionalSettings = nullptr, const UPCGComponent* OptionalComponent = nullptr);
-	PCG_API FRandomStream GetRandomStreamFromTwoSeeds(int32 SeedA, int32 SeedB, const UPCGSettings* OptionalSettings = nullptr, const UPCGComponent* OptionalComponent = nullptr);
+	PCG_API FRandomStream GetRandomStreamFromSeed(int32 Seed, const UPCGSettings* OptionalSettings = nullptr, const IPCGGraphExecutionSource* OptionalExecutionSource = nullptr);
+	PCG_API FRandomStream GetRandomStreamFromTwoSeeds(int32 SeedA, int32 SeedB, const UPCGSettings* OptionalSettings = nullptr, const IPCGGraphExecutionSource* OptionalExecutionSource = nullptr);
 
 	PCG_API bool IsInsideBounds(const FBox& InBox, const FVector& InPosition);
 	PCG_API bool IsInsideBoundsXY(const FBox& InBox, const FVector& InPosition);
@@ -80,6 +81,9 @@ namespace PCGHelpers
 	/** If hierarchical generation is enabled, returns all relevant grid sizes for the graph, otherwise returns partition grid size from world actor. */
 	PCG_API bool GetGenerationGridSizes(const UPCGGraph* InGraph, const APCGWorldActor* InWorldActor, PCGHiGenGrid::FSizeArray& OutGridSizes, bool& bOutHasUnbounded);
 
+	PCG_API int32 GetGenerationGridSize(const IPCGGraphExecutionSource* InExecutionSource);
+	PCG_API bool IsRuntimeGeneration(const IPCGGraphExecutionSource* InExecutionSource);
+
 #if WITH_EDITOR
 	PCG_API void GetGeneratedActorsFolderPath(const AActor* InTargetActor, FString& OutFolderPath);
 	PCG_API void GetGeneratedActorsFolderPath(const AActor* InTargetActor, const FPCGContext* InContext, EPCGAttachOptions AttachOptions, FString& OutFolderPath);
diff --git a/Engine/Plugins/PCG/Source/PCG/Public/Helpers/PCGSettingsHelpers.h b/Engine/Plugins/PCG/Source/PCG/Public/Helpers/PCGSettingsHelpers.h
--- a/Engine/Plugins/PCG/Source/PCG/Public/Helpers/PCGSettingsHelpers.h
+++ b/Engine/Plugins/PCG/Source/PCG/Public/Helpers/PCGSettingsHelpers.h
@@ -84,10 +84,10 @@ namespace PCGSettingsHelpers
 		const UPCGSpatialData* BoundingShape = Context->InputData.GetSpatialUnionOfInputsByPin(Context, BoundingShapeLabel, bOutUnionWasCreated);
 
 		// Fallback to getting bounds from actor
-		if (!BoundingShape && Context->SourceComponent.IsValid())
+		if (!BoundingShape && Context->ExecutionSource.IsValid())
 		{
 			check(bOutUnionWasCreated == false);
-			BoundingShape = Cast<UPCGSpatialData>(Context->SourceComponent->GetActorPCGData());
+			BoundingShape = Cast<UPCGSpatialData>(Context->ExecutionSource->GetExecutionState().GetSelfData());
 		}
 
 		return BoundingShape;
diff --git a/Engine/Plugins/PCG/Source/PCG/Public/PCGCommon.h b/Engine/Plugins/PCG/Source/PCG/Public/PCGCommon.h
--- a/Engine/Plugins/PCG/Source/PCG/Public/PCGCommon.h
+++ b/Engine/Plugins/PCG/Source/PCG/Public/PCGCommon.h
@@ -12,6 +12,7 @@ static const FPCGTaskId InvalidPCGTaskId = (uint64)-1;
 using FPCGPinId = uint64;
 
 class IPCGElement;
+class IPCGGraphExecutionSource;
 class UPCGComponent;
 class UPCGGraph;
 struct FPCGContext;
@@ -20,7 +21,7 @@ struct FPCGStack;
 namespace PCGLog
 {
 	/** Returns the component's owner name. In the case of commandlets, owner full object path will be returned. */
-	FString GetComponentOwnerName(const UPCGComponent* InComponent, bool bUseLabel = false, FString InDefaultName = TEXT("MISSINGOWNER"));
+	FString GetExecutionSourceName(const IPCGGraphExecutionSource* InExecutionSource, bool bUseLabel = false, FString InDefaultName = TEXT("MISSINGOWNER"));
 }
 
 namespace PCGPinIdHelpers
@@ -510,14 +511,14 @@ struct FPCGScheduleGraphParams
 {
 	FPCGScheduleGraphParams(
 		UPCGGraph* InGraph,
-		UPCGComponent* InSourceComponent,
+		IPCGGraphExecutionSource* InExecutionSource,
 		FPCGElementPtr InPreGraphElement,
 		FPCGElementPtr InInputElement,
 		TArray<FPCGTaskId> InExternalDependencies,
 		const FPCGStack* InFromStack,
 		bool bInAllowHierarchicalGeneration)
 		: Graph(InGraph)
-		, SourceComponent(InSourceComponent)
+		, ExecutionSource(InExecutionSource)
 		, PreGraphElement(InPreGraphElement)
 		, InputElement(InInputElement)
 		, ExternalDependencies(MoveTemp(InExternalDependencies))
@@ -529,8 +530,8 @@ struct FPCGScheduleGraphParams
 	// Graph to execute
 	UPCGGraph* Graph = nullptr;
 
-	// PCG component associated with this task. Can be null.
-	UPCGComponent* SourceComponent = nullptr;
+	// PCG execution source associated with this task. Can be null.
+	IPCGGraphExecutionSource* ExecutionSource = nullptr;
 
 	// First task to run.
 	FPCGElementPtr PreGraphElement;
@@ -550,19 +551,19 @@ struct FPCGScheduleGraphParams
 
 struct FPCGScheduleGenericParams
 {
-	FPCGScheduleGenericParams(TFunction<bool(FPCGContext*)> InOperation, UPCGComponent* InSourceComponent)
-		: Operation(MoveTemp(InOperation)), SourceComponent(InSourceComponent)
+	FPCGScheduleGenericParams(TFunction<bool(FPCGContext*)> InOperation, IPCGGraphExecutionSource* InExecutionSource)
+		: Operation(MoveTemp(InOperation)), ExecutionSource(InExecutionSource)
 	{
 	}
 
 	FPCGScheduleGenericParams(
 		TFunction<bool(FPCGContext*)> InOperation, 
-		UPCGComponent* InSourceComponent, 
+		IPCGGraphExecutionSource* InExecutionSource, 
 		const TArray<FPCGTaskId>& InExecutionDependencies, 
 		const TArray<FPCGTaskId>& InDataDependencies, 
 		bool bInSupportBasePointDataInput)
 		: Operation(MoveTemp(InOperation))
-		, SourceComponent(InSourceComponent)
+		, ExecutionSource(InExecutionSource)
 		, ExecutionDependencies(InExecutionDependencies)
 		, DataDependencies(InDataDependencies)
 		, bSupportBasePointDataInput(bInSupportBasePointDataInput)
@@ -572,12 +573,12 @@ struct FPCGScheduleGenericParams
 	FPCGScheduleGenericParams(
 		TFunction<bool(FPCGContext*)> InOperation, 
 		TFunction<void(FPCGContext*)> InAbortOperation, 
-		UPCGComponent* InSourceComponent, 
+		IPCGGraphExecutionSource* InExecutionSource,
 		const TArray<FPCGTaskId>& InExecutionDependencies, 
 		const TArray<FPCGTaskId>& InDataDependencies = {}, 
 		bool bInSupportBasePointDataInput = false)
 		: Operation(MoveTemp(InOperation))
-		, SourceComponent(InSourceComponent)
+		, ExecutionSource(InExecutionSource)
 		, AbortOperation(MoveTemp(InAbortOperation))
 		, ExecutionDependencies(InExecutionDependencies)
 		, DataDependencies(InDataDependencies)
@@ -588,8 +589,8 @@ struct FPCGScheduleGenericParams
 	// Callback that takes a Context as argument and returns true if the task is done, false otherwise
 	TFunction<bool(FPCGContext*)> Operation;
 
-	// PCG component associated with this task. Can be null.
-	UPCGComponent* SourceComponent = nullptr;
+	// PCG execution source associated with this task. Can be null.
+	IPCGGraphExecutionSource* ExecutionSource = nullptr;
 
 	// Callback that is called if the task is aborted (cancelled) before fully executed.
 	TFunction<void(FPCGContext*)> AbortOperation;
diff --git a/Engine/Plugins/PCG/Source/PCG/Public/PCGComponent.h b/Engine/Plugins/PCG/Source/PCG/Public/PCGComponent.h
--- a/Engine/Plugins/PCG/Source/PCG/Public/PCGComponent.h
+++ b/Engine/Plugins/PCG/Source/PCG/Public/PCGComponent.h
@@ -4,6 +4,7 @@
 
 #include "PCGCommon.h"
 #include "PCGGraphExecutionInspection.h"
+#include "PCGGraphExecutionStateInterface.h"
 #include "PCGNode.h"
 #include "PCGSettings.h"
 #include "Graph/PCGStackContext.h"
@@ -86,13 +87,44 @@ namespace PCGTestsCommon
 	struct FTestData;
 }
 
+class FPCGComponentExecutionState : public IPCGGraphExecutionState
+{
+	friend class UPCGComponent;
+
+public:
+	virtual UPCGData* GetSelfData() const override;
+	virtual int32 GetSeed() const override;
+	virtual FString GetDebugName() const override;
+	virtual FTransform GetTransform() const override;
+	virtual UWorld* GetWorld() const override;
+	virtual FBox GetBounds() const override;
+	virtual UPCGGraph* GetGraph() const override;
+	virtual UPCGGraphInstance* GetGraphInstance() const override;
+	virtual void OnGraphExecutionAborted(bool bQuiet, bool bCleanupUnusedResources) override;
+	virtual void Cancel() override;
+
+#if WITH_EDITOR
+	virtual const PCGUtils::FExtraCapture& GetExtraCapture() const override;
+	virtual PCGUtils::FExtraCapture& GetExtraCapture() override;
+
+	virtual const FPCGGraphExecutionInspection& GetInspection() const override;
+	virtual FPCGGraphExecutionInspection& GetInspection() override;
+
+	void RegisterDynamicTracking(const UPCGSettings* InSettings, const TArrayView<TPair<FPCGSelectionKey, bool>>& InDynamicKeysAndCulling);
+#endif // WITH_EDITOR
+
+private:
+	UPCGComponent* Component = nullptr;
+};
+
 UCLASS(BlueprintType, ClassGroup = (Procedural), meta = (BlueprintSpawnableComponent, PrioritizeCategories = "PCG"))
-class PCG_API UPCGComponent : public UActorComponent
+class PCG_API UPCGComponent : public UActorComponent, public IPCGGraphExecutionSource
 {
 	UPCGComponent(const FObjectInitializer& InObjectInitializer);
 
 	GENERATED_BODY()
 
+	friend class FPCGComponentExecutionState;
 	friend class UPCGManagedActors;
 	friend class UPCGSubsystem;
 	friend class FPCGGraphExecutor;
@@ -101,6 +133,11 @@ class PCG_API UPCGComponent : public UActorComponent
 	friend struct FPCGWorldPartitionBuilder;
 
 public:
+	/** ~Begin IPCGGraphExecutionSource interface */
+	virtual IPCGGraphExecutionState& GetExecutionState() override { return ExecutionState; }
+	virtual const IPCGGraphExecutionState& GetExecutionState() const override { return ExecutionState; }
+	/** ~End IPCGGraphExecutionSource interface */
+
 	/** ~Begin UObject interface */
 	virtual void Serialize(FArchive& Ar) override;
 	virtual void PostLoad() override;
@@ -140,7 +177,7 @@ public:
 	UPCGData* GetOriginalActorPCGData() const;
 
 	/** If this is a local component returns self, otherwise returns the original component. */
-	UPCGComponent* GetOriginalComponent();
+	UPCGComponent* GetOriginalComponent() const;
 	const UPCGComponent* GetConstOriginalComponent() const;
 
 	UPCGSchedulingPolicyBase* GetRuntimeGenSchedulingPolicy() const { return SchedulingPolicy; }
@@ -492,31 +529,28 @@ public:
 
 #if WITH_EDITOR
 
-	const FPCGGraphExecutionInspection& GetExecutionInspection() const { return ExecutionInspection; }
-
-	// todo_pcg: uncomment deprecation once this is moved to upcoming execution state
-	//UE_DEPRECATED(5.6, "Use FPCGGraphExecutionInspection::IsInspecting() instead")
+	UE_DEPRECATED(5.6, "Use FPCGGraphExecutionInspection::IsInspecting() instead")
 	bool IsInspecting() const;
 
-	//UE_DEPRECATED(5.6, "Use FPCGGraphExecutionInspection::EnableInspection() instead")
+	UE_DEPRECATED(5.6, "Use FPCGGraphExecutionInspection::EnableInspection() instead")
 	void EnableInspection();
 
-	//UE_DEPRECATED(5.6, "Use FPCGGraphExecutionInspection::DisableInspection() instead")
+	UE_DEPRECATED(5.6, "Use FPCGGraphExecutionInspection::DisableInspection() instead")
 	void DisableInspection();
 
-	//UE_DEPRECATED(5.6, "Use FPCGGraphExecutionInspection::StoreInspectionData() instead")
+	UE_DEPRECATED(5.6, "Use FPCGGraphExecutionInspection::StoreInspectionData() instead")
 	void StoreInspectionData(const FPCGStack* InStack, const UPCGNode* InNode, const PCGUtils::FCallTime* InTimer, const FPCGDataCollection& InInputData, const FPCGDataCollection& InOutputData, bool bUsedCache);
 
-	//UE_DEPRECATED(5.6, "Use FPCGGraphExecutionInspection::GetInspectionData() instead")
+	UE_DEPRECATED(5.6, "Use FPCGGraphExecutionInspection::GetInspectionData() instead")
 	const FPCGDataCollection* GetInspectionData(const FPCGStack& InStack) const;
 
-	//UE_DEPRECATED(5.6, "Use FPCGGraphExecutionInspection::ClearInspectionData() instead")
+	UE_DEPRECATED(5.6, "Use FPCGGraphExecutionInspection::ClearInspectionData() instead")
 	void ClearInspectionData(bool bClearPerNodeExecutionData = true);
 
-	//UE_DEPRECATED(5.6, "Use FPCGGraphExecutionInspection::WasNodeExecuted() instead")
+	UE_DEPRECATED(5.6, "Use FPCGGraphExecutionInspection::WasNodeExecuted() instead")
 	bool WasNodeExecuted(const UPCGNode* InNode, const FPCGStack& Stack) const;
 
-	//UE_DEPRECATED(5.6, "Use FPCGGraphExecutionInspection::NotifyNodeExecuted() instead")
+	UE_DEPRECATED(5.6, "Use FPCGGraphExecutionInspection::NotifyNodeExecuted() instead")
 	void NotifyNodeExecuted(const UPCGNode* InNode, const FPCGStack* InStack, const PCGUtils::FCallTime* InTimer, bool bNodeUsedCache);
 		
 	struct UE_DEPRECATED(5.6, "Use FPCGGraphExecutionInspection::FNodeExecutionNotificationData instead") NodeExecutedNotificationData
@@ -761,6 +795,9 @@ public:
 #if PCG_EXECUTION_CACHE_VALIDATION_ENABLED
 	bool bCanCreateExecutionCache = false;
 #endif
+
+private:
+	FPCGComponentExecutionState ExecutionState;
 };
 
 /** Used to store generated resources data during RerunConstructionScripts */
diff --git a/Engine/Plugins/PCG/Source/PCG/Public/PCGContext.h b/Engine/Plugins/PCG/Source/PCG/Public/PCGContext.h
--- a/Engine/Plugins/PCG/Source/PCG/Public/PCGContext.h
+++ b/Engine/Plugins/PCG/Source/PCG/Public/PCGContext.h
@@ -4,6 +4,7 @@
 
 #include "PCGData.h"
 #include "PCGNode.h" // IWYU pragma: keep
+#include "PCGGraphExecutionStateInterface.h"
 #include "Helpers/PCGAsyncState.h"
 #include "Utils/PCGExtraCapture.h"
 
@@ -19,6 +20,7 @@ class UPCGGraphInterface;
 class UPCGSettingsInterface;
 class UPCGSpatialData;
 struct FPCGGetFromCacheParams;
+struct FPCGInitializeElementParams;
 struct FPCGSettingsOverridableParam;
 struct FPCGStack;
 struct FPCGStoreInCacheParams;
@@ -79,10 +81,25 @@ struct PCG_API FPCGContext
 	friend struct FPCGGraphActiveTask;
 	friend struct FPCGGraphTask;
 
+	// Note: We need to explicitly disable warnings on these constructors/operators for clang to be happy with deprecated variables
+PRAGMA_DISABLE_DEPRECATION_WARNINGS
+	FPCGContext() = default;
+	FPCGContext(const FPCGContext&) = default;
+	FPCGContext(FPCGContext&&) = default;
+	FPCGContext& operator=(const FPCGContext&) = default;
+	FPCGContext& operator=(FPCGContext&&) = default;
+PRAGMA_ENABLE_DEPRECATION_WARNINGS
+
 	virtual ~FPCGContext();
 
+	void InitFromParams(const FPCGInitializeElementParams& InParams);
+
 	FPCGDataCollection InputData;
 	FPCGDataCollection OutputData;
+
+	TWeakInterfacePtr<IPCGGraphExecutionSource> ExecutionSource = nullptr;
+
+	UE_DEPRECATED(5.6, "Use ExecutionSource instead")
 	TWeakObjectPtr<UPCGComponent> SourceComponent = nullptr;
 
 	/** Used to track when data transformations (such as CPU readback) occurred on the input data collection as part of the element prepare data. */
@@ -148,7 +165,12 @@ struct PCG_API FPCGContext
 	FPCGTaskId GetGraphExecutionTaskId() const;
 
 	FString GetTaskName() const;
+
+	UE_DEPRECATED(5.6, "Use GetExecutionSourceName() instead")
 	FString GetComponentName() const;
+
+	FString GetExecutionSourceName() const;
+
 	bool ShouldStop() const { return AsyncState.ShouldStop(); }
 
 	AActor* GetTargetActor(const UPCGSpatialData* InSpatialData) const;
diff --git a/Engine/Plugins/PCG/Source/PCG/Public/PCGElement.h b/Engine/Plugins/PCG/Source/PCG/Public/PCGElement.h
--- a/Engine/Plugins/PCG/Source/PCG/Public/PCGElement.h
+++ b/Engine/Plugins/PCG/Source/PCG/Public/PCGElement.h
@@ -2,14 +2,16 @@
 
 #pragma once
 
+#include "PCGData.h"
+
 #include "Logging/LogVerbosity.h"
 #include "Templates/SharedPointer.h"
-#include "UObject/WeakObjectPtrTemplates.h"
+#include "UObject/WeakInterfacePtr.h"
 
+class IPCGGraphExecutionSource;
 class UPCGSettingsInterface;
 struct FPCGContext;
 struct FPCGCrc;
-struct FPCGDataCollection;
 
 class IPCGElement;
 class IPCGGraphCache;
@@ -42,7 +44,7 @@ namespace EPCGElementLogMode
 	UE_LOG(LogPCG, \
 		Verbosity, \
 		TEXT("[%s - %s]: %s"), \
-		(CustomContext) ? *((CustomContext)->GetComponentName()) : TEXT("UnknownComponent"), \
+		(CustomContext) ? *((CustomContext)->GetExecutionSourceName()) : TEXT("UnknownSource"), \
 		(CustomContext) ? *((CustomContext)->GetTaskName()) : TEXT("UnknownTask"), \
 		*Message.ToString())
 
@@ -90,6 +92,39 @@ namespace PCGElementHelpers
 	bool PCG_API SplitDataPerPrimaryPin(const UPCGSettings* Settings, const FPCGDataCollection& InCollection, EPCGElementExecutionLoopMode Mode, TArray<FPCGDataCollection>& OutPrimaryCollections, FPCGDataCollection& OutCommonCollection);
 };
 
+/**
+ * Struct containing params to initialize IPCGElement
+ */
+struct FPCGInitializeElementParams
+{
+	FPCGInitializeElementParams(const FPCGDataCollection* InInputData, TWeakInterfacePtr<IPCGGraphExecutionSource> InExecutionSource, const UPCGNode* InNode)
+		: InputData(InInputData)
+		, ExecutionSource(InExecutionSource)
+		, Node(InNode)
+	{
+		check(InputData);
+	}
+	
+	const FPCGDataCollection* InputData;
+	TWeakInterfacePtr<IPCGGraphExecutionSource> ExecutionSource;
+	const UPCGNode* Node;
+};
+
+struct FPCGGetDependenciesCrcParams
+{
+	FPCGGetDependenciesCrcParams(const FPCGDataCollection* InInputData, const UPCGSettings* InSettings, IPCGGraphExecutionSource* InExecutionSource)
+		: InputData(InInputData)
+		, Settings(InSettings)
+		, ExecutionSource(InExecutionSource)
+	{
+		check(InputData);
+	}
+
+	const FPCGDataCollection* InputData;
+	const UPCGSettings* Settings;
+	IPCGGraphExecutionSource* ExecutionSource;
+};
+
 /**
 * Base class for the processing bit of a PCG node/settings
 */
@@ -99,6 +134,9 @@ public:
 	virtual ~IPCGElement() = default;
 
 	/** Creates a custom context object paired to this element */
+	virtual FPCGContext* Initialize(const FPCGInitializeElementParams& InParams);
+
+	UE_DEPRECATED(5.6, "Use/Implement version with FPCGInitializeElementParams parameter instead")
 	virtual FPCGContext* Initialize(const FPCGDataCollection& InputData, TWeakObjectPtr<UPCGComponent> SourceComponent, const UPCGNode* Node);
 
 	/** Returns true if the element, in its current phase can be executed only from the main thread */
@@ -121,10 +159,13 @@ public:
 	 * external data) changes then this value should change. For some elements it is inefficient or not possible to output a Crc here. These can return an invalid
 	 * Crc and the Crc can either be computed during execution, or afterwards based on output data.
 	 */
+	virtual void GetDependenciesCrc(const FPCGGetDependenciesCrcParams& InParams, FPCGCrc& OutCrc) const;
+
+	UE_DEPRECATED(5.6, "Use/Implement version with FPCGGetDependenciesCrcParams parameter instead")
 	virtual void GetDependenciesCrc(const FPCGDataCollection& InInput, const UPCGSettings* InSettings, UPCGComponent* InComponent, FPCGCrc& OutCrc) const;
 
 	/** Gather input data (pre-context creation) and tries to retrieve matching data from the cache if the element is cacheable. */
-	EPCGCachingStatus RetrieveResultsFromCache(IPCGGraphCache* Cache, const UPCGNode* Node, const FPCGDataCollection& Input, UPCGComponent* Component, FPCGDataCollection& Output, FPCGCrc* OutCrc = nullptr) const;
+	EPCGCachingStatus RetrieveResultsFromCache(IPCGGraphCache* Cache, const UPCGNode* Node, const FPCGDataCollection& Input, IPCGGraphExecutionSource* ExecutionSource, FPCGDataCollection& Output, FPCGCrc* OutCrc = nullptr) const;
 
 	/** Public function that executes the element on the appropriately created context.
 	* The caller should call the Execute function until it returns true.
diff --git a/Engine/Plugins/PCG/Source/PCG/Public/PCGGraphExecutionInspection.h b/Engine/Plugins/PCG/Source/PCG/Public/PCGGraphExecutionInspection.h
--- a/Engine/Plugins/PCG/Source/PCG/Public/PCGGraphExecutionInspection.h
+++ b/Engine/Plugins/PCG/Source/PCG/Public/PCGGraphExecutionInspection.h
@@ -11,7 +11,7 @@
 #include "UObject/ObjectKey.h"
 
 #if WITH_EDITOR
-class FPCGGraphExecutionInspection
+class PCG_API FPCGGraphExecutionInspection
 {
 public:
 	friend struct FPCGComponentInstanceData;
diff --git a/Engine/Plugins/PCG/Source/PCG/Public/PCGGraphExecutionStateInterface.h b/Engine/Plugins/PCG/Source/PCG/Public/PCGGraphExecutionStateInterface.h
new file mode 100644
--- /dev/null
+++ b/Engine/Plugins/PCG/Source/PCG/Public/PCGGraphExecutionStateInterface.h
@@ -0,0 +1,95 @@
+// Copyright Epic Games, Inc. All Rights Reserved.
+
+#pragma once
+
+#include "PCGSettings.h"
+#include "Utils/PCGExtraCapture.h"
+
+#include "CoreMinimal.h"
+#include "UObject/Interface.h"
+
+#include "PCGGraphExecutionStateInterface.generated.h"
+
+class UPCGData;
+class FPCGGraphExecutionInspection;
+class IPCGGraphExecutionSource;
+class UPCGGraph;
+class UPCGGraphInstance;
+
+namespace PCGUtils
+{
+	class FExtraCapture;
+}
+
+/**
+* Interface returned by a IPCGGraphExecutionSource that is queried / updated during execution of a PCG Graph.
+*/
+class IPCGGraphExecutionState
+{
+public:
+
+	/** Returns a UPCGData representation of the ExecutionState. */
+	virtual UPCGData* GetSelfData() const = 0;
+
+	/** Returns a Seed for graph execution. */
+	virtual int32 GetSeed() const = 0;
+
+	/** Returns a Debug name that can be used for logging. */
+	virtual FString GetDebugName() const = 0;
+
+	/** Returns a World, can be null */
+	virtual UWorld* GetWorld() const = 0;
+
+	/** Returns a Transform if the ExecutionState is a spatial one. */
+	virtual FTransform GetTransform() const = 0;
+
+	/** Returns the ExecutionState bounds if the ExecutionState is a spatial one. */
+	virtual FBox GetBounds() const = 0;
+
+	/** Returns the UPCGGraph this ExecutionState is executing. */
+	virtual UPCGGraph* GetGraph() const = 0;
+
+	/** Returns the UPCGGrahInstance this ExecutionState is executing. */
+	virtual UPCGGraphInstance* GetGraphInstance() const = 0;
+
+	/** Cancel execution of this ExecutionState. */
+	virtual void Cancel() = 0;
+
+	/** Notify ExecutionState that its execution is being aborted */
+	virtual void OnGraphExecutionAborted(bool bQuiet = false, bool bCleanupUnusedResources = true) = 0;
+
+#if WITH_EDITOR
+
+	/** Returns the FExtraCapture object for this ExecutionState */
+	virtual const PCGUtils::FExtraCapture& GetExtraCapture() const = 0;
+
+	virtual PCGUtils::FExtraCapture& GetExtraCapture() = 0;
+
+	/** Returns the FPCGGraphExecutionInspection object for this ExecutionState */
+	virtual const FPCGGraphExecutionInspection& GetInspection() const = 0;
+
+	virtual FPCGGraphExecutionInspection& GetInspection() = 0;
+
+	/** Register tracking dependencies, so ExecutionState can be updated when they change */
+	virtual void RegisterDynamicTracking(const UPCGSettings* InSettings, const TArrayView<TPair<FPCGSelectionKey, bool>>& InDynamicKeysAndCulling) = 0;
+
+#endif
+};
+
+UINTERFACE(BlueprintType, MinimalAPI)
+class UPCGGraphExecutionSource : public UInterface
+{
+	GENERATED_UINTERFACE_BODY()
+};
+
+/**
+* Interface used by the FPCGGraphExecutor to get an IPCGGraphExecutionState used to query/update execution.
+*/
+class IPCGGraphExecutionSource
+{
+	GENERATED_IINTERFACE_BODY()
+
+public:
+	virtual IPCGGraphExecutionState& GetExecutionState() = 0;
+	virtual const IPCGGraphExecutionState& GetExecutionState() const = 0;
+};
diff --git a/Engine/Plugins/PCG/Source/PCG/Public/PCGSettings.h b/Engine/Plugins/PCG/Source/PCG/Public/PCGSettings.h
--- a/Engine/Plugins/PCG/Source/PCG/Public/PCGSettings.h
+++ b/Engine/Plugins/PCG/Source/PCG/Public/PCGSettings.h
@@ -240,8 +240,8 @@ public:
 	virtual bool UseSeed() const { return bUseSeed; }
 	PRAGMA_ENABLE_DEPRECATION_WARNINGS
 
-	// Get the seed, combined with optional PCGComponent seed
-	int GetSeed(const UPCGComponent* InSourceComponent = nullptr) const;
+	// Get the seed, combined with optional ExecutionSource seed
+	int GetSeed(const IPCGGraphExecutionSource* InExecutionSource = nullptr) const;
 
 	UE_DEPRECATED(5.6, "Use GetSettingsCrc instead")
 	const FPCGCrc& GetCachedCrc() const;
diff --git a/Engine/Plugins/PCG/Source/PCG/Public/PCGSubgraph.h b/Engine/Plugins/PCG/Source/PCG/Public/PCGSubgraph.h
--- a/Engine/Plugins/PCG/Source/PCG/Public/PCGSubgraph.h
+++ b/Engine/Plugins/PCG/Source/PCG/Public/PCGSubgraph.h
@@ -189,7 +189,7 @@ protected:
 class FPCGSubgraphElement : public IPCGElement
 {
 public:
-	virtual FPCGContext* Initialize(const FPCGDataCollection& InputData, TWeakObjectPtr<UPCGComponent> SourceComponent, const UPCGNode* Node) override;
+	virtual FPCGContext* Initialize(const FPCGInitializeElementParams& InParams) override;
 	virtual bool IsCacheable(const UPCGSettings* InSettings) const override { return false; }
 
 protected:
diff --git a/Engine/Plugins/PCG/Source/PCG/Public/Utils/PCGExtraCapture.h b/Engine/Plugins/PCG/Source/PCG/Public/Utils/PCGExtraCapture.h
--- a/Engine/Plugins/PCG/Source/PCG/Public/Utils/PCGExtraCapture.h
+++ b/Engine/Plugins/PCG/Source/PCG/Public/Utils/PCGExtraCapture.h
@@ -98,7 +98,7 @@ namespace PCGUtils
 
 		const TCapturedMessageMap& GetCapturedMessages() const { return CapturedMessages; }
 
-		FCallTreeInfo CalculateCallTreeInfo(const UPCGComponent* Component, const FPCGStack& RootStack) const;
+		FCallTreeInfo CalculateCallTreeInfo(const IPCGGraphExecutionSource* ExecutionSource, const FPCGStack& RootStack) const;
 
 	private:
 		TCapturedMessageMap CapturedMessages;
diff --git a/Engine/Plugins/PCG/Source/PCGEditor/Private/DataVisualizations/PCGSpatialDataVisualization.cpp b/Engine/Plugins/PCG/Source/PCGEditor/Private/DataVisualizations/PCGSpatialDataVisualization.cpp
--- a/Engine/Plugins/PCG/Source/PCGEditor/Private/DataVisualizations/PCGSpatialDataVisualization.cpp
+++ b/Engine/Plugins/PCG/Source/PCGEditor/Private/DataVisualizations/PCGSpatialDataVisualization.cpp
@@ -2,6 +2,7 @@
 
 #include "DataVisualizations/PCGSpatialDataVisualization.h"
 
+#include "PCGComponent.h"
 #include "PCGContext.h"
 #include "PCGData.h"
 #include "PCGDebug.h"
@@ -91,6 +92,12 @@ void IPCGSpatialDataVisualization::ExecuteDebugDisplay(FPCGContext* Context, con
 		return;
 	}
 
+	UPCGComponent* SourceComponent = Cast<UPCGComponent>(Context->ExecutionSource.Get());
+	if (!SourceComponent)
+	{
+		return;
+	}
+
 	const FPCGDebugVisualizationSettings& DebugSettings = SettingsInterface->DebugSettings;
 	UStaticMesh* Mesh = DebugSettings.PointMesh.LoadSynchronous();
 
@@ -213,7 +220,7 @@ void IPCGSpatialDataVisualization::ExecuteDebugDisplay(FPCGContext* Context, con
 			continue;
 		}
 
-		UInstancedStaticMeshComponent* ISMC = UPCGActorHelpers::GetOrCreateISMC(TargetActor, Context->SourceComponent.Get(), Params[Direction], Context);
+		UInstancedStaticMeshComponent* ISMC = UPCGActorHelpers::GetOrCreateISMC(TargetActor, SourceComponent, Params[Direction], Context);
 		check(ISMC && ISMC->NumCustomDataFloats == NumCustomData);
 
 		ISMC->ComponentTags.AddUnique(PCGHelpers::DefaultPCGDebugTag);
diff --git a/Engine/Plugins/PCG/Source/PCGEditor/Private/Nodes/PCGEditorGraphNodeBase.cpp b/Engine/Plugins/PCG/Source/PCGEditor/Private/Nodes/PCGEditorGraphNodeBase.cpp
--- a/Engine/Plugins/PCG/Source/PCGEditor/Private/Nodes/PCGEditorGraphNodeBase.cpp
+++ b/Engine/Plugins/PCG/Source/PCGEditor/Private/Nodes/PCGEditorGraphNodeBase.cpp
@@ -84,7 +84,7 @@ namespace PCGEditorGraphNodeBase
 			if (!Settings || !Settings->IsA<UPCGRerouteSettings>())
 			{
 				// Task will be displayed as active if it was executed.
-				return InComponentBeingDebugged->WasNodeExecuted(PCGNode, *InStackBeingInspected);
+				return InComponentBeingDebugged->GetExecutionState().GetInspection().WasNodeExecuted(PCGNode, *InStackBeingInspected);
 			}
 			else
 			{
diff --git a/Engine/Plugins/PCG/Source/PCGEditor/Private/PCGEditor.cpp b/Engine/Plugins/PCG/Source/PCGEditor/Private/PCGEditor.cpp
--- a/Engine/Plugins/PCG/Source/PCGEditor/Private/PCGEditor.cpp
+++ b/Engine/Plugins/PCG/Source/PCGEditor/Private/PCGEditor.cpp
@@ -279,9 +279,9 @@ void FPCGEditor::SetStackBeingInspected(const FPCGStack& FullStack)
 
 	if (NewComponent && NewComponent != LastComponent)
 	{
-		if (LastComponent && LastComponent->IsInspecting())
+		if (LastComponent && LastComponent->GetExecutionState().GetInspection().IsInspecting())
 		{
-			LastComponent->DisableInspection();
+			LastComponent->GetExecutionState().GetInspection().DisableInspection();
 		}
 
 		LastValidPCGComponentBeingInspected = NewComponent;
@@ -342,12 +342,12 @@ void FPCGEditor::UpdateAfterInspectedStackChanged()
 		// Implementation note: if we're inspecting and have not pre-run the graph, then it probably makes sense to enable inspection by default. 
 		// TODO This could be selected with a cvar though.
 		const bool bHasBeenGeneratedThisSession = Component->bGenerated && Component->WasGeneratedThisSession();
-		const bool bWasInspecting = Component->IsInspecting();
+		const bool bWasInspecting = Component->GetExecutionState().GetInspection().IsInspecting();
 		const bool bNeedsInspection = Algo::AnyOf(AttributesWidgets, [](const TSharedPtr<SPCGEditorGraphAttributeListView>& ALV) { return ALV->GetNodeBeingInspected() != nullptr; });
 
 		if (!bHasBeenGeneratedThisSession || (bNeedsInspection && !bWasInspecting))
 		{
-			Component->EnableInspection();
+			Component->GetExecutionState().GetInspection().EnableInspection();
 			UpdateDebugAfterComponentSelection(Component, Component, true);
 		}
 	}
@@ -2769,17 +2769,17 @@ void FPCGEditor::OnClose()
 
 	if (PCGComponentBeingInspected.IsValid())
 	{
-		if (PCGComponentBeingInspected->IsInspecting())
+		if (PCGComponentBeingInspected->GetExecutionState().GetInspection().IsInspecting())
 		{
-			PCGComponentBeingInspected->DisableInspection();
+			PCGComponentBeingInspected->GetExecutionState().GetInspection().DisableInspection();
 		}
 	}
 
 	if (LastValidPCGComponentBeingInspected.IsValid())
 	{
-		if (LastValidPCGComponentBeingInspected->IsInspecting())
+		if (LastValidPCGComponentBeingInspected->GetExecutionState().GetInspection().IsInspecting())
 		{
-			LastValidPCGComponentBeingInspected->DisableInspection();
+			LastValidPCGComponentBeingInspected->GetExecutionState().GetInspection().DisableInspection();
 		}
 	}
 
diff --git a/Engine/Plugins/PCG/Source/PCGEditor/Private/PCGEditorGraphSchema.cpp b/Engine/Plugins/PCG/Source/PCGEditor/Private/PCGEditorGraphSchema.cpp
--- a/Engine/Plugins/PCG/Source/PCGEditor/Private/PCGEditorGraphSchema.cpp
+++ b/Engine/Plugins/PCG/Source/PCGEditor/Private/PCGEditorGraphSchema.cpp
@@ -919,7 +919,7 @@ bool FPCGEditorConnectionDrawingPolicy::UpdateParamsIfDebugging(UEdGraphPin* Out
 	const FPCGStack* PCGStack = Editor ? Editor->GetStackBeingInspected() : nullptr;
 	UPCGComponent* PCGComponent = Editor ? Editor->GetPCGComponentBeingInspected() : nullptr;
 
-	if (!Editor || !PCGStack || !PCGComponent || !PCGComponent->IsInspecting())
+	if (!Editor || !PCGStack || !PCGComponent || !PCGComponent->GetExecutionState().GetInspection().IsInspecting())
 	{
 		return false;
 	}
@@ -930,7 +930,7 @@ bool FPCGEditorConnectionDrawingPolicy::UpdateParamsIfDebugging(UEdGraphPin* Out
 	StackFrames.Emplace(NodeToInspect);
 	StackFrames.Emplace(PinToInspect);
 
-	if (const FPCGDataCollection* DataCollection = PCGComponent->GetInspectionData(Stack))
+	if (const FPCGDataCollection* DataCollection = PCGComponent->GetExecutionState().GetInspection().GetInspectionData(Stack))
 	{
 		if (DataCollection->TaggedData.Num() > 1)
 		{
diff --git a/Engine/Plugins/PCG/Source/PCGEditor/Private/PCGEditorModule.cpp b/Engine/Plugins/PCG/Source/PCGEditor/Private/PCGEditorModule.cpp
--- a/Engine/Plugins/PCG/Source/PCGEditor/Private/PCGEditorModule.cpp
+++ b/Engine/Plugins/PCG/Source/PCGEditor/Private/PCGEditorModule.cpp
@@ -794,7 +794,7 @@ void FPCGEditorModule::ClearExecutionMetadata(UPCGComponent* InComponent)
 	if (InComponent)
 	{
 		ClearExecutedStacks(InComponent);
-		InComponent->ClearInspectionData();
+		InComponent->GetExecutionState().GetInspection().ClearInspectionData();
 
 		GetNodeVisualLogsMutable().ClearLogs(InComponent);
 	}
diff --git a/Engine/Plugins/PCG/Source/PCGEditor/Private/Widgets/SPCGEditorGraphAttributeListView.cpp b/Engine/Plugins/PCG/Source/PCGEditor/Private/Widgets/SPCGEditorGraphAttributeListView.cpp
--- a/Engine/Plugins/PCG/Source/PCGEditor/Private/Widgets/SPCGEditorGraphAttributeListView.cpp
+++ b/Engine/Plugins/PCG/Source/PCGEditor/Private/Widgets/SPCGEditorGraphAttributeListView.cpp
@@ -730,7 +730,7 @@ const FPCGDataCollection* SPCGEditorGraphAttributeListView::GetInspectionData(co
 	StackFrames.Emplace(PCGNode);
 	StackFrames.Emplace(Pin);
 
-	return PCGComponent->GetInspectionData(Stack);
+	return PCGComponent->GetExecutionState().GetInspection().GetInspectionData(Stack);
 }
 
 const FPCGDataCollection* SPCGEditorGraphAttributeListView::GetInspectionData() const
diff --git a/Engine/Plugins/PCG/Source/PCGEditor/Private/Widgets/SPCGEditorGraphLogView.cpp b/Engine/Plugins/PCG/Source/PCGEditor/Private/Widgets/SPCGEditorGraphLogView.cpp
--- a/Engine/Plugins/PCG/Source/PCGEditor/Private/Widgets/SPCGEditorGraphLogView.cpp
+++ b/Engine/Plugins/PCG/Source/PCGEditor/Private/Widgets/SPCGEditorGraphLogView.cpp
@@ -275,9 +275,9 @@ FReply SPCGEditorGraphLogView::Clear()
 		return FReply::Handled();
 	}
 
-	if (const UPCGComponent* Component = GetPCGComponent().Get())
+	if (IPCGGraphExecutionSource* ExecutionSource = GetExecutionSource())
 	{
-		Component->ExtraCapture.ResetCapturedMessages();
+		ExecutionSource->GetExecutionState().GetExtraCapture().ResetCapturedMessages();
 		Refresh();
 	}
 
@@ -300,9 +300,9 @@ FReply SPCGEditorGraphLogView::Refresh()
 		return FReply::Handled();
 	}
 
-	const UPCGComponent* Component = GetPCGComponent().Get();
+	const IPCGGraphExecutionSource* ExecutionSource = GetExecutionSource();
 
-	if (!Component)
+	if (!ExecutionSource)
 	{
 		return FReply::Handled();
 	}
@@ -356,14 +356,14 @@ void SPCGEditorGraphLogView::CreateAndAddItem(const UPCGEditorGraphNode* InPCGEd
 
 	check(PCGEditor);
 
-	const UPCGComponent* Component = GetPCGComponent().Get();
+	IPCGGraphExecutionSource* ExecutionSource = GetExecutionSource();
 
-	if (!Component || !InPCGEditorNode || !InPCGNode)
+	if (!ExecutionSource || !InPCGEditorNode || !InPCGNode)
 	{
 		return;
 	}
 
-	const TArray<PCGUtils::FCapturedMessage>* CapturedMessagesPtr = Component->ExtraCapture.GetCapturedMessages().Find(InPCGNode);
+	const TArray<PCGUtils::FCapturedMessage>* CapturedMessagesPtr = ExecutionSource->GetExecutionState().GetExtraCapture().GetCapturedMessages().Find(InPCGNode);
 	if (!CapturedMessagesPtr)
 	{
 		return;
@@ -405,7 +405,7 @@ void SPCGEditorGraphLogView::CreateAndAddItem(const UPCGEditorGraphNode* InPCGEd
 	}
 }
 
-TWeakObjectPtr<UPCGComponent> SPCGEditorGraphLogView::GetPCGComponent() const
+IPCGGraphExecutionSource* SPCGEditorGraphLogView::GetExecutionSource() const
 {
 	return PCGEditorPtr.IsValid() ? PCGEditorPtr.Pin()->GetPCGComponentBeingInspected() : nullptr;
 }
diff --git a/Engine/Plugins/PCG/Source/PCGEditor/Private/Widgets/SPCGEditorGraphLogView.h b/Engine/Plugins/PCG/Source/PCGEditor/Private/Widgets/SPCGEditorGraphLogView.h
--- a/Engine/Plugins/PCG/Source/PCGEditor/Private/Widgets/SPCGEditorGraphLogView.h
+++ b/Engine/Plugins/PCG/Source/PCGEditor/Private/Widgets/SPCGEditorGraphLogView.h
@@ -9,9 +9,9 @@
 class STableViewBase;
 
 class FPCGEditor;
-class UPCGComponent;
 class UPCGEditorGraphNode;
 class UPCGEditorGraph;
+class IPCGGraphExecutionSource;
 class UPCGNode;
 struct FPCGStack;
 
@@ -82,8 +82,8 @@ private:
 	/** Cached PCGGraph being viewed */
 	UPCGEditorGraph* PCGEditorGraph = nullptr;
 
-	/** Cached PCGComponent being viewed */
-	TWeakObjectPtr<UPCGComponent> GetPCGComponent() const;
+	/** Cached PCG Execution source being viewed */
+	IPCGGraphExecutionSource* GetExecutionSource() const;
 
 	TSharedPtr<SHeaderRow> ListViewHeader;
 	TSharedPtr<SListView<PCGLogListViewItemPtr>> ListView;
diff --git a/Engine/Plugins/PCG/Source/PCGEditor/Private/Widgets/SPCGEditorGraphProfilingView.cpp b/Engine/Plugins/PCG/Source/PCGEditor/Private/Widgets/SPCGEditorGraphProfilingView.cpp
--- a/Engine/Plugins/PCG/Source/PCGEditor/Private/Widgets/SPCGEditorGraphProfilingView.cpp
+++ b/Engine/Plugins/PCG/Source/PCGEditor/Private/Widgets/SPCGEditorGraphProfilingView.cpp
@@ -701,8 +701,8 @@ FReply SPCGEditorGraphProfilingView::Refresh()
 		return FReply::Handled();
 	}
 
-	const UPCGComponent* Component = GetPCGComponent().Get();
-	if (!Component)
+	const IPCGGraphExecutionSource* ExecutionSource = GetExecutionSource();
+	if (!ExecutionSource)
 	{
 		return FReply::Handled();
 	}
@@ -716,7 +716,7 @@ FReply SPCGEditorGraphProfilingView::Refresh()
 		EditorNodeLookup.Add(EditorNode->GetPCGNode(), EditorNode);
 	}
 
-	PCGUtils::FCallTreeInfo TreeInfo = Component->ExtraCapture.CalculateCallTreeInfo(Component, PCGStack);
+	PCGUtils::FCallTreeInfo TreeInfo = ExecutionSource->GetExecutionState().GetExtraCapture().CalculateCallTreeInfo(ExecutionSource, PCGStack);
 
 	ListViewItems.Reserve(TreeInfo.Children.Num());
 
@@ -799,7 +799,7 @@ void SPCGEditorGraphProfilingView::OnDebugStackChanged(const FPCGStack& InPCGSta
 	RequestRefresh();
 }
 
-TWeakObjectPtr<UPCGComponent> SPCGEditorGraphProfilingView::GetPCGComponent() const
+IPCGGraphExecutionSource* SPCGEditorGraphProfilingView::GetExecutionSource() const
 {
 	return PCGEditorPtr.IsValid() ? PCGEditorPtr.Pin()->GetPCGComponentBeingInspected() : nullptr;
 }
diff --git a/Engine/Plugins/PCG/Source/PCGEditor/Private/Widgets/SPCGEditorGraphProfilingView.h b/Engine/Plugins/PCG/Source/PCGEditor/Private/Widgets/SPCGEditorGraphProfilingView.h
--- a/Engine/Plugins/PCG/Source/PCGEditor/Private/Widgets/SPCGEditorGraphProfilingView.h
+++ b/Engine/Plugins/PCG/Source/PCGEditor/Private/Widgets/SPCGEditorGraphProfilingView.h
@@ -87,8 +87,8 @@ private:
 	/** Called when user changes commits text to the search box */
 	void OnSearchTextCommitted(const FText& InText, ETextCommit::Type InCommitType);
 
-	/** Returns the current PCG component */
-	TWeakObjectPtr<UPCGComponent> GetPCGComponent() const;
+	/** Returns the current PCG execution source */
+	IPCGGraphExecutionSource* GetExecutionSource() const;
 
 	/** Pointer back to the PCG editor that owns us */
 	TWeakPtr<FPCGEditor> PCGEditorPtr;
diff --git a/Engine/Plugins/PCGInterops/PCGExternalDataInterop/Source/PCGExternalDataInterop/Private/Elements/PCGLoadAlembicElement.cpp b/Engine/Plugins/PCGInterops/PCGExternalDataInterop/Source/PCGExternalDataInterop/Private/Elements/PCGLoadAlembicElement.cpp
--- a/Engine/Plugins/PCGInterops/PCGExternalDataInterop/Source/PCGExternalDataInterop/Private/Elements/PCGLoadAlembicElement.cpp
+++ b/Engine/Plugins/PCGInterops/PCGExternalDataInterop/Source/PCGExternalDataInterop/Private/Elements/PCGLoadAlembicElement.cpp
@@ -74,12 +74,12 @@ void UPCGLoadAlembicSettings::SetupFromStandard(EPCGLoadAlembicStandardSetup InS
 	}
 }
 
-void FPCGLoadAlembicElement::GetDependenciesCrc(const FPCGDataCollection& InInput, const UPCGSettings* InSettings, UPCGComponent* InComponent, FPCGCrc& OutCrc) const
+void FPCGLoadAlembicElement::GetDependenciesCrc(const FPCGGetDependenciesCrcParams& InParams, FPCGCrc& OutCrc) const
 {
 	FPCGCrc Crc;
-	FPCGExternalDataElement::GetDependenciesCrc(InInput, InSettings, InComponent, Crc);
+	FPCGExternalDataElement::GetDependenciesCrc(InParams, Crc);
 
-	if (const UPCGLoadAlembicSettings* Settings = Cast<UPCGLoadAlembicSettings>(InSettings))
+	if (const UPCGLoadAlembicSettings* Settings = Cast<UPCGLoadAlembicSettings>(InParams.Settings))
 	{
 		const FDateTime FileTimeStamp = IFileManager::Get().GetTimeStamp(*Settings->AlembicFilePath.FilePath);
 		if (FileTimeStamp != FDateTime::MinValue())
diff --git a/Engine/Plugins/PCGInterops/PCGExternalDataInterop/Source/PCGExternalDataInterop/Public/Elements/PCGLoadAlembicElement.h b/Engine/Plugins/PCGInterops/PCGExternalDataInterop/Source/PCGExternalDataInterop/Public/Elements/PCGLoadAlembicElement.h
--- a/Engine/Plugins/PCGInterops/PCGExternalDataInterop/Source/PCGExternalDataInterop/Public/Elements/PCGLoadAlembicElement.h
+++ b/Engine/Plugins/PCGInterops/PCGExternalDataInterop/Source/PCGExternalDataInterop/Public/Elements/PCGLoadAlembicElement.h
@@ -69,7 +69,7 @@ public:
 class PCGEXTERNALDATAINTEROP_API FPCGLoadAlembicElement : public FPCGExternalDataElement
 {
 public:
-	virtual void GetDependenciesCrc(const FPCGDataCollection& InInput, const UPCGSettings* InSettings, UPCGComponent* InComponent, FPCGCrc& OutCrc) const override;
+	virtual void GetDependenciesCrc(const FPCGGetDependenciesCrcParams& InParams, FPCGCrc& OutCrc) const override;
 protected:
 	virtual FPCGContext* CreateContext() override;
 	virtual bool PrepareLoad(FPCGExternalDataContext* Context) const override;
diff --git a/Engine/Plugins/PCGInterops/PCGExternalDataInterop/Source/PCGExternalDataInteropEditor/Private/PCGLoadAlembic.cpp b/Engine/Plugins/PCGInterops/PCGExternalDataInterop/Source/PCGExternalDataInteropEditor/Private/PCGLoadAlembic.cpp
--- a/Engine/Plugins/PCGInterops/PCGExternalDataInterop/Source/PCGExternalDataInteropEditor/Private/PCGLoadAlembic.cpp
+++ b/Engine/Plugins/PCGInterops/PCGExternalDataInterop/Source/PCGExternalDataInteropEditor/Private/PCGLoadAlembic.cpp
@@ -2,6 +2,7 @@
 
 #include "PCGLoadAlembic.h"
 #include "PCGAssetExporterUtils.h"
+#include "PCGGraphExecutionStateInterface.h"
 #include "PCGModule.h"
 
 #include "Elements/PCGLoadAlembicElement.h"
@@ -64,7 +65,8 @@ void UPCGLoadAlembicFunctionLibrary::LoadAlembicFileToPCGInternal(const FPCGLoad
 	Settings->AttributeMapping = InSettings.AttributeMapping;
 
 	FPCGLoadAlembicElement Element;
-	FPCGContext* Context = Element.Initialize(FPCGDataCollection(), TWeakObjectPtr<UPCGComponent>(), nullptr);
+	FPCGDataCollection InputData;
+	FPCGContext* Context = Element.Initialize(FPCGInitializeElementParams(&InputData, nullptr, nullptr));
 	Context->AsyncState.NumAvailableTasks = 1; // TODO : change this
 	FPCGTaggedData& DataSettings = Context->InputData.TaggedData.Emplace_GetRef();
 	DataSettings.Data = Settings;
diff --git a/Engine/Plugins/PCGInterops/PCGGeometryScriptInterop/Source/PCGGeometryScriptInterop/Private/Editor/PCGDynamicMeshDataVisualization.cpp b/Engine/Plugins/PCGInterops/PCGGeometryScriptInterop/Source/PCGGeometryScriptInterop/Private/Editor/PCGDynamicMeshDataVisualization.cpp
--- a/Engine/Plugins/PCGInterops/PCGGeometryScriptInterop/Source/PCGGeometryScriptInterop/Private/Editor/PCGDynamicMeshDataVisualization.cpp
+++ b/Engine/Plugins/PCGInterops/PCGGeometryScriptInterop/Source/PCGGeometryScriptInterop/Private/Editor/PCGDynamicMeshDataVisualization.cpp
@@ -18,7 +18,7 @@
 void FPCGDynamicMeshDataVisualization::ExecuteDebugDisplay(FPCGContext* Context, const UPCGSettingsInterface* SettingsInterface, const UPCGData* Data, AActor* TargetActor) const
 {
 	const UPCGDynamicMeshData* DynMeshData = CastChecked<UPCGDynamicMeshData>(Data);
-	UPCGComponent* SourcePCGComponent = Context->SourceComponent.Get();
+	UPCGComponent* SourcePCGComponent = Cast<UPCGComponent>(Context->ExecutionSource.Get());
 	if (!DynMeshData || !SourcePCGComponent)
 	{
 		return;
diff --git a/Engine/Plugins/PCGInterops/PCGGeometryScriptInterop/Source/PCGGeometryScriptInterop/Private/Elements/PCGMeshSampler.cpp b/Engine/Plugins/PCGInterops/PCGGeometryScriptInterop/Source/PCGGeometryScriptInterop/Private/Elements/PCGMeshSampler.cpp
--- a/Engine/Plugins/PCGInterops/PCGGeometryScriptInterop/Source/PCGGeometryScriptInterop/Private/Elements/PCGMeshSampler.cpp
+++ b/Engine/Plugins/PCGInterops/PCGGeometryScriptInterop/Source/PCGGeometryScriptInterop/Private/Elements/PCGMeshSampler.cpp
@@ -401,7 +401,7 @@ bool FPCGMeshSamplerElement::PrepareDataInternal(FPCGContext* InContext) const
 
 	check(Context);
 
-	if (!Context->SourceComponent.IsValid())
+	if (!Context->ExecutionSource.IsValid())
 	{
 		return true;
 	}
diff --git a/Engine/Plugins/PCGInterops/PCGGeometryScriptInterop/Source/PCGGeometryScriptInterop/Private/Elements/PCGSpawnDynamicMesh.cpp b/Engine/Plugins/PCGInterops/PCGGeometryScriptInterop/Source/PCGGeometryScriptInterop/Private/Elements/PCGSpawnDynamicMesh.cpp
--- a/Engine/Plugins/PCGInterops/PCGGeometryScriptInterop/Source/PCGGeometryScriptInterop/Private/Elements/PCGSpawnDynamicMesh.cpp
+++ b/Engine/Plugins/PCGInterops/PCGGeometryScriptInterop/Source/PCGGeometryScriptInterop/Private/Elements/PCGSpawnDynamicMesh.cpp
@@ -51,7 +51,7 @@ bool FPCGSpawnDynamicMeshElement::ExecuteInternal(FPCGContext* InContext) const
 		return true;
 	}
 	
-	UPCGComponent* SourcePCGComponent = InContext->SourceComponent.Get();
+	UPCGComponent* SourcePCGComponent = Cast<UPCGComponent>(InContext->ExecutionSource.Get());
 	if (!SourcePCGComponent)
 	{
 		return true;
diff --git a/Engine/Plugins/PCGInterops/PCGGeometryScriptInterop/Source/PCGGeometryScriptInterop/Private/Resources/PCGDynamicMeshManagedComponent.cpp b/Engine/Plugins/PCGInterops/PCGGeometryScriptInterop/Source/PCGGeometryScriptInterop/Private/Resources/PCGDynamicMeshManagedComponent.cpp
--- a/Engine/Plugins/PCGInterops/PCGGeometryScriptInterop/Source/PCGGeometryScriptInterop/Private/Resources/PCGDynamicMeshManagedComponent.cpp
+++ b/Engine/Plugins/PCGInterops/PCGGeometryScriptInterop/Source/PCGGeometryScriptInterop/Private/Resources/PCGDynamicMeshManagedComponent.cpp
@@ -53,7 +53,7 @@ UPCGDynamicMeshManagedComponent* PCGDynamicMeshManagedComponent::GetOrCreateDyna
 		return nullptr;
 	}
 
-	UPCGComponent* SourceComponent = Context->SourceComponent.Get();
+	UPCGComponent* SourceComponent = Cast<UPCGComponent>(Context->ExecutionSource.Get());
 	if (!SourceComponent || !SettingsInterface)
 	{
 		return nullptr;
