[PCG] Enable 'pcg.Experimental.SpatialData.EnablePrepareForSpatialQuery' by default
- in editor code UPCGLandscapeData prepare for spatial query tasks are GameThread bound because of FLandscapeComponentDataInterface which isn't thread safe

#rnx

diff --git a/Engine/Plugins/PCG/Source/PCG/Private/Data/PCGLandscapeData.cpp b/Engine/Plugins/PCG/Source/PCG/Private/Data/PCGLandscapeData.cpp
--- a/Engine/Plugins/PCG/Source/PCG/Private/Data/PCGLandscapeData.cpp
+++ b/Engine/Plugins/PCG/Source/PCG/Private/Data/PCGLandscapeData.cpp
@@ -704,7 +704,14 @@ TArray<FPCGTaskId> UPCGLandscapeData::PrepareForSpatialQuery(FPCGContext* InCont
 		return true;
 	},
 	InContext->ExecutionSource.Get());
+
+#if WITH_EDITOR
+	// Creation of new cache entries is not thread safe because of FLandscapeComponentDataInterface 
+	Params.bCanExecuteOnlyOnMainThread = true;
+#else
+	// Outside of editor we can only load existing entries which is thread safe
 	Params.bCanExecuteOnlyOnMainThread = false;
+#endif
 
 	return { InContext->ScheduleGeneric(Params) };
 }
\ No newline at end of file
diff --git a/Engine/Plugins/PCG/Source/PCG/Private/Data/PCGSpatialData.cpp b/Engine/Plugins/PCG/Source/PCG/Private/Data/PCGSpatialData.cpp
--- a/Engine/Plugins/PCG/Source/PCG/Private/Data/PCGSpatialData.cpp
+++ b/Engine/Plugins/PCG/Source/PCG/Private/Data/PCGSpatialData.cpp
@@ -19,7 +19,7 @@ namespace PCGSpatialData
 {
 	TAutoConsoleVariable<bool> CVarEnablePrepareForSpatialQuery(
 		TEXT("pcg.Experimental.SpatialData.EnablePrepareForSpatialQuery"),
-		false,
+		true,
 		TEXT("Enable UPCGSpatialData subclass PrepareForSpatialQuery task scheduling"));
 }
 
