[PCG] Removed duplicate generation sources (player controller vs wp source) by adding named generation sources.
Usage of WP streaming sources-as-generation sources in general can now be opted-out.

diff --git a/Engine/Plugins/PCG/Source/PCG/Private/RuntimeGen/PCGGenSourceManager.cpp b/Engine/Plugins/PCG/Source/PCG/Private/RuntimeGen/PCGGenSourceManager.cpp
--- a/Engine/Plugins/PCG/Source/PCG/Private/RuntimeGen/PCGGenSourceManager.cpp
+++ b/Engine/Plugins/PCG/Source/PCG/Private/RuntimeGen/PCGGenSourceManager.cpp
@@ -62,13 +62,22 @@ TSet<IPCGGenSourceBase*> FPCGGenSourceManager::GetAllGenSources(const APCGWorldA
 		bDirty = false;
 	}
 
-	TSet<IPCGGenSourceBase*> AllGenSources = RegisteredGenSources;
+	TSet<IPCGGenSourceBase*> AllGenSources;
+	AllGenSources.Reserve(RegisteredGenSources.Num() + RegisteredNamedGenSources.Num());
+	AllGenSources.Append(RegisteredGenSources);
 
-	// Acquire a generation source for each active streaming source.
-	for (int32 I = 0; I < NumWorldPartitionStreamingSources; ++I)
+	for (const TPair<FName, IPCGGenSourceBase*>& NamedGenSource : RegisteredNamedGenSources)
 	{
-		check(WorldPartitionGenSources.IsValidIndex(I));
-		AllGenSources.Add(WorldPartitionGenSources[I]);
+		AllGenSources.Add(NamedGenSource.Value);
+	}
+
+	// Acquire a generation source for each active streaming source that doesn't have a generation source already (based on name collision test).
+	for (UPCGGenSourceWPStreamingSource* WPGenSource : WorldPartitionGenSources)
+	{
+		if (WPGenSource && WPGenSource->StreamingSource && !RegisteredNamedGenSources.Contains(WPGenSource->StreamingSource->Name))
+		{
+			AllGenSources.Add(WPGenSource);
+		}
 	}
 
 #if WITH_EDITOR
@@ -82,7 +91,7 @@ TSet<IPCGGenSourceBase*> FPCGGenSourceManager::GetAllGenSources(const APCGWorldA
 	return AllGenSources;
 }
 
-bool FPCGGenSourceManager::RegisterGenSource(IPCGGenSourceBase* InGenSource)
+bool FPCGGenSourceManager::RegisterGenSource(IPCGGenSourceBase* InGenSource, FName InGenSourceName)
 {
 	if (UPCGGenSourceComponent* GenSourceComponent = Cast<UPCGGenSourceComponent>(InGenSource))
 	{
@@ -92,7 +101,22 @@ bool FPCGGenSourceManager::RegisterGenSource(IPCGGenSourceBase* InGenSource)
 		}
 	}
 
-	return RegisteredGenSources.Add(InGenSource).IsValidId();
+	if (InGenSourceName == NAME_None)
+	{
+		return RegisteredGenSources.Add(InGenSource).IsValidId();
+	}
+	else
+	{
+		if (RegisteredNamedGenSources.Find(InGenSourceName))
+		{
+			return false;
+		}
+		else
+		{
+			RegisteredNamedGenSources.Add(InGenSourceName, InGenSource);
+			return true;
+		}
+	}
 }
 
 bool FPCGGenSourceManager::UnregisterGenSource(const IPCGGenSourceBase* InGenSource)
@@ -105,7 +129,28 @@ bool FPCGGenSourceManager::UnregisterGenSource(const IPCGGenSourceBase* InGenSou
 		}
 	}
 
-	return RegisteredGenSources.Remove(InGenSource) > 0;
+	// Start by removing from the unnamed gen sources.
+	// @todo_pcg We could remove from named sources here too, but that might be limiting for more complex cases, BUT it might be less error-prone too.
+	if (RegisteredGenSources.Remove(InGenSource) > 0)
+	{
+		return true;
+	}
+	else // Then try to remove from the named sources
+	{
+		if (const FName* GenSourceNamePtr = RegisteredNamedGenSources.FindKey(const_cast<IPCGGenSourceBase*>(InGenSource)))
+		{
+			FName GenSourceName = *GenSourceNamePtr;
+			return UnregisterGenSource(GenSourceName);
+		}
+	}
+
+	return false;
+}
+
+bool FPCGGenSourceManager::UnregisterGenSource(FName InGenSourceName)
+{
+	// Note: unnamed generation sources can't be unregistered by name, by definition.
+	return RegisteredNamedGenSources.Remove(InGenSourceName) > 0;
 }
 
 void FPCGGenSourceManager::OnGameModePostLogin(AGameModeBase* InGameMode, APlayerController* InPlayerController)
@@ -120,8 +165,7 @@ void FPCGGenSourceManager::OnGameModePostLogin(AGameModeBase* InGameMode, APlaye
 	UPCGGenSourcePlayer* GenSource = NewObject<UPCGGenSourcePlayer>();
 	GenSource->SetPlayerController(InPlayerController);
 
-	FSetElementId ElementId = RegisteredGenSources.Add(GenSource);
-	if (ElementId.IsValidId())
+	if (RegisterGenSource(GenSource, InPlayerController->GetFName()))
 	{
 		GenSource->AddToRoot();
 	}
@@ -129,61 +173,75 @@ void FPCGGenSourceManager::OnGameModePostLogin(AGameModeBase* InGameMode, APlaye
 
 void FPCGGenSourceManager::OnGameModePostLogout(AGameModeBase* InGameMode, AController* InController)
 {
-	if (InController->GetWorld() != World)
+	if (!InController || InController->GetWorld() != World)
 	{
 		return;
 	}
 
 	ensure(IsInGameThread());
 
-	for (IPCGGenSourceBase* GenSource : RegisteredGenSources)
+	IPCGGenSourceBase* GenSource = nullptr;
+	if (RegisteredNamedGenSources.RemoveAndCopyValue(InController->GetFName(), GenSource))
 	{
+		check(GenSource);
 		if (UPCGGenSourcePlayer* GenSourcePlayer = Cast<UPCGGenSourcePlayer>(GenSource))
 		{
-			if (GenSourcePlayer->GetPlayerController() == InController || !GenSourcePlayer->IsValid())
-			{
-				GenSourcePlayer->RemoveFromRoot();
-				RegisteredGenSources.Remove(GenSourcePlayer);
-				break;
-			}
+			GenSourcePlayer->RemoveFromRoot();
+		}
+		else
+		{
+			// The gen source we removed isn't a player controller gen source, so let's put it back.
+			RegisteredNamedGenSources.Add(InController->GetFName(), GenSource);
 		}
 	}
 }
 
 void FPCGGenSourceManager::UpdatePerTickGenSources(const APCGWorldActor* InPCGWorldActor)
 {
-	if (const UWorld* InWorld = InPCGWorldActor->GetWorld())
+	// Start by updating, adding and removing WP generation sources if it's needed.
+	// TODO: Grab StreamingSourceProviders instead of StreamingSources?
+	// TODO: Is it possible to avoid adding a StreamingSource for the Player, which we already capture in OnGameModePostLogin?
+	// Note: GetStreamingSources only works in GameWorld, so StreamingSources do not act as generation sources in editor.
+	const TArray<FWorldPartitionStreamingSource>* WPStreamingSources = nullptr;
+
+	if (InPCGWorldActor->bEnableWorldPartitionGenerationSources &&
+		InPCGWorldActor->GetWorld() &&
+		InPCGWorldActor->GetWorld()->GetWorldPartition())
 	{
-		if (const UWorldPartition* WorldPartition = InWorld->GetWorldPartition())
-		{
-			// TODO: Grab StreamingSourceProviders instead of StreamingSources?
-			// TODO: Is it possible to avoid adding a StreamingSource for the Player, which we already capture in OnGameModePostLogin?
-			// Note: GetStreamingSources only works in GameWorld, so StreamingSources do not act as generation sources in editor.
-			const TArray<FWorldPartitionStreamingSource>& StreamingSources = WorldPartition->GetStreamingSources();
-
-			NumWorldPartitionStreamingSources = StreamingSources.Num();
-
-			for (int32 I = 0; I < NumWorldPartitionStreamingSources; ++I)
-			{
-				const bool bCreateNewWorldPartitionGenSource = I >= WorldPartitionGenSources.Num();
-
-				// Allocate new WP generation sources as necessary.
-				UPCGGenSourceWPStreamingSource* GenSource = bCreateNewWorldPartitionGenSource ? WorldPartitionGenSources.Add_GetRef(NewObject<UPCGGenSourceWPStreamingSource>()) : WorldPartitionGenSources[I];
-
-				// Capture each streaming source in a single generation source.
-				if (GenSource)
-				{
-					if (bCreateNewWorldPartitionGenSource)
-					{
-						// New generation sources need to be rooted.
-						GenSource->AddToRoot();
-					}
+		WPStreamingSources = &InPCGWorldActor->GetWorld()->GetWorldPartition()->GetStreamingSources();
+	}
 
-					GenSource->StreamingSource = &StreamingSources[I];
-				}
-			}
+	const int NumWPStreamingSources = WPStreamingSources ? WPStreamingSources->Num() : 0;
+	const int NumWPGenerationSources = WorldPartitionGenSources.Num();
+	const int NumGenerationSourcesToUpdate = FMath::Min(NumWPStreamingSources, NumWPGenerationSources);
+	const int NumGenerationSourcesToRemove = FMath::Max(0, NumWPGenerationSources - NumWPStreamingSources);
+	const int NumGenerationSourcesToAdd = FMath::Max(0, NumWPStreamingSources - NumWPGenerationSources);
+	check((NumGenerationSourcesToRemove == 0 && NumGenerationSourcesToAdd == 0) || ((NumGenerationSourcesToRemove > 0) != (NumGenerationSourcesToAdd > 0)));
+	
+	for (int I = 0; I < NumGenerationSourcesToAdd; ++I)
+	{
+		UPCGGenSourceWPStreamingSource* GenSource = WorldPartitionGenSources.Add_GetRef(NewObject<UPCGGenSourceWPStreamingSource>());
+		// New generation sources need to be rooted.
+		GenSource->AddToRoot();
+	}
+	
+	check((NumGenerationSourcesToUpdate + NumGenerationSourcesToAdd == 0) || WPStreamingSources);
+	// Implementation note: not really needed but we need to appease the static analysis gods
+	if (WPStreamingSources)
+	{
+		for (int I = 0; I < NumGenerationSourcesToUpdate + NumGenerationSourcesToAdd; ++I)
+		{
+			WorldPartitionGenSources[I]->StreamingSource = &(*WPStreamingSources)[I];
 		}
 	}
+	
+	for (int I = NumGenerationSourcesToUpdate + NumGenerationSourcesToAdd; I < NumGenerationSourcesToUpdate + NumGenerationSourcesToAdd + NumGenerationSourcesToRemove; ++I)
+	{
+		// Unused generation sources need to be forgotten & removed from root
+		WorldPartitionGenSources[I]->RemoveFromRoot();
+	}
+
+	WorldPartitionGenSources.SetNum(NumGenerationSourcesToUpdate + NumGenerationSourcesToAdd);
 
 #if WITH_EDITOR
 	EditorCameraGenSource->EditorViewportClient = nullptr;
diff --git a/Engine/Plugins/PCG/Source/PCG/Public/PCGWorldActor.h b/Engine/Plugins/PCG/Source/PCG/Public/PCGWorldActor.h
--- a/Engine/Plugins/PCG/Source/PCG/Public/PCGWorldActor.h
+++ b/Engine/Plugins/PCG/Source/PCG/Public/PCGWorldActor.h
@@ -76,6 +76,10 @@ public:
 	bool bTreatEditorViewportAsGenerationSource = false;
 #endif
 
+	/** Allows World Partition streaming sources to act as Runtime Generation Sources. */
+	UPROPERTY(BlueprintReadWrite, EditAnywhere, Category = RuntimeGeneration)
+	bool bEnableWorldPartitionGenerationSources = true;
+
 private:
 	friend class FPCGActorAndComponentMapping;
 
diff --git a/Engine/Plugins/PCG/Source/PCG/Public/RuntimeGen/PCGGenSourceManager.h b/Engine/Plugins/PCG/Source/PCG/Public/RuntimeGen/PCGGenSourceManager.h
--- a/Engine/Plugins/PCG/Source/PCG/Public/RuntimeGen/PCGGenSourceManager.h
+++ b/Engine/Plugins/PCG/Source/PCG/Public/RuntimeGen/PCGGenSourceManager.h
@@ -30,11 +30,14 @@ public:
 	TSet<IPCGGenSourceBase*> PCG_API GetAllGenSources(const APCGWorldActor* InPCGWorldActor);
 
 	/** Adds a UPCGGenSource to be tracked by the GenSourceManager. */
-	bool PCG_API RegisterGenSource(IPCGGenSourceBase* InGenSource);
+	bool PCG_API RegisterGenSource(IPCGGenSourceBase* InGenSource, FName InGenSourceName = NAME_None);
 
 	/** Removes a UPCGGenSource from being tracked by the GenSourceManager. */
 	bool PCG_API UnregisterGenSource(const IPCGGenSourceBase* InGenSource);
 
+	/** Removes a UPCGGenSource from being tracked by the GenSourceManaged on a name-basis. */
+	bool PCG_API UnregisterGenSource(FName InGenSourceName);
+
 protected:
 	void OnGameModePostLogin(AGameModeBase* InGameMode, APlayerController* InPlayerController);
 	void OnGameModePostLogout(AGameModeBase* InGameMode, AController* InController);
@@ -43,7 +46,10 @@ protected:
 	void UpdatePerTickGenSources(const APCGWorldActor* InPCGWorldActor);
 
 protected:
-	/** Tracks registered generation sources, such as UPCGGenSourceComponent and UPCGGenSourcePlayer. */
+	/** Tracks named registered generation sources, such as UPCGGenSourcePlayer. */
+	TMap<FName, IPCGGenSourceBase*> RegisteredNamedGenSources;
+
+	/** Tracks unnamed registered generation sources, such as UPCGGenSourceComponent. */
 	TSet<IPCGGenSourceBase*> RegisteredGenSources;
 
 #if WITH_EDITORONLY_DATA
@@ -54,9 +60,6 @@ protected:
 	/** Pool of GenSources dedicated to WorldPartition StreamingSources. This list grows as necessary and refreshes every tick. */
 	TArray<UPCGGenSourceWPStreamingSource*> WorldPartitionGenSources;
 
-	/** There may be more generation sources in 'WorldPartitionGenSources' than there are streaming sources in the world, so we only use the first N generation sources. */
-	int32 NumWorldPartitionStreamingSources = 0;
-
 	const UWorld* World = nullptr;
 	bool bDirty = false;
 };
