[PCG] Add data visualization viewport
* Tied to attribute list view. Each ALV has its own viewport window, and will try to visualize whatever data is currently inspected.
* Add Static Mesh Resource Data Visualization, with an implementation for viewport visualization

#rnx

diff --git a/Engine/Plugins/PCG/Source/PCG/Public/PCGDataVisualization.h b/Engine/Plugins/PCG/Source/PCG/Public/PCGDataVisualization.h
--- a/Engine/Plugins/PCG/Source/PCG/Public/PCGDataVisualization.h
+++ b/Engine/Plugins/PCG/Source/PCG/Public/PCGDataVisualization.h
@@ -13,6 +13,16 @@ class UPCGData;
 class UPCGSettingsInterface;
 struct FPCGContext;
 
+class FAdvancedPreviewScene;
+struct FStreamableHandle;
+
+using FPCGSetupSceneFunc = TFunction<void(
+	FAdvancedPreviewScene* InScene,
+	const TArray<UObject*>& InResources,
+	FVector& OutViewLocation,
+	float& OutOrbitDistance,
+	TArray<TObjectPtr<UObject>>& OutManagedResources)>;
+
 enum class EPCGTableVisualizerColumnSortingMode
 {
 	None,
@@ -83,5 +93,11 @@ PRAGMA_ENABLE_DEPRECATION_WARNINGS
 		Data->SetDomainFromDomainID(DomainID, Selector);
 		return Selector.GetDomainName().ToString();
 	}
+
+	/** Initiates an async load on any resources this data needs in order to be visualized. */
+	virtual TArray<TSharedPtr<FStreamableHandle>> LoadRequiredResources(const UPCGData* Data) const { return {}; }
+
+	/** Optionally provide a function to setup the data viewport. */
+	virtual FPCGSetupSceneFunc GetViewportSetupFunc(const UPCGData* Data) const { return FPCGSetupSceneFunc(); }
 };
 #endif
diff --git a/Engine/Plugins/PCG/Source/PCGEditor/PCGEditor.Build.cs b/Engine/Plugins/PCG/Source/PCGEditor/PCGEditor.Build.cs
--- a/Engine/Plugins/PCG/Source/PCGEditor/PCGEditor.Build.cs
+++ b/Engine/Plugins/PCG/Source/PCGEditor/PCGEditor.Build.cs
@@ -30,6 +30,7 @@ namespace UnrealBuildTool.Rules
 			PrivateDependencyModuleNames.AddRange(
 				new string[]
 				{
+					"AdvancedPreviewScene",
 					"AppFramework",
 					"ApplicationCore",
 					"AssetDefinition",
diff --git a/Engine/Plugins/PCG/Source/PCGEditor/Private/DataVisualizations/PCGStaticMeshDataVisualization.cpp b/Engine/Plugins/PCG/Source/PCGEditor/Private/DataVisualizations/PCGStaticMeshDataVisualization.cpp
new file mode 100644
--- /dev/null
+++ b/Engine/Plugins/PCG/Source/PCGEditor/Private/DataVisualizations/PCGStaticMeshDataVisualization.cpp
@@ -0,0 +1,68 @@
+// Copyright Epic Games, Inc. All Rights Reserved.
+
+#include "DataVisualizations/PCGStaticMeshDataVisualization.h"
+
+#include "Data/PCGStaticMeshResourceData.h"
+#include "DataVisualizations/PCGDataVisualizationHelpers.h"
+
+#include "AdvancedPreviewScene.h"
+#include "MaterialEditor/MaterialEditorMeshComponent.h"
+
+FPCGTableVisualizerInfo IPCGStaticMeshDataVisualization::GetTableVisualizerInfoWithDomain(const UPCGData* Data, const FPCGMetadataDomainID& DomainID) const
+{
+	FPCGTableVisualizerInfo Info;
+	Info.Data = Data;
+	// @todo_pcg: Add row for mesh path? Support metadata as well?
+ 
+	return Info;
+}
+
+TArray<TSharedPtr<FStreamableHandle>> IPCGStaticMeshDataVisualization::LoadRequiredResources(const UPCGData* Data) const
+{
+	TArray<TSharedPtr<FStreamableHandle>> LoadHandles;
+
+	if (const UPCGStaticMeshResourceData* StaticMeshData = Cast<UPCGStaticMeshResourceData>(Data))
+	{
+		LoadHandles.Add(StaticMeshData->RequestResourceLoad());
+	}
+
+	return LoadHandles;
+}
+
+FPCGSetupSceneFunc IPCGStaticMeshDataVisualization::GetViewportSetupFunc(const UPCGData* Data) const
+{
+	return [this](FAdvancedPreviewScene* InScene, const TArray<UObject*>& InResources, FVector& OutViewLocation, float& OutOrbitDistance, TArray<TObjectPtr<UObject>>& OutManagedResources)
+	{
+		check(InScene);
+
+		if (InResources.IsEmpty())
+		{
+			return;
+		}
+
+		// Note: Using UMaterialEditorMeshComponent subclass for more accurate mesh bounds.
+		TObjectPtr<UStaticMeshComponent> MeshComponent = NewObject<UMaterialEditorMeshComponent>(GetTransientPackage(), NAME_None, RF_Transient);
+		OutManagedResources.Add(MeshComponent);
+
+		if (GEditor->PreviewPlatform.GetEffectivePreviewFeatureLevel() <= ERHIFeatureLevel::ES3_1)
+		{
+			MeshComponent->SetMobility(EComponentMobility::Static);
+		}
+
+		InScene->AddComponent(MeshComponent, FTransform::Identity);
+
+		UStaticMesh* StaticMesh = Cast<UStaticMesh>(InResources[0]);
+		MeshComponent->SetStaticMesh(StaticMesh);
+
+		// Bounds will be updated already by SetStaticMesh() call.
+		const FBoxSphereBounds Bounds = MeshComponent->Bounds;
+		InScene->SetFloorOffset(-Bounds.Origin.Z + Bounds.BoxExtent.Z);
+
+		OutViewLocation = Bounds.Origin;
+
+		if (Bounds.SphereRadius > 0.0f)
+		{
+			OutOrbitDistance = Bounds.SphereRadius * 2.0f;
+		}
+	};
+}
diff --git a/Engine/Plugins/PCG/Source/PCGEditor/Private/PCGEditor.cpp b/Engine/Plugins/PCG/Source/PCGEditor/Private/PCGEditor.cpp
--- a/Engine/Plugins/PCG/Source/PCGEditor/Private/PCGEditor.cpp
+++ b/Engine/Plugins/PCG/Source/PCGEditor/Private/PCGEditor.cpp
@@ -40,6 +40,7 @@
 #include "Widgets/SPCGEditorGraphParamsView.h"
 #include "Widgets/SPCGEditorGraphProfilingView.h"
 #include "Widgets/SPCGEditorNodeSource.h"
+#include "Widgets/AssetEditorViewport/SPCGEditorViewport.h"
 
 #include "AssetToolsModule.h"
 #include "EdGraphUtilities.h"
@@ -93,6 +94,11 @@ namespace FPCGEditor_private
 	const FName LogID = FName(TEXT("Log"));
 	const FName NodeSourceID = FName(TEXT("NodeSource"));
 	const FName UserParamsID = FName(TEXT("UserParams"));
+	const FName ViewportID[] = {
+		FName(TEXT("Viewport")),
+		FName(TEXT("Viewport2")),
+		FName(TEXT("Viewport3")),
+		FName(TEXT("Viewport4")) };
 
 	const FText UserParamsTabName = LOCTEXT("UserParamsTab", "Graph Parameters");
 }
@@ -152,6 +158,7 @@ void FPCGEditor::Initialize(const EToolkitMode::Type InMode, const TSharedPtr<cl
 	for (int AttributesIndex = 0; AttributesIndex < 4; ++AttributesIndex)
 	{
 		AttributesWidgets.Add(CreateAttributesWidget());
+		AttributesWidgets[AttributesIndex]->SetViewportWidget(CreateViewportWidget());
 	}
 	
 	DeterminismWidget = CreateDeterminismWidget();
@@ -560,6 +567,7 @@ void FPCGEditor::RegisterTabSpawners(const TSharedRef<FTabManager>& InTabManager
 	WorkspaceMenuCategory = InTabManager->AddLocalWorkspaceMenuCategory(LOCTEXT("WorkspaceMenu_PCGEditor", "PCG Editor"));
 	TSharedRef<FWorkspaceItem> DetailsGroup = WorkspaceMenuCategory->AddGroup(LOCTEXT("WorkspaceMenu_PCGEditor_Details", "Details"));
 	TSharedRef<FWorkspaceItem> AttributesGroup = WorkspaceMenuCategory->AddGroup(LOCTEXT("WorkspaceMenu_PCGEditor_Attributes", "Attributes"));
+	TSharedRef<FWorkspaceItem> ViewportGroup = WorkspaceMenuCategory->AddGroup(LOCTEXT("WorkspaceMenu_PCGEditor_Viewport", "Data Viewport"));
 	const TSharedRef<FWorkspaceItem>& WorkspaceMenuCategoryRef = WorkspaceMenuCategory.ToSharedRef();
 
 	FAssetEditorToolkit::RegisterTabSpawners(InTabManager);
@@ -630,8 +638,23 @@ void FPCGEditor::RegisterTabSpawners(const TSharedRef<FTabManager>& InTabManager
 		.SetGroup(WorkspaceMenuCategoryRef);
 
 	InTabManager->RegisterTabSpawner(FPCGEditor_private::UserParamsID, FOnSpawnTab::CreateSP(this, &FPCGEditor::SpawnTab_UserParams))
-		.SetDisplayName(FPCGEditor_private::UserParamsTabName)
-		.SetGroup(WorkspaceMenuCategoryRef);
+		.SetDisplayName(FPCGEditor_private::UserParamsTabName);
+
+	InTabManager->RegisterTabSpawner(FPCGEditor_private::ViewportID[0], FOnSpawnTab::CreateSP(this, &FPCGEditor::SpawnTab_Viewport, 0))
+		.SetDisplayName(LOCTEXT("ViewportTab1", "Data Viewport 1"))
+		.SetGroup(ViewportGroup);
+
+	InTabManager->RegisterTabSpawner(FPCGEditor_private::ViewportID[1], FOnSpawnTab::CreateSP(this, &FPCGEditor::SpawnTab_Viewport, 1))
+		.SetDisplayName(LOCTEXT("ViewportTab2", "Data Viewport 2"))
+		.SetGroup(ViewportGroup);
+
+	InTabManager->RegisterTabSpawner(FPCGEditor_private::ViewportID[2], FOnSpawnTab::CreateSP(this, &FPCGEditor::SpawnTab_Viewport, 2))
+		.SetDisplayName(LOCTEXT("ViewportTab3", "Data Viewport 3"))
+		.SetGroup(ViewportGroup);
+
+	InTabManager->RegisterTabSpawner(FPCGEditor_private::ViewportID[3], FOnSpawnTab::CreateSP(this, &FPCGEditor::SpawnTab_Viewport, 3))
+		.SetDisplayName(LOCTEXT("ViewportTab4", "Data Viewport 4"))
+		.SetGroup(ViewportGroup);
 }
 
 void FPCGEditor::UnregisterTabSpawners(const TSharedRef<class FTabManager>& InTabManager)
@@ -2860,6 +2883,11 @@ TSharedRef<SPCGEditorGraphUserParametersView> FPCGEditor::CreateGraphParamsWidge
 	return SNew(SPCGEditorGraphUserParametersView, SharedThis(this));
 }
 
+TSharedRef<SPCGEditorViewport> FPCGEditor::CreateViewportWidget()
+{
+	return SNew(SPCGEditorViewport);
+}
+
 void FPCGEditor::OnSelectedNodesChanged(const TSet<UObject*>& NewSelection)
 {
 	TArray<TWeakObjectPtr<UObject>> SelectedObjects;
@@ -3280,6 +3308,18 @@ TSharedRef<SDockTab> FPCGEditor::SpawnTab_UserParams(const FSpawnTabArgs& Args)
 		];
 }
 
+TSharedRef<SDockTab> FPCGEditor::SpawnTab_Viewport(const FSpawnTabArgs& Args, int ViewportIndex)
+{
+	TAttribute<FText> Label = TAttribute<FText>::Create(TAttribute<FText>::FGetter::CreateRaw(this, &FPCGEditor::GetViewportTabLabel, ViewportIndex));
+
+	return SNew(SDockTab)
+		.Label(Label)
+		.TabColorScale(GetTabColorScale())
+		[
+			AttributesWidgets[ViewportIndex]->GetViewportWidget().ToSharedRef()
+		];
+}
+
 FText FPCGEditor::GetDetailsTabLabel(int DetailsIndex)
 {
 	if (DetailsIndex == 0)
@@ -3309,4 +3349,16 @@ FText FPCGEditor::GetAttributesTabLabel(int AttributesIndex)
 	}
 }
 
+FText FPCGEditor::GetViewportTabLabel(int ViewportIndex)
+{
+	if (ViewportIndex == 0)
+	{
+		return LOCTEXT("PCGViewportTitle", "Data Viewport");
+	}
+	else
+	{
+		return FText::Format(LOCTEXT("PCGViewportTitle_Multi", "Data Viewport {0}"), ViewportIndex + 1);
+	}
+}
+
 #undef LOCTEXT_NAMESPACE
\ No newline at end of file
diff --git a/Engine/Plugins/PCG/Source/PCGEditor/Private/PCGEditor.h b/Engine/Plugins/PCG/Source/PCGEditor/Private/PCGEditor.h
--- a/Engine/Plugins/PCG/Source/PCGEditor/Private/PCGEditor.h
+++ b/Engine/Plugins/PCG/Source/PCGEditor/Private/PCGEditor.h
@@ -24,6 +24,7 @@ class SPCGEditorGraphNodePalette;
 class SPCGEditorGraphProfilingView;
 class SPCGEditorNodeSource;
 class SPCGEditorGraphUserParametersView;
+class SPCGEditorViewport;
 class UEdGraphNode;
 class UPCGComponent;
 class UPCGEditorGraph;
@@ -349,6 +350,9 @@ private:
 	/** Create a new user graph parameters tab widget */
 	TSharedRef<SPCGEditorGraphUserParametersView> CreateGraphParamsWidget();
 
+	/** Create a new viewport widget */
+	TSharedRef<SPCGEditorViewport> CreateViewportWidget();
+
 	/** Called when the selection changes in the GraphEditor */
 	void OnSelectedNodesChanged(const TSet<UObject*>& NewSelection);
 
@@ -409,10 +413,12 @@ private:
 	TSharedRef<SDockTab> SpawnTab_Log(const FSpawnTabArgs& Args);
 	TSharedRef<SDockTab> SpawnTab_NodeSource(const FSpawnTabArgs& Args);
 	TSharedRef<SDockTab> SpawnTab_UserParams(const FSpawnTabArgs& Args);
+	TSharedRef<SDockTab> SpawnTab_Viewport(const FSpawnTabArgs& Args, int ViewportIndex);
 
 	FText GetDetailsTabLabel(int DetailsIndex);
 	FText GetDetailsViewObjectName(int DetailsIndex);
 	FText GetAttributesTabLabel(int AttributesIndex);
+	FText GetViewportTabLabel(int ViewportIndex);
 
 	TSharedPtr<SGraphEditor> GraphEditorWidget;
 	TArray<TSharedPtr<SPCGEditorGraphDetailsView>> PropertyDetailsWidgets;
diff --git a/Engine/Plugins/PCG/Source/PCGEditor/Private/PCGEditorModule.cpp b/Engine/Plugins/PCG/Source/PCGEditor/Private/PCGEditorModule.cpp
--- a/Engine/Plugins/PCG/Source/PCGEditor/Private/PCGEditorModule.cpp
+++ b/Engine/Plugins/PCG/Source/PCGEditor/Private/PCGEditorModule.cpp
@@ -14,6 +14,7 @@
 #include "Data/PCGPrimitiveData.h"
 #include "Data/PCGSpatialData.h"
 #include "Data/PCGSplineData.h"
+#include "Data/PCGStaticMeshResourceData.h"
 #include "Data/PCGVolumeData.h"
 #include "Grid/PCGPartitionActor.h"
 #include "WorldPartitionBuilder/PCGWorldPartitionBuilder.h"
@@ -28,6 +29,7 @@
 #include "DataVisualizations/PCGParamDataVisualization.h"
 #include "DataVisualizations/PCGSpatialDataVisualization.h"
 #include "DataVisualizations/PCGSplineDataVisualization.h"
+#include "DataVisualizations/PCGStaticMeshDataVisualization.h"
 #include "DataVisualizations/PCGVolumeDataVisualization.h"
 #include "Details/PCGAttributePropertySelectorDetails.h"
 #include "Details/PCGBlueprintSettingsDetails.h"
@@ -731,10 +733,11 @@ void FPCGEditorModule::RegisterPCGDataVisualizations()
 	DataVisRegistry.InternalRegistry.Add(UPCGSpatialData::StaticClass(), MakeUnique<const IPCGSpatialDataVisualization>());
 	DataVisRegistry.InternalRegistry.Add(UPCGSplineData::StaticClass(), MakeUnique<const IPCGSplineDataVisualization>());
 	DataVisRegistry.InternalRegistry.Add(UPCGVolumeData::StaticClass(), MakeUnique<const IPCGPrimitiveDataVisualization>());
-    DataVisRegistry.InternalRegistry.Add(UPCGPrimitiveData::StaticClass(), MakeUnique<const IPCGPrimitiveDataVisualization>());
-    DataVisRegistry.InternalRegistry.Add(UPCGCollisionShapeData::StaticClass(), MakeUnique<const IPCGPrimitiveDataVisualization>());
-    DataVisRegistry.InternalRegistry.Add(UPCGCollisionWrapperData::StaticClass(), MakeUnique<const IPCGPrimitiveDataVisualization>());
-    DataVisRegistry.InternalRegistry.Add(UPCGLandscapeData::StaticClass(), MakeUnique<const IPCGLandscapeDataVisualization>());
+	DataVisRegistry.InternalRegistry.Add(UPCGPrimitiveData::StaticClass(), MakeUnique<const IPCGPrimitiveDataVisualization>());
+	DataVisRegistry.InternalRegistry.Add(UPCGCollisionShapeData::StaticClass(), MakeUnique<const IPCGPrimitiveDataVisualization>());
+	DataVisRegistry.InternalRegistry.Add(UPCGCollisionWrapperData::StaticClass(), MakeUnique<const IPCGPrimitiveDataVisualization>());
+	DataVisRegistry.InternalRegistry.Add(UPCGLandscapeData::StaticClass(), MakeUnique<const IPCGLandscapeDataVisualization>());
+	DataVisRegistry.InternalRegistry.Add(UPCGStaticMeshResourceData::StaticClass(), MakeUnique<const IPCGStaticMeshDataVisualization>());
 }
 
 void FPCGEditorModule::UnregisterPCGDataVisualizations()
diff --git a/Engine/Plugins/PCG/Source/PCGEditor/Private/Widgets/AssetEditorViewport/SPCGEditorViewport.cpp b/Engine/Plugins/PCG/Source/PCGEditor/Private/Widgets/AssetEditorViewport/SPCGEditorViewport.cpp
new file mode 100644
--- /dev/null
+++ b/Engine/Plugins/PCG/Source/PCGEditor/Private/Widgets/AssetEditorViewport/SPCGEditorViewport.cpp
@@ -0,0 +1,204 @@
+// Copyright Epic Games, Inc. All Rights Reserved.
+
+#include "Widgets/AssetEditorViewport/SPCGEditorViewport.h"
+
+#include "AdvancedPreviewScene.h"
+#include "AssetViewerSettings.h"
+#include "EditorViewportClient.h"
+
+namespace PCGEditorViewportConstants
+{
+	const FVector DefaultViewLocation = FVector::ZeroVector;
+	const float DefaultOrbitDistance = 500.0f;
+}
+
+class FPCGEditorViewportClient : public FEditorViewportClient
+{
+public:
+	FPCGEditorViewportClient(const TSharedRef<SPCGEditorViewport>& InAssetEditorViewport);
+
+	//~ Begin FViewportClient Interface
+	virtual bool InputKey(const FInputKeyEventArgs& EventArgs) override;
+	virtual bool InputAxis(const FInputKeyEventArgs& Args) override;
+	//~ End FViewportClient Interface
+
+	//~ Begin FEditorViewportClient Interface
+	virtual FLinearColor GetBackgroundColor() const override;
+	virtual void Tick(float DeltaSeconds) override;
+	//~ End FEditorViewportClient Interface
+
+	FAdvancedPreviewScene* ResetScene();
+
+private:
+	TSharedPtr<FAdvancedPreviewScene> Scene = nullptr;
+};
+
+FPCGEditorViewportClient::FPCGEditorViewportClient(const TSharedRef<SPCGEditorViewport>& InAssetEditorViewport)
+	: FEditorViewportClient(nullptr, nullptr, StaticCastSharedRef<SEditorViewport>(InAssetEditorViewport))
+{
+	ResetScene();
+
+	bUsesDrawHelper = true;
+
+	// Setup defaults for the common draw helper.
+	DrawHelper.bDrawPivot = false;
+	DrawHelper.bDrawWorldBox = false;
+	DrawHelper.bDrawKillZ = false;
+	DrawHelper.bDrawGrid = true;
+	DrawHelper.GridColorAxis = FColor(80, 80, 80);
+	DrawHelper.GridColorMajor = FColor(72, 72, 72);
+	DrawHelper.GridColorMinor = FColor(64, 64, 64);
+	DrawHelper.PerspectiveGridSize = HALF_WORLD_MAX1;
+
+	EngineShowFlags.DisableAdvancedFeatures();
+	EngineShowFlags.SetSnap(0);
+	EngineShowFlags.SetSeparateTranslucency(true);
+}
+
+bool FPCGEditorViewportClient::InputKey(const FInputKeyEventArgs& Args)
+{
+	if (!Scene.IsValid())
+	{
+		return true;
+	}
+
+	bool bHandled = FEditorViewportClient::InputKey(Args);
+	bHandled |= InputTakeScreenshot(Args.Viewport, Args.Key, Args.Event);
+	bHandled |= Scene->HandleInputKey(Args);
+
+	return bHandled;
+}
+
+bool FPCGEditorViewportClient::InputAxis(const FInputKeyEventArgs& Args)
+{
+	if (!Scene.IsValid())
+	{
+		return true;
+	}
+
+	bool bResult = true;
+
+	if (!bDisableInput)
+	{
+		bResult = Scene->HandleViewportInput(Args.Viewport, Args.InputDevice, Args.Key, Args.AmountDepressed, Args.DeltaTime, Args.NumSamples, Args.IsGamepad());
+
+		if (bResult)
+		{
+			Invalidate();
+		}
+		else
+		{
+			bResult = FEditorViewportClient::InputAxis(Args);
+		}
+	}
+
+	return bResult;
+}
+
+FLinearColor FPCGEditorViewportClient::GetBackgroundColor() const
+{
+	return Scene.IsValid() ? Scene->GetBackgroundColor() : FColor(64, 64, 64);
+}
+
+void FPCGEditorViewportClient::Tick(float DeltaSeconds)
+{
+	FEditorViewportClient::Tick(DeltaSeconds);
+
+	if (!GIntraFrameDebuggingGameThread && Scene.IsValid() && Scene->GetWorld())
+	{
+		Scene->GetWorld()->Tick(LEVELTICK_All, DeltaSeconds);
+	}
+}
+
+FAdvancedPreviewScene* FPCGEditorViewportClient::ResetScene()
+{
+	Scene = MakeShareable(new FAdvancedPreviewScene(FPreviewScene::ConstructionValues()));
+	PreviewScene = Scene.Get();
+
+	// Restore last used feature level.
+	if (UWorld* World = Scene->GetWorld())
+	{
+		World->ChangeFeatureLevel(GWorld->GetFeatureLevel());
+	}
+
+	UAssetViewerSettings* Settings = UAssetViewerSettings::Get();
+	const int32 ProfileIndex = Scene->GetCurrentProfileIndex();
+
+	if (Settings && Settings->Profiles.IsValidIndex(ProfileIndex))
+	{
+		Scene->SetEnvironmentVisibility(Settings->Profiles[ProfileIndex].bShowEnvironment, true);
+	}
+
+	return Scene.Get();
+}
+
+SPCGEditorViewport::~SPCGEditorViewport()
+{
+	if (EditorViewportClient.IsValid())
+	{
+		EditorViewportClient->Viewport = nullptr;
+	}
+
+	ReleaseManagedResources();
+}
+
+void SPCGEditorViewport::Construct(const FArguments& InArgs)
+{
+	SEditorViewport::Construct(SEditorViewport::FArguments());
+}
+
+void SPCGEditorViewport::SetupScene(const TArray<UObject*>& InResources, const FPCGSetupSceneFunc& SetupFunc)
+{
+	ReleaseManagedResources();
+	
+	if (!ensure(EditorViewportClient.IsValid()))
+	{
+		return;
+	}
+
+	FAdvancedPreviewScene* Scene = EditorViewportClient->ResetScene();
+	FVector ViewLocation = PCGEditorViewportConstants::DefaultViewLocation;
+	float OrbitDistance = PCGEditorViewportConstants::DefaultOrbitDistance;
+
+	if (Scene && SetupFunc.IsSet())
+	{
+		SetupFunc(Scene, InResources, ViewLocation, OrbitDistance, ManagedResources);
+
+		for (TObjectPtr<UObject> ManagedResource : ManagedResources)
+		{
+			// @todo_pcg: Ideally we should avoid using AddTo/RemoveFromRoot. Once we have a PCGEditorSubsystem, we should use AddReferencedObject instead.
+			ManagedResource->AddToRoot();
+		}
+	}
+
+	EditorViewportClient->SetViewLocation(ViewLocation);
+	EditorViewportClient->SetViewLocationForOrbiting(ViewLocation, OrbitDistance);
+	EditorViewportClient->Invalidate();
+}
+
+TSharedRef<FEditorViewportClient> SPCGEditorViewport::MakeEditorViewportClient()
+{
+	EditorViewportClient = MakeShareable(new FPCGEditorViewportClient(SharedThis(this)));
+	EditorViewportClient->SetViewLocation(PCGEditorViewportConstants::DefaultViewLocation);
+	EditorViewportClient->SetViewRotation(FRotator(-25.0f, -135.0f, 0.0f));
+	EditorViewportClient->SetViewLocationForOrbiting(PCGEditorViewportConstants::DefaultViewLocation, PCGEditorViewportConstants::DefaultOrbitDistance);
+	EditorViewportClient->bSetListenerPosition = false;
+	EditorViewportClient->EngineShowFlags.EnableAdvancedFeatures();
+	EditorViewportClient->EngineShowFlags.SetGrid(false);
+	EditorViewportClient->EngineShowFlags.SetLighting(true);
+	EditorViewportClient->EngineShowFlags.SetIndirectLightingCache(true);
+	EditorViewportClient->EngineShowFlags.SetPostProcessing(true);
+	EditorViewportClient->Invalidate();
+
+	return EditorViewportClient.ToSharedRef();
+}
+
+void SPCGEditorViewport::ReleaseManagedResources()
+{
+	for (TObjectPtr<UObject> ManagedResource : ManagedResources)
+	{
+		ManagedResource->RemoveFromRoot();
+	}
+
+	ManagedResources.Empty();
+}
diff --git a/Engine/Plugins/PCG/Source/PCGEditor/Private/Widgets/AssetEditorViewport/SPCGEditorViewport.h b/Engine/Plugins/PCG/Source/PCGEditor/Private/Widgets/AssetEditorViewport/SPCGEditorViewport.h
new file mode 100644
--- /dev/null
+++ b/Engine/Plugins/PCG/Source/PCGEditor/Private/Widgets/AssetEditorViewport/SPCGEditorViewport.h
@@ -0,0 +1,36 @@
+// Copyright Epic Games, Inc. All Rights Reserved.
+
+#pragma once
+
+#include "PCGDataVisualization.h"
+
+#include "SAssetEditorViewport.h"
+
+class FPCGEditorViewportClient;
+
+class SPCGEditorViewport : public SAssetEditorViewport
+{
+public:
+	SLATE_BEGIN_ARGS(SPCGEditorViewport) {}
+	SLATE_END_ARGS()
+
+	virtual ~SPCGEditorViewport();
+
+	void Construct(const FArguments& InArgs);
+
+	/** Rebuilds the scene using the provided resources and setup callback. */
+	void SetupScene(const TArray<UObject*>& InResources, const FPCGSetupSceneFunc& SetupFunc);
+
+protected:
+	//~ Begin SAssetEditorViewport Interface
+	virtual TSharedRef<FEditorViewportClient> MakeEditorViewportClient() override;
+	//~ End SAssetEditorViewport Interface
+
+	void ReleaseManagedResources();
+
+private:
+	TSharedPtr<class FPCGEditorViewportClient> EditorViewportClient = nullptr;
+
+	/** Objects used by the scene. Must be managed for GC. */
+	TArray<TObjectPtr<UObject>> ManagedResources;
+};
diff --git a/Engine/Plugins/PCG/Source/PCGEditor/Private/Widgets/SPCGEditorGraphAttributeListView.cpp b/Engine/Plugins/PCG/Source/PCGEditor/Private/Widgets/SPCGEditorGraphAttributeListView.cpp
--- a/Engine/Plugins/PCG/Source/PCGEditor/Private/Widgets/SPCGEditorGraphAttributeListView.cpp
+++ b/Engine/Plugins/PCG/Source/PCGEditor/Private/Widgets/SPCGEditorGraphAttributeListView.cpp
@@ -25,7 +25,9 @@
 #include "PCGEditorGraph.h"
 #include "PCGEditorStyle.h"
 #include "Nodes/PCGEditorGraphNodeBase.h"
+#include "Widgets/AssetEditorViewport/SPCGEditorViewport.h"
 
+#include "Engine/StreamableManager.h"
 #include "Fonts/FontMeasure.h"
 #include "Framework/Commands/GenericCommands.h"
 #include "Framework/MultiBox/MultiBoxBuilder.h"
@@ -753,6 +755,7 @@ void SPCGEditorGraphAttributeListView::RefreshAttributeList()
 	ListViewItems.Empty();
 	ListViewHeader->ClearColumns();
 	InfoTextBlock->SetText(FText::GetEmpty());
+	ViewportWidget->SetupScene({}, FPCGSetupSceneFunc());
 	DataPtr.Reset();
 
 	const FPCGDataCollection* InspectionData = GetInspectionData();
@@ -812,6 +815,43 @@ void SPCGEditorGraphAttributeListView::RefreshAttributeList()
 		const FPCGTableVisualizerInfo TableVisualizerInfo = DataVisualization->GetTableVisualizerInfoWithDomain(PCGData, MetadataDomainID);
 		DataToVisualize = TableVisualizerInfo.Data;
 
+		if (DataToVisualize != DataPtr.Get())
+		{
+			LoadHandles = DataVisualization->LoadRequiredResources(DataToVisualize);
+		}
+
+		bool bAllResourcesLoaded = true;
+		TArray<UObject*> LoadedResources;
+
+		for (TSharedPtr<FStreamableHandle> LoadHandle : LoadHandles)
+		{
+			if (!LoadHandle.IsValid())
+			{
+				continue;
+			}
+
+			if (LoadHandle->HasLoadCompleted())
+			{
+				LoadedResources.Add(LoadHandle->GetLoadedAsset());
+			}
+			else
+			{
+				bAllResourcesLoaded = false;
+			}
+		}
+
+		if (bAllResourcesLoaded)
+		{
+			ViewportWidget->SetupScene(LoadedResources, DataVisualization->GetViewportSetupFunc(PCGData));
+		}
+		else
+		{
+			// Poll next tick.
+			bNeedsRefresh = true;
+			DataPtr = DataToVisualize;
+			return;
+		}
+
 		for (const FPCGTableVisualizerColumnInfo& ColumnInfo : TableVisualizerInfo.ColumnInfos)
 		{
 			AddColumn(DataToVisualize, ColumnInfo, TableVisualizerInfo.AccessorKeys);
diff --git a/Engine/Plugins/PCG/Source/PCGEditor/Private/Widgets/SPCGEditorGraphAttributeListView.h b/Engine/Plugins/PCG/Source/PCGEditor/Private/Widgets/SPCGEditorGraphAttributeListView.h
--- a/Engine/Plugins/PCG/Source/PCGEditor/Private/Widgets/SPCGEditorGraphAttributeListView.h
+++ b/Engine/Plugins/PCG/Source/PCGEditor/Private/Widgets/SPCGEditorGraphAttributeListView.h
@@ -22,6 +22,7 @@ class FPCGEditor;
 class FPCGMetadataAttributeBase;
 class FUICommandList;
 class SPCGEditorGraphAttributeListView;
+class SPCGEditorViewport;
 class UPCGComponent;
 class UPCGData;
 class UPCGEditorGraphNodeBase;
@@ -37,6 +38,7 @@ class SComboButton;
 class SHeaderRow;
 class SSearchBox;
 struct FSlateBrush;
+struct FStreamableHandle;
 
 namespace PCGEditorGraphAttributeListView
 {
@@ -154,6 +156,9 @@ public:
 	bool IsLocked() const { return bIsLocked; }
 	void SetIsLocked(bool bInIsLocked) { bIsLocked = bInIsLocked; }
 
+	TSharedPtr<SPCGEditorViewport> GetViewportWidget() const { return ViewportWidget; }
+	void SetViewportWidget(TSharedPtr<SPCGEditorViewport> InViewportWidget) { ViewportWidget = InViewportWidget; }
+
 private:
 	TSharedRef<SHeaderRow> CreateHeaderRowWidget() const;
 
@@ -252,6 +257,9 @@ private:
 	/** Pointer back to the PCG editor that owns us */
 	TWeakPtr<FPCGEditor> PCGEditorPtr;
 
+	/** Data visualization viewport. */
+	TSharedPtr<SPCGEditorViewport> ViewportWidget;
+
 	/** Cached PCGGraphNode being viewed */
 	TWeakObjectPtr<UPCGEditorGraphNodeBase> PCGEditorGraphNode;
 	bool bPCGEditorGraphNodeChanged = false;
@@ -311,4 +319,7 @@ private:
 
 	/** Used to ensure data collapsed for inspection is kept alive. */
 	TWeakObjectPtr<const UPCGData> DataPtr = nullptr;
+
+	/** Handles for any resources the data needs to load before visualization. */
+	TArray<TSharedPtr<FStreamableHandle>> LoadHandles;
 };
diff --git a/Engine/Plugins/PCG/Source/PCGEditor/Public/DataVisualizations/PCGStaticMeshDataVisualization.h b/Engine/Plugins/PCG/Source/PCGEditor/Public/DataVisualizations/PCGStaticMeshDataVisualization.h
new file mode 100644
--- /dev/null
+++ b/Engine/Plugins/PCG/Source/PCGEditor/Public/DataVisualizations/PCGStaticMeshDataVisualization.h
@@ -0,0 +1,20 @@
+// Copyright Epic Games, Inc. All Rights Reserved.
+
+#pragma once
+
+#include "PCGDataVisualization.h"
+#include "PCGEditorModule.h"
+
+class UPCGData;
+struct FPCGContext;
+
+class IPCGStaticMeshDataVisualization : public IPCGDataVisualization
+{
+public:
+	// ~Begin IPCGDataVisualization interface
+	virtual void ExecuteDebugDisplay(FPCGContext* Context, const UPCGSettingsInterface* SettingsInterface, const UPCGData* Data, AActor* TargetActor) const override { /* Do nothing. */ };
+	virtual FPCGTableVisualizerInfo GetTableVisualizerInfoWithDomain(const UPCGData* Data, const FPCGMetadataDomainID& DomainID) const override;
+	virtual TArray<TSharedPtr<FStreamableHandle>> LoadRequiredResources(const UPCGData* Data) const override;
+	virtual FPCGSetupSceneFunc GetViewportSetupFunc(const UPCGData* Data) const override;
+	// ~End IPCGDataVisualization interface
+};
