[PCGEditor] Add node context menu action to jump to source file

diff --git a/Engine/Plugins/PCG/Source/PCGEditor/Private/Nodes/PCGEditorGraphNodeBase.cpp b/Engine/Plugins/PCG/Source/PCGEditor/Private/Nodes/PCGEditorGraphNodeBase.cpp
--- a/Engine/Plugins/PCG/Source/PCGEditor/Private/Nodes/PCGEditorGraphNodeBase.cpp
+++ b/Engine/Plugins/PCG/Source/PCGEditor/Private/Nodes/PCGEditorGraphNodeBase.cpp
@@ -346,6 +346,7 @@ void UPCGEditorGraphNodeBase::GetNodeContextMenuActions(UToolMenu* Menu, class U
 		Section.AddMenuEntry(FPCGEditorCommands::Get().ExportNodes, TAttribute<FText>(), TAttribute<FText>(), FSlateIcon(FPCGEditorStyle::Get().GetStyleSetName(), "ClassIcon.PCGSettings"));
 		Section.AddMenuEntry(FPCGEditorCommands::Get().CollapseNodes);
 		Section.AddMenuEntry(FPCGEditorCommands::Get().ConvertToStandaloneNodes);
+		Section.AddMenuEntry(FPCGEditorCommands::Get().JumpToSource);
 	}
 
 	// Organizational actions
diff --git a/Engine/Plugins/PCG/Source/PCGEditor/Private/PCGEditor.cpp b/Engine/Plugins/PCG/Source/PCGEditor/Private/PCGEditor.cpp
--- a/Engine/Plugins/PCG/Source/PCGEditor/Private/PCGEditor.cpp
+++ b/Engine/Plugins/PCG/Source/PCGEditor/Private/PCGEditor.cpp
@@ -929,6 +929,10 @@ void FPCGEditor::BindCommands()
 		PCGEditorCommands.SelectNamedRerouteDeclaration,
 		FExecuteAction::CreateSP(this, &FPCGEditor::OnSelectNamedRerouteDeclaration),
 		FCanExecuteAction::CreateSP(this, &FPCGEditor::CanSelectNamedRerouteDeclaration));
+
+	GraphEditorCommands->MapAction(
+		PCGEditorCommands.JumpToSource,
+		FExecuteAction::CreateSP(this, &FPCGEditor::OnJumpToSource));
 }
 
 void FPCGEditor::OnFind()
@@ -1545,6 +1549,26 @@ bool FPCGEditor::CanSelectNamedRerouteDeclaration() const
 	return false;
 }
 
+void FPCGEditor::OnJumpToSource()
+{
+	if (!InternalValidationOnAction())
+	{
+		return;
+	}
+
+	for (UObject* Object : GraphEditorWidget->GetSelectedNodes())
+	{
+		const UPCGEditorGraphNodeBase* PCGEditorGraphNode = Cast<const UPCGEditorGraphNodeBase>(Object);
+		const UPCGNode* PCGNode = PCGEditorGraphNode ? PCGEditorGraphNode->GetPCGNode() : nullptr;
+		const UPCGSettings* Settings = PCGNode ? PCGNode->GetSettings() : nullptr;
+
+		if (Settings)
+		{
+			JumpToDefinition(Settings->GetClass());
+		}
+	}
+}
+
 FReply FPCGEditor::OnSpawnNodeByShortcut(FInputChord InChord, const FVector2D& InPosition, UPCGEditorGraph* InGraph)
 {
 	const TSharedPtr<FEdGraphSchemaAction> Action = FPCGEditorSpawnNodeCommands::Get().GetGraphActionByChord(InChord);
diff --git a/Engine/Plugins/PCG/Source/PCGEditor/Private/PCGEditor.h b/Engine/Plugins/PCG/Source/PCGEditor/Private/PCGEditor.h
--- a/Engine/Plugins/PCG/Source/PCGEditor/Private/PCGEditor.h
+++ b/Engine/Plugins/PCG/Source/PCGEditor/Private/PCGEditor.h
@@ -303,6 +303,9 @@ private:
 	/** Whether the user can find the declaration from the selection */
 	bool CanSelectNamedRerouteDeclaration() const;
 
+	/** Jumps to source definition for the selected nodes. */
+	void OnJumpToSource();
+
 	/** Internal method that validates a few things (& logs errors) prior to executing actions. */
 	bool InternalValidationOnAction();
 
diff --git a/Engine/Plugins/PCG/Source/PCGEditor/Private/PCGEditorCommands.cpp b/Engine/Plugins/PCG/Source/PCGEditor/Private/PCGEditorCommands.cpp
--- a/Engine/Plugins/PCG/Source/PCGEditor/Private/PCGEditorCommands.cpp
+++ b/Engine/Plugins/PCG/Source/PCGEditor/Private/PCGEditorCommands.cpp
@@ -42,6 +42,7 @@ void FPCGEditorCommands::RegisterCommands()
 	UI_COMMAND(RenameNode, "Rename Node", "Rename the selected node", EUserInterfaceActionType::Button, FInputChord(EKeys::F2));
 	UI_COMMAND(SelectNamedRerouteUsages, "Select Named Reroute Usages", "Selects all usages of this Named Reroute Declaration", EUserInterfaceActionType::Button, FInputChord());
 	UI_COMMAND(SelectNamedRerouteDeclaration, "Select Named Reroute Declaration", "Selects the associated Named Reroute Declaration matching this Named Reroute Usage", EUserInterfaceActionType::Button, FInputChord());
+	UI_COMMAND(JumpToSource, "Jump to Source", "Jumps to the associated source file.", EUserInterfaceActionType::Button, FInputChord());
 
 	//~ PLATFORM SPECIFIC ~//
 #if PLATFORM_MAC
diff --git a/Engine/Plugins/PCG/Source/PCGEditor/Private/PCGEditorCommands.h b/Engine/Plugins/PCG/Source/PCGEditor/Private/PCGEditorCommands.h
--- a/Engine/Plugins/PCG/Source/PCGEditor/Private/PCGEditorCommands.h
+++ b/Engine/Plugins/PCG/Source/PCGEditor/Private/PCGEditorCommands.h
@@ -40,6 +40,7 @@ public:
 	TSharedPtr<FUICommandInfo> RenameNode;
 	TSharedPtr<FUICommandInfo> SelectNamedRerouteUsages;
 	TSharedPtr<FUICommandInfo> SelectNamedRerouteDeclaration;
+	TSharedPtr<FUICommandInfo> JumpToSource;
 };
 
 struct FPCGSpawnNodeCommandInfo
