[PCG] Add Generate/Cleanup all options in the tool menu

* Will call generate/cleanup on all components registered.

#rnx

diff --git a/Engine/Plugins/PCG/Source/PCG/Private/PCGSubsystem.cpp b/Engine/Plugins/PCG/Source/PCG/Private/PCGSubsystem.cpp
--- a/Engine/Plugins/PCG/Source/PCG/Private/PCGSubsystem.cpp
+++ b/Engine/Plugins/PCG/Source/PCG/Private/PCGSubsystem.cpp
@@ -1821,6 +1821,29 @@ uint32 UPCGSubsystem::GetGraphCacheEntryCount(IPCGElement* InElement) const
 	return GraphExecutor ? GraphExecutor->GetGraphCacheEntryCount(InElement) : 0;
 }
 
+void UPCGSubsystem::GenerateAllPCGComponents(bool bForce) const
+{
+	for (UPCGComponent* Component : ActorAndComponentMapping.GetAllRegisteredComponents())
+	{
+		Component->Generate(bForce);
+	}
+}
+
+void UPCGSubsystem::CleanupAllPCGComponents(bool bPurge) const
+{
+	for (UPCGComponent* Component : ActorAndComponentMapping.GetAllRegisteredComponents())
+	{
+		if (bPurge)
+		{
+			Component->CleanupLocalDeleteAllGeneratedObjects({});
+		}
+		else
+		{
+			Component->Cleanup(/*bRemoveComponents=*/true);
+		}
+	}
+}
+
 const FPCGNodeVisualLogs& UPCGSubsystem::GetNodeVisualLogs() const
 {
 	if (IPCGEditorModule* PCGEditorModule = IPCGEditorModule::Get())
diff --git a/Engine/Plugins/PCG/Source/PCG/Public/PCGSubsystem.h b/Engine/Plugins/PCG/Source/PCG/Public/PCGSubsystem.h
--- a/Engine/Plugins/PCG/Source/PCG/Public/PCGSubsystem.h
+++ b/Engine/Plugins/PCG/Source/PCG/Public/PCGSubsystem.h
@@ -344,6 +344,12 @@ public:
 	/** Returns how many times InElement is present in the cache. */
 	uint32 GetGraphCacheEntryCount(IPCGElement* InElement) const;
 
+	/** Will gather all the components registered, and ask for generate. */
+	void GenerateAllPCGComponents(bool bForce) const;
+
+	/** Will gather all the components registered, and ask for cleanup. */
+	void CleanupAllPCGComponents(bool bPurge) const;
+
 	/** Get graph warnings and errors for all nodes. */
 	UE_DEPRECATED(5.6, "Use IPCGEditorModule::GetNodeVisualLogs instead")
 	const FPCGNodeVisualLogs& GetNodeVisualLogs() const;
diff --git a/Engine/Plugins/PCG/Source/PCGEditor/Private/PCGEditorModule.cpp b/Engine/Plugins/PCG/Source/PCGEditor/Private/PCGEditorModule.cpp
--- a/Engine/Plugins/PCG/Source/PCGEditor/Private/PCGEditorModule.cpp
+++ b/Engine/Plugins/PCG/Source/PCGEditor/Private/PCGEditorModule.cpp
@@ -705,6 +705,62 @@ void FPCGEditorModule::PopulateMenuActions(FMenuBuilder& MenuBuilder)
 				}
 			})),
 		NAME_None);
+
+	MenuBuilder.AddMenuEntry(
+		LOCTEXT("GenerateAllComponents", "Generate all PCG components"),
+		LOCTEXT("GenerateAllComponents_Tooltip", "Generate or refresh all the loaded PCG components, whenever they are dirty or not."),
+		FSlateIcon(),
+		FUIAction(
+			FExecuteAction::CreateLambda([]()
+			{
+				if (UPCGSubsystem* PCGSubsystem = UPCGSubsystem::GetSubsystemForCurrentWorld())
+				{
+					PCGSubsystem->GenerateAllPCGComponents(/*bForce=*/true);
+				}
+			})),
+		NAME_None);
+
+	MenuBuilder.AddMenuEntry(
+	LOCTEXT("GenerateAllComponents", "Generate all dirty PCG components"),
+	LOCTEXT("GenerateAllComponents_Tooltip", "Generate or refresh all the loaded and dirty PCG components."),
+	FSlateIcon(),
+	FUIAction(
+		FExecuteAction::CreateLambda([]()
+		{
+			if (UPCGSubsystem* PCGSubsystem = UPCGSubsystem::GetSubsystemForCurrentWorld())
+			{
+				PCGSubsystem->GenerateAllPCGComponents(/*bForce=*/false);
+			}
+		})),
+	NAME_None);
+
+	MenuBuilder.AddMenuEntry(
+	LOCTEXT("CleanupAllComponents", "Cleanup all PCG components"),
+	LOCTEXT("CleanupAllComponents_Tooltip", "Cleanup all the loaded PCG components."),
+	FSlateIcon(),
+	FUIAction(
+		FExecuteAction::CreateLambda([]()
+		{
+			if (UPCGSubsystem* PCGSubsystem = UPCGSubsystem::GetSubsystemForCurrentWorld())
+			{
+				PCGSubsystem->CleanupAllPCGComponents(/*bPurge=*/false);
+			}
+		})),
+	NAME_None);
+
+	MenuBuilder.AddMenuEntry(
+	LOCTEXT("PurgeAllComponents", "Purge all PCG components"),
+	LOCTEXT("PurgeAllComponents_Tooltip", "Cleanup all the loaded PCG components and also remove any ghost resources on the components owners."),
+	FSlateIcon(),
+	FUIAction(
+		FExecuteAction::CreateLambda([]()
+		{
+			if (UPCGSubsystem* PCGSubsystem = UPCGSubsystem::GetSubsystemForCurrentWorld())
+			{
+				PCGSubsystem->CleanupAllPCGComponents(/*bPurge=*/true);
+			}
+		})),
+	NAME_None);
 }
 
 void FPCGEditorModule::RegisterSettings()
