[PCG] Add error message when pulling PCG Component data if there is a mismatch in pins

* Useful if multiple components are pulled at the same time and some are ill-formed.

#rnx

diff --git a/Engine/Plugins/PCG/Source/PCG/Private/Elements/PCGDataFromActor.cpp b/Engine/Plugins/PCG/Source/PCG/Private/Elements/PCGDataFromActor.cpp
--- a/Engine/Plugins/PCG/Source/PCG/Private/Elements/PCGDataFromActor.cpp
+++ b/Engine/Plugins/PCG/Source/PCG/Private/Elements/PCGDataFromActor.cpp
@@ -968,6 +968,15 @@ void FPCGDataFromActorElement::ProcessActor(FPCGContext* InContext, const UPCGDa
 #endif
 
 					DuplicatedTaggedData.Tags.Add(PCGDataFromActorConstants::PCGComponentDataGridSizeTagPrefix + FString::FromInt(PCGHiGenGrid::GridToGridSize(Component->GetGenerationGrid())));
+
+					// Add some logging if we have un-expected output pin
+					if (DuplicatedTaggedData.Pin != PCGPinConstants::DefaultOutputLabel && !Settings->ExpectedPins.IsEmpty() && !Settings->ExpectedPins.Contains(TaggedData.Pin))
+					{
+						PCGLog::LogErrorOnGraph(FText::Format(LOCTEXT("InvalidOutputPinPCGComponentData", "Component '{0}' on actor '{1}' has data from an unexpected output pin: '{2}'.\nMaybe the data is outdated, try regenerate this component to refresh its output data"),
+							FText::FromString(Component->GetName()),
+							FText::FromString(PCGLog::GetExecutionSourceName(Component, /*bUseLabel=*/true)),
+							FText::FromName(TaggedData.Pin)), Context);
+					}
 				}
 			}
 		}
