[PCG] Spline mesh override support for Start/EndRollDegrees

#rnx

diff --git a/Engine/Plugins/PCG/Source/PCG/Private/Elements/PCGSpawnSplineMesh.cpp b/Engine/Plugins/PCG/Source/PCG/Private/Elements/PCGSpawnSplineMesh.cpp
--- a/Engine/Plugins/PCG/Source/PCG/Private/Elements/PCGSpawnSplineMesh.cpp
+++ b/Engine/Plugins/PCG/Source/PCG/Private/Elements/PCGSpawnSplineMesh.cpp
@@ -177,6 +177,9 @@ bool FPCGSpawnSplineMeshElement::ExecuteInternal(FPCGContext* InContext) const
 		const int NumSegments = SplineData->GetNumSegments();
 		const bool bIsClosed = SplineData->IsClosed();
 
+		const bool bStartRollOverridden = Settings->SplineMeshParamsOverride.ContainsByPredicate([](const FPCGObjectPropertyOverrideDescription& It) { return It.PropertyTarget == GET_MEMBER_NAME_CHECKED(FPCGSplineMeshParams, StartRollDegrees); });
+		const bool bEndRollOverridden = Settings->SplineMeshParamsOverride.ContainsByPredicate([](const FPCGObjectPropertyOverrideDescription& It) { return It.PropertyTarget == GET_MEMBER_NAME_CHECKED(FPCGSplineMeshParams, EndRollDegrees); });
+
 		if (IterState.ElementIndex == 0)
 		{
 			// Copy the descriptor, for the overrides
@@ -184,7 +187,7 @@ bool FPCGSpawnSplineMeshElement::ExecuteInternal(FPCGContext* InContext) const
 			IterState.SMCBuilderParams.SplineMeshParams = Settings->SplineMeshParams;
 			IterState.DescriptionOverrides.UpdateTemplateObject(&IterState.SMCBuilderParams.Descriptor);
 			IterState.ParamsOverrides.UpdateTemplateObject(&IterState.SMCBuilderParams.SplineMeshParams);
-		}		
+		}
 
 		while (IterState.ElementIndex < NumSegments)
 		{
@@ -231,11 +234,21 @@ bool FPCGSpawnSplineMeshElement::ExecuteInternal(FPCGContext* InContext) const
 			IterState.SMCBuilderParams.SplineMeshParams.EndTangent = ArriveTangent;
 
 			// Set Roll (Rotation)
-			const FQuat StartRotation = Transform.GetRotation();
-			const FQuat EndRotation = NextTransform.GetRotation();
+			if (!bStartRollOverridden || !bEndRollOverridden)
+			{
+				const FQuat StartRotation = Transform.GetRotation();
+				const FQuat EndRotation = NextTransform.GetRotation();
 
-			IterState.SMCBuilderParams.SplineMeshParams.StartRollDegrees = StartRotation.Rotator().Roll;
-			IterState.SMCBuilderParams.SplineMeshParams.EndRollDegrees = EndRotation.Rotator().Roll;
+				if (!bStartRollOverridden)
+				{
+					IterState.SMCBuilderParams.SplineMeshParams.StartRollDegrees = StartRotation.Rotator().Roll;
+				}
+				
+				if (!bEndRollOverridden)
+				{
+					IterState.SMCBuilderParams.SplineMeshParams.EndRollDegrees = EndRotation.Rotator().Roll;
+				}
+			}
 
 			// Set Scale
 			FVector StartScale = Transform.GetScale3D() * Bounds.GetExtent();
diff --git a/Engine/Plugins/PCG/Source/PCG/Public/Elements/PCGSplineMeshParams.h b/Engine/Plugins/PCG/Source/PCG/Public/Elements/PCGSplineMeshParams.h
--- a/Engine/Plugins/PCG/Source/PCG/Public/Elements/PCGSplineMeshParams.h
+++ b/Engine/Plugins/PCG/Source/PCG/Public/Elements/PCGSplineMeshParams.h
@@ -68,7 +68,8 @@ struct FPCGSplineMeshParams
 	UPROPERTY(meta = (PCG_NotOverridable))
 	FVector StartTangent = FVector::ZeroVector;
 
-	UPROPERTY(meta = (PCG_NotOverridable))
+	/** Start roll of the mesh. Computed automatically so can't be set there, but can be overridden in the Spline Mesh Params Overrides, if needed. */
+	UPROPERTY(VisibleAnywhere, Category = "Spline Mesh Params", meta = (PCG_NotOverridable))
 	float StartRollDegrees = 0.0f;
 
 	UPROPERTY(meta = (PCG_NotOverridable))
@@ -80,7 +81,8 @@ struct FPCGSplineMeshParams
 	UPROPERTY(meta = (PCG_NotOverridable))
 	FVector EndTangent = FVector::ZeroVector;
 
-	UPROPERTY(meta = (PCG_NotOverridable))
+	/** End roll of the mesh. Computed automatically so can't be set there, but can be overridden in the Spline Mesh Params Overrides, if needed. */
+	UPROPERTY(VisibleAnywhere, Category = "Spline Mesh Params", meta = (PCG_NotOverridable))
 	float EndRollDegrees = 0.0f;
 
 	UPROPERTY(meta = (PCG_NotOverridable))
