[PCG] Fix GC triggered every frame when pcg.Cache.Enabled=0

#rnx

diff --git a/Engine/Plugins/PCG/Source/PCG/Private/Graph/PCGGraphCache.cpp b/Engine/Plugins/PCG/Source/PCG/Private/Graph/PCGGraphCache.cpp
--- a/Engine/Plugins/PCG/Source/PCG/Private/Graph/PCGGraphCache.cpp
+++ b/Engine/Plugins/PCG/Source/PCG/Private/Graph/PCGGraphCache.cpp
@@ -149,7 +149,7 @@ bool FPCGGraphCache::EnforceMemoryBudget()
 	TRACE_CPUPROFILER_EVENT_SCOPE(FPCGGraphCache::FPCGGraphCache::EnforceMemoryBudget);
 	if (!CVarCacheEnabled.GetValueOnAnyThread())
 	{
-		return true; // Always return true so we can report that a GC might be needed
+		return false;
 	}
 
 	if (!CVarCacheMemoryBudgetEnabled.GetValueOnAnyThread())
