[PCG] Switch context stack pointer from raw to a safe handle with ref counting.

Fixes crash when context gets destroyed by an async task releasing a context handle which may happen after the stack is destroyed.

#rnx

diff --git a/Engine/Plugins/PCG/Source/PCG/Private/Compute/Elements/PCGComputeGraphElement.cpp b/Engine/Plugins/PCG/Source/PCG/Private/Compute/Elements/PCGComputeGraphElement.cpp
--- a/Engine/Plugins/PCG/Source/PCG/Private/Compute/Elements/PCGComputeGraphElement.cpp
+++ b/Engine/Plugins/PCG/Source/PCG/Private/Compute/Elements/PCGComputeGraphElement.cpp
@@ -27,7 +27,8 @@
 
 FPCGComputeGraphContext::~FPCGComputeGraphContext()
 {
-	TWeakObjectPtr<const UPCGGraph> GraphWeak = Stack ? Stack->GetNearestNonInlinedGraphForCurrentFrame() : nullptr;
+	const FPCGStack* StackPtr = GetStack();
+	TWeakObjectPtr<const UPCGGraph> GraphWeak = StackPtr ? StackPtr->GetNearestNonInlinedGraphForCurrentFrame() : nullptr;
 
 	// Can be here on a worker thread, after unpinning the context in one of the async tasks.
 	// Will execute immediately if on game thread, otherwise on next game thread frame.
@@ -119,13 +120,14 @@ bool FPCGComputeGraphElement::ExecuteInternal(FPCGContext* InContext) const
 
 		const IPCGGraphExecutionSource* Source = InContext->ExecutionSource.Get();
 		const UPCGGraph* TopGraph = Source ? Source->GetExecutionState().GetGraph() : nullptr;
-		if (!ensure(TopGraph))
+		const FPCGStack* StackPtr = Context->GetStack();
+		if (!ensure(TopGraph) || !ensure(StackPtr))
 		{
 			return true;
 		}
 
 		// Higen is always disabled within dynamic subgraphs - will retrieve Uninitialized tasks (which are always cooked).
-		const UPCGGraph* DynamicSubgraph = InContext->Stack ? InContext->Stack->GetNearestDynamicSubgraphForCurrentFrame() : nullptr;
+		const UPCGGraph* DynamicSubgraph = StackPtr ? StackPtr->GetNearestDynamicSubgraphForCurrentFrame() : nullptr;
 
 		if (TopGraph->IsHierarchicalGenerationEnabled() && !DynamicSubgraph)
 		{
@@ -185,7 +187,8 @@ bool FPCGComputeGraphElement::ExecuteInternal(FPCGContext* InContext) const
 			Context->DataBinding->InitializeTables(Context);
 		});
 
-		const UPCGGraph* Graph = Context->Stack ? Context->Stack->GetNearestNonInlinedGraphForCurrentFrame() : nullptr;
+		const FPCGStack* StackPtr = Context->GetStack();
+		const UPCGGraph* Graph = StackPtr ? StackPtr->GetNearestNonInlinedGraphForCurrentFrame() : nullptr;
 		const UPCGGraph::FComputeGraphInstanceKey Key = { Context->GenerationGridSize, ComputeGraphIndex };
 
 		if (!ensure(Graph))
@@ -561,8 +564,8 @@ void FPCGComputeGraphElement::PostExecuteInternal(FPCGContext* InContext) const
 {
 	check(InContext);
 	FPCGComputeGraphContext* Context = static_cast<FPCGComputeGraphContext*>(InContext);
-
-	if (!Context->DataBinding)
+	const FPCGStack* Stack = Context->GetStack();
+	if (!Context->DataBinding || !ensure(Stack))
 	{
 		return;
 	}
@@ -574,9 +577,9 @@ void FPCGComputeGraphElement::PostExecuteInternal(FPCGContext* InContext) const
 		{
 			const UPCGNode* Node = NodePtr.Get();
 			IPCGGraphExecutionSource* ExecutionSource = Context->ExecutionSource.Get();
-			if (ExecutionSource && Context->Stack && Node)
+			if (ExecutionSource && Stack && Node)
 			{
-				ExecutionSource->GetExecutionState().GetInspection().NotifyNodeExecuted(Node, Context->Stack, /*InTimer=*/nullptr, /*bUsedCache*/false);
+				ExecutionSource->GetExecutionState().GetInspection().NotifyNodeExecuted(Node, Stack, /*InTimer=*/nullptr, /*bUsedCache*/false);
 			}
 		}
 	}
@@ -715,7 +718,7 @@ bool FPCGComputeGraphElement::StoreDataForInspection(FPCGComputeGraphContext* In
 		const UPCGNode* Node = Cast<UPCGNode>(SettingsAndData.Key->GetOuter());
 
 		// TODO: Input data not yet supported.
-		InContext->ExecutionSource->GetExecutionState().GetInspection().StoreInspectionData(InContext->Stack, Node, /*InTimer=*/nullptr, /*InInputData=*/{}, SettingsAndData.Get<1>(), /*bUsedCache*/false);
+		InContext->ExecutionSource->GetExecutionState().GetInspection().StoreInspectionData(InContext->GetStack(), Node, /*InTimer=*/nullptr, /*InInputData=*/{}, SettingsAndData.Get<1>(), /*bUsedCache*/false);
 	}
 
 	InContext->DataBinding->DataToInspect.Empty();
@@ -725,7 +728,7 @@ bool FPCGComputeGraphElement::StoreDataForInspection(FPCGComputeGraphContext* In
 
 void FPCGComputeGraphElement::LogCompilationMessages(FPCGComputeGraphContext* InContext) const
 {
-	if (InContext->ExecutionSource.IsValid() && InContext->Stack)
+	if (InContext->ExecutionSource.IsValid() && InContext->GetStack())
 	{
 		using FNodeAndCompileMessages = const TPair<TObjectKey<const UPCGNode>, TArray<FComputeKernelCompileMessage>>;
 		for (FNodeAndCompileMessages& NodeAndCompileMessages : InContext->ComputeGraph->KernelToCompileMessages)
@@ -755,7 +758,7 @@ void FPCGComputeGraphElement::LogCompilationMessages(FPCGComputeGraphContext* In
 				{
 					if (IPCGEditorModule* PCGEditorModule = IPCGEditorModule::Get())
 					{
-						FPCGStack StackWithNode = *InContext->Stack;
+						FPCGStack StackWithNode = *InContext->GetStack();
 						StackWithNode.PushFrame(NodeAndCompileMessages.Get<0>().ResolveObjectPtr());
 
 						FText LogText;
diff --git a/Engine/Plugins/PCG/Source/PCG/Private/Compute/PCGComputeCommon.cpp b/Engine/Plugins/PCG/Source/PCG/Private/Compute/PCGComputeCommon.cpp
--- a/Engine/Plugins/PCG/Source/PCG/Private/Compute/PCGComputeCommon.cpp
+++ b/Engine/Plugins/PCG/Source/PCG/Private/Compute/PCGComputeCommon.cpp
@@ -67,7 +67,8 @@ namespace PCGComputeHelpers
 		{
 			if (IPCGEditorModule* PCGEditorModule = IPCGEditorModule::Get())
 			{
-				FPCGStack StackWithNode = Context->Stack ? *Context->Stack : FPCGStack();
+				const FPCGStack* Stack = Context->GetStack();
+				FPCGStack StackWithNode = Stack ? *Stack : FPCGStack();
 				StackWithNode.PushFrame(Settings->GetOuter());
 
 				PCGEditorModule->GetNodeVisualLogsMutable().Log(StackWithNode, ELogVerbosity::Warning, InText);
@@ -84,7 +85,7 @@ namespace PCGComputeHelpers
 		{
 			if (IPCGEditorModule* PCGEditorModule = IPCGEditorModule::Get())
 			{
-				FPCGStack StackWithNode = Context->Stack ? *Context->Stack : FPCGStack();
+				FPCGStack StackWithNode = Context->GetStack() ? *Context->GetStack() : FPCGStack();
 				StackWithNode.PushFrame(Settings->GetOuter());
 
 				PCGEditorModule->GetNodeVisualLogsMutable().Log(StackWithNode, ELogVerbosity::Error, InText);
diff --git a/Engine/Plugins/PCG/Source/PCG/Private/Elements/PCGCreateTargetActor.cpp b/Engine/Plugins/PCG/Source/PCG/Private/Elements/PCGCreateTargetActor.cpp
--- a/Engine/Plugins/PCG/Source/PCG/Private/Elements/PCGCreateTargetActor.cpp
+++ b/Engine/Plugins/PCG/Source/PCG/Private/Elements/PCGCreateTargetActor.cpp
@@ -430,8 +430,11 @@ bool FPCGCreateTargetActorElement::ExecuteInternal(FPCGContext* Context) const
 		{
 			ResourceCrc = Context->DependenciesCrc;
 
-			const FPCGCrc StackCRC = Context->Stack->GetCrc();
-			ResourceCrc.Combine(StackCRC);
+			if (const FPCGStack* Stack = Context->GetStack())
+			{
+				const FPCGCrc StackCRC = Stack->GetCrc();
+				ResourceCrc.Combine(StackCRC);
+			}
 
 			const FPCGCrc AdditionalCRC = PCGCreateTargetActor::GetAdditionalDependenciesCrc(SpawnDefaultActorParams, TargetActor);
 			ResourceCrc.Combine(AdditionalCRC);
diff --git a/Engine/Plugins/PCG/Source/PCG/Private/Elements/PCGDataFromActor.cpp b/Engine/Plugins/PCG/Source/PCG/Private/Elements/PCGDataFromActor.cpp
--- a/Engine/Plugins/PCG/Source/PCG/Private/Elements/PCGDataFromActor.cpp
+++ b/Engine/Plugins/PCG/Source/PCG/Private/Elements/PCGDataFromActor.cpp
@@ -660,13 +660,13 @@ void FPCGDataFromActorElement::GatherWaitTasks(AActor* FoundActor, FPCGContext*
 			const FPCGTaskId GenerateTask = Component->GenerateLocalGetTaskId(EPCGComponentGenerationTrigger::GenerateOnDemand, /*bForce=*/false);
 			if (GenerateTask != InvalidPCGTaskId)
 			{
-				PCGGraphExecutionLogging::LogGraphScheduleDependency(Component, Context->Stack);
+				PCGGraphExecutionLogging::LogGraphScheduleDependency(Component, Context->GetStack());
 
 				OutWaitTasks.Add(GenerateTask);
 			}
 			else
 			{
-				PCGGraphExecutionLogging::LogGraphScheduleDependencyFailed(Component, Context->Stack);
+				PCGGraphExecutionLogging::LogGraphScheduleDependencyFailed(Component, Context->GetStack());
 			}
 		}
 	}
diff --git a/Engine/Plugins/PCG/Source/PCG/Private/Elements/PCGGetLoopIndex.cpp b/Engine/Plugins/PCG/Source/PCG/Private/Elements/PCGGetLoopIndex.cpp
--- a/Engine/Plugins/PCG/Source/PCG/Private/Elements/PCGGetLoopIndex.cpp
+++ b/Engine/Plugins/PCG/Source/PCG/Private/Elements/PCGGetLoopIndex.cpp
@@ -45,14 +45,13 @@ bool FPCGGetLoopIndexElement::ExecuteInternal(FPCGContext* Context) const
 	const UPCGGetLoopIndexSettings* Settings = Context->GetInputSettings<UPCGGetLoopIndexSettings>();
 	check(Settings);
 
-	if (!ensure(Context->Stack))
+	const FPCGStack* Stack = Context->GetStack();
+	if (!ensure(Stack))
 	{
 		PCGE_LOG(Error, LogOnly, LOCTEXT("ContextHasNoExecutionStack", "The execution context is malformed and has no call stack."));
 		return true;
 	}
 
-	const FPCGStack* Stack = Context->Stack;
-
 	int LoopIndex = INDEX_NONE;
 
 	const TArray<FPCGStackFrame>& StackFrames = Stack->GetStackFrames();
diff --git a/Engine/Plugins/PCG/Source/PCG/Private/Elements/PCGGetSubgraphDepth.cpp b/Engine/Plugins/PCG/Source/PCG/Private/Elements/PCGGetSubgraphDepth.cpp
--- a/Engine/Plugins/PCG/Source/PCG/Private/Elements/PCGGetSubgraphDepth.cpp
+++ b/Engine/Plugins/PCG/Source/PCG/Private/Elements/PCGGetSubgraphDepth.cpp
@@ -73,13 +73,13 @@ bool FPCGGetSubgraphDepthElement::ExecuteInternal(FPCGContext* Context) const
 	bool bRecursive = Settings->Mode == EPCGSubgraphDepthMode::RecursiveDepth;
 	const UPCGGraph* TargetGraph = nullptr;
 
-	if (!ensure(Context->Stack))
+	const FPCGStack* Stack = Context->GetStack();
+	if (!ensure(Stack))
 	{
 		PCGE_LOG(Error, LogOnly, LOCTEXT("ContextHasNoExecutionStack", "The execution context is malformed and has no call stack."));
 		return true;
 	}
 
-	const FPCGStack* Stack = Context->Stack;
 	const TArray<FPCGStackFrame>& StackFrames = Stack->GetStackFrames();
 
 	// Find target graph first
diff --git a/Engine/Plugins/PCG/Source/PCG/Private/Elements/PCGLoopElement.cpp b/Engine/Plugins/PCG/Source/PCG/Private/Elements/PCGLoopElement.cpp
--- a/Engine/Plugins/PCG/Source/PCG/Private/Elements/PCGLoopElement.cpp
+++ b/Engine/Plugins/PCG/Source/PCG/Private/Elements/PCGLoopElement.cpp
@@ -277,7 +277,8 @@ bool FPCGLoopElement::ExecuteInternal(FPCGContext* InContext) const
 			}
 
 			// Prepare the invocation stack - which is the stack up to this node, and then this node, then a loop index
-			FPCGStack InvocationStack = ensure(Context->Stack) ? *Context->Stack : FPCGStack();
+			const FPCGStack* Stack = Context->GetStack();
+			FPCGStack InvocationStack = ensure(Stack) ? *Stack : FPCGStack();
 
 			TArray<FPCGStackFrame>& StackFrames = InvocationStack.GetStackFramesMutable();
 			StackFrames.Reserve(StackFrames.Num() + 2);
diff --git a/Engine/Plugins/PCG/Source/PCG/Private/Elements/PCGPrintElement.cpp b/Engine/Plugins/PCG/Source/PCG/Private/Elements/PCGPrintElement.cpp
--- a/Engine/Plugins/PCG/Source/PCG/Private/Elements/PCGPrintElement.cpp
+++ b/Engine/Plugins/PCG/Source/PCG/Private/Elements/PCGPrintElement.cpp
@@ -167,7 +167,7 @@ bool FPCGPrintElement::ExecuteInternal(FPCGContext* Context) const
 			}
 
 			// Add the stack hash to include the loop indices
-			HashKey32 = HashCombine(HashKey32, GetTypeHash(Context->Stack));
+			HashKey32 = HashCombine(HashKey32, GetTypeHash(Context->GetStack()));
 
 			check(SourceComponent);
 			UPCGManagedDebugStringMessageKey* ManagedMessageKey = FPCGContext::NewObject_AnyThread<UPCGManagedDebugStringMessageKey>(Context, SourceComponent);
diff --git a/Engine/Plugins/PCG/Source/PCG/Private/Elements/PCGSpawnActor.cpp b/Engine/Plugins/PCG/Source/PCG/Private/Elements/PCGSpawnActor.cpp
--- a/Engine/Plugins/PCG/Source/PCG/Private/Elements/PCGSpawnActor.cpp
+++ b/Engine/Plugins/PCG/Source/PCG/Private/Elements/PCGSpawnActor.cpp
@@ -637,7 +637,8 @@ bool FPCGSpawnActorElement::SpawnAndPrepareSubgraphs(FPCGSubgraphContext* Contex
 							Context->AddToReferencedObjects(SubgraphInputData);
 
 							// Prepare the invocation stack - which is the stack up to this node, and then this node, then a loop index
-							FPCGStack InvocationStack = ensure(Context->Stack) ? *Context->Stack : FPCGStack();
+							const FPCGStack* Stack = Context->GetStack();
+							FPCGStack InvocationStack = ensure(Stack) ? *Stack : FPCGStack();
 
 							UPCGGraph* Graph = GraphInterface->GetGraph();
 
diff --git a/Engine/Plugins/PCG/Source/PCG/Private/Graph/PCGGraphExecutor.cpp b/Engine/Plugins/PCG/Source/PCG/Private/Graph/PCGGraphExecutor.cpp
--- a/Engine/Plugins/PCG/Source/PCG/Private/Graph/PCGGraphExecutor.cpp
+++ b/Engine/Plugins/PCG/Source/PCG/Private/Graph/PCGGraphExecutor.cpp
@@ -1377,7 +1377,10 @@ void FPCGGraphExecutor::PrepareForExecute(FPCGGraphTask& Task, FCachedResult*& O
 		Task.Context->TaskId = Task.NodeId;
 		Task.Context->CompiledTaskId = Task.CompiledTaskId;
 		Task.Context->DependenciesCrc = DependenciesCrc;
+		Task.Context->StackHandle = Task.GetStackHandle();
+		PRAGMA_DISABLE_DEPRECATION_WARNINGS
 		Task.Context->Stack = Task.GetStack();
+		PRAGMA_ENABLE_DEPRECATION_WARNINGS
 
 		// Safe to remove async flags on CombineParams here
 		PCGGraphExecutor::ClearAsyncFlags(Task.CombineParamsAsyncObjects);
@@ -2911,7 +2914,7 @@ namespace PCGGraphExecutor
 				if (ToGridSize == PCGHiGenGrid::UnboundedGridSize())
 				{
 					PCGEditorModule->GetNodeVisualLogsMutable().Log(
-						*InContext->Stack,
+						*InContext->GetStack(),
 						ELogVerbosity::Error,
 						FText::Format(
 							NSLOCTEXT("PCGGraphCompiler", "InvalidLinkageToUnbounded", "Could not read data across grid levels - cannot read from grid size {0} to Unbounded domain."),
@@ -2921,7 +2924,7 @@ namespace PCGGraphExecutor
 				else
 				{
 					PCGEditorModule->GetNodeVisualLogsMutable().Log(
-						*InContext->Stack,
+						*InContext->GetStack(),
 						ELogVerbosity::Error,
 						FText::Format(
 							NSLOCTEXT("PCGGraphCompiler", "InvalidLinkageInvalidGridSizes", "Could not read data across grid levels - origin grid size {0} must be greater than destination grid size {1}. Graph default grid size may need increasing."),
diff --git a/Engine/Plugins/PCG/Source/PCG/Private/Graph/PCGGraphTask.cpp b/Engine/Plugins/PCG/Source/PCG/Private/Graph/PCGGraphTask.cpp
--- a/Engine/Plugins/PCG/Source/PCG/Private/Graph/PCGGraphTask.cpp
+++ b/Engine/Plugins/PCG/Source/PCG/Private/Graph/PCGGraphTask.cpp
@@ -18,6 +18,11 @@ const FPCGStack* FPCGGraphTask::GetStack() const
 	return (StackContext && StackIndex != INDEX_NONE) ? StackContext->GetStack(StackIndex) : nullptr;
 }
 
+FPCGStackHandle FPCGGraphTask::GetStackHandle() const
+{
+	return FPCGStackHandle(StackContext, StackIndex);
+}
+
 FPCGTaskId FPCGGraphTask::GetGraphExecutionTaskId() const
 {
 	return StackContext ? StackContext->GetGraphExecutionTaskId() : InvalidPCGTaskId;
diff --git a/Engine/Plugins/PCG/Source/PCG/Private/Graph/PCGGraphTask.h b/Engine/Plugins/PCG/Source/PCG/Private/Graph/PCGGraphTask.h
--- a/Engine/Plugins/PCG/Source/PCG/Private/Graph/PCGGraphTask.h
+++ b/Engine/Plugins/PCG/Source/PCG/Private/Graph/PCGGraphTask.h
@@ -94,6 +94,7 @@ struct FPCGGraphTask
 #endif
 
 	const FPCGStack* GetStack() const;
+	FPCGStackHandle GetStackHandle() const;
 	FPCGTaskId GetGraphExecutionTaskId() const;
 	bool CanExecuteOnlyOnMainThread() const;
 
diff --git a/Engine/Plugins/PCG/Source/PCG/Private/Helpers/PCGHelpers.cpp b/Engine/Plugins/PCG/Source/PCG/Private/Helpers/PCGHelpers.cpp
--- a/Engine/Plugins/PCG/Source/PCG/Private/Helpers/PCGHelpers.cpp
+++ b/Engine/Plugins/PCG/Source/PCG/Private/Helpers/PCGHelpers.cpp
@@ -606,9 +606,9 @@ namespace PCGHelpers
 		{
 			GetGeneratedActorsFolderPath(InTargetActor, OutFolderPath);
 		}
-		else if (AttachOptions == EPCGAttachOptions::InGraphFolder && InContext && InContext->Stack && InContext->Stack->GetRootGraph())
+		else if (AttachOptions == EPCGAttachOptions::InGraphFolder && InContext && InContext->GetStack() && InContext->GetStack()->GetRootGraph())
 		{
-			OutFolderPath = InContext->Stack->GetRootGraph()->GetName() + "_Generated";
+			OutFolderPath = InContext->GetStack()->GetRootGraph()->GetName() + "_Generated";
 		}
 		else // Generated folder
 		{
diff --git a/Engine/Plugins/PCG/Source/PCG/Private/PCGCommon.cpp b/Engine/Plugins/PCG/Source/PCG/Private/PCGCommon.cpp
--- a/Engine/Plugins/PCG/Source/PCG/Private/PCGCommon.cpp
+++ b/Engine/Plugins/PCG/Source/PCG/Private/PCGCommon.cpp
@@ -283,3 +283,19 @@ namespace PCGQualityHelpers
 		return SelectedPinLabel;
 	}
 }
+
+FPCGStackHandle::FPCGStackHandle(TSharedPtr<const FPCGStackContext> InStackContext, int32 InStackIndex)
+	: StackContext(InStackContext)
+	, StackIndex(InStackIndex)
+{
+}
+
+bool FPCGStackHandle::IsValid() const
+{
+	return StackContext && StackIndex >= 0 && StackIndex < StackContext->GetNumStacks();
+}
+
+const FPCGStack* FPCGStackHandle::GetStack() const
+{
+	return IsValid() ? StackContext->GetStack(StackIndex) : nullptr;
+}
diff --git a/Engine/Plugins/PCG/Source/PCG/Private/PCGContext.cpp b/Engine/Plugins/PCG/Source/PCG/Private/PCGContext.cpp
--- a/Engine/Plugins/PCG/Source/PCG/Private/PCGContext.cpp
+++ b/Engine/Plugins/PCG/Source/PCG/Private/PCGContext.cpp
@@ -90,7 +90,8 @@ UPCGBasePointData* FPCGContext::NewPointData_AnyThread(FPCGContext* Context)
 
 FPCGTaskId FPCGContext::GetGraphExecutionTaskId() const
 {
-	return ensure(Stack) ? Stack->GetGraphExecutionTaskId() : InvalidPCGTaskId;
+	const FPCGStack* StackPtr = GetStack();
+	return ensure(StackPtr) ? StackPtr->GetGraphExecutionTaskId() : InvalidPCGTaskId;
 }
 
 FString FPCGContext::GetTaskName() const
@@ -489,7 +490,8 @@ void FPCGContext::LogVisual(ELogVerbosity::Type InVerbosity, const FText& InMess
 {
 	if (IPCGEditorModule* PCGEditorModule = IPCGEditorModule::Get())
 	{
-		FPCGStack StackWithNode = Stack ? FPCGStack(*Stack) : FPCGStack();
+		const FPCGStack* StackPtr = GetStack();
+		FPCGStack StackWithNode = StackPtr ? FPCGStack(*StackPtr) : FPCGStack();
 		StackWithNode.PushFrame(Node);
 		PCGEditorModule->GetNodeVisualLogsMutable().Log(StackWithNode, InVerbosity, InMessage);
 	}
@@ -499,7 +501,8 @@ bool FPCGContext::HasVisualLogs() const
 {
 	if (IPCGEditorModule* PCGEditorModule = IPCGEditorModule::Get())
 	{
-		FPCGStack StackWithNode = Stack ? FPCGStack(*Stack) : FPCGStack();
+		const FPCGStack* StackPtr = GetStack();
+		FPCGStack StackWithNode = StackPtr ? FPCGStack(*StackPtr) : FPCGStack();
 		StackWithNode.PushFrame(Node);
 		return PCGEditorModule->GetNodeVisualLogs().HasLogs(StackWithNode);
 	}
diff --git a/Engine/Plugins/PCG/Source/PCG/Private/PCGInputOutputSettings.cpp b/Engine/Plugins/PCG/Source/PCG/Private/PCGInputOutputSettings.cpp
--- a/Engine/Plugins/PCG/Source/PCG/Private/PCGInputOutputSettings.cpp
+++ b/Engine/Plugins/PCG/Source/PCG/Private/PCGInputOutputSettings.cpp
@@ -87,7 +87,8 @@ bool FPCGInputOutputElement::SupportsGPUResidentData(FPCGContext* InContext) con
 	check(InContext);
 
 	const UPCGGraphInputOutputSettings* Settings = InContext->GetOriginalSettings<UPCGGraphInputOutputSettings>();
-	if (!Settings || !ensure(InContext->Stack))
+	const FPCGStack* Stack = InContext->GetStack();
+	if (!Settings || !ensure(Stack))
 	{
 		return false;
 	}
@@ -99,7 +100,7 @@ bool FPCGInputOutputElement::SupportsGPUResidentData(FPCGContext* InContext) con
 	else
 	{
 		// The output node of subgraphs can return GPU data proxies, but the top graph output must read back.
-		return !InContext->Stack->IsCurrentFrameInRootGraph();
+		return !Stack->IsCurrentFrameInRootGraph();
 	}
 }
 
diff --git a/Engine/Plugins/PCG/Source/PCG/Private/PCGSubgraph.cpp b/Engine/Plugins/PCG/Source/PCG/Private/PCGSubgraph.cpp
--- a/Engine/Plugins/PCG/Source/PCG/Private/PCGSubgraph.cpp
+++ b/Engine/Plugins/PCG/Source/PCG/Private/PCGSubgraph.cpp
@@ -738,9 +738,9 @@ bool FPCGSubgraphElement::ExecuteInternal(FPCGContext* InContext) const
 	bool bIsRecursive = false;
 	if (Subgraph)
 	{
-		if (Context->Stack)
+		if (const FPCGStack* Stack = Context->GetStack())
 		{
-			bIsRecursive = (bIsDynamic && Context->Stack->HasObject(Subgraph)) || Subgraph->Contains(Context->Stack->GetGraphForCurrentFrame());
+			bIsRecursive = (bIsDynamic && Stack->HasObject(Subgraph)) || Subgraph->Contains(Stack->GetGraphForCurrentFrame());
 		}
 		else if(Context->ExecutionSource.Get())
 		{
@@ -774,7 +774,8 @@ bool FPCGSubgraphElement::ExecuteInternal(FPCGContext* InContext) const
 				}
 
 				// Prepare the invocation stack - which is the stack up to this node, and then this node and the 'not-a-loop index' which we use to differentiate dynamic vs static subgraphs
-				FPCGStack InvocationStack = ensure(Context->Stack) ? *Context->Stack : FPCGStack();
+				const FPCGStack* Stack = Context->GetStack();
+				FPCGStack InvocationStack = ensure(Stack) ? *Stack : FPCGStack();
 				InvocationStack.GetStackFramesMutable().Emplace(Context->Node);
 				InvocationStack.GetStackFramesMutable().Emplace(INDEX_NONE); // not a loop index
 
diff --git a/Engine/Plugins/PCG/Source/PCG/Private/Utils/PCGExtraCapture.cpp b/Engine/Plugins/PCG/Source/PCG/Private/Utils/PCGExtraCapture.cpp
--- a/Engine/Plugins/PCG/Source/PCG/Private/Utils/PCGExtraCapture.cpp
+++ b/Engine/Plugins/PCG/Source/PCG/Private/Utils/PCGExtraCapture.cpp
@@ -19,7 +19,7 @@ void PCGUtils::FExtraCapture::ResetCapturedMessages()
 
 void PCGUtils::FExtraCapture::Update(const PCGUtils::FScopedCall& InScopedCall)
 {
-	if (!InScopedCall.Context || !InScopedCall.Context->Stack)
+	if (!InScopedCall.Context || !InScopedCall.Context->GetStack())
 	{
 		return;
 	}
diff --git a/Engine/Plugins/PCG/Source/PCG/Public/PCGCommon.h b/Engine/Plugins/PCG/Source/PCG/Public/PCGCommon.h
--- a/Engine/Plugins/PCG/Source/PCG/Public/PCGCommon.h
+++ b/Engine/Plugins/PCG/Source/PCG/Public/PCGCommon.h
@@ -17,6 +17,7 @@ class UPCGComponent;
 class UPCGGraph;
 struct FPCGContext;
 struct FPCGStack;
+struct FPCGStackContext;
 
 namespace PCGLog
 {
@@ -608,6 +609,22 @@ struct FPCGScheduleGenericParams
 	bool bCanExecuteOnlyOnMainThread = true;
 };
 
+/** A handle to a call stack which holds onto the stack memory. */
+struct PCG_API FPCGStackHandle
+{
+	FPCGStackHandle() = default;
+	FPCGStackHandle(TSharedPtr<const FPCGStackContext> InStackContext, int32 InStackIndex);
+
+	bool IsValid() const;
+
+	const FPCGStack* GetStack() const;
+
+private:
+	TSharedPtr<const FPCGStackContext> StackContext;
+
+	int32 StackIndex = INDEX_NONE;
+};
+
 // Enable to debug if some PCG Elements are creating some data that should be pre-created by the graph PreGraph element
 #define PCG_EXECUTION_CACHE_VALIDATION_ENABLED 0
 
diff --git a/Engine/Plugins/PCG/Source/PCG/Public/PCGContext.h b/Engine/Plugins/PCG/Source/PCG/Public/PCGContext.h
--- a/Engine/Plugins/PCG/Source/PCG/Public/PCGContext.h
+++ b/Engine/Plugins/PCG/Source/PCG/Public/PCGContext.h
@@ -118,7 +118,10 @@ PRAGMA_ENABLE_DEPRECATION_WARNINGS
 	EPCGExecutionPhase CurrentPhase = EPCGExecutionPhase::NotExecuted;
 	TArray<TPair<FPCGDataCollection, FPCGDataCollection>> CachedInputToOutputInternalResults;
 
-	/** The current call stack. */
+	/** Get the current call stack. */
+	const FPCGStack* GetStack() const { return StackHandle.GetStack(); }
+
+	UE_DEPRECATED(5.6, "Access the stack via GetStack()")
 	const FPCGStack* Stack = nullptr;
 
 	bool CanExecuteOnlyOnMainThread() const { return bOverrideSettingsOnMainThread; } 
@@ -314,6 +317,9 @@ private:
 	// Graph Executor that scheduled the task
 	FPCGGraphExecutor* GraphExecutor = nullptr;
 
+	/** Handle for the current call stack. */
+	FPCGStackHandle StackHandle;
+
 #if WITH_EDITOR
 	friend class PCGUtils::FExtraCapture;
 	PCGUtils::FCallTime Timer;
