[PCG] Add missing profile scope around the buffer reuse code which has non-trivial cost.

#rnx

diff --git a/Engine/Plugins/PCG/Source/PCG/Private/Compute/PCGDataUploadAdaptor.cpp b/Engine/Plugins/PCG/Source/PCG/Private/Compute/PCGDataUploadAdaptor.cpp
--- a/Engine/Plugins/PCG/Source/PCG/Private/Compute/PCGDataUploadAdaptor.cpp
+++ b/Engine/Plugins/PCG/Source/PCG/Private/Compute/PCGDataUploadAdaptor.cpp
@@ -229,6 +229,8 @@ FRDGBufferSRVRef FPCGDataUploadAdaptor::GetDataElementCountsBufferSRV(FRDGBuilde
 /** Gets the buffer that can then be used to read the data in kernels. */
 FRDGBufferRef FPCGDataUploadAdaptor::GetBuffer_RenderThread(FRDGBuilder& InGraphBuilder, EPCGExportMode InExportMode)
 {
+	RDG_EVENT_SCOPE(InGraphBuilder, "FPCGDataUploadAdaptor::GetBuffer_RenderThread");
+
 	if (!ExternalBufferForReuse)
 	{
 		// Not reusing an existing buffer, create a new buffer and pack the input data collection.
