[PCG] Added some robustness to the PA original to local map by using object keys instead of weak pointers.

#rnx

diff --git a/Engine/Plugins/PCG/Source/PCG/Private/Grid/PCGPartitionActor.cpp b/Engine/Plugins/PCG/Source/PCG/Private/Grid/PCGPartitionActor.cpp
--- a/Engine/Plugins/PCG/Source/PCG/Private/Grid/PCGPartitionActor.cpp
+++ b/Engine/Plugins/PCG/Source/PCG/Private/Grid/PCGPartitionActor.cpp
@@ -213,7 +213,7 @@ void APCGPartitionActor::BeginPlay()
 		// Pass through all the pcg components, to verify if we need to generate them.
 		for (auto& It : OriginalToLocal)
 		{
-			UPCGComponent* OriginalComponent = It.Key.Get();
+			UPCGComponent* OriginalComponent = It.Key.ResolveObjectPtr();
 			UPCGComponent* LocalComponent = It.Value;
 			if (OriginalComponent && LocalComponent)
 			{
@@ -240,10 +240,10 @@ void APCGPartitionActor::EndPlay(const EEndPlayReason::Type EndPlayReason)
 	// It is possible for a PA to get called EndPlay before its original volume actor so make sure to cleanup properly here
 	if (IsRuntimeGenerated())
 	{
-		TMap<TWeakObjectPtr<UPCGComponent>, TObjectPtr<UPCGComponent>> OriginalToLocalCopy(OriginalToLocal);
-		for (const TTuple<TWeakObjectPtr<UPCGComponent>, TObjectPtr<UPCGComponent>>& OriginalToLocalItem : OriginalToLocalCopy)
+		TMap<TObjectKey<UPCGComponent>, TObjectPtr<UPCGComponent>> OriginalToLocalCopy(OriginalToLocal);
+		for (const TTuple<TObjectKey<UPCGComponent>, TObjectPtr<UPCGComponent>>& OriginalToLocalItem : OriginalToLocalCopy)
 		{
-			RemoveGraphInstance(OriginalToLocalItem.Key.GetEvenIfUnreachable());
+			RemoveGraphInstance(OriginalToLocalItem.Key.ResolveObjectPtrEvenIfGarbage());
 		}
 	}
 
@@ -432,7 +432,7 @@ void APCGPartitionActor::CleanupDeadGraphInstances(bool bRemoveNonNullOnly)
 void APCGPartitionActor::CleanupDeadGraphInstancesInternal()
 {
 	// First find if we have any local dead instance (= nullptr) hooked to an original component.
-	TSet<TWeakObjectPtr<UPCGComponent>> DeadOriginalInstances;
+	TSet<TObjectKey<UPCGComponent>> DeadOriginalInstances;
 	for (const auto& OriginalToLocalItem : OriginalToLocal)
 	{
 		if (!OriginalToLocalItem.Value)
@@ -445,7 +445,7 @@ void APCGPartitionActor::CleanupDeadGraphInstancesInternal()
 	{
 		Modify();
 
-		for (const TWeakObjectPtr<UPCGComponent>& DeadInstance : DeadOriginalInstances)
+		for (const TObjectKey<UPCGComponent>& DeadInstance : DeadOriginalInstances)
 		{
 			OriginalToLocal.Remove(DeadInstance);
 		}
@@ -705,9 +705,9 @@ TSet<TObjectPtr<UPCGComponent>> APCGPartitionActor::GetAllLocalPCGComponents() c
 TSet<TObjectPtr<UPCGComponent>> APCGPartitionActor::GetAllOriginalPCGComponents() const
 {
 	TSet<TObjectPtr<UPCGComponent>> ResultComponents;
-	for(const TPair<TWeakObjectPtr<UPCGComponent>, TObjectPtr<UPCGComponent>>& OriginalPair : OriginalToLocal)
+	for(const TPair<TObjectKey<UPCGComponent>, TObjectPtr<UPCGComponent>>& OriginalPair : OriginalToLocal)
 	{
-		if (UPCGComponent* OriginalComponent = OriginalPair.Key.Get())
+		if (UPCGComponent* OriginalComponent = OriginalPair.Key.ResolveObjectPtr())
 		{
 			ResultComponents.Add(OriginalComponent);
 		}
diff --git a/Engine/Plugins/PCG/Source/PCG/Public/Grid/PCGPartitionActor.h b/Engine/Plugins/PCG/Source/PCG/Public/Grid/PCGPartitionActor.h
--- a/Engine/Plugins/PCG/Source/PCG/Public/Grid/PCGPartitionActor.h
+++ b/Engine/Plugins/PCG/Source/PCG/Public/Grid/PCGPartitionActor.h
@@ -145,7 +145,7 @@ private:
 	UPCGComponent* GetLocalComponent(const UPCGComponent* OriginalComponent, bool bRebuildMappingOnNullEntries) const;
 
 	// Note: this map is not a property and not serialized since we will rebuild it from the LocalToOriginal
-	mutable TMap<TWeakObjectPtr<UPCGComponent>, TObjectPtr<UPCGComponent>> OriginalToLocal;
+	mutable TMap<TObjectKey<UPCGComponent>, TObjectPtr<UPCGComponent>> OriginalToLocal;
 
 	UPROPERTY(NonTransactional)
 	TMap<TObjectPtr<UPCGComponent>, TSoftObjectPtr<UPCGComponent>> LocalToOriginal;
