[WorldPartition Minimap] Fix Excluded Data layers (broken since we used assets).
* Add the Included/Excluded Data layers asset directly on the WP Builder and update the builders (In use by Miinimap, NavMesh & PCG)

#tests Excluding DL on a map
[RN]

diff --git a/Engine/Plugins/PCG/Source/PCGEditor/Private/WorldPartitionBuilder/PCGWorldPartitionBuilder.cpp b/Engine/Plugins/PCG/Source/PCGEditor/Private/WorldPartitionBuilder/PCGWorldPartitionBuilder.cpp
--- a/Engine/Plugins/PCG/Source/PCGEditor/Private/WorldPartitionBuilder/PCGWorldPartitionBuilder.cpp
+++ b/Engine/Plugins/PCG/Source/PCGEditor/Private/WorldPartitionBuilder/PCGWorldPartitionBuilder.cpp
@@ -519,35 +519,12 @@ bool UPCGWorldPartitionBuilder::PreRun(UWorld* World, FPackageSourceControlHelpe
 	Args.bIterativeCellLoading &= World->IsPartitionedWorld();
 
 	// Setup DataLayers
-	if (UDataLayerManager* DataLayerManager = UDataLayerManager::GetDataLayerManager(World))
-	{
-		auto FillDataLayerShortNames = [World, DataLayerManager](const TArray<TObjectPtr<UDataLayerAsset>>& InDataLayerAssets, TSet<FName>& OutDataLayerShortNames)
-		{
-			for(const TObjectPtr<UDataLayerAsset>& DataLayer : InDataLayerAssets)
-			{
-				if (DataLayer != nullptr)
-				{
-					const UDataLayerInstance* DataLayerInstance = DataLayerManager->GetDataLayerInstance(DataLayer);
-					if (DataLayerInstance == nullptr)
-					{
-						UE_LOG(LogPCGWorldPartitionBuilder, Display, TEXT("World '%s' does not contain a DataLayerInstance for '%s'"), *World->GetPathName(), *DataLayer->GetPathName());
-					}
-					else
-					{
-						OutDataLayerShortNames.Add(FName(DataLayerInstance->GetDataLayerShortName()));
-					}
-				}
-			}
-		};
-
-		// Set Base Builder Class fields
-		bLoadNonDynamicDataLayers = Args.bLoadEditorOnlyDataLayers;
-		bLoadInitiallyActiveDataLayers = Args.bLoadActivatedRuntimeDataLayers;
-
-		FillDataLayerShortNames(Args.IncludedDataLayers, DataLayerShortNames);
-		FillDataLayerShortNames(Args.ExcludedDataLayers, ExcludedDataLayerShortNames);
-	}
+	// Set Base Builder Class fields
+	bLoadNonDynamicDataLayers = Args.bLoadEditorOnlyDataLayers;
+	bLoadInitiallyActiveDataLayers = Args.bLoadActivatedRuntimeDataLayers;
 
+	IncludedDataLayers.Append(Args.IncludedDataLayers);
+	ExcludedDataLayers.Append(Args.ExcludedDataLayers);
 
 	if (GetLoadingMode() == ELoadingMode::IterativeCells2D)
 	{
diff --git a/Engine/Source/Editor/UnrealEd/Private/WorldPartition/WorldPartitionBuilder.cpp b/Engine/Source/Editor/UnrealEd/Private/WorldPartition/WorldPartitionBuilder.cpp
--- a/Engine/Source/Editor/UnrealEd/Private/WorldPartition/WorldPartitionBuilder.cpp
+++ b/Engine/Source/Editor/UnrealEd/Private/WorldPartition/WorldPartitionBuilder.cpp
@@ -316,10 +316,13 @@ void UWorldPartitionBuilder::LoadDataLayers(UWorld* InWorld)
 		{
 			const FName DataLayerShortName(DataLayer->GetDataLayerShortName());
 
+			PRAGMA_DISABLE_DEPRECATION_WARNINGS
 			// Load all Non Excluded Data Layers + Non DynamicallyLoaded Data Layers + Initially Active Data Layers + Data Layers provided by builder
-			const bool bLoadedInEditor = !ExcludedDataLayerShortNames.Contains(DataLayerShortName) && ((bLoadNonDynamicDataLayers && !DataLayer->IsRuntime()) ||
+			const bool bLoadedInEditor = !ExcludedDataLayers.Contains(DataLayer->GetAsset()) && !ExcludedDataLayerShortNames.Contains(DataLayerShortName) && ((bLoadNonDynamicDataLayers && !DataLayer->IsRuntime()) ||
 											(bLoadInitiallyActiveDataLayers && DataLayer->GetInitialRuntimeState() == EDataLayerRuntimeState::Activated) ||
-										DataLayerShortNames.Contains(DataLayerShortName));
+											IncludedDataLayers.Contains(DataLayer->GetAsset()) || DataLayerShortNames.Contains(DataLayerShortName));
+			PRAGMA_ENABLE_DEPRECATION_WARNINGS
+			
 			if (DataLayer->IsLoadedInEditor() != bLoadedInEditor)
 			{
 				bUpdateEditorCells = true;
diff --git a/Engine/Source/Editor/UnrealEd/Private/WorldPartition/WorldPartitionMiniMapBuilder.cpp b/Engine/Source/Editor/UnrealEd/Private/WorldPartition/WorldPartitionMiniMapBuilder.cpp
--- a/Engine/Source/Editor/UnrealEd/Private/WorldPartition/WorldPartitionMiniMapBuilder.cpp
+++ b/Engine/Source/Editor/UnrealEd/Private/WorldPartition/WorldPartitionMiniMapBuilder.cpp
@@ -82,19 +82,7 @@ bool UWorldPartitionMiniMapBuilder::PreRun(UWorld* World, FPackageSourceControlH
 		WorldToMinimap *= FScaleMatrix(Scaling);
 	}
 
-	// Gather excluded data layers
-	{
-		UDataLayerManager* DataLayerManager = UDataLayerManager::GetDataLayerManager(World);
-		for (const FActorDataLayer& ActorDataLayer : WorldMiniMap->ExcludedDataLayers)
-		{
-			const UDataLayerInstance* DataLayerInstance = DataLayerManager->GetDataLayerInstance(ActorDataLayer.Name);
-
-			if (DataLayerInstance != nullptr)
-			{
-				ExcludedDataLayerShortNames.Add(FName(DataLayerInstance->GetDataLayerShortName()));
-			}
-		}
-	}
+	ExcludedDataLayers = WorldMiniMap->ExcludedDataLayers;
 
 	return true;
 }
diff --git a/Engine/Source/Editor/UnrealEd/Private/WorldPartition/WorldPartitionNavigationDataBuilder.cpp b/Engine/Source/Editor/UnrealEd/Private/WorldPartition/WorldPartitionNavigationDataBuilder.cpp
--- a/Engine/Source/Editor/UnrealEd/Private/WorldPartition/WorldPartitionNavigationDataBuilder.cpp
+++ b/Engine/Source/Editor/UnrealEd/Private/WorldPartition/WorldPartitionNavigationDataBuilder.cpp
@@ -41,7 +41,7 @@ bool UWorldPartitionNavigationDataBuilder::PreRun(UWorld* World, FPackageSourceC
 			}
 			else if (DataLayerInstance->IsRuntime())
 			{
-				DataLayerShortNames.Add(FName(DataLayerInstance->GetDataLayerShortName()));
+				IncludedDataLayers.Add(const_cast<UDataLayerAsset*>(DataLayerInstance->GetAsset()));
 			}
 		}
 	}
diff --git a/Engine/Source/Editor/UnrealEd/Public/WorldPartition/WorldPartitionBuilder.h b/Engine/Source/Editor/UnrealEd/Public/WorldPartition/WorldPartitionBuilder.h
--- a/Engine/Source/Editor/UnrealEd/Public/WorldPartition/WorldPartitionBuilder.h
+++ b/Engine/Source/Editor/UnrealEd/Public/WorldPartition/WorldPartitionBuilder.h
@@ -147,8 +147,14 @@ protected:
 	int32 IterativeCellOverlapSize = 0;
 	FBox  IterativeWorldBounds;
 
+	UE_DEPRECATED(5.6, "Please use IncludedDataLayers")
 	TSet<FName> DataLayerShortNames;
+	UE_DEPRECATED(5.6, "Please use ExcludedDataLayers")
 	TSet<FName> ExcludedDataLayerShortNames;
+
+	TSet<TObjectPtr<UDataLayerAsset>> IncludedDataLayers;
+	TSet<TObjectPtr<UDataLayerAsset>> ExcludedDataLayers;
+
 	bool bLoadNonDynamicDataLayers = true;
 	bool bLoadInitiallyActiveDataLayers = true;
 
diff --git a/Engine/Source/Runtime/Engine/Public/WorldPartition/WorldPartitionMiniMap.h b/Engine/Source/Runtime/Engine/Public/WorldPartition/WorldPartitionMiniMap.h
--- a/Engine/Source/Runtime/Engine/Public/WorldPartition/WorldPartitionMiniMap.h
+++ b/Engine/Source/Runtime/Engine/Public/WorldPartition/WorldPartitionMiniMap.h
@@ -52,7 +52,7 @@ public:
 
 	/* Datalayers excluded from MiniMap rendering */
 	UPROPERTY(EditAnywhere, Category = WorldPartitionMiniMap)
-	TSet<FActorDataLayer> ExcludedDataLayers;
+	TSet<TObjectPtr<UDataLayerAsset>> ExcludedDataLayers;
 
 	/**
 	 * Target world units per pixel for the minimap texture. 
