[PCG] Discard generation sources that lie in unloaded WP areas.

#rnx

diff --git a/Engine/Plugins/PCG/Source/PCG/Private/RuntimeGen/PCGGenSourceManager.cpp b/Engine/Plugins/PCG/Source/PCG/Private/RuntimeGen/PCGGenSourceManager.cpp
--- a/Engine/Plugins/PCG/Source/PCG/Private/RuntimeGen/PCGGenSourceManager.cpp
+++ b/Engine/Plugins/PCG/Source/PCG/Private/RuntimeGen/PCGGenSourceManager.cpp
@@ -2,6 +2,7 @@
 
 #include "RuntimeGen/PCGGenSourceManager.h"
 
+#include "PCGModule.h"
 #include "PCGWorldActor.h"
 #include "RuntimeGen/GenSources/PCGGenSourceBase.h"
 #include "RuntimeGen/GenSources/PCGGenSourceComponent.h"
@@ -13,6 +14,8 @@
 #include "Engine/World.h"
 #include "GameFramework/GameModeBase.h"
 #include "WorldPartition/WorldPartition.h"
+#include "WorldPartition/WorldPartitionRuntimeCell.h"
+#include "WorldPartition/WorldPartitionSubsystem.h"
 
 #if WITH_EDITOR
 #include "EditorViewportClient.h"
@@ -56,25 +59,70 @@ FPCGGenSourceManager::~FPCGGenSourceManager()
 
 TSet<IPCGGenSourceBase*> FPCGGenSourceManager::GetAllGenSources(const APCGWorldActor* InPCGWorldActor)
 {
+	TRACE_CPUPROFILER_EVENT_SCOPE(FPCGGenSourceManager::GetAllGenSources);
+
 	if (bDirty)
 	{
 		UpdatePerTickGenSources(InPCGWorldActor);
 		bDirty = false;
 	}
 
+	UWorldPartitionSubsystem* WorldPartitionSubsystem = InPCGWorldActor ? UWorld::GetSubsystem<UWorldPartitionSubsystem>(InPCGWorldActor->GetWorld()) : nullptr;
+	
+	// Prepare query upfront.
+	TArray<FWorldPartitionStreamingQuerySource> QuerySources;
+	FWorldPartitionStreamingQuerySource& QuerySource = QuerySources.Emplace_GetRef();
+	QuerySource.bSpatialQuery = true;
+	QuerySource.bUseGridLoadingRange = false;
+	// Providing 1cm is standard when checking loading state.
+	QuerySource.Radius = 1.0f;
+	QuerySource.bDataLayersOnly = false;
+
+	auto IsWorldLoaded = [WorldPartitionSubsystem, &QuerySources](IPCGGenSourceBase* InGenSource)
+	{
+		TRACE_CPUPROFILER_EVENT_SCOPE(QueryIsWorldLoaded);
+
+		if (TOptional<FVector> Position = InGenSource->GetPosition())
+		{
+			QuerySources[0].Location = *Position;
+
+			const bool bIsLoaded = WorldPartitionSubsystem->IsStreamingCompleted(EWorldPartitionRuntimeCellState::Activated, QuerySources, /*bExactState*/ false);
+			if (!bIsLoaded)
+			{
+				UE_LOG(LogPCG, Verbose, TEXT("Discarded gen source due to world not loaded."));
+			}
+
+			return bIsLoaded;
+		}
+		else
+		{
+			return false;
+		}
+	};
+
 	TSet<IPCGGenSourceBase*> AllGenSources;
 	AllGenSources.Reserve(RegisteredGenSources.Num() + RegisteredNamedGenSources.Num());
-	AllGenSources.Append(RegisteredGenSources);
+
+	for (IPCGGenSourceBase* GenSource : RegisteredGenSources)
+	{
+		if (IsWorldLoaded(GenSource))
+		{
+			AllGenSources.Add(GenSource);
+		}
+	}
 
 	for (const TPair<FName, IPCGGenSourceBase*>& NamedGenSource : RegisteredNamedGenSources)
 	{
-		AllGenSources.Add(NamedGenSource.Value);
+		if (IsWorldLoaded(NamedGenSource.Value))
+		{
+			AllGenSources.Add(NamedGenSource.Value);
+		}
 	}
 
 	// Acquire a generation source for each active streaming source that doesn't have a generation source already (based on name collision test).
 	for (UPCGGenSourceWPStreamingSource* WPGenSource : WorldPartitionGenSources)
 	{
-		if (WPGenSource && WPGenSource->StreamingSource && !RegisteredNamedGenSources.Contains(WPGenSource->StreamingSource->Name))
+		if (WPGenSource && WPGenSource->StreamingSource && !RegisteredNamedGenSources.Contains(WPGenSource->StreamingSource->Name) && IsWorldLoaded(WPGenSource))
 		{
 			AllGenSources.Add(WPGenSource);
 		}
@@ -82,7 +130,7 @@ TSet<IPCGGenSourceBase*> FPCGGenSourceManager::GetAllGenSources(const APCGWorldA
 
 #if WITH_EDITOR
 	// Acquire a generation source for the active editor viewport if one exists.
-	if (EditorCameraGenSource && EditorCameraGenSource->EditorViewportClient)
+	if (EditorCameraGenSource && EditorCameraGenSource->EditorViewportClient && IsWorldLoaded(EditorCameraGenSource))
 	{
 		AllGenSources.Add(EditorCameraGenSource);
 	}
