[PCG] Make the visualization data a strong ptr on the attribute list view to keep it alive

* This was changed to a weak, in order to fix a leak.
  - The leak was happening when we were inspecting the Output node, and the data we were inspecting was serialized on the PCG Component, and we were changing maps.
  - We were holding a strong reference in the ALV to something that was outered and referenced on an actor.
* Add the same logic than the debug tree view to the ALV. On map change, ask for a refresh, and release the strong reference.
* TODO: Move the strong ptr to a AddReferencedObjects.

#rnx

diff --git a/Engine/Plugins/PCG/Source/PCGEditor/Private/PCGEditor.cpp b/Engine/Plugins/PCG/Source/PCGEditor/Private/PCGEditor.cpp
--- a/Engine/Plugins/PCG/Source/PCGEditor/Private/PCGEditor.cpp
+++ b/Engine/Plugins/PCG/Source/PCGEditor/Private/PCGEditor.cpp
@@ -3182,10 +3182,7 @@ void FPCGEditor::OnMapChanged(UWorld* InWorld, EMapChangeType InMapChangedType)
 {
 	if (InMapChangedType != EMapChangeType::SaveMap)
 	{
-		if (DebugObjectTreeWidget.IsValid())
-		{
-			DebugObjectTreeWidget->RequestRefresh();
-		}
+		RefreshViewsOnLevelChange();
 
 		// Subsystem has been torn down and rebuilt.
 		if (GEditor)
@@ -3215,6 +3212,22 @@ void FPCGEditor::OnLevelActorDeleted(AActor* InActor)
 	}
 }
 
+void FPCGEditor::RefreshViewsOnLevelChange()
+{
+	if (DebugObjectTreeWidget.IsValid())
+	{
+		DebugObjectTreeWidget->RequestRefresh();
+	}
+
+	for (TSharedPtr<SPCGEditorGraphAttributeListView>& AttributeWidget : AttributesWidgets)
+	{
+		if (AttributeWidget.IsValid())
+		{
+			AttributeWidget->RequestRefresh();
+		}
+	}
+}
+
 TSharedRef<SDockTab> FPCGEditor::SpawnTab_GraphEditor(const FSpawnTabArgs& Args)
 {
 	return SNew(SDockTab)
diff --git a/Engine/Plugins/PCG/Source/PCGEditor/Private/PCGEditor.h b/Engine/Plugins/PCG/Source/PCGEditor/Private/PCGEditor.h
--- a/Engine/Plugins/PCG/Source/PCGEditor/Private/PCGEditor.h
+++ b/Engine/Plugins/PCG/Source/PCGEditor/Private/PCGEditor.h
@@ -405,6 +405,8 @@ private:
 	void OnEndPIE(bool bIsSimulating);
 	void OnLevelActorDeleted(AActor* InActor);
 
+	void RefreshViewsOnLevelChange();
+
 	TSharedRef<SDockTab> SpawnTab_GraphEditor(const FSpawnTabArgs& Args);
 	TSharedRef<SDockTab> SpawnTab_PropertyDetails(const FSpawnTabArgs& Args, int PropertyDetailsIndex);
 	TSharedRef<SDockTab> SpawnTab_Palette(const FSpawnTabArgs& Args);
diff --git a/Engine/Plugins/PCG/Source/PCGEditor/Private/Widgets/SPCGEditorGraphAttributeListView.cpp b/Engine/Plugins/PCG/Source/PCGEditor/Private/Widgets/SPCGEditorGraphAttributeListView.cpp
--- a/Engine/Plugins/PCG/Source/PCGEditor/Private/Widgets/SPCGEditorGraphAttributeListView.cpp
+++ b/Engine/Plugins/PCG/Source/PCGEditor/Private/Widgets/SPCGEditorGraphAttributeListView.cpp
@@ -143,7 +143,8 @@ TSharedRef<SWidget> SPCGListViewItemRow::GenerateWidgetForColumn(const FName& Co
 
 	if (FPCGColumnData* PCGColumnData = SharedAttributeListView->PCGColumnData.Find(ColumnId))
 	{
-		if (PCGColumnData->DataAccessor.IsValid() && PCGColumnData->DataKeys.IsValid())
+		// We have to make sure the data is still valid otherwise this will crash.
+		if (ensure(SharedAttributeListView->DataPtr.IsValid()) && PCGColumnData->DataAccessor.IsValid() && PCGColumnData->DataKeys.IsValid())
 		{
 			int32 Index = InternalItem->Index;
 			auto Callback = [&PCGColumnData, Index, &RowText, &bShouldComputeCellWidth, &bIsClickable] (auto Dummy)
@@ -628,6 +629,13 @@ void SPCGEditorGraphAttributeListView::Tick(const FGeometry& AllottedGeometry, c
 	}
 }
 
+void SPCGEditorGraphAttributeListView::RequestRefresh()
+{
+	bNeedsRefresh = true;
+	// When we ask for refresh, we should also release the hold on the visualized data
+	DataPtr.Reset();
+}
+
 TSharedRef<SHeaderRow> SPCGEditorGraphAttributeListView::CreateHeaderRowWidget() const
 {
 	return SNew(SHeaderRow);
@@ -848,7 +856,7 @@ void SPCGEditorGraphAttributeListView::RefreshAttributeList()
 		{
 			// Poll next tick.
 			bNeedsRefresh = true;
-			DataPtr = DataToVisualize;
+			DataPtr.Reset(DataToVisualize);
 			return;
 		}
 
@@ -878,7 +886,7 @@ void SPCGEditorGraphAttributeListView::RefreshAttributeList()
 		FocusOnDataCallback.Reset();
 	}
 
-	DataPtr = DataToVisualize;
+	DataPtr.Reset(DataToVisualize);
 
 	ListView->SetItemsSource(&ListViewItems);
 	ListView->RequestListRefresh();
diff --git a/Engine/Plugins/PCG/Source/PCGEditor/Private/Widgets/SPCGEditorGraphAttributeListView.h b/Engine/Plugins/PCG/Source/PCGEditor/Private/Widgets/SPCGEditorGraphAttributeListView.h
--- a/Engine/Plugins/PCG/Source/PCGEditor/Private/Widgets/SPCGEditorGraphAttributeListView.h
+++ b/Engine/Plugins/PCG/Source/PCGEditor/Private/Widgets/SPCGEditorGraphAttributeListView.h
@@ -148,7 +148,7 @@ public:
 
 	virtual void Tick(const FGeometry& AllottedGeometry, const double InCurrentTime, const float InDeltaTime) override;
 
-	void RequestRefresh() { bNeedsRefresh = true; }
+	void RequestRefresh();
 
 	UPCGEditorGraphNodeBase* GetNodeBeingInspected() const;
 	void SetNodeBeingInspected(UPCGEditorGraphNodeBase* InNode);
@@ -317,8 +317,12 @@ private:
 
 	TSharedPtr<FPCGListViewUpdater> CurrentUpdateTask = nullptr;
 
-	/** Used to ensure data collapsed for inspection is kept alive. */
-	TWeakObjectPtr<const UPCGData> DataPtr = nullptr;
+	/** In some cases, like when inspecting temporary collapsed point data, there is no owner of the data.
+	* As a temporary solution, since the ALV needs the data for visualization and double-click functionality,
+	* it can store a strong reference to the data to keep it alive during inspection. 
+	* @todo_pcg: Replace this with AddReferencedObjects
+	*/
+	TStrongObjectPtr<const UPCGData> DataPtr = nullptr;
 
 	/** Handles for any resources the data needs to load before visualization. */
 	TArray<TSharedPtr<FStreamableHandle>> LoadHandles;
