[PCG] Add guard against invalid execution source after crash was hit.

#rnx

diff --git a/Engine/Plugins/PCG/Source/PCG/Private/Elements/PCGHiGenGridSize.cpp b/Engine/Plugins/PCG/Source/PCG/Private/Elements/PCGHiGenGridSize.cpp
--- a/Engine/Plugins/PCG/Source/PCG/Private/Elements/PCGHiGenGridSize.cpp
+++ b/Engine/Plugins/PCG/Source/PCG/Private/Elements/PCGHiGenGridSize.cpp
@@ -208,6 +208,12 @@ bool FPCGHiGenGridSizeElement::ExecuteInternal(FPCGContext* Context) const
 	TRACE_CPUPROFILER_EVENT_SCOPE(FPCGHiGenGridSizeElement::Execute);
 	check(Context);
 
+	if (!Context->ExecutionSource.IsValid())
+	{
+		UE_LOG(LogPCG, Error, TEXT("FPCGHiGenGridSizeElement: Invalid execution source, aborting."));
+		return true;
+	}
+
 	const UPCGHiGenGridSizeSettings* Settings = Context->GetInputSettings<UPCGHiGenGridSizeSettings>();
 	check(Settings);
 
