[PCG] When matching generation sources by name, also check location to fix broken tracking of cinematic cameras.

#rnx

diff --git a/Engine/Plugins/PCG/Source/PCG/Private/RuntimeGen/PCGGenSourceManager.cpp b/Engine/Plugins/PCG/Source/PCG/Private/RuntimeGen/PCGGenSourceManager.cpp
--- a/Engine/Plugins/PCG/Source/PCG/Private/RuntimeGen/PCGGenSourceManager.cpp
+++ b/Engine/Plugins/PCG/Source/PCG/Private/RuntimeGen/PCGGenSourceManager.cpp
@@ -122,7 +122,20 @@ TSet<IPCGGenSourceBase*> FPCGGenSourceManager::GetAllGenSources(const APCGWorldA
 	// Acquire a generation source for each active streaming source that doesn't have a generation source already (based on name collision test).
 	for (UPCGGenSourceWPStreamingSource* WPGenSource : WorldPartitionGenSources)
 	{
-		if (WPGenSource && WPGenSource->StreamingSource && !RegisteredNamedGenSources.Contains(WPGenSource->StreamingSource->Name) && IsWorldLoaded(WPGenSource))
+		bool bAlreadyTracked = false;
+
+		if (IPCGGenSourceBase** FoundGenSource = (WPGenSource && WPGenSource->StreamingSource) ? RegisteredNamedGenSources.Find(WPGenSource->StreamingSource->Name) : nullptr)
+		{
+			// We've seen during cinematics there are two generation sources with the same name but with different locations in the world, so match based on position too.
+			TOptional<FVector> WPGenSourcePosition = WPGenSource->GetPosition();
+			TOptional<FVector> FoundGenSourcePosition = (*FoundGenSource)->GetPosition();
+			if (WPGenSourcePosition && FoundGenSourcePosition && (*FoundGenSourcePosition - *WPGenSourcePosition).SquaredLength() < UE_KINDA_SMALL_NUMBER)
+			{
+				bAlreadyTracked = true;
+			}
+		}
+
+		if (!bAlreadyTracked && IsWorldLoaded(WPGenSource))
 		{
 			AllGenSources.Add(WPGenSource);
 		}
