[PCG] Convert most remaining PCG Nodes / Datas to support UPCGPointArrayData

#rnx

diff --git a/Engine/Plugins/Experimental/PCGInterops/PCGNiagaraInterop/Source/PCGNiagaraInterop/Public/Elements/PCGWriteToNiagaraDataChannel.h b/Engine/Plugins/Experimental/PCGInterops/PCGNiagaraInterop/Source/PCGNiagaraInterop/Public/Elements/PCGWriteToNiagaraDataChannel.h
--- a/Engine/Plugins/Experimental/PCGInterops/PCGNiagaraInterop/Source/PCGNiagaraInterop/Public/Elements/PCGWriteToNiagaraDataChannel.h
+++ b/Engine/Plugins/Experimental/PCGInterops/PCGNiagaraInterop/Source/PCGNiagaraInterop/Public/Elements/PCGWriteToNiagaraDataChannel.h
@@ -66,5 +66,6 @@ public:
 protected:
 	virtual bool PrepareDataInternal(FPCGContext* InContext) const override;
 	virtual bool ExecuteInternal(FPCGContext* InContext) const override;
+	virtual bool SupportsBasePointDataInputs(FPCGContext* InContext) const override { return true; }
 };
 
diff --git a/Engine/Plugins/PCG/Source/PCG/Private/Compute/PCGDataForGPU.cpp b/Engine/Plugins/PCG/Source/PCG/Private/Compute/PCGDataForGPU.cpp
--- a/Engine/Plugins/PCG/Source/PCG/Private/Compute/PCGDataForGPU.cpp
+++ b/Engine/Plugins/PCG/Source/PCG/Private/Compute/PCGDataForGPU.cpp
@@ -1270,11 +1270,9 @@ EPCGUnpackDataCollectionResult FPCGDataCollectionDesc::UnpackDataCollection(FPCG
 		{
 			TRACE_CPUPROFILER_EVENT_SCOPE(UnpackPointDataItem);
 
-			UPCGPointData* OutPointData = FPCGContext::NewObject_AnyThread<UPCGPointData>(InContext);
+			UPCGBasePointData* OutPointData = FPCGContext::NewPointData_AnyThread(InContext);
 			UPCGMetadata* Metadata = OutPointData->MutableMetadata();
-			TArray<FPCGPoint>& OutPoints = OutPointData->GetMutablePoints();
-
-			OutPoints.SetNumUninitialized(NumElements);
+			OutPointData->SetNumPoints(NumElements, /*bInitializeValues=*/false);
 
 			{
 				TRACE_CPUPROFILER_EVENT_SCOPE(InitializeOutput);
@@ -1282,25 +1280,32 @@ EPCGUnpackDataCollectionResult FPCGDataCollectionDesc::UnpackDataCollection(FPCG
 				// We only need to add the entry keys if there are actually attributes to unpack.
 				if (DataDesc.HasMetadataAttributes())
 				{
+					TPCGValueRange<int64> MetadataEntryRange = OutPointData->GetMetadataEntryValueRange(/*bAllocate=*/true);
 					TArray<PCGMetadataEntryKey*> ParentEntryKeys;
 					ParentEntryKeys.Reserve(NumElements);
 
 					for (uint32 ElementIndex = 0; ElementIndex < NumElements; ++ElementIndex)
 					{
-						PCGMetadataEntryKey& EntryKey = OutPoints[ElementIndex].MetadataEntry;
+						PCGMetadataEntryKey& EntryKey = MetadataEntryRange[ElementIndex];
 						EntryKey = PCGInvalidEntryKey;
 						ParentEntryKeys.Add(&EntryKey);
 					}
 
 					Metadata->AddEntriesInPlace(ParentEntryKeys);
 				}
-				else
+				else if(OutPointData->IsA<UPCGPointData>())
 				{
-					ParallelFor(NumElements, [&OutPoints](int32 ElementIndex)
+					// Special case so that we can use ParallelFor (UPCGPointArrayData doesn't need this)
+					TPCGValueRange<int64> MetadataEntryRange = OutPointData->GetMetadataEntryValueRange(/*bAllocate=*/true);
+					ParallelFor(NumElements, [&MetadataEntryRange](int32 ElementIndex)
 					{
-						OutPoints[ElementIndex].MetadataEntry = -1;
+						MetadataEntryRange[ElementIndex] = -1;
 					});
 				}
+				else
+				{
+					OutPointData->SetMetadataEntry(-1);
+				}
 			}
 
 			FPCGTaggedData& OutTaggedData = OutData.Emplace_GetRef();
@@ -1345,20 +1350,22 @@ EPCGUnpackDataCollectionResult FPCGDataCollectionDesc::UnpackDataCollection(FPCG
 					{
 					case POINT_POSITION_ATTRIBUTE_ID:
 					{
-						ParallelFor(NumElements, [DataAsFloat, AddressUints, &OutPoints](int32 ElementIndex)
+						TPCGValueRange<FTransform> TransformRange = OutPointData->GetTransformValueRange(/*bAlocate=*/true);
+						ParallelFor(NumElements, [DataAsFloat, AddressUints, &TransformRange](int32 ElementIndex)
 						{
 							const FVector Location = FVector(
 								DataAsFloat[AddressUints + ElementIndex * 3 + 0],
 								DataAsFloat[AddressUints + ElementIndex * 3 + 1],
 								DataAsFloat[AddressUints + ElementIndex * 3 + 2]);
 
-							OutPoints[ElementIndex].Transform.SetLocation(Location);
+							TransformRange[ElementIndex].SetLocation(Location);
 						});
 						break;
 					}
 					case POINT_ROTATION_ATTRIBUTE_ID:
 					{
-						ParallelFor(NumElements, [DataAsFloat, AddressUints, &OutPoints](int32 ElementIndex)
+						TPCGValueRange<FTransform> TransformRange = OutPointData->GetTransformValueRange(/*bAlocate=*/true);
+						ParallelFor(NumElements, [DataAsFloat, AddressUints, &TransformRange](int32 ElementIndex)
 						{
 							const FQuat Rotation = FQuat(
 								DataAsFloat[AddressUints + ElementIndex * 4 + 0],
@@ -1367,13 +1374,14 @@ EPCGUnpackDataCollectionResult FPCGDataCollectionDesc::UnpackDataCollection(FPCG
 								DataAsFloat[AddressUints + ElementIndex * 4 + 3]);
 
 							// Normalize here with default tolerance (zero quat will return identity).
-							OutPoints[ElementIndex].Transform.SetRotation(Rotation.GetNormalized());
+							TransformRange[ElementIndex].SetRotation(Rotation.GetNormalized());
 						});
 						break;
 					}
 					case POINT_SCALE_ATTRIBUTE_ID:
 					{
-						ParallelFor(NumElements, [DataAsFloat, AddressUints, &OutPoints](int32 ElementIndex)
+						TPCGValueRange<FTransform> TransformRange = OutPointData->GetTransformValueRange(/*bAlocate=*/true);
+						ParallelFor(NumElements, [DataAsFloat, AddressUints, &TransformRange](int32 ElementIndex)
 						{
 							const FVector Scale = FVector
 							(
@@ -1381,39 +1389,42 @@ EPCGUnpackDataCollectionResult FPCGDataCollectionDesc::UnpackDataCollection(FPCG
 								DataAsFloat[AddressUints + ElementIndex * 3 + 1],
 								DataAsFloat[AddressUints + ElementIndex * 3 + 2]);
 
-							OutPoints[ElementIndex].Transform.SetScale3D(Scale);
+							TransformRange[ElementIndex].SetScale3D(Scale);
 						});
 						break;
 					}
 					case POINT_BOUNDS_MIN_ATTRIBUTE_ID:
 					{
-						ParallelFor(NumElements, [DataAsFloat, AddressUints, &OutPoints](int32 ElementIndex)
+						TPCGValueRange<FVector> BoundsMinRange = OutPointData->GetBoundsMinValueRange(/*bAllocate=*/true);
+						ParallelFor(NumElements, [DataAsFloat, AddressUints, &BoundsMinRange](int32 ElementIndex)
 						{
 							const FVector BoundsMin = FVector(
 								DataAsFloat[AddressUints + ElementIndex * 3 + 0],
 								DataAsFloat[AddressUints + ElementIndex * 3 + 1],
 								DataAsFloat[AddressUints + ElementIndex * 3 + 2]);
 
-							OutPoints[ElementIndex].BoundsMin = BoundsMin;
+							BoundsMinRange[ElementIndex] = BoundsMin;
 						});
 						break;
 					}
 					case POINT_BOUNDS_MAX_ATTRIBUTE_ID:
 					{
-						ParallelFor(NumElements, [DataAsFloat, AddressUints, &OutPoints](int32 ElementIndex)
+						TPCGValueRange<FVector> BoundsMaxRange = OutPointData->GetBoundsMaxValueRange(/*bAllocate=*/true);
+						ParallelFor(NumElements, [DataAsFloat, AddressUints, &BoundsMaxRange](int32 ElementIndex)
 						{
 							const FVector BoundsMax = FVector(
 								DataAsFloat[AddressUints + ElementIndex * 3 + 0],
 								DataAsFloat[AddressUints + ElementIndex * 3 + 1],
 								DataAsFloat[AddressUints + ElementIndex * 3 + 2]);
 
-							OutPoints[ElementIndex].BoundsMax = BoundsMax;
+							BoundsMaxRange[ElementIndex] = BoundsMax;
 						});
 						break;
 					}
 					case POINT_COLOR_ATTRIBUTE_ID:
 					{
-						ParallelFor(NumElements, [DataAsFloat, AddressUints, &OutPoints](int32 ElementIndex)
+						TPCGValueRange<FVector4> ColorRange = OutPointData->GetColorValueRange(/*bAllocate=*/true);
+						ParallelFor(NumElements, [DataAsFloat, AddressUints, &ColorRange](int32 ElementIndex)
 						{
 							const FVector4 Color = FVector4(
 								DataAsFloat[AddressUints + ElementIndex * 4 + 0],
@@ -1421,31 +1432,34 @@ EPCGUnpackDataCollectionResult FPCGDataCollectionDesc::UnpackDataCollection(FPCG
 								DataAsFloat[AddressUints + ElementIndex * 4 + 2],
 								DataAsFloat[AddressUints + ElementIndex * 4 + 3]);
 
-							OutPoints[ElementIndex].Color = Color;
+							ColorRange[ElementIndex] = Color;
 						});
 						break;
 					}
 					case POINT_DENSITY_ATTRIBUTE_ID:
 					{
-						ParallelFor(NumElements, [DataAsFloat, AddressUints, &OutPoints](int32 ElementIndex)
+						TPCGValueRange<float> DensityRange = OutPointData->GetDensityValueRange(/*bAllocate=*/true);
+						ParallelFor(NumElements, [DataAsFloat, AddressUints, &DensityRange](int32 ElementIndex)
 						{
-							OutPoints[ElementIndex].Density = DataAsFloat[AddressUints + ElementIndex];
+							DensityRange[ElementIndex] = DataAsFloat[AddressUints + ElementIndex];
 						});
 						break;
 					}
 					case POINT_SEED_ATTRIBUTE_ID:
 					{
-						ParallelFor(NumElements, [DataAsInt, AddressUints, &OutPoints](int32 ElementIndex)
+						TPCGValueRange<int32> SeedRange = OutPointData->GetSeedValueRange(/*bAllocate=*/true);
+						ParallelFor(NumElements, [DataAsInt, AddressUints, &SeedRange](int32 ElementIndex)
 						{
-							OutPoints[ElementIndex].Seed = DataAsInt[AddressUints + ElementIndex];
+							SeedRange[ElementIndex] = DataAsInt[AddressUints + ElementIndex];
 						});
 						break;
 					}
 					case POINT_STEEPNESS_ATTRIBUTE_ID:
 					{
-						ParallelFor(NumElements, [DataAsFloat, AddressUints, &OutPoints](int32 ElementIndex)
+						TPCGValueRange<float> SteepnessRange = OutPointData->GetSteepnessValueRange(/*bAllocate=*/true);
+						ParallelFor(NumElements, [DataAsFloat, AddressUints, &SteepnessRange](int32 ElementIndex)
 						{
-							OutPoints[ElementIndex].Steepness = DataAsFloat[AddressUints + ElementIndex];
+							SteepnessRange[ElementIndex] = DataAsFloat[AddressUints + ElementIndex];
 						});
 						break;
 					}
@@ -1466,20 +1480,49 @@ EPCGUnpackDataCollectionResult FPCGDataCollectionDesc::UnpackDataCollection(FPCG
 			}
 
 			// TODO: It may be more efficient to create a mapping from input point index to final output point index and do everything in one pass.
-			auto DiscardInvalidPoints = [&OutPoints](int32 Index, FPCGPoint& OutPoint) -> bool
+			auto ProcessRangeFunc = [OutPointData](int32 StartReadIndex, int32 StartWriteIndex, int32 Count)
 			{
-				if (OutPoints[Index].Density == PCGComputeConstants::INVALID_DENSITY)
+				int32 NumWritten = 0;
+
+				const FConstPCGPointValueRanges InRanges(OutPointData);
+				FPCGPointValueRanges OutRanges(OutPointData, /*bAllocate=*/false);
+
+				for (int32 ReadIndex = StartReadIndex; ReadIndex < StartReadIndex + Count; ++ReadIndex)
 				{
-					return false;
+					if (InRanges.DensityRange[ReadIndex] == PCGComputeConstants::INVALID_DENSITY)
+					{
+						continue;
+					}
+
+					const int32 WriteIndex = StartWriteIndex + NumWritten;
+					OutRanges.SetFromValueRanges(WriteIndex, InRanges, ReadIndex);
+					++NumWritten;
 				}
 
-				OutPoint = OutPoints[Index];
-				return true;
+				return NumWritten;
+			};
+
+			auto MoveDataRangeFunc = [OutPointData](int32 RangeStartIndex, int32 MoveToIndex, int32 NumElements)
+			{
+				OutPointData->MoveRange(RangeStartIndex, MoveToIndex, NumElements);
+			};
+
+			auto FinishedFunc = [OutPointData](int32 NumWritten)
+			{
+				OutPointData->SetNumPoints(NumWritten);
 			};
 
 			{
 				TRACE_CPUPROFILER_EVENT_SCOPE(DiscardInvalidPoints);
-				FPCGAsync::AsyncPointProcessing(InContext, OutPoints.Num(), OutPoints, DiscardInvalidPoints);
+				FPCGAsyncState* AsyncState = InContext ? &InContext->AsyncState : nullptr;
+				FPCGAsync::AsyncProcessingRangeEx(
+					AsyncState, 
+					OutPointData->GetNumPoints(),
+					[]{},
+					ProcessRangeFunc,
+					MoveDataRangeFunc,
+					FinishedFunc,
+					/*bEnableTimeSlicing=*/false);
 			}
 		}
 		else if (TypeId == PARAM_DATA_TYPE_ID)
diff --git a/Engine/Plugins/PCG/Source/PCG/Private/Data/PCGSplineInteriorSurfaceData.cpp b/Engine/Plugins/PCG/Source/PCG/Private/Data/PCGSplineInteriorSurfaceData.cpp
--- a/Engine/Plugins/PCG/Source/PCG/Private/Data/PCGSplineInteriorSurfaceData.cpp
+++ b/Engine/Plugins/PCG/Source/PCG/Private/Data/PCGSplineInteriorSurfaceData.cpp
@@ -3,6 +3,7 @@
 #include "Data/PCGSplineInteriorSurfaceData.h"
 
 #include "PCGContext.h"
+#include "Data/PCGPointArrayData.h"
 #include "Data/PCGPointData.h"
 #include "Data/PCGSplineData.h"
 #include "Elements/PCGSplineSampler.h"
@@ -100,10 +101,20 @@ bool UPCGSplineInteriorSurfaceData::ProjectPoint(const FTransform& InTransform,
 }
 
 const UPCGPointData* UPCGSplineInteriorSurfaceData::CreatePointData(FPCGContext* Context) const
+{
+	return CastChecked<UPCGPointData>(CreateBasePointData(Context, UPCGPointData::StaticClass()));
+}
+
+const UPCGPointArrayData* UPCGSplineInteriorSurfaceData::CreatePointArrayData(FPCGContext* Context, const FBox& InBounds) const
+{
+	return CastChecked<UPCGPointArrayData>(CreateBasePointData(Context, UPCGPointArrayData::StaticClass()));
+}
+
+const UPCGBasePointData* UPCGSplineInteriorSurfaceData::CreateBasePointData(FPCGContext* Context, TSubclassOf<UPCGBasePointData> PointDataClass) const
 {
 	TRACE_CPUPROFILER_EVENT_SCOPE(UPCGSplineData::CreatePointData);
 
-	UPCGPointData* Data = FPCGContext::NewObject_AnyThread<UPCGPointData>(Context);
+	UPCGBasePointData* Data = FPCGContext::NewObject_AnyThread<UPCGBasePointData>(Context, GetTransientPackage(), PointDataClass);
 	Data->InitializeFromData(this);
 
 	FPCGSplineSamplerParams SamplerParams;
@@ -115,9 +126,7 @@ const UPCGPointData* UPCGSplineInteriorSurfaceData::CreatePointData(FPCGContext*
 
 	// TODO: Ideally SampleInteriorData could just consume a UPCGSplineInteriorSurfaceData or FPCGSplineStruct, to avoid the extra copy on the spline struct.
 	PCGSplineSamplerHelpers::SampleInteriorData(Context, SplineData, /*InBoundingShape=*/nullptr, /*InProjectionTarget=*/nullptr, /*InProjectionParams=*/{}, SamplerParams, Data);
-	UE_LOG(LogPCG, Verbose, TEXT("SplineInteriorSurface generated %d points."), Data->GetPoints().Num());
-
-	SplineData->MarkAsGarbage();
+	UE_LOG(LogPCG, Verbose, TEXT("SplineInteriorSurface generated %d points."), Data->GetNumPoints());
 
 	return Data;
 }
@@ -162,7 +171,7 @@ void UPCGSplineInteriorSurfaceData::CacheData(FPCGContext* Context)
 	SamplerParams.Mode = EPCGSplineSamplingMode::Subdivision;
 	SamplerParams.SubdivisionsPerSegment = 5;
 
-	UPCGPointData* PointData = FPCGContext::NewObject_AnyThread<UPCGPointData>(Context);
+	UPCGBasePointData* PointData = FPCGContext::NewPointData_AnyThread(Context);
 	UPCGSplineData* SplineData = FPCGContext::NewObject_AnyThread<UPCGSplineData>(Context);
 	SplineData->Initialize(SplineStruct);
 
@@ -171,10 +180,11 @@ void UPCGSplineInteriorSurfaceData::CacheData(FPCGContext* Context)
 	PCGSplineSamplerHelpers::SampleLineData(Context, SplineData, /*InBoundingShape=*/nullptr, /*InProjectionTarget=*/nullptr, /*InProjectionParams=*/{}, SamplerParams, PointData);
 
 	// Cache the points which describe the polygon of our spline.
-	for (const FPCGPoint& Point : PointData->GetPoints())
+	const TConstPCGValueRange<FTransform> TransformRange = PointData->GetConstTransformValueRange();
+	for (const FTransform& PointTransform : TransformRange)
 	{
-		CachedSplinePoints.Add(Point.Transform.GetLocation());
-		CachedSplinePoints2D.Add(FVector2D(Point.Transform.GetLocation()));
+		CachedSplinePoints.Add(PointTransform.GetLocation());
+		CachedSplinePoints2D.Add(FVector2D(PointTransform.GetLocation()));
 	}
 }
 
diff --git a/Engine/Plugins/PCG/Source/PCG/Private/Data/PCGTextureData.cpp b/Engine/Plugins/PCG/Source/PCG/Private/Data/PCGTextureData.cpp
--- a/Engine/Plugins/PCG/Source/PCG/Private/Data/PCGTextureData.cpp
+++ b/Engine/Plugins/PCG/Source/PCG/Private/Data/PCGTextureData.cpp
@@ -4,6 +4,7 @@
 
 #include "PCGContext.h"
 #include "PCGTextureReadback.h"
+#include "Data/PCGPointArrayData.h"
 #include "Data/PCGPointData.h"
 #include "Helpers/PCGAsync.h"
 #include "Helpers/PCGHelpers.h"
@@ -276,6 +277,16 @@ bool UPCGBaseTextureData::SamplePoint(const FTransform& InTransform, const FBox&
 }
 
 const UPCGPointData* UPCGBaseTextureData::CreatePointData(FPCGContext* Context) const
+{
+	return CastChecked<UPCGPointData>(CreateBasePointData(Context, UPCGPointData::StaticClass()));
+}
+
+const UPCGPointArrayData* UPCGBaseTextureData::CreatePointArrayData(FPCGContext* Context, const FBox& InBounds) const
+{
+	return CastChecked<UPCGPointArrayData>(CreateBasePointData(Context, UPCGPointArrayData::StaticClass()));
+}
+
+const UPCGBasePointData* UPCGBaseTextureData::CreateBasePointData(FPCGContext* Context, TSubclassOf<UPCGBasePointData> PointDataClass) const
 {
 	TRACE_CPUPROFILER_EVENT_SCOPE(UPCGBaseTextureData::CreatePointData);
 	// TODO: this is a trivial implementation
@@ -283,7 +294,7 @@ const UPCGPointData* UPCGBaseTextureData::CreatePointData(FPCGContext* Context)
 	// or based on a given texel size
 	FBox2D LocalSurfaceBounds(FVector2D(-1.0f, -1.0f), FVector2D(1.0f, 1.0f));
 
-	UPCGPointData* Data = FPCGContext::NewObject_AnyThread<UPCGPointData>(Context);
+	UPCGBasePointData* Data = FPCGContext::NewObject_AnyThread<UPCGBasePointData>(Context, GetTransientPackage(), PointDataClass);
 	Data->InitializeFromData(this);
 
 	// Early out for invalid data
@@ -299,8 +310,6 @@ const UPCGPointData* UPCGBaseTextureData::CreatePointData(FPCGContext* Context)
 		return Data;
 	}
 
-	TArray<FPCGPoint>& Points = Data->GetMutablePoints();
-
 	// Map target texel size to the current physical size of the texture data.
 	const FVector::FReal XSize = 2.0 * Transform.GetScale3D().X;
 	const FVector::FReal YSize = 2.0 * Transform.GetScale3D().Y;
@@ -317,35 +326,73 @@ const UPCGPointData* UPCGBaseTextureData::CreatePointData(FPCGContext* Context)
 
 	const FBox2D Surface(FVector2D(-1.0f, -1.0f), FVector2D(1.0f, 1.0f));
 
-	FPCGAsync::AsyncPointProcessing(Context, PointCount, Points, [this, XCount, YCount, &Surface](int32 Index, FPCGPoint& OutPoint)
+	Data->SetNumPoints(PointCount, /*bInitializeValues=*/false);
+	Data->AllocateProperties(EPCGPointNativeProperties::Transform | EPCGPointNativeProperties::Seed | EPCGPointNativeProperties::Density | EPCGPointNativeProperties::Color);
+
+	// Check if we are dealing with always allocated properties or not (if not set the constant extents
+	if (Data->GetAllocatedProperties() != EPCGPointNativeProperties::All)
 	{
-		const int X = (Index % XCount);
-		const int Y = (Index / XCount);
+		Data->SetExtents(FVector(TexelSize / 2.0));
+	}
 
-		// TODO: we should have a 0.5 bias here
-		FVector2D LocalCoordinate((2.0 * X + 0.5) / XCount - 1.0, (2.0 * Y + 0.5) / YCount - 1.0);
-		FLinearColor Color = FLinearColor(EForceInit::ForceInit);
+	auto ProcessRangeFunc = [this, XCount, YCount, &Surface, Data](int32 StartReadIndex, int32 StartWriteIndex, int32 Count)
+	{
+		int32 NumWritten = 0;
+		FPCGPointValueRanges OutRanges(Data, /*bAllocate=*/false);
 
-		if (PCGTextureSamplingHelpers::Sample<FLinearColor>(LocalCoordinate, Surface, this, Width, Height, Color, [this](int32 Index) { return ColorData[Index]; }))
+		for (int32 ReadIndex = StartReadIndex; ReadIndex < StartReadIndex + Count; ++ReadIndex)
 		{
-			const float Density = bUseDensitySourceChannel ? PCGTextureSamplingHelpers::SampleFloatChannel(Color, ColorChannel) : 1.0f;
-			if (Density > 0 || bKeepZeroDensityPoints)
+			const int X = (ReadIndex % XCount);
+			const int Y = (ReadIndex / XCount);
+
+			const int32 WriteIndex = StartWriteIndex + NumWritten;
+
+			// TODO: we should have a 0.5 bias here
+			FVector2D LocalCoordinate((2.0 * X + 0.5) / XCount - 1.0, (2.0 * Y + 0.5) / YCount - 1.0);
+			FLinearColor Color = FLinearColor(EForceInit::ForceInit);
+
+			if (PCGTextureSamplingHelpers::Sample<FLinearColor>(LocalCoordinate, Surface, this, Width, Height, Color, [this](int32 Index) { return ColorData[Index]; }))
 			{
-				FVector LocalPosition(LocalCoordinate, 0);
-				OutPoint = FPCGPoint(FTransform(Transform.TransformPosition(LocalPosition)),
-					Density,
-					PCGHelpers::ComputeSeed(X, Y));
+				const float Density = bUseDensitySourceChannel ? PCGTextureSamplingHelpers::SampleFloatChannel(Color, ColorChannel) : 1.0f;
+				if (Density > 0 || bKeepZeroDensityPoints)
+				{
+					FVector LocalPosition(LocalCoordinate, 0);
+					FPCGPoint OutPoint(FTransform(Transform.TransformPosition(LocalPosition)),
+						Density,
+						PCGHelpers::ComputeSeed(X, Y));
 
-				OutPoint.SetExtents(FVector(TexelSize / 2.0));
-				OutPoint.Color = Color;
+					// Always the same extents so property wasn't allocated
+					OutPoint.SetExtents(FVector(TexelSize / 2.0));
+					OutPoint.Color = Color;
 
-				return true;
+					OutRanges.SetFromPoint(WriteIndex, OutPoint);
+					++NumWritten;
+				}
 			}
 		}
+		
+		return NumWritten;
+	};
 
-		return false;
-	});
+	auto MoveDataRangeFunc = [Data](int32 RangeStartIndex, int32 MoveToIndex, int32 NumElements)
+	{
+		Data->MoveRange(RangeStartIndex, MoveToIndex, NumElements);
+	};
+
+	auto FinishedFunc = [Data](int32 NumWritten)
+	{
+		Data->SetNumPoints(NumWritten);
+	};
 
+	FPCGAsync::AsyncProcessingRangeEx(
+		Context ? &Context->AsyncState : nullptr,
+		PointCount,
+		[] {},
+		ProcessRangeFunc,
+		MoveDataRangeFunc,
+		FinishedFunc,
+		/*bEnableTimeSlicing=*/false);
+		
 	return Data;
 }
 
diff --git a/Engine/Plugins/PCG/Source/PCG/Private/Data/PCGWorldData.cpp b/Engine/Plugins/PCG/Source/PCG/Private/Data/PCGWorldData.cpp
--- a/Engine/Plugins/PCG/Source/PCG/Private/Data/PCGWorldData.cpp
+++ b/Engine/Plugins/PCG/Source/PCG/Private/Data/PCGWorldData.cpp
@@ -323,7 +323,7 @@ bool UPCGWorldRayHitData::SamplePoint(const FTransform& InTransform, const FBox&
 		{
 			const FHitResult& Hit = HitResult.GetValue();
 			OutPoint = FPCGPoint(PCGWorldQueryHelpers::GetOrthonormalImpactTransform(Hit), 1.0f, UPCGBlueprintHelpers::ComputeSeedFromPosition(Hit.Location));
-			PCGWorldQueryHelpers::ApplyRayHitMetadata(Hit, QueryParams, QueryParams.RayDirection, OutPoint, OutMetadata, World);
+			PCGWorldQueryHelpers::ApplyRayHitMetadata(Hit, QueryParams, QueryParams.RayDirection, OutPoint.Transform, OutPoint.MetadataEntry, OutMetadata, World);
 			return true;
 		}
 	}
@@ -333,8 +333,16 @@ bool UPCGWorldRayHitData::SamplePoint(const FTransform& InTransform, const FBox&
 
 const UPCGPointData* UPCGWorldRayHitData::CreatePointData(FPCGContext* Context, const FBox& InBounds) const
 {
-	UPCGPointData* Data = FPCGContext::NewObject_AnyThread<UPCGPointData>(Context);
-	Data->InitializeFromData(this);
+	return CastChecked<UPCGPointData>(CreateBasePointData(Context, InBounds, UPCGPointData::StaticClass()), ECastCheckedType::NullAllowed);
+}
+
+const UPCGPointArrayData* UPCGWorldRayHitData::CreatePointArrayData(FPCGContext* Context, const FBox& InBounds) const
+{
+	return CastChecked<UPCGPointArrayData>(CreateBasePointData(Context, InBounds, UPCGPointArrayData::StaticClass()), ECastCheckedType::NullAllowed);
+}
+
+const UPCGBasePointData* UPCGWorldRayHitData::CreateBasePointData(FPCGContext* Context, const FBox& InBounds, TSubclassOf<UPCGBasePointData> PointDataClass) const
+{
 	PCGWorldQueryHelpers::CreateRayHitAttributes(QueryParams, Metadata);
 
 	FBox EffectiveBounds = Bounds;
@@ -358,14 +366,14 @@ const UPCGPointData* UPCGWorldRayHitData::CreatePointData(FPCGContext* Context,
 			UE_LOG(LogPCG, Error, TEXT("PCG World Ray Hit Data cannot generate without sampling bounds. Consider using a Surface Sampler with the Unbounded option disabled."));
 		}
 
+		UPCGBasePointData* Data = FPCGContext::NewObject_AnyThread<UPCGBasePointData>(Context, GetTransientPackage(), PointDataClass);
+		Data->InitializeFromData(this);
 		return Data;
 	}
 
 	// The default params will be fine in this case
 	const PCGSurfaceSampler::FSurfaceSamplerParams Params;
-	Data = PCGSurfaceSampler::SampleSurface(Context, /*InSurface=*/this, /*InBoundingShape=*/nullptr, EffectiveBounds, Params);
-
-	return Data;
+	return PCGSurfaceSampler::SampleSurface(Context, Params, /*InSurface=*/this, /*InBoundingShape=*/nullptr, EffectiveBounds, PointDataClass);
 }
 
 UPCGSpatialData* UPCGWorldRayHitData::CopyInternal(FPCGContext* Context) const
diff --git a/Engine/Plugins/PCG/Source/PCG/Private/Elements/Grammar/PCGDuplicateCrossSections.cpp b/Engine/Plugins/PCG/Source/PCG/Private/Elements/Grammar/PCGDuplicateCrossSections.cpp
--- a/Engine/Plugins/PCG/Source/PCG/Private/Elements/Grammar/PCGDuplicateCrossSections.cpp
+++ b/Engine/Plugins/PCG/Source/PCG/Private/Elements/Grammar/PCGDuplicateCrossSections.cpp
@@ -4,7 +4,6 @@
 
 #include "PCGContext.h"
 #include "PCGParamData.h"
-#include "Data/PCGPointData.h"
 #include "Data/PCGSplineData.h"
 #include "Grammar/PCGGrammarParser.h"
 #include "Metadata/Accessors/IPCGAttributeAccessor.h"
diff --git a/Engine/Plugins/PCG/Source/PCG/Private/Elements/Grammar/PCGSelectGrammar.cpp b/Engine/Plugins/PCG/Source/PCG/Private/Elements/Grammar/PCGSelectGrammar.cpp
--- a/Engine/Plugins/PCG/Source/PCG/Private/Elements/Grammar/PCGSelectGrammar.cpp
+++ b/Engine/Plugins/PCG/Source/PCG/Private/Elements/Grammar/PCGSelectGrammar.cpp
@@ -4,7 +4,7 @@
 
 #include "PCGContext.h"
 #include "PCGParamData.h"
-#include "Data/PCGPointData.h"
+#include "Data/PCGBasePointData.h"
 #include "Elements/Metadata/PCGMetadataElementCommon.h"
 #include "Helpers/PCGPropertyHelpers.h"
 
@@ -257,14 +257,14 @@ bool FPCGSelectGrammarElement::ExecuteInternal(FPCGContext* InContext) const
 	TArray<FPCGTaggedData>& Outputs = InContext->OutputData.TaggedData;
 	for (const FPCGTaggedData& PointInput : PointInputs)
 	{
-		const UPCGPointData* InputPointData = Cast<UPCGPointData>(PointInput.Data);
-		if (!InputPointData || InputPointData->GetPoints().IsEmpty())
+		const UPCGBasePointData* InputPointData = Cast<UPCGBasePointData>(PointInput.Data);
+		if (!InputPointData || InputPointData->IsEmpty())
 		{
 			continue;
 		}
 
 		UPCGData* OutputData = InputPointData->DuplicateData(InContext);
-		UPCGPointData* OutputPointData = CastChecked<UPCGPointData>(OutputData);
+		UPCGBasePointData* OutputPointData = CastChecked<UPCGBasePointData>(OutputData);
 		check(OutputPointData->Metadata)
 		Outputs.Emplace_GetRef(PointInput).Data = OutputData;
 
diff --git a/Engine/Plugins/PCG/Source/PCG/Private/Elements/Grammar/PCGSplineToSegment.cpp b/Engine/Plugins/PCG/Source/PCG/Private/Elements/Grammar/PCGSplineToSegment.cpp
--- a/Engine/Plugins/PCG/Source/PCG/Private/Elements/Grammar/PCGSplineToSegment.cpp
+++ b/Engine/Plugins/PCG/Source/PCG/Private/Elements/Grammar/PCGSplineToSegment.cpp
@@ -3,7 +3,7 @@
 #include "Elements/Grammar/PCGSplineToSegment.h"
 
 #include "PCGContext.h"
-#include "Data/PCGPointData.h"
+#include "Data/PCGBasePointData.h"
 #include "Data/PCGSplineData.h"
 #include "Helpers/PCGHelpers.h"
 
@@ -91,7 +91,7 @@ bool FPCGSplineToSegmentElement::ExecuteInternal(FPCGContext* InContext) const
 			continue;
 		}
 
-		UPCGPointData* OutputPointData = FPCGContext::NewObject_AnyThread<UPCGPointData>(InContext);
+		UPCGBasePointData* OutputPointData = FPCGContext::NewPointData_AnyThread(InContext);
 		OutputPointData->InitializeFromData(InputSplineData);
 
 		const int32 NumSegments = InputSplineData->GetNumSegments();
@@ -152,8 +152,10 @@ bool FPCGSplineToSegmentElement::ExecuteInternal(FPCGContext* InContext) const
 			bHasMetadata = true;
 		}
 
-		TArray<FPCGPoint>& OutputPoints = OutputPointData->GetMutablePoints();
-		OutputPoints.Reserve(NumSegments);
+		OutputPointData->SetNumPoints(NumSegments);
+		OutputPointData->AllocateProperties(EPCGPointNativeProperties::Transform | EPCGPointNativeProperties::Seed | EPCGPointNativeProperties::MetadataEntry | EPCGPointNativeProperties::BoundsMin | EPCGPointNativeProperties::BoundsMax);
+
+		FPCGPointValueRanges OutRanges(OutputPointData, /*bAllocate*/false);
 
 		TArray<FBox> Segments;
 		Segments.Reserve(NumSegments);
@@ -175,28 +177,27 @@ bool FPCGSplineToSegmentElement::ExecuteInternal(FPCGContext* InContext) const
 			const double AbsAngle = FMath::Acos(ArriveTangent.CosineAngle2D(LeaveTangent));
 			const double Angle = FMath::RadiansToDegrees(FMath::Sign(CrossProduct) * AbsAngle);
 
-			FPCGPoint& OutputPoint = OutputPoints.Emplace_GetRef();
 			if (bHasMetadata)
 			{
-				OutputPointData->Metadata->InitializeOnSet(OutputPoint.MetadataEntry);
+				OutputPointData->Metadata->InitializeOnSet(OutRanges.MetadataEntryRange[Index]);
 			}
 
 			// Points will be placed at the middle of the segment with bounds min and max matching control points on X (Y and Z will be set to -1 and 1 for min and max respectively)
 			// We'll also reset the point scale
 			const FVector Extents = (NextControlPoint - ControlPoint) * 0.5;
 			const double ExtentsLength = Extents.Length();
-			OutputPoint.Transform.SetLocation(ControlPoint + Extents);
-			OutputPoint.Transform.SetRotation(FRotationMatrix::MakeFromXZ(Extents, ControlPointZUp).ToQuat());
-			OutputPoint.BoundsMin = FVector(-ExtentsLength, 0.0, 0.0);
-			OutputPoint.BoundsMax = FVector(ExtentsLength, 1.0, 1.0);
+			OutRanges.TransformRange[Index].SetLocation(ControlPoint + Extents);
+			OutRanges.TransformRange[Index].SetRotation(FRotationMatrix::MakeFromXZ(Extents, ControlPointZUp).ToQuat());
+			OutRanges.BoundsMinRange[Index] = FVector(-ExtentsLength, 0.0, 0.0);
+			OutRanges.BoundsMaxRange[Index] = FVector(ExtentsLength, 1.0, 1.0);
 
-			OutputPoint.Seed = PCGHelpers::ComputeSeedFromPosition(OutputPoint.Transform.GetLocation());
+			OutRanges.SeedRange[Index] = PCGHelpers::ComputeSeedFromPosition(OutRanges.TransformRange[Index].GetLocation());
 
 			auto SetMetadata = [&](const int32 Index, const FVector& Tangent, const double Angle)
 			{
 				const int32 PreviousIndex = (Index + NumSegments - 1) % NumSegments;
-				const PCGMetadataEntryKey PreviousPointKey = OutputPoints[PreviousIndex].MetadataEntry;
-				const PCGMetadataEntryKey CurrentPointKey = OutputPoints[Index].MetadataEntry;
+				const PCGMetadataEntryKey PreviousPointKey = OutRanges.MetadataEntryRange[PreviousIndex];
+				const PCGMetadataEntryKey CurrentPointKey = OutRanges.MetadataEntryRange[Index];
 
 				if (PreviousTangentAttribute && NextTangentAttribute)
 				{
diff --git a/Engine/Plugins/PCG/Source/PCG/Private/Elements/Grammar/PCGSubdivideSegment.cpp b/Engine/Plugins/PCG/Source/PCG/Private/Elements/Grammar/PCGSubdivideSegment.cpp
--- a/Engine/Plugins/PCG/Source/PCG/Private/Elements/Grammar/PCGSubdivideSegment.cpp
+++ b/Engine/Plugins/PCG/Source/PCG/Private/Elements/Grammar/PCGSubdivideSegment.cpp
@@ -4,9 +4,10 @@
 
 #include "PCGContext.h"
 #include "PCGParamData.h"
-#include "Data/PCGPointData.h"
+#include "Data/PCGBasePointData.h"
 #include "Elements/Metadata/PCGMetadataElementCommon.h"
 #include "Grammar/PCGGrammarParser.h"
+#include "Helpers/PCGPointHelpers.h"
 #include "Metadata/Accessors/IPCGAttributeAccessor.h"
 #include "Metadata/Accessors/PCGAttributeAccessorHelpers.h"
 #include "Metadata/Accessors/PCGAttributeAccessorKeys.h"
@@ -35,9 +36,8 @@ public:
 
 		const UPCGSubdivideSegmentSettings* Settings = nullptr;
 		FPCGContext* Context = nullptr;
-		const TArray<FPCGPoint>* InPoints = nullptr;
-		TArray<FPCGPoint>* OutPoints = nullptr;
-		UPCGMetadata* OutputMetadata = nullptr;
+		const UPCGBasePointData* InPointData = nullptr;
+		UPCGBasePointData* OutPointData = nullptr;
 	};
 
 	static void Process(FParameters& InOutParameters, const FString& InGrammar, const bool bFlipAxis, int32 AdditionalSeed, int32 Index)
@@ -55,12 +55,12 @@ public:
 			return;
 		}
 
-		const FPCGPoint& Point = (*InOutParameters.InPoints)[Index];
+		FConstPCGPointValueRanges InRanges(InOutParameters.InPointData);
+		const FPCGPoint Point = InRanges.GetPoint(Index);
+
 		FTransform TransformNoTranslation = Point.Transform;
 		TransformNoTranslation.SetLocation(FVector::ZeroVector);
 
-		const int32 FirstModuleIndex = InOutParameters.OutPoints->Num();
-
 		FBox Segment = Point.GetLocalBounds();
 		FVector PointScaledSize = Point.GetScaledLocalSize();
 		if (bFlipAxis)
@@ -97,6 +97,13 @@ public:
 		FVector CurrentPos = InitialPos;
 		int32 ModuleIndex = 0;
 
+		const int32 FirstModuleIndex = InOutParameters.OutPointData->GetNumPoints();
+
+		InOutParameters.OutPointData->SetNumPoints(InOutParameters.OutPointData->GetNumPoints() + ModulesInstances.Num(), /*bInitializeValues=*/false);
+		InOutParameters.OutPointData->AllocateProperties(InOutParameters.InPointData->GetAllocatedProperties());
+		
+		FPCGPointValueRanges OutRanges(InOutParameters.OutPointData, /*bAllocateProperties=*/false);
+
 		for (int32 ModuleInstanceIndex = 0; ModuleInstanceIndex < ModulesInstances.Num(); ModuleInstanceIndex++)
 		{
 			const PCGSubdivisionBase::TModuleInstance<PCGGrammar::FTokenizedModule>& ModuleInstance = ModulesInstances[ModuleInstanceIndex];
@@ -115,10 +122,10 @@ public:
 			const FVector Position = CurrentPos + HalfStep;
 			CurrentPos = Position + HalfStep;
 
-			FPCGPoint& OutPoint = InOutParameters.OutPoints->Emplace_GetRef(Point);
+			FPCGPoint OutPoint(Point);
 			OutPoint.Transform = FTransform(Point.Transform.GetRotation(), Position + OtherDirection, Scale);
 			OutPoint.SetLocalBounds(FBox(-LocalBoundsExtents, LocalBoundsExtents));
-			InOutParameters.OutputMetadata->InitializeOnSet(OutPoint.MetadataEntry);
+			InOutParameters.OutPointData->Metadata->InitializeOnSet(OutPoint.MetadataEntry);
 			if (InOutParameters.SymbolAttribute)
 			{
 				InOutParameters.SymbolAttribute->SetValue(OutPoint.MetadataEntry, Symbol);
@@ -143,17 +150,19 @@ public:
 			{
 				InOutParameters.IsFinalPointAttribute->SetValue(OutPoint.MetadataEntry, true);
 			}
+
+			OutRanges.SetFromPoint(FirstModuleIndex + ModuleInstanceIndex, OutPoint);
 		}
 
 		if (bFlipAxis)
 		{
-			InOutParameters.CornerIndexes.Add(InOutParameters.OutPoints->Num() - 1);
+			InOutParameters.CornerIndexes.Add(InOutParameters.OutPointData->GetNumPoints() - 1);
 			InOutParameters.CornerIndexes.Add(FirstModuleIndex);
 		}
 		else
 		{
 			InOutParameters.CornerIndexes.Add(FirstModuleIndex);
-			InOutParameters.CornerIndexes.Add(InOutParameters.OutPoints->Num() - 1);
+			InOutParameters.CornerIndexes.Add(InOutParameters.OutPointData->GetNumPoints() - 1);
 		}
 	}
 };
@@ -234,7 +243,7 @@ bool FPCGSegmentSubdivisionElement::ExecuteInternal(FPCGContext* InContext) cons
 
 	for (const FPCGTaggedData& Input : Inputs)
 	{
-		const UPCGPointData* InputPointData = Cast<const UPCGPointData>(Input.Data);
+		const UPCGBasePointData* InputPointData = Cast<const UPCGBasePointData>(Input.Data);
 
 		if (!InputPointData)
 		{
@@ -244,7 +253,8 @@ bool FPCGSegmentSubdivisionElement::ExecuteInternal(FPCGContext* InContext) cons
 		TUniquePtr<const IPCGAttributeAccessor> GrammarAccessor;
 		TUniquePtr<const IPCGAttributeAccessor> FlipAxisAccessor;
 		TUniquePtr<const IPCGAttributeAccessor> SeedAccessor;
-		const FPCGAttributeAccessorKeysPoints Keys(InputPointData->GetPoints());
+		TUniquePtr<const IPCGAttributeAccessorKeys> Keys = PCGAttributeAccessorHelpers::CreateConstKeys(InputPointData, FPCGAttributePropertyOutputSelector());
+
 
 		const bool bShouldUseAccessors = Settings->GrammarSelection.bGrammarAsAttribute || Settings->bFlipAxisAsAttribute || Settings->bUseSeedAttribute;
 		constexpr int32 DefaultAdditionalSeed = 0;
@@ -281,14 +291,16 @@ bool FPCGSegmentSubdivisionElement::ExecuteInternal(FPCGContext* InContext) cons
 			}
 		}
 
-		UPCGPointData* OutputPointData = FPCGContext::NewObject_AnyThread<UPCGPointData>(InContext);
-		OutputPointData->InitializeFromData(InputPointData);
+		UPCGBasePointData* OutputPointData = FPCGContext::NewPointData_AnyThread(InContext);
+
+		FPCGInitializeFromDataParams InitializeFromDataParams(InputPointData);
+		InitializeFromDataParams.bInheritSpatialData = false;
+		OutputPointData->InitializeFromDataWithParams(InitializeFromDataParams);
 
 		Parameters.Settings = Settings;
 		Parameters.Context = InContext;
-		Parameters.InPoints = &InputPointData->GetPoints();
-		Parameters.OutPoints = &OutputPointData->GetMutablePoints();
-		Parameters.OutputMetadata = OutputPointData->Metadata;
+		Parameters.InPointData = InputPointData;
+		Parameters.OutPointData = OutputPointData;
 
 		auto CreateAndValidateAttribute = [InContext, OutputMetadata = OutputPointData->Metadata]<typename T>(const FName AttributeName, const T DefaultValue, const bool bShouldCreate, FPCGMetadataAttribute<T>*& OutAttribute) -> bool
 		{
@@ -323,21 +335,23 @@ bool FPCGSegmentSubdivisionElement::ExecuteInternal(FPCGContext* InContext) cons
 				PCGSubdivideSegmentHelpers::Process(Parameters, InGrammar, bFlipAxis, AdditionalSeed, Index);
 			};
 
-			PCGMetadataElementCommon::ApplyOnMultiAccessors<FString, bool, int32>(Keys, { GrammarAccessor.Get(), FlipAxisAccessor.Get(), SeedAccessor.Get()}, Process);
+			PCGMetadataElementCommon::ApplyOnMultiAccessors<FString, bool, int32>(*Keys, { GrammarAccessor.Get(), FlipAxisAccessor.Get(), SeedAccessor.Get()}, Process);
 		}
 		else
 		{
-			for (int32 SegmentIndex = 0; SegmentIndex < Parameters.InPoints->Num(); ++SegmentIndex)
+			for (int32 SegmentIndex = 0; SegmentIndex < Parameters.InPointData->GetNumPoints(); ++SegmentIndex)
 			{
 				PCGSubdivideSegmentHelpers::Process(Parameters, Settings->GrammarSelection.GrammarString, Settings->bShouldFlipAxis, DefaultAdditionalSeed, SegmentIndex);
 			}
 		}
 
-		if (!Parameters.OutPoints->IsEmpty())
+		if (!Parameters.OutPointData->IsEmpty())
 		{
 			// Set the extremity neighbor indexes
 			if (Parameters.ExtremityNeighborIndexAttribute)
 			{
+				TPCGValueRange<int64> MetadataEntryRange = Parameters.OutPointData->GetMetadataEntryValueRange(/*bAllocate=*/true);
+
 				// Go 2 by 2
 				check(Parameters.CornerIndexes.Num() % 2 == 0);
 
@@ -346,8 +360,8 @@ bool FPCGSegmentSubdivisionElement::ExecuteInternal(FPCGContext* InContext) cons
 					const int32 PreviousModuleIndex = Parameters.CornerIndexes[j == 0 ? Parameters.CornerIndexes.Num() - 1 : j - 1];
 					const int32 NextModuleIndex = Parameters.CornerIndexes[j == Parameters.CornerIndexes.Num() - 2 ? 0 : j + 2];
 
-					Parameters.ExtremityNeighborIndexAttribute->SetValue((*Parameters.OutPoints)[Parameters.CornerIndexes[j]].MetadataEntry, PreviousModuleIndex);
-					Parameters.ExtremityNeighborIndexAttribute->SetValue((*Parameters.OutPoints)[Parameters.CornerIndexes[j+1]].MetadataEntry, NextModuleIndex);
+					Parameters.ExtremityNeighborIndexAttribute->SetValue(MetadataEntryRange[Parameters.CornerIndexes[j]], PreviousModuleIndex);
+					Parameters.ExtremityNeighborIndexAttribute->SetValue(MetadataEntryRange[Parameters.CornerIndexes[j+1]], NextModuleIndex);
 				}
 			}
 
diff --git a/Engine/Plugins/PCG/Source/PCG/Private/Elements/Grammar/PCGSubdivideSpline.cpp b/Engine/Plugins/PCG/Source/PCG/Private/Elements/Grammar/PCGSubdivideSpline.cpp
--- a/Engine/Plugins/PCG/Source/PCG/Private/Elements/Grammar/PCGSubdivideSpline.cpp
+++ b/Engine/Plugins/PCG/Source/PCG/Private/Elements/Grammar/PCGSubdivideSpline.cpp
@@ -3,7 +3,7 @@
 #include "Elements/Grammar/PCGSubdivideSpline.h"
 
 #include "PCGContext.h"
-#include "Data/PCGPointData.h"
+#include "Data/PCGBasePointData.h"
 #include "Data/PCGSplineData.h"
 #include "Elements/Metadata/PCGMetadataElementCommon.h"
 #include "Helpers/PCGHelpers.h"
@@ -26,8 +26,7 @@ namespace PCGSubdivideSplineHelpers
 		double BisectionTolerance = 0.01;
 
 		const UPCGPolyLineData* PolyLineData = nullptr;
-		TArray<FPCGPoint>* OutPoints = nullptr;
-		UPCGMetadata* OutputMetadata = nullptr;
+		UPCGBasePointData* OutputPointData = nullptr;
 
 		FPCGMetadataAttribute<FName>* SymbolAttribute = nullptr;
 		FPCGMetadataAttribute<FVector4>* DebugColorAttribute = nullptr;
@@ -223,6 +222,12 @@ namespace PCGSubdivideSplineHelpers
 
 		FVector PreviousSegmentEndPoint = (ModuleAlphas.IsEmpty() ? FVector::Zero() : InOutParameters.PolyLineData->GetLocationAtAlpha(ModuleAlphas[0]));
 		const int32 NumIterations = FMath::Min(ModulesInstances.Num(), ModuleAlphas.Num() - 1);
+
+		const int32 WriteIndex = InOutParameters.OutputPointData->GetNumPoints();
+		InOutParameters.OutputPointData->SetNumPoints(WriteIndex + NumIterations, /*bInitializeValues=*/false);
+		InOutParameters.OutputPointData->AllocateProperties(EPCGPointNativeProperties::Transform | EPCGPointNativeProperties::BoundsMin | EPCGPointNativeProperties::BoundsMax | EPCGPointNativeProperties::Seed | EPCGPointNativeProperties::MetadataEntry);
+		FPCGPointValueRanges OutRanges(InOutParameters.OutputPointData, /*bAllocate=*/false);
+
 		for (int32 ModuleInstanceIndex = 0; ModuleInstanceIndex < NumIterations; ModuleInstanceIndex++)
 		{
 			const PCGSubdivisionBase::TModuleInstance<PCGGrammar::FTokenizedModule>& ModuleInstance = ModulesInstances[ModuleInstanceIndex];
@@ -247,44 +252,44 @@ namespace PCGSubdivideSplineHelpers
 			const FRotator Rotation = FRotationMatrix::MakeFromXZ(SubdivisionDirection, CenterPointTransform.GetRotation().GetUpVector()).Rotator();
 			const FVector Scale = FVector(1.0 + ModuleInstance.ExtraScale, 1.0, 1.0);
 
-			FPCGPoint& OutPoint = InOutParameters.OutPoints->Emplace_GetRef(FTransform(Rotation, Position, Scale), /*InDensity=*/1, PCGHelpers::ComputeSeedFromPosition(Position));
+			FPCGPoint OutPoint(FTransform(Rotation, Position, Scale), /*InDensity=*/1, PCGHelpers::ComputeSeedFromPosition(Position));
 
 			const double HalfSubmoduleSize = SubdivisionSubmodule.Size * 0.5;
 			OutPoint.SetLocalBounds(FBox(FVector(-HalfSubmoduleSize, 0, 0), FVector(HalfSubmoduleSize, 1, InOutParameters.ModuleHeight)));
 
 			// Now, handle the metadata attributes
-			if (!bHasMetadata)
+			if (bHasMetadata)
 			{
-				continue;
-			}
+				InOutParameters.OutputPointData->Metadata->InitializeOnSet(OutPoint.MetadataEntry);
+				if (InOutParameters.SymbolAttribute)
+				{
+					InOutParameters.SymbolAttribute->SetValue(OutPoint.MetadataEntry, Symbol);
+				}
 
-			InOutParameters.OutputMetadata->InitializeOnSet(OutPoint.MetadataEntry);
-			if (InOutParameters.SymbolAttribute)
-			{
-				InOutParameters.SymbolAttribute->SetValue(OutPoint.MetadataEntry, Symbol);
-			}
+				if (InOutParameters.DebugColorAttribute)
+				{
+					InOutParameters.DebugColorAttribute->SetValue(OutPoint.MetadataEntry, FVector4(SubdivisionSubmodule.DebugColor, 1.0));
+				}
 
-			if (InOutParameters.DebugColorAttribute)
-			{
-				InOutParameters.DebugColorAttribute->SetValue(OutPoint.MetadataEntry, FVector4(SubdivisionSubmodule.DebugColor, 1.0));
-			}
+				if (InOutParameters.ModuleIndexAttribute)
+				{
+					InOutParameters.ModuleIndexAttribute->SetValue(OutPoint.MetadataEntry, ModuleInstanceIndex);
+				}
 
-			if (InOutParameters.ModuleIndexAttribute)
-			{
-				InOutParameters.ModuleIndexAttribute->SetValue(OutPoint.MetadataEntry, ModuleInstanceIndex);
-			}
+				const bool bIsFirstModule = (ModuleInstanceIndex == 0);
+				if (bIsFirstModule && InOutParameters.IsFirstPointAttribute)
+				{
+					InOutParameters.IsFirstPointAttribute->SetValue(OutPoint.MetadataEntry, true);
+				}
 
-			const bool bIsFirstModule = (ModuleInstanceIndex == 0);
-			if (bIsFirstModule && InOutParameters.IsFirstPointAttribute)
-			{
-				InOutParameters.IsFirstPointAttribute->SetValue(OutPoint.MetadataEntry, true);
+				const bool bIsFinalModule = (ModuleInstanceIndex == NumIterations - 1);
+				if (bIsFinalModule && InOutParameters.IsFinalPointAttribute)
+				{
+					InOutParameters.IsFinalPointAttribute->SetValue(OutPoint.MetadataEntry, true);
+				}
 			}
 
-			const bool bIsFinalModule = (ModuleInstanceIndex == NumIterations - 1);
-			if (bIsFinalModule && InOutParameters.IsFinalPointAttribute)
-			{
-				InOutParameters.IsFinalPointAttribute->SetValue(OutPoint.MetadataEntry, true);
-			}
+			OutRanges.SetFromPoint(WriteIndex + ModuleInstanceIndex, OutPoint);
 		}
 	}
 }
@@ -364,15 +369,16 @@ bool FPCGSubdivideSplineElement::ExecuteInternal(FPCGContext* InContext) const
 			}
 		}
 
-		UPCGPointData* OutputPointData = FPCGContext::NewObject_AnyThread<UPCGPointData>(InContext);
+		UPCGBasePointData* OutputPointData = FPCGContext::NewPointData_AnyThread(InContext);
 		OutputPointData->InitializeFromData(InputPolyLineData);
+		OutputPointData->SetDensity(1.0f);
+
 		FPCGTaggedData& Output = Outputs.Emplace_GetRef(Input);
 		Output.Data = OutputPointData;
 
 		// Update 'per-input' parameters
 		Parameters.PolyLineData = InputPolyLineData;
-		Parameters.OutPoints = &OutputPointData->GetMutablePoints();
-		Parameters.OutputMetadata = OutputPointData->Metadata;
+		Parameters.OutputPointData = OutputPointData;
 		Parameters.ModuleHeight = ModuleHeight;
 
 		auto CreateAndValidateAttribute = [InContext, OutputMetadata = OutputPointData->Metadata]<typename T>(const FName AttributeName, const T DefaultValue, const bool bShouldCreate, FPCGMetadataAttribute<T>*& OutAttribute) -> bool
diff --git a/Engine/Plugins/PCG/Source/PCG/Private/Elements/Grammar/PCGSubdivisionBase.cpp b/Engine/Plugins/PCG/Source/PCG/Private/Elements/Grammar/PCGSubdivisionBase.cpp
--- a/Engine/Plugins/PCG/Source/PCG/Private/Elements/Grammar/PCGSubdivisionBase.cpp
+++ b/Engine/Plugins/PCG/Source/PCG/Private/Elements/Grammar/PCGSubdivisionBase.cpp
@@ -4,7 +4,7 @@
 
 #include "PCGContext.h"
 #include "PCGParamData.h"
-#include "Data/PCGPointData.h"
+#include "Data/PCGBasePointData.h"
 #include "Elements/Metadata/PCGMetadataElementCommon.h"
 #include "Grammar/PCGGrammarParser.h"
 #include "Helpers/PCGPropertyHelpers.h"
@@ -272,7 +272,7 @@ PCGGrammar::FTokenizedGrammar PCGSubdivisionBase::GetTokenizedGrammar(FPCGContex
 	return TokenizedGrammar;
 }
 
-TMap<FString, PCGGrammar::FTokenizedGrammar> FPCGSubdivisionBaseElement::GetTokenizedGrammarForPoints(FPCGContext* InContext, const UPCGPointData* InputData, const UPCGSubdivisionBaseSettings* InSettings, const FModuleInfoMap& InModulesInfo, double& OutMinSize) const
+TMap<FString, PCGGrammar::FTokenizedGrammar> FPCGSubdivisionBaseElement::GetTokenizedGrammarForPoints(FPCGContext* InContext, const UPCGBasePointData* InputData, const UPCGSubdivisionBaseSettings* InSettings, const FModuleInfoMap& InModulesInfo, double& OutMinSize) const
 {
 	TMap<FString, PCGGrammar::FTokenizedGrammar> Result;
 
@@ -377,7 +377,7 @@ bool FPCGSubdivisionBaseElement::MatchAndSetAttributes(const TArray<FPCGTaggedDa
 			continue;
 		}
 
-		if (UPCGPointData* OutPointData = Cast<UPCGPointData>(OutData))
+		if (UPCGBasePointData* OutPointData = Cast<UPCGBasePointData>(OutData))
 		{
 			TArray<const FPCGMetadataAttributeBase*> InAttributes;
 			TArray<FPCGMetadataAttributeBase*> OutAttributes;
@@ -404,9 +404,11 @@ bool FPCGSubdivisionBaseElement::MatchAndSetAttributes(const TArray<FPCGTaggedDa
 				OutAttributes.Add(OutputMetadata->CopyAttribute(InAttributes.Last(), AttributeName, /*bKeepParent=*/false, /*bCopyEntries=*/false, /*bCopyValues=*/true));
 			}
 
-			for (const FPCGPoint& OutPoint : OutPointData->GetPoints())
+			// @todo_pcg: potential optimization GetValuesFromItemKeys / SetValuesFromValueKeys
+			TPCGValueRange<int64> MetadataEntryRange = OutPointData->GetMetadataEntryValueRange(/*bAllocate=*/true);
+			for (int32 Index = 0; Index < OutPointData->GetNumPoints(); ++Index)
 			{
-				const FName Symbol = OutSymbolAttribute->GetValueFromItemKey(OutPoint.MetadataEntry);
+				const FName Symbol = OutSymbolAttribute->GetValueFromItemKey(MetadataEntryRange[Index]);
 
 				if (!SymbolToEntryKeyMapping.Contains(Symbol))
 				{
@@ -423,7 +425,7 @@ bool FPCGSubdivisionBaseElement::MatchAndSetAttributes(const TArray<FPCGTaggedDa
 				for (int32 i = 0; i < InAttributes.Num(); ++i)
 				{
 					PCGMetadataValueKey InputValueKey = InAttributes[i]->GetValueKey(InputEntryKey);
-					OutAttributes[i]->SetValueFromValueKey(OutPoint.MetadataEntry, InputValueKey);
+					OutAttributes[i]->SetValueFromValueKey(MetadataEntryRange[Index], InputValueKey);
 				}
 			}
 		}
diff --git a/Engine/Plugins/PCG/Source/PCG/Private/Elements/IO/PCGDataTableElement.cpp b/Engine/Plugins/PCG/Source/PCG/Private/Elements/IO/PCGDataTableElement.cpp
--- a/Engine/Plugins/PCG/Source/PCG/Private/Elements/IO/PCGDataTableElement.cpp
+++ b/Engine/Plugins/PCG/Source/PCG/Private/Elements/IO/PCGDataTableElement.cpp
@@ -5,7 +5,7 @@
 #include "PCGComponent.h"
 #include "PCGModule.h"
 #include "PCGParamData.h"
-#include "Data/PCGPointData.h"
+#include "Data/PCGBasePointData.h"
 #include "Elements/IO/PCGExternalDataContext.h"
 #include "Helpers/PCGDynamicTrackingHelpers.h"
 
@@ -127,7 +127,11 @@ bool FPCGLoadDataTableElement::PrepareLoad(FPCGExternalDataContext* Context) con
 
 	if (Settings->OutputType == EPCGExclusiveDataType::Point)
 	{
-		UPCGPointData* PointData = FPCGContext::NewObject_AnyThread<UPCGPointData>(Context);
+		UPCGBasePointData* PointData = FPCGContext::NewPointData_AnyThread(Context);
+		// Dummy implementation - note that we do initialize the points here in cases where some of the properties would not be set
+		// Set num points before creating accessors
+		PointData->SetNumPoints(DataTable->GetRowMap().Num());
+
 		check(PointData);
 		OutData = PointData;
 		OutMetadata = PointData->MutableMetadata();
@@ -241,16 +245,7 @@ bool FPCGLoadDataTableElement::ExecuteLoad(FPCGExternalDataContext* Context) con
 	UPCGData* Data = PointDataAccessorMapping.Data;
 	check(Data);
 
-	if (UPCGPointData* PointData = Cast<UPCGPointData>(Data))
-	{
-		TArray<FPCGPoint>& Points = PointData->GetMutablePoints();
-		if (Points.IsEmpty())
-		{
-			// Dummy implementation - note that we do initialize the points here in cases where some of the properties would not be set
-			Points.SetNum(DataTable->GetRowMap().Num());
-		}
-	}
-	else if (UPCGParamData* ParamData = Cast<UPCGParamData>(Data))
+	if (UPCGParamData* ParamData = Cast<UPCGParamData>(Data))
 	{
 		check(ParamData->Metadata);
 		TArray<FName> AttributeNames;
diff --git a/Engine/Plugins/PCG/Source/PCG/Private/Elements/IO/PCGExternalData.cpp b/Engine/Plugins/PCG/Source/PCG/Private/Elements/IO/PCGExternalData.cpp
--- a/Engine/Plugins/PCG/Source/PCG/Private/Elements/IO/PCGExternalData.cpp
+++ b/Engine/Plugins/PCG/Source/PCG/Private/Elements/IO/PCGExternalData.cpp
@@ -4,7 +4,6 @@
 #include "Elements/IO/PCGExternalDataContext.h"
 
 #include "PCGComponent.h"
-#include "Data/PCGPointData.h"
 
 #define LOCTEXT_NAMESPACE "PCGExternalData"
 
diff --git a/Engine/Plugins/PCG/Source/PCG/Private/Elements/PCGAttractElement.cpp b/Engine/Plugins/PCG/Source/PCG/Private/Elements/PCGAttractElement.cpp
--- a/Engine/Plugins/PCG/Source/PCG/Private/Elements/PCGAttractElement.cpp
+++ b/Engine/Plugins/PCG/Source/PCG/Private/Elements/PCGAttractElement.cpp
@@ -10,6 +10,7 @@
 #include "Metadata/Accessors/PCGAttributeAccessorHelpers.h"
 #include "SpatialAlgo/PCGOctreeQueries.h"
 
+#include "Algo/Count.h"
 #include "Algo/MaxElement.h"
 
 #define LOCTEXT_NAMESPACE "PCGAttractElement"
@@ -31,7 +32,7 @@ namespace PCGAttractElement
 
 			// Process
 			int32 OutIndex = 0;
-			const int32 SourcePointCount = AttractData.SourcePointData->GetPoints().Num();
+			const int32 SourcePointCount = AttractData.SourceData->GetNumPoints();
 
 			for (int32 SourceIndex = 0; SourceIndex < SourcePointCount; ++SourceIndex)
 			{
@@ -112,7 +113,7 @@ namespace PCGAttractElement
 		// Runs a single iteration of the attract algorithm.
 		static void RunAttractIteration(const int32 Index, FAttractState& OutAttractData, const UPCGAttractSettings* AttractSettings)
 		{
-			const TArray<FPCGPoint>& SourcePoints = OutAttractData.SourcePointData->GetPoints();
+			const TConstPCGValueRange<FTransform> TransformRange = OutAttractData.SourceData->GetConstTransformValueRange();
 
 			if (AttractSettings->Mode == EPCGAttractMode::Closest)
 			{
@@ -124,8 +125,8 @@ namespace PCGAttractElement
 					bool bFoundColocatedPoint = false;
 
 					UPCGOctreeQueries::ForEachPointInsideSphere(
-						OutAttractData.TargetPointData,
-						SourcePoints[Index].Transform.GetLocation(),
+						OutAttractData.TargetData,
+						TransformRange[Index].GetLocation(),
 						AttractSettings->Distance,
 						[Index, &OutAttractData, &MinDistanceSquared, &ClosestTargetIndex, &bFoundColocatedPoint](const UPCGBasePointData* PointData, int32 TargetIndex, const double DistanceSquared)
 						{
@@ -165,22 +166,23 @@ namespace PCGAttractElement
 				else
 				{
 					// Simple closest point query
-					if (const FPCGPoint* ClosestPointInTarget = UPCGOctreeQueries::GetClosestPoint(
-						OutAttractData.TargetPointData,
-						SourcePoints[Index].Transform.GetLocation(),
+					const int32 ClosestPointIndex = UPCGOctreeQueries::GetClosestPointIndex(
+						OutAttractData.TargetData,
+						TransformRange[Index].GetLocation(),
 						/*bInDiscardCenter=*/false,
-						AttractSettings->Distance))
+						AttractSettings->Distance);
+					
+					if(ClosestPointIndex != INDEX_NONE)
 					{
-						const int32 TargetIndex = ClosestPointInTarget - OutAttractData.TargetPointData->GetPoints().GetData();
-						OutAttractData.Mapping[Index] = TargetIndex;
+						OutAttractData.Mapping[Index] = ClosestPointIndex;
 					}
 				}
 			}
 			else // Min/Max attribute test inside of search radius
 			{
 				UPCGOctreeQueries::ForEachPointInsideSphere(
-					OutAttractData.TargetPointData,
-					SourcePoints[Index].Transform.GetLocation(),
+					OutAttractData.TargetData,
+					TransformRange[Index].GetLocation(),
 					AttractSettings->Distance,
 					[Index, &OutAttractData](const UPCGBasePointData* PointData, int32 TargetIndex, const double DistanceSquared)
 					{
@@ -216,11 +218,11 @@ namespace PCGAttractElement
 
 			check(InContext);
 
-			while (OutAttractData.IterationIndex < OutAttractData.SourcePointData->GetNumPoints())
+			while (OutAttractData.IterationIndex < OutAttractData.SourceData->GetNumPoints())
 			{
 				Helpers::RunAttractIteration(OutAttractData.IterationIndex++, OutAttractData, AttractSettings);
 
-				if (OutAttractData.IterationIndex != OutAttractData.SourcePointData->GetNumPoints() && InContext->ShouldStop())
+				if (OutAttractData.IterationIndex != OutAttractData.SourceData->GetNumPoints() && InContext->ShouldStop())
 				{
 					// Not done
 					return false;
@@ -321,35 +323,36 @@ bool FPCGAttractElement::PrepareDataInternal(FPCGContext* InContext) const
 
 	Context->InitializePerIterationStates(SourceInputs.Num(), [&SourceInputs, &TargetInputs, Settings, Context](IterStateType& OutState, const ExecStateType&, const uint32 IterationIndex)
 	{
-		const UPCGPointData* SourcePointData = Cast<UPCGPointData>(SourceInputs[IterationIndex].Data);
+		const UPCGBasePointData* SourcePointData = Cast<UPCGBasePointData>(SourceInputs[IterationIndex].Data);
 
 		if (!SourcePointData || SourcePointData->IsEmpty())
 		{
 			return EPCGTimeSliceInitResult::NoOperation;
 		}
 
-		const UPCGPointData* TargetPointData = Cast<UPCGPointData>(TargetInputs[IterationIndex % TargetInputs.Num()].Data);
+		const UPCGBasePointData* TargetPointData = Cast<UPCGBasePointData>(TargetInputs[IterationIndex % TargetInputs.Num()].Data);
 
 		if (!TargetPointData)
 		{
 			return EPCGTimeSliceInitResult::NoOperation;
 		}
 
-		UPCGPointData* OutPointData = FPCGContext::NewObject_AnyThread<UPCGPointData>(Context);
+		UPCGBasePointData* OutPointData = FPCGContext::NewPointData_AnyThread(Context);
 		OutPointData->InitializeFromData(SourcePointData);
 		Context->OutputData.TaggedData.Add_GetRef(SourceInputs[IterationIndex]).Data = OutPointData;
 
 		const bool bTargetIsSource = (TargetPointData == SourcePointData);
 
-		OutState = PCGAttractElement::FAttractState
-		{
-			.TargetPointData = TargetPointData,
-			.SourcePointData = SourcePointData,
-			.OutPointData = OutPointData,
-
-			.IterationIndex = 0,
-			.bTargetIsSource = bTargetIsSource
-		};
+		OutState.TargetData = TargetPointData;
+		OutState.SourceData = SourcePointData;
+		OutState.OutputData = OutPointData;
+		OutState.IterationIndex = 0;
+		OutState.bTargetIsSource = bTargetIsSource;
+PRAGMA_DISABLE_DEPRECATION_WARNINGS
+		OutState.TargetPointData = Cast<UPCGPointData>(TargetPointData);
+		OutState.SourcePointData = Cast<UPCGPointData>(SourcePointData);
+		OutState.OutPointData = Cast<UPCGPointData>(OutPointData);
+PRAGMA_ENABLE_DEPRECATION_WARNINGS
 
 		OutState.Mapping.SetNumUninitialized(SourcePointData->GetNumPoints());
 
@@ -499,40 +502,12 @@ bool FPCGAttractElement::PrepareDataInternal(FPCGContext* InContext) const
 			OutputSelector.ImportFromOtherSelector(SourceSelector);
 			OutputSelector = OutputSelector.CopyAndFixSource(&SourceSelector, OutPointData);
 
-			TUniquePtr<IPCGAttributeAccessor> OutputAccessor = PCGAttributeAccessorHelpers::CreateAccessor(OutPointData, OutputSelector);
-
-			if (!OutputAccessor)
-			{
-				PCGLog::Metadata::LogFailToCreateAccessorError(OutputSelector, Context);
-				continue;
-			}
-
 			PCGAttractElement::FAttractState::FAttributeInterpolationData& Interpolator = OutState.Interpolators.Emplace_GetRef();
 			Interpolator.SourceAccessor = std::move(SourceAccessor);
 			Interpolator.SourceKeys = std::move(SourceKeys);
 			Interpolator.TargetAccessor = std::move(TargetAccessor);
 			Interpolator.TargetKeys = std::move(TargetKeys);
-			Interpolator.OutputAccessor = std::move(OutputAccessor);
-		}
-
-		if (Settings->bOutputAttractIndex)
-		{
-			const FPCGAttributePropertyOutputNoSourceSelector AttractIndexSelector = Settings->OutputAttractIndexAttribute;
-
-			// Make sure we create the attribute if needed.
-			if(AttractIndexSelector.IsBasicAttribute())
-			{
-				OutPointData->Metadata->FindOrCreateAttribute<int32>(AttractIndexSelector.GetName(), 0, /*bAllowInterpolation=*/false);
-			}
-
-			TUniquePtr<IPCGAttributeAccessor> AttractIndexAccessor = PCGAttributeAccessorHelpers::CreateAccessor(OutPointData, AttractIndexSelector);
-			if (!AttractIndexAccessor)
-			{
-				PCGLog::Metadata::LogFailToCreateAccessorError(AttractIndexSelector, Context);
-				return EPCGTimeSliceInitResult::NoOperation;
-			}
-
-			OutState.OutputAttractIndexAccessor = std::move(AttractIndexAccessor);
+			Interpolator.OutputSelector = std::move(OutputSelector);
 		}
 
 		return EPCGTimeSliceInitResult::Success;
@@ -588,39 +563,77 @@ bool FPCGAttractElement::ExecuteInternal(FPCGContext* InContext) const
 
 		if (Settings->bRemoveUnattractedPoints)
 		{
-			TArray<FPCGPoint>& OutPoints = AttractData.OutPointData->GetMutablePoints();
-			const TArray<FPCGPoint>& SourcePoints = AttractData.SourcePointData->GetPoints();
+			const int32 NumPoints = Algo::CountIf(AttractData.Mapping, [](int32 Index)
+			{
+				return Index != INDEX_NONE;
+			});
 
-			Mapping.Reserve(AttractData.Mapping.Num());
+			AttractData.OutputData->SetNumPoints(NumPoints, /*bInitializeValues=*/false);
+			AttractData.OutputData->AllocateProperties(AttractData.SourceData->GetAllocatedProperties());
+			
+			Mapping.Reserve(NumPoints);
+
+			const FConstPCGPointValueRanges SourceRanges(AttractData.SourceData);
+			FPCGPointValueRanges OutRanges(AttractData.OutputData);
 
-			OutPoints.Reserve(SourcePoints.Num());
-			for (int PointIndex = 0; PointIndex < SourcePoints.Num(); ++PointIndex)
+			int32 WriteIndex = 0;
+			for (int PointIndex = 0; PointIndex < AttractData.SourceData->GetNumPoints(); ++PointIndex)
 			{
 				if (AttractData.Mapping[PointIndex] != INDEX_NONE)
 				{
-					OutPoints.Add(SourcePoints[PointIndex]);
+					OutRanges.SetFromValueRanges(WriteIndex, SourceRanges, PointIndex);
 					Mapping.Add(AttractData.Mapping[PointIndex]);
+					++WriteIndex;
 				}
 			}
 		}
 		else
 		{
-			AttractData.OutPointData->GetMutablePoints() = AttractData.SourcePointData->GetPoints();
+			if (!AttractData.OutputData->HasSpatialDataParent())
+			{
+				UPCGBasePointData::SetPoints(AttractData.SourceData, AttractData.OutputData, {}, /*bCopyAll=*/true);
+			}
+
 			Mapping = AttractData.Mapping;
 		}
 
 		// After we've created the points, we can now create the keys
-		AttractData.OutputKeys = PCGAttributeAccessorHelpers::CreateKeys(AttractData.OutPointData, FPCGAttributePropertyOutputSelector());
+		AttractData.OutputKeys = PCGAttributeAccessorHelpers::CreateKeys(AttractData.OutputData, FPCGAttributePropertyOutputSelector());
 
-		// Set the attact index if required
-		if (AttractData.OutputAttractIndexAccessor)
+		// After we've created the points we can now create the accessor
+		if (Settings->bOutputAttractIndex)
 		{
+			const FPCGAttributePropertyOutputNoSourceSelector AttractIndexSelector = Settings->OutputAttractIndexAttribute;
+
+			// Make sure we create the attribute if needed.
+			if (AttractIndexSelector.IsBasicAttribute())
+			{
+				AttractData.OutputData->Metadata->FindOrCreateAttribute<int32>(AttractIndexSelector.GetName(), 0, /*bAllowInterpolation=*/false);
+			}
+
+			TUniquePtr<IPCGAttributeAccessor> AttractIndexAccessor = PCGAttributeAccessorHelpers::CreateAccessor(AttractData.OutputData, AttractIndexSelector);
+			if (!AttractIndexAccessor)
+			{
+				PCGLog::Metadata::LogFailToCreateAccessorError(AttractIndexSelector, Context);
+				return true;
+			}
+
+			AttractData.OutputAttractIndexAccessor = std::move(AttractIndexAccessor);
 			AttractData.OutputAttractIndexAccessor->SetRange<int32>(Mapping, 0, *AttractData.OutputKeys);
 		}
 
 		// Apply weighting from source to target
 		while (AttractData.InterpolationIndex < AttractData.Interpolators.Num())
 		{
+			TUniquePtr<IPCGAttributeAccessor> OutputAccessor = PCGAttributeAccessorHelpers::CreateAccessor(AttractData.OutputData, AttractData.Interpolators[AttractData.InterpolationIndex].OutputSelector);
+
+			if (!OutputAccessor)
+			{
+				PCGLog::Metadata::LogFailToCreateAccessorError(AttractData.Interpolators[AttractData.InterpolationIndex].OutputSelector, Context);
+				return false;
+			}
+
+			AttractData.Interpolators[AttractData.InterpolationIndex].OutputAccessor = std::move(OutputAccessor);
 			PCGAttractElement::Helpers::InterpolatePoints(AttractData, AttractData.InterpolationIndex++, Settings);
 
 			// Check for time-slice
@@ -630,8 +643,6 @@ bool FPCGAttractElement::ExecuteInternal(FPCGContext* InContext) const
 			}
 		}
 
-		AttractData.OutPointData->GetMutablePoints().Shrink();
-
 		return true;
 	});
 }
diff --git a/Engine/Plugins/PCG/Source/PCG/Private/Elements/PCGAttributeSelectElement.cpp b/Engine/Plugins/PCG/Source/PCG/Private/Elements/PCGAttributeSelectElement.cpp
--- a/Engine/Plugins/PCG/Source/PCG/Private/Elements/PCGAttributeSelectElement.cpp
+++ b/Engine/Plugins/PCG/Source/PCG/Private/Elements/PCGAttributeSelectElement.cpp
@@ -6,7 +6,7 @@
 #include "PCGCustomVersion.h"
 #include "PCGParamData.h"
 #include "PCGPin.h"
-#include "Data/PCGPointData.h"
+#include "Data/PCGBasePointData.h"
 #include "Data/PCGSpatialData.h"
 #include "Elements/Metadata/PCGMetadataElementCommon.h"
 #include "Metadata/Accessors/PCGAttributeAccessorHelpers.h"
@@ -292,7 +292,7 @@ bool FPCGAttributeSelectElement::ExecuteInternal(FPCGContext* Context) const
 		return true;
 	}
 
-	const UPCGPointData* PointData = Cast<UPCGPointData>(SpatialData);
+	const UPCGBasePointData* PointData = Cast<UPCGBasePointData>(SpatialData);
 	if (!PointData && Context->Node && Context->Node->IsOutputPinConnected(PCGAttributeSelectConstants::OutputPointLabel))
 	{
 		PCGE_LOG(Warning, GraphAndLog, FText::Format(LOCTEXT("InputMissingPointData", "No point data in input, will output nothing in the '{0}' output pin"), FText::FromName(PCGAttributeSelectConstants::OutputPointLabel)));
@@ -433,9 +433,18 @@ bool FPCGAttributeSelectElement::ExecuteInternal(FPCGContext* Context) const
 	if(PointData && Context->Node && Context->Node->IsOutputPinConnected(PCGAttributeSelectConstants::OutputPointLabel))
 #endif
 	{
-		UPCGPointData* OutputPointData = FPCGContext::NewObject_AnyThread<UPCGPointData>(Context);
-		OutputPointData->InitializeFromData(PointData);
-		OutputPointData->GetMutablePoints().Add(PointData->GetPoint(OutputIndex));
+		UPCGBasePointData* OutputPointData = FPCGContext::NewPointData_AnyThread(Context);
+		
+		FPCGInitializeFromDataParams InitializeFromDataParams(PointData);
+		InitializeFromDataParams.bInheritSpatialData = false;
+		OutputPointData->InitializeFromDataWithParams(InitializeFromDataParams);
+
+		OutputPointData->SetNumPoints(1);
+
+		const FConstPCGPointValueRanges InRanges(PointData);
+		FPCGPointValueRanges OutRanges(OutputPointData, /*bAllocate=*/false);
+
+		OutRanges.SetFromValueRanges(0, InRanges, OutputIndex);
 
 		FPCGTaggedData& PointOutput = Outputs.Add_GetRef(Inputs[0]);
 		PointOutput.Data = OutputPointData;
diff --git a/Engine/Plugins/PCG/Source/PCG/Private/Elements/PCGBlurElement.cpp b/Engine/Plugins/PCG/Source/PCG/Private/Elements/PCGBlurElement.cpp
--- a/Engine/Plugins/PCG/Source/PCG/Private/Elements/PCGBlurElement.cpp
+++ b/Engine/Plugins/PCG/Source/PCG/Private/Elements/PCGBlurElement.cpp
@@ -3,7 +3,7 @@
 #include "Elements/PCGBlurElement.h"
 
 #include "PCGContext.h"
-#include "Data/PCGPointData.h"
+#include "Data/PCGBasePointData.h"
 #include "Helpers/PCGAsync.h"
 #include "Helpers/PCGHelpers.h"
 #include "Metadata/Accessors/PCGAttributeAccessorHelpers.h"
@@ -64,16 +64,17 @@ bool FPCGBlurElement::PrepareDataInternal(FPCGContext* InContext) const
 		FPCGTaggedData& OutputData = Outputs.Emplace_GetRef(Inputs[IterationIndex]);
 		OutputData.Pin = PCGPinConstants::DefaultOutputLabel;
 
-		OutState.InputPointData = Cast<UPCGPointData>(Inputs[IterationIndex].Data);
-		if (!OutState.InputPointData)
+
+		OutState.InputData = Cast<UPCGBasePointData>(Inputs[IterationIndex].Data);
+		if (!OutState.InputData)
 		{
 			return EPCGTimeSliceInitResult::NoOperation;
 		}
 
-		FPCGAttributePropertyInputSelector InputSelector = Settings->InputSource.CopyAndFixLast(OutState.InputPointData);
+		FPCGAttributePropertyInputSelector InputSelector = Settings->InputSource.CopyAndFixLast(OutState.InputData);
 
-		OutState.InputAccessor = PCGAttributeAccessorHelpers::CreateConstAccessor(OutState.InputPointData, InputSelector);
-		OutState.InputKeys = PCGAttributeAccessorHelpers::CreateConstKeys(OutState.InputPointData, InputSelector);
+		OutState.InputAccessor = PCGAttributeAccessorHelpers::CreateConstAccessor(OutState.InputData, InputSelector);
+		OutState.InputKeys = PCGAttributeAccessorHelpers::CreateConstKeys(OutState.InputData, InputSelector);
 
 		if (!OutState.InputAccessor || !OutState.InputKeys)
 		{
@@ -93,9 +94,9 @@ bool FPCGBlurElement::PrepareDataInternal(FPCGContext* InContext) const
 			return EPCGTimeSliceInitResult::NoOperation;
 		}
 
-		UPCGPointData* OutputPointData = CastChecked<UPCGPointData>(OutState.InputPointData->DuplicateData(InContext));
+		UPCGBasePointData* OutputPointData = CastChecked<UPCGBasePointData>(OutState.InputData->DuplicateData(InContext));
 
-		FPCGAttributePropertyOutputSelector OutputSelector = Settings->OutputTarget.CopyAndFixSource(&InputSelector, OutState.InputPointData);
+		FPCGAttributePropertyOutputSelector OutputSelector = Settings->OutputTarget.CopyAndFixSource(&InputSelector, OutState.InputData);
 
 		// Create attribute if needed
 		if (OutputSelector.IsBasicAttribute())
@@ -190,89 +191,98 @@ bool FPCGBlurElement::ExecuteInternal(FPCGContext* InContext) const
 				TArrayView<const AttributeType> TypedReadBuffer = MakeArrayView(reinterpret_cast<const AttributeType*>(ReadBuffer), IterState.InputKeys->GetNum());
 				TArrayView<AttributeType> TypedWriteBuffer = MakeArrayView(reinterpret_cast<AttributeType*>(WriteBuffer), IterState.InputKeys->GetNum());
 
-				const bool bIsDone = FPCGAsync::AsyncProcessingOneToOneEx(&Context->AsyncState, IterState.InputKeys->GetNum(), []() {}, [&IterState, Settings, &TypedReadBuffer, TypedWriteBuffer, GaussianExpRatio, GaussianRatio](int32 ReadIndex, int32 WriteIndex)
+				const bool bIsDone = FPCGAsync::AsyncProcessingOneToOneRangeEx(&Context->AsyncState, IterState.InputKeys->GetNum(), []() {}, [&IterState, Settings, &TypedReadBuffer, TypedWriteBuffer, GaussianExpRatio, GaussianRatio](int32 StartReadIndex, int32 StartWriteIndex, int32 Count)
 				{
-					const FPCGPoint& Point = IterState.InputPointData->GetPoints()[ReadIndex];
-					const FVector PointLocation = Point.Transform.GetLocation();
-					TArray<TTuple<int32, double>> NeighborsAndDistance;
+					const TConstPCGValueRange<FTransform> ReadTransformRange = IterState.InputData->GetConstTransformValueRange();
 
-					UPCGOctreeQueries::ForEachPointInsideSphere(IterState.InputPointData, PointLocation, Settings->SearchDistance, [&NeighborsAndDistance](const UPCGBasePointData* PointData, int32 Neighbor, double SquaredDistance)
-					{
-						NeighborsAndDistance.Emplace(Neighbor, SquaredDistance);
-					});
+					int32 NumWritten = 0;
 
-					if (!ensure(!NeighborsAndDistance.IsEmpty()))
+					for (int32 ReadIndex = StartReadIndex; ReadIndex < StartReadIndex + Count; ++ReadIndex)
 					{
-						return true;
-					}
+						const int32 WriteIndex = StartWriteIndex + NumWritten;
+						const FVector PointLocation = ReadTransformRange[ReadIndex].GetLocation();
+						TArray<TTuple<int32, double>> NeighborsAndDistance;
 
-					auto GetNeighborValue = [TypedReadBuffer, &IterState](int32 NeighborIndex) -> AttributeType
-					{
-						// Pointer arithmetic to know the neighbor index.
-						return TypedReadBuffer[NeighborIndex];
-					};
+						UPCGOctreeQueries::ForEachPointInsideSphere(IterState.InputData, PointLocation, Settings->SearchDistance, [&NeighborsAndDistance](const UPCGBasePointData* PointData, int32 Neighbor, double SquaredDistance)
+						{
+							NeighborsAndDistance.Emplace(Neighbor, SquaredDistance);
+						});
 
+						if (!ensure(!NeighborsAndDistance.IsEmpty()))
+						{
+							++NumWritten;
+							continue;
+						}
 
-					// If it is weighted by distance, we compute all the weights first, to have the sum of all of them (for normalization)
-					bool bValidWeights = false;
-					TArray<double> Weights;
-					if (Settings->BlurMode != EPCGBlurElementMode::Constant)
-					{
-						double WeightSum = 0.0;
-						bValidWeights = true;
-						Weights.Reserve(NeighborsAndDistance.Num());
-						for (const auto& [Neighbor, SquaredDistance] : NeighborsAndDistance)
+						auto GetNeighborValue = [TypedReadBuffer, &IterState](int32 NeighborIndex) -> AttributeType
+						{
+							// Pointer arithmetic to know the neighbor index.
+							return TypedReadBuffer[NeighborIndex];
+						};
+
+						// If it is weighted by distance, we compute all the weights first, to have the sum of all of them (for normalization)
+						bool bValidWeights = false;
+						TArray<double> Weights;
+						if (Settings->BlurMode != EPCGBlurElementMode::Constant)
 						{
-							double Weight = 0.0;
-							switch (Settings->BlurMode)
+							double WeightSum = 0.0;
+							bValidWeights = true;
+							Weights.Reserve(NeighborsAndDistance.Num());
+							for (const auto& [Neighbor, SquaredDistance] : NeighborsAndDistance)
 							{
-							case EPCGBlurElementMode::Linear:
-								Weight = 1.0 - FVector::Distance(PointLocation, IterState.InputPointData->GetPoints()[Neighbor].Transform.GetLocation()) / Settings->SearchDistance;
-								break;
-							case EPCGBlurElementMode::Gaussian:
-								Weight = GaussianRatio * FMath::Exp(-SquaredDistance * GaussianExpRatio);
-								break;
-							default:
-								checkNoEntry();
-								break;
+								double Weight = 0.0;
+								switch (Settings->BlurMode)
+								{
+								case EPCGBlurElementMode::Linear:
+									Weight = 1.0 - FVector::Distance(PointLocation, ReadTransformRange[Neighbor].GetLocation()) / Settings->SearchDistance;
+									break;
+								case EPCGBlurElementMode::Gaussian:
+									Weight = GaussianRatio * FMath::Exp(-SquaredDistance * GaussianExpRatio);
+									break;
+								default:
+									checkNoEntry();
+									break;
+								}
+
+								Weights.Add(Weight);
+								WeightSum += Weight;
 							}
 
-							Weights.Add(Weight);
-							WeightSum += Weight;
-						}
-
-						if (!FMath::IsNearlyZero(WeightSum))
-						{
-							for (double& Weight : Weights)
+							if (!FMath::IsNearlyZero(WeightSum))
 							{
-								Weight /= WeightSum;
+								for (double& Weight : Weights)
+								{
+									Weight /= WeightSum;
+								}
+							}
+							else
+							{
+								bValidWeights = false;
 							}
 						}
-						else
-						{
-							bValidWeights = false;
-						}
-					}
 
-					const double DefaultWeight = 1.0 / NeighborsAndDistance.Num();
-					const bool bUseDefaultWeight = !bValidWeights;
+						const double DefaultWeight = 1.0 / NeighborsAndDistance.Num();
+						const bool bUseDefaultWeight = !bValidWeights;
 
-					TypedWriteBuffer[WriteIndex] = PCG::Private::MetadataTraits<AttributeType>::ZeroValueForWeightedSum();
+						TypedWriteBuffer[WriteIndex] = PCG::Private::MetadataTraits<AttributeType>::ZeroValueForWeightedSum();
 
-					for (int32 i = 0; i < NeighborsAndDistance.Num(); ++i)
-					{
-						AttributeType Value = GetNeighborValue(NeighborsAndDistance[i].Get<0>());
-						const double Weight = bUseDefaultWeight ? DefaultWeight : Weights[i];
-						TypedWriteBuffer[WriteIndex] = PCG::Private::MetadataTraits<AttributeType>::WeightedSum(TypedWriteBuffer[WriteIndex], Value, Weight);
-					}
+						for (int32 i = 0; i < NeighborsAndDistance.Num(); ++i)
+						{
+							AttributeType Value = GetNeighborValue(NeighborsAndDistance[i].Get<0>());
+							const double Weight = bUseDefaultWeight ? DefaultWeight : Weights[i];
+							TypedWriteBuffer[WriteIndex] = PCG::Private::MetadataTraits<AttributeType>::WeightedSum(TypedWriteBuffer[WriteIndex], Value, Weight);
+						}
 
-					// For Quaternions, we need to normalize them
-					if constexpr (std::is_same_v<AttributeType, FQuat>)
-					{
-						TypedWriteBuffer[WriteIndex].Normalize();
+						// For Quaternions, we need to normalize them
+						if constexpr (std::is_same_v<AttributeType, FQuat>)
+						{
+							TypedWriteBuffer[WriteIndex].Normalize();
+						}
+						++NumWritten;
 					}
 
-					return true;
+					check(Count == NumWritten);
+					return Count;
 				},/*bEnableTimeSlicing=*/true);
 
 				if (!bIsDone)
diff --git a/Engine/Plugins/PCG/Source/PCG/Private/Elements/PCGBoundsFromMesh.cpp b/Engine/Plugins/PCG/Source/PCG/Private/Elements/PCGBoundsFromMesh.cpp
--- a/Engine/Plugins/PCG/Source/PCG/Private/Elements/PCGBoundsFromMesh.cpp
+++ b/Engine/Plugins/PCG/Source/PCG/Private/Elements/PCGBoundsFromMesh.cpp
@@ -4,7 +4,7 @@
 
 #include "PCGContext.h"
 #include "PCGParamData.h"
-#include "Data/PCGPointData.h"
+#include "Data/PCGBasePointData.h"
 #include "Elements/Metadata/PCGMetadataElementCommon.h"
 #include "Helpers/PCGAsync.h"
 #include "Metadata/Accessors/PCGAttributeAccessorHelpers.h"
@@ -129,7 +129,7 @@ bool FPCGBoundsFromMeshElement::PrepareDataInternal(FPCGContext* InContext) cons
 			while (Context->CurrentPrepareIndex < Inputs.Num())
 			{
 				const FPCGTaggedData& Input = Inputs[Context->CurrentPrepareIndex];
-				const UPCGPointData* PointData = Cast<UPCGPointData>(Input.Data);
+				const UPCGBasePointData* PointData = Cast<UPCGBasePointData>(Input.Data);
 
 				if (!PointData)
 				{
@@ -240,7 +240,7 @@ bool FPCGBoundsFromMeshElement::ExecuteInternal(FPCGContext* InContext) const
 		FPCGTaggedData& Output = Outputs.Add_GetRef(Input);
 
 		// If input is not a point data -> passthrough
-		const UPCGPointData* PointData = Cast<UPCGPointData>(Input.Data);
+		const UPCGBasePointData* PointData = Cast<UPCGBasePointData>(Input.Data);
 
 		if (!PointData)
 		{
@@ -262,24 +262,43 @@ bool FPCGBoundsFromMeshElement::ExecuteInternal(FPCGContext* InContext) const
 
 			const FBox StaticMeshBounds = SingleMesh->GetBoundingBox();
 			
-			UPCGPointData* OutputData = FPCGContext::NewObject_AnyThread<UPCGPointData>(Context);
-			OutputData->InitializeFromData(PointData);
+			UPCGBasePointData* OutputData = FPCGContext::NewPointData_AnyThread(Context);
 
-			const TArray<FPCGPoint>& InputPoints = PointData->GetPoints();
-			TArray<FPCGPoint>& OutputPoints = OutputData->GetMutablePoints();
-			OutputPoints.SetNumUninitialized(InputPoints.Num());
+			// Inherit from Input
+			OutputData->InitializeFromData(PointData);
+			OutputData->SetNumPoints(PointData->GetNumPoints(), /*bInitializeValues=*/false);
 
-			auto CopyAndSetBounds = [&InputPoints, &OutputPoints, &StaticMeshBounds](const int32 ReadIndex, const int32 WriteIndex)
+			if (OutputData->HasSpatialDataParent())
+			{
+				OutputData->SetBoundsMin(StaticMeshBounds.Min);
+				OutputData->SetBoundsMax(StaticMeshBounds.Max);
+			}
+			else
 			{
-				const FPCGPoint& InputPoint = InputPoints[ReadIndex];
-				FPCGPoint& OutputPoint = OutputPoints[WriteIndex];
+				OutputData->AllocateProperties(PointData->GetAllocatedProperties() | EPCGPointNativeProperties::BoundsMin | EPCGPointNativeProperties::BoundsMax);
+				auto CopyAndSetBounds = [OutputData, PointData, &StaticMeshBounds](const int32 StartReadIndex, const int32 StartWriteIndex, const int32 Count)
+				{
+					int32 NumWritten = 0;
+
+					const FConstPCGPointValueRanges InRanges(PointData);
+					FPCGPointValueRanges OutRanges(OutputData, /*bAllocate=*/false);
+
+					for (int32 ReadIndex = StartReadIndex; ReadIndex < StartReadIndex + Count; ++ReadIndex)
+					{
+						const int32 WriteIndex = StartWriteIndex + NumWritten;
+						OutRanges.SetFromValueRanges(WriteIndex, InRanges, ReadIndex);
+
+						OutRanges.BoundsMinRange[WriteIndex] = StaticMeshBounds.Min;
+						OutRanges.BoundsMaxRange[WriteIndex] = StaticMeshBounds.Max;
+						++NumWritten;
+					}
 
-				OutputPoint = InputPoint;
-				OutputPoint.BoundsMin = StaticMeshBounds.Min;
-				OutputPoint.BoundsMax = StaticMeshBounds.Max;
-			};
+					check(Count == NumWritten);
+					return Count;
+				};
 
-			FPCGAsync::AsyncProcessingOneToOneEx(&Context->AsyncState, InputPoints.Num(), /*Initialize=*/[]() {}, CopyAndSetBounds, /*bEnableTimeSlicing=*/false);
+				FPCGAsync::AsyncProcessingOneToOneRangeEx(&Context->AsyncState, PointData->GetNumPoints(), /*Initialize=*/[]() {}, CopyAndSetBounds, /*bEnableTimeSlicing=*/false);
+			}
 
 			Output.Data = OutputData;
 		}
@@ -288,27 +307,46 @@ bool FPCGBoundsFromMeshElement::ExecuteInternal(FPCGContext* InContext) const
 			FPCGBoundsFromMeshContext::InputMeshesData* InputData = Context->PerInputData.FindByPredicate([Index = Context->CurrentExecuteIndex - 1](const FPCGBoundsFromMeshContext::InputMeshesData& InputData) { return InputData.InputIndex == Index; });
 			if (InputData)
 			{
-				UPCGPointData* OutputData = FPCGContext::NewObject_AnyThread<UPCGPointData>(Context);
+				UPCGBasePointData* OutputData = FPCGContext::NewPointData_AnyThread(Context);
 				OutputData->InitializeFromData(PointData);
+				OutputData->SetNumPoints(PointData->GetNumPoints(), /*bInitializeValues=*/false);
+
+				if (!OutputData->HasSpatialDataParent())
+				{
+					OutputData->AllocateProperties(PointData->GetAllocatedProperties());
+				}
 
-				const TArray<FPCGPoint>& InputPoints = PointData->GetPoints();
-				TArray<FPCGPoint>& OutputPoints = OutputData->GetMutablePoints();
-				OutputPoints.SetNumUninitialized(InputPoints.Num());
+				OutputData->AllocateProperties(EPCGPointNativeProperties::BoundsMin | EPCGPointNativeProperties::BoundsMax);
 
-				auto CopyAndSetBounds = [&InputPoints, &OutputPoints, Context, InputData](const int32 ReadIndex, const int32 WriteIndex)
+				auto CopyAndSetBounds = [OutputData, PointData, Context, InputData](const int32 StartReadIndex, const int32 StartWriteIndex, const int32 Count)
 				{
-					const FPCGPoint& InputPoint = InputPoints[ReadIndex];
-					FPCGPoint& OutputPoint = OutputPoints[WriteIndex];
+					int32 NumWritten = 0;
+
+					const FConstPCGPointValueRanges InRanges(PointData);
+					FPCGPointValueRanges OutRanges(OutputData, /*bAllocate=*/false);
 
-					OutputPoint = InputPoint;
-					if (FBox* MatchingBounds = Context->MeshToBoundsMap.Find(InputData->MeshValueIndex[ReadIndex]))
+					for (int32 ReadIndex = StartReadIndex; ReadIndex < StartReadIndex + Count; ++ReadIndex)
 					{
-						OutputPoint.BoundsMin = MatchingBounds->Min;
-						OutputPoint.BoundsMax = MatchingBounds->Max;
+						const int32 WriteIndex = StartWriteIndex + NumWritten;
+
+						if (!OutputData->HasSpatialDataParent())
+						{
+							OutRanges.SetFromValueRanges(WriteIndex, InRanges, ReadIndex);
+						}
+
+						if (FBox* MatchingBounds = Context->MeshToBoundsMap.Find(InputData->MeshValueIndex[ReadIndex]))
+						{
+							OutRanges.BoundsMinRange[WriteIndex] = MatchingBounds->Min;
+							OutRanges.BoundsMaxRange[WriteIndex] = MatchingBounds->Max;
+						}
+
+						++NumWritten;
 					}
+					check(NumWritten == Count);
+					return Count;
 				};
 
-				FPCGAsync::AsyncProcessingOneToOneEx(&Context->AsyncState, InputPoints.Num(), /*Initialize=*/[]() {}, CopyAndSetBounds, /*bEnableTimeSlicing=*/false);
+				FPCGAsync::AsyncProcessingOneToOneRangeEx(&Context->AsyncState, PointData->GetNumPoints(), /*Initialize=*/[]() {}, CopyAndSetBounds, /*bEnableTimeSlicing=*/false);
 
 				Output.Data = OutputData;
 			}
diff --git a/Engine/Plugins/PCG/Source/PCG/Private/Elements/PCGClusterElement.cpp b/Engine/Plugins/PCG/Source/PCG/Private/Elements/PCGClusterElement.cpp
--- a/Engine/Plugins/PCG/Source/PCG/Private/Elements/PCGClusterElement.cpp
+++ b/Engine/Plugins/PCG/Source/PCG/Private/Elements/PCGClusterElement.cpp
@@ -47,7 +47,7 @@ namespace PCGClusterElement
 
 			// Set up a few references to state data.
 			uint8& ProcessingStage = ClusteringData.ProcessingStage;
-			const TArray<FPCGPoint>& Points = ClusteringData.OutPointData->GetPoints();
+			const TConstPCGValueRange<FTransform> TransformRange = ClusteringData.OutputData->GetConstTransformValueRange();
 			const TArray<FVector, TInlineAllocator<16>>& Centroids = ClusteringData.Centroids;
 			TArray<int32>& ClusterAssignments = ClusteringData.PointToClusterAssignments;
 
@@ -56,7 +56,7 @@ namespace PCGClusterElement
 				TRACE_CPUPROFILER_EVENT_SCOPE(FPCGClusterElement::KMeans::IterationLoop);
 
 				// Find the centroid that best represents each point.
-				auto ProcessPoints = [&Points, &Centroids, &bConverged = ClusteringData.bConverged](const int32& PointIndex, int32& InOutCentroidIndex)
+				auto ProcessPoints = [&TransformRange, &Centroids, &bConverged = ClusteringData.bConverged](const int32& PointIndex, int32& InOutCentroidIndex)
 				{
 					TRACE_CPUPROFILER_EVENT_SCOPE(FPCGClusterElement::KMeans::ProcessPoints);
 
@@ -66,7 +66,7 @@ namespace PCGClusterElement
 					 * first iteration.
 					 */
 					int CentroidIndex = InOutCentroidIndex;
-					double MinDistanceSquared = FVector::DistSquared(Points[PointIndex].Transform.GetLocation(), Centroids[InOutCentroidIndex]);
+					double MinDistanceSquared = FVector::DistSquared(TransformRange[PointIndex].GetLocation(), Centroids[InOutCentroidIndex]);
 					for (CentroidIndex = 0; CentroidIndex < Centroids.Num(); ++CentroidIndex)
 					{
 						// Stabilizes against flipping centroids in the event of equidistant points
@@ -75,7 +75,7 @@ namespace PCGClusterElement
 							continue;
 						}
 
-						const double DistanceSquared = FVector::DistSquared(Points[PointIndex].Transform.GetLocation(), Centroids[CentroidIndex]);
+						const double DistanceSquared = FVector::DistSquared(TransformRange[PointIndex].GetLocation(), Centroids[CentroidIndex]);
 						if (DistanceSquared < MinDistanceSquared)
 						{
 							MinDistanceSquared = DistanceSquared;
@@ -108,18 +108,18 @@ namespace PCGClusterElement
 				}
 
 				// Update the centroids with newly added or removed points
-				auto ProcessCentroids = [Points, ClusterAssignments](const int32 CentroidIndex, FVector& OutResult)
+				auto ProcessCentroids = [TransformRange, ClusterAssignments](const int32 CentroidIndex, FVector& OutResult)
 				{
 					TRACE_CPUPROFILER_EVENT_SCOPE(FPCGClusterElement::KMeans::ProcessCentroids);
 
 					FVector CentroidSum = FVector::ZeroVector;
 					int Count = 0;
 
-					for (int PointIndex = 0; PointIndex < Points.Num(); ++PointIndex)
+					for (int PointIndex = 0; PointIndex < TransformRange.Num(); ++PointIndex)
 					{
 						if (ClusterAssignments[PointIndex] == CentroidIndex)
 						{
-							CentroidSum += Points[PointIndex].Transform.GetLocation();
+							CentroidSum += TransformRange[PointIndex].GetLocation();
 							++Count;
 						}
 					}
@@ -473,16 +473,17 @@ bool FPCGClusterElement::PrepareDataInternal(FPCGContext* InContext) const
 				PCGLog::InputOutput::LogFirstInputOnlyWarning(PCGClusterElement::Constants::InputCentroidsLabel, Context);
 			}
 
-			const UPCGPointData* CentroidPoints = Cast<UPCGPointData>(CentroidInputs[0].Data);
+			const UPCGBasePointData* CentroidPoints = Cast<UPCGBasePointData>(CentroidInputs[0].Data);
 			if (!CentroidPoints || CentroidPoints->IsEmpty())
 			{
 				return EPCGTimeSliceInitResult::NoOperation;
 			}
 
-			OutState.InitialCentroids.Reserve(CentroidPoints->GetPoints().Num());
-			for (const FPCGPoint& Point : CentroidPoints->GetPoints())
+			OutState.InitialCentroids.Reserve(CentroidPoints->GetNumPoints());
+			const TConstPCGValueRange<FTransform> TransformRange = CentroidPoints->GetConstTransformValueRange();
+			for (const FTransform& Transform : TransformRange)
 			{
-				OutState.InitialCentroids.Emplace(Point.Transform.GetLocation());
+				OutState.InitialCentroids.Emplace(Transform.GetLocation());
 			}
 		}
 
@@ -493,15 +494,20 @@ bool FPCGClusterElement::PrepareDataInternal(FPCGContext* InContext) const
 	{
 		TRACE_CPUPROFILER_EVENT_SCOPE(FPCGClusterElement::PrepareData::InitializeIteration);
 
-		const UPCGPointData* InputPointData = Cast<UPCGPointData>(PointInputs[IterationIndex].Data);
+		const UPCGBasePointData* InputPointData = Cast<UPCGBasePointData>(PointInputs[IterationIndex].Data);
 		if (!InputPointData || InputPointData->IsEmpty())
 		{
 			return EPCGTimeSliceInitResult::NoOperation;
 		}
 
-		UPCGPointData* OutPointData = FPCGContext::NewObject_AnyThread<UPCGPointData>(Context);
+		UPCGBasePointData* OutPointData = FPCGContext::NewPointData_AnyThread(Context);
 		OutPointData->InitializeFromData(InputPointData);
-		OutPointData->GetMutablePoints() = InputPointData->GetPointsCopy();
+
+		if (!OutPointData->HasSpatialDataParent())
+		{
+			UPCGBasePointData::SetPoints(InputPointData, OutPointData, {}, /*bCopyAll=*/true);
+		}
+
 		UPCGMetadata* OutPointMetadata = OutPointData->MutableMetadata();
 		check(OutPointMetadata);
 
@@ -511,7 +517,10 @@ bool FPCGClusterElement::PrepareDataInternal(FPCGContext* InContext) const
 			PCGMetadataElementCommon::ClearOrCreateAttribute(OutPointMetadata, Selector, 0, Context);
 		}
 
-		OutState.OutPointData = OutPointData;
+		OutState.OutputData = OutPointData;
+PRAGMA_DISABLE_DEPRECATION_WARNINGS
+		OutState.OutPointData = Cast<UPCGPointData>(OutPointData);
+PRAGMA_ENABLE_DEPRECATION_WARNINGS
 
 		OutState.Accessor = PCGAttributeAccessorHelpers::CreateAccessor(OutPointData, Selector);
 		OutState.Keys = PCGAttributeAccessorHelpers::CreateKeys(OutPointData, Selector);
@@ -553,7 +562,7 @@ bool FPCGClusterElement::PrepareDataInternal(FPCGContext* InContext) const
 		{
 			if (ExecState.bSelectRandomCentroids)
 			{
-				const TArray<FPCGPoint>& Points = OutState.OutPointData->GetPoints();
+				const TConstPCGValueRange<FTransform> TransformRange = OutState.OutputData->GetConstTransformValueRange();
 
 				FRandomStream RandomStream(Settings->GetSeed());
 				TArray<int32> StartingIndices = PCGHelpers::GetRandomIndices(RandomStream, NumPoints, NumClusters);
@@ -561,7 +570,7 @@ bool FPCGClusterElement::PrepareDataInternal(FPCGContext* InContext) const
 				OutState.Centroids.Reserve(NumClusters);
 				for (int i = 0; i < NumClusters; ++i)
 				{
-					OutState.Centroids.Emplace(Points[StartingIndices[i]].Transform.GetLocation());
+					OutState.Centroids.Emplace(TransformRange[StartingIndices[i]].GetLocation());
 				}
 			}
 			else
@@ -591,13 +600,13 @@ bool FPCGClusterElement::PrepareDataInternal(FPCGContext* InContext) const
 			{
 				InitializeCentroids();
 
-				const TArray<FPCGPoint>& Points = OutPointData->GetPoints();
+				const TConstPCGValueRange<FTransform> TransformRange = OutState.OutputData->GetConstTransformValueRange();
 
 				// Normalization of point locations may be required to bring the GPD Function's Exponential Term into the -709.8 -> 709.8 range to prevent numerical overflow.
 				double MaxSquaredLength = std::numeric_limits<double>::min();
-				for (const FPCGPoint& Point : Points)
+				for (const FTransform& Transform : TransformRange)
 				{
-					MaxSquaredLength = FMath::Max(MaxSquaredLength, Point.Transform.GetLocation().SquaredLength());
+					MaxSquaredLength = FMath::Max(MaxSquaredLength, Transform.GetLocation().SquaredLength());
 				}
 
 				// Only normalize the system if needed.
@@ -614,12 +623,12 @@ bool FPCGClusterElement::PrepareDataInternal(FPCGContext* InContext) const
 					OutState.EMData.Gaussians.Emplace(OutState.Centroids[i] * OutState.EMData.ScalingFactor, FMatrix::Identity, 1.0 / NumClusters);
 				}
 
-				OutState.EMData.Probabilities.Initialize(Points.Num(), NumClusters);
+				OutState.EMData.Probabilities.Initialize(OutState.OutputData->GetNumPoints(), NumClusters);
 
-				OutState.EMData.NormalizedPointLocations.Reserve(Points.Num());
-				Algo::Transform(Points, OutState.EMData.NormalizedPointLocations, [ScalingFactor = OutState.EMData.ScalingFactor](const FPCGPoint& Point)
+				OutState.EMData.NormalizedPointLocations.Reserve(OutState.OutputData->GetNumPoints());
+				Algo::Transform(TransformRange, OutState.EMData.NormalizedPointLocations, [ScalingFactor = OutState.EMData.ScalingFactor](const FTransform& Transform)
 				{
-					return Point.Transform.GetLocation() * ScalingFactor;
+					return Transform.GetLocation() * ScalingFactor;
 				});
 
 				break;
@@ -671,10 +680,10 @@ bool FPCGClusterElement::ExecuteInternal(FPCGContext* InContext) const
 		}
 
 		// Once finished, set the attribute values.
-		check(ClusteringData.OutPointData);
+		check(ClusteringData.OutputData);
 		FPCGTaggedData& OutputData = Context->OutputData.TaggedData.Emplace_GetRef();
 		OutputData.Pin = PCGPinConstants::DefaultOutputLabel;
-		OutputData.Data = ClusteringData.OutPointData;
+		OutputData.Data = ClusteringData.OutputData;
 
 		// Update the cluster assignment attribute.
 		ClusteringData.Accessor->SetRange(MakeConstArrayView(ClusteringData.PointToClusterAssignments), 0, *ClusteringData.Keys);
@@ -702,10 +711,17 @@ bool FPCGClusterElement::ExecuteInternal(FPCGContext* InContext) const
 				OutCentroids = std::move(ClusteringData.Centroids);
 			}
 
-			UPCGPointData* OutputCentroidPointData = FPCGContext::NewObject_AnyThread<UPCGPointData>(Context);
-			TArray<FPCGPoint>& OutputCentroidPoints = OutputCentroidPointData->GetMutablePoints();
-			OutputCentroidPoints.Reserve(OutCentroids.Num());
-			Algo::Transform(OutCentroids, OutputCentroidPoints, [](const FVector& Location) { return FPCGPoint(FTransform(Location), 1.0, PCGHelpers::ComputeSeedFromPosition(Location)); });
+			UPCGBasePointData* OutputCentroidPointData = FPCGContext::NewPointData_AnyThread(Context);
+			OutputCentroidPointData->SetNumPoints(OutCentroids.Num());
+			OutputCentroidPointData->SetDensity(1.0f);
+			OutputCentroidPointData->AllocateProperties(EPCGPointNativeProperties::Transform | EPCGPointNativeProperties::Seed);
+
+			FPCGPointValueRanges OutRanges(OutputCentroidPointData, /*bAllocate=*/false);
+			for (int32 i = 0; i < OutCentroids.Num(); ++i)
+			{
+				OutRanges.TransformRange[i] = FTransform(OutCentroids[i]);
+				OutRanges.SeedRange[i] = PCGHelpers::ComputeSeedFromPosition(OutCentroids[i]);
+			}
 
 			FPCGTaggedData& OutputCentroidData = Context->OutputData.TaggedData.Emplace_GetRef();
 			OutputCentroidData.Pin = PCGClusterElement::Constants::FinalCentroidsLabel;
@@ -730,7 +746,7 @@ bool FPCGClusterElement::ExecuteInternal(FPCGContext* InContext) const
 				}
 
 				TArray<int32, TInlineAllocator<16>> CentroidElementCount;
-				CentroidElementCount.SetNumZeroed(OutputCentroidPoints.Num());
+				CentroidElementCount.SetNumZeroed(OutCentroids.Num());
 				// Tabulate the count of each ClusterID
 				for (int32 PointIndex = 0; PointIndex < ClusteringData.PointToClusterAssignments.Num(); ++PointIndex)
 				{
diff --git a/Engine/Plugins/PCG/Source/PCG/Private/Elements/PCGCombinePoints.cpp b/Engine/Plugins/PCG/Source/PCG/Private/Elements/PCGCombinePoints.cpp
--- a/Engine/Plugins/PCG/Source/PCG/Private/Elements/PCGCombinePoints.cpp
+++ b/Engine/Plugins/PCG/Source/PCG/Private/Elements/PCGCombinePoints.cpp
@@ -4,7 +4,7 @@
 
 #include "PCGContext.h"
 #include "PCGPin.h"
-#include "Data/PCGPointData.h"
+#include "Data/PCGBasePointData.h"
 #include "Helpers/PCGHelpers.h"
 
 #define LOCTEXT_NAMESPACE "PCGCombinePointsElement"
@@ -27,40 +27,42 @@ bool FPCGCombinePointsElement::ExecuteInternal(FPCGContext* Context) const
 
 	for (int i = 0; i < Inputs.Num(); ++i)
 	{
-		const UPCGPointData* InputPointData = Cast<UPCGPointData>(Inputs[i].Data);
+		const UPCGBasePointData* InputPointData = Cast<UPCGBasePointData>(Inputs[i].Data);
 		if (!InputPointData)
 		{
 			PCGE_LOG(Warning, GraphAndLog, LOCTEXT("InvalidInputPointData", "The input is not point data, skipped."));
 			continue;
 		}
 
-		const TArray<FPCGPoint>& InputPoints = InputPointData->GetPoints();
-		if (InputPoints.IsEmpty())
+		if (InputPointData->IsEmpty())
 		{
 			PCGE_LOG(Warning, GraphAndLog, LOCTEXT("NoPointsFound", "No points were found in the input data, skipped."));
 			continue;
 		}
 
-		FPCGPoint InputPoint = InputPoints[0];
+		const FConstPCGPointValueRanges InRanges(InputPointData);
+		FPCGPoint InputPoint = InRanges.GetPoint(0);
 		InputPoint.ApplyScaleToBounds();
 
 		const FTransform& PointTransform = Settings->bUseFirstPointTransform ? InputPoint.Transform : Settings->PointTransform;
 		const FMatrix InversePointTransform = PointTransform.ToInverseMatrixWithScale();
 
 		FPCGTaggedData& Output = Outputs.Add_GetRef(Inputs[i]);
-		UPCGPointData* OutputPointData = FPCGContext::NewObject_AnyThread<UPCGPointData>(Context);
-		OutputPointData->InitializeFromData(InputPointData);
-		FPCGPoint& OutputPoint = OutputPointData->GetMutablePoints().Emplace_GetRef();
+		UPCGBasePointData* OutputPointData = FPCGContext::NewPointData_AnyThread(Context);
+		
+		FPCGInitializeFromDataParams InitializeFromDataParams(InputPointData);
+		InitializeFromDataParams.bInheritSpatialData = false;
+		OutputPointData->InitializeFromDataWithParams(InitializeFromDataParams);
+		
+		OutputPointData->SetNumPoints(1);
 		Output.Data = OutputPointData;
 
 		FPCGPoint OutPoint = FPCGPoint();
 		FBox OutBox(EForceInit::ForceInit);
 
-		for (int j = 0; j < InputPoints.Num(); ++j)
+		for (int j = 0; j < InputPointData->GetNumPoints(); ++j)
 		{
-			const FPCGPoint& InPoint = InputPoints[j];
-
-			OutBox += InPoint.GetLocalBounds().TransformBy(InPoint.Transform.ToMatrixWithScale() * InversePointTransform);
+			OutBox += PCGPointHelpers::GetLocalBounds(InRanges.BoundsMinRange[j], InRanges.BoundsMaxRange[j]).TransformBy(InRanges.TransformRange[j].ToMatrixWithScale() * InversePointTransform);
 		};
 
 		OutPoint.SetLocalBounds(OutBox);
@@ -72,7 +74,8 @@ bool FPCGCombinePointsElement::ExecuteInternal(FPCGContext* Context) const
 			OutPoint.ResetPointCenter(FVector(0.5));
 		}
 
-		OutputPoint = OutPoint;
+		FPCGPointValueRanges OutRanges(OutputPointData, /*bAllocate=*/false);
+		OutRanges.SetFromPoint(0, OutPoint);
 	}
 
 	return true;
diff --git a/Engine/Plugins/PCG/Source/PCG/Private/Elements/PCGConvexHull2D.cpp b/Engine/Plugins/PCG/Source/PCG/Private/Elements/PCGConvexHull2D.cpp
--- a/Engine/Plugins/PCG/Source/PCG/Private/Elements/PCGConvexHull2D.cpp
+++ b/Engine/Plugins/PCG/Source/PCG/Private/Elements/PCGConvexHull2D.cpp
@@ -4,7 +4,7 @@
 
 #include "PCGContext.h"
 #include "PCGModule.h"
-#include "Data/PCGPointData.h"
+#include "Data/PCGBasePointData.h"
 
 #include "Math/ConvexHull2d.h"
 
@@ -28,7 +28,7 @@ bool FPCGConvexHull2DElement::ExecuteInternal(FPCGContext* Context) const
 
 	for (const FPCGTaggedData& Input : Context->InputData.GetInputsByPin(PCGPinConstants::DefaultInputLabel))
 	{
-		const UPCGPointData* PointData = Cast<UPCGPointData>(Input.Data);
+		const UPCGBasePointData* PointData = Cast<UPCGBasePointData>(Input.Data);
 
 		if (!PointData)
 		{
@@ -36,22 +36,30 @@ bool FPCGConvexHull2DElement::ExecuteInternal(FPCGContext* Context) const
 			continue;
 		}
 
-		const TArray<FPCGPoint>& Points = PointData->GetPoints();
-
 		TArray<FVector> PointsPositions;
 		TArray<int32> ConvexHullIndices;
-		Algo::Transform(Points, PointsPositions, [](const FPCGPoint& Point) { return Point.Transform.GetLocation(); });
+
+		const FConstPCGPointValueRanges InRanges(PointData);
+		Algo::Transform(InRanges.TransformRange, PointsPositions, [](const FTransform& Transform) { return Transform.GetLocation(); });
 
 		ConvexHull2D::ComputeConvexHull(PointsPositions, ConvexHullIndices);
 
 		FPCGTaggedData& Output = Outputs.Add_GetRef(Input);
-		UPCGPointData* OutputPointData = FPCGContext::NewObject_AnyThread<UPCGPointData>(Context);
-		OutputPointData->InitializeFromData(PointData);
-		TArray<FPCGPoint>& OutPoints = OutputPointData->GetMutablePoints();
-		OutPoints.Reserve(ConvexHullIndices.Num());
+		UPCGBasePointData* OutputPointData = FPCGContext::NewPointData_AnyThread(Context);
+		
+		FPCGInitializeFromDataParams InitializeFromDataParams(PointData);
+		InitializeFromDataParams.bInheritSpatialData = false;
+		OutputPointData->InitializeFromDataWithParams(InitializeFromDataParams);
+		
+		OutputPointData->SetNumPoints(ConvexHullIndices.Num());
+		OutputPointData->AllocateProperties(PointData->GetAllocatedProperties());
+
+		FPCGPointValueRanges OutRanges(OutputPointData, /*bAllocate=*/false);
+		int32 WriteIndex = 0;
 		for (int32 Index : ConvexHullIndices)
 		{
-			OutPoints.Add(Points[Index]);
+			OutRanges.SetFromValueRanges(WriteIndex, InRanges, Index);
+			++WriteIndex;
 		}
 
 		Output.Data = OutputPointData;
diff --git a/Engine/Plugins/PCG/Source/PCG/Private/Elements/PCGCreateCollisionData.cpp b/Engine/Plugins/PCG/Source/PCG/Private/Elements/PCGCreateCollisionData.cpp
--- a/Engine/Plugins/PCG/Source/PCG/Private/Elements/PCGCreateCollisionData.cpp
+++ b/Engine/Plugins/PCG/Source/PCG/Private/Elements/PCGCreateCollisionData.cpp
@@ -4,7 +4,7 @@
 
 #include "PCGContext.h"
 #include "Data/PCGCollisionWrapperData.h"
-#include "Data/PCGPointData.h"
+#include "Data/PCGBasePointData.h"
 
 #include UE_INLINE_GENERATED_CPP_BY_NAME(PCGCreateCollisionData)
 
@@ -48,16 +48,11 @@ FPCGElementPtr UPCGCreateCollisionDataSettings::CreateElement() const
 	return MakeShared<FPCGCreateCollisionDataElement>();
 }
 
-FPCGCreateCollisionContext::~FPCGCreateCollisionContext()
+void FPCGCreateCollisionContext::AddExtraStructReferencedObjects(FReferenceCollector& Collector)
 {
-	// If the execution is cancelled between the prepare data and the execution,
-	// we might have data that's been rooted temporarily - take care of this here.
 	for (InputMeshData& ThisInputData : PerInputData)
 	{
-		if (ThisInputData.Data && ThisInputData.Data->IsRooted())
-		{
-			ThisInputData.Data->RemoveFromRoot();
-		}
+		Collector.AddReferencedObject(ThisInputData.Data);
 	}
 }
 
@@ -79,7 +74,7 @@ bool FPCGCreateCollisionDataElement::PrepareDataInternal(FPCGContext* InContext)
 		{
 			const FPCGTaggedData& Input = Inputs[InputIndex];
 
-			if (const UPCGPointData* PointData = Cast<UPCGPointData>(Input.Data))
+			if (const UPCGBasePointData* PointData = Cast<UPCGBasePointData>(Input.Data))
 			{
 				FPCGCreateCollisionContext::InputMeshData InputMeshData;
 				InputMeshData.InputIndex = InputIndex;
@@ -94,9 +89,6 @@ bool FPCGCreateCollisionDataElement::PrepareDataInternal(FPCGContext* InContext)
 				}
 				else
 				{
-					// Root object so it can't get deleted until we finish initialize and properly manage it as part of the collection.
-					InputMeshData.Data->AddToRoot();
-
 					for (const FSoftObjectPath& MeshPath : InputMeshData.MeshPaths)
 					{
 						if (!MeshPath.IsNull())
@@ -140,9 +132,6 @@ bool FPCGCreateCollisionDataElement::ExecuteInternal(FPCGContext* InContext) con
 			check(MatchingData->Data);
 			MatchingData->Data->FinalizeInitializationEx(MatchingData->MeshPaths);
 
-			// Unroot collision wrapper data object since it'll be tracked through the collection
-			MatchingData->Data->RemoveFromRoot();
-
 			Output.Data = MatchingData->Data;
 			
 			// Null out the data so the context destructor doesn't interact with the rooting
diff --git a/Engine/Plugins/PCG/Source/PCG/Private/Elements/PCGCreatePointsSphere.cpp b/Engine/Plugins/PCG/Source/PCG/Private/Elements/PCGCreatePointsSphere.cpp
--- a/Engine/Plugins/PCG/Source/PCG/Private/Elements/PCGCreatePointsSphere.cpp
+++ b/Engine/Plugins/PCG/Source/PCG/Private/Elements/PCGCreatePointsSphere.cpp
@@ -3,7 +3,7 @@
 #include "Elements/PCGCreatePointsSphere.h"
 
 #include "PCGComponent.h"
-#include "Data/PCGPointData.h"
+#include "Data/PCGBasePointData.h"
 #include "Elements/PCGCreatePoints.h"
 #include "Helpers/PCGAsync.h"
 #include "Helpers/PCGHelpers.h"
@@ -437,8 +437,7 @@ bool FPCGCreatePointsSphereElement::ExecuteInternal(FPCGContext* InContext) cons
 	TArray<FPCGTaggedData>& Outputs = InContext->OutputData.TaggedData;
 	FPCGTaggedData& Output = Outputs.Emplace_GetRef();
 
-	UPCGPointData* PointData = FPCGContext::NewObject_AnyThread<UPCGPointData>(InContext);
-	TArray<FPCGPoint>& OutputPoints = PointData->GetMutablePoints();
+	UPCGBasePointData* PointData = FPCGContext::NewPointData_AnyThread(InContext);
 	Output.Data = PointData;
 
 	FTransform LocalTransform = Settings->CoordinateSpace == EPCGCoordinateSpace::World ? FTransform::Identity : InContext->ExecutionSource->GetExecutionState().GetTransform();
@@ -456,7 +455,8 @@ bool FPCGCreatePointsSphereElement::ExecuteInternal(FPCGContext* InContext) cons
 	// Point scale should not be relative to the coordinate space
 	LocalTransform.SetScale3D(FVector::One());
 
-	OutputPoints.Reserve(OutLocalPointLocations.Num());
+	TArray<TTuple<FTransform, int32>> Points;
+	Points.Reserve(OutLocalPointLocations.Num());
 
 	for (int Index = 0; Index < OutLocalPointLocations.Num(); ++Index)
 	{
@@ -491,11 +491,24 @@ bool FPCGCreatePointsSphereElement::ExecuteInternal(FPCGContext* InContext) cons
 
 		if (!CullingShape || (CullingShape->GetDensityAtPosition(PointTransform.GetLocation()) > 0.f))
 		{
-			FPCGPoint& OutPoint = OutputPoints.Emplace_GetRef(PointTransform, 1.f, PCGHelpers::ComputeSeedFromPosition(OutLocalPointLocations[Index]));
-			OutPoint.Steepness = Settings->PointSteepness;
+			Points.Add({ PointTransform, PCGHelpers::ComputeSeedFromPosition(OutLocalPointLocations[Index]) });
 		}
 	}
 
+	PointData->SetNumPoints(Points.Num());
+	PointData->SetDensity(1.0f);
+	PointData->SetSteepness(Settings->PointSteepness);
+	PointData->AllocateProperties(EPCGPointNativeProperties::Transform | EPCGPointNativeProperties::Seed);
+
+	TPCGValueRange<FTransform> TransformRange = PointData->GetTransformValueRange(/*bAllocate=*/false);
+	TPCGValueRange<int32> SeedRange = PointData->GetSeedValueRange(/*bAllocate=*/false);
+
+	for (int32 Index = 0; Index < PointData->GetNumPoints(); ++Index)
+	{
+		TransformRange[Index] = Points[Index].Get<0>();
+		SeedRange[Index] = Points[Index].Get<1>();
+	}
+
 	return true;
 }
 
diff --git a/Engine/Plugins/PCG/Source/PCG/Private/Elements/PCGCreateSpline.cpp b/Engine/Plugins/PCG/Source/PCG/Private/Elements/PCGCreateSpline.cpp
--- a/Engine/Plugins/PCG/Source/PCG/Private/Elements/PCGCreateSpline.cpp
+++ b/Engine/Plugins/PCG/Source/PCG/Private/Elements/PCGCreateSpline.cpp
@@ -7,7 +7,7 @@
 #include "PCGContext.h"
 #include "PCGManagedResource.h"
 #include "PCGPin.h"
-#include "Data/PCGPointData.h"
+#include "Data/PCGBasePointData.h"
 #include "Data/PCGSplineData.h"
 #include "Helpers/PCGActorHelpers.h"
 #include "Helpers/PCGHelpers.h"
@@ -83,7 +83,7 @@ bool FPCGCreateSplineElement::ExecuteInternal(FPCGContext* Context) const
 			continue;
 		}
 
-		const UPCGPointData* PointData = SpatialData->ToPointData(Context);
+		const UPCGBasePointData* PointData = SpatialData->ToBasePointData(Context);
 
 		if (!PointData)
 		{
@@ -126,12 +126,12 @@ bool FPCGCreateSplineElement::ExecuteInternal(FPCGContext* Context) const
 		SplineData->InitializeFromData(PointData);
 		AActor* SplineActor = TargetActor;
 
-		const TArray<FPCGPoint>& Points = PointData->GetPoints();
+		const int32 NumPoints = PointData->GetNumPoints();
 		TArray<FSplinePoint> SplinePoints;
-		SplinePoints.Reserve(Points.Num());
+		SplinePoints.Reserve(NumPoints);
 
 		TArray<PCGMetadataEntryKey> SplineEntryKeys;
-		SplineEntryKeys.Reserve(Points.Num());
+		SplineEntryKeys.Reserve(NumPoints);
 		bool bHasAValidEntry = false;
 
 		const FTransform SplineActorTransform = SplineActor->GetTransform();
@@ -146,22 +146,24 @@ bool FPCGCreateSplineElement::ExecuteInternal(FPCGContext* Context) const
 			PointType = ESplinePointType::CurveCustomTangent;
 		}
 
-		for(int32 PointIndex = 0; PointIndex < Points.Num(); ++PointIndex)
+		FConstPCGPointValueRanges InRanges(PointData);
+
+		for(int32 PointIndex = 0; PointIndex < NumPoints; ++PointIndex)
 		{
-			const FPCGPoint& Point = Points[PointIndex];
-			const FTransform& PointTransform = Point.Transform;
+			const FTransform& PointTransform = InRanges.TransformRange[PointIndex];
 			const FVector LocalPosition = PointTransform.GetLocation() - SplineActorTransform.GetLocation();
 
+			int64 PointMetadataEntry = InRanges.MetadataEntryRange[PointIndex];
 			SplinePoints.Emplace(static_cast<float>(PointIndex),
 				LocalPosition,
-				ArriveTangentAttribute ? ArriveTangentAttribute->GetValueFromItemKey(Point.MetadataEntry) : FVector::ZeroVector,
-				LeaveTangentAttribute ? LeaveTangentAttribute->GetValueFromItemKey(Point.MetadataEntry) : FVector::ZeroVector,
+				ArriveTangentAttribute ? ArriveTangentAttribute->GetValueFromItemKey(PointMetadataEntry) : FVector::ZeroVector,
+				LeaveTangentAttribute ? LeaveTangentAttribute->GetValueFromItemKey(PointMetadataEntry) : FVector::ZeroVector,
 				PointTransform.GetRotation().Rotator(),
 				PointTransform.GetScale3D(),
 				PointType);
 
-			SplineEntryKeys.Emplace(Point.MetadataEntry);
-			bHasAValidEntry |= (Point.MetadataEntry != PCGInvalidEntryKey);
+			SplineEntryKeys.Emplace(PointMetadataEntry);
+			bHasAValidEntry |= (PointMetadataEntry != PCGInvalidEntryKey);
 		}
 
 		if (!bHasAValidEntry)
diff --git a/Engine/Plugins/PCG/Source/PCG/Private/Elements/PCGElevationIsolines.cpp b/Engine/Plugins/PCG/Source/PCG/Private/Elements/PCGElevationIsolines.cpp
--- a/Engine/Plugins/PCG/Source/PCG/Private/Elements/PCGElevationIsolines.cpp
+++ b/Engine/Plugins/PCG/Source/PCG/Private/Elements/PCGElevationIsolines.cpp
@@ -5,7 +5,7 @@
 #include "PCGComponent.h"
 #include "PCGContext.h"
 #include "Data/PCGIntersectionData.h"
-#include "Data/PCGPointData.h"
+#include "Data/PCGBasePointData.h"
 #include "Data/PCGSplineData.h"
 #include "Data/PCGSurfaceData.h"
 #include "Elements/PCGSurfaceSampler.h"
@@ -286,17 +286,19 @@ bool FPCGElevationIsolinesElement::ExecuteInternal(FPCGContext* InContext) const
 				}
 				else
 				{
-					UPCGPointData* OutPointData = FPCGContext::NewObject_AnyThread<UPCGPointData>(InContext);
-					TArray<FPCGPoint>& OutPoints = OutPointData->GetMutablePoints();
-					OutPoints.Reserve(Result.LinkedGridCoordinates.Num());
-					OutputData.Data = OutPointData;
-
-					for (const FTransform& Transform : FinalTransforms)
+					UPCGBasePointData* OutPointData = FPCGContext::NewPointData_AnyThread(InContext);
+					OutPointData->SetNumPoints(FinalTransforms.Num());
+					OutPointData->SetBoundsMin(CellBounds.Min);
+					OutPointData->SetBoundsMax(CellBounds.Max);
+					OutPointData->AllocateProperties(EPCGPointNativeProperties::Transform);
+					
+					TPCGValueRange<FTransform> TransformRange = OutPointData->GetTransformValueRange(/*bAllocate=*/false);
+					for (int32 Index = 0; Index < FinalTransforms.Num(); ++Index)
 					{
-						FPCGPoint& Point = OutPoints.Emplace_GetRef();
-						Point.Transform = Transform;
-						Point.SetLocalBounds(CellBounds);
+						TransformRange[Index] = FinalTransforms[Index];
 					}
+
+					OutputData.Data = OutPointData;
 				}
 			}
 
diff --git a/Engine/Plugins/PCG/Source/PCG/Private/Elements/PCGFilterElementsByIndex.cpp b/Engine/Plugins/PCG/Source/PCG/Private/Elements/PCGFilterElementsByIndex.cpp
--- a/Engine/Plugins/PCG/Source/PCG/Private/Elements/PCGFilterElementsByIndex.cpp
+++ b/Engine/Plugins/PCG/Source/PCG/Private/Elements/PCGFilterElementsByIndex.cpp
@@ -4,7 +4,7 @@
 
 #include "PCGContext.h"
 #include "PCGParamData.h"
-#include "Data/PCGPointData.h"
+#include "Data/PCGBasePointData.h"
 #include "Helpers/Parsing/PCGParsing.h"
 #include "Metadata/Accessors/PCGAttributeAccessorHelpers.h"
 
@@ -26,7 +26,7 @@ namespace PCGFilterElementsByIndex
 	{
 		int32 GetDataElementCount(const UPCGData* Data)
 		{
-			if (const UPCGPointData* PointData = Cast<UPCGPointData>(Data))
+			if (const UPCGBasePointData* PointData = Cast<UPCGBasePointData>(Data))
 			{
 				return PointData->GetNumPoints();
 			}
@@ -201,7 +201,7 @@ bool FPCGFilterElementsByIndexElement::ExecuteInternal(FPCGContext* InContext) c
 		for (const FPCGTaggedData& FilterInput : FilterInputs)
 		{
 			const UPCGData* FilterInputData = FilterInput.Data;
-			const UPCGPointData* FilterPointData = Cast<UPCGPointData>(FilterInputData);
+			const UPCGBasePointData* FilterPointData = Cast<UPCGBasePointData>(FilterInputData);
 			const UPCGParamData* FilterParamData = Cast<UPCGParamData>(FilterInputData);
 			if (!(FilterPointData || FilterParamData))
 			{
@@ -294,43 +294,56 @@ bool FPCGFilterElementsByIndexElement::ExecuteInternal(FPCGContext* InContext) c
 			});
 		};
 
-		if (const UPCGPointData* InputPointData = Cast<UPCGPointData>(PrimaryInputData))
+		if (const UPCGBasePointData* InputPointData = Cast<UPCGBasePointData>(PrimaryInputData))
 		{
-			TArray<FPCGPoint, TInlineAllocator<256>> InFilterPoints;
-			TArray<FPCGPoint, TInlineAllocator<256>> OutFilterPoints;
-			InFilterPoints.Reserve(InputPointData->GetNumPoints());
+			UPCGBasePointData* InFilterPointData = FPCGContext::NewPointData_AnyThread(InContext);
+
+			FPCGInitializeFromDataParams InitializeFromDataParams(InputPointData);
+			InitializeFromDataParams.bInheritSpatialData = false;
+
+			InFilterPointData->InitializeFromDataWithParams(InitializeFromDataParams);
+
+			const int32 NumInFilterPoints = Settings->bInvertFilter ? InputPointData->GetNumPoints() - AllFilteredIndices.GetTotalIndexCount() : AllFilteredIndices.GetTotalIndexCount();
+			InFilterPointData->SetNumPoints(NumInFilterPoints);
+			InFilterPointData->AllocateProperties(InputPointData->GetAllocatedProperties());
+
+			AddNewOutput(InFilterPointData, Constants::InFilterOutputPinLabel);
+
+			UPCGBasePointData* OutFilterPointData = nullptr;
 			if (Settings->bOutputDiscardedElements)
 			{
-				OutFilterPoints.Reserve(InputPointData->GetNumPoints());
+				OutFilterPointData = FPCGContext::NewPointData_AnyThread(InContext);
+				
+				OutFilterPointData->InitializeFromDataWithParams(InitializeFromDataParams);
+				
+				const int32 NumOutFilterPoints = InputPointData->GetNumPoints() - NumInFilterPoints;
+				OutFilterPointData->SetNumPoints(NumOutFilterPoints);
+				OutFilterPointData->AllocateProperties(InputPointData->GetAllocatedProperties());
+
+				AddNewOutput(OutFilterPointData, Constants::OutFilterOutputPinLabel);
 			}
 
-			const TArray<FPCGPoint>& InputPoints = InputPointData->GetPoints();
-			for (int PointIndex = 0; PointIndex < InputPoints.Num(); ++PointIndex)
+			const FConstPCGPointValueRanges InRanges(InputPointData);
+			FPCGPointValueRanges InFilterRanges(InFilterPointData, /*bAllocate=*/false);
+			FPCGPointValueRanges OutFilterRanges = OutFilterPointData ? FPCGPointValueRanges(OutFilterPointData, /*bAllocate=*/false) : FPCGPointValueRanges();
+			
+			int32 InFilterWritten = 0;
+			int32 OutFilterWritten = 0;
+			for (int PointIndex = 0; PointIndex < InputPointData->GetNumPoints(); ++PointIndex)
 			{
 				bool bIndexIsFiltered = AllFilteredIndices.ContainsIndex(PointIndex);
 				bIndexIsFiltered = Settings->bInvertFilter ? !bIndexIsFiltered : bIndexIsFiltered;
 				if (bIndexIsFiltered)
 				{
-					InFilterPoints.Emplace(InputPoints[PointIndex]);
+					InFilterRanges.SetFromValueRanges(InFilterWritten, InRanges, PointIndex);
+					++InFilterWritten;
 				}
 				else if (Settings->bOutputDiscardedElements)
 				{
-					OutFilterPoints.Emplace(InputPoints[PointIndex]);
+					OutFilterRanges.SetFromValueRanges(OutFilterWritten, InRanges, PointIndex);
+					++OutFilterWritten;
 				}
 			}
-
-			UPCGPointData* InFilterPointData = FPCGContext::NewObject_AnyThread<UPCGPointData>(InContext);
-			InFilterPointData->InitializeFromData(InputPointData);
-			InFilterPointData->GetMutablePoints() = std::move(InFilterPoints);
-			AddNewOutput(InFilterPointData, Constants::InFilterOutputPinLabel);
-
-			if (Settings->bOutputDiscardedElements)
-			{
-				UPCGPointData* OutFilterPointData = FPCGContext::NewObject_AnyThread<UPCGPointData>(InContext);
-				OutFilterPointData->InitializeFromData(InputPointData);
-				OutFilterPointData->GetMutablePoints() = std::move(OutFilterPoints);
-				AddNewOutput(OutFilterPointData, Constants::OutFilterOutputPinLabel);
-			}
 		}
 		else if (const UPCGParamData* InputParamData = Cast<UPCGParamData>(PrimaryInputData))
 		{
diff --git a/Engine/Plugins/PCG/Source/PCG/Private/Elements/PCGGetActorDataLayers.cpp b/Engine/Plugins/PCG/Source/PCG/Private/Elements/PCGGetActorDataLayers.cpp
--- a/Engine/Plugins/PCG/Source/PCG/Private/Elements/PCGGetActorDataLayers.cpp
+++ b/Engine/Plugins/PCG/Source/PCG/Private/Elements/PCGGetActorDataLayers.cpp
@@ -6,7 +6,7 @@
 #include "PCGData.h"
 #include "PCGParamData.h"
 #include "PCGPin.h"
-#include "Data/PCGPointData.h"
+#include "Data/PCGBasePointData.h"
 #include "Helpers/PCGDataLayerHelpers.h"
 #include "Utils/PCGLogErrors.h"
 
diff --git a/Engine/Plugins/PCG/Source/PCG/Private/Elements/PCGMakeConcreteElement.cpp b/Engine/Plugins/PCG/Source/PCG/Private/Elements/PCGMakeConcreteElement.cpp
--- a/Engine/Plugins/PCG/Source/PCG/Private/Elements/PCGMakeConcreteElement.cpp
+++ b/Engine/Plugins/PCG/Source/PCG/Private/Elements/PCGMakeConcreteElement.cpp
@@ -3,7 +3,7 @@
 #include "Elements/PCGMakeConcreteElement.h"
 
 #include "PCGContext.h"
-#include "Data/PCGPointData.h"
+#include "Data/PCGBasePointData.h"
 #include "Data/PCGSpatialData.h"
 
 #define LOCTEXT_NAMESPACE "PCGMakeConcreteElement"
@@ -58,7 +58,7 @@ bool FPCGMakeConcreteElement::ExecuteInternal(FPCGContext* InContext) const
 		Output.Pin = PCGPinConstants::DefaultOutputLabel;
 		if (!bDataIsConcrete)
 		{
-			Output.Data = InputSpatialData->ToPointData(InContext);
+			Output.Data = InputSpatialData->ToBasePointData(InContext);
 		}
 	}
 
diff --git a/Engine/Plugins/PCG/Source/PCG/Private/Elements/PCGMatchAndSetAttributes.cpp b/Engine/Plugins/PCG/Source/PCG/Private/Elements/PCGMatchAndSetAttributes.cpp
--- a/Engine/Plugins/PCG/Source/PCG/Private/Elements/PCGMatchAndSetAttributes.cpp
+++ b/Engine/Plugins/PCG/Source/PCG/Private/Elements/PCGMatchAndSetAttributes.cpp
@@ -1097,7 +1097,7 @@ bool FPCGMatchAndSetAttributesElement::PrepareDataInternal(FPCGContext* InContex
 			InitializeFromDataParams.bInheritSpatialData = false;
 			OutPointData->InitializeFromDataWithParams(InitializeFromDataParams);
 			OutPointData->SetNumPoints(InPointData->GetNumPoints());
-			OutPointData->AllocateProperties(InPointData->GetAllocatedProperties());
+			OutPointData->AllocateProperties(InPointData->GetAllocatedProperties() | EPCGPointNativeProperties::MetadataEntry);
 
 			OutState.OutData = OutPointData;
 		}
diff --git a/Engine/Plugins/PCG/Source/PCG/Private/Elements/PCGPartitionByActorDataLayers.cpp b/Engine/Plugins/PCG/Source/PCG/Private/Elements/PCGPartitionByActorDataLayers.cpp
--- a/Engine/Plugins/PCG/Source/PCG/Private/Elements/PCGPartitionByActorDataLayers.cpp
+++ b/Engine/Plugins/PCG/Source/PCG/Private/Elements/PCGPartitionByActorDataLayers.cpp
@@ -5,7 +5,7 @@
 #include "PCGContext.h"
 #include "PCGParamData.h"
 #include "PCGPin.h"
-#include "Data/PCGPointData.h"
+#include "Data/PCGBasePointData.h"
 #include "Helpers/PCGDataLayerHelpers.h"
 #include "Utils/PCGLogErrors.h"
 
@@ -133,12 +133,12 @@ bool FPCGPartitionByActorDataLayersElement::ExecuteInternal(FPCGContext* Context
 	TArray<TSoftObjectPtr<UDataLayerAsset>> IncludedDataLayerAssets = PCGDataLayerHelpers::GetDataLayerAssetsFromInput(Context, PCGDataLayerHelpers::Constants::IncludedDataLayersAttribute, Settings->IncludedDataLayers);
 	TArray<TSoftObjectPtr<UDataLayerAsset>> ExcludedDataLayerAssets = PCGDataLayerHelpers::GetDataLayerAssetsFromInput(Context, PCGDataLayerHelpers::Constants::ExcludedDataLayersAttribute, Settings->ExcludedDataLayers);
 
-	TMap<int32, UPCGPointData*> DataLayerCrcToOutputData;
+	TMap<int32, UPCGBasePointData*> DataLayerCrcToOutputData;
 
 	for (const FPCGTaggedData& Input : Inputs)
 	{
 		// @todo_pcg: support ParamData
-		const UPCGPointData* InData = Cast<UPCGPointData>(Input.Data);
+		const UPCGBasePointData* InData = Cast<UPCGBasePointData>(Input.Data);
 		if (!InData)
 		{
 			continue;
@@ -151,7 +151,7 @@ bool FPCGPartitionByActorDataLayersElement::ExecuteInternal(FPCGContext* Context
 
 		for (UPCGData* PartitionData : PartitionDataArray)
 		{
-			UPCGPointData* PartitionPointData = CastChecked<UPCGPointData>(PartitionData);
+			UPCGBasePointData* PartitionPointData = CastChecked<UPCGBasePointData>(PartitionData);
 
 			const TUniquePtr<const IPCGAttributeAccessor> Accessor = PCGAttributeAccessorHelpers::CreateConstAccessor(PartitionData, PartitionAttributeSource);
 			const TUniquePtr<const IPCGAttributeAccessorKeys> Keys = PCGAttributeAccessorHelpers::CreateConstKeys(PartitionData, PartitionAttributeSource);
@@ -185,15 +185,19 @@ bool FPCGPartitionByActorDataLayersElement::ExecuteInternal(FPCGContext* Context
 			TArray<const UDataLayerAsset*> ActorDataLayers = PCGPartitionByActorDataLayers::GetDataLayersFromActor(Actor, IncludedDataLayerAssets, ExcludedDataLayerAssets);
 			int32 DataLayerCrc = PCGPartitionByActorDataLayers::GetDataLayersCrc(ActorDataLayers);
 
-			UPCGPointData* PartitionOutputData = nullptr;
-			if (UPCGPointData** FoundOutputData = DataLayerCrcToOutputData.Find(DataLayerCrc))
+			UPCGBasePointData* PartitionOutputData = nullptr;
+			if (UPCGBasePointData** FoundOutputData = DataLayerCrcToOutputData.Find(DataLayerCrc))
 			{
 				PartitionOutputData = *FoundOutputData;
 			}
 			else
 			{
-				PartitionOutputData = FPCGContext::NewObject_AnyThread<UPCGPointData>(Context);
-				PartitionOutputData->InitializeFromData(PartitionPointData);
+				PartitionOutputData = FPCGContext::NewPointData_AnyThread(Context);
+
+				FPCGInitializeFromDataParams InitializeFromDataParams(PartitionPointData);
+				InitializeFromDataParams.bInheritSpatialData = false;
+				PartitionOutputData->InitializeFromDataWithParams(InitializeFromDataParams);
+
 				DataLayerCrcToOutputData.Add(DataLayerCrc, PartitionOutputData);
 
 				FPCGTaggedData& Output = Context->OutputData.TaggedData.Emplace_GetRef();
@@ -215,8 +219,10 @@ bool FPCGPartitionByActorDataLayersElement::ExecuteInternal(FPCGContext* Context
 			}
 
 			check(PartitionOutputData);
-			TArray<FPCGPoint>& MutablePoints = PartitionOutputData->GetMutablePoints();
-			MutablePoints.Append(PartitionPointData->GetPoints());
+			const int32 NumPoints = PartitionOutputData->GetNumPoints();
+			PartitionOutputData->SetNumPoints(NumPoints + PartitionPointData->GetNumPoints());
+			PartitionOutputData->AllocateProperties(PartitionPointData->GetAllocatedProperties());
+			PartitionPointData->CopyPointsTo(PartitionOutputData, 0, NumPoints, PartitionPointData->GetNumPoints());
 		}
 	}
 #else
diff --git a/Engine/Plugins/PCG/Source/PCG/Private/Elements/PCGPointNeighborhood.cpp b/Engine/Plugins/PCG/Source/PCG/Private/Elements/PCGPointNeighborhood.cpp
--- a/Engine/Plugins/PCG/Source/PCG/Private/Elements/PCGPointNeighborhood.cpp
+++ b/Engine/Plugins/PCG/Source/PCG/Private/Elements/PCGPointNeighborhood.cpp
@@ -3,7 +3,7 @@
 #include "Elements/PCGPointNeighborhood.h"
 
 #include "PCGContext.h"
-#include "Data/PCGPointData.h"
+#include "Data/PCGBasePointData.h"
 #include "Helpers/PCGAsync.h"
 #include "Metadata/Accessors/PCGAttributeAccessorHelpers.h"
 
@@ -12,7 +12,7 @@
 namespace PCGPointNeighborhood
 {
 	template<typename T>
-	void SetAttributeHelper(UPCGPointData* Data, const FName& AttributeName, const TArrayView<const T> Values)
+	void SetAttributeHelper(UPCGBasePointData* Data, const FName& AttributeName, const TArrayView<const T> Values)
 	{
 		// Discard None attributes
 		if (AttributeName == NAME_None)
@@ -70,19 +70,47 @@ bool FPCGPointNeighborhoodElement::ExecuteInternal(FPCGContext* Context) const
 
 	for (int i = 0; i < Inputs.Num(); ++i)
 	{
-		const UPCGPointData* InputPointData = Cast<UPCGPointData>(Inputs[i].Data);
+		const UPCGBasePointData* InputPointData = Cast<UPCGBasePointData>(Inputs[i].Data);
 		if (!InputPointData)
 		{
 			PCGE_LOG(Verbose, GraphAndLog, FText::Format(LOCTEXT("InvalidPointData", "Input {0} is not point data"), i));
 			continue;
 		}
 
-		const TArray<FPCGPoint>& SrcPoints = InputPointData->GetPoints();
 		FPCGTaggedData& Output = Outputs.Add_GetRef(Inputs[i]);
-		UPCGPointData* OutputPointData = FPCGContext::NewObject_AnyThread<UPCGPointData>(Context);
+		UPCGBasePointData* OutputPointData = FPCGContext::NewPointData_AnyThread(Context);
 		OutputPointData->InitializeFromData(InputPointData);
-		TArray<FPCGPoint>& OutputPoints = OutputPointData->GetMutablePoints();
-		OutputPoints.SetNumUninitialized(SrcPoints.Num());
+
+		EPCGPointNativeProperties PropertiesToAllocate = EPCGPointNativeProperties::None;
+
+		if (Settings->SetDensity != EPCGPointNeighborhoodDensityMode::None)
+		{
+			PropertiesToAllocate |= EPCGPointNativeProperties::Density;
+		}
+
+		if (Settings->bSetAveragePosition)
+		{
+			PropertiesToAllocate |= EPCGPointNativeProperties::Transform;
+		}
+
+		if (Settings->bSetAverageColor)
+		{
+			PropertiesToAllocate |= EPCGPointNativeProperties::Color;
+		}
+
+		if (Settings->bSetDistanceToAttribute || Settings->bSetAveragePositionToAttribute)
+		{
+			PropertiesToAllocate |= EPCGPointNativeProperties::MetadataEntry;
+		}
+
+		OutputPointData->SetNumPoints(InputPointData->GetNumPoints(), /*bInitializeValues=*/false);
+		if (!OutputPointData->HasSpatialDataParent())
+		{
+			OutputPointData->AllocateProperties(InputPointData->GetAllocatedProperties());
+		}
+
+		OutputPointData->AllocateProperties(PropertiesToAllocate);
+
 		Output.Data = OutputPointData;
 		
 		const FBox SearchBounds(FVector(SearchDistance * -1.0), FVector(SearchDistance));
@@ -91,7 +119,7 @@ bool FPCGPointNeighborhoodElement::ExecuteInternal(FPCGContext* Context) const
 		Params.bProjectRotations = Params.bProjectScales = false;
 		Params.ColorBlendMode = Settings->bSetAverageColor ? EPCGProjectionColorBlendMode::TargetValue : EPCGProjectionColorBlendMode::SourceValue;
 
-		auto InitializeBuffers = [Settings, &OutputDataBuffers, Count = SrcPoints.Num()]()
+		auto InitializeBuffers = [Settings, &OutputDataBuffers, Count = InputPointData->GetNumPoints()]()
 		{
 			if (Settings->bSetDistanceToAttribute)
 			{
@@ -104,53 +132,71 @@ bool FPCGPointNeighborhoodElement::ExecuteInternal(FPCGContext* Context) const
 			}
 		};
 
-		auto ProcessPoint = [Settings, &OutputPoints, &OutputDataBuffers, &SrcPoints, &SearchBounds, &SearchDistance, &Params, &InputPointData, &OutputPointData](const int32 ReadIndex, const int32 WriteIndex)
+		auto ProcessPoint = [Settings, &OutputDataBuffers, &SearchBounds, &SearchDistance, &Params, InputPointData, OutputPointData](const int32 StartReadIndex, const int32 StartWriteIndex, const int32 Count)
 		{
+			const FConstPCGPointValueRanges InRanges(InputPointData);
+			TPCGValueRange<float> DensityRange = Settings->SetDensity != EPCGPointNeighborhoodDensityMode::None ? OutputPointData->GetDensityValueRange(/*bAllocate=*/false) : TPCGValueRange<float>();
+			TPCGValueRange<FTransform> TransformRange = Settings->bSetAveragePosition ? OutputPointData->GetTransformValueRange(/*bAllocate=*/false) : TPCGValueRange<FTransform>();
+			TPCGValueRange<FVector4> ColorRange = Settings->bSetAverageColor ? OutputPointData->GetColorValueRange(/*bAllocate=*/false) : TPCGValueRange<FVector4>();
 
-			const FPCGPoint& InPoint = SrcPoints[ReadIndex];
-			FPCGPoint& OutPoint = OutputPoints[WriteIndex];
-			OutPoint = InPoint;
-
-			const FTransform InTransform(InPoint.Transform.GetLocation());
-
-			FPCGPoint ProjectionPoint = FPCGPoint();
-			// Do not pass metadata as output as we don't care about attributes (Projected point will be discarded).
-			// @todo_pcg: Maybe support the weighting average of attributes as an option.
-			InputPointData->ProjectPoint(InTransform, SearchBounds, Params, ProjectionPoint, /*OutMetadata=*/nullptr, Settings->bWeightedAverage);
-
-			const double NormalizedDistance = FVector::Distance(InPoint.Transform.GetLocation(), ProjectionPoint.Transform.GetLocation()) / SearchDistance;
-
-			if (Settings->SetDensity == EPCGPointNeighborhoodDensityMode::SetNormalizedDistanceToDensity)
+			if (!OutputPointData->HasSpatialDataParent())
 			{
-				OutPoint.Density = FMath::Clamp(NormalizedDistance, 0.0, 1.0);
-			}
-			else if (Settings->SetDensity == EPCGPointNeighborhoodDensityMode::SetAverageDensity)
-			{
-				OutPoint.Density = ProjectionPoint.Density;
+				InputPointData->CopyPointsTo(OutputPointData, StartReadIndex, StartWriteIndex, Count);
 			}
 
-			if (Settings->bSetDistanceToAttribute)
+			int32 NumWritten = 0;
+			for (int32 ReadIndex = StartReadIndex; ReadIndex < StartReadIndex + Count; ++ReadIndex)
 			{
-				OutputDataBuffers.Distances[WriteIndex] = FVector::Distance(InPoint.Transform.GetLocation(), ProjectionPoint.Transform.GetLocation());
+				const int32 WriteIndex = StartWriteIndex + NumWritten;
+
+				FVector InLocation = InRanges.TransformRange[ReadIndex].GetLocation();
+				const FTransform InTransform(InLocation);
+
+				FPCGPoint ProjectionPoint = FPCGPoint();
+				// Do not pass metadata as output as we don't care about attributes (Projected point will be discarded).
+				// @todo_pcg: Maybe support the weighting average of attributes as an option.
+				// @todo_pcg: Might be better to use ProjectPoints using a range here?
+				InputPointData->ProjectPoint(InTransform, SearchBounds, Params, ProjectionPoint, /*OutMetadata=*/nullptr, Settings->bWeightedAverage);
+
+				const double NormalizedDistance = FVector::Distance(InLocation, ProjectionPoint.Transform.GetLocation()) / SearchDistance;
+
+				if (Settings->SetDensity == EPCGPointNeighborhoodDensityMode::SetNormalizedDistanceToDensity)
+				{
+					DensityRange[WriteIndex] = FMath::Clamp(NormalizedDistance, 0.0, 1.0);
+				}
+				else if (Settings->SetDensity == EPCGPointNeighborhoodDensityMode::SetAverageDensity)
+				{
+					DensityRange[WriteIndex] = ProjectionPoint.Density;
+				}
+
+				if (Settings->bSetDistanceToAttribute)
+				{
+					OutputDataBuffers.Distances[WriteIndex] = FVector::Distance(InLocation, ProjectionPoint.Transform.GetLocation());
+				}
+
+				if (Settings->bSetAveragePosition)
+				{
+					TransformRange[WriteIndex].SetLocation(ProjectionPoint.Transform.GetLocation());
+				}
+
+				if (Settings->bSetAveragePositionToAttribute)
+				{
+					OutputDataBuffers.AveragePositions[WriteIndex] = ProjectionPoint.Transform.GetLocation();
+				}
+
+				if (Settings->bSetAverageColor)
+				{
+					ColorRange[WriteIndex] = ProjectionPoint.Color;
+				}
+
+				++NumWritten;
 			}
 
-			if (Settings->bSetAveragePosition)
-			{
-				OutPoint.Transform.SetLocation(ProjectionPoint.Transform.GetLocation());
-			}
-
-			if (Settings->bSetAveragePositionToAttribute)
-			{
-				OutputDataBuffers.AveragePositions[WriteIndex] = ProjectionPoint.Transform.GetLocation();
-			}
-
-			if (Settings->bSetAverageColor)
-			{
-				OutPoint.Color = ProjectionPoint.Color;
-			}
+			check(NumWritten == Count);
+			return Count;
 		};
 
-		FPCGAsync::AsyncProcessingOneToOneEx(&Context->AsyncState, SrcPoints.Num(), InitializeBuffers, ProcessPoint, /* bEnableTimeSlicing=*/false);
+		FPCGAsync::AsyncProcessingOneToOneRangeEx(&Context->AsyncState, InputPointData->GetNumPoints(), InitializeBuffers, ProcessPoint, /* bEnableTimeSlicing=*/false);
 
 		if (Settings->bSetDistanceToAttribute)
 		{
diff --git a/Engine/Plugins/PCG/Source/PCG/Private/Elements/PCGRandomChoice.cpp b/Engine/Plugins/PCG/Source/PCG/Private/Elements/PCGRandomChoice.cpp
--- a/Engine/Plugins/PCG/Source/PCG/Private/Elements/PCGRandomChoice.cpp
+++ b/Engine/Plugins/PCG/Source/PCG/Private/Elements/PCGRandomChoice.cpp
@@ -5,7 +5,7 @@
 #include "PCGContext.h"
 #include "PCGParamData.h"
 #include "PCGPin.h"
-#include "Data/PCGPointData.h"
+#include "Data/PCGBasePointData.h"
 #include "Elements/Metadata/PCGMetadataElementCommon.h"
 #include "Helpers/PCGHelpers.h"
 #include "Metadata/Accessors/PCGAttributeAccessorHelpers.h"
@@ -17,20 +17,27 @@ namespace PCGRandomChoice
 	UPCGData* ChoosePointData(const UPCGData* InData, TArrayView<int32> InIndexes, FPCGContext* InContext)
 	{
 		// We know it is a point data
-		const UPCGPointData* InPointData = CastChecked<const UPCGPointData>(InData);
-		UPCGPointData* OutPointData = FPCGContext::NewObject_AnyThread<UPCGPointData>(InContext);
-		OutPointData->InitializeFromData(InPointData);
+		const UPCGBasePointData* InPointData = CastChecked<const UPCGBasePointData>(InData);
+		UPCGBasePointData* OutPointData = FPCGContext::NewPointData_AnyThread(InContext);
 
-		const TArray<FPCGPoint>& InPoints = InPointData->GetPoints();
-		TArray<FPCGPoint>& OutPoints = OutPointData->GetMutablePoints();
-		OutPoints.Reserve(InIndexes.Num());
+		FPCGInitializeFromDataParams InitializeFromDataParams(InPointData);
+		InitializeFromDataParams.bInheritSpatialData = false;
+		OutPointData->InitializeFromDataWithParams(InitializeFromDataParams);
+
+		OutPointData->SetNumPoints(InIndexes.Num());
+		OutPointData->AllocateProperties(InPointData->GetAllocatedProperties());
 
 		// Order needs to be stable to sort this part of the array
 		Algo::Sort(InIndexes);
 
-		for (int32 Index : InIndexes)
+		int32 WriteIndex = 0;
+		const FConstPCGPointValueRanges InRanges(InPointData);
+		FPCGPointValueRanges OutRanges(OutPointData, /*bAllocate=*/false);
+
+		for (int32 ReadIndex : InIndexes)
 		{
-			OutPoints.Add(InPoints[Index]);
+			OutRanges.SetFromValueRanges(WriteIndex, InRanges, ReadIndex);
+			++WriteIndex;
 		}
 
 		return OutPointData;
@@ -120,16 +127,16 @@ bool FPCGRandomChoiceElement::ExecuteInternal(FPCGContext* Context) const
 
 		UPCGData* (*ChooseFunc)(const UPCGData*, TArrayView<int32>, FPCGContext*) = nullptr;
 
-		if (const UPCGPointData* InputPointData = Cast<UPCGPointData>(CurrentInput.Data))
+		if (const UPCGBasePointData* InputPointData = Cast<UPCGBasePointData>(CurrentInput.Data))
 		{
-			const TArray<FPCGPoint>& InPoints = InputPointData->GetPoints();
-			NumElements = InPoints.Num();
+			const TConstPCGValueRange<int32> SeedRange = InputPointData->GetConstSeedValueRange();
+			NumElements = InputPointData->GetNumPoints();
 			ChooseFunc = &PCGRandomChoice::ChoosePointData;
 
 			// By default, combine the seed with the first point's seed so that multiple data produces different results.
-			if (!Settings->bHasCustomSeedSource && !InPoints.IsEmpty())
+			if (!Settings->bHasCustomSeedSource && NumElements > 0)
 			{
-				Seed = PCGHelpers::ComputeSeed(Seed, InPoints[0].Seed);
+				Seed = PCGHelpers::ComputeSeed(Seed, SeedRange[0]);
 			}
 		}
 		else if (const UPCGParamData* InputParamData = Cast<UPCGParamData>(CurrentInput.Data))
diff --git a/Engine/Plugins/PCG/Source/PCG/Private/Elements/PCGSampleTexture.cpp b/Engine/Plugins/PCG/Source/PCG/Private/Elements/PCGSampleTexture.cpp
--- a/Engine/Plugins/PCG/Source/PCG/Private/Elements/PCGSampleTexture.cpp
+++ b/Engine/Plugins/PCG/Source/PCG/Private/Elements/PCGSampleTexture.cpp
@@ -4,7 +4,7 @@
 
 #include "PCGContext.h"
 #include "PCGPin.h"
-#include "Data/PCGPointData.h"
+#include "Data/PCGBasePointData.h"
 #include "Data/PCGTextureData.h"
 #include "Helpers/PCGAsync.h"
 #include "Helpers/PCGHelpers.h"
@@ -62,22 +62,20 @@ bool FPCGSampleTextureElement::ExecuteInternal(FPCGContext* Context) const
 
 	for (int i = 0; i < PointInputs.Num(); ++i)
 	{
-		const UPCGPointData* PointData = Cast<UPCGPointData>(PointInputs[i].Data);
+		const UPCGBasePointData* PointData = Cast<UPCGBasePointData>(PointInputs[i].Data);
 		if (!PointData)
 		{
 			PCGE_LOG(Warning, GraphAndLog, FText::Format(LOCTEXT("InvalidPointData", "Point Input {0} is not point data."), i));
 			continue;
 		}
 
-		const TArray<FPCGPoint>& InputPoints = PointData->GetPoints();
-
 		TUniquePtr<const IPCGAttributeAccessor> InputAccessor = nullptr;
 		TUniquePtr<const IPCGAttributeAccessorKeys> InputKeys = nullptr;
 
 		if (Settings->TextureMappingMethod == EPCGTextureMappingMethod::UVCoordinates)
 		{
 			const FPCGAttributePropertyInputSelector UVSource = Settings->UVCoordinatesAttribute.CopyAndFixLast(PointData);
-	
+
 			InputAccessor = PCGAttributeAccessorHelpers::CreateConstAccessor(PointData, UVSource);
 			InputKeys = PCGAttributeAccessorHelpers::CreateConstKeys(PointData, UVSource);
 			if (!InputAccessor || !InputKeys)
@@ -93,72 +91,122 @@ bool FPCGSampleTextureElement::ExecuteInternal(FPCGContext* Context) const
 			}
 		}
 
-		UPCGPointData* OutPointData = FPCGContext::NewObject_AnyThread<UPCGPointData>(Context);
-		OutPointData->InitializeFromData(PointData);
+		UPCGBasePointData* OutPointData = FPCGContext::NewPointData_AnyThread(Context);
+
+		FPCGInitializeFromDataParams InitializeFromDataParams(PointData);
+		InitializeFromDataParams.bInheritSpatialData = false;
+		OutPointData->InitializeFromDataWithParams(InitializeFromDataParams);
 
-		auto ProcessPoint = [Settings, BaseTextureData, &InputAccessor, &InputKeys, &InputPoints, &DensityMergeFunc, OutPointData](int32 Index, FPCGPoint& OutPoint)
+		auto InitializeFunc = [OutPointData, PointData, Settings]()
 		{
-			OutPoint = InputPoints[Index];
+			OutPointData->SetNumPoints(PointData->GetNumPoints());
 
-			if (Settings->TextureMappingMethod == EPCGTextureMappingMethod::UVCoordinates)
+			EPCGPointNativeProperties PropertiesToAllocate = PointData->GetAllocatedProperties();
+
+			// This is based of the UPCGBaseTextureData::SamplePoint code, we should probably add a virtual to UPCGSpatialData to return the sampled properties to better futurproof this code
+			PropertiesToAllocate |= (EPCGPointNativeProperties::Color | EPCGPointNativeProperties::Density);
+			if (Settings->TextureMappingMethod != EPCGTextureMappingMethod::UVCoordinates)
 			{
-				FVector OutSamplePosition;
-				InputAccessor->Get(OutSamplePosition, Index, *InputKeys, EPCGAttributeAccessorFlags::AllowBroadcast | EPCGAttributeAccessorFlags::AllowConstructible);
+				PropertiesToAllocate |= EPCGPointNativeProperties::Transform;
+			}
 
-				if (Settings->TilingMode == EPCGTextureAddressMode::Clamp)
-				{
-					OutSamplePosition.X = FMath::Clamp(OutSamplePosition.X, 0.0, 1.0);
-					OutSamplePosition.Y = FMath::Clamp(OutSamplePosition.Y, 0.0, 1.0);
-				}
+			OutPointData->AllocateProperties(PropertiesToAllocate);
+		};
 
-				float SampleDensity = 1.0f;
-				if (BaseTextureData->SamplePointLocal(FVector2D(OutSamplePosition), OutPoint.Color, SampleDensity))
-				{
-					float ComputedDensity = DensityMergeFunc(OutPoint.Density, SampleDensity);
+		auto MoveDataRangeFunc = [OutPointData](int32 RangeStartIndex, int32 MoveToIndex, int32 NumElements)
+		{
+			OutPointData->MoveRange(RangeStartIndex, MoveToIndex, NumElements);
+		};
 
-					if (Settings->bClampOutputDensity)
-					{
-						ComputedDensity = FMath::Clamp(ComputedDensity, 0.0f, 1.0f);
-					}
+		auto FinishedFunc = [OutPointData](int32 NumPoints)
+		{
+			OutPointData->SetNumPoints(NumPoints);
+		};
 
-					OutPoint.Density = ComputedDensity;
-					return true;
-				}
-				else
-				{
-					return false;
-				}
+		auto ProcessRangeFunc = [Settings, BaseTextureData, &InputAccessor, &InputKeys, PointData, &DensityMergeFunc, OutPointData](int32 StartReadIndex, int32 StartWriteIndex, int32 Count)
+		{
+			int32 NumWritten = 0;
+
+			const FConstPCGPointValueRanges InRanges(PointData);
+			FPCGPointValueRanges OutRanges(OutPointData, /*bAllocate=*/false);
+
+			TArray<FVector> OutSampledPositions;
+			OutSampledPositions.SetNumUninitialized(Count);
+			if (Settings->TextureMappingMethod == EPCGTextureMappingMethod::UVCoordinates)
+			{
+				InputAccessor->GetRange<FVector>(OutSampledPositions, StartReadIndex, *InputKeys, EPCGAttributeAccessorFlags::AllowBroadcast | EPCGAttributeAccessorFlags::AllowConstructible);
 			}
-			else
+
+			for (int32 ReadIndex = StartReadIndex; ReadIndex < StartReadIndex + Count; ++ReadIndex)
 			{
-				if (BaseTextureData->SamplePoint(OutPoint.Transform, OutPoint.GetLocalBounds(), OutPoint, OutPointData->Metadata))
-				{
-					float ComputedDensity = DensityMergeFunc(InputPoints[Index].Density, OutPoint.Density);
+				const int32 WriteIndex = StartWriteIndex + NumWritten;
 
-					if (Settings->bClampOutputDensity)
+				if (Settings->TextureMappingMethod == EPCGTextureMappingMethod::UVCoordinates)
+				{
+					FVector& OutSamplePosition = OutSampledPositions[ReadIndex - StartReadIndex];
+					
+					if (Settings->TilingMode == EPCGTextureAddressMode::Clamp)
 					{
-						ComputedDensity = FMath::Clamp(ComputedDensity, 0.0f, 1.0f);
+						OutSamplePosition.X = FMath::Clamp(OutSamplePosition.X, 0.0, 1.0);
+						OutSamplePosition.Y = FMath::Clamp(OutSamplePosition.Y, 0.0, 1.0);
 					}
 
-					OutPoint.Density = ComputedDensity;
-					return true;
+					float SampleDensity = 1.0f;
+					FVector4 OutColor;
+					if (BaseTextureData->SamplePointLocal(FVector2D(OutSamplePosition), OutColor, SampleDensity))
+					{
+						OutRanges.SetFromValueRanges(WriteIndex, InRanges, ReadIndex);
+						OutRanges.ColorRange[WriteIndex] = OutColor;
+
+						float ComputedDensity = DensityMergeFunc(OutRanges.DensityRange[WriteIndex], SampleDensity);
+
+						if (Settings->bClampOutputDensity)
+						{
+							ComputedDensity = FMath::Clamp(ComputedDensity, 0.0f, 1.0f);
+						}
+
+						OutRanges.DensityRange[WriteIndex] = ComputedDensity;
+						++NumWritten;
+					}
 				}
 				else
 				{
-					return false;
+					FPCGPoint OutPoint = InRanges.GetPoint(ReadIndex);
+					if (BaseTextureData->SamplePoint(OutPoint.Transform, OutPoint.GetLocalBounds(), OutPoint, OutPointData->Metadata))
+					{
+						OutRanges.SetFromPoint(WriteIndex, OutPoint);
+
+						float ComputedDensity = DensityMergeFunc(InRanges.DensityRange[ReadIndex], OutRanges.DensityRange[WriteIndex]);
+
+						if (Settings->bClampOutputDensity)
+						{
+							ComputedDensity = FMath::Clamp(ComputedDensity, 0.0f, 1.0f);
+						}
+
+						OutRanges.DensityRange[WriteIndex] = ComputedDensity;
+						++NumWritten;
+					}
 				}
 			}
+
+			return NumWritten;
 		};
 
-		TArray<FPCGPoint>& OutPoints = OutPointData->GetMutablePoints();
-		OutPoints.SetNumUninitialized(InputPoints.Num());
 		FPCGTaggedData& Output = Outputs.Add_GetRef(PointInputs[i]);
 		Output.Data = OutPointData;
 
-		FPCGAsync::AsyncPointProcessing(Context, InputPoints.Num(), OutPoints, ProcessPoint);
+		FPCGAsync::AsyncProcessingRangeEx(
+			&Context->AsyncState,
+			PointData->GetNumPoints(),
+			InitializeFunc,
+			ProcessRangeFunc,
+			MoveDataRangeFunc,
+			FinishedFunc,
+			/*bEnableTimeSlicing=*/false);
+
 	}
 
 	return true;
 }
 
-#undef LOCTEXT_NAMESPACE
+#undef LOCTEXT_NAMESPACE
\ No newline at end of file
diff --git a/Engine/Plugins/PCG/Source/PCG/Private/Elements/PCGSanityCheckPointData.cpp b/Engine/Plugins/PCG/Source/PCG/Private/Elements/PCGSanityCheckPointData.cpp
--- a/Engine/Plugins/PCG/Source/PCG/Private/Elements/PCGSanityCheckPointData.cpp
+++ b/Engine/Plugins/PCG/Source/PCG/Private/Elements/PCGSanityCheckPointData.cpp
@@ -6,7 +6,7 @@
 #include "PCGContext.h"
 #include "PCGPin.h"
 
-#include "Data/PCGPointData.h"
+#include "Data/PCGBasePointData.h"
 #include "Data/PCGSpatialData.h"
 
 #define LOCTEXT_NAMESPACE "PCGSanityCheckPointDataElement"
@@ -47,14 +47,14 @@ bool FPCGSanityCheckPointDataElement::ExecuteInternal(FPCGContext* Context) cons
 	// first find the total Input bounds which will determine the size of each cell
 	for (const FPCGTaggedData& Source : Sources) 
 	{
-		const UPCGPointData* PointData = Cast<UPCGPointData>(Source.Data);
+		const UPCGBasePointData* PointData = Cast<UPCGBasePointData>(Source.Data);
 
 		if (!PointData)
 		{
 			continue;
 		}
 
-		InputPointCount += PointData->GetPoints().Num();
+		InputPointCount += PointData->GetNumPoints();
 	}
 
 	if (InputPointCount < Settings->MinPointCount)
diff --git a/Engine/Plugins/PCG/Source/PCG/Private/Elements/PCGSplineSampler.cpp b/Engine/Plugins/PCG/Source/PCG/Private/Elements/PCGSplineSampler.cpp
--- a/Engine/Plugins/PCG/Source/PCG/Private/Elements/PCGSplineSampler.cpp
+++ b/Engine/Plugins/PCG/Source/PCG/Private/Elements/PCGSplineSampler.cpp
@@ -831,6 +831,9 @@ namespace PCGSplineSamplerHelpers
 				// @todo_pcg: ideally we could preallocate the number of points here if the Samplers could give a hint on number of points
 				const int32 NumPoints = OutPointData->GetNumPoints();
 				OutPointData->SetNumPoints(NumPoints+1);
+				// @todo_pcg: be more precise about allocation (sampler could provide properties to allocate)
+				OutPointData->AllocateProperties(EPCGPointNativeProperties::All);
+
 				FPCGPointValueRanges OutRanges(OutPointData, /*bAllocate=*/false);
 				OutRanges.SetFromPoint(NumPoints, OutPoint);
 			}
@@ -949,6 +952,9 @@ namespace PCGSplineSamplerHelpers
 					// @todo_pcg: ideally we could preallocate the number of points here if the Samplers could give a hint on number of points
 					const int32 NumPoints = OutPointData->GetNumPoints();
 					OutPointData->SetNumPoints(NumPoints + 1);
+					// @todo_pcg: be more precise about allocation (sampler could provide properties to allocate)
+					OutPointData->AllocateProperties(EPCGPointNativeProperties::All);
+					
 					FPCGPointValueRanges OutRanges(OutPointData, /*bAllocate=*/false);
 					OutRanges.SetFromPoint(NumPoints, OutPoint);
 				}
@@ -997,9 +1003,6 @@ namespace PCGSplineSamplerHelpers
 		FSamplerResult* PreviousResult = &Results[0];
 		FSamplerResult* CurrentResult = &Results[1];
 
-		// @todo_pcg: be more precise about allocation (sampler could provide properties to allocate)
-		OutPointData->AllocateProperties(EPCGPointNativeProperties::All);
-
 		if (!Sampler->IsDone())
 		{
 			Sampler->Step(*PreviousResult);
diff --git a/Engine/Plugins/PCG/Source/PCG/Private/Elements/PCGSplitPoints.cpp b/Engine/Plugins/PCG/Source/PCG/Private/Elements/PCGSplitPoints.cpp
--- a/Engine/Plugins/PCG/Source/PCG/Private/Elements/PCGSplitPoints.cpp
+++ b/Engine/Plugins/PCG/Source/PCG/Private/Elements/PCGSplitPoints.cpp
@@ -4,7 +4,7 @@
 
 #include "PCGContext.h"
 #include "PCGPin.h"
-#include "Data/PCGPointData.h"
+#include "Data/PCGBasePointData.h"
 #include "Helpers/PCGAsync.h"
 #include "Metadata/PCGMetadataAccessor.h"
 
@@ -54,57 +54,92 @@ bool FPCGSplitPointsElement::ExecuteInternal(FPCGContext* Context) const
 
 	for (int i = 0; i < Inputs.Num(); ++i)
 	{
-		const UPCGPointData* InputPointData = Cast<UPCGPointData>(Inputs[i].Data);
+		const UPCGBasePointData* InputPointData = Cast<UPCGBasePointData>(Inputs[i].Data);
 		if (!InputPointData)
 		{
 			PCGE_LOG(Verbose, GraphAndLog, FText::Format(LOCTEXT("InvalidPointData", "Input {0} is not point data"), i));
 			continue;
 		}
 
-		const TArray<FPCGPoint>& InputPoints = InputPointData->GetPoints();
-
 		FPCGTaggedData& OutputA = Outputs.Add_GetRef(Inputs[i]);
-		UPCGPointData* OutPointDataA = FPCGContext::NewObject_AnyThread<UPCGPointData>(Context);
+		UPCGBasePointData* OutPointDataA = FPCGContext::NewPointData_AnyThread(Context);
+
 		OutPointDataA->InitializeFromData(InputPointData);
-		TArray<FPCGPoint>& PointsA = OutPointDataA->GetMutablePoints();
-		PointsA.SetNumUninitialized(InputPoints.Num());
+
+		OutPointDataA->SetNumPoints(InputPointData->GetNumPoints());
+
+		if (!OutPointDataA->HasSpatialDataParent())
+		{
+			OutPointDataA->AllocateProperties(InputPointData->GetAllocatedProperties());
+		}
+
+		OutPointDataA->AllocateProperties(EPCGPointNativeProperties::BoundsMax);
+
 		OutputA.Data = OutPointDataA;
 		OutputA.Pin = PCGSplitPointsConstants::OutputALabel;
 
 		FPCGTaggedData& OutputB = Outputs.Add_GetRef(Inputs[i]);
-		UPCGPointData* OutPointDataB = FPCGContext::NewObject_AnyThread<UPCGPointData>(Context);
+		UPCGBasePointData* OutPointDataB = FPCGContext::NewPointData_AnyThread(Context);
+
 		OutPointDataB->InitializeFromData(InputPointData);
-		TArray<FPCGPoint>& PointsB = OutPointDataB->GetMutablePoints();
-		PointsB.SetNumUninitialized(InputPoints.Num());
+
+		OutPointDataB->SetNumPoints(InputPointData->GetNumPoints());
+
+		if (!OutPointDataB->HasSpatialDataParent())
+		{
+			OutPointDataB->AllocateProperties(InputPointData->GetAllocatedProperties());
+		}
+
+		OutPointDataB->AllocateProperties(EPCGPointNativeProperties::BoundsMin);
+
 		OutputB.Data = OutPointDataB;
 		OutputB.Pin = PCGSplitPointsConstants::OutputBLabel;
 
-		auto ProcessPoint = [Settings, SplitPosition, &InputPoints, &InputPointData, &OutPointDataA, &OutPointDataB, &PointsA, &PointsB](const int32 ReadIndex, const int32 WriteIndex)
+		auto ProcessPoint = [Settings, SplitPosition, InputPointData, OutPointDataA, OutPointDataB](const int32 StartReadIndex, const int32 StartWriteIndex, const int32 Count)
 		{
-			const FPCGPoint& InPoint = InputPoints[ReadIndex];
-			FPCGPoint PointA = InPoint;
-			FPCGPoint PointB = InPoint;
+			const FConstPCGPointValueRanges InRanges(InputPointData);
+			FPCGPointValueRanges OutRangesA(OutPointDataA, /*bAllocate=*/false);
+			FPCGPointValueRanges OutRangesB(OutPointDataB, /*bAllocate=*/false);
 
-			const FVector SplitterPosition = (InPoint.BoundsMax - InPoint.BoundsMin) * SplitPosition;
-			const FVector MinPlusSplit = InPoint.BoundsMin + SplitterPosition;
+			int32 NumWritten = 0;
 
-			FVector SplitValues = FVector::ZeroVector;
-			const int AxisIndex = static_cast<int>(Settings->SplitAxis);
-			if (ensure(AxisIndex >= 0 && AxisIndex <= 2))
+			for (int32 ReadIndex = StartReadIndex; ReadIndex < StartReadIndex + Count; ++ReadIndex)
 			{
-				SplitValues[AxisIndex] = 1.0;
+				const int32 WriteIndex = StartWriteIndex + NumWritten;
+				const FVector SplitterPosition = (InRanges.BoundsMaxRange[ReadIndex] - InRanges.BoundsMinRange[ReadIndex]) * SplitPosition;
+				const FVector MinPlusSplit = InRanges.BoundsMinRange[ReadIndex] + SplitterPosition;
+
+				FVector SplitValues = FVector::ZeroVector;
+				const int AxisIndex = static_cast<int>(Settings->SplitAxis);
+				if (ensure(AxisIndex >= 0 && AxisIndex <= 2))
+				{
+					SplitValues[AxisIndex] = 1.0;
+				}
+
+				if (!OutPointDataA->HasSpatialDataParent())
+				{
+					OutRangesA.SetFromValueRanges(WriteIndex, InRanges, ReadIndex);
+				}
+
+				// Execution for PointsA portion
+				OutRangesA.BoundsMaxRange[WriteIndex] = OutRangesA.BoundsMaxRange[WriteIndex] + SplitValues * (MinPlusSplit - OutRangesA.BoundsMaxRange[WriteIndex]);
+				
+				if(!OutPointDataB->HasSpatialDataParent())
+				{
+					OutRangesB.SetFromValueRanges(WriteIndex, InRanges, ReadIndex);
+				}
+
+				// Execution of the PointsB portion
+				OutRangesB.BoundsMinRange[WriteIndex] = OutRangesB.BoundsMinRange[WriteIndex] + SplitValues * (MinPlusSplit - OutRangesB.BoundsMinRange[WriteIndex]);
+				
+				++NumWritten;
 			}
 
-			// Execution for PointsA portion
-			PointA.BoundsMax = PointA.BoundsMax + SplitValues * (MinPlusSplit - PointA.BoundsMax);
-			PointsA[WriteIndex] = PointA;
-
-			// Execution of the PointsB portion
-			PointB.BoundsMin = PointB.BoundsMin + SplitValues * (MinPlusSplit - PointB.BoundsMin);
-			PointsB[WriteIndex] = PointB;
+			check(NumWritten == Count);
+			return Count;
 		};
 
-		FPCGAsync::AsyncProcessingOneToOneEx(&Context->AsyncState, InputPoints.Num(), /*Initialize=*/[]() {}, ProcessPoint, /*bEnableTimeSlicing=*/false);
+		FPCGAsync::AsyncProcessingOneToOneRangeEx(&Context->AsyncState, InputPointData->GetNumPoints(), /*Initialize=*/[]() {}, ProcessPoint, /*bEnableTimeSlicing=*/false);
 	}
 
 	return true;
diff --git a/Engine/Plugins/PCG/Source/PCG/Private/Elements/PCGSurfaceSampler.cpp b/Engine/Plugins/PCG/Source/PCG/Private/Elements/PCGSurfaceSampler.cpp
--- a/Engine/Plugins/PCG/Source/PCG/Private/Elements/PCGSurfaceSampler.cpp
+++ b/Engine/Plugins/PCG/Source/PCG/Private/Elements/PCGSurfaceSampler.cpp
@@ -154,9 +154,9 @@ namespace PCGSurfaceSampler
 		return FIntVector2(CellMinX + (Index % CellCountX), CellMinY + (Index / CellCountX));
 	}
 
-	UPCGBasePointData* SampleSurface(FPCGContext* Context, const FSurfaceSamplerParams& ExecutionParams, const UPCGSurfaceData* InSurface, const UPCGSpatialData* InBoundingShape, const FBox& EffectiveBounds, bool bForcePointData)
+	UPCGBasePointData* SampleSurface(FPCGContext* Context, const FSurfaceSamplerParams& ExecutionParams, const UPCGSurfaceData* InSurface, const UPCGSpatialData* InBoundingShape, const FBox& EffectiveBounds, TSubclassOf<UPCGBasePointData> PointDataClass)
 	{
-		UPCGBasePointData* SampledData = bForcePointData ? FPCGContext::NewObject_AnyThread<UPCGPointData>(Context) : FPCGContext::NewPointData_AnyThread(Context);
+		UPCGBasePointData* SampledData = PointDataClass != nullptr ? FPCGContext::NewObject_AnyThread<UPCGBasePointData>(Context, GetTransientPackage(), PointDataClass) : FPCGContext::NewPointData_AnyThread(Context);
 		SampledData->InitializeFromData(InSurface);
 
 		FSurfaceSamplerData SamplerData;
@@ -170,14 +170,9 @@ namespace PCGSurfaceSampler
 
 	UPCGPointData* SampleSurface(FPCGContext* Context, const UPCGSurfaceData* InSurface, const UPCGSpatialData* InBoundingShape, const FBox& EffectiveBounds, const FSurfaceSamplerParams& ExecutionParams)
 	{
-		return CastChecked<UPCGPointData>(SampleSurface(Context, ExecutionParams, InSurface, InBoundingShape, EffectiveBounds, /*bForcePointData=*/true));
+		return CastChecked<UPCGPointData>(SampleSurface(Context, ExecutionParams, InSurface, InBoundingShape, EffectiveBounds, UPCGPointData::StaticClass()));
 	}
 
-	UPCGBasePointData* SampleSurface(FPCGContext* Context, const FSurfaceSamplerParams& ExecutionParams, const UPCGSurfaceData* InSurface, const UPCGSpatialData* InBoundingShape, const FBox& EffectiveBounds)
-	{
-		return SampleSurface(Context, ExecutionParams, InSurface, InBoundingShape, EffectiveBounds, /*bForcePointData=*/false);
-	}
-	
 	bool SampleSurface(FPCGContext* Context, const FSurfaceSamplerData& SamplerData, const UPCGSurfaceData* InSurface, const UPCGSpatialData* InBoundingShape, UPCGBasePointData* SampledData, const bool bTimeSlicingIsEnabled)
 	{
 		check(InSurface && SampledData);
diff --git a/Engine/Plugins/PCG/Source/PCG/Private/Elements/PCGVisualizeAttribute.cpp b/Engine/Plugins/PCG/Source/PCG/Private/Elements/PCGVisualizeAttribute.cpp
--- a/Engine/Plugins/PCG/Source/PCG/Private/Elements/PCGVisualizeAttribute.cpp
+++ b/Engine/Plugins/PCG/Source/PCG/Private/Elements/PCGVisualizeAttribute.cpp
@@ -5,7 +5,7 @@
 #include "PCGComponent.h"
 #include "PCGContext.h"
 #include "PCGDebugDrawComponent.h"
-#include "Data/PCGPointData.h"
+#include "Data/PCGBasePointData.h"
 #include "Metadata/Accessors/PCGAttributeAccessorHelpers.h"
 
 #define LOCTEXT_NAMESPACE "PCGVisualizeAttributeElement"
@@ -43,7 +43,7 @@ bool FPCGVisualizeAttribute::ExecuteInternal(FPCGContext* Context) const
 
 	for (const FPCGTaggedData& InputData : Inputs)
 	{
-		const UPCGPointData* PointData = Cast<const UPCGPointData>(InputData.Data);
+		const UPCGBasePointData* PointData = Cast<const UPCGBasePointData>(InputData.Data);
 		if (!PointData)
 		{
 			PCGLog::LogErrorOnGraph(LOCTEXT("InvalidInputType", "Input data type is not supported, only supports Point Data."), Context);
@@ -51,7 +51,7 @@ bool FPCGVisualizeAttribute::ExecuteInternal(FPCGContext* Context) const
 		}
 
 		// No points, early out
-		if (PointData->GetPoints().IsEmpty())
+		if (PointData->IsEmpty())
 		{
 			continue;
 		}
@@ -110,13 +110,15 @@ bool FPCGVisualizeAttribute::ExecuteInternal(FPCGContext* Context) const
 			SourceComponent->AddToManagedResources(ManagedComponent);
 		}
 
-		const TArray<FPCGPoint>& Points = PointData->GetPoints();
+		const TConstPCGValueRange<FTransform> TransformRanges = PointData->GetConstTransformValueRange();
 		TArray<FDebugRenderSceneProxy::FText3d> DebugStrings;
-		DebugStrings.Reserve(FMath::Min(Points.Num(), Settings->PointLimit));
+
+		const int32 NumPoints = FMath::Min(PointData->GetNumPoints(), Settings->PointLimit);
+		DebugStrings.Reserve(NumPoints);
 
 		// TODO: Optimize which points will be drawn by distance, or other measures, rather than by Index
 		// Debug messages are already culled by camera frustum, but a point limit is desired to prevent overloading
-		for (int I = 0; I < Points.Num() && I < Settings->PointLimit; ++I)
+		for (int I = 0; I < NumPoints; ++I)
 		{
 			FStringBuilderBase CompoundedString;
 
@@ -134,7 +136,7 @@ bool FPCGVisualizeAttribute::ExecuteInternal(FPCGContext* Context) const
 
 			CompoundedString = FString::Printf(TEXT("%s %s"), *CompoundedString, *ValuesToString[I]);
 
-			DebugStrings.Emplace(CompoundedString.ToString(), Points[I].Transform.GetLocation() + Settings->LocalOffset, Settings->Color);
+			DebugStrings.Emplace(CompoundedString.ToString(), TransformRanges[I].GetLocation() + Settings->LocalOffset, Settings->Color);
 		}
 
 		check(DebugDrawComponent);
diff --git a/Engine/Plugins/PCG/Source/PCG/Private/Elements/PCGWorldQuery.cpp b/Engine/Plugins/PCG/Source/PCG/Private/Elements/PCGWorldQuery.cpp
--- a/Engine/Plugins/PCG/Source/PCG/Private/Elements/PCGWorldQuery.cpp
+++ b/Engine/Plugins/PCG/Source/PCG/Private/Elements/PCGWorldQuery.cpp
@@ -5,7 +5,6 @@
 #include "PCGComponent.h"
 #include "PCGContext.h"
 #include "PCGSubsystem.h"
-#include "Data/PCGPointData.h"
 #include "Grid/PCGLandscapeCache.h"
 #include "Helpers/PCGHelpers.h"
 #include "Helpers/PCGWorldQueryHelpers.h"
diff --git a/Engine/Plugins/PCG/Source/PCG/Private/Elements/PCGWorldRaycast.cpp b/Engine/Plugins/PCG/Source/PCG/Private/Elements/PCGWorldRaycast.cpp
--- a/Engine/Plugins/PCG/Source/PCG/Private/Elements/PCGWorldRaycast.cpp
+++ b/Engine/Plugins/PCG/Source/PCG/Private/Elements/PCGWorldRaycast.cpp
@@ -5,7 +5,7 @@
 #include "PCGComponent.h"
 #include "PCGContext.h"
 #include "PCGPoint.h"
-#include "Data/PCGPointData.h"
+#include "Data/PCGBasePointData.h"
 #include "Helpers/PCGBlueprintHelpers.h"
 #include "Helpers/PCGHelpers.h"
 #include "Helpers/PCGSettingsHelpers.h"
@@ -14,6 +14,7 @@
 
 #include "Engine/HitResult.h"
 #include "Engine/World.h"
+#include "Misc/Optional.h"
 
 #include UE_INLINE_GENERATED_CPP_BY_NAME(PCGWorldRaycast)
 
@@ -277,14 +278,14 @@ bool FPCGWorldRaycastElement::PrepareDataInternal(FPCGContext* InContext) const
 		const UPCGData* FilterActorsData = !FilterActorInputData.IsEmpty() ? FilterActorInputData[Index % FilterActorInputData.Num()].Data : nullptr;
 
 		// Accept only point or param data
-		if (!OriginsData || (!OriginsData->IsA<UPCGPointData>() && !OriginsData->IsA<UPCGParamData>()))
+		if (!OriginsData || (!OriginsData->IsA<UPCGBasePointData>() && !OriginsData->IsA<UPCGParamData>()))
 		{
 			PCGLog::InputOutput::LogTypedDataNotFoundWarning(EPCGDataType::PointOrParam, Labels::OriginsInput, Context);
 			return EPCGTimeSliceInitResult::NoOperation;
 		}
 
 		// Accept only point or param data
-		if (FilterActorsData && !FilterActorsData->IsA<UPCGPointData>() && !FilterActorsData->IsA<UPCGParamData>())
+		if (FilterActorsData && !FilterActorsData->IsA<UPCGBasePointData>() && !FilterActorsData->IsA<UPCGParamData>())
 		{
 			PCGLog::InputOutput::LogTypedDataNotFoundWarning(EPCGDataType::PointOrParam, PCGWorldRayHitConstants::FilterActorPinLabel, Context);
 			return EPCGTimeSliceInitResult::NoOperation;
@@ -417,29 +418,24 @@ bool FPCGWorldRaycastElement::ExecuteInternal(FPCGContext* InContext) const
 
 		const TArray<FPCGTaggedData> OriginInputs = Context->InputData.GetInputsByPin(PCGWorldRaycastElement::Constants::Labels::OriginsInput);
 		const FPCGTaggedData& OriginInputData = OriginInputs[IterIndex];
-		const UPCGPointData* OriginInputPointData = Cast<UPCGPointData>(OriginInputData.Data);
-		const TArray<FPCGPoint>* OriginInputPoints = OriginInputPointData ? &OriginInputPointData->GetPoints() : nullptr;
+		const UPCGBasePointData* OriginInputPointData = Cast<UPCGBasePointData>(OriginInputData.Data);
 
 		TArray<FPCGTaggedData>& OutputData = Context->OutputData.TaggedData;
-		UPCGPointData* OutputPointData = FPCGContext::NewObject_AnyThread<UPCGPointData>(Context);
+		UPCGBasePointData* OutputPointData = FPCGContext::NewPointData_AnyThread(Context);
+
 		// TODO: This should initialize from the End Points point data, if Origin point data doesn't exist
-		OutputPointData->InitializeFromData(OriginInputPointData);
+		FPCGInitializeFromDataParams InitializeFromDataParams(OriginInputPointData);
+		InitializeFromDataParams.bInheritSpatialData = false;
+		OutputPointData->InitializeFromDataWithParams(InitializeFromDataParams);
+
 		FPCGTaggedData& Output = OutputData.Emplace_GetRef(OriginInputData);
 		Output.Data = OutputPointData;
 
 		UPCGMetadata* OutMetadata = OutputPointData->Metadata;
 
-		TArray<FPCGPoint>& OutputPoints = OutputPointData->GetMutablePoints();
-		// If all points will be kept, the size can be pre-allocated
-		if (Settings->bKeepOriginalPointOnMiss)
-		{
-			OutputPoints.Reserve(IterState.CachedRayOrigins.Num());
-		}
-
 		FPCGWorldRaycastQueryParams WorldQueryParams = Settings->WorldQueryParams;
 		WorldQueryParams.Initialize();
 
-
 		if (!PCGWorldQueryHelpers::CreateRayHitAttributes(WorldQueryParams, OutMetadata))
 		{
 			PCGLog::LogWarningOnGraph(LOCTEXT("UnableToCreateAllAttributes", "One or more attributes were unable to be created."), Context);
@@ -448,6 +444,12 @@ bool FPCGWorldRaycastElement::ExecuteInternal(FPCGContext* InContext) const
 		bool bAttributeSuccess = true;
 
 		const int32 PointCount = FMath::Max(IterState.CachedRayOrigins.Num(), IterState.CachedRayVectors.Num());
+		
+		TArray<TTuple<TOptional<FHitResult>, const FVector&, const FVector&>> HitResults;
+		HitResults.Reserve(PointCount);
+		
+		TArray<FHitResult> Hits;
+
 		for (int64 PointIndex = 0; PointIndex < PointCount; ++PointIndex)
 		{
 			const FVector& Origin = IterState.CachedRayOrigins[IterState.CachedRayOrigins.Num() == 1 ? 0 : PointIndex];
@@ -471,7 +473,7 @@ bool FPCGWorldRaycastElement::ExecuteInternal(FPCGContext* InContext) const
 			}
 
 			// The physics system will automatically follow the line trace path, if the collision shape is line or has no volume.
-			TArray<FHitResult> Hits;
+			Hits.Reset();
 			World->SweepMultiByObjectType(
 				Hits,
 				Origin,
@@ -480,36 +482,74 @@ bool FPCGWorldRaycastElement::ExecuteInternal(FPCGContext* InContext) const
 				ExecState.CollisionObjectQueryParams,
 				Settings->CollisionShape.ToCollisionShape(),
 				ExecState.CollisionQueryParams);
-			const TOptional<FHitResult>& HitResult = PCGWorldQueryHelpers::FilterRayHitResults(&WorldQueryParams, Cast<UPCGComponent>(Context->ExecutionSource.Get()), Hits, IterState.CachedFilterActors);
 
-			// If there's a hit, and it's within bounds, keep the point
+			TOptional<FHitResult> HitResult = PCGWorldQueryHelpers::FilterRayHitResults(&WorldQueryParams, Cast<UPCGComponent>(Context->ExecutionSource.Get()), Hits, IterState.CachedFilterActors);
 			if (HitResult.IsSet() && (Settings->bUnbounded || PCGHelpers::IsInsideBounds(ExecState.Bounds, HitResult.GetValue().ImpactPoint)))
+			{
+				HitResults.Add({ MoveTemp(HitResult), Origin, Direction });
+			}
+			else if (Settings->bKeepOriginalPointOnMiss)
+			{
+				HitResults.Add({ TOptional<FHitResult>(), Origin, Direction });
+			}
+		}
+
+		OutputPointData->SetNumPoints(HitResults.Num(), /*bInitializeValues=*/false);
+
+		if (OriginInputPointData)
+		{
+			OutputPointData->AllocateProperties(OriginInputPointData->GetAllocatedProperties());
+		}
+
+		OutputPointData->AllocateProperties(EPCGPointNativeProperties::Transform | EPCGPointNativeProperties::Seed | EPCGPointNativeProperties::MetadataEntry);
+		if (Settings->bKeepOriginalPointOnMiss)
+		{
+			OutputPointData->AllocateProperties(EPCGPointNativeProperties::Density);
+		}
+		else
+		{
+			OutputPointData->SetDensity(1.f);
+		}
+
+		const FConstPCGPointValueRanges InRanges = OriginInputPointData ? FConstPCGPointValueRanges(OriginInputPointData) : FConstPCGPointValueRanges();
+		FPCGPointValueRanges OutRanges(OutputPointData, /*bAllocate=*/false);
+
+		for (int64 PointIndex = 0; PointIndex < HitResults.Num(); ++PointIndex)
+		{
+			const TOptional<FHitResult>& HitResult = HitResults[PointIndex].Get<0>();
+			const FVector& Origin = HitResults[PointIndex].Get<1>();
+			const FVector& Direction = HitResults[PointIndex].Get<2>();
+
+			// If there's a hit, and it's within bounds, keep the point
+			if (HitResult.IsSet())
 			{
 				const FHitResult& Hit = HitResult.GetValue();
 				// Copy the input point, if it exists
-				FPCGPoint& OutPoint = OriginInputPoints
-					? OutputPoints.Emplace_GetRef((*OriginInputPoints)[OriginInputPoints->Num() == 1 ? 0 : PointIndex])
-					: OutputPoints.Emplace_GetRef();
+				FPCGPoint OutPoint = OriginInputPointData ? InRanges.GetPoint(OriginInputPointData->GetNumPoints() == 1 ? 0 : PointIndex) : FPCGPoint();
 
 				OutPoint.Transform = PCGWorldQueryHelpers::GetOrthonormalImpactTransform(Hit);
 				OutPoint.Density = 1.f;
-				UPCGBlueprintHelpers::SetSeedFromPosition(OutPoint);
+				OutPoint.Seed = PCGHelpers::ComputeSeedFromPosition(OutPoint.Transform.GetLocation());
+
+				bAttributeSuccess &= PCGWorldQueryHelpers::ApplyRayHitMetadata(Hit, WorldQueryParams, Direction, OutPoint.Transform, OutPoint.MetadataEntry, OutMetadata, World);
 
-				bAttributeSuccess &= PCGWorldQueryHelpers::ApplyRayHitMetadata(Hit, WorldQueryParams, Direction, OutputPoints.Last(), OutMetadata, World);
+				OutRanges.SetFromPoint(PointIndex, OutPoint);
 			}
-			else if (Settings->bKeepOriginalPointOnMiss)
+			else
 			{
 				// Copy the original point if possible
-				if (OriginInputPoints)
+				if (OriginInputPointData)
 				{
-					OutputPoints.Emplace_GetRef((*OriginInputPoints)[OriginInputPoints->Num() == 1 ? 0 : PointIndex]).Density = 0.f;
+					OutRanges.SetFromPoint(PointIndex, InRanges.GetPoint(OriginInputPointData->GetNumPoints() == 1 ? 0 : PointIndex));	
 				}
 				else
 				{
-					OutputPoints.Emplace(FTransform(Origin), /*InDensity=*/0.f, UPCGBlueprintHelpers::ComputeSeedFromPosition(Origin));
+					OutRanges.TransformRange[PointIndex] = FTransform(Origin);
+					OutRanges.SeedRange[PointIndex] = PCGHelpers::ComputeSeedFromPosition(Origin);
 				}
 
-				bAttributeSuccess &= PCGWorldQueryHelpers::ApplyRayMissMetadata(WorldQueryParams, OutputPoints.Last(), OutMetadata);
+				OutRanges.DensityRange[PointIndex] = 0.f;
+				bAttributeSuccess &= PCGWorldQueryHelpers::ApplyRayMissMetadata(WorldQueryParams, OutRanges.MetadataEntryRange[PointIndex], OutMetadata);
 			}
 		}
 
diff --git a/Engine/Plugins/PCG/Source/PCG/Private/Grid/PCGLandscapeCache.cpp b/Engine/Plugins/PCG/Source/PCG/Private/Grid/PCGLandscapeCache.cpp
--- a/Engine/Plugins/PCG/Source/PCG/Private/Grid/PCGLandscapeCache.cpp
+++ b/Engine/Plugins/PCG/Source/PCG/Private/Grid/PCGLandscapeCache.cpp
@@ -266,14 +266,14 @@ void FPCGLandscapeCacheEntry::GetInterpolatedPointInternal(const PCGLandscapeCac
 	OutPoint.BoundsMax = PointHalfSize;
 }
 
-void FPCGLandscapeCacheEntry::GetInterpolatedPointMetadataInternal(const PCGLandscapeCache::FSafeIndices& Indices, FPCGPoint& OutPoint, UPCGMetadata* OutMetadata) const
+void FPCGLandscapeCacheEntry::GetInterpolatedPointMetadataInternal(const PCGLandscapeCache::FSafeIndices& Indices, int64& OutMetadataEntry, UPCGMetadata* OutMetadata) const
 {
 	check(bDataLoaded);
 	check(2 * Indices.X1Y1 < PositionsAndNormals.Num());
 
 	if (OutMetadata && !LayerData.IsEmpty())
 	{
-		OutMetadata->InitializeOnSet(OutPoint.MetadataEntry);
+		OutMetadata->InitializeOnSet(OutMetadataEntry);
 
 		check(LayerData.Num() == LayerDataNames.Num());
 
@@ -289,7 +289,7 @@ void FPCGLandscapeCacheEntry::GetInterpolatedPointMetadataInternal(const PCGLand
 				const float Y1Data = FMath::Lerp((float)CurrentLayerData[Indices.X0Y1] / 255.0f, (float)CurrentLayerData[Indices.X1Y1] / 255.0f, Indices.XFraction);
 				const float Data = FMath::Lerp(Y0Data, Y1Data, Indices.YFraction);
 
-				static_cast<FPCGMetadataAttribute<float>*>(Attribute)->SetValue(OutPoint.MetadataEntry, Data);
+				static_cast<FPCGMetadataAttribute<float>*>(Attribute)->SetValue(OutMetadataEntry, Data);
 			}
 		}
 	}
@@ -300,14 +300,14 @@ void FPCGLandscapeCacheEntry::GetInterpolatedPoint(const FVector2D& LocalPoint,
 	check (bDataLoaded);
 	const PCGLandscapeCache::FSafeIndices Indices = PCGLandscapeCache::CalcSafeIndices(LocalPoint, Stride);
 	GetInterpolatedPointInternal(Indices, OutPoint);
-	GetInterpolatedPointMetadataInternal(Indices, OutPoint, OutMetadata);
+	GetInterpolatedPointMetadataInternal(Indices, OutPoint.MetadataEntry, OutMetadata);
 }
 
-void FPCGLandscapeCacheEntry::GetInterpolatedPointMetadataOnly(const FVector2D& LocalPoint, FPCGPoint& OutPoint, UPCGMetadata* OutMetadata) const
+void FPCGLandscapeCacheEntry::GetInterpolatedPointMetadataOnly(const FVector2D& LocalPoint, int64& OutMetadataEntry, UPCGMetadata* OutMetadata) const
 {
 	check(bDataLoaded);
 	const PCGLandscapeCache::FSafeIndices Indices = PCGLandscapeCache::CalcSafeIndices(LocalPoint, Stride);
-	GetInterpolatedPointMetadataInternal(Indices, OutPoint, OutMetadata);
+	GetInterpolatedPointMetadataInternal(Indices, OutMetadataEntry, OutMetadata);
 }
 
 void FPCGLandscapeCacheEntry::GetInterpolatedPointHeightOnly(const FVector2D& LocalPoint, FPCGPoint& OutPoint, UPCGMetadata* OutMetadata) const
@@ -315,7 +315,7 @@ void FPCGLandscapeCacheEntry::GetInterpolatedPointHeightOnly(const FVector2D& Lo
 	check(bDataLoaded);
 	const PCGLandscapeCache::FSafeIndices Indices = PCGLandscapeCache::CalcSafeIndices(LocalPoint, Stride);
 	GetInterpolatedPointInternal(Indices, OutPoint, /*bHeightOnly=*/true);
-	GetInterpolatedPointMetadataInternal(Indices, OutPoint, OutMetadata);
+	GetInterpolatedPointMetadataInternal(Indices, OutPoint.MetadataEntry, OutMetadata);
 }
 
 bool FPCGLandscapeCacheEntry::TouchAndLoad(int32 InTouch) const
@@ -973,7 +973,7 @@ TArray<FName> UPCGLandscapeCache::GetLayerNames(ALandscapeProxy* Landscape)
 	return CachedLayerNames.Array();
 }
 
-void UPCGLandscapeCache::SampleMetadataOnPoint(ALandscapeProxy* Landscape, FPCGPoint& InOutPoint, UPCGMetadata* OutMetadata)
+void UPCGLandscapeCache::SampleMetadataOnPoint(ALandscapeProxy* Landscape, const FTransform& InTransform, int64& OutMetadataEntry, UPCGMetadata* OutMetadata)
 {
 	if (!Landscape || !Landscape->GetWorld() || !OutMetadata || CachedLayerNames.IsEmpty())
 	{
@@ -988,7 +988,7 @@ void UPCGLandscapeCache::SampleMetadataOnPoint(ALandscapeProxy* Landscape, FPCGP
 	}
 
 	const FTransform LandscapeTransform = Landscape->LandscapeActorToWorld();
-	const FVector LocalPoint = LandscapeTransform.InverseTransformPosition(InOutPoint.Transform.GetLocation());
+	const FVector LocalPoint = LandscapeTransform.InverseTransformPosition(InTransform.GetLocation());
 	const FIntPoint ComponentMapKey(FMath::FloorToInt(LocalPoint.X / LandscapeInfo->ComponentSizeQuads), FMath::FloorToInt(LocalPoint.Y / LandscapeInfo->ComponentSizeQuads));
 
 	const FPCGLandscapeCacheEntry* CacheEntry = GetCacheEntry(LandscapeInfo, ComponentMapKey, Landscape);
@@ -999,7 +999,7 @@ void UPCGLandscapeCache::SampleMetadataOnPoint(ALandscapeProxy* Landscape, FPCGP
 	}
 
 	const FVector2D ComponentLocalPoint(LocalPoint.X - ComponentMapKey.X * LandscapeInfo->ComponentSizeQuads, LocalPoint.Y - ComponentMapKey.Y * LandscapeInfo->ComponentSizeQuads);
-	CacheEntry->GetInterpolatedPointMetadataOnly(ComponentLocalPoint, InOutPoint, OutMetadata);
+	CacheEntry->GetInterpolatedPointMetadataOnly(ComponentLocalPoint, OutMetadataEntry, OutMetadata);
 }
 
 #if WITH_EDITOR
diff --git a/Engine/Plugins/PCG/Source/PCG/Private/Helpers/PCGWorldQueryHelpers.cpp b/Engine/Plugins/PCG/Source/PCG/Private/Helpers/PCGWorldQueryHelpers.cpp
--- a/Engine/Plugins/PCG/Source/PCG/Private/Helpers/PCGWorldQueryHelpers.cpp
+++ b/Engine/Plugins/PCG/Source/PCG/Private/Helpers/PCGWorldQueryHelpers.cpp
@@ -226,7 +226,7 @@ namespace PCGWorldQueryHelpers
 		return bResult;
 	}
 
-	bool ApplyRayMissMetadata(const FPCGWorldRaycastQueryParams& QueryParams, FPCGPoint& OutPoint, UPCGMetadata* OutMetadata)
+	bool ApplyRayMissMetadata(const FPCGWorldRaycastQueryParams& QueryParams, int64& OutMetadataEntry, UPCGMetadata* OutMetadata)
 	{
 		if (!QueryParams.bGetImpact)
 		{
@@ -240,8 +240,8 @@ namespace PCGWorldQueryHelpers
 
 		if (FPCGMetadataAttribute<bool>* Attribute = OutMetadata->FindOrCreateAttribute<bool>(PCGWorldQueryConstants::ImpactAttribute, true, /*bAllowsInterpolation=*/true, /*bOverrideParent=*/false, /*bOverwriteIfTypeMismatch=*/true))
 		{
-			OutMetadata->InitializeOnSet(OutPoint.MetadataEntry);
-			Attribute->SetValue(OutPoint.MetadataEntry, false);
+			OutMetadata->InitializeOnSet(OutMetadataEntry);
+			Attribute->SetValue(OutMetadataEntry, false);
 			return true;
 		}
 
@@ -253,7 +253,8 @@ namespace PCGWorldQueryHelpers
 		const FHitResult& HitResult,
 		const FPCGWorldRaycastQueryParams& QueryParams,
 		const FVector& RayDirection,
-		FPCGPoint& OutPoint,
+		const FTransform& InTransform, 
+		int64& OutMetadataEntry,
 		UPCGMetadata* OutMetadata,
 		TWeakObjectPtr<UWorld> World)
 	{
@@ -262,7 +263,7 @@ namespace PCGWorldQueryHelpers
 			return false;
 		}
 
-		auto ApplyAttribute = [&OutPoint, OutMetadata]<typename Type>(FName AttributeName, const Type& Value, bool bShouldApply = true)
+		auto ApplyAttribute = [&OutMetadataEntry, OutMetadata]<typename Type>(FName AttributeName, const Type& Value, bool bShouldApply = true)
 		{
 			if (!bShouldApply)
 			{
@@ -271,8 +272,8 @@ namespace PCGWorldQueryHelpers
 
 			if (FPCGMetadataAttribute<Type>* Attribute = OutMetadata->GetMutableTypedAttribute<Type>(AttributeName))
 			{
-				OutMetadata->InitializeOnSet(OutPoint.MetadataEntry);
-				Attribute->SetValue(OutPoint.MetadataEntry, Value);
+				OutMetadata->InitializeOnSet(OutMetadataEntry);
+				Attribute->SetValue(OutMetadataEntry, Value);
 				return true;
 			}
 
@@ -317,7 +318,7 @@ namespace PCGWorldQueryHelpers
 						OutMetadata->FindOrCreateAttribute<float>(Layer, float());
 					}
 
-					LandscapeCache->SampleMetadataOnPoint(Landscape, OutPoint, OutMetadata);
+					LandscapeCache->SampleMetadataOnPoint(Landscape, InTransform, OutMetadataEntry, OutMetadata);
 				}
 			}
 		}
diff --git a/Engine/Plugins/PCG/Source/PCG/Private/SpatialAlgo/PCGOctreeQueries.cpp b/Engine/Plugins/PCG/Source/PCG/Private/SpatialAlgo/PCGOctreeQueries.cpp
--- a/Engine/Plugins/PCG/Source/PCG/Private/SpatialAlgo/PCGOctreeQueries.cpp
+++ b/Engine/Plugins/PCG/Source/PCG/Private/SpatialAlgo/PCGOctreeQueries.cpp
@@ -5,30 +5,28 @@
 namespace PCGOctreeQueries
 {
 	/** Factorize the logic for all GetClosest/GetFarthest, as the only thing that changes is the StartingDistance and the Condition. */
-	template <typename Func = TFunctionRef<bool(double ChosenPointDistance, int32 PointIndex, double SquaredDistance)>>
-	const FPCGPoint* GetPointInSphereUnderCondition(const UPCGPointData* InPointData, const FVector& InCenter, const double InSearchDistance, const double InStartingDistance, Func Condition)
+	template <typename Func = TFunctionRef<bool(double ChosenPointDistance, const UPCGBasePointData* BasePointData, int32 PointIndex, double SquaredDistance)>>
+	const int32 GetPointIndexInSphereUnderCondition(const UPCGBasePointData* InPointData, const FVector& InCenter, const double InSearchDistance, const double InStartingDistance, Func Condition)
 	{
 		if (!InPointData)
 		{
-			return nullptr;
+			return INDEX_NONE;
 		}
 
-		const FPCGPoint* ChosenPoint = nullptr;
+		int32 ChosenPointIndex = INDEX_NONE;
 		const double SearchDistance = FMath::Max(InSearchDistance, 0.0);
-		const TArray<FPCGPoint>& Points = InPointData->GetPoints();
 
 		double ChosenPointDistance = InStartingDistance;
-		UPCGOctreeQueries::ForEachPointInsideSphere(InPointData, InCenter, SearchDistance, [&Points, &ChosenPointDistance, &ChosenPoint, Condition = std::move(Condition)](const UPCGBasePointData* BasePointData, int32 PointIndex, double SquaredDistance)
+		UPCGOctreeQueries::ForEachPointInsideSphere(InPointData, InCenter, SearchDistance, [&ChosenPointDistance, &ChosenPointIndex, Condition = std::move(Condition)](const UPCGBasePointData* BasePointData, int32 PointIndex, double SquaredDistance)
 		{
-			const FPCGPoint& PCGPoint = Points[PointIndex];
-			if (Condition(ChosenPointDistance, PCGPoint, SquaredDistance))
+			if (Condition(ChosenPointDistance, BasePointData, PointIndex, SquaredDistance))
 			{
 				ChosenPointDistance = SquaredDistance;
-				ChosenPoint = &PCGPoint;
+				ChosenPointIndex = PointIndex;
 			}
 		});
 
-		return ChosenPoint;
+		return ChosenPointIndex;
 	}
 }
 
@@ -81,8 +79,14 @@ void UPCGOctreeQueries::GetClosestPoint(const UPCGPointData* InPointData, const
 
 const FPCGPoint* UPCGOctreeQueries::GetClosestPoint(const UPCGPointData* InPointData, const FVector& InCenter, const bool bInDiscardCenter, const double InSearchDistance)
 {
-	return PCGOctreeQueries::GetPointInSphereUnderCondition(InPointData, InCenter, InSearchDistance, std::numeric_limits<double>::max(),
-		[bInDiscardCenter](double ChosenPointDistance, const FPCGPoint& Point, double SquaredDistance)
+	int32 PointIndex = GetClosestPointIndex(InPointData, InCenter, bInDiscardCenter, InSearchDistance);
+	return PointIndex != INDEX_NONE ? &InPointData->GetPoints()[PointIndex] : nullptr;
+}
+
+int32 UPCGOctreeQueries::GetClosestPointIndex(const UPCGBasePointData* InPointData, const FVector& InCenter, const bool bInDiscardCenter, const double InSearchDistance)
+{
+	return PCGOctreeQueries::GetPointIndexInSphereUnderCondition(InPointData, InCenter, InSearchDistance, std::numeric_limits<double>::max(),
+		[bInDiscardCenter](double ChosenPointDistance, const UPCGBasePointData* BasePointData, int32 PointIndex, double SquaredDistance)
 		{
 			return SquaredDistance <= ChosenPointDistance && (!bInDiscardCenter || SquaredDistance > UE_DOUBLE_SMALL_NUMBER);
 		});
@@ -110,11 +114,13 @@ void UPCGOctreeQueries::GetClosestPointFromOtherPoint(const UPCGPointData* InPoi
 
 const FPCGPoint* UPCGOctreeQueries::GetClosestPointFromOtherPoint(const UPCGPointData* InPointData, const FPCGPoint& InPoint, const double InSearchDistance)
 {
-	return PCGOctreeQueries::GetPointInSphereUnderCondition(InPointData, InPoint.Transform.GetLocation(), InSearchDistance, std::numeric_limits<double>::max(),
-		[&InPoint](double ChosenPointDistance, const FPCGPoint& Point, double SquaredDistance)
+	int32 PointIndex = PCGOctreeQueries::GetPointIndexInSphereUnderCondition(InPointData, InPoint.Transform.GetLocation(), InSearchDistance, std::numeric_limits<double>::max(),
+		[&InPoint](double ChosenPointDistance, const UPCGBasePointData* BasePointData, int32 PointIndex, double SquaredDistance)
 		{
-			return SquaredDistance <= ChosenPointDistance && &InPoint != &Point;
+			const UPCGPointData* PointData = CastChecked<UPCGPointData>(BasePointData);
+			return SquaredDistance <= ChosenPointDistance && &InPoint != &PointData->GetPoints()[PointIndex];
 		});
+	return PointIndex != INDEX_NONE ? &InPointData->GetPoints()[PointIndex] : nullptr;
 }
 
 void UPCGOctreeQueries::GetFarthestPoint(const UPCGPointData* InPointData, const FVector& InCenter, bool& bOutFound, FPCGPoint& OutPoint, const double InSearchDistance)
@@ -129,8 +135,14 @@ void UPCGOctreeQueries::GetFarthestPoint(const UPCGPointData* InPointData, const
 
 const FPCGPoint* UPCGOctreeQueries::GetFarthestPoint(const UPCGPointData* InPointData, const FVector& InCenter, const double InSearchDistance)
 {
-	return PCGOctreeQueries::GetPointInSphereUnderCondition(InPointData, InCenter, InSearchDistance, std::numeric_limits<double>::min(),
-		[](double ChosenPointDistance, const FPCGPoint& Point, double SquaredDistance)
+	int32 PointIndex = GetFarthestPointIndex(InPointData, InCenter, InSearchDistance);
+	return PointIndex != INDEX_NONE ? &InPointData->GetPoints()[PointIndex] : nullptr;
+}
+
+int32 UPCGOctreeQueries::GetFarthestPointIndex(const UPCGBasePointData* InPointData, const FVector& InCenter, const double InSearchDistance)
+{
+	return PCGOctreeQueries::GetPointIndexInSphereUnderCondition(InPointData, InCenter, InSearchDistance, std::numeric_limits<double>::min(),
+		[](double ChosenPointDistance, const UPCGBasePointData* BasePointData, int32 PointIndex, double SquaredDistance)
 		{
 			return SquaredDistance >= ChosenPointDistance;
 		});
@@ -158,9 +170,11 @@ void UPCGOctreeQueries::GetFarthestPointFromOtherPoint(const UPCGPointData* InPo
 
 const FPCGPoint* UPCGOctreeQueries::GetFarthestPointFromOtherPoint(const UPCGPointData* InPointData, const FPCGPoint& InPoint, const double InSearchDistance)
 {
-	return PCGOctreeQueries::GetPointInSphereUnderCondition(InPointData, InPoint.Transform.GetLocation(), InSearchDistance, std::numeric_limits<double>::min(),
-		[&InPoint](double ChosenPointDistance, const FPCGPoint& Point, double SquaredDistance)
+	int32 PointIndex = PCGOctreeQueries::GetPointIndexInSphereUnderCondition(InPointData, InPoint.Transform.GetLocation(), InSearchDistance, std::numeric_limits<double>::min(),
+		[&InPoint](double ChosenPointDistance, const UPCGBasePointData* BasePointData, int32 PointIndex, double SquaredDistance)
 		{
-			return SquaredDistance >= ChosenPointDistance && &InPoint != &Point;
+			const UPCGPointData* PointData = CastChecked<UPCGPointData>(BasePointData);
+			return SquaredDistance >= ChosenPointDistance && &InPoint != &PointData->GetPoints()[PointIndex];
 		});
+	return PointIndex != INDEX_NONE ? &InPointData->GetPoints()[PointIndex] : nullptr;
 }
\ No newline at end of file
diff --git a/Engine/Plugins/PCG/Source/PCG/Private/Tests/Elements/PCGCombinePointsTest.cpp b/Engine/Plugins/PCG/Source/PCG/Private/Tests/Elements/PCGCombinePointsTest.cpp
--- a/Engine/Plugins/PCG/Source/PCG/Private/Tests/Elements/PCGCombinePointsTest.cpp
+++ b/Engine/Plugins/PCG/Source/PCG/Private/Tests/Elements/PCGCombinePointsTest.cpp
@@ -54,13 +54,11 @@ bool FPCGCombinePointsTest_Default::RunTest(const FString& Parameters)
 
 	UTEST_EQUAL("Output count", Outputs.Num(), 1);
 
-	const UPCGPointData* OutPointData = Cast<UPCGPointData>(Outputs[0].Data);
+	const UPCGBasePointData* OutPointData = Cast<UPCGBasePointData>(Outputs[0].Data);
 
 	UTEST_NOT_NULL("OutputA point data", OutPointData);
 
-	const TArray<FPCGPoint>& OutPoints = OutPointData->GetPoints();
-
-	UTEST_EQUAL("OutputA point count", OutPoints.Num(), 1);
+	UTEST_EQUAL("OutputA point count", OutPointData->GetNumPoints(), 1);
 
 	/** Point[0] at(0)->bounds(-1, 1)
 	 * Point [1] at (10) -> bounds (8, 12)
@@ -70,9 +68,10 @@ bool FPCGCombinePointsTest_Default::RunTest(const FString& Parameters)
 	 * 45 + (-1) / 2 == 22 <- check center pivot
 	 * 45 + |-1| / 2 == 23 <- check extents */
 
-	UTEST_EQUAL("Output Location", OutPoints[0].Transform.GetLocation(), FVector(22.0));
-	UTEST_EQUAL("Output Bounds Min", OutPoints[0].BoundsMin, FVector(-23.0));
-	UTEST_EQUAL("Output Bounds Max", OutPoints[0].BoundsMax, FVector(23.0));
+	const FConstPCGPointValueRanges OutRanges(OutPointData);
+	UTEST_EQUAL("Output Location", OutRanges.TransformRange[0].GetLocation(), FVector(22.0));
+	UTEST_EQUAL("Output Bounds Min", OutRanges.BoundsMinRange[0], FVector(-23.0));
+	UTEST_EQUAL("Output Bounds Max", OutRanges.BoundsMaxRange[0], FVector(23.0));
 	
 	return true;
 }
@@ -88,17 +87,16 @@ bool FPCGCombinePointsTest_SourceTransform::RunTest(const FString& Parameters)
 
 	UTEST_EQUAL("Output count", Outputs.Num(), 1);
 
-	const UPCGPointData* OutPointData = Cast<UPCGPointData>(Outputs[0].Data);
+	const UPCGBasePointData* OutPointData = Cast<UPCGBasePointData>(Outputs[0].Data);
 
 	UTEST_NOT_NULL("OutputA point data", OutPointData);
 
-	const TArray<FPCGPoint>& OutPoints = OutPointData->GetPoints();
-
-	UTEST_EQUAL("OutputA point count", OutPoints.Num(), 1);
+	UTEST_EQUAL("OutputA point count", OutPointData->GetNumPoints(), 1);
 
-	UTEST_EQUAL("Output Location", OutPoints[0].Transform.GetLocation(), FVector(0.0));
-	UTEST_EQUAL("Output Bounds Min", OutPoints[0].BoundsMin, FVector(-1.0));
-	UTEST_EQUAL("Output Bounds Max", OutPoints[0].BoundsMax, FVector(45.0));
+	const FConstPCGPointValueRanges OutRanges(OutPointData);
+	UTEST_EQUAL("Output Location", OutRanges.TransformRange[0].GetLocation(), FVector(0.0));
+	UTEST_EQUAL("Output Bounds Min", OutRanges.BoundsMinRange[0], FVector(-1.0));
+	UTEST_EQUAL("Output Bounds Max", OutRanges.BoundsMaxRange[0], FVector(45.0));
 
 	return true;
 }
@@ -114,20 +112,18 @@ bool FPCGCombinePointsTest_TransformLocation::RunTest(const FString& Parameters)
 
 	UTEST_EQUAL("Output count", Outputs.Num(), 1);
 
-	const UPCGPointData* OutPointData = Cast<UPCGPointData>(Outputs[0].Data);
+	const UPCGBasePointData* OutPointData = Cast<UPCGBasePointData>(Outputs[0].Data);
 
 	UTEST_NOT_NULL("OutputA point data", OutPointData);
 
-	const TArray<FPCGPoint>& OutPoints = OutPointData->GetPoints();
-
-	UTEST_EQUAL("OutputA point count", OutPoints.Num(), 1);
+	UTEST_EQUAL("OutputA point count", OutPointData->GetNumPoints(), 1);
 
 	// Has min at -1, max at 45 at origin
 	// Bound extents shift by (10) 
-
-	UTEST_EQUAL("Output Location", OutPoints[0].Transform.GetLocation(), FVector(10.0));
-	UTEST_EQUAL("Output Bounds Min", OutPoints[0].BoundsMin, FVector(-11.0));
-	UTEST_EQUAL("Output Bounds Max", OutPoints[0].BoundsMax, FVector(35.0));
+	const FConstPCGPointValueRanges OutRanges(OutPointData);
+	UTEST_EQUAL("Output Location", OutRanges.TransformRange[0].GetLocation(), FVector(10.0));
+	UTEST_EQUAL("Output Bounds Min", OutRanges.BoundsMinRange[0], FVector(-11.0));
+	UTEST_EQUAL("Output Bounds Max", OutRanges.BoundsMaxRange[0], FVector(35.0));
 
 	return true;
 }
@@ -143,18 +139,17 @@ bool FPCGCombinePointsTest_TransformRotation::RunTest(const FString& Parameters)
 
 	UTEST_EQUAL("Output count", Outputs.Num(), 1);
 
-	const UPCGPointData* OutPointData = Cast<UPCGPointData>(Outputs[0].Data);
+	const UPCGBasePointData* OutPointData = Cast<UPCGBasePointData>(Outputs[0].Data);
 
 	UTEST_NOT_NULL("OutputA point data", OutPointData);
 
-	const TArray<FPCGPoint>& OutPoints = OutPointData->GetPoints();
-
-	UTEST_EQUAL("OutputA point count", OutPoints.Num(), 1);
+	UTEST_EQUAL("OutputA point count", OutPointData->GetNumPoints(), 1);
 
 	// Rotate a cube around X axis
-	UTEST_EQUAL("Output Location", OutPoints[0].Transform.GetLocation(), FVector::ZeroVector);
-	UTEST_EQUAL("Output Bounds Min", OutPoints[0].BoundsMin, FVector(-1.0, -1.0, -45.0));
-	UTEST_EQUAL("Output Bounds Max", OutPoints[0].BoundsMax, FVector(45.0, 45.0, 1.0));
+	const FConstPCGPointValueRanges OutRanges(OutPointData);
+	UTEST_EQUAL("Output Location", OutRanges.TransformRange[0].GetLocation(), FVector::ZeroVector);
+	UTEST_EQUAL("Output Bounds Min", OutRanges.BoundsMinRange[0], FVector(-1.0, -1.0, -45.0));
+	UTEST_EQUAL("Output Bounds Max", OutRanges.BoundsMaxRange[0], FVector(45.0, 45.0, 1.0));
 
 	return true;
 }
@@ -170,17 +165,16 @@ bool FPCGCombinePointsTest_TransformScale::RunTest(const FString& Parameters)
 
 	UTEST_EQUAL("Output count", Outputs.Num(), 1);
 
-	const UPCGPointData* OutPointData = Cast<UPCGPointData>(Outputs[0].Data);
+	const UPCGBasePointData* OutPointData = Cast<UPCGBasePointData>(Outputs[0].Data);
 
 	UTEST_NOT_NULL("OutputA point data", OutPointData);
 
-	const TArray<FPCGPoint>& OutPoints = OutPointData->GetPoints();
-
-	UTEST_EQUAL("OutputA point count", OutPoints.Num(), 1);
+	UTEST_EQUAL("OutputA point count", OutPointData->GetNumPoints(), 1);
 
-	UTEST_EQUAL("Output Location", OutPoints[0].Transform.GetLocation(), FVector::ZeroVector);
-	UTEST_EQUAL("Output Bounds Min", OutPoints[0].BoundsMin, FVector(-0.5));
-	UTEST_EQUAL("Output Bounds Max", OutPoints[0].BoundsMax, FVector(22.5));
+	const FConstPCGPointValueRanges OutRanges(OutPointData);
+	UTEST_EQUAL("Output Location", OutRanges.TransformRange[0].GetLocation(), FVector::ZeroVector);
+	UTEST_EQUAL("Output Bounds Min", OutRanges.BoundsMinRange[0], FVector(-0.5));
+	UTEST_EQUAL("Output Bounds Max", OutRanges.BoundsMaxRange[0], FVector(22.5));
 
 	return true;
 }
@@ -196,17 +190,16 @@ bool FPCGCombinePointsTest_TransformNegScale::RunTest(const FString& Parameters)
 
 	UTEST_EQUAL("Output count", Outputs.Num(), 1);
 
-	const UPCGPointData* OutPointData = Cast<UPCGPointData>(Outputs[0].Data);
+	const UPCGBasePointData* OutPointData = Cast<UPCGBasePointData>(Outputs[0].Data);
 
 	UTEST_NOT_NULL("OutputA point data", OutPointData);
 
-	const TArray<FPCGPoint>& OutPoints = OutPointData->GetPoints();
-
-	UTEST_EQUAL("OutputA point count", OutPoints.Num(), 1);
+	UTEST_EQUAL("OutputA point count", OutPointData->GetNumPoints(), 1);
 
-	UTEST_EQUAL("Output Location", OutPoints[0].Transform.GetLocation(), FVector::ZeroVector);
-	UTEST_EQUAL("Output Bounds Min", OutPoints[0].BoundsMin, FVector(-22.5));
-	UTEST_EQUAL("Output Bounds Max", OutPoints[0].BoundsMax, FVector(0.5));
+	const FConstPCGPointValueRanges OutRanges(OutPointData);
+	UTEST_EQUAL("Output Location", OutRanges.TransformRange[0].GetLocation(), FVector::ZeroVector);
+	UTEST_EQUAL("Output Bounds Min", OutRanges.BoundsMinRange[0], FVector(-22.5));
+	UTEST_EQUAL("Output Bounds Max", OutRanges.BoundsMaxRange[0], FVector(0.5));
 
 	return true;
 }
\ No newline at end of file
diff --git a/Engine/Plugins/PCG/Source/PCG/Private/Tests/Elements/PCGPointNeighborhoodTest.cpp b/Engine/Plugins/PCG/Source/PCG/Private/Tests/Elements/PCGPointNeighborhoodTest.cpp
--- a/Engine/Plugins/PCG/Source/PCG/Private/Tests/Elements/PCGPointNeighborhoodTest.cpp
+++ b/Engine/Plugins/PCG/Source/PCG/Private/Tests/Elements/PCGPointNeighborhoodTest.cpp
@@ -94,17 +94,17 @@ bool FPCGPointNeighborhoodTest_AttrDistance::RunTest(const FString& Parameters)
 
 	UTEST_EQUAL("Output count", Outputs.Num(), 1);
 
-	const UPCGPointData* OutPointData = Cast<UPCGPointData>(Outputs[0].Data);
+	const UPCGBasePointData* OutPointData = Cast<UPCGBasePointData>(Outputs[0].Data);
 
 	UTEST_NOT_NULL("Output point data", OutPointData);
 
-	const TArray<FPCGPoint>& OutPoints = OutPointData->GetPoints();
-
-	UTEST_EQUAL("Output point count", OutPoints.Num(), 9);
+	UTEST_EQUAL("Output point count", OutPointData->GetNumPoints(), 9);
 
 	const FPCGMetadataAttribute<double>* DistanceAttr = OutPointData->Metadata->GetConstTypedAttribute<double>(DistanceAttrParameters.DistanceName);
 	UTEST_NOT_NULL("Distance Attribute exists", DistanceAttr);
-	UTEST_EQUAL("Point 4 of Distance Attribute", DistanceAttr->GetValueFromItemKey(OutPoints[4].MetadataEntry), 0.0)
+
+	const TConstPCGValueRange<int64> MetadataEntryRange = OutPointData->GetConstMetadataEntryValueRange();
+	UTEST_EQUAL("Point 4 of Distance Attribute", DistanceAttr->GetValueFromItemKey(MetadataEntryRange[4]), 0.0)
 
 	return true; 
 }
@@ -130,17 +130,17 @@ bool FPCGPointNeighborhoodTest_AttrPosition::RunTest(const FString& Parameters)
 
 	UTEST_EQUAL("Output count", Outputs.Num(), 1);
 
-	const UPCGPointData* OutPointData = Cast<UPCGPointData>(Outputs[0].Data);
+	const UPCGBasePointData* OutPointData = Cast<UPCGBasePointData>(Outputs[0].Data);
 
 	UTEST_NOT_NULL("Output point data", OutPointData);
 
-	const TArray<FPCGPoint>& OutPoints = OutPointData->GetPoints();
-
-	UTEST_EQUAL("Output point count", OutPoints.Num(), 9);
+	UTEST_EQUAL("Output point count", OutPointData->GetNumPoints(), 9);
 
 	const FPCGMetadataAttribute<FVector>* PositionAttr = OutPointData->Metadata->GetConstTypedAttribute<FVector>(PositionAttrParameters.PositionName);
 	UTEST_NOT_NULL("Average Position Attribute exists", PositionAttr);
-	UTEST_EQUAL("Point 4 of Position Attribute", PositionAttr->GetValueFromItemKey(OutPoints[4].MetadataEntry), FVector(100.0, 100.0, 0.0));
+
+	const TConstPCGValueRange<int64> MetadataEntryRange = OutPointData->GetConstMetadataEntryValueRange();
+	UTEST_EQUAL("Point 4 of Position Attribute", PositionAttr->GetValueFromItemKey(MetadataEntryRange[4]), FVector(100.0, 100.0, 0.0));
 
 	return true;
 }
@@ -166,22 +166,21 @@ bool FPCGPointNeighborhoodTest_DensityAvg::RunTest(const FString& Parameters)
 
 	UTEST_EQUAL("Output count", Outputs.Num(), 1);
 
-	const UPCGPointData* OutPointData = Cast<UPCGPointData>(Outputs[0].Data);
+	const UPCGBasePointData* OutPointData = Cast<UPCGBasePointData>(Outputs[0].Data);
 
 	UTEST_NOT_NULL("Output point data", OutPointData);
 
-	const TArray<FPCGPoint>& OutPoints = OutPointData->GetPoints();
-
-	UTEST_EQUAL("Output point count", OutPoints.Num(), 9);
+	UTEST_EQUAL("Output point count", OutPointData->GetNumPoints(), 9);
 
 	/** Steepness taken into account:
 	 * Point Bounds = 30 -> (20cm * (1 + 0.5))
 	 * Search Bounds = 450 (|-225| + 225)
 	 * Given 9 contributions:
 	 * 9 * (30 ^ 3) / (450 ^ 3) == 0.00266...*/
-	UTEST_EQUAL("Average Density in Point 0 Density Attribute", OutPoints[0].Density, 9.0f * FMath::Pow(30.0f, 3.0f) / FMath::Pow(450.0f, 3.0f));
-	UTEST_EQUAL("Average Density in Point 4 Density Attribute", OutPoints[4].Density, 9.0f * FMath::Pow(30.0f, 3.0f) / FMath::Pow(450.0f, 3.0f));
-	
+	const TConstPCGValueRange<float> DensityRange = OutPointData->GetConstDensityValueRange();
+	UTEST_EQUAL("Average Density in Point 0 Density Attribute", DensityRange[0], 9.0f * FMath::Pow(30.0f, 3.0f) / FMath::Pow(450.0f, 3.0f));
+	UTEST_EQUAL("Average Density in Point 4 Density Attribute", DensityRange[4], 9.0f * FMath::Pow(30.0f, 3.0f) / FMath::Pow(450.0f, 3.0f));
+
 	return true;
 }
 
@@ -206,18 +205,17 @@ bool FPCGPointNeighborhoodTest_DensityAvgNoBounds::RunTest(const FString& Parame
 
 	UTEST_EQUAL("Output count", Outputs.Num(), 1);
 
-	const UPCGPointData* OutPointData = Cast<UPCGPointData>(Outputs[0].Data);
+	const UPCGBasePointData* OutPointData = Cast<UPCGBasePointData>(Outputs[0].Data);
 
 	UTEST_NOT_NULL("Output point data", OutPointData);
 
-	const TArray<FPCGPoint>& OutPoints = OutPointData->GetPoints();
-
-	UTEST_EQUAL("Output point count", OutPoints.Num(), 9);
+	UTEST_EQUAL("Output point count", OutPointData->GetNumPoints(), 9);
 
 	/** Density / num points total = 1 / 9 contribution per point
 	 * 9 * 1/9 == 1 average density */ 
-	UTEST_EQUAL("Average Density in Point 0 Density Attribute", OutPoints[0].Density, 1.0f);
-	UTEST_EQUAL("Average Density in Point 4 Density Attribute", OutPoints[4].Density, 1.0f);
+	const TConstPCGValueRange<float> DensityRange = OutPointData->GetConstDensityValueRange();
+	UTEST_EQUAL("Average Density in Point 0 Density Attribute", DensityRange[0], 1.0f);
+	UTEST_EQUAL("Average Density in Point 4 Density Attribute", DensityRange[4], 1.0f);
 
 	return true;
 }
@@ -243,20 +241,19 @@ bool FPCGPointNeighborhoodTest_DensityNormal::RunTest(const FString& Parameters)
 
 	UTEST_EQUAL("Output count", Outputs.Num(), 1);
 
-	const UPCGPointData* OutPointData = Cast<UPCGPointData>(Outputs[0].Data);
+	const UPCGBasePointData* OutPointData = Cast<UPCGBasePointData>(Outputs[0].Data);
 
 	UTEST_NOT_NULL("Output point data", OutPointData);
 
-	const TArray<FPCGPoint>& OutPoints = OutPointData->GetPoints();
-
-	UTEST_EQUAL("Output point count", OutPoints.Num(), 9);
+	UTEST_EQUAL("Output point count", OutPointData->GetNumPoints(), 9);
 
 	/** Distance between each point / search distance (radius)
 	 * sqrt(100^2 + 100^2) == 100*sqrt(2) / 225
 	 * == 0.629...
 	*/
-	UTEST_EQUAL("Normalized Distance in Point 0 Density Attribute", OutPoints[0].Density, (100.0f * FMath::Sqrt(2.0f)) / 225.0f);
-	UTEST_EQUAL("Normalized Distance in Point 4 Density Attribute", OutPoints[4].Density, 0.0f);
+	const TConstPCGValueRange<float> DensityRange = OutPointData->GetConstDensityValueRange();
+	UTEST_EQUAL("Normalized Distance in Point 0 Density Attribute", DensityRange[0], (100.0f * FMath::Sqrt(2.0f)) / 225.0f);
+	UTEST_EQUAL("Normalized Distance in Point 4 Density Attribute", DensityRange[4], 0.0f);
 
 	return true;
 }
@@ -282,16 +279,15 @@ bool FPCGPointNeighborhoodTest_AvgPosition::RunTest(const FString& Parameters)
 
 	UTEST_EQUAL("Output count", Outputs.Num(), 1);
 
-	const UPCGPointData* OutPointData = Cast<UPCGPointData>(Outputs[0].Data);
+	const UPCGBasePointData* OutPointData = Cast<UPCGBasePointData>(Outputs[0].Data);
 
 	UTEST_NOT_NULL("Output point data", OutPointData);
 
-	const TArray<FPCGPoint>& OutPoints = OutPointData->GetPoints();
+	UTEST_EQUAL("Output point count", OutPointData->GetNumPoints(), 9);
 
-	UTEST_EQUAL("Output point count", OutPoints.Num(), 9);
-
-	UTEST_EQUAL("Average Position in Point 0 Location Attribute", OutPoints[0].Transform.GetLocation(), FVector(100.0, 100.0, 0.0));
-	UTEST_EQUAL("Average Position in Point 4 Location Attribute", OutPoints[4].Transform.GetLocation(), FVector(100.0, 100.0, 0.0));
+	const TConstPCGValueRange<FTransform> TransformRange = OutPointData->GetConstTransformValueRange();
+	UTEST_EQUAL("Average Position in Point 0 Location Attribute", TransformRange[0].GetLocation(), FVector(100.0, 100.0, 0.0));
+	UTEST_EQUAL("Average Position in Point 4 Location Attribute", TransformRange[4].GetLocation(), FVector(100.0, 100.0, 0.0));
 
 	return true;
 }
@@ -317,17 +313,16 @@ bool FPCGPointNeighborhoodTest_AvgColor::RunTest(const FString& Parameters)
 
 	UTEST_EQUAL("Output count", Outputs.Num(), 1);
 
-	const UPCGPointData* OutPointData = Cast<UPCGPointData>(Outputs[0].Data);
+	const UPCGBasePointData* OutPointData = Cast<UPCGBasePointData>(Outputs[0].Data);
 
 	UTEST_NOT_NULL("Output point data", OutPointData);
 
-	const TArray<FPCGPoint>& OutPoints = OutPointData->GetPoints();
-
-	UTEST_EQUAL("Output point count", OutPoints.Num(), 9);
+	UTEST_EQUAL("Output point count", OutPointData->GetNumPoints(), 9);
 
-	for (int i = 0; i < OutPoints.Num(); ++i)
+	const TConstPCGValueRange<FVector4> ColorRange = OutPointData->GetConstColorValueRange();
+	for (int i = 0; i < OutPointData->GetNumPoints(); ++i)
 	{
-		UTEST_EQUAL("Average Color in Location Attribute", OutPoints[i].Color, FVector4(1.0, 1.0, 1.0, 1.0));
+		UTEST_EQUAL("Average Color in Location Attribute", ColorRange[i], FVector4(1.0, 1.0, 1.0, 1.0));
 	}
 	
 	return true;
diff --git a/Engine/Plugins/PCG/Source/PCG/Private/Tests/Elements/PCGRandomChoiceTest.cpp b/Engine/Plugins/PCG/Source/PCG/Private/Tests/Elements/PCGRandomChoiceTest.cpp
--- a/Engine/Plugins/PCG/Source/PCG/Private/Tests/Elements/PCGRandomChoiceTest.cpp
+++ b/Engine/Plugins/PCG/Source/PCG/Private/Tests/Elements/PCGRandomChoiceTest.cpp
@@ -7,7 +7,7 @@
 #include "PCGComponent.h"
 #include "PCGContext.h"
 #include "PCGParamData.h"
-#include "Data/PCGPointData.h"
+#include "Data/PCGBasePointData.h"
 #include "Metadata/PCGMetadata.h"
 #include "Metadata/PCGMetadataAttributeTpl.h"
 
@@ -25,19 +25,21 @@ IMPLEMENT_CUSTOM_SIMPLE_AUTOMATION_TEST(FPCGRandomChoiceTest_Fixed_ParamData, FP
 
 namespace PCGRandomChoiceTest
 {
-	UPCGPointData* CreateInputPointData(FPCGContext* Context, const int NumPoints, const int Seed = 42)
+	UPCGBasePointData* CreateInputPointData(FPCGContext* Context, const int NumPoints, const int Seed = 42)
 	{
 		check(Context);
 
-		UPCGPointData* NewPointData = NewObject<UPCGPointData>();
+		UPCGBasePointData* NewPointData = PCGTestsCommon::CreateEmptyBasePointData();
 		NewPointData->SetFlags(RF_Transient);
-		TArray<FPCGPoint>& Points = NewPointData->GetMutablePoints();
-		Points.SetNum(NumPoints);
+		NewPointData->SetNumPoints(NumPoints);
+		NewPointData->SetSeed(Seed);
+		NewPointData->AllocateProperties(EPCGPointNativeProperties::Density);
+
+		TPCGValueRange<float> DensityRange = NewPointData->GetDensityValueRange(/*bAllocate=*/false);
 		for (int i = 0; i < NumPoints; ++i)
 		{
 			// Store the index in the density
-			Points[i].Density = i;
-			Points[i].Seed = Seed;
+			DensityRange[i] = i;
 		}
 
 		FPCGTaggedData& InputData = Context->InputData.TaggedData.Emplace_GetRef();
@@ -71,18 +73,17 @@ namespace PCGRandomChoiceTest
 		return NewParamData;
 	}
 
-	bool VerifyAllPointsThere(const int NumPoints, const UPCGPointData* ChosenPointData, const UPCGPointData* DiscardedPointData)
+	bool VerifyAllPointsThere(const int NumPoints, const UPCGBasePointData* ChosenPointData, const UPCGBasePointData* DiscardedPointData)
 	{
 		check(ChosenPointData && DiscardedPointData);
 		TSet<int> IndexesSeen;
 
-		auto Check = [&IndexesSeen](const UPCGPointData* PointData)
+		auto Check = [&IndexesSeen](const UPCGBasePointData* PointData)
 		{
-			const TArray<FPCGPoint>& Points = PointData->GetPoints();
-			for (int i = 0; i < Points.Num(); ++i)
+			const TConstPCGValueRange<float> DensityRange = PointData->GetConstDensityValueRange();
+			for (int i = 0; i < PointData->GetNumPoints(); ++i)
 			{
-				const FPCGPoint& Point = Points[i];
-				const int Index = static_cast<int>(Point.Density);
+				const int Index = static_cast<int>(DensityRange[i]);
 				if (IndexesSeen.Contains(Index))
 				{
 					return false;
@@ -91,7 +92,7 @@ namespace PCGRandomChoiceTest
 				IndexesSeen.Add(Index);
 
 				// It needs to be stable so density should be ascending.
-				if (i > 0 && Points[i].Density <= Points[i - 1].Density)
+				if (i > 0 && DensityRange[i] <= DensityRange[i - 1])
 				{
 					return false;
 				}
@@ -177,7 +178,7 @@ bool FPCGRandomChoiceTest_Fixed::RunTest(const FString& Parameters)
 	Settings->FixedNumber = ExpectedNumElementsChosen;
 
 	TUniquePtr<FPCGContext> Context = TestData.InitializeTestContext();
-	const UPCGPointData* InputPointData = PCGRandomChoiceTest::CreateInputPointData(Context.Get(), NumOfPoints);
+	const UPCGBasePointData* InputPointData = PCGRandomChoiceTest::CreateInputPointData(Context.Get(), NumOfPoints);
 
 	FPCGElementPtr TestElement = TestData.Settings->GetElement();
 	while (!TestElement->Execute(Context.Get())) {}
@@ -185,14 +186,14 @@ bool FPCGRandomChoiceTest_Fixed::RunTest(const FString& Parameters)
 	TArray<FPCGTaggedData> ChosenOutputTagged = Context->OutputData.GetInputsByPin(PCGRandomChoiceConstants::ChosenEntriesLabel);
 	TArray<FPCGTaggedData> DiscardedOutputTagged = Context->OutputData.GetInputsByPin(PCGRandomChoiceConstants::DiscardedEntriesLabel);
 
-	const UPCGPointData* ChosenOutputData = ChosenOutputTagged.Num() == 1 ? Cast<const UPCGPointData>(ChosenOutputTagged[0].Data) : nullptr;
-	const UPCGPointData* DiscardedOutputData = DiscardedOutputTagged.Num() == 1 ? Cast<const UPCGPointData>(DiscardedOutputTagged[0].Data) : nullptr;
+	const UPCGBasePointData* ChosenOutputData = ChosenOutputTagged.Num() == 1 ? Cast<const UPCGBasePointData>(ChosenOutputTagged[0].Data) : nullptr;
+	const UPCGBasePointData* DiscardedOutputData = DiscardedOutputTagged.Num() == 1 ? Cast<const UPCGBasePointData>(DiscardedOutputTagged[0].Data) : nullptr;
 
 	UTEST_NOT_NULL("There is a point data in chosen points", ChosenOutputData);
 	UTEST_NOT_NULL("There is a point data in discarded points", DiscardedOutputData);
 
-	UTEST_EQUAL("There is the right number of points in chosen", ChosenOutputData->GetPoints().Num(), ExpectedNumElementsChosen);
-	UTEST_EQUAL("There is the right number of points in discarded", DiscardedOutputData->GetPoints().Num(), ExpectedNumElementsDiscarded);
+	UTEST_EQUAL("There is the right number of points in chosen", ChosenOutputData->GetNumPoints(), ExpectedNumElementsChosen);
+	UTEST_EQUAL("There is the right number of points in discarded", DiscardedOutputData->GetNumPoints(), ExpectedNumElementsDiscarded);
 
 	UTEST_TRUE("All points are there and in the right order", PCGRandomChoiceTest::VerifyAllPointsThere(NumOfPoints, ChosenOutputData, DiscardedOutputData));
 
@@ -213,7 +214,7 @@ bool FPCGRandomChoiceTest_Ratio::RunTest(const FString& Parameters)
 	Settings->Ratio = 0.25f;
 
 	TUniquePtr<FPCGContext> Context = TestData.InitializeTestContext();
-	const UPCGPointData* InputPointData = PCGRandomChoiceTest::CreateInputPointData(Context.Get(), NumOfPoints);
+	const UPCGBasePointData* InputPointData = PCGRandomChoiceTest::CreateInputPointData(Context.Get(), NumOfPoints);
 
 	FPCGElementPtr TestElement = TestData.Settings->GetElement();
 	while (!TestElement->Execute(Context.Get())) {}
@@ -221,14 +222,14 @@ bool FPCGRandomChoiceTest_Ratio::RunTest(const FString& Parameters)
 	TArray<FPCGTaggedData> ChosenOutputTagged = Context->OutputData.GetInputsByPin(PCGRandomChoiceConstants::ChosenEntriesLabel);
 	TArray<FPCGTaggedData> DiscardedOutputTagged = Context->OutputData.GetInputsByPin(PCGRandomChoiceConstants::DiscardedEntriesLabel);
 
-	const UPCGPointData* ChosenOutputData = ChosenOutputTagged.Num() == 1 ? Cast<const UPCGPointData>(ChosenOutputTagged[0].Data) : nullptr;
-	const UPCGPointData* DiscardedOutputData = DiscardedOutputTagged.Num() == 1 ? Cast<const UPCGPointData>(DiscardedOutputTagged[0].Data) : nullptr;
+	const UPCGBasePointData* ChosenOutputData = ChosenOutputTagged.Num() == 1 ? Cast<const UPCGBasePointData>(ChosenOutputTagged[0].Data) : nullptr;
+	const UPCGBasePointData* DiscardedOutputData = DiscardedOutputTagged.Num() == 1 ? Cast<const UPCGBasePointData>(DiscardedOutputTagged[0].Data) : nullptr;
 
 	UTEST_NOT_NULL("There is a point data in chosen points", ChosenOutputData);
 	UTEST_NOT_NULL("There is a point data in discarded points", DiscardedOutputData);
 
-	UTEST_EQUAL("There is the right number of points in chosen", ChosenOutputData->GetPoints().Num(), ExpectedNumElementsChosen);
-	UTEST_EQUAL("There is the right number of points in discarded", DiscardedOutputData->GetPoints().Num(), ExpectedNumElementsDiscarded);
+	UTEST_EQUAL("There is the right number of points in chosen", ChosenOutputData->GetNumPoints(), ExpectedNumElementsChosen);
+	UTEST_EQUAL("There is the right number of points in discarded", DiscardedOutputData->GetNumPoints(), ExpectedNumElementsDiscarded);
 
 	UTEST_TRUE("All points are there and in the right order", PCGRandomChoiceTest::VerifyAllPointsThere(NumOfPoints, ChosenOutputData, DiscardedOutputData));
 
@@ -247,7 +248,7 @@ bool FPCGRandomChoiceTest_SelectNone::RunTest(const FString& Parameters)
 	Settings->FixedNumber = 0;
 
 	TUniquePtr<FPCGContext> Context = TestData.InitializeTestContext();
-	const UPCGPointData* InputPointData = PCGRandomChoiceTest::CreateInputPointData(Context.Get(), NumOfPoints);
+	const UPCGBasePointData* InputPointData = PCGRandomChoiceTest::CreateInputPointData(Context.Get(), NumOfPoints);
 
 	FPCGElementPtr TestElement = TestData.Settings->GetElement();
 	while (!TestElement->Execute(Context.Get())) {}
@@ -255,13 +256,13 @@ bool FPCGRandomChoiceTest_SelectNone::RunTest(const FString& Parameters)
 	TArray<FPCGTaggedData> ChosenOutputTagged = Context->OutputData.GetInputsByPin(PCGRandomChoiceConstants::ChosenEntriesLabel);
 	TArray<FPCGTaggedData> DiscardedOutputTagged = Context->OutputData.GetInputsByPin(PCGRandomChoiceConstants::DiscardedEntriesLabel);
 
-	const UPCGPointData* ChosenOutputData = ChosenOutputTagged.Num() == 1 ? Cast<const UPCGPointData>(ChosenOutputTagged[0].Data) : nullptr;
-	const UPCGPointData* DiscardedOutputData = DiscardedOutputTagged.Num() == 1 ? Cast<const UPCGPointData>(DiscardedOutputTagged[0].Data) : nullptr;
+	const UPCGBasePointData* ChosenOutputData = ChosenOutputTagged.Num() == 1 ? Cast<const UPCGBasePointData>(ChosenOutputTagged[0].Data) : nullptr;
+	const UPCGBasePointData* DiscardedOutputData = DiscardedOutputTagged.Num() == 1 ? Cast<const UPCGBasePointData>(DiscardedOutputTagged[0].Data) : nullptr;
 
 	UTEST_NOT_NULL("There is a point data in chosen points", ChosenOutputData);
 	UTEST_NOT_NULL("There is a point data in discarded points", DiscardedOutputData);
 
-	UTEST_EQUAL("Chosen points is empty", ChosenOutputData->GetPoints().Num(), 0);
+	UTEST_EQUAL("Chosen points is empty", ChosenOutputData->GetNumPoints(), 0);
 	UTEST_EQUAL("Discarded points is the input data", DiscardedOutputData, InputPointData);
 
 	return true;
@@ -279,7 +280,7 @@ bool FPCGRandomChoiceTest_SelectAll::RunTest(const FString& Parameters)
 	Settings->Ratio = 1.0f;
 
 	TUniquePtr<FPCGContext> Context = TestData.InitializeTestContext();
-	const UPCGPointData* InputPointData = PCGRandomChoiceTest::CreateInputPointData(Context.Get(), NumOfPoints);
+	const UPCGBasePointData* InputPointData = PCGRandomChoiceTest::CreateInputPointData(Context.Get(), NumOfPoints);
 
 	FPCGElementPtr TestElement = TestData.Settings->GetElement();
 	while (!TestElement->Execute(Context.Get())) {}
@@ -287,14 +288,14 @@ bool FPCGRandomChoiceTest_SelectAll::RunTest(const FString& Parameters)
 	TArray<FPCGTaggedData> ChosenOutputTagged = Context->OutputData.GetInputsByPin(PCGRandomChoiceConstants::ChosenEntriesLabel);
 	TArray<FPCGTaggedData> DiscardedOutputTagged = Context->OutputData.GetInputsByPin(PCGRandomChoiceConstants::DiscardedEntriesLabel);
 
-	const UPCGPointData* ChosenOutputData = ChosenOutputTagged.Num() == 1 ? Cast<const UPCGPointData>(ChosenOutputTagged[0].Data) : nullptr;
-	const UPCGPointData* DiscardedOutputData = DiscardedOutputTagged.Num() == 1 ? Cast<const UPCGPointData>(DiscardedOutputTagged[0].Data) : nullptr;
+	const UPCGBasePointData* ChosenOutputData = ChosenOutputTagged.Num() == 1 ? Cast<const UPCGBasePointData>(ChosenOutputTagged[0].Data) : nullptr;
+	const UPCGBasePointData* DiscardedOutputData = DiscardedOutputTagged.Num() == 1 ? Cast<const UPCGBasePointData>(DiscardedOutputTagged[0].Data) : nullptr;
 
 	UTEST_NOT_NULL("There is a point data in chosen points", ChosenOutputData);
 	UTEST_NOT_NULL("There is a point data in discarded points", DiscardedOutputData);
 
 	UTEST_EQUAL("Chosen points is the input data", ChosenOutputData, InputPointData);
-	UTEST_EQUAL("Discarded points is empty", DiscardedOutputData->GetPoints().Num(), 0);
+	UTEST_EQUAL("Discarded points is empty", DiscardedOutputData->GetNumPoints(), 0);
 
 	return true;
 }
@@ -313,7 +314,7 @@ bool FPCGRandomChoiceTest_NoDiscard::RunTest(const FString& Parameters)
 	Settings->bOutputDiscardedEntries = false;
 
 	TUniquePtr<FPCGContext> Context = TestData.InitializeTestContext();
-	const UPCGPointData* InputPointData = PCGRandomChoiceTest::CreateInputPointData(Context.Get(), NumOfPoints);
+	const UPCGBasePointData* InputPointData = PCGRandomChoiceTest::CreateInputPointData(Context.Get(), NumOfPoints);
 
 	FPCGElementPtr TestElement = TestData.Settings->GetElement();
 	while (!TestElement->Execute(Context.Get())) {}
@@ -321,13 +322,13 @@ bool FPCGRandomChoiceTest_NoDiscard::RunTest(const FString& Parameters)
 	TArray<FPCGTaggedData> ChosenOutputTagged = Context->OutputData.GetInputsByPin(PCGRandomChoiceConstants::ChosenEntriesLabel);
 	TArray<FPCGTaggedData> DiscardedOutputTagged = Context->OutputData.GetInputsByPin(PCGRandomChoiceConstants::DiscardedEntriesLabel);
 
-	const UPCGPointData* ChosenOutputData = ChosenOutputTagged.Num() == 1 ? Cast<const UPCGPointData>(ChosenOutputTagged[0].Data) : nullptr;
-	const UPCGPointData* DiscardedOutputData = DiscardedOutputTagged.Num() == 1 ? Cast<const UPCGPointData>(DiscardedOutputTagged[0].Data) : nullptr;
+	const UPCGBasePointData* ChosenOutputData = ChosenOutputTagged.Num() == 1 ? Cast<const UPCGBasePointData>(ChosenOutputTagged[0].Data) : nullptr;
+	const UPCGBasePointData* DiscardedOutputData = DiscardedOutputTagged.Num() == 1 ? Cast<const UPCGBasePointData>(DiscardedOutputTagged[0].Data) : nullptr;
 
 	UTEST_NOT_NULL("There is a point data in chosen points", ChosenOutputData);
 	UTEST_NULL("There is no point data in discarded points", DiscardedOutputData);
 
-	UTEST_EQUAL("There is the right number of points in chosen", ChosenOutputData->GetPoints().Num(), ExpectedNumElementsChosen);
+	UTEST_EQUAL("There is the right number of points in chosen", ChosenOutputData->GetNumPoints(), ExpectedNumElementsChosen);
 
 	return true;
 }
@@ -346,8 +347,8 @@ bool FPCGRandomChoiceTest_MultiData_SameSeed::RunTest(const FString& Parameters)
 	Settings->FixedNumber = ExpectedNumElementsChosen;
 
 	TUniquePtr<FPCGContext> Context = TestData.InitializeTestContext();
-	const UPCGPointData* FirstInputPointData = PCGRandomChoiceTest::CreateInputPointData(Context.Get(), NumOfPoints);
-	const UPCGPointData* SecondInputPointData = PCGRandomChoiceTest::CreateInputPointData(Context.Get(), NumOfPoints);
+	const UPCGBasePointData* FirstInputPointData = PCGRandomChoiceTest::CreateInputPointData(Context.Get(), NumOfPoints);
+	const UPCGBasePointData* SecondInputPointData = PCGRandomChoiceTest::CreateInputPointData(Context.Get(), NumOfPoints);
 
 	FPCGElementPtr TestElement = TestData.Settings->GetElement();
 	while (!TestElement->Execute(Context.Get())) {}
@@ -355,30 +356,32 @@ bool FPCGRandomChoiceTest_MultiData_SameSeed::RunTest(const FString& Parameters)
 	TArray<FPCGTaggedData> ChosenOutputTagged = Context->OutputData.GetInputsByPin(PCGRandomChoiceConstants::ChosenEntriesLabel);
 	TArray<FPCGTaggedData> DiscardedOutputTagged = Context->OutputData.GetInputsByPin(PCGRandomChoiceConstants::DiscardedEntriesLabel);
 
-	const UPCGPointData* FirstChosenOutputData = ChosenOutputTagged.Num() == 2 ? Cast<const UPCGPointData>(ChosenOutputTagged[0].Data) : nullptr;
-	const UPCGPointData* FirstDiscardedOutputData = DiscardedOutputTagged.Num() == 2 ? Cast<const UPCGPointData>(DiscardedOutputTagged[0].Data) : nullptr;
+	const UPCGBasePointData* FirstChosenOutputData = ChosenOutputTagged.Num() == 2 ? Cast<const UPCGBasePointData>(ChosenOutputTagged[0].Data) : nullptr;
+	const UPCGBasePointData* FirstDiscardedOutputData = DiscardedOutputTagged.Num() == 2 ? Cast<const UPCGBasePointData>(DiscardedOutputTagged[0].Data) : nullptr;
 
-	const UPCGPointData* SecondChosenOutputData = ChosenOutputTagged.Num() == 2 ? Cast<const UPCGPointData>(ChosenOutputTagged[1].Data) : nullptr;
-	const UPCGPointData* SecondDiscardedOutputData = DiscardedOutputTagged.Num() == 2 ? Cast<const UPCGPointData>(DiscardedOutputTagged[1].Data) : nullptr;
+	const UPCGBasePointData* SecondChosenOutputData = ChosenOutputTagged.Num() == 2 ? Cast<const UPCGBasePointData>(ChosenOutputTagged[1].Data) : nullptr;
+	const UPCGBasePointData* SecondDiscardedOutputData = DiscardedOutputTagged.Num() == 2 ? Cast<const UPCGBasePointData>(DiscardedOutputTagged[1].Data) : nullptr;
 
 	UTEST_TRUE("There is 2 point data in chosen points", FirstChosenOutputData && SecondChosenOutputData);
 	UTEST_TRUE("There is 2 point data in discarded points", FirstDiscardedOutputData && SecondDiscardedOutputData);
 
 	check(FirstChosenOutputData && SecondChosenOutputData && FirstDiscardedOutputData && SecondDiscardedOutputData);
 
-	UTEST_EQUAL("There is the right number of points in first chosen", FirstChosenOutputData->GetPoints().Num(), ExpectedNumElementsChosen);
-	UTEST_EQUAL("There is the right number of points in second chosen", SecondChosenOutputData->GetPoints().Num(), ExpectedNumElementsChosen);
+	UTEST_EQUAL("There is the right number of points in first chosen", FirstChosenOutputData->GetNumPoints(), ExpectedNumElementsChosen);
+	UTEST_EQUAL("There is the right number of points in second chosen", SecondChosenOutputData->GetNumPoints(), ExpectedNumElementsChosen);
 
-	UTEST_EQUAL("There is the right number of points in first discarded", FirstDiscardedOutputData->GetPoints().Num(), ExpectedNumElementsDiscarded);
-	UTEST_EQUAL("There is the right number of points in second discarded", SecondDiscardedOutputData->GetPoints().Num(), ExpectedNumElementsDiscarded);
+	UTEST_EQUAL("There is the right number of points in first discarded", FirstDiscardedOutputData->GetNumPoints(), ExpectedNumElementsDiscarded);
+	UTEST_EQUAL("There is the right number of points in second discarded", SecondDiscardedOutputData->GetNumPoints(), ExpectedNumElementsDiscarded);
 
 	UTEST_TRUE("All points are there and in the right order for first data", PCGRandomChoiceTest::VerifyAllPointsThere(NumOfPoints, FirstChosenOutputData, FirstDiscardedOutputData));
 	UTEST_TRUE("All points are there and in the right order for second data", PCGRandomChoiceTest::VerifyAllPointsThere(NumOfPoints, SecondChosenOutputData, SecondDiscardedOutputData));
 
 	// If the point seeds are the same, chosen points should be exactly the same
+	const TConstPCGValueRange<float> FirstDensityRange = FirstChosenOutputData->GetConstDensityValueRange();
+	const TConstPCGValueRange<float> SecondDensityRange = SecondChosenOutputData->GetConstDensityValueRange();
 	for (int32 i = 0; i < ExpectedNumElementsChosen; ++i)
 	{
-		UTEST_TRUE(*FString::Printf(TEXT("Point %d is the same for both chosen inputs"), i), FirstChosenOutputData->GetPoints()[i].Density == SecondChosenOutputData->GetPoints()[i].Density);
+		UTEST_TRUE(*FString::Printf(TEXT("Point %d is the same for both chosen inputs"), i), FirstDensityRange[i] == SecondDensityRange[i]);
 	}
 
 	return true;
@@ -398,8 +401,8 @@ bool FPCGRandomChoiceTest_MultiData_DifferentSeed::RunTest(const FString& Parame
 	Settings->FixedNumber = ExpectedNumElementsChosen;
 
 	TUniquePtr<FPCGContext> Context = TestData.InitializeTestContext();
-	const UPCGPointData* FirstInputPointData = PCGRandomChoiceTest::CreateInputPointData(Context.Get(), NumOfPoints);
-	const UPCGPointData* SecondInputPointData = PCGRandomChoiceTest::CreateInputPointData(Context.Get(), NumOfPoints, 4653465);
+	const UPCGBasePointData* FirstInputPointData = PCGRandomChoiceTest::CreateInputPointData(Context.Get(), NumOfPoints);
+	const UPCGBasePointData* SecondInputPointData = PCGRandomChoiceTest::CreateInputPointData(Context.Get(), NumOfPoints, 4653465);
 
 	FPCGElementPtr TestElement = TestData.Settings->GetElement();
 	while (!TestElement->Execute(Context.Get())) {}
@@ -407,31 +410,33 @@ bool FPCGRandomChoiceTest_MultiData_DifferentSeed::RunTest(const FString& Parame
 	TArray<FPCGTaggedData> ChosenOutputTagged = Context->OutputData.GetInputsByPin(PCGRandomChoiceConstants::ChosenEntriesLabel);
 	TArray<FPCGTaggedData> DiscardedOutputTagged = Context->OutputData.GetInputsByPin(PCGRandomChoiceConstants::DiscardedEntriesLabel);
 
-	const UPCGPointData* FirstChosenOutputData = ChosenOutputTagged.Num() == 2 ? Cast<const UPCGPointData>(ChosenOutputTagged[0].Data) : nullptr;
-	const UPCGPointData* FirstDiscardedOutputData = DiscardedOutputTagged.Num() == 2 ? Cast<const UPCGPointData>(DiscardedOutputTagged[0].Data) : nullptr;
+	const UPCGBasePointData* FirstChosenOutputData = ChosenOutputTagged.Num() == 2 ? Cast<const UPCGBasePointData>(ChosenOutputTagged[0].Data) : nullptr;
+	const UPCGBasePointData* FirstDiscardedOutputData = DiscardedOutputTagged.Num() == 2 ? Cast<const UPCGBasePointData>(DiscardedOutputTagged[0].Data) : nullptr;
 
-	const UPCGPointData* SecondChosenOutputData = ChosenOutputTagged.Num() == 2 ? Cast<const UPCGPointData>(ChosenOutputTagged[1].Data) : nullptr;
-	const UPCGPointData* SecondDiscardedOutputData = DiscardedOutputTagged.Num() == 2 ? Cast<const UPCGPointData>(DiscardedOutputTagged[1].Data) : nullptr;
+	const UPCGBasePointData* SecondChosenOutputData = ChosenOutputTagged.Num() == 2 ? Cast<const UPCGBasePointData>(ChosenOutputTagged[1].Data) : nullptr;
+	const UPCGBasePointData* SecondDiscardedOutputData = DiscardedOutputTagged.Num() == 2 ? Cast<const UPCGBasePointData>(DiscardedOutputTagged[1].Data) : nullptr;
 
 	UTEST_TRUE("There is 2 point data in chosen points", FirstChosenOutputData && SecondChosenOutputData);
 	UTEST_TRUE("There is 2 point data in discarded points", FirstDiscardedOutputData && SecondDiscardedOutputData);
 
 	check(FirstChosenOutputData && SecondChosenOutputData && FirstDiscardedOutputData && SecondDiscardedOutputData);
 
-	UTEST_EQUAL("There is the right number of points in first chosen", FirstChosenOutputData->GetPoints().Num(), ExpectedNumElementsChosen);
-	UTEST_EQUAL("There is the right number of points in second chosen", SecondChosenOutputData->GetPoints().Num(), ExpectedNumElementsChosen);
+	UTEST_EQUAL("There is the right number of points in first chosen", FirstChosenOutputData->GetNumPoints(), ExpectedNumElementsChosen);
+	UTEST_EQUAL("There is the right number of points in second chosen", SecondChosenOutputData->GetNumPoints(), ExpectedNumElementsChosen);
 
-	UTEST_EQUAL("There is the right number of points in first discarded", FirstDiscardedOutputData->GetPoints().Num(), ExpectedNumElementsDiscarded);
-	UTEST_EQUAL("There is the right number of points in second discarded", SecondDiscardedOutputData->GetPoints().Num(), ExpectedNumElementsDiscarded);
+	UTEST_EQUAL("There is the right number of points in first discarded", FirstDiscardedOutputData->GetNumPoints(), ExpectedNumElementsDiscarded);
+	UTEST_EQUAL("There is the right number of points in second discarded", SecondDiscardedOutputData->GetNumPoints(), ExpectedNumElementsDiscarded);
 
 	UTEST_TRUE("All points are there and in the right order for first data", PCGRandomChoiceTest::VerifyAllPointsThere(NumOfPoints, FirstChosenOutputData, FirstDiscardedOutputData));
 	UTEST_TRUE("All points are there and in the right order for second data", PCGRandomChoiceTest::VerifyAllPointsThere(NumOfPoints, SecondChosenOutputData, SecondDiscardedOutputData));
 
 	// If point seeds are different, chosen points should be different. Verify that at least one point is different.
 	bool bIsDifferent = false;
+	const TConstPCGValueRange<float> FirstDensityRange = FirstChosenOutputData->GetConstDensityValueRange();
+	const TConstPCGValueRange<float> SecondDensityRange = SecondChosenOutputData->GetConstDensityValueRange();
 	for (int32 i = 0; i < ExpectedNumElementsChosen && !bIsDifferent; ++i)
 	{
-		bIsDifferent = FirstChosenOutputData->GetPoints()[i].Density != SecondChosenOutputData->GetPoints()[i].Density;
+		bIsDifferent = FirstDensityRange[i] != SecondDensityRange[i];
 	}
 
 	UTEST_TRUE("Chosen points are different", bIsDifferent);
diff --git a/Engine/Plugins/PCG/Source/PCG/Private/Tests/Elements/PCGSampleTextureTest.cpp b/Engine/Plugins/PCG/Source/PCG/Private/Tests/Elements/PCGSampleTextureTest.cpp
--- a/Engine/Plugins/PCG/Source/PCG/Private/Tests/Elements/PCGSampleTextureTest.cpp
+++ b/Engine/Plugins/PCG/Source/PCG/Private/Tests/Elements/PCGSampleTextureTest.cpp
@@ -124,13 +124,11 @@ protected:
 
 		UTEST_EQUAL("Output count", Outputs.Num(), 1);
 
-		const UPCGPointData* OutPointData = Cast<UPCGPointData>(Outputs[0].Data);
+		const UPCGBasePointData* OutPointData = Cast<UPCGBasePointData>(Outputs[0].Data);
 
 		UTEST_NOT_NULL("Output point data", OutPointData);
 
-		const TArray<FPCGPoint>& OutPoints = OutPointData->GetPoints();
-
-		UTEST_EQUAL("Output point count", OutPoints.Num(), 16);
+		UTEST_EQUAL("Output point count", OutPointData->GetNumPoints(), 16);
 
 		const FPCGMetadataAttribute<FVector2D>* OutputUV = OutPointData->Metadata->GetConstTypedAttribute<FVector2D>(Parameters.UVAttribute);
 		UTEST_NOT_NULL("UV Attribute exists", OutputUV);
@@ -139,12 +137,13 @@ protected:
 
 		UTEST_TRUE("Output Value Key is not default value key.", ValueKey != PCGDefaultValueKey)
 
-		for (const FPCGPoint& OutPoint : OutPoints)
+		const FConstPCGPointValueRanges OutRanges(OutPointData);
+		for (int32 PointIndex = 0; PointIndex < OutPointData->GetNumPoints(); ++PointIndex)
 		{
-			if (OutputUV->GetValueKey(OutPoint.MetadataEntry) == ValueKey)
+			if (OutputUV->GetValueKey(OutRanges.MetadataEntryRange[PointIndex]) == ValueKey)
 			{
-				UTEST_EQUAL("Point Color", OutPoint.Color, Parameters.ExpectedColor);
-				UTEST_EQUAL("Point Density", OutPoint.Density, Parameters.ExpectedDensity);
+				UTEST_EQUAL("Point Color", OutRanges.ColorRange[PointIndex], Parameters.ExpectedColor);
+				UTEST_EQUAL("Point Density", OutRanges.DensityRange[PointIndex], Parameters.ExpectedDensity);
 			}
 		}
 
diff --git a/Engine/Plugins/PCG/Source/PCG/Private/Tests/Elements/PCGSplitPointsTest.cpp b/Engine/Plugins/PCG/Source/PCG/Private/Tests/Elements/PCGSplitPointsTest.cpp
--- a/Engine/Plugins/PCG/Source/PCG/Private/Tests/Elements/PCGSplitPointsTest.cpp
+++ b/Engine/Plugins/PCG/Source/PCG/Private/Tests/Elements/PCGSplitPointsTest.cpp
@@ -49,28 +49,26 @@ protected:
 		UTEST_EQUAL("Output count", Outputs.Num(), 2);
 
 		// checking the PointA portion
-		const UPCGPointData* OutPointDataA = Cast<UPCGPointData>(Outputs[0].Data);
+		const UPCGBasePointData* OutPointDataA = Cast<UPCGBasePointData>(Outputs[0].Data);
 
 		UTEST_NOT_NULL("OutputA point data", OutPointDataA);
 
-		const TArray<FPCGPoint>& OutPointsA = OutPointDataA->GetPoints();
+		UTEST_EQUAL("OutputA point count", OutPointDataA->GetNumPoints(), 1);
 
-		UTEST_EQUAL("OutputA point count", OutPointsA.Num(), 1);
-
-		UTEST_EQUAL("OutputA Bounds Min", OutPointsA[0].BoundsMin, Parameters.ExpectedABoundsMin);
-		UTEST_EQUAL("OutputA Bounds Max", OutPointsA[0].BoundsMax, Parameters.ExpectedABoundsMax);
+		const FConstPCGPointValueRanges RangesA(OutPointDataA);
+		UTEST_EQUAL("OutputA Bounds Min", RangesA.BoundsMinRange[0], Parameters.ExpectedABoundsMin);
+		UTEST_EQUAL("OutputA Bounds Max", RangesA.BoundsMaxRange[0], Parameters.ExpectedABoundsMax);
 
 		// checking the PointB portion
-		const UPCGPointData* OutPointDataB = Cast<UPCGPointData>(Outputs[1].Data);
+		const UPCGBasePointData* OutPointDataB = Cast<UPCGBasePointData>(Outputs[1].Data);
 
 		UTEST_NOT_NULL("OutputB point data", OutPointDataB);
 
-		const TArray<FPCGPoint>& OutPointsB = OutPointDataB->GetPoints();
-
-		UTEST_EQUAL("OutputB point count", OutPointsB.Num(), 1);
+		UTEST_EQUAL("OutputB point count", OutPointDataB->GetNumPoints(), 1);
 
-		UTEST_EQUAL("OutputB Bounds Min", OutPointsB[0].BoundsMin, Parameters.ExpectedBBoundsMin);
-		UTEST_EQUAL("OutputB Bounds Max", OutPointsB[0].BoundsMax, Parameters.ExpectedBBoundsMax);
+		const FConstPCGPointValueRanges RangesB(OutPointDataB);
+		UTEST_EQUAL("OutputB Bounds Min", RangesB.BoundsMinRange[0], Parameters.ExpectedBBoundsMin);
+		UTEST_EQUAL("OutputB Bounds Max", RangesB.BoundsMaxRange[0], Parameters.ExpectedBBoundsMax);
 
 		return true;
 	}
diff --git a/Engine/Plugins/PCG/Source/PCG/Private/Tests/PCGTestsCommon.cpp b/Engine/Plugins/PCG/Source/PCG/Private/Tests/PCGTestsCommon.cpp
--- a/Engine/Plugins/PCG/Source/PCG/Private/Tests/PCGTestsCommon.cpp
+++ b/Engine/Plugins/PCG/Source/PCG/Private/Tests/PCGTestsCommon.cpp
@@ -166,6 +166,18 @@ namespace PCGTestsCommon
 		return NewObject<UPCGParamData>();
 	}
 
+	UPCGBasePointData* CreateEmptyBasePointData()
+	{
+		if (CVarPCGEnablePointArrayData.GetValueOnAnyThread())
+		{
+			return CreateEmptyPointData<UPCGPointArrayData>();
+		}
+		else
+		{
+			return CreateEmptyPointData<UPCGPointData>();
+		}
+	}
+
 	TArray<FPCGDataCollection> GenerateAllowedData(const FPCGPinProperties& PinProperties)
 	{
 		TArray<FPCGDataCollection> Data;
diff --git a/Engine/Plugins/PCG/Source/PCG/Public/Data/PCGSplineInteriorSurfaceData.h b/Engine/Plugins/PCG/Source/PCG/Public/Data/PCGSplineInteriorSurfaceData.h
--- a/Engine/Plugins/PCG/Source/PCG/Public/Data/PCGSplineInteriorSurfaceData.h
+++ b/Engine/Plugins/PCG/Source/PCG/Public/Data/PCGSplineInteriorSurfaceData.h
@@ -47,6 +47,7 @@ protected:
 public:
 	// ~Begin UPCGSpatialDataWithPointCache interface
 	virtual const UPCGPointData* CreatePointData(FPCGContext* Context) const override;
+	virtual const UPCGPointArrayData* CreatePointArrayData(FPCGContext* Context, const FBox& InBounds) const override;
 	// ~End UPCGConcreteDataWithPointCache interface
 
 protected:
@@ -56,6 +57,8 @@ protected:
 	/** True if the given location falls in the top-down projection of the polygon given by our cached spline points. */
 	bool PointInsidePolygon(const FTransform& InTransform, const FBox& InBounds) const;
 
+	const UPCGBasePointData* CreateBasePointData(FPCGContext* Context, TSubclassOf<UPCGBasePointData> PointDataClass) const;
+
 protected:
 	/** Minimal data needed to replicate the behavior from USplineComponent. */
 	UPROPERTY()
diff --git a/Engine/Plugins/PCG/Source/PCG/Public/Data/PCGTextureData.h b/Engine/Plugins/PCG/Source/PCG/Public/Data/PCGTextureData.h
--- a/Engine/Plugins/PCG/Source/PCG/Public/Data/PCGTextureData.h
+++ b/Engine/Plugins/PCG/Source/PCG/Public/Data/PCGTextureData.h
@@ -83,6 +83,7 @@ public:
 
 	//~Begin UPCGSpatialDataWithPointCache interface
 	virtual const UPCGPointData* CreatePointData(FPCGContext* Context) const override;
+	virtual const UPCGPointArrayData* CreatePointArrayData(FPCGContext* Context, const FBox& InBounds) const override;
 	//~End UPCGSpatialDataWithPointCache interface
 
 	/** Sample using a local space 'UV' position. */
@@ -96,6 +97,8 @@ public:
 	virtual int GetTextureSlice() const { return 0; }
 
 protected:
+	const UPCGBasePointData* CreateBasePointData(FPCGContext* Context, TSubclassOf<UPCGBasePointData> PointDataClass) const;
+
 	PRAGMA_DISABLE_DEPRECATION_WARNINGS
 	UFUNCTION(BlueprintGetter, meta = (BlueprintInternalUseOnly = "true"))
 	EPCGTextureDensityFunction GetDensityFunctionEquivalent() const;
diff --git a/Engine/Plugins/PCG/Source/PCG/Public/Data/PCGWorldData.h b/Engine/Plugins/PCG/Source/PCG/Public/Data/PCGWorldData.h
--- a/Engine/Plugins/PCG/Source/PCG/Public/Data/PCGWorldData.h
+++ b/Engine/Plugins/PCG/Source/PCG/Public/Data/PCGWorldData.h
@@ -361,11 +361,13 @@ protected:
 	virtual UPCGSpatialData* CopyInternal(FPCGContext* Context) const override;
 	//~End UPCGSpatialData interface
 
+	const UPCGBasePointData* CreateBasePointData(FPCGContext* Context, const FBox& InBounds, TSubclassOf<UPCGBasePointData> PointDataClass) const;
 public:
 	// ~Begin UPCGSpatialDataWithPointCache interface
 	virtual bool SupportsBoundedPointData() const { return true; }
 	virtual const UPCGPointData* CreatePointData(FPCGContext* Context) const override { return CreatePointData(Context, FBox(EForceInit::ForceInit)); }
 	virtual const UPCGPointData* CreatePointData(FPCGContext* Context, const FBox& InBounds) const override;
+	virtual const UPCGPointArrayData* CreatePointArrayData(FPCGContext* Context, const FBox& InBounds) const override;
 	// ~End UPCGSpatialDataWithPointCache interface
 
 	UPROPERTY()
diff --git a/Engine/Plugins/PCG/Source/PCG/Public/Elements/ControlFlow/PCGQualitySelect.h b/Engine/Plugins/PCG/Source/PCG/Public/Elements/ControlFlow/PCGQualitySelect.h
--- a/Engine/Plugins/PCG/Source/PCG/Public/Elements/ControlFlow/PCGQualitySelect.h
+++ b/Engine/Plugins/PCG/Source/PCG/Public/Elements/ControlFlow/PCGQualitySelect.h
@@ -56,4 +56,5 @@ public:
 
 protected:
 	virtual bool ExecuteInternal(FPCGContext* Context) const override;
+	virtual bool SupportsBasePointDataInputs(FPCGContext* InContext) const override { return true; }
 };
diff --git a/Engine/Plugins/PCG/Source/PCG/Public/Elements/Grammar/PCGSelectGrammar.h b/Engine/Plugins/PCG/Source/PCG/Public/Elements/Grammar/PCGSelectGrammar.h
--- a/Engine/Plugins/PCG/Source/PCG/Public/Elements/Grammar/PCGSelectGrammar.h
+++ b/Engine/Plugins/PCG/Source/PCG/Public/Elements/Grammar/PCGSelectGrammar.h
@@ -146,4 +146,5 @@ class FPCGSelectGrammarElement : public FPCGSubdivisionBaseElement
 {
 protected:
 	virtual bool ExecuteInternal(FPCGContext* InContext) const override;
+	virtual bool SupportsBasePointDataInputs(FPCGContext* InContext) const override { return true; }
 };
\ No newline at end of file
diff --git a/Engine/Plugins/PCG/Source/PCG/Public/Elements/Grammar/PCGSubdivideSegment.h b/Engine/Plugins/PCG/Source/PCG/Public/Elements/Grammar/PCGSubdivideSegment.h
--- a/Engine/Plugins/PCG/Source/PCG/Public/Elements/Grammar/PCGSubdivideSegment.h
+++ b/Engine/Plugins/PCG/Source/PCG/Public/Elements/Grammar/PCGSubdivideSegment.h
@@ -86,4 +86,5 @@ public:
 	friend class PCGSubdivideSegmentHelpers;
 protected:
 	virtual bool ExecuteInternal(FPCGContext* InContext) const override;
+	virtual bool SupportsBasePointDataInputs(FPCGContext* InContext) const override { return true; }
 };
diff --git a/Engine/Plugins/PCG/Source/PCG/Public/Elements/Grammar/PCGSubdivisionBase.h b/Engine/Plugins/PCG/Source/PCG/Public/Elements/Grammar/PCGSubdivisionBase.h
--- a/Engine/Plugins/PCG/Source/PCG/Public/Elements/Grammar/PCGSubdivisionBase.h
+++ b/Engine/Plugins/PCG/Source/PCG/Public/Elements/Grammar/PCGSubdivisionBase.h
@@ -12,7 +12,7 @@
 #include "PCGSubdivisionBase.generated.h"
 
 class UPCGData;
-class UPCGPointData;
+class UPCGBasePointData;
 
 USTRUCT(BlueprintType)
 struct FPCGSubdivisionSubmodule
@@ -458,6 +458,6 @@ protected:
 	FModuleInfoMap GetModulesInfoMap(FPCGContext* InContext, const FPCGSubdivisionModuleAttributeNames& InSubdivisionModuleAttributeNames, const UPCGParamData*& OutModuleInfoParamData) const;
 	FModuleInfoMap GetModulesInfoMap(FPCGContext* InContext, const UPCGSubdivisionBaseSettings* InSettings, const UPCGParamData*& OutModuleInfoParamData) const;
 	PCGGrammar::FTokenizedGrammar GetTokenizedGrammar(FPCGContext* InContext, const UPCGData* InputData, const UPCGSubdivisionBaseSettings* InSettings, const FModuleInfoMap& InModulesInfo, double& OutMinSize) const;
-	TMap<FString, PCGGrammar::FTokenizedGrammar> GetTokenizedGrammarForPoints(FPCGContext* InContext, const UPCGPointData* InputData, const UPCGSubdivisionBaseSettings* InSettings, const FModuleInfoMap& InModulesInfo, double& OutMinSize) const;
+	TMap<FString, PCGGrammar::FTokenizedGrammar> GetTokenizedGrammarForPoints(FPCGContext* InContext, const UPCGBasePointData* InputData, const UPCGSubdivisionBaseSettings* InSettings, const FModuleInfoMap& InModulesInfo, double& OutMinSize) const;
 	bool MatchAndSetAttributes(const TArray<FPCGTaggedData>& InputData, TArray<FPCGTaggedData>& OutputData, const UPCGParamData* InModuleInfoParamData, const UPCGSubdivisionBaseSettings* InSettings) const;
 };
diff --git a/Engine/Plugins/PCG/Source/PCG/Public/Elements/Metadata/PCGAttributeCast.h b/Engine/Plugins/PCG/Source/PCG/Public/Elements/Metadata/PCGAttributeCast.h
--- a/Engine/Plugins/PCG/Source/PCG/Public/Elements/Metadata/PCGAttributeCast.h
+++ b/Engine/Plugins/PCG/Source/PCG/Public/Elements/Metadata/PCGAttributeCast.h
@@ -52,4 +52,5 @@ class FPCGAttributeCastElement : public IPCGElement
 protected:
 	virtual bool ExecuteInternal(FPCGContext* Context) const override;
 	virtual EPCGElementExecutionLoopMode ExecutionLoopMode(const UPCGSettings* Settings) const override { return EPCGElementExecutionLoopMode::SinglePrimaryPin; }
+	virtual bool SupportsBasePointDataInputs(FPCGContext* InContext) const override { return true; }
 };
diff --git a/Engine/Plugins/PCG/Source/PCG/Public/Elements/PCGAddComponent.h b/Engine/Plugins/PCG/Source/PCG/Public/Elements/PCGAddComponent.h
--- a/Engine/Plugins/PCG/Source/PCG/Public/Elements/PCGAddComponent.h
+++ b/Engine/Plugins/PCG/Source/PCG/Public/Elements/PCGAddComponent.h
@@ -118,6 +118,7 @@ public:
 
 protected:
 	virtual bool ExecuteInternal(FPCGContext* Context) const override;
+	virtual bool SupportsBasePointDataInputs(FPCGContext* InContext) const override { return true; }
 
 private:
 	void AddComponents(FPCGAddComponentContext* Context,
diff --git a/Engine/Plugins/PCG/Source/PCG/Public/Elements/PCGApplyOnActor.h b/Engine/Plugins/PCG/Source/PCG/Public/Elements/PCGApplyOnActor.h
--- a/Engine/Plugins/PCG/Source/PCG/Public/Elements/PCGApplyOnActor.h
+++ b/Engine/Plugins/PCG/Source/PCG/Public/Elements/PCGApplyOnActor.h
@@ -74,4 +74,5 @@ public:
 protected:
 	virtual bool PrepareDataInternal(FPCGContext* Context) const override;
 	virtual bool ExecuteInternal(FPCGContext* Context) const override;
+	virtual bool SupportsBasePointDataInputs(FPCGContext* InContext) const override { return true; }
 };
diff --git a/Engine/Plugins/PCG/Source/PCG/Public/Elements/PCGAttractElement.h b/Engine/Plugins/PCG/Source/PCG/Public/Elements/PCGAttractElement.h
--- a/Engine/Plugins/PCG/Source/PCG/Public/Elements/PCGAttractElement.h
+++ b/Engine/Plugins/PCG/Source/PCG/Public/Elements/PCGAttractElement.h
@@ -10,6 +10,7 @@
 
 #include "PCGAttractElement.generated.h"
 
+class UPCGBasePointData;
 class UPCGPointData;
 
 UENUM(BlueprintType)
@@ -96,8 +97,21 @@ namespace PCGAttractElement
 {
 	struct FAttractState
 	{
+		// Note: We need to explicitly disable warnings on these constructors/operators for clang to be happy with deprecated variables
+		PRAGMA_DISABLE_DEPRECATION_WARNINGS
+		FAttractState() = default;
+		~FAttractState() = default;
+		PRAGMA_ENABLE_DEPRECATION_WARNINGS
+
+		const UPCGBasePointData* TargetData = nullptr;
+		const UPCGBasePointData* SourceData = nullptr;
+		UPCGBasePointData* OutputData = nullptr;
+
+		UE_DEPRECATED(5.6, "Use TargetData instead")
 		const UPCGPointData* TargetPointData = nullptr;
+		UE_DEPRECATED(5.6, "Use SourceData instead")
 		const UPCGPointData* SourcePointData = nullptr;
+		UE_DEPRECATED(5.6, "Use OutputData instead")
 		UPCGPointData* OutPointData = nullptr;
 
 		TArray<int32> Mapping;
@@ -115,6 +129,8 @@ namespace PCGAttractElement
 			TUniquePtr<const IPCGAttributeAccessorKeys> SourceKeys;
 			TUniquePtr<const IPCGAttributeAccessor> TargetAccessor;
 			TUniquePtr<const IPCGAttributeAccessorKeys> TargetKeys;
+
+			FPCGAttributePropertyOutputSelector OutputSelector;
 			TUniquePtr<IPCGAttributeAccessor> OutputAccessor;
 		};
 
@@ -138,4 +154,5 @@ class FPCGAttractElement : public TPCGTimeSlicedElementBase<PCGAttractElement::F
 protected:
 	virtual bool PrepareDataInternal(FPCGContext* InContext) const override;
 	virtual bool ExecuteInternal(FPCGContext* InContext) const override;
+	virtual bool SupportsBasePointDataInputs(FPCGContext* InContext) const override { return true; }
 };
diff --git a/Engine/Plugins/PCG/Source/PCG/Public/Elements/PCGAttributeRemoveDuplicates.h b/Engine/Plugins/PCG/Source/PCG/Public/Elements/PCGAttributeRemoveDuplicates.h
--- a/Engine/Plugins/PCG/Source/PCG/Public/Elements/PCGAttributeRemoveDuplicates.h
+++ b/Engine/Plugins/PCG/Source/PCG/Public/Elements/PCGAttributeRemoveDuplicates.h
@@ -45,5 +45,6 @@ class FPCGAttributeRemoveDuplicatesElement : public IPCGElement
 {
 protected:
 	virtual bool ExecuteInternal(FPCGContext* Context) const override;
+	virtual bool SupportsBasePointDataInputs(FPCGContext* InContext) const override { return true; }
 	virtual EPCGElementExecutionLoopMode ExecutionLoopMode(const UPCGSettings* Settings) const override { return EPCGElementExecutionLoopMode::SinglePrimaryPin; }
 };
diff --git a/Engine/Plugins/PCG/Source/PCG/Public/Elements/PCGAttributeSelectElement.h b/Engine/Plugins/PCG/Source/PCG/Public/Elements/PCGAttributeSelectElement.h
--- a/Engine/Plugins/PCG/Source/PCG/Public/Elements/PCGAttributeSelectElement.h
+++ b/Engine/Plugins/PCG/Source/PCG/Public/Elements/PCGAttributeSelectElement.h
@@ -103,4 +103,5 @@ class FPCGAttributeSelectElement : public IPCGElement
 {
 protected:
 	virtual bool ExecuteInternal(FPCGContext* Context) const override;
+	virtual bool SupportsBasePointDataInputs(FPCGContext* InContext) const override { return true; }
 };
diff --git a/Engine/Plugins/PCG/Source/PCG/Public/Elements/PCGBlurElement.h b/Engine/Plugins/PCG/Source/PCG/Public/Elements/PCGBlurElement.h
--- a/Engine/Plugins/PCG/Source/PCG/Public/Elements/PCGBlurElement.h
+++ b/Engine/Plugins/PCG/Source/PCG/Public/Elements/PCGBlurElement.h
@@ -75,12 +75,22 @@ public:
 
 struct FPCGBlurIterState
 {
+	// Note: We need to explicitly disable warnings on these constructors/operators for clang to be happy with deprecated variables
+	PRAGMA_DISABLE_DEPRECATION_WARNINGS
+	FPCGBlurIterState() = default;
+	~FPCGBlurIterState() = default;
+	PRAGMA_ENABLE_DEPRECATION_WARNINGS
+
 	TUniquePtr<const IPCGAttributeAccessor> InputAccessor;
 	TUniquePtr<IPCGAttributeAccessor> OutputAccessor;
 	TUniquePtr<const IPCGAttributeAccessorKeys> InputKeys;
 	TUniquePtr<IPCGAttributeAccessorKeys> OutputKeys;
 
+	UE_DEPRECATED(5.6, "Use InputData instead")
 	const UPCGPointData* InputPointData = nullptr;
+
+	const UPCGBasePointData* InputData = nullptr;
+	
 	// Need two working buffers as we are working in parallel. One for read, one for write.
 	TArray<uint8> WorkingBuffer1;
 	TArray<uint8> WorkingBuffer2;
@@ -93,5 +103,6 @@ protected:
 	virtual bool PrepareDataInternal(FPCGContext* InContext) const override;
 	virtual bool ExecuteInternal(FPCGContext* InContext) const override;
 	virtual EPCGElementExecutionLoopMode ExecutionLoopMode(const UPCGSettings* Settings) const override { return EPCGElementExecutionLoopMode::SinglePrimaryPin; }
+	virtual bool SupportsBasePointDataInputs(FPCGContext* InContext) const override { return true; }
 };
 
diff --git a/Engine/Plugins/PCG/Source/PCG/Public/Elements/PCGBoundsFromMesh.h b/Engine/Plugins/PCG/Source/PCG/Public/Elements/PCGBoundsFromMesh.h
--- a/Engine/Plugins/PCG/Source/PCG/Public/Elements/PCGBoundsFromMesh.h
+++ b/Engine/Plugins/PCG/Source/PCG/Public/Elements/PCGBoundsFromMesh.h
@@ -77,5 +77,6 @@ public:
 protected:
 	virtual bool PrepareDataInternal(FPCGContext* Context) const override;
 	virtual bool ExecuteInternal(FPCGContext* Context) const override;
+	virtual bool SupportsBasePointDataInputs(FPCGContext* InContext) const override { return true; }
 	virtual EPCGElementExecutionLoopMode ExecutionLoopMode(const UPCGSettings* Settings) const override { return EPCGElementExecutionLoopMode::SinglePrimaryPin; }
 };
diff --git a/Engine/Plugins/PCG/Source/PCG/Public/Elements/PCGClusterElement.h b/Engine/Plugins/PCG/Source/PCG/Public/Elements/PCGClusterElement.h
--- a/Engine/Plugins/PCG/Source/PCG/Public/Elements/PCGClusterElement.h
+++ b/Engine/Plugins/PCG/Source/PCG/Public/Elements/PCGClusterElement.h
@@ -10,6 +10,7 @@
 #include "PCGClusterElement.generated.h"
 
 class UPCGPointData;
+class UPCGBasePointData;
 
 UENUM()
 enum class EPCGClusterAlgorithm
@@ -82,7 +83,17 @@ namespace PCGClusterElement
 {
 	struct FClusteringData
 	{
+		// Note: We need to explicitly disable warnings on these constructors/operators for clang to be happy with deprecated variables
+		PRAGMA_DISABLE_DEPRECATION_WARNINGS
+		FClusteringData() = default;
+		~FClusteringData() = default;
+		PRAGMA_ENABLE_DEPRECATION_WARNINGS
+
+		UE_DEPRECATED(5.6, "Use OutputData instead")
 		UPCGPointData* OutPointData = nullptr;
+
+		UPCGBasePointData* OutputData = nullptr;
+
 		TUniquePtr<IPCGAttributeAccessor> Accessor;
 		TUniquePtr<IPCGAttributeAccessorKeys> Keys;
 
@@ -164,4 +175,5 @@ class FPCGClusterElement : public TPCGTimeSlicedElementBase<PCGClusterElement::F
 protected:
 	virtual bool PrepareDataInternal(FPCGContext* InContext) const override;
 	virtual bool ExecuteInternal(FPCGContext* InContext) const override;
+	virtual bool SupportsBasePointDataInputs(FPCGContext* InContext) const override { return true; }
 };
diff --git a/Engine/Plugins/PCG/Source/PCG/Public/Elements/PCGCombinePoints.h b/Engine/Plugins/PCG/Source/PCG/Public/Elements/PCGCombinePoints.h
--- a/Engine/Plugins/PCG/Source/PCG/Public/Elements/PCGCombinePoints.h
+++ b/Engine/Plugins/PCG/Source/PCG/Public/Elements/PCGCombinePoints.h
@@ -49,5 +49,6 @@ class FPCGCombinePointsElement : public IPCGElement
 {
 protected:
 	virtual bool ExecuteInternal(FPCGContext* Context) const override;
+	virtual bool SupportsBasePointDataInputs(FPCGContext* InContext) const override { return true; }
 	virtual EPCGElementExecutionLoopMode ExecutionLoopMode(const UPCGSettings* Settings) const override { return EPCGElementExecutionLoopMode::SinglePrimaryPin; }
 };
\ No newline at end of file
diff --git a/Engine/Plugins/PCG/Source/PCG/Public/Elements/PCGConvexHull2D.h b/Engine/Plugins/PCG/Source/PCG/Public/Elements/PCGConvexHull2D.h
--- a/Engine/Plugins/PCG/Source/PCG/Public/Elements/PCGConvexHull2D.h
+++ b/Engine/Plugins/PCG/Source/PCG/Public/Elements/PCGConvexHull2D.h
@@ -36,4 +36,5 @@ class FPCGConvexHull2DElement : public IPCGElement
 protected:
 	virtual bool ExecuteInternal(FPCGContext* Context) const override;
 	virtual EPCGElementExecutionLoopMode ExecutionLoopMode(const UPCGSettings* Settings) const override { return EPCGElementExecutionLoopMode::SinglePrimaryPin; }
+	virtual bool SupportsBasePointDataInputs(FPCGContext* InContext) const override { return true; }
 };
\ No newline at end of file
diff --git a/Engine/Plugins/PCG/Source/PCG/Public/Elements/PCGCreateCollisionData.h b/Engine/Plugins/PCG/Source/PCG/Public/Elements/PCGCreateCollisionData.h
--- a/Engine/Plugins/PCG/Source/PCG/Public/Elements/PCGCreateCollisionData.h
+++ b/Engine/Plugins/PCG/Source/PCG/Public/Elements/PCGCreateCollisionData.h
@@ -57,16 +57,17 @@ private:
 
 struct FPCGCreateCollisionContext : public FPCGContext, public IPCGAsyncLoadingContext
 {
-	~FPCGCreateCollisionContext();
-
 	struct InputMeshData
 	{
 		int InputIndex = INDEX_NONE;
 		TArray<FSoftObjectPath> MeshPaths;
-		UPCGCollisionWrapperData* Data = nullptr;
+		TObjectPtr<UPCGCollisionWrapperData> Data = nullptr;
 	};
 
 	TArray<InputMeshData> PerInputData;
+
+protected:
+	virtual void AddExtraStructReferencedObjects(FReferenceCollector& Collector) override;
 };
 
 class FPCGCreateCollisionDataElement : public IPCGElementWithCustomContext<FPCGCreateCollisionContext>
@@ -78,5 +79,6 @@ public:
 protected:
 	virtual bool PrepareDataInternal(FPCGContext* Context) const override;
 	virtual bool ExecuteInternal(FPCGContext* Context) const override;
+	virtual bool SupportsBasePointDataInputs(FPCGContext* InContext) const override { return true; }
 	virtual EPCGElementExecutionLoopMode ExecutionLoopMode(const UPCGSettings* Settings) const override { return EPCGElementExecutionLoopMode::SinglePrimaryPin; }
 };
\ No newline at end of file
diff --git a/Engine/Plugins/PCG/Source/PCG/Public/Elements/PCGCreateSpline.h b/Engine/Plugins/PCG/Source/PCG/Public/Elements/PCGCreateSpline.h
--- a/Engine/Plugins/PCG/Source/PCG/Public/Elements/PCGCreateSpline.h
+++ b/Engine/Plugins/PCG/Source/PCG/Public/Elements/PCGCreateSpline.h
@@ -77,5 +77,6 @@ public:
 
 protected:	
 	virtual bool ExecuteInternal(FPCGContext* Context) const override;
+	virtual bool SupportsBasePointDataInputs(FPCGContext* InContext) const override { return true; }
 	virtual EPCGElementExecutionLoopMode ExecutionLoopMode(const UPCGSettings* Settings) const override { return EPCGElementExecutionLoopMode::SinglePrimaryPin; }
 };
\ No newline at end of file
diff --git a/Engine/Plugins/PCG/Source/PCG/Public/Elements/PCGElevationIsolines.h b/Engine/Plugins/PCG/Source/PCG/Public/Elements/PCGElevationIsolines.h
--- a/Engine/Plugins/PCG/Source/PCG/Public/Elements/PCGElevationIsolines.h
+++ b/Engine/Plugins/PCG/Source/PCG/Public/Elements/PCGElevationIsolines.h
@@ -67,6 +67,7 @@ class FPCGElevationIsolinesElement : public IPCGElement
 {
 protected:
 	virtual bool ExecuteInternal(FPCGContext* InContext) const override;
+	virtual bool SupportsBasePointDataInputs(FPCGContext* InContext) const override { return true; }
 	virtual EPCGElementExecutionLoopMode ExecutionLoopMode(const UPCGSettings* Settings) const override { return EPCGElementExecutionLoopMode::SinglePrimaryPin; }
 };
 
diff --git a/Engine/Plugins/PCG/Source/PCG/Public/Elements/PCGFilterElementsByIndex.h b/Engine/Plugins/PCG/Source/PCG/Public/Elements/PCGFilterElementsByIndex.h
--- a/Engine/Plugins/PCG/Source/PCG/Public/Elements/PCGFilterElementsByIndex.h
+++ b/Engine/Plugins/PCG/Source/PCG/Public/Elements/PCGFilterElementsByIndex.h
@@ -57,4 +57,5 @@ class FPCGFilterElementsByIndexElement : public IPCGElement
 {
 protected:
 	virtual bool ExecuteInternal(FPCGContext* InContext) const override;
+	virtual bool SupportsBasePointDataInputs(FPCGContext* InContext) const override { return true; }
 };
diff --git a/Engine/Plugins/PCG/Source/PCG/Public/Elements/PCGGetActorDataLayers.h b/Engine/Plugins/PCG/Source/PCG/Public/Elements/PCGGetActorDataLayers.h
--- a/Engine/Plugins/PCG/Source/PCG/Public/Elements/PCGGetActorDataLayers.h
+++ b/Engine/Plugins/PCG/Source/PCG/Public/Elements/PCGGetActorDataLayers.h
@@ -45,4 +45,5 @@ public:
 	virtual bool CanExecuteOnlyOnMainThread(FPCGContext* Context) const override { return true; }
 protected:
 	virtual bool ExecuteInternal(FPCGContext* Context) const override;
+	virtual bool SupportsBasePointDataInputs(FPCGContext* InContext) const override { return true; }
 };
\ No newline at end of file
diff --git a/Engine/Plugins/PCG/Source/PCG/Public/Elements/PCGGetActorProperty.h b/Engine/Plugins/PCG/Source/PCG/Public/Elements/PCGGetActorProperty.h
--- a/Engine/Plugins/PCG/Source/PCG/Public/Elements/PCGGetActorProperty.h
+++ b/Engine/Plugins/PCG/Source/PCG/Public/Elements/PCGGetActorProperty.h
@@ -132,4 +132,5 @@ public:
 protected:
 	virtual bool PrepareDataInternal(FPCGContext* Context) const override;
 	virtual bool ExecuteInternal(FPCGContext* Context) const override;
+	virtual bool SupportsBasePointDataInputs(FPCGContext* InContext) const override { return true; }
 };
diff --git a/Engine/Plugins/PCG/Source/PCG/Public/Elements/PCGGetBounds.h b/Engine/Plugins/PCG/Source/PCG/Public/Elements/PCGGetBounds.h
--- a/Engine/Plugins/PCG/Source/PCG/Public/Elements/PCGGetBounds.h
+++ b/Engine/Plugins/PCG/Source/PCG/Public/Elements/PCGGetBounds.h
@@ -35,5 +35,6 @@ class FPCGGetBoundsElement : public IPCGElement
 {
 protected:
 	virtual bool ExecuteInternal(FPCGContext* Context) const override;
+	virtual bool SupportsBasePointDataInputs(FPCGContext* InContext) const override { return true; }
 	virtual EPCGElementExecutionLoopMode ExecutionLoopMode(const UPCGSettings* Settings) const override { return EPCGElementExecutionLoopMode::SinglePrimaryPin; }
 };
diff --git a/Engine/Plugins/PCG/Source/PCG/Public/Elements/PCGMakeConcreteElement.h b/Engine/Plugins/PCG/Source/PCG/Public/Elements/PCGMakeConcreteElement.h
--- a/Engine/Plugins/PCG/Source/PCG/Public/Elements/PCGMakeConcreteElement.h
+++ b/Engine/Plugins/PCG/Source/PCG/Public/Elements/PCGMakeConcreteElement.h
@@ -37,5 +37,6 @@ public:
 
 protected:
 	virtual bool ExecuteInternal(FPCGContext* InContext) const override;
+	virtual bool SupportsBasePointDataInputs(FPCGContext* InContext) const override { return true; }
 	virtual EPCGElementExecutionLoopMode ExecutionLoopMode(const UPCGSettings* Settings) const override { return EPCGElementExecutionLoopMode::SinglePrimaryPin; }
 };
diff --git a/Engine/Plugins/PCG/Source/PCG/Public/Elements/PCGPartitionByActorDataLayers.h b/Engine/Plugins/PCG/Source/PCG/Public/Elements/PCGPartitionByActorDataLayers.h
--- a/Engine/Plugins/PCG/Source/PCG/Public/Elements/PCGPartitionByActorDataLayers.h
+++ b/Engine/Plugins/PCG/Source/PCG/Public/Elements/PCGPartitionByActorDataLayers.h
@@ -56,4 +56,5 @@ public:
 	virtual bool CanExecuteOnlyOnMainThread(FPCGContext* Context) const override { return true; }
 protected:
 	virtual bool ExecuteInternal(FPCGContext* Context) const override;
+	virtual bool SupportsBasePointDataInputs(FPCGContext* InContext) const override { return true; }
 };
\ No newline at end of file
diff --git a/Engine/Plugins/PCG/Source/PCG/Public/Elements/PCGPointNeighborhood.h b/Engine/Plugins/PCG/Source/PCG/Public/Elements/PCGPointNeighborhood.h
--- a/Engine/Plugins/PCG/Source/PCG/Public/Elements/PCGPointNeighborhood.h
+++ b/Engine/Plugins/PCG/Source/PCG/Public/Elements/PCGPointNeighborhood.h
@@ -78,5 +78,6 @@ class FPCGPointNeighborhoodElement : public IPCGElement
 {
 protected:
 	virtual bool ExecuteInternal(FPCGContext* Context) const override;
+	virtual bool SupportsBasePointDataInputs(FPCGContext* InContext) const override { return true; }
 	virtual EPCGElementExecutionLoopMode ExecutionLoopMode(const UPCGSettings* Settings) const override { return EPCGElementExecutionLoopMode::SinglePrimaryPin; }
 };
diff --git a/Engine/Plugins/PCG/Source/PCG/Public/Elements/PCGRandomChoice.h b/Engine/Plugins/PCG/Source/PCG/Public/Elements/PCGRandomChoice.h
--- a/Engine/Plugins/PCG/Source/PCG/Public/Elements/PCGRandomChoice.h
+++ b/Engine/Plugins/PCG/Source/PCG/Public/Elements/PCGRandomChoice.h
@@ -74,4 +74,5 @@ class FPCGRandomChoiceElement : public IPCGElement
 protected:
 	virtual bool ExecuteInternal(FPCGContext* Context) const override;
 	virtual EPCGElementExecutionLoopMode ExecutionLoopMode(const UPCGSettings* Settings) const override { return EPCGElementExecutionLoopMode::SinglePrimaryPin; }
+	virtual bool SupportsBasePointDataInputs(FPCGContext* InContext) const override { return true; }
 };
diff --git a/Engine/Plugins/PCG/Source/PCG/Public/Elements/PCGSampleTexture.h b/Engine/Plugins/PCG/Source/PCG/Public/Elements/PCGSampleTexture.h
--- a/Engine/Plugins/PCG/Source/PCG/Public/Elements/PCGSampleTexture.h
+++ b/Engine/Plugins/PCG/Source/PCG/Public/Elements/PCGSampleTexture.h
@@ -71,5 +71,6 @@ public:
 
 protected:
 	virtual bool ExecuteInternal(FPCGContext* Context) const override;
+	virtual bool SupportsBasePointDataInputs(FPCGContext* InContext) const override { return true; }
 	virtual EPCGElementExecutionLoopMode ExecutionLoopMode(const UPCGSettings* Settings) const override { return EPCGElementExecutionLoopMode::SinglePrimaryPin; }
 };
diff --git a/Engine/Plugins/PCG/Source/PCG/Public/Elements/PCGSanityCheckPointData.h b/Engine/Plugins/PCG/Source/PCG/Public/Elements/PCGSanityCheckPointData.h
--- a/Engine/Plugins/PCG/Source/PCG/Public/Elements/PCGSanityCheckPointData.h
+++ b/Engine/Plugins/PCG/Source/PCG/Public/Elements/PCGSanityCheckPointData.h
@@ -38,4 +38,6 @@ class FPCGSanityCheckPointDataElement : public IPCGElement
 {
 protected:
 	virtual bool ExecuteInternal(FPCGContext* Context) const override;
+	virtual bool SupportsBasePointDataInputs(FPCGContext* InContext) const override { return true; }
+	virtual bool CanExecuteOnlyOnMainThread(FPCGContext* Context) const override { return true; }
 };
diff --git a/Engine/Plugins/PCG/Source/PCG/Public/Elements/PCGSplitPoints.h b/Engine/Plugins/PCG/Source/PCG/Public/Elements/PCGSplitPoints.h
--- a/Engine/Plugins/PCG/Source/PCG/Public/Elements/PCGSplitPoints.h
+++ b/Engine/Plugins/PCG/Source/PCG/Public/Elements/PCGSplitPoints.h
@@ -51,4 +51,5 @@ class FPCGSplitPointsElement : public IPCGElement
 protected:
 	virtual bool ExecuteInternal(FPCGContext* Context) const override;
 	virtual EPCGElementExecutionLoopMode ExecutionLoopMode(const UPCGSettings* Settings) const override { return EPCGElementExecutionLoopMode::SinglePrimaryPin; }
+	virtual bool SupportsBasePointDataInputs(FPCGContext* InContext) const override { return true; }
 };
\ No newline at end of file
diff --git a/Engine/Plugins/PCG/Source/PCG/Public/Elements/PCGSurfaceSampler.h b/Engine/Plugins/PCG/Source/PCG/Public/Elements/PCGSurfaceSampler.h
--- a/Engine/Plugins/PCG/Source/PCG/Public/Elements/PCGSurfaceSampler.h
+++ b/Engine/Plugins/PCG/Source/PCG/Public/Elements/PCGSurfaceSampler.h
@@ -88,7 +88,7 @@ namespace PCGSurfaceSampler
 	};
 
 	/** Sample a surface and returns the resulting point data. Can't be timesliced. */
-	PCG_API UPCGBasePointData* SampleSurface(FPCGContext* Context, const FSurfaceSamplerParams& ExecutionParams, const UPCGSurfaceData* InSurface, const UPCGSpatialData* InBoundingShape, const FBox& EffectiveBounds);
+	PCG_API UPCGBasePointData* SampleSurface(FPCGContext* Context, const FSurfaceSamplerParams& ExecutionParams, const UPCGSurfaceData* InSurface, const UPCGSpatialData* InBoundingShape, const FBox& EffectiveBounds, TSubclassOf<UPCGBasePointData> PointDataClass = nullptr);
 		
 	PCG_API UPCGPointData* SampleSurface(FPCGContext* Context, const UPCGSurfaceData* InSurface, const UPCGSpatialData* InBoundingShape, const FBox& EffectiveBounds, const FSurfaceSamplerParams& ExecutionParams);
 
diff --git a/Engine/Plugins/PCG/Source/PCG/Public/Elements/PCGVisualizeAttribute.h b/Engine/Plugins/PCG/Source/PCG/Public/Elements/PCGVisualizeAttribute.h
--- a/Engine/Plugins/PCG/Source/PCG/Public/Elements/PCGVisualizeAttribute.h
+++ b/Engine/Plugins/PCG/Source/PCG/Public/Elements/PCGVisualizeAttribute.h
@@ -79,4 +79,5 @@ public:
 
 protected:
 	virtual bool ExecuteInternal(FPCGContext* Context) const override;
+	virtual bool SupportsBasePointDataInputs(FPCGContext* InContext) const override { return true; }
 };
\ No newline at end of file
diff --git a/Engine/Plugins/PCG/Source/PCG/Public/Elements/PCGVolumeSampler.h b/Engine/Plugins/PCG/Source/PCG/Public/Elements/PCGVolumeSampler.h
--- a/Engine/Plugins/PCG/Source/PCG/Public/Elements/PCGVolumeSampler.h
+++ b/Engine/Plugins/PCG/Source/PCG/Public/Elements/PCGVolumeSampler.h
@@ -114,5 +114,6 @@ public:
 protected:
 	virtual bool PrepareDataInternal(FPCGContext* Context) const override;
 	virtual bool ExecuteInternal(FPCGContext* Context) const override;
+	virtual bool SupportsBasePointDataInputs(FPCGContext* InContext) const override { return true; }
 	virtual EPCGElementExecutionLoopMode ExecutionLoopMode(const UPCGSettings* Settings) const override { return EPCGElementExecutionLoopMode::SinglePrimaryPin; }
 };
diff --git a/Engine/Plugins/PCG/Source/PCG/Public/Elements/PCGWorldRaycast.h b/Engine/Plugins/PCG/Source/PCG/Public/Elements/PCGWorldRaycast.h
--- a/Engine/Plugins/PCG/Source/PCG/Public/Elements/PCGWorldRaycast.h
+++ b/Engine/Plugins/PCG/Source/PCG/Public/Elements/PCGWorldRaycast.h
@@ -123,6 +123,7 @@ public:
 protected:
 	virtual bool PrepareDataInternal(FPCGContext* InContext) const override;
 	virtual bool ExecuteInternal(FPCGContext* InContext) const override;
+	virtual bool SupportsBasePointDataInputs(FPCGContext* InContext) const override { return true; }
 	// Physics required to run on main thread
 	virtual bool CanExecuteOnlyOnMainThread(FPCGContext* Context) const override { return true; }
 };
diff --git a/Engine/Plugins/PCG/Source/PCG/Public/Grid/PCGLandscapeCache.h b/Engine/Plugins/PCG/Source/PCG/Public/Grid/PCGLandscapeCache.h
--- a/Engine/Plugins/PCG/Source/PCG/Public/Grid/PCGLandscapeCache.h
+++ b/Engine/Plugins/PCG/Source/PCG/Public/Grid/PCGLandscapeCache.h
@@ -72,14 +72,14 @@ public:
 	void GetPoint(int32 PointIndex, FPCGPoint& OutPoint, UPCGMetadata* OutMetadata) const;
 	void GetPointHeightOnly(int32 PointIndex, FPCGPoint& OutPoint) const;
 	void GetInterpolatedPoint(const FVector2D& LocalPoint, FPCGPoint& OutPoint, UPCGMetadata* OutMetadata) const;
-	void GetInterpolatedPointMetadataOnly(const FVector2D& LocalPoint, FPCGPoint& OutPoint, UPCGMetadata* OutMetadata) const;
+	void GetInterpolatedPointMetadataOnly(const FVector2D& LocalPoint, int64& OutMetadataEntry, UPCGMetadata* OutMetadata) const;
 	void GetInterpolatedPointHeightOnly(const FVector2D& LocalPoint, FPCGPoint& OutPoint, UPCGMetadata* OutMetadata) const;
 	void GetInterpolatedLayerWeights(const FVector2D& LocalPoint, TArray<FPCGLandscapeLayerWeight>& OutLayerWeights) const;
 
 private:
 	// Private API to remove boilerplate
 	void GetInterpolatedPointInternal(const PCGLandscapeCache::FSafeIndices& Indices, FPCGPoint& OutPoint, bool bHeightOnly = false) const;
-	void GetInterpolatedPointMetadataInternal(const PCGLandscapeCache::FSafeIndices& Indices, FPCGPoint& OutPoint, UPCGMetadata* OutMetadata) const;
+	void GetInterpolatedPointMetadataInternal(const PCGLandscapeCache::FSafeIndices& Indices, int64& OutMetadataEntry, UPCGMetadata* OutMetadata) const;
 
 	// Private API for UPCGLandscapeCache usage
 	bool TouchAndLoad(int32 Touch) const;
@@ -153,7 +153,7 @@ public:
 	TArray<FName> GetLayerNames(ALandscapeProxy* Landscape);
 
 	/** Convenience method to get metadata from the landscape for a given pair of landscape and position */
-	void SampleMetadataOnPoint(ALandscapeProxy* Landscape, FPCGPoint& OutPoint, UPCGMetadata* OutMetadata);
+	void SampleMetadataOnPoint(ALandscapeProxy* Landscape, const FTransform& InTransform, int64& OutMetadataEntry, UPCGMetadata* OutMetadata);
 
 	UPROPERTY(BlueprintReadWrite, EditAnywhere, Category = "Landscape Cache")
 	EPCGLandscapeCacheSerializationMode SerializationMode = EPCGLandscapeCacheSerializationMode::NeverSerialize;
diff --git a/Engine/Plugins/PCG/Source/PCG/Public/Helpers/PCGWorldQueryHelpers.h b/Engine/Plugins/PCG/Source/PCG/Public/Helpers/PCGWorldQueryHelpers.h
--- a/Engine/Plugins/PCG/Source/PCG/Public/Helpers/PCGWorldQueryHelpers.h
+++ b/Engine/Plugins/PCG/Source/PCG/Public/Helpers/PCGWorldQueryHelpers.h
@@ -74,13 +74,14 @@ namespace PCGWorldQueryHelpers
 	bool CreateRayHitAttributes(const FPCGWorldRaycastQueryParams& QueryParams, UPCGMetadata* OutMetadata);
 
 	/** Applies a 'miss' hit result to the metadata. To be called if the ray misses the target and the point should be kept. */
-	bool ApplyRayMissMetadata(const FPCGWorldRaycastQueryParams& QueryParams, FPCGPoint& OutPoint, UPCGMetadata* OutMetadata);
+	bool ApplyRayMissMetadata(const FPCGWorldRaycastQueryParams& QueryParams, int64& OutMetadataEntry, UPCGMetadata* OutMetadata);
 
 	/** Applies common world ray hit results to attributes. */
 	bool ApplyRayHitMetadata(const FHitResult& HitResult,
 		const FPCGWorldRaycastQueryParams& QueryParams,
 		const FVector& RayDirection,
-		FPCGPoint& OutPoint,
+		const FTransform& InTransform, 
+		int64& OutMetadataEntry,
 		UPCGMetadata* OutMetadata,
 		TWeakObjectPtr<UWorld> World);
 }
diff --git a/Engine/Plugins/PCG/Source/PCG/Public/SpatialAlgo/PCGOctreeQueries.h b/Engine/Plugins/PCG/Source/PCG/Public/SpatialAlgo/PCGOctreeQueries.h
--- a/Engine/Plugins/PCG/Source/PCG/Public/SpatialAlgo/PCGOctreeQueries.h
+++ b/Engine/Plugins/PCG/Source/PCG/Public/SpatialAlgo/PCGOctreeQueries.h
@@ -26,12 +26,14 @@ public:
 
 	/** Get the closest point to a given position within the search distance, If bInDiscardCenter is true, will reject any points that are exactly at the center. Returns null if not found. */
 	static PCG_API const FPCGPoint* GetClosestPoint(const UPCGPointData* InPointData, const FVector& InCenter, const bool bInDiscardCenter, const double InSearchDistance = DefaultSearchDistance);
+	static PCG_API int32 GetClosestPointIndex(const UPCGBasePointData* InPointData, const FVector& InCenter, const bool bInDiscardCenter, const double InSearchDistance = DefaultSearchDistance);
 
 	/** Get the nearest point to a given point (that is not itself) within the search distance. Returns null if not found. */
 	static PCG_API const FPCGPoint* GetClosestPointFromOtherPoint(const UPCGPointData* InPointData, const FPCGPoint& InPoint, const double InSearchDistance = DefaultSearchDistance);
 
 	/** Get the farthest point from a given position within the search distance. Returns null if not found. */
 	static PCG_API const FPCGPoint* GetFarthestPoint(const UPCGPointData* InPointData, const FVector& InCenter, const double InSearchDistance = DefaultSearchDistance);
+	static PCG_API int32 GetFarthestPointIndex(const UPCGBasePointData* InPointData, const FVector& InCenter, const double InSearchDistance = DefaultSearchDistance);
 
 	/** Get the farthest point from a given point (that is not itself) within the search distance. Returns null if not found. */
 	static PCG_API const FPCGPoint* GetFarthestPointFromOtherPoint(const UPCGPointData* InPointData, const FPCGPoint& InPoint, const double InSearchDistance = DefaultSearchDistance);
diff --git a/Engine/Plugins/PCG/Source/PCG/Public/Tests/PCGTestsCommon.h b/Engine/Plugins/PCG/Source/PCG/Public/Tests/PCGTestsCommon.h
--- a/Engine/Plugins/PCG/Source/PCG/Public/Tests/PCGTestsCommon.h
+++ b/Engine/Plugins/PCG/Source/PCG/Public/Tests/PCGTestsCommon.h
@@ -57,6 +57,8 @@ namespace PCGTestsCommon
 	AActor* CreateTemporaryActor();
 	UPCGParamData* CreateEmptyParamData();
 
+	UPCGBasePointData* CreateEmptyBasePointData();
+
 	/** Creates a PointData with a single point at the origin */
 	template <typename PointDataType = UPCGPointData, std::enable_if_t<std::is_base_of_v<UPCGBasePointData, PointDataType>, int> = 0>
 	PointDataType* CreateEmptyPointData()
diff --git a/Engine/Plugins/PCGInterops/PCGExternalDataInterop/Source/PCGExternalDataInterop/Private/Alembic/PCGAlembicInterop.cpp b/Engine/Plugins/PCGInterops/PCGExternalDataInterop/Source/PCGExternalDataInterop/Private/Alembic/PCGAlembicInterop.cpp
--- a/Engine/Plugins/PCGInterops/PCGExternalDataInterop/Source/PCGExternalDataInterop/Private/Alembic/PCGAlembicInterop.cpp
+++ b/Engine/Plugins/PCGInterops/PCGExternalDataInterop/Source/PCGExternalDataInterop/Private/Alembic/PCGAlembicInterop.cpp
@@ -5,7 +5,7 @@
 #include "Elements/PCGLoadAlembicElement.h"
 
 #include "PCGComponent.h"
-#include "Data/PCGPointData.h"
+#include "Data/PCGBasePointData.h"
 #include "Elements/Metadata/PCGMetadataElementCommon.h"
 #include "Metadata/Accessors/IPCGAttributeAccessorTpl.h"
 #include "Metadata/Accessors/PCGAttributeAccessorHelpers.h"
@@ -173,7 +173,7 @@ TUniquePtr<const IPCGAttributeAccessor> CreateAlembicPropAccessor(Alembic::AbcGe
 	return nullptr;
 }
 
-bool CreatePointAccessorAndValidate(FPCGContext* Context, UPCGPointData* PointData, const TUniquePtr<const IPCGAttributeAccessor>& AlembicPropAccessor, const FPCGAttributePropertySelector& PointPropertySelector, const FString& PropName, TUniquePtr<IPCGAttributeAccessor>& PointPropertyAccessor)
+bool CreatePointAccessorAndValidate(FPCGContext* Context, UPCGBasePointData* PointData, const TUniquePtr<const IPCGAttributeAccessor>& AlembicPropAccessor, const FPCGAttributePropertySelector& PointPropertySelector, const FString& PropName, TUniquePtr<IPCGAttributeAccessor>& PointPropertyAccessor)
 {
 	if (!PointData || !PointData->Metadata)
 	{
@@ -249,7 +249,7 @@ bool ParseAlembicObject(FPCGExternalDataContext* Context, const Alembic::Abc::IO
 			bHasCreatedData = true;
 
 			// Create point data & mapping
-			UPCGPointData* PointData = NewObject<UPCGPointData>();
+			UPCGBasePointData* PointData = FPCGContext::NewPointData_AnyThread(Context);
 			check(PointData);
 
 			UPCGMetadata* PointMetadata = PointData->MutableMetadata();
@@ -260,8 +260,7 @@ bool ParseAlembicObject(FPCGExternalDataContext* Context, const Alembic::Abc::IO
 			// We're not going to use the input keys, but we still need to provide something
 			PointDataAccessorMapping.RowKeys = MakeUnique<FPCGAttributeAccessorKeysEntries>(PCGInvalidEntryKey);
 
-			TArray<FPCGPoint>& OutPoints = PointData->GetMutablePoints();
-			OutPoints.SetNum(NumPoints);
+			PointData->SetNumPoints(NumPoints);
 
 			// If the user has provided a position remapping, don't process points right now, instead push the transformation to the row accessors
 			if (const FPCGAttributePropertySelector* RemappedPositions = Settings->AttributeMapping.Find(TEXT("Position")))
@@ -280,10 +279,11 @@ bool ParseAlembicObject(FPCGExternalDataContext* Context, const Alembic::Abc::IO
 			}
 			else // Otherwise, write the positions directly
 			{
+				TPCGValueRange<FTransform> TransformRange = PointData->GetTransformValueRange(/*bAllocate=*/true);
 				for (uint32 PointIndex = 0; PointIndex < NumPoints; ++PointIndex)
 				{
 					Alembic::Abc::P3fArraySample::value_type Position = (*Positions)[PointIndex];
-					OutPoints[PointIndex].Transform.SetLocation(FVector(Position.x, Position.y, Position.z));
+					TransformRange[PointIndex].SetLocation(FVector(Position.x, Position.y, Position.z));
 				}
 			}
 
diff --git a/Engine/Plugins/PCGInterops/PCGExternalDataInterop/Source/PCGExternalDataInterop/Private/Elements/PCGLoadAlembicElement.cpp b/Engine/Plugins/PCGInterops/PCGExternalDataInterop/Source/PCGExternalDataInterop/Private/Elements/PCGLoadAlembicElement.cpp
--- a/Engine/Plugins/PCGInterops/PCGExternalDataInterop/Source/PCGExternalDataInterop/Private/Elements/PCGLoadAlembicElement.cpp
+++ b/Engine/Plugins/PCGInterops/PCGExternalDataInterop/Source/PCGExternalDataInterop/Private/Elements/PCGLoadAlembicElement.cpp
@@ -149,24 +149,24 @@ bool FPCGLoadAlembicElement::ExecuteLoad(FPCGExternalDataContext* InContext) con
 	{
 		for (const FPCGExternalDataContext::FPointDataAccessorsMapping& DataMapping : Context->PointDataAccessorsMapping)
 		{
-			UPCGPointData* PointData = Cast<UPCGPointData>(DataMapping.Data);
+			UPCGBasePointData* PointData = Cast<UPCGBasePointData>(DataMapping.Data);
 
 			if (!PointData)
 			{
 				continue;
 			}
 
-			TArray<FPCGPoint>& Points = PointData->GetMutablePoints();
+			TPCGValueRange<FTransform> TransformRange = PointData->GetTransformValueRange(/*bAllocate=*/true);
 
-			for (FPCGPoint& Point : Points)
+			for (int32 PointIndex = 0; PointIndex < PointData->GetNumPoints(); ++PointIndex)
 			{
-				Point.Transform = Point.Transform * ConversionTransform;
+				TransformRange[PointIndex] = TransformRange[PointIndex] * ConversionTransform;
 
 				if (bFlipRotationW)
 				{
-					FQuat Rotation = Point.Transform.GetRotation();
+					FQuat Rotation = TransformRange[PointIndex].GetRotation();
 					Rotation.W *= -1;
-					Point.Transform.SetRotation(Rotation);
+					TransformRange[PointIndex].SetRotation(Rotation);
 				}
 			}
 		}
diff --git a/Engine/Plugins/PCGInterops/PCGGeometryScriptInterop/Source/PCGGeometryScriptInterop/Private/Data/PCGDynamicMeshData.cpp b/Engine/Plugins/PCGInterops/PCGGeometryScriptInterop/Source/PCGGeometryScriptInterop/Private/Data/PCGDynamicMeshData.cpp
--- a/Engine/Plugins/PCGInterops/PCGGeometryScriptInterop/Source/PCGGeometryScriptInterop/Private/Data/PCGDynamicMeshData.cpp
+++ b/Engine/Plugins/PCGInterops/PCGGeometryScriptInterop/Source/PCGGeometryScriptInterop/Private/Data/PCGDynamicMeshData.cpp
@@ -4,6 +4,7 @@
 
 #include "PCGContext.h"
 #include "PCGModule.h"
+#include "Data/PCGPointArrayData.h"
 #include "Data/PCGPointData.h"
 #include "Elements/PCGVolumeSampler.h"
 
@@ -12,6 +13,7 @@
 #include "GeometryScript/MeshQueryFunctions.h"
 #include "Helpers/PCGHelpers.h"
 #include "Serialization/ArchiveCrc32.h"
+#include "UObject/Package.h"
 
 UPCGDynamicMeshData::UPCGDynamicMeshData(const FObjectInitializer& ObjectInitializer)
 	: UPCGSpatialData(ObjectInitializer)
@@ -157,7 +159,7 @@ const UE::Geometry::FDynamicMeshOctree3& UPCGDynamicMeshData::GetDynamicMeshOctr
 	return DynamicMeshOctree;
 }
 
-const UPCGPointData* UPCGDynamicMeshData::ToPointData(FPCGContext* Context, const FBox& InBounds) const
+const UPCGBasePointData* UPCGDynamicMeshData::ToBasePointData(FPCGContext* Context, const FBox& InBounds, TSubclassOf<UPCGBasePointData> PointDataClass) const
 {
 	TRACE_CPUPROFILER_EVENT_SCOPE(UPCGPrimitiveData::CreatePointData);
 
@@ -168,10 +170,14 @@ const UPCGPointData* UPCGDynamicMeshData::ToPointData(FPCGContext* Context, cons
 	UGeometryScriptLibrary_MeshQueryFunctions::GetAllVertexPositions(DynamicMesh, VerticesList, /*bSkipGaps=*/false, /*bHasVertexIDGaps=*/Dummy);
 
 	TArray<FVector>& Vertices = *VerticesList.List.Get();
-	
-	UPCGPointData* Data = FPCGContext::NewObject_AnyThread<UPCGPointData>(Context);
-	TArray<FPCGPoint>& Points = Data->GetMutablePoints();
-	Points.Reserve(Vertices.Num());
+
+	UPCGBasePointData* Data = FPCGContext::NewObject_AnyThread<UPCGBasePointData>(Context, GetTransientPackage(), PointDataClass);
+	Data->SetNumPoints(Vertices.Num());
+	Data->AllocateProperties(EPCGPointNativeProperties::Transform | EPCGPointNativeProperties::Seed);
+
+	int32 NumWritten = 0;
+	FPCGPointValueRanges OutRanges(Data, /*bAllocate=*/false);
+
 	for (const FVector& Vertex : Vertices)
 	{
 		if (InBounds.IsValid && !InBounds.IsInside(Vertex))
@@ -179,16 +185,29 @@ const UPCGPointData* UPCGDynamicMeshData::ToPointData(FPCGContext* Context, cons
 			continue;
 		}
 		
-		FPCGPoint& Point = Points.Emplace_GetRef();
-		Point.Transform.SetLocation(Vertex);
-		PCGHelpers::ComputeSeedFromPosition(Vertex);
+		OutRanges.TransformRange[NumWritten].SetLocation(Vertex);
+		OutRanges.SeedRange[NumWritten] = PCGHelpers::ComputeSeedFromPosition(Vertex);
+
+		++NumWritten;
 	}
-	
-	UE_LOG(LogPCG, Verbose, TEXT("Dynamic Mesh extracted %d points"), Points.Num());
+
+	Data->SetNumPoints(NumWritten);
+
+	UE_LOG(LogPCG, Verbose, TEXT("Dynamic Mesh extracted %d points"), NumWritten);
 
 	return Data;
 }
 
+const UPCGPointArrayData* UPCGDynamicMeshData::ToPointArrayData(FPCGContext* Context, const FBox& InBounds) const
+{
+	return CastChecked<UPCGPointArrayData>(ToBasePointData(Context, InBounds, UPCGPointArrayData::StaticClass()));
+}
+
+const UPCGPointData* UPCGDynamicMeshData::ToPointData(FPCGContext* Context, const FBox& InBounds) const
+{
+	return CastChecked<UPCGPointData>(ToBasePointData(Context, InBounds, UPCGPointData::StaticClass()));
+}
+
 void UPCGDynamicMeshData::SetMaterials(const TArray<UMaterialInterface*>& InMaterials)
 {
 	Materials.Reset(InMaterials.Num());
diff --git a/Engine/Plugins/PCGInterops/PCGGeometryScriptInterop/Source/PCGGeometryScriptInterop/Private/Editor/PCGDynamicMeshDataVisualization.h b/Engine/Plugins/PCGInterops/PCGGeometryScriptInterop/Source/PCGGeometryScriptInterop/Private/Editor/PCGDynamicMeshDataVisualization.h
--- a/Engine/Plugins/PCGInterops/PCGGeometryScriptInterop/Source/PCGGeometryScriptInterop/Private/Editor/PCGDynamicMeshDataVisualization.h
+++ b/Engine/Plugins/PCGInterops/PCGGeometryScriptInterop/Source/PCGGeometryScriptInterop/Private/Editor/PCGDynamicMeshDataVisualization.h
@@ -8,7 +8,6 @@
 
 class AActor;
 class UPCGData;
-class UPCGPointData;
 struct FPCGContext;
 
 class FPCGDynamicMeshDataVisualization : public IPCGSpatialDataVisualization
diff --git a/Engine/Plugins/PCGInterops/PCGGeometryScriptInterop/Source/PCGGeometryScriptInterop/Private/Elements/PCGAppendMeshesFromPoints.cpp b/Engine/Plugins/PCGInterops/PCGGeometryScriptInterop/Source/PCGGeometryScriptInterop/Private/Elements/PCGAppendMeshesFromPoints.cpp
--- a/Engine/Plugins/PCGInterops/PCGGeometryScriptInterop/Source/PCGGeometryScriptInterop/Private/Elements/PCGAppendMeshesFromPoints.cpp
+++ b/Engine/Plugins/PCGInterops/PCGGeometryScriptInterop/Source/PCGGeometryScriptInterop/Private/Elements/PCGAppendMeshesFromPoints.cpp
@@ -4,7 +4,7 @@
 
 #include "PCGContext.h"
 #include "Data/PCGDynamicMeshData.h"
-#include "Data/PCGPointData.h"
+#include "Data/PCGBasePointData.h"
 #include "Elements/Metadata/PCGMetadataElementCommon.h"
 #include "Helpers/PCGGeometryHelpers.h"
 #include "Metadata/Accessors/PCGAttributeAccessorHelpers.h"
@@ -121,7 +121,7 @@ bool FPCGAppendMeshesFromPointsElement::PrepareDataInternal(FPCGContext* InConte
 	if (Settings->Mode == EPCGAppendMeshesFromPointsMode::StaticMeshFromAttribute)
 	{
 		const TArray<FPCGTaggedData> InputPoints = InContext->InputData.GetInputsByPin(PCGAppendMeshesFromPoints::InPointsPinLabel);
-		const UPCGPointData* InPointData = !InputPoints.IsEmpty() ? Cast<const UPCGPointData>(InputPoints[0].Data) : nullptr;
+		const UPCGBasePointData* InPointData = !InputPoints.IsEmpty() ? Cast<const UPCGBasePointData>(InputPoints[0].Data) : nullptr;
 
 		if (!InPointData)
 		{
@@ -181,7 +181,7 @@ bool FPCGAppendMeshesFromPointsElement::ExecuteInternal(FPCGContext* InContext)
 	const TArray<FPCGTaggedData> InputDynMesh = InContext->InputData.GetInputsByPin(PCGAppendMeshesFromPoints::InDynMeshPinLabel);
 	const TArray<FPCGTaggedData> InputAppendDynMesh = InContext->InputData.GetInputsByPin(PCGAppendMeshesFromPoints::InAppendMeshPinLabel);
 
-	const UPCGPointData* InPointData = !InputPoints.IsEmpty() ? Cast<const UPCGPointData>(InputPoints[0].Data) : nullptr;
+	const UPCGBasePointData* InPointData = !InputPoints.IsEmpty() ? Cast<const UPCGBasePointData>(InputPoints[0].Data) : nullptr;
 	const UPCGDynamicMeshData* InDynMeshData = !InputDynMesh.IsEmpty() ? Cast<const UPCGDynamicMeshData>(InputDynMesh[0].Data) : nullptr;
 	const UPCGDynamicMeshData* InAppendDynMeshData = !InputAppendDynMesh.IsEmpty() ? Cast<const UPCGDynamicMeshData>(InputAppendDynMesh[0].Data) : nullptr;
 
@@ -192,7 +192,7 @@ bool FPCGAppendMeshesFromPointsElement::ExecuteInternal(FPCGContext* InContext)
 	}
 
 	FPCGTaggedData& OutputData = InContext->OutputData.TaggedData.Emplace_GetRef(InputDynMesh[0]);
-	if (InPointData->GetPoints().IsEmpty())
+	if (InPointData->IsEmpty())
 	{
 		return true;
 	}
@@ -291,20 +291,20 @@ bool FPCGAppendMeshesFromPointsElement::ExecuteInternal(FPCGContext* InContext)
 	UE::Geometry::FMeshIndexMappings MeshIndexMappings;
 	UE::Geometry::FDynamicMeshEditor Editor(OutDynMeshData->GetMutableDynamicMesh()->GetMeshPtr());
 
-	auto AppendMesh = [&Editor, &MeshIndexMappings](const UE::Geometry::FDynamicMesh3* MeshToAppend, const FPCGPoint& Point)
+	auto AppendMesh = [&Editor, &MeshIndexMappings](const UE::Geometry::FDynamicMesh3* MeshToAppend, const FTransform& PointTransform)
 	{
 		Editor.AppendMesh(MeshToAppend, MeshIndexMappings,
-		[&Point](int, const FVector& Position) { return Point.Transform.TransformPosition(Position); },
-		[&Point](int, const FVector& Normal)
+		[&PointTransform](int, const FVector& Position) { return PointTransform.TransformPosition(Position); },
+		[&PointTransform](int, const FVector& Normal)
 		{
-			const FVector& PointScale = Point.Transform.GetScale3D();
+			const FVector& PointScale = PointTransform.GetScale3D();
 			const double DetSign = FMathd::SignNonZero(PointScale.X * PointScale.Y * PointScale.Z); // we only need to multiply by the sign of the determinant, rather than divide by it, since we normalize later anyway
 			const FVector SafeInversePointScale(PointScale.Y * PointScale.Z * DetSign, PointScale.X * PointScale.Z * DetSign, PointScale.X * PointScale.Y * DetSign);
-			return Point.Transform.TransformVectorNoScale((SafeInversePointScale * Normal).GetSafeNormal());
+			return PointTransform.TransformVectorNoScale((SafeInversePointScale * Normal).GetSafeNormal());
 		});
 	};
 
-	const TArray<FPCGPoint>& InPoints = InPointData->GetPoints();
+	const TConstPCGValueRange<FTransform> TransformRange = InPointData->GetConstTransformValueRange();
 
 	if (Settings->Mode == EPCGAppendMeshesFromPointsMode::StaticMeshFromAttribute)
 	{
@@ -315,7 +315,7 @@ bool FPCGAppendMeshesFromPointsElement::ExecuteInternal(FPCGContext* InContext)
 
 			for (const int32 i : It.Value)
 			{
-				AppendMesh(MeshToAppend, InPoints[i]);
+				AppendMesh(MeshToAppend, TransformRange[i]);
 			}
 		}
 	}
@@ -337,9 +337,9 @@ bool FPCGAppendMeshesFromPointsElement::ExecuteInternal(FPCGContext* InContext)
 			return true;
 		}
 
-		for (const FPCGPoint& Point : InPoints)
+		for (const FTransform& PointTransform : TransformRange)
 		{
-			AppendMesh(MeshToAppend, Point);
+			AppendMesh(MeshToAppend, PointTransform);
 		}
 	}
 
diff --git a/Engine/Plugins/PCGInterops/PCGGeometryScriptInterop/Source/PCGGeometryScriptInterop/Private/Elements/PCGMeshSampler.cpp b/Engine/Plugins/PCGInterops/PCGGeometryScriptInterop/Source/PCGGeometryScriptInterop/Private/Elements/PCGMeshSampler.cpp
--- a/Engine/Plugins/PCGInterops/PCGGeometryScriptInterop/Source/PCGGeometryScriptInterop/Private/Elements/PCGMeshSampler.cpp
+++ b/Engine/Plugins/PCGInterops/PCGGeometryScriptInterop/Source/PCGGeometryScriptInterop/Private/Elements/PCGMeshSampler.cpp
@@ -9,7 +9,6 @@
 #include "PCGPin.h"
 #include "Data/PCGPointData.h"
 #include "Helpers/PCGAsync.h"
-#include "Helpers/PCGBlueprintHelpers.h"
 #include "Helpers/PCGDynamicTrackingHelpers.h"
 #include "Helpers/PCGHelpers.h"
 #include "Metadata/PCGMetadata.h"
@@ -41,36 +40,57 @@ namespace PCGMeshSampler
 
 		TRACE_CPUPROFILER_EVENT_SCOPE(PCGMeshSampler::SampleOnePointPerVertex);
 
-		auto IterationBody = [Settings, Context, SetPointDensityPtr](int32 ReadIndex, int32 WriteIndex) -> bool
+		auto IterationBody = [Settings, Context, SetPointDensityPtr](int32 StartReadIndex, int32 StartWriteIndex, int32 Count)
 		{
+			int32 NumWritten = 0;
 			int32 DataIndex = 0;
-			while (!Context->DynamicMeshes[DataIndex] || ReadIndex >= Context->StartingIndices[DataIndex + 1])
+
+			auto UpdateDataIndex = [&DataIndex, Context](int32 InReadIndex) -> bool
 			{
-				DataIndex++;
-			}
+				bool bUpdated = false;
+				while (!Context->DynamicMeshes[DataIndex] || InReadIndex >= Context->StartingIndices[DataIndex + 1])
+				{
+					DataIndex++;
+					bUpdated = true;
+				}
+				return bUpdated;
+			};
+
+			UpdateDataIndex(StartReadIndex);
+			FPCGPointValueRanges OutRanges(Context->OutputData[DataIndex], /*bAllocate=*/false);
+
+			for (int32 ReadIndex = StartReadIndex; ReadIndex < StartReadIndex + Count; ++ReadIndex)
+			{
+				if (UpdateDataIndex(ReadIndex))
+				{
+					OutRanges = FPCGPointValueRanges(Context->OutputData[DataIndex], /*bAllocate=*/false);
+				}
 
-			const TArray<FVector>& Positions = *Context->Positions[DataIndex].List.Get();
-			const TArray<FLinearColor>& Colors = *Context->Colors[DataIndex].List.Get();
-			const TArray<FVector>& Normals = *Context->Normals[DataIndex].List.Get();
+				const TArray<FVector>& Positions = *Context->Positions[DataIndex].List.Get();
+				const TArray<FLinearColor>& Colors = *Context->Colors[DataIndex].List.Get();
+				const TArray<FVector>& Normals = *Context->Normals[DataIndex].List.Get();
 
-			const int32 CurrentIndex = ReadIndex - Context->StartingIndices[DataIndex];
+				const int32 CurrentIndex = ReadIndex - Context->StartingIndices[DataIndex];
 
-			const FVector& Position = Positions[CurrentIndex];
-			const FLinearColor& Color = Colors[CurrentIndex];
-			const FVector& Normal = Normals[CurrentIndex];
+				const FVector& Position = Positions[CurrentIndex];
+				const FLinearColor& Color = Colors[CurrentIndex];
+				const FVector& Normal = Normals[CurrentIndex];
 
-			FPCGPoint OutPoint{};
-			OutPoint.Transform = FTransform{ FRotationMatrix::MakeFromZ(Normal).Rotator(), Position, FVector::OneVector };
-			OutPoint.Color = Color;
-			OutPoint.Steepness = Settings->PointSteepness;
+				FPCGPoint OutPoint{};
+				OutPoint.Transform = FTransform{ FRotationMatrix::MakeFromZ(Normal).Rotator(), Position, FVector::OneVector };
+				OutPoint.Color = Color;
+				OutPoint.Steepness = Settings->PointSteepness;
 
-			SetPointDensityPtr(Color, OutPoint);
+				SetPointDensityPtr(Color, OutPoint);
 
-			UPCGBlueprintHelpers::SetSeedFromPosition(OutPoint);
+				OutPoint.Seed = PCGHelpers::ComputeSeedFromPosition(OutPoint.Transform.GetLocation());
 
-			Context->OutPointData[DataIndex]->GetMutablePoints()[CurrentIndex] = OutPoint;
+				OutRanges.SetFromPoint(CurrentIndex, OutPoint);
+				++NumWritten;
+			}
 
-			return true;
+			check(Count == NumWritten);
+			return Count;
 		};
 
 		// Small verification to catch any problem that could arise in the iteration body
@@ -79,7 +99,7 @@ namespace PCGMeshSampler
 			return true;
 		}
 
-		return FPCGAsync::AsyncProcessingOneToOneEx(&Context->AsyncState, Context->StartingIndices.Last(), []() {}, IterationBody, /*bEnableTimeSlicing=*/true);
+		return FPCGAsync::AsyncProcessingOneToOneRangeEx(&Context->AsyncState, Context->StartingIndices.Last(), []() {}, IterationBody, /*bEnableTimeSlicing=*/true);
 	}
 
 	bool SampleOnePointPerTriangle(const UPCGMeshSamplerSettings* Settings, FPCGMeshSamplerContext* Context, FPCGMeshSamplerContext::SetPointDensityFunc SetPointDensityPtr)
@@ -88,43 +108,64 @@ namespace PCGMeshSampler
 
 		TRACE_CPUPROFILER_EVENT_SCOPE(FPCGMeshSamplerElement::Execute::OnePointPerTriangle);
 
-		auto IterationBody = [Context, Settings, SetPointDensityPtr](int32 ReadIndex, int32 WriteIndex) -> bool
+		auto IterationBody = [Context, Settings, SetPointDensityPtr](int32 StartReadIndex, int32 StartWriteIndex, int32 Count)
 		{
+			int32 NumWritten = 0;
 			int32 DataIndex = 0;
-			while (!Context->DynamicMeshes[DataIndex] || ReadIndex >= Context->StartingIndices[DataIndex + 1])
+
+			auto UpdateDataIndex = [&DataIndex, Context](int32 InReadIndex) -> bool
 			{
-				DataIndex++;
-			}
+				bool bUpdated = false;
+				while (!Context->DynamicMeshes[DataIndex] || InReadIndex >= Context->StartingIndices[DataIndex + 1])
+				{
+					DataIndex++;
+					bUpdated = true;
+				}
+				return bUpdated;
+			};
 
-			const TArray<int32>& TriangleIds = *Context->TriangleIds[DataIndex].List.Get();
-			const int32 CurrentIndex = ReadIndex - Context->StartingIndices[DataIndex];
+			UpdateDataIndex(StartReadIndex);
+			FPCGPointValueRanges OutRanges(Context->OutputData[DataIndex], /*bAllocate=*/false);
 
-			const int32 TriangleId = TriangleIds[CurrentIndex];
+			for (int32 ReadIndex = StartReadIndex; ReadIndex < StartReadIndex + Count; ++ReadIndex)
+			{
+				if (UpdateDataIndex(ReadIndex))
+				{
+					OutRanges = FPCGPointValueRanges(Context->OutputData[DataIndex], /*bAllocate=*/false);
+				}
 
-			FVector Vertex1, Vertex2, Vertex3;
-			bool bIsValidTriangle;
+				const TArray<int32>& TriangleIds = *Context->TriangleIds[DataIndex].List.Get();
+				const int32 CurrentIndex = ReadIndex - Context->StartingIndices[DataIndex];
 
-			UGeometryScriptLibrary_MeshQueryFunctions::GetTrianglePositions(Context->DynamicMeshes[DataIndex], TriangleId, bIsValidTriangle, Vertex1, Vertex2, Vertex3);
-			const FVector Normal = UGeometryScriptLibrary_MeshQueryFunctions::GetTriangleFaceNormal(Context->DynamicMeshes[DataIndex], TriangleId, bIsValidTriangle);
-			const FVector Position = (Vertex1 + Vertex2 + Vertex3) / 3.0;
+				const int32 TriangleId = TriangleIds[CurrentIndex];
 
-			FPCGPoint OutPoint{};
-			OutPoint.Transform = FTransform{ FRotationMatrix::MakeFromZ(Normal).Rotator(), Position, FVector::OneVector };
-			OutPoint.Steepness = Settings->PointSteepness;
+				FVector Vertex1, Vertex2, Vertex3;
+				bool bIsValidTriangle;
 
-			FVector Dummy1, Dummy2, Dummy3;
-			FVector BarycentricCoord;
-			bool bIsValid = false;
-			UGeometryScriptLibrary_MeshQueryFunctions::ComputeTriangleBarycentricCoords(Context->DynamicMeshes[DataIndex], TriangleId, bIsValid, Position, Dummy1, Dummy2, Dummy3, BarycentricCoord);
+				UGeometryScriptLibrary_MeshQueryFunctions::GetTrianglePositions(Context->DynamicMeshes[DataIndex], TriangleId, bIsValidTriangle, Vertex1, Vertex2, Vertex3);
+				const FVector Normal = UGeometryScriptLibrary_MeshQueryFunctions::GetTriangleFaceNormal(Context->DynamicMeshes[DataIndex], TriangleId, bIsValidTriangle);
+				const FVector Position = (Vertex1 + Vertex2 + Vertex3) / 3.0;
 
-			Context->SetPointColorAndDensity(SetPointDensityPtr, TriangleId, BarycentricCoord, OutPoint, DataIndex);
-			Context->SetAttributeValues(Settings->UVChannel, TriangleId, BarycentricCoord, OutPoint, DataIndex);
+				FPCGPoint OutPoint{};
+				OutPoint.Transform = FTransform{ FRotationMatrix::MakeFromZ(Normal).Rotator(), Position, FVector::OneVector };
+				OutPoint.Steepness = Settings->PointSteepness;
 
-			UPCGBlueprintHelpers::SetSeedFromPosition(OutPoint);
+				FVector Dummy1, Dummy2, Dummy3;
+				FVector BarycentricCoord;
+				bool bIsValid = false;
+				UGeometryScriptLibrary_MeshQueryFunctions::ComputeTriangleBarycentricCoords(Context->DynamicMeshes[DataIndex], TriangleId, bIsValid, Position, Dummy1, Dummy2, Dummy3, BarycentricCoord);
 
-			Context->OutPointData[DataIndex]->GetMutablePoints()[CurrentIndex] = OutPoint;
+				Context->SetPointColorAndDensity(SetPointDensityPtr, TriangleId, BarycentricCoord, OutPoint, DataIndex);
+				Context->SetAttributeValues(Settings->UVChannel, TriangleId, BarycentricCoord, OutPoint.MetadataEntry, DataIndex);
 
-			return true;
+				OutPoint.Seed = PCGHelpers::ComputeSeedFromPosition(OutPoint.Transform.GetLocation());
+
+				OutRanges.SetFromPoint(CurrentIndex, OutPoint);
+				++NumWritten;
+			}
+
+			check(Count == NumWritten);
+			return Count;
 		};
 
 		// Small verification to catch any problem that could arise in the iteration body
@@ -133,7 +174,7 @@ namespace PCGMeshSampler
 			return true;
 		}
 
-		return FPCGAsync::AsyncProcessingOneToOneEx(&Context->AsyncState, Context->StartingIndices.Last(), []() {}, IterationBody, /*bEnableTimeSlicing=*/true);
+		return FPCGAsync::AsyncProcessingOneToOneRangeEx(&Context->AsyncState, Context->StartingIndices.Last(), []() {}, IterationBody, /*bEnableTimeSlicing=*/true);
 	}
 
 	bool PoissonSampling(const UPCGMeshSamplerSettings* Settings, FPCGMeshSamplerContext* Context, FPCGMeshSamplerContext::SetPointDensityFunc SetPointDensityPtr)
@@ -143,6 +184,7 @@ namespace PCGMeshSampler
 		// For Poisson sampling, we are calling a "all-in-one" function, where we don't have control for timeslicing.
 		// Since Poisson sampling can be expensive (depending on the radius used), we will do the sampling in a future, put this task to sleep,
 		// and wait for the sampling to wake us up.
+		// @todo_pcg: convert to use ScheduleGeneric + ExecutionDependencies
 		if (Context->SamplingFutures.IsEmpty())
 		{
 			// Put the task asleep
@@ -183,7 +225,7 @@ namespace PCGMeshSampler
 						return true;
 					}
 
-					TArray<FPCGPoint>& Points = Context->OutPointData[DataIndex]->GetMutablePoints();
+					TArray<FPCGPoint> Points;
 					Points.Reserve(PointSampling.Samples.Num());
 
 					int Count = 0;
@@ -209,9 +251,18 @@ namespace PCGMeshSampler
 						OutPoint.Steepness = Settings->PointSteepness;
 
 						Context->SetPointColorAndDensity(SetPointDensityPtr, TriangleId, BarycentricCoords, OutPoint, DataIndex);
-						Context->SetAttributeValues(Settings->UVChannel, TriangleId, BarycentricCoords, OutPoint, DataIndex);
+						Context->SetAttributeValues(Settings->UVChannel, TriangleId, BarycentricCoords, OutPoint.MetadataEntry, DataIndex);
+
+						OutPoint.Seed = PCGHelpers::ComputeSeedFromPosition(OutPoint.Transform.GetLocation());
+					}
 
-						UPCGBlueprintHelpers::SetSeedFromPosition(OutPoint);
+					Context->OutputData[DataIndex]->SetNumPoints(Points.Num(), /*bInitializeValues=*/false);
+					// @todo_pcg: could probably be optimzed to save a little memory
+					Context->OutputData[DataIndex]->AllocateProperties(EPCGPointNativeProperties::All);
+					FPCGPointValueRanges OutRanges(Context->OutputData[DataIndex], /*bAllocate=*/false);
+					for (int32 PointIndex = 0; PointIndex < Points.Num(); ++PointIndex)
+					{
+						OutRanges.SetFromPoint(PointIndex, Points[PointIndex]);
 					}
 
 					// Unpause the task
@@ -236,7 +287,7 @@ namespace PCGMeshSampler
 	}
 }
 
-void FPCGMeshSamplerContext::SetAttributeValues(int32 UVChannel, int32 TriangleId, const FVector& BarycentricCoord, FPCGPoint& OutPoint, int32 DataIndex)
+void FPCGMeshSamplerContext::SetAttributeValues(int32 UVChannel, int32 TriangleId, const FVector& BarycentricCoord, int64& MetadataEntry, int32 DataIndex)
 {
 	if (UVAttributes.IsValidIndex(DataIndex) && UVAttributes[DataIndex])
 	{
@@ -245,9 +296,9 @@ void FPCGMeshSamplerContext::SetAttributeValues(int32 UVChannel, int32 TriangleI
 		UGeometryScriptLibrary_MeshQueryFunctions::GetInterpolatedTriangleUV(DynamicMeshes[DataIndex], /*UVSetIndex=*/ UVChannel, TriangleId, BarycentricCoord, bHasValidUVs, InterpolatedUV);
 		if (bHasValidUVs)
 		{
-			check(OutPointData[DataIndex]);
-			OutPointData[DataIndex]->Metadata->InitializeOnSet(OutPoint.MetadataEntry);
-			UVAttributes[DataIndex]->SetValue(OutPoint.MetadataEntry, InterpolatedUV);
+			check(OutputData[DataIndex]);
+			OutputData[DataIndex]->Metadata->InitializeOnSet(MetadataEntry);
+			UVAttributes[DataIndex]->SetValue(MetadataEntry, InterpolatedUV);
 		}
 	}
 
@@ -258,22 +309,22 @@ void FPCGMeshSamplerContext::SetAttributeValues(int32 UVChannel, int32 TriangleI
 
 		if (bIsValidTriangle)
 		{
-			check(OutPointData[DataIndex]);
-			OutPointData[DataIndex]->Metadata->InitializeOnSet(OutPoint.MetadataEntry);
-			MaterialIdAttributes[DataIndex]->SetValue(OutPoint.MetadataEntry, MaterialId);
+			check(OutputData[DataIndex]);
+			OutputData[DataIndex]->Metadata->InitializeOnSet(MetadataEntry);
+			MaterialIdAttributes[DataIndex]->SetValue(MetadataEntry, MaterialId);
 
 			if (MaterialAttributes.IsValidIndex(DataIndex) && MaterialAttributes[DataIndex] && AssetMaterialList.IsValidIndex(DataIndex) && AssetMaterialList[DataIndex].IsValidIndex(MaterialId) && AssetMaterialList[DataIndex][MaterialId])
 			{
-				MaterialAttributes[DataIndex]->SetValue(OutPoint.MetadataEntry, AssetMaterialList[DataIndex][MaterialId]);
+				MaterialAttributes[DataIndex]->SetValue(MetadataEntry, AssetMaterialList[DataIndex][MaterialId]);
 			}
 		}
 	}
 
 	if (TriangleIdAttributes.IsValidIndex(DataIndex) && TriangleIdAttributes[DataIndex])
 	{
-		check(OutPointData[DataIndex]);
-		OutPointData[DataIndex]->Metadata->InitializeOnSet(OutPoint.MetadataEntry);
-		TriangleIdAttributes[DataIndex]->SetValue(OutPoint.MetadataEntry, TriangleId);
+		check(OutputData[DataIndex]);
+		OutputData[DataIndex]->Metadata->InitializeOnSet(MetadataEntry);
+		TriangleIdAttributes[DataIndex]->SetValue(MetadataEntry, TriangleId);
 	}
 }
 
@@ -364,6 +415,8 @@ void UPCGMeshSamplerSettings::GetStaticTrackedKeys(FPCGSelectionKeyToSettingsMap
 }
 #endif // WITH_EDITOR
 
+PRAGMA_DISABLE_DEPRECATION_WARNINGS
+
 FPCGMeshSamplerContext::~FPCGMeshSamplerContext()
 {
 	// The context can be destroyed if the task is canceled. In that case, we need to notify the futures that they should stop,
@@ -378,6 +431,8 @@ FPCGMeshSamplerContext::~FPCGMeshSamplerContext()
 	}
 }
 
+PRAGMA_ENABLE_DEPRECATION_WARNINGS
+
 void FPCGMeshSamplerContext::AddExtraStructReferencedObjects(FReferenceCollector& Collector)
 {
 	for (TObjectPtr<UDynamicMesh>& DynamicMesh : DynamicMeshes)
@@ -581,8 +636,11 @@ bool FPCGMeshSamplerElement::PrepareDataInternal(FPCGContext* InContext) const
 		}
 
 		TArray<FPCGTaggedData>& Outputs = InContext->OutputData.TaggedData;
-		UPCGPointData* CurrentOutPointData = FPCGContext::NewObject_AnyThread<UPCGPointData>(Context);
-		Context->OutPointData.Emplace(CurrentOutPointData);
+		UPCGBasePointData* CurrentOutPointData = FPCGContext::NewPointData_AnyThread(Context);
+		Context->OutputData.Emplace(CurrentOutPointData);
+PRAGMA_DISABLE_DEPRECATION_WARNINGS
+		Context->OutPointData.Emplace(Cast<UPCGPointData>(CurrentOutPointData));
+PRAGMA_ENABLE_DEPRECATION_WARNINGS
 
 		int32 NumIterations = -1;
 
@@ -619,7 +677,9 @@ bool FPCGMeshSamplerElement::PrepareDataInternal(FPCGContext* InContext) const
 		if (NumIterations != -1)
 		{
 			Context->StartingIndices[i + 1] = Context->StartingIndices[i] + NumIterations;
-			CurrentOutPointData->GetMutablePoints().SetNumUninitialized(NumIterations);
+			CurrentOutPointData->SetNumPoints(NumIterations, /*bInitializeValues=*/false);
+			// @todo_pcg: could probably be optimzed to save a little memory 
+			CurrentOutPointData->AllocateProperties(EPCGPointNativeProperties::All);
 		}
 
 		auto CreateAttribute = [this, Context, Metadata = CurrentOutPointData->Metadata]<typename T>(const bool bShouldAdd, const FName AttributeName, T DefaultValue, bool bAllowInterpolation, TArray<FPCGMetadataAttribute<T>*>& OutAttributesArray, const TCHAR* What)
@@ -659,7 +719,7 @@ bool FPCGMeshSamplerElement::ExecuteInternal(FPCGContext* InContext) const
 
 	check(Context);
 
-	if (!Context->bDataPrepared || Context->OutPointData.IsEmpty())
+	if (!Context->bDataPrepared || Context->OutputData.IsEmpty())
 	{
 		return true;
 	}
diff --git a/Engine/Plugins/PCGInterops/PCGGeometryScriptInterop/Source/PCGGeometryScriptInterop/Public/Data/PCGDynamicMeshData.h b/Engine/Plugins/PCGInterops/PCGGeometryScriptInterop/Source/PCGGeometryScriptInterop/Public/Data/PCGDynamicMeshData.h
--- a/Engine/Plugins/PCGInterops/PCGGeometryScriptInterop/Source/PCGGeometryScriptInterop/Public/Data/PCGDynamicMeshData.h
+++ b/Engine/Plugins/PCGInterops/PCGGeometryScriptInterop/Source/PCGGeometryScriptInterop/Public/Data/PCGDynamicMeshData.h
@@ -72,12 +72,14 @@ protected:
 	virtual UPCGSpatialData* CopyInternal(FPCGContext* Context) const override;
 public:
 	virtual const UPCGPointData* ToPointData(FPCGContext* Context, const FBox& InBounds) const override;
+	virtual const UPCGPointArrayData* ToPointArrayData(FPCGContext* Context, const FBox& InBounds) const override;
 	//~End UPCGSpatialData interface
 
 private:
 	// const but will set the mutable CachedBounds
 	void ResetBounds() const;
 
+	const UPCGBasePointData* ToBasePointData(FPCGContext* Context, const FBox& InBounds, TSubclassOf<UPCGBasePointData> PointDataClass) const;
 protected:
 	UPROPERTY(BlueprintReadOnly, VisibleAnywhere, Instanced, Category="DynamicMesh")
 	TObjectPtr<UDynamicMesh> DynamicMesh;
diff --git a/Engine/Plugins/PCGInterops/PCGGeometryScriptInterop/Source/PCGGeometryScriptInterop/Public/Elements/PCGMeshSampler.h b/Engine/Plugins/PCGInterops/PCGGeometryScriptInterop/Source/PCGGeometryScriptInterop/Public/Elements/PCGMeshSampler.h
--- a/Engine/Plugins/PCGInterops/PCGGeometryScriptInterop/Source/PCGGeometryScriptInterop/Public/Elements/PCGMeshSampler.h
+++ b/Engine/Plugins/PCGInterops/PCGGeometryScriptInterop/Source/PCGGeometryScriptInterop/Public/Elements/PCGMeshSampler.h
@@ -185,11 +185,14 @@ protected:
 struct FPCGMeshSamplerContext : public FPCGLoadObjectsFromPathContext
 {
 public:
-	~FPCGMeshSamplerContext();
+PRAGMA_DISABLE_DEPRECATION_WARNINGS
+	FPCGMeshSamplerContext() = default;
+	virtual ~FPCGMeshSamplerContext();
+PRAGMA_ENABLE_DEPRECATION_WARNINGS
 
 	using SetPointDensityFunc = void(*)(const FLinearColor&, FPCGPoint&);
 
-	void SetAttributeValues(int32 UVChannel, int32 TriangleId, const FVector& BarycentricCoord, FPCGPoint& OutPoint, int32 DataIndex);
+	void SetAttributeValues(int32 UVChannel, int32 TriangleId, const FVector& BarycentricCoord, int64& MetadataEntry, int32 DataIndex);
 	void SetPointColorAndDensity(SetPointDensityFunc SetPointDensityFuncPtr, int32 TriangleId, const FVector& BarycentricCoord, FPCGPoint& OutPoint, int32 DataIndex);
 
 protected:
@@ -206,8 +209,11 @@ public:
 	TArray<FGeometryScriptIndexList> TriangleIds;
 
 	// Output point data.
+	UE_DEPRECATED(5.6, "Use OutputData instead")
 	TArray<UPCGPointData*> OutPointData;
 
+	TArray<UPCGBasePointData*> OutputData;
+
 	// Optional attributes.
 	TArray<FPCGMetadataAttribute<FVector2D>*> UVAttributes;
 	TArray<FPCGMetadataAttribute<int32>*> TriangleIdAttributes;
