[PCG] LandscapeCache early out if CreateCacheEntry is called outside of GameThread (this can happen if Landscape cache  gets invalidated between the PrepareForSpatialQuery and the Sampling)
- A better fix would be to keep the invalidated data (introduce some kind of version in the cache key) but for now this avoid potential crashes

#rnx

diff --git a/Engine/Plugins/PCG/Source/PCG/Private/Grid/PCGLandscapeCache.cpp b/Engine/Plugins/PCG/Source/PCG/Private/Grid/PCGLandscapeCache.cpp
--- a/Engine/Plugins/PCG/Source/PCG/Private/Grid/PCGLandscapeCache.cpp
+++ b/Engine/Plugins/PCG/Source/PCG/Private/Grid/PCGLandscapeCache.cpp
@@ -71,6 +71,11 @@ FPCGLandscapeCacheEntry* FPCGLandscapeCacheEntry::CreateCacheEntry(ULandscapeInf
 {
 	TRACE_CPUPROFILER_EVENT_SCOPE(FPCGLandscapeCacheEntry::CreateCacheEntry);
 
+	if (!IsInGameThread() && !IsInParallelGameThread())
+	{
+		return nullptr;
+	}
+
 	FLandscapeComponentDataInterface CDI(InComponent, 0, /*bWorkInEditingLayer=*/false);
 
 	// CDI.GetRawHeightData() may be nullptr if InComponent has no texture data.
@@ -707,6 +712,7 @@ void UPCGLandscapeCache::PrimeCache()
 
 	ParallelFor(CacheEntriesToBuild.Num(), [this, &CacheEntriesToBuild, &NewEntries](int32 EntryIndex)
 	{
+		FTaskTagScope Scope(ETaskTag::EParallelGameThread);
 		const TTuple<ULandscapeInfo*, ULandscapeComponent*, CacheMapKey>& CacheEntryInfo = CacheEntriesToBuild[EntryIndex];
 		NewEntries[EntryIndex] = FPCGLandscapeCacheEntry::CreateCacheEntry(CacheEntryInfo.Get<0>(), CacheEntryInfo.Get<1>());
 		NewEntries[EntryIndex]->OwningCache = this;
