[PCG] Loc fix

#rnx

diff --git a/Engine/Plugins/PCG/Source/PCGEditor/Private/PCGEditorModule.cpp b/Engine/Plugins/PCG/Source/PCGEditor/Private/PCGEditorModule.cpp
--- a/Engine/Plugins/PCG/Source/PCGEditor/Private/PCGEditorModule.cpp
+++ b/Engine/Plugins/PCG/Source/PCGEditor/Private/PCGEditorModule.cpp
@@ -721,8 +721,8 @@ void FPCGEditorModule::PopulateMenuActions(FMenuBuilder& MenuBuilder)
 		NAME_None);
 
 	MenuBuilder.AddMenuEntry(
-	LOCTEXT("GenerateAllComponents", "Generate all dirty PCG components"),
-	LOCTEXT("GenerateAllComponents_Tooltip", "Generate or refresh all the loaded and dirty PCG components."),
+	LOCTEXT("GenerateAllDirtyComponents", "Generate all dirty PCG components"),
+	LOCTEXT("GenerateAllDirtyComponents_Tooltip", "Generate or refresh all the loaded and dirty PCG components."),
 	FSlateIcon(),
 	FUIAction(
 		FExecuteAction::CreateLambda([]()
