[PCG] Multiple fixes on metadata ops

* Make sure we do not create empty data if the default value doesn't exist.
* Make sure we can edit the input source if the default value is not activated.
* Small clean up of unecessary jumps when checking if a default value is activated.

#rnx

diff --git a/Engine/Plugins/PCG/Source/PCG/Private/Elements/Metadata/PCGMetadataOpElementBase.cpp b/Engine/Plugins/PCG/Source/PCG/Private/Elements/Metadata/PCGMetadataOpElementBase.cpp
--- a/Engine/Plugins/PCG/Source/PCG/Private/Elements/Metadata/PCGMetadataOpElementBase.cpp
+++ b/Engine/Plugins/PCG/Source/PCG/Private/Elements/Metadata/PCGMetadataOpElementBase.cpp
@@ -120,21 +120,21 @@ bool UPCGMetadataSettingsBase::IsPinDefaultValueEnabled(const FName PinLabel) co
 
 bool UPCGMetadataSettingsBase::IsPinDefaultValueActivated(const FName PinLabel) const
 {
-	if (!DefaultValuesAreEnabled() || !IsPinDefaultValueEnabled(PinLabel))
+	if (!IsPinDefaultValueEnabled(PinLabel))
 	{
 		return false;
 	}
 
-	for (uint32 I = 0; I < GetOperandNum(); ++I)
+	const uint32 PinIndex = GetInputPinIndex(PinLabel);
+	if (PinIndex != static_cast<uint32>(INDEX_NONE))
 	{
-		if (PinLabel == GetInputPinLabel(I))
-		{
-			const FName PropertyName = PCGMetadataBase::Helpers::GetDefaultValuePropertyName(GetInputPinIndex(PinLabel));
-			return DefaultValues.IsPropertyActivated(PropertyName);
-		}
+		const FName PropertyName = PCGMetadataBase::Helpers::GetDefaultValuePropertyName(PinIndex);
+		return DefaultValues.IsPropertyActivated(PropertyName);
+	}
+	else
+	{
+		return false;
 	}
-
-	return false;
 }
 
 uint32 UPCGMetadataSettingsBase::GetInputPinToForward() const
@@ -491,7 +491,7 @@ bool UPCGMetadataSettingsBase::CanEditInputSource(const FProperty* InProperty, c
 		{
 			if (const UPCGNode* Node = Cast<UPCGNode>(GetOuter()))
 			{
-				return Node->IsInputPinConnected(GetInputPinLabel(SourceIndex));
+				return !IsPinDefaultValueActivated(GetInputPinLabel(SourceIndex)) || Node->IsInputPinConnected(GetInputPinLabel(SourceIndex));
 			}
 		}
 	}
@@ -682,36 +682,34 @@ bool FPCGMetadataElementBase::PrepareDataInternal(FPCGContext* Context) const
 			TArray<FPCGTaggedData> CurrentPinInputData = Context->InputData.GetInputsByPin(CurrentPinLabel);
 
 			// This only needs to be checked once
-			if (Settings->DefaultValuesAreEnabled() && !bIsInputConnected && CurrentPinInputData.IsEmpty())
+			if (Settings->DefaultValuesAreEnabled() && !bIsInputConnected && CurrentPinInputData.IsEmpty() && Settings->IsPinDefaultValueActivated(CurrentPinLabel))
 			{
 				FPCGTaggedData& DefaultData = CurrentPinInputData.Emplace_GetRef();
 				DefaultData.Pin = CurrentPinLabel;
-				if (Settings->IsPinDefaultValueActivated(CurrentPinLabel))
+
+				// @todo_pcg: Future optimizations/refactors - cache the param data on the settings, or use accessors on the default value struct, etc.
+				// Create from the Default Value Container if it exists
+				DefaultData.Data = Settings->CreateDefaultValueParamData(Context, CurrentPinLabel);
+
+				// Couldn't create a default value
+				if (!DefaultData.Data)
+				{
+					PCGE_LOG(Error, GraphAndLog, FText::Format(LOCTEXT("CantCreateDefaultValue", "Pin '{0}' supports default value but we could not create it."), FText::FromName(CurrentPinLabel)));
+					return EPCGTimeSliceInitResult::AbortExecution;
+				}
+				else
 				{
-					// @todo_pcg: Future optimizations/refactors - cache the param data on the settings, or use accessors on the default value struct, etc.
-					// Create from the Default Value Container if it exists
-					DefaultData.Data = Settings->CreateDefaultValueParamData(Context, CurrentPinLabel);
+					// Need to make sure the param data is properly tracked by the context to prevent garbage collection
+					TimeSlicedContext->TrackObject(DefaultData.Data);
 
-					// Couldn't create a default value
-					if (!DefaultData.Data)
+					// Need to make sure the param data has at least one entry
+					UPCGMetadata* DefaultParamMetadata = CastChecked<UPCGParamData>(DefaultData.Data)->Metadata;
+					if (DefaultParamMetadata->GetLocalItemCount() == 0)
 					{
-						PCGE_LOG(Error, GraphAndLog, FText::Format(LOCTEXT("CantCreateDefaultValue", "Pin '{0}' supports default value but we could not create it."), FText::FromName(CurrentPinLabel)));
-						return EPCGTimeSliceInitResult::AbortExecution;
+						DefaultParamMetadata->AddEntry();
 					}
-					else
-					{
-						// Need to make sure the param data is properly tracked by the context to prevent garbage collection
-						TimeSlicedContext->TrackObject(DefaultData.Data);
 
-						// Need to make sure the param data has at least one entry
-						UPCGMetadata* DefaultParamMetadata = CastChecked<UPCGParamData>(DefaultData.Data)->Metadata;
-						if (DefaultParamMetadata->GetLocalItemCount() == 0)
-						{
-							DefaultParamMetadata->AddEntry();
-						}
-
-						OutState.DefaultValueOverriddenPins[OperandPinIndex] = true;
-					}
+					OutState.DefaultValueOverriddenPins[OperandPinIndex] = true;
 				}
 			}
 
