[PCG] Support targetting different domains for CopyAttributes.

* Will throw an error if we try to copy from elements to global and there are more than one value
* Copy All Attributes will have more options for copy (All Domains or Target domains)

#rnx

diff --git a/Engine/Plugins/PCG/Source/PCG/Private/Data/PCGBasePointData.cpp b/Engine/Plugins/PCG/Source/PCG/Private/Data/PCGBasePointData.cpp
--- a/Engine/Plugins/PCG/Source/PCG/Private/Data/PCGBasePointData.cpp
+++ b/Engine/Plugins/PCG/Source/PCG/Private/Data/PCGBasePointData.cpp
@@ -449,7 +449,9 @@ FPCGAttributeAccessorMethods UPCGBasePointData::GetPointAccessorMethods()
 		if (DomainID.IsDefault() || DomainID == PCGMetadataDomainID::Elements)
 		{
 			UPCGBasePointData* PointData = CastChecked<UPCGBasePointData>(InData);
-			return MakeUnique<FPCGAttributeAccessorKeysPointIndices>(PointData);
+
+			// If we know the keys are used with a selector of an attribute, allocate the entries.
+			return MakeUnique<FPCGAttributeAccessorKeysPointIndices>(PointData, /*bAllocateMetadataEntries=*/InSelector.GetSelection() == EPCGAttributePropertySelection::Attribute);
 		}
 		else
 		{
diff --git a/Engine/Plugins/PCG/Source/PCG/Private/Elements/PCGCopyAttributes.cpp b/Engine/Plugins/PCG/Source/PCG/Private/Elements/PCGCopyAttributes.cpp
--- a/Engine/Plugins/PCG/Source/PCG/Private/Elements/PCGCopyAttributes.cpp
+++ b/Engine/Plugins/PCG/Source/PCG/Private/Elements/PCGCopyAttributes.cpp
@@ -252,7 +252,29 @@ bool FPCGCopyAttributesElement::ExecuteInternal(FPCGContext* Context) const
 
 			if (Settings->bCopyAllAttributes)
 			{
-				bSuccess |= PCGMetadataHelpers::CopyAllAttributes(SourceData, OutputData, Context);
+				PCGMetadataHelpers::FPCGCopyAllAttributesParams Params
+				{
+					.SourceData = SourceData,
+					.TargetData = OutputData,
+					.OptionalContext = Context,
+				};
+
+				if (Settings->bCopyAllDomains)
+				{
+					Params.InitializeMappingForAllDomains();
+				}
+				else
+				{
+					Params.InitializeMappingFromDomainNames(Settings->MetadataDomainsMapping);
+				}
+
+				// Nothing to do
+				if (Params.DomainMapping.IsEmpty())
+				{
+					return;
+				}
+				
+				bSuccess |= PCGMetadataHelpers::CopyAllAttributes(Params);
 			}
 			else
 			{
diff --git a/Engine/Plugins/PCG/Source/PCG/Private/Elements/PCGCreateAttribute.cpp b/Engine/Plugins/PCG/Source/PCG/Private/Elements/PCGCreateAttribute.cpp
--- a/Engine/Plugins/PCG/Source/PCG/Private/Elements/PCGCreateAttribute.cpp
+++ b/Engine/Plugins/PCG/Source/PCG/Private/Elements/PCGCreateAttribute.cpp
@@ -471,7 +471,29 @@ bool FPCGAddAttributeElement::ExecuteInternal(FPCGContext* Context) const
 			{
 				if (const UPCGParamData* ParamData = Cast<UPCGParamData>(SourceParamData.Data))
 				{
-					bSuccess &= PCGMetadataHelpers::CopyAllAttributes(ParamData, TargetData, Context);
+					PCGMetadataHelpers::FPCGCopyAllAttributesParams Params
+					{
+						.SourceData = ParamData,
+						.TargetData = TargetData,
+						.OptionalContext = Context
+					};
+
+					if (Settings->bCopyAllDomains)
+					{
+						Params.InitializeMappingForAllDomains();
+					}
+					else
+					{
+						Params.InitializeMappingFromDomainNames(Settings->MetadataDomainsMapping);
+					}
+
+					// Nothing to do
+					if (Params.DomainMapping.IsEmpty())
+					{
+						continue;
+					}
+
+					bSuccess &= PCGMetadataHelpers::CopyAllAttributes(Params);
 				}
 			}
 		}
diff --git a/Engine/Plugins/PCG/Source/PCG/Private/Graph/PCGGraphExecutor.cpp b/Engine/Plugins/PCG/Source/PCG/Private/Graph/PCGGraphExecutor.cpp
--- a/Engine/Plugins/PCG/Source/PCG/Private/Graph/PCGGraphExecutor.cpp
+++ b/Engine/Plugins/PCG/Source/PCG/Private/Graph/PCGGraphExecutor.cpp
@@ -2168,7 +2168,13 @@ void FPCGGraphExecutor::CombineParams(FPCGGraphTask& Task)
 				}
 				else
 				{
-					bSuccess &= PCGMetadataHelpers::CopyAllAttributes(ParamData, CombinedParamData, nullptr);
+					PCGMetadataHelpers::FPCGCopyAllAttributesParams Params
+					{
+						.SourceData = ParamData,
+						.TargetData = CombinedParamData
+					};
+					
+					bSuccess &= PCGMetadataHelpers::CopyAllAttributes(Params);
 				}
 			}
 		}
diff --git a/Engine/Plugins/PCG/Source/PCG/Private/Helpers/PCGMetadataHelpers.cpp b/Engine/Plugins/PCG/Source/PCG/Private/Helpers/PCGMetadataHelpers.cpp
--- a/Engine/Plugins/PCG/Source/PCG/Private/Helpers/PCGMetadataHelpers.cpp
+++ b/Engine/Plugins/PCG/Source/PCG/Private/Helpers/PCGMetadataHelpers.cpp
@@ -3,6 +3,7 @@
 #include "Helpers/PCGMetadataHelpers.h"
 
 #include "PCGParamData.h"
+#include "Data/PCGBasePointData.h"
 #include "Data/PCGSpatialData.h"
 #include "Elements/Metadata/PCGMetadataElementCommon.h"
 #include "Metadata/Accessors/IPCGAttributeAccessor.h"
@@ -146,17 +147,16 @@ namespace PCGMetadataHelpers
 			const FPCGAttributePropertyOutputSelector& OutputTarget = SelectorTuple.Get<1>();
 			const EPCGMetadataTypes RequestedOutputType = SelectorTuple.Get<2>();
 
-			const FName LocalSourceAttribute = InputSource.GetName();
-			const FName LocalDestinationAttribute = OutputTarget.GetName();
+			const FPCGAttributeIdentifier LocalSourceAttribute = FPCGAttributeIdentifier(InputSource.GetName(), SourceData->GetMetadataDomainIDFromSelector(InputSource));
+			const FPCGAttributeIdentifier LocalDestinationAttribute = FPCGAttributeIdentifier(OutputTarget.GetName(), TargetData->GetMetadataDomainIDFromSelector(OutputTarget));
 
 			if (InputSource.GetSelection() == EPCGAttributePropertySelection::Attribute && !SourceMetadata->HasAttribute(LocalSourceAttribute))
 			{
-				PCGLog::LogWarningOnGraph(FText::Format(LOCTEXT("InputMissingAttribute", "Input does not have the '{0}' attribute"), FText::FromName(LocalSourceAttribute)), OptionalContext);
+				PCGLog::LogWarningOnGraph(FText::Format(LOCTEXT("InputMissingAttribute", "Input does not have the '{0}' attribute"), InputSource.GetDisplayText()), OptionalContext);
 				continue;
 			}
 
 			// We need accessors if we have a multi entry source attribute or we have extractors
-			const bool bIsMultiEntries = SourceData->IsA<UPCGParamData>() && SourceMetadata->GetLocalItemCount() > 1;
 			const bool bInputHasAnyExtra = !InputSource.GetExtraNames().IsEmpty();
 			const bool bOutputHasAnyExtra = !OutputTarget.GetExtraNames().IsEmpty();
 			const bool bSourceIsAttribute = InputSource.GetSelection() == EPCGAttributePropertySelection::Attribute;
@@ -164,7 +164,7 @@ namespace PCGMetadataHelpers
 			// Cast is only required if it is on an output attribute that has no extra (that we will create)
 			const bool bOutputTypeCast = bTargetIsAttribute && !bOutputHasAnyExtra && (RequestedOutputType != EPCGMetadataTypes::Unknown);
 
-			const bool bNeedAccessors = bIsMultiEntries || bInputHasAnyExtra || bOutputHasAnyExtra || !bSourceIsAttribute || !bTargetIsAttribute || bOutputTypeCast;
+			const bool bNeedAccessors = bInputHasAnyExtra || bOutputHasAnyExtra || !bSourceIsAttribute || !bTargetIsAttribute || bOutputTypeCast;
 
 			// If no accessor, copy over the attribute
 			if (!bNeedAccessors)
@@ -179,10 +179,20 @@ namespace PCGMetadataHelpers
 				// Presence of attribute was already checked before, this should not return null
 				check(SourceAttribute);
 
-				if (FPCGMetadataAttributeBase* NewAttr = TargetMetadata->CopyAttribute(SourceAttribute, LocalDestinationAttribute, /*bKeepParent=*/false, /*bCopyEntries=*/bSameOrigin, /*bCopyValues=*/true))
+				FPCGMetadataDomain* TargetMetadataDomain = TargetMetadata->GetMetadataDomain(LocalDestinationAttribute.MetadataDomain);
+				if (!TargetMetadataDomain)
+				{
+					PCGLog::Metadata::LogInvalidMetadataDomain(OutputTarget, OptionalContext);
+					continue;
+				}
+
+				// Copy entries only if they come from the same data and they are on the same domain.
+				const bool bCopyEntries = bSameOrigin && SourceAttribute->GetMetadataDomain()->GetDomainID() == TargetMetadataDomain->GetDomainID();
+				const bool bSourceSupportsMultiEntries = SourceMetadata->MetadataDomainSupportsMultiEntries(LocalSourceAttribute.MetadataDomain);
+				if (FPCGMetadataAttributeBase* NewAttr = TargetMetadataDomain->CopyAttribute(SourceAttribute, LocalDestinationAttribute.Name, /*bKeepParent=*/false, /*bCopyEntries=*/bCopyEntries, /*bCopyValues=*/true))
 				{
 					// And finally, we create our source and target keys, to get all the metadata entry keys, to remap if we are not the same origin and that we have multiple entries.
-					if (!bSameOrigin)
+					if (!bCopyEntries && bSourceSupportsMultiEntries)
 					{
 						TUniquePtr<const IPCGAttributeAccessorKeys> SourceKeys = PCGAttributeAccessorHelpers::CreateConstKeys(SourceData, InputSource);
 						TUniquePtr<IPCGAttributeAccessorKeys> TargetKeys = PCGAttributeAccessorHelpers::CreateKeys(TargetData, OutputTarget);
@@ -230,7 +240,7 @@ namespace PCGMetadataHelpers
 								for (int EntryIndex = 0; EntryIndex < AllTargetEntryKeysPtr.Num(); ++EntryIndex)
 								{
 									PCGMetadataEntryKey& EntryKey = *AllTargetEntryKeysPtr[EntryIndex];
-									if (EntryKey == PCGInvalidEntryKey || (EntryKey < TargetMetadata->GetItemKeyCountForParent()))
+									if (EntryKey == PCGInvalidEntryKey || (EntryKey < TargetMetadataDomain->GetItemKeyCountForParent()))
 									{
 										AllTargetEntryKeysPtrTemp.Add(&EntryKey);
 									}
@@ -239,7 +249,7 @@ namespace PCGMetadataHelpers
 								if (!AllTargetEntryKeysPtrTemp.IsEmpty())
 								{
 									TRACE_CPUPROFILER_EVENT_SCOPE(FPCGAttributeAccessor::Prepare::AddEntriesInPlace);
-									TargetMetadata->AddEntriesInPlace(AllTargetEntryKeysPtrTemp);
+									TargetMetadataDomain->AddEntriesInPlace(AllTargetEntryKeysPtrTemp);
 								}
 
 								NewAttr->SetValuesFromValueKeys(AllTargetEntryKeysPtr, FinalValueKeys, /*bResetValueOnDefaultValueKey=*/true);
@@ -252,10 +262,14 @@ namespace PCGMetadataHelpers
 							}
 						}
 					}
+					else if (!bCopyEntries && !bSourceSupportsMultiEntries)
+					{
+						NewAttr->SetDefaultValueToFirstEntry();
+					}
 				}
 				else
 				{
-					PCGLog::LogWarningOnGraph(FText::Format(LOCTEXT("FailedCreateNewAttribute", "Failed to create new attribute '{0}'"), FText::FromName(LocalDestinationAttribute)));
+					PCGLog::LogWarningOnGraph(FText::Format(LOCTEXT("FailedCreateNewAttribute", "Failed to create new attribute '{0}'"), OutputTarget.GetDisplayText()), OptionalContext);
 					continue;
 				}
 			}
@@ -266,7 +280,7 @@ namespace PCGMetadataHelpers
 
 				if (!InputAccessor.IsValid() || !InputKeys.IsValid())
 				{
-					PCGLog::LogWarningOnGraph(LOCTEXT("FailedToCreateInputAccessor", "Failed to create input accessor or iterator"), OptionalContext);
+					PCGLog::Metadata::LogFailToCreateAccessorError(InputSource, OptionalContext);
 					continue;
 				}
 
@@ -289,7 +303,7 @@ namespace PCGMetadataHelpers
 				// If it has any extra, it will try to write to it.
 				if (!bOutputHasAnyExtra && bTargetIsAttribute)
 				{
-					auto CreateAttribute = [TargetMetadata, LocalDestinationAttribute, &InputAccessor](auto Dummy) -> bool
+					auto CreateAttribute = [TargetMetadata, &OutputTarget, &InputAccessor, OptionalContext](auto Dummy) -> bool
 					{
 						using AttributeType = decltype(Dummy);
 						AttributeType DefaultValue{};
@@ -299,12 +313,19 @@ namespace PCGMetadataHelpers
 							DefaultValue = PCG::Private::MetadataTraits<AttributeType>::ZeroValue();
 						}
 
-						return PCGMetadataElementCommon::ClearOrCreateAttribute<AttributeType>(TargetMetadata, LocalDestinationAttribute, std::move(DefaultValue)) != nullptr;
+						if (!PCGMetadataElementCommon::ClearOrCreateAttribute<AttributeType>(TargetMetadata, OutputTarget, std::move(DefaultValue)))
+						{
+							PCGLog::Metadata::LogFailToCreateAttributeError<AttributeType>(OutputTarget.GetDisplayText(), OptionalContext);
+							return false;
+						}
+						else
+						{
+							return true;
+						}
 					};
 
 					if (!PCGMetadataAttribute::CallbackWithRightType(OutputType, CreateAttribute))
 					{
-						PCGLog::LogWarningOnGraph(FText::Format(LOCTEXT("FailedToCreateNewAttribute", "Failed to create new attribute '{0}'"), FText::FromName(LocalDestinationAttribute)), OptionalContext);
 						continue;
 					}
 				}
@@ -314,7 +335,7 @@ namespace PCGMetadataHelpers
 
 				if (!OutputAccessor.IsValid() || !OutputKeys.IsValid())
 				{
-					PCGLog::LogWarningOnGraph(LOCTEXT("FailedToCreateOutputAccessor", "Failed to create output accessor or iterator"), OptionalContext);
+					PCGLog::Metadata::LogFailToCreateAccessorError(OutputTarget, OptionalContext);
 					continue;
 				}
 
@@ -368,33 +389,130 @@ namespace PCGMetadataHelpers
 		return CopyAttributes(InParams.TargetData, InParams.SourceData, AttributeSelectors, InParams.bSameOrigin, InParams.OptionalContext);
 	}
 
+	void FPCGCopyAllAttributesParams::InitializeMappingFromDomainNames(const TMap<FName, FName>& MetadataDomainsMapping)
+	{
+		if (MetadataDomainsMapping.IsEmpty())
+		{
+			DomainMapping.Empty(1);
+			DomainMapping = {{PCGMetadataDomainID::Default, PCGMetadataDomainID::Default}};
+			return;
+		}
+		
+		DomainMapping.Empty(MetadataDomainsMapping.Num());
+		
+		if (!SourceData || !TargetData)
+		{
+			return;
+		}
+		
+		FPCGAttributePropertySelector TempSelector{};
+	
+		for (const TPair<FName, FName>& It : MetadataDomainsMapping)
+		{
+			TempSelector.SetDomainName(It.Key);
+			const FPCGMetadataDomainID SourceDomain = SourceData->GetMetadataDomainIDFromSelector(TempSelector);
+			if (!SourceDomain.IsValid())
+			{
+				PCGLog::Metadata::LogInvalidMetadataDomain(TempSelector);
+				continue;
+			}
+				
+			TempSelector.SetDomainName(It.Value);
+			const FPCGMetadataDomainID TargetDomain = TargetData->GetMetadataDomainIDFromSelector(TempSelector);
+			if (!TargetDomain.IsValid())
+			{
+				PCGLog::Metadata::LogInvalidMetadataDomain(TempSelector);
+				continue;
+			}
+		
+			DomainMapping.Emplace(SourceDomain, TargetDomain);
+		}
+	}
+
+	void FPCGCopyAllAttributesParams::InitializeMappingForAllDomains()
+	{
+		DomainMapping.Reset();
+		
+		if (!SourceData || !TargetData)
+		{
+			return;
+		}
+
+		for (const FPCGMetadataDomainID SourceDomain : SourceData->GetAllSupportedMetadataDomainIDs())
+		{
+			if (TargetData->IsSupportedMetadataDomainID(SourceDomain))
+			{
+				DomainMapping.Emplace(SourceDomain, SourceDomain);
+			}
+		}
+	}
+
 	bool CopyAllAttributes(const UPCGData* SourceData, UPCGData* TargetData, FPCGContext* OptionalContext)
 	{
-		if (!TargetData || !SourceData)
+		const FPCGCopyAllAttributesParams InParams
+		{
+			.SourceData = SourceData,
+			.TargetData = TargetData,
+			.OptionalContext = OptionalContext
+		};
+
+		return CopyAllAttributes(InParams);
+	}
+
+	bool CopyAllAttributes(const FPCGCopyAllAttributesParams& InParams)
+	{
+		if (!InParams.TargetData || !InParams.SourceData)
 		{
 			return false;
 		}
 
-		const UPCGMetadata* SourceMetadata = SourceData->ConstMetadata();
+		const UPCGMetadata* SourceMetadata = InParams.SourceData->ConstMetadata();
 		if (!SourceMetadata)
 		{
 			return false;
 		}
 
 		TArray<TTuple<FPCGAttributePropertyInputSelector, FPCGAttributePropertyOutputSelector, EPCGMetadataTypes>> AttributeSelectors;
-		TArray<FName> AttributeNames;
+		TArray<FPCGAttributeIdentifier> AttributeIDs;
 		TArray<EPCGMetadataTypes> AttributeTypes;
-		SourceMetadata->GetAttributes(AttributeNames, AttributeTypes);
+		SourceMetadata->GetAllAttributes(AttributeIDs, AttributeTypes);
+
+		const FPCGMetadataDomainID DefaultSourceDomain = SourceMetadata->GetConstDefaultMetadataDomain()->GetDomainID();
 
-		for (const FName& AttributeName : AttributeNames)
+		for (const FPCGAttributeIdentifier& AttributeID : AttributeIDs)
 		{
+			const bool bIsDefaultDomain = AttributeID.MetadataDomain == DefaultSourceDomain;
+			const FPCGMetadataDomainID* TargetDomainId = nullptr;
+			if (!InParams.DomainMapping.IsEmpty())
+			{
+				TargetDomainId = InParams.DomainMapping.Find(AttributeID.MetadataDomain);
+				if (!TargetDomainId && bIsDefaultDomain)
+				{
+					TargetDomainId = InParams.DomainMapping.Find(PCGMetadataDomainID::Default);
+				}
+
+				if (!TargetDomainId && AttributeID.MetadataDomain.IsDefault())
+				{
+					TargetDomainId = InParams.DomainMapping.Find(DefaultSourceDomain);
+				}
+
+				if (!TargetDomainId)
+				{
+					// Didn't find the domain in the mapping, continue...
+					continue;
+				}
+			}
+			
 			TTuple<FPCGAttributePropertyInputSelector, FPCGAttributePropertyOutputSelector, EPCGMetadataTypes>& Selectors = AttributeSelectors.Emplace_GetRef();
-			Selectors.Get<0>().SetAttributeName(AttributeName);
-			Selectors.Get<1>().SetAttributeName(AttributeName);
+			Selectors.Get<0>().SetAttributeName(AttributeID.Name);
+			Selectors.Get<1>().SetAttributeName(AttributeID.Name);
 			Selectors.Get<2>() = EPCGMetadataTypes::Unknown;
+
+			InParams.SourceData->SetDomainFromDomainID(AttributeID.MetadataDomain, Selectors.Get<0>());
+			InParams.TargetData->SetDomainFromDomainID(TargetDomainId ? *TargetDomainId : AttributeID.MetadataDomain, Selectors.Get<1>());
 		}
 
-		return CopyAttributes(TargetData, SourceData, AttributeSelectors, /*bSameOrigin=*/false, OptionalContext);
+		return CopyAttributes(InParams.TargetData, InParams.SourceData, AttributeSelectors, /*bSameOrigin=*/false, InParams.OptionalContext);
 	}
 
 	void ComputePointWeightedAttribute(FPCGMetadataDomain* InOutMetadata, FPCGPoint& OutPoint, const TArrayView<TPair<const FPCGPoint*, float>>& InWeightedPoints, const FPCGMetadataDomain* InMetadata)
diff --git a/Engine/Plugins/PCG/Source/PCG/Private/Metadata/PCGMetadata.cpp b/Engine/Plugins/PCG/Source/PCG/Private/Metadata/PCGMetadata.cpp
--- a/Engine/Plugins/PCG/Source/PCG/Private/Metadata/PCGMetadata.cpp
+++ b/Engine/Plugins/PCG/Source/PCG/Private/Metadata/PCGMetadata.cpp
@@ -202,7 +202,7 @@ void UPCGMetadata::ForEachValidUniqueMetadataDomain(Func InFunc)
 
 FPCGMetadataDomain* UPCGMetadata::GetMetadataDomain(const FPCGMetadataDomainID& InMetadataDomainID)
 {
-	return FindOrCreateMetadataDomain(InMetadataDomainID);
+	return InMetadataDomainID.IsValid() ? FindOrCreateMetadataDomain(InMetadataDomainID) : nullptr;
 }
 
 FPCGMetadataDomain* UPCGMetadata::GetMetadataDomainFromSelector(const FPCGAttributePropertySelector& InSelector)
diff --git a/Engine/Plugins/PCG/Source/PCG/Private/Tests/Elements/PCGCopyAttributeTest.cpp b/Engine/Plugins/PCG/Source/PCG/Private/Tests/Elements/PCGCopyAttributeTest.cpp
--- a/Engine/Plugins/PCG/Source/PCG/Private/Tests/Elements/PCGCopyAttributeTest.cpp
+++ b/Engine/Plugins/PCG/Source/PCG/Private/Tests/Elements/PCGCopyAttributeTest.cpp
@@ -7,7 +7,7 @@
 #include "PCGComponent.h"
 #include "PCGContext.h"
 #include "PCGParamData.h"
-#include "Data/PCGPointData.h"
+#include "Data/PCGBasePointData.h"
 #include "Metadata/PCGMetadata.h"
 #include "Metadata/PCGMetadataAttributeTpl.h"
 
@@ -26,26 +26,42 @@ IMPLEMENT_CUSTOM_SIMPLE_AUTOMATION_TEST(FPCGCopyAttributeTests_Params_MultiValue
 IMPLEMENT_CUSTOM_SIMPLE_AUTOMATION_TEST(FPCGCopyAttributeTests_Params_CopyingToItself, FPCGTestBaseClass, "Plugins.PCG.CopyAttribute.Params.CopyingToItself", PCGTestsCommon::TestFlags)
 IMPLEMENT_CUSTOM_SIMPLE_AUTOMATION_TEST(FPCGCopyAttributeTests_Params_CopyingAllToItself, FPCGTestBaseClass, "Plugins.PCG.CopyAttribute.Params.CopyingAllToItself", PCGTestsCommon::TestFlags)
 
+IMPLEMENT_CUSTOM_SIMPLE_AUTOMATION_TEST(FPCGCopyAttributeTests_MultiDomain_GlobalToGlobal, FPCGTestBaseClass, "Plugins.PCG.CopyAttribute.MultiDomain.GlobalToGlobal", PCGTestsCommon::TestFlags)
+IMPLEMENT_CUSTOM_SIMPLE_AUTOMATION_TEST(FPCGCopyAttributeTests_MultiDomain_GlobalToElements, FPCGTestBaseClass, "Plugins.PCG.CopyAttribute.MultiDomain.GlobalToElements", PCGTestsCommon::TestFlags)
+IMPLEMENT_CUSTOM_SIMPLE_AUTOMATION_TEST(FPCGCopyAttributeTests_MultiDomain_ElementsToGlobal, FPCGTestBaseClass, "Plugins.PCG.CopyAttribute.MultiDomain.ElementsToGlobal", PCGTestsCommon::TestFlags)
+IMPLEMENT_CUSTOM_SIMPLE_AUTOMATION_TEST(FPCGCopyAttributeTests_MultiDomain_ElementsToGlobalTooMany, FPCGTestBaseClass, "Plugins.PCG.CopyAttribute.MultiDomain.ElementsToGlobalTooMany", PCGTestsCommon::TestFlags)
+IMPLEMENT_CUSTOM_SIMPLE_AUTOMATION_TEST(FPCGCopyAttributeTests_MultiDomain_InvalidDomain, FPCGTestBaseClass, "Plugins.PCG.CopyAttribute.MultiDomain.InvalidDomain", PCGTestsCommon::TestFlags)
+
+IMPLEMENT_CUSTOM_SIMPLE_AUTOMATION_TEST(FPCGCopyAttributeTests_MultiDomain_CopyAll_GlobalToGlobal, FPCGTestBaseClass, "Plugins.PCG.CopyAttribute.MultiDomain.CopyAll.GlobalToGlobal", PCGTestsCommon::TestFlags)
+IMPLEMENT_CUSTOM_SIMPLE_AUTOMATION_TEST(FPCGCopyAttributeTests_MultiDomain_CopyAll_GlobalToElements, FPCGTestBaseClass, "Plugins.PCG.CopyAttribute.MultiDomain.CopyAll.GlobalToElements", PCGTestsCommon::TestFlags)
+IMPLEMENT_CUSTOM_SIMPLE_AUTOMATION_TEST(FPCGCopyAttributeTests_MultiDomain_CopyAll_ElementsToGlobal, FPCGTestBaseClass, "Plugins.PCG.CopyAttribute.MultiDomain.CopyAll.ElementsToGlobal", PCGTestsCommon::TestFlags)
+IMPLEMENT_CUSTOM_SIMPLE_AUTOMATION_TEST(FPCGCopyAttributeTests_MultiDomain_CopyAll_ElementsToGlobalTooMany, FPCGTestBaseClass, "Plugins.PCG.CopyAttribute.MultiDomain.CopyAll.ElementsToGlobalTooMany", PCGTestsCommon::TestFlags)
+IMPLEMENT_CUSTOM_SIMPLE_AUTOMATION_TEST(FPCGCopyAttributeTests_MultiDomain_CopyAll_ElementsToElementsExplicit, FPCGTestBaseClass, "Plugins.PCG.CopyAttribute.MultiDomain.CopyAll.ElementsToElementsExplicit", PCGTestsCommon::TestFlags)
+IMPLEMENT_CUSTOM_SIMPLE_AUTOMATION_TEST(FPCGCopyAttributeTests_MultiDomain_CopyAll_ElementsToElementsDefault, FPCGTestBaseClass, "Plugins.PCG.CopyAttribute.MultiDomain.CopyAll.ElementsToElementsDefault", PCGTestsCommon::TestFlags)
+IMPLEMENT_CUSTOM_SIMPLE_AUTOMATION_TEST(FPCGCopyAttributeTests_MultiDomain_CopyAll_AllToAll, FPCGTestBaseClass, "Plugins.PCG.CopyAttribute.MultiDomain.CopyAll.AllToAll", PCGTestsCommon::TestFlags)
+IMPLEMENT_CUSTOM_SIMPLE_AUTOMATION_TEST(FPCGCopyAttributeTests_MultiDomain_CopyAll_InvalidDomain, FPCGTestBaseClass, "Plugins.PCG.CopyAttribute.MultiDomain.CopyAll.InvalidDomain", PCGTestsCommon::TestFlags)
 
 namespace PCGCopyAttributeTests
 {
 	static const FName AttributeName = TEXT("InputAttr");
 
-	UPCGPointData* CreateInputPointData(FPCGContext* Context, const int NumPoints)
+	UPCGBasePointData* CreateInputPointData(FPCGContext* Context, const int NumPoints)
 	{
 		check(Context);
 
-		UPCGPointData* NewPointData = NewObject<UPCGPointData>();
+		UPCGBasePointData* NewPointData = FPCGContext::NewPointData_AnyThread(Context);
 		NewPointData->SetFlags(RF_Transient);
 
 		check(NewPointData);
+		
+		NewPointData->SetNumPoints(NumPoints);
+		NewPointData->AllocateProperties(EPCGPointNativeProperties::Density);
 
-		TArray<FPCGPoint>& Points = NewPointData->GetMutablePoints();
-		Points.SetNum(NumPoints);
+		TPCGValueRange<float> DensityValues = NewPointData->GetDensityValueRange();
 		for (int i = 0; i < NumPoints; ++i)
 		{
 			// Store the index in the density
-			Points[i].Density = i;
+			DensityValues[i] = i;
 		}
 
 		return NewPointData;
@@ -96,25 +112,26 @@ bool FPCGCopyAttributeTests_Points_PropertyToProperty::RunTest(const FString& Pa
 	Settings->OutputTarget.Update(TEXT("$Position.X"));
 
 	TUniquePtr<FPCGContext> Context = TestData.InitializeTestContext();
-	const UPCGPointData* InputPointData = PCGCopyAttributeTests::CreateInputPointData(Context.Get(), NumOfPoints);
+	const UPCGBasePointData* InputPointData = PCGCopyAttributeTests::CreateInputPointData(Context.Get(), NumOfPoints);
 	PCGCopyAttributeTests::ConnectToSourceAndTarget(Context.Get(), InputPointData);
 
 	FPCGElementPtr TestElement = TestData.Settings->GetElement();
 	while (!TestElement->Execute(Context.Get())) {}
 
 	TArray<FPCGTaggedData> OutputTagged = Context->OutputData.GetInputsByPin(PCGPinConstants::DefaultOutputLabel);
-	const UPCGPointData* OutputData = OutputTagged.Num() == 1 ? Cast<const UPCGPointData>(OutputTagged[0].Data) : nullptr;
+	const UPCGBasePointData* OutputData = OutputTagged.Num() == 1 ? Cast<const UPCGBasePointData>(OutputTagged[0].Data) : nullptr;
 
 	UTEST_NOT_NULL("There is a point data in output points", OutputData);
-	UTEST_EQUAL("There is the right number of points in output", OutputData->GetPoints().Num(), NumOfPoints);
+	UTEST_EQUAL("There is the right number of points in output", OutputData->GetNumPoints(), NumOfPoints);
 
 	check(OutputData);
 
 	// Density was copied in the position correctly
+	TConstPCGValueRange<float> DensityRange = OutputData->GetConstDensityValueRange();
+	TConstPCGValueRange<FTransform> TransformRange = OutputData->GetConstTransformValueRange();
 	for (int i = 0; i < NumOfPoints; ++i)
 	{
-		const FPCGPoint& OutputPoint = OutputData->GetPoints()[i];
-		UTEST_EQUAL(*FString::Printf(TEXT("Position.X has the same value as density for point %d"), i), OutputPoint.Transform.GetLocation().X, static_cast<double>(OutputPoint.Density));
+		UTEST_EQUAL(*FString::Printf(TEXT("Position.X has the same value as density for point %d"), i), TransformRange[i].GetLocation().X, static_cast<double>(DensityRange[i]));
 	}
 
 	return true;
@@ -134,17 +151,17 @@ bool FPCGCopyAttributeTests_Points_PropertyToAttribute::RunTest(const FString& P
 	Settings->OutputTarget.SetAttributeName(OutputAttributeName);
 
 	TUniquePtr<FPCGContext> Context = TestData.InitializeTestContext();
-	const UPCGPointData* InputPointData = PCGCopyAttributeTests::CreateInputPointData(Context.Get(), NumOfPoints);
+	const UPCGBasePointData* InputPointData = PCGCopyAttributeTests::CreateInputPointData(Context.Get(), NumOfPoints);
 	PCGCopyAttributeTests::ConnectToSourceAndTarget(Context.Get(), InputPointData);
 
 	FPCGElementPtr TestElement = TestData.Settings->GetElement();
 	while (!TestElement->Execute(Context.Get())) {}
 
 	TArray<FPCGTaggedData> OutputTagged = Context->OutputData.GetInputsByPin(PCGPinConstants::DefaultOutputLabel);
-	const UPCGPointData* OutputData = OutputTagged.Num() == 1 ? Cast<const UPCGPointData>(OutputTagged[0].Data) : nullptr;
+	const UPCGBasePointData* OutputData = OutputTagged.Num() == 1 ? Cast<const UPCGBasePointData>(OutputTagged[0].Data) : nullptr;
 
 	UTEST_NOT_NULL("There is a point data in output points", OutputData);
-	UTEST_EQUAL("There is the right number of points in output", OutputData->GetPoints().Num(), NumOfPoints);
+	UTEST_EQUAL("There is the right number of points in output", OutputData->GetNumPoints(), NumOfPoints);
 
 	check(OutputData && OutputData->Metadata);
 	const FPCGMetadataAttribute<double>* OutputAttribute = OutputData->Metadata->GetConstTypedAttribute<double>(OutputAttributeName);
@@ -153,10 +170,11 @@ bool FPCGCopyAttributeTests_Points_PropertyToAttribute::RunTest(const FString& P
 	check(OutputAttribute);
 
 	// Attribute was copied in the position correctly
+	TConstPCGValueRange<float> DensityRange = OutputData->GetConstDensityValueRange();
+	TConstPCGValueRange<PCGMetadataEntryKey> MetadataEntryRange = OutputData->GetConstMetadataEntryValueRange();
 	for (int i = 0; i < NumOfPoints; ++i)
 	{
-		const FPCGPoint& OutputPoint = OutputData->GetPoints()[i];
-		UTEST_EQUAL(*FString::Printf(TEXT("Output attribute has the same value than density for point %d"), i), OutputAttribute->GetValueFromItemKey(OutputPoint.MetadataEntry), static_cast<double>(OutputPoint.Density));
+		UTEST_EQUAL(*FString::Printf(TEXT("Output attribute has the same value than density for point %d"), i), OutputAttribute->GetValueFromItemKey(MetadataEntryRange[i]), static_cast<double>(DensityRange[i]));
 	}
 
 	return true;
@@ -175,7 +193,7 @@ bool FPCGCopyAttributeTests_Points_AttributeToProperty::RunTest(const FString& P
 	Settings->OutputTarget.Update(TEXT("$Position"));
 
 	TUniquePtr<FPCGContext> Context = TestData.InitializeTestContext();
-	UPCGPointData* InputPointData = PCGCopyAttributeTests::CreateInputPointData(Context.Get(), NumOfPoints);
+	UPCGBasePointData* InputPointData = PCGCopyAttributeTests::CreateInputPointData(Context.Get(), NumOfPoints);
 	PCGTestsCommon::CreateAndFillRandomAttribute<FVector>(InputPointData, PCGCopyAttributeTests::AttributeName, FVector::ZeroVector, NumOfPoints);
 	PCGCopyAttributeTests::ConnectToSourceAndTarget(Context.Get(), InputPointData);
 
@@ -183,20 +201,21 @@ bool FPCGCopyAttributeTests_Points_AttributeToProperty::RunTest(const FString& P
 	while (!TestElement->Execute(Context.Get())) {}
 
 	TArray<FPCGTaggedData> OutputTagged = Context->OutputData.GetInputsByPin(PCGPinConstants::DefaultOutputLabel);
-	const UPCGPointData* OutputData = OutputTagged.Num() == 1 ? Cast<const UPCGPointData>(OutputTagged[0].Data) : nullptr;
+	const UPCGBasePointData* OutputData = OutputTagged.Num() == 1 ? Cast<const UPCGBasePointData>(OutputTagged[0].Data) : nullptr;
 
 	UTEST_NOT_NULL("There is a point data in output points", OutputData);
-	UTEST_EQUAL("There is the right number of points in output", OutputData->GetPoints().Num(), NumOfPoints);
+	UTEST_EQUAL("There is the right number of points in output", OutputData->GetNumPoints(), NumOfPoints);
 
 	check(OutputData && OutputData->Metadata);
 	const FPCGMetadataAttribute<FVector>* InputAttribute = OutputData->Metadata->GetConstTypedAttribute<FVector>(PCGCopyAttributeTests::AttributeName);
 	UTEST_NOT_NULL("Input Attribute exists in the output data", InputAttribute);
 
-	// Attribute value was copied in Position.Y correctly
+	// Attribute value was copied in Position correctly
+	TConstPCGValueRange<FTransform> TransformRange = OutputData->GetConstTransformValueRange();
+	TConstPCGValueRange<PCGMetadataEntryKey> MetadataEntryRange = OutputData->GetConstMetadataEntryValueRange();
 	for (int i = 0; i < NumOfPoints; ++i)
 	{
-		const FPCGPoint& OutputPoint = OutputData->GetPoints()[i];
-		UTEST_EQUAL(*FString::Printf(TEXT("Position has the same value than the input attribute for point %d"), i), OutputPoint.Transform.GetLocation(), InputAttribute->GetValueFromItemKey(OutputPoint.MetadataEntry));
+		UTEST_EQUAL(*FString::Printf(TEXT("Position has the same value than the input attribute for point %d"), i), TransformRange[i].GetLocation(), InputAttribute->GetValueFromItemKey(MetadataEntryRange[i]));
 	}
 
 	return true;
@@ -216,7 +235,7 @@ bool FPCGCopyAttributeTests_Points_AttributeToAttribute::RunTest(const FString&
 	Settings->OutputTarget.SetAttributeName(OutputAttributeName);
 
 	TUniquePtr<FPCGContext> Context = TestData.InitializeTestContext();
-	UPCGPointData* InputPointData = PCGCopyAttributeTests::CreateInputPointData(Context.Get(), NumOfPoints);
+	UPCGBasePointData* InputPointData = PCGCopyAttributeTests::CreateInputPointData(Context.Get(), NumOfPoints);
 	PCGTestsCommon::CreateAndFillRandomAttribute<FString>(InputPointData, PCGCopyAttributeTests::AttributeName, FString(), NumOfPoints);
 	PCGCopyAttributeTests::ConnectToSourceAndTarget(Context.Get(), InputPointData);
 
@@ -224,10 +243,10 @@ bool FPCGCopyAttributeTests_Points_AttributeToAttribute::RunTest(const FString&
 	while (!TestElement->Execute(Context.Get())) {}
 
 	TArray<FPCGTaggedData> OutputTagged = Context->OutputData.GetInputsByPin(PCGPinConstants::DefaultOutputLabel);
-	const UPCGPointData* OutputData = OutputTagged.Num() == 1 ? Cast<const UPCGPointData>(OutputTagged[0].Data) : nullptr;
+	const UPCGBasePointData* OutputData = OutputTagged.Num() == 1 ? Cast<const UPCGBasePointData>(OutputTagged[0].Data) : nullptr;
 
 	UTEST_NOT_NULL("There is a point data in output points", OutputData);
-	UTEST_EQUAL("There is the right number of points in output", OutputData->GetPoints().Num(), NumOfPoints);
+	UTEST_EQUAL("There is the right number of points in output", OutputData->GetNumPoints(), NumOfPoints);
 
 	check(OutputData && OutputData->Metadata);
 
@@ -241,10 +260,10 @@ bool FPCGCopyAttributeTests_Points_AttributeToAttribute::RunTest(const FString&
 	UTEST_EQUAL("Output Attribute default value is the same as default value of Input Attribute", OutputAttribute->GetValueFromItemKey(PCGInvalidEntryKey), InputAttribute->GetValueFromItemKey(PCGInvalidEntryKey));
 
 	// Input attribute was copied in the output attribute correctly
+	TConstPCGValueRange<PCGMetadataEntryKey> MetadataEntryRange = OutputData->GetConstMetadataEntryValueRange();
 	for (int i = 0; i < NumOfPoints; ++i)
 	{
-		const FPCGPoint& OutputPoint = OutputData->GetPoints()[i];
-		UTEST_EQUAL(*FString::Printf(TEXT("Output Attribute has the same value as Input Attribute for point %d"), i), OutputAttribute->GetValueFromItemKey(OutputPoint.MetadataEntry), InputAttribute->GetValueFromItemKey(OutputPoint.MetadataEntry));
+		UTEST_EQUAL(*FString::Printf(TEXT("Output Attribute has the same value as Input Attribute for point %d"), i), OutputAttribute->GetValueFromItemKey(MetadataEntryRange[i]), InputAttribute->GetValueFromItemKey(MetadataEntryRange[i]));
 	}
 
 	return true;
@@ -264,15 +283,15 @@ bool FPCGCopyAttributeTests_Points_CopyingToItself::RunTest(const FString& Param
 
 	TUniquePtr<FPCGContext> Context = TestData.InitializeTestContext();
 	// Need constness for UTEST_EQUAL....
-	const UPCGPointData* InputPointData = PCGCopyAttributeTests::CreateInputPointData(Context.Get(), NumOfPoints);
-	PCGTestsCommon::CreateAndFillRandomAttribute<double>(const_cast<UPCGPointData*>(InputPointData), PCGCopyAttributeTests::AttributeName, 0.0, NumOfPoints);
+	const UPCGBasePointData* InputPointData = PCGCopyAttributeTests::CreateInputPointData(Context.Get(), NumOfPoints);
+	PCGTestsCommon::CreateAndFillRandomAttribute<double>(const_cast<UPCGBasePointData*>(InputPointData), PCGCopyAttributeTests::AttributeName, 0.0, NumOfPoints);
 	PCGCopyAttributeTests::ConnectToSourceAndTarget(Context.Get(), InputPointData);
 
 	FPCGElementPtr TestElement = TestData.Settings->GetElement();
 	while (!TestElement->Execute(Context.Get())) {}
 
 	TArray<FPCGTaggedData> OutputTagged = Context->OutputData.GetInputsByPin(PCGPinConstants::DefaultOutputLabel);
-	const UPCGPointData* OutputData = OutputTagged.Num() == 1 ? Cast<const UPCGPointData>(OutputTagged[0].Data) : nullptr;
+	const UPCGBasePointData* OutputData = OutputTagged.Num() == 1 ? Cast<const UPCGBasePointData>(OutputTagged[0].Data) : nullptr;
 
 	UTEST_NOT_NULL("There is a point data in output points", OutputData);
 	UTEST_EQUAL("It's the same data as input", OutputData, InputPointData);
@@ -291,15 +310,15 @@ bool FPCGCopyAttributeTests_Points_CopyingAllToItself::RunTest(const FString& Pa
 	Settings->bCopyAllAttributes = true;
 
 	TUniquePtr<FPCGContext> Context = TestData.InitializeTestContext();
-	const UPCGPointData* InputPointData = PCGCopyAttributeTests::CreateInputPointData(Context.Get(), NumOfPoints);
-	PCGTestsCommon::CreateAndFillRandomAttribute<double>(const_cast<UPCGPointData*>(InputPointData), PCGCopyAttributeTests::AttributeName, 0.0, NumOfPoints);
+	const UPCGBasePointData* InputPointData = PCGCopyAttributeTests::CreateInputPointData(Context.Get(), NumOfPoints);
+	PCGTestsCommon::CreateAndFillRandomAttribute<double>(const_cast<UPCGBasePointData*>(InputPointData), PCGCopyAttributeTests::AttributeName, 0.0, NumOfPoints);
 	PCGCopyAttributeTests::ConnectToSourceAndTarget(Context.Get(), InputPointData);
 
 	FPCGElementPtr TestElement = TestData.Settings->GetElement();
 	while (!TestElement->Execute(Context.Get())) {}
 
 	TArray<FPCGTaggedData> OutputTagged = Context->OutputData.GetInputsByPin(PCGPinConstants::DefaultOutputLabel);
-	const UPCGPointData* OutputData = OutputTagged.Num() == 1 ? Cast<const UPCGPointData>(OutputTagged[0].Data) : nullptr;
+	const UPCGBasePointData* OutputData = OutputTagged.Num() == 1 ? Cast<const UPCGBasePointData>(OutputTagged[0].Data) : nullptr;
 
 	UTEST_NOT_NULL("There is a point data in output points", OutputData);
 	UTEST_EQUAL("It's the same data as input", OutputData, InputPointData);
@@ -447,5 +466,628 @@ bool FPCGCopyAttributeTests_Params_CopyingAllToItself::RunTest(const FString& Pa
 	return true;
 }
 
+bool FPCGCopyAttributeTests_MultiDomain_GlobalToGlobal::RunTest(const FString& Parameters)
+{
+	PCGTestsCommon::FTestData TestData;
+	UPCGCopyAttributesSettings* Settings = PCGTestsCommon::GenerateSettings<UPCGCopyAttributesSettings>(TestData);
+	check(Settings);
+
+	constexpr int NumOfPoints = 20;
+
+	Settings->InputSource.Update(TEXT("@Global.MyAttr"));
+	Settings->OutputTarget.Update(TEXT("@Global.MyAttr2"));
+
+	TUniquePtr<FPCGContext> Context = TestData.InitializeTestContext();
+	UPCGBasePointData* InputPointData = PCGCopyAttributeTests::CreateInputPointData(Context.Get(), NumOfPoints);
+	InputPointData->MutableMetadata()->GetMetadataDomain(PCGMetadataDomainID::Data)->CreateAttribute<int32>(TEXT("MyAttr"), /*DefaultValue=*/5, /*bAllowInterpolation=*/false, /*bOverrideParent=*/false);
+
+	PCGCopyAttributeTests::ConnectToSourceAndTarget(Context.Get(), InputPointData);
+
+	FPCGElementPtr TestElement = TestData.Settings->GetElement();
+	while (!TestElement->Execute(Context.Get())) {}
+
+	TArray<FPCGTaggedData> OutputTagged = Context->OutputData.GetInputsByPin(PCGPinConstants::DefaultOutputLabel);
+	const UPCGBasePointData* OutputData = OutputTagged.Num() == 1 ? Cast<const UPCGBasePointData>(OutputTagged[0].Data) : nullptr;
+
+	UTEST_NOT_NULL("There is a point data in output points", OutputData);
+	UTEST_EQUAL("There is the right number of points in output", OutputData->GetNumPoints(), NumOfPoints);
+
+	check(OutputData && OutputData->ConstMetadata());
+
+	const FPCGMetadataDomain* GlobalDomain = OutputData->ConstMetadata()->GetConstMetadataDomain(PCGMetadataDomainID::Data);
+	check(GlobalDomain);
+
+	const FPCGMetadataAttribute<int32>* InputAttribute = GlobalDomain->GetConstTypedAttribute<int32>(TEXT("MyAttr"));
+	UTEST_NOT_NULL("Input Attribute exists in the output data", InputAttribute);
+	const FPCGMetadataAttribute<int32>* OutputAttribute = GlobalDomain->GetConstTypedAttribute<int32>(TEXT("MyAttr2"));
+	UTEST_NOT_NULL("Output Attribute exists in the output data", OutputAttribute);
+
+	check(InputAttribute && OutputAttribute);
+
+	UTEST_EQUAL("Output Attribute default value is the same as default value of Input Attribute", OutputAttribute->GetValueFromItemKey(PCGInvalidEntryKey), InputAttribute->GetValueFromItemKey(PCGInvalidEntryKey));
+
+	return true;
+}
+
+bool FPCGCopyAttributeTests_MultiDomain_GlobalToElements::RunTest(const FString& Parameters)
+{
+	PCGTestsCommon::FTestData TestData;
+	UPCGCopyAttributesSettings* Settings = PCGTestsCommon::GenerateSettings<UPCGCopyAttributesSettings>(TestData);
+	check(Settings);
+
+	constexpr int NumOfPoints = 20;
+
+	Settings->InputSource.Update(TEXT("@Global.MyAttr"));
+	Settings->OutputTarget.Update(TEXT("MyAttr2"));
+
+	TUniquePtr<FPCGContext> Context = TestData.InitializeTestContext();
+	UPCGBasePointData* InputPointData = PCGCopyAttributeTests::CreateInputPointData(Context.Get(), NumOfPoints);
+	InputPointData->MutableMetadata()->GetMetadataDomain(PCGMetadataDomainID::Data)->CreateAttribute<int32>(TEXT("MyAttr"), /*DefaultValue=*/5, /*bAllowInterpolation=*/false, /*bOverrideParent=*/false);
+	InputPointData->MutableMetadata()->GetMetadataDomain(PCGMetadataDomainID::Data)->AddEntry();
+
+	PCGCopyAttributeTests::ConnectToSourceAndTarget(Context.Get(), InputPointData);
+
+	FPCGElementPtr TestElement = TestData.Settings->GetElement();
+	while (!TestElement->Execute(Context.Get())) {}
+
+	TArray<FPCGTaggedData> OutputTagged = Context->OutputData.GetInputsByPin(PCGPinConstants::DefaultOutputLabel);
+	const UPCGBasePointData* OutputData = OutputTagged.Num() == 1 ? Cast<const UPCGBasePointData>(OutputTagged[0].Data) : nullptr;
+
+	UTEST_NOT_NULL("There is a point data in output points", OutputData);
+	UTEST_EQUAL("There is the right number of points in output", OutputData->GetNumPoints(), NumOfPoints);
+
+	check(OutputData && OutputData->ConstMetadata());
+
+	const FPCGMetadataDomain* ElementsDomain = OutputData->ConstMetadata()->GetConstMetadataDomain(PCGMetadataDomainID::Elements);
+	const FPCGMetadataDomain* GlobalDomain = OutputData->ConstMetadata()->GetConstMetadataDomain(PCGMetadataDomainID::Data);
+	check(GlobalDomain && ElementsDomain);
+
+	const FPCGMetadataAttribute<int32>* InputAttribute = GlobalDomain->GetConstTypedAttribute<int32>(TEXT("MyAttr"));
+	UTEST_NOT_NULL("Input Attribute exists in the output global metadata domain", InputAttribute);
+	const FPCGMetadataAttribute<int32>* OutputAttribute = ElementsDomain->GetConstTypedAttribute<int32>(TEXT("MyAttr2"));
+	UTEST_NOT_NULL("Output Attribute exists in the output elements metadata domain", OutputAttribute);
+
+	check(InputAttribute && OutputAttribute);
+
+	UTEST_EQUAL("Output Attribute default value is the same as default value of Input Attribute", OutputAttribute->GetValueFromItemKey(PCGInvalidEntryKey), InputAttribute->GetValueFromItemKey(PCGInvalidEntryKey));
+
+	// Input attribute was copied in the output attribute correctly
+	TConstPCGValueRange<PCGMetadataEntryKey> MetadataEntryRange = OutputData->GetConstMetadataEntryValueRange();
+	for (int i = 0; i < NumOfPoints; ++i)
+	{
+		UTEST_EQUAL(*FString::Printf(TEXT("Output Attribute has the same value as Input Attribute for point %d"), i), OutputAttribute->GetValueFromItemKey(MetadataEntryRange[i]), InputAttribute->GetValueFromItemKey(0));
+	}
+	
+	return true;
+}
+
+bool FPCGCopyAttributeTests_MultiDomain_ElementsToGlobal::RunTest(const FString& Parameters)
+{
+	PCGTestsCommon::FTestData TestData;
+	UPCGCopyAttributesSettings* Settings = PCGTestsCommon::GenerateSettings<UPCGCopyAttributesSettings>(TestData);
+	check(Settings);
+
+	constexpr int NumOfPoints = 1;
+
+	Settings->InputSource.Update(TEXT("MyAttr"));
+	Settings->OutputTarget.Update(TEXT("@Global.MyAttr2"));
+
+	TUniquePtr<FPCGContext> Context = TestData.InitializeTestContext();
+	UPCGBasePointData* InputPointData = PCGCopyAttributeTests::CreateInputPointData(Context.Get(), NumOfPoints);
+	FPCGMetadataDomain* InputElementsDomain = InputPointData->MutableMetadata()->GetMetadataDomain(PCGMetadataDomainID::Elements);
+	check(InputElementsDomain);
+	FPCGMetadataAttribute<int32>* Attribute = InputElementsDomain->CreateAttribute<int32>(TEXT("MyAttr"), /*DefaultValue=*/5, /*bAllowInterpolation=*/false, /*bOverrideParent=*/false);
+	PCGMetadataEntryKey EntryKey = PCGInvalidEntryKey;
+	InputElementsDomain->InitializeOnSet(EntryKey);
+	Attribute->SetValue(EntryKey, 5);
+	InputPointData->SetMetadataEntry(EntryKey);
+
+	PCGCopyAttributeTests::ConnectToSourceAndTarget(Context.Get(), InputPointData);
+
+	FPCGElementPtr TestElement = TestData.Settings->GetElement();
+	while (!TestElement->Execute(Context.Get())) {}
+
+	TArray<FPCGTaggedData> OutputTagged = Context->OutputData.GetInputsByPin(PCGPinConstants::DefaultOutputLabel);
+	const UPCGBasePointData* OutputData = OutputTagged.Num() == 1 ? Cast<const UPCGBasePointData>(OutputTagged[0].Data) : nullptr;
+
+	UTEST_NOT_NULL("There is a point data in output points", OutputData);
+	UTEST_EQUAL("There is the right number of points in output", OutputData->GetNumPoints(), NumOfPoints);
+
+	check(OutputData && OutputData->ConstMetadata());
+
+	const FPCGMetadataDomain* ElementsDomain = OutputData->ConstMetadata()->GetConstMetadataDomain(PCGMetadataDomainID::Elements);
+	const FPCGMetadataDomain* GlobalDomain = OutputData->ConstMetadata()->GetConstMetadataDomain(PCGMetadataDomainID::Data);
+	check(GlobalDomain && ElementsDomain);
+
+	const FPCGMetadataAttribute<int32>* InputAttribute = ElementsDomain->GetConstTypedAttribute<int32>(TEXT("MyAttr"));
+	UTEST_NOT_NULL("Input Attribute exists in the output elements metadata domain", InputAttribute);
+	const FPCGMetadataAttribute<int32>* OutputAttribute = GlobalDomain->GetConstTypedAttribute<int32>(TEXT("MyAttr2"));
+	UTEST_NOT_NULL("Output Attribute exists in the output global metadata domain", OutputAttribute);
+
+	check(InputAttribute && OutputAttribute);
+
+	UTEST_EQUAL("Output Attribute default value is the same as default value of Input Attribute", OutputAttribute->GetValueFromItemKey(PCGInvalidEntryKey), InputAttribute->GetValueFromItemKey(PCGInvalidEntryKey));
+	UTEST_EQUAL("Output Attribute value is the same as value of Input Attribute", OutputAttribute->GetValueFromItemKey(0), InputAttribute->GetValueFromItemKey(0));
+	
+	return true;
+}
+
+bool FPCGCopyAttributeTests_MultiDomain_ElementsToGlobalTooMany::RunTest(const FString& Parameters)
+{
+	PCGTestsCommon::FTestData TestData;
+	UPCGCopyAttributesSettings* Settings = PCGTestsCommon::GenerateSettings<UPCGCopyAttributesSettings>(TestData);
+	check(Settings);
+
+	constexpr int NumOfPoints = 20;
+
+	Settings->InputSource.Update(TEXT("MyAttr"));
+	Settings->OutputTarget.Update(TEXT("@Global.MyAttr2"));
+
+	TUniquePtr<FPCGContext> Context = TestData.InitializeTestContext();
+	UPCGBasePointData* InputPointData = PCGCopyAttributeTests::CreateInputPointData(Context.Get(), NumOfPoints);
+	PCGTestsCommon::CreateAndFillRandomAttribute<double>(InputPointData, TEXT("MyAttr"), 0.0, NumOfPoints);
+	PCGCopyAttributeTests::ConnectToSourceAndTarget(Context.Get(), InputPointData);
+	
+	FPCGElementPtr TestElement = TestData.Settings->GetElement();
+	while (!TestElement->Execute(Context.Get())) {}
+
+	TArray<FPCGTaggedData> OutputTagged = Context->OutputData.GetInputsByPin(PCGPinConstants::DefaultOutputLabel);
+	const UPCGBasePointData* OutputData = OutputTagged.Num() == 1 ? Cast<const UPCGBasePointData>(OutputTagged[0].Data) : nullptr;
+
+	UTEST_NOT_NULL("There is a point data in output points", OutputData);
+	UTEST_EQUAL("There is the right number of points in output", OutputData->GetNumPoints(), NumOfPoints);
+
+	check(OutputData && OutputData->ConstMetadata());
+
+	const FPCGMetadataDomain* ElementsDomain = OutputData->ConstMetadata()->GetConstMetadataDomain(PCGMetadataDomainID::Elements);
+	const FPCGMetadataDomain* GlobalDomain = OutputData->ConstMetadata()->GetConstMetadataDomain(PCGMetadataDomainID::Data);
+	check(GlobalDomain && ElementsDomain);
+
+	const FPCGMetadataAttribute<double>* InputAttribute = ElementsDomain->GetConstTypedAttribute<double>(TEXT("MyAttr"));
+	UTEST_NOT_NULL("Input Attribute exists in the output elements metadata domain", InputAttribute);
+	const FPCGMetadataAttribute<double>* OutputAttribute = GlobalDomain->GetConstTypedAttribute<double>(TEXT("MyAttr2"));
+	UTEST_NOT_NULL("Output Attribute exists in the output global metadata domain", OutputAttribute);
+
+	check(InputAttribute && OutputAttribute);
+
+	UTEST_EQUAL("Output Attribute default value is the same as default value of Input Attribute", OutputAttribute->GetValueFromItemKey(PCGInvalidEntryKey), InputAttribute->GetValueFromItemKey(PCGInvalidEntryKey));
+
+	TConstPCGValueRange<PCGMetadataEntryKey> EntryKeyRange = OutputData->GetConstMetadataEntryValueRange();
+	UTEST_EQUAL("Output Attribute value is the same as the first value of Input Attribute", OutputAttribute->GetValueFromItemKey(0), InputAttribute->GetValueFromItemKey(EntryKeyRange[0]));
+	
+	return true;
+}
+
+bool FPCGCopyAttributeTests_MultiDomain_InvalidDomain::RunTest(const FString& Parameters)
+{
+	PCGTestsCommon::FTestData TestData;
+	UPCGCopyAttributesSettings* Settings = PCGTestsCommon::GenerateSettings<UPCGCopyAttributesSettings>(TestData);
+	check(Settings);
+
+	constexpr int NumOfPoints = 20;
+
+	Settings->InputSource.Update(TEXT("MyAttr"));
+	Settings->OutputTarget.Update(TEXT("@Blablabla.MyAttr2"));
+
+	TUniquePtr<FPCGContext> Context = TestData.InitializeTestContext();
+	const UPCGBasePointData* InputPointData = PCGCopyAttributeTests::CreateInputPointData(Context.Get(), NumOfPoints);
+	PCGTestsCommon::CreateAndFillRandomAttribute<double>(const_cast<UPCGBasePointData*>(InputPointData), TEXT("MyAttr"), 0.0, NumOfPoints);
+	PCGCopyAttributeTests::ConnectToSourceAndTarget(Context.Get(), InputPointData);
+	
+	AddExpectedError(TEXT("Metadata domain Blablabla is invalid for this data."), EAutomationExpectedErrorFlags::MatchType::Contains, 1, /*bIsRegex=*/false);
+
+	FPCGElementPtr TestElement = TestData.Settings->GetElement();
+	while (!TestElement->Execute(Context.Get())) {}
+
+	TArray<FPCGTaggedData> OutputTagged = Context->OutputData.GetInputsByPin(PCGPinConstants::DefaultOutputLabel);
+	const UPCGBasePointData* OutputData = OutputTagged.Num() == 1 ? Cast<const UPCGBasePointData>(OutputTagged[0].Data) : nullptr;
+
+	UTEST_NOT_NULL("There is a point data in output points", OutputData);
+	UTEST_EQUAL("It's the same data as input", OutputData, InputPointData);
+	
+	return true;
+}
+
+bool FPCGCopyAttributeTests_MultiDomain_CopyAll_GlobalToGlobal::RunTest(const FString& Parameters)
+{
+	PCGTestsCommon::FTestData TestData;
+	UPCGCopyAttributesSettings* Settings = PCGTestsCommon::GenerateSettings<UPCGCopyAttributesSettings>(TestData);
+	check(Settings);
+	
+	Settings->MetadataDomainsMapping.Emplace(PCGDataConstants::GlobalDomainName, PCGDataConstants::GlobalDomainName);
+	Settings->bCopyAllAttributes = true;
+
+	constexpr int NumOfPoints = 20;
+	
+	TUniquePtr<FPCGContext> Context = TestData.InitializeTestContext();
+	UPCGBasePointData* TargetPointData = PCGCopyAttributeTests::CreateInputPointData(Context.Get(), NumOfPoints);
+	UPCGParamData* SourceParamData = PCGCopyAttributeTests::CreateInputParamData(Context.Get());
+	
+	TargetPointData->MutableMetadata()->GetMetadataDomain(PCGMetadataDomainID::Data)->CreateAttribute<int32>(TEXT("MyAttr"), /*DefaultValue=*/5, /*bAllowInterpolation=*/false, /*bOverrideParent=*/false);
+	TargetPointData->MutableMetadata()->GetMetadataDomain(PCGMetadataDomainID::Data)->AddEntry();
+
+	const FPCGMetadataAttribute<int32>* SourceAttribute = SourceParamData->MutableMetadata()->GetMetadataDomain(PCGMetadataDomainID::Data)->CreateAttribute<int32>(TEXT("MyAttr2"), /*DefaultValue=*/6, /*bAllowInterpolation=*/false, /*bOverrideParent=*/false);
+	SourceParamData->MutableMetadata()->GetMetadataDomain(PCGMetadataDomainID::Data)->AddEntry();
+
+	PCGCopyAttributeTests::ConnectToSource(Context.Get(), SourceParamData);
+	PCGCopyAttributeTests::ConnectToTarget(Context.Get(), TargetPointData);
+
+	FPCGElementPtr TestElement = TestData.Settings->GetElement();
+	while (!TestElement->Execute(Context.Get())) {}
+
+	TArray<FPCGTaggedData> OutputTagged = Context->OutputData.GetInputsByPin(PCGPinConstants::DefaultOutputLabel);
+	const UPCGBasePointData* OutputData = OutputTagged.Num() == 1 ? Cast<const UPCGBasePointData>(OutputTagged[0].Data) : nullptr;
+
+	UTEST_NOT_NULL("There is a point data in output points", OutputData);
+	UTEST_EQUAL("There is the right number of points in output", OutputData->GetNumPoints(), NumOfPoints);
+
+	check(OutputData && OutputData->ConstMetadata());
+
+	const FPCGMetadataDomain* ElementsDomain = OutputData->ConstMetadata()->GetConstMetadataDomain(PCGMetadataDomainID::Elements);
+	const FPCGMetadataDomain* GlobalDomain = OutputData->ConstMetadata()->GetConstMetadataDomain(PCGMetadataDomainID::Data);
+	check(GlobalDomain && ElementsDomain);
+
+	UTEST_EQUAL("No attributes on the output elements metadata domain", ElementsDomain->GetAttributeCount(), 0);
+
+	const FPCGMetadataAttribute<int32>* ExistingAttribute = GlobalDomain->GetConstTypedAttribute<int32>(TEXT("MyAttr"));
+	UTEST_NOT_NULL("Existing attribute is still there in the output global metadata domain", ExistingAttribute);
+	
+	const FPCGMetadataAttribute<int32>* Attribute = GlobalDomain->GetConstTypedAttribute<int32>(TEXT("MyAttr2"));
+	UTEST_NOT_NULL("Source Attribute exists in the output global metadata domain", Attribute);
+
+	check(Attribute);
+
+	UTEST_EQUAL("Target Attribute default value is the same as default value of Source Attribute", Attribute->GetValueFromItemKey(PCGInvalidEntryKey), SourceAttribute->GetValueFromItemKey(PCGInvalidEntryKey));
+	UTEST_EQUAL("Target Attribute value is the same as value of Source Attribute", Attribute->GetValueFromItemKey(PCGMetadataEntryKey(0)), SourceAttribute->GetValueFromItemKey(PCGMetadataEntryKey(0)));
+	
+	return true;
+}
+
+bool FPCGCopyAttributeTests_MultiDomain_CopyAll_AllToAll::RunTest(const FString& Parameters)
+{
+	PCGTestsCommon::FTestData TestData;
+	UPCGCopyAttributesSettings* Settings = PCGTestsCommon::GenerateSettings<UPCGCopyAttributesSettings>(TestData);
+	check(Settings);
+
+	Settings->bCopyAllDomains = true;
+	Settings->bCopyAllAttributes = true;
+
+	constexpr int NumOfPoints = 20;
+	
+	TUniquePtr<FPCGContext> Context = TestData.InitializeTestContext();
+	UPCGParamData* TargetParamData = PCGCopyAttributeTests::CreateInputParamData(Context.Get());
+	UPCGBasePointData* SourcePointData = PCGCopyAttributeTests::CreateInputPointData(Context.Get(), NumOfPoints);
+	
+	TargetParamData->MutableMetadata()->GetMetadataDomain(PCGMetadataDomainID::Data)->CreateAttribute<int32>(TEXT("MyGlobalAttr"), 5, false, false);
+	TargetParamData->MutableMetadata()->GetMetadataDomain(PCGMetadataDomainID::Data)->AddEntry();
+
+	PCGTestsCommon::CreateAndFillRandomAttribute<double>(TargetParamData, TEXT("MyAttr"), 0.0, NumOfPoints);
+
+	const FPCGMetadataAttribute<int32>* SourceAttribute = SourcePointData->MutableMetadata()->GetMetadataDomain(PCGMetadataDomainID::Data)->CreateAttribute<int32>(TEXT("MyGlobalAttr2"), /*DefaultValue=*/6, /*bAllowInterpolation=*/false, /*bOverrideParent=*/false);
+	SourcePointData->MutableMetadata()->GetMetadataDomain(PCGMetadataDomainID::Data)->AddEntry();
+
+	PCGTestsCommon::CreateAndFillRandomAttribute<double>(SourcePointData, TEXT("MyAttr2"), 0.0, NumOfPoints);
+	const FPCGMetadataAttribute<double>* SourceElementsAttribute = SourcePointData->ConstMetadata()->GetConstMetadataDomain(PCGMetadataDomainID::Elements)->GetConstTypedAttribute<double>(TEXT("MyAttr2"));
+	check(SourceElementsAttribute);
+
+	PCGCopyAttributeTests::ConnectToSource(Context.Get(), SourcePointData);
+	PCGCopyAttributeTests::ConnectToTarget(Context.Get(), TargetParamData);
+
+	FPCGElementPtr TestElement = TestData.Settings->GetElement();
+	while (!TestElement->Execute(Context.Get())) {}
 
+	TArray<FPCGTaggedData> OutputTagged = Context->OutputData.GetInputsByPin(PCGPinConstants::DefaultOutputLabel);
+	const UPCGParamData* OutputData = OutputTagged.Num() == 1 ? Cast<const UPCGParamData>(OutputTagged[0].Data) : nullptr;
+
+	UTEST_NOT_NULL("There is a param data in output points", OutputData);
+
+	check(OutputData && OutputData->ConstMetadata());
+
+	const FPCGMetadataDomain* ElementsDomain = OutputData->ConstMetadata()->GetConstMetadataDomain(PCGMetadataDomainID::Elements);
+	const FPCGMetadataDomain* GlobalDomain = OutputData->ConstMetadata()->GetConstMetadataDomain(PCGMetadataDomainID::Data);
+	check(GlobalDomain && ElementsDomain);
+
+	const FPCGMetadataAttribute<int32>* ExistingAttribute = GlobalDomain->GetConstTypedAttribute<int32>(TEXT("MyGlobalAttr"));
+	UTEST_NOT_NULL("Existing attribute is still there in the output global metadata domain", ExistingAttribute);
+	
+	const FPCGMetadataAttribute<int32>* Attribute = GlobalDomain->GetConstTypedAttribute<int32>(TEXT("MyGlobalAttr2"));
+	UTEST_NOT_NULL("Source Attribute exists in the output global metadata domain", Attribute);
+
+	check(Attribute);
+
+	UTEST_EQUAL("Target Attribute default value is the same as default value of Source Attribute", Attribute->GetValueFromItemKey(PCGInvalidEntryKey), SourceAttribute->GetValueFromItemKey(PCGInvalidEntryKey));
+	UTEST_EQUAL("Target Attribute value is the same as value of Source Attribute", Attribute->GetValueFromItemKey(PCGMetadataEntryKey(0)), SourceAttribute->GetValueFromItemKey(PCGMetadataEntryKey(0)));
+
+	const FPCGMetadataAttribute<double>* ExistingElementsAttribute = ElementsDomain->GetConstTypedAttribute<double>(TEXT("MyAttr"));
+	UTEST_NOT_NULL("Existing attribute is still there in the output elements metadata domain", ExistingElementsAttribute);
+	
+	const FPCGMetadataAttribute<double>* TargetElementsAttribute = ElementsDomain->GetConstTypedAttribute<double>(TEXT("MyAttr2"));
+	UTEST_NOT_NULL("Source Attribute exists in the output elements metadata domain", TargetElementsAttribute);
+
+	check(TargetElementsAttribute);
+
+	UTEST_EQUAL("Target Attribute default value is the same as default value of Source Attribute", TargetElementsAttribute->GetValueFromItemKey(PCGInvalidEntryKey), SourceElementsAttribute->GetValueFromItemKey(PCGInvalidEntryKey));
+	TConstPCGValueRange<PCGMetadataEntryKey> MetadataEntryRange = SourcePointData->GetConstMetadataEntryValueRange();
+	for (int32 i = 0; i < NumOfPoints; ++i)
+	{
+		UTEST_EQUAL(*FString::Printf(TEXT("Target Attribute has the same value as Source Attribute for index %d"), i), TargetElementsAttribute->GetValueFromItemKey(i), SourceElementsAttribute->GetValueFromItemKey(MetadataEntryRange[i]));
+	}
+	
+	return true;
+}
+
+bool FPCGCopyAttributeTests_MultiDomain_CopyAll_GlobalToElements::RunTest(const FString& Parameters)
+{
+	PCGTestsCommon::FTestData TestData;
+	UPCGCopyAttributesSettings* Settings = PCGTestsCommon::GenerateSettings<UPCGCopyAttributesSettings>(TestData);
+	check(Settings);
+
+	constexpr int NumOfPoints = 20;
+
+	Settings->MetadataDomainsMapping.Emplace(PCGDataConstants::GlobalDomainName, PCGPointDataConstants::ElementsDomainName);
+	Settings->bCopyAllAttributes = true;
+
+	TUniquePtr<FPCGContext> Context = TestData.InitializeTestContext();
+	UPCGBasePointData* TargetPointData = PCGCopyAttributeTests::CreateInputPointData(Context.Get(), NumOfPoints);
+	UPCGParamData* SourceParamData = PCGCopyAttributeTests::CreateInputParamData(Context.Get());
+
+	const FPCGMetadataAttribute<int32>* InputAttribute = SourceParamData->MutableMetadata()->GetMetadataDomain(PCGMetadataDomainID::Data)->CreateAttribute<int32>(TEXT("MyGlobalAttr"), /*DefaultValue=*/5, /*bAllowInterpolation=*/false, /*bOverrideParent=*/false);
+	SourceParamData->MutableMetadata()->GetMetadataDomain(PCGMetadataDomainID::Data)->AddEntry();
+	
+	PCGCopyAttributeTests::ConnectToSource(Context.Get(), SourceParamData);
+	PCGCopyAttributeTests::ConnectToTarget(Context.Get(), TargetPointData);
+
+	FPCGElementPtr TestElement = TestData.Settings->GetElement();
+	while (!TestElement->Execute(Context.Get())) {}
+
+	TArray<FPCGTaggedData> OutputTagged = Context->OutputData.GetInputsByPin(PCGPinConstants::DefaultOutputLabel);
+	const UPCGBasePointData* OutputData = OutputTagged.Num() == 1 ? Cast<const UPCGBasePointData>(OutputTagged[0].Data) : nullptr;
+
+	UTEST_NOT_NULL("There is a point data in output points", OutputData);
+	UTEST_EQUAL("There is the right number of points in output", OutputData->GetNumPoints(), NumOfPoints);
+
+	check(OutputData && OutputData->ConstMetadata());
+	const FPCGMetadataDomain* ElementsDomain = OutputData->ConstMetadata()->GetConstMetadataDomain(PCGMetadataDomainID::Elements);
+	check(ElementsDomain);
+	
+	const FPCGMetadataAttribute<int32>* OutputAttribute = ElementsDomain->GetConstTypedAttribute<int32>(TEXT("MyGlobalAttr"));
+	UTEST_NOT_NULL("Output Attribute exists in the output elements metadata domain", OutputAttribute);
+
+	check(InputAttribute && OutputAttribute);
+
+	UTEST_EQUAL("Output Attribute default value is the same as default value of Input Attribute", OutputAttribute->GetValueFromItemKey(PCGInvalidEntryKey), InputAttribute->GetValueFromItemKey(PCGInvalidEntryKey));
+	
+	return true;
+}
+
+bool FPCGCopyAttributeTests_MultiDomain_CopyAll_ElementsToGlobal::RunTest(const FString& Parameters)
+{
+	PCGTestsCommon::FTestData TestData;
+	UPCGCopyAttributesSettings* Settings = PCGTestsCommon::GenerateSettings<UPCGCopyAttributesSettings>(TestData);
+	check(Settings);
+
+	constexpr int NumOfPoints = 20;
+
+	Settings->MetadataDomainsMapping.Emplace(PCGPointDataConstants::ElementsDomainName, PCGDataConstants::GlobalDomainName);
+	Settings->bCopyAllAttributes = true;
+
+	TUniquePtr<FPCGContext> Context = TestData.InitializeTestContext();
+	UPCGBasePointData* TargetPointData = PCGCopyAttributeTests::CreateInputPointData(Context.Get(), NumOfPoints);
+	UPCGBasePointData* SourcePointData = PCGCopyAttributeTests::CreateInputPointData(Context.Get(), 1);
+
+	FPCGMetadataAttribute<int32>* InputAttribute = SourcePointData->MutableMetadata()->GetMetadataDomain(PCGMetadataDomainID::Elements)->CreateAttribute<int32>(TEXT("MyAttr"), /*DefaultValue=*/6, /*bAllowInterpolation=*/false, /*bOverrideParent=*/false);
+	PCGMetadataEntryKey EntryKey = PCGInvalidEntryKey;
+	SourcePointData->MutableMetadata()->GetMetadataDomain(PCGMetadataDomainID::Elements)->InitializeOnSet(EntryKey);
+	InputAttribute->SetValue(EntryKey, 5);
+	SourcePointData->SetMetadataEntry(EntryKey);
+	
+	PCGCopyAttributeTests::ConnectToSource(Context.Get(), SourcePointData);
+	PCGCopyAttributeTests::ConnectToTarget(Context.Get(), TargetPointData);
+
+	FPCGElementPtr TestElement = TestData.Settings->GetElement();
+	while (!TestElement->Execute(Context.Get())) {}
+
+	TArray<FPCGTaggedData> OutputTagged = Context->OutputData.GetInputsByPin(PCGPinConstants::DefaultOutputLabel);
+	const UPCGBasePointData* OutputData = OutputTagged.Num() == 1 ? Cast<const UPCGBasePointData>(OutputTagged[0].Data) : nullptr;
+
+	UTEST_NOT_NULL("There is a point data in output points", OutputData);
+	UTEST_EQUAL("There is the right number of points in output", OutputData->GetNumPoints(), NumOfPoints);
+
+	check(OutputData && OutputData->ConstMetadata());
+	const FPCGMetadataDomain* GlobalDomain = OutputData->ConstMetadata()->GetConstMetadataDomain(PCGMetadataDomainID::Data);
+	check(GlobalDomain);
+	
+	const FPCGMetadataAttribute<int32>* OutputAttribute = GlobalDomain->GetConstTypedAttribute<int32>(TEXT("MyAttr"));
+	UTEST_NOT_NULL("Output Attribute exists in the output global metadata domain", OutputAttribute);
+
+	check(InputAttribute && OutputAttribute);
+
+	UTEST_EQUAL("Output Attribute default value is the same as the first value of Input Attribute", OutputAttribute->GetValueFromItemKey(PCGInvalidEntryKey), InputAttribute->GetValueFromItemKey(0));
+	UTEST_EQUAL("Output Attribute value is the same as value of Input Attribute", OutputAttribute->GetValueFromItemKey(0), InputAttribute->GetValueFromItemKey(0));
+	
+	return true;
+}
+
+bool FPCGCopyAttributeTests_MultiDomain_CopyAll_ElementsToGlobalTooMany::RunTest(const FString& Parameters)
+{
+	PCGTestsCommon::FTestData TestData;
+	UPCGCopyAttributesSettings* Settings = PCGTestsCommon::GenerateSettings<UPCGCopyAttributesSettings>(TestData);
+	check(Settings);
+
+	constexpr int NumOfPoints = 20;
+
+	Settings->MetadataDomainsMapping.Emplace(PCGPointDataConstants::ElementsDomainName, PCGDataConstants::GlobalDomainName);
+	Settings->bCopyAllAttributes = true;
+
+	TUniquePtr<FPCGContext> Context = TestData.InitializeTestContext();
+	UPCGBasePointData* TargetPointData = PCGCopyAttributeTests::CreateInputPointData(Context.Get(), NumOfPoints);
+	UPCGBasePointData* SourcePointData = PCGCopyAttributeTests::CreateInputPointData(Context.Get(), NumOfPoints);
+
+	PCGTestsCommon::CreateAndFillRandomAttribute<double>(SourcePointData, TEXT("MyAttr"), 0.0, NumOfPoints);
+	const FPCGMetadataAttribute<double>* InputAttribute = SourcePointData->MutableMetadata()->GetMetadataDomain(PCGMetadataDomainID::Elements)->GetConstTypedAttribute<double>(TEXT("MyAttr"));
+	
+	PCGCopyAttributeTests::ConnectToSource(Context.Get(), SourcePointData);
+	PCGCopyAttributeTests::ConnectToTarget(Context.Get(), TargetPointData);
+
+	FPCGElementPtr TestElement = TestData.Settings->GetElement();
+	while (!TestElement->Execute(Context.Get())) {}
+
+	TArray<FPCGTaggedData> OutputTagged = Context->OutputData.GetInputsByPin(PCGPinConstants::DefaultOutputLabel);
+	const UPCGBasePointData* OutputData = OutputTagged.Num() == 1 ? Cast<const UPCGBasePointData>(OutputTagged[0].Data) : nullptr;
+
+	UTEST_NOT_NULL("There is a point data in output points", OutputData);
+	UTEST_EQUAL("There is the right number of points in output", OutputData->GetNumPoints(), NumOfPoints);
+
+	check(OutputData && OutputData->ConstMetadata());
+	const FPCGMetadataDomain* GlobalDomain = OutputData->ConstMetadata()->GetConstMetadataDomain(PCGMetadataDomainID::Data);
+	check(GlobalDomain);
+	
+	const FPCGMetadataAttribute<double>* OutputAttribute = GlobalDomain->GetConstTypedAttribute<double>(TEXT("MyAttr"));
+	UTEST_NOT_NULL("Output Attribute exists in the output global metadata domain", OutputAttribute);
+
+	check(InputAttribute && OutputAttribute);
+
+	UTEST_EQUAL("Output Attribute default value is the same as default value of Input Attribute", OutputAttribute->GetValueFromItemKey(PCGInvalidEntryKey), InputAttribute->GetValueFromItemKey(PCGInvalidEntryKey));
+
+	TConstPCGValueRange<PCGMetadataEntryKey> EntryKeyRange = SourcePointData->GetConstMetadataEntryValueRange();
+	UTEST_EQUAL("Output Attribute value is the same as the first value of Input Attribute", OutputAttribute->GetValueFromItemKey(0), InputAttribute->GetValueFromItemKey(EntryKeyRange[0]));
+	
+	return true;
+}
+
+bool FPCGCopyAttributeTests_MultiDomain_CopyAll_ElementsToElementsExplicit::RunTest(const FString& Parameters)
+{
+	PCGTestsCommon::FTestData TestData;
+	UPCGCopyAttributesSettings* Settings = PCGTestsCommon::GenerateSettings<UPCGCopyAttributesSettings>(TestData);
+	check(Settings);
+
+	constexpr int NumOfPoints = 20;
+
+	Settings->MetadataDomainsMapping.Emplace(PCGPointDataConstants::ElementsDomainName, PCGPointDataConstants::ElementsDomainName);
+	Settings->bCopyAllAttributes = true;
+
+	TUniquePtr<FPCGContext> Context = TestData.InitializeTestContext();
+	UPCGBasePointData* SourcePointData = PCGCopyAttributeTests::CreateInputPointData(Context.Get(), NumOfPoints);
+	PCGTestsCommon::CreateAndFillRandomAttribute<FString>(SourcePointData, PCGCopyAttributeTests::AttributeName, FString(), NumOfPoints);
+	
+	UPCGBasePointData* TargetPointData = PCGCopyAttributeTests::CreateInputPointData(Context.Get(), NumOfPoints);
+	
+	PCGCopyAttributeTests::ConnectToSource(Context.Get(), SourcePointData);
+	PCGCopyAttributeTests::ConnectToTarget(Context.Get(), TargetPointData);
+
+	FPCGElementPtr TestElement = TestData.Settings->GetElement();
+	while (!TestElement->Execute(Context.Get())) {}
+
+	TArray<FPCGTaggedData> OutputTagged = Context->OutputData.GetInputsByPin(PCGPinConstants::DefaultOutputLabel);
+	const UPCGBasePointData* OutputData = OutputTagged.Num() == 1 ? Cast<const UPCGBasePointData>(OutputTagged[0].Data) : nullptr;
+
+	UTEST_NOT_NULL("There is a point data in output points", OutputData);
+	UTEST_EQUAL("There is the right number of points in output", OutputData->GetNumPoints(), NumOfPoints);
+
+	check(OutputData && OutputData->Metadata);
+
+	const FPCGMetadataAttribute<FString>* InputAttribute = SourcePointData->ConstMetadata()->GetConstTypedAttribute<FString>(PCGCopyAttributeTests::AttributeName);
+	UTEST_NOT_NULL("Input Attribute exists in the source data", InputAttribute);
+	const FPCGMetadataAttribute<FString>* OutputAttribute = OutputData->ConstMetadata()->GetConstTypedAttribute<FString>(PCGCopyAttributeTests::AttributeName);
+	UTEST_NOT_NULL("Output Attribute exists in the output data", OutputAttribute);
+
+	check(InputAttribute && OutputAttribute);
+
+	UTEST_EQUAL("Output Attribute default value is the same as default value of Input Attribute", OutputAttribute->GetValueFromItemKey(PCGInvalidEntryKey), InputAttribute->GetValueFromItemKey(PCGInvalidEntryKey));
+
+	// Input attribute was copied in the output attribute correctly
+	TConstPCGValueRange<PCGMetadataEntryKey> OutputMetadataEntryRange = OutputData->GetConstMetadataEntryValueRange();
+	TConstPCGValueRange<PCGMetadataEntryKey> SourceMetadataEntryRange = SourcePointData->GetConstMetadataEntryValueRange();
+	for (int i = 0; i < NumOfPoints; ++i)
+	{
+		UTEST_EQUAL(*FString::Printf(TEXT("Output Attribute has the same value as Input Attribute for point %d"), i), OutputAttribute->GetValueFromItemKey(OutputMetadataEntryRange[i]), InputAttribute->GetValueFromItemKey(SourceMetadataEntryRange[i]));
+	}
+
+	return true;
+}
+
+bool FPCGCopyAttributeTests_MultiDomain_CopyAll_ElementsToElementsDefault::RunTest(const FString& Parameters)
+{
+	PCGTestsCommon::FTestData TestData;
+	UPCGCopyAttributesSettings* Settings = PCGTestsCommon::GenerateSettings<UPCGCopyAttributesSettings>(TestData);
+	check(Settings);
+
+	constexpr int NumOfPoints = 20;
+	
+	Settings->bCopyAllAttributes = true;
+
+	TUniquePtr<FPCGContext> Context = TestData.InitializeTestContext();
+	UPCGBasePointData* SourcePointData = PCGCopyAttributeTests::CreateInputPointData(Context.Get(), NumOfPoints);
+	PCGTestsCommon::CreateAndFillRandomAttribute<FString>(SourcePointData, PCGCopyAttributeTests::AttributeName, FString(), NumOfPoints);
+	
+	UPCGBasePointData* TargetPointData = PCGCopyAttributeTests::CreateInputPointData(Context.Get(), NumOfPoints);
+	
+	PCGCopyAttributeTests::ConnectToSource(Context.Get(), SourcePointData);
+	PCGCopyAttributeTests::ConnectToTarget(Context.Get(), TargetPointData);
+
+	FPCGElementPtr TestElement = TestData.Settings->GetElement();
+	while (!TestElement->Execute(Context.Get())) {}
+
+	TArray<FPCGTaggedData> OutputTagged = Context->OutputData.GetInputsByPin(PCGPinConstants::DefaultOutputLabel);
+	const UPCGBasePointData* OutputData = OutputTagged.Num() == 1 ? Cast<const UPCGBasePointData>(OutputTagged[0].Data) : nullptr;
+
+	UTEST_NOT_NULL("There is a point data in output points", OutputData);
+	UTEST_EQUAL("There is the right number of points in output", OutputData->GetNumPoints(), NumOfPoints);
+
+	check(OutputData && OutputData->Metadata);
+
+	const FPCGMetadataAttribute<FString>* InputAttribute = SourcePointData->ConstMetadata()->GetConstTypedAttribute<FString>(PCGCopyAttributeTests::AttributeName);
+	UTEST_NOT_NULL("Input Attribute exists in the source data", InputAttribute);
+	const FPCGMetadataAttribute<FString>* OutputAttribute = OutputData->ConstMetadata()->GetConstTypedAttribute<FString>(PCGCopyAttributeTests::AttributeName);
+	UTEST_NOT_NULL("Output Attribute exists in the output data", OutputAttribute);
+
+	check(InputAttribute && OutputAttribute);
+
+	UTEST_EQUAL("Output Attribute default value is the same as default value of Input Attribute", OutputAttribute->GetValueFromItemKey(PCGInvalidEntryKey), InputAttribute->GetValueFromItemKey(PCGInvalidEntryKey));
+
+	// Input attribute was copied in the output attribute correctly
+	TConstPCGValueRange<PCGMetadataEntryKey> OutputMetadataEntryRange = OutputData->GetConstMetadataEntryValueRange();
+	TConstPCGValueRange<PCGMetadataEntryKey> SourceMetadataEntryRange = SourcePointData->GetConstMetadataEntryValueRange();
+	for (int i = 0; i < NumOfPoints; ++i)
+	{
+		UTEST_EQUAL(*FString::Printf(TEXT("Output Attribute has the same value as Input Attribute for point %d"), i), OutputAttribute->GetValueFromItemKey(OutputMetadataEntryRange[i]), InputAttribute->GetValueFromItemKey(SourceMetadataEntryRange[i]));
+	}
+
+	return true;
+}
+
+bool FPCGCopyAttributeTests_MultiDomain_CopyAll_InvalidDomain::RunTest(const FString& Parameters)
+{
+	PCGTestsCommon::FTestData TestData;
+	UPCGCopyAttributesSettings* Settings = PCGTestsCommon::GenerateSettings<UPCGCopyAttributesSettings>(TestData);
+	check(Settings);
+
+	constexpr int NumOfPoints = 20;
+
+	Settings->MetadataDomainsMapping.Reset();
+	Settings->MetadataDomainsMapping.Emplace(TEXT("Bliblibli"), TEXT("Bloubloublou"));
+	Settings->bCopyAllAttributes = true;
+
+	TUniquePtr<FPCGContext> Context = TestData.InitializeTestContext();
+	const UPCGBasePointData* SourcePointData = PCGCopyAttributeTests::CreateInputPointData(Context.Get(), NumOfPoints);
+	PCGTestsCommon::CreateAndFillRandomAttribute<double>(const_cast<UPCGBasePointData*>(SourcePointData), TEXT("MyAttr"), 0.0, NumOfPoints);
+	
+	const UPCGBasePointData* TargetPointData = PCGCopyAttributeTests::CreateInputPointData(Context.Get(), NumOfPoints);
+	
+	PCGCopyAttributeTests::ConnectToSource(Context.Get(), SourcePointData);
+	PCGCopyAttributeTests::ConnectToTarget(Context.Get(), TargetPointData);
+
+	AddExpectedError(TEXT("Metadata domain Bliblibli is invalid for this data."), EAutomationExpectedErrorFlags::MatchType::Exact, 1, /*bIsRegex=*/false);
+
+	FPCGElementPtr TestElement = TestData.Settings->GetElement();
+	while (!TestElement->Execute(Context.Get())) {}
+
+	TArray<FPCGTaggedData> OutputTagged = Context->OutputData.GetInputsByPin(PCGPinConstants::DefaultOutputLabel);
+	const UPCGBasePointData* OutputData = OutputTagged.Num() == 1 ? Cast<const UPCGBasePointData>(OutputTagged[0].Data) : nullptr;
+
+	UTEST_NOT_NULL("There is a point data in output points", OutputData);
+	UTEST_EQUAL("It's the same data as input", OutputData, TargetPointData);
+	
+	return true;
+}
 #endif // WITH_EDITOR
diff --git a/Engine/Plugins/PCG/Source/PCG/Public/Elements/PCGCopyAttributes.h b/Engine/Plugins/PCG/Source/PCG/Public/Elements/PCGCopyAttributes.h
--- a/Engine/Plugins/PCG/Source/PCG/Public/Elements/PCGCopyAttributes.h
+++ b/Engine/Plugins/PCG/Source/PCG/Public/Elements/PCGCopyAttributes.h
@@ -70,6 +70,14 @@ public:
 	UPROPERTY(BlueprintReadWrite, EditAnywhere, Category = Settings, meta = (PCG_Overridable))
 	bool bCopyAllAttributes = false;
 
+	/** If checked, it is copying all attributes from all domains, as long as the source domain is supported on the target data.*/
+	UPROPERTY(BlueprintReadWrite, EditAnywhere, Category = Settings, meta = (EditCondition="bCopyAllAttributes", PCG_Overridable))
+	bool bCopyAllDomains = false;
+
+	/** When copying all attributes, a mapping can be specified. If it is empty, it's going to be Default -> Default. */
+	UPROPERTY(BlueprintReadWrite, EditAnywhere, Category = Settings, meta = (EditCondition="bCopyAllAttributes && !bCopyAllDomains", EditConditionHides))
+	TMap<FName, FName> MetadataDomainsMapping;
+
 	UPROPERTY(BlueprintReadWrite, EditAnywhere, Category = Settings, meta = (EditCondition="!bCopyAllAttributes", PCG_Overridable))
 	FPCGAttributePropertyInputSelector InputSource;
 
diff --git a/Engine/Plugins/PCG/Source/PCG/Public/Elements/PCGCreateAttribute.h b/Engine/Plugins/PCG/Source/PCG/Public/Elements/PCGCreateAttribute.h
--- a/Engine/Plugins/PCG/Source/PCG/Public/Elements/PCGCreateAttribute.h
+++ b/Engine/Plugins/PCG/Source/PCG/Public/Elements/PCGCreateAttribute.h
@@ -62,6 +62,14 @@ public:
 	UPROPERTY(BlueprintReadWrite, EditAnywhere, Category = Settings, meta = (PCG_Overridable))
 	bool bCopyAllAttributes = false;
 
+	/** If checked, it is copying all attributes from all domains, as long as the source domain is supported on the target data.*/
+	UPROPERTY(BlueprintReadWrite, EditAnywhere, Category = Settings, meta = (EditCondition="bCopyAllAttributes", PCG_Overridable))
+	bool bCopyAllDomains = false;
+
+	/** When copying all attributes, a mapping can be specified. If it is empty, it's going to be Default -> Default. */
+	UPROPERTY(BlueprintReadWrite, EditAnywhere, Category = Settings, meta = (EditCondition="bCopyAllAttributes && !bCopyAllDomains", EditConditionHides))
+	TMap<FName, FName> MetadataDomainsMapping;
+
 #if WITH_EDITORONLY_DATA
 	UPROPERTY()
 	FName SourceParamAttributeName_DEPRECATED = NAME_None;
diff --git a/Engine/Plugins/PCG/Source/PCG/Public/Helpers/PCGMetadataHelpers.h b/Engine/Plugins/PCG/Source/PCG/Public/Helpers/PCGMetadataHelpers.h
--- a/Engine/Plugins/PCG/Source/PCG/Public/Helpers/PCGMetadataHelpers.h
+++ b/Engine/Plugins/PCG/Source/PCG/Public/Helpers/PCGMetadataHelpers.h
@@ -36,7 +36,7 @@ namespace PCGMetadataHelpers
 	/** Create a lambda that will construct a soft object path from an underlying attribute of type FSoftObjectPath or FString. Returns true if successful. */
 	[[nodiscard]] PCG_API bool CreateObjectOrClassPathGetter(const FPCGMetadataAttributeBase* InAttributeBase, TFunction<void(int64, FSoftObjectPath&)>& OutGetter);
 
-	struct PCG_API FPCGCopyAttributeParams
+	struct FPCGCopyAttributeParams
 	{
 		/** Source data where the attribute is coming from */
 		const UPCGData* SourceData = nullptr;
@@ -63,9 +63,36 @@ namespace PCGMetadataHelpers
 	/** Copy the attribute coming from Source Data into Target Data. */
 	PCG_API bool CopyAttribute(const FPCGCopyAttributeParams& InParams);
 
-	/** Copy all the attributes coming from Source Data into Target Data. */
+	struct FPCGCopyAllAttributesParams
+	{
+		/** Source data where the attribute is coming from */
+		const UPCGData* SourceData = nullptr;
+
+		/** Target data to write to */
+		UPCGData* TargetData = nullptr;
+
+		/**
+		 * Metadata domains mapping. Empty means copying all domains from target to source, as long as they are compatible.
+		 * For retro-compatibility, it is initialized as Default -> Default only.
+		 */
+		TMap<FPCGMetadataDomainID, FPCGMetadataDomainID> DomainMapping = {{PCGMetadataDomainID::Default, PCGMetadataDomainID::Default}}; 
+
+		/** Optional context for logging */
+		FPCGContext* OptionalContext = nullptr;
+
+		/** Will initialize DomainMapping using SourceData and TargetData to convert names into domain IDs. If the mapping is empty, will be default -> default.*/
+		PCG_API void InitializeMappingFromDomainNames(const TMap<FName, FName>& MetadataDomainsMapping);
+
+		/** Will map all matching domains.*/
+		PCG_API void InitializeMappingForAllDomains();
+	};
+
+	UE_DEPRECATED(5.6, "Use the version with FPCGCopyAllAttributesParams")
 	PCG_API bool CopyAllAttributes(const UPCGData* SourceData, UPCGData* TargetData, FPCGContext* OptionalContext = nullptr);
 
+	/** Copy all the attributes coming from Source Data into Target Data. */
+	PCG_API bool CopyAllAttributes(const FPCGCopyAllAttributesParams& InParams);
+	
 	PCG_API void SetPointAttributes(FPCGMetadataDomain* InOutMetadata, const TArrayView<const FPCGPoint>& InPoints, const FPCGMetadataDomain* InMetadata, const TArrayView<FPCGPoint>& OutPoints, FPCGContext* OptionalContext = nullptr);
 
 	PCG_API void ComputePointWeightedAttribute(FPCGMetadataDomain* InOutMetadata, FPCGPoint& OutPoint, const TArrayView<TPair<const FPCGPoint*, float>>& InWeightedPoints, const FPCGMetadataDomain* InMetadata);
diff --git a/Engine/Plugins/PCG/Source/PCG/Public/Metadata/Accessors/PCGCustomAccessor.h b/Engine/Plugins/PCG/Source/PCG/Public/Metadata/Accessors/PCGCustomAccessor.h
--- a/Engine/Plugins/PCG/Source/PCG/Public/Metadata/Accessors/PCGCustomAccessor.h
+++ b/Engine/Plugins/PCG/Source/PCG/Public/Metadata/Accessors/PCGCustomAccessor.h
@@ -219,11 +219,13 @@ public:
 class FPCGAttributeAccessorKeysPointIndices : public IPCGAttributeAccessorKeys
 {
 public:
-	FPCGAttributeAccessorKeysPointIndices(UPCGBasePointData* InPointData)
+	FPCGAttributeAccessorKeysPointIndices(UPCGBasePointData* InPointData, bool bAllocateMetadataEntries = false)
 		: IPCGAttributeAccessorKeys(/*bInReadOnly=*/false)
 		, PointData(InPointData)
 	{
-		MetadataEntryKeys = PointData->GetMetadataEntryValueRange(/*bAllocate=*/false);
+		// By default don't allocate Metadata Entries, since we don't know if the keys are going to be used to write into attributes.
+		// If it is known that it will write into attributes, set bAllocateMetadataEntries to true.
+		MetadataEntryKeys = PointData->GetMetadataEntryValueRange(bAllocateMetadataEntries);
 		ConstMetadataEntryKeys = PointData->GetConstMetadataEntryValueRange();
 	}
 
diff --git a/Engine/Plugins/PCG/Source/PCG/Public/Tests/PCGTestsCommon.h b/Engine/Plugins/PCG/Source/PCG/Public/Tests/PCGTestsCommon.h
--- a/Engine/Plugins/PCG/Source/PCG/Public/Tests/PCGTestsCommon.h
+++ b/Engine/Plugins/PCG/Source/PCG/Public/Tests/PCGTestsCommon.h
@@ -241,9 +241,10 @@ namespace PCGTestsCommon
 		check(InData);
 
 		UPCGMetadata* Metadata = InData->MutableMetadata();
-		UPCGPointData* PointData = Cast<UPCGPointData>(InData);
+		UPCGBasePointData* PointData = Cast<UPCGBasePointData>(InData);
+		TPCGValueRange<PCGMetadataEntryKey> MetadataEntryRange = PointData ? PointData->GetMetadataEntryValueRange() : TPCGValueRange<PCGMetadataEntryKey>();
 
-		check(Metadata && (!PointData || NumValues == 0 || NumValues == PointData->GetPoints().Num()));
+		check(Metadata && (!PointData || NumValues == 0 || NumValues == PointData->GetNumPoints()));
 
 		const bool bAllowInterpolation = bForceAllowInterpolation ? *bForceAllowInterpolation : PCG::Private::MetadataTraits<T>::CanInterpolate;
 		FPCGMetadataAttribute<T>* NewAttribute = Metadata->CreateAttribute<T>(AttributeName, /*DefaultValue=*/std::move(DefaultValue), /*bAllowInterpolation=*/bAllowInterpolation, /*bOverrideParent=*/false);
@@ -256,8 +257,8 @@ namespace PCGTestsCommon
 			PCGMetadataEntryKey EntryKey;
 			if (PointData)
 			{
-				Metadata->InitializeOnSet(PointData->GetMutablePoints()[i].MetadataEntry);
-				EntryKey = PointData->GetMutablePoints()[i].MetadataEntry;
+				Metadata->InitializeOnSet(MetadataEntryRange[i]);
+				EntryKey = MetadataEntryRange[i];
 			}
 			else
 			{
