[PCG] Fix metadata sampling on spline closed loops

* Also make the control points the default metadata domain for splines.

#rnx

diff --git a/Engine/Plugins/PCG/Source/PCG/Private/Data/PCGSplineData.cpp b/Engine/Plugins/PCG/Source/PCG/Private/Data/PCGSplineData.cpp
--- a/Engine/Plugins/PCG/Source/PCG/Private/Data/PCGSplineData.cpp
+++ b/Engine/Plugins/PCG/Source/PCG/Private/Data/PCGSplineData.cpp
@@ -22,7 +22,7 @@
 UPCGSplineData::UPCGSplineData()
 {
 	check(Metadata);
-	Metadata->SetupDomain(PCGMetadataDomainID::Elements, /*bIsDefault=*/false);
+	Metadata->SetupDomain(PCGMetadataDomainID::Elements, /*bIsDefault=*/true);
 }
 
 void UPCGSplineData::Initialize(const USplineComponent* InSpline)
@@ -215,7 +215,15 @@ FTransform UPCGSplineData::GetTransformAtAlpha(float Alpha) const
 
 void UPCGSplineData::WriteMetadataToPoint(float InputKey, FPCGPoint& OutPoint, UPCGMetadata* OutMetadata) const
 {
-	if (!OutMetadata)
+	if (!OutMetadata || !ensure(Metadata))
+	{
+		return;
+	}
+
+	const FPCGMetadataDomain* MetadataDomain = Metadata->GetConstMetadataDomain(PCGMetadataDomainID::Elements);
+	FPCGMetadataDomain* OutMetadataDomain = OutMetadata->GetMetadataDomain(PCGMetadataDomainID::Elements);
+
+	if (!MetadataDomain || !OutMetadataDomain)
 	{
 		return;
 	}
@@ -224,20 +232,42 @@ void UPCGSplineData::WriteMetadataToPoint(float InputKey, FPCGPoint& OutPoint, U
 	TConstArrayView<PCGMetadataEntryKey> EntryKeys = SplineStruct.GetConstControlPointsEntryKeys();
 	if (!EntryKeys.IsEmpty())
 	{
-		int32 PreviousIndex = FMath::Floor(InputKey);
-		int32 NextIndex = FMath::CeilToInt32(InputKey);
-		if (EntryKeys.IsValidIndex(PreviousIndex) && EntryKeys.IsValidIndex(NextIndex))
+		const TTuple<int, float> PreviousIndexAndKey = SplineStruct.GetSegmentStartIndexAndKeyAtInputKey(InputKey);
+		const int PreviousIndex = PreviousIndexAndKey.Get<0>();
+		const float PreviousInputKey = PreviousIndexAndKey.Get<1>();
+
+		if (PreviousIndex == INDEX_NONE)
 		{
-			OutMetadata->InitializeOnSet(OutPoint.MetadataEntry, EntryKeys[PreviousIndex], Metadata);
-			
-			if (PreviousIndex != NextIndex)
-			{
-				const float Alpha = InputKey - PreviousIndex;
-				TStaticArray<TPair<PCGMetadataEntryKey, float>, 2> Coefficients;
-				Coefficients[0] = {EntryKeys[PreviousIndex], 1.0f - Alpha};
-				Coefficients[1] = {EntryKeys[NextIndex], Alpha};
-				OutMetadata->ComputeWeightedAttribute(OutPoint.MetadataEntry, Coefficients, Metadata);
-			}
+			return;
+		}
+		
+		int32 NextIndex = PreviousIndex + 1;
+		float NextInputKey = 0;
+		// In case the spline is not closed, we cannot go further than the last point.
+		// In case the spline is closed, we have to get the input key first (to get the input key associated with the last point, not the first point)
+		// then reset it to zero if it is the last point.
+		if (!IsClosed())
+		{
+			NextIndex = FMath::Min(NextIndex, SplineStruct.GetNumberOfPoints() - 1);
+			NextInputKey = SplineStruct.GetInputKeyAtSegmentStart(NextIndex);
+		}
+		else
+		{
+			NextInputKey = SplineStruct.GetInputKeyAtSegmentStart(NextIndex);
+			NextIndex %= SplineStruct.GetNumberOfPoints();
+		}
+
+		check(EntryKeys.IsValidIndex(PreviousIndex) && EntryKeys.IsValidIndex(NextIndex));
+		
+		OutMetadataDomain->InitializeOnSet(OutPoint.MetadataEntry, EntryKeys[PreviousIndex], MetadataDomain);
+		
+		if (PreviousIndex != NextIndex && !FMath::IsNearlyEqual(PreviousInputKey, NextInputKey))
+		{
+			const float Alpha = (InputKey - PreviousInputKey) / (NextInputKey - PreviousInputKey);
+			TStaticArray<TPair<PCGMetadataEntryKey, float>, 2> Coefficients;
+			Coefficients[0] = {EntryKeys[PreviousIndex], 1.0f - Alpha};
+			Coefficients[1] = {EntryKeys[NextIndex], Alpha};
+			OutMetadataDomain->ComputeWeightedAttribute(OutPoint.MetadataEntry, Coefficients, MetadataDomain);
 		}
 	}
 }
@@ -318,14 +348,16 @@ FPCGAttributeAccessorMethods UPCGSplineData::GetSplineAccessorMethods()
 		{
 			return MakeUnique<FPCGAttributeAccessorKeysSplineData>(SplineData);
 		}
-		else if (Selection == EPCGAttributePropertySelection::Attribute && InSelector.GetDomainName() == PCGSplineData::ControlPointDomainName)
-		{
-			return MakeUnique<FPCGAttributeAccessorKeysSplineDataEntries>(SplineData);
-		}
-		else
+		else if (Selection == EPCGAttributePropertySelection::Attribute)
 		{
-			return {};
+			const FPCGMetadataDomainID DomainID = InData->GetMetadataDomainIDFromSelector(InSelector);
+			if (DomainID.IsDefault() || DomainID == PCGMetadataDomainID::Elements)
+			{
+				return MakeUnique<FPCGAttributeAccessorKeysSplineDataEntries>(SplineData);
+			}
 		}
+
+		return {};
 	};
 	
 	auto CreateConstAccessorKeysFunc = [](const UPCGData* InData, const FPCGAttributePropertySelector& InSelector, bool bQuiet) -> TUniquePtr<const IPCGAttributeAccessorKeys>
@@ -337,14 +369,16 @@ FPCGAttributeAccessorMethods UPCGSplineData::GetSplineAccessorMethods()
 		{
 			return MakeUnique<const FPCGAttributeAccessorKeysSplineData>(SplineData);
 		}
-		else if (Selection == EPCGAttributePropertySelection::Attribute && InSelector.GetDomainName() == PCGSplineData::ControlPointDomainName)
-		{
-			return MakeUnique<const FPCGAttributeAccessorKeysSplineDataEntries>(SplineData);
-		}
-		else
+		else if (Selection == EPCGAttributePropertySelection::Attribute)
 		{
-			return {};
+			const FPCGMetadataDomainID DomainID = InData->GetMetadataDomainIDFromSelector(InSelector);
+			if (DomainID.IsDefault() || DomainID == PCGMetadataDomainID::Elements)
+			{
+				return MakeUnique<FPCGAttributeAccessorKeysSplineDataEntries>(SplineData);
+			}
 		}
+
+		return {};
 	};
 	
 	FPCGAttributeAccessorMethods Methods
diff --git a/Engine/Plugins/PCG/Source/PCG/Private/Data/PCGSplineStruct.cpp b/Engine/Plugins/PCG/Source/PCG/Private/Data/PCGSplineStruct.cpp
--- a/Engine/Plugins/PCG/Source/PCG/Private/Data/PCGSplineStruct.cpp
+++ b/Engine/Plugins/PCG/Source/PCG/Private/Data/PCGSplineStruct.cpp
@@ -212,6 +212,11 @@ int FPCGSplineStruct::GetNumberOfSplineSegments() const
 	return (bClosedLoop ? NumPoints : NumPoints - 1);
 }
 
+int FPCGSplineStruct::GetNumberOfPoints() const
+{
+	return SplineCurves.Position.Points.Num();
+}
+
 FVector::FReal FPCGSplineStruct::GetSplineLength() const
 {
 	return SplineCurves.GetSplineLength();
@@ -337,6 +342,25 @@ float FPCGSplineStruct::FindInputKeyClosestToWorldLocation(const FVector& WorldL
 	return SplineCurves.Position.InaccurateFindNearest(LocalLocation, Dummy);
 }
 
+TTuple<int, float> FPCGSplineStruct::GetSegmentStartIndexAndKeyAtInputKey(float InKey) const
+{
+	const int32 Index = SplineCurves.Position.GetPointIndexForInputValue(InKey);
+	return {Index, GetInputKeyAtSegmentStart(Index)};
+}
+
+float FPCGSplineStruct::GetInputKeyAtSegmentStart(int InSegmentIndex) const
+{
+	if (IsClosedLoop() && !SplineCurves.Position.Points.IsEmpty() && InSegmentIndex >= SplineCurves.Position.Points.Num())
+	{
+		// In case of a closed loop, the last point is the first point, and the input key is the last point + LoopKeyOffset
+		return SplineCurves.Position.Points.Last().InVal + SplineCurves.Position.LoopKeyOffset;
+	}
+	else
+	{
+		return SplineCurves.Position.Points.IsValidIndex(InSegmentIndex) ? SplineCurves.Position.Points[InSegmentIndex].InVal : 0.0f;
+	}
+}
+
 void FPCGSplineStruct::AllocateMetadataEntries()
 {
 	// Add robustness to cleanup everything if we have a mismatch between the number of points and the number of entry keys
diff --git a/Engine/Plugins/PCG/Source/PCG/Public/Data/PCGSplineStruct.h b/Engine/Plugins/PCG/Source/PCG/Public/Data/PCGSplineStruct.h
--- a/Engine/Plugins/PCG/Source/PCG/Public/Data/PCGSplineStruct.h
+++ b/Engine/Plugins/PCG/Source/PCG/Public/Data/PCGSplineStruct.h
@@ -38,6 +38,7 @@ struct PCG_API FPCGSplineStruct
 	void UpdateSpline();
 
 	int GetNumberOfSplineSegments() const;
+	int GetNumberOfPoints() const;
 	bool IsClosedLoop() const { return bClosedLoop; }
 	FVector::FReal GetSplineLength() const;
 	FBox GetBounds() const;
@@ -65,6 +66,15 @@ struct PCG_API FPCGSplineStruct
 	
 	float FindInputKeyClosestToWorldLocation(const FVector& WorldLocation) const;
 
+	/** Return the index of the segment for a given InputKey, and the input key at the segment start.
+	 * It's for future-proof-ness as even if there is no guarantee a segment start is at an integer input key, right now it is not well supported.
+	 * If the input key is invalid, returns {INDEX_NONE, 0}.
+	 */
+	TTuple<int, float> GetSegmentStartIndexAndKeyAtInputKey(float InKey) const;
+
+	/** Return the input key at segment start for the given segment. Takes into account if the spline is closed. */
+	float GetInputKeyAtSegmentStart(int InSegmentIndex) const;
+
 	/** To be called at any moment we want to access individual metadata, like with accessors. Does nothing if it was already allocated. */
 	void AllocateMetadataEntries();
 	TConstArrayView<PCGMetadataEntryKey> GetConstControlPointsEntryKeys() const { return ControlPointsEntryKeys; }
