[PCG] Fix repeated preparation of data providers for execution.

#rnx

diff --git a/Engine/Plugins/PCG/Source/PCG/Private/Compute/DataInterfaces/PCGInstanceDataInterface.cpp b/Engine/Plugins/PCG/Source/PCG/Private/Compute/DataInterfaces/PCGInstanceDataInterface.cpp
--- a/Engine/Plugins/PCG/Source/PCG/Private/Compute/DataInterfaces/PCGInstanceDataInterface.cpp
+++ b/Engine/Plugins/PCG/Source/PCG/Private/Compute/DataInterfaces/PCGInstanceDataInterface.cpp
@@ -154,7 +154,8 @@ bool UPCGInstanceDataProvider::PrepareForExecute_GameThread(UPCGDataBinding* InB
 
 	if (!FoundPrimitives)
 	{
-		return true;
+		// The static mesh data interface(s) set this up, so wait until it is ready.
+		return false;
 	}
 
 	if (!ensure(FoundPrimitives->Primitives.Num() <= PCGComputeConstants::MAX_PRIMITIVE_COMPONENTS_PER_SPAWNER))
diff --git a/Engine/Plugins/PCG/Source/PCG/Private/Compute/Elements/PCGComputeGraphElement.cpp b/Engine/Plugins/PCG/Source/PCG/Private/Compute/Elements/PCGComputeGraphElement.cpp
--- a/Engine/Plugins/PCG/Source/PCG/Private/Compute/Elements/PCGComputeGraphElement.cpp
+++ b/Engine/Plugins/PCG/Source/PCG/Private/Compute/Elements/PCGComputeGraphElement.cpp
@@ -374,6 +374,11 @@ bool FPCGComputeGraphElement::ExecuteInternal(FPCGContext* InContext) const
 			return true;
 		}
 
+		Algo::Transform(Context->ComputeGraphInstance->GetDataProviders(), Context->DataProvidersPendingReadyForExecute, [](UComputeDataProvider* InProvider)
+		{
+			return Cast<UPCGComputeDataProvider>(InProvider);
+		});
+
 		Context->ExecutionSubPhase = EPCGComputeGraphExecutionPhase::PrepareForExecute;
 
 		if (Context->AsyncState.ShouldStop())
@@ -385,20 +390,17 @@ bool FPCGComputeGraphElement::ExecuteInternal(FPCGContext* InContext) const
 	{
 		TRACE_CPUPROFILER_EVENT_SCOPE(PrepareForExecute);
 
-		bool bAllReady = true;
-
-		for (UComputeDataProvider* ComputeDataProvider : Context->ComputeGraphInstance->GetDataProviders())
+		for (int Index = Context->DataProvidersPendingReadyForExecute.Num() - 1; Index >= 0; --Index)
 		{
-			if (UPCGComputeDataProvider* DataProvider = Cast<UPCGComputeDataProvider>(ComputeDataProvider))
+			UPCGComputeDataProvider* DataProvider = Context->DataProvidersPendingReadyForExecute[Index];
+
+			if (!DataProvider || DataProvider->PrepareForExecute_GameThread(Context->DataBinding.Get()))
 			{
-				if (!DataProvider->PrepareForExecute_GameThread(Context->DataBinding.Get()))
-				{
-					bAllReady = false;
-				}
+				Context->DataProvidersPendingReadyForExecute.RemoveAtSwap(Index);
 			}
 		}
 
-		if (bAllReady)
+		if (Context->DataProvidersPendingReadyForExecute.IsEmpty())
 		{
 			Context->ExecutionSubPhase = EPCGComputeGraphExecutionPhase::ValidateComputeGraphCompilation;
 		}
diff --git a/Engine/Plugins/PCG/Source/PCG/Private/Compute/Elements/PCGComputeGraphElement.h b/Engine/Plugins/PCG/Source/PCG/Private/Compute/Elements/PCGComputeGraphElement.h
--- a/Engine/Plugins/PCG/Source/PCG/Private/Compute/Elements/PCGComputeGraphElement.h
+++ b/Engine/Plugins/PCG/Source/PCG/Private/Compute/Elements/PCGComputeGraphElement.h
@@ -14,6 +14,8 @@
 
 #include "PCGComputeGraphElement.generated.h"
 
+class UPCGComputeDataProvider;
+
 UENUM()
 enum class EPCGComputeGraphExecutionPhase : uint8
 {
@@ -63,6 +65,9 @@ public:
 	/** Graph executed successfully. */
 	bool bExecutionSuccess = false;
 
+	/** The data providers which are not yet ready for execution. */
+	TArray<UPCGComputeDataProvider*> DataProvidersPendingReadyForExecute;
+
 	/** Data providers with buffers that are passed to downstream tasks. The buffer will be created on the render thread, and then passed
 	* back to main thread, upon which a reference is taken to the buffer and the provider is removed from this set to signal completion.
 	*/
diff --git a/Engine/Plugins/PCG/Source/PCG/Private/Compute/PCGDataUploadAdaptor.cpp b/Engine/Plugins/PCG/Source/PCG/Private/Compute/PCGDataUploadAdaptor.cpp
--- a/Engine/Plugins/PCG/Source/PCG/Private/Compute/PCGDataUploadAdaptor.cpp
+++ b/Engine/Plugins/PCG/Source/PCG/Private/Compute/PCGDataUploadAdaptor.cpp
@@ -115,6 +115,8 @@ bool FPCGDataUploadAdaptor::PrepareData_GameThread()
 
 	if (ExternalBufferForReuse)
 	{
+		ensureAlways(SourceToTargetAttributeId.IsEmpty());
+
 		SourceToTargetAttributeId.Reserve(Binding->GetAttributeIdRemap().Num());
 
 		Algo::Transform(Binding->GetAttributeIdRemap(), SourceToTargetAttributeId, [](const TPair<UPCGDataBinding::FSourceBufferAndAttributeIndex, int32>& Remapping)
