[PCG] Added option to check world is streamed in under local components before scheduling generation on them.

#rnx

diff --git a/Engine/Plugins/PCG/Source/PCG/Private/Helpers/PCGActorHelpers.cpp b/Engine/Plugins/PCG/Source/PCG/Private/Helpers/PCGActorHelpers.cpp
--- a/Engine/Plugins/PCG/Source/PCG/Private/Helpers/PCGActorHelpers.cpp
+++ b/Engine/Plugins/PCG/Source/PCG/Private/Helpers/PCGActorHelpers.cpp
@@ -806,6 +806,12 @@ FIntVector UPCGActorHelpers::GetCellCoord(FVector InPosition, int InGridSize, bo
 	);
 }
 
+FVector UPCGActorHelpers::GetCellCenter(const FVector& InPosition, uint32 InGridSize, bool bUse2DGrid)
+{
+	const FIntVector GridCoords = GetCellCoord(InPosition, InGridSize, bUse2DGrid);
+	return FVector(GridCoords.X + 0.5, GridCoords.Y + 0.5, GridCoords.Z + 0.5) * InGridSize;
+}
+
 int UPCGActorHelpers::ComputeHashFromActorTagsAndReference(const AActor* InActor, const bool bIncludeTags, const bool bIncludeActorReference)
 {
 	if (!InActor)
diff --git a/Engine/Plugins/PCG/Source/PCG/Private/RuntimeGen/PCGRuntimeGenScheduler.cpp b/Engine/Plugins/PCG/Source/PCG/Private/RuntimeGen/PCGRuntimeGenScheduler.cpp
--- a/Engine/Plugins/PCG/Source/PCG/Private/RuntimeGen/PCGRuntimeGenScheduler.cpp
+++ b/Engine/Plugins/PCG/Source/PCG/Private/RuntimeGen/PCGRuntimeGenScheduler.cpp
@@ -19,6 +19,9 @@
 #include "RuntimeGen/PCGGenSourceManager.h"
 
 #include "UObject/UObjectGlobals.h"
+#include "WorldPartition/WorldPartition.h"
+#include "WorldPartition/WorldPartitionRuntimeCell.h"
+#include "WorldPartition/WorldPartitionSubsystem.h"
 
 #if WITH_EDITOR
 #include "EditorViewportClient.h"
@@ -78,6 +81,11 @@ namespace PCGRuntimeGenSchedulerHelpers
 				PCGSubsystem->GetRuntimeGenScheduler()->FlushAllGeneratedActors();
 			}
 		}));
+
+	static TAutoConsoleVariable<bool> CVarCheckWorldStreamedIn(
+		TEXT("pcg.RuntimeGeneration.CheckWorldStreamingStatusForLocalComponents"),
+		false,
+		TEXT("Checks that the world is streamed in before triggering generation of local (partitioned) components."));
 }
 
 FPCGGridDescriptor FPCGRuntimeGenScheduler::FGridGenerationKey::GetGridDescriptor() const
@@ -268,6 +276,35 @@ void FPCGRuntimeGenScheduler::TickQueueComponentsForGeneration(
 		}
 	};
 
+	// Prepare streaming queries up front.
+	TArray<FWorldPartitionStreamingQuerySource> StreamingQuerySources;
+	UWorldPartitionSubsystem* WorldPartitionSubsystem = InPCGWorldActor ? UWorld::GetSubsystem<UWorldPartitionSubsystem>(InPCGWorldActor->GetWorld()) : nullptr;
+
+	const bool bCheckStreaming = PCGRuntimeGenSchedulerHelpers::CVarCheckWorldStreamedIn.GetValueOnGameThread();
+	if (bCheckStreaming)
+	{
+		FWorldPartitionStreamingQuerySource& QuerySource = StreamingQuerySources.Emplace_GetRef();
+		QuerySource.bSpatialQuery = true;
+		QuerySource.bUseGridLoadingRange = false;
+		QuerySource.bDataLayersOnly = false;
+	}
+
+	auto IsWorldStreamingComplete = [WorldPartitionSubsystem, &StreamingQuerySources](const FVector& InLocation, float InGridSize)
+	{
+		TRACE_CPUPROFILER_EVENT_SCOPE(IsWorldStreamingComplete);
+
+		StreamingQuerySources[0].Radius = InGridSize / 2.0f;
+		StreamingQuerySources[0].Location = InLocation;
+
+		const bool bIsLoaded = WorldPartitionSubsystem->IsStreamingCompleted(EWorldPartitionRuntimeCellState::Activated, StreamingQuerySources, /*bExactState*/ false);
+		if (!bIsLoaded)
+		{
+			UE_LOG(LogPCG, Verbose, TEXT("Holding back generation of cell at (%.2f, %.2f, %.2f), grid size %f, due to world not loaded."), InLocation.X, InLocation.Y, InLocation.Z, InGridSize);
+		}
+
+		return bIsLoaded;
+	};
+
 	// Collect local components from all partitioned components.
 	for (UPCGComponent* OriginalComponent : ActorAndComponentMapping->PartitionedOctree.GetAllComponents())
 	{
@@ -324,6 +361,8 @@ void FPCGRuntimeGenScheduler::TickQueueComponentsForGeneration(
 
 				if (bHasUnbounded)
 				{
+					// TODO check that unbounded is loaded, if dependent on world state. Or temporarily don't allow querying world state in unbounded, for runtime gen.
+					
 					// Ignore components that have already been generated or marked for generation. Unbounded grid size means not-partitioned.
 					FGridGenerationKey Key(PCGHiGenGrid::UnboundedGridSize(), FIntVector(0), OriginalComponent);
 					if (!GeneratedComponents.Contains(Key) && 
@@ -336,8 +375,10 @@ void FPCGRuntimeGenScheduler::TickQueueComponentsForGeneration(
 
 				// TODO: once one of the larger grid sizes is out of range, we can forego checking any smaller grid sizes. they can't possibly be closer!
 				// This assumes we enforce generation radii to increase monotonically.
-				for (const uint32 GridSize : GridSizes)
+				for (int GridIndex = 0; GridIndex < GridSizes.Num(); ++GridIndex)
 				{
+					const uint32 GridSize = GridSizes[GridIndex];
+
 					ensure(PCGHiGenGrid::IsValidGridSize(GridSize));
 
 					const FIntVector GenSourceGridPosition = UPCGActorHelpers::GetCellCoord(GenSourcePosition, GridSize, bIs2DGrid);
@@ -391,7 +432,35 @@ void FPCGRuntimeGenScheduler::TickQueueComponentsForGeneration(
 								if (LocalDistanceSquared <= GenerationRadius * GenerationRadius && 
 									Policy->ShouldGenerate(GenSource, IntersectedBounds, Key.GetGridDescriptor().Is2DGrid()))
 								{
-									AddComponentToGenerate(Key, GenSource, Policy, IntersectedBounds, Key.GetGridDescriptor().Is2DGrid());
+									bool bStreamingComplete = true;
+
+									if (bCheckStreaming)
+									{
+										if (GridIndex > 0)
+										{
+											// We use the topmost parent because we don't currently know which ancestor(s) this grid actually relies on.
+											// todo_pcg: Precompute the topmost ancestor that relies on world streaming.
+											const uint32 ParentGridSize = GridSizes[0];
+											const FVector ParentCenter = UPCGActorHelpers::GetCellCenter(Center, ParentGridSize, Key.GetGridDescriptor().Is2DGrid());
+
+											bStreamingComplete = IsWorldStreamingComplete(ParentCenter, ParentGridSize);
+
+											if (!bStreamingComplete)
+											{
+												UE_LOG(LogPCG, VeryVerbose, TEXT("Cell %u (%d, %d, %d) rejected as parent on grid size %u was not fully loaded."),
+													GridSize, GridCoords.X, GridCoords.Y, GridCoords.Z, ParentGridSize);
+											}
+										}
+										else
+										{
+											bStreamingComplete = IsWorldStreamingComplete(Center, GridSize);
+										}
+									}
+
+									if (bStreamingComplete)
+									{
+										AddComponentToGenerate(Key, GenSource, Policy, IntersectedBounds, Key.GetGridDescriptor().Is2DGrid());
+									}
 								}
 							}
 						}
diff --git a/Engine/Plugins/PCG/Source/PCG/Public/Helpers/PCGActorHelpers.h b/Engine/Plugins/PCG/Source/PCG/Public/Helpers/PCGActorHelpers.h
--- a/Engine/Plugins/PCG/Source/PCG/Public/Helpers/PCGActorHelpers.h
+++ b/Engine/Plugins/PCG/Source/PCG/Public/Helpers/PCGActorHelpers.h
@@ -230,6 +230,11 @@ PRAGMA_ENABLE_DEPRECATION_WARNINGS
 	 */
 	static FIntVector GetCellCoord(FVector InPosition, int InGridSize, bool bUse2DGrid);
 
+	/**
+	 * Return the center of the PCG partition grid cell given a position and the grid size.
+	 */
+	static FVector GetCellCenter(const FVector& InPosition, uint32 InGridSize, bool bUse2DGrid);
+
 	/** Extract the tags and the actor reference of the given actor and hash it. Useful for CRC dependencies that depends on the tags or the instance of the actor.*/
 	static int ComputeHashFromActorTagsAndReference(const AActor* InActor, const bool bIncludeTags, const bool bIncludeActorReference);
 };
