[PCG] Fix deprecation for PCG Output Attribute Selector

* Was breaking in Electric Dreams

#rnx

diff --git a/Engine/Plugins/PCG/Source/PCG/Private/Metadata/PCGAttributePropertySelector.cpp b/Engine/Plugins/PCG/Source/PCG/Private/Metadata/PCGAttributePropertySelector.cpp
--- a/Engine/Plugins/PCG/Source/PCG/Private/Metadata/PCGAttributePropertySelector.cpp
+++ b/Engine/Plugins/PCG/Source/PCG/Private/Metadata/PCGAttributePropertySelector.cpp
@@ -555,7 +555,7 @@ void FPCGAttributePropertyInputSelector::ApplyDeprecation(int32 InPCGCustomVersi
 
 bool FPCGAttributePropertyOutputNoSourceSelector::SerializeFromMismatchedTag(const struct FPropertyTag& Tag, FStructuredArchive::FSlot Slot)
 {
-	if (Tag.GetType().IsStruct(StaticStruct()->GetFName()))
+	if (Tag.GetType().IsStruct(FPCGAttributePropertySelector::StaticStruct()->GetFName()))
 	{
 		StaticStruct()->SerializeItem(Slot, this, nullptr);
 		return true;
