[PCG] Convert CollapsePoints node to PointArrayData

#rnx

diff --git a/Engine/Plugins/PCG/Source/PCG/Private/Elements/PCGCollapsePoints.cpp b/Engine/Plugins/PCG/Source/PCG/Private/Elements/PCGCollapsePoints.cpp
--- a/Engine/Plugins/PCG/Source/PCG/Private/Elements/PCGCollapsePoints.cpp
+++ b/Engine/Plugins/PCG/Source/PCG/Private/Elements/PCGCollapsePoints.cpp
@@ -3,7 +3,7 @@
 #include "Elements/PCGCollapsePoints.h"
 
 #include "PCGContext.h"
-#include "Data/PCGPointData.h"
+#include "Data/PCGBasePointData.h"
 #include "Helpers/PCGAsync.h"
 #include "Helpers/PCGHelpers.h"
 #include "Metadata/Accessors/PCGAttributeAccessorHelpers.h"
@@ -35,7 +35,7 @@ namespace PCGCollapsePoints
 				const double WeightSum = PrimaryWeight + SecondaryWeight;
 
 				double Alpha = FMath::IsNearlyZero(WeightSum) ? 0.5 : (SecondaryWeight / WeightSum);
-				const FVector DeltaPosition = Settings.GetPointPositionFunc(OutState.PointTransforms, OutState.SourceData, SecondaryPointIndex) - Settings.GetPointPositionFunc(OutState.PointTransforms, OutState.SourceData, PrimaryPointIndex);
+				const FVector DeltaPosition = Settings.GetPointPositionFunc(OutState.PointTransforms, OutState.SourcePointData, SecondaryPointIndex) - Settings.GetPointPositionFunc(OutState.PointTransforms, OutState.SourcePointData, PrimaryPointIndex);
 
 				OutState.PointTransforms[PrimaryPointIndex].AddToTranslation(Alpha * DeltaPosition);
 				OutState.Weights[PrimaryPointIndex] = WeightSum;
@@ -46,14 +46,14 @@ namespace PCGCollapsePoints
 		void RebuildOctree(FPCGContext* InContext, const FCollapsePointsSettings& Settings, FCollapsePointsState& OutState)
 		{
 			TRACE_CPUPROFILER_EVENT_SCOPE(PCGCollapsePointsElement::Algo::RebuildOctree);
-			check(OutState.SourceData && OutState.Merged.Num() == OutState.SourceData->GetNumPoints());
+			check(OutState.SourcePointData && OutState.Merged.Num() == OutState.SourcePointData->GetNumPoints());
 
-			PCGPointOctree::FPointOctree NewOctree(OutState.SourceData->GetBounds().GetCenter(), OutState.SourceData->GetBounds().GetExtent().Length());
-			for (int PointIndex = 0; PointIndex < OutState.SourceData->GetNumPoints(); ++PointIndex)
+			PCGPointOctree::FPointOctree NewOctree(OutState.SourcePointData->GetBounds().GetCenter(), OutState.SourcePointData->GetBounds().GetExtent().Length());
+			for (int PointIndex = 0; PointIndex < OutState.SourcePointData->GetNumPoints(); ++PointIndex)
 			{
 				if (OutState.Merged[PointIndex] == INDEX_NONE)
 				{
-					NewOctree.AddElement(Settings.GetPointReferenceFunc(OutState.PointTransforms, OutState.SourceData, PointIndex));
+					NewOctree.AddElement(Settings.GetPointReferenceFunc(OutState.PointTransforms, OutState.SourcePointData, PointIndex));
 				}
 			}
 
@@ -95,7 +95,7 @@ namespace PCGCollapsePoints
 				bool bHasColocatedPoint = false;
 
 				const double Extents = UE_DOUBLE_SQRT_2 * Settings.DistanceThreshold;
-				FBoxCenterAndExtent SearchBounds = Settings.GetPointSearchBoundsFunc(OutState.PointTransforms, OutState.SourceData, PointIndex, Extents);
+				FBoxCenterAndExtent SearchBounds = Settings.GetPointSearchBoundsFunc(OutState.PointTransforms, OutState.SourcePointData, PointIndex, Extents);
 				OutState.PointOctree.FindElementsWithBoundsTest(SearchBounds, [&Settings, PointIndex, &MinSqrDistance, &ClosestUnvisitedIndex, &bHasColocatedPoint, &OutState](const PCGPointOctree::FPointRef& PointRef)
 				{
 					int NeighborIndex = PointRef.Index;
@@ -105,7 +105,7 @@ namespace PCGCollapsePoints
 						return;
 					}
 
-					double SqrDistance = (Settings.GetPointPositionFunc(OutState.PointTransforms, OutState.SourceData, PointIndex) - Settings.GetPointPositionFunc(OutState.PointTransforms, OutState.SourceData, NeighborIndex)).SquaredLength();
+					double SqrDistance = (Settings.GetPointPositionFunc(OutState.PointTransforms, OutState.SourcePointData, PointIndex) - Settings.GetPointPositionFunc(OutState.PointTransforms, OutState.SourcePointData, NeighborIndex)).SquaredLength();
 
 					if (FMath::IsNearlyZero(SqrDistance))
 					{
@@ -154,36 +154,36 @@ namespace PCGCollapsePoints
 
 	namespace ComparisonModes
 	{
-		FVector GetPosition(const TArray<FTransform>& InTransforms, const UPCGPointData* InPointData, int32 InIndex)
+		FVector GetPosition(const TArray<FTransform>& InTransforms, const UPCGBasePointData* InPointData, int32 InIndex)
 		{
 			return InTransforms[InIndex].GetLocation();
 		}
 
-		FVector GetCenter(const TArray<FTransform>& InTransforms, const UPCGPointData* InPointData, int32 InIndex)
+		FVector GetCenter(const TArray<FTransform>& InTransforms, const UPCGBasePointData* InPointData, int32 InIndex)
 		{
 			return InTransforms[InIndex].TransformPosition(InPointData->GetLocalCenter(InIndex));
 		}
 
-		PCGPointOctree::FPointRef GetPositionPointRef(const TArray<FTransform>& InTransforms, const UPCGPointData* InPointData, int32 InIndex)
+		PCGPointOctree::FPointRef GetPositionPointRef(const TArray<FTransform>& InTransforms, const UPCGBasePointData* InPointData, int32 InIndex)
 		{
 			check(InPointData);
 			return PCGPointOctree::FPointRef(InIndex, FBox(FVector::ZeroVector, FVector::ZeroVector).TransformBy(InTransforms[InIndex]));
 		}
 
-		PCGPointOctree::FPointRef GetCenterPointRef(const TArray<FTransform>& InTransforms, const UPCGPointData* InPointData, int32 InIndex)
+		PCGPointOctree::FPointRef GetCenterPointRef(const TArray<FTransform>& InTransforms, const UPCGBasePointData* InPointData, int32 InIndex)
 		{
 			check(InPointData);
 			const FVector LocalCenter = InPointData->GetLocalCenter(InIndex);
 			return PCGPointOctree::FPointRef(InIndex, FBox(LocalCenter, LocalCenter).TransformBy(InTransforms[InIndex]));
 		}
 
-		FBoxCenterAndExtent GetPositionSearchBounds(const TArray<FTransform>& InTransforms, const UPCGPointData* InPointData, int32 InIndex, const double& Extents)
+		FBoxCenterAndExtent GetPositionSearchBounds(const TArray<FTransform>& InTransforms, const UPCGBasePointData* InPointData, int32 InIndex, const double& Extents)
 		{
 			check(InPointData);
 			return FBoxCenterAndExtent(InTransforms[InIndex].GetLocation(), FVector(Extents, Extents, Extents));
 		}
 
-		FBoxCenterAndExtent GetCenterSearchBounds(const TArray<FTransform>& InTransforms, const UPCGPointData* InPointData, int32 InIndex, const double& Extents)
+		FBoxCenterAndExtent GetCenterSearchBounds(const TArray<FTransform>& InTransforms, const UPCGBasePointData* InPointData, int32 InIndex, const double& Extents)
 		{
 			check(InPointData);
 			return FBoxCenterAndExtent(InTransforms[InIndex].TransformPosition(InPointData->GetLocalCenter(InIndex)), FVector(Extents, Extents, Extents));
@@ -293,14 +293,17 @@ bool FPCGCollapsePointsElement::PrepareDataInternal(FPCGContext* InContext) cons
 
 	Context->InitializePerIterationStates(Inputs.Num(), [&Inputs, Settings, Context](IterStateType& OutState, const ExecStateType& ExecState, const uint32 IterationIndex)
 	{
-		const UPCGPointData* PointData = Cast<UPCGPointData>(Inputs[IterationIndex].Data);
+		const UPCGBasePointData* PointData = Cast<UPCGBasePointData>(Inputs[IterationIndex].Data);
 
 		if (!PointData || PointData->IsEmpty())
 		{
 			return EPCGTimeSliceInitResult::NoOperation;
 		}
 
-		OutState.SourceData = PointData;
+		OutState.SourcePointData = PointData;
+PRAGMA_DISABLE_DEPRECATION_WARNINGS
+		OutState.SourceData = Cast<UPCGPointData>(PointData);
+PRAGMA_ENABLE_DEPRECATION_WARNINGS
 
 		// Get points, octree will be rebuilt later
 		OutState.PointTransforms = PointData->GetTransformsCopy();
@@ -393,9 +396,16 @@ bool FPCGCollapsePointsElement::PrepareDataInternal(FPCGContext* InContext) cons
 		}
 
 		// Create the output point data now. Operate on a copy of the input point data.
-		UPCGPointData* OutPointData = FPCGContext::NewObject_AnyThread<UPCGPointData>(Context);
-		OutPointData->InitializeFromData(PointData);
-		OutState.OutData = OutPointData;
+		UPCGBasePointData* OutPointData = FPCGContext::NewPointData_AnyThread(Context);
+
+		FPCGInitializeFromDataParams InitializeFromDataParams(PointData);
+		InitializeFromDataParams.bInheritSpatialData = false;
+		OutPointData->InitializeFromDataWithParams(InitializeFromDataParams);
+
+		OutState.OutputPointData = OutPointData;
+PRAGMA_DISABLE_DEPRECATION_WARNINGS
+		OutState.OutData = Cast<UPCGPointData>(OutPointData);
+PRAGMA_ENABLE_DEPRECATION_WARNINGS
 
 		// Initialize merge accessors
 		for(const FPCGAttributePropertyOutputNoSourceSelector& InAttributeToMerge : Settings->AttributesToMerge)
@@ -415,15 +425,8 @@ bool FPCGCollapsePointsElement::PrepareDataInternal(FPCGContext* InContext) cons
 			FPCGAttributePropertyOutputSelector OutputAttributeSelector;
 			OutputAttributeSelector.ImportFromOtherSelector(AttributeToMerge); 
 
-			TUniquePtr<IPCGAttributeAccessor> OutputAttributeAccessor = PCGAttributeAccessorHelpers::CreateAccessor(OutPointData, OutputAttributeSelector);
-			if (!OutputAttributeAccessor)
-			{
-				PCGLog::Metadata::LogFailToCreateAccessorError(OutputAttributeSelector, Context);
-				continue;
-			}
-
 			OutState.SourceMergeAccessors.Add(std::move(AttributeAccessor));
-			OutState.OutputMergeAccessors.Add(std::move(OutputAttributeAccessor));
+			OutState.OutputMergeSelectors.Add(std::move(OutputAttributeSelector));
 		}
 
 		OutState.SourceMergeKeys = PCGAttributeAccessorHelpers::CreateConstKeys(PointData, FPCGAttributePropertyInputSelector());
@@ -516,21 +519,36 @@ bool FPCGCollapsePointsElement::ExecuteInternal(FPCGContext* InContext) const
 			}
 
 			// Create one point per partition
-			check(CollapseState.SourceData);
-			const TArray<FPCGPoint>& SourcePoints = CollapseState.SourceData->GetPoints();
-			check(CollapseState.OutData);
-			TArray<FPCGPoint>& OutPoints = CollapseState.OutData->GetMutablePoints();
-
+			check(CollapseState.SourcePointData);
+			check(CollapseState.OutputPointData);
+			
 			// Need to allocate prior to key creation
-			OutPoints.SetNum(Partition.Num());
+			CollapseState.OutputPointData->SetNumPoints(Partition.Num(), /*bInitializeValues=*/false);
+			CollapseState.OutputPointData->AllocateProperties(CollapseState.SourcePointData->GetAllocatedProperties());
 
 			// Build accessor list(s) + keys, based on settings
 			TArray<double> PartitionWeights;
-			TUniquePtr<IPCGAttributeAccessorKeys> OutputKeys = PCGAttributeAccessorHelpers::CreateKeys(CollapseState.OutData, FPCGAttributePropertyOutputSelector());
+			TUniquePtr<IPCGAttributeAccessorKeys> OutputKeys = PCGAttributeAccessorHelpers::CreateKeys(CollapseState.OutputPointData, FPCGAttributePropertyOutputSelector());
 			check(OutputKeys);
 			TUniquePtr<const IPCGAttributeAccessorKeys>& SourceKeys = CollapseState.SourceMergeKeys;
 			check(SourceKeys);
 			
+			FPCGPointValueRanges OutRanges(CollapseState.OutputPointData, /*bAllocate=*/false);
+			const FConstPCGPointValueRanges InRanges(CollapseState.SourcePointData);
+
+			// Create Accessors now that data is allocated
+			check(CollapseState.OutputMergeAccessors.IsEmpty());
+			for (int AccessorIndex = 0; AccessorIndex < CollapseState.SourceMergeAccessors.Num(); ++AccessorIndex)
+			{
+				TUniquePtr<IPCGAttributeAccessor> OutputAccessor = PCGAttributeAccessorHelpers::CreateAccessor(CollapseState.OutputPointData, CollapseState.OutputMergeSelectors[AccessorIndex]);
+				if (!OutputAccessor)
+				{
+					PCGLog::Metadata::LogFailToCreateAccessorError(CollapseState.OutputMergeSelectors[AccessorIndex], Context);
+					return true;
+				}
+				CollapseState.OutputMergeAccessors.Add(std::move(OutputAccessor));
+			}
+
 			// TODO: do this in parallel
 			int PartitionIndex = 0;
 			for(const TPair<int, TArray<int, TInlineAllocator<8>>>& PartitionEntry : Partition)
@@ -539,8 +557,7 @@ bool FPCGCollapsePointsElement::ExecuteInternal(FPCGContext* InContext) const
 				const TArray<int, TInlineAllocator<8>>& SecondaryPointIndices = PartitionEntry.Value;
 
 				// First, copy the primary point as-is.
-				FPCGPoint& MergedPoint = OutPoints[PartitionIndex];
-				MergedPoint = SourcePoints[PrimaryPointIndex];
+				OutRanges.SetFromValueRanges(PartitionIndex, InRanges, PrimaryPointIndex);
 
 				if (SecondaryPointIndices.Num() == 1)
 				{
diff --git a/Engine/Plugins/PCG/Source/PCG/Public/Elements/PCGCollapsePoints.h b/Engine/Plugins/PCG/Source/PCG/Public/Elements/PCGCollapsePoints.h
--- a/Engine/Plugins/PCG/Source/PCG/Public/Elements/PCGCollapsePoints.h
+++ b/Engine/Plugins/PCG/Source/PCG/Public/Elements/PCGCollapsePoints.h
@@ -12,8 +12,6 @@
 
 #include "PCGCollapsePoints.generated.h"
 
-class UPCGPointData;
-
 UENUM(BlueprintType)
 enum class EPCGCollapseMode : uint8
 {
@@ -112,13 +110,20 @@ PRAGMA_DISABLE_DEPRECATION_WARNINGS
 		TArray<int> VisitOrder;
 		TArray<double> OriginalWeights;
 
+		const UPCGBasePointData* SourcePointData = nullptr;
+		UPCGBasePointData* OutputPointData = nullptr;
+
 		// Output data
+		UE_DEPRECATED(5.6, "Use SourcePointData instead")
 		const UPCGPointData* SourceData = nullptr;
+
+		UE_DEPRECATED(5.6, "Use OutputPointData instead")
 		UPCGPointData* OutData = nullptr;
 
 		TArray<TUniquePtr<const IPCGAttributeAccessor>> SourceMergeAccessors;
 		TUniquePtr<const IPCGAttributeAccessorKeys> SourceMergeKeys;
 		TArray<TUniquePtr<IPCGAttributeAccessor>> OutputMergeAccessors;
+		TArray<FPCGAttributePropertyOutputSelector> OutputMergeSelectors;
 
 		UE_DEPRECATED(5.6, "Use PointTransforms instead")
 		TArray<FPCGPoint> Points;
@@ -137,13 +142,13 @@ PRAGMA_DISABLE_DEPRECATION_WARNINGS
 		using MergeSelectionFuncType = void(*)(FPCGContext* /*InContext*/, const FCollapsePointsSettings& /*Settings*/, FCollapsePointsState& /*OutState*/);
 		MergeSelectionFuncType MergeSelectionFunc = nullptr;
 
-		using GetPointPositionFuncType = FVector(*)(const TArray<FTransform>& /*InTransforms*/,const UPCGPointData* /*InPointData*/, int32 /*InPointIndex*/);
+		using GetPointPositionFuncType = FVector(*)(const TArray<FTransform>& /*InTransforms*/,const UPCGBasePointData* /*InPointData*/, int32 /*InPointIndex*/);
 		GetPointPositionFuncType GetPointPositionFunc = nullptr;
 
-		using GetPointReferenceFuncType = PCGPointOctree::FPointRef(*)(const TArray<FTransform>& /*InTransforms*/, const UPCGPointData* /*InPointData*/, int32 /*InPointIndex*/);
+		using GetPointReferenceFuncType = PCGPointOctree::FPointRef(*)(const TArray<FTransform>& /*InTransforms*/, const UPCGBasePointData* /*InPointData*/, int32 /*InPointIndex*/);
 		GetPointReferenceFuncType GetPointReferenceFunc = nullptr;
 
-		using GetPointSearchBoundsFuncType = FBoxCenterAndExtent(*)(const TArray<FTransform>& /*InTransforms*/, const UPCGPointData* /*InPointData*/, int32 /*InPointIndex*/, const double& /*InExtents*/);
+		using GetPointSearchBoundsFuncType = FBoxCenterAndExtent(*)(const TArray<FTransform>& /*InTransforms*/, const UPCGBasePointData* /*InPointData*/, int32 /*InPointIndex*/, const double& /*InExtents*/);
 		GetPointSearchBoundsFuncType GetPointSearchBoundsFunc = nullptr;
 
 		double DistanceThreshold = 0.0;
