[PCG] Make sure we always register primitives in the binding to signal that the SM DI has run, so that Instance DI knows whether it should wait or not.

Fixes knock on issue from previous commit that refactored data interface preparation.

#rnx

diff --git a/Engine/Plugins/PCG/Source/PCG/Private/Compute/DataInterfaces/Elements/PCGStaticMeshSpawnerDataInterface.cpp b/Engine/Plugins/PCG/Source/PCG/Private/Compute/DataInterfaces/Elements/PCGStaticMeshSpawnerDataInterface.cpp
--- a/Engine/Plugins/PCG/Source/PCG/Private/Compute/DataInterfaces/Elements/PCGStaticMeshSpawnerDataInterface.cpp
+++ b/Engine/Plugins/PCG/Source/PCG/Private/Compute/DataInterfaces/Elements/PCGStaticMeshSpawnerDataInterface.cpp
@@ -294,6 +294,12 @@ bool UPCGStaticMeshSpawnerDataProvider::PrepareForExecute_GameThread(UPCGDataBin
 		return false;
 	}
 
+	if (!bRegisteredPrimitives)
+	{
+		InBinding->MeshSpawnersToPrimitives.Add(ProducerKernel, {});
+		bRegisteredPrimitives = true;
+	}
+
 	TSharedPtr<FPCGContextHandle> ContextHandle = InBinding->ContextHandle.Pin();
 	FPCGContext* Context = ContextHandle ? ContextHandle->GetContext() : nullptr;
 	if (!ensure(Context))
@@ -388,8 +394,6 @@ bool UPCGStaticMeshSpawnerDataProvider::PrepareForExecute_GameThread(UPCGDataBin
 		}
 		else
 		{
-			InBinding->MeshSpawnersToPrimitives.Add(ProducerKernel, {});
-
 			// No component set up means we have no more work to do.
 			return true;
 		}
@@ -439,6 +443,7 @@ void UPCGStaticMeshSpawnerDataProvider::Reset()
 	CustomFloatCount = 0;
 	bPrimitivesSetUp = false;
 	NumPrimitivesSetup = 0;
+	bRegisteredPrimitives = false;
 
 	Super::Reset();
 }
diff --git a/Engine/Plugins/PCG/Source/PCG/Private/Compute/DataInterfaces/Elements/PCGStaticMeshSpawnerDataInterface.h b/Engine/Plugins/PCG/Source/PCG/Private/Compute/DataInterfaces/Elements/PCGStaticMeshSpawnerDataInterface.h
--- a/Engine/Plugins/PCG/Source/PCG/Private/Compute/DataInterfaces/Elements/PCGStaticMeshSpawnerDataInterface.h
+++ b/Engine/Plugins/PCG/Source/PCG/Private/Compute/DataInterfaces/Elements/PCGStaticMeshSpawnerDataInterface.h
@@ -91,6 +91,8 @@ public:
 
 	uint32 CustomFloatCount = 0;
 
+	bool bRegisteredPrimitives = false;
+
 	bool bPrimitivesSetUp = false;
 
 	/** Number of primitives created during primitive setup. */
