Determinism fix supplied by partners at The Coalition:
Fixing PCGGraph diff in IterativeValidate.
Summary: With a multi-platform cook used to export snapshots, PCGGraph should create the per-platform sub-object (CookedCompilationData) with a deterministic name. Addressing that by resetting and reusing the instance if it exists.

diff --git a/Engine/Plugins/PCG/Source/PCG/Private/PCGGraph.cpp b/Engine/Plugins/PCG/Source/PCG/Private/PCGGraph.cpp
--- a/Engine/Plugins/PCG/Source/PCG/Private/PCGGraph.cpp
+++ b/Engine/Plugins/PCG/Source/PCG/Private/PCGGraph.cpp
@@ -676,7 +676,18 @@ void UPCGGraph::PreSave(FObjectPreSaveContext ObjectSaveContext)
 		TMap<uint32, FPCGStackContext>* CompiledStackContexts = Cache.TopGraphToStackContextMap.Find(this);
 		TMap<uint32, TArray<TObjectPtr<UPCGComputeGraph>>>* CompiledComputeGraphs = Cache.TopGraphToComputeGraphMap.Find(this);
 
-		CookedCompilationData = NewObject<UPCGGraphCompilationData>(this);
+		if (CookedCompilationData != nullptr)
+		{
+			// Reset and reuse the instance for determinism during multi-platform cooks,
+			// to avoid the sub-object name being different for each platform.
+			CookedCompilationData->Tasks.Reset();
+			CookedCompilationData->StackContexts.Reset();
+			CookedCompilationData->ComputeGraphs.Reset();
+		}
+		else
+		{
+			CookedCompilationData = NewObject<UPCGGraphCompilationData>(this);
+		}
 
 		if (ensure(CompiledTasks))
 		{
