[PCG] PCGExecuteBlueprint Element will convert UPCGPointData output to UPCGPointArrayData when cvar is enabled
- Eventually PCG BP Apis will support writing/reading from UPCGPointArrayData but until then node will not inject UPCGPointData into the rest of the graph

#rnx

diff --git a/Engine/Plugins/PCG/Source/PCG/Private/Elements/PCGExecuteBlueprint.cpp b/Engine/Plugins/PCG/Source/PCG/Private/Elements/PCGExecuteBlueprint.cpp
--- a/Engine/Plugins/PCG/Source/PCG/Private/Elements/PCGExecuteBlueprint.cpp
+++ b/Engine/Plugins/PCG/Source/PCG/Private/Elements/PCGExecuteBlueprint.cpp
@@ -6,6 +6,7 @@
 #include "Engine/Blueprint.h"
 #include "Math/RandomStream.h"
 #include "PCGComponent.h"
+#include "Data/PCGPointArrayData.h"
 #include "Data/PCGPointData.h"
 #include "Helpers/PCGAsync.h"
 #include "Helpers/PCGHelpers.h"
@@ -860,10 +861,11 @@ void FPCGExecuteBlueprintElement::PostExecuteInternal(FPCGContext* InContext) co
 		// Log info on outputs
 		for (int32 OutputIndex = 0; OutputIndex < Context->OutputData.TaggedData.Num(); ++OutputIndex)
 		{
-			const FPCGTaggedData& Output = Context->OutputData.TaggedData[OutputIndex];
-			if (const UPCGPointData* PointData = Cast<UPCGPointData>(Output.Data))
+			FPCGTaggedData& Output = Context->OutputData.TaggedData[OutputIndex];
+
+			if (const UPCGBasePointData* PointData = Cast<UPCGBasePointData>(Output.Data))
 			{
-				PCGE_LOG(Verbose, LogOnly, FText::Format(LOCTEXT("OutputPointInfo", "Output {0} has {1} points"), OutputIndex, PointData->GetPoints().Num()));
+				PCGE_LOG(Verbose, LogOnly, FText::Format(LOCTEXT("OutputPointInfo", "Output {0} has {1} points"), OutputIndex, PointData->GetNumPoints()));
 			}
 
 			// Important implementation note:
@@ -918,6 +920,15 @@ void FPCGExecuteBlueprintElement::PostExecuteInternal(FPCGContext* InContext) co
 						Metadata = Metadata->GetParent();
 					}
 				}
+
+				// @todo_pcg: expose PointArrayData to the Blueprint API but for now if we have PointArrayData enabled convert any PointData output to PointArrayData
+				if (CVarPCGEnablePointArrayData.GetValueOnAnyThread())
+				{
+					if (const UPCGPointData* PointData = Cast<UPCGPointData>(Output.Data))
+					{
+						Output.Data = PointData->ToPointArrayData(Context);
+					}
+				}
 			}
 		}
 	}
