[PCG] Minor code clean up, added consts.

#rnx

diff --git a/Engine/Plugins/PCG/Source/PCG/Private/Elements/PCGDifferenceElement.cpp b/Engine/Plugins/PCG/Source/PCG/Private/Elements/PCGDifferenceElement.cpp
--- a/Engine/Plugins/PCG/Source/PCG/Private/Elements/PCGDifferenceElement.cpp
+++ b/Engine/Plugins/PCG/Source/PCG/Private/Elements/PCGDifferenceElement.cpp
@@ -177,7 +177,7 @@ bool FPCGDifferenceElement::ExecuteInternal(FPCGContext* Context) const
 	}
 
 	// Iterate over the sources and subtract the differences
-	for (FPCGTaggedData& Source : Sources)
+	for (const FPCGTaggedData& Source : Sources)
 	{
 		const UPCGSpatialData* SourceSpatialData = Cast<UPCGSpatialData>(Source.Data);
 
@@ -193,7 +193,7 @@ bool FPCGDifferenceElement::ExecuteInternal(FPCGContext* Context) const
 		UPCGDifferenceData* DifferenceData = FPCGContext::NewObject_AnyThread<UPCGDifferenceData>(Context);
 		DifferenceData->Initialize(SourceSpatialData);
 		
-		for (FPCGTaggedData& Difference : Differences)
+		for (const FPCGTaggedData& Difference : Differences)
 		{
 			if (const UPCGSpatialData* DifferenceSpatialData = Cast<const UPCGSpatialData>(Difference.Data))
 			{
diff --git a/Engine/Plugins/PCG/Source/PCG/Private/Elements/PCGInnerIntersectionElement.cpp b/Engine/Plugins/PCG/Source/PCG/Private/Elements/PCGInnerIntersectionElement.cpp
--- a/Engine/Plugins/PCG/Source/PCG/Private/Elements/PCGInnerIntersectionElement.cpp
+++ b/Engine/Plugins/PCG/Source/PCG/Private/Elements/PCGInnerIntersectionElement.cpp
@@ -64,7 +64,7 @@ bool FPCGInnerIntersectionElement::ExecuteInternal(FPCGContext* Context) const
 	UPCGIntersectionData* IntersectionData = nullptr;
 	int32 IntersectionTaggedDataIndex = -1;
 
-	for (FPCGTaggedData& Input : Inputs)
+	for (const FPCGTaggedData& Input : Inputs)
 	{
 		const UPCGSpatialData* SpatialData = Cast<UPCGSpatialData>(Input.Data);
 
diff --git a/Engine/Plugins/PCG/Source/PCG/Private/Elements/PCGProjectionElement.cpp b/Engine/Plugins/PCG/Source/PCG/Private/Elements/PCGProjectionElement.cpp
--- a/Engine/Plugins/PCG/Source/PCG/Private/Elements/PCGProjectionElement.cpp
+++ b/Engine/Plugins/PCG/Source/PCG/Private/Elements/PCGProjectionElement.cpp
@@ -88,7 +88,7 @@ bool FPCGProjectionElement::ExecuteInternal(FPCGContext* Context) const
 
 	const bool bKeepZeroDensityPoints = Settings->bKeepZeroDensityPoints;
 
-	for (FPCGTaggedData& Source : Sources)
+	for (const FPCGTaggedData& Source : Sources)
 	{
 		const UPCGSpatialData* ProjectionSource = Cast<UPCGSpatialData>(Source.Data);
 
diff --git a/Engine/Plugins/PCG/Source/PCG/Private/Elements/PCGVolumeSampler.cpp b/Engine/Plugins/PCG/Source/PCG/Private/Elements/PCGVolumeSampler.cpp
--- a/Engine/Plugins/PCG/Source/PCG/Private/Elements/PCGVolumeSampler.cpp
+++ b/Engine/Plugins/PCG/Source/PCG/Private/Elements/PCGVolumeSampler.cpp
@@ -398,7 +398,7 @@ bool FPCGVolumeSamplerElement::ExecuteInternal(FPCGContext* Context) const
 	// The execution would have resulted in an empty set of points for all iterations
 	if (TimeSlicedContext->GetExecutionStateResult() == EPCGTimeSliceInitResult::NoOperation)
 	{
-		for (FPCGTaggedData& Input : TimeSlicedContext->InputData.GetInputsByPin(PCGVolumeSamplerConstants::VolumeLabel))
+		for (const FPCGTaggedData& Input : TimeSlicedContext->InputData.GetInputsByPin(PCGVolumeSamplerConstants::VolumeLabel))
 		{
 			// TODO: Empty point data (to preserve previous behavior). Eventually, should be replaced with no output at all
 			FPCGTaggedData& Output = TimeSlicedContext->OutputData.TaggedData.Add_GetRef(Input);
diff --git a/Engine/Plugins/PCG/Source/PCG/Private/Tests/Determinism/PCGDeterminismTestsCommon.cpp b/Engine/Plugins/PCG/Source/PCG/Private/Tests/Determinism/PCGDeterminismTestsCommon.cpp
--- a/Engine/Plugins/PCG/Source/PCG/Private/Tests/Determinism/PCGDeterminismTestsCommon.cpp
+++ b/Engine/Plugins/PCG/Source/PCG/Private/Tests/Determinism/PCGDeterminismTestsCommon.cpp
@@ -1448,7 +1448,7 @@ namespace PCGDeterminismTests
 
 	void ShuffleAllInternalData(PCGTestsCommon::FTestData& TestData)
 	{
-		for (FPCGTaggedData& TaggedData : TestData.InputData.TaggedData)
+		for (const FPCGTaggedData& TaggedData : TestData.InputData.TaggedData)
 		{
 			if (DataCanBeShuffled(TaggedData.Data))
 			{
diff --git a/Engine/Plugins/PCGInterops/PCGExternalDataInterop/Source/PCGExternalDataInteropEditor/Private/PCGLoadAlembic.cpp b/Engine/Plugins/PCGInterops/PCGExternalDataInterop/Source/PCGExternalDataInteropEditor/Private/PCGLoadAlembic.cpp
--- a/Engine/Plugins/PCGInterops/PCGExternalDataInterop/Source/PCGExternalDataInteropEditor/Private/PCGLoadAlembic.cpp
+++ b/Engine/Plugins/PCGInterops/PCGExternalDataInterop/Source/PCGExternalDataInteropEditor/Private/PCGLoadAlembic.cpp
@@ -82,7 +82,7 @@ void UPCGLoadAlembicFunctionLibrary::LoadAlembicFileToPCGInternal(const FPCGLoad
 	// Reouter the data if needed
 	if (TargetOuter)
 	{
-		for (FPCGTaggedData& TaggedData : Data.TaggedData)
+		for (const FPCGTaggedData& TaggedData : Data.TaggedData)
 		{
 			if (TaggedData.Data)
 			{
