[PCG] Fix missing instance issues by communicating when static mesh spawn DI setup is complete.

#rnx

diff --git a/Engine/Plugins/PCG/Source/PCG/Private/Compute/DataInterfaces/Elements/PCGStaticMeshSpawnerDataInterface.cpp b/Engine/Plugins/PCG/Source/PCG/Private/Compute/DataInterfaces/Elements/PCGStaticMeshSpawnerDataInterface.cpp
--- a/Engine/Plugins/PCG/Source/PCG/Private/Compute/DataInterfaces/Elements/PCGStaticMeshSpawnerDataInterface.cpp
+++ b/Engine/Plugins/PCG/Source/PCG/Private/Compute/DataInterfaces/Elements/PCGStaticMeshSpawnerDataInterface.cpp
@@ -304,12 +304,14 @@ bool UPCGStaticMeshSpawnerDataProvider::PrepareForExecute_GameThread(UPCGDataBin
 	FPCGContext* Context = ContextHandle ? ContextHandle->GetContext() : nullptr;
 	if (!ensure(Context))
 	{
+		InBinding->CompletedMeshSpawners.Add(ProducerKernel);
 		return true;
 	}
 
 	UPCGComponent* SourceComponent = Cast<UPCGComponent>(Context->ExecutionSource.Get());
 	if (!ensure(SourceComponent))
 	{
+		InBinding->CompletedMeshSpawners.Add(ProducerKernel);
 		return true;
 	}
 
@@ -321,6 +323,7 @@ bool UPCGStaticMeshSpawnerDataProvider::PrepareForExecute_GameThread(UPCGDataBin
 
 		if (!ensure(InputDataDesc))
 		{
+			InBinding->CompletedMeshSpawners.Add(ProducerKernel);
 			return true;
 		}
 
@@ -355,6 +358,7 @@ bool UPCGStaticMeshSpawnerDataProvider::PrepareForExecute_GameThread(UPCGDataBin
 					FText::FromName(SelectorName)));
 			}
 
+			InBinding->CompletedMeshSpawners.Add(ProducerKernel);
 			return true;
 		}
 
@@ -395,6 +399,7 @@ bool UPCGStaticMeshSpawnerDataProvider::PrepareForExecute_GameThread(UPCGDataBin
 		else
 		{
 			// No component set up means we have no more work to do.
+			InBinding->CompletedMeshSpawners.Add(ProducerKernel);
 			return true;
 		}
 	}
@@ -418,6 +423,7 @@ bool UPCGStaticMeshSpawnerDataProvider::PrepareForExecute_GameThread(UPCGDataBin
 	SelectedMeshAttributeId = InBinding->GetAttributeId(Settings->OutAttributeName, EPCGKernelAttributeType::StringKey);
 	ensure(SelectedMeshAttributeId != INDEX_NONE);
 
+	InBinding->CompletedMeshSpawners.Add(ProducerKernel);
 	return true;
 }
 
diff --git a/Engine/Plugins/PCG/Source/PCG/Private/Compute/DataInterfaces/PCGInstanceDataInterface.cpp b/Engine/Plugins/PCG/Source/PCG/Private/Compute/DataInterfaces/PCGInstanceDataInterface.cpp
--- a/Engine/Plugins/PCG/Source/PCG/Private/Compute/DataInterfaces/PCGInstanceDataInterface.cpp
+++ b/Engine/Plugins/PCG/Source/PCG/Private/Compute/DataInterfaces/PCGInstanceDataInterface.cpp
@@ -150,12 +150,17 @@ bool UPCGInstanceDataProvider::PrepareForExecute_GameThread(UPCGDataBinding* InB
 	TRACE_CPUPROFILER_EVENT_SCOPE(UPCGInstanceDataProvider::PrepareForExecute_GameThread);
 	UPCGDataBinding* Binding = CastChecked<UPCGDataBinding>(InBinding);
 
+	if (!Binding->CompletedMeshSpawners.Contains(ProducerKernel))
+	{
+		// The static mesh data interface(s) set this up, so wait until it is ready.
+		return false;
+	}
+
 	FPCGSpawnerPrimitives* FoundPrimitives = Binding->MeshSpawnersToPrimitives.Find(Cast<UPCGStaticMeshSpawnerKernel>(ProducerKernel));
 
 	if (!FoundPrimitives)
 	{
-		// The static mesh data interface(s) set this up, so wait until it is ready.
-		return false;
+		return true;
 	}
 
 	if (!ensure(FoundPrimitives->Primitives.Num() <= PCGComputeConstants::MAX_PRIMITIVE_COMPONENTS_PER_SPAWNER))
@@ -299,7 +304,7 @@ void FPCGInstanceDataProviderProxy::PostSubmit(FRDGBuilder& InGraphBuilder) cons
 		}
 
 #if PCG_INSTANCE_DATA_LOGGING
-		UE_LOG(LogPCG, Log, TEXT("\tProcess primitive [%d]:\tPrimitive ID %u\tProxy: %llu"), PrimitiveIndex, SceneInfo->PrimitiveComponentId.PrimIDValue, (uint64)this);
+		UE_LOG(LogPCG, Log, TEXT("\tProcess primitive [%d]:\tPrimitive ID %u\tProxy: %llu"), PrimitiveIndex, SceneInfo->GetPersistentIndex().Index, (uint64)this);
 #endif
 
 		const int32 NumInstancesAllocatedInGPUScene = PrimitiveSceneProxy->GetInstanceDataHeader().NumInstances;
diff --git a/Engine/Plugins/PCG/Source/PCG/Private/Compute/PCGDataBinding.h b/Engine/Plugins/PCG/Source/PCG/Private/Compute/PCGDataBinding.h
--- a/Engine/Plugins/PCG/Source/PCG/Private/Compute/PCGDataBinding.h
+++ b/Engine/Plugins/PCG/Source/PCG/Private/Compute/PCGDataBinding.h
@@ -133,6 +133,8 @@ public:
 	UPROPERTY()
 	TMap<TObjectPtr<const UPCGStaticMeshSpawnerKernel>, FPCGSpawnerPrimitives> MeshSpawnersToPrimitives;
 
+	TArray<TObjectPtr<const UPCGStaticMeshSpawnerKernel>> CompletedMeshSpawners;
+
 	TArray<FPCGDataToDebug> DataToDebug;
 	TArray<FPCGDataToDebug> DataToInspect;
 
