[PCG] Make sure duplicate texture is compiled before continuing with initialization

Avoids crash when initializing the duplicate texture.

Fast tracking in this review to unblock MW.

#rnx

diff --git a/Engine/Plugins/PCG/Source/PCG/Private/Data/PCGTextureData.cpp b/Engine/Plugins/PCG/Source/PCG/Private/Data/PCGTextureData.cpp
--- a/Engine/Plugins/PCG/Source/PCG/Private/Data/PCGTextureData.cpp
+++ b/Engine/Plugins/PCG/Source/PCG/Private/Data/PCGTextureData.cpp
@@ -556,11 +556,24 @@ void UPCGTextureData::InitializeInternal(UTexture* InTexture, uint32 InTextureIn
 			// Duplicate texture and change access flags (editor only). This duplicate texture will be used by the normal logic below.
 			FObjectDuplicationParameters DuplicationParams(Texture2D, /*Outer=*/this);
 			DuplicateTexture = CastChecked<UTexture2D>(StaticDuplicateObjectEx(DuplicationParams));
+		}
+
+		if (DuplicateTexture && !bDuplicateTextureInitialized)
+		{
+			// Wait until texture compilation is complete on the duplicated texture, otherwise we can crash in PreEditChange().
+			if (!DuplicateTexture->IsAsyncCacheComplete())
+			{
+				SetInitCompleted(false);
+				return;
+			}
+
 			DuplicateTexture->PreEditChange(nullptr);
 			DuplicateTexture->MipGenSettings = TextureMipGenSettings::TMGS_NoMipmaps;
 			DuplicateTexture->CompressionSettings = TextureCompressionSettings::TC_VectorDisplacementmap; // Allows the texture to be in a non-compressed format (B8G8R8A8), which is necessary to convince the data to remain CPU-side.
 			DuplicateTexture->SRGB = false;
 			DuplicateTexture->PostEditChange();
+
+			bDuplicateTextureInitialized = true;
 		}
 	}
 #endif
diff --git a/Engine/Plugins/PCG/Source/PCG/Public/Data/PCGTextureData.h b/Engine/Plugins/PCG/Source/PCG/Public/Data/PCGTextureData.h
--- a/Engine/Plugins/PCG/Source/PCG/Public/Data/PCGTextureData.h
+++ b/Engine/Plugins/PCG/Source/PCG/Public/Data/PCGTextureData.h
@@ -232,6 +232,9 @@ public:
 	/** Transient CPU visible duplicate of Texture created and used only when initialized with bCreateCPUDuplicateEditorOnly. */
 	UPROPERTY(Transient)
 	TObjectPtr<UTexture2D> DuplicateTexture = nullptr;
+
+	UPROPERTY(Transient)
+	bool bDuplicateTextureInitialized = false;
 #endif
 
 	UPROPERTY(BlueprintReadOnly, VisibleAnywhere, Category = Properties)
