[Backout] - CL41205612
#rnx
Original CL Desc
-----------------------------------------------------------------
Revert async load feature, which fails in an odd way (some instances appear missing when pcg budget is set low like 1ms).

#rnx

diff --git a/Engine/Plugins/PCG/Source/PCG/Private/Components/PCGProceduralISMComponent.cpp b/Engine/Plugins/PCG/Source/PCG/Private/Components/PCGProceduralISMComponent.cpp
--- a/Engine/Plugins/PCG/Source/PCG/Private/Components/PCGProceduralISMComponent.cpp
+++ b/Engine/Plugins/PCG/Source/PCG/Private/Components/PCGProceduralISMComponent.cpp
@@ -22,8 +22,8 @@ UPCGManagedProceduralISMComponent* PCGManagedProceduralISMComponent::GetOrCreate
 {
 	check(InTargetActor && InSourceComponent);
 
-	const UStaticMesh* StaticMesh = InParams.Descriptor.StaticMesh;
-	if (!StaticMesh)
+	const UStaticMesh* StaticMesh = InParams.Descriptor.StaticMesh.Get();
+	if (!ensure(StaticMesh))
 	{
 		return nullptr;
 	}
@@ -125,7 +125,7 @@ FPCGProceduralISMComponentDescriptor& FPCGProceduralISMComponentDescriptor::oper
 	InstanceStartCullDistance = Other.InstanceStartCullDistance;
 	InstanceEndCullDistance = Other.InstanceEndCullDistance;
 	OverlayMaterial = Other.OverlayMaterial.LoadSynchronous();
-	StaticMesh = Other.StaticMesh.LoadSynchronous();
+	StaticMesh = Other.StaticMesh;
 	ComponentTags = Other.ComponentTags;
 	Mobility = Other.Mobility;
 	VirtualTextureRenderPassType = Other.VirtualTextureRenderPassType;
@@ -168,7 +168,7 @@ FPCGProceduralISMComponentDescriptor& FPCGProceduralISMComponentDescriptor::oper
 
 void FPCGProceduralISMComponentDescriptor::InitComponent(UPCGProceduralISMComponent* InProceduralISMComponent) const
 {
-	InProceduralISMComponent->SetStaticMesh(StaticMesh);
+	InProceduralISMComponent->SetStaticMesh(StaticMesh.Get());
 	InProceduralISMComponent->OverrideMaterials = OverrideMaterials;
 	InProceduralISMComponent->OverlayMaterial = OverlayMaterial;
 	InProceduralISMComponent->RuntimeVirtualTextures = RuntimeVirtualTextures;
diff --git a/Engine/Plugins/PCG/Source/PCG/Private/Components/PCGProceduralISMComponent.h b/Engine/Plugins/PCG/Source/PCG/Private/Components/PCGProceduralISMComponent.h
--- a/Engine/Plugins/PCG/Source/PCG/Private/Components/PCGProceduralISMComponent.h
+++ b/Engine/Plugins/PCG/Source/PCG/Private/Components/PCGProceduralISMComponent.h
@@ -193,7 +193,7 @@ public:
 	mutable uint32 Hash = 0;
 
 	UPROPERTY()
-	TObjectPtr<UStaticMesh> StaticMesh = nullptr;
+	TSoftObjectPtr<UStaticMesh> StaticMesh = nullptr;
 
 	UPROPERTY()
 	TArray<TObjectPtr<UMaterialInterface>> OverrideMaterials;
diff --git a/Engine/Plugins/PCG/Source/PCG/Private/Compute/DataInterfaces/Elements/PCGStaticMeshSpawnerDataInterface.cpp b/Engine/Plugins/PCG/Source/PCG/Private/Compute/DataInterfaces/Elements/PCGStaticMeshSpawnerDataInterface.cpp
--- a/Engine/Plugins/PCG/Source/PCG/Private/Compute/DataInterfaces/Elements/PCGStaticMeshSpawnerDataInterface.cpp
+++ b/Engine/Plugins/PCG/Source/PCG/Private/Compute/DataInterfaces/Elements/PCGStaticMeshSpawnerDataInterface.cpp
@@ -372,6 +372,13 @@ bool UPCGStaticMeshSpawnerDataProvider::PrepareForExecute_GameThread(UPCGDataBin
 	{
 		CreatePrimitiveDescriptors(Context, InBinding);
 
+		for (const FPCGProceduralISMComponentDescriptor& PrimitiveDescriptor : PrimitiveDescriptors)
+		{
+			UE_LOG(LogPCG, Verbose, TEXT("Request '%s' to load."), *PrimitiveDescriptor.StaticMesh.ToString());
+
+			PrimitiveDescriptor.StaticMesh.LoadAsync({});
+		}
+
 		bPrimitiveDescriptorsCreated = true;
 
 		if (Context->AsyncState.ShouldStop())
@@ -380,6 +387,21 @@ bool UPCGStaticMeshSpawnerDataProvider::PrepareForExecute_GameThread(UPCGDataBin
 		}
 	}
 
+	if (!bStaticMeshesLoaded)
+	{
+		for (const FPCGProceduralISMComponentDescriptor& PrimitiveDescriptor : PrimitiveDescriptors)
+		{
+			if (PrimitiveDescriptor.StaticMesh.IsPending())
+			{
+				UE_LOG(LogPCG, Verbose, TEXT("Waiting for '%s' to load."), *PrimitiveDescriptor.StaticMesh.ToString());
+
+				return false;
+			}
+		}
+
+		bStaticMeshesLoaded = true;
+	}
+
 	if (!bPrimitivesSetUp)
 	{
 		if (!SetupPrimitives(Context, InBinding))
@@ -448,8 +470,10 @@ void UPCGStaticMeshSpawnerDataProvider::Reset()
 	PrimitiveDescriptors.Empty();
 	CustomFloatCount = 0;
 	bPrimitivesSetUp = false;
+	NumPrimitivesProcessed = 0;
 	NumPrimitivesSetup = 0;
 	bRegisteredPrimitives = false;
+	bStaticMeshesLoaded = false;
 
 	Super::Reset();
 }
@@ -585,21 +609,12 @@ void UPCGStaticMeshSpawnerDataProvider::CreatePrimitiveDescriptors(FPCGContext*
 				break;
 			}
 
-			UStaticMesh* StaticMesh = Cast<UStaticMesh>(FSoftObjectPath(MeshPathString).TryLoad());
-			if (!StaticMesh)
-			{
-				PCG_KERNEL_VALIDATION_ERR(InContext, Settings,
-					FText::Format(PCGStaticMeshSpawnerDataInterface::CouldNotLoadStaticMeshFormat, FText::FromString(MeshPathString)));
-
-				continue;
-			}
-
 			FPCGProceduralISMComponentDescriptor Descriptor;
 			Descriptor = SelectorByAttribute->TemplateDescriptor;
 			Descriptor.NumInstances = StringKeyAndInstanceCount.Value;
 			Descriptor.LocalBounds = LocalBounds;
 			Descriptor.NumCustomFloats = CustomFloatCount;
-			Descriptor.StaticMesh = StaticMesh;
+			Descriptor.StaticMesh = FSoftObjectPath(MeshPathString);
 
 			// Sanity check instance count.
 			if (!ensure(Descriptor.NumInstances <= TotalInputPointCount))
@@ -609,11 +624,6 @@ void UPCGStaticMeshSpawnerDataProvider::CreatePrimitiveDescriptors(FPCGContext*
 
 			PrimitiveStringKeys.Add(StringKeyAndInstanceCount.Key);
 			PrimitiveDescriptors.Add(MoveTemp(Descriptor));
-
-			if (Settings->bApplyMeshBoundsToPoints)
-			{
-				PrimitiveMeshBounds.Add(Descriptor.StaticMesh->GetBoundingBox());
-			}
 		}
 
 		PrimitiveSelectionCDF.SetNumZeroed(PrimitiveDescriptors.Num());
@@ -658,7 +668,7 @@ void UPCGStaticMeshSpawnerDataProvider::CreatePrimitiveDescriptors(FPCGContext*
 				break;
 			}
 
-			if (UStaticMesh* StaticMesh = Entry.Descriptor.StaticMesh.LoadSynchronous())
+			if (!Entry.Descriptor.StaticMesh.IsNull())
 			{
 				const float Weight = float(Entry.Weight) / TotalWeight;
 				CumulativeWeight += Weight;
@@ -669,7 +679,7 @@ void UPCGStaticMeshSpawnerDataProvider::CreatePrimitiveDescriptors(FPCGContext*
 				Descriptor = Entry.Descriptor;
 				Descriptor.LocalBounds = LocalBounds;
 				Descriptor.NumCustomFloats = CustomFloatCount;
-				Descriptor.StaticMesh = StaticMesh;
+				Descriptor.StaticMesh = Entry.Descriptor.StaticMesh;
 
 				int32 InstanceCount = FMath::CeilToInt(TotalInputPointCount * Weight);
 
@@ -691,11 +701,6 @@ void UPCGStaticMeshSpawnerDataProvider::CreatePrimitiveDescriptors(FPCGContext*
 				Descriptor.NumInstances = FMath::Min(InstanceCount, TotalInputPointCount);
 
 				PrimitiveDescriptors.Add(MoveTemp(Descriptor));
-
-				if (Settings->bApplyMeshBoundsToPoints)
-				{
-					PrimitiveMeshBounds.Add(Descriptor.StaticMesh->GetBoundingBox());
-				}
 			}
 			else
 			{
@@ -728,7 +733,6 @@ void UPCGStaticMeshSpawnerDataProvider::CreatePrimitiveDescriptors(FPCGContext*
 	Primitives.SelectionCDF = PrimitiveSelectionCDF;
 	Primitives.PrimitiveStringKeys = PrimitiveStringKeys;
 	Primitives.PrimitiveMeshBounds = PrimitiveMeshBounds;
-	Primitives.NumInstancesAllPrimitives = static_cast<uint32>(TotalInstanceCount);
 
 	// Useful for debugging instance counts.
 	//UE_LOG(LogPCG, Log, TEXT("Input points %d, total instance count %d (+ %.2f%%)"), TotalInputPointCount, Primitives.NumInstancesAllPrimitives, (float(Primitives.NumInstancesAllPrimitives - TotalInputPointCount) / TotalInputPointCount) * 100.0f);
@@ -750,9 +754,9 @@ bool UPCGStaticMeshSpawnerDataProvider::SetupPrimitives(FPCGContext* InContext,
 		return true;
 	}
 
-	while (NumPrimitivesSetup < PrimitiveDescriptors.Num())
+	while (NumPrimitivesProcessed < PrimitiveDescriptors.Num())
 	{
-		const FPCGProceduralISMComponentDescriptor& Desc = PrimitiveDescriptors[NumPrimitivesSetup];
+		const FPCGProceduralISMComponentDescriptor& Desc = PrimitiveDescriptors[NumPrimitivesProcessed];
 
 		FPCGProceduralISMCBuilderParameters Params;
 		Params.Descriptor = Desc;
@@ -763,17 +767,28 @@ bool UPCGStaticMeshSpawnerDataProvider::SetupPrimitives(FPCGContext* InContext,
 
 		UPCGManagedProceduralISMComponent* MISMC = PCGManagedProceduralISMComponent::GetOrCreateManagedProceduralISMC(TargetActor, SourceComponent, Settings->UID, Params);
 
-		check(MISMC);
-		MISMC->SetCrc(InContext->DependenciesCrc);
+		if (MISMC)
+		{
+			MISMC->SetCrc(InContext->DependenciesCrc);
+
+			// Don't bother registering the resource change as we're transient anyway.
+			//InContext->TouchedResources.Emplace(MISMC);
+
+			Primitives->Primitives.Add(MISMC->GetComponent());
+
+			if (Settings->bApplyMeshBoundsToPoints)
+			{
+				PrimitiveMeshBounds.Add(Desc.StaticMesh->GetBoundingBox());
+			}
 
-		// Don't bother registering the resource change as we're transient anyway.
-		//InContext->TouchedResources.Emplace(MISMC);
+			Primitives->NumInstancesAllPrimitives += Desc.NumInstances;
 
-		Primitives->Primitives.Add(MISMC->GetComponent());
+			++NumPrimitivesSetup;
+		}
 
-		++NumPrimitivesSetup;
+		++NumPrimitivesProcessed;
 
-		if (NumPrimitivesSetup < PrimitiveDescriptors.Num() && InContext->AsyncState.ShouldStop())
+		if (NumPrimitivesProcessed < PrimitiveDescriptors.Num() && InContext->AsyncState.ShouldStop())
 		{
 			return false;
 		}
diff --git a/Engine/Plugins/PCG/Source/PCG/Private/Compute/DataInterfaces/Elements/PCGStaticMeshSpawnerDataInterface.h b/Engine/Plugins/PCG/Source/PCG/Private/Compute/DataInterfaces/Elements/PCGStaticMeshSpawnerDataInterface.h
--- a/Engine/Plugins/PCG/Source/PCG/Private/Compute/DataInterfaces/Elements/PCGStaticMeshSpawnerDataInterface.h
+++ b/Engine/Plugins/PCG/Source/PCG/Private/Compute/DataInterfaces/Elements/PCGStaticMeshSpawnerDataInterface.h
@@ -93,6 +93,10 @@ public:
 
 	bool bRegisteredPrimitives = false;
 
+	bool bStaticMeshesLoaded = false;
+
+	int32 NumPrimitivesProcessed = 0;
+
 	bool bPrimitivesSetUp = false;
 
 	/** Number of primitives created during primitive setup. */
