[PCG] Fix a couple remaining issues related to PointArrayData
- Projection: Update remaining ToPointData calls to use ToBasePointData instead
- SpatialData with cache: Make sure the returned cache data is of the proper type by keeping two different caches, this should only happen with code that has not been converted yet
- Union: guard against no input points

#rnx

diff --git a/Engine/Plugins/PCG/Source/PCG/Private/Data/PCGProjectionData.cpp b/Engine/Plugins/PCG/Source/PCG/Private/Data/PCGProjectionData.cpp
--- a/Engine/Plugins/PCG/Source/PCG/Private/Data/PCGProjectionData.cpp
+++ b/Engine/Plugins/PCG/Source/PCG/Private/Data/PCGProjectionData.cpp
@@ -175,7 +175,7 @@ bool UPCGProjectionData::SamplePoint(const FTransform& InTransform, const FBox&
 	{
 		// Passing nullptr for the context means the operation will execute single threaded which is not ideal. To mitigate this we
 		// prewarm the point cache when this projection data is constructed in the projection element.
-		return ToPointData(nullptr)->SamplePoint(InTransform, InBounds, OutPoint, OutMetadata);
+		return ToBasePointData(nullptr)->SamplePoint(InTransform, InBounds, OutPoint, OutMetadata);
 	}
 
 	FPCGPoint PointFromSource;
diff --git a/Engine/Plugins/PCG/Source/PCG/Private/Data/PCGSpatialData.cpp b/Engine/Plugins/PCG/Source/PCG/Private/Data/PCGSpatialData.cpp
--- a/Engine/Plugins/PCG/Source/PCG/Private/Data/PCGSpatialData.cpp
+++ b/Engine/Plugins/PCG/Source/PCG/Private/Data/PCGSpatialData.cpp
@@ -63,10 +63,8 @@ const UPCGPointArrayData* UPCGSpatialData::ToPointArrayData(FPCGContext* Context
 	}
 }
 
-void UPCGSpatialDataWithPointCache::GetResourceSizeEx(FResourceSizeEx& CumulativeResourceSize)
+void FPCGPointDataCache::GetResourceSizeEx(FResourceSizeEx& CumulativeResourceSize)
 {
-	Super::GetResourceSizeEx(CumulativeResourceSize);
-
 	if (CachedPointData)
 	{
 		const_cast<UPCGBasePointData*>(CachedPointData.Get())->GetResourceSizeEx(CumulativeResourceSize);
@@ -83,12 +81,20 @@ void UPCGSpatialDataWithPointCache::GetResourceSizeEx(FResourceSizeEx& Cumulativ
 	}
 }
 
-const UPCGBasePointData* UPCGSpatialDataWithPointCache::ToBasePointDataInternal(FPCGContext* Context, const FBox& InBounds, TFunctionRef<const UPCGBasePointData*(FPCGContext*, const FBox&)> CreatePointDataFunc) const
+void UPCGSpatialDataWithPointCache::GetResourceSizeEx(FResourceSizeEx& CumulativeResourceSize)
+{
+	Super::GetResourceSizeEx(CumulativeResourceSize);
+
+	PointDataCache.GetResourceSizeEx(CumulativeResourceSize);
+	PointArrayDataCache.GetResourceSizeEx(CumulativeResourceSize);
+}
+
+const UPCGBasePointData* FPCGPointDataCache::ToBasePointDataInternal(FPCGContext* Context, const FBox& InBounds, bool bSupportsBoundedPointData, FCriticalSection& InCacheLock, TFunctionRef<const UPCGBasePointData*(FPCGContext*, const FBox&)> CreatePointDataFunc)
 {
-	if (InBounds.IsValid && SupportsBoundedPointData())
+	if (InBounds.IsValid && bSupportsBoundedPointData)
 	{
 		const UPCGBasePointData* BoundedPointData = nullptr;
-		CacheLock.Lock();
+		InCacheLock.Lock();
 		check(CachedBoundedPointDataBoxes.Num() == CachedBoundedPointData.Num());
 		for (int CachedDataIndex = 0; CachedDataIndex < CachedBoundedPointDataBoxes.Num(); ++CachedDataIndex)
 		{
@@ -105,7 +111,7 @@ const UPCGBasePointData* UPCGSpatialDataWithPointCache::ToBasePointDataInternal(
 			CachedBoundedPointDataBoxes.Add(InBounds);
 			CachedBoundedPointData.Add(BoundedPointData);
 		}
-		CacheLock.Unlock();
+		InCacheLock.Unlock();
 
 		return BoundedPointData;
 	}
@@ -113,14 +119,14 @@ const UPCGBasePointData* UPCGSpatialDataWithPointCache::ToBasePointDataInternal(
 	{
 		if (!CachedPointData)
 		{
-			CacheLock.Lock();
+			InCacheLock.Lock();
 
 			if (!CachedPointData)
 			{
 				CachedPointData = CreatePointDataFunc(Context, InBounds);
 			}
 
-			CacheLock.Unlock();
+			InCacheLock.Unlock();
 		}
 
 		return CachedPointData;
@@ -129,12 +135,12 @@ const UPCGBasePointData* UPCGSpatialDataWithPointCache::ToBasePointDataInternal(
 
 const UPCGPointData* UPCGSpatialDataWithPointCache::ToPointData(FPCGContext* Context, const FBox& InBounds) const
 {
-	return CastChecked<UPCGPointData>(ToBasePointDataInternal(Context, InBounds, [this](FPCGContext* InContext, const FBox& InBounds) { return CreatePointData(InContext, InBounds); }), ECastCheckedType::NullAllowed);
+	return CastChecked<UPCGPointData>(PointDataCache.ToBasePointDataInternal(Context, InBounds, SupportsBoundedPointData(), CacheLock, [this](FPCGContext* InContext, const FBox& InBounds) { return CreatePointData(InContext, InBounds); }), ECastCheckedType::NullAllowed);
 }
 
 const UPCGPointArrayData* UPCGSpatialDataWithPointCache::ToPointArrayData(FPCGContext* Context, const FBox& InBounds) const
 {
-	return CastChecked<UPCGPointArrayData>(ToBasePointDataInternal(Context, InBounds, [this](FPCGContext* InContext, const FBox& InBounds) { return CreatePointArrayData(InContext, InBounds); }), ECastCheckedType::NullAllowed);
+	return CastChecked<UPCGPointArrayData>(PointArrayDataCache.ToBasePointDataInternal(Context, InBounds, SupportsBoundedPointData(), CacheLock, [this](FPCGContext* InContext, const FBox& InBounds) { return CreatePointArrayData(InContext, InBounds); }), ECastCheckedType::NullAllowed);
 }
 
 const UPCGPointArrayData* UPCGSpatialDataWithPointCache::CreatePointArrayData(FPCGContext* Context, const FBox& InBounds) const
diff --git a/Engine/Plugins/PCG/Source/PCG/Private/Data/PCGUnionData.cpp b/Engine/Plugins/PCG/Source/PCG/Private/Data/PCGUnionData.cpp
--- a/Engine/Plugins/PCG/Source/PCG/Private/Data/PCGUnionData.cpp
+++ b/Engine/Plugins/PCG/Source/PCG/Private/Data/PCGUnionData.cpp
@@ -280,7 +280,7 @@ const UPCGBasePointData* UPCGUnionData::CreateBasePointData(FPCGContext* Context
 	FPCGInitializeFromDataParams InitializeFromDataParams(FirstSource);
 	InitializeFromDataParams.bInheritSpatialData = false;
 	PointData->InitializeFromDataWithParams(InitializeFromDataParams);
-	
+
 	switch (UnionType)
 	{
 	case EPCGUnionType::LeftToRightPriority:
@@ -305,80 +305,83 @@ const UPCGBasePointData* UPCGUnionData::CreateBasePointData(FPCGContext* Context
 				PropertiesToAllocate |= PointInputDatas[Index]->GetAllocatedProperties();
 			}
 
-			PointData->SetNumPoints(PointInputCount);
-			PointData->AllocateProperties(PropertiesToAllocate);
-
-			int32 NumWritten = 0;
-
-			for(int32 DataIndex = 0; DataIndex < DataRawPtr.Num(); ++DataIndex)
+			if (PointInputCount > 0)
 			{
-				const UPCGSpatialData* Datum = DataRawPtr[DataIndex];
-				const UPCGBasePointData* DatumPointData = PointInputDatas[DataIndex];
-				const UPCGMetadata* DatumPointMetadata = DatumPointData->Metadata;
+				PointData->SetNumPoints(PointInputCount);
+				PointData->AllocateProperties(PropertiesToAllocate);
 
-				const int32 DatumNumPoints = DatumPointData->GetNumPoints();
-				const int32 TargetPointOffset = NumWritten;
-				
-				if (DatumNumPoints > 0)
+				int32 NumWritten = 0;
+
+				for (int32 DataIndex = 0; DataIndex < DataRawPtr.Num(); ++DataIndex)
 				{
-					DatumPointData->CopyPointsTo(PointData, 0, TargetPointOffset, DatumNumPoints);
-					NumWritten += DatumNumPoints;
+					const UPCGSpatialData* Datum = DataRawPtr[DataIndex];
+					const UPCGBasePointData* DatumPointData = PointInputDatas[DataIndex];
+					const UPCGMetadata* DatumPointMetadata = DatumPointData->Metadata;
 
-					if (DataIndex > 0)
+					const int32 DatumNumPoints = DatumPointData->GetNumPoints();
+					const int32 TargetPointOffset = NumWritten;
+
+					if (DatumNumPoints > 0)
 					{
-						TPCGValueRange<int64> TargetMetadataEntryRange = PointData->GetMetadataEntryValueRange();
-						
-						// TODO: could optimize case where there is a common parent between Data 0 and current data, for points that still point to common parent metadata.
-						for (int32 TargetIndex = TargetPointOffset; TargetIndex < TargetPointOffset + DatumNumPoints; ++TargetIndex)
-						{
-							TargetMetadataEntryRange[TargetIndex] = PCGInvalidEntryKey;
-						}
+						DatumPointData->CopyPointsTo(PointData, 0, TargetPointOffset, DatumNumPoints);
+						NumWritten += DatumNumPoints;
 
-						if (OutMetadata && DatumPointMetadata && DatumPointMetadata->GetAttributeCount() > 0)
+						if (DataIndex > 0)
 						{
-							TArray<PCGMetadataEntryKey, TInlineAllocator<256>> DatumKeys;
-							DatumKeys.SetNumUninitialized(DatumNumPoints);
+							TPCGValueRange<int64> TargetMetadataEntryRange = PointData->GetMetadataEntryValueRange();
 
-							const TConstPCGValueRange<int64> DatumMetadataEntryRange = DatumPointData->GetConstMetadataEntryValueRange();
-							for (int32 DatumIndex = 0; DatumIndex < DatumKeys.Num(); ++DatumIndex)
+							// TODO: could optimize case where there is a common parent between Data 0 and current data, for points that still point to common parent metadata.
+							for (int32 TargetIndex = TargetPointOffset; TargetIndex < TargetPointOffset + DatumNumPoints; ++TargetIndex)
 							{
-								DatumKeys[DatumIndex] = DatumMetadataEntryRange[DatumIndex];
+								TargetMetadataEntryRange[TargetIndex] = PCGInvalidEntryKey;
 							}
 
-							TArray<PCGMetadataEntryKey, TInlineAllocator<256>> TargetKeys;
-							TargetKeys.SetNumUninitialized(DatumNumPoints);
-
-							for (int32 TargetIndex = 0; TargetIndex < TargetKeys.Num(); ++TargetIndex)
+							if (OutMetadata && DatumPointMetadata && DatumPointMetadata->GetAttributeCount() > 0)
 							{
-								TargetKeys[TargetIndex] = PCGInvalidEntryKey;
-							}
-
-							PointData->Metadata->SetAttributes(DatumKeys, DatumPointMetadata, TargetKeys, Context);
-
-							// Write back
-							for (int32 TargetIndex = 0; TargetIndex < TargetKeys.Num(); ++TargetIndex)
-							{
-								TargetMetadataEntryRange[TargetPointOffset + TargetIndex] = TargetKeys[TargetIndex];
+								TArray<PCGMetadataEntryKey, TInlineAllocator<256>> DatumKeys;
+								DatumKeys.SetNumUninitialized(DatumNumPoints);
+
+								const TConstPCGValueRange<int64> DatumMetadataEntryRange = DatumPointData->GetConstMetadataEntryValueRange();
+								for (int32 DatumIndex = 0; DatumIndex < DatumKeys.Num(); ++DatumIndex)
+								{
+									DatumKeys[DatumIndex] = DatumMetadataEntryRange[DatumIndex];
+								}
+
+								TArray<PCGMetadataEntryKey, TInlineAllocator<256>> TargetKeys;
+								TargetKeys.SetNumUninitialized(DatumNumPoints);
+
+								for (int32 TargetIndex = 0; TargetIndex < TargetKeys.Num(); ++TargetIndex)
+								{
+									TargetKeys[TargetIndex] = PCGInvalidEntryKey;
+								}
+
+								PointData->Metadata->SetAttributes(DatumKeys, DatumPointMetadata, TargetKeys, Context);
+
+								// Write back
+								for (int32 TargetIndex = 0; TargetIndex < TargetKeys.Num(); ++TargetIndex)
+								{
+									TargetMetadataEntryRange[TargetPointOffset + TargetIndex] = TargetKeys[TargetIndex];
+								}
 							}
 						}
 					}
 				}
-			}
 
-			// Correct density for binary-style union
-			if (bBinaryDensity)
-			{
-				if (EnumHasAllFlags(PointData->GetAllocatedProperties(), EPCGPointNativeProperties::Density))
+				// Correct density for binary-style union
+				if (bBinaryDensity)
 				{
-					TPCGValueRange<float> DensityRange = PointData->GetDensityValueRange(/*bAllocate=*/false);
-					for (float& Density : DensityRange)
+					if (EnumHasAllFlags(PointData->GetAllocatedProperties(), EPCGPointNativeProperties::Density))
 					{
-						Density = ((Density > 0) ? 1.0f : 0);
+						TPCGValueRange<float> DensityRange = PointData->GetDensityValueRange(/*bAllocate=*/false);
+						for (float& Density : DensityRange)
+						{
+							Density = ((Density > 0) ? 1.0f : 0);
+						}
+					}
+					else
+					{
+						PointData->SetDensity(PointData->GetDensity(0) > 0 ? 1.0f : 0);
 					}
-				}
-				else
-				{
-					PointData->SetDensity(PointData->GetDensity(0) > 0 ? 1.0f : 0);
 				}
 			}
 		}
diff --git a/Engine/Plugins/PCG/Source/PCG/Private/Elements/PCGProjectionElement.cpp b/Engine/Plugins/PCG/Source/PCG/Private/Elements/PCGProjectionElement.cpp
--- a/Engine/Plugins/PCG/Source/PCG/Private/Elements/PCGProjectionElement.cpp
+++ b/Engine/Plugins/PCG/Source/PCG/Private/Elements/PCGProjectionElement.cpp
@@ -109,9 +109,9 @@ bool FPCGProjectionElement::ExecuteInternal(FPCGContext* Context) const
 		}
 		else if (OutputData->RequiresCollapseToSample())
 		{
-			// Calling ToPointData will populate the point cache. Doing so here means we can pass in the Context object, which
+			// Calling ToBasePointData will populate the point cache. Doing so here means we can pass in the Context object, which
 			// means the operation will be multi-threaded. This primes the cache in the most efficient way.
-			OutputData->ToPointData(Context);
+			OutputData->ToBasePointData(Context);
 		}
 
 		FPCGTaggedData& ProjectionTaggedData = Outputs.Emplace_GetRef(Source);
diff --git a/Engine/Plugins/PCG/Source/PCG/Public/Data/PCGSpatialData.h b/Engine/Plugins/PCG/Source/PCG/Public/Data/PCGSpatialData.h
--- a/Engine/Plugins/PCG/Source/PCG/Public/Data/PCGSpatialData.h
+++ b/Engine/Plugins/PCG/Source/PCG/Public/Data/PCGSpatialData.h
@@ -254,6 +254,24 @@ private:
 	FPCGAttributePropertyInputSelector CachedLastSelector;
 };
 
+USTRUCT()
+struct FPCGPointDataCache
+{
+	GENERATED_BODY()
+
+	const UPCGBasePointData* ToBasePointDataInternal(FPCGContext* Context, const FBox& InBounds, bool bSupportsBoundedPointData, FCriticalSection& CacheLock, TFunctionRef<const UPCGBasePointData* (FPCGContext*, const FBox&)> CreatePointDataFunc);
+	void GetResourceSizeEx(FResourceSizeEx& CumulativeResourceSize);
+
+	UPROPERTY(Transient)
+	TObjectPtr<const UPCGBasePointData> CachedPointData;
+
+	UPROPERTY(Transient)
+	TArray<FBox> CachedBoundedPointDataBoxes;
+
+	UPROPERTY(Transient)
+	TArray<TObjectPtr<const UPCGBasePointData>> CachedBoundedPointData;
+};
+
 UCLASS(Abstract, ClassGroup = (Procedural))
 class PCG_API UPCGSpatialDataWithPointCache : public UPCGSpatialData
 {
@@ -270,8 +288,6 @@ public:
 	// ~End UPCGSpatialData implementation
 
 protected:
-	const UPCGBasePointData* ToBasePointDataInternal(FPCGContext* Context, const FBox& InBounds, TFunctionRef<const UPCGBasePointData* (FPCGContext*, const FBox&)> CreatePointDataFunc) const;
-
 	virtual bool SupportsBoundedPointData() const { return false; }
 	virtual const UPCGPointData* CreatePointData(FPCGContext* Context) const PURE_VIRTUAL(UPCGSpatialData::CreatePointData, return nullptr;);
 	virtual const UPCGPointData* CreatePointData(FPCGContext* Context, const FBox& InBounds) const { return CreatePointData(Context); }
@@ -279,13 +295,10 @@ protected:
 
 private:
 	UPROPERTY(Transient)
-	mutable TObjectPtr<const UPCGBasePointData> CachedPointData;
-
-	UPROPERTY(Transient)
-	mutable TArray<FBox> CachedBoundedPointDataBoxes;
+	mutable FPCGPointDataCache PointDataCache;
 
 	UPROPERTY(Transient)
-	mutable TArray<TObjectPtr<const UPCGBasePointData>> CachedBoundedPointData;
+	mutable FPCGPointDataCache PointArrayDataCache;
 
 	mutable FCriticalSection CacheLock;
 };
