[PCG] Fix issue where Ctrl+click on paths in the attribute list view could crash

* Now we check if the path is an asset, if so we open it
* If it is an Actor/ActorComponent, we jump to that one in the world
* Don't load world/levels, that would be destructive

#rnx

diff --git a/Engine/Plugins/PCG/Source/PCGEditor/Private/PCGEditorUtils.cpp b/Engine/Plugins/PCG/Source/PCGEditor/Private/PCGEditorUtils.cpp
--- a/Engine/Plugins/PCG/Source/PCGEditor/Private/PCGEditorUtils.cpp
+++ b/Engine/Plugins/PCG/Source/PCGEditor/Private/PCGEditorUtils.cpp
@@ -8,11 +8,12 @@
 
 #include "AssetToolsModule.h"
 #include "ContentBrowserModule.h"
+#include "Editor.h"
 #include "FrontendFilterBase.h"
 #include "IAssetTools.h"
 #include "IContentBrowserSingleton.h"
+#include "ScopedTransaction.h"
 #include "SPrimaryButton.h"
-#include "SSimpleButton.h"
 #include "Algo/Transform.h"
 #include "AssetRegistry/AssetData.h"
 #include "AssetRegistry/ARFilter.h"
@@ -20,15 +21,13 @@
 #include "Blueprint/BlueprintSupport.h"
 #include "Filters/SBasicFilterBar.h"
 #include "Framework/Application/SlateApplication.h"
-#include "Interfaces/IMainFrameModule.h"
 #include "Kismet2/BlueprintEditorUtils.h"
 #include "Kismet2/KismetEditorUtilities.h"
 #include "Misc/PackageName.h"
-#include "Styling/AppStyle.h"
+#include "Subsystems/AssetEditorSubsystem.h"
 #include "Types/SlateEnums.h"
 #include "Widgets/Layout/SBox.h"
 #include "Widgets/Input/SButton.h"
-#include "Widgets/Layout/SBorder.h"
 #include "Widgets/Layout/SSpacer.h"
 #include "Widgets/SBoxPanel.h"
 #include "Widgets/SWindow.h"
@@ -349,4 +348,79 @@ bool PCGEditorUtils::PickGraphTemplate(FAssetData& OutAssetData, const FText& Ti
 	return bClickedOk;
 }
 
-#undef LOCTEXT_NAMESPACE
\ No newline at end of file
+void PCGEditorUtils::OpenAssetOrMoveToActorOrComponent(const FSoftObjectPath& InPath)
+{
+	if (!GEditor || !GEditor->GetEditorSubsystem<UAssetEditorSubsystem>())
+	{
+		return;
+	}
+
+	FAssetRegistryModule& AssetRegistryModule = FModuleManager::LoadModuleChecked<FAssetRegistryModule>(TEXT("AssetRegistry"));
+	IAssetRegistry& AssetRegistry = AssetRegistryModule.Get();
+
+	// We need to know which class the object is. If it is not an asset, we can't jump to it.
+	const FAssetData AssetData = AssetRegistry.GetAssetByObjectPath(InPath);
+	if (!AssetData.IsValid() || AssetData.IsRedirector())
+	{
+		return;
+	}
+
+	const UClass* AssetClass = AssetData.GetClass(EResolveClass::Yes);
+	if (!AssetClass)
+	{
+		return;
+	}
+
+	// Don't jump to a world/level, would be pretty destructive to change levels
+	if (AssetClass->IsChildOf<UWorld>() || AssetClass->IsChildOf<ULevel>())
+	{
+		return;
+	}
+
+	// If it is not an actor or an actor component, we can try to open an editor for it.
+	if (!AssetClass->IsChildOf<AActor>() && !AssetClass->IsChildOf<UActorComponent>())
+	{
+		GEditor->GetEditorSubsystem<UAssetEditorSubsystem>()->OpenEditorForAsset(InPath);
+		return;
+	}
+	
+	// Otherwise, try to resolve the actor/actor component to be able to jump to it. Never try to load it.
+	// Encapsulate the move in a transaction to be able to undo the selection.
+	if (const UObject* Object = InPath.ResolveObject())
+	{	
+		FScopedTransaction Transaction(LOCTEXT("PCGHyperLinkSoftObjectPath", "[PCG] Jump to Actor/Component"));
+		bool bSuccess = false;
+		
+		if (const USceneComponent* SceneComponent = Cast<USceneComponent>(Object))
+		{
+			GEditor->MoveViewportCamerasToComponent(SceneComponent, /*bActiveViewportOnly=*/true);
+			GEditor->SelectNone(/*bNoteSelectionChange=*/ false, /*bDeselectBSPSurfsfalse=*/ true);
+			GEditor->SelectComponent(const_cast<USceneComponent*>(SceneComponent), /*bInSelected=*/ true, /*bNotify=*/ true);
+			bSuccess = true;
+		}
+		else if (const AActor* Actor = Cast<AActor>(Object))
+		{
+			GEditor->MoveViewportCamerasToActor(const_cast<AActor&>(*Actor), /*bActiveViewportOnly=*/true);
+			GEditor->SelectNone(/*bNoteSelectionChange=*/ false, /*bDeselectBSPSurfsfalse=*/ true);
+			GEditor->SelectActor(const_cast<AActor*>(Actor), /*bInSelected=*/ true, /*bNotify=*/ true);
+			bSuccess = true;
+		}
+		else if (const UActorComponent* ActorComponent = Cast<UActorComponent>(Object))
+		{
+			if (AActor* OwnerActor = ActorComponent->GetOwner())
+			{
+				GEditor->MoveViewportCamerasToActor(*OwnerActor, /*bActiveViewportOnly=*/true);
+				GEditor->SelectNone(/*bNoteSelectionChange=*/ false, /*bDeselectBSPSurfsfalse=*/ true);
+				GEditor->SelectComponent(const_cast<UActorComponent*>(ActorComponent), /*bInSelected=*/ true, /*bNotify=*/ true);
+				bSuccess = true;
+			}
+		}
+
+		if (!bSuccess)
+		{
+			Transaction.Cancel();
+		}
+	}
+}
+
+#undef LOCTEXT_NAMESPACE
diff --git a/Engine/Plugins/PCG/Source/PCGEditor/Private/PCGEditorUtils.h b/Engine/Plugins/PCG/Source/PCGEditor/Private/PCGEditorUtils.h
--- a/Engine/Plugins/PCG/Source/PCGEditor/Private/PCGEditorUtils.h
+++ b/Engine/Plugins/PCG/Source/PCGEditor/Private/PCGEditorUtils.h
@@ -34,4 +34,7 @@ namespace PCGEditorUtils
 
 	/** Opens a template picker dialog. Returns true if user selected 'ok'. */
 	bool PickGraphTemplate(FAssetData& OutData, const FText& InTitleOverride = FText());
+
+	/** Open an asset or jumps to an actor/component and select it. Will load it, and won't load a new level*/
+	void OpenAssetOrMoveToActorOrComponent(const FSoftObjectPath& InPath);
 }
diff --git a/Engine/Plugins/PCG/Source/PCGEditor/Private/Widgets/SPCGEditorGraphAttributeListView.cpp b/Engine/Plugins/PCG/Source/PCGEditor/Private/Widgets/SPCGEditorGraphAttributeListView.cpp
--- a/Engine/Plugins/PCG/Source/PCGEditor/Private/Widgets/SPCGEditorGraphAttributeListView.cpp
+++ b/Engine/Plugins/PCG/Source/PCGEditor/Private/Widgets/SPCGEditorGraphAttributeListView.cpp
@@ -8,22 +8,20 @@
 #include "PCGDataVisualization.h"
 #include "PCGGraph.h"
 #include "PCGNode.h"
-#include "PCGParamData.h"
 #include "PCGPin.h"
 #include "PCGSubsystem.h"
 #include "Compute/Data/PCGProxyForGPUData.h"
-#include "Data/PCGSpatialData.h"
 #include "Elements/IO/PCGSaveAssetElement.h"
 #include "Metadata/PCGAttributePropertySelector.h"
 #include "Metadata/Accessors/IPCGAttributeAccessor.h"
 #include "Metadata/Accessors/PCGAttributeAccessor.h"
 #include "Metadata/Accessors/PCGAttributeAccessorHelpers.h"
 #include "Metadata/Accessors/PCGAttributeAccessorKeys.h"
-#include "Metadata/Accessors/PCGCustomAccessor.h"
 
 #include "PCGEditor.h"
 #include "PCGEditorGraph.h"
 #include "PCGEditorStyle.h"
+#include "PCGEditorUtils.h"
 #include "Nodes/PCGEditorGraphNodeBase.h"
 #include "Widgets/AssetEditorViewport/SPCGEditorViewport.h"
 
@@ -211,22 +209,12 @@ TSharedRef<SWidget> SPCGListViewItemRow::GenerateWidgetForColumn(const FName& Co
 	if (bIsClickable)
 	{
 		// @todo_pcg: We're missing Margin and Overflow policy with this widget, maybe something we can re-add in the future.
-
-		// Put it in a shared to pass it as the text of the hyperlink and to the OnNavigate callback
-		TSharedPtr<FText> Text = MakeShared<FText>(std::move(RowText));
-		
 		TSharedPtr<SHyperlink> Hyperlink;
 		SAssignNew(Hyperlink, SHyperlink)
-			.Text(TAttribute<FText>::CreateLambda([Text]() { return *Text;}))
+			.Text(RowText)
 			.ToolTipText(LOCTEXT("Hyperlink", "Ctrl + Click to jump to the asset."))
 			.Style(FAppStyle::Get(), "HoverOnlyHyperlink")
-			.OnNavigate(FSimpleDelegate::CreateLambda([Text]()
-			{
-				if (FSlateApplication::Get().GetModifierKeys().IsControlDown() && GEditor && GEditor->GetEditorSubsystem<UAssetEditorSubsystem>())
-				{
-					GEditor->GetEditorSubsystem<UAssetEditorSubsystem>()->OpenEditorForAsset(Text->ToString());
-				}
-			}));
+			.OnNavigate(FSimpleDelegate::CreateLambda([Text = std::move(RowText)](){ SPCGListViewItemRow::OnSoftObjectPathHyperlinkClicked(Text); }));
 
 		return Hyperlink.ToSharedRef();
 	}
@@ -239,6 +227,14 @@ TSharedRef<SWidget> SPCGListViewItemRow::GenerateWidgetForColumn(const FName& Co
 	}
 }
 
+void SPCGListViewItemRow::OnSoftObjectPathHyperlinkClicked(const FText& InText)
+{
+	if (FSlateApplication::Get().GetModifierKeys().IsControlDown())
+	{
+		PCGEditorUtils::OpenAssetOrMoveToActorOrComponent(FSoftObjectPath(InText.ToString()));
+	}
+}
+
 FPCGPointFilterExpressionContext::FPCGPointFilterExpressionContext(const FPCGListViewItem* InRowItem, const TMap<FName, FPCGColumnData>* InPCGColumnData)
 	: RowItem(InRowItem)
 	, PCGColumnData(InPCGColumnData)
diff --git a/Engine/Plugins/PCG/Source/PCGEditor/Private/Widgets/SPCGEditorGraphAttributeListView.h b/Engine/Plugins/PCG/Source/PCGEditor/Private/Widgets/SPCGEditorGraphAttributeListView.h
--- a/Engine/Plugins/PCG/Source/PCGEditor/Private/Widgets/SPCGEditorGraphAttributeListView.h
+++ b/Engine/Plugins/PCG/Source/PCGEditor/Private/Widgets/SPCGEditorGraphAttributeListView.h
@@ -114,6 +114,8 @@ public:
 
 	virtual TSharedRef<SWidget> GenerateWidgetForColumn(const FName& ColumnId) override;
 
+	static void OnSoftObjectPathHyperlinkClicked(const FText& InText);
+
 private:
 	TWeakPtr<SPCGEditorGraphAttributeListView> AttributeListView;
 	PCGListViewItemPtr InternalItem;
