[PCG] Fix Material Overrides on template descriptor ignored when using Mesh Selector by Attribute

* If we had no override by attribute for the materials, we were using an empty list instead of using the one from the template descriptor.

#rnx

diff --git a/Engine/Plugins/PCG/Source/PCG/Private/MeshSelectors/PCGMeshSelectorByAttribute.cpp b/Engine/Plugins/PCG/Source/PCG/Private/MeshSelectors/PCGMeshSelectorByAttribute.cpp
--- a/Engine/Plugins/PCG/Source/PCG/Private/MeshSelectors/PCGMeshSelectorByAttribute.cpp
+++ b/Engine/Plugins/PCG/Source/PCG/Private/MeshSelectors/PCGMeshSelectorByAttribute.cpp
@@ -306,8 +306,7 @@ bool UPCGMeshSelectorByAttribute::SelectMeshInstances(
 				}
 				else
 				{
-					TArray<TSoftObjectPtr<UMaterialInterface>> DummyMaterialList;
-					FPCGMeshInstanceList& InstanceList = PCGMeshSelectorAttribute::GetInstanceList(OutMeshInstances, CurrentPartitionDescriptor, CurrentPartitionDescriptor.StaticMesh, DummyMaterialList, bReverseTransform, InPointData, PartitionIndex);
+					FPCGMeshInstanceList& InstanceList = PCGMeshSelectorAttribute::GetInstanceList(OutMeshInstances, CurrentPartitionDescriptor, CurrentPartitionDescriptor.StaticMesh, CurrentPartitionDescriptor.OverrideMaterials, bReverseTransform, InPointData, PartitionIndex);
 
 					check(InstanceList.Instances.Num() == InstanceList.InstancesIndices.Num());
 					const int32 InstanceOffset = InstanceList.Instances.Num();
