[PCG] Make sure GraphExecutor respects frame budget in the scheduling code (not only in the FPCGGraphExecutor::Execute top level loop)

#rnx

diff --git a/Engine/Plugins/PCG/Source/PCG/Private/Graph/PCGGraphExecutor.cpp b/Engine/Plugins/PCG/Source/PCG/Private/Graph/PCGGraphExecutor.cpp
--- a/Engine/Plugins/PCG/Source/PCG/Private/Graph/PCGGraphExecutor.cpp
+++ b/Engine/Plugins/PCG/Source/PCG/Private/Graph/PCGGraphExecutor.cpp
@@ -845,13 +845,43 @@ void FPCGGraphExecutor::ClearOutputData(FPCGTaskId TaskId)
 	}
 }
 
-bool FPCGGraphExecutor::ProcessScheduledTasks()
+void FPCGGraphExecutor::CheckState(bool bIsInGameThread)
+{
+	if (bIsInGameThread)
+	{
+		// Check every second to avoid locking too often for an error that shouldn't happen
+		double TimeInSeconds = FPlatformTime::Seconds();
+		if (TimeInSeconds - LastSchedulingErrorCheck > 1.0)
+		{
+			LastSchedulingErrorCheck = TimeInSeconds;
+
+			PCGGraphExecutor::TScopeLock ScopeLock(TasksLock);
+			PCGGraphExecutor::TScopeLock ChildScopeLock(LiveTasksLock);
+			PCGGraphExecutor::TScopeLock PausedLock(PausedTaskSuccessorsLock);
+			PCGGraphExecutor::TScopeLock CachingResultsScopeLock(CachingResultsLock);
+			// This is a safeguard to check if we're in a stuck state
+			if (CollectGCCachingResults.Num() == 0 && ReadyTasks.Num() == 0 && ActiveTasks.Num() == 0 && ActiveTasksGameThreadOnly.Num() == 0 && PausedTasks.Num() == 0 && Tasks.Num() > 0)
+			{
+				UE_LOG(LogPCG, Error, TEXT("PCG Graph executor error: tasks are in a deadlocked state. Will drop all tasks."));
+				ClearAllTasks();
+			}
+		}
+	}
+}
+
+bool FPCGGraphExecutor::ProcessScheduledTasks(double EndTime)
 {
 	TRACE_CPUPROFILER_EVENT_SCOPE(FPCGGraphExecutor::ProcessScheduledTasks);
 	bool bProcessedNewTasks = false;
 	const bool bIsInGameThread = IsInGameThread();
 	TArray<FCachedResult*> CachedResults;
-	TArray<FPCGGraphTask, TInlineAllocator<16>> NewReadyTasks;
+
+	// Avoid locking if empty
+	if (ScheduledTasks.IsEmpty())
+	{
+		CheckState(bIsInGameThread);
+		return bProcessedNewTasks;
+	}
 
 	{
 		// Process any newly scheduled graphs to execute
@@ -860,43 +890,66 @@ bool FPCGGraphExecutor::ProcessScheduledTasks()
 		{
 			if (ScheduledTasks.Num() > 0)
 			{
-				PCGGraphExecutor::TScopeLock ChildScopeLock(TasksLock);
-				PCGGraphExecutor::TScopeLock OutputLock(TaskOutputsLock);
-
-				for (int32 ScheduledTaskIndex = ScheduledTasks.Num() - 1; ScheduledTaskIndex >= 0; --ScheduledTaskIndex)
+				int32 ScheduledTaskIndex = ScheduledTasks.Num() - 1;
+				while (ScheduledTaskIndex >= 0)
 				{
-					FPCGGraphScheduleTask& ScheduledTask = ScheduledTasks[ScheduledTaskIndex];
-
-					check(ScheduledTask.Tasks.Num() > 0);
-					// Push tasks to the primary task list & build successors map
-					for (FPCGGraphTask& Task : ScheduledTask.Tasks)
+					TArray<FPCGGraphTask, TInlineAllocator<16>> NewReadyTasks;
 					{
-						const FPCGTaskId TaskId = Task.NodeId;
+						PCGGraphExecutor::TScopeLock ChildScopeLock(TasksLock);
+						PCGGraphExecutor::TScopeLock OutputLock(TaskOutputsLock);
 
-						// TODO: review if it's actually possible for a task with inputs to be ready at this point - it seems very unlikely
-						bool bPushToReady = true;
+						// Every X ScheduledTasks, we want to check the frame budget and early out if we consumed it.
+						const int32 MaxScheduledTasksPerBudgetCheck = 32;
+						const int32 FirstScheduledTaskIndex = ScheduledTaskIndex;
+						while (ScheduledTaskIndex >= 0 && (FirstScheduledTaskIndex - ScheduledTaskIndex) < MaxScheduledTasksPerBudgetCheck)
 						{
-							for (const FPCGGraphTaskInput& Input : Task.Inputs)
+							FPCGGraphScheduleTask& ScheduledTask = ScheduledTasks[ScheduledTaskIndex];
+
+							check(ScheduledTask.Tasks.Num() > 0);
+							// Push tasks to the primary task list & build successors map
+							for (FPCGGraphTask& Task : ScheduledTask.Tasks)
 							{
-								if (!TaskOutputs.Contains(Input.TaskId))
+								const FPCGTaskId TaskId = Task.NodeId;
+
+								bool bPushToReady = true;
+								{
+									for (const FPCGGraphTaskInput& Input : Task.Inputs)
+									{
+										if (!TaskOutputs.Contains(Input.TaskId))
+										{
+											TaskSuccessors.FindOrAdd(Input.TaskId).Add(TaskId);
+											bPushToReady = false;
+										}
+									}
+								}
+
+								// Automatically push inputless/already satisfied tasks to the ready queue
+								if (bPushToReady)
+								{
+									NewReadyTasks.Add(MoveTemp(Task));
+									++NewReadyTaskCount;
+								}
+								else
 								{
-									TaskSuccessors.FindOrAdd(Input.TaskId).Add(TaskId);
-									bPushToReady = false;
+									Tasks.Add(TaskId, MoveTemp(Task));
 								}
+								bProcessedNewTasks = true;
 							}
-						}
 
-						// Automatically push inputless/already satisfied tasks to the ready queue
-						if (bPushToReady)
-						{
-							NewReadyTasks.Add(MoveTemp(Task));
-							++NewReadyTaskCount;
+							--ScheduledTaskIndex;
 						}
-						else
-						{
-							Tasks.Add(TaskId, MoveTemp(Task));
-						}
-						bProcessedNewTasks = true;
+					}
+
+					// Process ready tasks
+					for (FPCGGraphTask& ReadyTask : NewReadyTasks)
+					{
+						OnTaskInputsReady(ReadyTask, CachedResults, bIsInGameThread);
+						--NewReadyTaskCount;
+					}
+
+					if (FPlatformTime::Seconds() > EndTime)
+					{
+						break;
 					}
 				}
 
@@ -905,40 +958,22 @@ bool FPCGGraphExecutor::ProcessScheduledTasks()
 					PCGGraphExecutionLogging::LogGraphPostSchedule(Tasks, TaskSuccessors);
 				}
 
-				ScheduledTasks.Reset();
-			}
-
-			// Process ready tasks
-			for (FPCGGraphTask& ReadyTask : NewReadyTasks)
-			{
-				OnTaskInputsReady(ReadyTask, CachedResults, bIsInGameThread);
-				--NewReadyTaskCount;
+				if (ScheduledTaskIndex >= 0)
+				{
+					// [0, ScheduledTaskIndex] were not processed
+					ScheduledTasks.SetNum(ScheduledTaskIndex + 1);
+				}
+				else
+				{
+					ScheduledTasks.Reset();
+				}
 			}
 		}
 	}
 
 	ProcessCachedResults(MoveTemp(CachedResults));
-		
-	if (bIsInGameThread)
-	{
-		// Check every second to avoid locking too often for an error that shouldn't happen
-		double TimeInSeconds = FPlatformTime::Seconds();
-		if (TimeInSeconds - LastSchedulingErrorCheck > 1.0)
-		{
-			LastSchedulingErrorCheck = TimeInSeconds;
 
-			PCGGraphExecutor::TScopeLock ScopeLock(TasksLock);
-			PCGGraphExecutor::TScopeLock ChildScopeLock(LiveTasksLock);
-			PCGGraphExecutor::TScopeLock PausedLock(PausedTaskSuccessorsLock);
-			PCGGraphExecutor::TScopeLock CachingResultsScopeLock(CachingResultsLock);
-			// This is a safeguard to check if we're in a stuck state
-			if (CollectGCCachingResults.Num() == 0 && ReadyTasks.Num() == 0 && ActiveTasks.Num() == 0 && ActiveTasksGameThreadOnly.Num() == 0 && PausedTasks.Num() == 0 && Tasks.Num() > 0)
-			{
-				UE_LOG(LogPCG, Error, TEXT("PCG Graph executor error: tasks are in a deadlocked state. Will drop all tasks."));
-				ClearAllTasks();
-			}
-		}
-	}
+	CheckState(bIsInGameThread);
 
 	return bProcessedNewTasks;
 }
@@ -1447,6 +1482,9 @@ bool FPCGGraphExecutor::ExecuteScheduling(double EndTime, TSharedPtr<FPCGGraphAc
 		return false;
 	}
 
+	const float OutOfTickBudgetInSeconds = FMath::Max(1.0f, FPCGAsync::ConsoleVar::CVarAsyncOutOfTickBudgetInMilliseconds.GetValueOnAnyThread()) / 1000.f;
+	check(EndTime > 0.0);
+
 	bool bStateChanged = false;
 	const bool bIsInGameThread = IsInGameThread();
 	TRACE_CPUPROFILER_EVENT_SCOPE(FPCGGraphExecutor::ExecuteScheduling);
@@ -1481,6 +1519,7 @@ bool FPCGGraphExecutor::ExecuteScheduling(double EndTime, TSharedPtr<FPCGGraphAc
 				{
 					TSharedPtr<FPCGGraphActiveTask>& PausedTask = PausedTasks[PausedTaskIndex];
 
+					// It would be surprising that just looping through paused tasks would go over budget so no need to check time here
 					if (PausedTask->Context->bIsPaused)
 					{
 						continue;
@@ -1494,6 +1533,12 @@ bool FPCGGraphExecutor::ExecuteScheduling(double EndTime, TSharedPtr<FPCGGraphAc
 					PCGGraphExecutor::InsertToActiveTaskArrayNoLock(0, ActiveTaskArray, PausedTask);
 					PCGGraphExecutor::RemoveAtFromActiveTaskArrayNoLock(PausedTasks, PausedTaskIndex);
 					bStateChanged = true;
+
+					if (FPlatformTime::Seconds() > EndTime)
+					{
+						// Allow at least one loop of ReadyTasks
+						break;
+					}
 				}
 			}
 
@@ -1522,6 +1567,12 @@ bool FPCGGraphExecutor::ExecuteScheduling(double EndTime, TSharedPtr<FPCGGraphAc
 						CachedResults.Add(MoveTemp(CachedResult));
 						ReadyTasks.RemoveAtSwap(ReadyTaskIndex, EAllowShrinking::No);
 						bStateChanged = true;
+
+						if (FPlatformTime::Seconds() > EndTime)
+						{
+							// Make sure CachedResults are processed and that we process one loop of active tasks
+							break;
+						}
 						continue;
 					}
 
@@ -1547,10 +1598,16 @@ bool FPCGGraphExecutor::ExecuteScheduling(double EndTime, TSharedPtr<FPCGGraphAc
 				// Move the task up front if it needs to run on the main thread
 				ReadyTasks.RemoveAtSwap(ReadyTaskIndex, EAllowShrinking::No);
 				bStateChanged = true;
+
+				if (FPlatformTime::Seconds() > EndTime)
+				{
+					// Make sure CachedResults are processed and that we process one loop of active tasks
+					break;
+				}
 			}
 		}
 	}
-		
+
 	ProcessCachedResults(MoveTemp(CachedResults));
 
 	// Dispatch Tasks
@@ -1597,6 +1654,8 @@ bool FPCGGraphExecutor::ExecuteScheduling(double EndTime, TSharedPtr<FPCGGraphAc
 			// Consider the Main Thread task (+1)
 			if (MaxExecutingThreads > 0 && (MaxExecutingThreads+1) > FPCGGraphActiveTask::NumExecuting)
 			{
+				const double SchedulingBudgetInSeconds = GetTickBudgetInSeconds();
+
 				TRACE_CPUPROFILER_EVENT_SCOPE(FPCGGraphExecutor::ExecuteScheduling::LaunchTasks);
 				for (int32 ExecutionIndex = 0; ExecutionIndex < ActiveTasks.Num(); ++ExecutionIndex)
 				{
@@ -1630,7 +1689,7 @@ bool FPCGGraphExecutor::ExecuteScheduling(double EndTime, TSharedPtr<FPCGGraphAc
 
 						// Capture copy of SharedPtr so task stays valid while Async task is alive
 						ActiveTask->StartExecuting();
-						ActiveTask->ExecutingTask = UE::Tasks::Launch(UE_SOURCE_LOCATION, [this, ActiveTask]() -> bool
+						ActiveTask->ExecutingTask = UE::Tasks::Launch(UE_SOURCE_LOCATION, [this, ActiveTask, SchedulingBudgetInSeconds]() -> bool
 						{
 							TRACE_CPUPROFILER_EVENT_SCOPE(FPCGGraphExecutor::ExecuteAsyncTask);
 
@@ -1645,26 +1704,36 @@ bool FPCGGraphExecutor::ExecuteScheduling(double EndTime, TSharedPtr<FPCGGraphAc
 
 							if (bIsDone || bIsPaused)
 							{
+								const double EndTime = FPlatformTime::Seconds() + SchedulingBudgetInSeconds;
 								// If this task is paused it is possible it scheduled some new tasks, process those new tasks before calling ExecuteScheduling (so we are not one frame behind)
 								if (bIsPaused)
 								{
-									ProcessScheduledTasks();
+									ProcessScheduledTasks(EndTime);
 								}
 
-								// Budget is provided by CVarAsyncOutOfTickBudgetInMilliseconds
-								ExecuteScheduling(-1.0, nullptr,  /*bForceCheckPausedTasks=*/bIsPaused);
+								if (FPlatformTime::Seconds() < EndTime)
+								{
+									// Budget is provided by CVarAsyncOutOfTickBudgetInMilliseconds
+									ExecuteScheduling(EndTime, nullptr,  /*bForceCheckPausedTasks=*/bIsPaused);
+								}
 							}
 
 							return bIsDone;
 						});
 						bStateChanged = true;
+
+						if (FPlatformTime::Seconds() > EndTime)
+						{
+							// Make sure ProcessScheduledTasks runs
+							break;
+						}
 					}
 				}
 			}
 		}
 	}
 
-	bStateChanged |= ProcessScheduledTasks();
+	bStateChanged |= ProcessScheduledTasks(EndTime);
 
 	return bStateChanged;
 }
@@ -1726,18 +1795,28 @@ void FPCGGraphExecutor::Execute()
 	}
 #endif // WITH_EDITOR
 
-	ProcessScheduledTasks();
+	const double StartTime = FPlatformTime::Seconds();
+
+	const double ProcessScheduledTasksBudgetPercentage = 0.3;
+	const double ProcessScheduledTasksEndTime = StartTime + GetTickBudgetInSeconds() * ProcessScheduledTasksBudgetPercentage;
+
+	ProcessScheduledTasks(ProcessScheduledTasksEndTime);
 
+	const double ProcessScheduledTasksTime = FPlatformTime::Seconds() - StartTime;
+	
+	double NotificationTime = 0.0;
 #if WITH_EDITOR
 	// Update Notifications before capturing StartTime so that this call doesn't eat up task budget
 	UpdateGenerationNotification();
+
+	// Remove notifications from budget
+	NotificationTime = FPlatformTime::Seconds() - StartTime - ProcessScheduledTasksTime;
 #endif // WITH_EDITOR
 				
 	bool bContinueExecute = GetNonScheduledRemainingTaskCount() > 0;
 	bool bFirstLoop = true;
 
-	const double StartTime = FPlatformTime::Seconds(); 
-	const double EndTime = StartTime + GetTickBudgetInSeconds();
+	const double EndTime = StartTime + NotificationTime + GetTickBudgetInSeconds() - ProcessScheduledTasksTime;
 
 	while ((bFirstLoop || (FPlatformTime::Seconds() < EndTime)) && bContinueExecute)
 	{
@@ -3137,4 +3216,4 @@ FPCGElementPtr UPCGGridLinkageSettings::CreateElement() const
 	};
 
 	return MakeShared<PCGGraphExecutor::FPCGGridLinkageElement>(GridLinkageOperation, ContextAllocator, FromGrid, ToGrid, GenerationGrid, ResourceKey, UpstreamPin.Get());
-}
\ No newline at end of file
+}
diff --git a/Engine/Plugins/PCG/Source/PCG/Private/Graph/PCGGraphExecutor.h b/Engine/Plugins/PCG/Source/PCG/Private/Graph/PCGGraphExecutor.h
--- a/Engine/Plugins/PCG/Source/PCG/Private/Graph/PCGGraphExecutor.h
+++ b/Engine/Plugins/PCG/Source/PCG/Private/Graph/PCGGraphExecutor.h
@@ -209,9 +209,10 @@ private:
 		const TArray<FPCGTaskId>& TaskExecutionDependencies, const TArray<FPCGTaskId>& TaskDataDependencies);
 
 	void PostTaskExecute(TSharedPtr<FPCGGraphActiveTask> ActiveTask, bool bIsDone);
-	bool ProcessScheduledTasks();
+	bool ProcessScheduledTasks(double EndTime);
 	void ExecuteTasksEnded();
 	bool ExecuteScheduling(double EndTime, TSharedPtr<FPCGGraphActiveTask>* OutMainThreadTask = nullptr, bool bForceCheckPausedTasks = false);
+	void CheckState(bool bIsInGameThread);
 
 	TSet<IPCGGraphExecutionSource*> Cancel(TFunctionRef<bool(TWeakInterfacePtr<IPCGGraphExecutionSource>)> CancelFilter);
 	void ClearAllTasks();
