[PCG] Request runtime virtual textures to be primed before PCG component generation
* Add FPCGVirtualTexturePrimingInfo to inform VT priming
* Add a pass during the RuntimeGenScheduler tick to request priming for original components as necessary. This polls the OC Graph Instance parameters for params of type FPCGVirtualTexturePrimingInfo to inform priming.
* Knock out Landscape Normals RVT sampling behind a cvar 'pcg.GPU.VirtualTextures.EnableNormals'
* Disable virtual texture sampling in LS DI for non-RuntimeGen components, since the virtual textures are currently only primed for runtime generated components at this time.

diff --git a/Engine/Plugins/PCG/Source/PCG/Private/Compute/DataInterfaces/PCGLandscapeDataInterface.cpp b/Engine/Plugins/PCG/Source/PCG/Private/Compute/DataInterfaces/PCGLandscapeDataInterface.cpp
--- a/Engine/Plugins/PCG/Source/PCG/Private/Compute/DataInterfaces/PCGLandscapeDataInterface.cpp
+++ b/Engine/Plugins/PCG/Source/PCG/Private/Compute/DataInterfaces/PCGLandscapeDataInterface.cpp
@@ -34,6 +34,12 @@ namespace PCGLandscapeDataInterface
 		TEXT("pcg.GPU.UseLandscapeVTs"),
 		true,
 		TEXT("Allow the landscape data interface to use landscape virtual textures if available."));
+
+	/** Normal VTs may not match the actual landscape surface normals, so generally prefer computing normals from multiple height samples. */
+	static TAutoConsoleVariable<bool> CVarUseLandscapeNormalVTs(
+		TEXT("pcg.GPU.VirtualTextures.EnableNormals"),
+		false,
+		TEXT("Allow the landscape data interface to use landscape normal virtual textures if available."));
 }
 
 /** RenderResource used to hold textures generated by this DI, pulled from the collision geometry of the terrain. */
@@ -289,6 +295,10 @@ void UPCGLandscapeDataProvider::Initialize(const UComputeDataInterface* InDataIn
 
 	const UWorld* World = ExecutionSource->GetExecutionState().GetWorld();
 	const FBox ComponentBounds = ExecutionSource->GetExecutionState().GetBounds();
+	const UPCGComponent* SourceComponent = Cast<UPCGComponent>(ExecutionSource);
+
+	// Sampling of virtual textures is not supported outside of runtime generation at this time, because the RuntimeGenScheduler is responsible for priming the virtual textures.
+	bCanSampleVirtualTextures = SourceComponent ? SourceComponent->IsManagedByRuntimeGenSystem() : false;
 
 	auto TestLandscape = [World, &ComponentBounds](const ALandscape* InLandscape)
 	{
@@ -348,7 +358,7 @@ void UPCGLandscapeDataProvider::InitFromLandscape(ALandscape* InLandscape, const
 		return;
 	}
 
-	if (PCGLandscapeDataInterface::CVarUseLandscapeVTs.GetValueOnAnyThread())
+	if (bCanSampleVirtualTextures && PCGLandscapeDataInterface::CVarUseLandscapeVTs.GetValueOnAnyThread()) 
 	{
 		InitializeRuntimeVirtualTextures(InLandscape);
 	}
@@ -432,9 +442,12 @@ void UPCGLandscapeDataProvider::InitializeRuntimeVirtualTextures(ALandscape* InL
 		?  InLandscape->RuntimeVirtualTextures[HeightVirtualTextureIndex]
 		: nullptr;
 
-	NormalVirtualTexture = InLandscape->RuntimeVirtualTextures.IsValidIndex(NormalVirtualTextureIndex)
-		? InLandscape->RuntimeVirtualTextures[NormalVirtualTextureIndex]
-		: nullptr;
+	if (PCGLandscapeDataInterface::CVarUseLandscapeNormalVTs.GetValueOnAnyThread())
+	{
+		NormalVirtualTexture = InLandscape->RuntimeVirtualTextures.IsValidIndex(NormalVirtualTextureIndex)
+			? InLandscape->RuntimeVirtualTextures[NormalVirtualTextureIndex]
+			: nullptr;
+	}
 }
 
 void UPCGLandscapeDataProvider::InitializeFromLandscapeCollision(ALandscape* InLandscape, const FBox& Bounds)
@@ -554,6 +567,7 @@ void UPCGLandscapeDataProvider::Reset()
 	NormalVirtualTexture = nullptr;
 
 	Resource.Reset();
+	bCanSampleVirtualTextures = false;
 	bBaseColorSRGB = false;
 	BaseColorVirtualTextureUnpackType = EPCGVirtualTextureBaseColorUnpackType::None;
 	NormalVirtualTextureMode = ERuntimeVirtualTextureMaterialType::Count;
diff --git a/Engine/Plugins/PCG/Source/PCG/Private/Compute/DataInterfaces/PCGLandscapeDataInterface.h b/Engine/Plugins/PCG/Source/PCG/Private/Compute/DataInterfaces/PCGLandscapeDataInterface.h
--- a/Engine/Plugins/PCG/Source/PCG/Private/Compute/DataInterfaces/PCGLandscapeDataInterface.h
+++ b/Engine/Plugins/PCG/Source/PCG/Private/Compute/DataInterfaces/PCGLandscapeDataInterface.h
@@ -102,6 +102,7 @@ private:
 	TObjectPtr<const URuntimeVirtualTexture> NormalVirtualTexture = nullptr;
 
 	TSharedPtr<FPCGLandscapeResource> Resource;
+	bool bCanSampleVirtualTextures = false;
 	bool bBaseColorSRGB = false;
 	EPCGVirtualTextureBaseColorUnpackType BaseColorVirtualTextureUnpackType = EPCGVirtualTextureBaseColorUnpackType::None;
 	ERuntimeVirtualTextureMaterialType NormalVirtualTextureMode = ERuntimeVirtualTextureMaterialType::Count;
diff --git a/Engine/Plugins/PCG/Source/PCG/Private/RuntimeGen/PCGRuntimeGenScheduler.cpp b/Engine/Plugins/PCG/Source/PCG/Private/RuntimeGen/PCGRuntimeGenScheduler.cpp
--- a/Engine/Plugins/PCG/Source/PCG/Private/RuntimeGen/PCGRuntimeGenScheduler.cpp
+++ b/Engine/Plugins/PCG/Source/PCG/Private/RuntimeGen/PCGRuntimeGenScheduler.cpp
@@ -18,12 +18,15 @@
 #include "RuntimeGen/GenSources/PCGGenSourceComponent.h"
 #include "RuntimeGen/PCGGenSourceManager.h"
 
+#include "Components/RuntimeVirtualTextureComponent.h"
 #include "UObject/UObjectGlobals.h"
+#include "VT/RuntimeVirtualTexture.h"
 #include "WorldPartition/WorldPartition.h"
 #include "WorldPartition/WorldPartitionRuntimeCell.h"
 #include "WorldPartition/WorldPartitionSubsystem.h"
 
 #if WITH_EDITOR
+#include "DrawDebugHelpers.h"
 #include "EditorViewportClient.h"
 #endif
 
@@ -86,6 +89,18 @@ namespace PCGRuntimeGenSchedulerHelpers
 		TEXT("pcg.RuntimeGeneration.CheckWorldStreamingStatusForLocalComponents"),
 		false,
 		TEXT("Checks that the world is streamed in before triggering generation of local (partitioned) components."));
+
+	static TAutoConsoleVariable<bool> CVarEnableVirtualTexturePriming(
+		TEXT("pcg.VirtualTexturePriming.Enable"),
+		false,
+		TEXT("Enable priming of virtual textures for PCG Components which request it."));
+
+#if WITH_EDITOR
+	static TAutoConsoleVariable<bool> CVarDebugDrawTexturePrimingBounds(
+		TEXT("pcg.VirtualTexturePriming.DebugDrawTexturePrimingBounds"),
+		false,
+		TEXT("Draws debug boxes to indicate regions where PCG is requesting virtual texture priming."));
+#endif
 }
 
 FPCGGridDescriptor FPCGRuntimeGenScheduler::FGridGenerationKey::GetGridDescriptor() const
@@ -141,6 +156,13 @@ void FPCGRuntimeGenScheduler::Tick(APCGWorldActor* InPCGWorldActor)
 		}
 	}
 
+	if (PCGRuntimeGenSchedulerHelpers::CVarEnableVirtualTexturePriming.GetValueOnAnyThread())
+	{
+		// @todo_pcg: To support VT priming outside of RuntimeGen, this should probably move outside of the RGS tick, and be ticked directly by the PCGSubsystem.
+		// However, that would require also moving the GenSourceManager out of the RGS, 
+		TickRequestVirtualTexturePriming(GenSources);
+	}
+
 	// Initialize RuntimeGen PA pool if necessary. If PoolSize is 0, then we have not initialized the pool yet.
 	if (!GenSources.IsEmpty() || !GeneratedComponents.IsEmpty())
 	{
@@ -817,6 +839,183 @@ void FPCGRuntimeGenScheduler::TickScheduleGeneration(TMap<FGridGenerationKey, do
 	}
 }
 
+void FPCGRuntimeGenScheduler::TickRequestVirtualTexturePriming(const TSet<IPCGGenSourceBase*>& InGenSources)
+{
+	TRACE_CPUPROFILER_EVENT_SCOPE(FPCGRuntimeGenScheduler::TickRequestVirtualTexturePriming);
+
+	check(ActorAndComponentMapping);
+
+	TArray<const UPCGComponent*> OriginalComponents;
+
+	{
+		TRACE_CPUPROFILER_EVENT_SCOPE(GatherOriginalComponents);
+
+		TSet<UPCGComponent*> PartitionedComponents = ActorAndComponentMapping->PartitionedOctree.GetAllComponents();
+		TSet<UPCGComponent*> NonPartitionedComponents = ActorAndComponentMapping->NonPartitionedOctree.GetAllComponents();
+
+		OriginalComponents.Reserve(PartitionedComponents.Num() + NonPartitionedComponents.Num());
+
+		auto CanComponentRequestVirtualTexturePriming = [](const UPCGComponent* OriginalComponent)
+		{
+			return OriginalComponent && OriginalComponent->bActivated && OriginalComponent->IsManagedByRuntimeGenSystem() && OriginalComponent->GetGraphInstance();
+		};
+
+		for (UPCGComponent* OriginalComponent : PartitionedComponents)
+		{
+			if (CanComponentRequestVirtualTexturePriming(OriginalComponent))
+			{
+				OriginalComponents.Add(MoveTemp(OriginalComponent));
+			}
+		}
+
+		for (UPCGComponent* OriginalComponent : NonPartitionedComponents)
+		{
+			if (CanComponentRequestVirtualTexturePriming(OriginalComponent))
+			{
+				OriginalComponents.Add(MoveTemp(OriginalComponent));
+			}
+		}
+	}
+
+	if (OriginalComponents.IsEmpty())
+	{
+		return;
+	}
+
+	TMap<TSoftObjectPtr<URuntimeVirtualTexture>, TArray<URuntimeVirtualTextureComponent*>> VirtualTextureToComponents;
+
+	{
+		TRACE_CPUPROFILER_EVENT_SCOPE(FindVirtualTextureComponents);
+
+		// @todo_pcg: We could avoid polling for VT components every frame if they were registered somewhere instead.
+		TArray<UObject*> FoundComponents;
+		GetObjectsOfClass(URuntimeVirtualTextureComponent::StaticClass(), FoundComponents, /*bIncludeDerivedClasses=*/false);
+
+		for (UObject* FoundComponent : FoundComponents)
+		{
+			URuntimeVirtualTextureComponent* VirtualTextureComponent = Cast<URuntimeVirtualTextureComponent>(FoundComponent);
+			URuntimeVirtualTexture* VirtualTexture = VirtualTextureComponent ? VirtualTextureComponent->GetVirtualTexture() : nullptr;
+
+			if (VirtualTexture)
+			{
+				TArray<URuntimeVirtualTextureComponent*>& VirtualTextureComponents = VirtualTextureToComponents.FindOrAdd(VirtualTexture);
+				VirtualTextureComponents.Add(VirtualTextureComponent);
+			}
+		}
+	}
+
+	if (VirtualTextureToComponents.IsEmpty())
+	{
+		return;
+	}
+
+	for (const UPCGComponent* OriginalComponent : OriginalComponents)
+	{
+		TRACE_CPUPROFILER_EVENT_SCOPE(RequestVirtualTexturePrimingForComponent);
+		check(OriginalComponent);
+
+		// @todo_pcg: Instead of polling the PrimingInfos every tick, we could probably cache them and only update when the graph params change.
+		// Polling the PrimingInfos in this way is cheap for a single OC, but does not scale well if there are many OCs in the level.
+		TArray<FPCGVirtualTexturePrimingInfo*> PrimingInfos;
+
+		{
+			TRACE_CPUPROFILER_EVENT_SCOPE(PollVirtualTexturePrimingInfos);
+
+			const UPCGGraphInstance* GraphInstance = OriginalComponent->GetGraphInstance();
+			const FInstancedPropertyBag* UserParametersStruct = GraphInstance ? GraphInstance->GetUserParametersStruct() : nullptr;
+			const UPropertyBag* PropertyBag = UserParametersStruct ? UserParametersStruct->GetPropertyBagStruct() : nullptr;
+
+			if (!PropertyBag)
+			{
+				continue;
+			}
+
+			TConstArrayView<FPropertyBagPropertyDesc> PropertyDescs = PropertyBag->GetPropertyDescs();
+
+			for (const FPropertyBagPropertyDesc& PropertyDesc : PropertyDescs)
+			{
+				if (PropertyDesc.ValueType == EPropertyBagPropertyType::Struct && PropertyDesc.ValueTypeObject == TBaseStructure<FPCGVirtualTexturePrimingInfo>::Get())
+				{
+					TValueOrError<FPCGVirtualTexturePrimingInfo*, EPropertyBagResult> Property = UserParametersStruct->GetValueStruct<FPCGVirtualTexturePrimingInfo>(PropertyDesc);
+
+					if (Property.HasValue() && Property.GetValue())
+					{
+						PrimingInfos.Add(Property.GetValue());
+					}
+				}
+			}
+		}
+
+		const FBox OriginalComponentBounds = OriginalComponent->GetGridBounds();
+
+		for (const FPCGVirtualTexturePrimingInfo* PrimingInfo : PrimingInfos)
+		{
+			if (!PrimingInfo || !PrimingInfo->VirtualTexture)
+			{
+				continue;
+			}
+
+			TArray<URuntimeVirtualTextureComponent*>* VirtualTextureComponents = VirtualTextureToComponents.Find(PrimingInfo->VirtualTexture);
+
+			if (!VirtualTextureComponents)
+			{
+				continue;
+			}
+
+			for (URuntimeVirtualTextureComponent* VirtualTextureComponent : *VirtualTextureComponents)
+			{
+				check(VirtualTextureComponent);
+
+				const double PrimingRadius = OriginalComponent->GetGenerationRadiusFromGrid(PrimingInfo->Grid) + PCGHiGenGrid::GridToGridSize(PrimingInfo->Grid);
+
+				for (const IPCGGenSourceBase* GenSource : InGenSources)
+				{
+					const TOptional<FVector> GenSourcePositionOptional = GenSource->GetPosition();
+
+					if (!GenSourcePositionOptional.IsSet())
+					{
+						continue;
+					}
+
+					const FVector GenSourcePosition = GenSourcePositionOptional.GetValue();
+					FBox GenSourcePrimingBounds = FBox(FSphere(GenSourcePosition, PrimingRadius));
+					FVector ModifiedGenSourcePosition = GenSourcePosition;
+
+					if (OriginalComponent->Use2DGrid())
+					{
+						ModifiedGenSourcePosition.Z = OriginalComponentBounds.Min.Z;
+						GenSourcePrimingBounds.Min.Z = OriginalComponentBounds.Min.Z;
+						GenSourcePrimingBounds.Max.Z = OriginalComponentBounds.Max.Z;
+					}
+
+					const FBox PrimingBounds = OriginalComponentBounds.Overlap(GenSourcePrimingBounds);
+
+					if (PrimingBounds.GetVolume() > 0)
+					{
+						TRACE_CPUPROFILER_EVENT_SCOPE(RequestVirtualTexturePreload);
+
+						VirtualTextureComponent->RequestPreload(PrimingBounds, /*Level=*/0);
+
+#if WITH_EDITOR
+						if (PCGRuntimeGenSchedulerHelpers::CVarDebugDrawTexturePrimingBounds.GetValueOnAnyThread())
+						{
+							check(World);
+
+							DrawDebugBox(World,
+								PrimingBounds.GetCenter(),
+								PrimingBounds.GetExtent(),
+								FColor::Red,
+								/*bPersistentLines=*/false,
+								/*LifeTime=*/0.02f);
+						}
+#endif
+					}
+				}
+			}
+		}
+	}
+}
+
 void FPCGRuntimeGenScheduler::TickCVars(const APCGWorldActor* InPCGWorldActor)
 {
 	if (bActorFlushRequested && Subsystem && Subsystem->GetPCGWorldActor())
diff --git a/Engine/Plugins/PCG/Source/PCG/Public/RuntimeGen/PCGRuntimeGenScheduler.h b/Engine/Plugins/PCG/Source/PCG/Public/RuntimeGen/PCGRuntimeGenScheduler.h
--- a/Engine/Plugins/PCG/Source/PCG/Public/RuntimeGen/PCGRuntimeGenScheduler.h
+++ b/Engine/Plugins/PCG/Source/PCG/Public/RuntimeGen/PCGRuntimeGenScheduler.h
@@ -7,6 +7,8 @@
 #include "Grid/PCGGridDescriptor.h"
 #include "UObject/WeakObjectPtr.h"
 
+#include "PCGRuntimeGenScheduler.generated.h"
+
 class APCGPartitionActor;
 class APCGWorldActor;
 class FPCGActorAndComponentMapping;
@@ -16,6 +18,21 @@ class UPCGComponent;
 class UPCGSubsystem;
 class UWorld;
 
+/** Used to inform what virtual textures to prime and on what grids they need to be present. */
+USTRUCT(BlueprintType)
+struct FPCGVirtualTexturePrimingInfo
+{
+	GENERATED_BODY()
+
+	/** Virtual texture asset to be primed. */
+	UPROPERTY(EditAnywhere, Category = "")
+	TSoftObjectPtr<class URuntimeVirtualTexture> VirtualTexture;
+
+	/** Largest grid on which this virtual texture is sampled. */
+	UPROPERTY(EditAnywhere, Category = "")
+	EPCGHiGenGrid Grid = EPCGHiGenGrid::Grid32;
+};
+
 /**
  * The Runtime Generation Scheduler system handles the scheduling of PCG Components marked as GenerateAtRuntime.
  * It searches the level for Partitioned and Non-Partitioned components in range of the currently active
@@ -79,6 +96,9 @@ protected:
 	/** Schedule generation on components in priority order. */
 	void TickScheduleGeneration(TMap<FGridGenerationKey, double>& ComponentsToGenerate);
 
+	/** Request any required virtual textures to be primed within the necessary generation radius. */
+	void TickRequestVirtualTexturePriming(const TSet<IPCGGenSourceBase*>& InGenSources);
+
 	/** Detects changes in RuntimeGen CVars to keep the PA pool in a valid state. */
 	void TickCVars(const APCGWorldActor* InPCGWorldActor);
 
