[PCG] Soften check()s of execution source to guards.

#rnx

diff --git a/Engine/Plugins/PCG/Source/PCG/Private/Compute/Elements/PCGComputeGraphElement.cpp b/Engine/Plugins/PCG/Source/PCG/Private/Compute/Elements/PCGComputeGraphElement.cpp
--- a/Engine/Plugins/PCG/Source/PCG/Private/Compute/Elements/PCGComputeGraphElement.cpp
+++ b/Engine/Plugins/PCG/Source/PCG/Private/Compute/Elements/PCGComputeGraphElement.cpp
@@ -93,7 +93,12 @@ bool FPCGComputeGraphElement::ExecuteInternal(FPCGContext* InContext) const
 {
 	TRACE_CPUPROFILER_EVENT_SCOPE(FPCGComputeGraphElement::ExecuteInternal);
 	check(InContext);
-	check(InContext->ExecutionSource.IsValid());
+
+	if (!InContext->ExecutionSource.IsValid())
+	{
+		UE_LOG(LogPCG, Warning, TEXT("FPCGComputeGraphElement: Execution source lost, element execution halted."));
+		return true;
+	}
 
 	FPCGComputeGraphContext* Context = static_cast<FPCGComputeGraphContext*>(InContext);
 
diff --git a/Engine/Plugins/PCG/Source/PCG/Private/Elements/PCGCreatePoints.cpp b/Engine/Plugins/PCG/Source/PCG/Private/Elements/PCGCreatePoints.cpp
--- a/Engine/Plugins/PCG/Source/PCG/Private/Elements/PCGCreatePoints.cpp
+++ b/Engine/Plugins/PCG/Source/PCG/Private/Elements/PCGCreatePoints.cpp
@@ -67,8 +67,12 @@ FPCGElementPtr UPCGCreatePointsSettings::CreateElement() const
 bool FPCGCreatePointsElement::ExecuteInternal(FPCGContext* Context) const
 {
 	TRACE_CPUPROFILER_EVENT_SCOPE(FPCGCreatePointsElement::Execute);
-	
-	check(Context && Context->ExecutionSource.Get());
+	check(Context);
+
+	if (!Context->ExecutionSource.IsValid())
+	{
+		return true;
+	}
 
 	const UPCGCreatePointsSettings* Settings = Context->GetInputSettings<UPCGCreatePointsSettings>();
 	check(Settings);
diff --git a/Engine/Plugins/PCG/Source/PCG/Private/Elements/PCGCreatePointsGrid.cpp b/Engine/Plugins/PCG/Source/PCG/Private/Elements/PCGCreatePointsGrid.cpp
--- a/Engine/Plugins/PCG/Source/PCG/Private/Elements/PCGCreatePointsGrid.cpp
+++ b/Engine/Plugins/PCG/Source/PCG/Private/Elements/PCGCreatePointsGrid.cpp
@@ -51,8 +51,12 @@ FPCGElementPtr UPCGCreatePointsGridSettings::CreateElement() const
 bool FPCGCreatePointsGridElement::ExecuteInternal(FPCGContext* Context) const
 {
 	TRACE_CPUPROFILER_EVENT_SCOPE(FPCGCreatePointsGridElement::Execute);
-	
-	check(Context && Context->ExecutionSource.Get());
+	check(Context);
+
+	if (!Context->ExecutionSource.IsValid())
+	{
+		return true;
+	}
 
 	const UPCGCreatePointsGridSettings* Settings = Context->GetInputSettings<UPCGCreatePointsGridSettings>();
 	check(Settings);
