[PCG] Indirection Element, add missing call to InnerContext AddExtraStructReferencedObjects

#rnx

diff --git a/Engine/Plugins/PCG/Source/PCG/Private/Elements/PCGIndirectionElement.cpp b/Engine/Plugins/PCG/Source/PCG/Private/Elements/PCGIndirectionElement.cpp
--- a/Engine/Plugins/PCG/Source/PCG/Private/Elements/PCGIndirectionElement.cpp
+++ b/Engine/Plugins/PCG/Source/PCG/Private/Elements/PCGIndirectionElement.cpp
@@ -234,6 +234,11 @@ void FPCGIndirectionContext::AddExtraStructReferencedObjects(FReferenceCollector
 	{
 		Collector.AddReferencedObject(InnerSettings);
 	}
+
+	if (InnerContext)
+	{
+		InnerContext->AddExtraStructReferencedObjects(Collector);
+	}
 }
 
 FPCGContext* FPCGIndirectionElement::CreateContext()
