[PCG] Cleanup AddToRoot/RemoveFromRoot/TObjectStrongPtr in favor of proper AddReferencedObject calls

diff --git a/Engine/Plugins/PCG/Source/PCG/Private/Compute/Elements/PCGComputeGraphElement.cpp b/Engine/Plugins/PCG/Source/PCG/Private/Compute/Elements/PCGComputeGraphElement.cpp
--- a/Engine/Plugins/PCG/Source/PCG/Private/Compute/Elements/PCGComputeGraphElement.cpp
+++ b/Engine/Plugins/PCG/Source/PCG/Private/Compute/Elements/PCGComputeGraphElement.cpp
@@ -80,6 +80,9 @@ void FPCGComputeGraphContext::AddExtraStructReferencedObjects(FReferenceCollecto
 	{
 		Collector.AddReferencedObject(Provider);
 	}
+
+	Collector.AddReferencedObject(DataBinding);
+	Collector.AddReferencedObject(ComputeGraph);
 }
 
 #if WITH_EDITOR
@@ -147,7 +150,7 @@ bool FPCGComputeGraphElement::ExecuteInternal(FPCGContext* InContext) const
 
 		if (UPCGSubsystem* Subsystem = UWorld::GetSubsystem<UPCGSubsystem>(Source->GetExecutionState().GetWorld()))
 		{
-			Context->ComputeGraph.Reset(Subsystem->GetComputeGraph(DynamicSubgraph ? DynamicSubgraph : TopGraph, Context->GenerationGridSize, ComputeGraphIndex));
+			Context->ComputeGraph = Subsystem->GetComputeGraph(DynamicSubgraph ? DynamicSubgraph : TopGraph, Context->GenerationGridSize, ComputeGraphIndex);
 		}
 
 		if (!ensure(Context->ComputeGraph))
@@ -175,7 +178,7 @@ bool FPCGComputeGraphElement::ExecuteInternal(FPCGContext* InContext) const
 		ensure(!Context->DataBinding);
 
 		UPCGDataBinding* DataBinding = FPCGContext::NewObject_AnyThread<UPCGDataBinding>(Context);
-		Context->DataBinding.Reset(DataBinding);
+		Context->DataBinding = DataBinding;
 
 		DataBinding->Initialize(Context->ComputeGraph.Get(), Context);
 
diff --git a/Engine/Plugins/PCG/Source/PCG/Private/Compute/Elements/PCGComputeGraphElement.h b/Engine/Plugins/PCG/Source/PCG/Private/Compute/Elements/PCGComputeGraphElement.h
--- a/Engine/Plugins/PCG/Source/PCG/Private/Compute/Elements/PCGComputeGraphElement.h
+++ b/Engine/Plugins/PCG/Source/PCG/Private/Compute/Elements/PCGComputeGraphElement.h
@@ -42,8 +42,8 @@ protected:
 	virtual void AddExtraStructReferencedObjects(FReferenceCollector& Collector) override;
 
 public:
-	TStrongObjectPtr<UPCGDataBinding> DataBinding = nullptr;
-	TStrongObjectPtr<UPCGComputeGraph> ComputeGraph = nullptr;
+	TObjectPtr<UPCGDataBinding> DataBinding = nullptr;
+	TObjectPtr<UPCGComputeGraph> ComputeGraph = nullptr;
 
 	uint32 GenerationGridSize = PCGHiGenGrid::UninitializedGridSize();
 
diff --git a/Engine/Plugins/PCG/Source/PCG/Private/Compute/PCGDataBinding.cpp b/Engine/Plugins/PCG/Source/PCG/Private/Compute/PCGDataBinding.cpp
--- a/Engine/Plugins/PCG/Source/PCG/Private/Compute/PCGDataBinding.cpp
+++ b/Engine/Plugins/PCG/Source/PCG/Private/Compute/PCGDataBinding.cpp
@@ -182,16 +182,12 @@ void UPCGDataBinding::ReceiveDataFromGPU_GameThread(
 
 	if (!!(InExportMode & EPCGExportMode::DebugVisualization))
 	{
-		TStrongObjectPtr<UPCGProxyForGPUData> Proxy;
-		Proxy.Reset(InData);
-		DataToDebug.Add(FPCGDataToDebug{ .Proxy = Proxy, .ProducerSettings = InProducerSettings, .PinLabel = InPinLabel, .PinLabelAlias = InPinLabelAlias });
+		DataToDebug.Add(FPCGDataToDebug{ .Proxy = InData, .ProducerSettings = InProducerSettings, .PinLabel = InPinLabel, .PinLabelAlias = InPinLabelAlias });
 	}
 
 	if (!!(InExportMode & EPCGExportMode::Inspection))
 	{
-		TStrongObjectPtr<UPCGProxyForGPUData> Proxy;
-		Proxy.Reset(InData);
-		DataToInspect.Add(FPCGDataToDebug{ .Proxy = Proxy, .ProducerSettings = InProducerSettings, .PinLabel = InPinLabel, .PinLabelAlias = InPinLabelAlias });
+		DataToInspect.Add(FPCGDataToDebug{ .Proxy = InData, .ProducerSettings = InProducerSettings, .PinLabel = InPinLabel, .PinLabelAlias = InPinLabelAlias });
 	}
 }
 
diff --git a/Engine/Plugins/PCG/Source/PCG/Private/Compute/PCGDataBinding.h b/Engine/Plugins/PCG/Source/PCG/Private/Compute/PCGDataBinding.h
--- a/Engine/Plugins/PCG/Source/PCG/Private/Compute/PCGDataBinding.h
+++ b/Engine/Plugins/PCG/Source/PCG/Private/Compute/PCGDataBinding.h
@@ -49,11 +49,21 @@ struct FPCGSpawnerPrimitives
 	uint32 NumInstancesAllPrimitives = 0;
 };
 
+USTRUCT()
 struct FPCGDataToDebug
 {
-	TStrongObjectPtr<UPCGProxyForGPUData> Proxy;
+	GENERATED_BODY()
+
+	UPROPERTY()
+	TObjectPtr<UPCGProxyForGPUData> Proxy;
+	
+	UPROPERTY()
 	TWeakObjectPtr<const UPCGSettings> ProducerSettings;
+	
+	UPROPERTY()
 	FName PinLabel;
+	
+	UPROPERTY()
 	FName PinLabelAlias;
 };
 
@@ -133,9 +143,13 @@ public:
 	UPROPERTY()
 	TMap<TObjectPtr<const UPCGStaticMeshSpawnerKernel>, FPCGSpawnerPrimitives> MeshSpawnersToPrimitives;
 
+	UPROPERTY()	
 	TArray<TObjectPtr<const UPCGStaticMeshSpawnerKernel>> CompletedMeshSpawners;
 
+	UPROPERTY()
 	TArray<FPCGDataToDebug> DataToDebug;
+
+	UPROPERTY()
 	TArray<FPCGDataToDebug> DataToInspect;
 
 	TWeakPtr<FPCGContextHandle> ContextHandle;
diff --git a/Engine/Plugins/PCG/Source/PCG/Private/PCGSubsystem.cpp b/Engine/Plugins/PCG/Source/PCG/Private/PCGSubsystem.cpp
--- a/Engine/Plugins/PCG/Source/PCG/Private/PCGSubsystem.cpp
+++ b/Engine/Plugins/PCG/Source/PCG/Private/PCGSubsystem.cpp
@@ -155,6 +155,18 @@ UPCGSubsystem::UPCGSubsystem()
 {
 }
 
+void UPCGSubsystem::AddReferencedObjects(UObject* InThis, FReferenceCollector& Collector)
+{
+	Super::AddReferencedObjects(InThis, Collector);
+
+	UPCGSubsystem* PCGSubsystem = CastChecked<UPCGSubsystem>(InThis);
+	
+	if (FPCGRuntimeGenScheduler* RuntimeGenScheduler = PCGSubsystem->GetRuntimeGenScheduler())
+	{
+		RuntimeGenScheduler->AddReferencedObjects(Collector);
+	}
+}
+
 UPCGSubsystem* UPCGSubsystem::GetSubsystemForCurrentWorld()
 {
 	UWorld* World = nullptr;
diff --git a/Engine/Plugins/PCG/Source/PCG/Private/RuntimeGen/PCGGenSourceManager.cpp b/Engine/Plugins/PCG/Source/PCG/Private/RuntimeGen/PCGGenSourceManager.cpp
--- a/Engine/Plugins/PCG/Source/PCG/Private/RuntimeGen/PCGGenSourceManager.cpp
+++ b/Engine/Plugins/PCG/Source/PCG/Private/RuntimeGen/PCGGenSourceManager.cpp
@@ -32,7 +32,6 @@ FPCGGenSourceManager::FPCGGenSourceManager(const UWorld* InWorld)
 
 #if WITH_EDITOR
 	EditorCameraGenSource = NewObject<UPCGGenSourceEditorCamera>();
-	EditorCameraGenSource->AddToRoot();
 #endif
 }
 
@@ -40,21 +39,6 @@ FPCGGenSourceManager::~FPCGGenSourceManager()
 {
 	FGameModeEvents::GameModePostLoginEvent.RemoveAll(this);
 	FGameModeEvents::GameModeLogoutEvent.RemoveAll(this);
-
-#if WITH_EDITOR
-	if (EditorCameraGenSource && EditorCameraGenSource->IsRooted())
-	{
-		EditorCameraGenSource->RemoveFromRoot();
-	}
-#endif
-
-	for (UPCGGenSourceWPStreamingSource* WorldPartitionGenSource : WorldPartitionGenSources)
-	{
-		if (WorldPartitionGenSource && WorldPartitionGenSource->IsRooted())
-		{
-			WorldPartitionGenSource->RemoveFromRoot();
-		}
-	}
 }
 
 TSet<IPCGGenSourceBase*> FPCGGenSourceManager::GetAllGenSources(const APCGWorldActor* InPCGWorldActor)
@@ -103,28 +87,28 @@ TSet<IPCGGenSourceBase*> FPCGGenSourceManager::GetAllGenSources(const APCGWorldA
 	TSet<IPCGGenSourceBase*> AllGenSources;
 	AllGenSources.Reserve(RegisteredGenSources.Num() + RegisteredNamedGenSources.Num());
 
-	for (IPCGGenSourceBase* GenSource : RegisteredGenSources)
+	for (TScriptInterface<IPCGGenSourceBase> GenSource : RegisteredGenSources)
 	{
-		if (IsWorldLoaded(GenSource))
+		if (IsWorldLoaded(GenSource.GetInterface()))
 		{
-			AllGenSources.Add(GenSource);
+			AllGenSources.Add(GenSource.GetInterface());
 		}
 	}
 
-	for (const TPair<FName, IPCGGenSourceBase*>& NamedGenSource : RegisteredNamedGenSources)
+	for (const TPair<FName, TScriptInterface<IPCGGenSourceBase>>& NamedGenSource : RegisteredNamedGenSources)
 	{
-		if (IsWorldLoaded(NamedGenSource.Value))
+		if (IsWorldLoaded(NamedGenSource.Value.GetInterface()))
 		{
-			AllGenSources.Add(NamedGenSource.Value);
+			AllGenSources.Add(NamedGenSource.Value.GetInterface());
 		}
 	}
 
 	// Acquire a generation source for each active streaming source that doesn't have a generation source already (based on name collision test).
-	for (UPCGGenSourceWPStreamingSource* WPGenSource : WorldPartitionGenSources)
+	for (TObjectPtr<UPCGGenSourceWPStreamingSource> WPGenSource : WorldPartitionGenSources)
 	{
 		bool bAlreadyTracked = false;
 
-		if (IPCGGenSourceBase** FoundGenSource = (WPGenSource && WPGenSource->StreamingSource) ? RegisteredNamedGenSources.Find(WPGenSource->StreamingSource->Name) : nullptr)
+		if (TScriptInterface<IPCGGenSourceBase>* FoundGenSource = (WPGenSource && WPGenSource->StreamingSource) ? RegisteredNamedGenSources.Find(WPGenSource->StreamingSource->Name) : nullptr)
 		{
 			// We've seen during cinematics there are two generation sources with the same name but with different locations in the world, so match based on position too.
 			TOptional<FVector> WPGenSourcePosition = WPGenSource->GetPosition();
@@ -162,9 +146,10 @@ bool FPCGGenSourceManager::RegisterGenSource(IPCGGenSourceBase* InGenSource, FNa
 		}
 	}
 
+	TScriptInterface<IPCGGenSourceBase> InterfacePtr(Cast<UObject>(InGenSource));
 	if (InGenSourceName == NAME_None)
 	{
-		return RegisteredGenSources.Add(InGenSource).IsValidId();
+		return RegisteredGenSources.Add(InterfacePtr).IsValidId();
 	}
 	else
 	{
@@ -174,7 +159,7 @@ bool FPCGGenSourceManager::RegisterGenSource(IPCGGenSourceBase* InGenSource, FNa
 		}
 		else
 		{
-			RegisteredNamedGenSources.Add(InGenSourceName, InGenSource);
+			RegisteredNamedGenSources.Add(InGenSourceName, InterfacePtr);
 			return true;
 		}
 	}
@@ -192,13 +177,14 @@ bool FPCGGenSourceManager::UnregisterGenSource(const IPCGGenSourceBase* InGenSou
 
 	// Start by removing from the unnamed gen sources.
 	// @todo_pcg We could remove from named sources here too, but that might be limiting for more complex cases, BUT it might be less error-prone too.
-	if (RegisteredGenSources.Remove(InGenSource) > 0)
+	TScriptInterface<IPCGGenSourceBase> InterfacePtr(Cast<UObject>(const_cast<IPCGGenSourceBase*>(InGenSource)));
+	if (RegisteredGenSources.Remove(InterfacePtr) > 0)
 	{
 		return true;
 	}
 	else // Then try to remove from the named sources
 	{
-		if (const FName* GenSourceNamePtr = RegisteredNamedGenSources.FindKey(const_cast<IPCGGenSourceBase*>(InGenSource)))
+		if (const FName* GenSourceNamePtr = RegisteredNamedGenSources.FindKey(InterfacePtr))
 		{
 			FName GenSourceName = *GenSourceNamePtr;
 			return UnregisterGenSource(GenSourceName);
@@ -226,10 +212,7 @@ void FPCGGenSourceManager::OnGameModePostLogin(AGameModeBase* InGameMode, APlaye
 	UPCGGenSourcePlayer* GenSource = NewObject<UPCGGenSourcePlayer>();
 	GenSource->SetPlayerController(InPlayerController);
 
-	if (RegisterGenSource(GenSource, InPlayerController->GetFName()))
-	{
-		GenSource->AddToRoot();
-	}
+	RegisterGenSource(GenSource, InPlayerController->GetFName());
 }
 
 void FPCGGenSourceManager::OnGameModePostLogout(AGameModeBase* InGameMode, AController* InController)
@@ -241,15 +224,12 @@ void FPCGGenSourceManager::OnGameModePostLogout(AGameModeBase* InGameMode, ACont
 
 	ensure(IsInGameThread());
 
-	IPCGGenSourceBase* GenSource = nullptr;
+	TScriptInterface<IPCGGenSourceBase> GenSource = nullptr;
 	if (RegisteredNamedGenSources.RemoveAndCopyValue(InController->GetFName(), GenSource))
 	{
 		check(GenSource);
-		if (UPCGGenSourcePlayer* GenSourcePlayer = Cast<UPCGGenSourcePlayer>(GenSource))
-		{
-			GenSourcePlayer->RemoveFromRoot();
-		}
-		else
+		UPCGGenSourcePlayer* GenSourcePlayer = Cast<UPCGGenSourcePlayer>(GenSource.GetInterface());
+		if (!GenSourcePlayer)
 		{
 			// The gen source we removed isn't a player controller gen source, so let's put it back.
 			RegisteredNamedGenSources.Add(InController->GetFName(), GenSource);
@@ -282,8 +262,6 @@ void FPCGGenSourceManager::UpdatePerTickGenSources(const APCGWorldActor* InPCGWo
 	for (int I = 0; I < NumGenerationSourcesToAdd; ++I)
 	{
 		UPCGGenSourceWPStreamingSource* GenSource = WorldPartitionGenSources.Add_GetRef(NewObject<UPCGGenSourceWPStreamingSource>());
-		// New generation sources need to be rooted.
-		GenSource->AddToRoot();
 	}
 	
 	check((NumGenerationSourcesToUpdate + NumGenerationSourcesToAdd == 0) || WPStreamingSources);
@@ -295,12 +273,6 @@ void FPCGGenSourceManager::UpdatePerTickGenSources(const APCGWorldActor* InPCGWo
 			WorldPartitionGenSources[I]->StreamingSource = &(*WPStreamingSources)[I];
 		}
 	}
-	
-	for (int I = NumGenerationSourcesToUpdate + NumGenerationSourcesToAdd; I < NumGenerationSourcesToUpdate + NumGenerationSourcesToAdd + NumGenerationSourcesToRemove; ++I)
-	{
-		// Unused generation sources need to be forgotten & removed from root
-		WorldPartitionGenSources[I]->RemoveFromRoot();
-	}
 
 	WorldPartitionGenSources.SetNum(NumGenerationSourcesToUpdate + NumGenerationSourcesToAdd);
 
@@ -323,3 +295,21 @@ void FPCGGenSourceManager::UpdatePerTickGenSources(const APCGWorldActor* InPCGWo
 	}
 #endif
 }
+
+void FPCGGenSourceManager::AddReferencedObjects(FReferenceCollector& Collector)
+{
+#if WITH_EDITOR
+	Collector.AddReferencedObject(EditorCameraGenSource);
+#endif
+	Collector.AddReferencedObjects(WorldPartitionGenSources);
+
+	for (TScriptInterface<IPCGGenSourceBase>& GenSource : RegisteredGenSources)
+	{
+		GenSource.AddReferencedObjects(Collector);
+	}
+
+	for (auto& KeyValuePair : RegisteredNamedGenSources)
+	{
+		KeyValuePair.Value.AddReferencedObjects(Collector);
+	}
+}
diff --git a/Engine/Plugins/PCG/Source/PCG/Private/RuntimeGen/PCGRuntimeGenScheduler.cpp b/Engine/Plugins/PCG/Source/PCG/Private/RuntimeGen/PCGRuntimeGenScheduler.cpp
--- a/Engine/Plugins/PCG/Source/PCG/Private/RuntimeGen/PCGRuntimeGenScheduler.cpp
+++ b/Engine/Plugins/PCG/Source/PCG/Private/RuntimeGen/PCGRuntimeGenScheduler.cpp
@@ -1628,3 +1628,11 @@ void FPCGRuntimeGenScheduler::ResetPartitionActorPoolToSize(uint32 NewPoolSize)
 	PartitionActorPoolSize = 0;
 	AddPartitionActorPoolCount(NewPoolSize);
 }
+
+void FPCGRuntimeGenScheduler::AddReferencedObjects(FReferenceCollector& Collector)
+{
+	if (GenSourceManager)
+	{
+		GenSourceManager->AddReferencedObjects(Collector);
+	}
+}
\ No newline at end of file
diff --git a/Engine/Plugins/PCG/Source/PCG/Public/PCGSubsystem.h b/Engine/Plugins/PCG/Source/PCG/Public/PCGSubsystem.h
--- a/Engine/Plugins/PCG/Source/PCG/Public/PCGSubsystem.h
+++ b/Engine/Plugins/PCG/Source/PCG/Public/PCGSubsystem.h
@@ -62,6 +62,9 @@ public:
 
 	UPCGSubsystem();
 
+	/** Add UObject references for GC */
+	static void AddReferencedObjects(UObject* InThis, FReferenceCollector& Collector);
+
 	/** To be used when a PCG component can not have a world anymore, to unregister itself. */
 	static UPCGSubsystem* GetSubsystemForCurrentWorld();
 
diff --git a/Engine/Plugins/PCG/Source/PCG/Public/RuntimeGen/PCGGenSourceManager.h b/Engine/Plugins/PCG/Source/PCG/Public/RuntimeGen/PCGGenSourceManager.h
--- a/Engine/Plugins/PCG/Source/PCG/Public/RuntimeGen/PCGGenSourceManager.h
+++ b/Engine/Plugins/PCG/Source/PCG/Public/RuntimeGen/PCGGenSourceManager.h
@@ -38,6 +38,9 @@ public:
 	/** Removes a UPCGGenSource from being tracked by the GenSourceManaged on a name-basis. */
 	bool PCG_API UnregisterGenSource(FName InGenSourceName);
 
+	/** Add UObject references for GC */
+	void AddReferencedObjects(FReferenceCollector& Collector);
+
 protected:
 	void OnGameModePostLogin(AGameModeBase* InGameMode, APlayerController* InPlayerController);
 	void OnGameModePostLogout(AGameModeBase* InGameMode, AController* InController);
@@ -47,18 +50,18 @@ protected:
 
 protected:
 	/** Tracks named registered generation sources, such as UPCGGenSourcePlayer. */
-	TMap<FName, IPCGGenSourceBase*> RegisteredNamedGenSources;
+	TMap<FName, TScriptInterface<IPCGGenSourceBase>> RegisteredNamedGenSources;
 
 	/** Tracks unnamed registered generation sources, such as UPCGGenSourceComponent. */
-	TSet<IPCGGenSourceBase*> RegisteredGenSources;
+	TSet<TScriptInterface<IPCGGenSourceBase>> RegisteredGenSources;
 
 #if WITH_EDITORONLY_DATA
 	/** Tracks the active/main editor viewport client. This is refreshed every tick to keep a handle to whichever viewport is active. */
-	UPCGGenSourceEditorCamera* EditorCameraGenSource = nullptr;
+	TObjectPtr<UPCGGenSourceEditorCamera> EditorCameraGenSource = nullptr;
 #endif
 
 	/** Pool of GenSources dedicated to WorldPartition StreamingSources. This list grows as necessary and refreshes every tick. */
-	TArray<UPCGGenSourceWPStreamingSource*> WorldPartitionGenSources;
+	TArray<TObjectPtr<UPCGGenSourceWPStreamingSource>> WorldPartitionGenSources;
 
 	const UWorld* World = nullptr;
 	bool bDirty = false;
diff --git a/Engine/Plugins/PCG/Source/PCG/Public/RuntimeGen/PCGRuntimeGenScheduler.h b/Engine/Plugins/PCG/Source/PCG/Public/RuntimeGen/PCGRuntimeGenScheduler.h
--- a/Engine/Plugins/PCG/Source/PCG/Public/RuntimeGen/PCGRuntimeGenScheduler.h
+++ b/Engine/Plugins/PCG/Source/PCG/Public/RuntimeGen/PCGRuntimeGenScheduler.h
@@ -125,6 +125,9 @@ protected:
 	/** Destroy all pooled partition actors and rebuild with the NewPoolSize. */
 	void ResetPartitionActorPoolToSize(uint32 NewPoolSize);
 
+	/** Add UObject references for GC */
+	void AddReferencedObjects(FReferenceCollector& Collector);
+
 private:
 	/** Cleanup a local compnent from the Partition actor. */
 	void CleanupLocalComponent(APCGPartitionActor* PartitionActor, UPCGComponent* LocalComponent);
diff --git a/Engine/Plugins/PCG/Source/PCGEditor/Private/PCGEditor.cpp b/Engine/Plugins/PCG/Source/PCGEditor/Private/PCGEditor.cpp
--- a/Engine/Plugins/PCG/Source/PCGEditor/Private/PCGEditor.cpp
+++ b/Engine/Plugins/PCG/Source/PCGEditor/Private/PCGEditor.cpp
@@ -683,6 +683,14 @@ void FPCGEditor::UnregisterTabSpawners(const TSharedRef<class FTabManager>& InTa
 void FPCGEditor::AddReferencedObjects(FReferenceCollector& Collector)
 {
 	Collector.AddReferencedObject(PCGGraphBeingEdited);
+
+	for (TSharedPtr<SPCGEditorGraphAttributeListView> ALV : AttributesWidgets)
+	{
+		if (ALV)
+		{
+			ALV->AddReferencedObjects(Collector);
+		}
+	}
 }
 
 bool FPCGEditor::MatchesContext(const FTransactionContext& InContext, const TArray<TPair<UObject*, FTransactionObjectEvent>>& TransactionObjectContexts) const
diff --git a/Engine/Plugins/PCG/Source/PCGEditor/Private/Widgets/AssetEditorViewport/SPCGEditorViewport.cpp b/Engine/Plugins/PCG/Source/PCGEditor/Private/Widgets/AssetEditorViewport/SPCGEditorViewport.cpp
--- a/Engine/Plugins/PCG/Source/PCGEditor/Private/Widgets/AssetEditorViewport/SPCGEditorViewport.cpp
+++ b/Engine/Plugins/PCG/Source/PCGEditor/Private/Widgets/AssetEditorViewport/SPCGEditorViewport.cpp
@@ -163,12 +163,6 @@ void SPCGEditorViewport::SetupScene(const TArray<UObject*>& InResources, const F
 	if (Scene && SetupFunc.IsSet())
 	{
 		SetupFunc(Scene, InResources, ViewLocation, OrbitDistance, ManagedResources);
-
-		for (TObjectPtr<UObject> ManagedResource : ManagedResources)
-		{
-			// @todo_pcg: Ideally we should avoid using AddTo/RemoveFromRoot. Once we have a PCGEditorSubsystem, we should use AddReferencedObject instead.
-			ManagedResource->AddToRoot();
-		}
 	}
 
 	EditorViewportClient->SetViewLocation(ViewLocation);
@@ -195,10 +189,10 @@ TSharedRef<FEditorViewportClient> SPCGEditorViewport::MakeEditorViewportClient()
 
 void SPCGEditorViewport::ReleaseManagedResources()
 {
-	for (TObjectPtr<UObject> ManagedResource : ManagedResources)
-	{
-		ManagedResource->RemoveFromRoot();
-	}
-
 	ManagedResources.Empty();
 }
+
+void SPCGEditorViewport::AddReferencedObjects(FReferenceCollector& Collector)
+{
+	Collector.AddReferencedObjects(ManagedResources);
+}
diff --git a/Engine/Plugins/PCG/Source/PCGEditor/Private/Widgets/AssetEditorViewport/SPCGEditorViewport.h b/Engine/Plugins/PCG/Source/PCGEditor/Private/Widgets/AssetEditorViewport/SPCGEditorViewport.h
--- a/Engine/Plugins/PCG/Source/PCGEditor/Private/Widgets/AssetEditorViewport/SPCGEditorViewport.h
+++ b/Engine/Plugins/PCG/Source/PCGEditor/Private/Widgets/AssetEditorViewport/SPCGEditorViewport.h
@@ -21,6 +21,8 @@ public:
 	/** Rebuilds the scene using the provided resources and setup callback. */
 	void SetupScene(const TArray<UObject*>& InResources, const FPCGSetupSceneFunc& SetupFunc);
 
+	/** Add UObject references for GC */
+	void AddReferencedObjects(FReferenceCollector& Collector);
 protected:
 	//~ Begin SAssetEditorViewport Interface
 	virtual TSharedRef<FEditorViewportClient> MakeEditorViewportClient() override;
diff --git a/Engine/Plugins/PCG/Source/PCGEditor/Private/Widgets/SPCGEditorGraphAttributeListView.cpp b/Engine/Plugins/PCG/Source/PCGEditor/Private/Widgets/SPCGEditorGraphAttributeListView.cpp
--- a/Engine/Plugins/PCG/Source/PCGEditor/Private/Widgets/SPCGEditorGraphAttributeListView.cpp
+++ b/Engine/Plugins/PCG/Source/PCGEditor/Private/Widgets/SPCGEditorGraphAttributeListView.cpp
@@ -142,7 +142,7 @@ TSharedRef<SWidget> SPCGListViewItemRow::GenerateWidgetForColumn(const FName& Co
 	if (FPCGColumnData* PCGColumnData = SharedAttributeListView->PCGColumnData.Find(ColumnId))
 	{
 		// We have to make sure the data is still valid otherwise this will crash.
-		if (ensure(SharedAttributeListView->DataPtr.IsValid()) && PCGColumnData->DataAccessor.IsValid() && PCGColumnData->DataKeys.IsValid())
+		if (ensure(SharedAttributeListView->DataPtr) && PCGColumnData->DataAccessor.IsValid() && PCGColumnData->DataKeys.IsValid())
 		{
 			int32 Index = InternalItem->Index;
 			auto Callback = [&PCGColumnData, Index, &RowText, &bShouldComputeCellWidth, &bIsClickable] (auto Dummy)
@@ -315,7 +315,7 @@ SPCGEditorGraphAttributeListView::~SPCGEditorGraphAttributeListView()
 		PCGEditorPtr.Pin()->OnInspectedStackChangedDelegate.RemoveAll(this);
 	}
 
-	DataPtr.Reset();
+	DataPtr = nullptr;
 }
 
 void SPCGEditorGraphAttributeListView::Construct(const FArguments& InArgs, TSharedPtr<FPCGEditor> InPCGEditor)
@@ -614,7 +614,7 @@ void SPCGEditorGraphAttributeListView::Tick(const FGeometry& AllottedGeometry, c
 				InfoTextBlock->SetText(FText::Format(PCGEditorGraphAttributeListView::CrcLabelFormat, InfoTextBlock->GetText(), Crc.GetValue()));
 			}
 
-			if (DataPtr.IsValid() && DataPtr->HasCachedLastSelector())
+			if (DataPtr && DataPtr->HasCachedLastSelector())
 			{
 				const FText LastSelector = DataPtr->GetCachedLastSelector().GetDisplayText();
 				InfoTextBlock->SetText(FText::Format(PCGEditorGraphAttributeListView::LastLabelFormat, InfoTextBlock->GetText(), LastSelector));
@@ -629,7 +629,7 @@ void SPCGEditorGraphAttributeListView::RequestRefresh()
 {
 	bNeedsRefresh = true;
 	// When we ask for refresh, we should also release the hold on the visualized data
-	DataPtr.Reset();
+	DataPtr = nullptr;
 }
 
 TSharedRef<SHeaderRow> SPCGEditorGraphAttributeListView::CreateHeaderRowWidget() const
@@ -681,6 +681,16 @@ void SPCGEditorGraphAttributeListView::SetNodeBeingInspected(UPCGEditorGraphNode
 	RequestRefresh();
 }
 
+void SPCGEditorGraphAttributeListView::AddReferencedObjects(FReferenceCollector& Collector)
+{
+	Collector.AddReferencedObject(DataPtr);
+
+	if (ViewportWidget)
+	{
+		ViewportWidget->AddReferencedObjects(Collector);
+	}
+}
+
 void SPCGEditorGraphAttributeListView::OnGenerateUpdated(UPCGComponent* /*InPCGComponent*/)
 {
 	RequestRefresh();
@@ -760,7 +770,7 @@ void SPCGEditorGraphAttributeListView::RefreshAttributeList()
 	ListViewHeader->ClearColumns();
 	InfoTextBlock->SetText(FText::GetEmpty());
 	ViewportWidget->SetupScene({}, FPCGSetupSceneFunc());
-	DataPtr.Reset();
+	DataPtr = nullptr;
 
 	const FPCGDataCollection* InspectionData = GetInspectionData();
 	if (!InspectionData)
@@ -852,7 +862,7 @@ void SPCGEditorGraphAttributeListView::RefreshAttributeList()
 		{
 			// Poll next tick.
 			bNeedsRefresh = true;
-			DataPtr.Reset(DataToVisualize);
+			DataPtr = DataToVisualize;
 			return;
 		}
 
@@ -882,7 +892,7 @@ void SPCGEditorGraphAttributeListView::RefreshAttributeList()
 		FocusOnDataCallback.Reset();
 	}
 
-	DataPtr.Reset(DataToVisualize);
+	DataPtr = DataToVisualize;
 
 	ListView->SetItemsSource(&ListViewItems);
 	ListView->RequestListRefresh();
diff --git a/Engine/Plugins/PCG/Source/PCGEditor/Private/Widgets/SPCGEditorGraphAttributeListView.h b/Engine/Plugins/PCG/Source/PCGEditor/Private/Widgets/SPCGEditorGraphAttributeListView.h
--- a/Engine/Plugins/PCG/Source/PCGEditor/Private/Widgets/SPCGEditorGraphAttributeListView.h
+++ b/Engine/Plugins/PCG/Source/PCGEditor/Private/Widgets/SPCGEditorGraphAttributeListView.h
@@ -161,6 +161,8 @@ public:
 	TSharedPtr<SPCGEditorViewport> GetViewportWidget() const { return ViewportWidget; }
 	void SetViewportWidget(TSharedPtr<SPCGEditorViewport> InViewportWidget) { ViewportWidget = InViewportWidget; }
 
+	/** Add UObject references for GC */
+	void AddReferencedObjects(FReferenceCollector& Collector);
 private:
 	TSharedRef<SHeaderRow> CreateHeaderRowWidget() const;
 
@@ -324,7 +326,7 @@ private:
 	* it can store a strong reference to the data to keep it alive during inspection. 
 	* @todo_pcg: Replace this with AddReferencedObjects
 	*/
-	TStrongObjectPtr<const UPCGData> DataPtr = nullptr;
+	TObjectPtr<const UPCGData> DataPtr = nullptr;
 
 	/** Handles for any resources the data needs to load before visualization. */
 	TArray<TSharedPtr<FStreamableHandle>> LoadHandles;
