[PCG] Fix crashes in VT priming.

#rnx

diff --git a/Engine/Plugins/PCG/Source/PCG/Private/RuntimeGen/PCGRuntimeGenScheduler.cpp b/Engine/Plugins/PCG/Source/PCG/Private/RuntimeGen/PCGRuntimeGenScheduler.cpp
--- a/Engine/Plugins/PCG/Source/PCG/Private/RuntimeGen/PCGRuntimeGenScheduler.cpp
+++ b/Engine/Plugins/PCG/Source/PCG/Private/RuntimeGen/PCGRuntimeGenScheduler.cpp
@@ -889,17 +889,20 @@ void FPCGRuntimeGenScheduler::TickRequestVirtualTexturePriming(const TSet<IPCGGe
 
 		// @todo_pcg: We could avoid polling for VT components every frame if they were registered somewhere instead.
 		TArray<UObject*> FoundComponents;
-		GetObjectsOfClass(URuntimeVirtualTextureComponent::StaticClass(), FoundComponents, /*bIncludeDerivedClasses=*/false);
+		GetObjectsOfClass(URuntimeVirtualTextureComponent::StaticClass(), FoundComponents, /*bIncludeDerivedClasses=*/false, RF_ClassDefaultObject, EInternalObjectFlags::Garbage);
 
 		for (UObject* FoundComponent : FoundComponents)
 		{
-			URuntimeVirtualTextureComponent* VirtualTextureComponent = Cast<URuntimeVirtualTextureComponent>(FoundComponent);
-			URuntimeVirtualTexture* VirtualTexture = VirtualTextureComponent ? VirtualTextureComponent->GetVirtualTexture() : nullptr;
-
-			if (VirtualTexture)
+			if (FoundComponent && FoundComponent->GetWorld() == World)
 			{
-				TArray<URuntimeVirtualTextureComponent*>& VirtualTextureComponents = VirtualTextureToComponents.FindOrAdd(VirtualTexture);
-				VirtualTextureComponents.Add(VirtualTextureComponent);
+				URuntimeVirtualTextureComponent* VirtualTextureComponent = Cast<URuntimeVirtualTextureComponent>(FoundComponent);
+				URuntimeVirtualTexture* VirtualTexture = VirtualTextureComponent ? VirtualTextureComponent->GetVirtualTexture() : nullptr;
+
+				if (VirtualTexture)
+				{
+					TArray<URuntimeVirtualTextureComponent*>& VirtualTextureComponents = VirtualTextureToComponents.FindOrAdd(VirtualTexture);
+					VirtualTextureComponents.Add(VirtualTextureComponent);
+				}
 			}
 		}
 	}
@@ -1635,4 +1638,4 @@ void FPCGRuntimeGenScheduler::AddReferencedObjects(FReferenceCollector& Collecto
 	{
 		GenSourceManager->AddReferencedObjects(Collector);
 	}
-}
\ No newline at end of file
+}
