[PCG] Transient resources tracking through PCG data collections. Currently accessible through a cvar. Will break set members in struct on PCG Tagged Data.

diff --git a/Engine/Plugins/PCG/Source/PCG/Private/Data/PCGDataPtrWrapper.cpp b/Engine/Plugins/PCG/Source/PCG/Private/Data/PCGDataPtrWrapper.cpp
new file mode 100644
--- /dev/null
+++ b/Engine/Plugins/PCG/Source/PCG/Private/Data/PCGDataPtrWrapper.cpp
@@ -0,0 +1,140 @@
+// Copyright Epic Games, Inc. All Rights Reserved.
+
+#include "Data/PCGDataPtrWrapper.h"
+#include "PCGData.h"
+
+FPCGDataPtrWrapper::~FPCGDataPtrWrapper()
+{
+	DecRefCount();
+}
+
+FPCGDataPtrWrapper::FPCGDataPtrWrapper(const FPCGDataPtrWrapper& InWrapper)
+{
+	Data = InWrapper.Data;
+	WeakData = InWrapper.Data;
+	IncRefCount();
+}
+
+FPCGDataPtrWrapper::FPCGDataPtrWrapper(FPCGDataPtrWrapper&& InWrapper)
+{
+	Data = InWrapper.Data;
+	WeakData = InWrapper.WeakData;
+	InWrapper.Data = nullptr;
+	InWrapper.WeakData = nullptr;
+}
+
+FPCGDataPtrWrapper& FPCGDataPtrWrapper::operator=(const FPCGDataPtrWrapper& InWrapper)
+{
+	TWeakObjectPtr<const UPCGData> PreviousWeakData = WeakData;
+
+	Data = InWrapper.Data;
+	WeakData = InWrapper.WeakData;
+
+	IncRefCount();
+	DecRefCount(PreviousWeakData);
+	return *this;
+}
+
+FPCGDataPtrWrapper& FPCGDataPtrWrapper::operator=(FPCGDataPtrWrapper&& InWrapper)
+{
+	TWeakObjectPtr<const UPCGData> PreviousWeakData = WeakData;
+
+	Data = InWrapper.Data;
+	WeakData = InWrapper.WeakData;
+
+	InWrapper.Data = nullptr;
+	InWrapper.WeakData = nullptr;
+
+	DecRefCount(PreviousWeakData);
+	return *this;
+}
+
+FPCGDataPtrWrapper::FPCGDataPtrWrapper(const UPCGData* InData)
+{
+	Data = InData;
+	WeakData = InData;
+
+	IncRefCount();
+}
+
+FPCGDataPtrWrapper& FPCGDataPtrWrapper::operator=(const UPCGData* InData)
+{
+	TWeakObjectPtr<const UPCGData> PreviousWeakData = WeakData;
+
+	Data = InData;
+	WeakData = InData;
+
+	IncRefCount();
+	DecRefCount(PreviousWeakData);
+	return *this;
+}
+
+void FPCGDataPtrWrapper::IncRefCount(const UPCGData* InData)
+{
+	if (InData)
+	{
+		InData->IncCollectionRefCount();
+	}
+}
+
+void FPCGDataPtrWrapper::IncRefCount() 
+{ 
+	IncRefCount(Data); 
+}
+
+void FPCGDataPtrWrapper::DecRefCount(const TWeakObjectPtr<const UPCGData>& InWeakData)
+{
+	if (TStrongObjectPtr<const UPCGData> InData = InWeakData.Pin())
+	{
+		InData->DecCollectionRefCount();
+	}
+}
+
+void FPCGDataPtrWrapper::DecRefCount() 
+{ 
+	DecRefCount(WeakData);
+}
+
+bool FPCGDataPtrWrapper::Serialize(FArchive& Ar)
+{
+	TWeakObjectPtr<const UPCGData> WeakDataBeforeSerialization = WeakData;
+	Ar << Data;
+
+	if (Ar.IsLoading())
+	{
+		WeakData = Data;
+		IncRefCount(Data);
+		DecRefCount(WeakDataBeforeSerialization);
+	}
+
+	return true;
+}
+
+bool FPCGDataPtrWrapper::ImportTextItem(const TCHAR*& Buffer, int32 PortFlags, UObject* Parent, FOutputDevice* ErrorText)
+{
+	TWeakObjectPtr<const UPCGData> WeakDataBeforeImport = WeakData;
+	StaticStruct()->ImportText(Buffer, this, Parent, PortFlags, ErrorText, StaticStruct()->GetName(), false);
+
+	WeakData = Data;
+	IncRefCount(Data);
+	DecRefCount(WeakDataBeforeImport);
+
+	return true;
+}
+
+bool FPCGDataPtrWrapper::SerializeFromMismatchedTag(const FPropertyTag& Tag, FStructuredArchive::FSlot Slot)
+{
+	if (Tag.GetType().GetName() == NAME_ObjectProperty) // "Data"
+	{
+		TObjectPtr<const UPCGData> DataPtr;
+		Slot << DataPtr;
+		Data = DataPtr;
+		WeakData = DataPtr;
+		IncRefCount();
+		return true;
+	}
+	else
+	{
+		return false;
+	}
+}
\ No newline at end of file
diff --git a/Engine/Plugins/PCG/Source/PCG/Private/Elements/PCGExecuteBlueprint.cpp b/Engine/Plugins/PCG/Source/PCG/Private/Elements/PCGExecuteBlueprint.cpp
--- a/Engine/Plugins/PCG/Source/PCG/Private/Elements/PCGExecuteBlueprint.cpp
+++ b/Engine/Plugins/PCG/Source/PCG/Private/Elements/PCGExecuteBlueprint.cpp
@@ -812,6 +812,10 @@ bool FPCGExecuteBlueprintElement::ExecuteInternal(FPCGContext* InContext) const
 		TSharedPtr<FPCGContextHandle> Handle = Context->Handle;
 		Context->Handle = nullptr;
 
+		// Similarly we don't want to register to GC-protection on release.
+		TWeakPtr<FPCGGraphExecutor> Executor = Context->GraphExecutor;
+		Context->GraphExecutor = nullptr;
+
 		// When running outside of main thread make sure GC can't run (BP Nodes might create objects which can't happen while GC runs)
 		if(!IsInGameThread())
 		{
@@ -824,7 +828,8 @@ bool FPCGExecuteBlueprintElement::ExecuteInternal(FPCGContext* InContext) const
 			Context->BlueprintElementInstance->ExecuteWithContext(*Context, Context->InputData, Context->OutputData);
 		}
 
-		// Restore handle
+		// Restore handle & executor
+		Context->GraphExecutor = Executor;
 		Context->Handle = Handle;
 		Context->BlueprintElementInstance->SetCurrentContext(nullptr);
 	}
diff --git a/Engine/Plugins/PCG/Source/PCG/Private/Graph/PCGGraphExecutor.cpp b/Engine/Plugins/PCG/Source/PCG/Private/Graph/PCGGraphExecutor.cpp
--- a/Engine/Plugins/PCG/Source/PCG/Private/Graph/PCGGraphExecutor.cpp
+++ b/Engine/Plugins/PCG/Source/PCG/Private/Graph/PCGGraphExecutor.cpp
@@ -860,7 +860,7 @@ void FPCGGraphExecutor::CheckState(bool bIsInGameThread)
 			PCGGraphExecutor::TScopeLock PausedLock(PausedTaskSuccessorsLock);
 			PCGGraphExecutor::TScopeLock CachingResultsScopeLock(CachingResultsLock);
 			// This is a safeguard to check if we're in a stuck state
-			if (CollectGCCachingResults.Num() == 0 && ReadyTasks.Num() == 0 && ActiveTasks.Num() == 0 && ActiveTasksGameThreadOnly.Num() == 0 && PausedTasks.Num() == 0 && Tasks.Num() > 0)
+			if (CachingResultsForGC.Num() == 0 && ReadyTasks.Num() == 0 && ActiveTasks.Num() == 0 && ActiveTasksGameThreadOnly.Num() == 0 && PausedTasks.Num() == 0 && Tasks.Num() > 0)
 			{
 				UE_LOG(LogPCG, Error, TEXT("PCG Graph executor error: tasks are in a deadlocked state. Will drop all tasks."));
 				ClearAllTasks();
@@ -1182,7 +1182,9 @@ void FPCGGraphExecutor::PostTaskExecute(TSharedPtr<FPCGGraphActiveTask> ActiveTa
 		// If the task is a post execute, then we can safely clear the data after getting it from the results.
 		bool bNeedsManualClear = (ActiveTask.Element == GraphCompiler.GetSharedTrivialPostGraphElement());
 
-		if (!bNeedsManualClear)
+		// Data holding transient resources needed originally to be released late (at the end of the generation)
+		// because we had no way to assess input usage. With the cvar here (if it's true), it's not the case anymore.
+		if (!bNeedsManualClear && !PCGSystemSwitches::CVarReleaseTransientResourcesEarly.GetValueOnAnyThread())
 		{
 			for (const FPCGTaggedData& TaggedData : ActiveTask.Context->OutputData.TaggedData)
 			{
@@ -1211,6 +1213,7 @@ void FPCGGraphExecutor::PostTaskExecute(TSharedPtr<FPCGGraphActiveTask> ActiveTa
 		{
 			check(ActiveTaskArray[TaskIndex] == ActiveTaskPtr);
 			// Remove current active task from list
+			AddReleasedContextForGC(ActiveTaskArray[TaskIndex]->Context.Get());
 			PCGGraphExecutor::RemoveAtFromActiveTaskArrayNoLock(ActiveTaskArray, TaskIndex);
 		}
 
@@ -1346,8 +1349,8 @@ void FPCGGraphExecutor::PrepareForExecute(FPCGGraphTask& Task, FCachedResult*& O
 		bResultAlreadyInCache = (CacheStatus == EPCGCachingStatus::Cached);
 		if (bResultAlreadyInCache)
 		{
-			check(!CollectGCCachingResults.Contains(Task.NodeId));
-			TUniquePtr<FCachedResult>& CachingResultPtr = CollectGCCachingResults.Add(Task.NodeId, MakeUnique<FCachedResult>());
+			check(!CachingResultsForGC.Contains(Task.NodeId));
+			TUniquePtr<FCachedResult>& CachingResultPtr = CachingResultsForGC.Add(Task.NodeId, MakeUnique<FCachedResult>());
 			OutCachedResult = CachingResultPtr.Get();
 			OutCachedResult->TaskId = Task.NodeId;
 			OutCachedResult->Output = MoveTemp(LocalCachedResult.Output);
@@ -1400,7 +1403,7 @@ void FPCGGraphExecutor::PrepareForExecute(FPCGGraphTask& Task, FCachedResult*& O
 	
 	{
 		Task.Context = Task.Element->Initialize(FPCGInitializeElementParams(&Task.TaskInput, Task.ExecutionSource.Get(), Task.Node));
-		Task.Context->GraphExecutor = this;
+		Task.Context->GraphExecutor = this->AsWeak();
 
 		// Skip PostLoad only when not on GameThread
 		// The reason is that some nodes like SpawnActor / CreateTargetActor have TemplateActor subobjects that need PostLoad calls to function properly
@@ -1431,7 +1434,7 @@ void FPCGGraphExecutor::PrepareForExecute(FPCGGraphTask& Task, FCachedResult*& O
 		PCGGraphExecutor::TScopeLock ScopeLock(CachingResultsLock);
 		Task.Context->OutputData = MoveTemp(OutCachedResult->Output);
 		OutCachedResult = nullptr;
-		CollectGCCachingResults.Remove(Task.NodeId);
+		CachingResultsForGC.Remove(Task.NodeId);
 	}
 #endif
 }
@@ -1974,8 +1977,8 @@ TArray<FPCGTaskId> FPCGGraphExecutor::ProcessCachedResultsInternal(TArray<FCache
 		PCGGraphExecutor::TScopeLock ScopeLock(CachingResultsLock);
 		for (const FCachedResult* CachedResult : CachedResults)
 		{
-			check(CollectGCCachingResults.Contains(CachedResult->TaskId));
-			CollectGCCachingResults.Remove(CachedResult->TaskId);
+			check(CachingResultsForGC.Contains(CachedResult->TaskId));
+			CachingResultsForGC.Remove(CachedResult->TaskId);
 		}
 
 		// Next Scheduling call needs to check if this task completion unblocked some paused task(s)
@@ -2499,6 +2502,18 @@ void FPCGGraphExecutor::SendInactivePinNotification(const UPCGNode* InNode, cons
 }
 #endif
 
+void FPCGGraphExecutor::AddReleasedContextForGC(FPCGContext* InContext)
+{
+	PCGGraphExecutor::TScopeLock ScopeLock(ReleasedContextLock);
+	ReleasedContextsForGC.Add(InContext);
+}
+
+void FPCGGraphExecutor::RemoveReleasedContextForGC(FPCGContext* InContext)
+{
+	PCGGraphExecutor::TScopeLock ScopeLock(ReleasedContextLock);
+	ensure(ReleasedContextsForGC.Remove(InContext) > 0);
+}
+
 void FPCGGraphExecutor::AddReferencedObjects(FReferenceCollector& Collector)
 {
 	TRACE_CPUPROFILER_EVENT_SCOPE(FPCGGraphExecutor::AddReferencedObjects);
@@ -2551,7 +2566,7 @@ void FPCGGraphExecutor::AddReferencedObjects(FReferenceCollector& Collector)
 
 	{
 		PCGGraphExecutor::TScopeLock ScopeLock(CachingResultsLock);
-		Algo::ForEach(CollectGCCachingResults, [&Collector](auto& CachingResult)
+		Algo::ForEach(CachingResultsForGC, [&Collector](auto& CachingResult)
 		{
 			CachingResult.Value->Output.AddReferences(Collector);
 		});
@@ -2578,6 +2593,17 @@ void FPCGGraphExecutor::AddReferencedObjects(FReferenceCollector& Collector)
 		}
 	}
 
+	{
+		PCGGraphExecutor::TScopeLock ScopeLock(ReleasedContextLock);
+		Algo::ForEach(ReleasedContextsForGC, [&Collector](FPCGContext* InContext)
+		{
+			if(ensure(InContext))
+			{
+				InContext->AddStructReferencedObjects(Collector);
+			}
+		});
+	}
+
 	GraphCompiler.AddReferencedObjects(Collector);
 }
 
@@ -3216,4 +3242,4 @@ FPCGElementPtr UPCGGridLinkageSettings::CreateElement() const
 	};
 
 	return MakeShared<PCGGraphExecutor::FPCGGridLinkageElement>(GridLinkageOperation, ContextAllocator, FromGrid, ToGrid, GenerationGrid, ResourceKey, UpstreamPin.Get());
-}
+}
\ No newline at end of file
diff --git a/Engine/Plugins/PCG/Source/PCG/Private/Graph/PCGGraphExecutor.h b/Engine/Plugins/PCG/Source/PCG/Private/Graph/PCGGraphExecutor.h
--- a/Engine/Plugins/PCG/Source/PCG/Private/Graph/PCGGraphExecutor.h
+++ b/Engine/Plugins/PCG/Source/PCG/Private/Graph/PCGGraphExecutor.h
@@ -42,7 +42,7 @@ namespace PCGGraphExecutor
 	void ClearAsyncFlags(TSet<TObjectPtr<UObject>>& AsyncObjects);
 }
 
-class FPCGGraphExecutor : public FGCObject
+class FPCGGraphExecutor : public FGCObject, public TSharedFromThis<FPCGGraphExecutor>
 {
 public:
 	// Default constructor used by unittests
@@ -146,6 +146,9 @@ public:
 	virtual FString GetReferencerName() const override { return TEXT("FPCGGraphExecutor"); }
 	//~End FGCObject interface
 
+	void AddReleasedContextForGC(FPCGContext* InContext);
+	void RemoveReleasedContextForGC(FPCGContext* InContext);
+
 private:
 	friend UPCGSubsystem;
 
@@ -329,7 +332,11 @@ private:
 	/** Lock level leaf */
 	FPCGLock CachingResultsLock;
 	// Used to keep GC references to in flight caching results (not yet stored to output and might not be in cache anymore)
-	TMap<FPCGTaskId, TUniquePtr<FCachedResult>> CollectGCCachingResults;
+	TMap<FPCGTaskId, TUniquePtr<FCachedResult>> CachingResultsForGC;
+
+	/** Lock level leaf */
+	FPCGLock ReleasedContextLock;
+	TSet<FPCGContext*> ReleasedContextsForGC;
 
 	/** Lock level leaf */
 	FPCGLock TaskOutputsLock;
diff --git a/Engine/Plugins/PCG/Source/PCG/Private/PCGCommon.cpp b/Engine/Plugins/PCG/Source/PCG/Private/PCGCommon.cpp
--- a/Engine/Plugins/PCG/Source/PCG/Private/PCGCommon.cpp
+++ b/Engine/Plugins/PCG/Source/PCG/Private/PCGCommon.cpp
@@ -102,6 +102,11 @@ namespace PCGSystemSwitches
 		false,
 		TEXT("Forces all subgraph executions to be dynamic. Performance warning, also requires to flush all compiled graphs."));
 #endif
+
+	TAutoConsoleVariable<bool> CVarReleaseTransientResourcesEarly(
+		TEXT("pcg.Experimental.ReleaseTransientResourcesEarly"),
+		false,
+		TEXT("Aggressively release transient (esp. GPU) data when possible."));
 }
 
 namespace PCGHiGenGrid
diff --git a/Engine/Plugins/PCG/Source/PCG/Private/PCGComponent.cpp b/Engine/Plugins/PCG/Source/PCG/Private/PCGComponent.cpp
--- a/Engine/Plugins/PCG/Source/PCG/Private/PCGComponent.cpp
+++ b/Engine/Plugins/PCG/Source/PCG/Private/PCGComponent.cpp
@@ -2100,11 +2100,12 @@ bool UPCGComponent::CanEditChange(const FProperty* InProperty) const
 
 void UPCGComponent::BeginDestroy()
 {
+	GeneratedGraphOutput.Reset();
+	PerPinGeneratedOutput.Reset();
+
 #if WITH_EDITOR
 	// BeginDestroy will end up calling OnComponentDestroyed and in this specific case we don't want to rename our Outputs to transient package.
 	// Known case where GeneratedGraphOutput can be non empty here is the editor exit purge.
-	GeneratedGraphOutput.Reset();
-
 	if (GraphInstance)
 	{
 		GraphInstance->OnGraphChangedDelegate.RemoveAll(this);
diff --git a/Engine/Plugins/PCG/Source/PCG/Private/PCGContext.cpp b/Engine/Plugins/PCG/Source/PCG/Private/PCGContext.cpp
--- a/Engine/Plugins/PCG/Source/PCG/Private/PCGContext.cpp
+++ b/Engine/Plugins/PCG/Source/PCG/Private/PCGContext.cpp
@@ -48,6 +48,15 @@ namespace PCGContextHelpers
 
 FPCGContext::~FPCGContext()
 {
+	// Release the collection data before removing from GC gathering list from the executor to make sure there's no gap in coverage here
+	InputData.Reset();
+	OutputData.Reset();
+
+	if (TSharedPtr<FPCGGraphExecutor> Executor = GraphExecutor.Pin())
+	{
+		Executor->RemoveReleasedContextForGC(this);
+	}
+
 	// This should not happen, might be a missing call to FPCGContext::Release
 	// Clear Handle context pointer so that we don't double delete the FPCGContext
 	ensureMsgf(!Handle, TEXT("FPCGContext deleted without first setting the Handle to nullptr, probably a delete that should be replaced by a FPCGContext::Release call"));
@@ -530,11 +539,36 @@ void FPCGContext::AddStructReferencedObjects(FReferenceCollector& Collector)
 	AddExtraStructReferencedObjects(Collector);
 }
 
+void FPCGContext::Release(FPCGContext* InContext)
+{
+	if (!InContext)
+	{
+		return;
+	}
+
+	// Adds this to the tracked contexts so we ensure that the data isn't GC'ed too early for the context destructor to behave properly.
+	if(TSharedPtr<FPCGGraphExecutor> Executor = InContext->GraphExecutor.Pin())
+	{
+		Executor->AddReleasedContextForGC(InContext);
+	}
+
+	if (!InContext->Handle)
+	{
+		delete InContext;
+	}
+	else
+	{
+		// Release the Handle SharedPtr (at this point only captured and pinned Handles can keep the Context alive)
+		// The Context will get deleted when the handle gets released
+		InContext->Handle = nullptr;
+	}
+}
+
 FPCGTaskId FPCGContext::ScheduleGraph(const FPCGScheduleGraphParams& InParams)
 {
-	if (GraphExecutor && InParams.ExecutionSource)
+	if (TSharedPtr<FPCGGraphExecutor> Executor = GraphExecutor.Pin(); Executor && InParams.ExecutionSource)
 	{
-		return GraphExecutor->ScheduleGraph(InParams);
+		return Executor->ScheduleGraph(InParams);
 	}
 	else
 	{
@@ -544,9 +578,9 @@ FPCGTaskId FPCGContext::ScheduleGraph(const FPCGScheduleGraphParams& InParams)
 
 FPCGTaskId FPCGContext::ScheduleGeneric(const FPCGScheduleGenericParams& InParams)
 {
-	if (GraphExecutor)
+	if (TSharedPtr<FPCGGraphExecutor> Executor = GraphExecutor.Pin())
 	{
-		return GraphExecutor->ScheduleGeneric(InParams);
+		return Executor->ScheduleGeneric(InParams);
 	}
 	else
 	{
@@ -556,9 +590,9 @@ FPCGTaskId FPCGContext::ScheduleGeneric(const FPCGScheduleGenericParams& InParam
 
 bool FPCGContext::GetOutputData(FPCGTaskId InTaskId, FPCGDataCollection& OutData)
 {
-	if (GraphExecutor)
+	if (TSharedPtr<FPCGGraphExecutor> Executor = GraphExecutor.Pin())
 	{
-		return GraphExecutor->GetOutputData(InTaskId, OutData);
+		return Executor->GetOutputData(InTaskId, OutData);
 	}
 	else
 	{
@@ -568,9 +602,9 @@ bool FPCGContext::GetOutputData(FPCGTaskId InTaskId, FPCGDataCollection& OutData
 
 void FPCGContext::ClearOutputData(FPCGTaskId InTaskId)
 {
-	if (GraphExecutor)
+	if (TSharedPtr<FPCGGraphExecutor> Executor = GraphExecutor.Pin())
 	{
-		GraphExecutor->ClearOutputData(InTaskId);
+		Executor->ClearOutputData(InTaskId);
 	}
 }
 
@@ -582,17 +616,17 @@ void FPCGContext::InitializeGraphExecutor(FPCGContext* InContext)
 
 void FPCGContext::StoreInCache(const FPCGStoreInCacheParams& Params, const FPCGDataCollection& InOutput)
 {
-	if (GraphExecutor)
+	if (TSharedPtr<FPCGGraphExecutor> Executor = GraphExecutor.Pin())
 	{
-		GraphExecutor->GetCache().StoreInCache(Params, InOutput);
+		Executor->GetCache().StoreInCache(Params, InOutput);
 	}
 }
 
 bool FPCGContext::GetFromCache(const FPCGGetFromCacheParams& Params, FPCGDataCollection& OutOutput) const
 {
-	if (GraphExecutor)
+	if (TSharedPtr<FPCGGraphExecutor> Executor = GraphExecutor.Pin())
 	{
-		return GraphExecutor->GetCache().GetFromCache(Params, OutOutput);
+		return Executor->GetCache().GetFromCache(Params, OutOutput);
 	}
 	else
 	{
diff --git a/Engine/Plugins/PCG/Source/PCG/Private/PCGData.cpp b/Engine/Plugins/PCG/Source/PCG/Private/PCGData.cpp
--- a/Engine/Plugins/PCG/Source/PCG/Private/PCGData.cpp
+++ b/Engine/Plugins/PCG/Source/PCG/Private/PCGData.cpp
@@ -2,6 +2,7 @@
 
 #include "PCGData.h"
 
+#include "PCGCommon.h"
 #include "PCGContext.h"
 #include "PCGNode.h"
 #include "PCGParamData.h"
@@ -17,6 +18,8 @@
 
 #include UE_INLINE_GENERATED_CPP_BY_NAME(PCGData)
 
+//#define PCG_DATA_USAGE_LOGGING
+
 UPCGData::UPCGData(const FObjectInitializer& ObjectInitializer)
 	: Super(ObjectInitializer)
 {
@@ -69,7 +72,45 @@ void UPCGData::ClearUsage(EPCGDataUsage InUsage) const
 
 	Usage &= ~static_cast<__underlying_type(EPCGDataUsage)>(InUsage);
 
-	if (Usage == static_cast<__underlying_type(EPCGDataUsage)>(EPCGDataUsage::None))
+	if (Usage == static_cast<__underlying_type(EPCGDataUsage)>(EPCGDataUsage::None) && 
+		(!PCGSystemSwitches::CVarReleaseTransientResourcesEarly.GetValueOnAnyThread() || CollectionRefCount == 0))
+	{
+#ifdef PCG_DATA_USAGE_LOGGING
+		UE_LOG(LogPCG, Warning, TEXT("%s: !! ReleaseTransientResources after usage unmarked !!"), *GetName());
+#endif
+
+		VisitDataNetwork([](const UPCGData* InData)
+		{
+			if (InData)
+			{
+				const_cast<UPCGData*>(InData)->ReleaseTransientResources(TEXT("Data unused"));
+			}
+		});
+	}
+}
+
+bool UPCGData::HasUsage(EPCGDataUsage InUsage) const
+{
+	return !!(Usage & static_cast<__underlying_type(EPCGDataUsage)>(InUsage));
+}
+
+void UPCGData::IncCollectionRefCount() const
+{
+	++CollectionRefCount;
+}
+
+void UPCGData::DecCollectionRefCount() const
+{
+	--CollectionRefCount;
+
+	if (!PCGSystemSwitches::CVarReleaseTransientResourcesEarly.GetValueOnAnyThread())
+	{
+		return;
+	}
+
+	ensure(CollectionRefCount >= 0);
+
+	if (Usage == static_cast<__underlying_type(EPCGDataUsage)>(EPCGDataUsage::None) && CollectionRefCount == 0)
 	{
 #ifdef PCG_DATA_USAGE_LOGGING
 		UE_LOG(LogPCG, Warning, TEXT("%s: !! ReleaseTransientResources after usage unmarked !!"), *GetName());
@@ -77,7 +118,10 @@ void UPCGData::ClearUsage(EPCGDataUsage InUsage) const
 
 		VisitDataNetwork([](const UPCGData* InData)
 		{
-			const_cast<UPCGData*>(InData)->ReleaseTransientResources(TEXT("Data unused"));
+			if (InData)
+			{
+				const_cast<UPCGData*>(InData)->ReleaseTransientResources(TEXT("Data unused"));
+			}
 		});
 	}
 }
@@ -512,7 +556,7 @@ void FPCGDataCollection::AddReferences(FReferenceCollector& Collector)
 	{
 		if (Data.Data)
 		{
-			Collector.AddReferencedObject(Data.Data);
+			Collector.AddReferencedObject(Data.Data.Data);
 		}
 	}
 }
diff --git a/Engine/Plugins/PCG/Source/PCG/Private/PCGSubsystem.cpp b/Engine/Plugins/PCG/Source/PCG/Private/PCGSubsystem.cpp
--- a/Engine/Plugins/PCG/Source/PCG/Private/PCGSubsystem.cpp
+++ b/Engine/Plugins/PCG/Source/PCG/Private/PCGSubsystem.cpp
@@ -197,7 +197,6 @@ void UPCGSubsystem::Deinitialize()
 {
 	// Cancel all tasks
 	// TODO
-	delete GraphExecutor;
 	GraphExecutor = nullptr;
 
 	delete RuntimeGenScheduler;
@@ -232,7 +231,7 @@ void UPCGSubsystem::PostInitialize()
 
 	// Initialize graph executor
 	check(!GraphExecutor);
-	GraphExecutor = new FPCGGraphExecutor(GetWorld());
+	GraphExecutor = MakeShared<FPCGGraphExecutor>(GetWorld());
 
 	// Initialize runtime generation scheduler
 	check(!RuntimeGenScheduler);
diff --git a/Engine/Plugins/PCG/Source/PCG/Private/Tests/Determinism/PCGDeterminismTestsCommon.cpp b/Engine/Plugins/PCG/Source/PCG/Private/Tests/Determinism/PCGDeterminismTestsCommon.cpp
--- a/Engine/Plugins/PCG/Source/PCG/Private/Tests/Determinism/PCGDeterminismTestsCommon.cpp
+++ b/Engine/Plugins/PCG/Source/PCG/Private/Tests/Determinism/PCGDeterminismTestsCommon.cpp
@@ -109,13 +109,13 @@ namespace PCGDeterminismTests
 			PCGComponentCopy->RegisterComponentWithWorld(PCGActor->GetWorld());
 		}
 
-		FPCGGraphExecutor Executor;
+		TSharedPtr<FPCGGraphExecutor> Executor = MakeShared<FPCGGraphExecutor>();
 
 		ON_SCOPE_EXIT
 		{
-			if (Executor.IsGraphCurrentlyExecuting(InPCGGraph))
+			if (Executor->IsGraphCurrentlyExecuting(InPCGGraph))
 			{
-				Executor.CancelAll();
+				Executor->CancelAll();
 			}
 
 			// Clean up anything generated
@@ -134,10 +134,10 @@ namespace PCGDeterminismTests
 			volatile bool bTasksComplete = false;
 
 			// Clear the cache between runs
-			Executor.GetCache().ClearCache();
+			Executor->GetCache().ClearCache();
 
 			// Schedule final task for waiting
-			Executor.ScheduleGeneric([&bTasksComplete]()
+			Executor->ScheduleGeneric([&bTasksComplete]()
 			{
 				bTasksComplete = true;
 				return true;
@@ -152,11 +152,11 @@ namespace PCGDeterminismTests
 			{
 				if (ExecutionCount++ < CVarDeterminismGraphTestCycleLimit.GetValueOnAnyThread())
 				{
-					Executor.Execute();
+					Executor->Execute();
 				}
 				else
 				{
-					Executor.CancelAll();
+					Executor->CancelAll();
 					return false;
 				}
 			}
@@ -174,7 +174,7 @@ namespace PCGDeterminismTests
 		};
 
 		// Schedule the graph execution
-		FPCGTaskId FinalTaskID = Executor.ScheduleDebugWithTaskCallback(PCGComponentCopy, RecordMappedOutput);
+		FPCGTaskId FinalTaskID = Executor->ScheduleDebugWithTaskCallback(PCGComponentCopy, RecordMappedOutput);
 
 		TMap<FPCGTaskId, TTuple<const UPCGNode*, const FPCGDataCollection>> FirstIntermediateOutputArray;
 		bool bExecutionCompleted = ScheduleAndWaitForExecution(FinalTaskID);
@@ -188,7 +188,7 @@ namespace PCGDeterminismTests
 			PCGComponentCopy->Cleanup();
 
 			// Run again
-			FinalTaskID = Executor.ScheduleDebugWithTaskCallback(PCGComponentCopy, RecordMappedOutput);
+			FinalTaskID = Executor->ScheduleDebugWithTaskCallback(PCGComponentCopy, RecordMappedOutput);
 			bExecutionCompleted &= ScheduleAndWaitForExecution(FinalTaskID);
 		}
 
diff --git a/Engine/Plugins/PCG/Source/PCG/Private/Tests/Grid/PCGGraphExecutorTest.cpp b/Engine/Plugins/PCG/Source/PCG/Private/Tests/Grid/PCGGraphExecutorTest.cpp
--- a/Engine/Plugins/PCG/Source/PCG/Private/Tests/Grid/PCGGraphExecutorTest.cpp
+++ b/Engine/Plugins/PCG/Source/PCG/Private/Tests/Grid/PCGGraphExecutorTest.cpp
@@ -17,7 +17,7 @@ bool FPCGGraphExecutorTests::RunTest(const FString& Parameters)
 	TArray<PCGTestsCommon::FTestData> TestData;
 	TestData.SetNum(ComponentsNum);
 
-	FPCGGraphExecutor Executor;
+	TSharedPtr<FPCGGraphExecutor> Executor = MakeShared<FPCGGraphExecutor>();
 	TArray<TArray<FPCGTaskId>> Tasks;
 	TArray<TArray<bool>> bExecuted;
 	TArray<TArray<bool>> bAborted;
@@ -49,7 +49,7 @@ bool FPCGGraphExecutorTests::RunTest(const FString& Parameters)
 					PreviousTasks.Add(Tasks[ComponentIndex][TaskIndex - 1]);
 				}
 
-				ComponentTasks[TaskIndex] = Executor.ScheduleGeneric(
+				ComponentTasks[TaskIndex] = Executor->ScheduleGeneric(
 					[&bExecuted, ComponentIndex, TaskIndex]()
 					{
 						bExecuted[ComponentIndex][TaskIndex] = true;
@@ -67,7 +67,7 @@ bool FPCGGraphExecutorTests::RunTest(const FString& Parameters)
 	{
 		for (int i = 0; i < NTimes; ++i)
 		{
-			Executor.Execute();
+			Executor->Execute();
 		}
 	};
 
@@ -79,21 +79,21 @@ bool FPCGGraphExecutorTests::RunTest(const FString& Parameters)
 	// First test: schedule all three tasks, and cancel them, checking that the appropriate number is cancelled.
 	{
 		AddTasks(1);
-		Executor.Cancel(TestData[0].TestPCGComponent);
+		Executor->Cancel(TestData[0].TestPCGComponent);
 		ExecuteAll();
 
 		bTestPassed &= TestTrue(TEXT("Expect that no tasks are executed"), !bExecuted[0][0] && !bExecuted[1][0] && !bExecuted[2][0]);
 		bTestPassed &= TestTrue(TEXT("Expect that all tasks are aborted"), bAborted[0][0] && bAborted[1][0] && bAborted[2][0]);
 
 		AddTasks(1);
-		Executor.Cancel(TestData[1].TestPCGComponent);
+		Executor->Cancel(TestData[1].TestPCGComponent);
 		ExecuteAll();
 
 		bTestPassed &= TestTrue(TEXT("Expect that only task from A is executed"), bExecuted[0][0] && !bExecuted[1][0] && !bExecuted[2][0]);
 		bTestPassed &= TestTrue(TEXT("Expect that only task from A is not aborted"), !bAborted[0][0] && bAborted[1][0] && bAborted[2][0]);
 
 		AddTasks(1);
-		Executor.Cancel(TestData[2].TestPCGComponent);
+		Executor->Cancel(TestData[2].TestPCGComponent);
 		ExecuteAll();
 
 		bTestPassed &= TestTrue(TEXT("Expect that only task from C is not executed"), bExecuted[0][0] && bExecuted[1][0] && !bExecuted[2][0]);
@@ -131,7 +131,7 @@ bool FPCGGraphExecutorTests::RunTest(const FString& Parameters)
 			if (ComponentIndex < ComponentsNum)
 			{
 				FirstTaskCancelled = Tasks[ComponentIndex][ComponentTaskIndex];
-				Executor.Cancel(TestData[ComponentIndex].TestPCGComponent);
+				Executor->Cancel(TestData[ComponentIndex].TestPCGComponent);
 			}
 
 			// Finish execution
diff --git a/Engine/Plugins/PCG/Source/PCG/Private/Tests/Subgraphs/PCGCollapseSubgraphTest.cpp b/Engine/Plugins/PCG/Source/PCG/Private/Tests/Subgraphs/PCGCollapseSubgraphTest.cpp
--- a/Engine/Plugins/PCG/Source/PCG/Private/Tests/Subgraphs/PCGCollapseSubgraphTest.cpp
+++ b/Engine/Plugins/PCG/Source/PCG/Private/Tests/Subgraphs/PCGCollapseSubgraphTest.cpp
@@ -271,7 +271,7 @@ bool FPCGCollapseSubgraphCombinationOfCollapses::RunTest(const FString& Paramete
 	// Temporary solution to run a full graph and extract the resulting points.
 	auto ExecuteGraph = [&TestData](UPCGGraph* GraphToExecute, double FirstConstant, double SecondConstant) -> TArray<FPCGPoint>
 	{
-		FPCGGraphExecutor GraphExecutor;
+		TSharedPtr<FPCGGraphExecutor> GraphExecutor = MakeShared<FPCGGraphExecutor>();
 		TArray<FPCGPoint> OutputPoints;
 		bool bDone = false;
 
@@ -279,15 +279,15 @@ bool FPCGCollapseSubgraphCombinationOfCollapses::RunTest(const FString& Paramete
 		TestData.TestPCGComponent->GetGraphInstance()->SetGraphParameter<double>(TEXT("FloatConstant"), FirstConstant);
 		TestData.TestPCGComponent->GetGraphInstance()->SetGraphParameter<double>(TEXT("AnotherConstant"), SecondConstant);
 
-		const FPCGTaskId TaskId = GraphExecutor.Schedule(TestData.TestPCGComponent);
+		const FPCGTaskId TaskId = GraphExecutor->Schedule(TestData.TestPCGComponent);
 
 		// TODO: Post-release GraphExecutor should be setting this on the Component
 		TestData.SetCurrentGenerationTask(TaskId);
 
-		const FPCGTaskId FinalTaskId = GraphExecutor.ScheduleGeneric([&OutputPoints, &GraphExecutor, TaskId, &bDone]() -> bool
+		const FPCGTaskId FinalTaskId = GraphExecutor->ScheduleGeneric([&OutputPoints, &GraphExecutor, TaskId, &bDone]() -> bool
 		{
 			FPCGDataCollection OutputData{};
-			GraphExecutor.GetOutputData(TaskId, OutputData);
+			GraphExecutor->GetOutputData(TaskId, OutputData);
 
 			if (OutputData.TaggedData.Num() == 1)
 			{
@@ -303,7 +303,7 @@ bool FPCGCollapseSubgraphCombinationOfCollapses::RunTest(const FString& Paramete
 
 		while(!bDone)
 		{
-			GraphExecutor.Execute();
+			GraphExecutor->Execute();
 		}
 
 		return OutputPoints;
diff --git a/Engine/Plugins/PCG/Source/PCG/Public/Data/PCGDataPtrWrapper.h b/Engine/Plugins/PCG/Source/PCG/Public/Data/PCGDataPtrWrapper.h
new file mode 100644
--- /dev/null
+++ b/Engine/Plugins/PCG/Source/PCG/Public/Data/PCGDataPtrWrapper.h
@@ -0,0 +1,84 @@
+// Copyright Epic Games, Inc. All Rights Reserved.
+
+#pragma once
+
+#include "PCGCommon.h"
+
+#include "UObject/ObjectPtr.h"
+
+#include "PCGDataPtrWrapper.generated.h"
+
+class UPCGData;
+
+/** Wrapper to do ref-counting inside of FPCGDataCollection, so we can release temporary resources earlier. */
+USTRUCT()
+struct PCG_API FPCGDataPtrWrapper
+{
+	friend struct FPCGDataCollection;
+
+	GENERATED_BODY()
+
+	FPCGDataPtrWrapper() = default;
+	~FPCGDataPtrWrapper();
+	FPCGDataPtrWrapper(const FPCGDataPtrWrapper& InWrapper);
+	FPCGDataPtrWrapper(FPCGDataPtrWrapper&& InWrapper);
+	FPCGDataPtrWrapper& operator=(const FPCGDataPtrWrapper& InWrapper);
+	FPCGDataPtrWrapper& operator=(FPCGDataPtrWrapper&& InWrapper);
+
+	FPCGDataPtrWrapper(const UPCGData* InData);
+	FPCGDataPtrWrapper& operator=(const UPCGData* InData);
+
+	operator const UPCGData* () const { return Data; }
+	const UPCGData* operator->() const { return Data; }
+	const UPCGData* Get() const { return Data; }
+	const TObjectPtr<const UPCGData> GetObjectPtr() const { return Data; }
+
+	template<typename Class>
+	bool IsA() const
+	{
+		return Data && Data.IsA<Class>();
+	}
+
+	bool Serialize(FArchive& Ar);
+	bool ImportTextItem(const TCHAR*& Buffer, int32 PortFlags, UObject* Parent, FOutputDevice* ErrorText);
+	bool SerializeFromMismatchedTag(const FPropertyTag& Tag, FStructuredArchive::FSlot Slot);
+
+protected:
+	static void IncRefCount(const UPCGData* InData);
+	static void DecRefCount(const TWeakObjectPtr<const UPCGData>& InWeakData);
+
+	void IncRefCount();
+	void DecRefCount();
+
+protected:
+	UPROPERTY(VisibleAnywhere, Category = Data)
+	TObjectPtr<const UPCGData> Data = nullptr;
+
+	// Weakptr to be able to make sure we don't try to decrement collection ref count if the data was already destroyed.
+	// Normally, we don't really need it, except that in some cases the order of deletion isn't easily guaranteed, and enforcing BeginDestroy implementation
+	// in every user of the FPCGDataCollection class seems like a tall order.
+	TWeakObjectPtr<const UPCGData> WeakData;
+};
+
+template<>
+struct TStructOpsTypeTraits<FPCGDataPtrWrapper> : public TStructOpsTypeTraitsBase2<FPCGDataPtrWrapper>
+{
+	enum
+	{
+		WithStructuredSerializeFromMismatchedTag = true,
+		WithSerializer = true,
+		WithImportTextItem = true
+	};
+};
+
+template<typename To>
+FORCEINLINE TCopyQualifiersFromTo_T<const TObjectPtr<const UPCGData>, To>* Cast(const FPCGDataPtrWrapper& InSrc)
+{
+	return Cast<To>(InSrc.GetObjectPtr());
+}
+
+template<typename To>
+FORCEINLINE TCopyQualifiersFromTo_T<const TObjectPtr<const UPCGData>, To>* CastChecked(const FPCGDataPtrWrapper& InSrc, ECastCheckedType::Type CheckType = ECastCheckedType::NullChecked)
+{
+	return CastChecked<To>(InSrc.GetObjectPtr(), CheckType);
+}
\ No newline at end of file
diff --git a/Engine/Plugins/PCG/Source/PCG/Public/PCGCommon.h b/Engine/Plugins/PCG/Source/PCG/Public/PCGCommon.h
--- a/Engine/Plugins/PCG/Source/PCG/Public/PCGCommon.h
+++ b/Engine/Plugins/PCG/Source/PCG/Public/PCGCommon.h
@@ -274,6 +274,8 @@ namespace PCGSystemSwitches
 	extern TAutoConsoleVariable<bool> CVarDirtyLoadAsPreviewOnLoad;
 	extern TAutoConsoleVariable<bool> CVarForceDynamicGraphDispatch;
 #endif
+
+	extern PCG_API TAutoConsoleVariable<bool> CVarReleaseTransientResourcesEarly;
 }
 
 /** Describes space referential for operations that create data */
diff --git a/Engine/Plugins/PCG/Source/PCG/Public/PCGContext.h b/Engine/Plugins/PCG/Source/PCG/Public/PCGContext.h
--- a/Engine/Plugins/PCG/Source/PCG/Public/PCGContext.h
+++ b/Engine/Plugins/PCG/Source/PCG/Public/PCGContext.h
@@ -217,26 +217,7 @@ PRAGMA_ENABLE_DEPRECATION_WARNINGS
 		return Handle.ToWeakPtr();
 	}
 
-	static void Release(FPCGContext* InContext)
-	{
-		if(!InContext)
-		{
-			return;
-		}
-		
-		InContext->GraphExecutor = nullptr;
-
-		if (!InContext->Handle)
-		{
-			delete InContext;
-		}
-		else
-		{
-			// Release the Handle SharedPtr (at this point only captured and pinned Handles can keep the Context alive)
-			// The Context will get deleted when the handle gets released
-			InContext->Handle = nullptr;
-		}
-	}
+	static void Release(FPCGContext* InContext);
 
 	template <typename ContextType>
 	UE_DEPRECATED(5.6, "Use FSharedContext instead")
@@ -316,7 +297,7 @@ private:
 	TSharedPtr<FPCGContextHandle> Handle;
 
 	// Graph Executor that scheduled the task
-	FPCGGraphExecutor* GraphExecutor = nullptr;
+	TWeakPtr<FPCGGraphExecutor> GraphExecutor;
 
 	/** Handle for the current call stack. */
 	FPCGStackHandle StackHandle;
diff --git a/Engine/Plugins/PCG/Source/PCG/Public/PCGData.h b/Engine/Plugins/PCG/Source/PCG/Public/PCGData.h
--- a/Engine/Plugins/PCG/Source/PCG/Public/PCGData.h
+++ b/Engine/Plugins/PCG/Source/PCG/Public/PCGData.h
@@ -4,6 +4,7 @@
 
 #include "PCGCommon.h"
 #include "PCGCrc.h"
+#include "Data/PCGDataPtrWrapper.h"
 #include "Metadata/PCGAttributePropertySelector.h"
 #include "Metadata/PCGMetadataCommon.h"
 
@@ -137,6 +138,10 @@ public:
 
 	void MarkUsage(EPCGDataUsage InUsage) const;
 	void ClearUsage(EPCGDataUsage InUsage) const;
+	bool HasUsage(EPCGDataUsage InUsage) const;
+	void IncCollectionRefCount() const;
+	void DecCollectionRefCount() const;
+	int32 GetCollectionRefCount() const { return CollectionRefCount; }
 
 protected:
 	/** Computes Crc for this and any connected data. */
@@ -161,15 +166,17 @@ private:
 
 	/** Usage flags to record any current usages of this data. Used to help manage lifetimes of any internal transient resources. */
 	mutable std::atomic<__underlying_type(EPCGDataUsage)> Usage{ 0 };
+	mutable std::atomic<int32> CollectionRefCount{ 0 };
 };
 
-USTRUCT(BlueprintType)
+USTRUCT(BlueprintType, meta = (HasNativeBreak = "/Script/PCG.PCGDataFunctionLibrary.BreakTaggedData", HasNativeMake = "/Script/PCG.PCGDataFunctionLibrary.MakeTaggedData"))
 struct PCG_API FPCGTaggedData
 {
 	GENERATED_BODY()
 
-	UPROPERTY(BlueprintReadWrite, EditAnywhere, Category = Data)
-	TObjectPtr<const UPCGData> Data = nullptr;
+	/** Wraps a TObjectPtr<const UPCGData> so we can track ref counting easier. */
+	UPROPERTY(VisibleAnywhere, Category = Data, meta = (ShowOnlyInnerProperties))
+	FPCGDataPtrWrapper Data;
 
 	UPROPERTY(BlueprintReadWrite, EditAnywhere, Category = Data)
 	TSet<FString> Tags;
@@ -330,6 +337,19 @@ struct PCG_API FPCGDataCollection
 		}
 	}
 
+	bool HasUsage(EPCGDataUsage InUsage) const
+	{
+		for (const FPCGTaggedData& Data : TaggedData)
+		{
+			if (Data.Data && Data.Data->HasUsage(InUsage))
+			{
+				return true;
+			}
+		}
+
+		return false;
+	}
+
 	UPROPERTY(BlueprintReadWrite, EditAnywhere, Category = Data)
 	TArray<FPCGTaggedData> TaggedData;
 
@@ -413,6 +433,26 @@ public:
 	UFUNCTION(BlueprintCallable, Category = Data)
 	static TArray<FPCGTaggedData> GetAllSettings(const FPCGDataCollection& InCollection);
 
+	UFUNCTION(BlueprintPure, Category = "PCG|TaggedData", meta = (NativeBreakFunc))
+	static void BreakTaggedData(const FPCGTaggedData& TaggedData, UPCGData*& Data, TSet<FString>& Tags, FName& Pin, bool& bIsUsedMultipleTimes)
+	{
+		Data = const_cast<UPCGData*>(TaggedData.Data.Get());
+		Tags = TaggedData.Tags;
+		Pin = TaggedData.Pin;
+		bIsUsedMultipleTimes = TaggedData.bIsUsedMultipleTimes;
+	}
+
+	UFUNCTION(BlueprintPure, Category = "PCG|TaggedData", meta = (NativeMakeFunc, AutoCreateRefTerm = "Tags"))
+	static FPCGTaggedData MakeTaggedData(UPCGData* Data, const TSet<FString>& Tags, FName Pin)
+	{
+		FPCGTaggedData TaggedData;
+		TaggedData.Data = Data;
+		TaggedData.Tags = Tags;
+		TaggedData.Pin = Pin;
+
+		return TaggedData;
+	}
+
 protected:
 	static TArray<UPCGData*> GetInputsByPredicate(const FPCGDataCollection& InCollection, TArray<FPCGTaggedData>& OutTaggedData, TFunctionRef<bool(const FPCGTaggedData&)> InPredicate);
 };
diff --git a/Engine/Plugins/PCG/Source/PCG/Public/PCGSubsystem.h b/Engine/Plugins/PCG/Source/PCG/Public/PCGSubsystem.h
--- a/Engine/Plugins/PCG/Source/PCG/Public/PCGSubsystem.h
+++ b/Engine/Plugins/PCG/Source/PCG/Public/PCGSubsystem.h
@@ -431,7 +431,7 @@ private:
 	void ExecuteBeginTickActions();
 
 	APCGWorldActor* PCGWorldActor = nullptr;
-	FPCGGraphExecutor* GraphExecutor = nullptr;
+	TSharedPtr<FPCGGraphExecutor> GraphExecutor;
 	FPCGRuntimeGenScheduler* RuntimeGenScheduler = nullptr;
 	bool bHasTickedOnce = false;
 	FPCGActorAndComponentMapping ActorAndComponentMapping;
