[PCG] Avoid ticking VT priming if there are no generation sources

#rnx

diff --git a/Engine/Plugins/PCG/Source/PCG/Private/RuntimeGen/PCGRuntimeGenScheduler.cpp b/Engine/Plugins/PCG/Source/PCG/Private/RuntimeGen/PCGRuntimeGenScheduler.cpp
--- a/Engine/Plugins/PCG/Source/PCG/Private/RuntimeGen/PCGRuntimeGenScheduler.cpp
+++ b/Engine/Plugins/PCG/Source/PCG/Private/RuntimeGen/PCGRuntimeGenScheduler.cpp
@@ -156,7 +156,7 @@ void FPCGRuntimeGenScheduler::Tick(APCGWorldActor* InPCGWorldActor)
 		}
 	}
 
-	if (PCGRuntimeGenSchedulerHelpers::CVarEnableVirtualTexturePriming.GetValueOnAnyThread())
+	if (!GenSources.IsEmpty() && PCGRuntimeGenSchedulerHelpers::CVarEnableVirtualTexturePriming.GetValueOnGameThread())
 	{
 		// @todo_pcg: To support VT priming outside of RuntimeGen, this should probably move outside of the RGS tick, and be ticked directly by the PCGSubsystem.
 		// However, that would require also moving the GenSourceManager out of the RGS, 
