[PCG] Fix bounding shape test issue in the Surface sampler

* With the conversion to support Point Array, the out point bounds were not set correctly anymore before the bounding shape test.

#rnx

diff --git a/Engine/Plugins/PCG/Source/PCG/Private/Elements/PCGSurfaceSampler.cpp b/Engine/Plugins/PCG/Source/PCG/Private/Elements/PCGSurfaceSampler.cpp
--- a/Engine/Plugins/PCG/Source/PCG/Private/Elements/PCGSurfaceSampler.cpp
+++ b/Engine/Plugins/PCG/Source/PCG/Private/Elements/PCGSurfaceSampler.cpp
@@ -255,7 +255,11 @@ namespace PCGSurfaceSampler
 				{
 					continue;
 				}
-					
+
+				// Set physical properties that are needed for the bounding shape checks, etc.
+				OutPoint.SetExtents(SamplerParams.PointExtents);
+				OutPoint.Steepness = SamplerParams.PointSteepness;
+
 				// Set default density
 				DensityRange[WriteIndex] = 1.0f;
 
