[PCG] Allow PCG compute kernels to add additional sources
* Now instead of adding includes to the kernel source, shader source files should be added through CreateAdditionalSources, so that they are also hashed and cause the kernel to be recompiled when their source is changed
* Duplicate Quaternion.ush to sandbox ourselves from engine shaders, so that our data won't break if they change

diff --git a/Engine/Plugins/PCG/Shaders/Private/BuiltInKernels/PCGCountUniqueAttributeValues.usf b/Engine/Plugins/PCG/Shaders/Private/BuiltInKernels/PCGCountUniqueAttributeValues.usf
--- a/Engine/Plugins/PCG/Shaders/Private/BuiltInKernels/PCGCountUniqueAttributeValues.usf
+++ b/Engine/Plugins/PCG/Shaders/Private/BuiltInKernels/PCGCountUniqueAttributeValues.usf
@@ -5,15 +5,6 @@
 // Input: Input elements with attribute.
 // Output: Attribute set with a ValueCount attribute which has the count of each value. Number of elements is set to the number of unique values (known CPU-side).
 
-#include "/Engine/Private/ComputeShaderUtils.ush"
-#include "/Engine/Private/Quaternion.ush"
-#include "/Plugin/PCG/Private/PCGShaderUtils.ush"
-#include "/Plugin/PCG/Private/PCGShaderUtilsInternal.ush"
-
-// TODO: should be code-gen'd from PCGComputeCommon.h
-#define MAX_NUM_ATTRS 128
-#define NUM_RESERVED_ATTRS 32
-
 [numthreads(64, 1, 1)]
 void PCGCountUniqueAttributeValuesCS(uint3 GroupId : SV_GroupID, uint GroupIndex : SV_GroupIndex)
 {
diff --git a/Engine/Plugins/PCG/Shaders/Private/Elements/PCGCopyPoints.usf b/Engine/Plugins/PCG/Shaders/Private/Elements/PCGCopyPoints.usf
--- a/Engine/Plugins/PCG/Shaders/Private/Elements/PCGCopyPoints.usf
+++ b/Engine/Plugins/PCG/Shaders/Private/Elements/PCGCopyPoints.usf
@@ -1,10 +1,5 @@
 // Copyright Epic Games, Inc. All Rights Reserved.
 
-#include "/Engine/Private/ComputeShaderUtils.ush"
-#include "/Engine/Private/Quaternion.ush"
-#include "/Plugin/PCG/Private/PCGShaderUtils.ush"
-#include "/Plugin/PCG/Private/PCGShaderUtilsInternal.ush"
-
 // Mirrors EPCGCopyPointsInheritanceMode
 #define COPY_POINTS_INHERITANCE_MODE_Relative 0
 #define COPY_POINTS_INHERITANCE_MODE_Source 1
@@ -17,10 +12,6 @@
 #define COPY_POINTS_METADATA_INHERITANCE_MODE_TargetOnly 3
 #define COPY_POINTS_METADATA_INHERITANCE_MODE_None 4
 
-// TODO: should be code-gen'd from PCGComputeCommon.h
-#define MAX_NUM_ATTRS 128
-#define NUM_RESERVED_ATTRS 32
-
 [numthreads(64, 1, 1)]
 void PCGCopyPointsCS(uint3 GroupId : SV_GroupID, uint GroupIndex : SV_GroupIndex)
 {
diff --git a/Engine/Plugins/PCG/Shaders/Private/Elements/PCGCopyPointsAnalysis.usf b/Engine/Plugins/PCG/Shaders/Private/Elements/PCGCopyPointsAnalysis.usf
--- a/Engine/Plugins/PCG/Shaders/Private/Elements/PCGCopyPointsAnalysis.usf
+++ b/Engine/Plugins/PCG/Shaders/Private/Elements/PCGCopyPointsAnalysis.usf
@@ -1,14 +1,5 @@
 // Copyright Epic Games, Inc. All Rights Reserved.
 
-#include "/Engine/Private/ComputeShaderUtils.ush"
-#include "/Engine/Private/Quaternion.ush"
-#include "/Plugin/PCG/Private/PCGShaderUtils.ush"
-#include "/Plugin/PCG/Private/PCGShaderUtilsInternal.ush"
-
-// TODO: should be code-gen'd from PCGComputeCommon.h
-#define MAX_NUM_ATTRS 128
-#define NUM_RESERVED_ATTRS 32
-
 [numthreads(64, 1, 1)]
 void PCGCopyPointsAnalysisCS(uint3 GroupId : SV_GroupID, uint GroupIndex : SV_GroupIndex)
 {
diff --git a/Engine/Plugins/PCG/Shaders/Private/Elements/PCGMetadataPartition.usf b/Engine/Plugins/PCG/Shaders/Private/Elements/PCGMetadataPartition.usf
--- a/Engine/Plugins/PCG/Shaders/Private/Elements/PCGMetadataPartition.usf
+++ b/Engine/Plugins/PCG/Shaders/Private/Elements/PCGMetadataPartition.usf
@@ -2,15 +2,6 @@
 
 // Kernel for GPU implementation of Metadata Partition node.
 
-#include "/Engine/Private/ComputeShaderUtils.ush"
-#include "/Engine/Private/Quaternion.ush"
-#include "/Plugin/PCG/Private/PCGShaderUtils.ush"
-#include "/Plugin/PCG/Private/PCGShaderUtilsInternal.ush"
-
-// TODO: should be code-gen'd from PCGComputeCommon.h
-#define MAX_NUM_ATTRS 128
-#define NUM_RESERVED_ATTRS 32
-
 [numthreads(64, 1, 1)]
 void PCGMetadataPartitionCS(uint3 GroupId : SV_GroupID, uint GroupIndex : SV_GroupIndex)
 {
diff --git a/Engine/Plugins/PCG/Shaders/Private/Elements/PCGStaticMeshSpawner.usf b/Engine/Plugins/PCG/Shaders/Private/Elements/PCGStaticMeshSpawner.usf
--- a/Engine/Plugins/PCG/Shaders/Private/Elements/PCGStaticMeshSpawner.usf
+++ b/Engine/Plugins/PCG/Shaders/Private/Elements/PCGStaticMeshSpawner.usf
@@ -5,15 +5,6 @@
 
 // When spawning by attribute, we don't know in advance how many instances of each mesh will be spawned, so the worst case NumPoints * NumPrims is allocated in the GPU Scene.
 
-#include "/Engine/Private/ComputeShaderUtils.ush"
-#include "/Engine/Private/Quaternion.ush"
-#include "/Plugin/PCG/Private/PCGShaderUtils.ush"
-#include "/Plugin/PCG/Private/PCGShaderUtilsInternal.ush"
-
-// TODO: should be code-gen'd from PCGComputeCommon.h
-#define MAX_NUM_ATTRS 128
-#define NUM_RESERVED_ATTRS 32
-
 [numthreads(64, 1, 1)]
 void PCGStaticMeshSpawnerCS(uint3 GroupId : SV_GroupID, uint GroupIndex : SV_GroupIndex)
 {
diff --git a/Engine/Plugins/PCG/Shaders/Private/PCGDataCollectionDataInterface.ush b/Engine/Plugins/PCG/Shaders/Private/PCGDataCollectionDataInterface.ush
--- a/Engine/Plugins/PCG/Shaders/Private/PCGDataCollectionDataInterface.ush
+++ b/Engine/Plugins/PCG/Shaders/Private/PCGDataCollectionDataInterface.ush
@@ -16,8 +16,6 @@
 //   For each present attribute:
 //     <Packed element data> (element count * stride)
 
-#include "/Engine/Private/Quaternion.ush"
-
 // TODO: These defines could be code-gen'd to stay in sync with our PCGComputeConstants, see UComputeDataInterface::GetDefines()
 #define POINT_POSITION_ID 0
 #define POINT_ROTATION_ID 1
@@ -39,10 +37,6 @@
 #define ATTRIBUTE_HEADER_SIZE_BYTES 8
 #define KERNEL_EXECUTED_FLAG (1u << 31)
 
-// TODO: should be code-gen'd from PCGComputeCommon.h
-#define MAX_NUM_ATTRS 128
-#define NUM_RESERVED_ATTRS 32
-
 RWByteAddressBuffer {DataInterfaceName}_DataCollectionBuffer;
 
 // Multiplier that can be configured on output pins.
diff --git a/Engine/Plugins/PCG/Shaders/Private/PCGQuaternion.ush b/Engine/Plugins/PCG/Shaders/Private/PCGQuaternion.ush
new file mode 100644
--- /dev/null
+++ b/Engine/Plugins/PCG/Shaders/Private/PCGQuaternion.ush
@@ -0,0 +1,149 @@
+// Copyright Epic Games, Inc. All Rights Reserved.
+
+#pragma once
+
+// ~~ Begin duplicated from /Engine/Private/Quaternion.ush
+
+#define FQuat half4
+
+// Unless otherwise specified, use lerp + normalize to interpolate quats
+#ifndef QUAT_SLERP_APPROXIMATE
+	#define QUAT_SLERP_APPROXIMATE 1
+#endif
+
+FQuat QuatIdentity()
+{
+	return FQuat(0, 0, 0, 1);
+}
+
+FQuat QuatFromAxisAngle(in half3 AxisN, in float AngleRad)
+{
+	float S, C;
+	sincos(AngleRad * 0.5f, S, C);
+
+	return FQuat(AxisN * S, C);
+}
+
+FQuat QuatConjugate(in FQuat Q)
+{
+	return FQuat(-Q.xyz, Q.w);
+}
+
+float3 QuatRotateVector(in FQuat Q, in float3 V)
+{
+	float3 T = 2.0f * cross(Q.xyz, V);
+	return V + Q.w * T + cross(Q.xyz,T);
+}
+
+FQuat QuatMultiply(in FQuat Q1, in FQuat Q2)
+{
+	return FQuat(
+		(Q1.w * Q2.x) + (Q1.x * Q2.w) + (Q1.y * Q2.z) - (Q1.z * Q2.y),
+		(Q1.w * Q2.y) - (Q1.x * Q2.z) + (Q1.y * Q2.w) + (Q1.z * Q2.x),
+		(Q1.w * Q2.z) + (Q1.x * Q2.y) - (Q1.y * Q2.x) + (Q1.z * Q2.w),
+		(Q1.w * Q2.w) - (Q1.x * Q2.x) - (Q1.y * Q2.y) - (Q1.z * Q2.z)
+	);
+}
+
+half3x3 QuatToMatrix(in FQuat Q)
+{
+	return half3x3(
+		QuatRotateVector(Q, float3(1, 0, 0)),
+		QuatRotateVector(Q, float3(0, 1, 0)),
+		QuatRotateVector(Q, float3(0, 0, 1))
+	);
+}
+
+FQuat QuatFromMatrix(in half3x3 M)
+{
+	FQuat Q;
+	half T;
+
+	if (M[2][2] < 0.0f)
+	{
+		if (M[0][0] > M[1][1])
+		{
+			T = 1.0f + M[0][0] - M[1][1] - M[2][2];
+			Q = FQuat(T, M[0][1] + M[1][0], M[2][0] + M[0][2], M[1][2] - M[2][1]);
+		}
+		else
+		{
+			T = 1.0f - M[0][0] + M[1][1] - M[2][2];
+			Q = FQuat(M[0][1] + M[1][0], T, M[1][2] + M[2][1], M[2][0] - M[0][2]);
+		}
+	}
+	else
+	{
+		if (M[0][0] < -M[1][1])
+		{
+			T = 1.0f - M[0][0] - M[1][1] + M[2][2];
+			Q = FQuat(M[2][0] + M[0][2], M[1][2] + M[2][1], T, M[0][1] - M[1][0]);
+		}
+		else
+		{
+			T = 1.0f + M[0][0] + M[1][1] + M[2][2];
+			Q = FQuat(M[1][2] - M[2][1], M[2][0] - M[0][2], M[0][1] - M[1][0], T);
+		}
+	}
+
+	return Q * 0.5f * rsqrt(T);
+}
+
+FQuat QuatSlerpPrecise(FQuat Q0, FQuat Q1, half Slerp)
+{
+	// Get the half angle between quats.
+	const half C = dot(Q0, Q1);
+	const half AbsC = abs(C);
+
+	half Scale0, Scale1;
+	if (AbsC > 0.999f)
+	{
+		// Fall back to lerp when the angle between is tiny (prevents NaN from 1/sin(0))
+		Scale0 = half(1.0) - Slerp;
+		Scale1 = Slerp;
+	}
+	else
+	{
+		const half HalfAngle = acos(AbsC);
+		const half InvS = rcp(sin(HalfAngle));
+		Scale0 = sin((half(1.0) - Slerp) * HalfAngle) * InvS;
+		Scale1 = sin(Slerp * HalfAngle) * InvS;
+	}
+	
+	// Negate Scale1 if angle > 180 degrees to lerp in the shorter direction
+	Scale1 = C < 0.0f ? -Scale1 : Scale1;
+
+	return Scale0 * Q0 + Scale1 * Q1;
+}
+
+FQuat QuatSlerpApproximate(FQuat Q0, FQuat Q1, half Slerp)
+{
+	// Just use linear interpolation and normalize. Negate Q1 if the angle between them is
+	// greater than 180 degrees to lerp in the shorter direction.
+	const half C = dot(Q0, Q1); // cosine of the half angle between them
+	return normalize(lerp(Q0, C < 0.0f ? -Q1 : Q1, Slerp));
+}
+
+FQuat QuatSlerp(FQuat Q0, FQuat Q1, half Slerp)
+{
+#if QUAT_SLERP_APPROXIMATE
+	return QuatSlerpApproximate(Q0, Q1, Slerp);
+#else
+	return QuatSlerpPrecise(Q0, Q1, Slerp);
+#endif
+}
+
+FQuat QuatNormalize(in FQuat Q)
+{
+	float Len2 = dot(Q, Q);
+	if (Len2 > 0.0f)
+	{
+		return Q * rsqrt(Len2);
+	}
+	else
+	{
+		return QuatIdentity();
+	}
+}
+
+// ~~ End duplicated from /Engine/Private/Quaternion.ush
diff --git a/Engine/Plugins/PCG/Shaders/Private/PCGShaderUtils.ush b/Engine/Plugins/PCG/Shaders/Private/PCGShaderUtils.ush
--- a/Engine/Plugins/PCG/Shaders/Private/PCGShaderUtils.ush
+++ b/Engine/Plugins/PCG/Shaders/Private/PCGShaderUtils.ush
@@ -1,5 +1,7 @@
 // Copyright Epic Games, Inc. All Rights Reserved.
 
+#pragma once
+
 // These follow the corresponding C++ functions in PCGHelpers.
 uint ComputeSeed(uint A, uint B)
 {
@@ -30,36 +32,38 @@ float FRand(inout uint InOutSeed)
 
 float3 CreateGrid2D(int ElementIndex, int NumPoints, int NumRows, float3 Min, float3 Max)
 {
-    if (NumRows == 0)
+	if (NumRows == 0)
 	{
 		return Min;
 	}
 
-    const int NumCols = ceil((0.f + NumPoints) / NumRows);
+	const int NumCols = ceil((0.f + NumPoints) / NumRows);
 
-    if (NumCols == 0)
+	if (NumCols == 0)
 	{
 		return Min;
 	}
 
-    const int IndexX = ElementIndex / NumRows;
-    const int IndexY = ElementIndex - NumRows * IndexX;
+	const int IndexX = ElementIndex / NumRows;
+	const int IndexY = ElementIndex - NumRows * IndexX;
+
+	float3 Spacing = (Max - Min);
 
-    float3 Spacing = (Max - Min);
-    if (NumCols > 1)
+	if (NumCols > 1)
 	{
 		Spacing.x /= (NumCols - 1);
 	}
-    if (NumRows > 1)
+
+	if (NumRows > 1)
 	{
 		Spacing.y /= (NumRows - 1);
 	}
 
-    float3 Position = Min;
-    Position.x += IndexX * Spacing.x;
-    Position.y += IndexY * Spacing.y;
-
-    return Position;
+	float3 Position = Min;
+	Position.x += IndexX * Spacing.x;
+	Position.y += IndexY * Spacing.y;
+	
+	return Position;
 }
 
 float3 CreateGrid2D(int ElementIndex, int NumPoints, float3 Min, float3 Max)
@@ -69,44 +73,44 @@ float3 CreateGrid2D(int ElementIndex, int NumPoints, float3 Min, float3 Max)
 		return Min;
 	}
 
-    const int NumRows = ceil(sqrt(NumPoints));
-    return CreateGrid2D(ElementIndex, NumPoints, NumRows, Min, Max);
+	const int NumRows = ceil(sqrt(NumPoints));
+	return CreateGrid2D(ElementIndex, NumPoints, NumRows, Min, Max);
 }
 
 float3 CreateGrid3D(int ElementIndex, int NumPoints, int NumRows, int NumCols, float3 Min, float3 Max)
 {
-    if (NumRows == 0 || NumCols == 0)
+	if (NumRows == 0 || NumCols == 0)
 	{
 		return Min;
 	}
 
-    const int NumLayers = ceil((0.f + NumPoints) / (NumRows * NumCols));
-
-    if (NumLayers == 0)
+	const int NumLayers = ceil((0.f + NumPoints) / (NumRows * NumCols));
+	
+	if (NumLayers == 0)
 	{
 		return Min;
 	}
-
-    int3 Index;
-    Index.x = ElementIndex % NumCols;
-    Index.y = (ElementIndex / NumCols) % NumRows;
-    Index.z = ElementIndex / (NumRows * NumCols);
-
-    float3 Spacing = Max - Min;
-    if (NumCols > 1)
+	
+	int3 Index;
+	Index.x = ElementIndex % NumCols;
+	Index.y = (ElementIndex / NumCols) % NumRows;
+	Index.z = ElementIndex / (NumRows * NumCols);
+	
+	float3 Spacing = Max - Min;
+	if (NumCols > 1)
 	{
 		Spacing.x /= (NumCols - 1);
 	}
-    if (NumRows > 1)
+	if (NumRows > 1)
 	{
 		Spacing.y /= (NumRows - 1);
 	}
-    if (NumLayers > 1)
+	if (NumLayers > 1)
 	{
 		Spacing.z /= (NumLayers - 1);
 	}
-
-    return Min + Index * Spacing;
+	
+	return Min + Index * Spacing;
 }
 
 float3 CreateGrid3D(int ElementIndex, int NumPoints, float3 Min, float3 Max)
@@ -116,8 +120,8 @@ float3 CreateGrid3D(int ElementIndex, int NumPoints, float3 Min, float3 Max)
 		return Min;
 	}
 
-    const int NumRows = round(pow(NumPoints, 1.f / 3.f));
-    const int NumCols = NumRows;
-
-    return CreateGrid3D(ElementIndex, NumPoints, NumRows, NumCols, Min, Max);
+	const int NumRows = round(pow(NumPoints, 1.f / 3.f));
+	const int NumCols = NumRows;
+	
+	return CreateGrid3D(ElementIndex, NumPoints, NumRows, NumCols, Min, Max);
 }
diff --git a/Engine/Plugins/PCG/Shaders/Private/PCGShaderUtilsInternal.ush b/Engine/Plugins/PCG/Shaders/Private/PCGShaderUtilsInternal.ush
--- a/Engine/Plugins/PCG/Shaders/Private/PCGShaderUtilsInternal.ush
+++ b/Engine/Plugins/PCG/Shaders/Private/PCGShaderUtilsInternal.ush
@@ -1,5 +1,7 @@
 // Copyright Epic Games, Inc. All Rights Reserved.
 
+#pragma once
+
 // Must match C++ constants in PCGComputeCommon.h
 #ifndef NUM_RESERVED_ATTRS
 #define NUM_RESERVED_ATTRS 32
diff --git a/Engine/Plugins/PCG/Source/PCG/Internal/Compute/PCGComputeCommon.h b/Engine/Plugins/PCG/Source/PCG/Internal/Compute/PCGComputeCommon.h
--- a/Engine/Plugins/PCG/Source/PCG/Internal/Compute/PCGComputeCommon.h
+++ b/Engine/Plugins/PCG/Source/PCG/Internal/Compute/PCGComputeCommon.h
@@ -129,6 +129,10 @@ namespace PCGComputeHelpers
 
 	/** Produces the data interface name of a data label resolver. */
 	FString GetDataLabelResolverName(FName InPinLabel);
+
+#if WITH_EDITOR
+	void ConvertObjectPathToShaderFilePath(FString& InOutPath);
+#endif
 }
 
 namespace PCGComputeDummies
diff --git a/Engine/Plugins/PCG/Source/PCG/Private/Compute/DataInterfaces/PCGDataCollectionDataInterface.cpp b/Engine/Plugins/PCG/Source/PCG/Private/Compute/DataInterfaces/PCGDataCollectionDataInterface.cpp
--- a/Engine/Plugins/PCG/Source/PCG/Private/Compute/DataInterfaces/PCGDataCollectionDataInterface.cpp
+++ b/Engine/Plugins/PCG/Source/PCG/Private/Compute/DataInterfaces/PCGDataCollectionDataInterface.cpp
@@ -432,9 +432,6 @@ TCHAR const* UPCGDataCollectionDataInterface::GetShaderVirtualPath() const
 void UPCGDataCollectionDataInterface::GetShaderHash(FString& InOutKey) const
 {
 	GetShaderFileHash(TemplateFilePath, EShaderPlatform::SP_PCD3D_SM5).AppendString(InOutKey);
-
-	// Quaternion library
-	GetShaderFileHash(TEXT("/Engine/Private/Quaternion.ush"), EShaderPlatform::SP_PCD3D_SM5).AppendString(InOutKey);
 }
 
 void UPCGDataCollectionDataInterface::GetHLSL(FString& OutHLSL, FString const& InDataInterfaceName) const
diff --git a/Engine/Plugins/PCG/Source/PCG/Private/Compute/Elements/PCGCustomHLSLKernel.cpp b/Engine/Plugins/PCG/Source/PCG/Private/Compute/Elements/PCGCustomHLSLKernel.cpp
--- a/Engine/Plugins/PCG/Source/PCG/Private/Compute/Elements/PCGCustomHLSLKernel.cpp
+++ b/Engine/Plugins/PCG/Source/PCG/Private/Compute/Elements/PCGCustomHLSLKernel.cpp
@@ -7,10 +7,11 @@
 #include "PCGModule.h"
 #include "Compute/PCGComputeGraph.h"
 #include "Compute/PCGDataBinding.h"
+#include "Graph/PCGGPUGraphCompilationContext.h"
 #include "Compute/Elements/PCGCustomHLSL.h"
 
-#include "Internationalization/Regex.h"
 #include "Containers/StaticArray.h"
+#include "Internationalization/Regex.h"
 
 #include UE_INLINE_GENERATED_CPP_BY_NAME(PCGCustomHLSLKernel)
 
@@ -21,15 +22,6 @@ namespace PCGCustomHLSLKernel
 	constexpr TCHAR AttributeFunctionGet[] = { TEXT("Get") };
 	constexpr TCHAR AttributeFunctionSet[] = { TEXT("Set") };
 
-	void ConvertObjectPathToShaderFilePath(FString& InOutPath)
-	{
-		// Shader compiler recognizes "/Engine/Generated/..." path as special. 
-		// It doesn't validate file suffix etc.
-		InOutPath = FString::Printf(TEXT("/Engine/Generated/UObject%s.ush"), *InOutPath);
-		// Shader compilation result parsing will break if it finds ':' where it doesn't expect.
-		InOutPath.ReplaceCharInline(TEXT(':'), TEXT('@'));
-	}
-
 	FString GetKernelAttributeKeyAsString(const FPCGKernelAttributeKey& Key)
 	{
 		return FString::Format(TEXT("'{0}'"), { Key.Name.ToString() });
@@ -324,30 +316,19 @@ FString UPCGCustomHLSLKernel::GetCookedSource(const UPCGComputeGraph* InComputeG
 {
 	check(InComputeGraph);
 
+	// @todo_pcg: Allow overriding kernel source from an external kernel source
 	const UPCGCustomHLSLSettings* CustomHLSLSettings = CastChecked<UPCGCustomHLSLSettings>(Settings);
 	FString Source = CustomHLSLSettings->ShaderSource;
 	FString Functions = TEXT("\n\n// User kernel functions\n") + CustomHLSLSettings->ShaderFunctions;
 
+	// @todo_pcg: Include additional source files in custom HLSL parsing pass, so that they support attribute names, data labels, etc.
 	ProcessShaderSource(InComputeGraph, Source, Functions);
 
 	FString ShaderPathName = CustomHLSLSettings->GetPathName();
-	PCGCustomHLSLKernel::ConvertObjectPathToShaderFilePath(ShaderPathName);
+	PCGComputeHelpers::ConvertObjectPathToShaderFilePath(ShaderPathName);
 
 	//const bool bHasKernelKeyword = Source.Contains(TEXT("KERNEL"), ESearchCase::CaseSensitive);
 
-	FString Includes;
-	{
-		// Add with caution: Pulling in external includes has the danger that 1) we may pull in more than we expect
-		// if they include additional things in the future, and 2) if their functions change it could break user's
-		// kernel source. The latter may be mitigated by branching our own PCG version of these if needed in the future.
-
-		// TODO bring these in via additional sources so that their contents get hashed. Also these could generate declarations.
-		Includes += TEXT("#include \"/Engine/Private/ComputeShaderUtils.ush\"\n");
-		Includes += TEXT("#include \"/Engine/Private/Quaternion.ush\"\n");
-		Includes += TEXT("#include \"/Plugin/PCG/Private/PCGShaderUtils.ush\"\n");
-		Includes += TEXT("#include \"/Plugin/PCG/Private/PCGShaderUtilsInternal.ush\"\n");
-	}
-
 	const FIntVector GroupSize(PCGComputeConstants::THREAD_GROUP_SIZE, 1, 1);
 	const FString KernelFunc = FString::Printf(
 		TEXT("[numthreads(%d, %d, %d)]\nvoid %s(uint3 GroupId : SV_GroupID, uint GroupIndex : SV_GroupIndex)"),
@@ -444,7 +425,6 @@ FString UPCGCustomHLSLKernel::GetCookedSource(const UPCGComputeGraph* InComputeG
 	else*/
 	{
 		Result = FString::Printf(TEXT(
-			"%s\n\n" // Includes
 			"%s\n\n" // Functions
 			"%s\n" // KernelFunc
 			"{\n"
@@ -455,7 +435,7 @@ FString UPCGCustomHLSLKernel::GetCookedSource(const UPCGComputeGraph* InComputeG
 			"#line 0 \"%s\"\n" // ShaderPathName
 			"%s\n" // Source
 			"}\n"),
-			*Includes, *Functions, *KernelFunc, *SetAsExecuted, *UnWrappedDispatchThreadId, *KernelSpecificPreamble, *ShaderPathName, *Source);
+			*Functions, *KernelFunc, *SetAsExecuted, *UnWrappedDispatchThreadId, *KernelSpecificPreamble, *ShaderPathName, *Source);
 	}
 
 	return Result;
diff --git a/Engine/Plugins/PCG/Source/PCG/Private/Compute/PCGComputeCommon.cpp b/Engine/Plugins/PCG/Source/PCG/Private/Compute/PCGComputeCommon.cpp
--- a/Engine/Plugins/PCG/Source/PCG/Private/Compute/PCGComputeCommon.cpp
+++ b/Engine/Plugins/PCG/Source/PCG/Private/Compute/PCGComputeCommon.cpp
@@ -141,6 +141,17 @@ namespace PCGComputeHelpers
 	{
 		return FString::Format(TEXT("{0}_DataResolver"), { *InPinLabel.ToString() });
 	}
+
+#if WITH_EDITOR
+	void ConvertObjectPathToShaderFilePath(FString& InOutPath)
+	{
+		// Shader compiler recognizes "/Engine/Generated/..." path as special. 
+		// It doesn't validate file suffix etc.
+		InOutPath = FString::Printf(TEXT("/Engine/Generated/UObject%s.ush"), *InOutPath);
+		// Shader compilation result parsing will break if it finds ':' where it doesn't expect.
+		InOutPath.ReplaceCharInline(TEXT(':'), TEXT('@'));
+	}
+#endif
 }
 
 namespace PCGComputeDummies
diff --git a/Engine/Plugins/PCG/Source/PCG/Private/Compute/PCGComputeKernel.cpp b/Engine/Plugins/PCG/Source/PCG/Private/Compute/PCGComputeKernel.cpp
--- a/Engine/Plugins/PCG/Source/PCG/Private/Compute/PCGComputeKernel.cpp
+++ b/Engine/Plugins/PCG/Source/PCG/Private/Compute/PCGComputeKernel.cpp
@@ -8,8 +8,12 @@
 #include "PCGNode.h"
 #include "PCGSettings.h"
 #include "Compute/PCGComputeGraph.h"
+#include "Compute/PCGComputeSource.h"
 #include "Compute/PCGDataBinding.h"
 #include "Compute/PCGPinPropertiesGPU.h"
+#include "Graph/PCGGPUGraphCompilationContext.h"
+
+#include "ShaderCompilerCore.h"
 
 #include UE_INLINE_GENERATED_CPP_BY_NAME(PCGComputeKernel)
 
@@ -59,6 +63,30 @@ bool UPCGComputeKernel::AreKernelSettingsValid(FPCGContext* InContext) const
 	return !bHasStaticValidationErrors;
 }
 
+#if WITH_EDITOR
+void UPCGComputeKernel::CreateAdditionalSources(FPCGGPUCompilationContext& InOutContext, UObject* InObjectOuter, TArray<TObjectPtr<UComputeSource>>& OutAdditionalSources) const
+{
+	auto AddComputeSourceForFile = [&InOutContext, InObjectOuter, &OutAdditionalSources](const FString& InSourceFile)
+	{
+		FString Source;
+
+		if (ensure(LoadShaderSourceFile(*InSourceFile, EShaderPlatform::SP_PCD3D_SM5, &Source, nullptr)))
+		{
+			UPCGComputeSource* AdditionalSource = InOutContext.NewObject_AnyThread<UPCGComputeSource>(InObjectOuter);
+			AdditionalSource->SetSource(MoveTemp(Source));
+			OutAdditionalSources.Add(AdditionalSource);
+		}
+	};
+
+	AddComputeSourceForFile(TEXT("/Engine/Private/ComputeShaderUtils.ush"));
+	AddComputeSourceForFile(TEXT("/Plugin/PCG/Private/PCGShaderUtils.ush"));
+	AddComputeSourceForFile(TEXT("/Plugin/PCG/Private/PCGShaderUtilsInternal.ush"));
+
+	// Note: PCGDataCollectionDataInterface.ush depends on PCGQuaternion.ush, therefore all kernels also depend on PCGQuaternion.ush
+	AddComputeSourceForFile(TEXT("/Plugin/PCG/Private/PCGQuaternion.ush"));
+}
+#endif
+
 void UPCGComputeKernel::ComputeDataDescFromPinProperties(const FPCGPinPropertiesGPU& OutputPinProps, const TArrayView<const FPCGPinProperties>& InInputPinProps, UPCGDataBinding* InBinding, FPCGDataCollectionDesc& OutPinDesc) const
 {
 	check(InBinding);
diff --git a/Engine/Plugins/PCG/Source/PCG/Private/Compute/PCGComputeKernel.h b/Engine/Plugins/PCG/Source/PCG/Private/Compute/PCGComputeKernel.h
--- a/Engine/Plugins/PCG/Source/PCG/Private/Compute/PCGComputeKernel.h
+++ b/Engine/Plugins/PCG/Source/PCG/Private/Compute/PCGComputeKernel.h
@@ -78,6 +78,8 @@ public:
 	/** Get the name of the main kernel function in the source. This name is also displayed in GPU profile scopes. */
 	virtual FString GetEntryPoint() const PURE_VIRTUAL(UPCGComputeKernel::GetEntryPoint, return TEXT("Main"););
 
+	virtual void CreateAdditionalSources(FPCGGPUCompilationContext& InOutContext, UObject* InObjectOuter, TArray<TObjectPtr<class UComputeSource>>& OutAdditionalSources) const;
+
 	/** Create additional input data interfaces to marshal any required input data. */
 	virtual void CreateAdditionalInputDataInterfaces(FPCGGPUCompilationContext& InOutContext, UObject* InObjectOuter, TArray<TObjectPtr<UComputeDataInterface>>& OutDataInterfaces) const {}
 
diff --git a/Engine/Plugins/PCG/Source/PCG/Private/Compute/PCGComputeKernelSource.h b/Engine/Plugins/PCG/Source/PCG/Private/Compute/PCGComputeKernelSource.h
--- a/Engine/Plugins/PCG/Source/PCG/Private/Compute/PCGComputeKernelSource.h
+++ b/Engine/Plugins/PCG/Source/PCG/Private/Compute/PCGComputeKernelSource.h
@@ -12,19 +12,22 @@ class PCG_API UPCGComputeKernelSource : public UComputeKernelSource
 	GENERATED_BODY()
 
 public:
-	void SetSource(FString const& InSource)
-	{
-		Source = InSource;
-	}
-	
 	//~ Begin UComputeKernelSource Interface.
 	FString GetSource() const override
 	{
+#if WITH_EDITOR
 		return Source;
+#else
+		ensure(false);
+		return "";
+#endif
 	}
 	//~ End UComputeKernelSource Interface.
 
+#if WITH_EDITOR
+	void SetSource(FString const& InSource) { Source = InSource; }
+
 protected:
-	UPROPERTY()
 	FString Source;
+#endif
 };
diff --git a/Engine/Plugins/PCG/Source/PCG/Private/Compute/PCGComputeSource.cpp b/Engine/Plugins/PCG/Source/PCG/Private/Compute/PCGComputeSource.cpp
new file mode 100644
--- /dev/null
+++ b/Engine/Plugins/PCG/Source/PCG/Private/Compute/PCGComputeSource.cpp
@@ -0,0 +1,29 @@
+// Copyright Epic Games, Inc. All Rights Reserved.
+
+#include "Compute/PCGComputeSource.h"
+
+#include "Compute/PCGComputeCommon.h"
+
+#include UE_INLINE_GENERATED_CPP_BY_NAME(PCGComputeSource)
+
+FString UPCGComputeSource::GetSource() const
+{
+#if WITH_EDITOR
+	return Source;
+#else
+	ensure(false);
+	return "";
+#endif
+}
+
+FString UPCGComputeSource::GetVirtualPath() const
+{
+#if WITH_EDITOR
+	FString ShaderPathName = GetPathName();
+	PCGComputeHelpers::ConvertObjectPathToShaderFilePath(ShaderPathName);
+	return ShaderPathName;
+#else
+	ensure(false);
+	return "";
+#endif
+}
diff --git a/Engine/Plugins/PCG/Source/PCG/Private/Compute/PCGComputeSource.h b/Engine/Plugins/PCG/Source/PCG/Private/Compute/PCGComputeSource.h
new file mode 100644
--- /dev/null
+++ b/Engine/Plugins/PCG/Source/PCG/Private/Compute/PCGComputeSource.h
@@ -0,0 +1,26 @@
+// Copyright Epic Games, Inc. All Rights Reserved.
+
+#pragma once
+
+#include "ComputeFramework/ComputeSource.h"
+
+#include "PCGComputeSource.generated.h"
+
+UCLASS()
+class UPCGComputeSource : public UComputeSource
+{
+	GENERATED_BODY()
+
+public:
+	//~ Begin UComputeSource Interface.
+	FString GetSource() const override;
+	FString GetVirtualPath() const override;
+	//~ End UComputeSource Interface.
+
+#if WITH_EDITOR
+	void SetSource(FString const& InSource) { Source = InSource; }
+
+protected:
+	FString Source;
+#endif
+};
diff --git a/Engine/Plugins/PCG/Source/PCG/Private/Graph/PCGGraphCompilerGPU.cpp b/Engine/Plugins/PCG/Source/PCG/Private/Graph/PCGGraphCompilerGPU.cpp
--- a/Engine/Plugins/PCG/Source/PCG/Private/Graph/PCGGraphCompilerGPU.cpp
+++ b/Engine/Plugins/PCG/Source/PCG/Private/Graph/PCGGraphCompilerGPU.cpp
@@ -1424,6 +1424,7 @@ void FPCGGraphCompilerGPU::CompileComputeGraph(
 			// TODO: Instead of passing a half-baked compute graph into this virtual call, we should pass in a compilation context per compute graph.
 			FString Source = KernelWithBindings.Kernel->GetCookedSource(InOutComputeGraph);
 			KernelSource->SetSource(MoveTemp(Source));
+			KernelWithBindings.Kernel->CreateAdditionalSources(InOutContext, KernelWithBindings.Kernel, KernelSource->AdditionalSources);
 
 			if (Settings->bDumpCookedHLSL)
 			{
