[PCG] Expose polyline/spline data API to BP

#rnx

diff --git a/Engine/Plugins/PCG/Source/PCG/Private/Data/PCGPolyLineData.cpp b/Engine/Plugins/PCG/Source/PCG/Private/Data/PCGPolyLineData.cpp
--- a/Engine/Plugins/PCG/Source/PCG/Private/Data/PCGPolyLineData.cpp
+++ b/Engine/Plugins/PCG/Source/PCG/Private/Data/PCGPolyLineData.cpp
@@ -32,6 +32,11 @@ FVector::FReal UPCGPolyLineData::GetLength() const
 	return Length;
 }
 
+FTransform UPCGPolyLineData::K2_GetTransformAtDistance(int SegmentIndex, double Distance, FBox& OutBounds, bool bWorldSpace) const
+{
+	return GetTransformAtDistance(SegmentIndex, Distance, bWorldSpace, &OutBounds);
+}
+
 float UPCGPolyLineData::GetInputKeyAtAlpha(float Alpha) const
 {
 	Alpha = FMath::Clamp(Alpha, 0, 1);
diff --git a/Engine/Plugins/PCG/Source/PCG/Private/Data/PCGSplineData.cpp b/Engine/Plugins/PCG/Source/PCG/Private/Data/PCGSplineData.cpp
--- a/Engine/Plugins/PCG/Source/PCG/Private/Data/PCGSplineData.cpp
+++ b/Engine/Plugins/PCG/Source/PCG/Private/Data/PCGSplineData.cpp
@@ -64,6 +64,12 @@ void UPCGSplineData::Initialize(const TArray<FSplinePoint>& InSplinePoints, bool
 	CachedBounds = CachedBounds.TransformBy(InTransform);
 }
 
+void UPCGSplineData::K2_Initialize(const TArray<FSplinePoint>& InSplinePoints, bool bInClosedLoop, const FTransform& InTransform,
+	TArray<int64> InOptionalEntryKeys)
+{
+	Initialize(InSplinePoints, bInClosedLoop, InTransform, std::move(InOptionalEntryKeys));
+}
+
 void UPCGSplineData::Initialize(const FPCGSplineStruct& InSplineStruct)
 {
 	SplineStruct = InSplineStruct;
@@ -427,6 +433,34 @@ const UPCGBasePointData* UPCGSplineData::CreateBasePointData(FPCGContext* Contex
 	return Data;
 }
 
+TArray<FSplinePoint> UPCGSplineData::GetSplinePoints() const
+{
+	const FInterpCurveVector& ControlPointsPosition = SplineStruct.GetSplinePointsPosition();
+	const FInterpCurveQuat& ControlPointsRotation = SplineStruct.GetSplinePointsRotation();
+	const FInterpCurveVector& ControlPointsScale = SplineStruct.GetSplinePointsScale();
+
+	TArray<FSplinePoint> ControlPoints;
+	ControlPoints.Reserve(ControlPointsPosition.Points.Num());
+
+	for (int i = 0; i < ControlPointsPosition.Points.Num() - 1; ++i)
+	{
+		ControlPoints.Emplace(static_cast<float>(ControlPoints.Num()),
+			ControlPointsPosition.Points[i].OutVal,
+			ControlPointsPosition.Points[i].ArriveTangent,
+			ControlPointsPosition.Points[i].LeaveTangent,
+			ControlPointsRotation.Points[i].OutVal.Rotator(),
+			ControlPointsScale.Points[i].OutVal,
+			ConvertInterpCurveModeToSplinePointType(ControlPointsPosition.Points[i].InterpMode));
+	}
+
+	return ControlPoints;
+}
+
+TArray<int64> UPCGSplineData::GetMetadataEntryKeysForSplinePoints() const
+{
+	return TArray<int64>(SplineStruct.GetConstControlPointsEntryKeys());
+}
+
 FBox UPCGSplineData::GetBounds() const
 {
 	return CachedBounds;
diff --git a/Engine/Plugins/PCG/Source/PCG/Public/Data/PCGPolyLineData.h b/Engine/Plugins/PCG/Source/PCG/Public/Data/PCGPolyLineData.h
--- a/Engine/Plugins/PCG/Source/PCG/Public/Data/PCGPolyLineData.h
+++ b/Engine/Plugins/PCG/Source/PCG/Public/Data/PCGPolyLineData.h
@@ -22,55 +22,74 @@ public:
 	//~End UPCGSpatialData interface
 
 	/** Get the world-space transform of the entire line. */
+	UFUNCTION(BlueprintCallable, Category = "PolylineData")
 	virtual FTransform GetTransform() const { return FTransform::Identity; }
 
 	/** Get the number of segments in this line. If the line is closed, this will be the same as the number of control points in the line. */
+	UFUNCTION(BlueprintCallable, Category = "PolylineData")
 	virtual int GetNumSegments() const PURE_VIRTUAL(UPCGPolyLineData::GetNumSegments, return 0;);
 
 	/** Get the length of a specific segment of the line. */
-	virtual FVector::FReal GetSegmentLength(int SegmentIndex) const PURE_VIRTUAL(UPCGPolyLineData::GetSegmentLength, return 0;);
+	UFUNCTION(BlueprintCallable, Category = "PolylineData")
+	virtual double GetSegmentLength(int SegmentIndex) const PURE_VIRTUAL(UPCGPolyLineData::GetSegmentLength, return 0;);
 
 	/** Get the location of the point at the normalized [0, 1] parameter across the entire the poly line. */
+	UFUNCTION(BlueprintCallable, Category = "PolylineData")
 	virtual FVector GetLocationAtAlpha(float Alpha) const PURE_VIRTUAL(UPCGPolyLine::GetLocationAtAlpha, return FVector::ZeroVector;)
 
 	/** Get the full transform at the normalized [0, 1] parameter across the entire the poly line. */
+	UFUNCTION(BlueprintCallable, Category = "PolylineData")
 	virtual FTransform GetTransformAtAlpha(float Alpha) const PURE_VIRTUAL(UPCGPolyLine::GetTransformAtAlpha, return FTransform::Identity;)
 
 	/** Get the total length of the line. */
-	virtual FVector::FReal GetLength() const;
+	UFUNCTION(BlueprintCallable, Category = "PolylineData")
+	virtual double GetLength() const;
 
 	/** Get the location at a distance along the line. */
-	virtual FTransform GetTransformAtDistance(int SegmentIndex, FVector::FReal Distance, bool bWorldSpace = true, FBox* OutBounds = nullptr) const PURE_VIRTUAL(UPCGPolyLine::GetTransformAtDistance, return FTransform(););
+	virtual FTransform GetTransformAtDistance(int SegmentIndex, double Distance, bool bWorldSpace = true, FBox* OutBounds = nullptr) const PURE_VIRTUAL(UPCGPolyLine::GetTransformAtDistance, return FTransform(););
 
 	/** Get the location at a distance along the line. */
-	virtual FVector GetLocationAtDistance(int SegmentIndex, FVector::FReal Distance, bool bWorldSpace = true) const { return GetTransformAtDistance(SegmentIndex, Distance, bWorldSpace).GetLocation(); }
+	UFUNCTION(BlueprintCallable, Category = "PolylineData", meta = (DisplayName = "Get Transform At Distance", AutoCreateRefTerm = "OutBounds"))
+	FTransform K2_GetTransformAtDistance(int SegmentIndex, double Distance, UPARAM(ref) FBox& OutBounds, bool bWorldSpace = true) const;
+	
+	/** Get the location at a distance along the line. */
+	UFUNCTION(BlueprintCallable, Category = "PolylineData")
+	virtual FVector GetLocationAtDistance(int SegmentIndex, double Distance, bool bWorldSpace = true) const { return GetTransformAtDistance(SegmentIndex, Distance, bWorldSpace).GetLocation(); }
 
 	/** Get the curvature at a distance along the line. */
-	virtual FVector::FReal GetCurvatureAtDistance(int SegmentIndex, FVector::FReal Distance) const { return 0; }
+	UFUNCTION(BlueprintCallable, Category = "PolylineData")
+	virtual double GetCurvatureAtDistance(int SegmentIndex, double Distance) const { return 0; }
 
 	/**
 	 * Get a value [0,1] representing how far along the point is to the end of the line. Each segment on the line represents a same-size interval.
 	 * For example, if there are three segments, each segment will take up 0.333... of the interval.
 	 */
-	virtual float GetAlphaAtDistance(int SegmentIndex, FVector::FReal Distance) const;
+	UFUNCTION(BlueprintCallable, Category = "PolylineData")
+	virtual float GetAlphaAtDistance(int SegmentIndex, double Distance) const;
 
 	/** Get the input key at a distance along the line. InputKey is a float value in [0, N], where N is the number of control points. Each range [i, i+1] represents an interpolation from 0 to 1 across spline segment i. */
-	virtual float GetInputKeyAtDistance(int SegmentIndex, FVector::FReal Distance) const { return 0; }
+	UFUNCTION(BlueprintCallable, Category = "PolylineData")
+	virtual float GetInputKeyAtDistance(int SegmentIndex, double Distance) const { return 0; }
 
 	/** Get the input key from the normalized distance of [0, 1] across the entire the poly line. */
+	UFUNCTION(BlueprintCallable, Category = "PolylineData")
 	virtual float GetInputKeyAtAlpha(float Alpha) const;
 
 	/** Get the arrive and leave tangents for a control point via its segment index. */
+	UFUNCTION(BlueprintCallable, Category = "PolylineData")
 	virtual void GetTangentsAtSegmentStart(int SegmentIndex, FVector& OutArriveTangent, FVector& OutLeaveTangent) const;
 
 	/** Get the cumulative distance along the line to the start of a segment. */
-	virtual FVector::FReal GetDistanceAtSegmentStart(int SegmentIndex) const { return 0; }
+	UFUNCTION(BlueprintCallable, Category = "PolylineData")
+	virtual double GetDistanceAtSegmentStart(int SegmentIndex) const { return 0; }
 
 	/** True if the line is a closed loop. */
+	UFUNCTION(BlueprintCallable, Category = "PolylineData")
 	virtual bool IsClosed() const { return false; }
 	
 	/** This function should be called in the Sample/Project point function, but can also be called if the sampling is done manually.
 	 * This is meant for each child class to write its metadata given the InputKey, so interpolation can also be done.
 	 */
+	UFUNCTION(BlueprintCallable, Category = "PolylineData")
 	virtual void WriteMetadataToPoint(float InputKey, FPCGPoint& OutPoint, UPCGMetadata* OutMetadata) const {}
 };
diff --git a/Engine/Plugins/PCG/Source/PCG/Public/Data/PCGSplineData.h b/Engine/Plugins/PCG/Source/PCG/Public/Data/PCGSplineData.h
--- a/Engine/Plugins/PCG/Source/PCG/Public/Data/PCGSplineData.h
+++ b/Engine/Plugins/PCG/Source/PCG/Public/Data/PCGSplineData.h
@@ -28,8 +28,31 @@ class PCG_API UPCGSplineData : public UPCGPolyLineData
 public:
 	UPCGSplineData();
 
+	/**
+	 * Initialize a spline data from a Spline Component.
+	 */
+	UFUNCTION(BlueprintCallable, Category = "SplineData", meta = (DisplayName = "Initialize From Spline Component"))
 	void Initialize(const USplineComponent* InSpline);
+
+	/**
+	 * Initialize a spline data
+	 * @param InSplinePoints       List of control points.
+	 * @param bInClosedLoop        If the spline is closed or not.
+	 * @param InTransform          Transform of the spline
+	 * @param InOptionalEntryKeys  Optional list of metadata entry keys to setup for the control points. Need to the same size as the number of control points.
+	 */
 	void Initialize(const TArray<FSplinePoint>& InSplinePoints, bool bInClosedLoop, const FTransform& InTransform, TArray<PCGMetadataEntryKey> InOptionalEntryKeys = {});
+
+	/**
+	 * Initialize a spline data
+	 * @param InSplinePoints       List of control points.
+	 * @param bInClosedLoop        If the spline is closed or not.
+	 * @param InTransform          Transform of the spline
+	 * @param InOptionalEntryKeys  Optional list of metadata entry keys to setup for the control points. Need to the same size as the number of control points.
+	 */
+	UFUNCTION(BlueprintCallable, Category = "SplineData", meta = (DisplayName = "Initialize From Spline Points", AutoCreateRefTerm = "InOptionalEntryKeys"))
+	void K2_Initialize(const TArray<FSplinePoint>& InSplinePoints, bool bInClosedLoop, const FTransform& InTransform, TArray<int64> InOptionalEntryKeys);
+	
 	void Initialize(const FPCGSplineStruct& InSplineStruct);
 	void ApplyTo(USplineComponent* InSpline) const;
 
@@ -83,6 +106,12 @@ protected:
 
 	const UPCGBasePointData* CreateBasePointData(FPCGContext* Context, TSubclassOf<UPCGBasePointData> PointDataClass) const;
 
+	UFUNCTION(BlueprintCallable, Category = "SplineData")
+	TArray<FSplinePoint> GetSplinePoints() const;
+
+	UFUNCTION(BlueprintCallable, Category = "SplineData")
+	TArray<int64> GetMetadataEntryKeysForSplinePoints() const;
+
 public:
 	// Minimal data needed to replicate the behavior from USplineComponent
 	UPROPERTY()
