[PCG] Allow drawing debug box for virtual texture priming bounds in standalone

#rnx

diff --git a/Engine/Plugins/PCG/Source/PCG/Private/RuntimeGen/PCGRuntimeGenScheduler.cpp b/Engine/Plugins/PCG/Source/PCG/Private/RuntimeGen/PCGRuntimeGenScheduler.cpp
--- a/Engine/Plugins/PCG/Source/PCG/Private/RuntimeGen/PCGRuntimeGenScheduler.cpp
+++ b/Engine/Plugins/PCG/Source/PCG/Private/RuntimeGen/PCGRuntimeGenScheduler.cpp
@@ -18,6 +18,7 @@
 #include "RuntimeGen/GenSources/PCGGenSourceComponent.h"
 #include "RuntimeGen/PCGGenSourceManager.h"
 
+#include "DrawDebugHelpers.h"
 #include "Components/RuntimeVirtualTextureComponent.h"
 #include "UObject/UObjectGlobals.h"
 #include "VT/RuntimeVirtualTexture.h"
@@ -26,7 +27,6 @@
 #include "WorldPartition/WorldPartitionSubsystem.h"
 
 #if WITH_EDITOR
-#include "DrawDebugHelpers.h"
 #include "EditorViewportClient.h"
 #endif
 
@@ -95,12 +95,10 @@ namespace PCGRuntimeGenSchedulerHelpers
 		false,
 		TEXT("Enable priming of virtual textures for PCG Components which request it."));
 
-#if WITH_EDITOR
 	static TAutoConsoleVariable<bool> CVarDebugDrawTexturePrimingBounds(
 		TEXT("pcg.VirtualTexturePriming.DebugDrawTexturePrimingBounds"),
 		false,
 		TEXT("Draws debug boxes to indicate regions where PCG is requesting virtual texture priming."));
-#endif
 }
 
 FPCGGridDescriptor FPCGRuntimeGenScheduler::FGridGenerationKey::GetGridDescriptor() const
@@ -999,7 +997,6 @@ void FPCGRuntimeGenScheduler::TickRequestVirtualTexturePriming(const TSet<IPCGGe
 
 						VirtualTextureComponent->RequestPreload(PrimingBounds, /*Level=*/0);
 
-#if WITH_EDITOR
 						if (PCGRuntimeGenSchedulerHelpers::CVarDebugDrawTexturePrimingBounds.GetValueOnAnyThread())
 						{
 							check(World);
@@ -1011,7 +1008,6 @@ void FPCGRuntimeGenScheduler::TickRequestVirtualTexturePriming(const TSet<IPCGGe
 								/*bPersistentLines=*/false,
 								/*LifeTime=*/0.02f);
 						}
-#endif
 					}
 				}
 			}
