[PCG] Normal to Density GPU conversion.

#rnx

diff --git a/Engine/Plugins/PCG/Shaders/Private/Elements/PCGNormalToDensity.usf b/Engine/Plugins/PCG/Shaders/Private/Elements/PCGNormalToDensity.usf
new file mode 100644
--- /dev/null
+++ b/Engine/Plugins/PCG/Shaders/Private/Elements/PCGNormalToDensity.usf
@@ -0,0 +1,96 @@
+// Copyright Epic Games, Inc. All Rights Reserved.
+
+#include "/Engine/Private/ComputeShaderUtils.ush"
+#include "/Engine/Private/Quaternion.ush"
+#include "/Plugin/PCG/Private/PCGShaderUtils.ush"
+#include "/Plugin/PCG/Private/PCGShaderUtilsInternal.ush"
+
+// Mirrors PCGNormalToDensityMode
+#define NORMAL_TO_DENSITY_MODE_Set 0
+#define NORMAL_TO_DENSITY_MODE_Minimum 1
+#define NORMAL_TO_DENSITY_MODE_Maximum 2
+#define NORMAL_TO_DENSITY_MODE_Add 3
+#define NORMAL_TO_DENSITY_MODE_Subtract 4
+#define NORMAL_TO_DENSITY_MODE_Multiply 5
+#define NORMAL_TO_DENSITY_MODE_Divide 6
+
+[numthreads(64, 1, 1)]
+void PCGNormalToDensityCS(uint3 GroupId : SV_GroupID, uint GroupIndex : SV_GroupIndex)
+{
+	// Mark the kernel as having executed. Must run before we early out via thread index, because the kernel is still 'executed' even if the number of
+	// threads to iterate on is zero. Even if GetNumThreads() returns 0, the kernel will still have been dispatched on a single thread to set this flag.
+	if (all(GroupId == 0) && GroupIndex == 0)
+	{
+		Out_SetAsExecutedInternal();
+	}
+
+	const uint ThreadIndex = GetUnWrappedDispatchThreadId(GroupId, GroupIndex, 64);
+	if (ThreadIndex >= GetNumThreads().x) return;
+
+	uint InDataIndex, InElemIndex;
+	if (!In_GetThreadData(ThreadIndex, InDataIndex, InElemIndex))
+	{
+		return;
+	}
+
+	uint OutDataIndex, OutElemIndex;
+	if (!Out_GetThreadData(ThreadIndex, OutDataIndex, OutElemIndex))
+	{
+		return;
+	}
+
+	// Propagate the 'Removed' status to the output point. Otherwise, this point will not be culled.
+	if (In_IsPointRemoved(InDataIndex, InElemIndex))
+	{
+		Out_RemovePoint(OutDataIndex, OutElemIndex);
+		return;
+	}
+
+	// Do the actual 'work' of this node.
+	const float3 Normal = normalize(NormalToDensity_GetNormal());
+	const float Offset = NormalToDensity_GetOffset();
+	const float Strength = max(NormalToDensity_GetStrength(), 0.0001);
+	const uint DensityMode = NormalToDensity_GetDensityMode();
+
+	const float InDensity = In_GetDensity(InDataIndex, InElemIndex);
+	const float4 InRotation =  In_GetRotation(InDataIndex, InElemIndex);
+	const float3 PointUp = QuatRotateVector(InRotation, float3(0, 0, 1));
+	float NormalDensity = pow(clamp(dot(PointUp, Normal) + Offset, 0.0, 1.0), 1 / Strength);
+
+	float OutDensity = 0;
+
+	if (DensityMode == NORMAL_TO_DENSITY_MODE_Set)
+	{
+		OutDensity = NormalDensity;
+	}
+	if (DensityMode == NORMAL_TO_DENSITY_MODE_Minimum)
+	{
+		OutDensity = min(InDensity, NormalDensity);
+	}
+	else if (DensityMode == NORMAL_TO_DENSITY_MODE_Maximum)
+	{
+		OutDensity = max(InDensity, NormalDensity);
+	}
+	else if (DensityMode == NORMAL_TO_DENSITY_MODE_Add)
+	{
+		OutDensity = InDensity + NormalDensity;
+	}
+	else if (DensityMode == NORMAL_TO_DENSITY_MODE_Subtract)
+	{
+		OutDensity = InDensity - NormalDensity;
+	}
+	else if (DensityMode == NORMAL_TO_DENSITY_MODE_Multiply)
+	{
+		OutDensity = InDensity * NormalDensity;
+	}
+	else if (DensityMode == NORMAL_TO_DENSITY_MODE_Divide)
+	{
+		OutDensity = (NormalDensity != 0) ? InDensity / NormalDensity  : 0.0f;
+	}
+	
+	// Copy all attributes
+	PCG_COPY_ALL_ATTRIBUTES_TO_OUTPUT(Out, In, OutDataIndex, OutElemIndex, InDataIndex, InElemIndex);
+	
+	// Write output density.
+	Out_SetDensity(OutDataIndex, OutElemIndex, OutDensity);
+}
diff --git a/Engine/Plugins/PCG/Source/PCG/Private/Compute/DataInterfaces/Elements/PCGNormalToDensityDataInterface.cpp b/Engine/Plugins/PCG/Source/PCG/Private/Compute/DataInterfaces/Elements/PCGNormalToDensityDataInterface.cpp
new file mode 100644
--- /dev/null
+++ b/Engine/Plugins/PCG/Source/PCG/Private/Compute/DataInterfaces/Elements/PCGNormalToDensityDataInterface.cpp
@@ -0,0 +1,120 @@
+// Copyright Epic Games, Inc. All Rights Reserved.
+
+#include "PCGNormalToDensityDataInterface.h"
+
+#include "Compute/PCGComputeGraph.h"
+#include "Compute/PCGDataBinding.h"
+#include "Compute/Elements/PCGNormalToDensityKernel.h"
+
+#include "RenderGraphBuilder.h"
+#include "RenderGraphResources.h"
+#include "RHIResources.h"
+#include "ShaderCompilerCore.h"
+#include "ShaderCore.h"
+#include "ShaderParameterMetadataBuilder.h"
+#include "ComputeFramework/ShaderParamTypeDefinition.h"
+
+void UPCGNormalToDensityDataInterface::GetSupportedInputs(TArray<FShaderFunctionDefinition>& OutFunctions) const
+{
+	OutFunctions.AddDefaulted_GetRef()
+		.SetName(TEXT("NormalToDensity_GetNormal"))
+		.AddReturnType(FShaderValueType::Get(EShaderFundamentalType::Float, 3));
+	
+	OutFunctions.AddDefaulted_GetRef()
+		.SetName(TEXT("NormalToDensity_GetOffset"))
+		.AddReturnType(FShaderValueType::Get(EShaderFundamentalType::Float));
+
+	OutFunctions.AddDefaulted_GetRef()
+		.SetName(TEXT("NormalToDensity_GetStrength"))
+		.AddReturnType(FShaderValueType::Get(EShaderFundamentalType::Float));
+
+	OutFunctions.AddDefaulted_GetRef()
+		.SetName(TEXT("NormalToDensity_GetDensityMode"))
+		.AddReturnType(FShaderValueType::Get(EShaderFundamentalType::Uint));
+}
+
+BEGIN_SHADER_PARAMETER_STRUCT(FPCGNormalToDensityDataInterfaceParameters,)
+	SHADER_PARAMETER(FVector3f, Normal)
+	SHADER_PARAMETER(float, Offset)
+	SHADER_PARAMETER(float, Strength)
+	SHADER_PARAMETER(uint32, DensityMode)
+END_SHADER_PARAMETER_STRUCT()
+
+void UPCGNormalToDensityDataInterface::GetShaderParameters(TCHAR const* UID, FShaderParametersMetadataBuilder& InOutBuilder, FShaderParametersMetadataAllocations& InOutAllocations) const
+{
+	InOutBuilder.AddNestedStruct<FPCGNormalToDensityDataInterfaceParameters>(UID);
+}
+
+void UPCGNormalToDensityDataInterface::GetHLSL(FString& OutHLSL, FString const& InDataInterfaceName) const
+{
+	TMap<FString, FStringFormatArg> TemplateArgs =
+	{
+		{ TEXT("DataInterfaceName"), InDataInterfaceName },
+	};
+
+	OutHLSL += FString::Format(TEXT(
+		"float3 {DataInterfaceName}_Normal;\n"
+		"float {DataInterfaceName}_Offset;\n"
+		"float {DataInterfaceName}_Strength;\n"
+		"uint {DataInterfaceName}_DensityMode;\n"
+		"\n"
+		"float3 NormalToDensity_GetNormal_{DataInterfaceName}() { return {DataInterfaceName}_Normal; }\n"
+		"float NormalToDensity_GetOffset_{DataInterfaceName}() { return {DataInterfaceName}_Offset; }\n"
+		"float NormalToDensity_GetStrength_{DataInterfaceName}() { return {DataInterfaceName}_Strength; }\n"
+		"uint NormalToDensity_GetDensityMode_{DataInterfaceName}() { return {DataInterfaceName}_DensityMode; }\n"
+		), TemplateArgs);
+}
+
+UComputeDataProvider* UPCGNormalToDensityDataInterface::CreateDataProvider() const
+{
+	return NewObject<UPCGNormalToDensityProvider>();
+}
+
+void UPCGNormalToDensityProvider::Initialize(const UComputeDataInterface* InDataInterface, UObject* InBinding, uint64 InInputMask, uint64 InOutputMask)
+{
+	TRACE_CPUPROFILER_EVENT_SCOPE(UPCGNormalToDensityProvider::Initialize);
+
+	Super::Initialize(InDataInterface, InBinding, InInputMask, InOutputMask);
+
+	const UPCGNormalToDensityDataInterface* DataInterface = CastChecked<UPCGNormalToDensityDataInterface>(InDataInterface);
+	ProducerKernel = DataInterface->ProducerKernel;
+}
+
+FComputeDataProviderRenderProxy* UPCGNormalToDensityProvider::GetRenderProxy()
+{
+	const UPCGNormalToDensitySettings* Settings = CastChecked<UPCGNormalToDensitySettings>(ProducerKernel->GetSettings());
+
+	FPCGNormalToDensityDataProviderProxy::FNormalToDensityData_RenderThread ProxyData;
+	ProxyData.Normal = Settings->Normal;
+	ProxyData.Offset = Settings->Offset;
+	ProxyData.Strength = Settings->Strength;
+	ProxyData.DensityMode = Settings->DensityMode;
+
+	return new FPCGNormalToDensityDataProviderProxy(MoveTemp(ProxyData));
+}
+
+void UPCGNormalToDensityProvider::Reset()
+{
+	ProducerKernel = nullptr;
+
+	Super::Reset();
+}
+
+bool FPCGNormalToDensityDataProviderProxy::IsValid(FValidationData const& InValidationData) const
+{
+	return InValidationData.ParameterStructSize == sizeof(FParameters);
+}
+
+void FPCGNormalToDensityDataProviderProxy::GatherDispatchData(FDispatchData const& InDispatchData)
+{
+	const TStridedView<FParameters> ParameterArray = MakeStridedParameterView<FParameters>(InDispatchData);
+
+	for (int32 InvocationIndex = 0; InvocationIndex < ParameterArray.Num(); ++InvocationIndex)
+	{
+		FParameters& Parameters = ParameterArray[InvocationIndex];
+		Parameters.Normal = static_cast<FVector3f>(Data.Normal);
+		Parameters.Offset = static_cast<float>(Data.Offset);
+		Parameters.Strength = static_cast<float>(Data.Strength);
+		Parameters.DensityMode = static_cast<uint32>(Data.DensityMode);
+	}
+}
diff --git a/Engine/Plugins/PCG/Source/PCG/Private/Compute/DataInterfaces/Elements/PCGNormalToDensityDataInterface.h b/Engine/Plugins/PCG/Source/PCG/Private/Compute/DataInterfaces/Elements/PCGNormalToDensityDataInterface.h
new file mode 100644
--- /dev/null
+++ b/Engine/Plugins/PCG/Source/PCG/Private/Compute/DataInterfaces/Elements/PCGNormalToDensityDataInterface.h
@@ -0,0 +1,77 @@
+// Copyright Epic Games, Inc. All Rights Reserved.
+
+#pragma once
+
+#include "Compute/DataInterfaces/PCGComputeDataInterface.h"
+#include "ComputeFramework/ComputeDataProvider.h"
+#include "Elements/PCGNormalToDensity.h"
+
+#include "PCGNormalToDensityDataInterface.generated.h"
+
+class FPCGNormalToDensityDataInterfaceParameters;
+
+/** Data Interface to marshal Normal to density settings to the GPU. */
+UCLASS(ClassGroup = (Procedural))
+class UPCGNormalToDensityDataInterface : public UPCGComputeDataInterface
+{
+	GENERATED_BODY()
+
+public:
+	//~ Begin UComputeDataInterface Interface
+	TCHAR const* GetClassName() const override { return TEXT("PCGNormalToDensity"); }
+	void GetSupportedInputs(TArray<FShaderFunctionDefinition>& OutFunctions) const override;
+	void GetShaderParameters(TCHAR const* UID, FShaderParametersMetadataBuilder& InOutBuilder, FShaderParametersMetadataAllocations& InOutAllocations) const override;
+	void GetHLSL(FString& OutHLSL, FString const& InDataInterfaceName) const override;
+	UComputeDataProvider* CreateDataProvider() const override;
+	//~ End UComputeDataInterface Interface
+};
+
+UCLASS()
+class UPCGNormalToDensityProvider : public UPCGComputeDataProvider
+{
+	GENERATED_BODY()
+
+public:
+	//~ Begin UComputeDataProvider Interface
+	virtual void Initialize(const UComputeDataInterface* InDataInterface, UObject* InBinding, uint64 InInputMask, uint64 InOutputMask) override;
+	FComputeDataProviderRenderProxy* GetRenderProxy() override;
+	virtual void Reset() override;
+	//~ End UComputeDataProvider Interface
+
+public:
+	UPROPERTY()
+	TObjectPtr<const UPCGComputeKernel> ProducerKernel;
+};
+
+class FPCGNormalToDensityDataProviderProxy : public FComputeDataProviderRenderProxy
+{
+public:
+	struct FNormalToDensityData_RenderThread
+	{
+		// The normal to compare against
+		FVector Normal = FVector::UpVector;
+
+		// This biases the value towards or against the normal (positive or negative)
+		double Offset = 0.0;
+
+		// This applies a curve to scale the result density with Result = Result^(1/Strength)
+		double Strength = 1.0;
+
+		// The operator to apply to the output density 
+		PCGNormalToDensityMode DensityMode = PCGNormalToDensityMode::Set;
+	};
+
+	FPCGNormalToDensityDataProviderProxy(FNormalToDensityData_RenderThread InData)
+		: Data(MoveTemp(InData))
+	{}
+
+	//~ Begin FComputeDataProviderRenderProxy Interface
+	bool IsValid(FValidationData const& InValidationData) const override;
+	void GatherDispatchData(FDispatchData const& InDispatchData) override;
+	//~ End FComputeDataProviderRenderProxy Interface
+
+protected:
+	using FParameters = FPCGNormalToDensityDataInterfaceParameters;
+
+	FNormalToDensityData_RenderThread Data;
+};
diff --git a/Engine/Plugins/PCG/Source/PCG/Private/Compute/Elements/PCGNormalToDensityKernel.cpp b/Engine/Plugins/PCG/Source/PCG/Private/Compute/Elements/PCGNormalToDensityKernel.cpp
new file mode 100644
--- /dev/null
+++ b/Engine/Plugins/PCG/Source/PCG/Private/Compute/Elements/PCGNormalToDensityKernel.cpp
@@ -0,0 +1,83 @@
+// Copyright Epic Games, Inc. All Rights Reserved.
+
+#include "Compute/Elements/PCGNormalToDensityKernel.h"
+
+#include "PCGContext.h"
+#include "Graph/PCGGPUGraphCompilationContext.h"
+#include "Compute/Elements/PCGComputeGraphElement.h"
+#include "Compute/DataInterfaces/Elements/PCGNormalToDensityDataInterface.h"
+
+#include "ShaderCompilerCore.h"
+
+#define LOCTEXT_NAMESPACE "PCGTemplateKernel"
+
+const UPCGSettings* UPCGNormalToDensityKernel::GetSettings() const
+{
+	return Settings;
+}
+
+bool UPCGNormalToDensityKernel::ComputeOutputBindingDataDesc(const UPCGComputeGraph* InGraph, FName InOutputPinLabel, UPCGDataBinding* InBinding, FPCGDataCollectionDesc& OutDataDesc) const
+{
+	check(InGraph);
+	check(InBinding);
+
+	// Code assumes single output pin.
+	if (!ensure(InOutputPinLabel == PCGPinConstants::DefaultOutputLabel))
+	{
+		return false;
+	}
+
+	const FPCGKernelPin SourceKernelPin(KernelIndex, PCGPinConstants::DefaultInputLabel, /*bIsInput=*/true);
+
+	ensure(InBinding->ComputeKernelPinDataDesc(SourceKernelPin, &OutDataDesc));
+
+	return true;
+}
+
+int UPCGNormalToDensityKernel::ComputeThreadCount(const UPCGDataBinding* InBinding) const
+{
+	int ThreadCount = 0;
+
+	const FPCGDataCollectionDesc* OutputPinDesc = InBinding->GetCachedKernelPinDataDesc(this, PCGPinConstants::DefaultOutputLabel, /*bIsInput=*/false);
+
+	if (ensure(OutputPinDesc))
+	{
+		for (const FPCGDataDesc& DataDesc : OutputPinDesc->DataDescs)
+		{
+			ThreadCount += DataDesc.ElementCount;
+		}
+	}
+
+	return ThreadCount;
+}
+
+#if WITH_EDITOR
+FString UPCGNormalToDensityKernel::GetCookedSource(const UPCGComputeGraph* InComputeGraph) const
+{
+	FString TemplateFile;
+	ensure(LoadShaderSourceFile(TEXT("/Plugin/PCG/Private/Elements/PCGNormalToDensity.usf"), EShaderPlatform::SP_PCD3D_SM5, &TemplateFile, nullptr));
+	return TemplateFile;
+}
+
+void UPCGNormalToDensityKernel::CreateAdditionalInputDataInterfaces(FPCGGPUCompilationContext& InOutContext, UObject* InObjectOuter, TArray<TObjectPtr<UComputeDataInterface>>& OutDataInterfaces) const
+{
+	Super::CreateAdditionalInputDataInterfaces(InOutContext, InObjectOuter, OutDataInterfaces);
+
+	TObjectPtr<UPCGNormalToDensityDataInterface> NodeDI = InOutContext.NewObject_AnyThread<UPCGNormalToDensityDataInterface>(InObjectOuter);
+	NodeDI->ProducerKernel = this;
+
+	OutDataInterfaces.Add(NodeDI);
+}
+#endif
+
+void UPCGNormalToDensityKernel::GetInputPins(TArray<FPCGPinProperties>& OutPins) const
+{
+	OutPins.Emplace(PCGPinConstants::DefaultInputLabel, EPCGDataType::Point);
+}
+
+void UPCGNormalToDensityKernel::GetOutputPins(TArray<FPCGPinPropertiesGPU>& OutPins) const
+{
+	OutPins.Emplace(PCGPinConstants::DefaultOutputLabel, EPCGDataType::Point);
+}
+
+#undef LOCTEXT_NAMESPACE
diff --git a/Engine/Plugins/PCG/Source/PCG/Private/Compute/Elements/PCGNormalToDensityKernel.h b/Engine/Plugins/PCG/Source/PCG/Private/Compute/Elements/PCGNormalToDensityKernel.h
new file mode 100644
--- /dev/null
+++ b/Engine/Plugins/PCG/Source/PCG/Private/Compute/Elements/PCGNormalToDensityKernel.h
@@ -0,0 +1,30 @@
+// Copyright Epic Games, Inc. All Rights Reserved.
+
+#pragma once
+
+#include "Compute/PCGComputeKernel.h"
+#include "Compute/PCGPinPropertiesGPU.h"
+
+#include "PCGNormalToDensityKernel.generated.h"
+
+class UPCGNormalToDensitySettings;
+
+UCLASS()
+class UPCGNormalToDensityKernel : public UPCGComputeKernel
+{	
+	GENERATED_BODY()
+
+public:
+	virtual const UPCGSettings* GetSettings() const override;
+	virtual bool ComputeOutputBindingDataDesc(const UPCGComputeGraph* InGraph, FName InOutputPinLabel, UPCGDataBinding* InBinding, FPCGDataCollectionDesc& OutDataDesc) const override;
+	virtual int ComputeThreadCount(const UPCGDataBinding* Binding) const override;
+
+#if WITH_EDITOR
+	virtual FString GetCookedSource(const UPCGComputeGraph* InComputeGraph) const override;
+	virtual FString GetEntryPoint() const override { return TEXT("PCGNormalToDensityCS"); }
+	void CreateAdditionalInputDataInterfaces(FPCGGPUCompilationContext& InOutContext, UObject* InObjectOuter, TArray<TObjectPtr<UComputeDataInterface>>& OutDataInterfaces) const override;
+#endif
+	virtual void GetInputPins(TArray<FPCGPinProperties>& OutPins) const override;
+	virtual void GetOutputPins(TArray<FPCGPinPropertiesGPU>& OutPins) const override;
+};
+
diff --git a/Engine/Plugins/PCG/Source/PCG/Private/Elements/PCGNormalToDensity.cpp b/Engine/Plugins/PCG/Source/PCG/Private/Elements/PCGNormalToDensity.cpp
--- a/Engine/Plugins/PCG/Source/PCG/Private/Elements/PCGNormalToDensity.cpp
+++ b/Engine/Plugins/PCG/Source/PCG/Private/Elements/PCGNormalToDensity.cpp
@@ -1,10 +1,27 @@
 // Copyright Epic Games, Inc. All Rights Reserved.
 
 #include "Elements/PCGNormalToDensity.h"
-#include "Helpers/PCGSettingsHelpers.h"
+
+#include "PCGContext.h"
+#include "Compute/Elements/PCGNormalToDensityKernel.h"
 #include "Data/PCGPointData.h"
+#include "Graph/PCGGPUGraphCompilationContext.h"
+#include "Helpers/PCGSettingsHelpers.h"
+
 #include "Kismet/KismetMathLibrary.h"
-#include "PCGContext.h"
+
+#if WITH_EDITOR
+void UPCGNormalToDensitySettings::CreateKernels(FPCGGPUCompilationContext& InOutContext, UObject* InObjectOuter, TArray<UPCGComputeKernel*>& OutKernels, TArray<FPCGKernelEdge>& OutEdges) const
+{
+	UPCGNormalToDensityKernel* Kernel = InOutContext.NewObject_AnyThread<UPCGNormalToDensityKernel>(InObjectOuter);
+	Kernel->Initialize(this, bDumpDataDescriptions);
+	OutKernels.Add(Kernel);
+
+	// Connect node pins to kernel pins
+	OutEdges.Emplace(FPCGPinReference(PCGPinConstants::DefaultInputLabel), FPCGPinReference(Kernel, PCGPinConstants::DefaultInputLabel));
+	OutEdges.Emplace(FPCGPinReference(Kernel, PCGPinConstants::DefaultOutputLabel), FPCGPinReference(PCGPinConstants::DefaultOutputLabel));
+}
+#endif // WITH_EDITOR
 
 FPCGElementPtr UPCGNormalToDensitySettings::CreateElement() const
 {
diff --git a/Engine/Plugins/PCG/Source/PCG/Public/Elements/PCGNormalToDensity.h b/Engine/Plugins/PCG/Source/PCG/Public/Elements/PCGNormalToDensity.h
--- a/Engine/Plugins/PCG/Source/PCG/Public/Elements/PCGNormalToDensity.h
+++ b/Engine/Plugins/PCG/Source/PCG/Public/Elements/PCGNormalToDensity.h
@@ -32,6 +32,9 @@ public:
 	virtual FName GetDefaultNodeName() const override { return FName(TEXT("NormalToDensity")); }
 	virtual FText GetDefaultNodeTitle() const override { return NSLOCTEXT("PCGNormalToDensitySettings", "NodeTitle", "Normal To Density"); }
 	virtual EPCGSettingsType GetType() const override { return EPCGSettingsType::Spatial; }
+
+	virtual bool DisplayExecuteOnGPUSetting() const override { return true; }
+	virtual void CreateKernels(FPCGGPUCompilationContext& InOutContext, UObject* InObjectOuter, TArray<UPCGComputeKernel*>& OutKernels, TArray<FPCGKernelEdge>& OutEdges) const override;
 #endif
 	
 
@@ -45,7 +48,7 @@ public:
 
 	// The normal to compare against
 	UPROPERTY(BlueprintReadWrite, EditAnywhere, Category = Settings, meta = (PCG_Overridable))
-	FVector Normal = {0, 0.0, 1.0};
+	FVector Normal = FVector::UpVector;
 
 	// This is biases the value towards or against the normal (positive or negative)
 	UPROPERTY(BlueprintReadWrite, EditAnywhere, Category = Settings, meta = (PCG_Overridable))
