[PCG Test] Update unittests to support PointArrayData when enabled

#rnx

diff --git a/Engine/Plugins/PCG/Source/PCG/Private/SpatialAlgo/PCGOctreeQueries.cpp b/Engine/Plugins/PCG/Source/PCG/Private/SpatialAlgo/PCGOctreeQueries.cpp
--- a/Engine/Plugins/PCG/Source/PCG/Private/SpatialAlgo/PCGOctreeQueries.cpp
+++ b/Engine/Plugins/PCG/Source/PCG/Private/SpatialAlgo/PCGOctreeQueries.cpp
@@ -79,7 +79,7 @@ void UPCGOctreeQueries::GetClosestPoint(const UPCGPointData* InPointData, const
 
 const FPCGPoint* UPCGOctreeQueries::GetClosestPoint(const UPCGPointData* InPointData, const FVector& InCenter, const bool bInDiscardCenter, const double InSearchDistance)
 {
-	int32 PointIndex = GetClosestPointIndex(InPointData, InCenter, bInDiscardCenter, InSearchDistance);
+	const int32 PointIndex = GetClosestPointIndex(InPointData, InCenter, bInDiscardCenter, InSearchDistance);
 	return PointIndex != INDEX_NONE ? &InPointData->GetPoints()[PointIndex] : nullptr;
 }
 
@@ -114,7 +114,7 @@ void UPCGOctreeQueries::GetClosestPointFromOtherPoint(const UPCGPointData* InPoi
 
 const FPCGPoint* UPCGOctreeQueries::GetClosestPointFromOtherPoint(const UPCGPointData* InPointData, const FPCGPoint& InPoint, const double InSearchDistance)
 {
-	int32 PointIndex = PCGOctreeQueries::GetPointIndexInSphereUnderCondition(InPointData, InPoint.Transform.GetLocation(), InSearchDistance, std::numeric_limits<double>::max(),
+	const int32 PointIndex = PCGOctreeQueries::GetPointIndexInSphereUnderCondition(InPointData, InPoint.Transform.GetLocation(), InSearchDistance, std::numeric_limits<double>::max(),
 		[&InPoint](double ChosenPointDistance, const UPCGBasePointData* BasePointData, int32 PointIndex, double SquaredDistance)
 		{
 			const UPCGPointData* PointData = CastChecked<UPCGPointData>(BasePointData);
@@ -123,6 +123,22 @@ const FPCGPoint* UPCGOctreeQueries::GetClosestPointFromOtherPoint(const UPCGPoin
 	return PointIndex != INDEX_NONE ? &InPointData->GetPoints()[PointIndex] : nullptr;
 }
 
+int32 UPCGOctreeQueries::GetClosestPointIndexFromOtherPointIndex(const UPCGBasePointData* InPointData, int32 OtherPointIndex, const double InSearchDistance)
+{
+	const TConstPCGValueRange<FTransform> TransformRange = InPointData->GetConstTransformValueRange();
+	if (!TransformRange.IsValidIndex(OtherPointIndex))
+	{
+		return INDEX_NONE;
+	}
+
+	const int32 PointIndex = PCGOctreeQueries::GetPointIndexInSphereUnderCondition(InPointData, TransformRange[OtherPointIndex].GetLocation(), InSearchDistance, std::numeric_limits<double>::max(),
+		[OtherPointIndex](double ChosenPointDistance, const UPCGBasePointData* BasePointData, int32 PointIndex, double SquaredDistance)
+		{
+			return SquaredDistance <= ChosenPointDistance && PointIndex != OtherPointIndex;
+		});
+	return PointIndex;
+}
+
 void UPCGOctreeQueries::GetFarthestPoint(const UPCGPointData* InPointData, const FVector& InCenter, bool& bOutFound, FPCGPoint& OutPoint, const double InSearchDistance)
 {
 	bOutFound = false;
@@ -135,7 +151,7 @@ void UPCGOctreeQueries::GetFarthestPoint(const UPCGPointData* InPointData, const
 
 const FPCGPoint* UPCGOctreeQueries::GetFarthestPoint(const UPCGPointData* InPointData, const FVector& InCenter, const double InSearchDistance)
 {
-	int32 PointIndex = GetFarthestPointIndex(InPointData, InCenter, InSearchDistance);
+	const int32 PointIndex = GetFarthestPointIndex(InPointData, InCenter, InSearchDistance);
 	return PointIndex != INDEX_NONE ? &InPointData->GetPoints()[PointIndex] : nullptr;
 }
 
@@ -170,11 +186,27 @@ void UPCGOctreeQueries::GetFarthestPointFromOtherPoint(const UPCGPointData* InPo
 
 const FPCGPoint* UPCGOctreeQueries::GetFarthestPointFromOtherPoint(const UPCGPointData* InPointData, const FPCGPoint& InPoint, const double InSearchDistance)
 {
-	int32 PointIndex = PCGOctreeQueries::GetPointIndexInSphereUnderCondition(InPointData, InPoint.Transform.GetLocation(), InSearchDistance, std::numeric_limits<double>::min(),
+	const int32 PointIndex = PCGOctreeQueries::GetPointIndexInSphereUnderCondition(InPointData, InPoint.Transform.GetLocation(), InSearchDistance, std::numeric_limits<double>::min(),
 		[&InPoint](double ChosenPointDistance, const UPCGBasePointData* BasePointData, int32 PointIndex, double SquaredDistance)
 		{
 			const UPCGPointData* PointData = CastChecked<UPCGPointData>(BasePointData);
 			return SquaredDistance >= ChosenPointDistance && &InPoint != &PointData->GetPoints()[PointIndex];
 		});
 	return PointIndex != INDEX_NONE ? &InPointData->GetPoints()[PointIndex] : nullptr;
+}
+
+int32 UPCGOctreeQueries::GetFarthestPointIndexFromOtherPointIndex(const UPCGBasePointData* InPointData, int32 OtherPointIndex, const double InSearchDistance)
+{
+	const TConstPCGValueRange<FTransform> TransformRange = InPointData->GetConstTransformValueRange();
+	if (!TransformRange.IsValidIndex(OtherPointIndex))
+	{
+		return INDEX_NONE;
+	}
+
+	const int32 PointIndex = PCGOctreeQueries::GetPointIndexInSphereUnderCondition(InPointData, TransformRange[OtherPointIndex].GetLocation(), InSearchDistance, std::numeric_limits<double>::min(),
+		[OtherPointIndex](double ChosenPointDistance, const UPCGBasePointData* BasePointData, int32 PointIndex, double SquaredDistance)
+		{
+			return SquaredDistance >= ChosenPointDistance && PointIndex != OtherPointIndex;
+		});
+	return PointIndex;
 }
\ No newline at end of file
diff --git a/Engine/Plugins/PCG/Source/PCG/Private/Tests/Accessors/PCGAttributeAccessorTest.cpp b/Engine/Plugins/PCG/Source/PCG/Private/Tests/Accessors/PCGAttributeAccessorTest.cpp
--- a/Engine/Plugins/PCG/Source/PCG/Private/Tests/Accessors/PCGAttributeAccessorTest.cpp
+++ b/Engine/Plugins/PCG/Source/PCG/Private/Tests/Accessors/PCGAttributeAccessorTest.cpp
@@ -22,14 +22,14 @@ namespace FPCGAttributeAccessorTestHelpers
 	struct AttributeData
 	{
 		AttributeData()
-			: PointData(NewObject<UPCGPointData>(GetTransientPackage(), NAME_None, RF_Transient))
+			: PointData(PCGTestsCommon::CreateEmptyBasePointData())
 			, FloatAttribute(PointData->Metadata->CreateAttribute<float>(TEXT("FloatAttribute"), 0.0f, true, false))
 			, VectorAttribute(PointData->Metadata->CreateAttribute<FVector>(TEXT("VectorAttribute"), FVector::ZeroVector, false, false))
 			, StringAttribute(PointData->Metadata->CreateAttribute<FString>(TEXT("StringAttribute"), FString(), false, false))
 		{
 			// All points will point to a given entry, in reverse order, except the last one, that will point to DefaultValue
-			TArray<FPCGPoint>& Points = PointData->GetMutablePoints();
-			Points.SetNum(ValuesSize + 1);
+			PointData->SetNumPoints(ValuesSize + 1);
+			TPCGValueRange<int64> MetadataEntryRange = PointData->GetMetadataEntryValueRange();
 
 			for (int32 i = 0; i < ValuesSize; ++i)
 			{
@@ -37,7 +37,7 @@ namespace FPCGAttributeAccessorTestHelpers
 				FloatAttribute->SetValue(EntryKey, FloatValues[i]);
 				VectorAttribute->SetValue(EntryKey, VectorValues[i]);
 				StringAttribute->SetValue(EntryKey, StringValues[i]);
-				Points[ValuesSize - i - 1].MetadataEntry = EntryKey;
+				MetadataEntryRange[ValuesSize - i - 1] = EntryKey;
 			}
 
 			FloatAccessor = PCGAttributeAccessorHelpers::CreateAccessor(FloatAttribute, PointData->Metadata);
@@ -52,7 +52,7 @@ namespace FPCGAttributeAccessorTestHelpers
 			TempFloats[2] = 3.0f;
 		}
 
-		TObjectPtr<UPCGPointData> PointData;
+		TObjectPtr<UPCGBasePointData> PointData;
 		FPCGMetadataAttribute<float>* FloatAttribute;
 		FPCGMetadataAttribute<FVector>* VectorAttribute;
 		FPCGMetadataAttribute<FString>* StringAttribute;
@@ -507,23 +507,24 @@ bool FPCGAttributeAccessorSetRangeOutsideRangeTest::RunTest(const FString& Param
 bool FPCGAttributeAccessorSingleGetPointsTest::RunTest(const FString& Parameters)
 {
 	FPCGAttributeAccessorTestHelpers::AttributeData AttributeData;
-	const TArray<FPCGPoint>& Points = AttributeData.PointData->GetPoints();
-	FPCGAttributeAccessorKeysPoints Keys(Points);
 
-	if (!TestEqual("Right number of keys", Keys.GetNum(), Points.Num()))
+	const TConstPCGValueRange<int64> MetadataEntryRange = AttributeData.PointData->GetConstMetadataEntryValueRange();
+	TUniquePtr<const IPCGAttributeAccessorKeys> Keys = PCGAttributeAccessorHelpers::CreateConstKeys(AttributeData.PointData, FPCGAttributePropertySelector());
+
+	if (!TestEqual("Right number of keys", Keys->GetNum(), AttributeData.PointData->GetNumPoints()))
 	{
 		return false;
 	}
 
-	if (!TestTrue("Keys are read only", Keys.IsReadOnly()))
+	if (!TestTrue("Keys are read only", Keys->IsReadOnly()))
 	{
 		return false;
 	}
 
 	for (int32 i = 0; i < 4; ++i)
 	{
-		if (!TestTrue(FString::Printf(TEXT("Get Float value %d"), i), AttributeData.FloatAccessor->Get(*AttributeData.TempFloats, i, Keys))
-			|| !TestEqual(FString::Printf(TEXT("Float value %d"), i), *AttributeData.TempFloats, AttributeData.FloatAttribute->GetValueFromItemKey(Points[i].MetadataEntry)))
+		if (!TestTrue(FString::Printf(TEXT("Get Float value %d"), i), AttributeData.FloatAccessor->Get(*AttributeData.TempFloats, i, *Keys))
+			|| !TestEqual(FString::Printf(TEXT("Float value %d"), i), *AttributeData.TempFloats, AttributeData.FloatAttribute->GetValueFromItemKey(MetadataEntryRange[i])))
 		{
 			return false;
 		}
@@ -535,10 +536,11 @@ bool FPCGAttributeAccessorSingleGetPointsTest::RunTest(const FString& Parameters
 bool FPCGAttributeAccessorGetRangePointsTest::RunTest(const FString& Parameters)
 {
 	FPCGAttributeAccessorTestHelpers::AttributeData AttributeData;
-	const TArray<FPCGPoint>& Points = AttributeData.PointData->GetPoints();
-	FPCGAttributeAccessorKeysPoints Keys(Points);
 
-	if (!TestTrue("Keys are read only", Keys.IsReadOnly()))
+	const TConstPCGValueRange<int64> MetadataEntryRange = AttributeData.PointData->GetConstMetadataEntryValueRange();
+	TUniquePtr<const IPCGAttributeAccessorKeys> Keys = PCGAttributeAccessorHelpers::CreateConstKeys(AttributeData.PointData, FPCGAttributePropertySelector());
+
+	if (!TestTrue("Keys are read only", Keys->IsReadOnly()))
 	{
 		return false;
 	}
@@ -547,7 +549,7 @@ bool FPCGAttributeAccessorGetRangePointsTest::RunTest(const FString& Parameters)
 	const int32 Index = 0;
 
 	// Points entry are in reverse order, and last one is the default value.
-	if (!TestTrue("Get float attribute full range", AttributeData.FloatAccessor->GetRange(FPCGAttributeAccessorTestHelpers::MakeArrayView(TempFloats), Index, Keys))
+	if (!TestTrue("Get float attribute full range", AttributeData.FloatAccessor->GetRange(FPCGAttributeAccessorTestHelpers::MakeArrayView(TempFloats), Index, *Keys))
 		|| !TestEqual("float value 0", TempFloats[0], FPCGAttributeAccessorTestHelpers::FloatValues[2])
 		|| !TestEqual("float value 1", TempFloats[1], FPCGAttributeAccessorTestHelpers::FloatValues[1])
 		|| !TestEqual("float value 2", TempFloats[2], FPCGAttributeAccessorTestHelpers::FloatValues[0])
@@ -562,15 +564,16 @@ bool FPCGAttributeAccessorGetRangePointsTest::RunTest(const FString& Parameters)
 bool FPCGAttributeAccessorSingleSetPointsTest::RunTest(const FString& Parameters)
 {
 	FPCGAttributeAccessorTestHelpers::AttributeData AttributeData;
-	TArray<FPCGPoint>& Points = AttributeData.PointData->GetMutablePoints();
-	FPCGAttributeAccessorKeysPoints Keys(TArrayView<FPCGPoint>(Points.GetData(), Points.Num()));
 
-	if (!TestEqual("Right number of keys", Keys.GetNum(), Points.Num()))
+	const TConstPCGValueRange<int64> MetadataEntryRange = AttributeData.PointData->GetConstMetadataEntryValueRange();
+	TUniquePtr<IPCGAttributeAccessorKeys> Keys = PCGAttributeAccessorHelpers::CreateKeys(AttributeData.PointData, FPCGAttributePropertySelector());
+
+	if (!TestEqual("Right number of keys", Keys->GetNum(), AttributeData.PointData->GetNumPoints()))
 	{
 		return false;
 	}
 
-	if (!TestFalse("Keys are not read only", Keys.IsReadOnly()))
+	if (!TestFalse("Keys are not read only", Keys->IsReadOnly()))
 	{
 		return false;
 	}
@@ -580,17 +583,17 @@ bool FPCGAttributeAccessorSingleSetPointsTest::RunTest(const FString& Parameters
 
 	for (int32 i = 0; i < 4; ++i)
 	{
-		if (!TestTrue(FString::Printf(TEXT("Set Float value %d"), i), AttributeData.FloatAccessor->Set(TempFloats[i], i, Keys))
-			|| !TestEqual(FString::Printf(TEXT("Float value %d"), i), AttributeData.FloatAttribute->GetValueFromItemKey(Points[i].MetadataEntry), TempFloats[i]))
+		if (!TestTrue(FString::Printf(TEXT("Set Float value %d"), i), AttributeData.FloatAccessor->Set(TempFloats[i], i, *Keys))
+			|| !TestEqual(FString::Printf(TEXT("Float value %d"), i), AttributeData.FloatAttribute->GetValueFromItemKey(MetadataEntryRange[i]), TempFloats[i]))
 		{
 			return false;
 		}
 	}
 
 	// Points should not have a different entry key, if it was not invalid in the first place
-	if (!TestEqual("Point 0 should have entry 2", Points[0].MetadataEntry, PCGMetadataEntryKey(2))
-		|| !TestEqual("Point 1 should have entry 1", Points[1].MetadataEntry, PCGMetadataEntryKey(1))
-		|| !TestEqual("Point 2 should have entry 0", Points[2].MetadataEntry, PCGMetadataEntryKey(0)))
+	if (!TestEqual("Point 0 should have entry 2", MetadataEntryRange[0], PCGMetadataEntryKey(2))
+		|| !TestEqual("Point 1 should have entry 1", MetadataEntryRange[1], PCGMetadataEntryKey(1))
+		|| !TestEqual("Point 2 should have entry 0", MetadataEntryRange[2], PCGMetadataEntryKey(0)))
 	{
 		return false;
 	}
@@ -601,7 +604,7 @@ bool FPCGAttributeAccessorSingleSetPointsTest::RunTest(const FString& Parameters
 		return false;
 	}
 
-	if (!TestEqual("Last point should have a valid entry", Points[3].MetadataEntry, PCGMetadataEntryKey(3)))
+	if (!TestEqual("Last point should have a valid entry", MetadataEntryRange[3], PCGMetadataEntryKey(3)))
 	{
 		return false;
 	}
@@ -612,25 +615,26 @@ bool FPCGAttributeAccessorSingleSetPointsTest::RunTest(const FString& Parameters
 bool FPCGAttributeAccessorSetRangePointsTest::RunTest(const FString& Parameters)
 {
 	FPCGAttributeAccessorTestHelpers::AttributeData AttributeData;
-	TArray<FPCGPoint>& Points = AttributeData.PointData->GetMutablePoints();
-	FPCGAttributeAccessorKeysPoints Keys(TArrayView<FPCGPoint>(Points.GetData(), Points.Num()));
 
+	const TConstPCGValueRange<int64> MetadataEntryRange = AttributeData.PointData->GetConstMetadataEntryValueRange();
+	TUniquePtr<IPCGAttributeAccessorKeys> Keys = PCGAttributeAccessorHelpers::CreateKeys(AttributeData.PointData, FPCGAttributePropertySelector());
+	
 	float TempFloats[4] = { 1.0f, 2.0f, 3.0f, 4.0f };
 	const int32 Index = 0;
 
-	if (!TestTrue("Set Float attribute full range", AttributeData.FloatAccessor->SetRange(FPCGAttributeAccessorTestHelpers::MakeConstArrayView(TempFloats), Index, Keys))
-		|| !TestEqual("Float value 0", AttributeData.FloatAttribute->GetValueFromItemKey(Points[0].MetadataEntry), TempFloats[0])
-		|| !TestEqual("Float value 1", AttributeData.FloatAttribute->GetValueFromItemKey(Points[1].MetadataEntry), TempFloats[1])
-		|| !TestEqual("Float value 2", AttributeData.FloatAttribute->GetValueFromItemKey(Points[2].MetadataEntry), TempFloats[2])
-		|| !TestEqual("Float value 3", AttributeData.FloatAttribute->GetValueFromItemKey(Points[3].MetadataEntry), TempFloats[3]))
+	if (!TestTrue("Set Float attribute full range", AttributeData.FloatAccessor->SetRange(FPCGAttributeAccessorTestHelpers::MakeConstArrayView(TempFloats), Index, *Keys))
+		|| !TestEqual("Float value 0", AttributeData.FloatAttribute->GetValueFromItemKey(MetadataEntryRange[0]), TempFloats[0])
+		|| !TestEqual("Float value 1", AttributeData.FloatAttribute->GetValueFromItemKey(MetadataEntryRange[1]), TempFloats[1])
+		|| !TestEqual("Float value 2", AttributeData.FloatAttribute->GetValueFromItemKey(MetadataEntryRange[2]), TempFloats[2])
+		|| !TestEqual("Float value 3", AttributeData.FloatAttribute->GetValueFromItemKey(MetadataEntryRange[3]), TempFloats[3]))
 	{
 		return false;
 	}
 
 	// Points should not have a different entry key, if it was not invalid in the first place
-	if (!TestEqual("Point 0 should have entry 2", Points[0].MetadataEntry, PCGMetadataEntryKey(2))
-		|| !TestEqual("Point 1 should have entry 1", Points[1].MetadataEntry, PCGMetadataEntryKey(1))
-		|| !TestEqual("Point 2 should have entry 0", Points[2].MetadataEntry, PCGMetadataEntryKey(0)))
+	if (!TestEqual("Point 0 should have entry 2", MetadataEntryRange[0], PCGMetadataEntryKey(2))
+		|| !TestEqual("Point 1 should have entry 1", MetadataEntryRange[1], PCGMetadataEntryKey(1))
+		|| !TestEqual("Point 2 should have entry 0", MetadataEntryRange[2], PCGMetadataEntryKey(0)))
 	{
 		return false;
 	}
@@ -641,7 +645,7 @@ bool FPCGAttributeAccessorSetRangePointsTest::RunTest(const FString& Parameters)
 		return false;
 	}
 
-	if (!TestEqual("Last point should have a valid entry", Points[3].MetadataEntry, PCGMetadataEntryKey(3)))
+	if (!TestEqual("Last point should have a valid entry", MetadataEntryRange[3], PCGMetadataEntryKey(3)))
 	{
 		return false;
 	}
diff --git a/Engine/Plugins/PCG/Source/PCG/Private/Tests/Data/PCGDifferenceDataTest.cpp b/Engine/Plugins/PCG/Source/PCG/Private/Tests/Data/PCGDifferenceDataTest.cpp
--- a/Engine/Plugins/PCG/Source/PCG/Private/Tests/Data/PCGDifferenceDataTest.cpp
+++ b/Engine/Plugins/PCG/Source/PCG/Private/Tests/Data/PCGDifferenceDataTest.cpp
@@ -16,13 +16,13 @@ bool FPCGDifferenceDataTest::RunTest(const FString& Parameters)
 	const UPCGVolumeData* ThirdVolume = PCGTestsCommon::CreateVolumeData(FBox::BuildAABB(FVector::OneVector * 10000, FVector::OneVector * 500));
 
 	// Inside first volume
-	UPCGPointData* FirstVolumeTestPoint = PCGTestsCommon::CreatePointData(FVector::OneVector * 500);
+	UPCGBasePointData* FirstVolumeTestPoint = PCGTestsCommon::CreateBasePointData(FVector::OneVector * 500);
 	// Inside second volume
-	UPCGPointData* SecondVolumeTestPoint = PCGTestsCommon::CreatePointData(FVector::OneVector * -500);
+	UPCGBasePointData* SecondVolumeTestPoint = PCGTestsCommon::CreateBasePointData(FVector::OneVector * -500);
 	// Point at origin
-	UPCGPointData* OriginTestPoint = PCGTestsCommon::CreatePointData();
+	UPCGBasePointData* OriginTestPoint = PCGTestsCommon::CreateBasePointData();
 	// Point far from anything
-	UPCGPointData* FarOffTestPoint = PCGTestsCommon::CreatePointData(FVector::OneVector * -10000);
+	UPCGBasePointData* FarOffTestPoint = PCGTestsCommon::CreateBasePointData(FVector::OneVector * -10000);
 
 	// Tests for the DifferenceData itself
 	auto ValidateVolumeDifferenceData = [this](const UPCGVolumeData* OriginalVolume, const UPCGVolumeData* DifferenceVolume) -> const UPCGDifferenceData*
@@ -53,7 +53,7 @@ bool FPCGDifferenceDataTest::RunTest(const FString& Parameters)
 	};
 
 	// Tests points in the various volumes
-	auto ValidatePointsShouldExist = [this](const UPCGDifferenceData* DifferenceData, const UPCGPointData* PointData, bool PointShouldExist)
+	auto ValidatePointsShouldExist = [this](const UPCGDifferenceData* DifferenceData, const UPCGBasePointData* PointData, bool PointShouldExist)
 	{
 		check(PointData);
 
@@ -63,9 +63,10 @@ bool FPCGDifferenceDataTest::RunTest(const FString& Parameters)
 			return;
 		}
 
-		check(PointData->GetPoints().Num() == 1);
+		check(PointData->GetNumPoints() == 1);
 		// Validate that we're able to sample a point
-		const FPCGPoint& Point = PointData->GetPoints()[0];
+		const FConstPCGPointValueRanges ValueRanges(PointData);
+		FPCGPoint Point = ValueRanges.GetPoint(0);
 
 		FPCGPoint SampledPoint;
 		bool PointWasSampled = DifferenceData->SamplePoint(Point.Transform, Point.GetLocalBounds(), SampledPoint, nullptr);
@@ -78,11 +79,11 @@ bool FPCGDifferenceDataTest::RunTest(const FString& Parameters)
 		}
 
 		// Generate PointData
-		const UPCGPointData* OutputPointData = DifferenceData->ToPointData(nullptr);
+		const UPCGBasePointData* OutputPointData = DifferenceData->ToBasePointData(nullptr);
 		if (TestNotNull("PointData generated", OutputPointData))
 		{
 			// If the point should be generated
-			TestTrue("PointData exists", OutputPointData->GetPoints().Num() > 0);
+			TestTrue("PointData exists", OutputPointData->GetNumPoints() > 0);
 		}
 	};
 
diff --git a/Engine/Plugins/PCG/Source/PCG/Private/Tests/Data/PCGIntersectionDataTest.cpp b/Engine/Plugins/PCG/Source/PCG/Private/Tests/Data/PCGIntersectionDataTest.cpp
--- a/Engine/Plugins/PCG/Source/PCG/Private/Tests/Data/PCGIntersectionDataTest.cpp
+++ b/Engine/Plugins/PCG/Source/PCG/Private/Tests/Data/PCGIntersectionDataTest.cpp
@@ -9,11 +9,11 @@ IMPLEMENT_CUSTOM_SIMPLE_AUTOMATION_TEST(FPCGIntersectionDataTest, FPCGTestBaseCl
 
 bool FPCGIntersectionDataTest::RunTest(const FString& Parameters)
 {
-	UPCGPointData* InsidePoint = PCGTestsCommon::CreatePointData();
-	check(InsidePoint->GetPoints().Num() == 1);
+	UPCGBasePointData* InsidePoint = PCGTestsCommon::CreateBasePointData();
+	check(InsidePoint->GetNumPoints() == 1);
 
-	UPCGPointData* OutsidePoint = PCGTestsCommon::CreatePointData(FVector::OneVector * 10000);
-	check(OutsidePoint->GetPoints().Num() == 1);
+	UPCGBasePointData* OutsidePoint = PCGTestsCommon::CreateBasePointData(FVector::OneVector * 10000);
+	check(OutsidePoint->GetNumPoints() == 1);
 
 	UPCGVolumeData* Volume = PCGTestsCommon::CreateVolumeData(FBox::BuildAABB(FVector::ZeroVector, FVector::OneVector * 100));
 
@@ -37,23 +37,25 @@ bool FPCGIntersectionDataTest::RunTest(const FString& Parameters)
 		TestTrue("Valid bounds", Intersection->GetBounds() == InsidePoint->GetBounds());
 		TestTrue("Valid strict bounds", Intersection->GetStrictBounds() == InsidePoint->GetStrictBounds());
 
-		// Validate sample point		
-		const FPCGPoint& Point = InsidePoint->GetPoints()[0];
+		// Validate sample point
+		FConstPCGPointValueRanges ValueRanges(InsidePoint);
+		const FPCGPoint Point = ValueRanges.GetPoint(0);
 
 		FPCGPoint SampledPoint;
 		TestTrue("Successful point sampling", Intersection->SamplePoint(Point.Transform, Point.GetLocalBounds(), SampledPoint, nullptr));
 		TestTrue("Correct sampled point", PCGTestsCommon::PointsAreIdentical(Point, SampledPoint));
 
 		// Validate create point data
-		const UPCGPointData* OutputPointData = Intersection->ToPointData(nullptr);
+		const UPCGBasePointData* OutputPointData = Intersection->ToBasePointData(nullptr);
 		TestTrue("Successful ToPoint", OutputPointData != nullptr);
 		
 		if (OutputPointData)
 		{
-			TestTrue("Valid number of points in ToPoint", OutputPointData->GetPoints().Num() == 1);
-			if (OutputPointData->GetPoints().Num() == 1)
+			TestTrue("Valid number of points in ToPoint", OutputPointData->GetNumPoints() == 1);
+			if (OutputPointData->GetNumPoints() == 1)
 			{
-				TestTrue("Correct point in ToPoint", PCGTestsCommon::PointsAreIdentical(Point, OutputPointData->GetPoints()[0]));
+				const FConstPCGPointValueRanges OutValueRanges(OutputPointData);
+				TestTrue("Correct point in ToPoint", PCGTestsCommon::PointsAreIdentical(Point, OutValueRanges.GetPoint(0)));
 			}
 		}
 	};
@@ -75,18 +77,19 @@ bool FPCGIntersectionDataTest::RunTest(const FString& Parameters)
 		TestTrue("Null strict bounds", !Intersection->GetStrictBounds().IsValid);
 
 		// Validate that we're not able to sample a point
-		const FPCGPoint& Point = OutsidePoint->GetPoints()[0];
+		FConstPCGPointValueRanges ValueRanges(OutsidePoint);
+		const FPCGPoint Point = ValueRanges.GetPoint(0);
 
 		FPCGPoint SampledPoint;
 		TestTrue("Unsuccessful point sampling", !Intersection->SamplePoint(Point.Transform, Point.GetLocalBounds(), SampledPoint, nullptr));
 
 		// Validate empty point data
-		const UPCGPointData* OutputPointData = Intersection->ToPointData(nullptr);
+		const UPCGBasePointData* OutputPointData = Intersection->ToBasePointData(nullptr);
 		TestTrue("Successful ToPoint", OutputPointData != nullptr);
 
 		if (OutputPointData)
 		{
-			TestTrue("Empty point data", OutputPointData->GetPoints().Num() == 0);
+			TestTrue("Empty point data", OutputPointData->GetNumPoints() == 0);
 		}
 	};
 
diff --git a/Engine/Plugins/PCG/Source/PCG/Private/Tests/Determinism/PCGDeterminismTestsCommon.cpp b/Engine/Plugins/PCG/Source/PCG/Private/Tests/Determinism/PCGDeterminismTestsCommon.cpp
--- a/Engine/Plugins/PCG/Source/PCG/Private/Tests/Determinism/PCGDeterminismTestsCommon.cpp
+++ b/Engine/Plugins/PCG/Source/PCG/Private/Tests/Determinism/PCGDeterminismTestsCommon.cpp
@@ -584,7 +584,7 @@ namespace PCGDeterminismTests
 
 	void AddSinglePointInputData(FPCGDataCollection& InputData, const FVector& Location, const FName& PinName)
 	{
-		UPCGPointData* PointData = PCGTestsCommon::CreatePointData(Location);
+		UPCGBasePointData* PointData = PCGTestsCommon::CreateBasePointData(Location);
 
 		FPCGTaggedData& TaggedData = InputData.TaggedData.Emplace_GetRef();
 		TaggedData.Data = PointData;
@@ -593,8 +593,15 @@ namespace PCGDeterminismTests
 
 	void AddMultiplePointsInputData(FPCGDataCollection& InputData, const TArray<FPCGPoint>& Points, const FName& PinName)
 	{
-		UPCGPointData* PointData = PCGTestsCommon::CreateEmptyPointData();
-		PointData->SetPoints(Points);
+		UPCGBasePointData* PointData = PCGTestsCommon::CreateEmptyBasePointData();
+		PointData->SetNumPoints(Points.Num());
+		PointData->AllocateProperties(EPCGPointNativeProperties::All);
+
+		FPCGPointValueRanges ValueRanges(PointData, /*bAllocate=*/false);
+		for (int32 Index = 0; Index < Points.Num(); ++Index)
+		{
+			ValueRanges.SetFromPoint(Index, Points[Index]);
+		}
 
 		FPCGTaggedData& TaggedData = InputData.TaggedData.Emplace_GetRef();
 		TaggedData.Data = PointData;
diff --git a/Engine/Plugins/PCG/Source/PCG/Private/Tests/Elements/Metadata/PCGMetadataBreakTransformTest.cpp b/Engine/Plugins/PCG/Source/PCG/Private/Tests/Elements/Metadata/PCGMetadataBreakTransformTest.cpp
--- a/Engine/Plugins/PCG/Source/PCG/Private/Tests/Elements/Metadata/PCGMetadataBreakTransformTest.cpp
+++ b/Engine/Plugins/PCG/Source/PCG/Private/Tests/Elements/Metadata/PCGMetadataBreakTransformTest.cpp
@@ -40,7 +40,7 @@ namespace PCGBreakTransformTest
 	void GenerateSpatialData(PCGTestsCommon::FTestData& TestData)
 	{
 		FPCGTaggedData& SpatialTaggedData = TestData.InputData.TaggedData.Emplace_GetRef();
-		TObjectPtr<UPCGPointData> PointData = PCGTestsCommon::CreateRandomPointData(10, TestData.Seed);
+		TObjectPtr<UPCGBasePointData> PointData = PCGTestsCommon::CreateRandomBasePointData(10, TestData.Seed);
 		SpatialTaggedData.Data = PointData;
 		SpatialTaggedData.Pin = PCGPinConstants::DefaultInputLabel;
 
@@ -50,14 +50,14 @@ namespace PCGBreakTransformTest
 		PointData->Metadata->CreateTransformAttribute(TransformAttribute, FTransform::Identity, bAllowsInterpolation, bOverrideParent);
 		PointData->Metadata->CreateFloatAttribute(InvalidAttribute, 0.0f, bAllowsInterpolation, bOverrideParent);
 
-		TArray<FPCGPoint>& SourcePoints = PointData->GetMutablePoints();
+		TPCGValueRange<int64> MetadataEntryRange = PointData->GetMetadataEntryValueRange();
 
 		// Randomize attribute values, and leave the second half to default value
 		FRandomStream RandomSource(TestData.Seed);
-		for (int I = 0; I < SourcePoints.Num() / 2; ++I)
+		for (int I = 0; I < PointData->GetNumPoints() / 2; ++I)
 		{
-			UPCGMetadataAccessorHelpers::SetTransformAttribute(SourcePoints[I], PointData->Metadata, TransformAttribute, RandomTransform(RandomSource));
-			UPCGMetadataAccessorHelpers::SetFloatAttribute(SourcePoints[I], PointData->Metadata, InvalidAttribute, RandomSource.FRand());
+			UPCGMetadataAccessorHelpers::SetTransformAttributeByMetadataKey(MetadataEntryRange[I], PointData->Metadata, TransformAttribute, RandomTransform(RandomSource));
+			UPCGMetadataAccessorHelpers::SetFloatAttributeByMetadataKey(MetadataEntryRange[I], PointData->Metadata, InvalidAttribute, RandomSource.FRand());
 		}
 	}
 
diff --git a/Engine/Plugins/PCG/Source/PCG/Private/Tests/Elements/Metadata/PCGMetadataBreakVectorTest.cpp b/Engine/Plugins/PCG/Source/PCG/Private/Tests/Elements/Metadata/PCGMetadataBreakVectorTest.cpp
--- a/Engine/Plugins/PCG/Source/PCG/Private/Tests/Elements/Metadata/PCGMetadataBreakVectorTest.cpp
+++ b/Engine/Plugins/PCG/Source/PCG/Private/Tests/Elements/Metadata/PCGMetadataBreakVectorTest.cpp
@@ -34,7 +34,7 @@ namespace PCGBreakVectorTest
 	void GenerateSpatialData(PCGTestsCommon::FTestData& TestData)
 	{
 		FPCGTaggedData& SpatialTaggedData = TestData.InputData.TaggedData.Emplace_GetRef();
-		TObjectPtr<UPCGPointData> PointData = PCGTestsCommon::CreateRandomPointData(10, TestData.Seed);
+		TObjectPtr<UPCGBasePointData> PointData = PCGTestsCommon::CreateRandomBasePointData(10, TestData.Seed);
 		SpatialTaggedData.Data = PointData;
 		SpatialTaggedData.Pin = PCGPinConstants::DefaultInputLabel;
 
@@ -47,17 +47,17 @@ namespace PCGBreakVectorTest
 		PointData->Metadata->CreateRotatorAttribute(RotatorAttribute, FRotator::ZeroRotator, bAllowsInterpolation, bOverrideParent);
 		PointData->Metadata->CreateFloatAttribute(FloatAttribute, 0, bAllowsInterpolation, bOverrideParent);
 
-		TArray<FPCGPoint>& SourcePoints = PointData->GetMutablePoints();
+		TPCGValueRange<int64> MetadataEntryRange = PointData->GetMetadataEntryValueRange();
 
 		// Randomize attribute values, and leave the second half to default value
 		FRandomStream RandomSource(TestData.Seed);
-		for (int I = 0; I < SourcePoints.Num() / 2; ++I)
+		for (int I = 0; I < PointData->GetNumPoints() / 2; ++I)
 		{
-			UPCGMetadataAccessorHelpers::SetVector2Attribute(SourcePoints[I], PointData->Metadata, Vec2Attribute, FVector2D(RandomSource.VRand()));
-			UPCGMetadataAccessorHelpers::SetVectorAttribute(SourcePoints[I], PointData->Metadata, Vec3Attribute, RandomSource.VRand());
-			UPCGMetadataAccessorHelpers::SetVector4Attribute(SourcePoints[I], PointData->Metadata, Vec4Attribute, FVector4(RandomSource.VRand(), RandomSource.FRand()));
-			UPCGMetadataAccessorHelpers::SetRotatorAttribute(SourcePoints[I], PointData->Metadata, RotatorAttribute, FRotator(RandomSource.FRand(), RandomSource.FRand(), RandomSource.FRand()));
-			UPCGMetadataAccessorHelpers::SetFloatAttribute(SourcePoints[I], PointData->Metadata, FloatAttribute, RandomSource.FRand());
+			UPCGMetadataAccessorHelpers::SetVector2AttributeByMetadataKey(MetadataEntryRange[I], PointData->Metadata, Vec2Attribute, FVector2D(RandomSource.VRand()));
+			UPCGMetadataAccessorHelpers::SetVectorAttributeByMetadataKey(MetadataEntryRange[I], PointData->Metadata, Vec3Attribute, RandomSource.VRand());
+			UPCGMetadataAccessorHelpers::SetVector4AttributeByMetadataKey(MetadataEntryRange[I], PointData->Metadata, Vec4Attribute, FVector4(RandomSource.VRand(), RandomSource.FRand()));
+			UPCGMetadataAccessorHelpers::SetRotatorAttributeByMetadataKey(MetadataEntryRange[I], PointData->Metadata, RotatorAttribute, FRotator(RandomSource.FRand(), RandomSource.FRand(), RandomSource.FRand()));
+			UPCGMetadataAccessorHelpers::SetFloatAttributeByMetadataKey(MetadataEntryRange[I], PointData->Metadata, FloatAttribute, RandomSource.FRand());
 		}
 	}
 
diff --git a/Engine/Plugins/PCG/Source/PCG/Private/Tests/Elements/Metadata/PCGMetadataMakeTransformTest.cpp b/Engine/Plugins/PCG/Source/PCG/Private/Tests/Elements/Metadata/PCGMetadataMakeTransformTest.cpp
--- a/Engine/Plugins/PCG/Source/PCG/Private/Tests/Elements/Metadata/PCGMetadataMakeTransformTest.cpp
+++ b/Engine/Plugins/PCG/Source/PCG/Private/Tests/Elements/Metadata/PCGMetadataMakeTransformTest.cpp
@@ -37,7 +37,7 @@ namespace PCGMetadataMakeTransformTest
 
 	void GenerateSpatialData(PCGTestsCommon::FTestData& TestData)
 	{
-		TObjectPtr<UPCGPointData> PointData = PCGTestsCommon::CreateRandomPointData(10, TestData.Seed);
+		TObjectPtr<UPCGBasePointData> PointData = PCGTestsCommon::CreateRandomBasePointData(10, TestData.Seed);
 		
 		ConnectData(TestData.InputData, PointData);
 
@@ -47,14 +47,14 @@ namespace PCGMetadataMakeTransformTest
 		PointData->Metadata->CreateVectorAttribute(Vec3Attribute, FVector::ZeroVector, bAllowsInterpolation, bOverrideParent);
 		PointData->Metadata->CreateQuatAttribute(QuatAttribute, FQuat::Identity, bAllowsInterpolation, bOverrideParent);
 
-		TArray<FPCGPoint>& SourcePoints = PointData->GetMutablePoints();
+		TPCGValueRange<int64> MetadataEntryRange = PointData->GetMetadataEntryValueRange();
 
 		// Randomize attribute values, and leave the second half to default value
 		FRandomStream RandomSource(TestData.Seed);
-		for (int I = 0; I < SourcePoints.Num() / 2; ++I)
+		for (int I = 0; I < PointData->GetNumPoints() / 2; ++I)
 		{
-			UPCGMetadataAccessorHelpers::SetVectorAttribute(SourcePoints[I], PointData->Metadata, Vec3Attribute, RandomSource.VRand());
-			UPCGMetadataAccessorHelpers::SetQuatAttribute(SourcePoints[I], PointData->Metadata, QuatAttribute, FRotator(RandomSource.FRand(), RandomSource.FRand(), RandomSource.FRand()).Quaternion());
+			UPCGMetadataAccessorHelpers::SetVectorAttributeByMetadataKey(MetadataEntryRange[I], PointData->Metadata, Vec3Attribute, RandomSource.VRand());
+			UPCGMetadataAccessorHelpers::SetQuatAttributeByMetadataKey(MetadataEntryRange[I], PointData->Metadata, QuatAttribute, FRotator(RandomSource.FRand(), RandomSource.FRand(), RandomSource.FRand()).Quaternion());
 		}
 	}
 
diff --git a/Engine/Plugins/PCG/Source/PCG/Private/Tests/Elements/Metadata/PCGMetadataMathsOpTest.cpp b/Engine/Plugins/PCG/Source/PCG/Private/Tests/Elements/Metadata/PCGMetadataMathsOpTest.cpp
--- a/Engine/Plugins/PCG/Source/PCG/Private/Tests/Elements/Metadata/PCGMetadataMathsOpTest.cpp
+++ b/Engine/Plugins/PCG/Source/PCG/Private/Tests/Elements/Metadata/PCGMetadataMathsOpTest.cpp
@@ -144,9 +144,11 @@ protected:
 			for (int32 InputIndex = 0; InputIndex < NumData; ++InputIndex)
 			{
 				int32 M = NumData * PinIndex + InputIndex + 1;
-				TObjectPtr<UPCGPointData> PointData = PCGTestsCommon::CreateEmptyPointData();
-				TArray<FPCGPoint>& Points = PointData->GetMutablePoints();
-				Points.Emplace(FTransform(M * DefaultTranslation), static_cast<double>(NumData * NumPins) / M, M);
+				TObjectPtr<UPCGBasePointData> PointData = PCGTestsCommon::CreateEmptyBasePointData();
+				PointData->SetNumPoints(1);
+				PointData->SetTransform(FTransform(M * DefaultTranslation));
+				PointData->SetDensity(static_cast<double>(NumData * NumPins) / M);
+				PointData->SetSeed(M);
 				TaggedData.Emplace_GetRef().Data = PointData;
 			}
 		}
diff --git a/Engine/Plugins/PCG/Source/PCG/Private/Tests/Elements/PCGAddTagTest.cpp b/Engine/Plugins/PCG/Source/PCG/Private/Tests/Elements/PCGAddTagTest.cpp
--- a/Engine/Plugins/PCG/Source/PCG/Private/Tests/Elements/PCGAddTagTest.cpp
+++ b/Engine/Plugins/PCG/Source/PCG/Private/Tests/Elements/PCGAddTagTest.cpp
@@ -23,7 +23,7 @@ bool FPCGAddTagTest_ZeroTags::RunTest(const FString& Parameters)
 
 	// At least one input data is required to add tags and get an output
 	FPCGTaggedData& NewTaggedData = Context->InputData.TaggedData.Emplace_GetRef();
-	NewTaggedData.Data = PCGTestsCommon::CreatePointData();
+	NewTaggedData.Data = PCGTestsCommon::CreateBasePointData();
 	NewTaggedData.Pin = PCGPinConstants::DefaultInputLabel;
 
 	while (!TestElement->Execute(Context.Get())) {}
@@ -63,7 +63,7 @@ bool FPCGAddTagTest_OneTag::RunTest(const FString& Parameters)
 
 	// At least one input data is required to add tags and get an output
 	FPCGTaggedData& NewTaggedData = Context->InputData.TaggedData.Emplace_GetRef();
-	NewTaggedData.Data = PCGTestsCommon::CreatePointData();
+	NewTaggedData.Data = PCGTestsCommon::CreateBasePointData();
 	NewTaggedData.Pin = PCGPinConstants::DefaultInputLabel;
 
 	while (!TestElement->Execute(Context.Get())) {}
@@ -110,7 +110,7 @@ bool FPCGAddTagTest_MultipleTags::RunTest(const FString& Parameters)
 
 	// At least one input data is required to add tags and get an output
 	FPCGTaggedData& NewTaggedData = Context->InputData.TaggedData.Emplace_GetRef();
-	NewTaggedData.Data = PCGTestsCommon::CreatePointData();
+	NewTaggedData.Data = PCGTestsCommon::CreateBasePointData();
 	NewTaggedData.Pin = PCGPinConstants::DefaultInputLabel;
 
 	while (!TestElement->Execute(Context.Get())) {}
diff --git a/Engine/Plugins/PCG/Source/PCG/Private/Tests/Elements/PCGApplyScaleToBoundsTest.cpp b/Engine/Plugins/PCG/Source/PCG/Private/Tests/Elements/PCGApplyScaleToBoundsTest.cpp
--- a/Engine/Plugins/PCG/Source/PCG/Private/Tests/Elements/PCGApplyScaleToBoundsTest.cpp
+++ b/Engine/Plugins/PCG/Source/PCG/Private/Tests/Elements/PCGApplyScaleToBoundsTest.cpp
@@ -16,15 +16,14 @@ bool FPCGApplyScaleToBoundsTest_Basic::RunTest(const FString& Parameters)
 	FPCGTaggedData& Inputs = TestData.InputData.TaggedData.Emplace_GetRef();
 	Inputs.Pin = PCGPinConstants::DefaultInputLabel;
 	UPCGApplyScaleToBoundsSettings* Settings = PCGTestsCommon::GenerateSettings<UPCGApplyScaleToBoundsSettings>(TestData);
-	UPCGPointData* InData = PCGTestsCommon::CreateRandomPointData(10, 42, false);
+	UPCGBasePointData* InData = PCGTestsCommon::CreateRandomBasePointData(10, 42, false);
 	Inputs.Data = InData;
 
 	// Setting the scale explicitly
-	UPCGPointData* InputPointData = Cast<UPCGPointData>(InData);
-	TArray<FPCGPoint>& InPoints = InputPointData->GetMutablePoints();
-	for (FPCGPoint& InPoint : InPoints)
+	TPCGValueRange<FTransform> TransformRange = InData->GetTransformValueRange();
+	for (FTransform& PointTransform : TransformRange)
 	{
-		InPoint.Transform.SetScale3D(FVector(3.0));
+		PointTransform.SetScale3D(FVector(3.0));
 	}
 
 	FPCGElementPtr TestElement = TestData.Settings->GetElement();
@@ -65,15 +64,14 @@ bool FPCGApplyScaleToBoundsTest_NegativeScale::RunTest(const FString& Parameters
 	FPCGTaggedData& Inputs = TestData.InputData.TaggedData.Emplace_GetRef();
 	Inputs.Pin = PCGPinConstants::DefaultInputLabel;
 	UPCGApplyScaleToBoundsSettings* Settings = PCGTestsCommon::GenerateSettings<UPCGApplyScaleToBoundsSettings>(TestData);
-	UPCGPointData* InData = PCGTestsCommon::CreateRandomPointData(10, 42, false);
+	UPCGBasePointData* InData = PCGTestsCommon::CreateRandomBasePointData(10, 42, false);
 	Inputs.Data = InData;
 
 	// Setting the scale explicitly
-	UPCGPointData* InputPointData = Cast<UPCGPointData>(InData);
-	TArray<FPCGPoint>& InPoints = InputPointData->GetMutablePoints();
-	for (FPCGPoint& InPoint : InPoints)
+	TPCGValueRange<FTransform> TransformRange = InData->GetTransformValueRange();
+	for (FTransform& PointTransform : TransformRange)
 	{
-		InPoint.Transform.SetScale3D(FVector(-2.0));
+		PointTransform.SetScale3D(FVector(-2.0));
 	}
 
 	FPCGElementPtr TestElement = TestData.Settings->GetElement(); // returns nullptr
@@ -93,15 +91,15 @@ bool FPCGApplyScaleToBoundsTest_NegativeScale::RunTest(const FString& Parameters
 
 	UTEST_EQUAL("Output point count", OutPointData->GetNumPoints(), 10);
 
-	const TConstPCGValueRange<FTransform> TransformRange = OutPointData->GetConstTransformValueRange();
-	const TConstPCGValueRange<FVector> BoundsMinRange = OutPointData->GetConstBoundsMinValueRange();
-	const TConstPCGValueRange<FVector> BoundsMaxRange = OutPointData->GetConstBoundsMaxValueRange();
+	const TConstPCGValueRange<FTransform> OutTransformRange = OutPointData->GetConstTransformValueRange();
+	const TConstPCGValueRange<FVector> OutBoundsMinRange = OutPointData->GetConstBoundsMinValueRange();
+	const TConstPCGValueRange<FVector> OutBoundsMaxRange = OutPointData->GetConstBoundsMaxValueRange();
 
 	for (int i = 0; i < OutPointData->GetNumPoints(); ++i)
 	{
-		UTEST_EQUAL("Output points should have a Bounds Min of -2", BoundsMinRange[i], FVector(-2.0));
-		UTEST_EQUAL("Output points should have a Bounds Max of 2", BoundsMaxRange[i], FVector(2.0));
-		UTEST_EQUAL("Output points Bounds Max should be negative -1", TransformRange[i].GetScale3D(), FVector(-1.0));
+		UTEST_EQUAL("Output points should have a Bounds Min of -2", OutBoundsMinRange[i], FVector(-2.0));
+		UTEST_EQUAL("Output points should have a Bounds Max of 2", OutBoundsMaxRange[i], FVector(2.0));
+		UTEST_EQUAL("Output points Bounds Max should be negative -1", OutTransformRange[i].GetScale3D(), FVector(-1.0));
 	}
 
 	return true;
diff --git a/Engine/Plugins/PCG/Source/PCG/Private/Tests/Elements/PCGAttributeCastTest.cpp b/Engine/Plugins/PCG/Source/PCG/Private/Tests/Elements/PCGAttributeCastTest.cpp
--- a/Engine/Plugins/PCG/Source/PCG/Private/Tests/Elements/PCGAttributeCastTest.cpp
+++ b/Engine/Plugins/PCG/Source/PCG/Private/Tests/Elements/PCGAttributeCastTest.cpp
@@ -31,17 +31,25 @@ protected:
 		Settings->OutputType = static_cast<EPCGMetadataTypes>(PCG::Private::MetadataTypes<OutType>::Id);
 
 		PCGDataType* InputData = NewObject<PCGDataType>();
+		TPCGValueRange<int64> MetadataEntryRange;
+		if constexpr (std::is_same_v<PCGDataType, UPCGPointData> || std::is_same_v<PCGDataType, UPCGPointArrayData>)
+		{
+			InputData->SetNumPoints(InData.Num()-1);
+			MetadataEntryRange = InputData->GetMetadataEntryValueRange();
+		}
+
 		FPCGMetadataAttribute<InType>* Attribute = InputData->Metadata->template CreateAttribute<InType>(InputAttributeName, InData[0], true, false);
+		int32 WriteIndex = 0;
 		for (int32 i = 1; i < InData.Num(); ++i)
 		{
 			PCGMetadataEntryKey EntryKey = InputData->Metadata->AddEntry();
 			Attribute->SetValue(EntryKey, InData[i]);
 
-			if constexpr (std::is_same_v<PCGDataType, UPCGPointData>)
+			if constexpr (std::is_same_v<PCGDataType, UPCGPointData> || std::is_same_v<PCGDataType, UPCGPointArrayData>)
 			{
-				FPCGPoint& Point = InputData->GetMutablePoints().Emplace_GetRef();
-				Point.MetadataEntry = EntryKey;
+				MetadataEntryRange[WriteIndex] = EntryKey;
 			}
+			++WriteIndex;
 		}
 
 		FPCGTaggedData& InputTaggedData = TestData.InputData.TaggedData.Emplace_GetRef();
@@ -59,9 +67,9 @@ protected:
 
 		UTEST_NOT_NULL("Output is an of the same type as input", OutputData);
 
-		if constexpr (std::is_same_v<PCGDataType, UPCGPointData>)
+		if constexpr (std::is_same_v<PCGDataType, UPCGPointData> || std::is_same_v<PCGDataType, UPCGPointArrayData>)
 		{
-			UTEST_EQUAL("Output has the same number of points", InputData->GetPoints().Num(), OutputData->GetPoints().Num());
+			UTEST_EQUAL("Output has the same number of points", InputData->GetNumPoints(), OutputData->GetNumPoints());
 		}
 
 		const FPCGMetadataAttribute<OutType>* OutputAttribute = OutputData->Metadata->template GetConstTypedAttribute<OutType>(OutputAttributeName);
@@ -69,12 +77,19 @@ protected:
 		UTEST_NOT_NULL("Output metadata has the expected attribute", OutputAttribute);
 
 		UTEST_EQUAL("Default value was casted correctly", OutputAttribute->GetValue(PCGDefaultValueKey), ExpectedOutData[0]);
+
+		TConstPCGValueRange<int64> OutMetadataEntryRange;
+		if constexpr (std::is_same_v<PCGDataType, UPCGPointData> || std::is_same_v<PCGDataType, UPCGPointArrayData>)
+		{
+			OutMetadataEntryRange = OutputData->GetConstMetadataEntryValueRange();
+		}
+
 		for (int32 i = 1; i < ExpectedOutData.Num(); ++i)
 		{
 			PCGMetadataEntryKey EntryKey(i - 1);
-			if constexpr (std::is_same_v<PCGDataType, UPCGPointData>)
+			if constexpr (std::is_same_v<PCGDataType, UPCGPointData> || std::is_same_v<PCGDataType, UPCGPointArrayData>)
 			{
-				EntryKey = OutputData->GetPoints()[i-1].MetadataEntry;
+				EntryKey = OutMetadataEntryRange[i-1];
 			}
 
 			UTEST_EQUAL(*FString::Printf(TEXT("Entry %d was casted correctly"), i - 1), OutputAttribute->GetValueFromItemKey(EntryKey), ExpectedOutData[i]);
@@ -143,7 +158,14 @@ bool FPCGAttributeCast_AttributeSetQuatToRotator::RunTest(const FString& Paramet
 
 bool FPCGAttributeCast_PointsIntToDouble::RunTest(const FString& Parameters)
 {
-	return RunTestInternal<int, double, UPCGPointData>({ 5, 1, 2 }, { 5.0, 1.0, 2.0 });
+	if (CVarPCGEnablePointArrayData.GetValueOnAnyThread())
+	{
+		return RunTestInternal<int, double, UPCGPointArrayData>({ 5, 1, 2 }, { 5.0, 1.0, 2.0 });
+	}
+	else
+	{
+		return RunTestInternal<int, double, UPCGPointData>({ 5, 1, 2 }, { 5.0, 1.0, 2.0 });
+	}
 }
 
 bool FPCGAttributeCast_PointsIntToDoubleToDensity::RunTest(const FString& Parameters)
@@ -157,16 +179,16 @@ bool FPCGAttributeCast_PointsIntToDoubleToDensity::RunTest(const FString& Parame
 	Settings->OutputTarget.SetPointProperty(EPCGPointProperties::Density);
 	Settings->OutputType = EPCGMetadataTypes::Double;
 
-	UPCGPointData* InputPoints = NewObject<UPCGPointData>();
+	UPCGBasePointData* InputPoints = PCGTestsCommon::CreateEmptyBasePointData();
 	FPCGMetadataAttribute<int32>* Attribute = InputPoints->Metadata->CreateAttribute<int32>(InputAttributeName, 5, true, false);
 
-	TArray<FPCGPoint>& Points = InputPoints->GetMutablePoints();
-	FPCGPoint& Point1 = Points.Emplace_GetRef();
-	FPCGPoint& Point2 = Points.Emplace_GetRef();
-	InputPoints->Metadata->InitializeOnSet(Point1.MetadataEntry);
-	InputPoints->Metadata->InitializeOnSet(Point2.MetadataEntry);
-	Attribute->SetValue(Point1.MetadataEntry, 1);
-	Attribute->SetValue(Point2.MetadataEntry, 2);
+	InputPoints->SetNumPoints(2);
+	TPCGValueRange<int64> MetadataEntryRange = InputPoints->GetMetadataEntryValueRange();
+
+	InputPoints->Metadata->InitializeOnSet(MetadataEntryRange[0]);
+	InputPoints->Metadata->InitializeOnSet(MetadataEntryRange[1]);
+	Attribute->SetValue(MetadataEntryRange[0], 1);
+	Attribute->SetValue(MetadataEntryRange[1], 2);
 
 	FPCGTaggedData& InputTaggedData = TestData.InputData.TaggedData.Emplace_GetRef();
 	InputTaggedData.Data = InputPoints;
@@ -179,13 +201,14 @@ bool FPCGAttributeCast_PointsIntToDoubleToDensity::RunTest(const FString& Parame
 
 	UTEST_EQUAL("There is one output", Context->OutputData.TaggedData.Num(), 1);
 
-	const UPCGPointData* OutputPoints = Cast<const UPCGPointData>(Context->OutputData.TaggedData[0].Data);
+	const UPCGBasePointData* OutputPoints = Cast<const UPCGBasePointData>(Context->OutputData.TaggedData[0].Data);
 
 	UTEST_NOT_NULL("Output is a point data", OutputPoints);
-	UTEST_EQUAL("Output has 2 points", OutputPoints->GetPoints().Num(), 2);
+	UTEST_EQUAL("Output has 2 points", OutputPoints->GetNumPoints(), 2);
 
-	UTEST_EQUAL("First point has the right value", OutputPoints->GetPoint(0).Density, 1.0f);
-	UTEST_EQUAL("Second point has the right value", OutputPoints->GetPoint(1).Density, 2.0f);
+	const TConstPCGValueRange<float> DensityRange = OutputPoints->GetConstDensityValueRange();
+	UTEST_EQUAL("First point has the right value", DensityRange[0], 1.0f);
+	UTEST_EQUAL("Second point has the right value", DensityRange[1], 2.0f);
 
 	return true;
 }
diff --git a/Engine/Plugins/PCG/Source/PCG/Private/Tests/Elements/PCGAttributeFilterTest.cpp b/Engine/Plugins/PCG/Source/PCG/Private/Tests/Elements/PCGAttributeFilterTest.cpp
--- a/Engine/Plugins/PCG/Source/PCG/Private/Tests/Elements/PCGAttributeFilterTest.cpp
+++ b/Engine/Plugins/PCG/Source/PCG/Private/Tests/Elements/PCGAttributeFilterTest.cpp
@@ -26,16 +26,21 @@ namespace PCGPointFilterTest
 	const FName OutsideFilterLabel = TEXT("OutsideFilter");
 	const FName FilterLabel = TEXT("Filter");
 
-	UPCGPointData* GeneratePointDataWithRandomDensity(int32 InNumPoints, int32 InRandomSeed)
+	UPCGBasePointData* GeneratePointDataWithRandomDensity(int32 InNumPoints, int32 InRandomSeed)
 	{
-		UPCGPointData* PointData = PCGTestsCommon::CreateEmptyPointData();
-		TArray<FPCGPoint>& Points = PointData->GetMutablePoints();
+		UPCGBasePointData* PointData = PCGTestsCommon::CreateEmptyBasePointData();
+		
+		PointData->SetNumPoints(InNumPoints);
+		PointData->AllocateProperties(EPCGPointNativeProperties::Transform | EPCGPointNativeProperties::Density | EPCGPointNativeProperties::Seed);
 
 		FRandomStream RandomSource(InRandomSeed);
 
+		FPCGPointValueRanges ValueRanges(PointData, /*bAllocate=*/false);
 		for (int I = 0; I < InNumPoints; ++I)
 		{
-			FPCGPoint& Point = Points.Emplace_GetRef(FTransform(), RandomSource.FRand(), I);
+			ValueRanges.TransformRange[I] = FTransform();
+			ValueRanges.DensityRange[I] = RandomSource.FRand();
+			ValueRanges.SeedRange[I] = I;
 		}
 
 		return PointData;
@@ -322,34 +327,32 @@ bool FPCGAttributeFilterSkipTestBug::RunTest(const FString& Parameters)
 	Settings->bUseSpatialQuery = true;
 
 	FPCGElementPtr TestElement = TestData.Settings->GetElement();
-	UPCGPointData* InputPointData = NewObject<UPCGPointData>();
-	UPCGPointData* ThresholdPointData = NewObject<UPCGPointData>();
-
-	TArray<FPCGPoint>& InputPoints = InputPointData->GetMutablePoints();
-	TArray<FPCGPoint>& ThresholdPoints = ThresholdPointData->GetMutablePoints();
+	UPCGBasePointData* InputPointData = PCGTestsCommon::CreateEmptyBasePointData();
+	UPCGBasePointData* ThresholdPointData = PCGTestsCommon::CreateEmptyBasePointData();
 
 	// Take a big number to make sure we go over the 256 default chunk size
 	constexpr int32 NumPoints = 2048;
 	constexpr int32 HalfNumPoints = NumPoints / 2;
-	InputPoints.Reserve(NumPoints);
-	ThresholdPoints.Reserve(NumPoints);
+	InputPointData->SetNumPoints(NumPoints);
+	TPCGValueRange<FTransform> InputTransformRange = InputPointData->GetTransformValueRange();
+
+	ThresholdPointData->SetNumPoints(NumPoints);
+	TPCGValueRange<FTransform> ThresholdTransformRange = ThresholdPointData->GetTransformValueRange();
+	TPCGValueRange<float> ThresholdDensityRange = ThresholdPointData->GetDensityValueRange();
 
 	for (int32 i = 0; i < NumPoints; ++i)
 	{
 		// First half are very different so the sampling should fail, second half are the same the sampling to succeed but the filtering to fail.
 		if (i >= HalfNumPoints)
 		{
-			FPCGPoint& Point = InputPoints.Emplace_GetRef();
-			Point.Transform.SetLocation(FVector(10 * i, 10 * i, 10 * i));
-			FPCGPoint& ThresholdPoint = ThresholdPoints.Add_GetRef(Point);
-			ThresholdPoint.Density = 0.5f;
+			InputTransformRange[i].SetLocation(FVector(10 * i, 10 * i, 10 * i));
+			ThresholdTransformRange[i] = InputTransformRange[i];
+			ThresholdDensityRange[i] = 0.5f;
 		}
 		else
 		{
-			FPCGPoint& Point = InputPoints.Emplace_GetRef();
-			Point.Transform.SetLocation(FVector(i, i, i));
-			FPCGPoint& ThresholdPoint = ThresholdPoints.Emplace_GetRef();
-			ThresholdPoint.Transform.SetLocation(FVector(-10 * i - 1000, -10 * i - 1000, -10 * i - 1000));
+			InputTransformRange[i].SetLocation(FVector(i, i, i));
+			ThresholdTransformRange[i].SetLocation(FVector(-10 * i - 1000, -10 * i - 1000, -10 * i - 1000));
 		}
 	}
 
diff --git a/Engine/Plugins/PCG/Source/PCG/Private/Tests/Elements/PCGAttributeNoiseTest.cpp b/Engine/Plugins/PCG/Source/PCG/Private/Tests/Elements/PCGAttributeNoiseTest.cpp
--- a/Engine/Plugins/PCG/Source/PCG/Private/Tests/Elements/PCGAttributeNoiseTest.cpp
+++ b/Engine/Plugins/PCG/Source/PCG/Private/Tests/Elements/PCGAttributeNoiseTest.cpp
@@ -19,15 +19,20 @@ bool FPCGAttributeNoiseDensityTest::RunTest(const FString& Parameters)
 	Settings->InputSource.SetPointProperty(EPCGPointProperties::Density);
 	FPCGElementPtr NoiseElement = TestData.Settings->GetElement();
 
-	TObjectPtr<UPCGPointData> PointData = PCGTestsCommon::CreateEmptyPointData();
-	TArray<FPCGPoint>& Points = PointData->GetMutablePoints();
+	TObjectPtr<UPCGBasePointData> PointData = PCGTestsCommon::CreateEmptyBasePointData();
 
 	FRandomStream RandomSource(TestData.Seed);
 	const int PointCount = 5;
+
+	PointData->SetNumPoints(PointCount);
+	
+	TPCGValueRange<int32> SeedRange = PointData->GetSeedValueRange();
+	TPCGValueRange<float> DensityRange = PointData->GetDensityValueRange();
+
 	for (int I = 0; I < PointCount; ++I)
 	{
-		FPCGPoint& Point = Points.Emplace_GetRef(FTransform(), 1, I);
-		Point.Density = RandomSource.GetFraction();
+		DensityRange[I] = RandomSource.GetFraction();
+		SeedRange[I] = I;
 	}
 
 	FPCGTaggedData& TaggedData = TestData.InputData.TaggedData.Emplace_GetRef(FPCGTaggedData());
@@ -59,7 +64,7 @@ bool FPCGAttributeNoiseDensityTest::RunTest(const FString& Parameters)
 			const UPCGSpatialData* InSpatialData = Cast<UPCGSpatialData>(Input.Data);
 			check(InSpatialData);
 
-			const UPCGPointData* InPointData = InSpatialData->ToPointData(Context.Get());
+			const UPCGBasePointData* InPointData = InSpatialData->ToBasePointData(Context.Get());
 			check(InPointData);
 
 			const UPCGSpatialData* OutSpatialData = Cast<UPCGSpatialData>(Output.Data);
@@ -70,7 +75,7 @@ bool FPCGAttributeNoiseDensityTest::RunTest(const FString& Parameters)
 				continue;
 			}
 
-			const UPCGPointData* OutPointData = OutSpatialData->ToPointData(Context.Get());
+			const UPCGBasePointData* OutPointData = OutSpatialData->ToBasePointData(Context.Get());
 
 			if (!TestNotNull("Valid output PointData", OutPointData))
 			{
@@ -78,18 +83,16 @@ bool FPCGAttributeNoiseDensityTest::RunTest(const FString& Parameters)
 				continue;
 			}
 
-			const TArray<FPCGPoint>& InPoints = InPointData->GetPoints();
-			const TArray<FPCGPoint>& OutPoints = OutPointData->GetPoints();
-
-			if (!TestEqual("Input and output point counts match", InPoints.Num(), OutPoints.Num()))
+			if (!TestEqual("Input and output point counts match", InPointData->GetNumPoints(), OutPointData->GetNumPoints()))
 			{ 
 				bTestPassed = false;
 				continue;
 			}
 
+			const TConstPCGValueRange<float> DensityRange = OutPointData->GetConstDensityValueRange();
 			for (int PointIndex = 0; PointIndex < ExpectedOutput.Num(); ++PointIndex)
 			{
-				bTestPassed &= TestEqual("Correct density", OutPoints[PointIndex].Density, ExpectedOutput[PointIndex]);
+				bTestPassed &= TestEqual("Correct density", DensityRange[PointIndex], ExpectedOutput[PointIndex]);
 			}
 		}
 
@@ -214,11 +217,11 @@ bool FPCGAttributeNoiseDensityTest::RunTest(const FString& Parameters)
 		Settings->NoiseMax = 0.f;
 		Settings->Mode = EPCGAttributeNoiseMode::Add;
 		bTestPassed &= ValidateDensityNoise({ 
-			Points[0].Density,
-			Points[1].Density,
-			Points[2].Density,
-			Points[3].Density,
-			Points[4].Density,
+			DensityRange[0],
+			DensityRange[1],
+			DensityRange[2],
+			DensityRange[3],
+			DensityRange[4],
 		});
 	}
 
@@ -243,7 +246,7 @@ bool FPCGAttributeNoiseNotExistingAttributeTest::RunTest(const FString& Paramete
 
 	Settings->InputSource.SetAttributeName(TEXT("Hi"));
 
-	UPCGPointData* PointData = PCGTestsCommon::CreateRandomPointData(/*PointCount*/5, /*Seed*/ 42, /*bRandomDensity=*/true);
+	UPCGBasePointData* PointData = PCGTestsCommon::CreateRandomBasePointData(/*PointCount*/5, /*Seed*/ 42, /*bRandomDensity=*/true);
 
 	FPCGTaggedData& TaggedData = TestData.InputData.TaggedData.Emplace_GetRef();
 	TaggedData.Data = PointData;
@@ -271,12 +274,14 @@ bool FPCGAttributeNoiseInvalidTypeTest::RunTest(const FString& Parameters)
 
 	Settings->InputSource.SetAttributeName(TEXT("MyStr"));
 
-	UPCGPointData* PointData = PCGTestsCommon::CreateRandomPointData(/*PointCount*/5, /*Seed*/ 42, /*bRandomDensity=*/true);
+	UPCGBasePointData* PointData = PCGTestsCommon::CreateRandomBasePointData(/*PointCount*/5, /*Seed*/ 42, /*bRandomDensity=*/true);
 	FPCGMetadataAttribute<FString>* StrAttribute = PointData->Metadata->CreateAttribute<FString>(TEXT("MyStr"), FString{}, /*bAllowInterpolation=*/false, /*bOverrideParent=*/false);
-	for (FPCGPoint& Point : PointData->GetMutablePoints())
+
+	TPCGValueRange<int64> MetadataEntryRange = PointData->GetMetadataEntryValueRange();
+	for (int64& MetadataEntry : MetadataEntryRange)
 	{
-		PointData->Metadata->InitializeOnSet(Point.MetadataEntry);
-		StrAttribute->SetValue(Point.MetadataEntry, TEXT("Hey"));
+		PointData->Metadata->InitializeOnSet(MetadataEntry);
+		StrAttribute->SetValue(MetadataEntry, TEXT("Hey"));
 	}
 
 	FPCGTaggedData& TaggedData = TestData.InputData.TaggedData.Emplace_GetRef();
diff --git a/Engine/Plugins/PCG/Source/PCG/Private/Tests/Elements/PCGAttributePartitionTest.cpp b/Engine/Plugins/PCG/Source/PCG/Private/Tests/Elements/PCGAttributePartitionTest.cpp
--- a/Engine/Plugins/PCG/Source/PCG/Private/Tests/Elements/PCGAttributePartitionTest.cpp
+++ b/Engine/Plugins/PCG/Source/PCG/Private/Tests/Elements/PCGAttributePartitionTest.cpp
@@ -44,15 +44,15 @@ namespace PCGAttributePartitionTest
 			OverrideTaggedData.Pin = GET_MEMBER_NAME_CHECKED(UPCGMetadataPartitionSettings, PartitionAttributeNames);
 		}
 
-		UPCGPointData* InputPointData = NewObject<UPCGPointData>();
-		TArray<FPCGPoint>& Points = InputPointData->GetMutablePoints();
-		Points.Reserve(20);
-		for (int32 i = 0; i < 20; ++i)
+		UPCGBasePointData* InputPointData = PCGTestsCommon::CreateEmptyBasePointData();
+		InputPointData->SetNumPoints(20);
+		TPCGValueRange<FTransform> TransformRange = InputPointData->GetTransformValueRange();
+
+		for (int32 i = 0; i < InputPointData->GetNumPoints(); ++i)
 		{
-			FPCGPoint& Point = Points.Emplace_GetRef();
 			// Also add a very small number on Y, to verify that approximation errors are not giving different partitions.
 			const double Jitter = i >= 10 ? UE_DOUBLE_SMALL_NUMBER : 0.0;
-			Point.Transform.SetLocation(FVector(i % 10, i % 5 + Jitter, 0.0));
+			TransformRange[i].SetLocation(FVector(i % 10, i % 5 + Jitter, 0.0));
 		}
 
 		FPCGTaggedData& InputTaggedData = TestData.InputData.TaggedData.Emplace_GetRef();
@@ -86,9 +86,9 @@ namespace PCGAttributePartitionTest
 			bool bFirstPoint = true;
 			bool bAllEquals = true;
 			
-			const TConstPCGValueRange<FTransform> TransformRange = OutputPointData->GetConstTransformValueRange();
+			const TConstPCGValueRange<FTransform> OutTransformRange = OutputPointData->GetConstTransformValueRange();
 			
-			for (const FTransform& Transform : TransformRange)
+			for (const FTransform& Transform : OutTransformRange)
 			{
 				if (bFirstPoint)
 				{
@@ -127,13 +127,13 @@ namespace PCGAttributePartitionTest
 		Settings->bDoNotPartition = bNoPartition;
 		Settings->PartitionIndexAttributeName = PartitionIndexAttributeName;
 
-		UPCGPointData* InputPointData = NewObject<UPCGPointData>();
-		TArray<FPCGPoint>& Points = InputPointData->GetMutablePoints();
-		Points.Reserve(100);
-		for (int32 i = 0; i < 100; ++i)
+		UPCGBasePointData* InputPointData = PCGTestsCommon::CreateEmptyBasePointData();
+		InputPointData->SetNumPoints(100);
+
+		TPCGValueRange<float> DensityRange = InputPointData->GetDensityValueRange();
+		for (int32 i = 0; i < InputPointData->GetNumPoints(); ++i)
 		{
-			FPCGPoint& Point = Points.Emplace_GetRef();
-			Point.Density = i % 10;
+			DensityRange[i] = i % 10;
 		}
 
 		const int32 ExpectedPointsPerData = bNoPartition ? 100 : 10;
@@ -223,26 +223,29 @@ bool FPCGAttributePartition_Order::RunTest(const FString& Parameters)
 	const FName AttributeName = TEXT("Attr");
 	Settings->PartitionAttributeSelectors[0].SetAttributeName(AttributeName);
 
-	UPCGPointData* InputPointData = NewObject<UPCGPointData>();
+	UPCGBasePointData* InputPointData = PCGTestsCommon::CreateEmptyBasePointData();
 	FPCGMetadataAttribute<int>* Attribute = InputPointData->Metadata->CreateAttribute<int>(AttributeName, 0, false, false);
 
 	constexpr int NumPoints = 10;
 
-	TArray<FPCGPoint>& Points = InputPointData->GetMutablePoints();
-	Points.Reserve(NumPoints);
+	InputPointData->SetNumPoints(NumPoints);
+	
+	TPCGValueRange<float> DensityRange = InputPointData->GetDensityValueRange();
+	TPCGValueRange<int64> MetadataEntryRange = InputPointData->GetMetadataEntryValueRange();
+
 	for (int32 i = 0; i < NumPoints; ++i)
 	{
-		FPCGPoint& Point = Points.Emplace_GetRef();
 		// Store the index in the density
-		Point.Density = i;
-		InputPointData->Metadata->InitializeOnSet(Point.MetadataEntry);
-		Attribute->SetValue(Point.MetadataEntry, (2 * i) / NumPoints);
+		DensityRange[i] = i;
+		InputPointData->Metadata->InitializeOnSet(MetadataEntryRange[i]);
+		Attribute->SetValue(MetadataEntryRange[i], (2 * i) / NumPoints);
 	}
 
 	// Reverse the points order
 	for (int32 i = 0; i < NumPoints / 2; ++i)
 	{
-		std::swap(Points[i], Points[NumPoints - i - 1]);
+		std::swap(DensityRange[i], DensityRange[NumPoints - i - 1]);
+		std::swap(MetadataEntryRange[i], MetadataEntryRange[NumPoints - i - 1]);
 	}
 
 	FPCGTaggedData& InputTaggedData = TestData.InputData.TaggedData.Emplace_GetRef();
@@ -269,14 +272,14 @@ bool FPCGAttributePartition_Order::RunTest(const FString& Parameters)
 
 		UTEST_EQUAL(*FString::Printf(TEXT("Output %d has %d points"), i, NumPoints / 2), OutputPointData->GetNumPoints(), NumPoints / 2);
 
-		const TConstPCGValueRange<float> DensityRange = OutputPointData->GetConstDensityValueRange();
-		const TConstPCGValueRange<int64> MetadataEntryRange = OutputPointData->GetConstMetadataEntryValueRange();
+		const TConstPCGValueRange<float> OutDensityRange = OutputPointData->GetConstDensityValueRange();
+		const TConstPCGValueRange<int64> OutMetadataEntryRange = OutputPointData->GetConstMetadataEntryValueRange();
 
 		for (int32 j = 0; j < OutputPointData->GetNumPoints() - 1; ++j)
 		{
-			UTEST_TRUE(*FString::Printf(TEXT("Output %d: Point_%d is at the right place"), i, j), DensityRange[j] > DensityRange[j + 1]);
+			UTEST_TRUE(*FString::Printf(TEXT("Output %d: Point_%d is at the right place"), i, j), OutDensityRange[j] > OutDensityRange[j + 1]);
 			// First partition should have the value 1, and second partition should have the value 0
-			UTEST_EQUAL(*FString::Printf(TEXT("Output %d: Point_%d has the right attribute value"), i, j), OutputAttribute->GetValue(MetadataEntryRange[j]), (i + 1) % 2);
+			UTEST_EQUAL(*FString::Printf(TEXT("Output %d: Point_%d has the right attribute value"), i, j), OutputAttribute->GetValue(OutMetadataEntryRange[j]), (i + 1) % 2);
 		}
 	}
 
diff --git a/Engine/Plugins/PCG/Source/PCG/Private/Tests/Elements/PCGAttributeReduceTest.cpp b/Engine/Plugins/PCG/Source/PCG/Private/Tests/Elements/PCGAttributeReduceTest.cpp
--- a/Engine/Plugins/PCG/Source/PCG/Private/Tests/Elements/PCGAttributeReduceTest.cpp
+++ b/Engine/Plugins/PCG/Source/PCG/Private/Tests/Elements/PCGAttributeReduceTest.cpp
@@ -99,14 +99,21 @@ protected:
 
 	TArray<FPCGTaggedData> PreparePointsAndSelectors(FPCGAttributePropertyInputSelector& DensitySelector, FPCGAttributePropertyInputSelector& PositionSelector)
 	{
-		UPCGPointData* PointData = PCGTestsCommon::CreateEmptyPointData();
-		TArray<FPCGPoint>& Points = PointData->GetMutablePoints();
+		UPCGBasePointData* PointData = PCGTestsCommon::CreateEmptyBasePointData();
 
 		constexpr int32 PointCount = 5;
+		PointData->SetNumPoints(PointCount);
+		
+		TPCGValueRange<FTransform> TransformRange = PointData->GetTransformValueRange();
+		TPCGValueRange<float> DensityRange = PointData->GetDensityValueRange();
+		TPCGValueRange<int32> SeedRange = PointData->GetSeedValueRange();
+		
 		for (int32 i = 0; i < PointCount; ++i)
 		{
 			const FVector Position{ 3.0 * i, 3.0 * i + 1.0, 3.0 * i + 2.0 };
-			FPCGPoint& Point = Points.Emplace_GetRef(FTransform(Position), 0.1 * (i + 1), i);
+			TransformRange[i] = FTransform(Position);
+			DensityRange[i] = 0.1 * (i + 1);
+			SeedRange[i] = i;
 		}
 
 		TArray<FPCGTaggedData> Inputs;
@@ -152,19 +159,21 @@ protected:
 	{
 		const FName DoubleName = TEXT("double");
 
-		UPCGPointData* PointData = PCGTestsCommon::CreateEmptyPointData();
+		UPCGBasePointData* PointData = PCGTestsCommon::CreateEmptyBasePointData();
 		FPCGMetadataAttribute<double>* PointsAttribute = PointData->Metadata->CreateAttribute<double>(DoubleName, 0.0, /*bAllowsInterpolation=*/ true, /*bOverrideParent=*/false);
 
 		UPCGParamData* ParamData = NewObject<UPCGParamData>();
 		FPCGMetadataAttribute<double>* ParamAttribute = ParamData->Metadata->CreateAttribute<double>(DoubleName, 0.0, /*bAllowsInterpolation=*/ true, /*bOverrideParent=*/false);
 
-		TArray<FPCGPoint>& Points = PointData->GetMutablePoints();
 		constexpr int32 EntryCount = 5;
+
+		PointData->SetNumPoints(5);
+		TPCGValueRange<int64> MetadataEntryRange = PointData->GetMetadataEntryValueRange();
+
 		for (int32 i = 0; i < EntryCount; ++i)
 		{
-			FPCGPoint& Point = Points.Emplace_GetRef();
-			PointData->Metadata->InitializeOnSet(Point.MetadataEntry);
-			PointsAttribute->SetValue(Point.MetadataEntry, 0.1 * (i + 1));
+			PointData->Metadata->InitializeOnSet(MetadataEntryRange[i]);
+			PointsAttribute->SetValue(MetadataEntryRange[i], 0.1 * (i + 1));
 
 			ParamAttribute->SetValue(ParamData->Metadata->AddEntry(), 0.2 * (i + 1));
 		}
diff --git a/Engine/Plugins/PCG/Source/PCG/Private/Tests/Elements/PCGBlurElementTest.cpp b/Engine/Plugins/PCG/Source/PCG/Private/Tests/Elements/PCGBlurElementTest.cpp
--- a/Engine/Plugins/PCG/Source/PCG/Private/Tests/Elements/PCGBlurElementTest.cpp
+++ b/Engine/Plugins/PCG/Source/PCG/Private/Tests/Elements/PCGBlurElementTest.cpp
@@ -36,16 +36,23 @@ bool FPCGBlurElementTest_Density::RunTest(const FString& Parameters)
 	Settings->SearchDistance = SquareSize + 1.0; //+ 1 for approximation errors
 
 
-	UPCGPointData* InputPointData = PCGTestsCommon::CreateEmptyPointData();
-	TArray<FPCGPoint>& InputPoints = InputPointData->GetMutablePoints();
+	UPCGBasePointData* InputPointData = PCGTestsCommon::CreateEmptyBasePointData();
+	
+	InputPointData->SetNumPoints(4);
+	InputPointData->SetSeed(42);
 
-	for (int i = 0; i < 4; ++i)
+	TPCGValueRange<FTransform> TransformRange = InputPointData->GetTransformValueRange();
+	TPCGValueRange<float> DensityRange = InputPointData->GetDensityValueRange();
+
+	for (int i = 0; i < InputPointData->GetNumPoints(); ++i)
 	{
 		const int X = i / 2;
 		const int Y = i % 2;
 		const FVector Location = FVector(SquareSize * X, SquareSize * Y, 0.0);
 		const float Density = i == 0 ? 1.0f : 0.0f;
-		InputPoints.Emplace(FTransform(Location), Density, 42);
+		
+		TransformRange[i] = FTransform(Location);
+		DensityRange[i] = Density;
 	}
 
 	FPCGTaggedData& TaggedData = TestData.InputData.TaggedData.Emplace_GetRef(FPCGTaggedData());
@@ -61,18 +68,19 @@ bool FPCGBlurElementTest_Density::RunTest(const FString& Parameters)
 	TArray<FPCGTaggedData> Outputs = Context->OutputData.GetInputsByPin(PCGPinConstants::DefaultOutputLabel);
 	UTEST_EQUAL("There is one output", Outputs.Num(), 1);
 
-	const UPCGPointData* OutputPointData = Cast<UPCGPointData>(Outputs[0].Data);
+	const UPCGBasePointData* OutputPointData = Cast<UPCGBasePointData>(Outputs[0].Data);
 	UTEST_NOT_NULL("Output is a point data", OutputPointData);
 
 	check(OutputPointData);
 
-	const TArray<FPCGPoint>& OutputPoints = OutputPointData->GetPoints();
-	UTEST_EQUAL("Output has as many points as input", OutputPoints.Num(), InputPoints.Num());
+	UTEST_EQUAL("Output has as many points as input", OutputPointData->GetNumPoints(), InputPointData->GetNumPoints());
 
 	constexpr float DensityValues[] = { 0.333333f, 0.222222f, 0.222222f, 0.222222f };
+
+	const TConstPCGValueRange<float> OutDensityRange = OutputPointData->GetConstDensityValueRange();
 	for (int32 i = 0; i < 4; ++i)
 	{
-		UTEST_EQUAL(*FString::Printf(TEXT("Point %d has the expected density."), i), OutputPoints[i].Density, DensityValues[i]);
+		UTEST_EQUAL(*FString::Printf(TEXT("Point %d has the expected density."), i), OutDensityRange[i], DensityValues[i]);
 	}
 
 	return true;
diff --git a/Engine/Plugins/PCG/Source/PCG/Private/Tests/Elements/PCGCombinePointsTest.cpp b/Engine/Plugins/PCG/Source/PCG/Private/Tests/Elements/PCGCombinePointsTest.cpp
--- a/Engine/Plugins/PCG/Source/PCG/Private/Tests/Elements/PCGCombinePointsTest.cpp
+++ b/Engine/Plugins/PCG/Source/PCG/Private/Tests/Elements/PCGCombinePointsTest.cpp
@@ -20,17 +20,16 @@ namespace CombinePointsCommonTests
 
 		FPCGTaggedData& Inputs = TestData.InputData.TaggedData.Emplace_GetRef();
 		Inputs.Pin = PCGPinConstants::DefaultInputLabel;
-		UPCGPointData* InData = PCGTestsCommon::CreateRandomPointData(5, 42, false);
+		UPCGBasePointData* InData = PCGTestsCommon::CreateRandomBasePointData(5, 42, false);
 		Inputs.Data = InData;
 
 		// Setting the transform explicitly
-		UPCGPointData* InputPointData = Cast<UPCGPointData>(InData);
-		TArray<FPCGPoint>& InPoints = InputPointData->GetMutablePoints();
-		for (int i = 0; i < InPoints.Num(); ++i)
+		TPCGValueRange<FTransform> TransformRange = InData->GetTransformValueRange();
+		for (int i = 0; i < InData->GetNumPoints(); ++i)
 		{
-			InPoints[i].Transform.SetLocation(FVector(10.0, 10.0, 10.0) * i);
-			InPoints[i].Transform.SetRotation(FQuat::Identity);
-			InPoints[i].Transform.SetScale3D(FVector::One() * i + 1);
+			TransformRange[i].SetLocation(FVector(10.0, 10.0, 10.0) * i);
+			TransformRange[i].SetRotation(FQuat::Identity);
+			TransformRange[i].SetScale3D(FVector::One() * i + 1);
 		}
 
 		FPCGElementPtr TestElement = TestData.Settings->GetElement();
diff --git a/Engine/Plugins/PCG/Source/PCG/Private/Tests/Elements/PCGCopyPointsTest.cpp b/Engine/Plugins/PCG/Source/PCG/Private/Tests/Elements/PCGCopyPointsTest.cpp
--- a/Engine/Plugins/PCG/Source/PCG/Private/Tests/Elements/PCGCopyPointsTest.cpp
+++ b/Engine/Plugins/PCG/Source/PCG/Private/Tests/Elements/PCGCopyPointsTest.cpp
@@ -26,8 +26,8 @@ protected:
 		EPCGCopyPointsInheritanceMode ColorInheritance = EPCGCopyPointsInheritanceMode::Relative;
 		EPCGCopyPointsInheritanceMode SeedInheritance = EPCGCopyPointsInheritanceMode::Relative;
 		EPCGCopyPointsMetadataInheritanceMode AttributeInheritance = EPCGCopyPointsMetadataInheritanceMode::SourceFirst;
-		TObjectPtr<UPCGPointData> SourceData = nullptr;
-		TObjectPtr<UPCGPointData> TargetData = nullptr;
+		TObjectPtr<UPCGBasePointData> SourceData = nullptr;
+		TObjectPtr<UPCGBasePointData> TargetData = nullptr;
 		TArray<FPCGPoint> ExpectedPoints;
 	};
 
@@ -36,8 +36,8 @@ protected:
 
 bool FPCGCopyPointsTestBase::RunTestInternal(FParameters& Parameters)
 {
-	UPCGPointData* SourceData = Parameters.SourceData;
-	UPCGPointData* TargetData = Parameters.TargetData;
+	UPCGBasePointData* SourceData = Parameters.SourceData;
+	UPCGBasePointData* TargetData = Parameters.TargetData;
 
 	check(SourceData && TargetData);
 
@@ -64,20 +64,19 @@ bool FPCGCopyPointsTestBase::RunTestInternal(FParameters& Parameters)
 	TargetData->Metadata->CreateFloatAttribute(LengthName, 1, bAllowsInterpolation, bOverrideParent);
 	TargetData->Metadata->CreateFloatAttribute(HeightName, 1, bAllowsInterpolation, bOverrideParent);
 
-	TArray<FPCGPoint>& SourcePoints = SourceData->GetMutablePoints();
-	TArray<FPCGPoint>& TargetPoints = TargetData->GetMutablePoints();
-
+	TPCGValueRange<int64> SourceMetadataEntryRange = SourceData->GetMetadataEntryValueRange();
 	FRandomStream RandomSource(TestData.Seed);
-	for (int I = 0; I < SourcePoints.Num(); ++I)
+	for (int I = 0; I < SourceData->GetNumPoints(); ++I)
 	{
 		// first half of the points have a non-default Length attribute value, the rest have a non-default width attribute value
-	 	UPCGMetadataAccessorHelpers::SetFloatAttribute(SourcePoints[I], SourceData->Metadata, (I < SourcePoints.Num() / 2) ? LengthName : WidthName, RandomSource.FRand());
+	 	UPCGMetadataAccessorHelpers::SetFloatAttributeByMetadataKey(SourceMetadataEntryRange[I], SourceData->Metadata, (I < SourceData->GetNumPoints() / 2) ? LengthName : WidthName, RandomSource.FRand());
 	}
 
-	for (int I = 0; I < TargetPoints.Num(); ++I)
+	TPCGValueRange<int64> TargetMetadataEntryRange = TargetData->GetMetadataEntryValueRange();
+	for (int I = 0; I < TargetData->GetNumPoints(); ++I)
 	{
 		// first quarter of the points have a non-default Length attribute value, the rest have a non-default Height attribute value
-		UPCGMetadataAccessorHelpers::SetFloatAttribute(TargetPoints[I], TargetData->Metadata, (I < SourcePoints.Num() / 4) ? LengthName : HeightName, RandomSource.FRand());
+		UPCGMetadataAccessorHelpers::SetFloatAttributeByMetadataKey(TargetMetadataEntryRange[I], TargetData->Metadata, (I < SourceData->GetNumPoints() / 4) ? LengthName : HeightName, RandomSource.FRand());
 	}
 
 	FPCGTaggedData& SourceTaggedData = TestData.InputData.TaggedData.Emplace_GetRef(FPCGTaggedData());
@@ -106,7 +105,7 @@ bool FPCGCopyPointsTestBase::RunTestInternal(FParameters& Parameters)
 	UTEST_NOT_NULL("Valid output point data", OutPointData);
 	check(OutPointData);
 
-	UTEST_EQUAL("Valid number of output points", SourcePoints.Num() * TargetPoints.Num(), OutPointData->GetNumPoints());
+	UTEST_EQUAL("Valid number of output points", SourceData->GetNumPoints() * TargetData->GetNumPoints(), OutPointData->GetNumPoints());
 	check(Parameters.ExpectedPoints.IsEmpty() || Parameters.ExpectedPoints.Num() == OutPointData->GetNumPoints());
 
 	const UPCGMetadata* RootMetadata = nullptr;
@@ -185,11 +184,13 @@ bool FPCGCopyPointsTestBase::RunTestInternal(FParameters& Parameters)
 	check(InheritedAttributeModes.Num() == InheritedAttributeModes.Num() && InheritedAttributeModes.Num() == OutAttributes.Num())
 
 	const FConstPCGPointValueRanges OutRanges(OutPointData);
+	const FConstPCGPointValueRanges SourceRanges(SourceData);
+	const FConstPCGPointValueRanges TargetRanges(TargetData);
 
 	for (int PointIndex = 0; PointIndex < OutPointData->GetNumPoints(); ++PointIndex)
 	{
-		const FPCGPoint& SourcePoint = SourcePoints[PointIndex / TargetPoints.Num()];
-		const FPCGPoint& TargetPoint = TargetPoints[PointIndex % TargetPoints.Num()];
+		const FPCGPoint SourcePoint = SourceRanges.GetPoint(PointIndex / TargetData->GetNumPoints());
+		const FPCGPoint TargetPoint = TargetRanges.GetPoint(PointIndex % TargetData->GetNumPoints());
 
 		auto FormatMessage = [PointIndex](const FString& Message)
 		{
@@ -301,22 +302,32 @@ bool FPCGCopyPointsTest_Relative::RunTest(const FString&)
 	Parameters.SeedInheritance = EPCGCopyPointsInheritanceMode::Relative;
 
 	// For relative, we hard code the results
-	UPCGPointData* SourceData = NewObject<UPCGPointData>();
-	UPCGPointData* TargetData = NewObject<UPCGPointData>();
-	TArray<FPCGPoint>& SourcePoints = SourceData->GetMutablePoints();
-	TArray<FPCGPoint>& TargetPoints = TargetData->GetMutablePoints();
+	UPCGBasePointData* SourceData = PCGTestsCommon::CreateEmptyBasePointData();
+	UPCGBasePointData* TargetData = PCGTestsCommon::CreateEmptyBasePointData();
 	
-	SourcePoints.Emplace(FTransform(FQuat::Identity, FVector(3.0, 4.0, 5.0), FVector(1.0, 1.0, 1.0)), 1.0f, 42);
-	SourcePoints.Emplace(FTransform(FQuat::MakeFromEuler(FVector(45.0, 45.0, 45.0)), FVector(2.0, 1.0, -1.0), FVector(1.0, 1.0, 1.0)), 1.0f, 43);
-	SourcePoints.Emplace(FTransform(FQuat::MakeFromEuler(FVector(46.0, 47.0, 48.0)), FVector(6.0, 2.0, 3.0), FVector(-1.0, 1.0, -1.0)), 1.0f, 44);
-	SourcePoints[0].Color = FVector4(0.1, 0.5, 0.2, 0.3);
-	SourcePoints[1].Color = FVector4(0.4, 0.1, 0.1, 0.6);
-	SourcePoints[2].Color = FVector4(0.7, 0.4, 0.9, 0.9);
-
-	TargetPoints.Emplace(FTransform(FQuat::MakeFromEuler(FVector(20.0, 21.0, 22.0)), FVector(-6.0, -4.0, -1.0), FVector(1.0, 1.0, 1.0)), 1.0f, 45);
-	TargetPoints.Emplace(FTransform(FQuat::MakeFromEuler(FVector(22.0, 23.0, 24.0)), FVector(-7.0, -5.0, -10.0), FVector(-1.0, 1.0, -1.0)), 1.0f, 46);
-	TargetPoints[0].Color = FVector4(0.4, 0.3, 1.0, 1.0);
-	TargetPoints[1].Color = FVector4(0.1, 0.2, 0.7, 0.4);
+	SourceData->SetNumPoints(3);
+	SourceData->AllocateProperties(EPCGPointNativeProperties::Transform | EPCGPointNativeProperties::Seed | EPCGPointNativeProperties::Color);
+	SourceData->SetDensity(1.0f);
+
+	FPCGPointValueRanges SourceRanges(SourceData, /*bAllocate=*/false);
+	SourceRanges.SetFromPoint(0, FPCGPoint(FTransform(FQuat::Identity, FVector(3.0, 4.0, 5.0), FVector(1.0, 1.0, 1.0)), 1.0f, 42));
+	SourceRanges.SetFromPoint(1, FPCGPoint(FTransform(FQuat::MakeFromEuler(FVector(45.0, 45.0, 45.0)), FVector(2.0, 1.0, -1.0), FVector(1.0, 1.0, 1.0)), 1.0f, 43));
+	SourceRanges.SetFromPoint(2, FPCGPoint(FTransform(FQuat::MakeFromEuler(FVector(46.0, 47.0, 48.0)), FVector(6.0, 2.0, 3.0), FVector(-1.0, 1.0, -1.0)), 1.0f, 44));
+		
+	SourceRanges.ColorRange[0] = FVector4(0.1, 0.5, 0.2, 0.3);
+	SourceRanges.ColorRange[1] = FVector4(0.4, 0.1, 0.1, 0.6);
+	SourceRanges.ColorRange[2] = FVector4(0.7, 0.4, 0.9, 0.9);
+
+	TargetData->SetNumPoints(2);
+	TargetData->AllocateProperties(EPCGPointNativeProperties::Transform | EPCGPointNativeProperties::Seed | EPCGPointNativeProperties::Color);
+	TargetData->SetDensity(1.0f);
+
+	FPCGPointValueRanges TargetRanges(TargetData, /*bAllocate=*/false);
+	TargetRanges.SetFromPoint(0, FPCGPoint(FTransform(FQuat::MakeFromEuler(FVector(20.0, 21.0, 22.0)), FVector(-6.0, -4.0, -1.0), FVector(1.0, 1.0, 1.0)), 1.0f, 45));
+	TargetRanges.SetFromPoint(1, FPCGPoint(FTransform(FQuat::MakeFromEuler(FVector(22.0, 23.0, 24.0)), FVector(-7.0, -5.0, -10.0), FVector(-1.0, 1.0, -1.0)), 1.0f, 46));
+
+	TargetRanges.ColorRange[0] = FVector4(0.4, 0.3, 1.0, 1.0);
+	TargetRanges.ColorRange[1] = FVector4(0.1, 0.2, 0.7, 0.4);
 
 	// Hard coded values for the values above.
 	Parameters.ExpectedPoints.Emplace(FTransform(FQuat(-0.133360, -0.208748, 0.153700, 0.956564), FVector(-6.558469, 1.672731, 3.184284), FVector(1.000000, 1.000000, 1.000000)), 1.0f, -21834409);
diff --git a/Engine/Plugins/PCG/Source/PCG/Private/Tests/Elements/PCGCreatePointsTest.cpp b/Engine/Plugins/PCG/Source/PCG/Private/Tests/Elements/PCGCreatePointsTest.cpp
--- a/Engine/Plugins/PCG/Source/PCG/Private/Tests/Elements/PCGCreatePointsTest.cpp
+++ b/Engine/Plugins/PCG/Source/PCG/Private/Tests/Elements/PCGCreatePointsTest.cpp
@@ -15,14 +15,23 @@ bool FPCGCreatePointsTest_Basic::RunTest(const FString& Parameters)
 	PCGTestsCommon::GenerateSettings<UPCGCreatePointsSettings>(TestData);
 	UPCGCreatePointsSettings* Settings = CastChecked<UPCGCreatePointsSettings>(TestData.Settings);
 	
-	TObjectPtr<UPCGPointData> PointData = PCGTestsCommon::CreateRandomPointData(100, 42, false);
+	TObjectPtr<UPCGBasePointData> PointData = PCGTestsCommon::CreateRandomBasePointData(100, 42, false);
 
 	// Setting the seed explicitly, first point will have a seed of 0, which mean it going to be changed and computed depending on the point position
 	// All the others should stay the same.
-	TArray<FPCGPoint>& Points = PointData->GetMutablePoints();
-	for (int32 i = 0; i < Points.Num(); ++i)
+	TPCGValueRange<int32> SeedRange = PointData->GetSeedValueRange();
+	for (int32 i = 0; i < PointData->GetNumPoints(); ++i)
 	{
-		Points[i].Seed = i;
+		SeedRange[i] = i;
+	}
+
+	const FConstPCGPointValueRanges ValueRange(PointData);
+	TArray<FPCGPoint> Points;
+	Points.Reserve(PointData->GetNumPoints());
+
+	for (int32 i = 0; i < PointData->GetNumPoints(); ++i)
+	{
+		Points.Add(ValueRange.GetPoint(i));
 	}
 
 	Settings->PointsToCreate = Points;
diff --git a/Engine/Plugins/PCG/Source/PCG/Private/Tests/Elements/PCGDensityRemapTest.cpp b/Engine/Plugins/PCG/Source/PCG/Private/Tests/Elements/PCGDensityRemapTest.cpp
--- a/Engine/Plugins/PCG/Source/PCG/Private/Tests/Elements/PCGDensityRemapTest.cpp
+++ b/Engine/Plugins/PCG/Source/PCG/Private/Tests/Elements/PCGDensityRemapTest.cpp
@@ -21,15 +21,20 @@ bool FPCGDensityRemapTest::RunTest(const FString& Parameters)
 	Settings->InputSource.SetPointProperty(EPCGPointProperties::Density);
 	Settings->bClampToUnitRange = true;
 
-	TObjectPtr<UPCGPointData> PointData = PCGTestsCommon::CreateEmptyPointData();
-	TArray<FPCGPoint>& Points = PointData->GetMutablePoints();
+	TObjectPtr<UPCGBasePointData> PointData = PCGTestsCommon::CreateEmptyBasePointData();
 
 	FRandomStream RandomSource(TestData.Seed);
 	const int PointCount = 6;
+
+	PointData->SetNumPoints(PointCount);
+	
+	TPCGValueRange<int32> SeedRange = PointData->GetSeedValueRange();
+	TPCGValueRange<float> DensityRange = PointData->GetDensityValueRange();
+
 	for (int I = 0; I < PointCount; ++I)
 	{
-		FPCGPoint& Point = Points.Emplace_GetRef(FTransform(), 1, I);
-		Point.Density = static_cast<float>(I) / (PointCount - 1);
+		SeedRange[I] = I;
+		DensityRange[I] = static_cast<float>(I) / (PointCount - 1);
 	}
 
 	FPCGTaggedData& TaggedData = TestData.InputData.TaggedData.Emplace_GetRef(FPCGTaggedData());
@@ -61,7 +66,7 @@ bool FPCGDensityRemapTest::RunTest(const FString& Parameters)
 			const UPCGSpatialData* InSpatialData = Cast<UPCGSpatialData>(Input.Data);
 			check(InSpatialData);
 
-			const UPCGPointData* InPointData = InSpatialData->ToPointData(Context.Get());
+			const UPCGBasePointData* InPointData = InSpatialData->ToBasePointData(Context.Get());
 			check(InPointData);
 
 			const UPCGSpatialData* OutSpatialData = Cast<UPCGSpatialData>(Output.Data);
@@ -72,7 +77,7 @@ bool FPCGDensityRemapTest::RunTest(const FString& Parameters)
 				continue;
 			}
 
-			const UPCGPointData* OutPointData = OutSpatialData->ToPointData(Context.Get());
+			const UPCGBasePointData* OutPointData = OutSpatialData->ToBasePointData(Context.Get());
 
 			if (!TestNotNull("Valid output PointData", OutPointData))
 			{
@@ -80,18 +85,16 @@ bool FPCGDensityRemapTest::RunTest(const FString& Parameters)
 				continue;
 			}
 
-			const TArray<FPCGPoint>& InPoints = InPointData->GetPoints();
-			const TArray<FPCGPoint>& OutPoints = OutPointData->GetPoints();
-
-			if (!TestEqual("Input and output point counts match", OutPoints.Num(), InPoints.Num()))
+			if (!TestEqual("Input and output point counts match", OutPointData->GetNumPoints(), InPointData->GetNumPoints()))
 			{ 
 				bTestPassed = false;
 				continue;
 			}
 
-			for (int PointIndex = 0; PointIndex < InPoints.Num(); ++PointIndex)
+			const TConstPCGValueRange<float> DensityRange = OutPointData->GetConstDensityValueRange();
+			for (int PointIndex = 0; PointIndex < InPointData->GetNumPoints(); ++PointIndex)
 			{
-				bTestPassed &= TestEqual("Correct density", OutPoints[PointIndex].Density, CorrectDensities[PointIndex]);
+				bTestPassed &= TestEqual("Correct density", DensityRange[PointIndex], CorrectDensities[PointIndex]);
 			}
 		}
 
diff --git a/Engine/Plugins/PCG/Source/PCG/Private/Tests/Elements/PCGDistanceTest.cpp b/Engine/Plugins/PCG/Source/PCG/Private/Tests/Elements/PCGDistanceTest.cpp
--- a/Engine/Plugins/PCG/Source/PCG/Private/Tests/Elements/PCGDistanceTest.cpp
+++ b/Engine/Plugins/PCG/Source/PCG/Private/Tests/Elements/PCGDistanceTest.cpp
@@ -22,11 +22,11 @@ bool FPCGDistanceTest_PointToPoint::RunTest(const FString& Parameters)
 	FPCGElementPtr TestElement = TestData.Settings->GetElement();
 
 	{
-		UPCGPointData *SourceData = PCGTestsCommon::CreateRandomPointData(2, 42);
+		UPCGBasePointData* SourceData = PCGTestsCommon::CreateRandomBasePointData(2, 42);
 
-		TArray<FPCGPoint>& SourcePoints = SourceData->GetMutablePoints();
-		SourcePoints[0].Transform.SetTranslation(FVector(100, 0, 0));
-		SourcePoints[1].Transform.SetTranslation(FVector(0, 50, 0));
+		TPCGValueRange<FTransform> TransformRange = SourceData->GetTransformValueRange();
+		TransformRange[0].SetTranslation(FVector(100, 0, 0));
+		TransformRange[1].SetTranslation(FVector(0, 50, 0));
 
 		FPCGTaggedData& SourcePin = TestData.InputData.TaggedData.Emplace_GetRef();
 		SourcePin.Pin = PCGDistance::SourceLabel;
@@ -36,7 +36,7 @@ bool FPCGDistanceTest_PointToPoint::RunTest(const FString& Parameters)
 	{
 		FPCGTaggedData& TargetPin = TestData.InputData.TaggedData.Emplace_GetRef();
 		TargetPin.Pin = PCGDistance::TargetLabel;
-		TargetPin.Data = PCGTestsCommon::CreatePointData(FVector(0, 0, 0));		
+		TargetPin.Data = PCGTestsCommon::CreateBasePointData(FVector(0, 0, 0));		
 	}
 
 	TUniquePtr<FPCGContext> Context = TestData.InitializeTestContext();
@@ -80,11 +80,11 @@ bool FPCGDistanceTest_SetDensity::RunTest(const FString& Parameters)
 	FPCGElementPtr TestElement = TestData.Settings->GetElement();
 
 	{
-		UPCGPointData *SourceData = PCGTestsCommon::CreateRandomPointData(2, 42);
+		UPCGBasePointData* SourceData = PCGTestsCommon::CreateRandomBasePointData(2, 42);
 
-		TArray<FPCGPoint>& SourcePoints = SourceData->GetMutablePoints();
-		SourcePoints[0].Transform.SetTranslation(FVector(100, 0, 0));
-		SourcePoints[1].Transform.SetTranslation(FVector(0, 50, 0));
+		TPCGValueRange<FTransform> TransformRange = SourceData->GetTransformValueRange();
+		TransformRange[0].SetTranslation(FVector(100, 0, 0));
+		TransformRange[1].SetTranslation(FVector(0, 50, 0));
 
 		FPCGTaggedData& SourcePin = TestData.InputData.TaggedData.Emplace_GetRef();
 		SourcePin.Pin = PCGDistance::SourceLabel;
@@ -94,7 +94,7 @@ bool FPCGDistanceTest_SetDensity::RunTest(const FString& Parameters)
 	{
 		FPCGTaggedData& TargetPin = TestData.InputData.TaggedData.Emplace_GetRef();
 		TargetPin.Pin = PCGDistance::TargetLabel;
-		TargetPin.Data = PCGTestsCommon::CreatePointData(FVector(0, 0, 0));		
+		TargetPin.Data = PCGTestsCommon::CreateBasePointData(FVector(0, 0, 0));		
 	}
 
 	TUniquePtr<FPCGContext> Context = TestData.InitializeTestContext();
@@ -132,11 +132,11 @@ bool FPCGDistanceTest_PointToSphere::RunTest(const FString& Parameters)
 	FPCGElementPtr TestElement = TestData.Settings->GetElement();
 
 	{
-		UPCGPointData *SourceData = PCGTestsCommon::CreateRandomPointData(2, 42);
+		UPCGBasePointData* SourceData = PCGTestsCommon::CreateRandomBasePointData(2, 42);
 
-		TArray<FPCGPoint>& SourcePoints = SourceData->GetMutablePoints();
-		SourcePoints[0].Transform.SetTranslation(FVector(100, 0, 0));
-		SourcePoints[1].Transform.SetTranslation(FVector(0, 50, 0));
+		TPCGValueRange<FTransform> TransformRange = SourceData->GetTransformValueRange();
+		TransformRange[0].SetTranslation(FVector(100, 0, 0));
+		TransformRange[1].SetTranslation(FVector(0, 50, 0));
 
 		FPCGTaggedData& SourcePin = TestData.InputData.TaggedData.Emplace_GetRef();
 		SourcePin.Pin = PCGDistance::SourceLabel;
@@ -150,9 +150,11 @@ bool FPCGDistanceTest_PointToSphere::RunTest(const FString& Parameters)
 		FPCGTaggedData& TargetPin = TestData.InputData.TaggedData.Emplace_GetRef();
 		TargetPin.Pin = PCGDistance::TargetLabel;
 
-		UPCGPointData *TargetData = PCGTestsCommon::CreatePointData(FVector(0, 0, 0));
-		TArray<FPCGPoint>& TargetPoints = TargetData->GetMutablePoints();
-		TargetPoints[0].SetExtents(TargetExtents);
+		UPCGBasePointData *TargetData = PCGTestsCommon::CreateBasePointData(FVector(0, 0, 0));
+		TargetData->AllocateProperties(EPCGPointNativeProperties::BoundsMin | EPCGPointNativeProperties::BoundsMax);
+		FPCGPointValueRanges ValueRanges(TargetData, /*bAllocate=*/false);
+
+		PCGPointHelpers::SetExtents(TargetExtents, ValueRanges.BoundsMinRange[0], ValueRanges.BoundsMaxRange[0]);
 
 		TargetPin.Data = TargetData;
 	}
@@ -195,11 +197,12 @@ bool FPCGDistanceTest_PointToBox::RunTest(const FString& Parameters)
 	FPCGElementPtr TestElement = TestData.Settings->GetElement();
 
 	{
-		UPCGPointData *SourceData = PCGTestsCommon::CreateRandomPointData(2, 42);
+		UPCGBasePointData *SourceData = PCGTestsCommon::CreateRandomBasePointData(2, 42);
+
+		TPCGValueRange<FTransform> TransformRange = SourceData->GetTransformValueRange(/*bAllocate=*/false);
 
-		TArray<FPCGPoint>& SourcePoints = SourceData->GetMutablePoints();
-		SourcePoints[0].Transform.SetTranslation(FVector(100, 0, 0));
-		SourcePoints[1].Transform.SetTranslation(FVector(0, 50, 0));
+		TransformRange[0].SetTranslation(FVector(100, 0, 0));
+		TransformRange[1].SetTranslation(FVector(0, 50, 0));
 
 		FPCGTaggedData& SourcePin = TestData.InputData.TaggedData.Emplace_GetRef();
 		SourcePin.Pin = PCGDistance::SourceLabel;
@@ -213,9 +216,11 @@ bool FPCGDistanceTest_PointToBox::RunTest(const FString& Parameters)
 		FPCGTaggedData& TargetPin = TestData.InputData.TaggedData.Emplace_GetRef();
 		TargetPin.Pin = PCGDistance::TargetLabel;
 
-		UPCGPointData *TargetData = PCGTestsCommon::CreatePointData(FVector(0, 0, 0));
-		TArray<FPCGPoint>& TargetPoints = TargetData->GetMutablePoints();
-		TargetPoints[0].SetExtents(TargetExtents);
+		UPCGBasePointData *TargetData = PCGTestsCommon::CreateBasePointData(FVector(0, 0, 0));
+		TargetData->AllocateProperties(EPCGPointNativeProperties::BoundsMin | EPCGPointNativeProperties::BoundsMax);
+		FPCGPointValueRanges ValueRanges(TargetData, /*bAllocate=*/false);
+
+		PCGPointHelpers::SetExtents(TargetExtents, ValueRanges.BoundsMinRange[0], ValueRanges.BoundsMaxRange[0]);
 
 		TargetPin.Data = TargetData;
 	}
diff --git a/Engine/Plugins/PCG/Source/PCG/Private/Tests/Elements/PCGDuplicatePointTest.cpp b/Engine/Plugins/PCG/Source/PCG/Private/Tests/Elements/PCGDuplicatePointTest.cpp
--- a/Engine/Plugins/PCG/Source/PCG/Private/Tests/Elements/PCGDuplicatePointTest.cpp
+++ b/Engine/Plugins/PCG/Source/PCG/Private/Tests/Elements/PCGDuplicatePointTest.cpp
@@ -38,7 +38,7 @@ protected:
 
 		FPCGTaggedData& Inputs = TestData.InputData.TaggedData.Emplace_GetRef();
 		Inputs.Pin = PCGPinConstants::DefaultInputLabel;
-		Inputs.Data = PCGTestsCommon::CreatePointData();
+		Inputs.Data = PCGTestsCommon::CreateBasePointData();
 
 		FPCGElementPtr TestElement = TestData.Settings->GetElement();
 
diff --git a/Engine/Plugins/PCG/Source/PCG/Private/Tests/Elements/PCGFilterByIndexTest.cpp b/Engine/Plugins/PCG/Source/PCG/Private/Tests/Elements/PCGFilterByIndexTest.cpp
--- a/Engine/Plugins/PCG/Source/PCG/Private/Tests/Elements/PCGFilterByIndexTest.cpp
+++ b/Engine/Plugins/PCG/Source/PCG/Private/Tests/Elements/PCGFilterByIndexTest.cpp
@@ -23,7 +23,7 @@ namespace PCGFilterByIndexTestPrivate
 		{
 			FPCGTaggedData& InputData = TestData.InputData.TaggedData.Emplace_GetRef();
 			InputData.Pin = PCGPinConstants::DefaultInputLabel;
-			InputData.Data = PCGTestsCommon::CreateRandomPointData(Index + 1, 42);
+			InputData.Data = PCGTestsCommon::CreateRandomBasePointData(Index + 1, 42);
 		}
 
 		TUniquePtr<FPCGContext> Context = TestData.InitializeTestContext();
@@ -41,8 +41,8 @@ namespace PCGFilterByIndexTestPrivate
 
 		for (int I = 0; I < PCGFilterByIndexTestConstants::InputNum; ++I)
 		{
-			const UPCGPointData* OutPointData = Cast<UPCGPointData>(Outputs[I].Data);
-			if (!OutPointData || (OutPointData->GetPoints().Num() != I + 1))
+			const UPCGBasePointData* OutPointData = Cast<UPCGBasePointData>(Outputs[I].Data);
+			if (!OutPointData || (OutPointData->GetNumPoints() != I + 1))
 			{
 				return false;
 			}
diff --git a/Engine/Plugins/PCG/Source/PCG/Private/Tests/Elements/PCGGatherTest.cpp b/Engine/Plugins/PCG/Source/PCG/Private/Tests/Elements/PCGGatherTest.cpp
--- a/Engine/Plugins/PCG/Source/PCG/Private/Tests/Elements/PCGGatherTest.cpp
+++ b/Engine/Plugins/PCG/Source/PCG/Private/Tests/Elements/PCGGatherTest.cpp
@@ -20,7 +20,7 @@ bool FPCGGatherTest_Basic::RunTest(const FString& Parameters)
 	{
 		FPCGTaggedData& SourcePin = TestData.InputData.TaggedData.Emplace_GetRef();
 		SourcePin.Pin = PCGPinConstants::DefaultInputLabel;
-		SourcePin.Data = PCGTestsCommon::CreateRandomPointData(100, 42);
+		SourcePin.Data = PCGTestsCommon::CreateRandomBasePointData(100, 42);
 	}
 
 	TUniquePtr<FPCGContext> Context = TestData.InitializeTestContext();
@@ -32,13 +32,11 @@ bool FPCGGatherTest_Basic::RunTest(const FString& Parameters)
 
 	UTEST_EQUAL("Output count", Outputs.Num(), 1);
 
-	const UPCGPointData *OutPointData = Cast<UPCGPointData>(Outputs[0].Data);
+	const UPCGBasePointData *OutPointData = Cast<UPCGBasePointData>(Outputs[0].Data);
 
 	UTEST_NOT_NULL("Output point data", OutPointData);
 
-	const TArray<FPCGPoint>& OutPoints = OutPointData->GetPoints();
-
-	UTEST_EQUAL("Output point count", OutPoints.Num(), 100);
+	UTEST_EQUAL("Output point count", OutPointData->GetNumPoints(), 100);
 
 	return true;
 }
diff --git a/Engine/Plugins/PCG/Source/PCG/Private/Tests/Elements/PCGInstancePackerTest.cpp b/Engine/Plugins/PCG/Source/PCG/Private/Tests/Elements/PCGInstancePackerTest.cpp
--- a/Engine/Plugins/PCG/Source/PCG/Private/Tests/Elements/PCGInstancePackerTest.cpp
+++ b/Engine/Plugins/PCG/Source/PCG/Private/Tests/Elements/PCGInstancePackerTest.cpp
@@ -106,13 +106,13 @@ bool FPCGStaticMeshSpawnerInstancePackerByAttributeTest::RunTest(const FString&
 	check(Settings);
 
 	FPCGTaggedData& SourceTaggedData = TestData.InputData.TaggedData.Emplace_GetRef(FPCGTaggedData());
-	TObjectPtr<UPCGPointData> SourceData = PCGTestsCommon::CreateRandomPointData(5, TestData.Seed);
+	TObjectPtr<UPCGBasePointData> SourceData = PCGTestsCommon::CreateRandomBasePointData(5, TestData.Seed);
 
 	// TODO - remove this code when the tests support flipped instances (generates multiple ISMC components)
-	TArray<FPCGPoint>& Points = SourceData->GetMutablePoints();
-	for (FPCGPoint& Point : Points)
+	TPCGValueRange<FTransform> TransformRange = SourceData->GetTransformValueRange();
+	for (FTransform& PointTransform : TransformRange)
 	{
-		Point.Transform.SetScale3D(Point.Transform.GetScale3D().GetAbs());
+		PointTransform.SetScale3D(PointTransform.GetScale3D().GetAbs());
 	}
 
 	SourceTaggedData.Data = SourceData;
@@ -148,8 +148,6 @@ bool FPCGStaticMeshSpawnerInstancePackerByAttributeTest::RunTest(const FString&
 	SourceData->Metadata->CreateRotatorAttribute(RotatorName, FRotator::ZeroRotator, bAllowsInterpolation, bOverrideParent);
 	SourceData->Metadata->CreateQuatAttribute(QuatName, FQuat::Identity, bAllowsInterpolation, bOverrideParent);
 
-	TArray<FPCGPoint>& SourcePoints = SourceData->GetMutablePoints();
-
 	TArray<float> BoolValues;
 	TArray<float> FloatValues;
 	TArray<float> DoubleValues;
@@ -159,6 +157,8 @@ bool FPCGStaticMeshSpawnerInstancePackerByAttributeTest::RunTest(const FString&
 	TArray<float> RotatorValues;
 	TArray<float> QuatValues;
 
+	TPCGValueRange<int64> MetadataEntryRange = SourceData->GetMetadataEntryValueRange();
+
 	FRandomStream RandomSource(TestData.Seed);
 	int NumPoints = 5;
 	for (int I = 0; I < NumPoints; ++I)
@@ -191,14 +191,14 @@ bool FPCGStaticMeshSpawnerInstancePackerByAttributeTest::RunTest(const FString&
 		QuatValues.Add(QuatValue.Z);
 		QuatValues.Add(QuatValue.W);
 
-		UPCGMetadataAccessorHelpers::SetBoolAttribute(SourcePoints[I], SourceData->Metadata, BoolName, BoolValue);
-		UPCGMetadataAccessorHelpers::SetFloatAttribute(SourcePoints[I], SourceData->Metadata, FloatName, FloatValue);
-		UPCGMetadataAccessorHelpers::SetDoubleAttribute(SourcePoints[I], SourceData->Metadata, DoubleName, DoubleValue);
-		UPCGMetadataAccessorHelpers::SetInteger64Attribute(SourcePoints[I], SourceData->Metadata, IntName, IntValue);
-		UPCGMetadataAccessorHelpers::SetVectorAttribute(SourcePoints[I], SourceData->Metadata, VecName, VecValue);
-		UPCGMetadataAccessorHelpers::SetVector4Attribute(SourcePoints[I], SourceData->Metadata, Vec4Name, Vec4Value);
-		UPCGMetadataAccessorHelpers::SetRotatorAttribute(SourcePoints[I], SourceData->Metadata, RotatorName, RotatorValue);
-		UPCGMetadataAccessorHelpers::SetQuatAttribute(SourcePoints[I], SourceData->Metadata, QuatName, QuatValue);
+		UPCGMetadataAccessorHelpers::SetBoolAttributeByMetadataKey(MetadataEntryRange[I], SourceData->Metadata, BoolName, BoolValue);
+		UPCGMetadataAccessorHelpers::SetFloatAttributeByMetadataKey(MetadataEntryRange[I], SourceData->Metadata, FloatName, FloatValue);
+		UPCGMetadataAccessorHelpers::SetDoubleAttributeByMetadataKey(MetadataEntryRange[I], SourceData->Metadata, DoubleName, DoubleValue);
+		UPCGMetadataAccessorHelpers::SetInteger64AttributeByMetadataKey(MetadataEntryRange[I], SourceData->Metadata, IntName, IntValue);
+		UPCGMetadataAccessorHelpers::SetVectorAttributeByMetadataKey(MetadataEntryRange[I], SourceData->Metadata, VecName, VecValue);
+		UPCGMetadataAccessorHelpers::SetVector4AttributeByMetadataKey(MetadataEntryRange[I], SourceData->Metadata, Vec4Name, Vec4Value);
+		UPCGMetadataAccessorHelpers::SetRotatorAttributeByMetadataKey(MetadataEntryRange[I], SourceData->Metadata, RotatorName, RotatorValue);
+		UPCGMetadataAccessorHelpers::SetQuatAttributeByMetadataKey(MetadataEntryRange[I], SourceData->Metadata, QuatName, QuatValue);
 	}
 
 	Settings->SetInstancePackerType(UPCGInstanceDataPackerByAttribute::StaticClass());
@@ -307,13 +307,13 @@ bool FPCGStaticMeshSpawnerInstancePackerByRegexTest::RunTest(const FString& Para
 	check(Settings);
 
 	FPCGTaggedData& SourceTaggedData = TestData.InputData.TaggedData.Emplace_GetRef(FPCGTaggedData());
-	TObjectPtr<UPCGPointData> SourceData = PCGTestsCommon::CreateRandomPointData(5, TestData.Seed);
+	TObjectPtr<UPCGBasePointData> SourceData = PCGTestsCommon::CreateRandomBasePointData(5, TestData.Seed);
 
 	// TODO - remove this code when the tests support flipped instances (generates multiple ISMC components)
-	TArray<FPCGPoint>& Points = SourceData->GetMutablePoints();
-	for (FPCGPoint& Point : Points)
+	TPCGValueRange<FTransform> TransformRange = SourceData->GetTransformValueRange();
+	for (FTransform& PointTransform : TransformRange)
 	{
-		Point.Transform.SetScale3D(Point.Transform.GetScale3D().GetAbs());
+		PointTransform.SetScale3D(PointTransform.GetScale3D().GetAbs());
 	}
 
 	SourceTaggedData.Data = SourceData;
@@ -334,14 +334,14 @@ bool FPCGStaticMeshSpawnerInstancePackerByRegexTest::RunTest(const FString& Para
 	SourceData->Metadata->CreateVectorAttribute(NameB1, FVector::Zero(), bAllowsInterpolation, bOverrideParent);
 	SourceData->Metadata->CreateVector4Attribute(NameB2, FVector4::Zero(), bAllowsInterpolation, bOverrideParent);
 
-	TArray<FPCGPoint>& SourcePoints = SourceData->GetMutablePoints();
-
 	TArray<float> FloatValues;
 	TArray<float> DoubleValues;
 	TArray<float> IntValues;
 	TArray<float> VecValues;
 	TArray<float> Vec4Values;
 
+	TPCGValueRange<int64> MetadataEntryRange = SourceData->GetMetadataEntryValueRange();
+
 	FRandomStream RandomSource(TestData.Seed);
 	int NumPoints = 5;
 	for (int I = 0; I < NumPoints; ++I)
@@ -363,11 +363,11 @@ bool FPCGStaticMeshSpawnerInstancePackerByRegexTest::RunTest(const FString& Para
 		Vec4Values.Add(Vec4Value.Z);
 		Vec4Values.Add(Vec4Value.W);
 
-		UPCGMetadataAccessorHelpers::SetFloatAttribute(SourcePoints[I], SourceData->Metadata, NameA1, FloatValue);
-		UPCGMetadataAccessorHelpers::SetDoubleAttribute(SourcePoints[I], SourceData->Metadata, NameA2, DoubleValue);
-		UPCGMetadataAccessorHelpers::SetInteger64Attribute(SourcePoints[I], SourceData->Metadata, NameA3, IntValue);
-		UPCGMetadataAccessorHelpers::SetVectorAttribute(SourcePoints[I], SourceData->Metadata, NameB1, VecValue);
-		UPCGMetadataAccessorHelpers::SetVector4Attribute(SourcePoints[I], SourceData->Metadata, NameB2, Vec4Value);
+		UPCGMetadataAccessorHelpers::SetFloatAttributeByMetadataKey(MetadataEntryRange[I], SourceData->Metadata, NameA1, FloatValue);
+		UPCGMetadataAccessorHelpers::SetDoubleAttributeByMetadataKey(MetadataEntryRange[I], SourceData->Metadata, NameA2, DoubleValue);
+		UPCGMetadataAccessorHelpers::SetInteger64AttributeByMetadataKey(MetadataEntryRange[I], SourceData->Metadata, NameA3, IntValue);
+		UPCGMetadataAccessorHelpers::SetVectorAttributeByMetadataKey(MetadataEntryRange[I], SourceData->Metadata, NameB1, VecValue);
+		UPCGMetadataAccessorHelpers::SetVector4AttributeByMetadataKey(MetadataEntryRange[I], SourceData->Metadata, NameB2, Vec4Value);
 	}
 
 	Settings->SetInstancePackerType(UPCGInstanceDataPackerByRegex::StaticClass());
diff --git a/Engine/Plugins/PCG/Source/PCG/Private/Tests/Elements/PCGNormalToDensityTest.cpp b/Engine/Plugins/PCG/Source/PCG/Private/Tests/Elements/PCGNormalToDensityTest.cpp
--- a/Engine/Plugins/PCG/Source/PCG/Private/Tests/Elements/PCGNormalToDensityTest.cpp
+++ b/Engine/Plugins/PCG/Source/PCG/Private/Tests/Elements/PCGNormalToDensityTest.cpp
@@ -21,11 +21,13 @@ bool FPCGNormalToDensityTest_Set::RunTest(const FString& Parameters)
 	Settings->DensityMode = PCGNormalToDensityMode::Set;
 
 	{
-		TObjectPtr<UPCGPointData> PointData = PCGTestsCommon::CreatePointData();
-		FPCGPoint& Point = PointData->GetMutablePoints()[0];
+		TObjectPtr<UPCGBasePointData> PointData = PCGTestsCommon::CreateBasePointData();
+		
+		TPCGValueRange<FTransform> TransformRange = PointData->GetTransformValueRange();
+		TPCGValueRange<float> DensityRange = PointData->GetDensityValueRange();
 
-		Point.Transform.SetRotation(FQuat::MakeFromEuler(FVector(45.0, 0, 0)));
-		Point.Density = 0.5f;
+		TransformRange[0].SetRotation(FQuat::MakeFromEuler(FVector(45.0, 0, 0)));
+		DensityRange[0] = 0.5f;
 
 		FPCGTaggedData& TaggedData = TestData.InputData.TaggedData.Emplace_GetRef(FPCGTaggedData());
 		TaggedData.Data = PointData;
@@ -45,14 +47,13 @@ bool FPCGNormalToDensityTest_Set::RunTest(const FString& Parameters)
 
 	UTEST_NOT_NULL("Output data", OutputData);
 
-	const UPCGPointData* OutputPointData = OutputData->ToPointData(Context.Get());
+	const UPCGBasePointData* OutputPointData = OutputData->ToBasePointData(Context.Get());
 	UTEST_NOT_NULL("Output point data", OutputPointData);
 
-	const TArray<FPCGPoint>& OutputPoints = OutputPointData->GetPoints();
+	UTEST_EQUAL("Output point count", OutputPointData->GetNumPoints(), 1);
 
-	UTEST_EQUAL("Output point count", OutputPoints.Num(), 1);
-
-	UTEST_EQUAL_TOLERANCE("Output point density", OutputPoints[0].Density, 0.707f, 0.001f);
+	const TConstPCGValueRange<float> DensityRange = OutputPointData->GetConstDensityValueRange();
+	UTEST_EQUAL_TOLERANCE("Output point density", DensityRange[0], 0.707f, 0.001f);
 
 	return true;
 }
@@ -70,9 +71,10 @@ bool FPCGNormalToDensityTest_CustomNormal::RunTest(const FString& Parameters)
 	Settings->Normal = FVector(0.707, 0, 0.707);
 
 	{
-		TObjectPtr<UPCGPointData> PointData = PCGTestsCommon::CreatePointData();
-		FPCGPoint& Point = PointData->GetMutablePoints()[0];
-		Point.Density = 0.5f;
+		TObjectPtr<UPCGBasePointData> PointData = PCGTestsCommon::CreateBasePointData();
+
+		TPCGValueRange<float> DensityRange = PointData->GetDensityValueRange();
+		DensityRange[0] = 0.5f;
 
 		FPCGTaggedData& TaggedData = TestData.InputData.TaggedData.Emplace_GetRef(FPCGTaggedData());
 		TaggedData.Data = PointData;
@@ -92,14 +94,13 @@ bool FPCGNormalToDensityTest_CustomNormal::RunTest(const FString& Parameters)
 
 	UTEST_NOT_NULL("Output data", OutputData);
 
-	const UPCGPointData* OutputPointData = OutputData->ToPointData(Context.Get());
+	const UPCGBasePointData* OutputPointData = OutputData->ToBasePointData(Context.Get());
 	UTEST_NOT_NULL("Output point data", OutputPointData);
 
-	const TArray<FPCGPoint>& OutputPoints = OutputPointData->GetPoints();
-
-	UTEST_EQUAL("Output point count", OutputPoints.Num(), 1);
+	UTEST_EQUAL("Output point count", OutputPointData->GetNumPoints(), 1);
 
-	UTEST_EQUAL_TOLERANCE("Output point density", OutputPoints[0].Density, 0.707f, 0.001f);
+	const TConstPCGValueRange<float> DensityRange = OutputPointData->GetConstDensityValueRange();
+	UTEST_EQUAL_TOLERANCE("Output point density", DensityRange[0], 0.707f, 0.001f);
 
 	return true;
 }
@@ -118,9 +119,11 @@ bool FPCGNormalToDensityTest_Strength::RunTest(const FString& Parameters)
 	Settings->Normal = FVector(0.866, 0, 0.5); // so the dot with up is 0.5
 
 	{
-		TObjectPtr<UPCGPointData> PointData = PCGTestsCommon::CreatePointData();
-		FPCGPoint& Point = PointData->GetMutablePoints()[0];
-		Point.Density = 0.5f;
+		TObjectPtr<UPCGBasePointData> PointData = PCGTestsCommon::CreateBasePointData();
+
+		TPCGValueRange<float> DensityRange = PointData->GetDensityValueRange();
+		DensityRange[0] = 0.5f;
+
 		FPCGTaggedData& TaggedData = TestData.InputData.TaggedData.Emplace_GetRef(FPCGTaggedData());
 		TaggedData.Data = PointData;
 		TaggedData.Pin = PCGPinConstants::DefaultInputLabel;
@@ -139,14 +142,13 @@ bool FPCGNormalToDensityTest_Strength::RunTest(const FString& Parameters)
 
 	UTEST_NOT_NULL("Output data", OutputData);
 
-	const UPCGPointData* OutputPointData = OutputData->ToPointData(Context.Get());
+	const UPCGBasePointData* OutputPointData = OutputData->ToBasePointData(Context.Get());
 	UTEST_NOT_NULL("Output point data", OutputPointData);
 
-	const TArray<FPCGPoint>& OutputPoints = OutputPointData->GetPoints();
+	UTEST_EQUAL("Output point count", OutputPointData->GetNumPoints(), 1);
 
-	UTEST_EQUAL("Output point count", OutputPoints.Num(), 1);
-
-	UTEST_EQUAL_TOLERANCE("Output point density", OutputPoints[0].Density, 0.707f, 0.001f);
+	const TConstPCGValueRange<float> DensityRange = OutputPointData->GetConstDensityValueRange();
+	UTEST_EQUAL_TOLERANCE("Output point density", DensityRange[0], 0.707f, 0.001f);
 
 	return true;
 }
@@ -163,11 +165,13 @@ bool FPCGNormalToDensityTest_Minimum::RunTest(const FString& Parameters)
 	Settings->DensityMode = PCGNormalToDensityMode::Minimum;
 
 	{
-		TObjectPtr<UPCGPointData> PointData = PCGTestsCommon::CreatePointData();
-		FPCGPoint& Point = PointData->GetMutablePoints()[0];
+		TObjectPtr<UPCGBasePointData> PointData = PCGTestsCommon::CreateBasePointData();
+
+		TPCGValueRange<FTransform> TransformRange = PointData->GetTransformValueRange();
+		TPCGValueRange<float> DensityRange = PointData->GetDensityValueRange();
 
-		Point.Transform.SetRotation(FQuat::MakeFromEuler(FVector(45.0, 0, 0)));
-		Point.Density = 0.5f;
+		TransformRange[0].SetRotation(FQuat::MakeFromEuler(FVector(45.0, 0, 0)));
+		DensityRange[0] = 0.5f;
 
 		FPCGTaggedData& TaggedData = TestData.InputData.TaggedData.Emplace_GetRef(FPCGTaggedData());
 		TaggedData.Data = PointData;
@@ -187,14 +191,13 @@ bool FPCGNormalToDensityTest_Minimum::RunTest(const FString& Parameters)
 
 	UTEST_NOT_NULL("Output data", OutputData);
 
-	const UPCGPointData* OutputPointData = OutputData->ToPointData(Context.Get());
+	const UPCGBasePointData* OutputPointData = OutputData->ToBasePointData(Context.Get());
 	UTEST_NOT_NULL("Output point data", OutputPointData);
 
-	const TArray<FPCGPoint>& OutputPoints = OutputPointData->GetPoints();
+	UTEST_EQUAL("Output point count", OutputPointData->GetNumPoints(), 1);
 
-	UTEST_EQUAL("Output point count", OutputPoints.Num(), 1);
-
-	UTEST_EQUAL_TOLERANCE("Output point density", OutputPoints[0].Density, 0.5f, 0.001f);
+	const TConstPCGValueRange<float> DensityRange = OutputPointData->GetConstDensityValueRange();
+	UTEST_EQUAL_TOLERANCE("Output point density", DensityRange[0], 0.5f, 0.001f);
 
 	return true;
 }
@@ -211,11 +214,13 @@ bool FPCGNormalToDensityTest_Maximum::RunTest(const FString& Parameters)
 	Settings->DensityMode = PCGNormalToDensityMode::Maximum;
 
 	{
-		TObjectPtr<UPCGPointData> PointData = PCGTestsCommon::CreatePointData();
-		FPCGPoint& Point = PointData->GetMutablePoints()[0];
+		TObjectPtr<UPCGBasePointData> PointData = PCGTestsCommon::CreateBasePointData();
+
+		TPCGValueRange<FTransform> TransformRange = PointData->GetTransformValueRange();
+		TPCGValueRange<float> DensityRange = PointData->GetDensityValueRange();
 
-		Point.Transform.SetRotation(FQuat::MakeFromEuler(FVector(45.0, 0, 0)));
-		Point.Density = 0.5f;
+		TransformRange[0].SetRotation(FQuat::MakeFromEuler(FVector(45.0, 0, 0)));
+		DensityRange[0] = 0.5f;
 
 		FPCGTaggedData& TaggedData = TestData.InputData.TaggedData.Emplace_GetRef(FPCGTaggedData());
 		TaggedData.Data = PointData;
@@ -235,14 +240,13 @@ bool FPCGNormalToDensityTest_Maximum::RunTest(const FString& Parameters)
 
 	UTEST_NOT_NULL("Output data", OutputData);
 
-	const UPCGPointData* OutputPointData = OutputData->ToPointData(Context.Get());
+	const UPCGBasePointData* OutputPointData = OutputData->ToBasePointData(Context.Get());
 	UTEST_NOT_NULL("Output point data", OutputPointData);
 
-	const TArray<FPCGPoint>& OutputPoints = OutputPointData->GetPoints();
+	UTEST_EQUAL("Output point count", OutputPointData->GetNumPoints(), 1);
 
-	UTEST_EQUAL("Output point count", OutputPoints.Num(), 1);
-
-	UTEST_EQUAL_TOLERANCE("Output point density", OutputPoints[0].Density, 0.707f, 0.001f);
+	const TConstPCGValueRange<float> DensityRange = OutputPointData->GetConstDensityValueRange();
+	UTEST_EQUAL_TOLERANCE("Output point density", DensityRange[0], 0.707f, 0.001f);
 
 	return true;
 }
@@ -260,11 +264,13 @@ bool FPCGNormalToDensityTest_Add::RunTest(const FString& Parameters)
 	Settings->DensityMode = PCGNormalToDensityMode::Add;
 
 	{
-		TObjectPtr<UPCGPointData> PointData = PCGTestsCommon::CreatePointData();
-		FPCGPoint& Point = PointData->GetMutablePoints()[0];
+		TObjectPtr<UPCGBasePointData> PointData = PCGTestsCommon::CreateBasePointData();
+
+		TPCGValueRange<FTransform> TransformRange = PointData->GetTransformValueRange();
+		TPCGValueRange<float> DensityRange = PointData->GetDensityValueRange();
 
-		Point.Transform.SetRotation(FQuat::MakeFromEuler(FVector(45.0, 0, 0)));
-		Point.Density = 0.1f;
+		TransformRange[0].SetRotation(FQuat::MakeFromEuler(FVector(45.0, 0, 0)));
+		DensityRange[0] = 0.1f;
 
 		FPCGTaggedData& TaggedData = TestData.InputData.TaggedData.Emplace_GetRef(FPCGTaggedData());
 		TaggedData.Data = PointData;
@@ -284,14 +290,13 @@ bool FPCGNormalToDensityTest_Add::RunTest(const FString& Parameters)
 
 	UTEST_NOT_NULL("Output data", OutputData);
 
-	const UPCGPointData* OutputPointData = OutputData->ToPointData(Context.Get());
+	const UPCGBasePointData* OutputPointData = OutputData->ToBasePointData(Context.Get());
 	UTEST_NOT_NULL("Output point data", OutputPointData);
 
-	const TArray<FPCGPoint>& OutputPoints = OutputPointData->GetPoints();
-
-	UTEST_EQUAL("Output point count", OutputPoints.Num(), 1);
+	UTEST_EQUAL("Output point count", OutputPointData->GetNumPoints(), 1);
 
-	UTEST_EQUAL_TOLERANCE("Output point density", OutputPoints[0].Density, 0.807f, 0.001f);
+	const TConstPCGValueRange<float> DensityRange = OutputPointData->GetConstDensityValueRange();
+	UTEST_EQUAL_TOLERANCE("Output point density", DensityRange[0], 0.807f, 0.001f);
 
 	return true;
 }
@@ -308,11 +313,13 @@ bool FPCGNormalToDensityTest_Subtract::RunTest(const FString& Parameters)
 	Settings->DensityMode = PCGNormalToDensityMode::Subtract;
 
 	{
-		TObjectPtr<UPCGPointData> PointData = PCGTestsCommon::CreatePointData();
-		FPCGPoint& Point = PointData->GetMutablePoints()[0];
+		TObjectPtr<UPCGBasePointData> PointData = PCGTestsCommon::CreateBasePointData();
 
-		Point.Transform.SetRotation(FQuat::MakeFromEuler(FVector(45.0, 0, 0)));
-		Point.Density = 0.9f;
+		TPCGValueRange<FTransform> TransformRange = PointData->GetTransformValueRange();
+		TPCGValueRange<float> DensityRange = PointData->GetDensityValueRange();
+
+		TransformRange[0].SetRotation(FQuat::MakeFromEuler(FVector(45.0, 0, 0)));
+		DensityRange[0] = 0.9f;
 
 		FPCGTaggedData& TaggedData = TestData.InputData.TaggedData.Emplace_GetRef(FPCGTaggedData());
 		TaggedData.Data = PointData;
@@ -332,14 +339,13 @@ bool FPCGNormalToDensityTest_Subtract::RunTest(const FString& Parameters)
 
 	UTEST_NOT_NULL("Output data", OutputData);
 
-	const UPCGPointData* OutputPointData = OutputData->ToPointData(Context.Get());
+	const UPCGBasePointData* OutputPointData = OutputData->ToBasePointData(Context.Get());
 	UTEST_NOT_NULL("Output point data", OutputPointData);
 
-	const TArray<FPCGPoint>& OutputPoints = OutputPointData->GetPoints();
-
-	UTEST_EQUAL("Output point count", OutputPoints.Num(), 1);
+	UTEST_EQUAL("Output point count", OutputPointData->GetNumPoints(), 1);
 
-	UTEST_EQUAL_TOLERANCE("Output point density", OutputPoints[0].Density, 0.9f-0.707f, 0.001f);
+	const TConstPCGValueRange<float> DensityRange = OutputPointData->GetConstDensityValueRange();
+	UTEST_EQUAL_TOLERANCE("Output point density", DensityRange[0], 0.9f - 0.707f, 0.001f);
 
 	return true;
 }
@@ -356,11 +362,13 @@ bool FPCGNormalToDensityTest_Multiply::RunTest(const FString& Parameters)
 	Settings->DensityMode = PCGNormalToDensityMode::Multiply;
 
 	{
-		TObjectPtr<UPCGPointData> PointData = PCGTestsCommon::CreatePointData();
-		FPCGPoint& Point = PointData->GetMutablePoints()[0];
+		TObjectPtr<UPCGBasePointData> PointData = PCGTestsCommon::CreateBasePointData();
 
-		Point.Transform.SetRotation(FQuat::MakeFromEuler(FVector(45.0, 0, 0)));
-		Point.Density = 0.5f;
+		TPCGValueRange<FTransform> TransformRange = PointData->GetTransformValueRange();
+		TPCGValueRange<float> DensityRange = PointData->GetDensityValueRange();
+
+		TransformRange[0].SetRotation(FQuat::MakeFromEuler(FVector(45.0, 0, 0)));
+		DensityRange[0] = 0.5f;
 
 		FPCGTaggedData& TaggedData = TestData.InputData.TaggedData.Emplace_GetRef(FPCGTaggedData());
 		TaggedData.Data = PointData;
@@ -380,14 +388,13 @@ bool FPCGNormalToDensityTest_Multiply::RunTest(const FString& Parameters)
 
 	UTEST_NOT_NULL("Output data", OutputData);
 
-	const UPCGPointData* OutputPointData = OutputData->ToPointData(Context.Get());
+	const UPCGBasePointData* OutputPointData = OutputData->ToBasePointData(Context.Get());
 	UTEST_NOT_NULL("Output point data", OutputPointData);
 
-	const TArray<FPCGPoint>& OutputPoints = OutputPointData->GetPoints();
-
-	UTEST_EQUAL("Output point count", OutputPoints.Num(), 1);
+	UTEST_EQUAL("Output point count", OutputPointData->GetNumPoints(), 1);
 
-	UTEST_EQUAL_TOLERANCE("Output point density", OutputPoints[0].Density, 0.5f*0.707f, 0.001f);
+	const TConstPCGValueRange<float> DensityRange = OutputPointData->GetConstDensityValueRange();
+	UTEST_EQUAL_TOLERANCE("Output point density", DensityRange[0], 0.5f*0.707f, 0.001f);
 
 	return true;
 }
@@ -404,11 +411,13 @@ bool FPCGNormalToDensityTest_Divide::RunTest(const FString& Parameters)
 	Settings->DensityMode = PCGNormalToDensityMode::Divide;
 
 	{
-		TObjectPtr<UPCGPointData> PointData = PCGTestsCommon::CreatePointData();
-		FPCGPoint& Point = PointData->GetMutablePoints()[0];
+		TObjectPtr<UPCGBasePointData> PointData = PCGTestsCommon::CreateBasePointData();
 
-		Point.Transform.SetRotation(FQuat::MakeFromEuler(FVector(45.0, 0, 0)));
-		Point.Density = 0.5f;
+		TPCGValueRange<FTransform> TransformRange = PointData->GetTransformValueRange();
+		TPCGValueRange<float> DensityRange = PointData->GetDensityValueRange();
+
+		TransformRange[0].SetRotation(FQuat::MakeFromEuler(FVector(45.0, 0, 0)));
+		DensityRange[0] = 0.5f;
 
 		FPCGTaggedData& TaggedData = TestData.InputData.TaggedData.Emplace_GetRef(FPCGTaggedData());
 		TaggedData.Data = PointData;
@@ -428,14 +437,13 @@ bool FPCGNormalToDensityTest_Divide::RunTest(const FString& Parameters)
 
 	UTEST_NOT_NULL("Output data", OutputData);
 
-	const UPCGPointData* OutputPointData = OutputData->ToPointData(Context.Get());
+	const UPCGBasePointData* OutputPointData = OutputData->ToBasePointData(Context.Get());
 	UTEST_NOT_NULL("Output point data", OutputPointData);
 
-	const TArray<FPCGPoint>& OutputPoints = OutputPointData->GetPoints();
-
-	UTEST_EQUAL("Output point count", OutputPoints.Num(), 1);
+	UTEST_EQUAL("Output point count", OutputPointData->GetNumPoints(), 1);
 
-	UTEST_EQUAL_TOLERANCE("Output point density", OutputPoints[0].Density, 0.5f/0.707f, 0.001f);
+	const TConstPCGValueRange<float> DensityRange = OutputPointData->GetConstDensityValueRange();
+	UTEST_EQUAL_TOLERANCE("Output point density", DensityRange[0], 0.5f/0.707f, 0.001f);
 
 	return true;
 }
diff --git a/Engine/Plugins/PCG/Source/PCG/Private/Tests/Elements/PCGPointExtentsModifierTest.cpp b/Engine/Plugins/PCG/Source/PCG/Private/Tests/Elements/PCGPointExtentsModifierTest.cpp
--- a/Engine/Plugins/PCG/Source/PCG/Private/Tests/Elements/PCGPointExtentsModifierTest.cpp
+++ b/Engine/Plugins/PCG/Source/PCG/Private/Tests/Elements/PCGPointExtentsModifierTest.cpp
@@ -18,21 +18,28 @@ bool FPCGPointsBoundsModiferTest::RunTest(const FString& Parameters)
 	UPCGPointExtentsModifierSettings* Settings = CastChecked<UPCGPointExtentsModifierSettings>(TestData.Settings);
 	FPCGElementPtr BoundsModifierElement = TestData.Settings->GetElement();
 
-	TObjectPtr<UPCGPointData> PointData = PCGTestsCommon::CreateEmptyPointData();
-	TArray<FPCGPoint>& Points = PointData->GetMutablePoints();
-
+	TObjectPtr<UPCGBasePointData> PointData = PCGTestsCommon::CreateEmptyBasePointData();
+	
 	static const int PointCount = 4;
+	
+	PointData->SetNumPoints(PointCount);
+	PointData->SetDensity(1);
+	PointData->AllocateProperties(EPCGPointNativeProperties::Seed | EPCGPointNativeProperties::BoundsMin | EPCGPointNativeProperties::BoundsMax);
+	
+	FPCGPointValueRanges ValueRanges(PointData, /*bAllocate=*/false);
+
 	for (int I = 0; I < PointCount; ++I)
 	{
-		FPCGPoint& Point = Points.Emplace_GetRef(FTransform(), 1, I);
+		ValueRanges.SeedRange[I] = I;
+
 		const float Value = float(I+1);
 
-		Point.SetExtents(FVector(Value, Value, Value));
+		PCGPointHelpers::SetExtents(FVector(Value, Value, Value), ValueRanges.BoundsMinRange[I], ValueRanges.BoundsMaxRange[I]);
 
 		if (I == PointCount-1)
 		{
 			// make the last one off-center, so (-3,-3,-3),(5,5,5)
-			Point.SetLocalCenter(FVector(1.0f, 1.0f, 1.0f));
+			PCGPointHelpers::SetLocalCenter(FVector(1.0f, 1.0f, 1.0f), ValueRanges.BoundsMinRange[I], ValueRanges.BoundsMaxRange[I]);
 		}
 	}
 
@@ -70,7 +77,7 @@ bool FPCGPointsBoundsModiferTest::RunTest(const FString& Parameters)
 			const UPCGSpatialData* InSpatialData = Cast<UPCGSpatialData>(Input.Data);
 			check(InSpatialData);
 
-			const UPCGPointData* InPointData = InSpatialData->ToPointData(Context.Get());
+			const UPCGBasePointData* InPointData = InSpatialData->ToBasePointData(Context.Get());
 			check(InPointData);
 
 			const UPCGSpatialData* OutSpatialData = Cast<UPCGSpatialData>(Output.Data);
@@ -81,7 +88,7 @@ bool FPCGPointsBoundsModiferTest::RunTest(const FString& Parameters)
 				continue;
 			}
 
-			const UPCGPointData* OutPointData = OutSpatialData->ToPointData(Context.Get());
+			const UPCGBasePointData* OutPointData = OutSpatialData->ToBasePointData(Context.Get());
 
 			if (!TestNotNull("Valid output PointData", OutPointData))
 			{
@@ -89,25 +96,25 @@ bool FPCGPointsBoundsModiferTest::RunTest(const FString& Parameters)
 				continue;
 			}
 
-			const TArray<FPCGPoint>& InPoints = InPointData->GetPoints();
-			const TArray<FPCGPoint>& OutPoints = OutPointData->GetPoints();
-
-			if (!TestEqual("Input and output point counts match", InPoints.Num(), OutPoints.Num()))
+			if (!TestEqual("Input and output point counts match", InPointData->GetNumPoints(), OutPointData->GetNumPoints()))
 			{ 
 				bTestPassed = false;
 				continue;
 			}
 
-			if (!TestEqual("Test data has enough points", PointCount, OutPoints.Num()))
+			if (!TestEqual("Test data has enough points", PointCount, OutPointData->GetNumPoints()))
 			{
 				bTestPassed = false;
 				continue;
 			}
 
+			const TConstPCGValueRange<FVector> OutBoundsMinRange = OutPointData->GetConstBoundsMinValueRange();
+			const TConstPCGValueRange<FVector> OutBoundsMaxRange = OutPointData->GetConstBoundsMaxValueRange();
+
 			for (int PointIndex = 0; PointIndex < ExpectedOutput.Num()/2; ++PointIndex)
 			{
-				bTestPassed &= TestEqual("Correct BoundsMin", OutPoints[PointIndex].BoundsMin, ExpectedOutput[PointIndex*2+0]);
-				bTestPassed &= TestEqual("Correct BoundsMax", OutPoints[PointIndex].BoundsMax, ExpectedOutput[PointIndex*2+1]);
+				bTestPassed &= TestEqual("Correct BoundsMin", OutBoundsMinRange[PointIndex], ExpectedOutput[PointIndex*2+0]);
+				bTestPassed &= TestEqual("Correct BoundsMax", OutBoundsMaxRange[PointIndex], ExpectedOutput[PointIndex*2+1]);
 			}
 		}
 
diff --git a/Engine/Plugins/PCG/Source/PCG/Private/Tests/Elements/PCGPointNeighborhoodTest.cpp b/Engine/Plugins/PCG/Source/PCG/Private/Tests/Elements/PCGPointNeighborhoodTest.cpp
--- a/Engine/Plugins/PCG/Source/PCG/Private/Tests/Elements/PCGPointNeighborhoodTest.cpp
+++ b/Engine/Plugins/PCG/Source/PCG/Private/Tests/Elements/PCGPointNeighborhoodTest.cpp
@@ -43,26 +43,23 @@ protected:
 
 		FPCGTaggedData& Inputs = TestData.InputData.TaggedData.Emplace_GetRef();
 		Inputs.Pin = PCGPinConstants::DefaultInputLabel;
-		UPCGPointData* InData = PCGTestsCommon::CreatePointData();
+		UPCGBasePointData* InData = PCGTestsCommon::CreateBasePointData();
 		Inputs.Data = InData;
 
-		TArray<FPCGPoint>& InPoints = InData->GetMutablePoints();
-		InPoints.SetNum(9);
+		InData->SetNumPoints(9);
+		InData->SetBoundsMin(FVector(-10));
+		InData->SetBoundsMax(FVector(10));
+
+		TPCGValueRange<FTransform> TransformRange = InData->GetTransformValueRange();
 
 		for (int i = 0; i < 3; ++i)
 		{
 			for (int j = 0; j < 3; ++j)
 			{
-				InPoints[(i * 3) + j].Transform.SetLocation(FVector(j * 100.0, i * 100.0, 0.0));
+				TransformRange[(i * 3) + j].SetLocation(FVector(j * 100.0, i * 100.0, 0.0));
 			}
 		}
 
-		for (FPCGPoint& InPoint : InPoints)
-		{
-			InPoint.BoundsMin = FVector(-10.0);
-			InPoint.BoundsMax = FVector(10.0);
-		}
-
 		FPCGElementPtr TestElement = TestData.Settings->GetElement();
 
 		TUniquePtr<FPCGContext> Context = TestData.InitializeTestContext();
diff --git a/Engine/Plugins/PCG/Source/PCG/Private/Tests/Elements/PCGReplaceTagTest.cpp b/Engine/Plugins/PCG/Source/PCG/Private/Tests/Elements/PCGReplaceTagTest.cpp
--- a/Engine/Plugins/PCG/Source/PCG/Private/Tests/Elements/PCGReplaceTagTest.cpp
+++ b/Engine/Plugins/PCG/Source/PCG/Private/Tests/Elements/PCGReplaceTagTest.cpp
@@ -33,7 +33,7 @@ public:
 		TUniquePtr<FPCGContext> Context = TestData.InitializeTestContext();
 
 		FPCGTaggedData& InputTaggedData = Context->InputData.TaggedData.Emplace_GetRef();
-		InputTaggedData.Data = PCGTestsCommon::CreatePointData();
+		InputTaggedData.Data = PCGTestsCommon::CreateBasePointData();
 		InputTaggedData.Pin = PCGPinConstants::DefaultInputLabel;
 		InputTaggedData.Tags = TSet<FString>{ TEXT("Awesome"), TEXT("Fun"), TEXT("Cool") };
 
diff --git a/Engine/Plugins/PCG/Source/PCG/Private/Tests/Elements/PCGResetPointCenterTest.cpp b/Engine/Plugins/PCG/Source/PCG/Private/Tests/Elements/PCGResetPointCenterTest.cpp
--- a/Engine/Plugins/PCG/Source/PCG/Private/Tests/Elements/PCGResetPointCenterTest.cpp
+++ b/Engine/Plugins/PCG/Source/PCG/Private/Tests/Elements/PCGResetPointCenterTest.cpp
@@ -18,16 +18,15 @@ namespace ResetPointCenterCommonTests
 
 		FPCGTaggedData& Inputs = TestData.InputData.TaggedData.Emplace_GetRef();
 		Inputs.Pin = PCGPinConstants::DefaultInputLabel;
-		UPCGPointData* InData = PCGTestsCommon::CreatePointData();
+		UPCGBasePointData* InData = PCGTestsCommon::CreateBasePointData();
 		Inputs.Data = InData;
 
 		// Setting the scale and rotation explicitly
-		UPCGPointData* InputPointData = Cast<UPCGPointData>(InData);
-		TArray<FPCGPoint>& InPoints = InputPointData->GetMutablePoints();
-		for (FPCGPoint& InPoint : InPoints)
+		TPCGValueRange<FTransform> TransformRange = InData->GetTransformValueRange();
+		for (FTransform& PointTransform : TransformRange)
 		{
-			InPoint.Transform.SetRotation(PointRotation.Quaternion());
-			InPoint.Transform.SetScale3D(PointScale);
+			PointTransform.SetRotation(PointRotation.Quaternion());
+			PointTransform.SetScale3D(PointScale);
 		}
 		
 		FPCGElementPtr TestElement = TestData.Settings->GetElement();
diff --git a/Engine/Plugins/PCG/Source/PCG/Private/Tests/Elements/PCGSampleTextureTest.cpp b/Engine/Plugins/PCG/Source/PCG/Private/Tests/Elements/PCGSampleTextureTest.cpp
--- a/Engine/Plugins/PCG/Source/PCG/Private/Tests/Elements/PCGSampleTextureTest.cpp
+++ b/Engine/Plugins/PCG/Source/PCG/Private/Tests/Elements/PCGSampleTextureTest.cpp
@@ -88,13 +88,16 @@ protected:
 		Settings->TextureMappingMethod = Parameters.TextureMappingMethod;
 	
 		FPCGTaggedData& PointInputs = TestData.InputData.TaggedData.Emplace_GetRef();
-		UPCGPointData* InData = PCGTestsCommon::CreatePointData();
+		UPCGBasePointData* InData = PCGTestsCommon::CreateBasePointData();
 		PointInputs.Data = InData;
 		PointInputs.Pin = PCGSampleTextureConstants::InputPointLabel;
 
 		FPCGMetadataAttribute<FVector2D>* InputUV = InData->Metadata->CreateAttribute<FVector2D>(Parameters.UVAttribute, FVector2D::ZeroVector, /*bAllowsInterpolation=*/true, /*bOverrideParent=*/true);
-		TArray<FPCGPoint>& Points = InData->GetMutablePoints();
-		Points.SetNum(16);
+		InData->SetNumPoints(16);
+
+		TPCGValueRange<FTransform> TransformRange = InData->GetTransformValueRange();
+		TPCGValueRange<int64> MetadataEntryRange = InData->GetMetadataEntryValueRange();
+
 		for (int i = 0; i < 4; ++i)
 		{
 			for (int j = 0; j < 4; ++j)
@@ -102,11 +105,11 @@ protected:
 				int CurrentPoint = i * 4 + j;
 
 				// Place each point at the texel center for easy UV access
-				Points[CurrentPoint].Transform.SetLocation(FVector((j + 0.5) * 0.25, (i + 0.5) * 0.25, 0.0));
+				TransformRange[CurrentPoint].SetLocation(FVector((j + 0.5) * 0.25, (i + 0.5) * 0.25, 0.0));
 
 				// Set the UV metadata for each point
-				InData->Metadata->InitializeOnSet(Points[CurrentPoint].MetadataEntry);
-				InputUV->SetValue(Points[CurrentPoint].MetadataEntry, FVector2D(Points[CurrentPoint].Transform.GetLocation()));
+				InData->Metadata->InitializeOnSet(MetadataEntryRange[CurrentPoint]);
+				InputUV->SetValue(MetadataEntryRange[CurrentPoint], FVector2D(TransformRange[CurrentPoint].GetLocation()));
 			}
 		}
 
diff --git a/Engine/Plugins/PCG/Source/PCG/Private/Tests/Elements/PCGSanityCheckPointDataTest.cpp b/Engine/Plugins/PCG/Source/PCG/Private/Tests/Elements/PCGSanityCheckPointDataTest.cpp
--- a/Engine/Plugins/PCG/Source/PCG/Private/Tests/Elements/PCGSanityCheckPointDataTest.cpp
+++ b/Engine/Plugins/PCG/Source/PCG/Private/Tests/Elements/PCGSanityCheckPointDataTest.cpp
@@ -24,7 +24,7 @@ bool FPCGSanityCheckPointData_MinPointCount::RunTest(const FString& Parameters)
 
 	FPCGElementPtr TestElement = TestData.Settings->GetElement();
 
-	UPCGPointData *SourceData = PCGTestsCommon::CreateRandomPointData(2, 42);
+	UPCGBasePointData *SourceData = PCGTestsCommon::CreateRandomBasePointData(2, 42);
 
 	FPCGTaggedData& SourcePin = TestData.InputData.TaggedData.Emplace_GetRef();
 	SourcePin.Pin = PCGPinConstants::DefaultInputLabel;
@@ -52,7 +52,7 @@ bool FPCGSanityCheckPointData_MaxPointCount::RunTest(const FString& Parameters)
 
 	FPCGElementPtr TestElement = TestData.Settings->GetElement();
 
-	UPCGPointData *SourceData = PCGTestsCommon::CreateRandomPointData(200, 42);
+	UPCGBasePointData *SourceData = PCGTestsCommon::CreateRandomBasePointData(200, 42);
 
 	FPCGTaggedData& SourcePin = TestData.InputData.TaggedData.Emplace_GetRef();
 	SourcePin.Pin = PCGPinConstants::DefaultInputLabel;
@@ -80,7 +80,7 @@ bool FPCGSanityCheckPointData_Success::RunTest(const FString& Parameters)
 
 	FPCGElementPtr TestElement = TestData.Settings->GetElement();
 
-	UPCGPointData *SourceData = PCGTestsCommon::CreateRandomPointData(50, 42);
+	UPCGBasePointData *SourceData = PCGTestsCommon::CreateRandomBasePointData(50, 42);
 
 	FPCGTaggedData& SourcePin = TestData.InputData.TaggedData.Emplace_GetRef();
 	SourcePin.Pin = PCGPinConstants::DefaultInputLabel;
diff --git a/Engine/Plugins/PCG/Source/PCG/Private/Tests/Elements/PCGSplitPointsTest.cpp b/Engine/Plugins/PCG/Source/PCG/Private/Tests/Elements/PCGSplitPointsTest.cpp
--- a/Engine/Plugins/PCG/Source/PCG/Private/Tests/Elements/PCGSplitPointsTest.cpp
+++ b/Engine/Plugins/PCG/Source/PCG/Private/Tests/Elements/PCGSplitPointsTest.cpp
@@ -35,7 +35,7 @@ protected:
 
 		FPCGTaggedData& Inputs = TestData.InputData.TaggedData.Emplace_GetRef();
 		Inputs.Pin = PCGPinConstants::DefaultInputLabel;
-		Inputs.Data = PCGTestsCommon::CreatePointData();
+		Inputs.Data = PCGTestsCommon::CreateBasePointData();
 
 		FPCGElementPtr TestElement = TestData.Settings->GetElement();
 
diff --git a/Engine/Plugins/PCG/Source/PCG/Private/Tests/Elements/PCGStaticMeshSpawnerTest.cpp b/Engine/Plugins/PCG/Source/PCG/Private/Tests/Elements/PCGStaticMeshSpawnerTest.cpp
--- a/Engine/Plugins/PCG/Source/PCG/Private/Tests/Elements/PCGStaticMeshSpawnerTest.cpp
+++ b/Engine/Plugins/PCG/Source/PCG/Private/Tests/Elements/PCGStaticMeshSpawnerTest.cpp
@@ -47,8 +47,8 @@ namespace
 	};
 }
 
-/** Finds the output UPCGPointData and builds a mapping from world-space position to each FPCGPoint */
-void GetOutputDataAndBuildPositionMap(FPCGTestBaseClass* Test, const TArray<FPCGTaggedData>& OutputData, TMap<FVector, FPCGPoint>& OutLocationToPoint, const UPCGPointData*& OutPointData)
+/** Finds the output UPCGBasePointData and builds a mapping from world-space position to each Point Index */
+void GetOutputDataAndBuildPositionMap(FPCGTestBaseClass* Test, const TArray<FPCGTaggedData>& OutputData, TMap<FVector, int32>& OutLocationToPoint, const UPCGBasePointData*& OutPointData)
 {
 	check(Test);
 
@@ -57,16 +57,16 @@ void GetOutputDataAndBuildPositionMap(FPCGTestBaseClass* Test, const TArray<FPCG
 		TObjectPtr<const UPCGSettings> OutSettings = Cast<UPCGSettings>(OutputData[0].Data);
 		Test->TestTrue("Valid output settings", OutSettings != nullptr);
 
-		OutPointData = Cast<UPCGPointData>(OutputData[1].Data);
+		OutPointData = Cast<UPCGBasePointData>(OutputData[1].Data);
 	}
 
 	if (Test->TestNotNull("Valid output data", OutPointData))
 	{
-		const TArray<FPCGPoint>& OutPoints = OutPointData->GetPoints();
-
-		for (const FPCGPoint& OutPoint : OutPoints)
+		const TConstPCGValueRange<FTransform> TransformRange = OutPointData->GetConstTransformValueRange();
+		
+		for (int32 Index = 0; Index < OutPointData->GetNumPoints(); ++Index)
 		{
-			OutLocationToPoint.Emplace(OutPoint.Transform.GetLocation(), OutPoint);
+			OutLocationToPoint.Emplace(TransformRange[Index].GetLocation(), Index);
 		}
 	}
 }
@@ -75,9 +75,9 @@ void GetOutputDataAndBuildPositionMap(FPCGTestBaseClass* Test, const TArray<FPCG
 void ValidateOutPointData(
 	FPCGTestBaseClass* Test, 
 	const UInstancedStaticMeshComponent* ISMC, 
-	const TMap<FVector, FPCGPoint>& LocationToPoint, 
+	const TMap<FVector, int32>& LocationToPointIndex, 
 	const FName& OutputAttributeName, 
-	const UPCGPointData* OutPointData) 
+	const UPCGBasePointData* OutPointData) 
 {
 	check(Test);
 	check(ISMC);
@@ -107,16 +107,17 @@ void ValidateOutPointData(
 		}
 	}
 
+	const TConstPCGValueRange<int64> MetadataEntryRange = OutPointData->GetConstMetadataEntryValueRange();
 	for (int InstanceIndex = 0; InstanceIndex < ISMC->GetInstanceCount(); ++InstanceIndex)
 	{
 		FTransform Transform;
 		ISMC->GetInstanceTransform(InstanceIndex, Transform, /*bWorldSpace=*/true);
 
-		const FPCGPoint* OutPoint = LocationToPoint.Find(Transform.GetLocation());
+		const int32* OutPointIndex = LocationToPointIndex.Find(Transform.GetLocation());
 
-		if (Test->TestNotNull("Output point exists", OutPoint))
+		if (Test->TestNotNull("Output point index exists", OutPointIndex) && Test->TestNotEqual("Ouput point index is valid", *OutPointIndex, (int32)INDEX_NONE))
 		{
-			const PCGMetadataValueKey ValueKey = OutAttribute->GetValueKey(OutPoint->MetadataEntry);
+			const PCGMetadataValueKey ValueKey = OutAttribute->GetValueKey(MetadataEntryRange[*OutPointIndex]);
 			Test->TestEqual("Valid mesh stored in output point metadata", ValueKey, MeshValueKey);
 		}
 	}
@@ -144,12 +145,12 @@ void TestMeshSelectorByAttribute(
 	while (!StaticMeshSpawner->Execute(Context.Get()))
 	{}
 
-	TMap<FVector, FPCGPoint> LocationToPoint;
-	const UPCGPointData* OutPointData = nullptr;
+	TMap<FVector, int32> LocationToPointIndex;
+	const UPCGBasePointData* OutPointData = nullptr;
 
 	if (bValidateOutput)
 	{
-		GetOutputDataAndBuildPositionMap(Test, Context->OutputData.TaggedData, LocationToPoint, OutPointData);
+		GetOutputDataAndBuildPositionMap(Test, Context->OutputData.TaggedData, LocationToPointIndex, OutPointData);
 	}
 
 	struct FPCGByAttributeValidIndexList
@@ -234,7 +235,7 @@ void TestMeshSelectorByAttribute(
 
 			if (bValidateOutput)
 			{
-				ValidateOutPointData(Test, ISMC, LocationToPoint, Settings->OutAttributeName, OutPointData);
+				ValidateOutPointData(Test, ISMC, LocationToPointIndex, Settings->OutAttributeName, OutPointData);
 			}
 		}
 	}
@@ -267,12 +268,12 @@ void TestMeshSelectorWeighted(
 	while (!StaticMeshSpawner->Execute(Context.Get()))
 	{}
 
-	TMap<FVector, FPCGPoint> LocationToPoint;
-	const UPCGPointData* OutPointData = nullptr;
+	TMap<FVector, int32> LocationToPointIndex;
+	const UPCGBasePointData* OutPointData = nullptr;
 
 	if (bValidateOutput)
 	{
-		GetOutputDataAndBuildPositionMap(Test, Context->OutputData.TaggedData, LocationToPoint, OutPointData);
+		GetOutputDataAndBuildPositionMap(Test, Context->OutputData.TaggedData, LocationToPointIndex, OutPointData);
 	}
 
 	int TotalWeight = 0;
@@ -318,7 +319,7 @@ void TestMeshSelectorWeighted(
 
 			if (bValidateOutput)
 			{
-				ValidateOutPointData(Test, ISMC, LocationToPoint, Settings->OutAttributeName, OutPointData);
+				ValidateOutPointData(Test, ISMC, LocationToPointIndex, Settings->OutAttributeName, OutPointData);
 			}
 		}
 	}
@@ -355,12 +356,12 @@ void TestMeshSelectorWeightedByCategory(
 	{
 	}
 
-	TMap<FVector, FPCGPoint> LocationToPoint;
-	const UPCGPointData* OutPointData = nullptr;
+	TMap<FVector, int32> LocationToPointIndex;
+	const UPCGBasePointData* OutPointData = nullptr;
 
 	if (bValidateOutput)
 	{
-		GetOutputDataAndBuildPositionMap(Test, Context->OutputData.TaggedData, LocationToPoint, OutPointData);
+		GetOutputDataAndBuildPositionMap(Test, Context->OutputData.TaggedData, LocationToPointIndex, OutPointData);
 	}
 
 	check(TestData.TestPCGComponent);
@@ -396,7 +397,7 @@ void TestMeshSelectorWeightedByCategory(
 
 			if (bValidateOutput)
 			{
-				ValidateOutPointData(Test, ISMC, LocationToPoint, Settings->OutAttributeName, OutPointData);
+				ValidateOutPointData(Test, ISMC, LocationToPointIndex, Settings->OutAttributeName, OutPointData);
 			}
 		}
 	}
@@ -428,7 +429,7 @@ bool FPCGStaticMeshSpawnerByAttributeTest::RunTest(const FString& Parameters)
 	TSoftObjectPtr<UStaticMesh> ConeMesh(ConePath);
 
 	// compose the test data
-	TObjectPtr<UPCGPointData> PointData = PCGTestsCommon::CreateEmptyPointData();
+	TObjectPtr<UPCGBasePointData> PointData = PCGTestsCommon::CreateEmptyBasePointData();
 	PointData->TargetActor = TestData.TestActor;
 	PointData->Metadata->CreateStringAttribute(RockAttribute, ConePath.ToString(), false);
 	PointData->Metadata->CreateStringAttribute(TreeAttribute, ConePath.ToString(), false);
@@ -436,11 +437,18 @@ bool FPCGStaticMeshSpawnerByAttributeTest::RunTest(const FString& Parameters)
 
 	// first 20 points are Rock, next 50 are Tree, next 30 are Bush, for a total of 100 points
 	const int PointCount = 100;
-	TArray<FPCGPoint>& Points = PointData->GetMutablePoints();
+	PointData->SetNumPoints(PointCount);
+	PointData->SetDensity(1);
+
+	TPCGValueRange<FTransform> TransformRange = PointData->GetTransformValueRange();
+	TPCGValueRange<int32> SeedRange = PointData->GetSeedValueRange();
+	TPCGValueRange<int64> MetadataEntryRange = PointData->GetMetadataEntryValueRange();
+	
 	for (int i = 0; i < PointCount; i++)
 	{
-		FVector Position(i, 0, 0);
-		Points.Emplace(FTransform(Position), 1, i);
+		const FVector Position(i, 0, 0);
+		TransformRange[i] = FTransform(Position);
+		SeedRange[i] = i;
 
 		FName AttributeName;
 		FSoftObjectPath ObjectPath;
@@ -462,7 +470,7 @@ bool FPCGStaticMeshSpawnerByAttributeTest::RunTest(const FString& Parameters)
 			ObjectPath = (i < 85) ? CubePath : CylinderPath; // half are CubePath, half are CylinderPath
 		}
 		
-		UPCGMetadataAccessorHelpers::SetStringAttribute(Points[i], PointData->Metadata, AttributeName, ObjectPath.ToString());
+		UPCGMetadataAccessorHelpers::SetStringAttributeByMetadataKey(MetadataEntryRange[i], PointData->Metadata, AttributeName, ObjectPath.ToString());
 	}
 	
 	FPCGTaggedData& TaggedDataPoints = TestData.InputData.TaggedData.Emplace_GetRef(FPCGTaggedData());
@@ -513,14 +521,20 @@ bool FPCGStaticMeshSpawnerWeightedTest::RunTest(const FString& Parameters)
 
 	// compose the test data
 	const int PointCount = 1000;
-	TObjectPtr<UPCGPointData> PointData = PCGTestsCommon::CreateEmptyPointData();
+	TObjectPtr<UPCGBasePointData> PointData = PCGTestsCommon::CreateEmptyBasePointData();
 	PointData->TargetActor = TestData.TestActor;
 
-	TArray<FPCGPoint>& Points = PointData->GetMutablePoints();
+	PointData->SetNumPoints(PointCount);
+	PointData->SetDensity(1);
+
+	TPCGValueRange<FTransform> TransformRange = PointData->GetTransformValueRange();
+	TPCGValueRange<int32> SeedRange = PointData->GetSeedValueRange();
+	
 	for (int i = 0; i < PointCount; i++)
 	{
-		FVector Position(i, 0, 0);
-		Points.Emplace(FTransform(Position), 1, i);
+		const FVector Position(i, 0, 0);
+		TransformRange[i] = FTransform(Position);
+		SeedRange[i] = i;
 	}
 	
 	FPCGTaggedData& TaggedDataPoints = TestData.InputData.TaggedData.Emplace_GetRef(FPCGTaggedData());
@@ -591,12 +605,10 @@ bool FPCGStaticMeshSpawnerWeightedByCategoryTest::RunTest(const FString& Paramet
 	const FString InvalidRockType = TEXT("InvalidRockType");
 
 	// compose the test data
-	TObjectPtr<UPCGPointData> PointData = PCGTestsCommon::CreateEmptyPointData();
+	TObjectPtr<UPCGBasePointData> PointData = PCGTestsCommon::CreateEmptyBasePointData();
 	PointData->TargetActor = TestData.TestActor;
 	PointData->Metadata->CreateStringAttribute(AttributeName, DefaultRockType, false);
 
-	TArray<FPCGPoint>& Points = PointData->GetMutablePoints();
-
 	// first 1k are igneous, next 1.5k are sedimenatary, next 500 are metamorphic, then last 500 are Default
 	const int PointCount = 3500;
 	const int IgneousCategorySize = 1000;
@@ -604,10 +616,18 @@ bool FPCGStaticMeshSpawnerWeightedByCategoryTest::RunTest(const FString& Paramet
 	const int MetamorphicCategorySize = 500;
 	const int DefaultCategorySize = 500;
 
+	PointData->SetNumPoints(PointCount);
+	PointData->SetDensity(1);
+
+	TPCGValueRange<FTransform> TransformRange = PointData->GetTransformValueRange();
+	TPCGValueRange<int32> SeedRange = PointData->GetSeedValueRange();
+	TPCGValueRange<int64> MetadataEntryRange = PointData->GetMetadataEntryValueRange();
+
 	for (int i = 0; i < PointCount; i++)
 	{
-		FVector Position(i, 0, 0);
-		Points.Emplace(FTransform(Position), 1, i);
+		const FVector Position(i, 0, 0);
+		TransformRange[i] = FTransform(Position);
+		SeedRange[i] = i;
 
 		FString AttributeValue;
 		
@@ -628,7 +648,7 @@ bool FPCGStaticMeshSpawnerWeightedByCategoryTest::RunTest(const FString& Paramet
 			AttributeValue = DefaultRockType;
 		} 
 		
-		UPCGMetadataAccessorHelpers::SetStringAttribute(Points[i], PointData->Metadata, AttributeName, AttributeValue);
+		UPCGMetadataAccessorHelpers::SetStringAttributeByMetadataKey(MetadataEntryRange[i], PointData->Metadata, AttributeName, AttributeValue);
 	}
 	
 	FPCGTaggedData& TaggedDataPoints = TestData.InputData.TaggedData.Emplace_GetRef(FPCGTaggedData());
diff --git a/Engine/Plugins/PCG/Source/PCG/Private/Tests/Elements/PCGTransformPointsTest.cpp b/Engine/Plugins/PCG/Source/PCG/Private/Tests/Elements/PCGTransformPointsTest.cpp
--- a/Engine/Plugins/PCG/Source/PCG/Private/Tests/Elements/PCGTransformPointsTest.cpp
+++ b/Engine/Plugins/PCG/Source/PCG/Private/Tests/Elements/PCGTransformPointsTest.cpp
@@ -20,24 +20,34 @@ bool FPCGTransformPointsTest::RunTest(const FString& Parameters)
 	UPCGTransformPointsSettings* Settings = CastChecked<UPCGTransformPointsSettings>(TestData.Settings);
 	FPCGElementPtr TransformPointsElement = TestData.Settings->GetElement();
 
-	TObjectPtr<UPCGPointData> EmptyData = PCGTestsCommon::CreateEmptyPointData();
-	TObjectPtr<UPCGPointData> SimpleData = PCGTestsCommon::CreateEmptyPointData();
-	TObjectPtr<UPCGPointData> ComplexData = PCGTestsCommon::CreateEmptyPointData();
-
-	// SimplePoints are empty transforms, ComplexPoints have unique location, rotation, and scale
-	TArray<FPCGPoint>& SimplePoints = SimpleData->GetMutablePoints();
-	TArray<FPCGPoint>& ComplexPoints = ComplexData->GetMutablePoints();
+	TObjectPtr<UPCGBasePointData> EmptyData = PCGTestsCommon::CreateEmptyBasePointData();
+	TObjectPtr<UPCGBasePointData> SimpleData = PCGTestsCommon::CreateEmptyBasePointData();
+	TObjectPtr<UPCGBasePointData> ComplexData = PCGTestsCommon::CreateEmptyBasePointData();
 
 	FRandomStream RandomSource(PCGHelpers::ComputeSeed(TestData.Seed));
 	const int PointCount = 100;
+
+	// SimplePoints are empty transforms
+	SimpleData->SetNumPoints(PointCount);
+	SimpleData->SetDensity(1);
+	TPCGValueRange<int32> SimpleSeedRange = SimpleData->GetSeedValueRange();
+
+	// ComplexPoints have unique location, rotation, and scale
+	ComplexData->SetNumPoints(PointCount);
+	ComplexData->SetDensity(1);
+	TPCGValueRange<FTransform> ComplexTransformRange = ComplexData->GetTransformValueRange();
+	TPCGValueRange<int32> ComplexSeedRange = ComplexData->GetSeedValueRange();
+
 	for (int I = 0; I < PointCount; ++I)
 	{
-		SimplePoints.Emplace(FTransform(), 1, I);
+		SimpleSeedRange[I] = I;
 
 		const FRotator RandomRotation = FRotator(RandomSource.FRandRange(0.f, 360.f));
 		const FVector RandomLocation = RandomSource.VRand();
 		const FVector RandomScale = RandomSource.VRand().GetAbs();
-		ComplexPoints.Emplace(FTransform(RandomRotation, RandomLocation, RandomScale), 1, I);
+
+		ComplexSeedRange[I] = I;
+		ComplexTransformRange[I] = FTransform(RandomRotation, RandomLocation, RandomScale);
 	}
 
 	FPCGTaggedData& EmptyTaggedData = TestData.InputData.TaggedData.Emplace_GetRef(FPCGTaggedData());
@@ -87,7 +97,7 @@ bool FPCGTransformPointsTest::RunTest(const FString& Parameters)
 			const UPCGSpatialData* InSpatialData = Cast<UPCGSpatialData>(Input.Data);
 			check(InSpatialData);
 
-			const UPCGPointData* InPointData = InSpatialData->ToPointData(Context.Get());
+			const UPCGBasePointData* InPointData = InSpatialData->ToBasePointData(Context.Get());
 			check(InPointData);
 
 			if (!TestTrue("Valid output data", Output.Data != nullptr))
@@ -104,7 +114,7 @@ bool FPCGTransformPointsTest::RunTest(const FString& Parameters)
 				continue;
 			}
 
-			const UPCGPointData* OutPointData = OutSpatialData->ToPointData(Context.Get());
+			const UPCGBasePointData* OutPointData = OutSpatialData->ToBasePointData(Context.Get());
 
 			if (!TestNotNull("Valid output PointData", OutPointData))
 			{
@@ -112,19 +122,19 @@ bool FPCGTransformPointsTest::RunTest(const FString& Parameters)
 				continue;
 			}
 
-			const TArray<FPCGPoint>& InPoints = InPointData->GetPoints();
-			const TArray<FPCGPoint>& OutPoints = OutPointData->GetPoints();
-
-			if (!TestEqual("Input and output point counts match", InPoints.Num(), OutPoints.Num()))
+			if (!TestEqual("Input and output point counts match", InPointData->GetNumPoints(), OutPointData->GetNumPoints()))
 			{ 
 				bTestPassed = false;
 				continue;
 			}
 
-			for (int PointIndex = 0; PointIndex < InPoints.Num(); ++PointIndex)
+			const FConstPCGPointValueRanges InRanges(InPointData);
+			const FConstPCGPointValueRanges OutRanges(OutPointData);
+
+			for (int PointIndex = 0; PointIndex < InPointData->GetNumPoints(); ++PointIndex)
 			{
-				const FPCGPoint& InPoint = InPoints[PointIndex];
-				const FPCGPoint& OutPoint = OutPoints[PointIndex];
+				const FPCGPoint InPoint = InRanges.GetPoint(PointIndex);
+				const FPCGPoint OutPoint = OutRanges.GetPoint(PointIndex);
 
 				FPCGPoint RepositionedPoint = InPoint;
 				RepositionedPoint.Transform = OutPoint.Transform;
diff --git a/Engine/Plugins/PCG/Source/PCG/Private/Tests/Metadata/PCGMetadataFlattenTest.cpp b/Engine/Plugins/PCG/Source/PCG/Private/Tests/Metadata/PCGMetadataFlattenTest.cpp
--- a/Engine/Plugins/PCG/Source/PCG/Private/Tests/Metadata/PCGMetadataFlattenTest.cpp
+++ b/Engine/Plugins/PCG/Source/PCG/Private/Tests/Metadata/PCGMetadataFlattenTest.cpp
@@ -19,23 +19,24 @@ bool FPCGMetadataFlatten::RunTest(const FString& Parameters)
 	static const FName Attribute2Name = TEXT("StringAttr");
 	static const FName Attribute3Name = TEXT("IntAttr");
 
-	UPCGPointData* RootPointData = NewObject<UPCGPointData>();
-	TArray<FPCGPoint>& Points = RootPointData->GetMutablePoints();
+	UPCGBasePointData* RootPointData = PCGTestsCommon::CreateEmptyBasePointData();
 	FPCGMetadataAttribute<float>* Attribute1 = RootPointData->Metadata->CreateAttribute<float>(Attribute1Name, -0.1f, true, true);
 	FPCGMetadataAttribute<FString>* Attribute2 = RootPointData->Metadata->CreateAttribute<FString>(Attribute2Name, TEXT("Default"), true, true);
 
 	check(Attribute1 && Attribute2);
 	
-	for (int32 i = 0; i < 10; ++i)
+	RootPointData->SetNumPoints(10);
+
+	TPCGValueRange<int64> MetadataEntryRange = RootPointData->GetMetadataEntryValueRange();
+	for (int32 i = 0; i < RootPointData->GetNumPoints(); ++i)
 	{
-		FPCGPoint& Point = Points.Emplace_GetRef();
-		RootPointData->Metadata->InitializeOnSet(Point.MetadataEntry);
-		Attribute1->SetValue(Point.MetadataEntry, static_cast<float>(i) * 0.1f);
+		RootPointData->Metadata->InitializeOnSet(MetadataEntryRange[i]);
+		Attribute1->SetValue(MetadataEntryRange[i], static_cast<float>(i) * 0.1f);
 
 		if (i % 2 == 0)
 		{
 			// Will be either "0" or "1"
-			Attribute2->SetValue(Point.MetadataEntry, FString::Printf(TEXT("%d"), i % 4));
+			Attribute2->SetValue(MetadataEntryRange[i], FString::Printf(TEXT("%d"), i % 4));
 		}
 	}
 
@@ -44,7 +45,7 @@ bool FPCGMetadataFlatten::RunTest(const FString& Parameters)
 	// for Attribute 2: ["0", "Default", "2", "Default", "0", ...]
 	UTEST_EQUAL("RootMetadata has 10 entries", RootPointData->Metadata->GetItemCountForChild(), 10);
 
-	UPCGPointData* FirstChildPointData = Cast<UPCGPointData>(RootPointData->DuplicateData(nullptr));
+	UPCGBasePointData* FirstChildPointData = Cast<UPCGBasePointData>(RootPointData->DuplicateData(nullptr));
 	Attribute1 = FirstChildPointData->Metadata->GetMutableTypedAttribute<float>(Attribute1Name);
 	Attribute2 = FirstChildPointData->Metadata->GetMutableTypedAttribute<FString>(Attribute2Name);
 
@@ -56,21 +57,20 @@ bool FPCGMetadataFlatten::RunTest(const FString& Parameters)
 	FPCGMetadataAttribute<int>* Attribute3 = FirstChildPointData->Metadata->CreateAttribute<int>(Attribute3Name, -1, true, true);
 	check(Attribute3);
 
-	TArray<FPCGPoint>& FirstChildPoints = FirstChildPointData->GetMutablePoints();
-	
-	UTEST_EQUAL("First child has 10 points", FirstChildPoints.Num(), 10);
+	UTEST_EQUAL("First child has 10 points", FirstChildPointData->GetNumPoints(), 10);
 
-	for (int32 i = 0; i < 10; ++i)
+	TPCGValueRange<int64> FirstChildMetadataEntryRange = FirstChildPointData->GetMetadataEntryValueRange();
+	for (int32 i = 0; i < FirstChildPointData->GetNumPoints(); ++i)
 	{
-		FirstChildPointData->Metadata->InitializeOnSet(FirstChildPoints[i].MetadataEntry);
-		Attribute1->SetValue(FirstChildPoints[i].MetadataEntry, static_cast<float>(i) * 1.1f);
+		FirstChildPointData->Metadata->InitializeOnSet(FirstChildMetadataEntryRange[i]);
+		Attribute1->SetValue(FirstChildMetadataEntryRange[i], static_cast<float>(i) * 1.1f);
 
 		if (i % 2 == 0)
 		{
-			Attribute2->SetValue(FirstChildPoints[i].MetadataEntry, TEXT("0"));
+			Attribute2->SetValue(FirstChildMetadataEntryRange[i], TEXT("0"));
 		}
 
-		Attribute3->SetValue(FirstChildPoints[i].MetadataEntry, i);
+		Attribute3->SetValue(FirstChildMetadataEntryRange[i], i);
 	}
 
 	// At the end of the second set, metadata has 20 entries, and values for each points are
@@ -80,25 +80,30 @@ bool FPCGMetadataFlatten::RunTest(const FString& Parameters)
 	UTEST_EQUAL("FirstChildMetadata has 20 entries", FirstChildPointData->Metadata->GetItemCountForChild(), 20);
 
 	// Second child is keeping the metadata entry for even numbers and reset metadata entry for the rest (point back to default). Also override Attribute 3 values for even numbers
-	UPCGPointData* SecondChildPointData = NewObject<UPCGPointData>();
-	SecondChildPointData->InitializeFromData(FirstChildPointData);
-	TArray<FPCGPoint>& SecondChildPoints = SecondChildPointData->GetMutablePoints();
+	UPCGBasePointData* SecondChildPointData = PCGTestsCommon::CreateEmptyBasePointData();
 
+	FPCGInitializeFromDataParams InitializeFromDataParams(FirstChildPointData);
+	InitializeFromDataParams.bInheritSpatialData = false;
+	SecondChildPointData->InitializeFromDataWithParams(InitializeFromDataParams);
+	
 	Attribute3 = SecondChildPointData->Metadata->GetMutableTypedAttribute<int>(Attribute3Name);
 	UTEST_NOT_NULL("Attributes exists in first child", Attribute3);
 	check(Attribute3);
 
-	for (int32 i = 0; i < 10; ++i)
+	SecondChildPointData->SetNumPoints(FirstChildPointData->GetNumPoints());
+	UPCGBasePointData::SetPoints(FirstChildPointData, SecondChildPointData, {}, /*bCopyAll=*/true);
+
+	TPCGValueRange<int64> SecondChildMetadataEntryRange = SecondChildPointData->GetMetadataEntryValueRange();
+	for (int32 i = 0; i < SecondChildPointData->GetNumPoints(); ++i)
 	{
-		FPCGPoint& Point = SecondChildPoints.Add_GetRef(FirstChildPoints[i]);
 		if (i % 2 == 0)
 		{
-			SecondChildPointData->Metadata->InitializeOnSet(Point.MetadataEntry);
-			Attribute3->SetValue(Point.MetadataEntry, 10 * i);
+			SecondChildPointData->Metadata->InitializeOnSet(SecondChildMetadataEntryRange[i]);
+			Attribute3->SetValue(SecondChildMetadataEntryRange[i], 10 * i);
 		}
 		else
 		{
-			Point.MetadataEntry = PCGInvalidEntryKey;
+			SecondChildMetadataEntryRange[i] = PCGInvalidEntryKey;
 		}
 	}
 
@@ -109,7 +114,7 @@ bool FPCGMetadataFlatten::RunTest(const FString& Parameters)
 	UTEST_EQUAL("SecondChildMetadata has 25 entries", SecondChildPointData->Metadata->GetItemCountForChild(), 25);
 
 	// For final set, duplicate the data and flatten it
-	UPCGPointData* FinalPointData = Cast<UPCGPointData>(SecondChildPointData->DuplicateData(nullptr));
+	UPCGBasePointData* FinalPointData = Cast<UPCGBasePointData>(SecondChildPointData->DuplicateData(nullptr));
 	FinalPointData->Flatten();
 
 	Attribute1 = FinalPointData->Metadata->GetMutableTypedAttribute<float>(Attribute1Name);
@@ -123,27 +128,28 @@ bool FPCGMetadataFlatten::RunTest(const FString& Parameters)
 	UTEST_EQUAL("Attribute 1 has 5 values", Attribute1->GetValueKeyOffsetForChild(), 5);
 	UTEST_EQUAL("Attribute 2 has 1 value", Attribute2->GetValueKeyOffsetForChild(), 1);
 	UTEST_EQUAL("Attribute 3 has 5 values", Attribute3->GetValueKeyOffsetForChild(), 5);
-	UTEST_EQUAL("Final point data has 10 points", FinalPointData->GetPoints().Num(), 10);
+	UTEST_EQUAL("Final point data has 10 points", FinalPointData->GetNumPoints(), 10);
 
 	// Validate the values
-	for (int32 i = 0; i < 10; ++i)
+	const TConstPCGValueRange<int64> FinalMetadataEntryRange = FinalPointData->GetConstMetadataEntryValueRange();
+	for (int32 i = 0; i < FinalPointData->GetNumPoints(); ++i)
 	{
 		auto FormatMessage = [i](const FString& Msg) { return FString::Printf(TEXT("Point %d: %s"), i, *Msg); };
 
-		const FPCGPoint& Point = FinalPointData->GetPoints()[i];
+		const int64 MetadataEntry = FinalMetadataEntryRange[i];
 		if (i % 2 == 0)
 		{
-			UTEST_EQUAL(*FormatMessage(TEXT("Valid metadata entry")), Point.MetadataEntry, static_cast<PCGMetadataEntryKey>(i / 2));
-			UTEST_EQUAL(*FormatMessage(TEXT("Attribute 1 value")), Attribute1->GetValueFromItemKey(Point.MetadataEntry), i * 1.1f);
-			UTEST_EQUAL(*FormatMessage(TEXT("Attribute 2 value")), Attribute2->GetValueFromItemKey(Point.MetadataEntry), TEXT("0"));
-			UTEST_EQUAL(*FormatMessage(TEXT("Attribute 3 value")), Attribute3->GetValueFromItemKey(Point.MetadataEntry), i * 10);
+			UTEST_EQUAL(*FormatMessage(TEXT("Valid metadata entry")), MetadataEntry, static_cast<PCGMetadataEntryKey>(i / 2));
+			UTEST_EQUAL(*FormatMessage(TEXT("Attribute 1 value")), Attribute1->GetValueFromItemKey(MetadataEntry), i * 1.1f);
+			UTEST_EQUAL(*FormatMessage(TEXT("Attribute 2 value")), Attribute2->GetValueFromItemKey(MetadataEntry), TEXT("0"));
+			UTEST_EQUAL(*FormatMessage(TEXT("Attribute 3 value")), Attribute3->GetValueFromItemKey(MetadataEntry), i * 10);
 		}
 		else
 		{
-			UTEST_EQUAL(*FormatMessage(TEXT("Invalid metadata entry")), Point.MetadataEntry, PCGInvalidEntryKey);
-			UTEST_EQUAL(*FormatMessage(TEXT("Attribute 1 value")), Attribute1->GetValueFromItemKey(Point.MetadataEntry), -0.1f);
-			UTEST_EQUAL(*FormatMessage(TEXT("Attribute 2 value")), Attribute2->GetValueFromItemKey(Point.MetadataEntry), TEXT("Default"));
-			UTEST_EQUAL(*FormatMessage(TEXT("Attribute 3 value")), Attribute3->GetValueFromItemKey(Point.MetadataEntry), -1);
+			UTEST_EQUAL(*FormatMessage(TEXT("Invalid metadata entry")), MetadataEntry, PCGInvalidEntryKey);
+			UTEST_EQUAL(*FormatMessage(TEXT("Attribute 1 value")), Attribute1->GetValueFromItemKey(MetadataEntry), -0.1f);
+			UTEST_EQUAL(*FormatMessage(TEXT("Attribute 2 value")), Attribute2->GetValueFromItemKey(MetadataEntry), TEXT("Default"));
+			UTEST_EQUAL(*FormatMessage(TEXT("Attribute 3 value")), Attribute3->GetValueFromItemKey(MetadataEntry), -1);
 		}
 	}
 
diff --git a/Engine/Plugins/PCG/Source/PCG/Private/Tests/Metadata/PCGMetadataTest.cpp b/Engine/Plugins/PCG/Source/PCG/Private/Tests/Metadata/PCGMetadataTest.cpp
--- a/Engine/Plugins/PCG/Source/PCG/Private/Tests/Metadata/PCGMetadataTest.cpp
+++ b/Engine/Plugins/PCG/Source/PCG/Private/Tests/Metadata/PCGMetadataTest.cpp
@@ -23,22 +23,22 @@ bool FPCGMetadataTest_Inherit::RunTest(const FString& Parameters)
 {
 	const FName AttributeName = TEXT("MyAttr");
 	
-	UPCGPointData* PointData3 = nullptr;
+	UPCGBasePointData* PointData3 = nullptr;
 
 	{
-		TObjectPtr<UPCGPointData> PointData = NewObject<UPCGPointData>();
-		TArray<FPCGPoint>& Points = PointData->GetMutablePoints();
-		Points.SetNum(5);
+		TObjectPtr<UPCGBasePointData> PointData = PCGTestsCommon::CreateEmptyBasePointData();
+		PointData->SetNumPoints(5);
 		UPCGMetadata* PointMetadata = PointData->MutableMetadata();
 		FPCGMetadataAttribute<int32>* Attribute = PointMetadata->CreateAttribute<int32>(AttributeName, 0, true, false);
 	
-		for (int32 i = 0 ; i < Points.Num(); ++i)
+		TPCGValueRange<int64> MetadataEntryRange = PointData->GetMetadataEntryValueRange();
+		for (int32 i = 0 ; i < PointData->GetNumPoints(); ++i)
 		{
-			PointMetadata->InitializeOnSet(Points[i].MetadataEntry);
-			Attribute->SetValue(Points[i].MetadataEntry, i);
+			PointMetadata->InitializeOnSet(MetadataEntryRange[i]);
+			Attribute->SetValue(MetadataEntryRange[i], i);
 		}
 	
-		TObjectPtr<UPCGPointData> PointData2 = NewObject<UPCGPointData>();
+		TObjectPtr<UPCGBasePointData> PointData2 = PCGTestsCommon::CreateEmptyBasePointData();
 		UPCGMetadata* PointMetadata2 = PointData2->MutableMetadata();
 		PointMetadata2->Initialize(PointMetadata);
 
@@ -47,23 +47,25 @@ bool FPCGMetadataTest_Inherit::RunTest(const FString& Parameters)
 
 		FPCGMetadataAttribute<int32>* Attribute2 = PointMetadata2->GetMutableTypedAttribute<int32>(AttributeName);
 
-		TArray<FPCGPoint>& Points2 = PointData2->GetMutablePoints();
-		Points2 = Points;
-		Points2.SetNum(10);
+		UPCGBasePointData::SetPoints(PointData, PointData2, {}, /*bCopyAll=*/true);
+		PointData2->SetNumPoints(10);
 	
-		for (int32 i = Points.Num() ; i < Points2.Num(); ++i)
+		TPCGValueRange<int64> MetadataEntryRange2 = PointData2->GetMetadataEntryValueRange();
+		for (int32 i = PointData->GetNumPoints() ; i < PointData2->GetNumPoints(); ++i)
 		{
-			PointMetadata2->InitializeOnSet(Points2[i].MetadataEntry);
-			Attribute2->SetValue(Points2[i].MetadataEntry, i);
+			PointMetadata2->InitializeOnSet(MetadataEntryRange2[i]);
+			Attribute2->SetValue(MetadataEntryRange2[i], i);
 		}
 		
-		PointData3 = NewObject<UPCGPointData>();
+		PointData3 = PCGTestsCommon::CreateEmptyBasePointData();
 		UPCGMetadata* PointMetadata3 = PointData3->MutableMetadata();
 		PointMetadata3->Initialize(PointMetadata2);
 
 		UTEST_EQUAL("Initialization inherited", PointMetadata3->GetLocalItemCount(), 0);
-		
-		PointData3->GetMutablePoints() = Points2;
+
+		UPCGBasePointData::SetPoints(PointData2, PointData3, {}, /*bCopyAll=*/true);
+
+		UTEST_EQUAL("Same number of points", PointData2->GetNumPoints(), PointData3->GetNumPoints());
 	}
 
 	// Force a GC run
@@ -74,14 +76,14 @@ bool FPCGMetadataTest_Inherit::RunTest(const FString& Parameters)
 
 	// Verify that all the values are readable from the metadata
 	check(PointData3);
-	const TArray<FPCGPoint>& Points3 = PointData3->GetPoints();
 	const FPCGMetadataAttribute<int32>* Attribute = PointData3->ConstMetadata()->GetConstTypedAttribute<int32>(AttributeName);
 	UTEST_NOT_NULL("Attribute exists", Attribute);
 	check(Attribute);
-	
-	for (int32 i = 0; i < Points3.Num(); ++i)
+
+	TConstPCGValueRange<int64> MetadataEntryRange3 = PointData3->GetConstMetadataEntryValueRange();
+	for (int32 i = 0; i < PointData3->GetNumPoints(); ++i)
 	{
-		UTEST_EQUAL("Values are the same", Attribute->GetValueFromItemKey(Points3[i].MetadataEntry), i);
+		UTEST_EQUAL("Values are the same", Attribute->GetValueFromItemKey(MetadataEntryRange3[i]), i);
 	}
 	
 	return true;
@@ -91,22 +93,23 @@ bool FPCGMetadataTest_InheritCopy::RunTest(const FString& Parameters)
 {
 	const FName AttributeName = TEXT("MyAttr");
 	
-	UPCGPointData* PointData3 = nullptr;
+	UPCGBasePointData* PointData3 = nullptr;
 
 	{
-		TObjectPtr<UPCGPointData> PointData = NewObject<UPCGPointData>();
-		TArray<FPCGPoint>& Points = PointData->GetMutablePoints();
-		Points.SetNum(5);
+		TObjectPtr<UPCGBasePointData> PointData = PCGTestsCommon::CreateEmptyBasePointData();
+		PointData->SetNumPoints(5);
+		
 		UPCGMetadata* PointMetadata = PointData->MutableMetadata();
 		FPCGMetadataAttribute<int32>* Attribute = PointMetadata->CreateAttribute<int32>(AttributeName, 0, true, false);
 	
-		for (int32 i = 0 ; i < Points.Num(); ++i)
+		TPCGValueRange<int64> MetadataEntryRange = PointData->GetMetadataEntryValueRange();
+		for (int32 i = 0 ; i < PointData->GetNumPoints(); ++i)
 		{
-			PointMetadata->InitializeOnSet(Points[i].MetadataEntry);
-			Attribute->SetValue(Points[i].MetadataEntry, i);
+			PointMetadata->InitializeOnSet(MetadataEntryRange[i]);
+			Attribute->SetValue(MetadataEntryRange[i], i);
 		}
 	
-		TObjectPtr<UPCGPointData> PointData2 = NewObject<UPCGPointData>();
+		TObjectPtr<UPCGBasePointData> PointData2 = PCGTestsCommon::CreateEmptyBasePointData();
 		UPCGMetadata* PointMetadata2 = PointData2->MutableMetadata();
 		PointMetadata2->Initialize(PointMetadata);
 
@@ -115,24 +118,26 @@ bool FPCGMetadataTest_InheritCopy::RunTest(const FString& Parameters)
 
 		FPCGMetadataAttribute<int32>* Attribute2 = PointMetadata2->GetMutableTypedAttribute<int32>(AttributeName);
 
-		TArray<FPCGPoint>& Points2 = PointData2->GetMutablePoints();
-		Points2 = Points;
-		Points2.SetNum(10);
+		UPCGBasePointData::SetPoints(PointData, PointData2, {}, /*bCopyAll=*/true);
+		PointData2->SetNumPoints(10);
 	
-		for (int32 i = Points.Num() ; i < Points2.Num(); ++i)
+		TPCGValueRange<int64> MetadataEntryRange2 = PointData2->GetMetadataEntryValueRange();
+		for (int32 i = PointData->GetNumPoints(); i < PointData2->GetNumPoints(); ++i)
 		{
-			PointMetadata2->InitializeOnSet(Points2[i].MetadataEntry);
-			Attribute2->SetValue(Points2[i].MetadataEntry, i);
+			PointMetadata2->InitializeOnSet(MetadataEntryRange2[i]);
+			Attribute2->SetValue(MetadataEntryRange2[i], i);
 		}
 		
-		PointData3 = NewObject<UPCGPointData>();
+		PointData3 = PCGTestsCommon::CreateEmptyBasePointData();
 		UPCGMetadata* PointMetadata3 = PointData3->MutableMetadata();
 		PointMetadata3->InitializeAsCopy(FPCGMetadataInitializeParams(PointMetadata2));
 
 		// Verify that the initialization made a copy
 		UTEST_EQUAL("Initialization made a copy", PointMetadata3->GetLocalItemCount(), 5);
 		
-		PointData3->GetMutablePoints() = Points2;
+		UPCGBasePointData::SetPoints(PointData2, PointData3, {}, /*bCopyAll=*/true);
+
+		UTEST_EQUAL("Same number of points", PointData2->GetNumPoints(), PointData3->GetNumPoints());
 	}
 
 	// Force a GC run
@@ -143,14 +148,14 @@ bool FPCGMetadataTest_InheritCopy::RunTest(const FString& Parameters)
 
 	// Verify that all the values are readable from the metadata
 	check(PointData3);
-	const TArray<FPCGPoint>& Points3 = PointData3->GetPoints();
 	const FPCGMetadataAttribute<int32>* Attribute = PointData3->ConstMetadata()->GetConstTypedAttribute<int32>(AttributeName);
 	UTEST_NOT_NULL("Attribute exists", Attribute);
 	check(Attribute);
-	
-	for (int32 i = 0; i < Points3.Num(); ++i)
+
+	TConstPCGValueRange<int64> MetadataEntryRange3 = PointData3->GetConstMetadataEntryValueRange();
+	for (int32 i = 0; i < PointData3->GetNumPoints(); ++i)
 	{
-		UTEST_EQUAL("Values are the same", Attribute->GetValueFromItemKey(Points3[i].MetadataEntry), i);
+		UTEST_EQUAL("Values are the same", Attribute->GetValueFromItemKey(MetadataEntryRange3[i]), i);
 	}
 	
 	return true;
diff --git a/Engine/Plugins/PCG/Source/PCG/Private/Tests/PCGTestsCommon.cpp b/Engine/Plugins/PCG/Source/PCG/Private/Tests/PCGTestsCommon.cpp
--- a/Engine/Plugins/PCG/Source/PCG/Private/Tests/PCGTestsCommon.cpp
+++ b/Engine/Plugins/PCG/Source/PCG/Private/Tests/PCGTestsCommon.cpp
@@ -178,6 +178,42 @@ namespace PCGTestsCommon
 		}
 	}
 
+	UPCGBasePointData* CreateBasePointData()
+	{
+		if (CVarPCGEnablePointArrayData.GetValueOnAnyThread())
+		{
+			return CreatePointData<UPCGPointArrayData>();
+		}
+		else
+		{
+			return CreatePointData<UPCGPointData>();
+		}
+	}
+
+	UPCGBasePointData* CreateBasePointData(const FVector& InLocation)
+	{
+		if (CVarPCGEnablePointArrayData.GetValueOnAnyThread())
+		{
+			return CreatePointData<UPCGPointArrayData>(InLocation);
+		}
+		else
+		{
+			return CreatePointData<UPCGPointData>(InLocation);
+		}
+	}
+
+	UPCGBasePointData* CreateRandomBasePointData(int32 PointCount, int32 Seed, bool bRandomDensity)
+	{
+		if (CVarPCGEnablePointArrayData.GetValueOnAnyThread())
+		{
+			return CreateRandomPointData<UPCGPointArrayData>(PointCount, Seed, bRandomDensity);
+		}
+		else
+		{
+			return CreateRandomPointData<UPCGPointData>(PointCount, Seed, bRandomDensity);
+		}
+	}
+
 	TArray<FPCGDataCollection> GenerateAllowedData(const FPCGPinProperties& PinProperties)
 	{
 		TArray<FPCGDataCollection> Data;
diff --git a/Engine/Plugins/PCG/Source/PCG/Private/Tests/SpatialAlgo/PCGOctreeQueriesTest.cpp b/Engine/Plugins/PCG/Source/PCG/Private/Tests/SpatialAlgo/PCGOctreeQueriesTest.cpp
--- a/Engine/Plugins/PCG/Source/PCG/Private/Tests/SpatialAlgo/PCGOctreeQueriesTest.cpp
+++ b/Engine/Plugins/PCG/Source/PCG/Private/Tests/SpatialAlgo/PCGOctreeQueriesTest.cpp
@@ -23,16 +23,19 @@ namespace PCGOctreeQueriesTest
 	static constexpr double Distance = 100.0;
 	static constexpr int32 NumPoints = 10;
 
-	UPCGPointData* CreatePointData()
+	UPCGBasePointData* CreatePointData()
 	{
-		UPCGPointData* InputPointData = PCGTestsCommon::CreateEmptyPointData();
-		TArray<FPCGPoint>& InputPoints = InputPointData->GetMutablePoints();
-		InputPoints.Reserve(NumPoints);
+		UPCGBasePointData* InputPointData = PCGTestsCommon::CreateEmptyBasePointData();
+		InputPointData->SetNumPoints(NumPoints);
+		InputPointData->SetDensity(1.0f);
+		InputPointData->SetSeed(42);
+
+		TPCGValueRange<FTransform> TransformRange = InputPointData->GetTransformValueRange();
 
 		for (int32 i = 0; i < NumPoints; ++i)
 		{
 			FVector Location = FVector::OneVector * Distance * i;
-			InputPoints.Emplace(FTransform(Location), 1.0f, 42);
+			TransformRange[i] = FTransform(Location);
 		}
 
 		return InputPointData;
@@ -42,14 +45,13 @@ namespace PCGOctreeQueriesTest
 
 bool FPCGOctreeQueriesTest_Sphere::RunTest(const FString& Parameters)
 {
-	const UPCGPointData* InputPointData = PCGOctreeQueriesTest::CreatePointData();
+	const UPCGBasePointData* InputPointData = PCGOctreeQueriesTest::CreatePointData();
 	check(InputPointData);
 
 	TArray<int32> ExpectedIndexes = { 0, 1, 2 };
 	int32 CountFound = 0;
-	UPCGOctreeQueries::ForEachPointInsideSphere(InputPointData, FVector::ZeroVector, 350.0, [&ExpectedIndexes, &CountFound, &InputPoints = InputPointData->GetPoints()](const FPCGPoint& InPoint, double)
+	UPCGOctreeQueries::ForEachPointInsideSphere(InputPointData, FVector::ZeroVector, 350.0, [&ExpectedIndexes, &CountFound](const UPCGBasePointData* PointData, int32 PointIndex, double)
 	{
-		const int32 PointIndex = static_cast<int32>(&InPoint - InputPoints.GetData());
 		if (ExpectedIndexes.Contains(PointIndex))
 		{
 			++CountFound;
@@ -74,13 +76,11 @@ bool FPCGOctreeQueriesTest_Sphere::RunTest(const FString& Parameters)
 
 bool FPCGOctreeQueriesTest_ClosestPoint::RunTest(const FString& Parameters)
 {
-	const UPCGPointData* InputPointData = PCGOctreeQueriesTest::CreatePointData();
+	const UPCGBasePointData* InputPointData = PCGOctreeQueriesTest::CreatePointData();
 	check(InputPointData);
 
-	const FPCGPoint* ClosestPoint = UPCGOctreeQueries::GetClosestPoint(InputPointData, FVector::ZeroVector, /*bDiscardCenter=*/false, 350.0);
-	UTEST_NOT_NULL("Closest point was found", ClosestPoint);
-
-	const int32 PointIndex = static_cast<int32>(ClosestPoint - InputPointData->GetPoints().GetData());
+	const int32 PointIndex = UPCGOctreeQueries::GetClosestPointIndex(InputPointData, FVector::ZeroVector, /*bDiscardCenter=*/false, 350.0);
+	UTEST_NOT_EQUAL("Closest point was found", PointIndex, (int32)INDEX_NONE);
 
 	UTEST_EQUAL("Closest point is the right index", PointIndex, 0);
 	return true;
@@ -88,13 +88,11 @@ bool FPCGOctreeQueriesTest_ClosestPoint::RunTest(const FString& Parameters)
 
 bool FPCGOctreeQueriesTest_ClosestPointDiscardCenter::RunTest(const FString& Parameters)
 {
-	const UPCGPointData* InputPointData = PCGOctreeQueriesTest::CreatePointData();
+	const UPCGBasePointData* InputPointData = PCGOctreeQueriesTest::CreatePointData();
 	check(InputPointData);
 
-	const FPCGPoint* ClosestPoint = UPCGOctreeQueries::GetClosestPoint(InputPointData, FVector::ZeroVector, /*bDiscardCenter=*/true, 350.0);
-	UTEST_NOT_NULL("Closest point was found", ClosestPoint);
-
-	const int32 PointIndex = static_cast<int32>(ClosestPoint - InputPointData->GetPoints().GetData());
+	const int32 PointIndex = UPCGOctreeQueries::GetClosestPointIndex(InputPointData, FVector::ZeroVector, /*bDiscardCenter=*/true, 350.0);
+	UTEST_NOT_EQUAL("Closest point was found", PointIndex, (int32)INDEX_NONE);
 
 	UTEST_EQUAL("Closest point is the right index", PointIndex, 1);
 	return true;
@@ -102,13 +100,11 @@ bool FPCGOctreeQueriesTest_ClosestPointDiscardCenter::RunTest(const FString& Par
 
 bool FPCGOctreeQueriesTest_ClosestPointFromOtherPoint::RunTest(const FString& Parameters)
 {
-	const UPCGPointData* InputPointData = PCGOctreeQueriesTest::CreatePointData();
+	const UPCGBasePointData* InputPointData = PCGOctreeQueriesTest::CreatePointData();
 	check(InputPointData);
 
-	const FPCGPoint* ClosestPoint = UPCGOctreeQueries::GetClosestPointFromOtherPoint(InputPointData, InputPointData->GetPoints()[0], 350.0);
-	UTEST_NOT_NULL("Closest point was found", ClosestPoint);
-
-	const int32 PointIndex = static_cast<int32>(ClosestPoint - InputPointData->GetPoints().GetData());
+	const int32 PointIndex = UPCGOctreeQueries::GetClosestPointIndexFromOtherPointIndex(InputPointData, 0, 350.0);
+	UTEST_NOT_EQUAL("Closest point was found", PointIndex, (int32)INDEX_NONE);
 
 	UTEST_EQUAL("Closest point is the right index", PointIndex, 1);
 	return true;
@@ -116,13 +112,11 @@ bool FPCGOctreeQueriesTest_ClosestPointFromOtherPoint::RunTest(const FString& Pa
 
 bool FPCGOctreeQueriesTest_FarthestPoint::RunTest(const FString& Parameters)
 {
-	const UPCGPointData* InputPointData = PCGOctreeQueriesTest::CreatePointData();
+	const UPCGBasePointData* InputPointData = PCGOctreeQueriesTest::CreatePointData();
 	check(InputPointData);
 
-	const FPCGPoint* FarthestPoint = UPCGOctreeQueries::GetFarthestPoint(InputPointData, FVector::ZeroVector, 350.0);
-	UTEST_NOT_NULL("Farthest point was found", FarthestPoint);
-
-	const int32 PointIndex = static_cast<int32>(FarthestPoint - InputPointData->GetPoints().GetData());
+	const int32 PointIndex = UPCGOctreeQueries::GetFarthestPointIndex(InputPointData, FVector::ZeroVector, 350.0);
+	UTEST_NOT_EQUAL("Farthest point was found", PointIndex, (int32)INDEX_NONE);
 
 	UTEST_EQUAL("Farthest point is the right index", PointIndex, 2);
 	return true;
@@ -130,13 +124,11 @@ bool FPCGOctreeQueriesTest_FarthestPoint::RunTest(const FString& Parameters)
 
 bool FPCGOctreeQueriesTest_FarthestPointFromOtherPoint::RunTest(const FString& Parameters)
 {
-	const UPCGPointData* InputPointData = PCGOctreeQueriesTest::CreatePointData();
+	const UPCGBasePointData* InputPointData = PCGOctreeQueriesTest::CreatePointData();
 	check(InputPointData);
 
-	const FPCGPoint* FarthestPoint = UPCGOctreeQueries::GetFarthestPointFromOtherPoint(InputPointData, InputPointData->GetPoints()[0], 10000.0);
-	UTEST_NOT_NULL("Farthest point was found", FarthestPoint);
-
-	const int32 PointIndex = static_cast<int32>(FarthestPoint - InputPointData->GetPoints().GetData());
+	const int32 PointIndex = UPCGOctreeQueries::GetFarthestPointIndexFromOtherPointIndex(InputPointData, 0, 10000.0);
+	UTEST_NOT_EQUAL("Farthest point was found", PointIndex, (int32)INDEX_NONE);
 
 	UTEST_EQUAL("Farthest point is the right index", PointIndex, 9);
 	return true;
diff --git a/Engine/Plugins/PCG/Source/PCG/Private/Tests/Subgraphs/PCGCollapseSubgraphTest.cpp b/Engine/Plugins/PCG/Source/PCG/Private/Tests/Subgraphs/PCGCollapseSubgraphTest.cpp
--- a/Engine/Plugins/PCG/Source/PCG/Private/Tests/Subgraphs/PCGCollapseSubgraphTest.cpp
+++ b/Engine/Plugins/PCG/Source/PCG/Private/Tests/Subgraphs/PCGCollapseSubgraphTest.cpp
@@ -269,10 +269,10 @@ bool FPCGCollapseSubgraphCombinationOfCollapses::RunTest(const FString& Paramete
 	PCGTestsCommon::FTestData TestData{};
 
 	// Temporary solution to run a full graph and extract the resulting points.
-	auto ExecuteGraph = [&TestData](UPCGGraph* GraphToExecute, double FirstConstant, double SecondConstant) -> TArray<FPCGPoint>
+	auto ExecuteGraph = [&TestData](UPCGGraph* GraphToExecute, double FirstConstant, double SecondConstant) -> TArray<FTransform>
 	{
 		FPCGGraphExecutor GraphExecutor;
-		TArray<FPCGPoint> OutputPoints;
+		TArray<FTransform> OutputTransforms;
 		bool bDone = false;
 
 		TestData.TestPCGComponent->SetGraphLocal(GraphToExecute);
@@ -284,16 +284,16 @@ bool FPCGCollapseSubgraphCombinationOfCollapses::RunTest(const FString& Paramete
 		// TODO: Post-release GraphExecutor should be setting this on the Component
 		TestData.SetCurrentGenerationTask(TaskId);
 
-		const FPCGTaskId FinalTaskId = GraphExecutor.ScheduleGeneric([&OutputPoints, &GraphExecutor, TaskId, &bDone]() -> bool
+		const FPCGTaskId FinalTaskId = GraphExecutor.ScheduleGeneric([&OutputTransforms, &GraphExecutor, TaskId, &bDone]() -> bool
 		{
 			FPCGDataCollection OutputData{};
 			GraphExecutor.GetOutputData(TaskId, OutputData);
 
 			if (OutputData.TaggedData.Num() == 1)
 			{
-				if (const UPCGPointData* PointData = Cast<UPCGPointData>(OutputData.TaggedData[0].Data))
+				if (const UPCGBasePointData* PointData = Cast<UPCGBasePointData>(OutputData.TaggedData[0].Data))
 				{
-					OutputPoints = PointData->GetPoints();
+					OutputTransforms = PointData->GetTransformsCopy();
 				}
 			}
 
@@ -306,15 +306,15 @@ bool FPCGCollapseSubgraphCombinationOfCollapses::RunTest(const FString& Paramete
 			GraphExecutor.Execute();
 		}
 
-		return OutputPoints;
+		return OutputTransforms;
 	};
 
-	auto ComparePoints = [this](const TArray<FPCGPoint>& InPointsA, const TArray<FPCGPoint>& InPointsB, const FString& Name) -> bool
+	auto CompareTransforms = [this](const TArray<FTransform>& TransformsA, const TArray<FTransform>& TransformsB, const FString& Name) -> bool
 	{
-		UTEST_EQUAL(*FString::Printf(TEXT("%s: Same number of points"), *Name), InPointsB.Num(), InPointsA.Num());
-		for (int i = 0; i < InPointsA.Num(); ++i)
+		UTEST_EQUAL(*FString::Printf(TEXT("%s: Same number of points"), *Name), TransformsB.Num(), TransformsA.Num());
+		for (int i = 0; i < TransformsA.Num(); ++i)
 		{
-			UTEST_EQUAL(*FString::Printf(TEXT("%s: Same position"), *Name), InPointsB[i].Transform.GetLocation(), InPointsA[i].Transform.GetLocation());
+			UTEST_EQUAL(*FString::Printf(TEXT("%s: Same position"), *Name), TransformsB[i].GetLocation(), TransformsA[i].GetLocation());
 		}
 
 		return true;
@@ -326,9 +326,9 @@ bool FPCGCollapseSubgraphCombinationOfCollapses::RunTest(const FString& Paramete
 	constexpr double SecondConstantValue2 = 5000.0;
 
 	// Run 3 times with different graph parameters on each
-	const TArray<FPCGPoint> FirstRunPoints = ExecuteGraph(UnitTestGraph, FirstConstantValue, SecondConstantValue);
-	const TArray<FPCGPoint> SecondRunPoints = ExecuteGraph(UnitTestGraph, FirstConstantValue2, SecondConstantValue);
-	const TArray<FPCGPoint> ThirdRunPoints = ExecuteGraph(UnitTestGraph, FirstConstantValue, SecondConstantValue2);
+	const TArray<FTransform> FirstRunPoints = ExecuteGraph(UnitTestGraph, FirstConstantValue, SecondConstantValue);
+	const TArray<FTransform> SecondRunPoints = ExecuteGraph(UnitTestGraph, FirstConstantValue2, SecondConstantValue);
+	const TArray<FTransform> ThirdRunPoints = ExecuteGraph(UnitTestGraph, FirstConstantValue, SecondConstantValue2);
 
 	// Then do all the 4 nodes collapse combinations. We expect our initial node to have 12 nodes. 12 choose 4 is 495 (which is the sum of ExpectedNumberOfCycles and ExpectedNumberOfSuccess)
 	bool bSuccess = true;
@@ -358,13 +358,13 @@ bool FPCGCollapseSubgraphCombinationOfCollapses::RunTest(const FString& Paramete
 					{
 						Subgraph->SetFlags(RF_Transient);
 
-						const TArray<FPCGPoint> FirstRunCollapsedPoints = ExecuteGraph(OriginalGraph, FirstConstantValue, SecondConstantValue);
-						const TArray<FPCGPoint> SecondRunCollapsedPoints = ExecuteGraph(OriginalGraph, FirstConstantValue2, SecondConstantValue);
-						const TArray<FPCGPoint> ThirdRunCollapsedPoints = ExecuteGraph(OriginalGraph, FirstConstantValue, SecondConstantValue2);
+						const TArray<FTransform> FirstRunCollapsedPoints = ExecuteGraph(OriginalGraph, FirstConstantValue, SecondConstantValue);
+						const TArray<FTransform> SecondRunCollapsedPoints = ExecuteGraph(OriginalGraph, FirstConstantValue2, SecondConstantValue);
+						const TArray<FTransform> ThirdRunCollapsedPoints = ExecuteGraph(OriginalGraph, FirstConstantValue, SecondConstantValue2);
 
-						bool bLocalSuccess = ComparePoints(FirstRunPoints, FirstRunCollapsedPoints, Name + TEXT(" First run"));
-						bLocalSuccess &= ComparePoints(SecondRunPoints, SecondRunCollapsedPoints, Name + TEXT(" Second run"));
-						bLocalSuccess &= ComparePoints(ThirdRunPoints, ThirdRunCollapsedPoints, Name + TEXT(" Third run"));
+						bool bLocalSuccess = CompareTransforms(FirstRunPoints, FirstRunCollapsedPoints, Name + TEXT(" First run"));
+						bLocalSuccess &= CompareTransforms(SecondRunPoints, SecondRunCollapsedPoints, Name + TEXT(" Second run"));
+						bLocalSuccess &= CompareTransforms(ThirdRunPoints, ThirdRunCollapsedPoints, Name + TEXT(" Third run"));
 
 						if (!bLocalSuccess)
 						{
diff --git a/Engine/Plugins/PCG/Source/PCG/Public/SpatialAlgo/PCGOctreeQueries.h b/Engine/Plugins/PCG/Source/PCG/Public/SpatialAlgo/PCGOctreeQueries.h
--- a/Engine/Plugins/PCG/Source/PCG/Public/SpatialAlgo/PCGOctreeQueries.h
+++ b/Engine/Plugins/PCG/Source/PCG/Public/SpatialAlgo/PCGOctreeQueries.h
@@ -30,6 +30,7 @@ public:
 
 	/** Get the nearest point to a given point (that is not itself) within the search distance. Returns null if not found. */
 	static PCG_API const FPCGPoint* GetClosestPointFromOtherPoint(const UPCGPointData* InPointData, const FPCGPoint& InPoint, const double InSearchDistance = DefaultSearchDistance);
+	static PCG_API int32 GetClosestPointIndexFromOtherPointIndex(const UPCGBasePointData* InPointData, int32 OtherPointIndex, const double InSearchDistance = DefaultSearchDistance);
 
 	/** Get the farthest point from a given position within the search distance. Returns null if not found. */
 	static PCG_API const FPCGPoint* GetFarthestPoint(const UPCGPointData* InPointData, const FVector& InCenter, const double InSearchDistance = DefaultSearchDistance);
@@ -37,6 +38,7 @@ public:
 
 	/** Get the farthest point from a given point (that is not itself) within the search distance. Returns null if not found. */
 	static PCG_API const FPCGPoint* GetFarthestPointFromOtherPoint(const UPCGPointData* InPointData, const FPCGPoint& InPoint, const double InSearchDistance = DefaultSearchDistance);
+	static PCG_API int32 GetFarthestPointIndexFromOtherPointIndex(const UPCGBasePointData* InPointData, int32 OtherPointIndex, const double InSearchDistance);
 
 private:
 	// Blueprint functions are not exposed, use the functions above or the Octree directly
diff --git a/Engine/Plugins/PCG/Source/PCG/Public/Tests/PCGTestsCommon.h b/Engine/Plugins/PCG/Source/PCG/Public/Tests/PCGTestsCommon.h
--- a/Engine/Plugins/PCG/Source/PCG/Public/Tests/PCGTestsCommon.h
+++ b/Engine/Plugins/PCG/Source/PCG/Public/Tests/PCGTestsCommon.h
@@ -66,6 +66,8 @@ namespace PCGTestsCommon
 		return NewObject<PointDataType>();
 	}
 
+	UPCGBasePointData* CreateBasePointData();
+
 	template <typename PointDataType = UPCGPointData, std::enable_if_t<std::is_base_of_v<UPCGBasePointData, PointDataType>, int> = 0>
 	PointDataType* CreatePointData()
 	{
@@ -78,6 +80,8 @@ namespace PCGTestsCommon
 		return SinglePointData;
 	}
 
+	UPCGBasePointData* CreateBasePointData(const FVector& InLocation);
+
 	template <typename PointDataType = UPCGPointData, std::enable_if_t<std::is_base_of_v<UPCGBasePointData, PointDataType>, int> = 0>
 	PointDataType* CreatePointData(const FVector& InLocation)
 	{
@@ -93,6 +97,8 @@ namespace PCGTestsCommon
 		return SinglePointData;
 	}
 
+	UPCGBasePointData* CreateRandomBasePointData(int32 PointCount, int32 Seed, bool bRandomDensity = false);
+
 	/** Creates a point data with PointCount many points, and randomizes the Transform, Color, and Density */
 	template <typename PointDataType = UPCGPointData, std::enable_if_t<std::is_base_of_v<UPCGBasePointData, PointDataType>, int> = 0>
 	PointDataType* CreateRandomPointData(int32 PointCount, int32 Seed, bool bRandomDensity = false)
