[PCG] Fix FPCGPointValueRanges so that it doesn't create TPCGValueRanges for unallocated properties (which was flattening properties and wasn't thread safe)
- Refactor CopyUnallocatedPropertiesTo to CopyUnallocatedPropertiesFrom to improve code readability
- Now that only allocated ranges are created in FPCGPointValueRanges, fix nodes that were relying on the opposite (mostly add missing CopyUnallocatedPropertiesFrom)

#rnx

diff --git a/Engine/Plugins/PCG/Source/PCG/Private/Data/PCGIntersectionData.cpp b/Engine/Plugins/PCG/Source/PCG/Private/Data/PCGIntersectionData.cpp
--- a/Engine/Plugins/PCG/Source/PCG/Private/Data/PCGIntersectionData.cpp
+++ b/Engine/Plugins/PCG/Source/PCG/Private/Data/PCGIntersectionData.cpp
@@ -316,7 +316,7 @@ UPCGBasePointData* UPCGIntersectionData::CreateAndFilterPointData(FPCGContext* C
 		return NumWritten;
 	};
 
-	const EPCGPointNativeProperties PropertiesToAllocate = SourcePointData->GetAllocatedProperties() | EPCGPointNativeProperties::Density;
+	const EPCGPointNativeProperties PropertiesToAllocate = SourcePointData->GetAllocatedProperties() | EPCGPointNativeProperties::Density | EPCGPointNativeProperties::Color;
 
 	FPCGSpatialDataProcessing::SampleBasedRangeProcessing<ChunkSize>(Context ? &Context->AsyncState : nullptr, ChunkSamplePoints, SourcePointData, OutputData, PropertiesToAllocate);
 
diff --git a/Engine/Plugins/PCG/Source/PCG/Private/Data/PCGPointArrayData.cpp b/Engine/Plugins/PCG/Source/PCG/Private/Data/PCGPointArrayData.cpp
--- a/Engine/Plugins/PCG/Source/PCG/Private/Data/PCGPointArrayData.cpp
+++ b/Engine/Plugins/PCG/Source/PCG/Private/Data/PCGPointArrayData.cpp
@@ -224,18 +224,23 @@ void UPCGPointArrayData::MoveRange(int32 RangeStartIndex, int32 MoveToIndex, int
 	PointArray.MoveRange(RangeStartIndex, MoveToIndex, NumElements);
 }
 
-void UPCGPointArrayData::CopyUnallocatedPropertiesTo(UPCGBasePointData* OutPointData) const
+void UPCGPointArrayData::CopyUnallocatedPropertiesFrom(const UPCGBasePointData* InPointData)
 {
-	if (UPCGPointArrayData* OutPointArrayData = Cast<UPCGPointArrayData>(OutPointData); OutPointArrayData && !OutPointArrayData->HasSpatialDataParent())
+	if (HasSpatialDataParent())
 	{
-		GetProperty<FTransform>(EPCGPointNativeProperties::Transform)->CopyUnallocatedProperty(*OutPointArrayData->GetProperty<FTransform>(EPCGPointNativeProperties::Transform, /*bWithInheritance=*/false));
-		GetProperty<float>(EPCGPointNativeProperties::Density)->CopyUnallocatedProperty(*OutPointArrayData->GetProperty<float>(EPCGPointNativeProperties::Density, /*bWithInheritance=*/false));
-		GetProperty<FVector>(EPCGPointNativeProperties::BoundsMin)->CopyUnallocatedProperty(*OutPointArrayData->GetProperty<FVector>(EPCGPointNativeProperties::BoundsMin, /*bWithInheritance=*/false));
-		GetProperty<FVector>(EPCGPointNativeProperties::BoundsMax)->CopyUnallocatedProperty(*OutPointArrayData->GetProperty<FVector>(EPCGPointNativeProperties::BoundsMax, /*bWithInheritance=*/false));
-		GetProperty<FVector4>(EPCGPointNativeProperties::Color)->CopyUnallocatedProperty(*OutPointArrayData->GetProperty<FVector4>(EPCGPointNativeProperties::Color, /*bWithInheritance=*/false));
-		GetProperty<float>(EPCGPointNativeProperties::Steepness)->CopyUnallocatedProperty(*OutPointArrayData->GetProperty<float>(EPCGPointNativeProperties::Steepness, /*bWithInheritance=*/false));
-		GetProperty<int32>(EPCGPointNativeProperties::Seed)->CopyUnallocatedProperty(*OutPointArrayData->GetProperty<int32>(EPCGPointNativeProperties::Seed, /*bWithInheritance=*/false));
-		GetProperty<int64>(EPCGPointNativeProperties::MetadataEntry)->CopyUnallocatedProperty(*OutPointArrayData->GetProperty<int64>(EPCGPointNativeProperties::MetadataEntry, /*bWithInheritance=*/false));
+		return;
+	}
+
+	if (const UPCGPointArrayData* InPointArrayData = Cast<UPCGPointArrayData>(InPointData))
+	{
+		InPointArrayData->GetProperty<FTransform>(EPCGPointNativeProperties::Transform)->CopyUnallocatedProperty(*GetProperty<FTransform>(EPCGPointNativeProperties::Transform, /*bWithInheritance=*/false));
+		InPointArrayData->GetProperty<float>(EPCGPointNativeProperties::Density)->CopyUnallocatedProperty(*GetProperty<float>(EPCGPointNativeProperties::Density, /*bWithInheritance=*/false));
+		InPointArrayData->GetProperty<FVector>(EPCGPointNativeProperties::BoundsMin)->CopyUnallocatedProperty(*GetProperty<FVector>(EPCGPointNativeProperties::BoundsMin, /*bWithInheritance=*/false));
+		InPointArrayData->GetProperty<FVector>(EPCGPointNativeProperties::BoundsMax)->CopyUnallocatedProperty(*GetProperty<FVector>(EPCGPointNativeProperties::BoundsMax, /*bWithInheritance=*/false));
+		InPointArrayData->GetProperty<FVector4>(EPCGPointNativeProperties::Color)->CopyUnallocatedProperty(*GetProperty<FVector4>(EPCGPointNativeProperties::Color, /*bWithInheritance=*/false));
+		InPointArrayData->GetProperty<float>(EPCGPointNativeProperties::Steepness)->CopyUnallocatedProperty(*GetProperty<float>(EPCGPointNativeProperties::Steepness, /*bWithInheritance=*/false));
+		InPointArrayData->GetProperty<int32>(EPCGPointNativeProperties::Seed)->CopyUnallocatedProperty(*GetProperty<int32>(EPCGPointNativeProperties::Seed, /*bWithInheritance=*/false));
+		InPointArrayData->GetProperty<int64>(EPCGPointNativeProperties::MetadataEntry)->CopyUnallocatedProperty(*GetProperty<int64>(EPCGPointNativeProperties::MetadataEntry, /*bWithInheritance=*/false));
 	}
 }
 
diff --git a/Engine/Plugins/PCG/Source/PCG/Private/Data/PCGProjectionData.cpp b/Engine/Plugins/PCG/Source/PCG/Private/Data/PCGProjectionData.cpp
--- a/Engine/Plugins/PCG/Source/PCG/Private/Data/PCGProjectionData.cpp
+++ b/Engine/Plugins/PCG/Source/PCG/Private/Data/PCGProjectionData.cpp
@@ -311,7 +311,7 @@ const UPCGBasePointData* UPCGProjectionData::CreateBasePointData(FPCGContext* Co
 
 		// @todo_pcg: try to optimize allocations here
 		PointData->AllocateProperties(SourcePointData->GetAllocatedProperties() | EPCGPointNativeProperties::Transform | EPCGPointNativeProperties::Color | EPCGPointNativeProperties::Density);
-		SourcePointData->CopyUnallocatedPropertiesTo(PointData);
+		PointData->CopyUnallocatedPropertiesFrom(SourcePointData);
 	};
 
 	auto MoveDataRangeFunc = [PointData](int32 RangeStartIndex, int32 MoveToIndex, int32 NumElements)
diff --git a/Engine/Plugins/PCG/Source/PCG/Private/Data/PCGUnionData.cpp b/Engine/Plugins/PCG/Source/PCG/Private/Data/PCGUnionData.cpp
--- a/Engine/Plugins/PCG/Source/PCG/Private/Data/PCGUnionData.cpp
+++ b/Engine/Plugins/PCG/Source/PCG/Private/Data/PCGUnionData.cpp
@@ -410,7 +410,7 @@ void UPCGUnionData::CreateSequentialPointData(FPCGContext* Context, TArray<const
 	}
 	
 	OutPointData->SetNumPoints(PointInputCount, /*bInitializeValues=*/false);
-	OutPointData->AllocateProperties(PropertiesToAllocate);
+	OutPointData->AllocateProperties(PropertiesToAllocate | EPCGPointNativeProperties::MetadataEntry);
 
 	struct FIndexParams
 	{
@@ -543,6 +543,7 @@ void UPCGUnionData::CreateSequentialPointData(FPCGContext* Context, TArray<const
 		};
 
 		const UPCGBasePointData* CurrentPointData = PointInputDatas[DataIndex];
+		OutPointData->CopyUnallocatedPropertiesFrom(CurrentPointData);
 
 		FPCGAsync::AsyncProcessingRangeEx(
 			Context ? &Context->AsyncState : nullptr,
diff --git a/Engine/Plugins/PCG/Source/PCG/Private/Elements/Grammar/PCGSubdivideSegment.cpp b/Engine/Plugins/PCG/Source/PCG/Private/Elements/Grammar/PCGSubdivideSegment.cpp
--- a/Engine/Plugins/PCG/Source/PCG/Private/Elements/Grammar/PCGSubdivideSegment.cpp
+++ b/Engine/Plugins/PCG/Source/PCG/Private/Elements/Grammar/PCGSubdivideSegment.cpp
@@ -101,6 +101,7 @@ public:
 
 		InOutParameters.OutPointData->SetNumPoints(InOutParameters.OutPointData->GetNumPoints() + ModulesInstances.Num(), /*bInitializeValues=*/false);
 		InOutParameters.OutPointData->AllocateProperties(InOutParameters.InPointData->GetAllocatedProperties());
+		InOutParameters.OutPointData->CopyUnallocatedPropertiesFrom(InOutParameters.InPointData);
 		
 		FPCGPointValueRanges OutRanges(InOutParameters.OutPointData, /*bAllocateProperties=*/false);
 
diff --git a/Engine/Plugins/PCG/Source/PCG/Private/Elements/PCGApplyHierarchy.cpp b/Engine/Plugins/PCG/Source/PCG/Private/Elements/PCGApplyHierarchy.cpp
--- a/Engine/Plugins/PCG/Source/PCG/Private/Elements/PCGApplyHierarchy.cpp
+++ b/Engine/Plugins/PCG/Source/PCG/Private/Elements/PCGApplyHierarchy.cpp
@@ -588,7 +588,7 @@ bool FPCGApplyHierarchyElement::ExecuteInternal(FPCGContext* InContext) const
 				{
 					FilteredData->SetNumPoints(OriginalData->GetNumPoints());
 					FilteredData->AllocateProperties(OriginalData->GetAllocatedProperties());
-					OriginalData->CopyUnallocatedPropertiesTo(FilteredData);
+					FilteredData->CopyUnallocatedPropertiesFrom(OriginalData);
 				};
 
 				auto AsyncProcessRangeFunc = [FilteredData, OriginalData, &IterState](int32 StartReadIndex, int32 StartWriteIndex, int32 Count)
diff --git a/Engine/Plugins/PCG/Source/PCG/Private/Elements/PCGAttractElement.cpp b/Engine/Plugins/PCG/Source/PCG/Private/Elements/PCGAttractElement.cpp
--- a/Engine/Plugins/PCG/Source/PCG/Private/Elements/PCGAttractElement.cpp
+++ b/Engine/Plugins/PCG/Source/PCG/Private/Elements/PCGAttractElement.cpp
@@ -570,7 +570,8 @@ bool FPCGAttractElement::ExecuteInternal(FPCGContext* InContext) const
 
 			AttractData.OutputData->SetNumPoints(NumPoints, /*bInitializeValues=*/false);
 			AttractData.OutputData->AllocateProperties(AttractData.SourceData->GetAllocatedProperties());
-			
+			AttractData.OutputData->CopyUnallocatedPropertiesFrom(AttractData.SourceData);
+
 			Mapping.Reserve(NumPoints);
 
 			const FConstPCGPointValueRanges SourceRanges(AttractData.SourceData);
diff --git a/Engine/Plugins/PCG/Source/PCG/Private/Elements/PCGAttributeFilter.cpp b/Engine/Plugins/PCG/Source/PCG/Private/Elements/PCGAttributeFilter.cpp
--- a/Engine/Plugins/PCG/Source/PCG/Private/Elements/PCGAttributeFilter.cpp
+++ b/Engine/Plugins/PCG/Source/PCG/Private/Elements/PCGAttributeFilter.cpp
@@ -853,10 +853,12 @@ bool FPCGAttributeFilterElementBase::DoFiltering(FPCGContext* Context, EPCGAttri
 				int32 InFilterWriteIndex = 0;
 				OperationData.InFilterPointData->SetNumPoints(NumInFilterPoints);
 				OperationData.InFilterPointData->AllocateProperties(OperationData.OriginalPointData->GetAllocatedProperties());
+				OperationData.InFilterPointData->CopyUnallocatedPropertiesFrom(OperationData.OriginalPointData);
 
 				int32 OutFilterWriteIndex = 0;
 				OperationData.OutFilterPointData->SetNumPoints(NumOutFilterPoints);
 				OperationData.OutFilterPointData->AllocateProperties(OperationData.OriginalPointData->GetAllocatedProperties());
+				OperationData.OutFilterPointData->CopyUnallocatedPropertiesFrom(OperationData.OriginalPointData);
 
 				const FConstPCGPointValueRanges OriginalRanges(OperationData.OriginalPointData);
 				FPCGPointValueRanges InFilterRanges(OperationData.InFilterPointData, /*bAllocate=*/false);
diff --git a/Engine/Plugins/PCG/Source/PCG/Private/Elements/PCGAttributeSelectElement.cpp b/Engine/Plugins/PCG/Source/PCG/Private/Elements/PCGAttributeSelectElement.cpp
--- a/Engine/Plugins/PCG/Source/PCG/Private/Elements/PCGAttributeSelectElement.cpp
+++ b/Engine/Plugins/PCG/Source/PCG/Private/Elements/PCGAttributeSelectElement.cpp
@@ -440,6 +440,7 @@ bool FPCGAttributeSelectElement::ExecuteInternal(FPCGContext* Context) const
 		OutputPointData->InitializeFromDataWithParams(InitializeFromDataParams);
 
 		OutputPointData->SetNumPoints(1);
+		OutputPointData->CopyUnallocatedPropertiesFrom(PointData);
 
 		const FConstPCGPointValueRanges InRanges(PointData);
 		FPCGPointValueRanges OutRanges(OutputPointData, /*bAllocate=*/false);
diff --git a/Engine/Plugins/PCG/Source/PCG/Private/Elements/PCGBoundsFromMesh.cpp b/Engine/Plugins/PCG/Source/PCG/Private/Elements/PCGBoundsFromMesh.cpp
--- a/Engine/Plugins/PCG/Source/PCG/Private/Elements/PCGBoundsFromMesh.cpp
+++ b/Engine/Plugins/PCG/Source/PCG/Private/Elements/PCGBoundsFromMesh.cpp
@@ -276,6 +276,8 @@ bool FPCGBoundsFromMeshElement::ExecuteInternal(FPCGContext* InContext) const
 			else
 			{
 				OutputData->AllocateProperties(PointData->GetAllocatedProperties() | EPCGPointNativeProperties::BoundsMin | EPCGPointNativeProperties::BoundsMax);
+				OutputData->CopyUnallocatedPropertiesFrom(PointData);
+
 				auto CopyAndSetBounds = [OutputData, PointData, &StaticMeshBounds](const int32 StartReadIndex, const int32 StartWriteIndex, const int32 Count)
 				{
 					int32 NumWritten = 0;
@@ -317,6 +319,7 @@ bool FPCGBoundsFromMeshElement::ExecuteInternal(FPCGContext* InContext) const
 				}
 
 				OutputData->AllocateProperties(EPCGPointNativeProperties::BoundsMin | EPCGPointNativeProperties::BoundsMax);
+				OutputData->CopyUnallocatedPropertiesFrom(PointData);
 
 				auto CopyAndSetBounds = [OutputData, PointData, Context, InputData](const int32 StartReadIndex, const int32 StartWriteIndex, const int32 Count)
 				{
diff --git a/Engine/Plugins/PCG/Source/PCG/Private/Elements/PCGCollapsePoints.cpp b/Engine/Plugins/PCG/Source/PCG/Private/Elements/PCGCollapsePoints.cpp
--- a/Engine/Plugins/PCG/Source/PCG/Private/Elements/PCGCollapsePoints.cpp
+++ b/Engine/Plugins/PCG/Source/PCG/Private/Elements/PCGCollapsePoints.cpp
@@ -525,6 +525,7 @@ bool FPCGCollapsePointsElement::ExecuteInternal(FPCGContext* InContext) const
 			// Need to allocate prior to key creation
 			CollapseState.OutputPointData->SetNumPoints(Partition.Num(), /*bInitializeValues=*/false);
 			CollapseState.OutputPointData->AllocateProperties(CollapseState.SourcePointData->GetAllocatedProperties());
+			CollapseState.OutputPointData->CopyUnallocatedPropertiesFrom(CollapseState.SourcePointData);
 
 			// Build accessor list(s) + keys, based on settings
 			TArray<double> PartitionWeights;
diff --git a/Engine/Plugins/PCG/Source/PCG/Private/Elements/PCGCombinePoints.cpp b/Engine/Plugins/PCG/Source/PCG/Private/Elements/PCGCombinePoints.cpp
--- a/Engine/Plugins/PCG/Source/PCG/Private/Elements/PCGCombinePoints.cpp
+++ b/Engine/Plugins/PCG/Source/PCG/Private/Elements/PCGCombinePoints.cpp
@@ -55,6 +55,7 @@ bool FPCGCombinePointsElement::ExecuteInternal(FPCGContext* Context) const
 		OutputPointData->InitializeFromDataWithParams(InitializeFromDataParams);
 		
 		OutputPointData->SetNumPoints(1);
+		OutputPointData->AllocateProperties(EPCGPointNativeProperties::All);
 		Output.Data = OutputPointData;
 
 		FPCGPoint OutPoint = FPCGPoint();
diff --git a/Engine/Plugins/PCG/Source/PCG/Private/Elements/PCGConvexHull2D.cpp b/Engine/Plugins/PCG/Source/PCG/Private/Elements/PCGConvexHull2D.cpp
--- a/Engine/Plugins/PCG/Source/PCG/Private/Elements/PCGConvexHull2D.cpp
+++ b/Engine/Plugins/PCG/Source/PCG/Private/Elements/PCGConvexHull2D.cpp
@@ -53,6 +53,7 @@ bool FPCGConvexHull2DElement::ExecuteInternal(FPCGContext* Context) const
 		
 		OutputPointData->SetNumPoints(ConvexHullIndices.Num());
 		OutputPointData->AllocateProperties(PointData->GetAllocatedProperties());
+		OutputPointData->CopyUnallocatedPropertiesFrom(PointData);
 
 		FPCGPointValueRanges OutRanges(OutputPointData, /*bAllocate=*/false);
 		int32 WriteIndex = 0;
diff --git a/Engine/Plugins/PCG/Source/PCG/Private/Elements/PCGCopyPoints.cpp b/Engine/Plugins/PCG/Source/PCG/Private/Elements/PCGCopyPoints.cpp
--- a/Engine/Plugins/PCG/Source/PCG/Private/Elements/PCGCopyPoints.cpp
+++ b/Engine/Plugins/PCG/Source/PCG/Private/Elements/PCGCopyPoints.cpp
@@ -318,6 +318,7 @@ bool FPCGCopyPointsElement::ExecuteInternal(FPCGContext* Context) const
 			PropertiesToAllocate |= GetPropertyToAllocate(ColorInheritance, EPCGPointNativeProperties::Color);
 
 			OutPointData->AllocateProperties(PropertiesToAllocate);
+			OutPointData->CopyUnallocatedPropertiesFrom(SourcePointData);
 		};
 
 		const int32 TargetNumPoints = TargetPointData->GetNumPoints();
@@ -339,10 +340,7 @@ bool FPCGCopyPointsElement::ExecuteInternal(FPCGContext* Context) const
 				const int32 WriteIndex = StartWriteIndex + NumWritten;
 
 				// Copy properties from source
-				OutRanges.DensityRange[WriteIndex] = SourceRanges.DensityRange[SourceIndex];
-				OutRanges.SteepnessRange[WriteIndex] = SourceRanges.SteepnessRange[SourceIndex];
-				OutRanges.BoundsMinRange[WriteIndex] = SourceRanges.BoundsMinRange[SourceIndex];
-				OutRanges.BoundsMaxRange[WriteIndex] = SourceRanges.BoundsMaxRange[SourceIndex];
+				FPCGPoint OutPoint = SourceRanges.GetPoint(SourceIndex);
 
 				const FTransform& SourcePointTransform = SourceRanges.TransformRange[SourceIndex];
 				const FTransform& TargetPointTransform = TargetRanges.TransformRange[TargetIndex];
@@ -371,43 +369,43 @@ bool FPCGCopyPointsElement::ExecuteInternal(FPCGContext* Context) const
 					TargetTransform.SetScale3D(TargetPointTransform.GetScale3D());
 				}
 
-				OutRanges.TransformRange[WriteIndex] = SourceTransform * TargetTransform;
+				OutPoint.Transform = SourceTransform * TargetTransform;
 
 				if (RotationInheritance == EPCGCopyPointsInheritanceMode::Source && bApplyTargetRotationToPositions)
 				{
-					OutRanges.TransformRange[WriteIndex].SetRotation(SourcePointTransform.GetRotation());
+					OutPoint.Transform.SetRotation(SourcePointTransform.GetRotation());
 				}
 
 				if (ScaleInheritance == EPCGCopyPointsInheritanceMode::Source && bApplyTargetScaleToPositions)
 				{
-					OutRanges.TransformRange[WriteIndex].SetScale3D(SourcePointTransform.GetScale3D());
+					OutPoint.Transform.SetScale3D(SourcePointTransform.GetScale3D());
 				}
 
 				if (ColorInheritance == EPCGCopyPointsInheritanceMode::Relative)
 				{
-					OutRanges.ColorRange[WriteIndex] = SourceRanges.ColorRange[SourceIndex] * TargetRanges.ColorRange[TargetIndex];
+					OutPoint.Color = SourceRanges.ColorRange[SourceIndex] * TargetRanges.ColorRange[TargetIndex];
 				}
 				else if (ColorInheritance == EPCGCopyPointsInheritanceMode::Target)
 				{
-					OutRanges.ColorRange[WriteIndex] = TargetRanges.ColorRange[TargetIndex];
+					OutPoint.Color = TargetRanges.ColorRange[TargetIndex];
 				}
 				else // if (ColorInheritance == EPCGCopyPointsInheritanceMode::Source)
 				{
-					OutRanges.ColorRange[WriteIndex] = SourceRanges.ColorRange[SourceIndex];
+					OutPoint.Color = SourceRanges.ColorRange[SourceIndex];
 				}
 
 				// Set seed based on inheritance mode
 				if (SeedInheritance == EPCGCopyPointsInheritanceMode::Relative)
 				{
-					OutRanges.SeedRange[WriteIndex] = PCGHelpers::ComputeSeed(SourceRanges.SeedRange[SourceIndex], TargetRanges.SeedRange[TargetIndex]);
+					OutPoint.Seed = PCGHelpers::ComputeSeed(SourceRanges.SeedRange[SourceIndex], TargetRanges.SeedRange[TargetIndex]);
 				}
 				else if (SeedInheritance == EPCGCopyPointsInheritanceMode::Target)
 				{
-					OutRanges.SeedRange[WriteIndex] = TargetRanges.SeedRange[TargetIndex];
+					OutPoint.Seed = TargetRanges.SeedRange[TargetIndex];
 				}
 				else // if (SeedInheritance == EPCGCopyPointsInheritanceMode::Source)
 				{
-					OutRanges.SeedRange[WriteIndex] = SourceRanges.SeedRange[SourceIndex];
+					OutPoint.Seed = SourceRanges.SeedRange[SourceIndex];
 				}
 
 				if (bProcessMetadata)
@@ -415,8 +413,8 @@ bool FPCGCopyPointsElement::ExecuteInternal(FPCGContext* Context) const
 					const int64 RootMetadataEntry = bInheritMetadataFromSource ? SourceRanges.MetadataEntryRange[SourceIndex] : TargetRanges.MetadataEntryRange[TargetIndex];
 					const int64 NonRootMetadataEntry = bInheritMetadataFromSource ? TargetRanges.MetadataEntryRange[TargetIndex] : SourceRanges.MetadataEntryRange[SourceIndex];
 
-					OutRanges.MetadataEntryRange[WriteIndex] = OutMetadata->AddEntryPlaceholder();
-					AllMetadataEntries[Index] = TTuple<int64, int64>(OutRanges.MetadataEntryRange[WriteIndex], RootMetadataEntry);
+					OutPoint.MetadataEntry = OutMetadata->AddEntryPlaceholder();
+					AllMetadataEntries[Index] = TTuple<int64, int64>(OutPoint.MetadataEntry, RootMetadataEntry);
 
 					if (NonRootMetadata)
 					{
@@ -425,14 +423,14 @@ bool FPCGCopyPointsElement::ExecuteInternal(FPCGContext* Context) const
 						{
 							for (int32 AttributeIndex = 0; AttributeIndex < NonRootAttributes.Num(); ++AttributeIndex)
 							{
-								AttributeValuesToSet[AttributeIndex][Index] = TTuple<PCGMetadataEntryKey, PCGMetadataValueKey>(OutRanges.MetadataEntryRange[WriteIndex], NonRootAttributes[AttributeIndex]->GetValueKey(NonRootMetadataEntry));
+								AttributeValuesToSet[AttributeIndex][Index] = TTuple<PCGMetadataEntryKey, PCGMetadataValueKey>(OutPoint.MetadataEntry, NonRootAttributes[AttributeIndex]->GetValueKey(NonRootMetadataEntry));
 							}
 						}
 						else
 						{
 							for (int32 AttributeIndex = 0; AttributeIndex < NonRootAttributes.Num(); ++AttributeIndex)
 							{
-								AttributeValuesToSet[AttributeIndex][Index] = TTuple<PCGMetadataEntryKey, PCGMetadataValueKey>(OutRanges.MetadataEntryRange[WriteIndex], PCGDefaultValueKey);
+								AttributeValuesToSet[AttributeIndex][Index] = TTuple<PCGMetadataEntryKey, PCGMetadataValueKey>(OutPoint.MetadataEntry, PCGDefaultValueKey);
 							}
 						}
 					}
@@ -440,9 +438,10 @@ bool FPCGCopyPointsElement::ExecuteInternal(FPCGContext* Context) const
 				else
 				{
 					// Reset the metadata entry if we have no metadata.
-					OutRanges.MetadataEntryRange[WriteIndex] = PCGInvalidEntryKey;
+					OutPoint.MetadataEntry = PCGInvalidEntryKey;
 				}
 
+				OutRanges.SetFromPoint(WriteIndex, OutPoint);
 				++NumWritten;
 			}
 
diff --git a/Engine/Plugins/PCG/Source/PCG/Private/Elements/PCGDensityFilter.cpp b/Engine/Plugins/PCG/Source/PCG/Private/Elements/PCGDensityFilter.cpp
--- a/Engine/Plugins/PCG/Source/PCG/Private/Elements/PCGDensityFilter.cpp
+++ b/Engine/Plugins/PCG/Source/PCG/Private/Elements/PCGDensityFilter.cpp
@@ -87,8 +87,8 @@ bool FPCGDensityFilterElement::ExecuteInternal(FPCGContext* Context) const
 		auto InitializeFunc = [FilteredData, OriginalData]()
 		{
 			FilteredData->SetNumPoints(OriginalData->GetNumPoints(), /*bInitializeValues=*/false);
-			FilteredData->AllocateProperties(OriginalData->GetAllocatedProperties());
-			OriginalData->CopyUnallocatedPropertiesTo(FilteredData);
+			FilteredData->AllocateProperties(OriginalData->GetAllocatedProperties() | EPCGPointNativeProperties::Density);
+			FilteredData->CopyUnallocatedPropertiesFrom(OriginalData);
 		};
 
 		auto AsyncProcessRangeFunc = [FilteredData, OriginalData, MinBound, MaxBound, bInvertFilter, bKeepZeroDensityPoints](int32 StartReadIndex, int32 StartWriteIndex, int32 Count)
@@ -103,32 +103,21 @@ bool FPCGDensityFilterElement::ExecuteInternal(FPCGContext* Context) const
 				const int32 WriteIndex = StartWriteIndex + NumWritten;
 				const float ReadDensity = ReadRanges.DensityRange[ReadIndex];
 				const bool bInRange = (ReadDensity >= MinBound && ReadDensity <= MaxBound);
-				bool bCopyPoint = false;
 			
 				if (bInRange != bInvertFilter)
 				{
-					bCopyPoint = true;
+					WriteRanges.SetFromValueRanges(WriteIndex, ReadRanges, ReadIndex);
 					WriteRanges.DensityRange[WriteIndex] = ReadRanges.DensityRange[ReadIndex];
+					++NumWritten;
 				}
 #if WITH_EDITOR
 				else if (bKeepZeroDensityPoints)
 				{
-					bCopyPoint = true;
+					WriteRanges.SetFromValueRanges(WriteIndex, ReadRanges, ReadIndex);
 					WriteRanges.DensityRange[WriteIndex] = 0.f;
-				}
-#endif
-
-				if (bCopyPoint)
-				{
-					WriteRanges.TransformRange[WriteIndex] = ReadRanges.TransformRange[ReadIndex];
-					WriteRanges.SteepnessRange[WriteIndex] = ReadRanges.SteepnessRange[ReadIndex];
-					WriteRanges.BoundsMinRange[WriteIndex] = ReadRanges.BoundsMinRange[ReadIndex];
-					WriteRanges.BoundsMaxRange[WriteIndex] = ReadRanges.BoundsMaxRange[ReadIndex];
-					WriteRanges.ColorRange[WriteIndex] = ReadRanges.ColorRange[ReadIndex];
-					WriteRanges.MetadataEntryRange[WriteIndex] = ReadRanges.MetadataEntryRange[ReadIndex];
-					WriteRanges.SeedRange[WriteIndex] = ReadRanges.SeedRange[ReadIndex];
 					++NumWritten;
 				}
+#endif
 			}
 
 			return NumWritten;
diff --git a/Engine/Plugins/PCG/Source/PCG/Private/Elements/PCGDuplicatePoint.cpp b/Engine/Plugins/PCG/Source/PCG/Private/Elements/PCGDuplicatePoint.cpp
--- a/Engine/Plugins/PCG/Source/PCG/Private/Elements/PCGDuplicatePoint.cpp
+++ b/Engine/Plugins/PCG/Source/PCG/Private/Elements/PCGDuplicatePoint.cpp
@@ -71,6 +71,7 @@ bool FPCGDuplicatePointElement::ExecuteInternal(FPCGContext* Context) const
 		{
 			OutPointData->SetNumPoints(NumIterations, /*bInitializeValues=*/false);
 			OutPointData->AllocateProperties(InputPointData->GetAllocatedProperties() | EPCGPointNativeProperties::Transform | EPCGPointNativeProperties::Seed);
+			OutPointData->CopyUnallocatedPropertiesFrom(InputPointData);
 		};
 
 		const FTransform& SourceDuplicateTransform = Settings->PointTransform;
@@ -125,6 +126,7 @@ bool FPCGDuplicatePointElement::ExecuteInternal(FPCGContext* Context) const
 			{
 				const FConstPCGPointValueRanges InRanges(InputPointData);
 				FPCGPointValueRanges OutRanges(OutPointData, /*bAllocate=*/false);
+				const FConstPCGPointValueRanges ConstOutRanges(OutPointData);
 
 				int32 NumWritten = 0;
 				const int32 NumInputPoints = InputPointData->GetNumPoints();
@@ -139,7 +141,7 @@ bool FPCGDuplicatePointElement::ExecuteInternal(FPCGContext* Context) const
 
 					if (DuplicateIndex != 0)
 					{
-						const FVector DuplicateLocationOffset = ((OutRanges.BoundsMaxRange[WriteIndex] - OutRanges.BoundsMinRange[WriteIndex]) * Direction + SourceDuplicateTransform.GetLocation()) * DuplicateIndex;
+						const FVector DuplicateLocationOffset = ((ConstOutRanges.BoundsMaxRange[WriteIndex] - ConstOutRanges.BoundsMinRange[WriteIndex]) * Direction + SourceDuplicateTransform.GetLocation()) * DuplicateIndex;
 						const FRotator DuplicateRotationOffset = SourceDuplicateTransform.Rotator() * DuplicateIndex;
 						const FVector DuplicateScaleMultiplier = FVector(
 							FMath::Pow(SourceDuplicateTransform.GetScale3D().X, DuplicateIndex),
diff --git a/Engine/Plugins/PCG/Source/PCG/Private/Elements/PCGFilterElementsByIndex.cpp b/Engine/Plugins/PCG/Source/PCG/Private/Elements/PCGFilterElementsByIndex.cpp
--- a/Engine/Plugins/PCG/Source/PCG/Private/Elements/PCGFilterElementsByIndex.cpp
+++ b/Engine/Plugins/PCG/Source/PCG/Private/Elements/PCGFilterElementsByIndex.cpp
@@ -319,6 +319,7 @@ bool FPCGFilterElementsByIndexElement::ExecuteInternal(FPCGContext* InContext) c
 				const int32 NumOutFilterPoints = InputPointData->GetNumPoints() - NumInFilterPoints;
 				OutFilterPointData->SetNumPoints(NumOutFilterPoints);
 				OutFilterPointData->AllocateProperties(InputPointData->GetAllocatedProperties());
+				OutFilterPointData->CopyUnallocatedPropertiesFrom(InputPointData);
 
 				AddNewOutput(OutFilterPointData, Constants::OutFilterOutputPinLabel);
 			}
diff --git a/Engine/Plugins/PCG/Source/PCG/Private/Elements/PCGMatchAndSetAttributes.cpp b/Engine/Plugins/PCG/Source/PCG/Private/Elements/PCGMatchAndSetAttributes.cpp
--- a/Engine/Plugins/PCG/Source/PCG/Private/Elements/PCGMatchAndSetAttributes.cpp
+++ b/Engine/Plugins/PCG/Source/PCG/Private/Elements/PCGMatchAndSetAttributes.cpp
@@ -1098,6 +1098,7 @@ bool FPCGMatchAndSetAttributesElement::PrepareDataInternal(FPCGContext* InContex
 			OutPointData->InitializeFromDataWithParams(InitializeFromDataParams);
 			OutPointData->SetNumPoints(InPointData->GetNumPoints());
 			OutPointData->AllocateProperties(InPointData->GetAllocatedProperties() | EPCGPointNativeProperties::MetadataEntry);
+			OutPointData->CopyUnallocatedPropertiesFrom(InPointData);
 
 			OutState.OutData = OutPointData;
 		}
diff --git a/Engine/Plugins/PCG/Source/PCG/Private/Elements/PCGRandomChoice.cpp b/Engine/Plugins/PCG/Source/PCG/Private/Elements/PCGRandomChoice.cpp
--- a/Engine/Plugins/PCG/Source/PCG/Private/Elements/PCGRandomChoice.cpp
+++ b/Engine/Plugins/PCG/Source/PCG/Private/Elements/PCGRandomChoice.cpp
@@ -26,6 +26,7 @@ namespace PCGRandomChoice
 
 		OutPointData->SetNumPoints(InIndexes.Num());
 		OutPointData->AllocateProperties(InPointData->GetAllocatedProperties());
+		OutPointData->CopyUnallocatedPropertiesFrom(InPointData);
 
 		// Order needs to be stable to sort this part of the array
 		Algo::Sort(InIndexes);
diff --git a/Engine/Plugins/PCG/Source/PCG/Private/Elements/PCGSampleTexture.cpp b/Engine/Plugins/PCG/Source/PCG/Private/Elements/PCGSampleTexture.cpp
--- a/Engine/Plugins/PCG/Source/PCG/Private/Elements/PCGSampleTexture.cpp
+++ b/Engine/Plugins/PCG/Source/PCG/Private/Elements/PCGSampleTexture.cpp
@@ -111,6 +111,7 @@ bool FPCGSampleTextureElement::ExecuteInternal(FPCGContext* Context) const
 			}
 
 			OutPointData->AllocateProperties(PropertiesToAllocate);
+			OutPointData->CopyUnallocatedPropertiesFrom(PointData);
 		};
 
 		auto MoveDataRangeFunc = [OutPointData](int32 RangeStartIndex, int32 MoveToIndex, int32 NumElements)
diff --git a/Engine/Plugins/PCG/Source/PCG/Private/Elements/PCGSelectPoints.cpp b/Engine/Plugins/PCG/Source/PCG/Private/Elements/PCGSelectPoints.cpp
--- a/Engine/Plugins/PCG/Source/PCG/Private/Elements/PCGSelectPoints.cpp
+++ b/Engine/Plugins/PCG/Source/PCG/Private/Elements/PCGSelectPoints.cpp
@@ -112,8 +112,8 @@ bool FPCGSelectPointsElement::ExecuteInternal(FPCGContext* Context) const
 			auto InitializeFunc = [SampledData, OriginalData]()
 			{
 				SampledData->SetNumPoints(OriginalData->GetNumPoints(), /*bInitializeValues=*/false);
-				SampledData->AllocateProperties(OriginalData->GetAllocatedProperties());
-				OriginalData->CopyUnallocatedPropertiesTo(SampledData);
+				SampledData->AllocateProperties(OriginalData->GetAllocatedProperties() | EPCGPointNativeProperties::Density);
+				SampledData->CopyUnallocatedPropertiesFrom(OriginalData);
 			};
 
 			auto AsyncProcessRangeFunc = [SampledData, OriginalData, Seed, Ratio, bKeepZeroDensityPoints](int32 StartReadIndex, int32 StartWriteIndex, int32 Count)
@@ -133,32 +133,20 @@ bool FPCGSelectPointsElement::ExecuteInternal(FPCGContext* Context) const
 					FRandomStream RandomSource(PCGHelpers::ComputeSeed(Seed, OriginalSeed));
 					float Chance = RandomSource.FRand();
 
-					bool bCopyPoint = false;
-
 					if (Chance < Ratio)
 					{
-						bCopyPoint = true;
+						SampledRanges.SetFromValueRanges(WriteIndex, OriginalRanges, ReadIndex);
 						SampledRanges.DensityRange[WriteIndex] = OriginalRanges.DensityRange[ReadIndex];
+						++NumWritten;
 					}
 #if WITH_EDITOR
 					else if (bKeepZeroDensityPoints)
 					{
-						bCopyPoint = true;
+						SampledRanges.SetFromValueRanges(WriteIndex, OriginalRanges, ReadIndex);
 						SampledRanges.DensityRange[WriteIndex] = 0.f;
-					}
-#endif
-					
-					if (bCopyPoint)
-					{
-						SampledRanges.TransformRange[WriteIndex] = OriginalRanges.TransformRange[ReadIndex];
-						SampledRanges.BoundsMinRange[WriteIndex] = OriginalRanges.BoundsMinRange[ReadIndex];
-						SampledRanges.BoundsMaxRange[WriteIndex] = OriginalRanges.BoundsMaxRange[ReadIndex];
-						SampledRanges.SeedRange[WriteIndex] = OriginalRanges.SeedRange[ReadIndex];
-						SampledRanges.MetadataEntryRange[WriteIndex] = OriginalRanges.MetadataEntryRange[ReadIndex];
-						SampledRanges.ColorRange[WriteIndex] = OriginalRanges.ColorRange[ReadIndex];
-						SampledRanges.SteepnessRange[WriteIndex] = OriginalRanges.SteepnessRange[ReadIndex];
 						++NumWritten;
 					}
+#endif
 				}
 							
 				return NumWritten;
diff --git a/Engine/Plugins/PCG/Source/PCG/Private/Elements/PCGSplitPoints.cpp b/Engine/Plugins/PCG/Source/PCG/Private/Elements/PCGSplitPoints.cpp
--- a/Engine/Plugins/PCG/Source/PCG/Private/Elements/PCGSplitPoints.cpp
+++ b/Engine/Plugins/PCG/Source/PCG/Private/Elements/PCGSplitPoints.cpp
@@ -71,6 +71,7 @@ bool FPCGSplitPointsElement::ExecuteInternal(FPCGContext* Context) const
 		if (!OutPointDataA->HasSpatialDataParent())
 		{
 			OutPointDataA->AllocateProperties(InputPointData->GetAllocatedProperties());
+			OutPointDataA->CopyUnallocatedPropertiesFrom(InputPointData);
 		}
 
 		OutPointDataA->AllocateProperties(EPCGPointNativeProperties::BoundsMax);
@@ -88,6 +89,7 @@ bool FPCGSplitPointsElement::ExecuteInternal(FPCGContext* Context) const
 		if (!OutPointDataB->HasSpatialDataParent())
 		{
 			OutPointDataB->AllocateProperties(InputPointData->GetAllocatedProperties());
+			OutPointDataB->CopyUnallocatedPropertiesFrom(InputPointData);
 		}
 
 		OutPointDataB->AllocateProperties(EPCGPointNativeProperties::BoundsMin);
diff --git a/Engine/Plugins/PCG/Source/PCG/Private/Elements/PCGWorldRaycast.cpp b/Engine/Plugins/PCG/Source/PCG/Private/Elements/PCGWorldRaycast.cpp
--- a/Engine/Plugins/PCG/Source/PCG/Private/Elements/PCGWorldRaycast.cpp
+++ b/Engine/Plugins/PCG/Source/PCG/Private/Elements/PCGWorldRaycast.cpp
@@ -499,6 +499,7 @@ bool FPCGWorldRaycastElement::ExecuteInternal(FPCGContext* InContext) const
 		if (OriginInputPointData)
 		{
 			OutputPointData->AllocateProperties(OriginInputPointData->GetAllocatedProperties());
+			OutputPointData->CopyUnallocatedPropertiesFrom(OriginInputPointData);
 		}
 
 		OutputPointData->AllocateProperties(EPCGPointNativeProperties::Transform | EPCGPointNativeProperties::Seed | EPCGPointNativeProperties::MetadataEntry);
diff --git a/Engine/Plugins/PCG/Source/PCG/Private/MeshSelectors/PCGMeshSelectorWeighted.cpp b/Engine/Plugins/PCG/Source/PCG/Private/MeshSelectors/PCGMeshSelectorWeighted.cpp
--- a/Engine/Plugins/PCG/Source/PCG/Private/MeshSelectors/PCGMeshSelectorWeighted.cpp
+++ b/Engine/Plugins/PCG/Source/PCG/Private/MeshSelectors/PCGMeshSelectorWeighted.cpp
@@ -239,15 +239,7 @@ bool UPCGMeshSelectorWeighted::SelectMeshInstances(
 
 				if (OutPointData && OutAttribute)
 				{
-					OutRanges.TransformRange[Context.CurrentWriteIndex] = InTransform;
-					OutRanges.MetadataEntryRange[Context.CurrentWriteIndex] = InMetadataEntry;
-					OutRanges.SeedRange[Context.CurrentWriteIndex] = InSeed;
-
-					OutRanges.SteepnessRange[Context.CurrentWriteIndex] = InRanges.SteepnessRange[InPointIndex];
-					OutRanges.DensityRange[Context.CurrentWriteIndex] = InRanges.DensityRange[InPointIndex];
-					OutRanges.BoundsMinRange[Context.CurrentWriteIndex] = InRanges.BoundsMinRange[InPointIndex];
-					OutRanges.BoundsMaxRange[Context.CurrentWriteIndex] = InRanges.BoundsMaxRange[InPointIndex];
-					OutRanges.ColorRange[Context.CurrentWriteIndex] = InRanges.ColorRange[InPointIndex];
+					OutRanges.SetFromValueRanges(Context.CurrentWriteIndex, InRanges, InPointIndex);
 
 					PCGMetadataValueKey* OutValueKey = MeshToValueKey.Find(Mesh);
 					if (!OutValueKey)
diff --git a/Engine/Plugins/PCG/Source/PCG/Private/MeshSelectors/PCGMeshSelectorWeightedByCategory.cpp b/Engine/Plugins/PCG/Source/PCG/Private/MeshSelectors/PCGMeshSelectorWeightedByCategory.cpp
--- a/Engine/Plugins/PCG/Source/PCG/Private/MeshSelectors/PCGMeshSelectorWeightedByCategory.cpp
+++ b/Engine/Plugins/PCG/Source/PCG/Private/MeshSelectors/PCGMeshSelectorWeightedByCategory.cpp
@@ -243,15 +243,7 @@ bool UPCGMeshSelectorWeightedByCategory::SelectMeshInstances(
 
 			if (OutPointData && OutAttribute)
 			{
-				OutRanges.TransformRange[Context.CurrentWriteIndex] = InTransform;
-				OutRanges.MetadataEntryRange[Context.CurrentWriteIndex] = InMetadataEntry;
-				OutRanges.SeedRange[Context.CurrentWriteIndex] = InSeed;
-
-				OutRanges.SteepnessRange[Context.CurrentWriteIndex] = InRanges.SteepnessRange[InPointIndex];
-				OutRanges.DensityRange[Context.CurrentWriteIndex] = InRanges.DensityRange[InPointIndex];
-				OutRanges.BoundsMinRange[Context.CurrentWriteIndex] = InRanges.BoundsMinRange[InPointIndex];
-				OutRanges.BoundsMaxRange[Context.CurrentWriteIndex] = InRanges.BoundsMaxRange[InPointIndex];
-				OutRanges.ColorRange[Context.CurrentWriteIndex] = InRanges.ColorRange[InPointIndex];
+				OutRanges.SetFromValueRanges(Context.CurrentWriteIndex, InRanges, InPointIndex);
 
 				PCGMetadataValueKey* OutValueKey = MeshToValueKey.Find(Mesh);
 				if (!OutValueKey)
diff --git a/Engine/Plugins/PCG/Source/PCG/Public/Data/PCGBasePointData.h b/Engine/Plugins/PCG/Source/PCG/Public/Data/PCGBasePointData.h
--- a/Engine/Plugins/PCG/Source/PCG/Public/Data/PCGBasePointData.h
+++ b/Engine/Plugins/PCG/Source/PCG/Public/Data/PCGBasePointData.h
@@ -70,7 +70,7 @@ public:
 	virtual void AllocateProperties(EPCGPointNativeProperties Properties) {}
 	virtual void FreeProperties(EPCGPointNativeProperties Properties) {}
 	virtual void MoveRange(int32 RangeStartIndex, int32 MoveToIndex, int32 NumElements) PURE_VIRTUAL(UPCGBasePointData::MoveRange, return;);
-	virtual void CopyUnallocatedPropertiesTo(UPCGBasePointData* OutPointData) const {}
+	virtual void CopyUnallocatedPropertiesFrom(const UPCGBasePointData* InPointData) {}
 
 	virtual TArray<FTransform> GetTransformsCopy() const PURE_VIRTUAL(UPCGBasePointData::GetTransformsCopy, return TArray<FTransform>(););
 	
@@ -349,14 +349,47 @@ struct FPCGPointValueRanges
 
 	FPCGPointValueRanges(UPCGBasePointData* InBasePointData, bool bAllocate = true)
 	{
-		TransformRange = InBasePointData->GetTransformValueRange(bAllocate);
-		DensityRange = InBasePointData->GetDensityValueRange(bAllocate);
-		SteepnessRange = InBasePointData->GetSteepnessValueRange(bAllocate);
-		BoundsMinRange = InBasePointData->GetBoundsMinValueRange(bAllocate);
-		BoundsMaxRange = InBasePointData->GetBoundsMaxValueRange(bAllocate);
-		ColorRange = InBasePointData->GetColorValueRange(bAllocate);
-		SeedRange = InBasePointData->GetSeedValueRange(bAllocate);
-		MetadataEntryRange = InBasePointData->GetMetadataEntryValueRange(bAllocate);
+		EPCGPointNativeProperties AllocatedProperties = InBasePointData->GetAllocatedProperties(/*bWithInheritance=*/false);
+
+		if (bAllocate || EnumHasAnyFlags(AllocatedProperties, EPCGPointNativeProperties::Transform))
+		{
+			TransformRange = InBasePointData->GetTransformValueRange(bAllocate);
+		}
+
+		if (bAllocate || EnumHasAnyFlags(AllocatedProperties, EPCGPointNativeProperties::Density))
+		{
+			DensityRange = InBasePointData->GetDensityValueRange(bAllocate);
+		}
+
+		if (bAllocate || EnumHasAnyFlags(AllocatedProperties, EPCGPointNativeProperties::Steepness))
+		{
+			SteepnessRange = InBasePointData->GetSteepnessValueRange(bAllocate);
+		}
+
+		if (bAllocate || EnumHasAnyFlags(AllocatedProperties, EPCGPointNativeProperties::BoundsMin))
+		{
+			BoundsMinRange = InBasePointData->GetBoundsMinValueRange(bAllocate);
+		}
+
+		if (bAllocate || EnumHasAnyFlags(AllocatedProperties, EPCGPointNativeProperties::BoundsMax))
+		{
+			BoundsMaxRange = InBasePointData->GetBoundsMaxValueRange(bAllocate);
+		}
+
+		if (bAllocate || EnumHasAnyFlags(AllocatedProperties, EPCGPointNativeProperties::Color))
+		{
+			ColorRange = InBasePointData->GetColorValueRange(bAllocate);
+		}
+
+		if (bAllocate || EnumHasAnyFlags(AllocatedProperties, EPCGPointNativeProperties::Seed))
+		{
+			SeedRange = InBasePointData->GetSeedValueRange(bAllocate);
+		}
+
+		if (bAllocate || EnumHasAnyFlags(AllocatedProperties, EPCGPointNativeProperties::MetadataEntry))
+		{
+			MetadataEntryRange = InBasePointData->GetMetadataEntryValueRange(bAllocate);
+		}
 	}
 
 	TPCGValueRange<FTransform> TransformRange;
@@ -370,25 +403,87 @@ struct FPCGPointValueRanges
 
 	void SetFromPoint(int32 Index, const FPCGPoint& Point)
 	{
-		TransformRange[Index] = Point.Transform;
-		DensityRange[Index] = Point.Density;
-		SteepnessRange[Index] = Point.Steepness;
-		BoundsMinRange[Index] = Point.BoundsMin;
-		BoundsMaxRange[Index] = Point.BoundsMax;
-		ColorRange[Index] = Point.Color;
-		SeedRange[Index] = Point.Seed;
-		MetadataEntryRange[Index] = Point.MetadataEntry;
+		if (!TransformRange.IsEmpty())
+		{
+			TransformRange[Index] = Point.Transform;
+		}
+
+		if (!DensityRange.IsEmpty())
+		{
+			DensityRange[Index] = Point.Density;
+		}
+		
+		if (!SteepnessRange.IsEmpty())
+		{
+			SteepnessRange[Index] = Point.Steepness;
+		}
+
+		if (!BoundsMinRange.IsEmpty())
+		{
+			BoundsMinRange[Index] = Point.BoundsMin;
+		}
+
+		if (!BoundsMaxRange.IsEmpty())
+		{
+			BoundsMaxRange[Index] = Point.BoundsMax;
+		}
+
+		if (!ColorRange.IsEmpty())
+		{
+			ColorRange[Index] = Point.Color;
+		}
+		
+		if (!SeedRange.IsEmpty())
+		{
+			SeedRange[Index] = Point.Seed;
+		}
+
+		if (!MetadataEntryRange.IsEmpty())
+		{
+			MetadataEntryRange[Index] = Point.MetadataEntry;
+		}
 	}
 		
 	void SetFromValueRanges(int32 WriteIndex, const FConstPCGPointValueRanges& ReadRange, int32 ReadIndex)
 	{
-		TransformRange[WriteIndex] = ReadRange.TransformRange[ReadIndex];
-		DensityRange[WriteIndex] = ReadRange.DensityRange[ReadIndex];
-		SteepnessRange[WriteIndex] = ReadRange.SteepnessRange[ReadIndex];
-		BoundsMinRange[WriteIndex] = ReadRange.BoundsMinRange[ReadIndex];
-		BoundsMaxRange[WriteIndex] = ReadRange.BoundsMaxRange[ReadIndex];
-		ColorRange[WriteIndex] = ReadRange.ColorRange[ReadIndex];
-		SeedRange[WriteIndex] = ReadRange.SeedRange[ReadIndex];
-		MetadataEntryRange[WriteIndex] = ReadRange.MetadataEntryRange[ReadIndex];
+		if (!TransformRange.IsEmpty())
+		{
+			TransformRange[WriteIndex] = ReadRange.TransformRange[ReadIndex];
+		}
+
+		if (!DensityRange.IsEmpty())
+		{
+			DensityRange[WriteIndex] = ReadRange.DensityRange[ReadIndex];
+		}
+
+		if (!SteepnessRange.IsEmpty())
+		{
+			SteepnessRange[WriteIndex] = ReadRange.SteepnessRange[ReadIndex];
+		}
+
+		if (!BoundsMinRange.IsEmpty())
+		{
+			BoundsMinRange[WriteIndex] = ReadRange.BoundsMinRange[ReadIndex];
+		}
+
+		if (!BoundsMaxRange.IsEmpty())
+		{
+			BoundsMaxRange[WriteIndex] = ReadRange.BoundsMaxRange[ReadIndex];
+		}
+
+		if (!ColorRange.IsEmpty())
+		{
+			ColorRange[WriteIndex] = ReadRange.ColorRange[ReadIndex];
+		}
+
+		if (!SeedRange.IsEmpty())
+		{
+			SeedRange[WriteIndex] = ReadRange.SeedRange[ReadIndex];
+		}
+
+		if (!MetadataEntryRange.IsEmpty())
+		{
+			MetadataEntryRange[WriteIndex] = ReadRange.MetadataEntryRange[ReadIndex];
+		}
 	}
 };
diff --git a/Engine/Plugins/PCG/Source/PCG/Public/Data/PCGPointArrayData.h b/Engine/Plugins/PCG/Source/PCG/Public/Data/PCGPointArrayData.h
--- a/Engine/Plugins/PCG/Source/PCG/Public/Data/PCGPointArrayData.h
+++ b/Engine/Plugins/PCG/Source/PCG/Public/Data/PCGPointArrayData.h
@@ -38,7 +38,7 @@ public:
 	virtual void AllocateProperties(EPCGPointNativeProperties Properties) override;
 	virtual void FreeProperties(EPCGPointNativeProperties Properties) override;
 	virtual void MoveRange(int32 RangeStartIndex, int32 MoveToIndex, int32 NumElements) override;
-	virtual void CopyUnallocatedPropertiesTo(UPCGBasePointData* OutPointData) const override;
+	virtual void CopyUnallocatedPropertiesFrom(const UPCGBasePointData* InPointData) override;
 
 	virtual TArray<FTransform> GetTransformsCopy() const { return PointArray.GetTransformCopy(); }
 
diff --git a/Engine/Plugins/PCG/Source/PCG/Public/Data/PCGSpatialDataTpl.h b/Engine/Plugins/PCG/Source/PCG/Public/Data/PCGSpatialDataTpl.h
--- a/Engine/Plugins/PCG/Source/PCG/Public/Data/PCGSpatialDataTpl.h
+++ b/Engine/Plugins/PCG/Source/PCG/Public/Data/PCGSpatialDataTpl.h
@@ -74,7 +74,7 @@ namespace FPCGSpatialDataProcessing
 		{
 			TargetData->SetNumPoints(NumIterations, /*bInitializeValues=*/false);
 			TargetData->AllocateProperties(PropertiesToAllocate);
-			SourceData->CopyUnallocatedPropertiesTo(TargetData);
+			TargetData->CopyUnallocatedPropertiesFrom(SourceData);
 		};
 
 		auto ProcessRange = [RangeFunc = MoveTemp(InProcessRange), SourceData, TargetData](int32 StartReadIndex, int32 StartWriteIndex, int32 Count) -> int32
