[PCG] Refactor ColumnInfo creation to reduce the number of calls

* ColumnInfos are now having their own keys, this will be important if the keys are differents for the same set of attributes
* Functions are now exposed, so someone from another plugin can use them
* Add utility functions to add columns for Enum Properties (which allow to extract the tooltip too)
* Add utility functions to add collumns for a given type (which allow to create the different extractors automatically)
* Since they were not exported, no deprecation needed.

#rnx

diff --git a/Engine/Plugins/PCG/Source/PCG/Private/Data/PCGBasePointData.cpp b/Engine/Plugins/PCG/Source/PCG/Private/Data/PCGBasePointData.cpp
--- a/Engine/Plugins/PCG/Source/PCG/Private/Data/PCGBasePointData.cpp
+++ b/Engine/Plugins/PCG/Source/PCG/Private/Data/PCGBasePointData.cpp
@@ -1112,7 +1112,7 @@ bool UPCGBasePointData::SetDomainFromDomainID(const FPCGMetadataDomainID& InDoma
 {
 	if (InDomainID == PCGMetadataDomainID::Elements)
 	{
-		InOutSelector.SetDomainName(PCGPointDataConstants::ElementsDomainName);
+		InOutSelector.SetDomainName(PCGPointDataConstants::ElementsDomainName, /*bResetExtraNames=*/false);
 		return true;
 	}
 	else
diff --git a/Engine/Plugins/PCG/Source/PCG/Private/Metadata/PCGAttributePropertySelector.cpp b/Engine/Plugins/PCG/Source/PCG/Private/Metadata/PCGAttributePropertySelector.cpp
--- a/Engine/Plugins/PCG/Source/PCG/Private/Metadata/PCGAttributePropertySelector.cpp
+++ b/Engine/Plugins/PCG/Source/PCG/Private/Metadata/PCGAttributePropertySelector.cpp
@@ -263,9 +263,9 @@ FString FPCGAttributePropertySelector::GetAttributePropertyAccessorsString(bool
 	}
 }
 
-FString FPCGAttributePropertySelector::ToString() const
+FString FPCGAttributePropertySelector::ToString(bool bSkipDomain) const
 {
-	const FString Domain = GetDomainString(/*bAddLeadingQualifier=*/true);
+	const FString Domain = !bSkipDomain ? GetDomainString(/*bAddLeadingQualifier=*/true) : FString{};
 	const FString Attribute = GetAttributePropertyString(/*bAddPropertyQualifier=*/true);
 	const FString Accessors = GetAttributePropertyAccessorsString(/*bAddLeadingSeparator*/true);
 
diff --git a/Engine/Plugins/PCG/Source/PCG/Private/PCGData.cpp b/Engine/Plugins/PCG/Source/PCG/Private/PCGData.cpp
--- a/Engine/Plugins/PCG/Source/PCG/Private/PCGData.cpp
+++ b/Engine/Plugins/PCG/Source/PCG/Private/PCGData.cpp
@@ -108,12 +108,12 @@ bool UPCGData::SetDomainFromDomainID(const FPCGMetadataDomainID& InDomainID, FPC
 {
 	if (InDomainID == PCGMetadataDomainID::Data)
 	{
-		InOutSelector.SetDomainName(PCGDataConstants::GlobalDomainName);
+		InOutSelector.SetDomainName(PCGDataConstants::GlobalDomainName, /*bResetExtraNames=*/false);
 		return true;
 	}
 	else if (InDomainID == PCGMetadataDomainID::Default)
 	{
-		InOutSelector.SetDomainName(NAME_None);
+		InOutSelector.SetDomainName(NAME_None, /*bResetExtraNames=*/false);
 		return true;
 	}
 	else
diff --git a/Engine/Plugins/PCG/Source/PCG/Private/PCGParamData.cpp b/Engine/Plugins/PCG/Source/PCG/Private/PCGParamData.cpp
--- a/Engine/Plugins/PCG/Source/PCG/Private/PCGParamData.cpp
+++ b/Engine/Plugins/PCG/Source/PCG/Private/PCGParamData.cpp
@@ -180,7 +180,7 @@ bool UPCGParamData::SetDomainFromDomainID(const FPCGMetadataDomainID& InDomainID
 {
 	if (InDomainID == PCGMetadataDomainID::Elements)
 	{
-		InOutSelector.SetDomainName(PCGParamDataConstants::ElementsDomainName);
+		InOutSelector.SetDomainName(PCGParamDataConstants::ElementsDomainName, /*bResetExtraNames=*/false);
 		return true;
 	}
 	else
diff --git a/Engine/Plugins/PCG/Source/PCG/Public/Metadata/PCGAttributePropertySelector.h b/Engine/Plugins/PCG/Source/PCG/Public/Metadata/PCGAttributePropertySelector.h
--- a/Engine/Plugins/PCG/Source/PCG/Public/Metadata/PCGAttributePropertySelector.h
+++ b/Engine/Plugins/PCG/Source/PCG/Public/Metadata/PCGAttributePropertySelector.h
@@ -80,7 +80,7 @@ public:
 	FName GetName() const;
 
 	// Returns qualified attribute/property name with the accessors.
-	FString ToString() const;
+	FString ToString(bool bSkipDomain = false) const;
 
 	// Returns attribute/property name only, with optional '@' qualifier for domains.
 	FString GetDomainString(bool bAddLeadingQualifier = true) const;
@@ -92,7 +92,7 @@ public:
 	FString GetAttributePropertyAccessorsString(bool bAddLeadingSeparator) const;
 
 	// Returns the text to display in the widget.
-	FText GetDisplayText() const { return FText::FromString(ToString()); }
+	FText GetDisplayText(bool bSkipDomain = false) const { return FText::FromString(ToString(bSkipDomain)); }
 
 	// Return true if the underlying name is valid.
 	bool IsValid() const;
diff --git a/Engine/Plugins/PCG/Source/PCG/Public/PCGCommon.h b/Engine/Plugins/PCG/Source/PCG/Public/PCGCommon.h
--- a/Engine/Plugins/PCG/Source/PCG/Public/PCGCommon.h
+++ b/Engine/Plugins/PCG/Source/PCG/Public/PCGCommon.h
@@ -233,6 +233,7 @@ namespace PCGObjectMetadata
 	const FName OverrideAliases = TEXT("PCG_OverrideAliases");
 	const FName DiscardPropertySelection = TEXT("PCG_DiscardPropertySelection");
 	const FName DiscardExtraSelection = TEXT("PCG_DiscardExtraSelection");
+	const FName EnumMetadataDomain = TEXT("PCG_MetadataDomain");
 
 	// Metadata usable in UPROPERTY for customizing the behavior when displaying the property in a property panel or graph node
 	enum
@@ -252,6 +253,13 @@ namespace PCGObjectMetadata
 		/// [PropertyMetadata] For FPCGAttributePropertySelector, won't display the extra property items in the dropdown
 		PCG_DiscardExtraSelection,
 	};
+
+	// Metadata usable in UENUM for customizing the entry in the attribute property selector
+	enum
+	{
+		/// [EnumMetadata] Specify the domain for this entry.
+		PCG_MetadataDomain
+	};
 }
 
 namespace PCGFeatureSwitches
diff --git a/Engine/Plugins/PCG/Source/PCG/Public/PCGDataVisualization.h b/Engine/Plugins/PCG/Source/PCG/Public/PCGDataVisualization.h
--- a/Engine/Plugins/PCG/Source/PCG/Public/PCGDataVisualization.h
+++ b/Engine/Plugins/PCG/Source/PCG/Public/PCGDataVisualization.h
@@ -46,15 +46,28 @@ struct FPCGTableVisualizerColumnInfo
 	float Width = -1.0f; // Will be calculated automatically if < 0
 	EPCGTableVisualizerCellAlignment CellAlignment = EPCGTableVisualizerCellAlignment::Right;
 	TSharedPtr<const IPCGAttributeAccessor> Accessor;
+	TSharedPtr<const IPCGAttributeAccessorKeys> AccessorKeys;
 };
 
 struct FPCGTableVisualizerInfo
 {
+	// Note: We need to explicitly disable warnings on these constructors/operators for clang to be happy with deprecated variables
+	PRAGMA_DISABLE_DEPRECATION_WARNINGS
+	FPCGTableVisualizerInfo() = default;
+	~FPCGTableVisualizerInfo() = default;
+	FPCGTableVisualizerInfo(const FPCGTableVisualizerInfo&) = default;
+	FPCGTableVisualizerInfo(FPCGTableVisualizerInfo&&) = default;
+	FPCGTableVisualizerInfo& operator=(const FPCGTableVisualizerInfo&) = default;
+	FPCGTableVisualizerInfo& operator=(FPCGTableVisualizerInfo&&) = default;
+	PRAGMA_ENABLE_DEPRECATION_WARNINGS
+	
 	const UPCGData* Data = nullptr;
 	TArray<FPCGTableVisualizerColumnInfo> ColumnInfos;
 	EPCGTableVisualizerColumnSortingMode SortingMode = EPCGTableVisualizerColumnSortingMode::Ascending;
 	FName SortingColumn = NAME_None;
 	TFunction<void(const UPCGData*, TArrayView<const int>)> FocusOnDataCallback = nullptr;
+	
+	UE_DEPRECATED(5.6, "Set the keys in the ColumnInfo")
 	TSharedPtr<const IPCGAttributeAccessorKeys> AccessorKeys;
 };
 
diff --git a/Engine/Plugins/PCG/Source/PCGEditor/Private/DataVisualizations/PCGDataVisualizationHelpers.cpp b/Engine/Plugins/PCG/Source/PCGEditor/Private/DataVisualizations/PCGDataVisualizationHelpers.cpp
--- a/Engine/Plugins/PCG/Source/PCGEditor/Private/DataVisualizations/PCGDataVisualizationHelpers.cpp
+++ b/Engine/Plugins/PCG/Source/PCGEditor/Private/DataVisualizations/PCGDataVisualizationHelpers.cpp
@@ -22,75 +22,110 @@ namespace PCGDataVisualizationConstants
 
 namespace PCGDataVisualizationHelpers
 {
-	FPCGTableVisualizerColumnInfo CreateColumnInfo(const UPCGData* Data, FName Id, const FText& Label, TFunction<TSharedPtr<const IPCGAttributeAccessor>()> CreateAccessorFunc)
+	void AddColumnInfo(FPCGTableVisualizerInfo& OutInfo, const UPCGData* Data, const FPCGAttributePropertySelector& InSelector, const FColumnInfoOverrides& Overrides)
 	{
-		FPCGTableVisualizerColumnInfo ColumnInfo;
-		ColumnInfo.Id = Id;
-		ColumnInfo.Label = Label;
-
-		FPCGAttributePropertyInputSelector Selector = FPCGAttributePropertySelector::CreateSelectorFromString<FPCGAttributePropertyInputSelector>(Id.ToString());
-		ColumnInfo.Accessor = CreateAccessorFunc ? CreateAccessorFunc() : TSharedPtr<const IPCGAttributeAccessor>(PCGAttributeAccessorHelpers::CreateConstAccessor(Data, Selector).Release());
+		FPCGTableVisualizerColumnInfo& ColumnInfo = OutInfo.ColumnInfos.Emplace_GetRef();
+		
+		ColumnInfo.Accessor = Overrides.CreateAccessorFuncOverride ? Overrides.CreateAccessorFuncOverride() : TSharedPtr<const IPCGAttributeAccessor>(PCGAttributeAccessorHelpers::CreateConstAccessor(Data, InSelector).Release());
+		ColumnInfo.AccessorKeys = Overrides.CreateAccessorKeysFuncOverride ? Overrides.CreateAccessorKeysFuncOverride() : TSharedPtr<const IPCGAttributeAccessorKeys>(PCGAttributeAccessorHelpers::CreateConstKeys(Data, InSelector).Release());
 
 		if (ensure(ColumnInfo.Accessor))
 		{
-			ColumnInfo.Tooltip = FText::FromString(PCG::Private::GetTypeName(ColumnInfo.Accessor->GetUnderlyingType()));
-		}
-
-		return ColumnInfo;
-	}
+			const uint16 Type = Overrides.TypeOverride == EPCGMetadataTypes::Unknown ? ColumnInfo.Accessor->GetUnderlyingType() : static_cast<uint16>(Overrides.TypeOverride);
 
-	FPCGTableVisualizerColumnInfo CreateMetadataColumnInfo(const FPCGMetadataDomain* MetadataDomain, const FPCGMetadataAttributeBase* Attribute, const FName& Id, EPCGMetadataTypes MetadataType, const TCHAR* PostFix = nullptr)
-	{
-		FString OriginalId = Id.ToString();
-		if (PostFix)
-		{
-			OriginalId.Append(PostFix);
-		}
+			if (Overrides.bAddTypeToTooltip && !Overrides.TooltipOverride.IsEmpty())
+			{
+				ColumnInfo.Tooltip = FText::Format(INVTEXT("{0} ({1})"), Overrides.TooltipOverride, PCG::Private::GetTypeNameText(Type));
+			}
+			else if (!Overrides.TooltipOverride.IsEmpty())
+			{
+				ColumnInfo.Tooltip = Overrides.TooltipOverride;
+			}
+			else
+			{
+				ColumnInfo.Tooltip = PCG::Private::GetTypeNameText(Type);
+			}
 
-		FString IdString(OriginalId);
-		FText Label;
+			if (PCG::Private::IsOfTypes<FString>(Type))
+			{
+				ColumnInfo.CellAlignment = EPCGTableVisualizerCellAlignment::Left;
+				ColumnInfo.Width = PCGEditorGraphAttributeListView::MaxColumnWidth;
+			}
+			
+			if (!Overrides.LabelOverride.IsEmpty())
+			{
+				ColumnInfo.Label = Overrides.LabelOverride;
+			}
+			else if (InSelector.GetName() == NAME_None)
+			{
+				ColumnInfo.Label = FText::Format(LOCTEXT("NoneLabelFormat", "{1} ({0})"), UEnum::GetDisplayValueAsText(static_cast<EPCGMetadataTypes>(ColumnInfo.Accessor->GetUnderlyingType())), InSelector.GetDisplayText(/*bSkipDomain=*/true));
+			}
+			else
+			{
+				ColumnInfo.Label = InSelector.GetDisplayText(/*bSkipDomain=*/true);
+			}
 
-		if (Id == NAME_None)
-		{
-			IdString = FString::Printf(TEXT("%s%s"), *PCGDataVisualizationConstants::NoneAttributeId, PostFix);
-			Label = FText::Format(LOCTEXT("NoneLabelFormat", "None{1} ({0})"), UEnum::GetDisplayValueAsText(MetadataType), FText::FromString(PostFix));
+			ColumnInfo.Id = *ColumnInfo.Label.ToString();
 		}
-		else if (Label.IsEmpty())
-		{
-			Label = FText::FromString(IdString);
-		}
-
-		FPCGTableVisualizerColumnInfo ColumnInfo;
-		ColumnInfo.Id = FName(IdString);
-		ColumnInfo.Label = Label;
+	}
 
-		if (MetadataType == EPCGMetadataTypes::String)
+	void AddColumnInfoExtraNames(FPCGTableVisualizerInfo& OutInfo, const UPCGData* Data, const FPCGAttributePropertySelector& BaseSelector, const TArray<FString>& AllExtraNames, const FColumnInfoOverrides& Overrides)
+	{
+		FPCGAttributePropertySelector SelectorCopy = BaseSelector;
+		FString& SelectorExtraName = SelectorCopy.GetExtraNamesMutable().Emplace_GetRef();
+		for (const FString& ExtraName : AllExtraNames)
 		{
-			ColumnInfo.CellAlignment = EPCGTableVisualizerCellAlignment::Left;
-			ColumnInfo.Width = PCGEditorGraphAttributeListView::MaxColumnWidth;
+			SelectorExtraName = ExtraName;
+			AddColumnInfo(OutInfo, Data, SelectorCopy, Overrides);
 		}
+	}
 
-		if (MetadataDomain && Attribute)
-		{
-			FPCGAttributePropertyInputSelector Selector;
-			Selector.Update(OriginalId);
-
-			// Need this gymnastic for the chain accessor
-			TUniquePtr<IPCGAttributeAccessor> Accessor = TUniquePtr<IPCGAttributeAccessor>(const_cast<IPCGAttributeAccessor*>(PCGAttributeAccessorHelpers::CreateConstAccessor(Attribute, MetadataDomain).Release()));
+	template <> void AddTypedColumnInfo_Impl<FVector2D>(FPCGTableVisualizerInfo& OutInfo, const UPCGData* Data, const FPCGAttributePropertySelector& InSelector, const FColumnInfoOverrides& Overrides)
+	{
+		AddColumnInfoExtraNames(OutInfo, Data, InSelector, {TEXT("X"), TEXT("Y")}, Overrides);
+	}
+	
+	template <> void AddTypedColumnInfo_Impl<FVector>(FPCGTableVisualizerInfo& OutInfo, const UPCGData* Data, const FPCGAttributePropertySelector& InSelector, const FColumnInfoOverrides& Overrides)
+	{
+		AddColumnInfoExtraNames(OutInfo, Data, InSelector, {TEXT("X"), TEXT("Y"), TEXT("Z")}, Overrides);
+	}
+	
+	template <> void AddTypedColumnInfo_Impl<FVector4>(FPCGTableVisualizerInfo& OutInfo, const UPCGData* Data, const FPCGAttributePropertySelector& InSelector, const FColumnInfoOverrides& Overrides)
+	{
+		AddColumnInfoExtraNames(OutInfo, Data, InSelector, {TEXT("X"), TEXT("Y"), TEXT("Z"), TEXT("W")}, Overrides);
+	}
 
-			bool bSuccess = false;
-			ColumnInfo.Accessor = TSharedPtr<const IPCGAttributeAccessor>(PCGAttributeAccessorHelpers::CreateChainAccessor(std::move(Accessor), Selector, bSuccess, /*bQuiet=*/true).Release());
+	template <> void AddTypedColumnInfo_Impl<FLinearColor>(FPCGTableVisualizerInfo& OutInfo, const UPCGData* Data, const FPCGAttributePropertySelector& InSelector, const FColumnInfoOverrides& Overrides)
+	{
+		AddColumnInfoExtraNames(OutInfo, Data, InSelector, {TEXT("R"), TEXT("G"), TEXT("B"), TEXT("A")}, Overrides);
+	}
+	
+	template <> void AddTypedColumnInfo_Impl<FQuat>(FPCGTableVisualizerInfo& OutInfo, const UPCGData* Data, const FPCGAttributePropertySelector& InSelector, const FColumnInfoOverrides& Overrides)
+	{
+		return AddTypedColumnInfo<FVector4>(OutInfo, Data, InSelector, Overrides);
+	}
+	
+	template <> void AddTypedColumnInfo_Impl<FRotator>(FPCGTableVisualizerInfo& OutInfo, const UPCGData* Data, const FPCGAttributePropertySelector& InSelector, const FColumnInfoOverrides& Overrides)
+	{
+		AddColumnInfoExtraNames(OutInfo, Data, InSelector, {TEXT("Roll"), TEXT("Pitch"), TEXT("Yaw")}, Overrides);
+	}
+	
+	template <> void AddTypedColumnInfo_Impl<FTransform>(FPCGTableVisualizerInfo& OutInfo, const UPCGData* Data, const FPCGAttributePropertySelector& InSelector, const FColumnInfoOverrides& Overrides)
+	{
+		FPCGAttributePropertySelector SelectorCopy = InSelector;
+		FString& SelectorExtraName = SelectorCopy.GetExtraNamesMutable().Emplace_GetRef();
 
-			if (bSuccess && ColumnInfo.Accessor)
-			{
-				ColumnInfo.Tooltip = FText::FromString(PCG::Private::GetTypeName(ColumnInfo.Accessor->GetUnderlyingType()));
-			}
-		}
+		SelectorExtraName = TEXT("Position");
+		AddTypedColumnInfo<FVector>(OutInfo, Data, SelectorCopy, Overrides);
 
-		return ColumnInfo;
+		SelectorExtraName = TEXT("Rotation");
+		AddTypedColumnInfo<FRotator>(OutInfo, Data, SelectorCopy, Overrides);
+		
+		SelectorExtraName = TEXT("Scale");
+		AddTypedColumnInfo<FVector>(OutInfo, Data, SelectorCopy, Overrides);
 	}
 
-	void CreateMetadataColumnInfos(const UPCGData* Data, TArray<FPCGTableVisualizerColumnInfo>& OutColumnInfos, const FPCGMetadataDomainID& DomainID)
+	void CreateMetadataColumnInfos(const UPCGData* Data, FPCGTableVisualizerInfo& OutInfo, const FPCGMetadataDomainID& DomainID)
 	{
 		const UPCGMetadata* Metadata = Data->ConstMetadata();
 		if (!Metadata)
@@ -103,84 +138,31 @@ namespace PCGDataVisualizationHelpers
 		{
 			return;
 		}
-
+		
 		TArray<FName> AttributeNames;
 		TArray<EPCGMetadataTypes> AttributeTypes;
 		MetadataDomain->GetAttributes(AttributeNames, AttributeTypes);
 
-		for (int32 I = 0; I < AttributeNames.Num(); I++)
+		// Domain is constant for all the attributes
+		FPCGAttributePropertySelector Selector{};
+		Data->SetDomainFromDomainID(DomainID, Selector);
+		
+		auto AddInfo = [Data, &OutInfo, &Selector]<typename T>(T, const FName AttributeName)
 		{
-			const FName& AttributeName = AttributeNames[I];
-			const EPCGMetadataTypes AttributeType = AttributeTypes[I];
-			FName ColumnName = AttributeName;
-			const FPCGMetadataAttributeBase* Attribute = MetadataDomain->GetConstAttribute(AttributeName);
-			check(Attribute);
+			Selector.SetAttributeName(AttributeName);
+			AddTypedColumnInfo<T>(OutInfo, Data, Selector);
+		};
 
-			switch (AttributeType)
-			{
-			case EPCGMetadataTypes::Float:
-			case EPCGMetadataTypes::Double:
-			case EPCGMetadataTypes::Integer32:
-			case EPCGMetadataTypes::Integer64:
-			case EPCGMetadataTypes::Boolean:
-			case EPCGMetadataTypes::String:
-			case EPCGMetadataTypes::Name:
-			case EPCGMetadataTypes::SoftObjectPath:
-			case EPCGMetadataTypes::SoftClassPath:
-			{
-				OutColumnInfos.Add(CreateMetadataColumnInfo(MetadataDomain, Attribute, ColumnName, AttributeType, TEXT("")));
-				break;
-			}
-			case EPCGMetadataTypes::Vector2:
-			{
-				OutColumnInfos.Add(CreateMetadataColumnInfo(MetadataDomain, Attribute, ColumnName, AttributeType, TEXT(".X")));
-				OutColumnInfos.Add(CreateMetadataColumnInfo(MetadataDomain, Attribute, ColumnName, AttributeType, TEXT(".Y")));
-				break;
-			}
-			case EPCGMetadataTypes::Vector:
-			{
-				OutColumnInfos.Add(CreateMetadataColumnInfo(MetadataDomain, Attribute, ColumnName, AttributeType, TEXT(".X")));
-				OutColumnInfos.Add(CreateMetadataColumnInfo(MetadataDomain, Attribute, ColumnName, AttributeType, TEXT(".Y")));
-				OutColumnInfos.Add(CreateMetadataColumnInfo(MetadataDomain, Attribute, ColumnName, AttributeType, TEXT(".Z")));
-				break;
-			}
-			case EPCGMetadataTypes::Vector4:
-			case EPCGMetadataTypes::Quaternion:
-			{
-				OutColumnInfos.Add(CreateMetadataColumnInfo(MetadataDomain, Attribute, ColumnName, AttributeType, TEXT(".X")));
-				OutColumnInfos.Add(CreateMetadataColumnInfo(MetadataDomain, Attribute, ColumnName, AttributeType, TEXT(".Y")));
-				OutColumnInfos.Add(CreateMetadataColumnInfo(MetadataDomain, Attribute, ColumnName, AttributeType, TEXT(".Z")));
-				OutColumnInfos.Add(CreateMetadataColumnInfo(MetadataDomain, Attribute, ColumnName, AttributeType, TEXT(".W")));
-				break;
-			}
-			case EPCGMetadataTypes::Rotator:
-			{
-				OutColumnInfos.Add(CreateMetadataColumnInfo(MetadataDomain, Attribute, ColumnName, AttributeType, TEXT(".Roll")));
-				OutColumnInfos.Add(CreateMetadataColumnInfo(MetadataDomain, Attribute, ColumnName, AttributeType, TEXT(".Pitch")));
-				OutColumnInfos.Add(CreateMetadataColumnInfo(MetadataDomain, Attribute, ColumnName, AttributeType, TEXT(".Yaw")));
-				break;
-			}
-			case EPCGMetadataTypes::Transform:
-			{
-				OutColumnInfos.Add(CreateMetadataColumnInfo(MetadataDomain, Attribute, ColumnName, AttributeType, TEXT(".Position.X")));
-				OutColumnInfos.Add(CreateMetadataColumnInfo(MetadataDomain, Attribute, ColumnName, AttributeType, TEXT(".Position.Y")));
-				OutColumnInfos.Add(CreateMetadataColumnInfo(MetadataDomain, Attribute, ColumnName, AttributeType, TEXT(".Position.Z")));
-				OutColumnInfos.Add(CreateMetadataColumnInfo(MetadataDomain, Attribute, ColumnName, AttributeType, TEXT(".Rotation.Roll")));
-				OutColumnInfos.Add(CreateMetadataColumnInfo(MetadataDomain, Attribute, ColumnName, AttributeType, TEXT(".Rotation.Pitch")));
-				OutColumnInfos.Add(CreateMetadataColumnInfo(MetadataDomain, Attribute, ColumnName, AttributeType, TEXT(".Rotation.Yaw")));
-				OutColumnInfos.Add(CreateMetadataColumnInfo(MetadataDomain, Attribute, ColumnName, AttributeType, TEXT(".Scale.X")));
-				OutColumnInfos.Add(CreateMetadataColumnInfo(MetadataDomain, Attribute, ColumnName, AttributeType, TEXT(".Scale.Y")));
-				OutColumnInfos.Add(CreateMetadataColumnInfo(MetadataDomain, Attribute, ColumnName, AttributeType, TEXT(".Scale.Z")));
-				break;
-			}
-			default:
-				break;
-			}
+		for (int32 I = 0; I < AttributeNames.Num(); I++)
+		{
+			PCGMetadataAttribute::CallbackWithRightType(static_cast<uint16>(AttributeTypes[I]), AddInfo, AttributeNames[I]);
 		}
 	}
 
 	FPCGTableVisualizerInfo CreateDefaultMetadataColumnInfos(const UPCGData* Data, const FPCGMetadataDomainID& DomainID)
 	{
+		check(Data);
+		
 		const UPCGMetadata* Metadata = Data->ConstMetadata();
 		const FPCGMetadataDomain* MetadataDomain = Metadata ? Metadata->GetConstMetadataDomain(DomainID) : nullptr;
 		if (!MetadataDomain)
@@ -190,15 +172,15 @@ namespace PCGDataVisualizationHelpers
 		
 		FPCGTableVisualizerInfo Info;
 		Info.Data = Data;
-		Info.ColumnInfos.Add(CreateColumnInfo(Data, PCGDataVisualizationConstants::NAME_Index, PCGDataVisualizationConstants::TEXT_Index));
-		Info.SortingColumn = PCGDataVisualizationConstants::NAME_Index;
 
+		FPCGAttributePropertySelector IndexSelector = FPCGAttributePropertySelector::CreateExtraPropertySelector(EPCGExtraProperties::Index);
+		Data->SetDomainFromDomainID(DomainID, IndexSelector);
 		
-		const bool bAddDefaultValueIfEmpty = DomainID == PCGMetadataDomainID::Data;
-		Info.AccessorKeys = MakeShared<const FPCGAttributeAccessorKeysEntries>(MetadataDomain, /*bAddDefaultValueIfEmpty=*/bAddDefaultValueIfEmpty);
+		AddColumnInfo(Info, Data, IndexSelector);
+		Info.SortingColumn = Info.ColumnInfos.Last().Id;
 
 		// Add Metadata Columns
-		CreateMetadataColumnInfos(Data, Info.ColumnInfos, DomainID);
+		CreateMetadataColumnInfos(Data, Info, DomainID);
 
 		return Info;
 	}
diff --git a/Engine/Plugins/PCG/Source/PCGEditor/Private/DataVisualizations/PCGSpatialDataVisualization.cpp b/Engine/Plugins/PCG/Source/PCGEditor/Private/DataVisualizations/PCGSpatialDataVisualization.cpp
--- a/Engine/Plugins/PCG/Source/PCGEditor/Private/DataVisualizations/PCGSpatialDataVisualization.cpp
+++ b/Engine/Plugins/PCG/Source/PCGEditor/Private/DataVisualizations/PCGSpatialDataVisualization.cpp
@@ -24,57 +24,11 @@
 
 namespace PCGPointDataVisualizationConstants
 {
-	/** Names of the columns in the attribute list. */
-	const FName NAME_Index = FName(TEXT("$Index"));
-	const FName NAME_PositionX = FName(TEXT("$Position.X"));
-	const FName NAME_PositionY = FName(TEXT("$Position.Y"));
-	const FName NAME_PositionZ = FName(TEXT("$Position.Z"));
-	const FName NAME_RotationRoll = FName(TEXT("$Rotation.Roll"));
-	const FName NAME_RotationPitch = FName(TEXT("$Rotation.Pitch"));
-	const FName NAME_RotationYaw = FName(TEXT("$Rotation.Yaw"));
-	const FName NAME_ScaleX = FName(TEXT("$Scale.X"));
-	const FName NAME_ScaleY = FName(TEXT("$Scale.Y"));
-	const FName NAME_ScaleZ = FName(TEXT("$Scale.Z"));
-	const FName NAME_BoundsMinX = FName(TEXT("$BoundsMin.X"));
-	const FName NAME_BoundsMinY = FName(TEXT("$BoundsMin.Y"));
-	const FName NAME_BoundsMinZ = FName(TEXT("$BoundsMin.Z"));
-	const FName NAME_BoundsMaxX = FName(TEXT("$BoundsMax.X"));
-	const FName NAME_BoundsMaxY = FName(TEXT("$BoundsMax.Y"));
-	const FName NAME_BoundsMaxZ = FName(TEXT("$BoundsMax.Z"));
-	const FName NAME_ColorR = FName(TEXT("$Color.R"));
-	const FName NAME_ColorG = FName(TEXT("$Color.G"));
-	const FName NAME_ColorB = FName(TEXT("$Color.B"));
-	const FName NAME_ColorA = FName(TEXT("$Color.A"));
-	const FName NAME_Density = FName(TEXT("$Density"));
-	const FName NAME_Steepness = FName(TEXT("$Steepness"));
-	const FName NAME_Seed = FName(TEXT("$Seed"));
+	/** Special names of the columns in the attribute list. */
 	const FName NAME_MetadataEntry = FName(TEXT("MetadataEntry"));
 	const FName NAME_MetadataEntryParent = FName(TEXT("PointMetadataEntryParent"));
 
-	/** Labels of the columns. */
-	const FText TEXT_Index = LOCTEXT("Index", "$Index");
-	const FText TEXT_PositionX = LOCTEXT("PositionX", "$Position.X");
-	const FText TEXT_PositionY = LOCTEXT("PositionY", "$Position.Y");
-	const FText TEXT_PositionZ = LOCTEXT("PositionZ", "$Position.Z");
-	const FText TEXT_RotationRoll = LOCTEXT("RotationRoll", "$Rotation.Roll");
-	const FText TEXT_RotationPitch = LOCTEXT("RotationPitch", "$Rotation.Pitch");
-	const FText TEXT_RotationYaw = LOCTEXT("RotationYaw", "$Rotation.Yaw");
-	const FText TEXT_ScaleX = LOCTEXT("ScaleX", "$Scale.X");
-	const FText TEXT_ScaleY = LOCTEXT("ScaleY", "$Scale.Y");
-	const FText TEXT_ScaleZ = LOCTEXT("ScaleZ", "$Scale.Z");
-	const FText TEXT_BoundsMinX = LOCTEXT("BoundsMinX", "$BoundsMin.X");
-	const FText TEXT_BoundsMinY = LOCTEXT("BoundsMinY", "$BoundsMin.Y");
-	const FText TEXT_BoundsMinZ = LOCTEXT("BoundsMinZ", "$BoundsMin.Z");
-	const FText TEXT_BoundsMaxX = LOCTEXT("BoundsMaxX", "$BoundsMax.X");
-	const FText TEXT_BoundsMaxY = LOCTEXT("BoundsMaxY", "$BoundsMax.Y");
-	const FText TEXT_BoundsMaxZ = LOCTEXT("BoundsMaxZ", "$BoundsMax.Z");
-	const FText TEXT_ColorR = LOCTEXT("ColorR", "$Color.R");
-	const FText TEXT_ColorG = LOCTEXT("ColorG", "$Color.G");
-	const FText TEXT_ColorB = LOCTEXT("ColorB", "$Color.B");
-	const FText TEXT_ColorA = LOCTEXT("ColorA", "$Color.A");
-	const FText TEXT_Density = LOCTEXT("Density", "$Density");
-	const FText TEXT_Steepness = LOCTEXT("Steepness", "$Steepness");
-	const FText TEXT_Seed = LOCTEXT("Seed", "$Seed");
+	/** Special labels of the columns. */
 	const FText TEXT_MetadataEntry = LOCTEXT("MetadataEntry", "Entry Key");
 	const FText TEXT_MetadataEntryParent = LOCTEXT("MetadataEntryParent", "Parent Key");
 }
@@ -291,41 +245,35 @@ FPCGTableVisualizerInfo IPCGSpatialDataVisualization::GetTableVisualizerInfoWith
 	FPCGTableVisualizerInfo Info;
 	Info.Data = PointData;
 
-	// Columns
-	Info.ColumnInfos.Add(CreateColumnInfo(PointData, NAME_Index,         TEXT_Index));
-	Info.ColumnInfos.Add(CreateColumnInfo(PointData, NAME_PositionX,     TEXT_PositionX));
-	Info.ColumnInfos.Add(CreateColumnInfo(PointData, NAME_PositionY,     TEXT_PositionY));
-	Info.ColumnInfos.Add(CreateColumnInfo(PointData, NAME_PositionZ,     TEXT_PositionZ));
-	Info.ColumnInfos.Add(CreateColumnInfo(PointData, NAME_RotationRoll,  TEXT_RotationRoll));
-	Info.ColumnInfos.Add(CreateColumnInfo(PointData, NAME_RotationPitch, TEXT_RotationPitch));
-	Info.ColumnInfos.Add(CreateColumnInfo(PointData, NAME_RotationYaw,   TEXT_RotationYaw));
-	Info.ColumnInfos.Add(CreateColumnInfo(PointData, NAME_ScaleX,        TEXT_ScaleX));
-	Info.ColumnInfos.Add(CreateColumnInfo(PointData, NAME_ScaleY,        TEXT_ScaleY));
-	Info.ColumnInfos.Add(CreateColumnInfo(PointData, NAME_ScaleZ,        TEXT_ScaleZ));
-	Info.ColumnInfos.Add(CreateColumnInfo(PointData, NAME_BoundsMinX,    TEXT_BoundsMinX));
-	Info.ColumnInfos.Add(CreateColumnInfo(PointData, NAME_BoundsMinY,    TEXT_BoundsMinY));
-	Info.ColumnInfos.Add(CreateColumnInfo(PointData, NAME_BoundsMinZ,    TEXT_BoundsMinZ));
-	Info.ColumnInfos.Add(CreateColumnInfo(PointData, NAME_BoundsMaxX,    TEXT_BoundsMaxX));
-	Info.ColumnInfos.Add(CreateColumnInfo(PointData, NAME_BoundsMaxY,    TEXT_BoundsMaxY));
-	Info.ColumnInfos.Add(CreateColumnInfo(PointData, NAME_BoundsMaxZ,    TEXT_BoundsMaxZ));
-	Info.ColumnInfos.Add(CreateColumnInfo(PointData, NAME_ColorR,        TEXT_ColorR));
-	Info.ColumnInfos.Add(CreateColumnInfo(PointData, NAME_ColorG,        TEXT_ColorG));
-	Info.ColumnInfos.Add(CreateColumnInfo(PointData, NAME_ColorB,        TEXT_ColorB));
-	Info.ColumnInfos.Add(CreateColumnInfo(PointData, NAME_ColorA,        TEXT_ColorA));
-	Info.ColumnInfos.Add(CreateColumnInfo(PointData, NAME_Density,       TEXT_Density));
-	Info.ColumnInfos.Add(CreateColumnInfo(PointData, NAME_Steepness,     TEXT_Steepness));
-	Info.ColumnInfos.Add(CreateColumnInfo(PointData, NAME_Seed,          TEXT_Seed));
+	// Column Sorting
+	AddColumnInfo(Info, PointData, FPCGAttributePropertySelector::CreateExtraPropertySelector(EPCGExtraProperties::Index));
+	Info.SortingColumn = Info.ColumnInfos.Last().Id;
+
+	AddPropertyEnumColumnInfo<FVector>(Info, PointData, EPCGPointProperties::Position);
+	AddPropertyEnumColumnInfo<FRotator>(Info, PointData, EPCGPointProperties::Rotation);
+	AddPropertyEnumColumnInfo<FVector>(Info, PointData, EPCGPointProperties::Scale);
+	AddPropertyEnumColumnInfo<FVector>(Info, PointData, EPCGPointProperties::BoundsMin);
+	AddPropertyEnumColumnInfo<FVector>(Info, PointData, EPCGPointProperties::BoundsMax);
+	AddPropertyEnumColumnInfo<FLinearColor>(Info, PointData, EPCGPointProperties::Color);
+	AddPropertyEnumColumnInfo<float>(Info, PointData, EPCGPointProperties::Density);
+	AddPropertyEnumColumnInfo<float>(Info, PointData, EPCGPointProperties::Steepness);
+	AddPropertyEnumColumnInfo<int32>(Info, PointData, EPCGPointProperties::Seed);
 
 	if (FPCGEditorCommon::CVarShowAdvancedAttributesFields.GetValueOnAnyThread())
 	{
-		Info.ColumnInfos.Add(CreateColumnInfo(PointData, NAME_MetadataEntry, TEXT_MetadataEntry, [PointData]()
+		FColumnInfoOverrides Overrides;
+		Overrides.LabelOverride = TEXT_MetadataEntry;
+		Overrides.CreateAccessorFuncOverride = [PointData]()
 		{
 			FPCGAttributePropertySelector MetadataEntrySelector;
 			MetadataEntrySelector.SetPropertyName(*StaticEnum<EPCGPointNativeProperties>()->GetNameStringByValue((int64)EPCGPointNativeProperties::MetadataEntry));
 			return TSharedPtr<const IPCGAttributeAccessor>(PCGAttributeAccessorHelpers::CreateConstAccessor(PointData, MetadataEntrySelector).Release());
-		}));
+		};
+		
+		AddTypedColumnInfo<int64>(Info, PointData, FPCGAttributePropertySelector{}, Overrides);
 
-		Info.ColumnInfos.Add(CreateColumnInfo(PointData, NAME_MetadataEntryParent, TEXT_MetadataEntryParent, [PointData]()
+		Overrides.LabelOverride = TEXT_MetadataEntryParent;
+		Overrides.CreateAccessorFuncOverride = [PointData]()
 		{
 			return MakeShared<FPCGCustomPointPropertyAccessor<int64, TConstPCGValueRange<int64>>>(PointData, [Metadata = PointData->ConstMetadata()](int32 Index, int64& OutValue, const TConstPCGValueRange<int64>& MetaDataEntryRange)
 			{
@@ -337,14 +285,13 @@ FPCGTableVisualizerInfo IPCGSpatialDataVisualization::GetTableVisualizerInfoWith
 				return false;
 			},
 			PointData->GetConstMetadataEntryValueRange());
-		}));
+		};
+		
+		AddTypedColumnInfo<int64>(Info, PointData, FPCGAttributePropertySelector{}, Overrides);
 	}
 
 	// Add Metadata Columns
-	CreateMetadataColumnInfos(PointData, Info.ColumnInfos, PCGMetadataDomainID::Elements);
-
-	// Column Sorting
-	Info.SortingColumn = NAME_Index;
+	CreateMetadataColumnInfos(PointData, Info, PCGMetadataDomainID::Elements);
 
 	// Focus on data behavior
 	Info.FocusOnDataCallback = [](const UPCGData* Data, TArrayView<const int> Indices)
@@ -382,9 +329,6 @@ FPCGTableVisualizerInfo IPCGSpatialDataVisualization::GetTableVisualizerInfoWith
 		}
 	};
 
-	// Accessor Keys
-	Info.AccessorKeys = TSharedPtr<const IPCGAttributeAccessorKeys>(PCGAttributeAccessorHelpers::CreateConstKeys(PointData, FPCGAttributePropertySelector()).Release());
-
 	return Info;
 }
 
diff --git a/Engine/Plugins/PCG/Source/PCGEditor/Private/Widgets/SPCGEditorGraphAttributeListView.cpp b/Engine/Plugins/PCG/Source/PCGEditor/Private/Widgets/SPCGEditorGraphAttributeListView.cpp
--- a/Engine/Plugins/PCG/Source/PCGEditor/Private/Widgets/SPCGEditorGraphAttributeListView.cpp
+++ b/Engine/Plugins/PCG/Source/PCGEditor/Private/Widgets/SPCGEditorGraphAttributeListView.cpp
@@ -866,19 +866,21 @@ void SPCGEditorGraphAttributeListView::RefreshAttributeList()
 			return;
 		}
 
+		int32 NumEntries = 0;
+
 		for (const FPCGTableVisualizerColumnInfo& ColumnInfo : TableVisualizerInfo.ColumnInfos)
 		{
-			AddColumn(DataToVisualize, ColumnInfo, TableVisualizerInfo.AccessorKeys);
+			AddColumn(DataToVisualize, ColumnInfo);
+			NumEntries = FMath::Max(NumEntries, ColumnInfo.AccessorKeys.IsValid() ? ColumnInfo.AccessorKeys->GetNum() : 0);
 		}
 
 		SortingColumn = TableVisualizerInfo.SortingColumn;
 		SortMode = (EColumnSortMode::Type)TableVisualizerInfo.SortingMode;
 		FocusOnDataCallback = TableVisualizerInfo.FocusOnDataCallback;
 
-		const uint32 NumEntries = TableVisualizerInfo.AccessorKeys.IsValid() ? TableVisualizerInfo.AccessorKeys->GetNum() : 0;
 		ListViewItems.Reserve(NumEntries);
 
-		for (uint32 Index = 0; Index < NumEntries; ++Index)
+		for (int32 Index = 0; Index < NumEntries; ++Index)
 		{
 			PCGListViewItemPtr ListViewItem = MakeShared<FPCGListViewItem>();
 			ListViewItem->Index = Index;
@@ -1508,11 +1510,11 @@ FReply SPCGEditorGraphAttributeListView::OnListViewKeyDown(const FGeometry& /*In
 	return FReply::Unhandled();
 }
 
-void SPCGEditorGraphAttributeListView::AddColumn(const UPCGData* InPCGData, const FPCGTableVisualizerColumnInfo& InColumnInfo, TSharedPtr<const IPCGAttributeAccessorKeys> InAccessorKeys)
+void SPCGEditorGraphAttributeListView::AddColumn(const UPCGData* InPCGData, const FPCGTableVisualizerColumnInfo& InColumnInfo)
 {
 	FPCGColumnData& ColumnData = PCGColumnData.Add(InColumnInfo.Id);
 	ColumnData.DataAccessor = InColumnInfo.Accessor;
-	ColumnData.DataKeys = InAccessorKeys;
+	ColumnData.DataKeys = InColumnInfo.AccessorKeys;
 
 	const float ColumnWidth = InColumnInfo.Width < 0.0f ? PCGEditorGraphAttributeListView::CalculateColumnWidth(InColumnInfo.Label) : InColumnInfo.Width;
 	ColumnsMaxWidthMapping.FindOrAdd(InColumnInfo.Id) = ColumnWidth;
diff --git a/Engine/Plugins/PCG/Source/PCGEditor/Private/Widgets/SPCGEditorGraphAttributeListView.h b/Engine/Plugins/PCG/Source/PCGEditor/Private/Widgets/SPCGEditorGraphAttributeListView.h
--- a/Engine/Plugins/PCG/Source/PCGEditor/Private/Widgets/SPCGEditorGraphAttributeListView.h
+++ b/Engine/Plugins/PCG/Source/PCGEditor/Private/Widgets/SPCGEditorGraphAttributeListView.h
@@ -236,7 +236,7 @@ private:
 
 	FReply OnListViewKeyDown(const FGeometry& InGeometry, const FKeyEvent& InKeyEvent) const;
 
-	void AddColumn(const UPCGData* InPCGData, const FPCGTableVisualizerColumnInfo& InColumnInfo, TSharedPtr<const IPCGAttributeAccessorKeys> InAccessorKeys);
+	void AddColumn(const UPCGData* InPCGData, const FPCGTableVisualizerColumnInfo& InColumnInfo);
 	void AddPointDataColumns(const UPCGPointData* InPCGPointData);
 	void AddMetadataColumn(const UPCGData* InPCGData, const FName& InColumnId, EPCGMetadataTypes InMetadataType, const TCHAR* PostFix = nullptr);
 
diff --git a/Engine/Plugins/PCG/Source/PCGEditor/Public/DataVisualizations/PCGDataVisualizationHelpers.h b/Engine/Plugins/PCG/Source/PCGEditor/Public/DataVisualizations/PCGDataVisualizationHelpers.h
--- a/Engine/Plugins/PCG/Source/PCGEditor/Public/DataVisualizations/PCGDataVisualizationHelpers.h
+++ b/Engine/Plugins/PCG/Source/PCGEditor/Public/DataVisualizations/PCGDataVisualizationHelpers.h
@@ -9,9 +9,95 @@ class UPCGData;
 
 namespace PCGDataVisualizationHelpers
 {
-	FPCGTableVisualizerColumnInfo CreateColumnInfo(const UPCGData* Data, FName Id, const FText& Label, TFunction<TSharedPtr<const IPCGAttributeAccessor>()> CreateAccessorFunc = nullptr);
+	struct FColumnInfoOverrides
+	{
+		FText LabelOverride;
+		FText TooltipOverride;
+		bool bAddTypeToTooltip = true;
+		EPCGMetadataTypes TypeOverride = EPCGMetadataTypes::Unknown;
+		TFunction<TSharedPtr<const IPCGAttributeAccessor>()> CreateAccessorFuncOverride;
+		TFunction<TSharedPtr<const IPCGAttributeAccessorKeys>()> CreateAccessorKeysFuncOverride;
+	};
+	
+	/**
+	 * Add a new column to the VisualizerInfo, using the Data and Selector.
+	 * By default, it will create a column with the label extracted from the selector, with accessor and keys created from the data and selector, but you can provide overrides.
+	 * @param OutInfo 
+	 * @param Data 
+	 * @param InSelector 
+	 * @param Overrides
+	 */
+	PCGEDITOR_API void AddColumnInfo(FPCGTableVisualizerInfo& OutInfo, const UPCGData* Data, const FPCGAttributePropertySelector& InSelector, const FColumnInfoOverrides& Overrides = {});
+	
+	PCGEDITOR_API void CreateMetadataColumnInfos(const UPCGData* Data, FPCGTableVisualizerInfo& OutInfo, const FPCGMetadataDomainID& DomainID = PCGMetadataDomainID::Default);
 
-	void CreateMetadataColumnInfos(const UPCGData* Data, TArray<FPCGTableVisualizerColumnInfo>& OutColumnInfos, const FPCGMetadataDomainID& DomainID = PCGMetadataDomainID::Default);
+	PCGEDITOR_API FPCGTableVisualizerInfo CreateDefaultMetadataColumnInfos(const UPCGData* Data, const FPCGMetadataDomainID& DomainID = PCGMetadataDomainID::Default);
 
-	FPCGTableVisualizerInfo CreateDefaultMetadataColumnInfos(const UPCGData* Data, const FPCGMetadataDomainID& DomainID = PCGMetadataDomainID::Default);
-}
+	template <typename T>
+	void AddTypedColumnInfo_Impl(FPCGTableVisualizerInfo& OutInfo, const UPCGData* Data, const FPCGAttributePropertySelector& InSelector, const FColumnInfoOverrides& Overrides = {})
+	{
+		AddColumnInfo(OutInfo, Data, InSelector, Overrides);
+	}
+
+	template <typename T>
+	void AddTypedColumnInfo(FPCGTableVisualizerInfo& OutInfo, const UPCGData* Data, const FPCGAttributePropertySelector& InSelector, const FColumnInfoOverrides& Overrides = {})
+	{
+		const FColumnInfoOverrides* OverridesToUse = &Overrides;
+		FColumnInfoOverrides CopyOverrides;
+		if (Overrides.TypeOverride == EPCGMetadataTypes::Unknown && Overrides.bAddTypeToTooltip && PCG::Private::IsPCGType<T>())
+		{
+			CopyOverrides = Overrides;
+			CopyOverrides.TypeOverride = static_cast<EPCGMetadataTypes>(PCG::Private::MetadataTypes<T>::Id);
+			OverridesToUse = &CopyOverrides;
+		}
+		
+		AddTypedColumnInfo_Impl<T>(OutInfo, Data, InSelector, *OverridesToUse);
+	}
+
+	template <typename T>
+	void AddPropertyEnumColumnInfo(FPCGTableVisualizerInfo& OutInfo, const UPCGData* Data, const UEnum* EnumClass, int64 EnumValue, const FColumnInfoOverrides& Overrides = {})
+	{
+		check(Data && EnumClass);
+
+		int32 EnumIndex = EnumClass->GetIndexByValue(EnumValue);
+		check(EnumIndex != INDEX_NONE);
+
+		const FColumnInfoOverrides* OverridesToUse = &Overrides;
+		FColumnInfoOverrides CopyOverrides;
+		if (Overrides.TooltipOverride.IsEmpty())
+		{
+			CopyOverrides = Overrides;
+			CopyOverrides.TooltipOverride = EnumClass->GetToolTipTextByIndex(EnumIndex);
+			OverridesToUse = &CopyOverrides;
+		}
+		
+		const FString InvariantName = EnumClass->GetNameStringByIndex(EnumIndex); // Use string version as it strips out the namespace.
+		FPCGAttributePropertySelector Selector = FPCGAttributePropertySelector::CreatePropertySelector(*InvariantName);
+
+		static const FString EnumMetadataDomain_MetadataFlag = PCGObjectMetadata::EnumMetadataDomain.ToString();
+		
+		if (EnumClass->HasMetaData(*EnumMetadataDomain_MetadataFlag, EnumIndex))
+		{
+			const FString MetadataDomain = EnumClass->GetMetaData(*EnumMetadataDomain_MetadataFlag, EnumIndex);
+			Selector.SetDomainName(*MetadataDomain);
+		}
+
+		AddTypedColumnInfo<T>(OutInfo, Data, Selector, *OverridesToUse);
+	}
+
+	template <typename T, typename EnumType>
+	void AddPropertyEnumColumnInfo(FPCGTableVisualizerInfo& OutInfo, const UPCGData* Data, EnumType EnumValue, const FColumnInfoOverrides& Overrides = {})
+	{
+		const UEnum* EnumClass = StaticEnum<EnumType>();
+		check(EnumClass);
+		AddPropertyEnumColumnInfo<T>(OutInfo, Data, EnumClass, static_cast<int64>(EnumValue));
+	}
+
+	template <> PCGEDITOR_API void AddTypedColumnInfo_Impl<FVector2D>(FPCGTableVisualizerInfo& OutInfo, const UPCGData* Data, const FPCGAttributePropertySelector& InSelector, const FColumnInfoOverrides& Overrides);
+	template <> PCGEDITOR_API void AddTypedColumnInfo_Impl<FVector>(FPCGTableVisualizerInfo& OutInfo, const UPCGData* Data, const FPCGAttributePropertySelector& InSelector, const FColumnInfoOverrides& Overrides);
+	template <> PCGEDITOR_API void AddTypedColumnInfo_Impl<FVector4>(FPCGTableVisualizerInfo& OutInfo, const UPCGData* Data, const FPCGAttributePropertySelector& InSelector, const FColumnInfoOverrides& Overrides);
+	template <> PCGEDITOR_API void AddTypedColumnInfo_Impl<FLinearColor>(FPCGTableVisualizerInfo& OutInfo, const UPCGData* Data, const FPCGAttributePropertySelector& InSelector, const FColumnInfoOverrides& Overrides);
+	template <> PCGEDITOR_API void AddTypedColumnInfo_Impl<FQuat>(FPCGTableVisualizerInfo& OutInfo, const UPCGData* Data, const FPCGAttributePropertySelector& InSelector, const FColumnInfoOverrides& Overrides);
+	template <> PCGEDITOR_API void AddTypedColumnInfo_Impl<FRotator>(FPCGTableVisualizerInfo& OutInfo, const UPCGData* Data, const FPCGAttributePropertySelector& InSelector, const FColumnInfoOverrides& Overrides);
+	template <> PCGEDITOR_API void AddTypedColumnInfo_Impl<FTransform>(FPCGTableVisualizerInfo& OutInfo, const UPCGData* Data, const FPCGAttributePropertySelector& InSelector, const FColumnInfoOverrides& Overrides);
+}
\ No newline at end of file
