[PCG] More accessors for Splines

* Change Position/Rotation/Scale to be in world coordinates like points
* Add LocalPosition/LocalRotation for the local value (stored on the spline struct)
* Add Transfrom/LocalTransform.
* Add global attributes like SplineTransform and IsClosed.

#rnx

diff --git a/Engine/Plugins/PCG/Source/PCG/Private/Data/PCGSplineData.cpp b/Engine/Plugins/PCG/Source/PCG/Private/Data/PCGSplineData.cpp
--- a/Engine/Plugins/PCG/Source/PCG/Private/Data/PCGSplineData.cpp
+++ b/Engine/Plugins/PCG/Source/PCG/Private/Data/PCGSplineData.cpp
@@ -126,7 +126,7 @@ bool UPCGSplineData::SetDomainFromDomainID(const FPCGMetadataDomainID& InDomainI
 {
 	if (InDomainID == PCGMetadataDomainID::Elements)
 	{
-		InOutSelector.SetDomainName(PCGSplineData::ControlPointDomainName);
+		InOutSelector.SetDomainName(PCGSplineData::ControlPointDomainName, /*bResetExtraNames=*/false);
 		return true;
 	}
 	else
@@ -280,7 +280,7 @@ void UPCGSplineData::WriteMetadataToPoint(float InputKey, FPCGPoint& OutPoint, U
 
 TUniquePtr<IPCGAttributeAccessor> UPCGSplineData::CreateStaticAccessor(const FPCGAttributePropertySelector& InSelector, bool bQuiet)
 {
-	static const FProperty* PCGSplineStructProperty = UPCGSplineData::StaticClass()->FindPropertyByName(GET_MEMBER_NAME_CHECKED(UPCGSplineData, SplineStruct));
+	static const FStructProperty* PCGSplineStructProperty = CastFieldChecked<FStructProperty>(UPCGSplineData::StaticClass()->FindPropertyByName(GET_MEMBER_NAME_CHECKED(UPCGSplineData, SplineStruct)));
 	static const FProperty* SplineCurveProperty = FPCGSplineStruct::StaticStruct()->FindPropertyByName(TEXT("SplineCurves"));
 
 	const FName DomainName = InSelector.GetDomainName();
@@ -288,22 +288,43 @@ TUniquePtr<IPCGAttributeAccessor> UPCGSplineData::CreateStaticAccessor(const FPC
 	if (InSelector.GetSelection() == EPCGAttributePropertySelection::Property && (DomainName.IsNone() || DomainName == PCGSplineData::ControlPointDomainName))
 	{
 		const FName PropertyName = InSelector.GetName();
-		TArray<const FProperty*> ExtraProperties = {PCGSplineStructProperty, SplineCurveProperty};
 		
-		if (PropertyName == TEXT("Position") || PropertyName == TEXT("Scale"))
+		if (PropertyName == TEXT("Position"))
 		{
-			const FStructProperty* VectorProperty = CastField<FStructProperty>(FSplineCurves::StaticStruct()->FindPropertyByName(PropertyName));
-			check(VectorProperty);
-			return MakeUnique<FPCGInterpCurveAccessor<FInterpCurveVector>>(VectorProperty, std::move(ExtraProperties));
+			return MakeUnique<FPCGControlPointsAccessor<FVector, EPCGControlPointsAccessorTarget::Location, /*bWorldCoordinates=*/true>>(PCGSplineStructProperty);
+		}
+		else if (PropertyName == TEXT("LocalPosition"))
+		{
+			return MakeUnique<FPCGControlPointsAccessor<FVector, EPCGControlPointsAccessorTarget::Location, /*bWorldCoordinates=*/false>>(PCGSplineStructProperty);
 		}
 		else if (PropertyName == TEXT("Rotation"))
 		{
-			const FStructProperty* QuatProperty = CastField<FStructProperty>(FSplineCurves::StaticStruct()->FindPropertyByName(PropertyName));
-			check(QuatProperty);
-			return MakeUnique<FPCGInterpCurveAccessor<FInterpCurveQuat>>(QuatProperty, std::move(ExtraProperties));
+			return MakeUnique<FPCGControlPointsAccessor<FQuat, EPCGControlPointsAccessorTarget::Rotation, /*bWorldCoordinates=*/true>>(PCGSplineStructProperty);
+		}
+		else if (PropertyName == TEXT("LocalRotation"))
+		{
+			return MakeUnique<FPCGControlPointsAccessor<FQuat, EPCGControlPointsAccessorTarget::Rotation, /*bWorldCoordinates=*/false>>(PCGSplineStructProperty);
+		}
+		else if (PropertyName == TEXT("Scale"))
+		{
+			return MakeUnique<FPCGControlPointsAccessor<FVector, EPCGControlPointsAccessorTarget::Scale, /*bWorldCoordinates=*/true>>(PCGSplineStructProperty);
+		}
+		else if (PropertyName == TEXT("LocalScale"))
+		{
+			return MakeUnique<FPCGControlPointsAccessor<FVector, EPCGControlPointsAccessorTarget::Scale, /*bWorldCoordinates=*/false>>(PCGSplineStructProperty);
+		}
+		else if (PropertyName == TEXT("Transform"))
+		{
+			return MakeUnique<FPCGControlPointsAccessor<FTransform, EPCGControlPointsAccessorTarget::Transform, /*bWorldCoordinates=*/true>>(PCGSplineStructProperty);
+		}
+		else if (PropertyName == TEXT("LocalTransform"))
+		{
+			return MakeUnique<FPCGControlPointsAccessor<FTransform, EPCGControlPointsAccessorTarget::Transform, /*bWorldCoordinates=*/false>>(PCGSplineStructProperty);
 		}
 		else if (PropertyName == TEXT("LeaveTangent") || PropertyName == TEXT("ArriveTangent") || PropertyName == TEXT("InterpType"))
 		{
+			TArray<const FProperty*> ExtraProperties = {PCGSplineStructProperty, SplineCurveProperty};
+			
 			// Implementation note:
 			// Tangents and Interpolation type are all linked to the Position property on FSplineCurves, even if technically the rotation and the scale also have tangents and interpolation type.
 			// But in the spline details view, and everywhere we have tangents and interpolation, they are always related to the Position interp curve.
@@ -329,6 +350,19 @@ TUniquePtr<IPCGAttributeAccessor> UPCGSplineData::CreateStaticAccessor(const FPC
 			PCGLog::LogErrorOnGraph(FText::Format(LOCTEXT("FailCreateAccessor", "Property {0} is not valid for a control point."), FText::FromName(PropertyName)));
 		}
 	}
+	else if (InSelector.GetSelection() == EPCGAttributePropertySelection::Property && DomainName == PCGDataConstants::GlobalDomainName)
+	{
+		const FName PropertyName = InSelector.GetName();
+		
+		if (PropertyName == TEXT("SplineTransform"))
+		{
+			return MakeUnique<FPCGSplineAccessor<FTransform, EPCGSplineAccessorTarget::Transform>>(PCGSplineStructProperty);
+		}
+		else if (PropertyName == TEXT("IsClosed"))
+		{
+			return MakeUnique<FPCGSplineAccessor<bool, EPCGSplineAccessorTarget::ClosedLoop>>(PCGSplineStructProperty);
+		}
+	}
 	
 	return {};
 }
@@ -352,7 +386,7 @@ FPCGAttributeAccessorMethods UPCGSplineData::GetSplineAccessorMethods()
 		
 		if (Selection == EPCGAttributePropertySelection::Property || Selection == EPCGAttributePropertySelection::ExtraProperty)
 		{
-			return MakeUnique<FPCGAttributeAccessorKeysSplineData>(SplineData);
+			return MakeUnique<FPCGAttributeAccessorKeysSplineData>(SplineData, InSelector.GetDomainName() == PCGDataConstants::GlobalDomainName);
 		}
 		else if (Selection == EPCGAttributePropertySelection::Attribute)
 		{
@@ -373,7 +407,7 @@ FPCGAttributeAccessorMethods UPCGSplineData::GetSplineAccessorMethods()
 		const EPCGAttributePropertySelection Selection = InSelector.GetSelection();
 		if (Selection == EPCGAttributePropertySelection::Property || Selection == EPCGAttributePropertySelection::ExtraProperty)
 		{
-			return MakeUnique<const FPCGAttributeAccessorKeysSplineData>(SplineData);
+			return MakeUnique<const FPCGAttributeAccessorKeysSplineData>(SplineData, InSelector.GetDomainName() == PCGDataConstants::GlobalDomainName);
 		}
 		else if (Selection == EPCGAttributePropertySelection::Attribute)
 		{
diff --git a/Engine/Plugins/PCG/Source/PCG/Private/Metadata/Accessors/PCGAttributeAccessorFactory.cpp b/Engine/Plugins/PCG/Source/PCG/Private/Metadata/Accessors/PCGAttributeAccessorFactory.cpp
--- a/Engine/Plugins/PCG/Source/PCG/Private/Metadata/Accessors/PCGAttributeAccessorFactory.cpp
+++ b/Engine/Plugins/PCG/Source/PCG/Private/Metadata/Accessors/PCGAttributeAccessorFactory.cpp
@@ -196,6 +196,8 @@ void FPCGAttributeAccessorMethods::FillSelectorMenuEntryFromEnum(const UEnum* En
 
 	check(CurrentMenu);
 
+	static const FString EnumMetadataDomain_MetadataFlag = PCGObjectMetadata::EnumMetadataDomain.ToString();
+
 	const int32 NumEnums = EnumType->ContainsExistingMax() ? EnumType->NumEnums() - 1 : EnumType->NumEnums();
 	for (int32 i = 0; i < NumEnums; ++i)
 	{
@@ -208,8 +210,15 @@ void FPCGAttributeAccessorMethods::FillSelectorMenuEntryFromEnum(const UEnum* En
 		FText EnumName = EnumType->GetDisplayNameTextByIndex(i);
 		FText Tooltip = EnumType->GetToolTipTextByIndex(i);
 		const FString InvariantName = EnumType->GetNameStringByIndex(i); // Use string version as it strips out the namespace.
+		FPCGAttributePropertySelector Selector = FPCGAttributePropertySelector::CreatePropertySelector(*InvariantName);		
+
+		if (EnumType->HasMetaData(*EnumMetadataDomain_MetadataFlag, i))
+		{
+			const FString MetadataDomain = EnumType->GetMetaData(*EnumMetadataDomain_MetadataFlag, i);
+			Selector.SetDomainName(*MetadataDomain);
+		}
 
-		CurrentMenu->Entries.Emplace(std::move(EnumName), std::move(Tooltip), FPCGAttributePropertySelector::CreatePropertySelector(*InvariantName));
+		CurrentMenu->Entries.Emplace(std::move(EnumName), std::move(Tooltip), std::move(Selector));
 	}
 }
 #endif // WITH_EDITOR
diff --git a/Engine/Plugins/PCG/Source/PCG/Private/Metadata/Accessors/PCGAttributeAccessorKeys.cpp b/Engine/Plugins/PCG/Source/PCG/Private/Metadata/Accessors/PCGAttributeAccessorKeys.cpp
--- a/Engine/Plugins/PCG/Source/PCG/Private/Metadata/Accessors/PCGAttributeAccessorKeys.cpp
+++ b/Engine/Plugins/PCG/Source/PCG/Private/Metadata/Accessors/PCGAttributeAccessorKeys.cpp
@@ -357,10 +357,34 @@ bool FPCGAttributeAccessorKeysPointsSubset::GetMetadataEntryKeys(int32 InStart,
 }
 
 /////////////////////////////////////////////////////////
+FPCGAttributeAccessorKeysSplineData::FPCGAttributeAccessorKeysSplineData()
+	: FPCGAttributeAccessorKeysSingleObjectPtr()
+{}
+
+FPCGAttributeAccessorKeysSplineData::FPCGAttributeAccessorKeysSplineData(UPCGSplineData* InPtr, bool bInGlobalData)
+	: FPCGAttributeAccessorKeysSingleObjectPtr(InPtr)
+	, bGlobalData(bInGlobalData)
+{}
+
+FPCGAttributeAccessorKeysSplineData::FPCGAttributeAccessorKeysSplineData(const UPCGSplineData* InPtr, bool bInGlobalData)
+	: FPCGAttributeAccessorKeysSingleObjectPtr(InPtr)
+	, bGlobalData(bInGlobalData)
+{}
 
 int32 FPCGAttributeAccessorKeysSplineData::GetNum() const
 {
-	return Ptr ? Ptr->SplineStruct.GetSplinePointsPosition().Points.Num() : 0;
+	if (!Ptr)
+	{
+		return 0; 
+	}
+	else if (bGlobalData)
+	{
+		return 1;
+	}
+	else
+	{
+		return Ptr->SplineStruct.GetSplinePointsPosition().Points.Num();
+	}
 }
 
 /////////////////////////////////////////////////////////
diff --git a/Engine/Plugins/PCG/Source/PCG/Public/Data/PCGSplineStruct.h b/Engine/Plugins/PCG/Source/PCG/Public/Data/PCGSplineStruct.h
--- a/Engine/Plugins/PCG/Source/PCG/Public/Data/PCGSplineStruct.h
+++ b/Engine/Plugins/PCG/Source/PCG/Public/Data/PCGSplineStruct.h
@@ -12,12 +12,27 @@
 UENUM(Blueprintable)
 enum class EPCGSplineStructProperties : uint8
 {
-	Position       UMETA(Tooltip = "Location component of the control point's transform."),
-	Rotation       UMETA(Tooltip = "Rotation component of the control point's transform."),
-	Scale          UMETA(Tooltip = "Scale component of the control point's transform."),
-	ArriveTangent  UMETA(Tooltip = "Arrive Tangent at the control point."),
-	LeaveTangent   UMETA(Tooltip = "Leave Tangent at the control point."),
-	InterpType     UMETA(Tooltip = "Type of interpolation at the control point for position (same as set on the spline control point). Underlying enum: ESplinePointType.")
+	Position         UMETA(Tooltip = "Location component of the control point's transform in world coordinates."),
+	Rotation         UMETA(Tooltip = "Rotation component of the control point's transform in world coordinates."),
+	Scale            UMETA(Tooltip = "Scale component of the control point's transform in world coordinates."),
+	Transform        UMETA(Tooltip = "Control point's world transform."),
+	ArriveTangent    UMETA(Tooltip = "Arrive Tangent at the control point."),
+	LeaveTangent     UMETA(Tooltip = "Leave Tangent at the control point."),
+	InterpType       UMETA(Tooltip = "Type of interpolation at the control point for position (same as set on the spline control point). Underlying enum: ESplinePointType."),
+	LocalPosition    UMETA(Tooltip = "Location component of the control point's transform."),
+	LocalRotation    UMETA(Tooltip = "Rotation component of the control point's transform."),
+	LocalScale       UMETA(Tooltip = "Scale component of the control point's transform."),
+	LocalTransform   UMETA(Tooltip = "Control point's transform."),
+	SplineTransform	 UMETA(Tooltip = "Transform of the spline. Global parameter", PCG_MetadataDomain="Global"),
+	IsClosed		 UMETA(Tooltip = "If the spline is closed. Global parameter", PCG_MetadataDomain="Global"),
+};
+
+enum class EPCGControlPointsAccessorTarget
+{
+	Location,
+	Rotation,
+	Scale,
+	Transform
 };
 
 /** Subset of the Spline Component API in a standalone struct */
@@ -109,4 +124,8 @@ protected:
 private:
 	// Internal helper function called by ConvertSplineSegmentToPolyLine -- assumes the input is within a half-segment, so testing the distance to midpoint will be an accurate guide to subdivision. Taken from USplineComponent.
 	bool DivideSplineIntoPolylineRecursiveWithDistancesHelper(float StartDistanceAlongSpline, float EndDistanceAlongSpline, ESplineCoordinateSpace::Type CoordinateSpace, const float MaxSquareDistanceFromSpline, TArray<FVector>& OutPoints, TArray<double>& OutDistancesAlongSpline) const;
+
+	// Spline accessors need to access the Spline Curves, so make it friend.
+	template <typename T, EPCGControlPointsAccessorTarget Target, bool bWorldCoordinates>
+	friend class FPCGControlPointsAccessor;
 };
\ No newline at end of file
diff --git a/Engine/Plugins/PCG/Source/PCG/Public/Metadata/Accessors/PCGAttributeAccessorKeys.h b/Engine/Plugins/PCG/Source/PCG/Public/Metadata/Accessors/PCGAttributeAccessorKeys.h
--- a/Engine/Plugins/PCG/Source/PCG/Public/Metadata/Accessors/PCGAttributeAccessorKeys.h
+++ b/Engine/Plugins/PCG/Source/PCG/Public/Metadata/Accessors/PCGAttributeAccessorKeys.h
@@ -439,9 +439,14 @@ protected:
 class FPCGAttributeAccessorKeysSplineData : public FPCGAttributeAccessorKeysSingleObjectPtr<UPCGSplineData>
 {
 public:
-	using FPCGAttributeAccessorKeysSingleObjectPtr::FPCGAttributeAccessorKeysSingleObjectPtr;
+	FPCGAttributeAccessorKeysSplineData();
+	explicit FPCGAttributeAccessorKeysSplineData(UPCGSplineData* InPtr, bool bInGlobalData);
+	explicit FPCGAttributeAccessorKeysSplineData(const UPCGSplineData* InPtr, bool bInGlobalData);
 	
 	virtual int32 GetNum() const override;
+	
+private:
+	bool bGlobalData = false;
 };
 
 /**
diff --git a/Engine/Plugins/PCG/Source/PCG/Public/Metadata/Accessors/PCGSplineAccessor.h b/Engine/Plugins/PCG/Source/PCG/Public/Metadata/Accessors/PCGSplineAccessor.h
--- a/Engine/Plugins/PCG/Source/PCG/Public/Metadata/Accessors/PCGSplineAccessor.h
+++ b/Engine/Plugins/PCG/Source/PCG/Public/Metadata/Accessors/PCGSplineAccessor.h
@@ -3,6 +3,7 @@
 #pragma once
 
 #include "PCGPropertyAccessor.h"
+#include "Data/PCGSplineStruct.h"
 
 #include "Components/SplineComponent.h"
 
@@ -14,6 +15,12 @@ enum class EPCGInterpCurveAccessorTarget
 	InterpMode
 };
 
+enum class EPCGSplineAccessorTarget
+{
+	Transform,
+	ClosedLoop
+};
+
 /**
 * Templated accessor for any interp curve. It's important that the keys only have a single value, the struct that holds the spline curve, since
 * the interp curve is basically an array.
@@ -32,10 +39,9 @@ public:
 	FPCGInterpCurveAccessor(const FStructProperty* InProperty, TArray<const FProperty*>&& ExtraProperties = {})
 		: Super(/*bInReadOnly=*/ false)
 		, IPCGPropertyChain(InProperty, std::move(ExtraProperties))
-		, Property(InProperty)
 	{
 		static_assert(PCG::Private::IsPCGType<Type>());		
-		check(Property && Property->Struct->IsChildOf(TBaseStructure<CurveType>::Get()));
+		check(InProperty && InProperty->Struct->IsChildOf(TBaseStructure<CurveType>::Get()));
 		TopPropertyStruct = GetTopPropertyStruct();
 	}
 
@@ -142,6 +148,316 @@ public:
 	}
 
 private:
-	const FStructProperty* Property = nullptr;
+	const UStruct* TopPropertyStruct = nullptr;
+};
+
+/**
+* Templated accessor for location/rotation/scale in world coordinates. It's important that the keys only have a single value, the struct that holds the spline data/struct, since
+* there is a single transform per spline.
+* Keys supported: PCGSplineData, FPCGSplineStruct
+*/
+template <typename T, EPCGControlPointsAccessorTarget Target, bool bWorldCoordinates>
+class FPCGControlPointsAccessor : public IPCGAttributeAccessorT<FPCGControlPointsAccessor<T, Target, bWorldCoordinates>>, IPCGPropertyChain
+{
+public:
+	// The underlying type is a quat if we target the rotation, otherwise a vector
+	using Type = T;
+	using Super = IPCGAttributeAccessorT<FPCGControlPointsAccessor<Type, Target, bWorldCoordinates>>;
+	
+	FPCGControlPointsAccessor(const FStructProperty* InProperty, TArray<const FProperty*>&& ExtraProperties = {})
+		: Super(/*bInReadOnly=*/ false)
+		, IPCGPropertyChain(InProperty, std::move(ExtraProperties))
+	{
+		static_assert(PCG::Private::IsPCGType<Type>());		
+		check(InProperty && InProperty->Struct->IsChildOf<FPCGSplineStruct>());
+		TopPropertyStruct = GetTopPropertyStruct();
+	}
+
+	bool GetRangeImpl(TArrayView<Type> OutValues, int32 Index, const IPCGAttributeAccessorKeys& Keys) const
+	{
+		// We want to access the transform on the spline struct, there is just a single one.
+		const void* ContainerKeys = nullptr;
+		TArrayView<const void*> ContainerKeysView(&ContainerKeys, 1);
+		if (!Keys.GetKeys(Index, ContainerKeysView))
+		{
+			return false;
+		}
+
+		// Validation to not access keys that are not the expected type. Done after the GetKeys, as we also want to discard other type of incompatible
+		// keys (like a Default Metadata entry key)
+		if (!ensure(Keys.IsClassSupported(TopPropertyStruct)))
+		{
+			return false;
+		}
+
+		PCGPropertyAccessor::AddressOffset(GetPropertyChain(), ContainerKeysView);
+
+		const FPCGSplineStruct& SplineStruct = *static_cast<const FPCGSplineStruct*>(ContainerKeys);
+		const FInterpCurveVector& Positions = SplineStruct.GetSplinePointsPosition();
+		const FInterpCurveQuat& Rotations = SplineStruct.GetSplinePointsRotation();
+		const FInterpCurveVector& Scale = SplineStruct.GetSplinePointsScale();
+		const FTransform SplineTransform = SplineStruct.GetTransform();
+
+		const int32 NumPoints = Positions.Points.Num();
+		if (NumPoints == 0)
+		{
+			return false;
+		}
+		
+		for (int32 i = 0; i < OutValues.Num(); ++i)
+		{
+			const int32 CurrIndex = (Index + i) % NumPoints;
+			
+			if constexpr (Target == EPCGControlPointsAccessorTarget::Location)
+			{
+				static_assert(std::is_same_v<Type, FVector>);
+				OutValues[i] = Positions.Points[CurrIndex].OutVal;
+				if constexpr (bWorldCoordinates)
+				{
+					OutValues[i] = SplineTransform.TransformPosition(OutValues[i]);
+				}
+			}
+			else if constexpr (Target == EPCGControlPointsAccessorTarget::Rotation)
+			{
+				static_assert(std::is_same_v<Type, FQuat>);
+				OutValues[i] = Rotations.Points[CurrIndex].OutVal;
+				if constexpr (bWorldCoordinates)
+				{
+					OutValues[i] = SplineTransform.TransformRotation(OutValues[i]);
+				}
+			}
+			else if constexpr (Target == EPCGControlPointsAccessorTarget::Scale)
+			{
+				
+				static_assert(std::is_same_v<Type, FVector>);
+				OutValues[i] = Scale.Points[CurrIndex].OutVal;
+				if constexpr (bWorldCoordinates)
+				{
+					OutValues[i] = (FTransform(FQuat::Identity, FVector::ZeroVector, OutValues[i]) * SplineTransform).GetScale3D();
+				}
+			}
+			else if constexpr (Target == EPCGControlPointsAccessorTarget::Transform)
+			{
+				static_assert(std::is_same_v<Type, FTransform>);
+				OutValues[i] = FTransform(Rotations.Points[CurrIndex].OutVal, Positions.Points[CurrIndex].OutVal, Scale.Points[CurrIndex].OutVal);
+				if constexpr (bWorldCoordinates)
+				{
+					OutValues[i] = OutValues[i] * SplineTransform;
+				}
+			}
+			else
+			{
+				// Pitfall static assert
+				static_assert(!std::is_same_v<Type, Type>);
+			}
+		}
+
+		return true;
+	}
+
+	bool SetRangeImpl(TArrayView<const Type> InValues, int32 Index, IPCGAttributeAccessorKeys& Keys, EPCGAttributeAccessorFlags)
+	{		
+		// We want to access the transform on the spline struct, there is just a single one.
+		void* ContainerKeys = nullptr;
+		TArrayView<void*> ContainerKeysView(&ContainerKeys, 1);
+		if (!Keys.GetKeys(Index, ContainerKeysView))
+		{
+			return false;
+		}
+
+		// Validation to not access keys that are not the expected type. Done after the GetKeys, as we also want to discard other type of incompatible
+		// keys (like a Default Metadata entry key)
+		if (!ensure(Keys.IsClassSupported(TopPropertyStruct)))
+		{
+			return false;
+		}
+
+		PCGPropertyAccessor::AddressOffset(GetPropertyChain(), ContainerKeysView);
+
+		FPCGSplineStruct& SplineStruct = *static_cast<FPCGSplineStruct*>(ContainerKeys);
+		const FTransform InverseSplineTransform = SplineStruct.GetTransform().Inverse();
+		
+		const int32 NumPoints = SplineStruct.SplineCurves.Position.Points.Num();
+		if (NumPoints == 0)
+		{
+			return false;
+		}
+		
+		for (int32 i = 0; i < InValues.Num(); ++i)
+		{
+			const int32 CurrIndex = (Index + i) % NumPoints;
+			
+			if constexpr (Target == EPCGControlPointsAccessorTarget::Location)
+			{
+				static_assert(std::is_same_v<Type, FVector>);
+				FVector& OutPosition = SplineStruct.SplineCurves.Position.Points[CurrIndex].OutVal;
+				OutPosition = InValues[i];
+				if constexpr (bWorldCoordinates)
+				{
+					OutPosition = InverseSplineTransform.TransformPosition(OutPosition);
+				}
+			}
+			else if constexpr (Target == EPCGControlPointsAccessorTarget::Rotation)
+			{
+				static_assert(std::is_same_v<Type, FQuat>);
+				FQuat& OutRotation = SplineStruct.SplineCurves.Rotation.Points[CurrIndex].OutVal;
+				OutRotation = InValues[i];
+				if constexpr (bWorldCoordinates)
+				{
+					OutRotation = InverseSplineTransform.TransformRotation(OutRotation);
+				}
+			}
+			else if constexpr (Target == EPCGControlPointsAccessorTarget::Scale)
+			{
+				static_assert(std::is_same_v<Type, FVector>);
+				FVector& OutScale = SplineStruct.SplineCurves.Scale.Points[CurrIndex].OutVal;
+				OutScale = InValues[i];
+				
+				if constexpr (bWorldCoordinates)
+				{
+					OutScale = (FTransform(FQuat::Identity, FVector::ZeroVector, OutScale) * InverseSplineTransform).GetScale3D();
+				}
+			}
+			else if constexpr (Target == EPCGControlPointsAccessorTarget::Transform)
+			{
+				static_assert(std::is_same_v<Type, FTransform>);
+				FTransform Transform = InValues[i];
+				if constexpr (bWorldCoordinates)
+				{
+					Transform = Transform * InverseSplineTransform;
+				}
+
+				SplineStruct.SplineCurves.Position.Points[CurrIndex].OutVal = InValues[i].GetLocation();
+				SplineStruct.SplineCurves.Rotation.Points[CurrIndex].OutVal = InValues[i].GetRotation();
+				SplineStruct.SplineCurves.Scale.Points[CurrIndex].OutVal = InValues[i].GetScale3D();
+			}
+			else
+			{
+				// Pitfall static assert
+				static_assert(!std::is_same_v<Type, Type>);
+			}
+		}
+
+		return true;
+	}
+
+private:
+	const UStruct* TopPropertyStruct = nullptr;
+};
+
+/**
+* Templated accessor for global spline data. Note that closed loop value is read-only.
+* Keys supported: PCGSplineData, FPCGSplineStruct
+*/
+template <typename T, EPCGSplineAccessorTarget Target>
+class FPCGSplineAccessor : public IPCGAttributeAccessorT<FPCGSplineAccessor<T, Target>>, IPCGPropertyChain
+{
+public:
+	// The underlying type is a quat if we target the rotation, otherwise a vector
+	using Type = T;
+	using Super = IPCGAttributeAccessorT<FPCGSplineAccessor<T, Target>>;
+	
+	FPCGSplineAccessor(const FStructProperty* InProperty, TArray<const FProperty*>&& ExtraProperties = {})
+		: Super(/*bInReadOnly=*/ Target == EPCGSplineAccessorTarget::ClosedLoop)
+		, IPCGPropertyChain(InProperty, std::move(ExtraProperties))
+	{
+		static_assert(PCG::Private::IsPCGType<Type>());		
+		check(InProperty && InProperty->Struct->IsChildOf<FPCGSplineStruct>());
+		TopPropertyStruct = GetTopPropertyStruct();
+	}
+
+	bool GetRangeImpl(TArrayView<Type> OutValues, int32 Index, const IPCGAttributeAccessorKeys& Keys) const
+	{
+		// We want to access the transform on the spline struct, there is just a single one.
+		const void* ContainerKeys = nullptr;
+		TArrayView<const void*> ContainerKeysView(&ContainerKeys, 1);
+		if (!Keys.GetKeys(Index, ContainerKeysView))
+		{
+			return false;
+		}
+
+		// Validation to not access keys that are not the expected type. Done after the GetKeys, as we also want to discard other type of incompatible
+		// keys (like a Default Metadata entry key)
+		if (!ensure(Keys.IsClassSupported(TopPropertyStruct)))
+		{
+			return false;
+		}
+
+		PCGPropertyAccessor::AddressOffset(GetPropertyChain(), ContainerKeysView);
+
+		const FPCGSplineStruct& SplineStruct = *static_cast<const FPCGSplineStruct*>(ContainerKeys);
+		
+		for (int32 i = 0; i < OutValues.Num(); ++i)
+		{
+			if constexpr (Target == EPCGSplineAccessorTarget::Transform)
+			{
+				static_assert(std::is_same_v<Type, FTransform>);
+				OutValues[i] = SplineStruct.GetTransform();
+			}
+			else if constexpr (Target == EPCGSplineAccessorTarget::ClosedLoop)
+			{
+				static_assert(std::is_same_v<Type, bool>);
+				OutValues[i] = SplineStruct.IsClosedLoop();
+			}
+			else
+			{
+				// Pitfall static assert
+				static_assert(!std::is_same_v<Type, Type>);
+			}
+		}
+
+		return true;
+	}
+
+	bool SetRangeImpl(TArrayView<const Type> InValues, int32 Index, IPCGAttributeAccessorKeys& Keys, EPCGAttributeAccessorFlags)
+	{
+		if (Target == EPCGSplineAccessorTarget::ClosedLoop)
+		{
+			// Not supported, as we need to update the spline when we change it, which is not threadsafe.
+			return false;
+		}
+			
+		// We want to access the transform on the spline struct, there is just a single one.
+		void* ContainerKeys = nullptr;
+		TArrayView<void*> ContainerKeysView(&ContainerKeys, 1);
+		if (!Keys.GetKeys(Index, ContainerKeysView))
+		{
+			return false;
+		}
+
+		// Validation to not access keys that are not the expected type. Done after the GetKeys, as we also want to discard other type of incompatible
+		// keys (like a Default Metadata entry key)
+		if (!ensure(Keys.IsClassSupported(TopPropertyStruct)))
+		{
+			return false;
+		}
+
+		PCGPropertyAccessor::AddressOffset(GetPropertyChain(), ContainerKeysView);
+
+		FPCGSplineStruct& SplineStruct = *static_cast<FPCGSplineStruct*>(ContainerKeys);
+		
+		for (int32 i = 0; i < InValues.Num(); ++i)
+		{
+			if constexpr (Target == EPCGSplineAccessorTarget::Transform)
+			{
+				static_assert(std::is_same_v<Type, FTransform>);
+				SplineStruct.Transform = InValues[i];
+			}
+			else if constexpr (Target == EPCGSplineAccessorTarget::ClosedLoop)
+			{
+				// Not supported, as we need to update the spline when we change it, which is not threadsafe.
+				// Need to stay here to avoid unreachable code.
+			}
+			else
+			{
+				// Pitfall static assert
+				static_assert(!std::is_same_v<Type, Type>);
+			}
+		}
+
+		return true;
+	}
+
+private:
 	const UStruct* TopPropertyStruct = nullptr;
 };
\ No newline at end of file
