[PCG] Use the new spline accessors for the spline visualisation

* Spline are no more default collapsed when inspecting the attribute list view
* It uses the new spline accessors, so we have a 1 for 1 with what you see and what you get using accessors.
* It also adds the global attributes [REDACTED_HANDLE].$SplineTransform and [REDACTED_HANDLE].$IsClosed

#rnx

diff --git a/Engine/Plugins/PCG/Source/PCGEditor/Private/DataVisualizations/PCGSplineDataVisualization.cpp b/Engine/Plugins/PCG/Source/PCGEditor/Private/DataVisualizations/PCGSplineDataVisualization.cpp
--- a/Engine/Plugins/PCG/Source/PCGEditor/Private/DataVisualizations/PCGSplineDataVisualization.cpp
+++ b/Engine/Plugins/PCG/Source/PCGEditor/Private/DataVisualizations/PCGSplineDataVisualization.cpp
@@ -6,6 +6,9 @@
 #include "Data/PCGPointData.h"
 #include "Data/PCGSplineData.h"
 #include "Data/PCGSplineStruct.h"
+#include "DataVisualizations/PCGDataVisualizationHelpers.h"
+
+#define LOCTEXT_NAMESPACE "PCGSplineDataVisualization"
 
 namespace PCGSplineDataVisualizationConstants
 {
@@ -14,8 +17,51 @@ namespace PCGSplineDataVisualizationConstants
 
 FPCGTableVisualizerInfo IPCGSplineDataVisualization::GetTableVisualizerInfoWithDomain(const UPCGData* Data, const FPCGMetadataDomainID& DomainID) const
 {
-	FPCGTableVisualizerInfo Info = IPCGSpatialDataVisualization::GetTableVisualizerInfoWithDomain(Data, DomainID);
+	using namespace PCGDataVisualizationHelpers;
+
+	const UPCGSplineData* SplineData = CastChecked<const UPCGSplineData>(Data);
+
+	FPCGTableVisualizerInfo Info;
+	Info.Data = SplineData;
 
+	if (DomainID == PCGMetadataDomainID::Data)
+	{
+		// Column Sorting
+		FPCGAttributePropertySelector IndexSelector = FPCGAttributePropertySelector::CreateExtraPropertySelector(EPCGExtraProperties::Index);
+		IndexSelector.SetDomainName(PCGDataConstants::GlobalDomainName);
+		
+		AddColumnInfo(Info, SplineData, IndexSelector);
+		Info.SortingColumn = Info.ColumnInfos.Last().Id;
+
+		// Columns
+		AddPropertyEnumColumnInfo<FTransform>(Info, SplineData, EPCGSplineStructProperties::SplineTransform);
+		AddPropertyEnumColumnInfo<bool>(Info, SplineData, EPCGSplineStructProperties::IsClosed);
+		
+		// Add Metadata Columns
+		CreateMetadataColumnInfos(SplineData, Info, PCGMetadataDomainID::Data);
+
+		return Info;
+	}
+	
+	// Column Sorting
+	AddColumnInfo(Info, SplineData, FPCGAttributePropertySelector::CreateExtraPropertySelector(EPCGExtraProperties::Index));
+	Info.SortingColumn = Info.ColumnInfos.Last().Id;
+	
+	// Columns
+	AddPropertyEnumColumnInfo<FVector>(Info, SplineData, EPCGSplineStructProperties::Position);
+	AddPropertyEnumColumnInfo<FRotator>(Info, SplineData, EPCGSplineStructProperties::Rotation);
+	AddPropertyEnumColumnInfo<FVector>(Info, SplineData, EPCGSplineStructProperties::Scale);
+	AddPropertyEnumColumnInfo<FVector>(Info, SplineData, EPCGSplineStructProperties::LocalPosition);
+	AddPropertyEnumColumnInfo<FRotator>(Info, SplineData, EPCGSplineStructProperties::LocalRotation);
+	AddPropertyEnumColumnInfo<FVector>(Info, SplineData, EPCGSplineStructProperties::LocalScale);
+	AddPropertyEnumColumnInfo<FVector>(Info, SplineData, EPCGSplineStructProperties::ArriveTangent);
+	AddPropertyEnumColumnInfo<FVector>(Info, SplineData, EPCGSplineStructProperties::LeaveTangent);
+	AddPropertyEnumColumnInfo<int32>(Info, SplineData, EPCGSplineStructProperties::InterpType);
+	
+	// Add Metadata Columns
+	CreateMetadataColumnInfos(SplineData, Info, PCGMetadataDomainID::Elements);
+	
+	// Focus on data behavior
 	if (DomainID == PCGMetadataDomainID::Elements)
 	{
 		Info.FocusOnDataCallback = [](const UPCGData* Data, TArrayView<const int> Indices)
