[PCG] Fix Collapse Points for rotations

* Rotations average needs to be done in Quaternions, and normalized at the end

#rnx

diff --git a/Engine/Plugins/PCG/Source/PCG/Private/Elements/PCGCollapsePoints.cpp b/Engine/Plugins/PCG/Source/PCG/Private/Elements/PCGCollapsePoints.cpp
--- a/Engine/Plugins/PCG/Source/PCG/Private/Elements/PCGCollapsePoints.cpp
+++ b/Engine/Plugins/PCG/Source/PCG/Private/Elements/PCGCollapsePoints.cpp
@@ -586,9 +586,10 @@ bool FPCGCollapsePointsElement::ExecuteInternal(FPCGContext* InContext) const
 					TUniquePtr<const IPCGAttributeAccessor>& SourceAccessor = CollapseState.SourceMergeAccessors[AccessorIndex];
 					TUniquePtr<IPCGAttributeAccessor>& OutputAccessor = CollapseState.OutputMergeAccessors[AccessorIndex];
 
-					auto Callback = [PartitionIndex, &SecondaryPointIndices, &PartitionWeights, &SourceAccessor, &SourceKeys, &OutputAccessor, &OutputKeys](auto Dummy)
+					auto Callback = [PartitionIndex, &SecondaryPointIndices, &PartitionWeights, &SourceAccessor, &SourceKeys, &OutputAccessor, &OutputKeys]<typename T>(T)
 					{
-						using OutputType = decltype(Dummy);
+						// If the attribute is a rotator, do the weighted sum in Quat.
+						using OutputType = std::conditional_t<std::is_same_v<T, FRotator>, FQuat, T>;
 
 						if constexpr (PCG::Private::MetadataTraits<OutputType>::CanInterpolate)
 						{
@@ -604,6 +605,12 @@ bool FPCGCollapsePointsElement::ExecuteInternal(FPCGContext* InContext) const
 								Value = PCG::Private::MetadataTraits<OutputType>::WeightedSum(Value, SecondaryValue, PartitionWeights[Index]);
 							}
 
+							if constexpr (std::is_same_v<OutputType, FQuat>)
+							{
+								// We need to normalize the resulting quaternion
+								Value.Normalize();
+							}
+							
 							OutputAccessor->Set(Value, PartitionIndex, *OutputKeys);
 						}
 						else
