[PCG] Add Remove Empty Data nodes

* Use the $Index accessor keys for genericity.

#rnx

diff --git a/Engine/Plugins/PCG/Source/PCG/Private/Elements/PCGRemoveEmptyData.cpp b/Engine/Plugins/PCG/Source/PCG/Private/Elements/PCGRemoveEmptyData.cpp
new file mode 100644
--- /dev/null
+++ b/Engine/Plugins/PCG/Source/PCG/Private/Elements/PCGRemoveEmptyData.cpp
@@ -0,0 +1,76 @@
+// Copyright Epic Games, Inc. All Rights Reserved.
+
+#include "Elements/PCGRemoveEmptyData.h"
+
+#include "PCGContext.h"
+#include "Metadata/Accessors/PCGAttributeAccessorHelpers.h"
+#include "Metadata/Accessors/PCGAttributeAccessorKeys.h"
+
+#define LOCTEXT_NAMESPACE "PCGRemoveEmptyDataElement"
+
+#if WITH_EDITOR
+FName UPCGRemoveEmptyDataSettings::GetDefaultNodeName() const
+{
+	return FName(TEXT("RemoveEmptyData"));
+}
+
+FText UPCGRemoveEmptyDataSettings::GetDefaultNodeTitle() const
+{
+	return LOCTEXT("NodeTitle", "Remove Empty Data");
+}
+
+FText UPCGRemoveEmptyDataSettings::GetNodeTooltipText() const
+{
+	return LOCTEXT("NodeTooltip", "Remove all data in the input pin that is empty.");
+}
+#endif // WITH_EDITOR
+
+FPCGElementPtr UPCGRemoveEmptyDataSettings::CreateElement() const
+{
+	return MakeShared<FPCGRemoveEmptyDataElement>();
+}
+
+TArray<FPCGPinProperties> UPCGRemoveEmptyDataSettings::InputPinProperties() const
+{
+	TArray<FPCGPinProperties> Properties;
+	Properties.Emplace_GetRef(PCGPinConstants::DefaultInputLabel, EPCGDataType::Any).SetRequiredPin();	
+	return Properties;
+}
+
+TArray<FPCGPinProperties> UPCGRemoveEmptyDataSettings::OutputPinProperties() const
+{ 
+	TArray<FPCGPinProperties> Properties;
+	Properties.Emplace(PCGPinConstants::DefaultOutputLabel, EPCGDataType::Any);	
+	return Properties;
+}
+
+bool FPCGRemoveEmptyDataElement::ExecuteInternal(FPCGContext* InContext) const
+{
+	TRACE_CPUPROFILER_EVENT_SCOPE(FPCGRemoveEmptyDataElement::Execute);
+
+	check(InContext);
+
+	const UPCGRemoveEmptyDataSettings* Settings = InContext->GetInputSettings<UPCGRemoveEmptyDataSettings>();
+	check(Settings);
+
+	for (const FPCGTaggedData& Input : InContext->InputData.GetInputsByPin(PCGPinConstants::DefaultInputLabel))
+	{
+		if (!Input.Data)
+		{
+			continue;
+		}
+
+		// Get the keys on the index, if it doesn't exist or is empty, discard the data.
+		const TUniquePtr<const IPCGAttributeAccessorKeys> Keys = PCGAttributeAccessorHelpers::CreateConstKeys(Input.Data, FPCGAttributePropertySelector::CreateExtraPropertySelector(EPCGExtraProperties::Index));
+		if (!Keys || Keys->GetNum() == 0)
+		{
+			continue;
+		}
+
+		InContext->OutputData.TaggedData.Add(Input);
+	}
+
+	return true;
+}
+
+#undef LOCTEXT_NAMESPACE
diff --git a/Engine/Plugins/PCG/Source/PCG/Public/Elements/PCGRemoveEmptyData.h b/Engine/Plugins/PCG/Source/PCG/Public/Elements/PCGRemoveEmptyData.h
new file mode 100644
--- /dev/null
+++ b/Engine/Plugins/PCG/Source/PCG/Public/Elements/PCGRemoveEmptyData.h
@@ -0,0 +1,35 @@
+// Copyright Epic Games, Inc. All Rights Reserved.
+
+#pragma once
+
+#include "PCGSettings.h"
+
+#include "PCGRemoveEmptyData.generated.h"
+
+UCLASS(MinimalAPI, BlueprintType, ClassGroup = (Procedural))
+class UPCGRemoveEmptyDataSettings : public UPCGSettings
+{
+	GENERATED_BODY()
+
+public:
+	//~Begin UPCGSettings interface
+#if WITH_EDITOR
+	virtual FName GetDefaultNodeName() const override;
+	virtual FText GetDefaultNodeTitle() const override;
+    virtual FText GetNodeTooltipText() const override;
+	virtual EPCGSettingsType GetType() const override { return EPCGSettingsType::Filter; }
+#endif
+
+protected:
+	virtual TArray<FPCGPinProperties> InputPinProperties() const override;
+	virtual TArray<FPCGPinProperties> OutputPinProperties() const override;
+	virtual FPCGElementPtr CreateElement() const override;
+	//~End UPCGSettings interface
+};
+
+class FPCGRemoveEmptyDataElement : public IPCGElement
+{
+protected:
+	virtual bool ExecuteInternal(FPCGContext* InContext) const override;
+};
+
