[PCG] Created a "Export Selected Attributes" node to export PCG Data Attributes to file

+ Supports Binary using FArchive
+ Support Json using custom formatting

diff --git a/Engine/Plugins/PCG/Source/PCG/PCG.Build.cs b/Engine/Plugins/PCG/Source/PCG/PCG.Build.cs
--- a/Engine/Plugins/PCG/Source/PCG/PCG.Build.cs
+++ b/Engine/Plugins/PCG/Source/PCG/PCG.Build.cs
@@ -31,6 +31,7 @@ namespace UnrealBuildTool.Rules
 					"PCGCompute",
 					"Renderer",
 					"Voronoi",
+					"Json"
 				});
 
 			if (Target.bBuildEditor == true)
diff --git a/Engine/Plugins/PCG/Source/PCG/Private/Elements/IO/PCGExportSelectedAttributes.cpp b/Engine/Plugins/PCG/Source/PCG/Private/Elements/IO/PCGExportSelectedAttributes.cpp
new file mode 100644
--- /dev/null
+++ b/Engine/Plugins/PCG/Source/PCG/Private/Elements/IO/PCGExportSelectedAttributes.cpp
@@ -0,0 +1,539 @@
+// Copyright Epic Games, Inc. All Rights Reserved.
+
+#include "Elements/IO/PCGExportSelectedAttributes.h"
+
+#include "Data/PCGBasePointData.h"
+#include "Data/PCGSplineData.h"
+#include "Metadata/PCGMetadata.h"
+#include "Metadata/PCGMetadataAttributeTraits.h"
+#include "Metadata/Accessors/PCGAttributeAccessorHelpers.h"
+#include "Utils/PCGLogErrors.h"
+
+#include "DesktopPlatformModule.h"
+#include "HAL/FileManager.h"
+#include "HAL/FileManagerGeneric.h"
+
+#if WITH_EDITOR
+#include "IContentBrowserSingleton.h"
+#endif // WITH_EDITOR
+#include "IDesktopPlatform.h"
+#include "Data/PCGSplineStruct.h"
+#include "Misc/FileHelper.h"
+#include "Serialization/JsonSerializer.h"
+#include "Serialization/JsonWriter.h"
+
+#define LOCTEXT_NAMESPACE "PCGExportSelectedAttributesElement"
+
+// Register the custom version with core
+namespace PCGExportSelectedAttributes
+{
+	struct FAccessorData
+	{
+		FPCGAttributePropertySelector Selector;
+		TUniquePtr<const IPCGAttributeAccessor> Accessor = nullptr;
+		TUniquePtr<const IPCGAttributeAccessorKeys> Keys = nullptr;
+	};
+
+	namespace Constants
+	{
+		static constexpr int32 DefaultInlineSourceCount = 16;
+
+		struct FCustomExportVersion
+		{
+			enum Type
+			{
+				InitialVersion = 0,
+
+				// New versions can be added above this line
+				VersionPlusOne,
+				LatestVersion = VersionPlusOne - 1
+			};
+
+			static FName GetFriendlyName() { return FName(TEXT("Initial Version")); }
+			const static FGuid GUID;
+		};
+
+		const FGuid FCustomExportVersion::GUID = FGuid(0x04E74488, 0x4BAC8717, 0xBBB18694, 0x39F8F3CE);
+		FCustomVersionRegistration GRegisterPCGExportSelectedAttributesCustomVersion(FCustomExportVersion::GUID, FCustomExportVersion::LatestVersion, TEXT("PCGExportSelectedAttributes"));
+	}
+
+	namespace Helpers
+	{
+		void AddAllPointPropertySelectors(TArray<FPCGAttributePropertySelector>& InOutSelectors)
+		{
+			// Some of the point properties are redundant, so pick these by hand for now
+			InOutSelectors.Add(FPCGAttributePropertySelector::CreatePointPropertySelector(EPCGPointProperties::Transform));
+			InOutSelectors.Add(FPCGAttributePropertySelector::CreatePointPropertySelector(EPCGPointProperties::Density));
+			InOutSelectors.Add(FPCGAttributePropertySelector::CreatePointPropertySelector(EPCGPointProperties::BoundsMin));
+			InOutSelectors.Add(FPCGAttributePropertySelector::CreatePointPropertySelector(EPCGPointProperties::BoundsMax));
+			InOutSelectors.Add(FPCGAttributePropertySelector::CreatePointPropertySelector(EPCGPointProperties::Color));
+			InOutSelectors.Add(FPCGAttributePropertySelector::CreatePointPropertySelector(EPCGPointProperties::Steepness));
+			InOutSelectors.Add(FPCGAttributePropertySelector::CreatePointPropertySelector(EPCGPointProperties::Seed));
+		}
+
+		void AddAllSplinePropertySelectors(TArray<FPCGAttributePropertySelector>& InOutSelectors)
+		{
+			InOutSelectors.Add(FPCGAttributePropertySelector::CreatePropertySelector(TEXT("Position")));
+			InOutSelectors.Add(FPCGAttributePropertySelector::CreatePropertySelector(TEXT("Rotation")));
+			InOutSelectors.Add(FPCGAttributePropertySelector::CreatePropertySelector(TEXT("Scale")));
+			InOutSelectors.Add(FPCGAttributePropertySelector::CreatePropertySelector(TEXT("ArriveTangent")));
+			InOutSelectors.Add(FPCGAttributePropertySelector::CreatePropertySelector(TEXT("LeaveTangent")));
+			InOutSelectors.Add(FPCGAttributePropertySelector::CreatePropertySelector(TEXT("InterpType")));
+		}
+
+		void AddAllAttributeSelectors(const UPCGData* InData, TArray<FPCGAttributePropertySelector>& InOutSelectors)
+		{
+			check(InData);
+			check(InData->ConstMetadata());
+
+			TArray<FPCGAttributeIdentifier> AttributeNames;
+			TArray<EPCGMetadataTypes> AttributeTypes;
+			InData->ConstMetadata()->GetAllAttributes(AttributeNames, AttributeTypes);
+			check(AttributeNames.Num() == AttributeTypes.Num());
+
+			for (int32 Index = 0; Index < AttributeNames.Num(); ++Index)
+			{
+				InOutSelectors.Add(FPCGAttributePropertySelector::CreateAttributeSelector(AttributeNames[Index].Name));
+			}
+		}
+
+		template <typename T>
+		bool GetAttributeValues(const FAccessorData& Data, TArray<T>& OutValues, EPCGAttributeAccessorFlags Flags, const FPCGContext* Context)
+		{
+			if constexpr (std::is_trivially_copyable_v<T>)
+			{
+				OutValues.SetNumUninitialized(Data.Keys->GetNum());
+			}
+			else
+			{
+				OutValues.SetNum(Data.Keys->GetNum());
+			}
+
+			if (!Data.Accessor->GetRange<T>(OutValues, 0, *Data.Keys, Flags))
+			{
+				PCGLog::Metadata::LogFailToGetAttributeError(Data.Selector, Context);
+				return false;
+			}
+
+			return true;
+		}
+
+		namespace String
+		{
+			FString MergeFileName(const FString& Path, const FString& File, const FString& Ext)
+			{
+				if (Path.IsEmpty() || File.IsEmpty() || Ext.IsEmpty())
+				{
+					return FString();
+				}
+
+				return FString::Format(TEXT("{0}/{1}{2}"), {Path, File, Ext});
+			}
+
+			void ToPrecisionString(const double Value, const int32 Precision, FString& OutString)
+			{
+				OutString = FString::SanitizeFloat(Value, Precision);
+			}
+		}
+
+		namespace Json
+		{
+			TSharedPtr<FJsonObject> CreateJsonHeaderObject(const bool bAddCustomDataVersion = false, const int32 CustomDataVersion = 0)
+			{
+				TSharedPtr<FJsonObject> JsonHeaderObject = MakeShared<FJsonObject>();
+				JsonHeaderObject->SetNumberField(LOCTEXT("JsonFieldVersionNumber", "Export Version").ToString(), Constants::FCustomExportVersion::LatestVersion);
+
+				if (bAddCustomDataVersion)
+				{
+					JsonHeaderObject->SetNumberField(LOCTEXT("JsonFieldVersionNumber", "Custom Data Version").ToString(), CustomDataVersion);
+				}
+
+				return JsonHeaderObject;
+			}
+
+			TSharedPtr<FJsonValueArray> ToJsonArray(const TConstArrayView<double> ArrayView)
+			{
+				TArray<TSharedPtr<FJsonValue>> JsonArray;
+				for (double Value : ArrayView)
+				{
+					JsonArray.Emplace(MakeShared<FJsonValueNumber>(Value));
+				}
+
+				return MakeShared<FJsonValueArray>(JsonArray);
+			}
+
+			// The order matters here in the JsonObject construction.
+			template <typename T>
+			TSharedPtr<FJsonValue> ConvertFloatingPointTypeToJson(const T& Value)
+			{
+				static_assert(PCG::Private::MetadataTraits<T>::IsFloatingPoint, "Helper is only designed to break floating point compound components.");
+
+				if constexpr (std::is_same_v<T, double>)
+				{
+					return MakeShared<FJsonValueNumber>(Value);
+				}
+				if constexpr (std::is_same_v<T, FVector2D>)
+				{
+					return ToJsonArray({Value.X, Value.Y});
+				}
+				else if constexpr (std::is_same_v<T, FVector>)
+				{
+					return ToJsonArray({Value.X, Value.Y, Value.Z});
+				}
+				else if constexpr (std::is_same_v<T, FVector4> || std::is_same_v<T, FQuat>)
+				{
+					return ToJsonArray({Value.X, Value.Y, Value.Z, Value.W});
+				}
+				else if constexpr (std::is_same_v<T, FTransform>)
+				{
+					TSharedPtr<FJsonObject> JsonTransformObject = MakeShared<FJsonObject>();
+					JsonTransformObject->SetArrayField(
+						LOCTEXT("JsonFieldTranslation", "Translation").ToString(),
+						ConvertFloatingPointTypeToJson(Value.GetTranslation())->AsArray());
+					JsonTransformObject->SetArrayField(
+						LOCTEXT("JsonFieldRotation", "Rotation").ToString(),
+						ConvertFloatingPointTypeToJson(Value.GetRotation())->AsArray());
+					JsonTransformObject->SetArrayField(
+						LOCTEXT("JsonFieldScale", "Scale").ToString(),
+						ConvertFloatingPointTypeToJson(Value.GetScale3D())->AsArray());
+
+					return MakeShared<FJsonValueObject>(JsonTransformObject);
+				}
+				else if constexpr (std::is_same_v<T, FRotator>)
+				{
+					return ToJsonArray({Value.Pitch, Value.Yaw, Value.Roll});
+				}
+
+				return nullptr;
+			}
+		}
+	}
+
+	using FSourceDataArray = TArray<FAccessorData, TInlineAllocator<Constants::DefaultInlineSourceCount>>;
+}
+
+UPCGExportSelectedAttributesSettings::UPCGExportSelectedAttributesSettings()
+{
+	// Initialize with one @Last source.
+	AttributeSelectors.Emplace();
+}
+
+TArray<FPCGPinProperties> UPCGExportSelectedAttributesSettings::InputPinProperties() const
+{
+	TArray<FPCGPinProperties> PinProperties;
+	PinProperties.Emplace_GetRef(PCGPinConstants::DefaultInputLabel, EPCGDataType::Any, /*bInAllowMultipleConnections=*/false, /*bInAllowMultipleData=*/false).SetRequiredPin();
+
+	return PinProperties;
+}
+
+TArray<FPCGPinProperties> UPCGExportSelectedAttributesSettings::OutputPinProperties() const
+{
+	TArray<FPCGPinProperties> PinProperties;
+	FPCGPinProperties& DepPin = PinProperties.Emplace_GetRef(PCGPinConstants::DefaultExecutionDependencyLabel, EPCGDataType::Any);
+#if WITH_EDITOR
+	DepPin.Tooltip = PCGPinConstants::Tooltips::ExecutionDependencyTooltip;
+#endif // WITH_EDITOR
+	DepPin.Usage = EPCGPinUsage::DependencyOnly;
+
+	return PinProperties;
+}
+
+FPCGElementPtr UPCGExportSelectedAttributesSettings::CreateElement() const
+{
+	return MakeShared<FPCGExportSelectedAttributesElement>();
+}
+
+bool FPCGExportSelectedAttributesElement::ExecuteInternal(FPCGContext* Context) const
+{
+	TRACE_CPUPROFILER_EVENT_SCOPE(FPCGExportSelectedAttributesElement::Execute);
+	using namespace PCGExportSelectedAttributes;
+
+	// Since this generates data on disk, only allow execution on editor approved platforms.
+#if PLATFORM_WINDOWS || PLATFORM_LINUX || PLATFORM_MAC
+
+	const UPCGExportSelectedAttributesSettings* Settings = Context->GetInputSettings<UPCGExportSelectedAttributesSettings>();
+	check(Settings);
+
+	const TArray<FPCGTaggedData> Inputs = Context->InputData.GetInputsByPin(PCGPinConstants::DefaultInputLabel);
+
+	if (Inputs.IsEmpty() || Settings->AttributeSelectors.IsEmpty())
+	{
+		return true;
+	}
+
+	// @todo_pcg: In the future, could be extended to support multiple inputs, but should also support separate file names for each, etc.
+	if (Inputs.Num() > 1)
+	{
+		PCGLog::InputOutput::LogFirstInputOnlyWarning(PCGPinConstants::DefaultInputLabel, Context);
+	}
+
+	// Take the property override (FString) if it exists, otherwise take the user property (FDirectoryPath).
+	FString OutputDirectory = Settings->Path.Path;
+	if (OutputDirectory.IsEmpty())
+	{
+#if WITH_EDITOR
+		IDesktopPlatform* DesktopPlatform = FDesktopPlatformModule::Get();
+		const bool bSuccess = DesktopPlatform->OpenDirectoryDialog(nullptr, TEXT("Choose output directory"), FPaths::GameUserDeveloperDir(), OutputDirectory);
+
+		if (!bSuccess || OutputDirectory.IsEmpty())
+#endif // WITH_EDITOR
+		{
+			PCGLog::LogErrorOnGraph(FText::Format(LOCTEXT("InvalidDirectoryError", "Invalid or empty export directory: {0}"), FText::FromString(OutputDirectory)));
+			return true;
+		}
+	}
+
+	if (Settings->FileName.IsEmpty())
+	{
+		PCGLog::LogErrorOnGraph(LOCTEXT("EmptyFileName", "A valid file name is required."));
+		return true;
+	}
+
+	// Create the serialization object pointers to be utilized and cleaned up later.
+	FArchive* Archive = nullptr;
+	TSharedPtr<FJsonObject> RootJsonObject = nullptr;
+
+	// Initialize the serialization object needed.
+	if (Settings->Format == EPCGExportAttributesFormat::Binary)
+	{
+		FString FinalPath = *Helpers::String::MergeFileName(OutputDirectory, Settings->FileName, TEXT(".bin"));
+		Archive = IFileManager::Get().CreateFileWriter(*FinalPath);
+		if (!Archive)
+		{
+			PCGLog::LogErrorOnGraph(FText::Format(LOCTEXT("FileWriterInvalidPath", "Could not create a valid archive writer with path: {0}"), FText::FromString(FinalPath)));
+			return true;
+		}
+
+		// Set up the custom export version
+		using namespace Constants;
+		Archive->SetCustomVersion(FCustomExportVersion::GUID, FCustomExportVersion::LatestVersion, FCustomExportVersion::GetFriendlyName());
+
+		// Set up the user's custom data version
+		if (Settings->bAddCustomDataVersion)
+		{
+			int32 CustomVersion = Settings->CustomVersion;
+			*Archive << CustomVersion;
+		}
+	}
+	else if (Settings->Format == EPCGExportAttributesFormat::Json)
+	{
+		RootJsonObject = MakeShared<FJsonObject>();
+
+		// Set up the Json header to include custom versions
+		TSharedPtr<FJsonObject> JsonHeaderObject = Helpers::Json::CreateJsonHeaderObject(Settings->bAddCustomDataVersion, Settings->CustomVersion);
+		RootJsonObject->SetObjectField(LOCTEXT("JsonHeaderField", "Header").ToString(), std::move(JsonHeaderObject));
+	}
+
+	const UPCGData* InputData = Inputs[0].Data;
+	const UPCGMetadata* Metadata = InputData ? InputData->ConstMetadata() : nullptr;
+	if (!InputData || !Metadata)
+	{
+		PCGLog::InputOutput::LogInvalidInputDataError(Context);
+		return true;
+	}
+
+	TArray<FPCGAttributePropertySelector> AttributeSelectors;
+
+	if (Settings->bExportAllAttributes)
+	{
+		using namespace Helpers;
+
+		/** @todo_pcg: Eventually, it would be useful to be able to iterate on all attributes and properties at once,
+		* to include all/some domains established by the Data. This can be utilized once the logic here is abstracted out
+		* into an API. */
+		if (InputData->IsA<UPCGBasePointData>())
+		{
+			AddAllPointPropertySelectors(AttributeSelectors);
+		}
+		else if (InputData->IsA<UPCGSplineData>())
+		{
+			AddAllSplinePropertySelectors(AttributeSelectors);
+		}
+
+		AddAllAttributeSelectors(InputData, AttributeSelectors);
+	}
+	else
+	{
+		for (const FPCGAttributePropertyInputSelector& Selector : Settings->AttributeSelectors)
+		{
+			AttributeSelectors.Emplace(Selector.CopyAndFixLast(InputData));
+		}
+	}
+
+	// Create and pre-process the attribute source data to be consumed later. Confirmed non-empty earlier.
+	FSourceDataArray SourceData;
+	SourceData.Reserve(Settings->AttributeSelectors.Num());
+
+	for (FPCGAttributePropertySelector& Selector : AttributeSelectors)
+	{
+		FAccessorData Data;
+		Data.Selector = std::move(Selector);
+		Data.Accessor = PCGAttributeAccessorHelpers::CreateConstAccessor(InputData, Data.Selector);
+		Data.Keys = PCGAttributeAccessorHelpers::CreateConstKeys(InputData, Data.Selector);
+		if (Data.Accessor && Data.Keys)
+		{
+			SourceData.Emplace(std::move(Data));
+		}
+		else
+		{
+			PCGLog::Metadata::LogFailToCreateAccessorError(Data.Selector, Context);
+		}
+	}
+
+	if (SourceData.IsEmpty())
+	{
+		return true;
+	}
+
+	if (Settings->Format == EPCGExportAttributesFormat::Binary)
+	{
+		check(Archive);
+
+		// Arrays of values can be consumed by the archive directly and should be consistent. Note: Archive serializes by attribute -only-.
+		for (const FAccessorData& Data : SourceData)
+		{
+			PCGMetadataAttribute::CallbackWithRightType(Data.Accessor->GetUnderlyingType(), [Archive, &Data, Context]<typename T>(const T&)
+			{
+				TArray<T> Values;
+				if (Helpers::GetAttributeValues(Data, Values, EPCGAttributeAccessorFlags::StrictType, Context))
+				{
+					*Archive << Values;
+				}
+				else
+				{
+					PCGLog::Metadata::LogFailToGetAttributeError(Data.Selector, Context);
+				}
+			});
+		}
+	}
+	else if (Settings->Format == EPCGExportAttributesFormat::Json)
+	{
+		check(RootJsonObject);
+
+		TArray<TTuple<FString, TArray<TSharedPtr<FJsonValue>>>> AllSourceJsonValues;
+		AllSourceJsonValues.SetNum(SourceData.Num());
+		// Pre-process the source data into Json values.
+		for (int32 SourceIndex = 0; SourceIndex < SourceData.Num(); ++SourceIndex)
+		{
+			FAccessorData& Data = SourceData[SourceIndex];
+			FString& SourceName = AllSourceJsonValues[SourceIndex].Get<0>();
+			SourceName = Data.Selector.ToString();
+
+			TArray<TSharedPtr<FJsonValue>>& SourceJsonValues = AllSourceJsonValues[SourceIndex].Get<1>();
+
+			PCGMetadataAttribute::CallbackWithRightType(Data.Accessor->GetUnderlyingType(), [&Data, &SourceJsonValues, Context]<typename T>(const T&)
+			{
+				// For floating point values and containers (FVector, FRotator, FTransform, etc), use the Json numerical to maintain precision.
+				if constexpr (PCG::Private::MetadataTraits<T>::IsFloatingPoint)
+				{
+					TArray<T> Values;
+					Helpers::GetAttributeValues(Data, Values, EPCGAttributeAccessorFlags::StrictType, Context);
+
+					SourceJsonValues.Reserve(Values.Num());
+					Algo::Transform(Values, SourceJsonValues, [&Data](const T& Value)
+					{
+						return Helpers::Json::ConvertFloatingPointTypeToJson(Value);
+					});
+				}
+				else // For everything else, automatically broadcast it to string
+				{
+					TArray<FString> Values;
+					Helpers::GetAttributeValues(Data, Values, EPCGAttributeAccessorFlags::AllowBroadcast, Context);
+
+					SourceJsonValues.Reserve(Values.Num());
+					Algo::Transform(Values, SourceJsonValues, [](const FString& Value)
+					{
+						return MakeShared<FJsonValueString>(Value);
+					});
+				}
+			});
+		}
+
+		// Iterate over the pre-processed source data by element index, to get the attribute values together "by element".
+		if (Settings->Layout == EPCGExportAttributesLayout::ByElement)
+		{
+			check(!SourceData.IsEmpty());
+
+			const int32 ElementNum = SourceData[0].Keys->GetNum();
+			for (int32 ElementIndex = 0; ElementIndex < ElementNum; ++ElementIndex)
+			{
+				TSharedPtr<FJsonObject> ElementJsonObject = MakeShared<FJsonObject>();
+
+				for (int32 SourceIndex = 0; SourceIndex < SourceData.Num(); ++SourceIndex)
+				{
+					check(AllSourceJsonValues.IsValidIndex(SourceIndex));
+
+					TArray<TSharedPtr<FJsonValue>>& SourceValues = AllSourceJsonValues[SourceIndex].Get<1>();
+					const FString& SourceName = AllSourceJsonValues[SourceIndex].Get<0>();
+
+					check(SourceValues.IsValidIndex(ElementIndex));
+					// Currently only support the following Json types:
+					check(SourceValues[ElementIndex]->Type == EJson::Array
+						|| SourceValues[ElementIndex]->Type == EJson::Number
+						|| SourceValues[ElementIndex]->Type == EJson::String
+						|| SourceValues[ElementIndex]->Type == EJson::Object);
+
+					ElementJsonObject->SetField(SourceName, std::move(SourceValues[ElementIndex]));
+				}
+
+				RootJsonObject->SetObjectField(
+					FText::Format(LOCTEXT("JsonElementPrefixFormat", "Element [{0}]"), FText::AsNumber(ElementIndex)).ToString(),
+					std::move(ElementJsonObject));
+			}
+		}
+		else if (Settings->Layout == EPCGExportAttributesLayout::ByAttribute)
+		{
+			for (int32 SourceIndex = 0; SourceIndex < SourceData.Num(); ++SourceIndex)
+			{
+				check(AllSourceJsonValues.IsValidIndex(SourceIndex));
+
+				TArray<TSharedPtr<FJsonValue>>& SourceValues = AllSourceJsonValues[SourceIndex].Get<1>();
+				const FString& SourceName = AllSourceJsonValues[SourceIndex].Get<0>();
+
+				RootJsonObject->SetArrayField(SourceName, std::move(SourceValues));
+			}
+		}
+		else
+		{
+			checkNoEntry();
+			return true;
+		}
+	}
+
+	// Clean up the serialization objects that exist.
+	if (Archive)
+	{
+		Archive->Close();
+	}
+	else if (RootJsonObject) // Conclude the Json serialization.
+	{
+		FString JsonString;
+		// @todo_pcg: Allow the user to set the final object name
+		const TSharedRef<TJsonWriter<>> JsonWriter = TJsonWriterFactory<>::Create(&JsonString);
+		if (!FJsonSerializer::Serialize(RootJsonObject.ToSharedRef(), JsonWriter))
+		{
+			PCGLog::LogErrorOnGraph(LOCTEXT("JsonSerializationError", "Serialization of Json data failed."));
+			return true;
+		}
+
+		const FString FinalPath = Helpers::String::MergeFileName(OutputDirectory, Settings->FileName, TEXT(".json"));
+		if (!FFileHelper::SaveStringToFile(JsonString, *FinalPath))
+		{
+			PCGLog::LogErrorOnGraph(FText::Format(LOCTEXT("JsonInvalidPath", "Could not create a valid Json file with path: {0}"), FText::FromString(FinalPath)));
+			return true;
+		}
+	}
+	else
+	{
+		// There was no export data. Was a new format created?
+		checkNoEntry();
+	}
+
+#else  // Not Windows, Linux, or Mac
+	UE_LOG(LogPCG, Verbose, TEXT("Running 'Export Selected Attributes' node, which has been disabled on this platform."));
+#endif // PLATFORM_WINDOWS || PLATFORM_LINUX || PLATFORM_MAC
+
+	return true;
+}
+
+#undef LOCTEXT_NAMESPACE
diff --git a/Engine/Plugins/PCG/Source/PCG/Private/Elements/Metadata/PCGMetadataCompareOpElement.cpp b/Engine/Plugins/PCG/Source/PCG/Private/Elements/Metadata/PCGMetadataCompareOpElement.cpp
--- a/Engine/Plugins/PCG/Source/PCG/Private/Elements/Metadata/PCGMetadataCompareOpElement.cpp
+++ b/Engine/Plugins/PCG/Source/PCG/Private/Elements/Metadata/PCGMetadataCompareOpElement.cpp
@@ -14,7 +14,7 @@ namespace PCGMetadataCompareSettings
 	{
 		if (Operation == EPCGMetadataCompareOperation::Equal)
 		{
-			if constexpr (PCG::Private::MetadataTraits<T>::NeedsCompareTolerance)
+			if constexpr (PCG::Private::MetadataTraits<T>::IsFloatingPoint)
 			{
 				return PCG::Private::MetadataTraits<T>::Equal(Input1, Input2, Tolerance);
 			}
@@ -25,7 +25,7 @@ namespace PCGMetadataCompareSettings
 		}
 		else if (Operation == EPCGMetadataCompareOperation::NotEqual)
 		{
-			if constexpr (PCG::Private::MetadataTraits<T>::NeedsCompareTolerance)
+			if constexpr (PCG::Private::MetadataTraits<T>::IsFloatingPoint)
 			{
 				return !PCG::Private::MetadataTraits<T>::Equal(Input1, Input2, Tolerance);
 			}
diff --git a/Engine/Plugins/PCG/Source/PCG/Public/Elements/IO/PCGExportSelectedAttributes.h b/Engine/Plugins/PCG/Source/PCG/Public/Elements/IO/PCGExportSelectedAttributes.h
new file mode 100644
--- /dev/null
+++ b/Engine/Plugins/PCG/Source/PCG/Public/Elements/IO/PCGExportSelectedAttributes.h
@@ -0,0 +1,99 @@
+// Copyright Epic Games, Inc. All Rights Reserved.
+
+#pragma once
+
+#include "PCGContext.h"
+#include "PCGSettings.h"
+
+#include "Serialization/CustomVersion.h"
+
+#include "PCGExportSelectedAttributes.generated.h"
+
+UENUM(BlueprintType)
+enum class EPCGExportAttributesFormat : uint8
+{
+	Binary UMETA(Tooltip = "Export binary data to file using the Unreal archive system."),
+	Json UMETA(Tooltip = "Output the asset values into a JSON format.")
+};
+
+UENUM(BlueprintType)
+enum class EPCGExportAttributesLayout : uint8
+{
+	ByElement UMETA(Tooltip = "Use the elements as the main data object. Each element contains all its attributes."),
+	ByAttribute UMETA(Tooltip = "Use the attributes as the main data object. Each attribute grouping will contain the element values in sequential order.")
+};
+
+/** @todo_pcg:
+ * - Provide a tooltip or other feedback to dynamically generate the expected format to make it easier for the user to consume.
+ * - Abstract this out to an API and have the node leverage the API.
+ * - Each Data Type can be in charge of the default "selection" of properties/attributes.
+ */
+
+/**
+ * Exports the selected attributes directly to file in a specified format.
+ * Note: This node is only operational on traditional development platforms (Windows, Linux, Mac) where the Editor can be used.
+ */
+UCLASS(BlueprintType, MinimalAPI, ClassGroup = (Procedural))
+class UPCGExportSelectedAttributesSettings : public UPCGSettings
+{
+	GENERATED_BODY()
+
+public:
+	UPCGExportSelectedAttributesSettings();
+
+	// ~Begin UPCGSettings interface
+#if WITH_EDITOR
+	virtual FName GetDefaultNodeName() const override { return FName(TEXT("PCGExportSelectedAttributes")); }
+	virtual FText GetDefaultNodeTitle() const override { return NSLOCTEXT("PCGExportSelectedAttributesElement", "NodeTitle", "Export Selected Attributes"); }
+	virtual FText GetNodeTooltipText() const override { return NSLOCTEXT("PCGExportSelectedAttributesElement", "NodeTooltip", "Exports selected attributes directly to a file in a specified format."); }
+	virtual EPCGSettingsType GetType() const override { return EPCGSettingsType::InputOutput; }
+#endif // WITH_EDITOR
+
+	virtual TArray<FPCGPinProperties> InputPinProperties() const override;
+	virtual TArray<FPCGPinProperties> OutputPinProperties() const override;
+
+protected:
+	virtual FPCGElementPtr CreateElement() const override;
+	// ~End UPCGSettings interface
+
+public:
+	/** Data will be exported to a local file in this format. */
+	UPROPERTY(EditAnywhere, BlueprintReadWrite, Category = Settings, meta = (PCG_Overridable))
+	EPCGExportAttributesFormat Format = EPCGExportAttributesFormat::Binary;
+
+	/** Determines how the data will be laid out in the export. */
+	UPROPERTY(EditAnywhere, BlueprintReadWrite, Category = Settings, meta = (EditCondition = "Format != EPCGExportAttributesFormat::Binary", PCG_Overridable))
+	EPCGExportAttributesLayout Layout = EPCGExportAttributesLayout::ByElement;
+
+	/** The directory to save the data within. If none is selected a dialog will open by default. */
+	UPROPERTY(EditAnywhere, BlueprintReadWrite, Category = Settings, meta = (PCG_Overridable))
+	FDirectoryPath Path;
+
+	/** The file name (without extension) to export the data. */
+	UPROPERTY(EditAnywhere, BlueprintReadWrite, Category = Settings, meta = (PCG_Overridable))
+	FString FileName;
+
+	UPROPERTY(EditAnywhere, BlueprintReadWrite, Category = Settings, meta = (PCG_Overridable))
+	bool bExportAllAttributes = true;
+
+	/** The attributes to use as sources for the data export. Only those selected will be exported from the input data. */
+	UPROPERTY(EditAnywhere, BlueprintReadWrite, Category = Settings, meta = (EditCondition = "!bExportAllAttributes", EditConditionHides, PCG_Overridable))
+	TArray<FPCGAttributePropertyInputSelector> AttributeSelectors;
+
+	UPROPERTY(EditAnywhere, BlueprintReadWrite, Category = Advanced)
+	bool bAddCustomDataVersion = false;
+
+	UPROPERTY(EditAnywhere, BlueprintReadWrite, Category = Advanced, meta = (EditCondition = "bAddCustomDataVersion", EditConditionHides))
+	int32 CustomVersion = 0;
+};
+
+class FPCGExportSelectedAttributesElement final : public IPCGElement
+{
+public:
+	virtual bool CanExecuteOnlyOnMainThread(FPCGContext* Context) const override { return true; }
+	// @todo_pcg: Crc the file bytes and version
+	virtual bool IsCacheable(const UPCGSettings* InSettings) const override { return false; }
+
+protected:
+	virtual bool ExecuteInternal(FPCGContext* Context) const override;
+};
\ No newline at end of file
diff --git a/Engine/Plugins/PCG/Source/PCG/Public/Metadata/PCGMetadataAttributeTraits.h b/Engine/Plugins/PCG/Source/PCG/Public/Metadata/PCGMetadataAttributeTraits.h
--- a/Engine/Plugins/PCG/Source/PCG/Public/Metadata/PCGMetadataAttributeTraits.h
+++ b/Engine/Plugins/PCG/Source/PCG/Public/Metadata/PCGMetadataAttributeTraits.h
@@ -331,13 +331,13 @@ namespace PCG
 		* - CanSearchString: Can do string search operations like matching substrings.
 		* - CompressData: Uses value keys, so if two values are identical, we only store one.
 		* - NeedsConstruction: Whether type requires constructor for valid state (i.e. can't simply be zero initialized).
-		* - NeedsCompareTolerance: Whether the type itself or the container's value type is floating point and needs to account for tolerance.
+		* - IsFloatingPoint: Whether the type itself or the container's value type is floating point.
 		*/
 
 		template<typename T>
 		struct DefaultStringTraits
 		{
-			enum { NeedsCompareTolerance = false };
+			enum { IsFloatingPoint = false };
 
 			static FString ToString(const T& A)
 			{
@@ -369,7 +369,7 @@ namespace PCG
 		{
 			enum { CanSubAdd = true };
 			enum { CanMulDiv = true };
-			enum { NeedsCompareTolerance = false };
+			enum { IsFloatingPoint = false };
 
 			static bool Equal(const T& A, const T& B)
 			{
@@ -556,7 +556,7 @@ namespace PCG
 			enum { CompressData = false };
 			enum { CanSearchString = false };
 			enum { NeedsConstruction = false };
-			enum { NeedsCompareTolerance = true };
+			enum { IsFloatingPoint = true };
 
 			static bool Equal(const float& A, const float& B, const double Tolerance = UE_SMALL_NUMBER)
 			{
@@ -577,7 +577,7 @@ namespace PCG
 			enum { CompressData = false };
 			enum { CanSearchString = false };
 			enum { NeedsConstruction = false };
-			enum { NeedsCompareTolerance = true };
+			enum { IsFloatingPoint = true };
 
 			static bool Equal(const double& A, const double& B, const double Tolerance = UE_DOUBLE_SMALL_NUMBER)
 			{
@@ -605,7 +605,7 @@ namespace PCG
 			enum { CanComputeDistance = false };
 			enum { CanSearchString = false };
 			enum { NeedsConstruction = false };
-			enum { NeedsCompareTolerance = false };
+			enum { IsFloatingPoint = false };
 
 			static bool ZeroValue()
 			{
@@ -676,7 +676,7 @@ namespace PCG
 			enum { NeedsConstruction = false };
 			enum { CanFindNearest = true };
 			enum { CanComputeDistance = true };
-			enum { NeedsCompareTolerance = true };
+			enum { IsFloatingPoint = true };
 
 			using DistanceType = typename T::FReal;
 
@@ -860,7 +860,7 @@ namespace PCG
 			enum { CanCompare = false };
 			enum { CanSearchString = false };
 			enum { NeedsConstruction = false };
-			enum { NeedsCompareTolerance = true };
+			enum { IsFloatingPoint = true };
 
 			using DistanceType = typename FQuat::FReal;
 
@@ -936,7 +936,7 @@ namespace PCG
 			enum { CanCompare = false };
 			enum { CanSearchString = false };
 			enum { NeedsConstruction = false };
-			enum { NeedsCompareTolerance = true };
+			enum { IsFloatingPoint = true };
 
 			using DistanceType = typename FRotator::FReal;
 
@@ -1014,7 +1014,7 @@ namespace PCG
 			enum { CanComputeDistance = false };
 			enum { CanSearchString = false };
 			enum { NeedsConstruction = false };
-			enum { NeedsCompareTolerance = true };
+			enum { IsFloatingPoint = true };
 
 			static bool Equal(const FTransform& A, const FTransform& B, const FTransform::FReal Tolerance = UE_DOUBLE_SMALL_NUMBER)
 			{
