[PCG] Add Get Graph Parameter node from drag and drop

diff --git a/Engine/Plugins/PCG/Source/PCG/Private/PCGPin.cpp b/Engine/Plugins/PCG/Source/PCG/Private/PCGPin.cpp
--- a/Engine/Plugins/PCG/Source/PCG/Private/PCGPin.cpp
+++ b/Engine/Plugins/PCG/Source/PCG/Private/PCGPin.cpp
@@ -33,6 +33,39 @@ namespace PCGPin
 
 		return true;
 	}
+
+	bool ArePinTypesCompatible(const EPCGDataType UpstreamTypes, const EPCGDataType DownstreamTypes)
+	{
+		if (UpstreamTypes == EPCGDataType::None || DownstreamTypes == EPCGDataType::None)
+		{
+			return false;
+		}
+
+		// Concrete can always be used as a composite - allow connections from concrete to composite
+		const bool bUpstreamInConcrete = !(UpstreamTypes & ~EPCGDataType::Concrete);
+		const bool bDownstreamIsSpatial = (DownstreamTypes == EPCGDataType::Spatial);
+		if (bUpstreamInConcrete && bDownstreamIsSpatial)
+		{
+			return true;
+		}
+
+		// Anything spatial can collapse to point as long as the downstream pin does not support any other spatial types. For example, Point | Param is collapsible, but Point | Surface is not
+		const bool bUpstreamIsSpatial = !(UpstreamTypes & ~EPCGDataType::Spatial);
+		const bool bDownstreamIsPoint = (DownstreamTypes & EPCGDataType::Spatial) == EPCGDataType::Point;
+		if (bUpstreamIsSpatial && bDownstreamIsPoint)
+		{
+			return true;
+		}
+
+		// Spline -> Surface conversion
+		if (UpstreamTypes == EPCGDataType::Spline && DownstreamTypes == EPCGDataType::Surface)
+		{
+			return true;
+		}
+
+		// Otherwise allow if there is overlap. Don't detect wide -> narrow issues - conversion nodes deal with that
+		return !!(UpstreamTypes & DownstreamTypes);
+	}
 }
 
 #if WITH_EDITOR
@@ -385,36 +418,13 @@ bool UPCGPin::IsCompatible(const UPCGPin* OtherPin) const
 	}
 	check(UpstreamPin && DownstreamPin);
 
-	// Types missing
-	if (UpstreamTypes == EPCGDataType::None || DownstreamTypes == EPCGDataType::None)
-	{
-		return false;
-	}
-
-	// Concrete can always be used as a composite - allow connections from concrete to composite
-	const bool bUpstreamInConcrete = !(UpstreamTypes & ~EPCGDataType::Concrete);
-	const bool bDownstreamIsSpatial = (DownstreamTypes == EPCGDataType::Spatial);
-	if (bUpstreamInConcrete && bDownstreamIsSpatial)
-	{
-		return true;
-	}
-
-	// Anything spatial can collapse to point as long as the downstream pin does not support any other spatial types. For example, Point | Param is collapsible, but Point | Surface is not
-	const bool bUpstreamIsSpatial = !(UpstreamTypes & ~EPCGDataType::Spatial);
-	const bool bDownstreamIsPoint = (DownstreamTypes & EPCGDataType::Spatial) == EPCGDataType::Point;
-	if (bUpstreamIsSpatial && bDownstreamIsPoint)
-	{
-		return true;
-	}
-
-	// Spline -> Surface conversion
-	if (UpstreamTypes == EPCGDataType::Spline && DownstreamTypes == EPCGDataType::Surface)
-	{
-		return true;
-	}
+	return PCGPin::ArePinTypesCompatible(UpstreamTypes, DownstreamTypes);
+}
 
-	// Otherwise allow if there is overlap. Don't detect wide -> narrow issues - conversion nodes deal with that
-	return !!(UpstreamTypes & DownstreamTypes);
+bool UPCGPin::IsDownstreamCompatible(const EPCGDataType UpstreamTypes) const
+{
+	const EPCGDataType AllowedTypes = Properties.AllowedTypes;
+	return PCGPin::ArePinTypesCompatible(UpstreamTypes, AllowedTypes);
 }
 
 bool UPCGPin::AllowsMultipleConnections() const
diff --git a/Engine/Plugins/PCG/Source/PCG/Public/PCGPin.h b/Engine/Plugins/PCG/Source/PCG/Public/PCGPin.h
--- a/Engine/Plugins/PCG/Source/PCG/Public/PCGPin.h
+++ b/Engine/Plugins/PCG/Source/PCG/Public/PCGPin.h
@@ -266,6 +266,8 @@ public:
 	bool AllowsMultipleConnections() const;
 	bool AllowsMultipleData() const;
 	bool IsCompatible(const UPCGPin* OtherPin) const;
+	/** For checking when the new node doesn't exist yet. */
+	bool IsDownstreamCompatible(EPCGDataType UpstreamTypes) const;
 	bool CanConnect(const UPCGPin* OtherPin) const;
 
 	/** Adds edge to pin if edge does not exist. Optionally returns list of affected nodes (including the pin owner). */
diff --git a/Engine/Plugins/PCG/Source/PCGEditor/Private/PCGEditorGraph.cpp b/Engine/Plugins/PCG/Source/PCGEditor/Private/PCGEditorGraph.cpp
--- a/Engine/Plugins/PCG/Source/PCGEditor/Private/PCGEditorGraph.cpp
+++ b/Engine/Plugins/PCG/Source/PCGEditor/Private/PCGEditorGraph.cpp
@@ -3,6 +3,7 @@
 #include "PCGEditorGraph.h"
 
 #include "PCGEdge.h"
+#include "PCGEditorGraphSchemaActions.h"
 #include "PCGEditorModule.h"
 #include "PCGGraph.h"
 #include "PCGPin.h"
@@ -42,6 +43,14 @@ namespace PCGEditorGraphUtils
 			}
 		}
 	}
+
+	UPCGPin* GetPCGPinFromEdGraphPin(const UEdGraphPin* Pin)
+	{
+		UEdGraphNode* GraphNode = Pin ? Pin->GetOwningNodeUnchecked() : nullptr;
+		UPCGEditorGraphNodeBase* PCGGraphNode = GraphNode ? CastChecked<UPCGEditorGraphNodeBase>(GraphNode) : nullptr;
+		UPCGNode* PCGNode = PCGGraphNode ? PCGGraphNode->GetPCGNode() : nullptr;
+		return PCGNode ? (Pin->Direction == EGPD_Input ? PCGNode->GetInputPin(Pin->PinName) : PCGNode->GetOutputPin(Pin->PinName)) : nullptr;
+	}
 }
 
 void UPCGEditorGraph::InitFromNodeGraph(UPCGGraph* InPCGGraph)
@@ -365,3 +374,53 @@ void UPCGEditorGraph::OnGraphUserParametersChanged(UPCGGraphInterface* InGraph,
 		NodeToRemove->DestroyNode();
 	}
 }
+
+bool UPCGEditorGraph::CanReceivePropertyBagDetailsDropOnGraphPin(const UEdGraphPin* Pin) const
+{
+	const UPCGPin* PCGPin = PCGEditorGraphUtils::GetPCGPinFromEdGraphPin(Pin);
+	return PCGPin
+	&& !PCGPin->IsOutputPin()
+	&& (PCGPin->EdgeCount() < 1 || PCGPin->AllowsMultipleConnections())
+	&& PCGPin->IsDownstreamCompatible(EPCGDataType::Param);
+}
+
+bool UPCGEditorGraph::CanReceivePropertyBagDetailsDropOnGraphNode(const UEdGraphNode* Node) const
+{
+	// Currently no useful way to interpret dropping a user parameter on a node.
+	return false;
+}
+
+bool UPCGEditorGraph::CanReceivePropertyBagDetailsDropOnGraph(const UEdGraph* Graph) const
+{
+	// Anywhere on the graph panel should be fine for creating a new get user parameter node.
+	return true;
+}
+
+FReply UPCGEditorGraph::OnPropertyBagDetailsDropOnGraphPin(const FPropertyBagPropertyDesc& PropertyDesc, UEdGraphPin* Pin, const FVector2D GraphPosition) const
+{
+	const UEdGraphNode* Node = Pin ? Pin->GetOwningNode() : nullptr;
+	if (Node && Node->GetGraph() && PropertyDesc.ID.IsValid() && PropertyDesc.CachedProperty)
+	{
+		FPCGEditorGraphSchemaAction_NewGetParameterElement Action;
+		Action.SettingsClass = UPCGUserParameterGetSettings::StaticClass();
+		Action.PropertyGuid = PropertyDesc.ID;
+		Action.PropertyName = PropertyDesc.Name;
+		Action.PerformAction(Node->GetGraph(), Pin, GraphPosition, /*bSelectNewNode=*/true);
+	}
+
+	return FReply::Handled();
+}
+
+FReply UPCGEditorGraph::OnPropertyBagDetailsDropOnGraph(const FPropertyBagPropertyDesc& PropertyDesc, UEdGraph* Graph, const FVector2D GraphPosition) const
+{
+	if (Graph && PropertyDesc.ID.IsValid() && PropertyDesc.CachedProperty)
+	{
+		FPCGEditorGraphSchemaAction_NewGetParameterElement Action;
+		Action.SettingsClass = UPCGUserParameterGetSettings::StaticClass();
+		Action.PropertyGuid = PropertyDesc.ID;
+		Action.PropertyName = PropertyDesc.Name;
+		Action.PerformAction(Graph, nullptr, GraphPosition, /*bSelectNewNode=*/true);
+	}
+
+	return FReply::Handled();
+}
diff --git a/Engine/Plugins/PCG/Source/PCGEditor/Private/PCGEditorGraph.h b/Engine/Plugins/PCG/Source/PCGEditor/Private/PCGEditorGraph.h
--- a/Engine/Plugins/PCG/Source/PCGEditor/Private/PCGEditorGraph.h
+++ b/Engine/Plugins/PCG/Source/PCGEditor/Private/PCGEditorGraph.h
@@ -2,6 +2,7 @@
 
 #pragma once
 
+#include "IPropertyBagEditorGraph.h"
 #include "EdGraph/EdGraph.h"
 
 #include "PCGGraph.h"
@@ -20,7 +21,7 @@ namespace PCGEditorGraphUtils
 }
 
 UCLASS()
-class UPCGEditorGraph : public UEdGraph
+class UPCGEditorGraph : public UEdGraph, public IPropertyBagEdGraph
 {
 	GENERATED_BODY()
 
@@ -60,6 +61,14 @@ protected:
 
 	void OnGraphUserParametersChanged(UPCGGraphInterface* InGraph, EPCGGraphParameterEvent ChangeType, FName ChangedPropertyName);
 
+	// ~Begin IPropertyBagEdGraph interface
+	virtual bool CanReceivePropertyBagDetailsDropOnGraphPin(const UEdGraphPin* Pin) const override;
+	virtual bool CanReceivePropertyBagDetailsDropOnGraphNode(const UEdGraphNode* Node) const override;
+	virtual bool CanReceivePropertyBagDetailsDropOnGraph(const UEdGraph* Graph) const override;
+	virtual FReply OnPropertyBagDetailsDropOnGraphPin(const FPropertyBagPropertyDesc& PropertyDesc, UEdGraphPin* Pin, FVector2D GraphPosition) const override;
+	virtual FReply OnPropertyBagDetailsDropOnGraph(const FPropertyBagPropertyDesc& PropertyDesc, UEdGraph* Graph, FVector2D GraphPosition) const override;
+	// ~End IPropertyBagEdGraph interface
+
 private:
 	UPROPERTY()
 	TObjectPtr<UPCGGraph> PCGGraph = nullptr;
