[PCG] Enable boolean math operations and Set node with other types

+ Enable Set for types FRotator, FQuat, FTransform, FSoftObjectPtr, FSoftClassPtr
+ Bool inline constant now available for Set node (boolean)
+ Permits pure boolean (input and output are all boolean) operations

diff --git a/Engine/Plugins/PCG/Source/PCG/Private/Elements/Metadata/PCGMetadataMathsOpElement.cpp b/Engine/Plugins/PCG/Source/PCG/Private/Elements/Metadata/PCGMetadataMathsOpElement.cpp
--- a/Engine/Plugins/PCG/Source/PCG/Private/Elements/Metadata/PCGMetadataMathsOpElement.cpp
+++ b/Engine/Plugins/PCG/Source/PCG/Private/Elements/Metadata/PCGMetadataMathsOpElement.cpp
@@ -146,8 +146,6 @@ namespace PCGMetadataMathsSettings
 			return PCGMetadataMaths::Pow(Value1, Value2);
 		case EPCGMetadataMathsOperation::Modulo:
 			return PCGMetadataMaths::Modulo(Value1, Value2);
-		case EPCGMetadataMathsOperation::Set:
-			return Value2;
 		default:
 			return T{};
 		}
@@ -168,6 +166,25 @@ namespace PCGMetadataMathsSettings
 			return T{};
 		}
 	}
+
+	// Specialize bool to int32, as some math operations won't compile with boolean values.
+	template <>
+	bool UnaryOp(const bool& Value, EPCGMetadataMathsOperation Op)
+	{
+		return static_cast<bool>(UnaryOp(static_cast<int32>(Value), Op));
+	}
+
+	template <>
+	bool BinaryOp(const bool& Value1, const bool& Value2, EPCGMetadataMathsOperation Op)
+	{
+		return static_cast<bool>(BinaryOp(static_cast<int32>(Value1), static_cast<int32>(Value2), Op));
+	}
+
+	template <>
+	bool TernaryOp(const bool& Value1, const bool& Value2, const bool& Value3, EPCGMetadataMathsOperation Op)
+	{
+		return static_cast<bool>(TernaryOp(static_cast<int32>(Value1), static_cast<int32>(Value2), static_cast<int32>(Value3), Op));
+	}
 }
 
 void UPCGMetadataMathsSettings::PostLoad()
@@ -234,7 +251,15 @@ uint32 UPCGMetadataMathsSettings::GetOperandNum() const
 bool UPCGMetadataMathsSettings::IsSupportedInputType(uint16 TypeId, uint32 InputIndex, bool& bHasSpecialRequirement) const
 {
 	bHasSpecialRequirement = false;
-	return PCG::Private::IsOfTypes<float, double, int32, int64, FVector2D, FVector, FVector4>(TypeId);
+
+	if (Operation == EPCGMetadataMathsOperation::Set)
+	{
+		return PCG::Private::IsPCGType(TypeId);
+	}
+	else
+	{
+		return PCG::Private::IsOfTypes<bool, float, double, int32, int64, FVector2D, FVector, FVector4>(TypeId);
+	}
 }
 
 bool UPCGMetadataMathsSettings::ShouldForceOutputToInt(uint16 InputTypeId) const
@@ -288,6 +313,11 @@ FPCGAttributePropertyInputSelector UPCGMetadataMathsSettings::GetInputSource(uin
 	}
 }
 
+bool UPCGMetadataMathsSettings::IsPinDefaultValueMetadataTypeValid(const FName PinLabel, const EPCGMetadataTypes DataType) const
+{
+	return (Operation == EPCGMetadataMathsOperation::Set) || Super::IsPinDefaultValueMetadataTypeValid(PinLabel, DataType);
+}
+
 FString UPCGMetadataMathsSettings::GetAdditionalTitleInformation() const
 {
 	if (const UEnum* EnumPtr = StaticEnum<EPCGMetadataMathsOperation>())
@@ -341,13 +371,17 @@ bool FPCGMetadataMathsElement::DoOperation(PCGMetadataOps::FOperationData& Opera
 
 	auto MathFunc = [this, Operation = Settings->Operation, &OperationData]<typename AttributeType>(AttributeType) -> bool
 	{
-		// Need to remove types that would not compile
-		if constexpr (!PCG::Private::IsOfTypes<AttributeType, float, double, int32, int64, FVector2D, FVector, FVector4>())
+		if constexpr (PCG::Private::IsPCGType<AttributeType>()) // Types that wouldn't compile with operations. Only works with Set operation.
 		{
-			ensure(false);
-			return true;
+			if (Operation == EPCGMetadataMathsOperation::Set)
+			{
+				TRACE_CPUPROFILER_EVENT_SCOPE(FPCGMetadataMathsElement::ExecuteInternal::BinaryOp);
+				DoBinaryOp<AttributeType, AttributeType>(OperationData, [](const AttributeType&, const AttributeType& Value2) -> AttributeType { return Value2; });
+				return true;
+			}
 		}
-		else
+
+		if constexpr (PCG::Private::IsOfTypes<AttributeType, bool, float, double, int32, int64, FVector2D, FVector, FVector4>())
 		{
 			if (PCGMetadataMathsSettings::IsUnaryOp(Operation))
 			{
@@ -380,6 +414,11 @@ bool FPCGMetadataMathsElement::DoOperation(PCGMetadataOps::FOperationData& Opera
 				return true;
 			}
 		}
+		else // Some other type not supported
+		{
+			ensure(false);
+			return true;
+		}
 	};
 
 	// If the output is double, force all to double.
diff --git a/Engine/Plugins/PCG/Source/PCG/Public/Elements/Metadata/PCGMetadataMathsOpElement.h b/Engine/Plugins/PCG/Source/PCG/Public/Elements/Metadata/PCGMetadataMathsOpElement.h
--- a/Engine/Plugins/PCG/Source/PCG/Public/Elements/Metadata/PCGMetadataMathsOpElement.h
+++ b/Engine/Plugins/PCG/Source/PCG/Public/Elements/Metadata/PCGMetadataMathsOpElement.h
@@ -80,6 +80,7 @@ public:
 
 	//~Begin IPCGSettingsDefaultValueProvider interface
 	virtual EPCGMetadataTypes GetPinInitialDefaultValueType(FName PinLabel) const override { return EPCGMetadataTypes::Double; }
+	virtual bool IsPinDefaultValueMetadataTypeValid(FName PinLabel, EPCGMetadataTypes DataType) const override;
 	//~End IPCGSettingsDefaultValueProvider interface
 
 protected:
diff --git a/Engine/Plugins/PCG/Source/PCG/Public/Metadata/PCGMetadataAttributeTraits.h b/Engine/Plugins/PCG/Source/PCG/Public/Metadata/PCGMetadataAttributeTraits.h
--- a/Engine/Plugins/PCG/Source/PCG/Public/Metadata/PCGMetadataAttributeTraits.h
+++ b/Engine/Plugins/PCG/Source/PCG/Public/Metadata/PCGMetadataAttributeTraits.h
@@ -167,6 +167,10 @@ namespace PCG
 #define PCGMetadataBroadcastable(FirstType, SecondType) Values[(uint8)EPCGMetadataTypes::FirstType][(uint8)EPCGMetadataTypes::SecondType] = true
 
 				// Finally set all cases where it is broadcastable
+				PCGMetadataBroadcastable(Boolean, Integer32);
+				PCGMetadataBroadcastable(Boolean, Integer64);
+				PCGMetadataBroadcastable(Boolean, Double);
+
 				PCGMetadataBroadcastable(Float, Double);
 				PCGMetadataBroadcastable(Float, Vector2);
 				PCGMetadataBroadcastable(Float, Vector);
