[PCG] Fix missing undef LOCTEXT_NAMESPACE

#rnx

diff --git a/Engine/Plugins/PCG/Source/PCGEditor/Private/DataVisualizations/PCGSplineDataVisualization.cpp b/Engine/Plugins/PCG/Source/PCGEditor/Private/DataVisualizations/PCGSplineDataVisualization.cpp
--- a/Engine/Plugins/PCG/Source/PCGEditor/Private/DataVisualizations/PCGSplineDataVisualization.cpp
+++ b/Engine/Plugins/PCG/Source/PCGEditor/Private/DataVisualizations/PCGSplineDataVisualization.cpp
@@ -136,3 +136,5 @@ const UPCGBasePointData* IPCGSplineDataVisualization::CollapseToDebugBasePointDa
 
 	return nullptr;
 }
+
+#undef LOCTEXT_NAMESPACE
