CIS fix

#rnx

diff --git a/Engine/Plugins/PCG/Source/PCGEditor/Private/PCGEditorGraph.cpp b/Engine/Plugins/PCG/Source/PCGEditor/Private/PCGEditorGraph.cpp
--- a/Engine/Plugins/PCG/Source/PCGEditor/Private/PCGEditorGraph.cpp
+++ b/Engine/Plugins/PCG/Source/PCGEditor/Private/PCGEditorGraph.cpp
@@ -14,6 +14,7 @@
 #include "Nodes/PCGEditorGraphNodeInput.h"
 #include "Nodes/PCGEditorGraphNodeOutput.h"
 #include "Nodes/PCGEditorGraphNodeReroute.h"
+#include "Types/SlateVector2.h"
 
 #include "EdGraph/EdGraphPin.h"
 
@@ -405,7 +406,7 @@ FReply UPCGEditorGraph::OnPropertyBagDetailsDropOnGraphPin(const FPropertyBagPro
 		Action.SettingsClass = UPCGUserParameterGetSettings::StaticClass();
 		Action.PropertyGuid = PropertyDesc.ID;
 		Action.PropertyName = PropertyDesc.Name;
-		Action.PerformAction(Node->GetGraph(), Pin, GraphPosition, /*bSelectNewNode=*/true);
+		Action.PerformAction(Node->GetGraph(), Pin, UE::Slate::FDeprecateVector2DParameter(GraphPosition), /*bSelectNewNode=*/true);
 	}
 
 	return FReply::Handled();
@@ -419,7 +420,7 @@ FReply UPCGEditorGraph::OnPropertyBagDetailsDropOnGraph(const FPropertyBagProper
 		Action.SettingsClass = UPCGUserParameterGetSettings::StaticClass();
 		Action.PropertyGuid = PropertyDesc.ID;
 		Action.PropertyName = PropertyDesc.Name;
-		Action.PerformAction(Graph, nullptr, GraphPosition, /*bSelectNewNode=*/true);
+		Action.PerformAction(Graph, nullptr, UE::Slate::FDeprecateVector2DParameter(GraphPosition), /*bSelectNewNode=*/true);
 	}
 
 	return FReply::Handled();
