[PCG] Fix build health issues related to Property Bag Details Drag Drop

+ Some slight additional code cleanup

#rnx

diff --git a/Engine/Plugins/PCG/Source/PCG/Private/PCGPin.cpp b/Engine/Plugins/PCG/Source/PCG/Private/PCGPin.cpp
--- a/Engine/Plugins/PCG/Source/PCG/Private/PCGPin.cpp
+++ b/Engine/Plugins/PCG/Source/PCG/Private/PCGPin.cpp
@@ -421,7 +421,7 @@ bool UPCGPin::IsCompatible(const UPCGPin* OtherPin) const
 	return PCGPin::ArePinTypesCompatible(UpstreamTypes, DownstreamTypes);
 }
 
-bool UPCGPin::IsDownstreamCompatible(const EPCGDataType UpstreamTypes) const
+bool UPCGPin::IsDownstreamPinTypeCompatible(const EPCGDataType UpstreamTypes) const
 {
 	const EPCGDataType AllowedTypes = Properties.AllowedTypes;
 	return PCGPin::ArePinTypesCompatible(UpstreamTypes, AllowedTypes);
diff --git a/Engine/Plugins/PCG/Source/PCG/Public/PCGPin.h b/Engine/Plugins/PCG/Source/PCG/Public/PCGPin.h
--- a/Engine/Plugins/PCG/Source/PCG/Public/PCGPin.h
+++ b/Engine/Plugins/PCG/Source/PCG/Public/PCGPin.h
@@ -267,7 +267,7 @@ public:
 	bool AllowsMultipleData() const;
 	bool IsCompatible(const UPCGPin* OtherPin) const;
 	/** For checking when the new node doesn't exist yet. */
-	bool IsDownstreamCompatible(EPCGDataType UpstreamTypes) const;
+	bool IsDownstreamPinTypeCompatible(EPCGDataType UpstreamTypes) const;
 	bool CanConnect(const UPCGPin* OtherPin) const;
 
 	/** Adds edge to pin if edge does not exist. Optionally returns list of affected nodes (including the pin owner). */
diff --git a/Engine/Plugins/PCG/Source/PCGEditor/Private/PCGEditorGraph.cpp b/Engine/Plugins/PCG/Source/PCGEditor/Private/PCGEditorGraph.cpp
--- a/Engine/Plugins/PCG/Source/PCGEditor/Private/PCGEditorGraph.cpp
+++ b/Engine/Plugins/PCG/Source/PCGEditor/Private/PCGEditorGraph.cpp
@@ -14,7 +14,6 @@
 #include "Nodes/PCGEditorGraphNodeInput.h"
 #include "Nodes/PCGEditorGraphNodeOutput.h"
 #include "Nodes/PCGEditorGraphNodeReroute.h"
-#include "Types/SlateVector2.h"
 
 #include "EdGraph/EdGraphPin.h"
 
@@ -48,7 +47,7 @@ namespace PCGEditorGraphUtils
 	UPCGPin* GetPCGPinFromEdGraphPin(const UEdGraphPin* Pin)
 	{
 		UEdGraphNode* GraphNode = Pin ? Pin->GetOwningNodeUnchecked() : nullptr;
-		UPCGEditorGraphNodeBase* PCGGraphNode = GraphNode ? CastChecked<UPCGEditorGraphNodeBase>(GraphNode) : nullptr;
+		UPCGEditorGraphNodeBase* PCGGraphNode = GraphNode ? CastChecked<UPCGEditorGraphNodeBase>(GraphNode, ECastCheckedType::NullAllowed) : nullptr;
 		UPCGNode* PCGNode = PCGGraphNode ? PCGGraphNode->GetPCGNode() : nullptr;
 		return PCGNode ? (Pin->Direction == EGPD_Input ? PCGNode->GetInputPin(Pin->PinName) : PCGNode->GetOutputPin(Pin->PinName)) : nullptr;
 	}
@@ -380,9 +379,9 @@ bool UPCGEditorGraph::CanReceivePropertyBagDetailsDropOnGraphPin(const UEdGraphP
 {
 	const UPCGPin* PCGPin = PCGEditorGraphUtils::GetPCGPinFromEdGraphPin(Pin);
 	return PCGPin
-	&& !PCGPin->IsOutputPin()
-	&& (PCGPin->EdgeCount() < 1 || PCGPin->AllowsMultipleConnections())
-	&& PCGPin->IsDownstreamCompatible(EPCGDataType::Param);
+		&& !PCGPin->IsOutputPin()
+		&& (PCGPin->EdgeCount() == 0 || PCGPin->AllowsMultipleConnections())
+		&& PCGPin->IsDownstreamPinTypeCompatible(EPCGDataType::Param);
 }
 
 bool UPCGEditorGraph::CanReceivePropertyBagDetailsDropOnGraphNode(const UEdGraphNode* Node) const
@@ -397,7 +396,7 @@ bool UPCGEditorGraph::CanReceivePropertyBagDetailsDropOnGraph(const UEdGraph* Gr
 	return true;
 }
 
-FReply UPCGEditorGraph::OnPropertyBagDetailsDropOnGraphPin(const FPropertyBagPropertyDesc& PropertyDesc, UEdGraphPin* Pin, const FVector2D GraphPosition) const
+FReply UPCGEditorGraph::OnPropertyBagDetailsDropOnGraphPin(const FPropertyBagPropertyDesc& PropertyDesc, UEdGraphPin* Pin, const FVector2f& GraphPosition) const
 {
 	const UEdGraphNode* Node = Pin ? Pin->GetOwningNode() : nullptr;
 	if (Node && Node->GetGraph() && PropertyDesc.ID.IsValid() && PropertyDesc.CachedProperty)
@@ -406,13 +405,13 @@ FReply UPCGEditorGraph::OnPropertyBagDetailsDropOnGraphPin(const FPropertyBagPro
 		Action.SettingsClass = UPCGUserParameterGetSettings::StaticClass();
 		Action.PropertyGuid = PropertyDesc.ID;
 		Action.PropertyName = PropertyDesc.Name;
-		Action.PerformAction(Node->GetGraph(), Pin, UE::Slate::FDeprecateVector2DParameter(GraphPosition), /*bSelectNewNode=*/true);
+		Action.PerformAction(Node->GetGraph(), Pin, GraphPosition, /*bSelectNewNode=*/true);
 	}
 
 	return FReply::Handled();
 }
 
-FReply UPCGEditorGraph::OnPropertyBagDetailsDropOnGraph(const FPropertyBagPropertyDesc& PropertyDesc, UEdGraph* Graph, const FVector2D GraphPosition) const
+FReply UPCGEditorGraph::OnPropertyBagDetailsDropOnGraph(const FPropertyBagPropertyDesc& PropertyDesc, UEdGraph* Graph, const FVector2f& GraphPosition) const
 {
 	if (Graph && PropertyDesc.ID.IsValid() && PropertyDesc.CachedProperty)
 	{
@@ -420,7 +419,7 @@ FReply UPCGEditorGraph::OnPropertyBagDetailsDropOnGraph(const FPropertyBagProper
 		Action.SettingsClass = UPCGUserParameterGetSettings::StaticClass();
 		Action.PropertyGuid = PropertyDesc.ID;
 		Action.PropertyName = PropertyDesc.Name;
-		Action.PerformAction(Graph, nullptr, UE::Slate::FDeprecateVector2DParameter(GraphPosition), /*bSelectNewNode=*/true);
+		Action.PerformAction(Graph, nullptr, GraphPosition, /*bSelectNewNode=*/true);
 	}
 
 	return FReply::Handled();
diff --git a/Engine/Plugins/PCG/Source/PCGEditor/Private/PCGEditorGraph.h b/Engine/Plugins/PCG/Source/PCGEditor/Private/PCGEditorGraph.h
--- a/Engine/Plugins/PCG/Source/PCGEditor/Private/PCGEditorGraph.h
+++ b/Engine/Plugins/PCG/Source/PCGEditor/Private/PCGEditorGraph.h
@@ -65,8 +65,8 @@ protected:
 	virtual bool CanReceivePropertyBagDetailsDropOnGraphPin(const UEdGraphPin* Pin) const override;
 	virtual bool CanReceivePropertyBagDetailsDropOnGraphNode(const UEdGraphNode* Node) const override;
 	virtual bool CanReceivePropertyBagDetailsDropOnGraph(const UEdGraph* Graph) const override;
-	virtual FReply OnPropertyBagDetailsDropOnGraphPin(const FPropertyBagPropertyDesc& PropertyDesc, UEdGraphPin* Pin, FVector2D GraphPosition) const override;
-	virtual FReply OnPropertyBagDetailsDropOnGraph(const FPropertyBagPropertyDesc& PropertyDesc, UEdGraph* Graph, FVector2D GraphPosition) const override;
+	virtual FReply OnPropertyBagDetailsDropOnGraphPin(const FPropertyBagPropertyDesc& PropertyDesc, UEdGraphPin* Pin, const FVector2f& GraphPosition) const override;
+	virtual FReply OnPropertyBagDetailsDropOnGraph(const FPropertyBagPropertyDesc& PropertyDesc, UEdGraph* Graph, const FVector2f& GraphPosition) const override;
 	// ~End IPropertyBagEdGraph interface
 
 private:
diff --git a/Engine/Source/Editor/StructUtilsEditor/Private/PropertyBagDragDropHandler.cpp b/Engine/Source/Editor/StructUtilsEditor/Private/PropertyBagDragDropHandler.cpp
--- a/Engine/Source/Editor/StructUtilsEditor/Private/PropertyBagDragDropHandler.cpp
+++ b/Engine/Source/Editor/StructUtilsEditor/Private/PropertyBagDragDropHandler.cpp
@@ -1,7 +1,5 @@
 // Copyright Epic Games, Inc. All Rights Reserved.
 
-#pragma once
-
 #include "PropertyBagDragDropHandler.h"
 
 #include "IPropertyBagEditorGraph.h"
@@ -88,7 +86,7 @@ FReply FPropertyBagDetailsDragDropOp::DroppedOnPin(const FVector2f& ScreenPositi
 		const IPropertyBagEdGraphDragAndDrop* PropertyBagGraph = GetPropertyBagEdGraphDragAndDropInterface();
 		if (PropertyBagGraph && PropertyBagGraph->CanReceivePropertyBagDetailsDropOnGraphPin(Pin))
 		{
-			return PropertyBagGraph->OnPropertyBagDetailsDropOnGraphPin(PropertyDesc, Pin, FDeprecateSlateVector2D(GraphPosition));
+			return PropertyBagGraph->OnPropertyBagDetailsDropOnGraphPin(PropertyDesc, Pin, GraphPosition);
 		}
 	}
 
@@ -103,7 +101,7 @@ FReply FPropertyBagDetailsDragDropOp::DroppedOnNode(const FVector2f& ScreenPosit
 		if (PropertyBagGraph && PropertyBagGraph->CanReceivePropertyBagDetailsDropOnGraphNode(Node))
 		{
 			// @todo_pcg: remove return type?
-			return PropertyBagGraph->OnPropertyBagDetailsDropOnGraphNode(PropertyDesc, Node, FDeprecateSlateVector2D(GraphPosition));
+			return PropertyBagGraph->OnPropertyBagDetailsDropOnGraphNode(PropertyDesc, Node, GraphPosition);
 		}
 	}
 
@@ -116,7 +114,7 @@ FReply FPropertyBagDetailsDragDropOp::DroppedOnPanel(const TSharedRef<SWidget>&
 	const IPropertyBagEdGraphDragAndDrop* PropertyBagGraph = GetPropertyBagEdGraphDragAndDropInterface();
 	if (PropertyBagGraph && PropertyBagGraph->CanReceivePropertyBagDetailsDropOnGraph(&Graph))
 	{
-		return PropertyBagGraph->OnPropertyBagDetailsDropOnGraph(PropertyDesc, &Graph, FDeprecateSlateVector2D(GraphPosition));
+		return PropertyBagGraph->OnPropertyBagDetailsDropOnGraph(PropertyDesc, &Graph, GraphPosition);
 	}
 
 	return FReply::Handled();
diff --git a/Engine/Source/Editor/StructUtilsEditor/Public/IPropertyBagEditorGraph.h b/Engine/Source/Editor/StructUtilsEditor/Public/IPropertyBagEditorGraph.h
--- a/Engine/Source/Editor/StructUtilsEditor/Public/IPropertyBagEditorGraph.h
+++ b/Engine/Source/Editor/StructUtilsEditor/Public/IPropertyBagEditorGraph.h
@@ -52,9 +52,9 @@ public:
 	virtual bool CanReceivePropertyBagDetailsDropOnGraphPin(const UEdGraphPin* Pin) const;
 	virtual bool CanReceivePropertyBagDetailsDropOnGraphNode(const UEdGraphNode* Node) const;
 	virtual bool CanReceivePropertyBagDetailsDropOnGraph(const UEdGraph* Graph) const;
-	virtual FReply OnPropertyBagDetailsDropOnGraphPin(const FPropertyBagPropertyDesc& PropertyDesc, UEdGraphPin* Pin, FVector2D GraphPosition) const;
-	virtual FReply OnPropertyBagDetailsDropOnGraphNode(const FPropertyBagPropertyDesc& PropertyDesc, UEdGraphNode* Node, FVector2D GraphPosition) const;
-	virtual FReply OnPropertyBagDetailsDropOnGraph(const FPropertyBagPropertyDesc& PropertyDesc, UEdGraph* Graph, FVector2D GraphPosition) const;
+	virtual FReply OnPropertyBagDetailsDropOnGraphPin(const FPropertyBagPropertyDesc& PropertyDesc, UEdGraphPin* Pin, const FVector2f& GraphPosition) const;
+	virtual FReply OnPropertyBagDetailsDropOnGraphNode(const FPropertyBagPropertyDesc& PropertyDesc, UEdGraphNode* Node, const FVector2f& GraphPosition) const;
+	virtual FReply OnPropertyBagDetailsDropOnGraph(const FPropertyBagPropertyDesc& PropertyDesc, UEdGraph* Graph, const FVector2f& GraphPosition) const;
 };
 
 UINTERFACE(MinimalAPI)
@@ -89,17 +89,17 @@ inline bool IPropertyBagEdGraphDragAndDrop::CanReceivePropertyBagDetailsDropOnGr
 	return false;
 }
 
-inline FReply IPropertyBagEdGraphDragAndDrop::OnPropertyBagDetailsDropOnGraphPin(const FPropertyBagPropertyDesc& PropertyDesc, UEdGraphPin* Pin, const FVector2D GraphPosition) const
+inline FReply IPropertyBagEdGraphDragAndDrop::OnPropertyBagDetailsDropOnGraphPin(const FPropertyBagPropertyDesc& PropertyDesc, UEdGraphPin* Pin, const FVector2f& GraphPosition) const
 {
 	return FReply::Handled();
 }
 
-inline FReply IPropertyBagEdGraphDragAndDrop::OnPropertyBagDetailsDropOnGraphNode(const FPropertyBagPropertyDesc& PropertyDesc, UEdGraphNode* Node, const FVector2D GraphPosition) const
+inline FReply IPropertyBagEdGraphDragAndDrop::OnPropertyBagDetailsDropOnGraphNode(const FPropertyBagPropertyDesc& PropertyDesc, UEdGraphNode* Node, const FVector2f& GraphPosition) const
 {
 	return FReply::Handled();
 }
 
-inline FReply IPropertyBagEdGraphDragAndDrop::OnPropertyBagDetailsDropOnGraph(const FPropertyBagPropertyDesc& PropertyDesc, UEdGraph* Graph, const FVector2D GraphPosition) const
+inline FReply IPropertyBagEdGraphDragAndDrop::OnPropertyBagDetailsDropOnGraph(const FPropertyBagPropertyDesc& PropertyDesc, UEdGraph* Graph, const FVector2f& GraphPosition) const
 {
 	return FReply::Handled();
 }
\ No newline at end of file
