[PCG] Change [REDACTED_HANDLE] for [REDACTED_HANDLE] for data wide metadata domain.

#rnx

diff --git a/Engine/Plugins/PCG/Source/PCG/Private/Data/PCGSplineData.cpp b/Engine/Plugins/PCG/Source/PCG/Private/Data/PCGSplineData.cpp
--- a/Engine/Plugins/PCG/Source/PCG/Private/Data/PCGSplineData.cpp
+++ b/Engine/Plugins/PCG/Source/PCG/Private/Data/PCGSplineData.cpp
@@ -350,7 +350,7 @@ TUniquePtr<IPCGAttributeAccessor> UPCGSplineData::CreateStaticAccessor(const FPC
 			PCGLog::LogErrorOnGraph(FText::Format(LOCTEXT("FailCreateAccessor", "Property {0} is not valid for a control point."), FText::FromName(PropertyName)));
 		}
 	}
-	else if (InSelector.GetSelection() == EPCGAttributePropertySelection::Property && DomainName == PCGDataConstants::GlobalDomainName)
+	else if (InSelector.GetSelection() == EPCGAttributePropertySelection::Property && DomainName == PCGDataConstants::DataDomainName)
 	{
 		const FName PropertyName = InSelector.GetName();
 		
@@ -386,7 +386,7 @@ FPCGAttributeAccessorMethods UPCGSplineData::GetSplineAccessorMethods()
 		
 		if (Selection == EPCGAttributePropertySelection::Property || Selection == EPCGAttributePropertySelection::ExtraProperty)
 		{
-			return MakeUnique<FPCGAttributeAccessorKeysSplineData>(SplineData, InSelector.GetDomainName() == PCGDataConstants::GlobalDomainName);
+			return MakeUnique<FPCGAttributeAccessorKeysSplineData>(SplineData, InSelector.GetDomainName() == PCGDataConstants::DataDomainName);
 		}
 		else if (Selection == EPCGAttributePropertySelection::Attribute)
 		{
@@ -407,7 +407,7 @@ FPCGAttributeAccessorMethods UPCGSplineData::GetSplineAccessorMethods()
 		const EPCGAttributePropertySelection Selection = InSelector.GetSelection();
 		if (Selection == EPCGAttributePropertySelection::Property || Selection == EPCGAttributePropertySelection::ExtraProperty)
 		{
-			return MakeUnique<const FPCGAttributeAccessorKeysSplineData>(SplineData, InSelector.GetDomainName() == PCGDataConstants::GlobalDomainName);
+			return MakeUnique<const FPCGAttributeAccessorKeysSplineData>(SplineData, InSelector.GetDomainName() == PCGDataConstants::DataDomainName);
 		}
 		else if (Selection == EPCGAttributePropertySelection::Attribute)
 		{
diff --git a/Engine/Plugins/PCG/Source/PCG/Private/PCGData.cpp b/Engine/Plugins/PCG/Source/PCG/Private/PCGData.cpp
--- a/Engine/Plugins/PCG/Source/PCG/Private/PCGData.cpp
+++ b/Engine/Plugins/PCG/Source/PCG/Private/PCGData.cpp
@@ -94,7 +94,7 @@ FPCGMetadataDomainID UPCGData::GetMetadataDomainIDFromSelector(const FPCGAttribu
 	{
 		return PCGMetadataDomainID::Default;
 	}
-	else if (DomainName == PCGDataConstants::GlobalDomainName)
+	else if (DomainName == PCGDataConstants::DataDomainName)
 	{
 		return PCGMetadataDomainID::Data;
 	}
@@ -108,7 +108,7 @@ bool UPCGData::SetDomainFromDomainID(const FPCGMetadataDomainID& InDomainID, FPC
 {
 	if (InDomainID == PCGMetadataDomainID::Data)
 	{
-		InOutSelector.SetDomainName(PCGDataConstants::GlobalDomainName, /*bResetExtraNames=*/false);
+		InOutSelector.SetDomainName(PCGDataConstants::DataDomainName, /*bResetExtraNames=*/false);
 		return true;
 	}
 	else if (InDomainID == PCGMetadataDomainID::Default)
@@ -124,7 +124,7 @@ bool UPCGData::SetDomainFromDomainID(const FPCGMetadataDomainID& InDomainID, FPC
 
 bool UPCGData::MetadataDomainSupportsMultiEntries(const FPCGMetadataDomainID& InDomainID) const
 {
-	// By default, global does not support multi entries
+	// By default, data domain does not support multi entries
 	return InDomainID.IsValid() && InDomainID != PCGMetadataDomainID::Data;
 }
 
diff --git a/Engine/Plugins/PCG/Source/PCG/Private/Tests/Elements/PCGCopyAttributeTest.cpp b/Engine/Plugins/PCG/Source/PCG/Private/Tests/Elements/PCGCopyAttributeTest.cpp
--- a/Engine/Plugins/PCG/Source/PCG/Private/Tests/Elements/PCGCopyAttributeTest.cpp
+++ b/Engine/Plugins/PCG/Source/PCG/Private/Tests/Elements/PCGCopyAttributeTest.cpp
@@ -26,16 +26,16 @@ IMPLEMENT_CUSTOM_SIMPLE_AUTOMATION_TEST(FPCGCopyAttributeTests_Params_MultiValue
 IMPLEMENT_CUSTOM_SIMPLE_AUTOMATION_TEST(FPCGCopyAttributeTests_Params_CopyingToItself, FPCGTestBaseClass, "Plugins.PCG.CopyAttribute.Params.CopyingToItself", PCGTestsCommon::TestFlags)
 IMPLEMENT_CUSTOM_SIMPLE_AUTOMATION_TEST(FPCGCopyAttributeTests_Params_CopyingAllToItself, FPCGTestBaseClass, "Plugins.PCG.CopyAttribute.Params.CopyingAllToItself", PCGTestsCommon::TestFlags)
 
-IMPLEMENT_CUSTOM_SIMPLE_AUTOMATION_TEST(FPCGCopyAttributeTests_MultiDomain_GlobalToGlobal, FPCGTestBaseClass, "Plugins.PCG.CopyAttribute.MultiDomain.GlobalToGlobal", PCGTestsCommon::TestFlags)
-IMPLEMENT_CUSTOM_SIMPLE_AUTOMATION_TEST(FPCGCopyAttributeTests_MultiDomain_GlobalToElements, FPCGTestBaseClass, "Plugins.PCG.CopyAttribute.MultiDomain.GlobalToElements", PCGTestsCommon::TestFlags)
-IMPLEMENT_CUSTOM_SIMPLE_AUTOMATION_TEST(FPCGCopyAttributeTests_MultiDomain_ElementsToGlobal, FPCGTestBaseClass, "Plugins.PCG.CopyAttribute.MultiDomain.ElementsToGlobal", PCGTestsCommon::TestFlags)
-IMPLEMENT_CUSTOM_SIMPLE_AUTOMATION_TEST(FPCGCopyAttributeTests_MultiDomain_ElementsToGlobalTooMany, FPCGTestBaseClass, "Plugins.PCG.CopyAttribute.MultiDomain.ElementsToGlobalTooMany", PCGTestsCommon::TestFlags)
+IMPLEMENT_CUSTOM_SIMPLE_AUTOMATION_TEST(FPCGCopyAttributeTests_MultiDomain_DataToData, FPCGTestBaseClass, "Plugins.PCG.CopyAttribute.MultiDomain.DataToData", PCGTestsCommon::TestFlags)
+IMPLEMENT_CUSTOM_SIMPLE_AUTOMATION_TEST(FPCGCopyAttributeTests_MultiDomain_DataToElements, FPCGTestBaseClass, "Plugins.PCG.CopyAttribute.MultiDomain.DataToElements", PCGTestsCommon::TestFlags)
+IMPLEMENT_CUSTOM_SIMPLE_AUTOMATION_TEST(FPCGCopyAttributeTests_MultiDomain_ElementsToData, FPCGTestBaseClass, "Plugins.PCG.CopyAttribute.MultiDomain.ElementsToData", PCGTestsCommon::TestFlags)
+IMPLEMENT_CUSTOM_SIMPLE_AUTOMATION_TEST(FPCGCopyAttributeTests_MultiDomain_ElementsToDataTooMany, FPCGTestBaseClass, "Plugins.PCG.CopyAttribute.MultiDomain.ElementsToDataTooMany", PCGTestsCommon::TestFlags)
 IMPLEMENT_CUSTOM_SIMPLE_AUTOMATION_TEST(FPCGCopyAttributeTests_MultiDomain_InvalidDomain, FPCGTestBaseClass, "Plugins.PCG.CopyAttribute.MultiDomain.InvalidDomain", PCGTestsCommon::TestFlags)
 
-IMPLEMENT_CUSTOM_SIMPLE_AUTOMATION_TEST(FPCGCopyAttributeTests_MultiDomain_CopyAll_GlobalToGlobal, FPCGTestBaseClass, "Plugins.PCG.CopyAttribute.MultiDomain.CopyAll.GlobalToGlobal", PCGTestsCommon::TestFlags)
-IMPLEMENT_CUSTOM_SIMPLE_AUTOMATION_TEST(FPCGCopyAttributeTests_MultiDomain_CopyAll_GlobalToElements, FPCGTestBaseClass, "Plugins.PCG.CopyAttribute.MultiDomain.CopyAll.GlobalToElements", PCGTestsCommon::TestFlags)
-IMPLEMENT_CUSTOM_SIMPLE_AUTOMATION_TEST(FPCGCopyAttributeTests_MultiDomain_CopyAll_ElementsToGlobal, FPCGTestBaseClass, "Plugins.PCG.CopyAttribute.MultiDomain.CopyAll.ElementsToGlobal", PCGTestsCommon::TestFlags)
-IMPLEMENT_CUSTOM_SIMPLE_AUTOMATION_TEST(FPCGCopyAttributeTests_MultiDomain_CopyAll_ElementsToGlobalTooMany, FPCGTestBaseClass, "Plugins.PCG.CopyAttribute.MultiDomain.CopyAll.ElementsToGlobalTooMany", PCGTestsCommon::TestFlags)
+IMPLEMENT_CUSTOM_SIMPLE_AUTOMATION_TEST(FPCGCopyAttributeTests_MultiDomain_CopyAll_DataToData, FPCGTestBaseClass, "Plugins.PCG.CopyAttribute.MultiDomain.CopyAll.DataToData", PCGTestsCommon::TestFlags)
+IMPLEMENT_CUSTOM_SIMPLE_AUTOMATION_TEST(FPCGCopyAttributeTests_MultiDomain_CopyAll_DataToElements, FPCGTestBaseClass, "Plugins.PCG.CopyAttribute.MultiDomain.CopyAll.DataToElements", PCGTestsCommon::TestFlags)
+IMPLEMENT_CUSTOM_SIMPLE_AUTOMATION_TEST(FPCGCopyAttributeTests_MultiDomain_CopyAll_ElementsToData, FPCGTestBaseClass, "Plugins.PCG.CopyAttribute.MultiDomain.CopyAll.ElementsToData", PCGTestsCommon::TestFlags)
+IMPLEMENT_CUSTOM_SIMPLE_AUTOMATION_TEST(FPCGCopyAttributeTests_MultiDomain_CopyAll_ElementsToDataTooMany, FPCGTestBaseClass, "Plugins.PCG.CopyAttribute.MultiDomain.CopyAll.ElementsToDataTooMany", PCGTestsCommon::TestFlags)
 IMPLEMENT_CUSTOM_SIMPLE_AUTOMATION_TEST(FPCGCopyAttributeTests_MultiDomain_CopyAll_ElementsToElementsExplicit, FPCGTestBaseClass, "Plugins.PCG.CopyAttribute.MultiDomain.CopyAll.ElementsToElementsExplicit", PCGTestsCommon::TestFlags)
 IMPLEMENT_CUSTOM_SIMPLE_AUTOMATION_TEST(FPCGCopyAttributeTests_MultiDomain_CopyAll_ElementsToElementsDefault, FPCGTestBaseClass, "Plugins.PCG.CopyAttribute.MultiDomain.CopyAll.ElementsToElementsDefault", PCGTestsCommon::TestFlags)
 IMPLEMENT_CUSTOM_SIMPLE_AUTOMATION_TEST(FPCGCopyAttributeTests_MultiDomain_CopyAll_AllToAll, FPCGTestBaseClass, "Plugins.PCG.CopyAttribute.MultiDomain.CopyAll.AllToAll", PCGTestsCommon::TestFlags)
@@ -466,7 +466,7 @@ bool FPCGCopyAttributeTests_Params_CopyingAllToItself::RunTest(const FString& Pa
 	return true;
 }
 
-bool FPCGCopyAttributeTests_MultiDomain_GlobalToGlobal::RunTest(const FString& Parameters)
+bool FPCGCopyAttributeTests_MultiDomain_DataToData::RunTest(const FString& Parameters)
 {
 	PCGTestsCommon::FTestData TestData;
 	UPCGCopyAttributesSettings* Settings = PCGTestsCommon::GenerateSettings<UPCGCopyAttributesSettings>(TestData);
@@ -474,8 +474,8 @@ bool FPCGCopyAttributeTests_MultiDomain_GlobalToGlobal::RunTest(const FString& P
 
 	constexpr int NumOfPoints = 20;
 
-	Settings->InputSource.Update(TEXT("@Global.MyAttr"));
-	Settings->OutputTarget.Update(TEXT("@Global.MyAttr2"));
+	Settings->InputSource.Update(TEXT("@Data.MyAttr"));
+	Settings->OutputTarget.Update(TEXT("@Data.MyAttr2"));
 
 	TUniquePtr<FPCGContext> Context = TestData.InitializeTestContext();
 	UPCGBasePointData* InputPointData = PCGCopyAttributeTests::CreateInputPointData(Context.Get(), NumOfPoints);
@@ -494,12 +494,12 @@ bool FPCGCopyAttributeTests_MultiDomain_GlobalToGlobal::RunTest(const FString& P
 
 	check(OutputData && OutputData->ConstMetadata());
 
-	const FPCGMetadataDomain* GlobalDomain = OutputData->ConstMetadata()->GetConstMetadataDomain(PCGMetadataDomainID::Data);
-	check(GlobalDomain);
+	const FPCGMetadataDomain* DataDomain = OutputData->ConstMetadata()->GetConstMetadataDomain(PCGMetadataDomainID::Data);
+	check(DataDomain);
 
-	const FPCGMetadataAttribute<int32>* InputAttribute = GlobalDomain->GetConstTypedAttribute<int32>(TEXT("MyAttr"));
+	const FPCGMetadataAttribute<int32>* InputAttribute = DataDomain->GetConstTypedAttribute<int32>(TEXT("MyAttr"));
 	UTEST_NOT_NULL("Input Attribute exists in the output data", InputAttribute);
-	const FPCGMetadataAttribute<int32>* OutputAttribute = GlobalDomain->GetConstTypedAttribute<int32>(TEXT("MyAttr2"));
+	const FPCGMetadataAttribute<int32>* OutputAttribute = DataDomain->GetConstTypedAttribute<int32>(TEXT("MyAttr2"));
 	UTEST_NOT_NULL("Output Attribute exists in the output data", OutputAttribute);
 
 	check(InputAttribute && OutputAttribute);
@@ -509,7 +509,7 @@ bool FPCGCopyAttributeTests_MultiDomain_GlobalToGlobal::RunTest(const FString& P
 	return true;
 }
 
-bool FPCGCopyAttributeTests_MultiDomain_GlobalToElements::RunTest(const FString& Parameters)
+bool FPCGCopyAttributeTests_MultiDomain_DataToElements::RunTest(const FString& Parameters)
 {
 	PCGTestsCommon::FTestData TestData;
 	UPCGCopyAttributesSettings* Settings = PCGTestsCommon::GenerateSettings<UPCGCopyAttributesSettings>(TestData);
@@ -517,7 +517,7 @@ bool FPCGCopyAttributeTests_MultiDomain_GlobalToElements::RunTest(const FString&
 
 	constexpr int NumOfPoints = 20;
 
-	Settings->InputSource.Update(TEXT("@Global.MyAttr"));
+	Settings->InputSource.Update(TEXT("@Data.MyAttr"));
 	Settings->OutputTarget.Update(TEXT("MyAttr2"));
 
 	TUniquePtr<FPCGContext> Context = TestData.InitializeTestContext();
@@ -539,11 +539,11 @@ bool FPCGCopyAttributeTests_MultiDomain_GlobalToElements::RunTest(const FString&
 	check(OutputData && OutputData->ConstMetadata());
 
 	const FPCGMetadataDomain* ElementsDomain = OutputData->ConstMetadata()->GetConstMetadataDomain(PCGMetadataDomainID::Elements);
-	const FPCGMetadataDomain* GlobalDomain = OutputData->ConstMetadata()->GetConstMetadataDomain(PCGMetadataDomainID::Data);
-	check(GlobalDomain && ElementsDomain);
+	const FPCGMetadataDomain* DataDomain = OutputData->ConstMetadata()->GetConstMetadataDomain(PCGMetadataDomainID::Data);
+	check(DataDomain && ElementsDomain);
 
-	const FPCGMetadataAttribute<int32>* InputAttribute = GlobalDomain->GetConstTypedAttribute<int32>(TEXT("MyAttr"));
-	UTEST_NOT_NULL("Input Attribute exists in the output global metadata domain", InputAttribute);
+	const FPCGMetadataAttribute<int32>* InputAttribute = DataDomain->GetConstTypedAttribute<int32>(TEXT("MyAttr"));
+	UTEST_NOT_NULL("Input Attribute exists in the output data metadata domain", InputAttribute);
 	const FPCGMetadataAttribute<int32>* OutputAttribute = ElementsDomain->GetConstTypedAttribute<int32>(TEXT("MyAttr2"));
 	UTEST_NOT_NULL("Output Attribute exists in the output elements metadata domain", OutputAttribute);
 
@@ -561,7 +561,7 @@ bool FPCGCopyAttributeTests_MultiDomain_GlobalToElements::RunTest(const FString&
 	return true;
 }
 
-bool FPCGCopyAttributeTests_MultiDomain_ElementsToGlobal::RunTest(const FString& Parameters)
+bool FPCGCopyAttributeTests_MultiDomain_ElementsToData::RunTest(const FString& Parameters)
 {
 	PCGTestsCommon::FTestData TestData;
 	UPCGCopyAttributesSettings* Settings = PCGTestsCommon::GenerateSettings<UPCGCopyAttributesSettings>(TestData);
@@ -570,7 +570,7 @@ bool FPCGCopyAttributeTests_MultiDomain_ElementsToGlobal::RunTest(const FString&
 	constexpr int NumOfPoints = 1;
 
 	Settings->InputSource.Update(TEXT("MyAttr"));
-	Settings->OutputTarget.Update(TEXT("@Global.MyAttr2"));
+	Settings->OutputTarget.Update(TEXT("@Data.MyAttr2"));
 
 	TUniquePtr<FPCGContext> Context = TestData.InitializeTestContext();
 	UPCGBasePointData* InputPointData = PCGCopyAttributeTests::CreateInputPointData(Context.Get(), NumOfPoints);
@@ -596,13 +596,13 @@ bool FPCGCopyAttributeTests_MultiDomain_ElementsToGlobal::RunTest(const FString&
 	check(OutputData && OutputData->ConstMetadata());
 
 	const FPCGMetadataDomain* ElementsDomain = OutputData->ConstMetadata()->GetConstMetadataDomain(PCGMetadataDomainID::Elements);
-	const FPCGMetadataDomain* GlobalDomain = OutputData->ConstMetadata()->GetConstMetadataDomain(PCGMetadataDomainID::Data);
-	check(GlobalDomain && ElementsDomain);
+	const FPCGMetadataDomain* DataDomain = OutputData->ConstMetadata()->GetConstMetadataDomain(PCGMetadataDomainID::Data);
+	check(DataDomain && ElementsDomain);
 
 	const FPCGMetadataAttribute<int32>* InputAttribute = ElementsDomain->GetConstTypedAttribute<int32>(TEXT("MyAttr"));
 	UTEST_NOT_NULL("Input Attribute exists in the output elements metadata domain", InputAttribute);
-	const FPCGMetadataAttribute<int32>* OutputAttribute = GlobalDomain->GetConstTypedAttribute<int32>(TEXT("MyAttr2"));
-	UTEST_NOT_NULL("Output Attribute exists in the output global metadata domain", OutputAttribute);
+	const FPCGMetadataAttribute<int32>* OutputAttribute = DataDomain->GetConstTypedAttribute<int32>(TEXT("MyAttr2"));
+	UTEST_NOT_NULL("Output Attribute exists in the output data metadata domain", OutputAttribute);
 
 	check(InputAttribute && OutputAttribute);
 
@@ -612,7 +612,7 @@ bool FPCGCopyAttributeTests_MultiDomain_ElementsToGlobal::RunTest(const FString&
 	return true;
 }
 
-bool FPCGCopyAttributeTests_MultiDomain_ElementsToGlobalTooMany::RunTest(const FString& Parameters)
+bool FPCGCopyAttributeTests_MultiDomain_ElementsToDataTooMany::RunTest(const FString& Parameters)
 {
 	PCGTestsCommon::FTestData TestData;
 	UPCGCopyAttributesSettings* Settings = PCGTestsCommon::GenerateSettings<UPCGCopyAttributesSettings>(TestData);
@@ -621,7 +621,7 @@ bool FPCGCopyAttributeTests_MultiDomain_ElementsToGlobalTooMany::RunTest(const F
 	constexpr int NumOfPoints = 20;
 
 	Settings->InputSource.Update(TEXT("MyAttr"));
-	Settings->OutputTarget.Update(TEXT("@Global.MyAttr2"));
+	Settings->OutputTarget.Update(TEXT("@Data.MyAttr2"));
 
 	TUniquePtr<FPCGContext> Context = TestData.InitializeTestContext();
 	UPCGBasePointData* InputPointData = PCGCopyAttributeTests::CreateInputPointData(Context.Get(), NumOfPoints);
@@ -640,13 +640,13 @@ bool FPCGCopyAttributeTests_MultiDomain_ElementsToGlobalTooMany::RunTest(const F
 	check(OutputData && OutputData->ConstMetadata());
 
 	const FPCGMetadataDomain* ElementsDomain = OutputData->ConstMetadata()->GetConstMetadataDomain(PCGMetadataDomainID::Elements);
-	const FPCGMetadataDomain* GlobalDomain = OutputData->ConstMetadata()->GetConstMetadataDomain(PCGMetadataDomainID::Data);
-	check(GlobalDomain && ElementsDomain);
+	const FPCGMetadataDomain* DataDomain = OutputData->ConstMetadata()->GetConstMetadataDomain(PCGMetadataDomainID::Data);
+	check(DataDomain && ElementsDomain);
 
 	const FPCGMetadataAttribute<double>* InputAttribute = ElementsDomain->GetConstTypedAttribute<double>(TEXT("MyAttr"));
 	UTEST_NOT_NULL("Input Attribute exists in the output elements metadata domain", InputAttribute);
-	const FPCGMetadataAttribute<double>* OutputAttribute = GlobalDomain->GetConstTypedAttribute<double>(TEXT("MyAttr2"));
-	UTEST_NOT_NULL("Output Attribute exists in the output global metadata domain", OutputAttribute);
+	const FPCGMetadataAttribute<double>* OutputAttribute = DataDomain->GetConstTypedAttribute<double>(TEXT("MyAttr2"));
+	UTEST_NOT_NULL("Output Attribute exists in the output data metadata domain", OutputAttribute);
 
 	check(InputAttribute && OutputAttribute);
 
@@ -688,13 +688,13 @@ bool FPCGCopyAttributeTests_MultiDomain_InvalidDomain::RunTest(const FString& Pa
 	return true;
 }
 
-bool FPCGCopyAttributeTests_MultiDomain_CopyAll_GlobalToGlobal::RunTest(const FString& Parameters)
+bool FPCGCopyAttributeTests_MultiDomain_CopyAll_DataToData::RunTest(const FString& Parameters)
 {
 	PCGTestsCommon::FTestData TestData;
 	UPCGCopyAttributesSettings* Settings = PCGTestsCommon::GenerateSettings<UPCGCopyAttributesSettings>(TestData);
 	check(Settings);
 	
-	Settings->MetadataDomainsMapping.Emplace(PCGDataConstants::GlobalDomainName, PCGDataConstants::GlobalDomainName);
+	Settings->MetadataDomainsMapping.Emplace(PCGDataConstants::DataDomainName, PCGDataConstants::DataDomainName);
 	Settings->bCopyAllAttributes = true;
 
 	constexpr int NumOfPoints = 20;
@@ -724,16 +724,16 @@ bool FPCGCopyAttributeTests_MultiDomain_CopyAll_GlobalToGlobal::RunTest(const FS
 	check(OutputData && OutputData->ConstMetadata());
 
 	const FPCGMetadataDomain* ElementsDomain = OutputData->ConstMetadata()->GetConstMetadataDomain(PCGMetadataDomainID::Elements);
-	const FPCGMetadataDomain* GlobalDomain = OutputData->ConstMetadata()->GetConstMetadataDomain(PCGMetadataDomainID::Data);
-	check(GlobalDomain && ElementsDomain);
+	const FPCGMetadataDomain* DataDomain = OutputData->ConstMetadata()->GetConstMetadataDomain(PCGMetadataDomainID::Data);
+	check(DataDomain && ElementsDomain);
 
 	UTEST_EQUAL("No attributes on the output elements metadata domain", ElementsDomain->GetAttributeCount(), 0);
 
-	const FPCGMetadataAttribute<int32>* ExistingAttribute = GlobalDomain->GetConstTypedAttribute<int32>(TEXT("MyAttr"));
-	UTEST_NOT_NULL("Existing attribute is still there in the output global metadata domain", ExistingAttribute);
+	const FPCGMetadataAttribute<int32>* ExistingAttribute = DataDomain->GetConstTypedAttribute<int32>(TEXT("MyAttr"));
+	UTEST_NOT_NULL("Existing attribute is still there in the output data metadata domain", ExistingAttribute);
 	
-	const FPCGMetadataAttribute<int32>* Attribute = GlobalDomain->GetConstTypedAttribute<int32>(TEXT("MyAttr2"));
-	UTEST_NOT_NULL("Source Attribute exists in the output global metadata domain", Attribute);
+	const FPCGMetadataAttribute<int32>* Attribute = DataDomain->GetConstTypedAttribute<int32>(TEXT("MyAttr2"));
+	UTEST_NOT_NULL("Source Attribute exists in the output data metadata domain", Attribute);
 
 	check(Attribute);
 
@@ -758,12 +758,12 @@ bool FPCGCopyAttributeTests_MultiDomain_CopyAll_AllToAll::RunTest(const FString&
 	UPCGParamData* TargetParamData = PCGCopyAttributeTests::CreateInputParamData(Context.Get());
 	UPCGBasePointData* SourcePointData = PCGCopyAttributeTests::CreateInputPointData(Context.Get(), NumOfPoints);
 	
-	TargetParamData->MutableMetadata()->GetMetadataDomain(PCGMetadataDomainID::Data)->CreateAttribute<int32>(TEXT("MyGlobalAttr"), 5, false, false);
+	TargetParamData->MutableMetadata()->GetMetadataDomain(PCGMetadataDomainID::Data)->CreateAttribute<int32>(TEXT("MyDataAttr"), 5, false, false);
 	TargetParamData->MutableMetadata()->GetMetadataDomain(PCGMetadataDomainID::Data)->AddEntry();
 
 	PCGTestsCommon::CreateAndFillRandomAttribute<double>(TargetParamData, TEXT("MyAttr"), 0.0, NumOfPoints);
 
-	const FPCGMetadataAttribute<int32>* SourceAttribute = SourcePointData->MutableMetadata()->GetMetadataDomain(PCGMetadataDomainID::Data)->CreateAttribute<int32>(TEXT("MyGlobalAttr2"), /*DefaultValue=*/6, /*bAllowInterpolation=*/false, /*bOverrideParent=*/false);
+	const FPCGMetadataAttribute<int32>* SourceAttribute = SourcePointData->MutableMetadata()->GetMetadataDomain(PCGMetadataDomainID::Data)->CreateAttribute<int32>(TEXT("MyDataAttr2"), /*DefaultValue=*/6, /*bAllowInterpolation=*/false, /*bOverrideParent=*/false);
 	SourcePointData->MutableMetadata()->GetMetadataDomain(PCGMetadataDomainID::Data)->AddEntry();
 
 	PCGTestsCommon::CreateAndFillRandomAttribute<double>(SourcePointData, TEXT("MyAttr2"), 0.0, NumOfPoints);
@@ -784,14 +784,14 @@ bool FPCGCopyAttributeTests_MultiDomain_CopyAll_AllToAll::RunTest(const FString&
 	check(OutputData && OutputData->ConstMetadata());
 
 	const FPCGMetadataDomain* ElementsDomain = OutputData->ConstMetadata()->GetConstMetadataDomain(PCGMetadataDomainID::Elements);
-	const FPCGMetadataDomain* GlobalDomain = OutputData->ConstMetadata()->GetConstMetadataDomain(PCGMetadataDomainID::Data);
-	check(GlobalDomain && ElementsDomain);
+	const FPCGMetadataDomain* DataDomain = OutputData->ConstMetadata()->GetConstMetadataDomain(PCGMetadataDomainID::Data);
+	check(DataDomain && ElementsDomain);
 
-	const FPCGMetadataAttribute<int32>* ExistingAttribute = GlobalDomain->GetConstTypedAttribute<int32>(TEXT("MyGlobalAttr"));
-	UTEST_NOT_NULL("Existing attribute is still there in the output global metadata domain", ExistingAttribute);
+	const FPCGMetadataAttribute<int32>* ExistingAttribute = DataDomain->GetConstTypedAttribute<int32>(TEXT("MyDataAttr"));
+	UTEST_NOT_NULL("Existing attribute is still there in the output data metadata domain", ExistingAttribute);
 	
-	const FPCGMetadataAttribute<int32>* Attribute = GlobalDomain->GetConstTypedAttribute<int32>(TEXT("MyGlobalAttr2"));
-	UTEST_NOT_NULL("Source Attribute exists in the output global metadata domain", Attribute);
+	const FPCGMetadataAttribute<int32>* Attribute = DataDomain->GetConstTypedAttribute<int32>(TEXT("MyDataAttr2"));
+	UTEST_NOT_NULL("Source Attribute exists in the output data metadata domain", Attribute);
 
 	check(Attribute);
 
@@ -816,7 +816,7 @@ bool FPCGCopyAttributeTests_MultiDomain_CopyAll_AllToAll::RunTest(const FString&
 	return true;
 }
 
-bool FPCGCopyAttributeTests_MultiDomain_CopyAll_GlobalToElements::RunTest(const FString& Parameters)
+bool FPCGCopyAttributeTests_MultiDomain_CopyAll_DataToElements::RunTest(const FString& Parameters)
 {
 	PCGTestsCommon::FTestData TestData;
 	UPCGCopyAttributesSettings* Settings = PCGTestsCommon::GenerateSettings<UPCGCopyAttributesSettings>(TestData);
@@ -824,14 +824,14 @@ bool FPCGCopyAttributeTests_MultiDomain_CopyAll_GlobalToElements::RunTest(const
 
 	constexpr int NumOfPoints = 20;
 
-	Settings->MetadataDomainsMapping.Emplace(PCGDataConstants::GlobalDomainName, PCGPointDataConstants::ElementsDomainName);
+	Settings->MetadataDomainsMapping.Emplace(PCGDataConstants::DataDomainName, PCGPointDataConstants::ElementsDomainName);
 	Settings->bCopyAllAttributes = true;
 
 	TUniquePtr<FPCGContext> Context = TestData.InitializeTestContext();
 	UPCGBasePointData* TargetPointData = PCGCopyAttributeTests::CreateInputPointData(Context.Get(), NumOfPoints);
 	UPCGParamData* SourceParamData = PCGCopyAttributeTests::CreateInputParamData(Context.Get());
 
-	const FPCGMetadataAttribute<int32>* InputAttribute = SourceParamData->MutableMetadata()->GetMetadataDomain(PCGMetadataDomainID::Data)->CreateAttribute<int32>(TEXT("MyGlobalAttr"), /*DefaultValue=*/5, /*bAllowInterpolation=*/false, /*bOverrideParent=*/false);
+	const FPCGMetadataAttribute<int32>* InputAttribute = SourceParamData->MutableMetadata()->GetMetadataDomain(PCGMetadataDomainID::Data)->CreateAttribute<int32>(TEXT("MyDataAttr"), /*DefaultValue=*/5, /*bAllowInterpolation=*/false, /*bOverrideParent=*/false);
 	SourceParamData->MutableMetadata()->GetMetadataDomain(PCGMetadataDomainID::Data)->AddEntry();
 	
 	PCGCopyAttributeTests::ConnectToSource(Context.Get(), SourceParamData);
@@ -850,7 +850,7 @@ bool FPCGCopyAttributeTests_MultiDomain_CopyAll_GlobalToElements::RunTest(const
 	const FPCGMetadataDomain* ElementsDomain = OutputData->ConstMetadata()->GetConstMetadataDomain(PCGMetadataDomainID::Elements);
 	check(ElementsDomain);
 	
-	const FPCGMetadataAttribute<int32>* OutputAttribute = ElementsDomain->GetConstTypedAttribute<int32>(TEXT("MyGlobalAttr"));
+	const FPCGMetadataAttribute<int32>* OutputAttribute = ElementsDomain->GetConstTypedAttribute<int32>(TEXT("MyDataAttr"));
 	UTEST_NOT_NULL("Output Attribute exists in the output elements metadata domain", OutputAttribute);
 
 	check(InputAttribute && OutputAttribute);
@@ -860,7 +860,7 @@ bool FPCGCopyAttributeTests_MultiDomain_CopyAll_GlobalToElements::RunTest(const
 	return true;
 }
 
-bool FPCGCopyAttributeTests_MultiDomain_CopyAll_ElementsToGlobal::RunTest(const FString& Parameters)
+bool FPCGCopyAttributeTests_MultiDomain_CopyAll_ElementsToData::RunTest(const FString& Parameters)
 {
 	PCGTestsCommon::FTestData TestData;
 	UPCGCopyAttributesSettings* Settings = PCGTestsCommon::GenerateSettings<UPCGCopyAttributesSettings>(TestData);
@@ -868,7 +868,7 @@ bool FPCGCopyAttributeTests_MultiDomain_CopyAll_ElementsToGlobal::RunTest(const
 
 	constexpr int NumOfPoints = 20;
 
-	Settings->MetadataDomainsMapping.Emplace(PCGPointDataConstants::ElementsDomainName, PCGDataConstants::GlobalDomainName);
+	Settings->MetadataDomainsMapping.Emplace(PCGPointDataConstants::ElementsDomainName, PCGDataConstants::DataDomainName);
 	Settings->bCopyAllAttributes = true;
 
 	TUniquePtr<FPCGContext> Context = TestData.InitializeTestContext();
@@ -894,11 +894,11 @@ bool FPCGCopyAttributeTests_MultiDomain_CopyAll_ElementsToGlobal::RunTest(const
 	UTEST_EQUAL("There is the right number of points in output", OutputData->GetNumPoints(), NumOfPoints);
 
 	check(OutputData && OutputData->ConstMetadata());
-	const FPCGMetadataDomain* GlobalDomain = OutputData->ConstMetadata()->GetConstMetadataDomain(PCGMetadataDomainID::Data);
-	check(GlobalDomain);
+	const FPCGMetadataDomain* DataDomain = OutputData->ConstMetadata()->GetConstMetadataDomain(PCGMetadataDomainID::Data);
+	check(DataDomain);
 	
-	const FPCGMetadataAttribute<int32>* OutputAttribute = GlobalDomain->GetConstTypedAttribute<int32>(TEXT("MyAttr"));
-	UTEST_NOT_NULL("Output Attribute exists in the output global metadata domain", OutputAttribute);
+	const FPCGMetadataAttribute<int32>* OutputAttribute = DataDomain->GetConstTypedAttribute<int32>(TEXT("MyAttr"));
+	UTEST_NOT_NULL("Output Attribute exists in the output data metadata domain", OutputAttribute);
 
 	check(InputAttribute && OutputAttribute);
 
@@ -908,7 +908,7 @@ bool FPCGCopyAttributeTests_MultiDomain_CopyAll_ElementsToGlobal::RunTest(const
 	return true;
 }
 
-bool FPCGCopyAttributeTests_MultiDomain_CopyAll_ElementsToGlobalTooMany::RunTest(const FString& Parameters)
+bool FPCGCopyAttributeTests_MultiDomain_CopyAll_ElementsToDataTooMany::RunTest(const FString& Parameters)
 {
 	PCGTestsCommon::FTestData TestData;
 	UPCGCopyAttributesSettings* Settings = PCGTestsCommon::GenerateSettings<UPCGCopyAttributesSettings>(TestData);
@@ -916,7 +916,7 @@ bool FPCGCopyAttributeTests_MultiDomain_CopyAll_ElementsToGlobalTooMany::RunTest
 
 	constexpr int NumOfPoints = 20;
 
-	Settings->MetadataDomainsMapping.Emplace(PCGPointDataConstants::ElementsDomainName, PCGDataConstants::GlobalDomainName);
+	Settings->MetadataDomainsMapping.Emplace(PCGPointDataConstants::ElementsDomainName, PCGDataConstants::DataDomainName);
 	Settings->bCopyAllAttributes = true;
 
 	TUniquePtr<FPCGContext> Context = TestData.InitializeTestContext();
@@ -939,11 +939,11 @@ bool FPCGCopyAttributeTests_MultiDomain_CopyAll_ElementsToGlobalTooMany::RunTest
 	UTEST_EQUAL("There is the right number of points in output", OutputData->GetNumPoints(), NumOfPoints);
 
 	check(OutputData && OutputData->ConstMetadata());
-	const FPCGMetadataDomain* GlobalDomain = OutputData->ConstMetadata()->GetConstMetadataDomain(PCGMetadataDomainID::Data);
-	check(GlobalDomain);
+	const FPCGMetadataDomain* DataDomain = OutputData->ConstMetadata()->GetConstMetadataDomain(PCGMetadataDomainID::Data);
+	check(DataDomain);
 	
-	const FPCGMetadataAttribute<double>* OutputAttribute = GlobalDomain->GetConstTypedAttribute<double>(TEXT("MyAttr"));
-	UTEST_NOT_NULL("Output Attribute exists in the output global metadata domain", OutputAttribute);
+	const FPCGMetadataAttribute<double>* OutputAttribute = DataDomain->GetConstTypedAttribute<double>(TEXT("MyAttr"));
+	UTEST_NOT_NULL("Output Attribute exists in the output data metadata domain", OutputAttribute);
 
 	check(InputAttribute && OutputAttribute);
 
diff --git a/Engine/Plugins/PCG/Source/PCG/Private/Tests/Metadata/PCGMetadataTest.cpp b/Engine/Plugins/PCG/Source/PCG/Private/Tests/Metadata/PCGMetadataTest.cpp
--- a/Engine/Plugins/PCG/Source/PCG/Private/Tests/Metadata/PCGMetadataTest.cpp
+++ b/Engine/Plugins/PCG/Source/PCG/Private/Tests/Metadata/PCGMetadataTest.cpp
@@ -14,10 +14,10 @@
 IMPLEMENT_CUSTOM_SIMPLE_AUTOMATION_TEST(FPCGMetadataTest_Inherit, FPCGTestBaseClass, "Plugins.PCG.Metadata.Inherit", PCGTestsCommon::TestFlags)
 IMPLEMENT_CUSTOM_SIMPLE_AUTOMATION_TEST(FPCGMetadataTest_InheritCopy, FPCGTestBaseClass, "Plugins.PCG.Metadata.InheritCopy", PCGTestsCommon::TestFlags)
 IMPLEMENT_CUSTOM_SIMPLE_AUTOMATION_TEST(FPCGMetadataTest_InheritWithNoParenting, FPCGTestBaseClass, "Plugins.PCG.Metadata.InheritWithNoParenting", PCGTestsCommon::TestFlags)
-IMPLEMENT_CUSTOM_SIMPLE_AUTOMATION_TEST(FPCGMetadataTest_TwoAddEntryToGlobal, FPCGTestBaseClass, "Plugins.PCG.Metadata.TwoAddEntryToGlobal", PCGTestsCommon::TestFlags)
-IMPLEMENT_CUSTOM_SIMPLE_AUTOMATION_TEST(FPCGMetadataTest_AddEntriesToGlobal, FPCGTestBaseClass, "Plugins.PCG.Metadata.AddEntriesToGlobal", PCGTestsCommon::TestFlags)
-IMPLEMENT_CUSTOM_SIMPLE_AUTOMATION_TEST(FPCGMetadataTest_AddEntriesInPlaceToGlobal, FPCGTestBaseClass, "Plugins.PCG.Metadata.AddEntriesToGlobal", PCGTestsCommon::TestFlags)
-IMPLEMENT_CUSTOM_SIMPLE_AUTOMATION_TEST(FPCGMetadataTest_DelayedAddEntriesToGlobal, FPCGTestBaseClass, "Plugins.PCG.Metadata.DelayedAddEntriesToGlobal", PCGTestsCommon::TestFlags)
+IMPLEMENT_CUSTOM_SIMPLE_AUTOMATION_TEST(FPCGMetadataTest_TwoAddEntryToData, FPCGTestBaseClass, "Plugins.PCG.Metadata.TwoAddEntryToData", PCGTestsCommon::TestFlags)
+IMPLEMENT_CUSTOM_SIMPLE_AUTOMATION_TEST(FPCGMetadataTest_AddEntriesToData, FPCGTestBaseClass, "Plugins.PCG.Metadata.AddEntriesToData", PCGTestsCommon::TestFlags)
+IMPLEMENT_CUSTOM_SIMPLE_AUTOMATION_TEST(FPCGMetadataTest_AddEntriesInPlaceToData, FPCGTestBaseClass, "Plugins.PCG.Metadata.AddEntriesToData", PCGTestsCommon::TestFlags)
+IMPLEMENT_CUSTOM_SIMPLE_AUTOMATION_TEST(FPCGMetadataTest_DelayedAddEntriesToData, FPCGTestBaseClass, "Plugins.PCG.Metadata.DelayedAddEntriesToData", PCGTestsCommon::TestFlags)
 
 bool FPCGMetadataTest_Inherit::RunTest(const FString& Parameters)
 {
@@ -223,7 +223,7 @@ bool FPCGMetadataTest_InheritWithNoParenting::RunTest(const FString& Parameters)
 	return true;
 }
 
-bool FPCGMetadataTest_TwoAddEntryToGlobal::RunTest(const FString& Parameters)
+bool FPCGMetadataTest_TwoAddEntryToData::RunTest(const FString& Parameters)
 {
 	const FName AttributeName = TEXT("MyAttr");
 	
@@ -247,7 +247,7 @@ bool FPCGMetadataTest_TwoAddEntryToGlobal::RunTest(const FString& Parameters)
 	return true;
 }
 
-bool FPCGMetadataTest_AddEntriesToGlobal::RunTest(const FString& Parameters)
+bool FPCGMetadataTest_AddEntriesToData::RunTest(const FString& Parameters)
 {
 	const FName AttributeName = TEXT("MyAttr");
 	
@@ -269,7 +269,7 @@ bool FPCGMetadataTest_AddEntriesToGlobal::RunTest(const FString& Parameters)
 	return true;
 }
 
-bool FPCGMetadataTest_AddEntriesInPlaceToGlobal::RunTest(const FString& Parameters)
+bool FPCGMetadataTest_AddEntriesInPlaceToData::RunTest(const FString& Parameters)
 {
 	const FName AttributeName = TEXT("MyAttr");
 	
@@ -293,7 +293,7 @@ bool FPCGMetadataTest_AddEntriesInPlaceToGlobal::RunTest(const FString& Paramete
 	return true;
 }
 
-bool FPCGMetadataTest_DelayedAddEntriesToGlobal::RunTest(const FString& Parameters)
+bool FPCGMetadataTest_DelayedAddEntriesToData::RunTest(const FString& Parameters)
 {
 	const FName AttributeName = TEXT("MyAttr");
 	
diff --git a/Engine/Plugins/PCG/Source/PCG/Public/Data/PCGSplineStruct.h b/Engine/Plugins/PCG/Source/PCG/Public/Data/PCGSplineStruct.h
--- a/Engine/Plugins/PCG/Source/PCG/Public/Data/PCGSplineStruct.h
+++ b/Engine/Plugins/PCG/Source/PCG/Public/Data/PCGSplineStruct.h
@@ -23,8 +23,8 @@ enum class EPCGSplineStructProperties : uint8
 	LocalRotation    UMETA(Tooltip = "Rotation component of the control point's transform."),
 	LocalScale       UMETA(Tooltip = "Scale component of the control point's transform."),
 	LocalTransform   UMETA(Tooltip = "Control point's transform."),
-	SplineTransform	 UMETA(Tooltip = "Transform of the spline. Global parameter", PCG_MetadataDomain="Global"),
-	IsClosed		 UMETA(Tooltip = "If the spline is closed. Global parameter", PCG_MetadataDomain="Global"),
+	SplineTransform	 UMETA(Tooltip = "Transform of the spline. Global parameter", PCG_MetadataDomain="Data"),
+	IsClosed		 UMETA(Tooltip = "If the spline is closed. Global parameter", PCG_MetadataDomain="Data"),
 };
 
 enum class EPCGControlPointsAccessorTarget
diff --git a/Engine/Plugins/PCG/Source/PCG/Public/PCGData.h b/Engine/Plugins/PCG/Source/PCG/Public/PCGData.h
--- a/Engine/Plugins/PCG/Source/PCG/Public/PCGData.h
+++ b/Engine/Plugins/PCG/Source/PCG/Public/PCGData.h
@@ -34,7 +34,7 @@ ENUM_CLASS_FLAGS(EPCGDataUsage);
 
 namespace PCGDataConstants
 {
-	const FName GlobalDomainName = "Global";
+	const FName DataDomainName = "Data";
 	const FName DefaultDomainName = "Default";
 }
 
diff --git a/Engine/Plugins/PCG/Source/PCGEditor/Private/DataVisualizations/PCGSplineDataVisualization.cpp b/Engine/Plugins/PCG/Source/PCGEditor/Private/DataVisualizations/PCGSplineDataVisualization.cpp
--- a/Engine/Plugins/PCG/Source/PCGEditor/Private/DataVisualizations/PCGSplineDataVisualization.cpp
+++ b/Engine/Plugins/PCG/Source/PCGEditor/Private/DataVisualizations/PCGSplineDataVisualization.cpp
@@ -28,7 +28,7 @@ FPCGTableVisualizerInfo IPCGSplineDataVisualization::GetTableVisualizerInfoWithD
 	{
 		// Column Sorting
 		FPCGAttributePropertySelector IndexSelector = FPCGAttributePropertySelector::CreateExtraPropertySelector(EPCGExtraProperties::Index);
-		IndexSelector.SetDomainName(PCGDataConstants::GlobalDomainName);
+		IndexSelector.SetDomainName(PCGDataConstants::DataDomainName);
 		
 		AddColumnInfo(Info, SplineData, IndexSelector);
 		Info.SortingColumn = Info.ColumnInfos.Last().Id;
