Added filtering to the debug object tree view

diff --git a/Engine/Plugins/PCG/Source/PCGEditor/Private/Widgets/SPCGEditorGraphDebugObjectTree.cpp b/Engine/Plugins/PCG/Source/PCGEditor/Private/Widgets/SPCGEditorGraphDebugObjectTree.cpp
--- a/Engine/Plugins/PCG/Source/PCGEditor/Private/Widgets/SPCGEditorGraphDebugObjectTree.cpp
+++ b/Engine/Plugins/PCG/Source/PCGEditor/Private/Widgets/SPCGEditorGraphDebugObjectTree.cpp
@@ -20,7 +20,9 @@
 #include "Algo/AnyOf.h"
 #include "Editor/UnrealEdEngine.h"
 #include "Styling/SlateIconFinder.h"
+#include "Widgets/Images/SLayeredImage.h"
 #include "Widgets/Images/SImage.h"
+#include "Widgets/Input/SComboBox.h"
 #include "Widgets/Layout/SScrollBar.h"
 #include "Widgets/Layout/SScrollBox.h"
 
@@ -37,6 +39,36 @@ namespace PCGEditorGraphDebugObjectTree
 		OutColorAndOpacity = Item->IsDebuggable() ? DefaultGraphColor : FLinearColor(1.0f, 1.0f, 1.0f, 0.5f);
 	}
 
+	bool StackContainsErrorOrWarning(const FPCGStack* InStack)
+	{
+		if (!InStack)
+		{
+			return false;
+		}
+
+		if (IPCGEditorModule* PCGEditorModule = IPCGEditorModule::Get())
+		{
+			bool bHasErrorOrWarning = false;
+			PCGEditorModule->GetNodeVisualLogs().ForAllMatchingLogs(*InStack, [&bHasErrorOrWarning](const FPCGStack& InStack, const FPCGPerNodeVisualLogs& InLogs)
+			{
+				for(const FPCGNodeLogEntry& LogEntry : InLogs)
+				{
+					if (LogEntry.Verbosity > ELogVerbosity::NoLogging && LogEntry.Verbosity <= ELogVerbosity::Warning)
+					{
+						bHasErrorOrWarning = true;
+						return false;
+					}
+				}
+
+				return true;
+			});
+
+			return bHasErrorOrWarning;
+		}
+
+		return false;
+	}
+
 	void GetErrorInfoInternal(FPCGEditorGraphDebugObjectItemPtr Item, ELogVerbosity::Type& OutMinVerbosity, ELogVerbosity::Type& OutMinLocalVerbosity, FPCGPerNodeVisualLogs& OutLogs)
 	{
 		if (!Item || !ensure(Item->GetPCGStack()))
@@ -452,6 +484,12 @@ void SPCGEditorGraphDebugObjectTree::Construct(const FArguments& InArgs, TShared
 		LOCTEXT("SetDebugObject", "Set debug object from Level Editor selection."),
 		TAttribute<bool>::Create(TAttribute<bool>::FGetter::CreateSP(this, &SPCGEditorGraphDebugObjectTree::IsSetDebugObjectFromSelectionButtonEnabled)));
 
+	TSharedPtr<SLayeredImage> FilterImage = SNew(SLayeredImage)
+		.Image(FAppStyle::Get().GetBrush("DetailsView.ViewOptions"))
+		.ColorAndOpacity(FSlateColor::UseForeground());
+
+	FilterImage->AddLayer(TAttribute<const FSlateBrush*>(this, &SPCGEditorGraphDebugObjectTree::GetFilterBadgeIcon));
+
 	ChildSlot
 	[
 		SNew(SVerticalBox)
@@ -484,6 +522,22 @@ void SPCGEditorGraphDebugObjectTree::Construct(const FArguments& InArgs, TShared
 					.ColorAndOpacity(FSlateColor::UseForeground())
 				]
 			]
+			+SHorizontalBox::Slot()
+			.AutoWidth()
+			.VAlign(VAlign_Center)
+			.Padding(4.0f)
+			[
+				SNew(SComboButton)
+				.HasDownArrow(false)
+				.ForegroundColor(FSlateColor::UseForeground())
+				.ButtonStyle(FAppStyle::Get(), "SimpleButton")
+				.OnGetMenuContent(this, &SPCGEditorGraphDebugObjectTree::OpenFilterMenu)
+				.ContentPadding(0.0f)
+				.ButtonContent()
+				[
+					FilterImage.ToSharedRef()
+				]
+			]
 		]
 		+SVerticalBox::Slot()
 		.FillHeight(1.0f)
@@ -844,8 +898,9 @@ void SPCGEditorGraphDebugObjectTree::AddStackToTree(const FPCGStack& Stack,
 
 	{
 		//TRACE_CPUPROFILER_EVENT_SCOPE(SPCGEditorGraphDebugObjectTree::AddStacksToTree::PerStackLoop);
-		// If the current graph is not in the stack at all then skip.
-		if (!Stack.HasObject(GraphBeingEdited))
+		if (!Stack.HasObject(GraphBeingEdited) ||
+			(!bShowDownstream && Stack.GetGraphForCurrentFrame() != GraphBeingEdited) ||
+			(bShowOnlyErrorsAndWarnings && !PCGEditorGraphDebugObjectTree::StackContainsErrorOrWarning(&Stack)))
 		{
 			return;
 		}
@@ -1154,7 +1209,7 @@ void SPCGEditorGraphDebugObjectTree::RefreshTree()
 		if(IPCGEditorModule* PCGEditorModule = IPCGEditorModule::Get())
 		{
 			TRACE_CPUPROFILER_EVENT_SCOPE(SPCGEditorGraphDebugObjectTree::RefreshTree::DynamicStacks);
-			TArray<FPCGStackSharedPtr> GraphStacks = PCGEditorModule->GetExecutedStacksPtrs(nullptr, PCGGraph, /*bOnlyWithSubgraphAsCurrentFrame=*/false);
+			TArray<FPCGStackSharedPtr> GraphStacks = PCGEditorModule->GetExecutedStacksPtrs(nullptr, PCGGraph, /*bOnlyWithSubgraphAsCurrentFrame=*/!bShowDownstream);
 			// TODO: do per component filtering?
 			AddStacksToTree(GraphStacks, ActorItems, StackToItem);
 		}
@@ -1492,6 +1547,45 @@ void SPCGEditorGraphDebugObjectTree::ExpandAndSelectFirstLeafDebugObject(const F
 	ExpandAndSelectDebugObject(ItemToSelect.IsValid() ? ItemToSelect : InItem);
 }
 
+const FSlateBrush* SPCGEditorGraphDebugObjectTree::GetFilterBadgeIcon() const
+{
+	if (!bShowDownstream || bShowOnlyErrorsAndWarnings)
+	{
+		return FAppStyle::Get().GetBrush("Icons.BadgeModified");
+	}
+	else
+	{
+		return nullptr;
+	}
+}
+
+TSharedRef<SWidget> SPCGEditorGraphDebugObjectTree::OpenFilterMenu()
+{
+	FMenuBuilder MenuBuilder(/*bInShouldCloseWindowAfterMenuSelection=*/true, /*InCommandList=*/nullptr);
+	
+	MenuBuilder.AddMenuEntry(
+		LOCTEXT("ToggleShowOnlyErrorsAndWarnings", "Show only errors/warnings"),
+		LOCTEXT("ToggleShowOnlyErrorsAndWarningsTooltip", "Toggles whether only executions that had errors and warnings are shown."),
+		FSlateIcon(),
+		FUIAction(FExecuteAction::CreateSP(this, &SPCGEditorGraphDebugObjectTree::ToggleShowOnlyErrorsAndWarnings),
+			FCanExecuteAction(),
+			FGetActionCheckState::CreateSP(this, &SPCGEditorGraphDebugObjectTree::IsShowingOnlyErrorsAndWarnings)),
+		NAME_None,
+		EUserInterfaceActionType::ToggleButton);
+
+	MenuBuilder.AddMenuEntry(
+		LOCTEXT("ToggleShowGraphsDownstream", "Show downstream"),
+		LOCTEXT("ToggleShowGraphsDownstreamTooltip", "Toggles whether all graphs downstream to this current graph are shown."),
+		FSlateIcon(),
+		FUIAction(FExecuteAction::CreateSP(this, &SPCGEditorGraphDebugObjectTree::ToggleShowDownstream),
+			FCanExecuteAction(),
+			FGetActionCheckState::CreateSP(this, &SPCGEditorGraphDebugObjectTree::IsShowingDownstream)),
+		NAME_None,
+		EUserInterfaceActionType::ToggleButton);
+
+	return MenuBuilder.MakeWidget();
+}
+
 TSharedPtr<SWidget> SPCGEditorGraphDebugObjectTree::OpenContextMenu()
 {
 	FMenuBuilder MenuBuilder(/*bInShouldCloseWindowAfterMenuSelection=*/true, /*InCommandList=*/nullptr);
@@ -1658,4 +1752,26 @@ bool SPCGEditorGraphDebugObjectTree::CanFocusOnItem(const FPCGEditorGraphDebugOb
 	return false;
 }
 
+void SPCGEditorGraphDebugObjectTree::ToggleShowOnlyErrorsAndWarnings()
+{
+	bShowOnlyErrorsAndWarnings = !bShowOnlyErrorsAndWarnings;
+	RequestRefresh();
+}
+
+ECheckBoxState SPCGEditorGraphDebugObjectTree::IsShowingOnlyErrorsAndWarnings() const
+{
+	return bShowOnlyErrorsAndWarnings ? ECheckBoxState::Checked : ECheckBoxState::Unchecked;
+}
+
+void SPCGEditorGraphDebugObjectTree::ToggleShowDownstream()
+{
+	bShowDownstream = !bShowDownstream;
+	RequestRefresh();
+}
+
+ECheckBoxState SPCGEditorGraphDebugObjectTree::IsShowingDownstream() const
+{
+	return bShowDownstream ? ECheckBoxState::Checked : ECheckBoxState::Unchecked;
+}
+
 #undef LOCTEXT_NAMESPACE
diff --git a/Engine/Plugins/PCG/Source/PCGEditor/Private/Widgets/SPCGEditorGraphDebugObjectTree.h b/Engine/Plugins/PCG/Source/PCGEditor/Private/Widgets/SPCGEditorGraphDebugObjectTree.h
--- a/Engine/Plugins/PCG/Source/PCGEditor/Private/Widgets/SPCGEditorGraphDebugObjectTree.h
+++ b/Engine/Plugins/PCG/Source/PCGEditor/Private/Widgets/SPCGEditorGraphDebugObjectTree.h
@@ -314,6 +314,9 @@ private:
 	/** Expand the given row and select the next occurrence of the current graph. */
 	void ExpandAndSelectFirstLeafDebugObject(const FPCGEditorGraphDebugObjectItemPtr& InItem);
 
+	TSharedRef<SWidget> OpenFilterMenu();
+	const FSlateBrush* GetFilterBadgeIcon() const;
+
 	TSharedPtr<SWidget> OpenContextMenu();
 
 	/** Jump-to context menu command. */
@@ -325,6 +328,11 @@ private:
 	void FocusOnItem(const FPCGEditorGraphDebugObjectItemPtr& Item);
 	bool CanFocusOnItem(const FPCGEditorGraphDebugObjectItemPtr& Item) const;
 
+	void ToggleShowOnlyErrorsAndWarnings();
+	ECheckBoxState IsShowingOnlyErrorsAndWarnings() const;
+	void ToggleShowDownstream();
+	ECheckBoxState IsShowingDownstream() const;
+
 	TWeakPtr<FPCGEditor> PCGEditor;
 
 	TSharedPtr<STreeView<FPCGEditorGraphDebugObjectItemPtr>> DebugObjectTreeView;
@@ -352,4 +360,10 @@ private:
 	FIntVector SelectedGridCoord = FIntVector::ZeroValue;
 	TSoftObjectPtr<UPCGComponent> SelectedOriginalComponent = nullptr;
 	TSoftObjectPtr<const UPCGNode> PCGNodeBeingInspected = nullptr;
+
+	/** Controls whether only stacks containing errors/warnings should be shown. */
+	bool bShowOnlyErrorsAndWarnings = false;
+
+	/** Controls whether downstream graphs are going to be shown. */
+	bool bShowDownstream = true;
 };
