[PCG] Strengthen and unify how we do interpolations for rotations/quat/transforms

* Add new trait IntepolationNeedsNormalization to know if after a WeightedSum we need to normalize
* Add new trait CanNormalize and a Normalize function to have a unified way of normalizing across types.
* Made a pass to unify everywhere we were doing WeightedSum to do the normalization

#rnx

diff --git a/Engine/Plugins/PCG/Source/PCG/Private/Elements/Metadata/PCGMetadataVectorOpElement.cpp b/Engine/Plugins/PCG/Source/PCG/Private/Elements/Metadata/PCGMetadataVectorOpElement.cpp
--- a/Engine/Plugins/PCG/Source/PCG/Private/Elements/Metadata/PCGMetadataVectorOpElement.cpp
+++ b/Engine/Plugins/PCG/Source/PCG/Private/Elements/Metadata/PCGMetadataVectorOpElement.cpp
@@ -174,25 +174,11 @@ namespace PCGMetadataVectorSettings
 			break;
 		}
 		case EPCGMetadataVectorOperation::Normalize:
-			if constexpr (std::is_same_v<FVector4, InType>)
-			{
-				double Length = Value1.Size();
-				if (Length >= UE_SMALL_NUMBER)
-				{
-					return Value1 / Length;
-				}
-				else
-				{
-					return FVector4::Zero();
-				}
-			}
-			else
-			{
-				InType Res = Value1;
-				Res.Normalize();
-				return Res;
-			}
-			break;
+		{
+			InType Res = Value1;
+			PCG::Private::MetadataTraits<InType>::Normalize(Res);
+			return Res;
+		}
 		case EPCGMetadataVectorOperation::Length:
 			DoubleValue = Value1.Size();
 			break;
diff --git a/Engine/Plugins/PCG/Source/PCG/Private/Elements/PCGAttractElement.cpp b/Engine/Plugins/PCG/Source/PCG/Private/Elements/PCGAttractElement.cpp
--- a/Engine/Plugins/PCG/Source/PCG/Private/Elements/PCGAttractElement.cpp
+++ b/Engine/Plugins/PCG/Source/PCG/Private/Elements/PCGAttractElement.cpp
@@ -73,26 +73,23 @@ namespace PCGAttractElement
 						Alpha = (SourceWeight + TargetWeight > UE_SMALL_NUMBER) ? (TargetWeight / (SourceWeight + TargetWeight)) : 0.5;
 					}
 
-					auto Callback = [OutIndex, SourceIndex, TargetIndex, Alpha, &Interpolator, &AttractData](auto Dummy)
+					auto Callback = [OutIndex, SourceIndex, TargetIndex, Alpha, &Interpolator, &AttractData]<typename T>(T)
 					{
-						using OutputType = decltype(Dummy);
-						OutputType SourceValue{}, TargetValue{};
+						T SourceValue{}, TargetValue{};
 						Interpolator.SourceAccessor->Get(SourceValue, SourceIndex, *Interpolator.SourceKeys);
 						Interpolator.TargetAccessor->Get(TargetValue, TargetIndex, *Interpolator.TargetKeys, EPCGAttributeAccessorFlags::AllowBroadcastAndConstructible);
 
-						OutputType OutputValue{};
-						if constexpr (PCG::Private::MetadataTraits<OutputType>::CanInterpolate)
+						T OutputValue{};
+						if constexpr (PCG::Private::MetadataTraits<T>::CanInterpolate)
 						{
-							if constexpr (std::is_same_v<OutputType, FTransform>)
-							{
-								// TODO: use BlendWith or PCGMetadataTransfromSettings::LerpTransform (this one has more options)
-								OutputValue = SourceValue;
-								OutputValue.BlendWith(TargetValue, Alpha);
-							}
-							else
+							OutputValue = PCG::Private::MetadataTraits<T>::WeightedSum(SourceValue, TargetValue, Alpha);
+
+							if constexpr (PCG::Private::MetadataTraits<T>::InterpolationNeedsNormalization)
 							{
-								// TODO: use PCGMetadataMaths::Lerp or PCG::Private::MetadataTraits<OutputType>::WeightedSum or FMath::Lerp
-								OutputValue = FMath::Lerp(SourceValue, TargetValue, Alpha);
+								static_assert(PCG::Private::MetadataTraits<T>::CanNormalize);
+
+								// We need to normalize the resulting value
+								PCG::Private::MetadataTraits<T>::Normalize(OutputValue);
 							}
 						}
 						else // Take value based on maximal weight
diff --git a/Engine/Plugins/PCG/Source/PCG/Private/Elements/PCGAttributeReduceElement.cpp b/Engine/Plugins/PCG/Source/PCG/Private/Elements/PCGAttributeReduceElement.cpp
--- a/Engine/Plugins/PCG/Source/PCG/Private/Elements/PCGAttributeReduceElement.cpp
+++ b/Engine/Plugins/PCG/Source/PCG/Private/Elements/PCGAttributeReduceElement.cpp
@@ -19,37 +19,34 @@ namespace PCGAttributeReduceElement
 	template <typename T>
 	bool Average(const IPCGAttributeAccessorKeys& Keys, const IPCGAttributeAccessor& Accessor, T& OutValue)
 	{
-		if constexpr (!PCG::Private::MetadataTraits<T>::CanSubAdd || !PCG::Private::MetadataTraits<T>::CanInterpolate)
+		if constexpr (!PCG::Private::MetadataTraits<T>::CanInterpolate)
 		{
 			return false;
 		}
 		else
 		{
-			if constexpr (std::is_same_v<FQuat, T>)
-			{
-				// Simple averaging will only work for similar quaternions, but the real answer needs a full solver with eigenvectors.
-				FQuat ZeroQuat{ 0.0, 0.0, 0.0, 0.0 };
-				OutValue = ZeroQuat;
+			const double Weight = 1.0f / Keys.GetNum();
 
-				FQuat FirstQuat = ZeroQuat;
-				const double Weight = 1.0f / Keys.GetNum();
+			// If we need normalization or can't sub/add, we need to do the weighted sum sequentially.
+			// Otherwise, we can just add all the values together and multiply the result by the weight.
+			if constexpr (PCG::Private::MetadataTraits<T>::InterpolationNeedsNormalization || !PCG::Private::MetadataTraits<T>::CanSubAdd)
+			{
+				OutValue = PCG::Private::MetadataTraits<T>::ZeroValueForWeightedSum();
 
-				bool bSuccess = PCGMetadataElementCommon::ApplyOnAccessor<T>(Keys, Accessor, [&FirstQuat, Weight, &OutValue, &ZeroQuat](const T& InValue, int32 Index)
+				bool bSuccess = PCGMetadataElementCommon::ApplyOnAccessor<T>(Keys, Accessor, [Weight, &OutValue](const T& InValue, int32 Index)
 				{
-					// Since doing a dot product with a quat equals to 0 will give us no info, we'll take the first non null quat.
-					if (FirstQuat.Equals(ZeroQuat))
-					{
-						FirstQuat = InValue;
-					}
-
-					// Because q and -q represent the same quaternion, but the average would not be correct. We need to inverse the quaternion if needed.
-					const double ThisWeight = (InValue | FirstQuat) < 0 ? -Weight : Weight;
-					OutValue += (ThisWeight * InValue);
+					OutValue = PCG::Private::MetadataTraits<T>::WeightedSum(OutValue, InValue, Weight);
 				});
 
 				if (bSuccess)
 				{
-					OutValue.Normalize();
+					if constexpr (PCG::Private::MetadataTraits<T>::InterpolationNeedsNormalization)
+					{
+						static_assert(PCG::Private::MetadataTraits<T>::CanNormalize);
+
+						// We need to normalize the resulting value
+						PCG::Private::MetadataTraits<T>::Normalize(OutValue);
+					}
 				}
 
 				return bSuccess;
@@ -65,7 +62,7 @@ namespace PCGAttributeReduceElement
 
 				if (bSuccess)
 				{
-					OutValue = PCG::Private::MetadataTraits<T>::WeightedSum(PCG::Private::MetadataTraits<T>::ZeroValue(), OutValue, 1.0f / Keys.GetNum());
+					OutValue = PCG::Private::MetadataTraits<T>::WeightedSum(PCG::Private::MetadataTraits<T>::ZeroValue(), OutValue, Weight);
 				}
 
 				return bSuccess;
diff --git a/Engine/Plugins/PCG/Source/PCG/Private/Elements/PCGBlurElement.cpp b/Engine/Plugins/PCG/Source/PCG/Private/Elements/PCGBlurElement.cpp
--- a/Engine/Plugins/PCG/Source/PCG/Private/Elements/PCGBlurElement.cpp
+++ b/Engine/Plugins/PCG/Source/PCG/Private/Elements/PCGBlurElement.cpp
@@ -273,11 +273,14 @@ bool FPCGBlurElement::ExecuteInternal(FPCGContext* InContext) const
 							TypedWriteBuffer[WriteIndex] = PCG::Private::MetadataTraits<AttributeType>::WeightedSum(TypedWriteBuffer[WriteIndex], Value, Weight);
 						}
 
-						// For Quaternions, we need to normalize them
-						if constexpr (std::is_same_v<AttributeType, FQuat>)
+						if constexpr (PCG::Private::MetadataTraits<AttributeType>::InterpolationNeedsNormalization)
 						{
-							TypedWriteBuffer[WriteIndex].Normalize();
+							static_assert(PCG::Private::MetadataTraits<AttributeType>::CanNormalize);
+								
+							// We need to normalize the resulting value
+							PCG::Private::MetadataTraits<AttributeType>::Normalize(TypedWriteBuffer[WriteIndex]);
 						}
+
 						++NumWritten;
 					}
 
diff --git a/Engine/Plugins/PCG/Source/PCG/Private/Elements/PCGCollapsePoints.cpp b/Engine/Plugins/PCG/Source/PCG/Private/Elements/PCGCollapsePoints.cpp
--- a/Engine/Plugins/PCG/Source/PCG/Private/Elements/PCGCollapsePoints.cpp
+++ b/Engine/Plugins/PCG/Source/PCG/Private/Elements/PCGCollapsePoints.cpp
@@ -588,27 +588,26 @@ bool FPCGCollapsePointsElement::ExecuteInternal(FPCGContext* InContext) const
 
 					auto Callback = [PartitionIndex, &SecondaryPointIndices, &PartitionWeights, &SourceAccessor, &SourceKeys, &OutputAccessor, &OutputKeys]<typename T>(T)
 					{
-						// If the attribute is a rotator, do the weighted sum in Quat.
-						using OutputType = std::conditional_t<std::is_same_v<T, FRotator>, FQuat, T>;
-
-						if constexpr (PCG::Private::MetadataTraits<OutputType>::CanInterpolate)
+						if constexpr (PCG::Private::MetadataTraits<T>::CanInterpolate)
 						{
-							OutputType Value = PCG::Private::MetadataTraits<OutputType>::ZeroValueForWeightedSum();
+							T Value = PCG::Private::MetadataTraits<T>::ZeroValueForWeightedSum();
 
 							for (int Index = 0; Index < SecondaryPointIndices.Num(); ++Index)
 							{
 								int PointIndex = SecondaryPointIndices[Index];
 
-								OutputType SecondaryValue{};
+								T SecondaryValue{};
 								SourceAccessor->Get(SecondaryValue, PointIndex, *SourceKeys);
 
-								Value = PCG::Private::MetadataTraits<OutputType>::WeightedSum(Value, SecondaryValue, PartitionWeights[Index]);
+								Value = PCG::Private::MetadataTraits<T>::WeightedSum(Value, SecondaryValue, PartitionWeights[Index]);
 							}
 
-							if constexpr (std::is_same_v<OutputType, FQuat>)
+							if constexpr (PCG::Private::MetadataTraits<T>::InterpolationNeedsNormalization)
 							{
-								// We need to normalize the resulting quaternion
-								Value.Normalize();
+								static_assert(PCG::Private::MetadataTraits<T>::CanNormalize);
+
+								// We need to normalize the resulting value
+								PCG::Private::MetadataTraits<T>::Normalize(Value);
 							}
 							
 							OutputAccessor->Set(Value, PartitionIndex, *OutputKeys);
diff --git a/Engine/Plugins/PCG/Source/PCG/Public/Metadata/PCGMetadataAttributeTpl.h b/Engine/Plugins/PCG/Source/PCG/Public/Metadata/PCGMetadataAttributeTpl.h
--- a/Engine/Plugins/PCG/Source/PCG/Public/Metadata/PCGMetadataAttributeTpl.h
+++ b/Engine/Plugins/PCG/Source/PCG/Public/Metadata/PCGMetadataAttributeTpl.h
@@ -1047,11 +1047,20 @@ protected:
 	template<typename IT = T, typename TEnableIf<PCG::Private::MetadataTraits<IT>::CanInterpolate>::Type* = nullptr>
 	void Accumulate(PCGMetadataEntryKey ItemKey, const FPCGMetadataAttributeBase* InAttribute, PCGMetadataEntryKey InEntryKey, float Weight)
 	{
-		SetValue(ItemKey,
-			PCG::Private::MetadataTraits<IT>::WeightedSum(
+		IT Value = PCG::Private::MetadataTraits<IT>::WeightedSum(
 				GetValueFromItemKey(ItemKey),
 				static_cast<const FPCGMetadataAttribute<T>*>(InAttribute)->GetValueFromItemKey(InEntryKey),
-				Weight));
+				Weight);
+
+		if constexpr (PCG::Private::MetadataTraits<IT>::InterpolationNeedsNormalization)
+		{
+			static_assert(PCG::Private::MetadataTraits<IT>::CanNormalize);
+
+			// We need to normalize the resulting value
+			PCG::Private::MetadataTraits<IT>::Normalize(Value);
+		}
+		
+		SetValue(ItemKey, Value);
 	}
 
 	template<typename IT = T, typename TEnableIf<!PCG::Private::MetadataTraits<IT>::CanInterpolate>::Type* = nullptr>
@@ -1072,6 +1081,14 @@ protected:
 				WeightedEntry.Value);
 		}
 
+		if constexpr (PCG::Private::MetadataTraits<IT>::InterpolationNeedsNormalization)
+		{
+			static_assert(PCG::Private::MetadataTraits<IT>::CanNormalize);
+								
+			// We need to normalize the resulting value
+			PCG::Private::MetadataTraits<IT>::Normalize(Value);
+		}
+
 		SetValue(ItemKey, Value);
 	}
 
diff --git a/Engine/Plugins/PCG/Source/PCG/Public/Metadata/PCGMetadataAttributeTraits.h b/Engine/Plugins/PCG/Source/PCG/Public/Metadata/PCGMetadataAttributeTraits.h
--- a/Engine/Plugins/PCG/Source/PCG/Public/Metadata/PCGMetadataAttributeTraits.h
+++ b/Engine/Plugins/PCG/Source/PCG/Public/Metadata/PCGMetadataAttributeTraits.h
@@ -471,6 +471,7 @@ namespace PCG
 		struct DefaultWeightedSumTraits
 		{
 			enum { CanInterpolate = true };
+			enum { InterpolationNeedsNormalization = false };
 
 			static T WeightedSum(const T& A, const T& B, float Weight)
 			{
@@ -529,6 +530,7 @@ namespace PCG
 			enum { CompressData = false };
 			enum { CanSearchString = false };
 			enum { NeedsConstruction = false };
+			enum { CanNormalize = false };
 
 			static int32 ParseFromString(const FString& InValue, bool& bOutSuccess)
 			{
@@ -545,6 +547,7 @@ namespace PCG
 			enum { CompressData = false };
 			enum { CanSearchString = false };
 			enum { NeedsConstruction = false };
+			enum { CanNormalize = false };
 
 			static int64 ParseFromString(const FString& InValue, bool& bOutSuccess)
 			{
@@ -561,6 +564,7 @@ namespace PCG
 			enum { CanSearchString = false };
 			enum { NeedsConstruction = false };
 			enum { IsFloatingPoint = true };
+			enum { CanNormalize = false };
 
 			static bool Equal(const float& A, const float& B, const double Tolerance = UE_SMALL_NUMBER)
 			{
@@ -582,6 +586,7 @@ namespace PCG
 			enum { CanSearchString = false };
 			enum { NeedsConstruction = false };
 			enum { IsFloatingPoint = true };
+			enum { CanNormalize = false };
 
 			static bool Equal(const double& A, const double& B, const double Tolerance = UE_DOUBLE_SMALL_NUMBER)
 			{
@@ -604,6 +609,7 @@ namespace PCG
 			enum { CanSubAdd = true };
 			enum { CanMulDiv = false };
 			enum { CanInterpolate = false };
+			enum { CanNormalize = false };
 			enum { CanCompare = true };
 			enum { CanFindNearest = false };
 			enum { CanComputeDistance = false };
@@ -675,6 +681,8 @@ namespace PCG
 			enum { CompressData = false };
 			enum { CanMinMax = true };
 			enum { CanInterpolate = true };
+			enum { InterpolationNeedsNormalization = false };
+			enum { CanNormalize = true };
 			enum { CanCompare = true };
 			enum { CanSearchString = false };
 			enum { NeedsConstruction = false };
@@ -756,6 +764,11 @@ namespace PCG
 				bOutSuccess = FDefaultValueHelper::ParseVector2D(InValue, OutValue) || OutValue.InitFromString(InValue);
 				return bOutSuccess ? OutValue : FVector2D::ZeroVector;
 			}
+
+			static void Normalize(FVector2D& InValue, const FVector2D::FReal Tolerance = UE_DOUBLE_SMALL_NUMBER)
+			{
+				InValue.Normalize(Tolerance);
+			}
 		};
 
 		template<>
@@ -803,8 +816,13 @@ namespace PCG
 				bOutSuccess = FDefaultValueHelper::ParseVector(InValue, OutValue) || OutValue.InitFromString(InValue);
 				return bOutSuccess ? OutValue : FVector::ZeroVector;
 			}
+
+			static void Normalize(FVector& InValue, const FVector::FReal Tolerance = UE_DOUBLE_SMALL_NUMBER)
+			{
+				InValue.Normalize(Tolerance);
+			}
 		};
-		
+
 		template<>
 		struct MetadataTraits<FVector4> : VectorTraits<FVector4>, DefaultHashTraits<FVector4>
 		{
@@ -850,6 +868,11 @@ namespace PCG
 				bOutSuccess = FDefaultValueHelper::ParseVector4(InValue, OutValue) || OutValue.InitFromString(InValue);
 				return bOutSuccess ? OutValue : FVector4::Zero();
 			}
+
+			static void Normalize(FVector4& InValue, const FVector4::FReal Tolerance = UE_DOUBLE_SMALL_NUMBER)
+			{
+				InValue = InValue.GetSafeNormal(Tolerance);
+			}
 		};
 
 		// Quaternion
@@ -861,10 +884,12 @@ namespace PCG
 			enum { CanSubAdd = true };
 			enum { CanMulDiv = true };
 			enum { CanInterpolate = true };
+			enum { InterpolationNeedsNormalization = true };
 			enum { CanCompare = false };
 			enum { CanSearchString = false };
 			enum { NeedsConstruction = false };
 			enum { IsFloatingPoint = true };
+			enum { CanNormalize = true };
 
 			using DistanceType = typename FQuat::FReal;
 
@@ -926,6 +951,11 @@ namespace PCG
 				bOutSuccess = FDefaultValueHelper::ParseVector4(InValue, Temp) || Temp.InitFromString(InValue);
 				return bOutSuccess ? FQuat(Temp.X, Temp.Y, Temp.Z, Temp.W) : FQuat::Identity;
 			}
+
+			static void Normalize(FQuat& InValue, const FQuat::FReal Tolerance = UE_DOUBLE_SMALL_NUMBER)
+			{
+				InValue.Normalize(Tolerance);
+			}
 		};
 
 		// Rotator
@@ -937,6 +967,8 @@ namespace PCG
 			enum { CanSubAdd = true };
 			enum { CanMulDiv = true };
 			enum { CanInterpolate = true };
+			enum { CanNormalize = true };
+			enum { InterpolationNeedsNormalization = true };
 			enum { CanCompare = false };
 			enum { CanSearchString = false };
 			enum { NeedsConstruction = false };
@@ -1002,6 +1034,11 @@ namespace PCG
 			{
 				return GetTypeHash(FVector(A.Pitch, A.Yaw, A.Roll));
 			}
+
+			static void Normalize(FRotator& InValue, const FRotator::FReal Tolerance = UE_DOUBLE_SMALL_NUMBER)
+			{
+				InValue.Normalize();
+			}
 		};
 
 		// Transform
@@ -1013,6 +1050,8 @@ namespace PCG
 			enum { CanSubAdd = false };
 			enum { CanMulDiv = true };
 			enum { CanInterpolate = true };
+			enum { CanNormalize = true };
+			enum { InterpolationNeedsNormalization = true };
 			enum { CanCompare = false };
 			enum { CanFindNearest = false };
 			enum { CanComputeDistance = false };
@@ -1035,23 +1074,13 @@ namespace PCG
 				return A * B.Inverse();
 			}
 
-			static FQuat WeightedQuatSum(const FQuat& Q, const FQuat& V, float Weight)
-			{
-				FQuat BlendQuat = V * Weight;
-
-				if ((Q | BlendQuat) >= 0.0f)
-					return Q + BlendQuat;
-				else
-					return Q - BlendQuat;
-			}
-
 			static FTransform WeightedSum(const FTransform& A, const FTransform& B, float Weight)
 			{
 				// WARNING: the rotation won't be normalized
 				return FTransform(
-					WeightedQuatSum(A.GetRotation(), B.GetRotation(), Weight),
-					A.GetLocation() + B.GetLocation() * Weight,
-					A.GetScale3D() + B.GetScale3D() * Weight);
+					MetadataTraits<FQuat>::WeightedSum(A.GetRotation(), B.GetRotation(), Weight),
+					MetadataTraits<FVector>::WeightedSum(A.GetLocation(), B.GetLocation(), Weight),
+					MetadataTraits<FVector>::WeightedSum(A.GetScale3D(), B.GetScale3D(), Weight));
 			}
 
 			static FTransform ZeroValue()
@@ -1070,6 +1099,11 @@ namespace PCG
 				bOutSuccess = OutValue.InitFromString(InValue);
 				return bOutSuccess ? OutValue : FTransform::Identity;
 			}
+
+			static void Normalize(FTransform& InValue, const FTransform::FReal Tolerance = UE_DOUBLE_SMALL_NUMBER)
+			{
+				InValue.SetRotation(InValue.GetRotation().GetNormalized(Tolerance));
+			}
 		};
 
 		// Strings
@@ -1081,6 +1115,8 @@ namespace PCG
 			enum { CanSubAdd = false };
 			enum { CanMulDiv = false };
 			enum { CanInterpolate = false };
+			enum { CanNormalize = false };
+			enum { InterpolationNeedsNormalization = false };
 			enum { CanSearchString = true };
 			enum { NeedsConstruction = true };
 
@@ -1118,6 +1154,8 @@ namespace PCG
 			enum { CanSubAdd = false };
 			enum { CanMulDiv = false };
 			enum { CanInterpolate = false };
+			enum { CanNormalize = false };
+			enum { InterpolationNeedsNormalization = false };
 			enum { CanCompare = true };
 			enum { CanSearchString = true };
 			enum { NeedsConstruction = false };
@@ -1171,6 +1209,8 @@ namespace PCG
 			enum { CanSubAdd = false };
 			enum { CanMulDiv = false };
 			enum { CanInterpolate = false };
+			enum { CanNormalize = false };
+			enum { InterpolationNeedsNormalization = false };
 			enum { CanSearchString = true };
 			// Contains FString objects
 			enum { NeedsConstruction = true };
