[PCG] Fix LOCTEXT macro in PCGExportSelectedAttributes

#rnx

diff --git a/Engine/Plugins/PCG/Source/PCG/Private/Elements/IO/PCGExportSelectedAttributes.cpp b/Engine/Plugins/PCG/Source/PCG/Private/Elements/IO/PCGExportSelectedAttributes.cpp
--- a/Engine/Plugins/PCG/Source/PCG/Private/Elements/IO/PCGExportSelectedAttributes.cpp
+++ b/Engine/Plugins/PCG/Source/PCG/Private/Elements/IO/PCGExportSelectedAttributes.cpp
@@ -141,11 +141,11 @@ namespace PCGExportSelectedAttributes
 			TSharedPtr<FJsonObject> CreateJsonHeaderObject(const bool bAddCustomDataVersion = false, const int32 CustomDataVersion = 0)
 			{
 				TSharedPtr<FJsonObject> JsonHeaderObject = MakeShared<FJsonObject>();
-				JsonHeaderObject->SetNumberField(LOCTEXT("JsonFieldVersionNumber", "Export Version").ToString(), Constants::FCustomExportVersion::LatestVersion);
+				JsonHeaderObject->SetNumberField(LOCTEXT("JsonFieldExportVersionNumber", "Export Version").ToString(), Constants::FCustomExportVersion::LatestVersion);
 
 				if (bAddCustomDataVersion)
 				{
-					JsonHeaderObject->SetNumberField(LOCTEXT("JsonFieldVersionNumber", "Custom Data Version").ToString(), CustomDataVersion);
+					JsonHeaderObject->SetNumberField(LOCTEXT("JsonFieldCustomVersionNumber", "Custom Data Version").ToString(), CustomDataVersion);
 				}
 
 				return JsonHeaderObject;
