[PCG] Fix crash occurring when multiple input edges added to node with GPU execution enabled.

#rnx

diff --git a/Engine/Plugins/PCG/Source/PCG/Private/Elements/PCGGather.cpp b/Engine/Plugins/PCG/Source/PCG/Private/Elements/PCGGather.cpp
--- a/Engine/Plugins/PCG/Source/PCG/Private/Elements/PCGGather.cpp
+++ b/Engine/Plugins/PCG/Source/PCG/Private/Elements/PCGGather.cpp
@@ -89,18 +89,20 @@ bool FPCGGatherElement::ExecuteInternal(FPCGContext* Context) const
 {
 	TRACE_CPUPROFILER_EVENT_SCOPE(FPCGGatherElement::Execute);
 
-	FPCGDataCollection GatheredData;
-	const UPCGGatherSettings* Settings = Context->GetInputSettings<UPCGGatherSettings>();
-	check(Settings);
-
-	for (FName PinLabel : Settings->GetNodeDefinedPinLabels())
+	if (const UPCGGatherSettings* Settings = Context->GetInputSettings<UPCGGatherSettings>())
+	{
+		for (FName PinLabel : Settings->GetNodeDefinedPinLabels())
+		{
+			Context->OutputData += PCGGather::GatherDataForPin(Context->InputData, PinLabel);
+		}
+	}
+	else
 	{
-		GatheredData += PCGGather::GatherDataForPin(Context->InputData, PinLabel);
+		// If running settings-less, just gather from default primary input pin.
+		Context->OutputData = PCGGather::GatherDataForPin(Context->InputData, PCGPinConstants::DefaultInputLabel);
 	}
 
-	Context->OutputData = GatheredData;
-
 	return true;
 }
 
-#undef LOCTEXT_NAMESPACE
\ No newline at end of file
+#undef LOCTEXT_NAMESPACE
diff --git a/Engine/Plugins/PCG/Source/PCG/Private/Graph/PCGGraphCompilerGPU.cpp b/Engine/Plugins/PCG/Source/PCG/Private/Graph/PCGGraphCompilerGPU.cpp
--- a/Engine/Plugins/PCG/Source/PCG/Private/Graph/PCGGraphCompilerGPU.cpp
+++ b/Engine/Plugins/PCG/Source/PCG/Private/Graph/PCGGraphCompilerGPU.cpp
@@ -485,8 +485,8 @@ void FPCGGraphCompilerGPU::CreateGatherTasksAtGPUInputs(UPCGGraph* InGraph, cons
 	// Add all compute graph task inputs and outputs.
 	for (FPCGTaskId GPUTaskId : InGPUCompatibleTaskIds)
 	{
-		EncounteredInputPins.Reset();
-		InputPinToGatherTask.Reset();
+		EncounteredInputPins.Empty(EncounteredInputPins.Num());
+		InputPinToGatherTask.Empty(InputPinToGatherTask.Num());
 
 		// First pass - create gather tasks for all original input pins which have more than one incident edge.
 		// This is so we can gather on the CPU (much more efficient than going it on the GPU).
@@ -530,41 +530,46 @@ void FPCGGraphCompilerGPU::CreateGatherTasksAtGPUInputs(UPCGGraph* InGraph, cons
 			InputPinToGatherTask.Add(PinKey, GatherTaskId);
 		}
 
-		EncounteredInputPins.Reset();
+		EncounteredInputPins.Empty(EncounteredInputPins.Num());
 
 		// Second pass - wire up the newly added gather tasks once we have the full picture of which edges are affected.
 		for (int InputIndex = 0; InputIndex < InOutCompiledTasks[GPUTaskId].Inputs.Num(); ++InputIndex)
 		{
-			auto CurrentInput = [&InOutCompiledTasks, InputIndex, GPUTaskId]() -> FPCGGraphTaskInput&
+			// Helper to get input rather than reference because we modify the inputs array below.
+			auto CurrentGPUTaskInput = [&InOutCompiledTasks, InputIndex, GPUTaskId]() -> FPCGGraphTaskInput&
 			{
 				return InOutCompiledTasks[GPUTaskId].Inputs[InputIndex];
 			};
 
-			if (!CurrentInput().DownstreamPin.IsSet())
+			if (!CurrentGPUTaskInput().DownstreamPin)
 			{
 				continue;
 			}
 
-			const FOriginalInputPinKey PinKey = { GPUTaskId, CurrentInput().DownstreamPin->Label };
+			const FOriginalInputPinKey PinKey = { GPUTaskId, CurrentGPUTaskInput().DownstreamPin->Label };
 
 			if (const FPCGTaskId* GatherTaskId = InputPinToGatherTask.Find(PinKey))
 			{
-				// Wire the upstream output pin to the gather task.
-				FPCGGraphTaskInput& WireUpstreamNodeToGather = InOutCompiledTasks[*GatherTaskId].Inputs.Add_GetRef(CurrentInput());
-				if (WireUpstreamNodeToGather.DownstreamPin.IsSet())
+				// Wire the upstream output pin to the gather task input.
 				{
-					WireUpstreamNodeToGather.DownstreamPin->Label = PCGPinConstants::DefaultInputLabel;
+					FPCGGraphTaskInput& GatherTaskInput = InOutCompiledTasks[*GatherTaskId].Inputs.Add_GetRef(CurrentGPUTaskInput());
+					if (GatherTaskInput.DownstreamPin)
+					{
+						GatherTaskInput.DownstreamPin->Label = PCGPinConstants::DefaultInputLabel;
+					}
 				}
 
 				if (!EncounteredInputPins.Contains(PinKey))
 				{
-					// First time we're encountering this input pin, wire it to the gather task.
+					// First time we're encountering this input pin, wire it to the gather task output.
 					EncounteredInputPins.Add(PinKey);
 
-					CurrentInput().TaskId = *GatherTaskId;
-					if (CurrentInput().UpstreamPin.IsSet())
+					CurrentGPUTaskInput().TaskId = *GatherTaskId;
+
+					// Wire the gather task output pin to the downstream GPU task input pin.
+					if (CurrentGPUTaskInput().UpstreamPin)
 					{
-						CurrentInput().UpstreamPin->Label = PCGPinConstants::DefaultOutputLabel;
+						CurrentGPUTaskInput().UpstreamPin->Label = PCGPinConstants::DefaultOutputLabel;
 					}
 				}
 				else
