[PCG] Attribute Delete refactor

* Don't distinguish between data, use polymorphism for that
* Add option to specify which domain to target

#rnx

diff --git a/Engine/Plugins/PCG/Source/PCG/Private/Elements/PCGDeleteAttributesElement.cpp b/Engine/Plugins/PCG/Source/PCG/Private/Elements/PCGDeleteAttributesElement.cpp
--- a/Engine/Plugins/PCG/Source/PCG/Private/Elements/PCGDeleteAttributesElement.cpp
+++ b/Engine/Plugins/PCG/Source/PCG/Private/Elements/PCGDeleteAttributesElement.cpp
@@ -159,38 +159,57 @@ bool FPCGDeleteAttributesElement::ExecuteInternal(FPCGContext* Context) const
 	for (const FPCGTaggedData& InputTaggedData : Inputs)
 	{
 		const UPCGData* InputData = InputTaggedData.Data;
-		UPCGData* OutputData = nullptr;
 
-		const UPCGMetadata* ParentMetadata = nullptr;
-		UPCGMetadata* Metadata = nullptr;
-
-		if (const UPCGSpatialData* InputSpatialData = Cast<UPCGSpatialData>(InputData))
+		const UPCGMetadata* InputMetadata = InputData->ConstMetadata();
+		if (!InputMetadata)
 		{
-			ParentMetadata = InputSpatialData->Metadata;
-
-			UPCGSpatialData* NewSpatialData = InputSpatialData->DuplicateData(Context, /*bInitializeFromThisData=*/false);
-			Metadata = NewSpatialData->Metadata;
-			NewSpatialData->Metadata->InitializeWithAttributeFilter(ParentMetadata, AttributesToFilter, FilterMode, Settings->Operator);
+			PCGLog::Metadata::LogInvalidMetadata(Context);
+			continue;
+		}
 
-			// No need to inherit metadata since we already initialized it.
-			NewSpatialData->InitializeFromData(InputSpatialData, /*InMetadataParentOverride=*/ nullptr, /*bInheritMetadata=*/ false);
+		const FPCGAttributePropertySelector DomainSelector = FPCGAttributePropertySelector::CreateAttributeSelector(NAME_None, Settings->MetadataDomain);
+		const FPCGMetadataDomainID SelectedDomainID = InputData->GetMetadataDomainIDFromSelector(DomainSelector);
+		const FPCGMetadataDomain* InputMetadataDomain = InputMetadata->GetConstMetadataDomain(SelectedDomainID);
 
-			OutputData = NewSpatialData;
+		if (!InputMetadataDomain)
+		{
+			PCGLog::Metadata::LogInvalidMetadataDomain(DomainSelector, Context);
+			continue;
 		}
-		else if (const UPCGParamData* InputParamData = Cast<UPCGParamData>(InputData))
+
+		UPCGData* OutputData = InputData->DuplicateData(Context, /*bInitializeMetadata=*/ false);
+		if (!ensure(OutputData))
 		{
-			ParentMetadata = InputParamData->Metadata;
+			continue;
+		}
 
-			UPCGParamData* NewParamData = FPCGContext::NewObject_AnyThread<UPCGParamData>(Context);
-			Metadata = NewParamData->Metadata;
-			Metadata->InitializeAsCopy(FPCGMetadataInitializeParams(ParentMetadata, AttributesToFilter, FilterMode));
+		UPCGMetadata* OutputMetadata = OutputData->MutableMetadata();
 
-			OutputData = NewParamData;
-		}
-		else
+		for (const FPCGMetadataDomainID DomainID : InputData->GetAllSupportedMetadataDomainIDs())
 		{
-			PCGE_LOG(Error, GraphAndLog, LOCTEXT("InvalidInputData", "Invalid data as input. Only Spatial and Attribute Set data are supported."));
-			continue;
+			const FPCGMetadataDomain* InDomain = InputMetadata->GetConstMetadataDomain(DomainID);
+			if (!InDomain)
+			{
+				// Not having an input domain is fine, it just means you don't have anything there.
+				continue;
+			}
+			
+			FPCGMetadataDomain* OutDomain = OutputMetadata->GetMetadataDomain(DomainID);
+			check(OutDomain);
+			
+			if (InputMetadataDomain != InDomain)
+			{
+				OutDomain->Initialize(InDomain);
+			}
+			else
+			{
+				FPCGMetadataDomainInitializeParams Params(InputMetadataDomain);
+				Params.FilterMode = FilterMode;
+				Params.MatchOperator = Settings->Operator;
+				Params.FilteredAttributes = AttributesToFilter;
+
+				OutDomain->Initialize(Params);
+			}
 		}
 
 		TArray<FPCGTaggedData>& Outputs = Context->OutputData.TaggedData;
diff --git a/Engine/Plugins/PCG/Source/PCG/Private/Metadata/PCGMetadataDomain.cpp b/Engine/Plugins/PCG/Source/PCG/Private/Metadata/PCGMetadataDomain.cpp
--- a/Engine/Plugins/PCG/Source/PCG/Private/Metadata/PCGMetadataDomain.cpp
+++ b/Engine/Plugins/PCG/Source/PCG/Private/Metadata/PCGMetadataDomain.cpp
@@ -326,6 +326,12 @@ void FPCGMetadataDomain::Initialize(const FPCGMetadataDomainInitializeParams& In
 		return;
 	}
 
+	// Make sure that the parent of the top metadata is also set correctly
+	if (TopMetadata->Parent == nullptr && InParams.Parent)
+	{
+		TopMetadata->Parent = InParams.Parent->TopMetadata;
+	}
+
 	Parent = ((InParams.Parent != this) ? InParams.Parent : nullptr);
 	ItemKeyOffset = Parent ? Parent->GetItemCountForChild() : 0;
 
diff --git a/Engine/Plugins/PCG/Source/PCG/Private/Tests/Elements/Metadata/PCGDeleteAttributesTests.cpp b/Engine/Plugins/PCG/Source/PCG/Private/Tests/Elements/Metadata/PCGDeleteAttributesTests.cpp
new file mode 100644
--- /dev/null
+++ b/Engine/Plugins/PCG/Source/PCG/Private/Tests/Elements/Metadata/PCGDeleteAttributesTests.cpp
@@ -0,0 +1,319 @@
+// Copyright Epic Games, Inc. All Rights Reserved.
+
+#include "PCGContext.h"
+#include "Elements/PCGDeleteAttributesElement.h"
+#if WITH_EDITOR
+
+#include "Tests/PCGTestsCommon.h"
+
+#include "PCGParamData.h"
+#include "Data/PCGSplineData.h"
+#include "Data/PCGPointData.h"
+#include "Metadata/PCGMetadata.h"
+
+template <typename DataType>
+class FPCGDeleteAttributesTests : public FPCGTestBaseClass
+{
+public:
+	using FPCGTestBaseClass::FPCGTestBaseClass;
+	
+protected:
+	static inline const FName AttributeName1 = TEXT("MyAttr123");
+	static inline const FName AttributeName2 = TEXT("MyAttr23");
+	static inline const FName AttributeName3 = TEXT("MyAttr3");
+	static inline const TStaticArray<FName, 3> DefaultAttributeNames = {AttributeName1, AttributeName2, AttributeName3};
+	
+	static UPCGData* CreateData()
+	{
+		if constexpr (std::is_same_v<DataType, UPCGBasePointData>)
+		{
+			return FPCGContext::NewPointData_AnyThread(nullptr);
+		}
+		else
+		{
+			return NewObject<DataType>();
+		}
+	}
+
+	static void CreateAttributes(UPCGData* Data, const FPCGMetadataDomainID DomainID, TArrayView<const FName> OptionalNames = {})
+	{
+		check(Data);
+		UPCGMetadata* Metadata = Data->MutableMetadata();
+		check(Metadata);
+		FPCGMetadataDomain* MetadataDomain = Metadata->GetMetadataDomain(DomainID);
+		check(MetadataDomain);
+
+		TArrayView<const FName> AttributeNames = OptionalNames.IsEmpty() ? DefaultAttributeNames : OptionalNames;
+
+		for (int32 i = 0; i < AttributeNames.Num(); ++i)
+		{
+			MetadataDomain->CreateAttribute<int32>(AttributeNames[i], i, false, false);
+		}
+	}
+	
+	const UPCGData* ExecuteElement(const UPCGData* InputData, UPCGDeleteAttributesSettings* Settings)
+	{
+		PCGTestsCommon::FTestData TestData(42, Settings);
+		
+		// Also push the settings
+		TestData.InputData.TaggedData.Emplace_GetRef().Data = TestData.Settings;
+		TestData.InputData.TaggedData.Last().Pin = TEXT("Settings");
+		
+		FPCGElementPtr TestElement = TestData.Settings->GetElement();
+		TUniquePtr<FPCGContext> Context = TestData.InitializeTestContext();
+
+		FPCGTaggedData& NewTaggedData = Context->InputData.TaggedData.Emplace_GetRef();
+		NewTaggedData.Data = InputData;
+		NewTaggedData.Pin = PCGPinConstants::DefaultInputLabel;
+
+
+		while (!TestElement->Execute(Context.Get())) {}
+
+		if (TestTrue(TEXT("Element execution has 1 output"), Context->OutputData.TaggedData.Num() == 1)
+			&& TestNotNull(TEXT("Element execution output data is valid"), Context->OutputData.TaggedData[0].Data.Get()))
+		{
+			return Context->OutputData.TaggedData[0].Data.Get();
+		}
+		else
+		{
+			return nullptr;
+		}
+	}
+
+	bool ValidateAttributes(const UPCGData* Data, const FPCGMetadataDomainID DomainID, TArrayView<const FName> ExpectedNames)
+	{
+		if (!Data)
+		{
+			return false;
+		}
+
+		const UPCGMetadata* Metadata = Data->ConstMetadata();
+		check(Metadata);
+		const FPCGMetadataDomain* MetadataDomain = Metadata->GetConstMetadataDomain(DomainID);
+		check(MetadataDomain);
+
+		TArray<FName> AttributeNames;
+		TArray<EPCGMetadataTypes> AttributeTypes;
+		MetadataDomain->GetAttributes(AttributeNames, AttributeTypes);
+
+		for (const FName ExpectedName : ExpectedNames)
+		{
+			UTEST_TRUE(*FString::Printf(TEXT("Attribute %s should exist in domain %s"), *ExpectedName.ToString(), *DomainID.DebugName.ToString()), AttributeNames.Contains(ExpectedName));
+		}
+
+		for (const FName AttributeName : AttributeNames)
+		{
+			UTEST_TRUE(*FString::Printf(TEXT("Attribute %s is expected in domain %s"), *AttributeName.ToString(), *DomainID.DebugName.ToString()), ExpectedNames.Contains(AttributeName));
+		}
+
+		return true;
+	}
+};
+
+IMPLEMENT_CUSTOM_SIMPLE_AUTOMATION_TEST(FPCGDeleteAttributesTests_Points_DeleteSelected, FPCGDeleteAttributesTests<UPCGBasePointData>, "Plugins.PCG.DeleteAttributes.Points.DeleteSelected", PCGTestsCommon::TestFlags)
+IMPLEMENT_CUSTOM_SIMPLE_AUTOMATION_TEST(FPCGDeleteAttributesTests_Points_KeepSelected, FPCGDeleteAttributesTests<UPCGBasePointData>, "Plugins.PCG.DeleteAttributes.Points.KeepSelected", PCGTestsCommon::TestFlags)
+IMPLEMENT_CUSTOM_SIMPLE_AUTOMATION_TEST(FPCGDeleteAttributesTests_Points_KeepSelectedEmpty, FPCGDeleteAttributesTests<UPCGBasePointData>, "Plugins.PCG.DeleteAttributes.Points.KeepSelectedEmpty", PCGTestsCommon::TestFlags)
+IMPLEMENT_CUSTOM_SIMPLE_AUTOMATION_TEST(FPCGDeleteAttributesTests_Points_DeleteSelectedSubstring, FPCGDeleteAttributesTests<UPCGBasePointData>, "Plugins.PCG.DeleteAttributes.Points.DeleteSelectedSubstring", PCGTestsCommon::TestFlags)
+IMPLEMENT_CUSTOM_SIMPLE_AUTOMATION_TEST(FPCGDeleteAttributesTests_Points_DeleteSelectedMatches, FPCGDeleteAttributesTests<UPCGBasePointData>, "Plugins.PCG.DeleteAttributes.Points.DeleteSelectedMatches", PCGTestsCommon::TestFlags)
+IMPLEMENT_CUSTOM_SIMPLE_AUTOMATION_TEST(FPCGDeleteAttributesTests_Points_KeepSelectedSubstring, FPCGDeleteAttributesTests<UPCGBasePointData>, "Plugins.PCG.DeleteAttributes.Points.KeepSelectedSubstring", PCGTestsCommon::TestFlags)
+IMPLEMENT_CUSTOM_SIMPLE_AUTOMATION_TEST(FPCGDeleteAttributesTests_Points_KeepSelectedMatches, FPCGDeleteAttributesTests<UPCGBasePointData>, "Plugins.PCG.DeleteAttributes.Points.KeepSelectedMatches", PCGTestsCommon::TestFlags)
+IMPLEMENT_CUSTOM_SIMPLE_AUTOMATION_TEST(FPCGDeleteAttributesTests_Points_DeleteSelectedData, FPCGDeleteAttributesTests<UPCGBasePointData>, "Plugins.PCG.DeleteAttributes.Points.DeleteSelectedData", PCGTestsCommon::TestFlags)
+IMPLEMENT_CUSTOM_SIMPLE_AUTOMATION_TEST(FPCGDeleteAttributesTests_Points_DeleteSelectedElements, FPCGDeleteAttributesTests<UPCGBasePointData>, "Plugins.PCG.DeleteAttributes.Points.DeleteSelectedElements", PCGTestsCommon::TestFlags)
+
+/**
+ * Delete
+ * Attributes In: MyAttr123, MyAttr23, MyAttr3
+ * Selection: "MyAttr123,MyAttr23"
+ * Expected out: MyAttr3
+ */
+bool FPCGDeleteAttributesTests_Points_DeleteSelected::RunTest(const FString& Parameters)
+{
+	UPCGData* InputData = CreateData();
+	CreateAttributes(InputData, PCGMetadataDomainID::Default);
+
+	UPCGDeleteAttributesSettings* Settings = NewObject<UPCGDeleteAttributesSettings>();
+	Settings->Operation = EPCGAttributeFilterOperation::DeleteSelectedAttributes;
+	Settings->Operator = EPCGStringMatchingOperator::Equal;
+	Settings->MetadataDomain = PCGDataConstants::DefaultDomainName;
+	Settings->SelectedAttributes = AttributeName1.ToString() + TEXT(",") + AttributeName2.ToString();
+	
+	return ValidateAttributes(ExecuteElement(InputData, Settings), PCGMetadataDomainID::Default, {AttributeName3});
+}
+
+/**
+ * Keep
+ * Attributes In: MyAttr123, MyAttr23, MyAttr3
+ * Selection: "MyAttr123,MyAttr23"
+ * Expected out: MyAttr123,MyAttr23
+ */
+bool FPCGDeleteAttributesTests_Points_KeepSelected::RunTest(const FString& Parameters)
+{
+	UPCGData* InputData = CreateData();
+	CreateAttributes(InputData, PCGMetadataDomainID::Default);
+
+	UPCGDeleteAttributesSettings* Settings = NewObject<UPCGDeleteAttributesSettings>();
+	Settings->Operation = EPCGAttributeFilterOperation::KeepSelectedAttributes;
+	Settings->Operator = EPCGStringMatchingOperator::Equal;
+	Settings->MetadataDomain = PCGDataConstants::DefaultDomainName;
+	Settings->SelectedAttributes = AttributeName1.ToString() + TEXT(",") + AttributeName2.ToString();
+	
+	return ValidateAttributes(ExecuteElement(InputData, Settings), PCGMetadataDomainID::Default, {AttributeName1, AttributeName2});
+}
+
+/**
+ * Keep
+ * Attributes In: MyAttr123, MyAttr23, MyAttr3
+ * Selection: ""
+ * Expected out: Nothing
+ */
+bool FPCGDeleteAttributesTests_Points_KeepSelectedEmpty::RunTest(const FString& Parameters)
+{
+	UPCGData* InputData = CreateData();
+	CreateAttributes(InputData, PCGMetadataDomainID::Default);
+
+	UPCGDeleteAttributesSettings* Settings = NewObject<UPCGDeleteAttributesSettings>();
+	Settings->Operation = EPCGAttributeFilterOperation::KeepSelectedAttributes;
+	Settings->Operator = EPCGStringMatchingOperator::Equal;
+	Settings->MetadataDomain = PCGDataConstants::DefaultDomainName;
+	Settings->SelectedAttributes = FString();
+	
+	return ValidateAttributes(ExecuteElement(InputData, Settings), PCGMetadataDomainID::Default, {});
+}
+
+/**
+ * Delete substring
+ * Attributes In: MyAttr123, MyAttr23, MyAttr3
+ * Selection: "MyAttr1"
+ * Expected out: MyAttr23, MyAttr3
+ */
+bool FPCGDeleteAttributesTests_Points_DeleteSelectedSubstring::RunTest(const FString& Parameters)
+{
+	UPCGData* InputData = CreateData();
+	CreateAttributes(InputData, PCGMetadataDomainID::Default);
+
+	UPCGDeleteAttributesSettings* Settings = NewObject<UPCGDeleteAttributesSettings>();
+	Settings->Operation = EPCGAttributeFilterOperation::DeleteSelectedAttributes;
+	Settings->Operator = EPCGStringMatchingOperator::Substring;
+	Settings->MetadataDomain = PCGDataConstants::DefaultDomainName;
+	Settings->SelectedAttributes = TEXT("MyAttr1");
+	
+	return ValidateAttributes(ExecuteElement(InputData, Settings), PCGMetadataDomainID::Default, {AttributeName2, AttributeName3});
+}
+
+/**
+ * Delete matches
+ * Attributes In: MyAttr123, MyAttr23, MyAttr3
+ * Selection: "MyAttr*3"
+ * Expected out:  Nothing
+ */
+bool FPCGDeleteAttributesTests_Points_DeleteSelectedMatches::RunTest(const FString& Parameters)
+{
+	UPCGData* InputData = CreateData();
+	CreateAttributes(InputData, PCGMetadataDomainID::Default);
+
+	UPCGDeleteAttributesSettings* Settings = NewObject<UPCGDeleteAttributesSettings>();
+	Settings->Operation = EPCGAttributeFilterOperation::DeleteSelectedAttributes;
+	Settings->Operator = EPCGStringMatchingOperator::Matches;
+	Settings->MetadataDomain = PCGDataConstants::DefaultDomainName;
+	Settings->SelectedAttributes = TEXT("MyAttr*3");
+	
+	return ValidateAttributes(ExecuteElement(InputData, Settings), PCGMetadataDomainID::Default, {});
+}
+
+/**
+ * Kepp substring
+ * Attributes In: MyAttr123, MyAttr23, MyAttr3
+ * Selection: "MyAttr1"
+ * Expected out: MyAttr123
+ */
+bool FPCGDeleteAttributesTests_Points_KeepSelectedSubstring::RunTest(const FString& Parameters)
+{
+	UPCGData* InputData = CreateData();
+	CreateAttributes(InputData, PCGMetadataDomainID::Default);
+
+	UPCGDeleteAttributesSettings* Settings = NewObject<UPCGDeleteAttributesSettings>();
+	Settings->Operation = EPCGAttributeFilterOperation::KeepSelectedAttributes;
+	Settings->Operator = EPCGStringMatchingOperator::Substring;
+	Settings->MetadataDomain = PCGDataConstants::DefaultDomainName;
+	Settings->SelectedAttributes = TEXT("MyAttr1");
+	
+	return ValidateAttributes(ExecuteElement(InputData, Settings), PCGMetadataDomainID::Default, {AttributeName1});
+}
+
+/**
+ * Keep matches
+ * Attributes In: MyAttr123, MyAttr23, MyAttr3
+ * Selection: "MyAttr*3"
+ * Expected out:  MyAttr123, MyAttr23, MyAttr3
+ */
+bool FPCGDeleteAttributesTests_Points_KeepSelectedMatches::RunTest(const FString& Parameters)
+{
+	UPCGData* InputData = CreateData();
+	CreateAttributes(InputData, PCGMetadataDomainID::Default);
+
+	UPCGDeleteAttributesSettings* Settings = NewObject<UPCGDeleteAttributesSettings>();
+	Settings->Operation = EPCGAttributeFilterOperation::KeepSelectedAttributes;
+	Settings->Operator = EPCGStringMatchingOperator::Matches;
+	Settings->MetadataDomain = PCGDataConstants::DefaultDomainName;
+	Settings->SelectedAttributes = TEXT("MyAttr*3");
+	
+	return ValidateAttributes(ExecuteElement(InputData, Settings), PCGMetadataDomainID::Default, {AttributeName1, AttributeName2, AttributeName3});
+}
+
+/**
+ * Delete Data
+ * Attributes In:
+ *    - Elements: MyAttr123, MyAttr23, MyAttr3
+ *    - Data: MyAttr123, MyAttr23, MyAttr3
+ * Selection: "MyAttr123"
+ * Expected out: 
+ *    - Elements: MyAttr123, MyAttr23, MyAttr3
+ *    - Data: MyAttr23, MyAttr3
+ */
+bool FPCGDeleteAttributesTests_Points_DeleteSelectedData::RunTest(const FString& Parameters)
+{
+	UPCGData* InputData = CreateData();
+	CreateAttributes(InputData, PCGMetadataDomainID::Elements);
+	CreateAttributes(InputData, PCGMetadataDomainID::Data);
+
+	UPCGDeleteAttributesSettings* Settings = NewObject<UPCGDeleteAttributesSettings>();
+	Settings->Operation = EPCGAttributeFilterOperation::DeleteSelectedAttributes;
+	Settings->Operator = EPCGStringMatchingOperator::Equal;
+	Settings->MetadataDomain = PCGDataConstants::DataDomainName;
+	Settings->SelectedAttributes = AttributeName1.ToString();
+
+	const UPCGData* OutputData = ExecuteElement(InputData, Settings);
+
+	return ValidateAttributes(OutputData, PCGMetadataDomainID::Data, {AttributeName2, AttributeName3})
+		&& ValidateAttributes(OutputData, PCGMetadataDomainID::Elements, {AttributeName1, AttributeName2, AttributeName3});
+}
+
+/**
+ * Delete Elements
+ * Attributes In:
+ *    - Elements: MyAttr123, MyAttr23, MyAttr3
+ *    - Data: MyAttr123, MyAttr23, MyAttr3
+ * Selection: "MyAttr123,MyAttr23"
+ * Expected out: 
+ *    - Elements: MyAttr3
+ *    - Data: MyAttr123, MyAttr23, MyAttr3
+ */
+bool FPCGDeleteAttributesTests_Points_DeleteSelectedElements::RunTest(const FString& Parameters)
+{
+	UPCGData* InputData = CreateData();
+	CreateAttributes(InputData, PCGMetadataDomainID::Elements);
+	CreateAttributes(InputData, PCGMetadataDomainID::Data);
+
+	UPCGDeleteAttributesSettings* Settings = NewObject<UPCGDeleteAttributesSettings>();
+	Settings->Operation = EPCGAttributeFilterOperation::DeleteSelectedAttributes;
+	Settings->Operator = EPCGStringMatchingOperator::Equal;
+	Settings->MetadataDomain = PCGPointDataConstants::ElementsDomainName;
+	Settings->SelectedAttributes = AttributeName1.ToString() + TEXT(",") + AttributeName2.ToString();
+
+	const UPCGData* OutputData = ExecuteElement(InputData, Settings);
+
+	return ValidateAttributes(OutputData, PCGMetadataDomainID::Data, {AttributeName1, AttributeName2, AttributeName3})
+		&& ValidateAttributes(OutputData, PCGMetadataDomainID::Elements, {AttributeName3});
+}
+
+#endif // WITH_EDITOR
\ No newline at end of file
diff --git a/Engine/Plugins/PCG/Source/PCG/Public/Elements/PCGDeleteAttributesElement.h b/Engine/Plugins/PCG/Source/PCG/Public/Elements/PCGDeleteAttributesElement.h
--- a/Engine/Plugins/PCG/Source/PCG/Public/Elements/PCGDeleteAttributesElement.h
+++ b/Engine/Plugins/PCG/Source/PCG/Public/Elements/PCGDeleteAttributesElement.h
@@ -69,6 +69,10 @@ public:
 	UPROPERTY(EditAnywhere, Category = Settings, meta = (EditCondition = "bTokenizeOnWhiteSpace", EditConditionHides, DeprecationMessage = "bTokenizeOnWhiteSpace has been deprecated."))
 	bool bTokenizeOnWhiteSpace = false;
 
+	/** When deleting attributes, it only target a single domain that can be specified here. */
+	UPROPERTY(BlueprintReadWrite, EditAnywhere, Category = Settings, meta = (PCG_Overridable))
+	FName MetadataDomain = PCGDataConstants::DefaultDomainName;
+
 #if WITH_EDITORONLY_DATA
 	UPROPERTY()
 	TArray<FName> AttributesToKeep_DEPRECATED;
