[PCG] Data Count Node - Enable GPU version

diff --git a/Engine/Plugins/PCG/Shaders/Private/BuiltInKernels/PCGDataNum.usf b/Engine/Plugins/PCG/Shaders/Private/BuiltInKernels/PCGDataNum.usf
new file mode 100644
--- /dev/null
+++ b/Engine/Plugins/PCG/Shaders/Private/BuiltInKernels/PCGDataNum.usf
@@ -0,0 +1,31 @@
+// Copyright Epic Games, Inc. All Rights Reserved.
+
+// Returns the total data count
+[numthreads(64, 1, 1)]
+void PCGDataNumCS(uint3 GroupId : SV_GroupID, uint GroupIndex : SV_GroupIndex)
+{
+	// Mark the kernel as having executed. Must run before we early out via thread index, because the kernel is still 'executed' even if the number of
+	// threads to iterate on is zero. Even if GetNumThreads() returns 0, the kernel will still have been dispatched on a single thread to set this flag.
+	if (all(GroupId == 0) && GroupIndex == 0)
+	{
+		Out_SetAsExecutedInternal();
+	}
+
+	const uint ThreadIndex = GetUnWrappedDispatchThreadId(GroupId, GroupIndex, 64);
+	if (ThreadIndex >= GetNumThreads().x) return;
+
+	uint InDataIndex, InElemIndex;
+	if (!In_GetThreadData(ThreadIndex, InDataIndex, InElemIndex))
+	{
+		return;
+	}
+
+	uint OutDataIndex, OutElemIndex;
+	if (!Out_GetThreadData(ThreadIndex, OutDataIndex, OutElemIndex))
+	{
+		return;
+	}
+
+	const int AttributeId = DataNum_GetOutputAttributeId();
+	Out_SetInt(OutDataIndex, OutElemIndex, AttributeId, (int)In_GetNumData());
+}
diff --git a/Engine/Plugins/PCG/Source/PCG/Private/Compute/BuiltInKernels/PCGDataNumKernel.cpp b/Engine/Plugins/PCG/Source/PCG/Private/Compute/BuiltInKernels/PCGDataNumKernel.cpp
new file mode 100644
--- /dev/null
+++ b/Engine/Plugins/PCG/Source/PCG/Private/Compute/BuiltInKernels/PCGDataNumKernel.cpp
@@ -0,0 +1,126 @@
+// Copyright Epic Games, Inc. All Rights Reserved.
+
+#include "PCGDataNumKernel.h"
+
+#include "PCGContext.h"
+#include "Compute/DataInterfaces/BuiltInKernels/PCGDataNumDataInterface.h"
+#include "Compute/Elements/PCGComputeGraphElement.h"
+#include "Graph/PCGGPUGraphCompilationContext.h"
+
+#include "ShaderCompilerCore.h"
+
+#include UE_INLINE_GENERATED_CPP_BY_NAME(PCGDataNumKernel)
+
+#define LOCTEXT_NAMESPACE "PCGDataNumKernel"
+
+bool UPCGDataNumKernel::ComputeOutputBindingDataDesc(const UPCGComputeGraph* InGraph, const FName InOutputPinLabel, UPCGDataBinding* InBinding, FPCGDataCollectionDesc& OutDataDesc) const
+{
+	check(InGraph);
+	check(InBinding);
+
+	// Only single output pin
+	if (!ensure(InOutputPinLabel == PCGPinConstants::DefaultOutputLabel))
+	{
+		return false;
+	}
+
+	FPCGDataDesc& OutputDataDesc = OutDataDesc.DataDescs.Emplace_GetRef(EPCGDataType::Param, 1);
+	OutputDataDesc.AddAttribute(OutCountAttributeKey, InBinding);
+
+	return true;
+}
+
+int UPCGDataNumKernel::ComputeThreadCount(const UPCGDataBinding* InBinding) const
+{
+	// One element returned. Only need 1 thread.
+	return 1;
+}
+
+bool UPCGDataNumKernel::IsKernelDataValid(FPCGContext* InContext) const
+{
+	TRACE_CPUPROFILER_EVENT_SCOPE(UPCGDataNumKernel::IsKernelDataValid);
+
+	if (!Super::IsKernelDataValid(InContext))
+	{
+		return false;
+	}
+
+	if (!InContext || !InContext->IsComputeContext())
+	{
+		return true;
+	}
+
+	const FPCGComputeGraphContext* Context = static_cast<FPCGComputeGraphContext*>(InContext);
+	if (const UPCGDataBinding* DataBinding = Context->DataBinding.Get())
+	{
+		if (const FPCGDataCollectionDesc* OutputDataDesc = DataBinding->GetCachedKernelPinDataDesc(this, PCGPinConstants::DefaultOutputLabel, /*bIsInputPin=*/false))
+		{
+			FPCGKernelAttributeDesc AttributeDesc;
+			bool bConflictingTypesFound = false;
+			if (OutputDataDesc->GetAttributeDesc(OutCountAttributeKey.Name, AttributeDesc, bConflictingTypesFound))
+			{
+				if (OutCountAttributeKey.Type == EPCGKernelAttributeType::Invalid)
+				{
+					PCG_KERNEL_VALIDATION_ERR(InContext, Settings, FText::Format(
+						LOCTEXT("InvalidAttributeKeyType", "Output attribute for the count could not be created with an 'invalid' type."),
+						FText::FromName(OutCountAttributeKey.Name)));
+
+					return false;
+				}
+
+				return true;
+			}
+			else
+			{
+				PCG_KERNEL_VALIDATION_ERR(InContext, Settings, FText::Format(
+					LOCTEXT("AttributeMissing", "Output attribute '{0}' not found."),
+					FText::FromName(OutCountAttributeKey.Name)));
+
+				return false;
+			}
+		}
+	}
+
+	return false;
+}
+
+#if WITH_EDITOR
+FString UPCGDataNumKernel::GetCookedSource(const UPCGComputeGraph* InComputeGraph) const
+{
+	FString DataNumFile;
+	ensure(LoadShaderSourceFile(TEXT("/Plugin/PCG/Private/BuiltInKernels/PCGDataNum.usf"), EShaderPlatform::SP_PCD3D_SM5, &DataNumFile, nullptr));
+	return DataNumFile;
+}
+
+void UPCGDataNumKernel::CreateAdditionalInputDataInterfaces(FPCGGPUCompilationContext& InOutContext, UObject* InObjectOuter, TArray<TObjectPtr<UComputeDataInterface>>& OutDataInterfaces) const
+{
+	Super::CreateAdditionalInputDataInterfaces(InOutContext, InObjectOuter, OutDataInterfaces);
+
+	TObjectPtr<UPCGDataNumDataInterface> NodeDI = InOutContext.NewObject_AnyThread<UPCGDataNumDataInterface>(InObjectOuter);
+	NodeDI->ProducerKernel = this;
+
+	OutDataInterfaces.Add(NodeDI);
+}
+#endif
+
+void UPCGDataNumKernel::GetKernelAttributeKeys(TArray<FPCGKernelAttributeKey>& OutKeys) const
+{
+	OutKeys.AddUnique(OutCountAttributeKey);
+}
+
+void UPCGDataNumKernel::GetInputPins(TArray<FPCGPinProperties>& OutPins) const
+{
+	OutPins.Emplace(PCGPinConstants::DefaultInputLabel, EPCGDataType::Point);
+}
+
+void UPCGDataNumKernel::GetOutputPins(TArray<FPCGPinPropertiesGPU>& OutPins) const
+{
+	OutPins.Emplace(PCGPinConstants::DefaultOutputLabel, EPCGDataType::Param);
+}
+
+void UPCGDataNumKernel::SetOutputAttribute(const FName InAttributeName)
+{
+	OutCountAttributeKey = FPCGKernelAttributeKey(InAttributeName, EPCGKernelAttributeType::Int);
+}
+
+#undef LOCTEXT_NAMESPACE
diff --git a/Engine/Plugins/PCG/Source/PCG/Private/Compute/BuiltInKernels/PCGDataNumKernel.h b/Engine/Plugins/PCG/Source/PCG/Private/Compute/BuiltInKernels/PCGDataNumKernel.h
new file mode 100644
--- /dev/null
+++ b/Engine/Plugins/PCG/Source/PCG/Private/Compute/BuiltInKernels/PCGDataNumKernel.h
@@ -0,0 +1,34 @@
+// Copyright Epic Games, Inc. All Rights Reserved.
+
+#pragma once
+
+#include "Compute/PCGComputeKernel.h"
+
+#include "Compute/PCGPinPropertiesGPU.h"
+
+#include "PCGDataNumKernel.generated.h"
+
+UCLASS()
+class UPCGDataNumKernel : public UPCGComputeKernel
+{
+	GENERATED_BODY()
+
+public:
+	virtual bool ComputeOutputBindingDataDesc(const UPCGComputeGraph* InGraph, FName InOutputPinLabel, UPCGDataBinding* InBinding, FPCGDataCollectionDesc& OutDataDesc) const override;
+	virtual int ComputeThreadCount(const UPCGDataBinding* Binding) const override;
+	virtual bool IsKernelDataValid(FPCGContext* InContext) const override;
+
+#if WITH_EDITOR
+	virtual FString GetCookedSource(const UPCGComputeGraph* InComputeGraph) const override;
+	virtual FString GetEntryPoint() const override { return TEXT("PCGDataNumCS"); }
+	virtual void CreateAdditionalInputDataInterfaces(FPCGGPUCompilationContext& InOutContext, UObject* InObjectOuter, TArray<TObjectPtr<UComputeDataInterface>>& OutDataInterfaces) const override;
+#endif
+	virtual void GetKernelAttributeKeys(TArray<FPCGKernelAttributeKey>& OutKeys) const override;
+	virtual void GetInputPins(TArray<FPCGPinProperties>& OutPins) const override;
+	virtual void GetOutputPins(TArray<FPCGPinPropertiesGPU>& OutPins) const override;
+
+	void SetOutputAttribute(FName InAttributeName);
+
+private:
+	FPCGKernelAttributeKey OutCountAttributeKey;
+};
diff --git a/Engine/Plugins/PCG/Source/PCG/Private/Compute/DataInterfaces/BuiltInKernels/PCGDataNumDataInterface.cpp b/Engine/Plugins/PCG/Source/PCG/Private/Compute/DataInterfaces/BuiltInKernels/PCGDataNumDataInterface.cpp
new file mode 100644
--- /dev/null
+++ b/Engine/Plugins/PCG/Source/PCG/Private/Compute/DataInterfaces/BuiltInKernels/PCGDataNumDataInterface.cpp
@@ -0,0 +1,112 @@
+// Copyright Epic Games, Inc. All Rights Reserved.
+
+#include "PCGDataNumDataInterface.h"
+
+#include "PCGContext.h"
+#include "Compute/PCGComputeGraph.h"
+#include "Compute/PCGDataBinding.h"
+#include "Compute/BuiltInKernels/PCGDataNumKernel.h"
+#include "Elements/PCGDataNum.h"
+
+#include "ShaderCompilerCore.h"
+#include "ShaderCore.h"
+#include "ShaderParameterMetadataBuilder.h"
+#include "ComputeFramework/ShaderParamTypeDefinition.h"
+
+#include UE_INLINE_GENERATED_CPP_BY_NAME(PCGDataNumDataInterface)
+
+void UPCGDataNumDataInterface::GetSupportedInputs(TArray<FShaderFunctionDefinition>& OutFunctions) const
+{
+	OutFunctions.AddDefaulted_GetRef()
+		.SetName(TEXT("DataNum_GetOutputAttributeId"))
+		.AddReturnType(FShaderValueType::Get(EShaderFundamentalType::Int));
+}
+
+BEGIN_SHADER_PARAMETER_STRUCT(FPCGDataNumDataInterfaceParameters,)
+	SHADER_PARAMETER(int32, OutputAttributeId)
+END_SHADER_PARAMETER_STRUCT()
+
+void UPCGDataNumDataInterface::GetShaderParameters(TCHAR const* UID, FShaderParametersMetadataBuilder& InOutBuilder, FShaderParametersMetadataAllocations& InOutAllocations) const
+{
+	InOutBuilder.AddNestedStruct<FPCGDataNumDataInterfaceParameters>(UID);
+}
+
+void UPCGDataNumDataInterface::GetHLSL(FString& OutHLSL, FString const& InDataInterfaceName) const
+{
+	const TMap<FString, FStringFormatArg> DataNumArgs = {{TEXT("DataInterfaceName"), InDataInterfaceName}};
+
+	OutHLSL += FString::Format(TEXT(
+		"int {DataInterfaceName}_OutputAttributeId;\n"
+		"\n"
+		"int DataNum_GetOutputAttributeId_{DataInterfaceName}()\n"
+		"{\n"
+		"	return {DataInterfaceName}_OutputAttributeId;\n"
+		"}\n"
+	), DataNumArgs);
+}
+
+UComputeDataProvider* UPCGDataNumDataInterface::CreateDataProvider() const
+{
+	return NewObject<UPCGDataNumDataProvider>();
+}
+
+void UPCGDataNumDataProvider::Initialize(const UComputeDataInterface* InDataInterface, UObject* InBinding, const uint64 InInputMask, const uint64 InOutputMask)
+{
+	TRACE_CPUPROFILER_EVENT_SCOPE(UPCGDataNumDataProvider::Initialize);
+
+	Super::Initialize(InDataInterface, InBinding, InInputMask, InOutputMask);
+
+	const UPCGDataNumDataInterface* DataInterface = CastChecked<UPCGDataNumDataInterface>(InDataInterface);
+	ProducerKernel = DataInterface->ProducerKernel;
+
+	const UPCGDataNumSettings* Settings = CastChecked<UPCGDataNumSettings>(ProducerKernel->GetSettings());
+	OutputAttributeName = Settings->OutputAttributeName;
+}
+
+FComputeDataProviderRenderProxy* UPCGDataNumDataProvider::GetRenderProxy()
+{
+	FPCGDataNumDataProviderProxy::FDataNumData_RenderThread ProxyData;
+	ProxyData.OutputAttributeId = OutputAttributeId;
+	return new FPCGDataNumDataProviderProxy(std::move(ProxyData));
+}
+
+bool UPCGDataNumDataProvider::PrepareForExecute_GameThread(UPCGDataBinding* InBinding)
+{
+	TRACE_CPUPROFILER_EVENT_SCOPE(UPCGDataNumDataProvider::PrepareForExecute_GameThread);
+	check(InBinding);
+
+	if (!Super::PrepareForExecute_GameThread(InBinding))
+	{
+		return false;
+	}
+
+	OutputAttributeId = InBinding->GetAttributeId(OutputAttributeName, EPCGKernelAttributeType::Int);
+	check(OutputAttributeId != INDEX_NONE);
+
+	return true;
+}
+
+void UPCGDataNumDataProvider::Reset()
+{
+	ProducerKernel = nullptr;
+	OutputAttributeName = NAME_None;
+	OutputAttributeId = INDEX_NONE;
+
+	Super::Reset();
+}
+
+bool FPCGDataNumDataProviderProxy::IsValid(FValidationData const& InValidationData) const
+{
+	return InValidationData.ParameterStructSize == sizeof(FParameters);
+}
+
+void FPCGDataNumDataProviderProxy::GatherDispatchData(FDispatchData const& InDispatchData)
+{
+	const TStridedView<FParameters> ParameterArray = MakeStridedParameterView<FParameters>(InDispatchData);
+
+	for (int32 InvocationIndex = 0; InvocationIndex < ParameterArray.Num(); ++InvocationIndex)
+	{
+		FParameters& Parameters = ParameterArray[InvocationIndex];
+		Parameters.OutputAttributeId = Data.OutputAttributeId;
+	}
+}
\ No newline at end of file
diff --git a/Engine/Plugins/PCG/Source/PCG/Private/Compute/DataInterfaces/BuiltInKernels/PCGDataNumDataInterface.h b/Engine/Plugins/PCG/Source/PCG/Private/Compute/DataInterfaces/BuiltInKernels/PCGDataNumDataInterface.h
new file mode 100644
--- /dev/null
+++ b/Engine/Plugins/PCG/Source/PCG/Private/Compute/DataInterfaces/BuiltInKernels/PCGDataNumDataInterface.h
@@ -0,0 +1,68 @@
+// Copyright Epic Games, Inc. All Rights Reserved.
+
+#pragma once
+
+#include "Compute/DataInterfaces/PCGComputeDataInterface.h"
+#include "ComputeFramework/ComputeDataProvider.h"
+
+#include "PCGDataNumDataInterface.generated.h"
+
+class FPCGDataNumDataInterfaceParameters;
+
+/** Data Interface to marshal UPCGDataNumSettings settings to the GPU. */
+UCLASS(ClassGroup = (Procedural))
+class UPCGDataNumDataInterface : public UPCGComputeDataInterface
+{
+	GENERATED_BODY()
+
+public:
+	//~ Begin UComputeDataInterface Interface
+	virtual TCHAR const* GetClassName() const override { return TEXT("PCGDataNum"); }
+	virtual void GetSupportedInputs(TArray<FShaderFunctionDefinition>& OutFunctions) const override;
+	virtual void GetShaderParameters(TCHAR const* UID, FShaderParametersMetadataBuilder& InOutBuilder, FShaderParametersMetadataAllocations& InOutAllocations) const override;
+	virtual void GetHLSL(FString& OutHLSL, FString const& InDataInterfaceName) const override;
+	virtual UComputeDataProvider* CreateDataProvider() const override;
+	//~ End UComputeDataInterface Interface
+};
+
+UCLASS()
+class UPCGDataNumDataProvider : public UPCGComputeDataProvider
+{
+	GENERATED_BODY()
+
+public:
+	//~ Begin UComputeDataProvider Interface
+	virtual void Initialize(const UComputeDataInterface* InDataInterface, UObject* InBinding, uint64 InInputMask, uint64 InOutputMask) override;
+	virtual FComputeDataProviderRenderProxy* GetRenderProxy() override;
+	virtual bool PrepareForExecute_GameThread(UPCGDataBinding* InBinding) override;
+	virtual void Reset() override;
+	//~ End UComputeDataProvider Interface
+
+	UPROPERTY()
+	TObjectPtr<const UPCGComputeKernel> ProducerKernel;
+
+	FName OutputAttributeName = NAME_None;
+	int32 OutputAttributeId = INDEX_NONE;
+};
+
+class FPCGDataNumDataProviderProxy : public FComputeDataProviderRenderProxy
+{
+public:
+	struct FDataNumData_RenderThread
+	{
+		int32 OutputAttributeId = INDEX_NONE;
+	};
+
+	FPCGDataNumDataProviderProxy(FDataNumData_RenderThread InData)
+		: Data(MoveTemp(InData)) {}
+
+	//~ Begin FComputeDataProviderRenderProxy Interface
+	virtual bool IsValid(FValidationData const& InValidationData) const override;
+	virtual void GatherDispatchData(FDispatchData const& InDispatchData) override;
+	//~ End FComputeDataProviderRenderProxy Interface
+
+protected:
+	using FParameters = FPCGDataNumDataInterfaceParameters;
+
+	FDataNumData_RenderThread Data;
+};
diff --git a/Engine/Plugins/PCG/Source/PCG/Private/Elements/PCGDataNum.cpp b/Engine/Plugins/PCG/Source/PCG/Private/Elements/PCGDataNum.cpp
--- a/Engine/Plugins/PCG/Source/PCG/Private/Elements/PCGDataNum.cpp
+++ b/Engine/Plugins/PCG/Source/PCG/Private/Elements/PCGDataNum.cpp
@@ -5,6 +5,8 @@
 #include "PCGContext.h"
 #include "PCGParamData.h"
 #include "PCGPin.h"
+#include "Compute/BuiltInKernels/PCGDataNumKernel.h"
+#include "Graph/PCGGPUGraphCompilationContext.h"
 
 #include UE_INLINE_GENERATED_CPP_BY_NAME(PCGDataNum)
 
@@ -25,6 +27,17 @@ FText UPCGDataNumSettings::GetNodeTooltipText() const
 {
 	return LOCTEXT("NodeTooltip", "Returns the count of data in the input data collection.");
 }
+
+void UPCGDataNumSettings::CreateKernels(FPCGGPUCompilationContext& InOutContext, UObject* InObjectOuter, TArray<UPCGComputeKernel*>& OutKernels, TArray<FPCGKernelEdge>& OutEdges) const
+{
+	UPCGDataNumKernel* Kernel = InOutContext.NewObject_AnyThread<UPCGDataNumKernel>(InObjectOuter);
+	Kernel->SetOutputAttribute(OutputAttributeName);
+	Kernel->Initialize(this, bDumpDataDescriptions);
+	OutKernels.Add(Kernel);
+
+	OutEdges.Emplace(FPCGPinReference(PCGPinConstants::DefaultInputLabel), FPCGPinReference(Kernel, PCGPinConstants::DefaultInputLabel));
+	OutEdges.Emplace(FPCGPinReference(Kernel, PCGPinConstants::DefaultOutputLabel), FPCGPinReference(PCGPinConstants::DefaultOutputLabel));
+}
 #endif // WITH_EDITOR
 
 TArray<FPCGPinProperties> UPCGDataNumSettings::OutputPinProperties() const
diff --git a/Engine/Plugins/PCG/Source/PCG/Public/Elements/PCGDataNum.h b/Engine/Plugins/PCG/Source/PCG/Public/Elements/PCGDataNum.h
--- a/Engine/Plugins/PCG/Source/PCG/Public/Elements/PCGDataNum.h
+++ b/Engine/Plugins/PCG/Source/PCG/Public/Elements/PCGDataNum.h
@@ -19,6 +19,8 @@ public:
 	virtual FText GetDefaultNodeTitle() const override;
 	virtual FText GetNodeTooltipText() const override;
 	virtual EPCGSettingsType GetType() const override { return EPCGSettingsType::Generic; }
+	virtual bool DisplayExecuteOnGPUSetting() const override { return true; }
+	virtual void CreateKernels(FPCGGPUCompilationContext& InOutContext, UObject* InObjectOuter, TArray<UPCGComputeKernel*>& OutKernels, TArray<FPCGKernelEdge>& OutEdges) const override;
 #endif
 
 protected:
