[PCG] Add world actor find when registering partition actors at runtime

* In the case of sublevels, the PA could be in a level that is not the level of the owning world
* This one could be without a world actor, so we are dependent on loading order where the world actor could be loaded after the first PA.
* Add a new function, that does that only at runtime, that if there is no World actor, we look for a world actor in the same level as the PA.

#rnx

diff --git a/Engine/Plugins/PCG/Source/PCG/Private/PCGActorAndComponentMapping.cpp b/Engine/Plugins/PCG/Source/PCG/Private/PCGActorAndComponentMapping.cpp
--- a/Engine/Plugins/PCG/Source/PCG/Private/PCGActorAndComponentMapping.cpp
+++ b/Engine/Plugins/PCG/Source/PCG/Private/PCGActorAndComponentMapping.cpp
@@ -650,8 +650,13 @@ void FPCGActorAndComponentMapping::RegisterPartitionActor(APCGPartitionActor* In
 		PartitionActorsMapGrid.Add(GridCoord, InActor);
 	}
 
-	APCGWorldActor* WorldActor = PCGSubsystem->GetPCGWorldActor();
-	check(WorldActor);
+	APCGWorldActor* WorldActor = PCGSubsystem->GetPCGWorldActorForPartitionActor(InActor);
+	if (!ensureMsgf(WorldActor, TEXT("When registering %s (Level: %s), the PCG World Actor does not exist in the world of the subsystem nor the partition actor level.\n"
+		  "This is ill-formed, make sure the PCG World Actor exists if you use PCG Partition actors."),
+		  *InActor->GetName(), InActor->GetLevel() ? *InActor->GetLevel()->GetName() : TEXT("Unknown")))
+	{
+		return;
+	}
 
 #if WITH_EDITOR
 PRAGMA_DISABLE_DEPRECATION_WARNINGS
diff --git a/Engine/Plugins/PCG/Source/PCG/Private/PCGSubsystem.cpp b/Engine/Plugins/PCG/Source/PCG/Private/PCGSubsystem.cpp
--- a/Engine/Plugins/PCG/Source/PCG/Private/PCGSubsystem.cpp
+++ b/Engine/Plugins/PCG/Source/PCG/Private/PCGSubsystem.cpp
@@ -343,18 +343,45 @@ APCGWorldActor* UPCGSubsystem::GetPCGWorldActor()
 #if WITH_EDITOR
 	if (!PCGWorldActor && !PCGHelpers::IsRuntimeOrPIE())
 	{
-		PCGWorldActorLock.Lock();
+		FScopeLock Lock(&PCGWorldActorLock);
+
 		if (!PCGWorldActor)
 		{
 			PCGWorldActor = APCGWorldActor::CreatePCGWorldActor(GetWorld());
 		}
-		PCGWorldActorLock.Unlock();
 	}
 #endif
 
 	return PCGWorldActor;
 }
 
+APCGWorldActor* UPCGSubsystem::GetPCGWorldActorForPartitionActor(APCGPartitionActor* InActor)
+{
+	APCGWorldActor* FoundWorldActor = GetPCGWorldActor();
+	if (FoundWorldActor)
+	{
+		return FoundWorldActor;
+	}
+	
+	// We're at runtime and we didn't find a World Actor. This can happen if the partition actor is not in the same level than the persistent
+	// level of the world. In that case, look for the PCG World Actor in the same level as the Partition Actor. If we find one, register it.
+	if (ensure(PCGHelpers::IsRuntimeOrPIE()) && ensure(InActor))
+	{
+		UPCGActorHelpers::ForEachActorInLevel<APCGWorldActor>(InActor->GetLevel(), [this](AActor* InActor)
+		{
+			FScopeLock Lock(&PCGWorldActorLock);
+			if (!PCGWorldActor)
+			{
+				PCGWorldActor = Cast<APCGWorldActor>(InActor);
+			}
+
+			return PCGWorldActor == nullptr;
+		});
+	}
+
+	return PCGWorldActor;
+}
+
 APCGWorldActor* UPCGSubsystem::FindPCGWorldActor()
 {
 	return PCGWorldActor;
diff --git a/Engine/Plugins/PCG/Source/PCG/Public/PCGSubsystem.h b/Engine/Plugins/PCG/Source/PCG/Public/PCGSubsystem.h
--- a/Engine/Plugins/PCG/Source/PCG/Public/PCGSubsystem.h
+++ b/Engine/Plugins/PCG/Source/PCG/Public/PCGSubsystem.h
@@ -412,6 +412,9 @@ private:
 #endif // WITH_EDITOR
 	
 private:
+	/** When registering PAs, we might not have a PCG World Actor already registered (at runtime). In that case we'll look for a World Actor in the same level.*/
+	APCGWorldActor* GetPCGWorldActorForPartitionActor(APCGPartitionActor* InActor);
+
 	void CancelGeneration(UPCGComponent* Component, bool bCleanupUnusedResources);
 
 	UPCGData* GetPCGData(FPCGTaskId InGraphExecutionTaskId);
@@ -439,9 +442,9 @@ private:
 	/** Functions will be executed at the beginning of the tick and then removed from this array. */
 	TArray<FTickAction> BeginTickActions;
 
-#if WITH_EDITOR
 	FCriticalSection PCGWorldActorLock;
-		
+
+#if WITH_EDITOR
 	using FConstructionScriptSourceComponents = TMap<FName, TObjectKey<UPCGComponent>>;
 	TMap<TObjectKey<AActor>, FConstructionScriptSourceComponents> PerActorConstructionScriptSourceComponents;
 
