[PCG] Small pass on accessors to cleanup UI

* Add EulerAngles as an extractor for rotators
* Add seperators in the extractors to be cleaner
* Add two categories ControlPoints and Global in Spline to differentiate which target it is
* Changed display name in Make Rotator for Make rot from angles to Make rot from Euler angles.

#rnx

diff --git a/Engine/Plugins/PCG/Source/PCG/Private/Data/PCGSplineData.cpp b/Engine/Plugins/PCG/Source/PCG/Private/Data/PCGSplineData.cpp
--- a/Engine/Plugins/PCG/Source/PCG/Private/Data/PCGSplineData.cpp
+++ b/Engine/Plugins/PCG/Source/PCG/Private/Data/PCGSplineData.cpp
@@ -430,7 +430,14 @@ FPCGAttributeAccessorMethods UPCGSplineData::GetSplineAccessorMethods()
 	};
 
 #if WITH_EDITOR
-	Methods.FillSelectorMenuEntryFromEnum<EPCGSplineStructProperties>({LOCTEXT("ControlPointSelectorMenuEntry", "Spline Control Point")});
+	TArray<FText, TInlineAllocator<2>> Menus = {LOCTEXT("ControlPointSelectorMenuEntry", "Spline")};
+	FText& SubMenu = Menus.Emplace_GetRef();
+	
+	SubMenu = LOCTEXT("ControlPointSelectorMenuEntryPoints", "Control Points");
+	Methods.FillSelectorMenuEntryFromEnum<EPCGSplineStructProperties>(Menus);
+	
+	SubMenu = LOCTEXT("ControlPointSelectorMenuEntryGlobal", "Global");
+	Methods.FillSelectorMenuEntryFromEnum<EPCGSplineDataProperties>(Menus);
 #endif // WITH_EDITOR
 
 	return Methods;
diff --git a/Engine/Plugins/PCG/Source/PCG/Private/Metadata/Accessors/PCGAttributeAccessorFactory.cpp b/Engine/Plugins/PCG/Source/PCG/Private/Metadata/Accessors/PCGAttributeAccessorFactory.cpp
--- a/Engine/Plugins/PCG/Source/PCG/Private/Metadata/Accessors/PCGAttributeAccessorFactory.cpp
+++ b/Engine/Plugins/PCG/Source/PCG/Private/Metadata/Accessors/PCGAttributeAccessorFactory.cpp
@@ -173,7 +173,7 @@ decltype(auto) FPCGAttributeAccessorFactory::CallOnMethod(PCGData* InData, Func
 }
 
 #if WITH_EDITOR
-void FPCGAttributeAccessorMethods::FillSelectorMenuEntryFromEnum(const UEnum* EnumType, const TArray<FText>& InMenuHierarchy)
+void FPCGAttributeAccessorMethods::FillSelectorMenuEntryFromEnum(const UEnum* EnumType, const TArrayView<const FText> InMenuHierarchy)
 {
 	if (!EnumType)
 	{
diff --git a/Engine/Plugins/PCG/Source/PCG/Private/Metadata/Accessors/PCGAttributeExtractor.cpp b/Engine/Plugins/PCG/Source/PCG/Private/Metadata/Accessors/PCGAttributeExtractor.cpp
--- a/Engine/Plugins/PCG/Source/PCG/Private/Metadata/Accessors/PCGAttributeExtractor.cpp
+++ b/Engine/Plugins/PCG/Source/PCG/Private/Metadata/Accessors/PCGAttributeExtractor.cpp
@@ -219,6 +219,13 @@ namespace PCGAttributeExtractor
 				[](FRotator& Value, const double& In) -> void { Value.Yaw = In; });
 		}
 
+		if (Name == PCGAttributeExtractorConstants::RotatorEulerAngles)
+		{
+			return MakeUnique<FPCGChainAccessor<FVector, FRotator>>(std::move(InAccessor),
+				[](const FRotator& Value) -> FVector { return Value.Euler(); },
+				[](FRotator& Value, const FVector& In) -> void { Value = FRotator::MakeFromEuler(In); });
+		}
+
 		// Read-only
 		if (Name == PCGAttributeExtractorConstants::RotatorForward)
 		{
@@ -267,6 +274,13 @@ namespace PCGAttributeExtractor
 				[](FQuat& Value, const double& In) -> void { FRotator Temp = Value.Rotator(); Temp.Yaw = In; Value = Temp.Quaternion(); });
 		}
 
+		if (Name == PCGAttributeExtractorConstants::RotatorEulerAngles)
+		{
+			return MakeUnique<FPCGChainAccessor<FVector, FQuat>>(std::move(InAccessor),
+				[](const FQuat& Value) -> FVector { return Value.Euler(); },
+				[](FQuat& Value, const FVector& In) -> void { Value = FQuat::MakeFromEuler(In); });
+		}
+
 		// Read-only
 		if (Name == PCGAttributeExtractorConstants::RotatorForward)
 		{
diff --git a/Engine/Plugins/PCG/Source/PCG/Private/Tests/Accessors/PCGAttributeExtractorTest.cpp b/Engine/Plugins/PCG/Source/PCG/Private/Tests/Accessors/PCGAttributeExtractorTest.cpp
--- a/Engine/Plugins/PCG/Source/PCG/Private/Tests/Accessors/PCGAttributeExtractorTest.cpp
+++ b/Engine/Plugins/PCG/Source/PCG/Private/Tests/Accessors/PCGAttributeExtractorTest.cpp
@@ -289,6 +289,7 @@ bool FPCGAttributeRotatorPropertyExtractorTest::RunTest(const FString & Paramete
 	bTestPassed &= Verify(Struct.RotatorValue.Quaternion().GetForwardVector(), PCG_STRUCT_NAME_CHECK(RotatorValue), PCGAttributeExtractorConstants::RotatorForward);
 	bTestPassed &= Verify(Struct.RotatorValue.Quaternion().GetRightVector(), PCG_STRUCT_NAME_CHECK(RotatorValue), PCGAttributeExtractorConstants::RotatorRight);
 	bTestPassed &= Verify(Struct.RotatorValue.Quaternion().GetUpVector(), PCG_STRUCT_NAME_CHECK(RotatorValue), PCGAttributeExtractorConstants::RotatorUp);
+	bTestPassed &= Verify(Struct.RotatorValue.Euler(), PCG_STRUCT_NAME_CHECK(RotatorValue), PCGAttributeExtractorConstants::RotatorEulerAngles);
 
 	// Quaternions
 	bTestPassed &= Verify(Struct.QuatValue.Rotator().Pitch, PCG_STRUCT_NAME_CHECK(QuatValue), PCGAttributeExtractorConstants::RotatorPitch);
@@ -297,6 +298,7 @@ bool FPCGAttributeRotatorPropertyExtractorTest::RunTest(const FString & Paramete
 	bTestPassed &= Verify(Struct.QuatValue.GetForwardVector(), PCG_STRUCT_NAME_CHECK(QuatValue), PCGAttributeExtractorConstants::RotatorForward);
 	bTestPassed &= Verify(Struct.QuatValue.GetRightVector(), PCG_STRUCT_NAME_CHECK(QuatValue), PCGAttributeExtractorConstants::RotatorRight);
 	bTestPassed &= Verify(Struct.QuatValue.GetUpVector(), PCG_STRUCT_NAME_CHECK(QuatValue), PCGAttributeExtractorConstants::RotatorUp);
+	bTestPassed &= Verify(Struct.QuatValue.Euler(), PCG_STRUCT_NAME_CHECK(QuatValue), PCGAttributeExtractorConstants::RotatorEulerAngles);
 
 	return bTestPassed;
 }
diff --git a/Engine/Plugins/PCG/Source/PCG/Public/Data/PCGSplineStruct.h b/Engine/Plugins/PCG/Source/PCG/Public/Data/PCGSplineStruct.h
--- a/Engine/Plugins/PCG/Source/PCG/Public/Data/PCGSplineStruct.h
+++ b/Engine/Plugins/PCG/Source/PCG/Public/Data/PCGSplineStruct.h
@@ -9,7 +9,7 @@
 
 #include "PCGSplineStruct.generated.h"
 
-UENUM(Blueprintable)
+UENUM()
 enum class EPCGSplineStructProperties : uint8
 {
 	Position         UMETA(Tooltip = "Location component of the control point's transform in world coordinates."),
@@ -23,8 +23,13 @@ enum class EPCGSplineStructProperties : uint8
 	LocalRotation    UMETA(Tooltip = "Rotation component of the control point's transform."),
 	LocalScale       UMETA(Tooltip = "Scale component of the control point's transform."),
 	LocalTransform   UMETA(Tooltip = "Control point's transform."),
-	SplineTransform	 UMETA(Tooltip = "Transform of the spline. Global parameter", PCG_MetadataDomain="Data"),
-	IsClosed		 UMETA(Tooltip = "If the spline is closed. Global parameter", PCG_MetadataDomain="Data"),
+};
+
+UENUM()
+enum class EPCGSplineDataProperties : uint8
+{
+	SplineTransform	 UMETA(Tooltip = "Transform of the spline.", PCG_MetadataDomain="Data"),
+	IsClosed		 UMETA(Tooltip = "If the spline is closed.", PCG_MetadataDomain="Data")
 };
 
 enum class EPCGControlPointsAccessorTarget
diff --git a/Engine/Plugins/PCG/Source/PCG/Public/Elements/Metadata/PCGMetadataMakeRotator.h b/Engine/Plugins/PCG/Source/PCG/Public/Elements/Metadata/PCGMetadataMakeRotator.h
--- a/Engine/Plugins/PCG/Source/PCG/Public/Elements/Metadata/PCGMetadataMakeRotator.h
+++ b/Engine/Plugins/PCG/Source/PCG/Public/Elements/Metadata/PCGMetadataMakeRotator.h
@@ -32,7 +32,7 @@ enum class EPCGMetadataMakeRotatorOp : uint8
 	MakeRotFromYZ,
 	MakeRotFromZY,
 	MakeRotFromAxes,
-	MakeRotFromAngles
+	MakeRotFromAngles UMETA(DisplayName = "Make Rot From Euler Angles")
 };
 
 /* Create a Rotator from 1, 2 or 3 axis. */
diff --git a/Engine/Plugins/PCG/Source/PCG/Public/Metadata/Accessors/PCGAttributeAccessorFactory.h b/Engine/Plugins/PCG/Source/PCG/Public/Metadata/Accessors/PCGAttributeAccessorFactory.h
--- a/Engine/Plugins/PCG/Source/PCG/Public/Metadata/Accessors/PCGAttributeAccessorFactory.h
+++ b/Engine/Plugins/PCG/Source/PCG/Public/Metadata/Accessors/PCGAttributeAccessorFactory.h
@@ -56,9 +56,9 @@ struct FPCGAttributeAccessorMethods
 
 #if WITH_EDITOR
 	template <typename EnumType>
-	void FillSelectorMenuEntryFromEnum(const TArray<FText>& InMenuHierarchy = {}) { FillSelectorMenuEntryFromEnum(StaticEnum<EnumType>(), InMenuHierarchy); }
+	void FillSelectorMenuEntryFromEnum(const TArrayView<const FText> InMenuHierarchy = {}) { FillSelectorMenuEntryFromEnum(StaticEnum<EnumType>(), InMenuHierarchy); }
 
-	PCG_API void FillSelectorMenuEntryFromEnum(const UEnum* EnumType, const TArray<FText>& InMenuHierarchy = {});
+	PCG_API void FillSelectorMenuEntryFromEnum(const UEnum* EnumType, const TArrayView<const FText> InMenuHierarchy = {});
 	
 	// Special option to gather all the possible options for a given data type in the Selector context menu.
 	FPCGAttributeSelectorMenu AttributeSelectorMenu;
diff --git a/Engine/Plugins/PCG/Source/PCG/Public/Metadata/Accessors/PCGAttributeExtractor.h b/Engine/Plugins/PCG/Source/PCG/Public/Metadata/Accessors/PCGAttributeExtractor.h
--- a/Engine/Plugins/PCG/Source/PCG/Public/Metadata/Accessors/PCGAttributeExtractor.h
+++ b/Engine/Plugins/PCG/Source/PCG/Public/Metadata/Accessors/PCGAttributeExtractor.h
@@ -58,6 +58,7 @@ namespace PCGAttributeExtractorConstants
 	const FName RotatorForward = TEXT("Forward");
 	const FName RotatorRight = TEXT("Right");
 	const FName RotatorUp = TEXT("Up");
+	const FName RotatorEulerAngles = TEXT("EulerAngles");
 
 	const FName RotatorExtractorNames[] =
 	{
@@ -66,7 +67,8 @@ namespace PCGAttributeExtractorConstants
 		RotatorRoll,
 		RotatorForward,
 		RotatorRight,
-		RotatorUp
+		RotatorUp,
+		RotatorEulerAngles
 	};
 }
 
diff --git a/Engine/Plugins/PCG/Source/PCGEditor/Private/DataVisualizations/PCGSplineDataVisualization.cpp b/Engine/Plugins/PCG/Source/PCGEditor/Private/DataVisualizations/PCGSplineDataVisualization.cpp
--- a/Engine/Plugins/PCG/Source/PCGEditor/Private/DataVisualizations/PCGSplineDataVisualization.cpp
+++ b/Engine/Plugins/PCG/Source/PCGEditor/Private/DataVisualizations/PCGSplineDataVisualization.cpp
@@ -34,8 +34,8 @@ FPCGTableVisualizerInfo IPCGSplineDataVisualization::GetTableVisualizerInfoWithD
 		Info.SortingColumn = Info.ColumnInfos.Last().Id;
 
 		// Columns
-		AddPropertyEnumColumnInfo<FTransform>(Info, SplineData, EPCGSplineStructProperties::SplineTransform);
-		AddPropertyEnumColumnInfo<bool>(Info, SplineData, EPCGSplineStructProperties::IsClosed);
+		AddPropertyEnumColumnInfo<FTransform>(Info, SplineData, EPCGSplineDataProperties::SplineTransform);
+		AddPropertyEnumColumnInfo<bool>(Info, SplineData, EPCGSplineDataProperties::IsClosed);
 		
 		// Add Metadata Columns
 		CreateMetadataColumnInfos(SplineData, Info, PCGMetadataDomainID::Data);
diff --git a/Engine/Plugins/PCG/Source/PCGEditor/Private/Details/PCGAttributePropertySelectorDetails.cpp b/Engine/Plugins/PCG/Source/PCGEditor/Private/Details/PCGAttributePropertySelectorDetails.cpp
--- a/Engine/Plugins/PCG/Source/PCGEditor/Private/Details/PCGAttributePropertySelectorDetails.cpp
+++ b/Engine/Plugins/PCG/Source/PCGEditor/Private/Details/PCGAttributePropertySelectorDetails.cpp
@@ -77,6 +77,8 @@ namespace PCGAttributePropertySelectorDetails
 			BuildSubMenuSectionForExtractor(InSubMenuBuilder, InDetailsObject, PCGAttributeExtractorConstants::VectorY, TooltipComponents);
 			BuildSubMenuSectionForExtractor(InSubMenuBuilder, InDetailsObject, PCGAttributeExtractorConstants::VectorZ, TooltipComponents);
 			BuildSubMenuSectionForExtractor(InSubMenuBuilder, InDetailsObject, PCGAttributeExtractorConstants::VectorW, TooltipComponents);
+			InSubMenuBuilder.AddSeparator();
+
 			BuildSubMenuSectionForExtractor(InSubMenuBuilder, InDetailsObject, PCGAttributeExtractorConstants::VectorR, TooltipComponents);
 			BuildSubMenuSectionForExtractor(InSubMenuBuilder, InDetailsObject, PCGAttributeExtractorConstants::VectorG, TooltipComponents);
 			BuildSubMenuSectionForExtractor(InSubMenuBuilder, InDetailsObject, PCGAttributeExtractorConstants::VectorB, TooltipComponents);
@@ -84,6 +86,7 @@ namespace PCGAttributePropertySelectorDetails
 
 			if (bIsInput)
 			{
+				InSubMenuBuilder.AddSeparator();
 				BuildSubMenuSectionForExtractor(InSubMenuBuilder, InDetailsObject, PCGAttributeExtractorConstants::VectorLength, LOCTEXT("VectorLengthTooltip", "Get the length of the vector (double value). Read only."));
 				BuildSubMenuSectionForExtractor(InSubMenuBuilder, InDetailsObject, PCGAttributeExtractorConstants::VectorSquaredLength, LOCTEXT("VectorSquaredLengthTooltip", "Get the squared length of the vector (double value). Read only."));
 				BuildSubMenuSectionForExtractor(InSubMenuBuilder, InDetailsObject, PCGAttributeExtractorConstants::VectorNormalized, LOCTEXT("VectorNormalizedTooltip", "Get the vector normalized (double value). Read only."));
@@ -104,8 +107,13 @@ namespace PCGAttributePropertySelectorDetails
 			BuildSubMenuSectionForExtractor(InSubMenuBuilder, InDetailsObject, PCGAttributeExtractorConstants::RotatorPitch, TooltipComponents);
 			BuildSubMenuSectionForExtractor(InSubMenuBuilder, InDetailsObject, PCGAttributeExtractorConstants::RotatorYaw, TooltipComponents);
 
+			InSubMenuBuilder.AddSeparator();
+				
+			BuildSubMenuSectionForExtractor(InSubMenuBuilder, InDetailsObject, PCGAttributeExtractorConstants::RotatorEulerAngles, LOCTEXT("TooltipComponentsRotatorEuler", "Euler angles as a Vector"));
+
 			if (bIsInput)
 			{
+				InSubMenuBuilder.AddSeparator();
 				BuildSubMenuSectionForExtractor(InSubMenuBuilder, InDetailsObject, PCGAttributeExtractorConstants::RotatorForward, LOCTEXT("RotatorForwardTooltip", "X Axis of the rotation. Read only."));
 				BuildSubMenuSectionForExtractor(InSubMenuBuilder, InDetailsObject, PCGAttributeExtractorConstants::RotatorRight, LOCTEXT("RotatorRightTooltip", "Y Axis of the rotation. Read only."));
 				BuildSubMenuSectionForExtractor(InSubMenuBuilder, InDetailsObject, PCGAttributeExtractorConstants::RotatorUp, LOCTEXT("RotatorUpTooltip", "Z Axis of the rotation. Read only."));
