[PCG] Parallelize component selection for cleanup in RuntimeGenScheduler

#rnx

diff --git a/Engine/Plugins/PCG/Source/PCG/Private/RuntimeGen/PCGRuntimeGenScheduler.cpp b/Engine/Plugins/PCG/Source/PCG/Private/RuntimeGen/PCGRuntimeGenScheduler.cpp
--- a/Engine/Plugins/PCG/Source/PCG/Private/RuntimeGen/PCGRuntimeGenScheduler.cpp
+++ b/Engine/Plugins/PCG/Source/PCG/Private/RuntimeGen/PCGRuntimeGenScheduler.cpp
@@ -19,6 +19,7 @@
 #include "RuntimeGen/PCGGenSourceManager.h"
 
 #include "DrawDebugHelpers.h"
+#include "Async/ParallelFor.h"
 #include "Components/RuntimeVirtualTextureComponent.h"
 #include "UObject/UObjectGlobals.h"
 #include "VT/RuntimeVirtualTexture.h"
@@ -593,21 +594,31 @@ void FPCGRuntimeGenScheduler::TickCleanup(const TSet<IPCGGenSourceBase*>& InGenS
 		return bAllGenSourcesWantToCleanup;
 	};
 
+	TArray<FGridGenerationKey> GeneratedComponentsArray = GeneratedComponents.Array();
+	const int32 NumGeneratedComponents = GeneratedComponentsArray.Num();
+
 	// Generated Entry Key, Local/Generated Component
 	using PCGComponentToClean = TTuple<FGridGenerationKey, UPCGComponent*>;
-	TArray<PCGComponentToClean, TInlineAllocator<16>> ComponentsToClean;
-
+	TArray<PCGComponentToClean, TInlineAllocator<128>> ComponentsToClean;
 	TArray<FGridGenerationKey, TInlineAllocator<16>> InvalidKeys;
 
-	// Find any generated components which should be cleaned.
-	for (const FGridGenerationKey& GenerationKey : GeneratedComponents)
+	ComponentsToClean.SetNumUninitialized(NumGeneratedComponents);
+	InvalidKeys.SetNumUninitialized(NumGeneratedComponents);
+
+	std::atomic<int> ComponentsToCleanCounter = 0;
+	std::atomic<int> InvalidKeysCounter = 0;
+
+	ParallelFor(NumGeneratedComponents, [&GeneratedComponentsArray, &ComponentsToClean, &InvalidKeys, &ComponentsToCleanCounter, &InvalidKeysCounter, ActorAndComponentMapping=ActorAndComponentMapping, &CheckIfAllGenSourcesWantToCleanup](int32 Index)
 	{
-		TRACE_CPUPROFILER_EVENT_SCOPE(FPCGRuntimeGenScheduler::SelectComponentsToCleanup);
+		TRACE_CPUPROFILER_EVENT_SCOPE(SelectComponentForCleanup);
+
+		const FGridGenerationKey& GenerationKey = GeneratedComponentsArray[Index];
 
 		if (!GenerationKey.IsValid())
 		{
-			InvalidKeys.Add(GenerationKey);
-			continue;
+			const int WriteIndex = InvalidKeysCounter.fetch_add(1);
+			InvalidKeys[WriteIndex] = GenerationKey;
+			return;
 		}
 
 		const FPCGGridDescriptor GridDescriptor = GenerationKey.GetGridDescriptor();
@@ -627,8 +638,9 @@ void FPCGRuntimeGenScheduler::TickCleanup(const TSet<IPCGGenSourceBase*>& InGenS
 		{
 			if (!OriginalComponent->bActivated)
 			{
-				ComponentsToClean.Add({ GenerationKey, OriginalComponent });
-				continue;
+				const int WriteIndex = ComponentsToCleanCounter.fetch_add(1);
+				ComponentsToClean[WriteIndex] = { GenerationKey, OriginalComponent };
+				return;
 			}
 
 			const FBox GridBounds = OriginalComponent->GetGridBounds();
@@ -636,7 +648,9 @@ void FPCGRuntimeGenScheduler::TickCleanup(const TSet<IPCGGenSourceBase*>& InGenS
 			// Only clean up if all generation sources agreed to clean up.
 			if(CheckIfAllGenSourcesWantToCleanup(Policy, GridDescriptor, GridBounds, CleanupRadiusSquared))
 			{
-				ComponentsToClean.Add({ GenerationKey, OriginalComponent });
+				const int WriteIndex = ComponentsToCleanCounter.fetch_add(1);
+				ComponentsToClean[WriteIndex] = { GenerationKey, OriginalComponent };
+				return;
 			}
 		}
 		// Otherwise, we have a local component.
@@ -647,8 +661,9 @@ void FPCGRuntimeGenScheduler::TickCleanup(const TSet<IPCGGenSourceBase*>& InGenS
 			if (!PartitionActor || !OriginalComponent->bActivated)
 			{
 				// Attempt to clean even in failure case to avoid leaking resources.
-				ComponentsToClean.Add({ GenerationKey, LocalComponent });
-				continue;
+				const int WriteIndex = ComponentsToCleanCounter.fetch_add(1);
+				ComponentsToClean[WriteIndex] = { GenerationKey, LocalComponent };
+				return;
 			}
 
 			const FBox GridBounds = LocalComponent->GetGridBounds();
@@ -656,18 +671,24 @@ void FPCGRuntimeGenScheduler::TickCleanup(const TSet<IPCGGenSourceBase*>& InGenS
 			// Only clean up if all generation sources agreed to clean up.
 			if (CheckIfAllGenSourcesWantToCleanup(Policy, GridDescriptor, GridBounds, CleanupRadiusSquared))
 			{
-				ComponentsToClean.Add({ GenerationKey, LocalComponent });
+				const int WriteIndex = ComponentsToCleanCounter.fetch_add(1);
+				ComponentsToClean[WriteIndex] = { GenerationKey, LocalComponent };
+				return;
 			}
 		}
-	}
+	});
 
-	for (const FGridGenerationKey& InvalidKey : InvalidKeys)
+	const int NumInvalidKeys = InvalidKeysCounter.load();
+	const int NumComponentsToClean = ComponentsToCleanCounter.load();
+
+	for (int I = 0; I < NumInvalidKeys; ++I)
 	{
-		GeneratedComponents.Remove(InvalidKey);
+		GeneratedComponents.Remove(InvalidKeys[I]);
 	}
 
-	for (const PCGComponentToClean& ComponentToClean : ComponentsToClean)
+	for (int I = 0; I < NumComponentsToClean; ++I)
 	{
+		const PCGComponentToClean& ComponentToClean = ComponentsToClean[I];
 		CleanupComponent(ComponentToClean.Get<0>(), ComponentToClean.Get<1>());
 	}
 }
