[PCG] Fixe issues with PCG Builder in Iterative Cell mode
- PCGPartitionActor: Properly set Local PCG Component grid size in PostLoad (transient property).
- Call FPCGWorldPartitionBuilder::UpdatePartitionedActors after expanding bounds in case some new PAs were loaded

diff --git a/Engine/Plugins/PCG/Source/PCG/Private/Grid/PCGPartitionActor.cpp b/Engine/Plugins/PCG/Source/PCG/Private/Grid/PCGPartitionActor.cpp
--- a/Engine/Plugins/PCG/Source/PCG/Private/Grid/PCGPartitionActor.cpp
+++ b/Engine/Plugins/PCG/Source/PCG/Private/Grid/PCGPartitionActor.cpp
@@ -83,6 +83,7 @@ void APCGPartitionActor::PostLoad()
 	for (UPCGComponent* LocalComponent : GetAllLocalPCGComponents())
 	{
 		LocalComponent->MarkAsLocalComponent();
+		LocalComponent->SetGenerationGridSize(PCGGridSize);
 		LocalComponent->ConditionalPostLoad();
 	}
 
diff --git a/Engine/Plugins/PCG/Source/PCGEditor/Private/WorldPartitionBuilder/PCGWorldPartitionBuilder.cpp b/Engine/Plugins/PCG/Source/PCGEditor/Private/WorldPartitionBuilder/PCGWorldPartitionBuilder.cpp
--- a/Engine/Plugins/PCG/Source/PCGEditor/Private/WorldPartitionBuilder/PCGWorldPartitionBuilder.cpp
+++ b/Engine/Plugins/PCG/Source/PCGEditor/Private/WorldPartitionBuilder/PCGWorldPartitionBuilder.cpp
@@ -818,6 +818,12 @@ bool UPCGWorldPartitionBuilder::RunPerCell(UWorld* InWorld, const FCellInfo& InC
 	// Do a Bounds expansion to make sure we are generating PCG Components on areas covering their owning actor bounds (which might be a bigger area then the currently loaded cell)
 	TUniquePtr<IWorldPartitionActorLoaderInterface::ILoaderAdapter> LoaderAdapter = FPCGWorldPartitionBuilder::ExpandLoadedBoundsIfNeeded(InWorld, InCellInfo.Bounds, ExpandedBounds, TEXT("Expanded Bounds"));
 	
+	if (LoaderAdapter)
+	{
+		// Update Mappings for new loaded PAs from Expanded Bounds
+		bCellHasChanges |= FPCGWorldPartitionBuilder::UpdatePartitionedActors(InWorld, PartitionedComponents, CellDeletedPackageNames);
+	}
+
 	// Add Components for sub grids
 	TArray<UPCGComponent*> SubGridComponentsToGenerate;
 	TMap<UPCGComponent*, UPCGComponent*> TaskDependencies;
