[PCG] Refactor GPU instancing path.

Old path had to be run within the scene update flow, and injected instances directly into the GPU scene. This required awkward caching of the primitive proxies and an awkward frame delay between creating primitive components and updating the instances.

New path exposes a scene API to update from compute, and avoids caching/handling primitive proxies from within the instance data interface. This is more in line with the other scene update paths and compute graph execution does not necessarily have to be in scene update.

#rnx

diff --git a/Engine/Plugins/PCG/Shaders/Private/PCGSceneWriter.usf b/Engine/Plugins/PCG/Shaders/Private/PCGSceneWriter.usf
--- a/Engine/Plugins/PCG/Shaders/Private/PCGSceneWriter.usf
+++ b/Engine/Plugins/PCG/Shaders/Private/PCGSceneWriter.usf
@@ -60,8 +60,9 @@ void SceneWriter_CS(const uint ThreadId : SV_DispatchThreadID)
  		/*RandomId=*/0.0f,
 		DFFromTileOffset(LWCPromote(LocalToWorld))
 	);
-
-	const FDFMatrix Transform = MakeDFMatrix4x3(PrimitiveData.PositionHigh*0.0, LocalToWorld);
+	
+	// todo_pcg: We should write local space transforms to the instance data buffer from the SM Spawner and use WriteInstanceLocalToRelativeWorld here instead.
+	const FDFMatrix Transform = MakeDFMatrix4x3((float3)0.0f, LocalToWorld);
 	WriteInstanceLocalToWorld(InstanceId, InPrimitiveId, Transform, false);
 
 	FInstancePayloadDataOffsets PayloadOffsets = GetInstancePayloadDataOffsets(InPrimitiveId, InPayloadDataFlags, LocalInstanceId);
diff --git a/Engine/Plugins/PCG/Source/PCG/Private/Components/PCGProceduralISMComponent.cpp b/Engine/Plugins/PCG/Source/PCG/Private/Components/PCGProceduralISMComponent.cpp
--- a/Engine/Plugins/PCG/Source/PCG/Private/Components/PCGProceduralISMComponent.cpp
+++ b/Engine/Plugins/PCG/Source/PCG/Private/Components/PCGProceduralISMComponent.cpp
@@ -12,6 +12,8 @@
 #include "InstanceDataSceneProxy.h"
 #include "InstancedStaticMeshSceneProxyDesc.h"
 #include "NaniteSceneProxy.h"
+#include "PrimitiveSceneDesc.h"
+#include "PrimitiveSceneProxyDesc.h"
 #include "SceneInterface.h"
 #include "StaticMeshSceneProxyDesc.h"
 #include "Engine/StaticMesh.h"
@@ -369,31 +371,8 @@ FPrimitiveSceneProxy* UPCGProceduralISMComponent::CreateStaticMeshSceneProxy(Nan
 		return nullptr;
 	}
 
-	FInstanceSceneDataBuffers InstanceSceneDataBuffers(/*InbInstanceDataIsGPUOnly=*/true);
-	{
-		FInstanceSceneDataBuffers::FAccessTag AccessTag(PointerHash(this));
-		FInstanceSceneDataBuffers::FWriteView ProxyData = InstanceSceneDataBuffers.BeginWriteAccess(AccessTag);
-
-		InstanceSceneDataBuffers.SetPrimitiveLocalToWorld(GetRenderMatrix(), AccessTag);
-
-		ProxyData.NumInstancesGPUOnly = GetNumInstances();
-		ProxyData.NumCustomDataFloats = GetNumCustomDataFloats();
-		ProxyData.InstanceLocalBounds.SetNum(1);
-		ProxyData.InstanceLocalBounds[0] = GetStaticMesh()->GetBounds();
-
-		ProxyData.Flags.bHasPerInstanceCustomData = ProxyData.NumCustomDataFloats > 0;
-
-		InstanceSceneDataBuffers.EndWriteAccess(AccessTag);
-		InstanceSceneDataBuffers.ValidateData();
-	}
-
 	FInstancedStaticMeshSceneProxyDesc Desc;
-	Desc.InitializeFromStaticMeshComponent(this);
-	Desc.InstanceDataSceneProxy = MakeShared<FInstanceDataSceneProxy, ESPMode::ThreadSafe>(MoveTemp(InstanceSceneDataBuffers));
-	Desc.InstanceMinDrawDistance = InstanceMinDrawDistance;
-	Desc.InstanceStartCullDistance = InstanceStartCullDistance;
-	Desc.InstanceEndCullDistance = InstanceEndCullDistance;
-	Desc.bUseGpuLodSelection = true;
+	GetSceneProxyDesc(Desc);
 
 	if (bCreateNanite)
 	{
@@ -494,6 +473,48 @@ void UPCGProceduralISMComponent::SetMinDrawDistance(int32 InMinDrawDistance)
 	}
 }
 
+void UPCGProceduralISMComponent::GetSceneProxyDesc(FInstancedStaticMeshSceneProxyDesc& OutSceneProxyDesc) const
+{
+	FInstanceSceneDataBuffers InstanceSceneDataBuffers(/*InbInstanceDataIsGPUOnly=*/true);
+	{
+		FInstanceSceneDataBuffers::FAccessTag AccessTag(PointerHash(this));
+		FInstanceSceneDataBuffers::FWriteView ProxyData = InstanceSceneDataBuffers.BeginWriteAccess(AccessTag);
+
+		InstanceSceneDataBuffers.SetPrimitiveLocalToWorld(GetRenderMatrix(), AccessTag);
+
+		ProxyData.NumInstancesGPUOnly = GetNumInstances();
+		ProxyData.NumCustomDataFloats = GetNumCustomDataFloats();
+		ProxyData.InstanceLocalBounds.SetNum(1);
+		ProxyData.InstanceLocalBounds[0] = GetStaticMesh()->GetBounds();
+
+		ProxyData.Flags.bHasPerInstanceCustomData = ProxyData.NumCustomDataFloats > 0;
+
+		InstanceSceneDataBuffers.EndWriteAccess(AccessTag);
+		InstanceSceneDataBuffers.ValidateData();
+	}
+
+	OutSceneProxyDesc.InitializeFromStaticMeshComponent(this);
+	OutSceneProxyDesc.InstanceDataSceneProxy = MakeShared<FInstanceDataSceneProxy, ESPMode::ThreadSafe>(MoveTemp(InstanceSceneDataBuffers));
+	OutSceneProxyDesc.InstanceMinDrawDistance = InstanceMinDrawDistance;
+	OutSceneProxyDesc.InstanceStartCullDistance = InstanceStartCullDistance;
+	OutSceneProxyDesc.InstanceEndCullDistance = InstanceEndCullDistance;
+	OutSceneProxyDesc.bUseGpuLodSelection = true;
+}
+
+void UPCGProceduralISMComponent::BuildSceneDesc(FPrimitiveSceneProxyDesc* InSceneProxyDesc, FPrimitiveSceneDesc& OutPrimitiveSceneDesc)
+{
+	check(InSceneProxyDesc);
+
+	OutPrimitiveSceneDesc.SceneProxy = GetSceneProxy();
+	OutPrimitiveSceneDesc.ProxyDesc = InSceneProxyDesc;
+	OutPrimitiveSceneDesc.PrimitiveSceneData = &GetSceneData();
+	OutPrimitiveSceneDesc.RenderMatrix = GetRenderMatrix();
+	OutPrimitiveSceneDesc.AttachmentRootPosition = GetComponentLocation();
+	OutPrimitiveSceneDesc.LocalBounds = LocalBounds;
+	OutPrimitiveSceneDesc.Bounds = CalcBounds(GetComponentToWorld());
+	OutPrimitiveSceneDesc.Mobility = InSceneProxyDesc->Mobility;
+}
+
 void UPCGProceduralISMComponent::ValidateComponentSetup()
 {
 	if (bAffectDynamicIndirectLighting)
diff --git a/Engine/Plugins/PCG/Source/PCG/Private/Components/PCGProceduralISMComponent.h b/Engine/Plugins/PCG/Source/PCG/Private/Components/PCGProceduralISMComponent.h
--- a/Engine/Plugins/PCG/Source/PCG/Private/Components/PCGProceduralISMComponent.h
+++ b/Engine/Plugins/PCG/Source/PCG/Private/Components/PCGProceduralISMComponent.h
@@ -19,8 +19,11 @@
 class FPrimitiveSceneProxy;
 class UPCGComponent;
 class UPCGManagedProceduralISMComponent;
+struct FInstancedStaticMeshSceneProxyDesc;
 struct FPCGContext;
 struct FPCGSoftISMComponentDescriptor;
+struct FPrimitiveSceneDesc;
+struct FPrimitiveSceneProxyDesc;
 
 /** Experimental instanced primitive that is not backed by CPU instance data. Instance data must be written direct to the GPU scene. */
 UCLASS(ClassGroup = Rendering, meta = (BlueprintSpawnableComponent), Blueprintable, MinimalAPI, hideCategories = (Collision, Cooking, HLOD, MeshPainting, Navigation, Physics, RayTracing))
@@ -112,6 +115,11 @@ protected:
 #endif
 	//~ End UObject Interface
 
+public:
+	void GetSceneProxyDesc(FInstancedStaticMeshSceneProxyDesc& OutSceneProxyDesc) const;
+
+	void BuildSceneDesc(FPrimitiveSceneProxyDesc* InSceneProxyDesc, FPrimitiveSceneDesc& OutPrimitiveSceneDesc);
+
 protected:
 	/** Validate component setup, disable features that are not supported (typically those that require instance transforms on CPU). */
 	virtual void ValidateComponentSetup();
diff --git a/Engine/Plugins/PCG/Source/PCG/Private/Compute/DataInterfaces/Elements/PCGStaticMeshSpawnerDataInterface.cpp b/Engine/Plugins/PCG/Source/PCG/Private/Compute/DataInterfaces/Elements/PCGStaticMeshSpawnerDataInterface.cpp
--- a/Engine/Plugins/PCG/Source/PCG/Private/Compute/DataInterfaces/Elements/PCGStaticMeshSpawnerDataInterface.cpp
+++ b/Engine/Plugins/PCG/Source/PCG/Private/Compute/DataInterfaces/Elements/PCGStaticMeshSpawnerDataInterface.cpp
@@ -414,9 +414,6 @@ bool UPCGStaticMeshSpawnerDataProvider::PrepareForExecute_GameThread(UPCGDataBin
 		if (NumPrimitivesSetup > 0)
 		{
 			SourceComponent->NotifyProceduralInstancesInUse();
-
-			// Signal not finished. Compute graph element will wait a frame, giving time for GPU Scene to be updated.
-			return false;
 		}
 		else
 		{
@@ -424,19 +421,10 @@ bool UPCGStaticMeshSpawnerDataProvider::PrepareForExecute_GameThread(UPCGDataBin
 			InBinding->CompletedMeshSpawners.Add(ProducerKernel);
 			return true;
 		}
-	}
 
-	// If any proxy is not yet ready, wait a frame.
-	if (const FPCGSpawnerPrimitives* Primitives = InBinding->MeshSpawnersToPrimitives.Find(ProducerKernel); ensure(Primitives))
-	{
-		for (const TObjectPtr<UPrimitiveComponent>& Primitive : Primitives->Primitives)
+		if (Context->AsyncState.ShouldStop())
 		{
-			const UPCGProceduralISMComponent* PISMC = Cast<UPCGProceduralISMComponent>(Primitive);
-			if (PISMC && PISMC->GetNumInstances() > 0 && !PISMC->GetSceneProxy())
-			{
-				// Waiting on a proxy. Wait a frame.
-				return false;
-			}
+			return false;
 		}
 	}
 
@@ -774,15 +762,17 @@ bool UPCGStaticMeshSpawnerDataProvider::SetupPrimitives(FPCGContext* InContext,
 			// Don't bother registering the resource change as we're transient anyway.
 			//InContext->TouchedResources.Emplace(MISMC);
 
-			Primitives->Primitives.Add(MISMC->GetComponent());
+			FPCGSpawnerPrimitive Primitive;
+			Primitive.Component = MISMC->GetComponent();
+			Primitive.NumInstances = Desc.NumInstances;
+
+			Primitives->Primitives.Add(MoveTemp(Primitive));
 
 			if (Settings->bApplyMeshBoundsToPoints)
 			{
 				PrimitiveMeshBounds.Add(Desc.StaticMesh->GetBoundingBox());
 			}
 
-			Primitives->NumInstancesAllPrimitives += Desc.NumInstances;
-
 			++NumPrimitivesSetup;
 		}
 
diff --git a/Engine/Plugins/PCG/Source/PCG/Private/Compute/DataInterfaces/PCGComputeDataInterface.h b/Engine/Plugins/PCG/Source/PCG/Private/Compute/DataInterfaces/PCGComputeDataInterface.h
--- a/Engine/Plugins/PCG/Source/PCG/Private/Compute/DataInterfaces/PCGComputeDataInterface.h
+++ b/Engine/Plugins/PCG/Source/PCG/Private/Compute/DataInterfaces/PCGComputeDataInterface.h
@@ -60,6 +60,9 @@ public:
 	/* Called prior to GetRenderProxy and execution. Any processing of data descriptions should be done here (after PerformPreExecuteReadbacks has been called). */
 	virtual bool PrepareForExecute_GameThread(UPCGDataBinding* InBinding) { return true; }
 
+	/** Return true when all done. */
+	virtual bool PostExecute(UPCGDataBinding* InBinding) { return true; }
+
 	/** Release any handles to transient resources like GPU buffers. */
 	virtual void ReleaseTransientResources(const TCHAR* InReason = nullptr) {}
 };
diff --git a/Engine/Plugins/PCG/Source/PCG/Private/Compute/DataInterfaces/PCGInstanceDataInterface.cpp b/Engine/Plugins/PCG/Source/PCG/Private/Compute/DataInterfaces/PCGInstanceDataInterface.cpp
--- a/Engine/Plugins/PCG/Source/PCG/Private/Compute/DataInterfaces/PCGInstanceDataInterface.cpp
+++ b/Engine/Plugins/PCG/Source/PCG/Private/Compute/DataInterfaces/PCGInstanceDataInterface.cpp
@@ -2,19 +2,19 @@
 
 #include "PCGInstanceDataInterface.h"
 
-#include "PCGData.h"
-#include "PCGNode.h"
+#include "PCGModule.h"
 #include "PCGSettings.h"
 #include "PCGSceneWriterCS.h"
-#include "Compute/PCGComputeCommon.h"
+#include "Components/PCGProceduralISMComponent.h"
 #include "Compute/PCGDataBinding.h"
-#include "Data/PCGTextureData.h"
 #include "Elements/PCGStaticMeshSpawnerKernel.h"
 
 #include "GlobalRenderResources.h"
 #include "GlobalShader.h"
 #include "GPUSceneWriter.h"
 #include "InstanceDataSceneProxy.h"
+#include "InstancedStaticMeshSceneProxyDesc.h"
+#include "PrimitiveSceneDesc.h"
 #include "PrimitiveSceneInfo.h"
 #include "PrimitiveSceneProxy.h"
 #include "RenderCaptureInterface.h"
@@ -23,14 +23,12 @@
 #include "RenderGraphUtils.h"
 #include "RHIResources.h"
 #include "RHIStaticStates.h"
-#include "SceneComputeUpdates.h"
+#include "SceneInterface.h"
 #include "ShaderCompilerCore.h"
 #include "ShaderCore.h"
 #include "ShaderParameterMetadataBuilder.h"
 #include "ShaderParameterStruct.h"
-#include "TextureResource.h"
 #include "Algo/AnyOf.h"
-#include "Components/PrimitiveComponent.h"
 #include "ComputeFramework/ComputeKernelPermutationSet.h"
 #include "ComputeFramework/ShaderParamTypeDefinition.h"
 #include "Engine/Texture.h"
@@ -169,61 +167,78 @@ bool UPCGInstanceDataProvider::PrepareForExecute_GameThread(UPCGDataBinding* InB
 		FoundPrimitives->Primitives.SetNum(PCGComputeConstants::MAX_PRIMITIVE_COMPONENTS_PER_SPAWNER);
 	}
 
-	Primitives = FoundPrimitives->Primitives;
-	NumInstancesAllPrimitives = FoundPrimitives->NumInstancesAllPrimitives;
+	NumInstancesAllPrimitives = 0;
+	Primitives.Reserve(FoundPrimitives->Primitives.Num());
+
+	for (const FPCGSpawnerPrimitive& Primitive : FoundPrimitives->Primitives)
+	{
+		Primitives.Add(FPCGPrimitiveComponentAndNumInstances(Primitive.Component, Primitive.NumInstances));
+		NumInstancesAllPrimitives += Primitive.NumInstances;
+	}
+
 	NumCustomFloatsPerInstance = FoundPrimitives->NumCustomFloats;
 
 	return true;
 }
 
-FComputeDataProviderRenderProxy* UPCGInstanceDataProvider::GetRenderProxy()
+bool UPCGInstanceDataProvider::PostExecute(UPCGDataBinding* InBinding)
 {
-	TRACE_CPUPROFILER_EVENT_SCOPE(UPCGInstanceDataProvider::GetRenderProxy);
-	TArray<FPrimitiveSceneProxy*> PrimitiveProxies;
-	PrimitiveProxies.Reserve(Primitives.Num());
-	
-	for (TObjectPtr<UPrimitiveComponent>& Primitive : Primitives)
+	if (!Super::PostExecute(InBinding))
 	{
-		UWorld* World = Primitive ? Primitive->GetWorld() : nullptr;
-		FSceneInterface* Scene = World ? World->Scene : nullptr;
-		if (FPrimitiveSceneProxy* PrimitiveSceneProxy = Primitive ? Primitive->GetSceneProxy() : nullptr)
-		{
-			PrimitiveProxies.Add(PrimitiveSceneProxy);
-		}
+		return false;
 	}
 
-	return new FPCGInstanceDataProviderProxy(MoveTemp(PrimitiveProxies), NumInstancesAllPrimitives, NumCustomFloatsPerInstance);
+	return bWroteInstances;
+}
+
+FComputeDataProviderRenderProxy* UPCGInstanceDataProvider::GetRenderProxy()
+{
+	TRACE_CPUPROFILER_EVENT_SCOPE(UPCGInstanceDataProvider::GetRenderProxy);
+
+	return new FPCGInstanceDataProviderProxy(Primitives, NumInstancesAllPrimitives, NumCustomFloatsPerInstance, MakeWeakObjectPtr(this));
 }
 
 void UPCGInstanceDataProvider::Reset()
 {
 	ProducerKernel = nullptr;
-	Primitives.Empty();
+	Primitives.Empty(Primitives.Num());
 	NumInstancesAllPrimitives = 0;
 	NumCustomFloatsPerInstance = 0;
+	bWroteInstances = false;
+
+	++GenerationCounter;
 
 	Super::Reset();
 }
 
-FPCGInstanceDataProviderProxy::FPCGInstanceDataProviderProxy(TArray<FPrimitiveSceneProxy*>&& InPrimitiveSceneProxies, uint32 InNumInstancesAllPrimitives, uint32 InNumCustomFloatsPerInstance)
-	: PrimitiveSceneProxies(MoveTemp(InPrimitiveSceneProxies))
+FPCGInstanceDataProviderProxy::FPCGInstanceDataProviderProxy(TArray<FPCGPrimitiveComponentAndNumInstances> InPrimitives, uint32 InNumInstancesAllPrimitives, uint32 InNumCustomFloatsPerInstance, TWeakObjectPtr<UPCGInstanceDataProvider> InDataProvider)
+	: Primitives(MoveTemp(InPrimitives))
 	, NumInstancesAllPrimitives(InNumInstancesAllPrimitives)
 	, NumCustomFloatsPerInstance(InNumCustomFloatsPerInstance)
+	, DataProvider(InDataProvider)
 {
-}
+	bIsValid = true;
 
-bool FPCGInstanceDataProviderProxy::IsValid(FValidationData const& InValidationData) const
-{
-	for (const FPrimitiveSceneProxy* Proxy : PrimitiveSceneProxies)
+	if (ensure(InDataProvider.IsValid()))
 	{
-		if (!Proxy)
-		{
-			return false;
-		}
+		DataProviderGeneration = InDataProvider->GenerationCounter;
+	}
+	else
+	{
+		UE_LOG(LogPCG, Warning, TEXT("FPCGInstanceDataProviderProxy: Data provider missing, proxy is invalid, compute graph will not execute."));
+		bIsValid = false;
 	}
 
-	return InValidationData.ParameterStructSize == sizeof(FParameters)
-		&& !Algo::AnyOf(PrimitiveSceneProxies, [](const FPrimitiveSceneProxy* InProxy) { return !InProxy; });
+	if (Algo::AnyOf(Primitives, [](const FPCGPrimitiveComponentAndNumInstances& InPrimitive) { return !InPrimitive.Get<0>().IsValid(); }))
+	{
+		UE_LOG(LogPCG, Warning, TEXT("FPCGInstanceDataProviderProxy: One or more components were invalid, proxy is invalid, compute graph will not execute."));
+		bIsValid = false;
+	}
+}
+
+bool FPCGInstanceDataProviderProxy::IsValid(FValidationData const& InValidationData) const
+{
+	return bIsValid && InValidationData.ParameterStructSize == sizeof(FParameters);
 }
 
 void FPCGInstanceDataProviderProxy::AllocateResources(FRDGBuilder& GraphBuilder, FAllocationData const& InAllocationData)
@@ -238,7 +253,7 @@ void FPCGInstanceDataProviderProxy::AllocateResources(FRDGBuilder& GraphBuilder,
 	InstanceCustomFloatDataSRV = GraphBuilder.CreateSRV(InstanceCustomFloatData);
 	InstanceCustomFloatDataUAV = GraphBuilder.CreateUAV(InstanceCustomFloatData);
 
-	const int32 NumCountersRequired = FMath::Max(1, PrimitiveSceneProxies.Num());
+	const int32 NumCountersRequired = FMath::Max(1, Primitives.Num());
 	WriteCounters = GraphBuilder.CreateBuffer(FRDGBufferDesc::CreateStructuredDesc(sizeof(uint32), NumCountersRequired), TEXT("PCGWriteCounters"));
 	WriteCountersSRV = GraphBuilder.CreateSRV(WriteCounters);
 	WriteCountersUAV = GraphBuilder.CreateUAV(WriteCounters);
@@ -253,7 +268,7 @@ void FPCGInstanceDataProviderProxy::GatherDispatchData(FDispatchData const& InDi
 	for (int32 InvocationIndex = 0; InvocationIndex < ParameterArray.Num(); ++InvocationIndex)
 	{
 		FParameters& Parameters = ParameterArray[InvocationIndex];
-		Parameters.NumPrimitives = PrimitiveSceneProxies.Num();
+		Parameters.NumPrimitives = Primitives.Num();
 		Parameters.NumInstancesAllPrimitives = NumInstancesAllPrimitives;
 		Parameters.NumCustomFloatsPerInstance = NumCustomFloatsPerInstance;
 		Parameters.InstanceData = InstanceDataUAV;
@@ -261,9 +276,9 @@ void FPCGInstanceDataProviderProxy::GatherDispatchData(FDispatchData const& InDi
 		Parameters.WriteCounters = WriteCountersUAV;
 
 		uint32 CumulativeInstanceCount = 0;
-		for (int32 Index = 0; Index < PrimitiveSceneProxies.Num(); ++Index)
+		for (int32 Index = 0; Index < Primitives.Num(); ++Index)
 		{
-			const uint32 PrimitiveInstanceCount = PrimitiveSceneProxies[Index]->GetInstanceDataHeader().NumInstances;
+			const uint32 PrimitiveInstanceCount = (Primitives[Index].Get<1>() > 0) ? static_cast<uint32>(Primitives[Index].Get<1>()) : 0u;
 
 			Parameters.PrimitiveParameters[Index] = FUintVector4(
 				PrimitiveInstanceCount,
@@ -284,84 +299,103 @@ void FPCGInstanceDataProviderProxy::PostSubmit(FRDGBuilder& InGraphBuilder) cons
 		return;
 	}
 
-	const FSceneComputeUpdatesBlackboardEntry* SceneComputeUpdatesBB = InGraphBuilder.Blackboard.Get<FSceneComputeUpdatesBlackboardEntry>();
-	if (!SceneComputeUpdatesBB || !ensure(SceneComputeUpdatesBB->SceneComputeUpdates))
-	{
-		UE_LOG(LogPCG, Warning, TEXT("Scene Compute Updates interface missing from graph builder blackboard. Scene will not be updated."));
-		return;
-	}
-
-	uint32 CumulativeInstanceCount = 0;
-
-	// One dispatch per primitive. If the scene compute update API supported writing to multiple primitives this could be done in one unified call.
-	for (int32 PrimitiveIndex = 0; PrimitiveIndex < PrimitiveSceneProxies.Num(); ++PrimitiveIndex)
+	const TRefCountPtr<FRDGPooledBuffer> InstanceDataExported = InGraphBuilder.ConvertToExternalBuffer(InstanceData);
+	const TRefCountPtr<FRDGPooledBuffer> InstanceCustomFloatDataExported = InGraphBuilder.ConvertToExternalBuffer(InstanceCustomFloatData);
+	const TRefCountPtr<FRDGPooledBuffer> WriteCountersExported = InGraphBuilder.ConvertToExternalBuffer(WriteCounters);
+	
+	ExecuteOnGameThread(TEXT("ApplyPrimitiveSceneUpdates"), [DataProviderWeak=DataProvider, DataProviderGeneration=DataProviderGeneration, InstanceDataExported, InstanceCustomFloatDataExported, WriteCountersExported]()
 	{
-		FPrimitiveSceneProxy* PrimitiveSceneProxy = PrimitiveSceneProxies[PrimitiveIndex];
-		FPrimitiveSceneInfo* SceneInfo = PrimitiveSceneProxy ? PrimitiveSceneProxy->GetPrimitiveSceneInfo() : nullptr;
-		if (!ensure(SceneInfo))
+		TStrongObjectPtr<UPCGInstanceDataProvider> DataProvider = DataProviderWeak.Pin();
+		if (!DataProvider || DataProvider->GenerationCounter != DataProviderGeneration)
 		{
+			UE_LOG(LogPCG, Verbose, TEXT("Data provider object lost, GPU instancing will fail."));
 			return;
 		}
 
-#if PCG_INSTANCE_DATA_LOGGING
-		UE_LOG(LogPCG, Log, TEXT("\tProcess primitive [%d]:\tPrimitive ID %u\tProxy: %llu"), PrimitiveIndex, SceneInfo->GetPersistentIndex().Index, (uint64)this);
-#endif
+		// All instance data is stored in a single buffer, so this is used to give the scene writer an index to the first primitive.
+		int32 CumulativeInstanceCount = 0;
 
-		const int32 NumInstancesAllocatedInGPUScene = PrimitiveSceneProxy->GetInstanceDataHeader().NumInstances;
+		for (int PrimitiveIndex = 0; PrimitiveIndex < DataProvider->Primitives.Num(); ++PrimitiveIndex)
+		{
+			FPCGPrimitiveComponentAndNumInstances& PrimitiveAndInstanceCount = DataProvider->Primitives[PrimitiveIndex];
+
+			const int32 PrimitiveNumInstances = PrimitiveAndInstanceCount.Get<1>();
+			if (PrimitiveNumInstances <= 0)
+			{
+				UE_LOG(LogPCG, Warning, TEXT("Primitive with 0 instances encountered during instancing, should not happen."));
+				continue;
+			}
 
+			UPCGProceduralISMComponent* Primitive = Cast<UPCGProceduralISMComponent>(PrimitiveAndInstanceCount.Get<0>().Get());
+			UWorld* World = Primitive ? Primitive->GetWorld() : nullptr;
+
+			if (FSceneInterface* Scene = World ? World->Scene : nullptr)
+			{
 #if UE_BUILD_SHIPPING
-		const bool bTriggerCapture = false;
+				const bool bTriggerCapture = false;
 #else
-		const bool bTriggerCapture = (PCGInstanceDataInterface::TriggerGPUCaptureDispatchIndex == 1);
-		PCGInstanceDataInterface::TriggerGPUCaptureDispatchIndex = FMath::Max(PCGInstanceDataInterface::TriggerGPUCaptureDispatchIndex - 1, 0);
+				const bool bTriggerCapture = (PCGInstanceDataInterface::TriggerGPUCaptureDispatchIndex == 1);
+				PCGInstanceDataInterface::TriggerGPUCaptureDispatchIndex = FMath::Max(PCGInstanceDataInterface::TriggerGPUCaptureDispatchIndex - 1, 0);
 #endif // UE_BUILD_SHIPPING
 
-		// Post graph dispatch, the instance data will have been populated. Enqueue a scene update to write the instance data into the scene.
-		// Will get executed on same frame and in same builder.
-		SceneComputeUpdatesBB->SceneComputeUpdates->EnqueueUpdate(SceneInfo,
-			FGPUSceneWriteDelegate::CreateLambda([PrimitiveIndex, NumInstancesAllocatedInGPUScene, CumulativeInstanceCount, InstanceDataSRV=InstanceDataSRV, InstanceCustomFloatDataSRV=InstanceCustomFloatDataSRV, WriteCountersSRV=WriteCountersSRV, bTriggerCapture](FRDGBuilder& GraphBuilder, const FGPUSceneWriteDelegateParams& Params)
-			{
-				RDG_EVENT_SCOPE(GraphBuilder, "SceneComputeUpdateInterface->EnqueueUpdate");
-				check(Params.PersistentPrimitiveId != (uint32)INDEX_NONE);
+				// todo_pcg: Can produce these upstream and move away from requiring PISMC components.
+				FInstancedStaticMeshSceneProxyDesc SceneProxyDesc;
+				Primitive->GetSceneProxyDesc(SceneProxyDesc);
+
+				FPrimitiveSceneDesc PrimitiveSceneDesc;
+				Primitive->BuildSceneDesc(&SceneProxyDesc, PrimitiveSceneDesc);
+
+				Scene->UpdatePrimitiveInstancesFromCompute(&PrimitiveSceneDesc, FGPUSceneWriteDelegate::CreateLambda([PrimitiveIndex, PrimitiveNumInstances, CumulativeInstanceCount, bTriggerCapture, InstanceDataExported, InstanceCustomFloatDataExported, WriteCountersExported](FRDGBuilder& InGraphBuilder, const FGPUSceneWriteDelegateParams& Params)
+				{
+					RDG_EVENT_SCOPE(InGraphBuilder, "SceneComputeUpdateInterface->EnqueueUpdate");
+					check(Params.PersistentPrimitiveId != (uint32)INDEX_NONE);
 
 #if !UE_BUILD_SHIPPING
-				RenderCaptureInterface::FScopedCapture RenderCapture(bTriggerCapture, GraphBuilder, TEXT("SceneComputeUpdateInterface->EnqueueUpdate"));
+					RenderCaptureInterface::FScopedCapture RenderCapture(bTriggerCapture, InGraphBuilder, TEXT("SceneComputeUpdateInterface->EnqueueUpdate"));
 #endif // !UE_BUILD_SHIPPING
 
-				// Write instances
-				FPCGSceneWriterCS::FParameters* Parameters = GraphBuilder.AllocParameters<FPCGSceneWriterCS::FParameters>();
-				Parameters->InPrimitiveIndex = PrimitiveIndex;
-				Parameters->InNumInstancesAllocatedInGPUScene = NumInstancesAllocatedInGPUScene;
-				Parameters->InInstanceOffset = CumulativeInstanceCount;
-				Parameters->InInstanceData = InstanceDataSRV;
-				Parameters->InInstanceCustomFloatData = InstanceCustomFloatDataSRV;
-				Parameters->InWriteCounters = WriteCountersSRV;
-				Parameters->InPrimitiveId = Params.PersistentPrimitiveId;
-				PRAGMA_DISABLE_DEPRECATION_WARNINGS
-				Parameters->GPUSceneWriterParameters = Params.GPUWriteParams;
-				PRAGMA_ENABLE_DEPRECATION_WARNINGS
-				Parameters->InCustomDataCount = Params.NumCustomDataFloats;
-				Parameters->InPayloadDataFlags = Params.PackedInstanceSceneDataFlags;
-			
+					FRDGBufferRef InstanceData = InGraphBuilder.RegisterExternalBuffer(InstanceDataExported);
+					FRDGBufferRef InstanceCustomFloatData = InGraphBuilder.RegisterExternalBuffer(InstanceCustomFloatDataExported);
+					FRDGBufferRef WriteCounters = InGraphBuilder.RegisterExternalBuffer(WriteCountersExported);
+
+					// Write instances
+					FPCGSceneWriterCS::FParameters* Parameters = InGraphBuilder.AllocParameters<FPCGSceneWriterCS::FParameters>();
+					Parameters->InPrimitiveIndex = PrimitiveIndex;
+					Parameters->InNumInstancesAllocatedInGPUScene = PrimitiveNumInstances;
+					Parameters->InInstanceOffset = CumulativeInstanceCount;
+					Parameters->InInstanceData = InGraphBuilder.CreateSRV(InstanceData);
+					Parameters->InInstanceCustomFloatData = InGraphBuilder.CreateSRV(InstanceCustomFloatData);
+					Parameters->InWriteCounters = InGraphBuilder.CreateSRV(WriteCounters);
+					Parameters->InPrimitiveId = Params.PersistentPrimitiveId;
+					PRAGMA_DISABLE_DEPRECATION_WARNINGS
+					Parameters->GPUSceneWriterParameters = Params.GPUWriteParams;
+					PRAGMA_ENABLE_DEPRECATION_WARNINGS
+					Parameters->InCustomDataCount = Params.NumCustomDataFloats;
+					Parameters->InPayloadDataFlags = Params.PackedInstanceSceneDataFlags;
+
 #if PCG_INSTANCE_DATA_LOGGING
-				UE_LOG(LogPCG, Log, TEXT("\tScene writer delegate [%d]:\tPrimitive ID %u,\tsource instance offset %u,\tInstanceSceneDataOffset %u, num instances %u"),
-					Parameters->InPrimitiveIndex,
-					Parameters->InPrimitiveId,
-					Parameters->InInstanceOffset,
-					Params.InstanceSceneDataOffset,
-					Parameters->InNumInstancesAllocatedInGPUScene);
+					UE_LOG(LogPCG, Log, TEXT("\tScene writer delegate [%d]:\tPrimitive ID %u,\tsource instance offset %u,\tInstanceSceneDataOffset %u, num instances %u"),
+						Parameters->InPrimitiveIndex,
+						Parameters->InPrimitiveId,
+						Parameters->InInstanceOffset,
+						Params.InstanceSceneDataOffset,
+						Parameters->InNumInstancesAllocatedInGPUScene);
 #endif
 
-				// We're using a custom compute shader here rather than using a compute graph kernel, because ultimately the scene update will incorporate the
-				// new instance data. In the future we will not write directly to the scene here.
-				TShaderMapRef<FPCGSceneWriterCS> Shader(GetGlobalShaderMap(GMaxRHIFeatureLevel));
-				int GroupCount = FMath::DivideAndRoundUp<int>(NumInstancesAllocatedInGPUScene, FPCGSceneWriterCS::NumThreadsPerGroup);
-				ensureAlways(GroupCount <= GRHIMaxDispatchThreadGroupsPerDimension.X);
-				GroupCount = FMath::Min(GRHIMaxDispatchThreadGroupsPerDimension.X, GroupCount);
+					// We're using a custom compute shader here rather than using a compute graph kernel, because ultimately the scene update will incorporate the
+					// new instance data. In the future we will not write directly to the scene here.
+					TShaderMapRef<FPCGSceneWriterCS> Shader(GetGlobalShaderMap(GMaxRHIFeatureLevel));
+					int GroupCount = FMath::DivideAndRoundUp<int>(PrimitiveNumInstances, FPCGSceneWriterCS::NumThreadsPerGroup);
+					ensureAlways(GroupCount <= GRHIMaxDispatchThreadGroupsPerDimension.X);
+					GroupCount = FMath::Min(GRHIMaxDispatchThreadGroupsPerDimension.X, GroupCount);
 
-				FComputeShaderUtils::AddPass(GraphBuilder, RDG_EVENT_NAME("PCGWriteInstanceData"), ERDGPassFlags::Compute | ERDGPassFlags::NeverCull, Shader, Parameters, FIntVector(GroupCount, 1, 1));
-			}));
+					FComputeShaderUtils::AddPass(InGraphBuilder, RDG_EVENT_NAME("PCGWriteInstanceData"), ERDGPassFlags::Compute | ERDGPassFlags::NeverCull, Shader, Parameters, FIntVector(GroupCount, 1, 1));
+				}));
+			}
 
-		CumulativeInstanceCount += NumInstancesAllocatedInGPUScene;
-	}
+			CumulativeInstanceCount += PrimitiveNumInstances;
+		}
+
+		DataProvider->bWroteInstances = true;
+	});
 }
diff --git a/Engine/Plugins/PCG/Source/PCG/Private/Compute/DataInterfaces/PCGInstanceDataInterface.h b/Engine/Plugins/PCG/Source/PCG/Private/Compute/DataInterfaces/PCGInstanceDataInterface.h
--- a/Engine/Plugins/PCG/Source/PCG/Private/Compute/DataInterfaces/PCGInstanceDataInterface.h
+++ b/Engine/Plugins/PCG/Source/PCG/Private/Compute/DataInterfaces/PCGInstanceDataInterface.h
@@ -13,6 +13,8 @@ class FPrimitiveSceneProxy;
 class UPCGSettings;
 class UPCGStaticMeshSpawnerKernel;
 
+using FPCGPrimitiveComponentAndNumInstances = TPair<TWeakObjectPtr<UPrimitiveComponent>, int>;
+
 /** Data Interface for writing instance data (transform and custom floats) to an intermediate buffer, which is then injected into the GPU Scene. */
 UCLASS(ClassGroup = (Procedural))
 class UPCGInstanceDataInterface : public UPCGComputeDataInterface
@@ -49,6 +51,7 @@ public:
 	//~ Begin UComputeDataProvider Interface
 	virtual void Initialize(const UComputeDataInterface* InDataInterface, UObject* InBinding, uint64 InInputMask, uint64 InOutputMask) override;
 	virtual bool PrepareForExecute_GameThread(UPCGDataBinding* InBinding) override;
+	virtual bool PostExecute(UPCGDataBinding* InBinding) override;
 	FComputeDataProviderRenderProxy* GetRenderProxy() override;
 	virtual void Reset() override;
 	//~ End UComputeDataProvider Interface
@@ -56,20 +59,23 @@ public:
 	UPROPERTY()
 	TObjectPtr<const UPCGStaticMeshSpawnerKernel> ProducerKernel;
 
-	UPROPERTY()
-	TArray<TObjectPtr<UPrimitiveComponent>> Primitives;
+	TArray<FPCGPrimitiveComponentAndNumInstances> Primitives;
 
-	UPROPERTY()
 	uint32 NumInstancesAllPrimitives = 0;
 
-	UPROPERTY()
 	uint32 NumCustomFloatsPerInstance = 0;
+
+	/** Whether instance data has been applied to scene and operation is complete. */
+	bool bWroteInstances = false;
+
+	/** Bumped each time the data provider is initialized or reset, so that async callbacks can detect if they originated from a previous usage of the data provider and no-op. */
+	std::atomic<uint64> GenerationCounter = 0;
 };
 
 class FPCGInstanceDataProviderProxy : public FComputeDataProviderRenderProxy
 {
 public:
-	FPCGInstanceDataProviderProxy(TArray<FPrimitiveSceneProxy*>&& InPrimitiveSceneProxies, uint32 InNumInstances, uint32 InNumCustomFloatsPerInstance);
+	FPCGInstanceDataProviderProxy(TArray<FPCGPrimitiveComponentAndNumInstances> InPrimitives, uint32 InNumInstancesAllPrimitives, uint32 InNumCustomFloatsPerInstance, TWeakObjectPtr<UPCGInstanceDataProvider> InDataProvider/*, uint64 InDataProviderGeneration*/);
 
 	//~ Begin FComputeDataProviderRenderProxy Interface
 	bool IsValid(FValidationData const& InValidationData) const override;
@@ -81,7 +87,9 @@ public:
 protected:
 	using FParameters = FPCGInstanceDataInterfaceParameters;
 
-	TArray<FPrimitiveSceneProxy*> PrimitiveSceneProxies;
+	bool bIsValid = false;
+
+	TArray<FPCGPrimitiveComponentAndNumInstances> Primitives;
 
 	/** Instance transforms of all instances across all primitives. Three float4s per instance. */
 	FRDGBufferRef InstanceData = nullptr;
@@ -100,4 +108,7 @@ protected:
 
 	uint32 NumInstancesAllPrimitives = 0;
 	uint32 NumCustomFloatsPerInstance = 0;
+
+	TWeakObjectPtr<UPCGInstanceDataProvider> DataProvider;
+	uint64 DataProviderGeneration = 0;
 };
diff --git a/Engine/Plugins/PCG/Source/PCG/Private/Compute/Elements/PCGComputeGraphElement.cpp b/Engine/Plugins/PCG/Source/PCG/Private/Compute/Elements/PCGComputeGraphElement.cpp
--- a/Engine/Plugins/PCG/Source/PCG/Private/Compute/Elements/PCGComputeGraphElement.cpp
+++ b/Engine/Plugins/PCG/Source/PCG/Private/Compute/Elements/PCGComputeGraphElement.cpp
@@ -507,18 +507,9 @@ bool FPCGComputeGraphElement::ExecuteInternal(FPCGContext* InContext) const
 
 		Context->ExecutionSubPhase = EPCGComputeGraphExecutionPhase::WaitForExecutionComplete;
 
-		if (Context->HasPendingAsyncOperations())
-		{
-			// If we're running async operations (like preparing exports), go to sleep and let them wake us up later.
-			Context->bIsPaused = true;
-			return false;
-		}
-		else
-		{
-			// No operations to wait for, so signal completion.
-			Context->bExecutionSuccess = true;
-			return true;
-		}
+		SleepUntilNextFrame();
+
+		return false;
 	}
 	case EPCGComputeGraphExecutionPhase::WaitForExecutionComplete:
 	{
@@ -530,44 +521,84 @@ bool FPCGComputeGraphElement::ExecuteInternal(FPCGContext* InContext) const
 			return true;
 		}
 
-		// Execution is complete when graph was enqueued/dispatched and all exports and readbacks are complete.
-		if (!Context->HasPendingAsyncOperations())
+		if (Context->HasPendingAsyncOperations())
 		{
-			// Apply any read back element counts to the final data description.
-			Context->DataBinding->ApplyElementCounts_GameThread();
+			// Still running. Likely we need a frame to pass in order to make progress with readbacks etc.
+			SleepUntilNextFrame();
 
-#if WITH_EDITOR
-			// Call both to start all required readbacks.
-			const bool bDebugDrawDone = ExecuteDebugDraw(Context);
-			const bool bInspectionDone = StoreDataForInspection(Context);
+			return false;
+		}
+
+		Context->bExecutionSuccess = true;
 
-			if (!bDebugDrawDone || !bInspectionDone)
+		// Apply any read back element counts to the final data description.
+		Context->DataBinding->ApplyElementCounts_GameThread();
+
+		Context->DataProvidersPendingPostExecute.Reserve(Context->ComputeGraphInstance->GetNumDataProviders());
+		for (UComputeDataProvider* Provider : Context->ComputeGraphInstance->GetDataProviders())
+		{
+			UPCGComputeDataProvider* PCGProvider = Cast<UPCGComputeDataProvider>(Provider);
+			if (PCGProvider)
 			{
-				SleepUntilNextFrame();
-				return false;
+				Context->DataProvidersPendingPostExecute.Add(PCGProvider);
 			}
-#endif
+		}
+
+		Context->ExecutionSubPhase = EPCGComputeGraphExecutionPhase::PostExecute;
 
-			Context->bExecutionSuccess = true;
+		if (Context->AsyncState.ShouldStop())
+		{
+			return false;
+		}
+	}
+	case EPCGComputeGraphExecutionPhase::PostExecute: // Fallthrough
+	{
+		for (int Index = Context->DataProvidersPendingPostExecute.Num() - 1; Index >= 0; --Index)
+		{
+			UPCGComputeDataProvider* DataProvider = Context->DataProvidersPendingPostExecute[Index];
 
-			// Currently we don't output anything if processing any readback data processing failed.
-			if (ensure(Context->bExecutionSuccess) && ensure(Context->DataBinding))
+			if (!DataProvider || DataProvider->PostExecute(Context->DataBinding.Get()))
 			{
-				Context->OutputData = Context->DataBinding->OutputDataCollection;
+				Context->DataProvidersPendingPostExecute.RemoveAtSwap(Index);
 			}
+		}
 
-			// Signal completion.
-			return true;
+		if (!Context->DataProvidersPendingPostExecute.IsEmpty())
+		{
+			SleepUntilNextFrame();
+			return false;
 		}
 
-		// Still running. Likely we need a frame to pass in order to make progress with readbacks etc.
-		SleepUntilNextFrame();
+		// Currently we don't output anything if processing any readback data processing failed.
+		if (ensure(Context->bExecutionSuccess) && ensure(Context->DataBinding))
+		{
+			Context->OutputData = Context->DataBinding->OutputDataCollection;
+		}
 
-		return false;
+		Context->ExecutionSubPhase = EPCGComputeGraphExecutionPhase::DebugAndInspection;
+	}
+	case EPCGComputeGraphExecutionPhase::DebugAndInspection: // Fallthrough
+	{
+#if WITH_EDITOR
+		if (Context->AsyncState.ShouldStop())
+		{
+			return false;
+		}
+
+		// Call both to start all required readbacks.
+		const bool bDebugDrawDone = ExecuteDebugDraw(Context);
+		const bool bInspectionDone = StoreDataForInspection(Context);
+
+		if (!bDebugDrawDone || !bInspectionDone)
+		{
+			SleepUntilNextFrame();
+			return false;
+		}
+#endif
 	}
 	}
 
-	return false;
+	return true;
 }
 
 void FPCGComputeGraphElement::PostExecuteInternal(FPCGContext* InContext) const
diff --git a/Engine/Plugins/PCG/Source/PCG/Private/Compute/Elements/PCGComputeGraphElement.h b/Engine/Plugins/PCG/Source/PCG/Private/Compute/Elements/PCGComputeGraphElement.h
--- a/Engine/Plugins/PCG/Source/PCG/Private/Compute/Elements/PCGComputeGraphElement.h
+++ b/Engine/Plugins/PCG/Source/PCG/Private/Compute/Elements/PCGComputeGraphElement.h
@@ -28,6 +28,8 @@ enum class EPCGComputeGraphExecutionPhase : uint8
 	ValidateComputeGraphCompilation,
 	ScheduleComputeGraph,
 	WaitForExecutionComplete,
+	PostExecute,
+	DebugAndInspection,
 };
 
 struct FPCGComputeGraphContext : public FPCGContext
@@ -68,6 +70,8 @@ public:
 	/** The data providers which are not yet ready for execution. */
 	TArray<UPCGComputeDataProvider*> DataProvidersPendingReadyForExecute;
 
+	TArray<UPCGComputeDataProvider*> DataProvidersPendingPostExecute;
+
 	/** Data providers with buffers that are passed to downstream tasks. The buffer will be created on the render thread, and then passed
 	* back to main thread, upon which a reference is taken to the buffer and the provider is removed from this set to signal completion.
 	*/
diff --git a/Engine/Plugins/PCG/Source/PCG/Private/Compute/PCGDataBinding.h b/Engine/Plugins/PCG/Source/PCG/Private/Compute/PCGDataBinding.h
--- a/Engine/Plugins/PCG/Source/PCG/Private/Compute/PCGDataBinding.h
+++ b/Engine/Plugins/PCG/Source/PCG/Private/Compute/PCGDataBinding.h
@@ -17,36 +17,39 @@ struct FPCGContextHandle;
 struct FPCGKernelPin;
 
 USTRUCT()
-struct FPCGSpawnerPrimitives
+struct FPCGSpawnerPrimitive
 {
 	GENERATED_BODY()
 
 	UPROPERTY()
-	TArray<TObjectPtr<UPrimitiveComponent>> Primitives;
+	TObjectPtr<UPrimitiveComponent> Component;
+
+	UPROPERTY()
+	int NumInstances = INDEX_NONE;
+};
+
+USTRUCT()
+struct FPCGSpawnerPrimitives
+{
+	GENERATED_BODY()
 
 	UPROPERTY()
+	TArray<FPCGSpawnerPrimitive> Primitives;
+
 	TArray<FBox> PrimitiveMeshBounds;
 
 	/** Cumulative distribution function values (one per primitive) to enable choosing a primitive based on a random draw value. */
-	UPROPERTY()
 	TArray<float> SelectionCDF;
 
 	// Same for all primitives
-	UPROPERTY()
 	uint32 NumCustomFloats = 0;
 
 	// Same for all primitives
-	UPROPERTY()
 	TArray<FUintVector4> AttributeIdOffsetStrides;
 
-	UPROPERTY()
 	int32 SelectorAttributeId = INDEX_NONE;
 
-	UPROPERTY()
 	TArray<int32> PrimitiveStringKeys;
-
-	UPROPERTY()
-	uint32 NumInstancesAllPrimitives = 0;
 };
 
 USTRUCT()
diff --git a/Engine/Plugins/Runtime/ComputeFramework/Source/ComputeFramework/Public/ComputeFramework/ComputeGraphInstance.h b/Engine/Plugins/Runtime/ComputeFramework/Source/ComputeFramework/Public/ComputeFramework/ComputeGraphInstance.h
--- a/Engine/Plugins/Runtime/ComputeFramework/Source/ComputeFramework/Public/ComputeFramework/ComputeGraphInstance.h
+++ b/Engine/Plugins/Runtime/ComputeFramework/Source/ComputeFramework/Public/ComputeFramework/ComputeGraphInstance.h
@@ -46,6 +46,9 @@ public:
 	/** Get the Data Provider objects. */
 	TArray< TObjectPtr<UComputeDataProvider> >& GetDataProviders() { return DataProviders; }
 
+	/** Get the number of Data Provider objects. */
+	int GetNumDataProviders() const { return DataProviders.Num(); }
+
 	/** Enqueue the ComputeGraph work. */
 	bool EnqueueWork(UComputeGraph* InComputeGraph, FSceneInterface const* InScene, FName InExecutionGroupName, FName InOwnerName, FSimpleDelegate InFallbackDelegate, UObject* InOwnerPointer = nullptr, uint8 InGraphSortPriorityOffset = 0);
 
diff --git a/Engine/Source/Runtime/Engine/Public/SceneInterface.h b/Engine/Source/Runtime/Engine/Public/SceneInterface.h
--- a/Engine/Source/Runtime/Engine/Public/SceneInterface.h
+++ b/Engine/Source/Runtime/Engine/Public/SceneInterface.h
@@ -67,6 +67,7 @@ template<int32 MaxSHOrder> class TSHVectorRGB;
 using FSHVectorRGB3 = TSHVectorRGB<3>;
 struct FCustomPrimitiveData;
 class FSceneViewFamily;
+class FGPUSceneWriteDelegate;
 
 struct FExponentialHeightFogDynamicState;
 struct FPrimitiveSceneDesc;
@@ -553,6 +554,9 @@ public:
 
 	virtual void UpdatePrimitiveInstances(FPrimitiveSceneDesc* Primitive) = 0;	
 
+	UE_EXPERIMENTAL(5.6, "Experimental API, scene writing delegate may be internalized in the future.")
+	virtual void UpdatePrimitiveInstancesFromCompute(FPrimitiveSceneDesc* Primitive, FGPUSceneWriteDelegate&& DataWriterGPU) = 0;
+
 	UE_DEPRECATED(5.6, "StaticMeshBounds is no longer a parameter of this interface, please update FPrimitiveSceneDesc version.")
 	virtual void UpdatePrimitiveInstances(FInstancedStaticMeshSceneDesc* Primitive) = 0;	
 
diff --git a/Engine/Source/Runtime/RenderCore/Private/RenderCore.cpp b/Engine/Source/Runtime/RenderCore/Private/RenderCore.cpp
--- a/Engine/Source/Runtime/RenderCore/Private/RenderCore.cpp
+++ b/Engine/Source/Runtime/RenderCore/Private/RenderCore.cpp
@@ -198,6 +198,7 @@ DEFINE_STAT(STAT_RemoveScenePrimitiveGT);
 DEFINE_STAT(STAT_AddScenePrimitiveGT);
 DEFINE_STAT(STAT_UpdatePrimitiveTransformGT);
 DEFINE_STAT(STAT_UpdatePrimitiveInstanceGT);
+DEFINE_STAT(STAT_UpdatePrimitiveInstanceFromComputeGT);
 
 DEFINE_STAT(STAT_Scene_SetShaderMapsOnMaterialResources_RT);
 DEFINE_STAT(STAT_Scene_UpdateStaticDrawLists_RT);
diff --git a/Engine/Source/Runtime/RenderCore/Public/RenderCore.h b/Engine/Source/Runtime/RenderCore/Public/RenderCore.h
--- a/Engine/Source/Runtime/RenderCore/Public/RenderCore.h
+++ b/Engine/Source/Runtime/RenderCore/Public/RenderCore.h
@@ -135,6 +135,7 @@ DECLARE_CYCLE_STAT_EXTERN(TEXT("RemovePrimitive (GT)"),STAT_RemoveScenePrimitive
 DECLARE_CYCLE_STAT_EXTERN(TEXT("AddPrimitive (GT)"),STAT_AddScenePrimitiveGT,STATGROUP_SceneUpdate, RENDERCORE_API);
 DECLARE_CYCLE_STAT_EXTERN(TEXT("UpdatePrimitiveTransform (GT)"),STAT_UpdatePrimitiveTransformGT,STATGROUP_SceneUpdate, RENDERCORE_API);
 DECLARE_CYCLE_STAT_EXTERN(TEXT("UpdatePrimitiveInstance (GT)"), STAT_UpdatePrimitiveInstanceGT, STATGROUP_SceneUpdate, RENDERCORE_API);
+DECLARE_CYCLE_STAT_EXTERN(TEXT("UpdatePrimitiveInstanceFromCompute (GT)"), STAT_UpdatePrimitiveInstanceFromComputeGT, STATGROUP_SceneUpdate, RENDERCORE_API);
 
 DECLARE_CYCLE_STAT_EXTERN(TEXT("Set Shader Maps On Material Resources (RT)"),STAT_Scene_SetShaderMapsOnMaterialResources_RT,STATGROUP_SceneUpdate, RENDERCORE_API);
 DECLARE_CYCLE_STAT_EXTERN(TEXT("Update Static Draw Lists (RT)"), STAT_Scene_UpdateStaticDrawLists_RT, STATGROUP_SceneUpdate, RENDERCORE_API);
diff --git a/Engine/Source/Runtime/Renderer/Internal/SceneComputeUpdates.h b/Engine/Source/Runtime/Renderer/Internal/SceneComputeUpdates.h
deleted file mode 100644
diff --git a/Engine/Source/Runtime/Renderer/Private/RendererScene.cpp b/Engine/Source/Runtime/Renderer/Private/RendererScene.cpp
--- a/Engine/Source/Runtime/Renderer/Private/RendererScene.cpp
+++ b/Engine/Source/Runtime/Renderer/Private/RendererScene.cpp
@@ -99,7 +99,6 @@
 #include "Nanite/NaniteMaterialsSceneExtension.h"
 #include "Nanite/NaniteSkinningSceneExtension.h"
 #include "ObjectCacheContext.h"
-#include "SceneComputeUpdates.h"
 #include "Animation/AnimBank.h"
 #include "Skinning/AnimBankTransformProvider.h"
 #include "Skinning/SkinningTransformProvider.h"
@@ -1706,6 +1705,25 @@ void FScene::UpdatePrimitiveInstances(UPrimitiveComponent* Primitive)
 	UpdatePrimitiveInstances(UpdateParams);
 }
 
+void FScene::UpdatePrimitiveInstancesFromCompute(FPrimitiveSceneDesc* Primitive, FGPUSceneWriteDelegate&& DataWriterGPU)
+{
+	SCOPE_CYCLE_COUNTER(STAT_UpdatePrimitiveInstanceFromComputeGT);
+	check(Primitive);
+
+	FPrimitiveSceneProxy* SceneProxy = Primitive->GetSceneProxy();
+
+	if (!ensureMsgf(SceneProxy, TEXT("Primitive must be added to scene prior to updating its instances from compute.")))
+	{
+		return;
+	}
+
+	FUpdateInstanceFromComputeCommand UpdateCommand;
+	UpdateCommand.PrimitiveSceneProxy = SceneProxy;
+	UpdateCommand.GPUSceneWriter = MoveTemp(DataWriterGPU);
+
+	UpdatePrimitiveInternal(SceneProxy, MoveTemp(UpdateCommand));
+}
+
 void FScene::UpdatePrimitiveInstancesFromCompute(FPrimitiveSceneInfo* PrimitiveSceneInfo, FGPUSceneWriteDelegate&& DataWriterGPU)
 {
 	SCOPED_NAMED_EVENT(FScene_UpdatePrimitiveInstanceFromCompute, FColor::Yellow);
@@ -4988,25 +5006,6 @@ bool FScene::CanSampleSkyLightRealTimeCaptureData() const
 	return SkyLight && SkyLight->bRealTimeCaptureEnabled && ConvolvedSkyRenderTargetReadyIndex >= 0;
 }
 
-class FSceneComputeUpdates : public ISceneComputeUpdates
-{
-public:
-	FSceneComputeUpdates(FScene* InScene) : Scene(InScene) {}
-
-	virtual void SetScene(FScene* InScene) override { Scene = InScene; }
-
-	virtual void EnqueueUpdateInternal(FPrimitiveSceneInfo* PrimitiveSceneInfo, FGPUSceneWriteDelegate&& DataWriterGPU) override
-	{
-		if (ensure(Scene))
-		{
-			Scene->UpdatePrimitiveInstancesFromCompute(PrimitiveSceneInfo, MoveTemp(DataWriterGPU));
-		}
-	}
-
-private:
-	FScene* Scene = nullptr;
-};
-
 FViewSceneChangeSet* FScene::ProcessViewChanges(FRDGBuilder& GraphBuilder, TConstArrayView<FViewInfo*> Views)
 {
 	FViewSceneChangeSet* ChangeSet = GraphBuilder.AllocObject<FViewSceneChangeSet>();
@@ -5106,18 +5105,6 @@ void FScene::Update(FRDGBuilder& GraphBuilder, const FUpdateParameters& Paramete
 		GPUSkinCacheTask = GPUSkinCache->Dispatch(GraphBuilder, GPUSkinTasks.SkinCache, GPUSkinCachePipeline);
 	}
 
-	// Add interface to blackboard to receive scene updates from compute.
-	if (const FSceneComputeUpdatesBlackboardEntry* FoundBlackboardEntry = GraphBuilder.Blackboard.Get<FSceneComputeUpdatesBlackboardEntry>())
-	{
-		check(FoundBlackboardEntry->SceneComputeUpdates);
-		FoundBlackboardEntry->SceneComputeUpdates->SetScene(this);
-	}
-	else
-	{
-		FSceneComputeUpdatesBlackboardEntry& BlackboardEntry = GraphBuilder.Blackboard.Create<FSceneComputeUpdatesBlackboardEntry>();
-		BlackboardEntry.SceneComputeUpdates = GraphBuilder.AllocObject<FSceneComputeUpdates>(this);
-	}
-
 	for (IComputeTaskWorker* ComputeTaskWorker : ComputeTaskWorkers)
 	{
 		if (ComputeTaskWorker->HasWork(ComputeTaskExecutionGroup::EndOfFrameUpdate))
@@ -5126,12 +5113,6 @@ void FScene::Update(FRDGBuilder& GraphBuilder, const FUpdateParameters& Paramete
 		}
 	}
 
-	// Disable scene updates on this builder from this point onwards.
-	{
-		const FSceneComputeUpdatesBlackboardEntry& BlackboardEntry = GraphBuilder.Blackboard.GetChecked<FSceneComputeUpdatesBlackboardEntry>();
-		BlackboardEntry.SceneComputeUpdates->SetScene(nullptr);
-	}
-
 	UE::Tasks::FTask UpdateUniformExpressionsTask;
 	FMaterialRenderProxy::UpdateDeferredCachedUniformExpressions(GraphBuilder.RHICmdList, EnumHasAnyFlags(Parameters.AsyncOps, EUpdateAllPrimitiveSceneInfosAsyncOps::CacheMaterialUniformExpressions) ? &UpdateUniformExpressionsTask : nullptr);
 
@@ -6471,6 +6452,7 @@ public:
 	/** Updates the transform of a primitive which has already been added to the scene. */
 	virtual void UpdatePrimitiveTransform(UPrimitiveComponent* Primitive) override {}
 	virtual void UpdatePrimitiveInstances(UPrimitiveComponent* Primitive) override {}
+	virtual void UpdatePrimitiveInstancesFromCompute(FPrimitiveSceneDesc* Primitive, FGPUSceneWriteDelegate&& DataWriterGPU) override {}
 	virtual void UpdatePrimitiveOcclusionBoundsSlack(UPrimitiveComponent* Primitive, float NewSlack) override {}
 	virtual void UpdatePrimitiveDrawDistance(UPrimitiveComponent* Primitive, float MinDrawDistance, float MaxDrawDistance, float VirtualTextureMaxDrawDistance) override {}
 	virtual void UpdateInstanceCullDistance(UPrimitiveComponent* Primitive, float StartCullDistance, float EndCullDistance) {}
diff --git a/Engine/Source/Runtime/Renderer/Private/ScenePrivate.h b/Engine/Source/Runtime/Renderer/Private/ScenePrivate.h
--- a/Engine/Source/Runtime/Renderer/Private/ScenePrivate.h
+++ b/Engine/Source/Runtime/Renderer/Private/ScenePrivate.h
@@ -3366,6 +3366,7 @@ public:
 	virtual void UpdateAllPrimitiveSceneInfos(FRDGBuilder& GraphBuilder, EUpdateAllPrimitiveSceneInfosAsyncOps AsyncOps = EUpdateAllPrimitiveSceneInfosAsyncOps::None) override;
 	virtual void UpdatePrimitiveTransform(UPrimitiveComponent* Primitive) override;
 	virtual void UpdatePrimitiveInstances(UPrimitiveComponent* Primitive) override;
+	virtual void UpdatePrimitiveInstancesFromCompute(FPrimitiveSceneDesc* Primitive, FGPUSceneWriteDelegate&& DataWriterGPU) override;
 	virtual void UpdatePrimitiveOcclusionBoundsSlack(UPrimitiveComponent* Primitive, float NewSlack) override;
 	virtual void UpdatePrimitiveDrawDistance(UPrimitiveComponent* Primitive, float MinDrawDistance, float MaxDrawDistance, float VirtualTextureMaxDrawDistance) override;
 	virtual void UpdateInstanceCullDistance(UPrimitiveComponent* Primitive, float StartCullDistance, float EndCullDistance) override;
@@ -3886,7 +3887,6 @@ public:
 protected:
 
 private:
-	friend class FSceneComputeUpdates;
 	void UpdatePrimitiveInstancesFromCompute(FPrimitiveSceneInfo* PrimitiveSceneInfo, FGPUSceneWriteDelegate&& DataWriterGPU);
 
 	template<class T> 	
