[PCG] Rule out GC from GetPartitionActorFromPool crash

#rnx

diff --git a/Engine/Plugins/PCG/Source/PCG/Private/RuntimeGen/PCGRuntimeGenScheduler.cpp b/Engine/Plugins/PCG/Source/PCG/Private/RuntimeGen/PCGRuntimeGenScheduler.cpp
--- a/Engine/Plugins/PCG/Source/PCG/Private/RuntimeGen/PCGRuntimeGenScheduler.cpp
+++ b/Engine/Plugins/PCG/Source/PCG/Private/RuntimeGen/PCGRuntimeGenScheduler.cpp
@@ -1655,4 +1655,7 @@ void FPCGRuntimeGenScheduler::AddReferencedObjects(FReferenceCollector& Collecto
 	{
 		GenSourceManager->AddReferencedObjects(Collector);
 	}
+
+	// The level should be keeping the pooled PAs visible to GC. This is just a tentative fix for a crash in GetPartitionActorFromPool(), to understand if the crash is happening because of unreferenced GCed actors.
+	Collector.AddReferencedObjects(PartitionActorPool);
 }
diff --git a/Engine/Plugins/PCG/Source/PCG/Public/RuntimeGen/PCGRuntimeGenScheduler.h b/Engine/Plugins/PCG/Source/PCG/Public/RuntimeGen/PCGRuntimeGenScheduler.h
--- a/Engine/Plugins/PCG/Source/PCG/Public/RuntimeGen/PCGRuntimeGenScheduler.h
+++ b/Engine/Plugins/PCG/Source/PCG/Public/RuntimeGen/PCGRuntimeGenScheduler.h
@@ -147,7 +147,7 @@ private:
 	TSet<FGridGenerationKey> GeneratedComponentsToRemove;
 
 	/** Pool of RuntimeGen PartitionActors used for hierarchical generation. */
-	TArray<APCGPartitionActor*> PartitionActorPool;
+	TArray<TObjectPtr<APCGPartitionActor>> PartitionActorPool;
 
 	/** PartitionActorPoolSize represents the current maximum capacity of the PartitionActorPool. */
 	int32 PartitionActorPoolSize = 0;
