[PCG] CIS fixes

#rnx

diff --git a/Engine/Plugins/PCG/Source/PCG/Private/Compute/DataInterfaces/PCGInstanceDataInterface.cpp b/Engine/Plugins/PCG/Source/PCG/Private/Compute/DataInterfaces/PCGInstanceDataInterface.cpp
--- a/Engine/Plugins/PCG/Source/PCG/Private/Compute/DataInterfaces/PCGInstanceDataInterface.cpp
+++ b/Engine/Plugins/PCG/Source/PCG/Private/Compute/DataInterfaces/PCGInstanceDataInterface.cpp
@@ -305,8 +305,8 @@ void FPCGInstanceDataProviderProxy::PostSubmit(FRDGBuilder& InGraphBuilder) cons
 	
 	ExecuteOnGameThread(TEXT("ApplyPrimitiveSceneUpdates"), [DataProviderWeak=DataProvider, DataProviderGeneration=DataProviderGeneration, InstanceDataExported, InstanceCustomFloatDataExported, WriteCountersExported]()
 	{
-		TStrongObjectPtr<UPCGInstanceDataProvider> DataProvider = DataProviderWeak.Pin();
-		if (!DataProvider || DataProvider->GenerationCounter != DataProviderGeneration)
+		TStrongObjectPtr<UPCGInstanceDataProvider> DataProviderInner = DataProviderWeak.Pin();
+		if (!DataProviderInner || DataProviderInner->GenerationCounter != DataProviderGeneration)
 		{
 			UE_LOG(LogPCG, Verbose, TEXT("Data provider object lost, GPU instancing will fail."));
 			return;
@@ -315,9 +315,9 @@ void FPCGInstanceDataProviderProxy::PostSubmit(FRDGBuilder& InGraphBuilder) cons
 		// All instance data is stored in a single buffer, so this is used to give the scene writer an index to the first primitive.
 		int32 CumulativeInstanceCount = 0;
 
-		for (int PrimitiveIndex = 0; PrimitiveIndex < DataProvider->Primitives.Num(); ++PrimitiveIndex)
+		for (int PrimitiveIndex = 0; PrimitiveIndex < DataProviderInner->Primitives.Num(); ++PrimitiveIndex)
 		{
-			FPCGPrimitiveComponentAndNumInstances& PrimitiveAndInstanceCount = DataProvider->Primitives[PrimitiveIndex];
+			FPCGPrimitiveComponentAndNumInstances& PrimitiveAndInstanceCount = DataProviderInner->Primitives[PrimitiveIndex];
 
 			const int32 PrimitiveNumInstances = PrimitiveAndInstanceCount.Get<1>();
 			if (PrimitiveNumInstances <= 0)
@@ -354,18 +354,18 @@ void FPCGInstanceDataProviderProxy::PostSubmit(FRDGBuilder& InGraphBuilder) cons
 					RenderCaptureInterface::FScopedCapture RenderCapture(bTriggerCapture, InGraphBuilder, TEXT("SceneComputeUpdateInterface->EnqueueUpdate"));
 #endif // !UE_BUILD_SHIPPING
 
-					FRDGBufferRef InstanceData = InGraphBuilder.RegisterExternalBuffer(InstanceDataExported);
-					FRDGBufferRef InstanceCustomFloatData = InGraphBuilder.RegisterExternalBuffer(InstanceCustomFloatDataExported);
-					FRDGBufferRef WriteCounters = InGraphBuilder.RegisterExternalBuffer(WriteCountersExported);
+					FRDGBufferRef InstanceDataInner  = InGraphBuilder.RegisterExternalBuffer(InstanceDataExported);
+					FRDGBufferRef InstanceCustomFloatDataInner = InGraphBuilder.RegisterExternalBuffer(InstanceCustomFloatDataExported);
+					FRDGBufferRef WriteCountersInner = InGraphBuilder.RegisterExternalBuffer(WriteCountersExported);
 
 					// Write instances
 					FPCGSceneWriterCS::FParameters* Parameters = InGraphBuilder.AllocParameters<FPCGSceneWriterCS::FParameters>();
 					Parameters->InPrimitiveIndex = PrimitiveIndex;
 					Parameters->InNumInstancesAllocatedInGPUScene = PrimitiveNumInstances;
 					Parameters->InInstanceOffset = CumulativeInstanceCount;
-					Parameters->InInstanceData = InGraphBuilder.CreateSRV(InstanceData);
-					Parameters->InInstanceCustomFloatData = InGraphBuilder.CreateSRV(InstanceCustomFloatData);
-					Parameters->InWriteCounters = InGraphBuilder.CreateSRV(WriteCounters);
+					Parameters->InInstanceData = InGraphBuilder.CreateSRV(InstanceDataInner);
+					Parameters->InInstanceCustomFloatData = InGraphBuilder.CreateSRV(InstanceCustomFloatDataInner);
+					Parameters->InWriteCounters = InGraphBuilder.CreateSRV(WriteCountersInner);
 					Parameters->InPrimitiveId = Params.PersistentPrimitiveId;
 					PRAGMA_DISABLE_DEPRECATION_WARNINGS
 					Parameters->GPUSceneWriterParameters = Params.GPUWriteParams;
@@ -396,6 +396,6 @@ void FPCGInstanceDataProviderProxy::PostSubmit(FRDGBuilder& InGraphBuilder) cons
 			CumulativeInstanceCount += PrimitiveNumInstances;
 		}
 
-		DataProvider->bWroteInstances = true;
+		DataProviderInner->bWroteInstances = true;
 	});
 }
