[PCG] Add BaseTexture support to data visualization viewport

diff --git a/Engine/Plugins/PCG/Source/PCG/Private/Data/PCGRenderTargetData.cpp b/Engine/Plugins/PCG/Source/PCG/Private/Data/PCGRenderTargetData.cpp
--- a/Engine/Plugins/PCG/Source/PCG/Private/Data/PCGRenderTargetData.cpp
+++ b/Engine/Plugins/PCG/Source/PCG/Private/Data/PCGRenderTargetData.cpp
@@ -65,6 +65,11 @@ void UPCGRenderTargetData::AddToCrc(FArchiveCrc32& Ar, bool bFullDataCrc) const
 	AddUIDToCrc(Ar);
 }
 
+UTexture* UPCGRenderTargetData::GetTexture() const
+{
+	return RenderTarget;
+}
+
 FTextureRHIRef UPCGRenderTargetData::GetTextureRHI() const
 {
 	// TODO: This makes no attempt to actually acquire the resource after it has been written or, nor does it ensure
diff --git a/Engine/Plugins/PCG/Source/PCG/Public/Data/PCGRenderTargetData.h b/Engine/Plugins/PCG/Source/PCG/Public/Data/PCGRenderTargetData.h
--- a/Engine/Plugins/PCG/Source/PCG/Public/Data/PCGRenderTargetData.h
+++ b/Engine/Plugins/PCG/Source/PCG/Public/Data/PCGRenderTargetData.h
@@ -8,6 +8,7 @@
 
 #include "PCGRenderTargetData.generated.h"
 
+class UTexture;
 class UTextureRenderTarget2D;
 
 //TODO: It's possible that caching the result in this class is not as efficient as it could be
@@ -27,6 +28,7 @@ public:
 	// ~End UPCGData interface
 
 	//~Begin UPCGBaseTextureData interface
+	virtual UTexture* GetTexture() const override;
 	virtual FTextureRHIRef GetTextureRHI() const override;
 	virtual EPCGTextureResourceType GetTextureResourceType() const override { return EPCGTextureResourceType::TextureObject; }
 	//~End UPCGBaseTextureData interface
diff --git a/Engine/Plugins/PCG/Source/PCG/Public/Data/PCGTextureData.h b/Engine/Plugins/PCG/Source/PCG/Public/Data/PCGTextureData.h
--- a/Engine/Plugins/PCG/Source/PCG/Public/Data/PCGTextureData.h
+++ b/Engine/Plugins/PCG/Source/PCG/Public/Data/PCGTextureData.h
@@ -91,6 +91,7 @@ public:
 
 	virtual bool IsValid() const;
 
+	virtual UTexture* GetTexture() const PURE_VIRTUAL(UPCGBaseTextureData::GetTexture, return nullptr;)
 	virtual FTextureRHIRef GetTextureRHI() const PURE_VIRTUAL(UPCGBaseTextureData::GetTextureResource, return nullptr;);
 	virtual EPCGTextureResourceType GetTextureResourceType() const PURE_VIRTUAL(UPCGBaseTextureData::GetTextureResourceType, return EPCGTextureResourceType::Invalid;);
 	virtual TRefCountPtr<IPooledRenderTarget> GetRefCountedTexture() const { return nullptr; }
@@ -200,6 +201,7 @@ public:
 	//~End UPCGData interface
 
 	//~Begin UPCGBaseTextureData interface
+	virtual UTexture* GetTexture() const override { return Texture.IsValid() ? Texture.Get() : nullptr; }
 	virtual FTextureRHIRef GetTextureRHI() const override;
 	virtual TRefCountPtr<IPooledRenderTarget> GetRefCountedTexture() const { return TextureHandle; }
 	virtual int GetTextureSlice() const override { return TextureIndex; }
diff --git a/Engine/Plugins/PCG/Source/PCGEditor/PCGEditor.Build.cs b/Engine/Plugins/PCG/Source/PCGEditor/PCGEditor.Build.cs
--- a/Engine/Plugins/PCG/Source/PCGEditor/PCGEditor.Build.cs
+++ b/Engine/Plugins/PCG/Source/PCGEditor/PCGEditor.Build.cs
@@ -53,6 +53,7 @@ namespace UnrealBuildTool.Rules
 					"Kismet",
 					"KismetWidgets",
 					"PCG",
+					"ProceduralMeshComponent",
 					"PropertyEditor",
 					"RenderCore",
 					"Slate",
diff --git a/Engine/Plugins/PCG/Source/PCGEditor/Private/DataVisualizations/PCGBaseTextureDataVisualization.cpp b/Engine/Plugins/PCG/Source/PCGEditor/Private/DataVisualizations/PCGBaseTextureDataVisualization.cpp
new file mode 100644
--- /dev/null
+++ b/Engine/Plugins/PCG/Source/PCGEditor/Private/DataVisualizations/PCGBaseTextureDataVisualization.cpp
@@ -0,0 +1,89 @@
+// Copyright Epic Games, Inc. All Rights Reserved.
+
+#include "DataVisualizations/PCGBaseTextureDataVisualization.h"
+
+#include "Data/PCGTextureData.h"
+
+#include "AdvancedPreviewScene.h"
+#include "ProceduralMeshComponent.h"
+#include "Engine/AssetManager.h"
+#include "Engine/Texture.h"
+#include "Materials/MaterialInstanceDynamic.h"
+
+namespace PCGBaseTextureVisualizationConstants
+{
+	static constexpr float DebugQuadSize = 500.0f;
+
+	static const TArray<FVector> DebugMeshVertices =
+	{
+		FVector(-1, 0, -1) * DebugQuadSize,
+		FVector( 1, 0, -1) * DebugQuadSize,
+		FVector(-1, 0,  1) * DebugQuadSize,
+		FVector( 1, 0,  1) * DebugQuadSize
+	};
+
+	static const TArray<int32> DebugMeshTriangles = { 0, 1, 2, 1, 3, 2 };
+
+	static const TArray<FVector2D> DebugMeshUVs =
+	{
+		FVector2D(0, 0),
+		FVector2D(1, 0),
+		FVector2D(0, 1),
+		FVector2D(1, 1)
+	};
+
+	static const FSoftObjectPath DebugMaterialPath(TEXT("Material'/PCG/DebugObjects/PCG_DebugMaterialTexture.PCG_DebugMaterialTexture'"));
+}
+
+TArray<TSharedPtr<FStreamableHandle>> FPCGBaseTextureDataVisualization::LoadRequiredResources(const UPCGData* Data) const
+{
+	return { UAssetManager::GetStreamableManager().RequestAsyncLoad(PCGBaseTextureVisualizationConstants::DebugMaterialPath, nullptr) };
+}
+
+FPCGSetupSceneFunc FPCGBaseTextureDataVisualization::GetViewportSetupFunc(const UPCGData* Data) const
+{
+	using namespace PCGBaseTextureVisualizationConstants;
+
+	return [this, WeakData=TWeakObjectPtr<const UPCGBaseTextureData>(Cast<UPCGBaseTextureData>(Data))](FAdvancedPreviewScene* InScene, const TArray<UObject*>& InResources, FVector& OutViewLocation, float& OutOrbitDistance, TArray<TObjectPtr<UObject>>& OutManagedResources)
+	{
+		check(InScene);
+
+		if (InResources.IsEmpty())
+		{
+			return;
+		}
+
+		if (!WeakData.IsValid())
+		{
+			UE_LOG(LogPCG, Error, TEXT("Failed to setup data viewport, the data was lost or invalid."));
+			return;
+		}
+
+		if (WeakData->GetTextureResourceType() != EPCGTextureResourceType::TextureObject)
+		{
+			// @todo_pcg: Need to create a dynamic UTexture asset from an exported buffer handle in order to support EPCGTextureResourceType::ExportedTexture
+			UE_LOG(LogPCG, Verbose, TEXT("Texture data uses an unsupported resource type for data viewport visualization."));
+		}
+
+		UMaterialInterface* DebugMaterial = Cast<UMaterialInterface>(InResources[0]);
+		UMaterialInstanceDynamic* MaterialInstance = UMaterialInstanceDynamic::Create(DebugMaterial, GetTransientPackage());
+		MaterialInstance->SetTextureParameterValue(FName(TEXT("DebugTexture")), WeakData->GetTexture());
+		OutManagedResources.Add(MaterialInstance);
+
+		UProceduralMeshComponent* MeshComponent = NewObject<UProceduralMeshComponent>(GetTransientPackage(), NAME_None, RF_Transient);
+		MeshComponent->CreateMeshSection(/*SectionIndex=*/0, DebugMeshVertices, DebugMeshTriangles, /*Normals=*/{}, DebugMeshUVs, /*UV1=*/{}, /*UV2=*/{}, /*UV3=*/{}, /*VertexColors=*/{}, /*Tangents=*/{}, /*bCreateCollision=*/false);
+		MeshComponent->OverrideMaterials.Add(MaterialInstance);
+		OutManagedResources.Add(MeshComponent);
+
+		if (GEditor->PreviewPlatform.GetEffectivePreviewFeatureLevel() <= ERHIFeatureLevel::ES3_1)
+		{
+			MeshComponent->SetMobility(EComponentMobility::Static);
+		}
+
+		InScene->AddComponent(MeshComponent, FTransform::Identity);
+		InScene->SetFloorOffset(DebugQuadSize * 2);
+
+		OutViewLocation = FVector::ZeroVector;
+		OutOrbitDistance = DebugQuadSize * 2;
+	};
+}
diff --git a/Engine/Plugins/PCG/Source/PCGEditor/Private/PCGEditorModule.cpp b/Engine/Plugins/PCG/Source/PCGEditor/Private/PCGEditorModule.cpp
--- a/Engine/Plugins/PCG/Source/PCGEditor/Private/PCGEditorModule.cpp
+++ b/Engine/Plugins/PCG/Source/PCGEditor/Private/PCGEditorModule.cpp
@@ -15,6 +15,7 @@
 #include "Data/PCGSpatialData.h"
 #include "Data/PCGSplineData.h"
 #include "Data/PCGStaticMeshResourceData.h"
+#include "Data/PCGTextureData.h"
 #include "Data/PCGVolumeData.h"
 #include "Grid/PCGPartitionActor.h"
 #include "WorldPartitionBuilder/PCGWorldPartitionBuilder.h"
@@ -25,6 +26,7 @@
 #include "PCGEditorSettings.h"
 #include "PCGEditorStyle.h"
 #include "PCGEditorUtils.h"
+#include "DataVisualizations/PCGBaseTextureDataVisualization.h"
 #include "DataVisualizations/PCGLandscapeDataVisualization.h"
 #include "DataVisualizations/PCGParamDataVisualization.h"
 #include "DataVisualizations/PCGSpatialDataVisualization.h"
@@ -794,6 +796,7 @@ void FPCGEditorModule::RegisterPCGDataVisualizations()
 	DataVisRegistry.InternalRegistry.Add(UPCGCollisionWrapperData::StaticClass(), MakeUnique<const IPCGPrimitiveDataVisualization>());
 	DataVisRegistry.InternalRegistry.Add(UPCGLandscapeData::StaticClass(), MakeUnique<const IPCGLandscapeDataVisualization>());
 	DataVisRegistry.InternalRegistry.Add(UPCGStaticMeshResourceData::StaticClass(), MakeUnique<const IPCGStaticMeshDataVisualization>());
+	DataVisRegistry.InternalRegistry.Add(UPCGBaseTextureData::StaticClass(), MakeUnique<const FPCGBaseTextureDataVisualization>());
 }
 
 void FPCGEditorModule::UnregisterPCGDataVisualizations()
diff --git a/Engine/Plugins/PCG/Source/PCGEditor/Public/DataVisualizations/PCGBaseTextureDataVisualization.h b/Engine/Plugins/PCG/Source/PCGEditor/Public/DataVisualizations/PCGBaseTextureDataVisualization.h
new file mode 100644
--- /dev/null
+++ b/Engine/Plugins/PCG/Source/PCGEditor/Public/DataVisualizations/PCGBaseTextureDataVisualization.h
@@ -0,0 +1,17 @@
+// Copyright Epic Games, Inc. All Rights Reserved.
+
+#pragma once
+
+#include "PCGSpatialDataVisualization.h"
+#include "PCGEditorModule.h"
+
+class UPCGData;
+
+class PCGEDITOR_API FPCGBaseTextureDataVisualization : public IPCGSpatialDataVisualization
+{
+public:
+	// ~Begin IPCGDataVisualization interface
+	virtual TArray<TSharedPtr<FStreamableHandle>> LoadRequiredResources(const UPCGData* Data) const override;
+	virtual FPCGSetupSceneFunc GetViewportSetupFunc(const UPCGData* Data) const override;
+	// ~End IPCGDataVisualization interface
+};
