[PCG] Fix build health issue from missing plugin dependency

#rnx

diff --git a/Engine/Plugins/PCG/PCG.uplugin b/Engine/Plugins/PCG/PCG.uplugin
--- a/Engine/Plugins/PCG/PCG.uplugin
+++ b/Engine/Plugins/PCG/PCG.uplugin
@@ -39,6 +39,10 @@
 		{
 			"Name": "ComputeFramework",
 			"Enabled": true
+		},
+		{
+			"Name": "ProceduralMeshComponent",
+			"Enabled": true
 		}
 	]
 }
