[PCG] Don't update the visualization viewport when it's not open

#rnx

diff --git a/Engine/Plugins/PCG/Source/PCGEditor/Private/PCGEditor.cpp b/Engine/Plugins/PCG/Source/PCGEditor/Private/PCGEditor.cpp
--- a/Engine/Plugins/PCG/Source/PCGEditor/Private/PCGEditor.cpp
+++ b/Engine/Plugins/PCG/Source/PCGEditor/Private/PCGEditor.cpp
@@ -157,8 +157,10 @@ void FPCGEditor::Initialize(const EToolkitMode::Type InMode, const TSharedPtr<cl
 
 	for (int AttributesIndex = 0; AttributesIndex < 4; ++AttributesIndex)
 	{
+		const int ViewportEditorPanel = static_cast<int>(EPCGEditorPanel::Viewport1) + AttributesIndex;
+
 		AttributesWidgets.Add(CreateAttributesWidget());
-		AttributesWidgets[AttributesIndex]->SetViewportWidget(CreateViewportWidget());
+		AttributesWidgets[AttributesIndex]->SetViewportWidget(CreateViewportWidget(), static_cast<EPCGEditorPanel>(ViewportEditorPanel));
 	}
 	
 	DeterminismWidget = CreateDeterminismWidget();
@@ -542,6 +544,14 @@ FName FPCGEditor::GetPanelID(const EPCGEditorPanel Panel) const
 			return FPCGEditor_private::PropertyDetailsID[3];
 		case EPCGEditorPanel::UserParams:
 			return FPCGEditor_private::UserParamsID;
+		case EPCGEditorPanel::Viewport1:
+			return FPCGEditor_private::ViewportID[0];
+		case EPCGEditorPanel::Viewport2:
+			return FPCGEditor_private::ViewportID[1];
+		case EPCGEditorPanel::Viewport3:
+			return FPCGEditor_private::ViewportID[2];
+		case EPCGEditorPanel::Viewport4:
+			return FPCGEditor_private::ViewportID[3];
 		default:
 			return NAME_None;
 	}
@@ -676,6 +686,10 @@ void FPCGEditor::UnregisterTabSpawners(const TSharedRef<class FTabManager>& InTa
 	InTabManager->UnregisterTabSpawner(FPCGEditor_private::LogID);
 	InTabManager->UnregisterTabSpawner(FPCGEditor_private::NodeSourceID);
 	InTabManager->UnregisterTabSpawner(FPCGEditor_private::UserParamsID);
+	InTabManager->UnregisterTabSpawner(FPCGEditor_private::ViewportID[0]);
+	InTabManager->UnregisterTabSpawner(FPCGEditor_private::ViewportID[1]);
+	InTabManager->UnregisterTabSpawner(FPCGEditor_private::ViewportID[2]);
+	InTabManager->UnregisterTabSpawner(FPCGEditor_private::ViewportID[3]);
 
 	FAssetEditorToolkit::UnregisterTabSpawners(InTabManager);
 }
@@ -1048,6 +1062,11 @@ void FPCGEditor::OnAttributeListViewTabClosed(TSharedRef<SDockTab> DockTab, int
 	}
 }
 
+void FPCGEditor::OnViewportViewTabClosed(TSharedRef<SDockTab> DockTab, int Index)
+{
+	AttributesWidgets[Index]->ResetViewport();
+}
+
 void FPCGEditor::OnPauseAutomaticRegeneration_Clicked()
 {
 	if (!PCGGraphBeingEdited)
@@ -3361,9 +3380,12 @@ TSharedRef<SDockTab> FPCGEditor::SpawnTab_UserParams(const FSpawnTabArgs& Args)
 TSharedRef<SDockTab> FPCGEditor::SpawnTab_Viewport(const FSpawnTabArgs& Args, int ViewportIndex)
 {
 	TAttribute<FText> Label = TAttribute<FText>::Create(TAttribute<FText>::FGetter::CreateRaw(this, &FPCGEditor::GetViewportTabLabel, ViewportIndex));
+	
+	AttributesWidgets[ViewportIndex]->RequestViewportRefresh();
 
 	return SNew(SDockTab)
 		.Label(Label)
+		.OnTabClosed_Raw(this, &FPCGEditor::OnViewportViewTabClosed, ViewportIndex)
 		.TabColorScale(GetTabColorScale())
 		[
 			AttributesWidgets[ViewportIndex]->GetViewportWidget().ToSharedRef()
diff --git a/Engine/Plugins/PCG/Source/PCGEditor/Private/PCGEditor.h b/Engine/Plugins/PCG/Source/PCGEditor/Private/PCGEditor.h
--- a/Engine/Plugins/PCG/Source/PCGEditor/Private/PCGEditor.h
+++ b/Engine/Plugins/PCG/Source/PCGEditor/Private/PCGEditor.h
@@ -53,7 +53,11 @@ enum class EPCGEditorPanel
 	PropertyDetails2,
 	PropertyDetails3,
 	PropertyDetails4,
-	UserParams
+	UserParams,
+	Viewport1,
+	Viewport2,
+	Viewport3,
+	Viewport4
 };
 
 class FPCGEditor : public FAssetEditorToolkit, public FGCObject, public FSelfRegisteringEditorUndoClient
@@ -172,6 +176,9 @@ private:
 	/** Called when an attribute list view tab is closed */
 	void OnAttributeListViewTabClosed(TSharedRef<SDockTab> DockTab, int Index);
 
+	/** Called when a viewport view tab is closed */
+	void OnViewportViewTabClosed(TSharedRef<SDockTab> DockTab, int Index);
+
 	/** Enable/Disable automatic PCG node generation */
 	void OnPauseAutomaticRegeneration_Clicked();
 	/** Has the user paused automatic regeneration in the Graph Editor */
diff --git a/Engine/Plugins/PCG/Source/PCGEditor/Private/Widgets/AssetEditorViewport/SPCGEditorViewport.cpp b/Engine/Plugins/PCG/Source/PCGEditor/Private/Widgets/AssetEditorViewport/SPCGEditorViewport.cpp
--- a/Engine/Plugins/PCG/Source/PCGEditor/Private/Widgets/AssetEditorViewport/SPCGEditorViewport.cpp
+++ b/Engine/Plugins/PCG/Source/PCGEditor/Private/Widgets/AssetEditorViewport/SPCGEditorViewport.cpp
@@ -27,6 +27,7 @@ public:
 	virtual void Tick(float DeltaSeconds) override;
 	//~ End FEditorViewportClient Interface
 
+	FAdvancedPreviewScene* GetPreviewScene() { return Scene.Get(); }
 	FAdvancedPreviewScene* ResetScene();
 
 private:
@@ -149,14 +150,14 @@ void SPCGEditorViewport::Construct(const FArguments& InArgs)
 
 void SPCGEditorViewport::SetupScene(const TArray<UObject*>& InResources, const FPCGSetupSceneFunc& SetupFunc)
 {
-	ReleaseManagedResources();
-	
+	ResetScene();
+
 	if (!ensure(EditorViewportClient.IsValid()))
 	{
 		return;
 	}
 
-	FAdvancedPreviewScene* Scene = EditorViewportClient->ResetScene();
+	FAdvancedPreviewScene* Scene = EditorViewportClient->GetPreviewScene();
 	FVector ViewLocation = PCGEditorViewportConstants::DefaultViewLocation;
 	float OrbitDistance = PCGEditorViewportConstants::DefaultOrbitDistance;
 
@@ -170,6 +171,21 @@ void SPCGEditorViewport::SetupScene(const TArray<UObject*>& InResources, const F
 	EditorViewportClient->Invalidate();
 }
 
+void SPCGEditorViewport::ResetScene()
+{
+	ReleaseManagedResources();
+	
+	if (!ensure(EditorViewportClient.IsValid()))
+	{
+		return;
+	}
+
+	EditorViewportClient->ResetScene();
+	EditorViewportClient->SetViewLocation(PCGEditorViewportConstants::DefaultViewLocation);
+	EditorViewportClient->SetViewLocationForOrbiting(PCGEditorViewportConstants::DefaultViewLocation, PCGEditorViewportConstants::DefaultOrbitDistance);
+	EditorViewportClient->Invalidate();
+}
+
 TSharedRef<FEditorViewportClient> SPCGEditorViewport::MakeEditorViewportClient()
 {
 	EditorViewportClient = MakeShareable(new FPCGEditorViewportClient(SharedThis(this)));
diff --git a/Engine/Plugins/PCG/Source/PCGEditor/Private/Widgets/AssetEditorViewport/SPCGEditorViewport.h b/Engine/Plugins/PCG/Source/PCGEditor/Private/Widgets/AssetEditorViewport/SPCGEditorViewport.h
--- a/Engine/Plugins/PCG/Source/PCGEditor/Private/Widgets/AssetEditorViewport/SPCGEditorViewport.h
+++ b/Engine/Plugins/PCG/Source/PCGEditor/Private/Widgets/AssetEditorViewport/SPCGEditorViewport.h
@@ -20,6 +20,7 @@ public:
 
 	/** Rebuilds the scene using the provided resources and setup callback. */
 	void SetupScene(const TArray<UObject*>& InResources, const FPCGSetupSceneFunc& SetupFunc);
+	void ResetScene();
 
 	/** Add UObject references for GC */
 	void AddReferencedObjects(FReferenceCollector& Collector);
diff --git a/Engine/Plugins/PCG/Source/PCGEditor/Private/Widgets/SPCGEditorGraphAttributeListView.cpp b/Engine/Plugins/PCG/Source/PCGEditor/Private/Widgets/SPCGEditorGraphAttributeListView.cpp
--- a/Engine/Plugins/PCG/Source/PCGEditor/Private/Widgets/SPCGEditorGraphAttributeListView.cpp
+++ b/Engine/Plugins/PCG/Source/PCGEditor/Private/Widgets/SPCGEditorGraphAttributeListView.cpp
@@ -595,6 +595,12 @@ void SPCGEditorGraphAttributeListView::Tick(const FGeometry& AllottedGeometry, c
 		RefreshAttributeList();
 	}
 
+	if (bViewportNeedsRefresh)
+	{
+		bViewportNeedsRefresh = false;
+		RefreshViewport();
+	}
+
 	if (CurrentUpdateTask.IsValid() && CurrentUpdateTask->IsCompleted())
 	{
 		FilteredListViewItems = MoveTemp(CurrentUpdateTask->ListViewItems);
@@ -632,6 +638,12 @@ void SPCGEditorGraphAttributeListView::RequestRefresh()
 	DataPtr = nullptr;
 }
 
+void SPCGEditorGraphAttributeListView::RequestViewportRefresh()
+{
+	bViewportNeedsRefresh = true;
+	bRefreshLoadHandles = true;
+}
+
 TSharedRef<SHeaderRow> SPCGEditorGraphAttributeListView::CreateHeaderRowWidget() const
 {
 	return SNew(SHeaderRow);
@@ -681,6 +693,17 @@ void SPCGEditorGraphAttributeListView::SetNodeBeingInspected(UPCGEditorGraphNode
 	RequestRefresh();
 }
 
+void SPCGEditorGraphAttributeListView::SetViewportWidget(TSharedPtr<SPCGEditorViewport> InViewportWidget, EPCGEditorPanel InViewportEditorPanel)
+{
+	ViewportWidget = InViewportWidget;
+	ViewportEditorPanel = InViewportEditorPanel;
+}
+
+void SPCGEditorGraphAttributeListView::ResetViewport()
+{
+	ViewportWidget->ResetScene();
+}
+
 void SPCGEditorGraphAttributeListView::AddReferencedObjects(FReferenceCollector& Collector)
 {
 	Collector.AddReferencedObject(DataPtr);
@@ -769,8 +792,8 @@ void SPCGEditorGraphAttributeListView::RefreshAttributeList()
 	ListViewItems.Empty();
 	ListViewHeader->ClearColumns();
 	InfoTextBlock->SetText(FText::GetEmpty());
-	ViewportWidget->SetupScene({}, FPCGSetupSceneFunc());
 	DataPtr = nullptr;
+	bViewportNeedsRefresh = true;
 
 	const FPCGDataCollection* InspectionData = GetInspectionData();
 	if (!InspectionData)
@@ -829,43 +852,6 @@ void SPCGEditorGraphAttributeListView::RefreshAttributeList()
 		const FPCGTableVisualizerInfo TableVisualizerInfo = DataVisualization->GetTableVisualizerInfoWithDomain(PCGData, MetadataDomainID);
 		DataToVisualize = TableVisualizerInfo.Data;
 
-		if (DataToVisualize != DataPtr.Get())
-		{
-			LoadHandles = DataVisualization->LoadRequiredResources(DataToVisualize);
-		}
-
-		bool bAllResourcesLoaded = true;
-		TArray<UObject*> LoadedResources;
-
-		for (TSharedPtr<FStreamableHandle> LoadHandle : LoadHandles)
-		{
-			if (!LoadHandle.IsValid())
-			{
-				continue;
-			}
-
-			if (LoadHandle->HasLoadCompleted())
-			{
-				LoadedResources.Add(LoadHandle->GetLoadedAsset());
-			}
-			else
-			{
-				bAllResourcesLoaded = false;
-			}
-		}
-
-		if (bAllResourcesLoaded)
-		{
-			ViewportWidget->SetupScene(LoadedResources, DataVisualization->GetViewportSetupFunc(PCGData));
-		}
-		else
-		{
-			// Poll next tick.
-			bNeedsRefresh = true;
-			DataPtr = DataToVisualize;
-			return;
-		}
-
 		int32 NumEntries = 0;
 
 		for (const FPCGTableVisualizerColumnInfo& ColumnInfo : TableVisualizerInfo.ColumnInfos)
@@ -894,6 +880,12 @@ void SPCGEditorGraphAttributeListView::RefreshAttributeList()
 		FocusOnDataCallback.Reset();
 	}
 
+	if (DataPtr.Get() != DataToVisualize)
+	{
+		// If the visualized data has changed for whatever reason, we should make sure to be loading the relevant assets.
+		bRefreshLoadHandles = true;
+	}
+
 	DataPtr = DataToVisualize;
 
 	ListView->SetItemsSource(&ListViewItems);
@@ -1097,6 +1089,103 @@ void SPCGEditorGraphAttributeListView::RefreshDomainComboBox(bool bKeepSelection
 	DomainsComboBoxItemsSelectedDomain = PCGMetadataDomainID::Invalid;
 }
 
+void SPCGEditorGraphAttributeListView::RefreshViewport()
+{
+	if (!IsViewportOpen())
+	{
+		return;
+	}
+
+	ResetViewport();
+
+	const FPCGDataCollection* InspectionData = GetInspectionData();
+	if (!InspectionData)
+	{
+		return;
+	}
+
+	const int32 DataIndex = GetSelectedDataIndex();
+	if (!InspectionData->TaggedData.IsValidIndex(DataIndex))
+	{
+		return;
+	}
+
+	const int32 DomainIndex = GetSelectedDomainIndex();
+	if (!DomainsComboBoxIds.IsValidIndex(DomainIndex))
+	{
+		return;
+	}
+
+	const UPCGData* Data = InspectionData->TaggedData[DataIndex].Data;
+
+	if (!Data)
+	{
+		return;
+	}
+
+	// If we have a proxy for GPU data, read back CPU data for inspection.
+	if (const UPCGProxyForGPUData* Proxy = Cast<UPCGProxyForGPUData>(Data))
+	{
+		UPCGProxyForGPUData::FReadbackResult Result = Proxy->GetCPUData(/*InContext=*/nullptr);
+
+		if (!Result.bComplete)
+		{
+			// Poll next tick.
+			bViewportNeedsRefresh = true;
+			return;
+		}
+
+		if (!Result.TaggedData.Data)
+		{
+			return;
+		}
+
+		Data = Result.TaggedData.Data;
+	}
+
+	const FPCGDataVisualizationRegistry& DataVisRegistry = FPCGModule::GetConstPCGDataVisualizationRegistry();
+	const IPCGDataVisualization* DataVisualization = DataVisRegistry.GetDataVisualization(Data->GetClass());
+
+	if (!DataVisualization)
+	{
+		return;
+	}
+
+	if (bRefreshLoadHandles)
+	{
+		LoadHandles = DataVisualization->LoadRequiredResources(Data);
+	}
+
+	bool bAllResourcesLoaded = true;
+	TArray<UObject*> LoadedResources;
+
+	for (TSharedPtr<FStreamableHandle> LoadHandle : LoadHandles)
+	{
+		if (!LoadHandle.IsValid())
+		{
+			continue;
+		}
+
+		if (LoadHandle->HasLoadCompleted())
+		{
+			LoadedResources.Add(LoadHandle->GetLoadedAsset());
+		}
+		else
+		{
+			bAllResourcesLoaded = false;
+		}
+	}
+
+	if (!bAllResourcesLoaded)
+	{
+		// Poll next tick.
+		bViewportNeedsRefresh = true;
+		return;
+	}
+
+	ViewportWidget->SetupScene(LoadedResources, DataVisualization->GetViewportSetupFunc(Data));
+}
+
 void SPCGEditorGraphAttributeListView::LaunchUpdateTask()
 {
 	// Discarding any currently running updater, task will still run and keep the old object alive but we wont care about the result.
@@ -1814,4 +1903,9 @@ bool SPCGEditorGraphAttributeListView::CanFocusOnSelection() const
 	return ListView->GetNumItemsSelected() > 0;
 }
 
+bool SPCGEditorGraphAttributeListView::IsViewportOpen() const
+{
+	return PCGEditorPtr.IsValid() && PCGEditorPtr.Pin()->IsPanelCurrentlyOpen(ViewportEditorPanel);
+}
+
 #undef LOCTEXT_NAMESPACE
diff --git a/Engine/Plugins/PCG/Source/PCGEditor/Private/Widgets/SPCGEditorGraphAttributeListView.h b/Engine/Plugins/PCG/Source/PCGEditor/Private/Widgets/SPCGEditorGraphAttributeListView.h
--- a/Engine/Plugins/PCG/Source/PCGEditor/Private/Widgets/SPCGEditorGraphAttributeListView.h
+++ b/Engine/Plugins/PCG/Source/PCGEditor/Private/Widgets/SPCGEditorGraphAttributeListView.h
@@ -33,6 +33,7 @@ struct FPCGTableVisualizerColumnInfo;
 struct FPCGDataCollection;
 struct FPCGPoint;
 enum class EPCGMetadataTypes : uint8;
+enum class EPCGEditorPanel;
 
 class SComboButton;
 class SHeaderRow;
@@ -151,6 +152,7 @@ public:
 	virtual void Tick(const FGeometry& AllottedGeometry, const double InCurrentTime, const float InDeltaTime) override;
 
 	void RequestRefresh();
+	void RequestViewportRefresh();
 
 	UPCGEditorGraphNodeBase* GetNodeBeingInspected() const;
 	void SetNodeBeingInspected(UPCGEditorGraphNodeBase* InNode);
@@ -159,7 +161,10 @@ public:
 	void SetIsLocked(bool bInIsLocked) { bIsLocked = bInIsLocked; }
 
 	TSharedPtr<SPCGEditorViewport> GetViewportWidget() const { return ViewportWidget; }
-	void SetViewportWidget(TSharedPtr<SPCGEditorViewport> InViewportWidget) { ViewportWidget = InViewportWidget; }
+	void SetViewportWidget(TSharedPtr<SPCGEditorViewport> InViewportWidget, EPCGEditorPanel InViewportEditorPanel);
+
+	/** Resets/clears the viewport scene. */
+	void ResetViewport();
 
 	/** Add UObject references for GC */
 	void AddReferencedObjects(FReferenceCollector& Collector);
@@ -180,6 +185,7 @@ private:
 	void RefreshPinComboBox(bool bKeepSelection, bool& bOutSelectionChanged);
 	void RefreshDataComboBox(bool bKeepSelection);
 	void RefreshDomainComboBox(bool bKeepSelection);
+	void RefreshViewport();
 
 	void LaunchUpdateTask();
 
@@ -199,6 +205,7 @@ private:
 
 	const FPCGDataCollection* GetInspectionData(const TSharedPtr<FPinComboBoxItem>& EditorPin) const;
 	const FPCGDataCollection* GetInspectionData() const;
+	const UPCGData* GetInspectedData(bool& bOutNeedsRefresh) const;
 
 	FText OnGenerateSelectedPinText() const;
 	void OnSelectionChangedPin(TSharedPtr<FPinComboBoxItem> InItem, ESelectInfo::Type InSelectInfo);
@@ -255,6 +262,8 @@ private:
 	void FocusOnSelection() const;
 	bool CanFocusOnSelection() const;
 
+	bool IsViewportOpen() const;
+
 	/** Returns the current PCG component */
 	TWeakObjectPtr<UPCGComponent> GetPCGComponent() const;
 
@@ -264,6 +273,9 @@ private:
 	/** Data visualization viewport. */
 	TSharedPtr<SPCGEditorViewport> ViewportWidget;
 
+	/** Data visualization viewport editor panel. */
+	EPCGEditorPanel ViewportEditorPanel;
+
 	/** Cached PCGGraphNode being viewed */
 	TWeakObjectPtr<UPCGEditorGraphNodeBase> PCGEditorGraphNode;
 	bool bPCGEditorGraphNodeChanged = false;
@@ -315,6 +327,7 @@ private:
 	TFunction<void(const UPCGData*, TArrayView<const int>)> FocusOnDataCallback;
 
 	bool bNeedsRefresh : 1 = false;
+	bool bViewportNeedsRefresh : 1 = false;
 
 	/** True if this property view is currently locked (i.e. The objects being observed are not changed automatically due to user selection)*/
 	bool bIsLocked : 1 = false;
@@ -330,4 +343,5 @@ private:
 
 	/** Handles for any resources the data needs to load before visualization. */
 	TArray<TSharedPtr<FStreamableHandle>> LoadHandles;
+	bool bRefreshLoadHandles : 1 = false;
 };
