[PCG] Add UPCGVirtualTextureData
* Supports Runtime Virtual Textures only for now. VT/SVT support can follow.
* Add 'Get Virtual Texture Data' node
* Add data interface for virtual texture sampling

diff --git a/Engine/Plugins/PCG/Shaders/Private/PCGLandscapeDataInterface.ush b/Engine/Plugins/PCG/Shaders/Private/PCGLandscapeDataInterface.ush
--- a/Engine/Plugins/PCG/Shaders/Private/PCGLandscapeDataInterface.ush
+++ b/Engine/Plugins/PCG/Shaders/Private/PCGLandscapeDataInterface.ush
@@ -1,17 +1,6 @@
 // Copyright Epic Games, Inc. All Rights Reserved.
 
-#define NUM_VIRTUALTEXTURE_SAMPLES 1
-// Workaround for unbound View uniform buffer. This pattern is used elsewhere in the engine as well.
-#define VT_DISABLE_VIEW_UNIFORM_BUFFER 1
-
-#include "/Engine/Private/VirtualTextureCommon.ush"
-#include "/Engine/Private/Landscape/LandscapeCommon.ush"
-
-// Mirrored from FPCGLandscapeDataProviderProxy::ENormalUnpackType
-#define VIRTUAL_TEXTURE_UNPACK_NORMAL_NONE 0
-#define VIRTUAL_TEXTURE_UNPACK_NORMAL_BC3BC3 1
-#define VIRTUAL_TEXTURE_UNPACK_NORMAL_BC5BC1 2
-#define VIRTUAL_TEXTURE_UNPACK_NORMAL_B5G6R5 3
+#include "/Plugin/PCG/Private/PCGVirtualTextureCommon.ush"
 
 // Height RVT
 Texture2D			{DataInterfaceName}_HeightVirtualTexture;
@@ -20,7 +9,7 @@ Texture2D<uint>		{DataInterfaceName}_HeightVirtualTexturePageTableIndirection;
 uint				{DataInterfaceName}_HeightVirtualTextureAdaptive;
 SamplerState		{DataInterfaceName}_HeightVirtualTextureSampler;
 float3				{DataInterfaceName}_HeightVirtualTextureLWCTile;
-float4x4			{DataInterfaceName}_HeightVirtualTextureWorldToUvTransform;
+float4x4			{DataInterfaceName}_HeightVirtualTextureWorldToUVTransform;
 uint				{DataInterfaceName}_HeightVirtualTextureEnabled;
 uint4				{DataInterfaceName}_HeightVirtualTexturePackedUniform0;
 uint4				{DataInterfaceName}_HeightVirtualTexturePackedUniform1;
@@ -35,7 +24,7 @@ uint				{DataInterfaceName}_NormalVirtualTextureAdaptive;
 SamplerState		{DataInterfaceName}_NormalVirtualTexture0Sampler;
 SamplerState		{DataInterfaceName}_NormalVirtualTexture1Sampler;
 float3				{DataInterfaceName}_NormalVirtualTextureLWCTile;
-float4x4			{DataInterfaceName}_NormalVirtualTextureWorldToUvTransform;
+float4x4			{DataInterfaceName}_NormalVirtualTextureWorldToUVTransform;
 uint4				{DataInterfaceName}_NormalVirtualTexturePackedUniform0;
 uint4				{DataInterfaceName}_NormalVirtualTexturePackedUniform1;
 uint4				{DataInterfaceName}_NormalVirtualTextureUniforms0;
@@ -50,7 +39,7 @@ Texture2D<uint>		{DataInterfaceName}_BaseColorVirtualTexturePageTableIndirection
 uint				{DataInterfaceName}_BaseColorVirtualTextureAdaptive;
 SamplerState		{DataInterfaceName}_BaseColorVirtualTextureSampler;
 float3				{DataInterfaceName}_BaseColorVirtualTextureLWCTile;
-float4x4			{DataInterfaceName}_BaseColorVirtualTextureWorldToUvTransform;
+float4x4			{DataInterfaceName}_BaseColorVirtualTextureWorldToUVTransform;
 uint				{DataInterfaceName}_BaseColorVirtualTextureUnpackSRGB;
 uint				{DataInterfaceName}_BaseColorVirtualTextureUnpackYCoCg;
 uint				{DataInterfaceName}_BaseColorVirtualTextureEnabled;
@@ -63,40 +52,40 @@ Texture2D			{DataInterfaceName}_CollisionHeightTexture;
 SamplerState		{DataInterfaceName}_CollisionHeightTextureSampler;
 uint				{DataInterfaceName}_CollisionHeightTextureEnabled;
 float3				{DataInterfaceName}_CollisionHeightTextureLWCTile;
-float4x4			{DataInterfaceName}_CollisionHeightTextureWorldToUvTransform;
-float4x4			{DataInterfaceName}_CollisionHeightTextureUvToWorldTransform;
-float4				{DataInterfaceName}_CollisionHeightTextureUvScaleBias;
+float4x4			{DataInterfaceName}_CollisionHeightTextureWorldToUVTransform;
+float4x4			{DataInterfaceName}_CollisionHeightTextureUVToWorldTransform;
+float4				{DataInterfaceName}_CollisionHeightTextureUVScaleBias;
 int2				{DataInterfaceName}_CollisionHeightTextureDimension;
 
 // Misc parameters
 float2				{DataInterfaceName}_LandscapeGridSize;
 
-float2 GetCollisionHeightTextureUv_{DataInterfaceName}(float3 InWorldPos)
+float2 GetCollisionHeightTextureUV_{DataInterfaceName}(float3 InWorldPos)
 {
 	// @todo_pcg: Support LWC landscape sampling
 	const float3 LWCTile = (float3)0;
 	const FLWCVector4 LwcWorldPos = MakeLWCVector4(LWCTile, float4(InWorldPos, 1.0f));
-	const FLWCInverseMatrix LwcWorldToUv = MakeLWCInverseMatrix({DataInterfaceName}_CollisionHeightTextureLWCTile, {DataInterfaceName}_CollisionHeightTextureWorldToUvTransform);
+	const FLWCInverseMatrix LwcWorldToUV = MakeLWCInverseMatrix({DataInterfaceName}_CollisionHeightTextureLWCTile, {DataInterfaceName}_CollisionHeightTextureWorldToUVTransform);
 
-	float2 Uv = LWCMultiply(LwcWorldPos, LwcWorldToUv).xy;
-	Uv = Uv * {DataInterfaceName}_CollisionHeightTextureUvScaleBias.xy + {DataInterfaceName}_CollisionHeightTextureUvScaleBias.zw;
+	float2 UV = LWCMultiply(LwcWorldPos, LwcWorldToUV).xy;
+	UV = UV * {DataInterfaceName}_CollisionHeightTextureUVScaleBias.xy + {DataInterfaceName}_CollisionHeightTextureUVScaleBias.zw;
 
-	return Uv;
+	return UV;
 }
 
 float GetHeight_{DataInterfaceName}(float3 InWorldPos)
 {
 	if ({DataInterfaceName}_HeightVirtualTextureEnabled != 0)
 	{
-		const float3 UnpackHeightOrigin = {DataInterfaceName}_HeightVirtualTextureWorldToUvTransform[0].xyz;
-		const float3 UnpackHeightU = {DataInterfaceName}_HeightVirtualTextureWorldToUvTransform[1].xyz;
-		const float3 UnpackHeightV = {DataInterfaceName}_HeightVirtualTextureWorldToUvTransform[2].xyz;
-		const float2 UnpackHeightScaleBias = {DataInterfaceName}_HeightVirtualTextureWorldToUvTransform[3].xy;
+		const float3 UnpackHeightOrigin = {DataInterfaceName}_HeightVirtualTextureWorldToUVTransform[0].xyz;
+		const float3 UnpackHeightU = {DataInterfaceName}_HeightVirtualTextureWorldToUVTransform[1].xyz;
+		const float3 UnpackHeightV = {DataInterfaceName}_HeightVirtualTextureWorldToUVTransform[2].xyz;
+		const float2 UnpackHeightScaleBias = {DataInterfaceName}_HeightVirtualTextureWorldToUVTransform[3].xy;
 		const float3 LWCTile = (float3)0;
 
 		const FDFVector3 LwcWorldPos = DFFromTileOffset_Hack(MakeLWCVector3(LWCTile, InWorldPos));
 		const FDFVector3 LwcHeightOrigin = DFFromTileOffset_Hack(MakeLWCVector3({DataInterfaceName}_HeightVirtualTextureLWCTile, UnpackHeightOrigin));
-		const float2 SampleUv = VirtualTextureWorldToUV(LwcWorldPos, LwcHeightOrigin, UnpackHeightU, UnpackHeightV);
+		const float2 SampleUV = VirtualTextureWorldToUV(LwcWorldPos, LwcHeightOrigin, UnpackHeightU, UnpackHeightV);
 
 		VTPageTableResult PageTable;
 
@@ -105,7 +94,7 @@ float GetHeight_{DataInterfaceName}(float3 InWorldPos)
 			PageTable = TextureLoadVirtualPageTableLevel(
 				{DataInterfaceName}_HeightVirtualTexturePageTable,
 				VTPageTableUniform_Unpack({DataInterfaceName}_HeightVirtualTexturePackedUniform0, {DataInterfaceName}_HeightVirtualTexturePackedUniform1),
-				SampleUv,
+				SampleUV,
 				VTADDRESSMODE_CLAMP,
 				VTADDRESSMODE_CLAMP,
 				0.0f);
@@ -116,7 +105,7 @@ float GetHeight_{DataInterfaceName}(float3 InWorldPos)
 				{DataInterfaceName}_HeightVirtualTexturePageTable,
 				{DataInterfaceName}_HeightVirtualTexturePageTableIndirection,
 				VTPageTableUniform_Unpack({DataInterfaceName}_HeightVirtualTexturePackedUniform0, {DataInterfaceName}_HeightVirtualTexturePackedUniform1),
-				SampleUv,
+				SampleUV,
 				VTADDRESSMODE_CLAMP,
 				VTADDRESSMODE_CLAMP,
 				0.0f);
@@ -133,8 +122,8 @@ float GetHeight_{DataInterfaceName}(float3 InWorldPos)
 	}
 	else if ({DataInterfaceName}_CollisionHeightTextureEnabled != 0)
 	{
-		const float2 Uv = GetCollisionHeightTextureUv_{DataInterfaceName}(InWorldPos);
-		return {DataInterfaceName}_CollisionHeightTexture.SampleLevel({DataInterfaceName}_CollisionHeightTextureSampler, Uv, 0).x;
+		const float2 UV = GetCollisionHeightTextureUV_{DataInterfaceName}(InWorldPos);
+		return {DataInterfaceName}_CollisionHeightTexture.SampleLevel({DataInterfaceName}_CollisionHeightTextureSampler, UV, 0).x;
 	}
 
 	return InWorldPos.z;
@@ -144,14 +133,14 @@ float3 GetNormal_{DataInterfaceName}(float3 InWorldPos)
 {
 	if ({DataInterfaceName}_NormalVirtualTextureEnabled != 0)
 	{
-		const float3 UnpackNormalOrigin = {DataInterfaceName}_NormalVirtualTextureWorldToUvTransform[0].xyz;
-		const float3 UnpackNormalU = {DataInterfaceName}_NormalVirtualTextureWorldToUvTransform[1].xyz;
-		const float3 UnpackNormalV = {DataInterfaceName}_NormalVirtualTextureWorldToUvTransform[2].xyz;
+		const float3 UnpackNormalOrigin = {DataInterfaceName}_NormalVirtualTextureWorldToUVTransform[0].xyz;
+		const float3 UnpackNormalU = {DataInterfaceName}_NormalVirtualTextureWorldToUVTransform[1].xyz;
+		const float3 UnpackNormalV = {DataInterfaceName}_NormalVirtualTextureWorldToUVTransform[2].xyz;
 		const float3 LWCTile = (float3)0;
 
 		const FDFVector3 LwcWorldPos = DFFromTileOffset_Hack(MakeLWCVector3(LWCTile, InWorldPos));
 		const FDFVector3 LwcNormalOrigin = DFFromTileOffset_Hack(MakeLWCVector3({DataInterfaceName}_NormalVirtualTextureLWCTile, UnpackNormalOrigin));
-		const float2 SampleUv = VirtualTextureWorldToUV(LwcWorldPos, LwcNormalOrigin, UnpackNormalU, UnpackNormalV);
+		const float2 SampleUV = VirtualTextureWorldToUV(LwcWorldPos, LwcNormalOrigin, UnpackNormalU, UnpackNormalV);
 
 		VTPageTableResult PageTable;
 
@@ -160,7 +149,7 @@ float3 GetNormal_{DataInterfaceName}(float3 InWorldPos)
 			PageTable = TextureLoadVirtualPageTableLevel(
 				{DataInterfaceName}_NormalVirtualTexturePageTable,
 				VTPageTableUniform_Unpack({DataInterfaceName}_NormalVirtualTexturePackedUniform0, {DataInterfaceName}_NormalVirtualTexturePackedUniform1),
-				SampleUv,
+				SampleUV,
 				VTADDRESSMODE_CLAMP,
 				VTADDRESSMODE_CLAMP,
 				0.0f);
@@ -171,7 +160,7 @@ float3 GetNormal_{DataInterfaceName}(float3 InWorldPos)
 				{DataInterfaceName}_NormalVirtualTexturePageTable,
 				{DataInterfaceName}_NormalVirtualTexturePageTableIndirection,
 				VTPageTableUniform_Unpack({DataInterfaceName}_NormalVirtualTexturePackedUniform0, {DataInterfaceName}_NormalVirtualTexturePackedUniform1),
-				SampleUv,
+				SampleUV,
 				VTADDRESSMODE_CLAMP,
 				VTADDRESSMODE_CLAMP,
 				0.0f);
@@ -224,14 +213,14 @@ float3 GetBaseColor_{DataInterfaceName}(float3 InWorldPos)
 {
 	if ({DataInterfaceName}_BaseColorVirtualTextureEnabled != 0)
 	{
-		const float3 UnpackBaseColorOrigin = {DataInterfaceName}_BaseColorVirtualTextureWorldToUvTransform[0].xyz;
-		const float3 UnpackBaseColorU = {DataInterfaceName}_BaseColorVirtualTextureWorldToUvTransform[1].xyz;
-		const float3 UnpackBaseColorV = {DataInterfaceName}_BaseColorVirtualTextureWorldToUvTransform[2].xyz;
+		const float3 UnpackBaseColorOrigin = {DataInterfaceName}_BaseColorVirtualTextureWorldToUVTransform[0].xyz;
+		const float3 UnpackBaseColorU = {DataInterfaceName}_BaseColorVirtualTextureWorldToUVTransform[1].xyz;
+		const float3 UnpackBaseColorV = {DataInterfaceName}_BaseColorVirtualTextureWorldToUVTransform[2].xyz;
 		const float3 LWCTile = (float3)0;
 
 		const FDFVector3 LwcWorldPos = DFFromTileOffset_Hack(MakeLWCVector3(LWCTile, InWorldPos));
 		const FDFVector3 LwcBaseColorOrigin = DFFromTileOffset_Hack(MakeLWCVector3({DataInterfaceName}_BaseColorVirtualTextureLWCTile, UnpackBaseColorOrigin));
-		const float2 SampleUv = VirtualTextureWorldToUV(LwcWorldPos, LwcBaseColorOrigin, UnpackBaseColorU, UnpackBaseColorV);
+		const float2 SampleUV = VirtualTextureWorldToUV(LwcWorldPos, LwcBaseColorOrigin, UnpackBaseColorU, UnpackBaseColorV);
 
 		VTPageTableResult PageTable;
 
@@ -240,7 +229,7 @@ float3 GetBaseColor_{DataInterfaceName}(float3 InWorldPos)
 			PageTable = TextureLoadVirtualPageTableLevel(
 				{DataInterfaceName}_BaseColorVirtualTexturePageTable,
 				VTPageTableUniform_Unpack({DataInterfaceName}_BaseColorVirtualTexturePackedUniform0, {DataInterfaceName}_BaseColorVirtualTexturePackedUniform1),
-				SampleUv,
+				SampleUV,
 				VTADDRESSMODE_CLAMP,
 				VTADDRESSMODE_CLAMP,
 				0.0f);
@@ -251,7 +240,7 @@ float3 GetBaseColor_{DataInterfaceName}(float3 InWorldPos)
 				{DataInterfaceName}_BaseColorVirtualTexturePageTable,
 				{DataInterfaceName}_BaseColorVirtualTexturePageTableIndirection,
 				VTPageTableUniform_Unpack({DataInterfaceName}_BaseColorVirtualTexturePackedUniform0, {DataInterfaceName}_BaseColorVirtualTexturePackedUniform1),
-				SampleUv,
+				SampleUV,
 				VTADDRESSMODE_CLAMP,
 				VTADDRESSMODE_CLAMP,
 				0.0f);
diff --git a/Engine/Plugins/PCG/Shaders/Private/PCGVirtualTextureCommon.ush b/Engine/Plugins/PCG/Shaders/Private/PCGVirtualTextureCommon.ush
new file mode 100644
--- /dev/null
+++ b/Engine/Plugins/PCG/Shaders/Private/PCGVirtualTextureCommon.ush
@@ -0,0 +1,26 @@
+// Copyright Epic Games, Inc. All Rights Reserved.
+
+#pragma once
+
+#define NUM_VIRTUALTEXTURE_SAMPLES 1
+// Workaround for unbound View uniform buffer. This pattern is used elsewhere in the engine as well.
+#define VT_DISABLE_VIEW_UNIFORM_BUFFER 1
+
+#include "/Engine/Private/VirtualTextureCommon.ush"
+
+// Mirror of ERuntimeVirtualTextureMaterialType
+#define ERuntimeVirtualTextureMaterialType_BaseColor								0
+#define ERuntimeVirtualTextureMaterialType_Mask4									1
+#define ERuntimeVirtualTextureMaterialType_BaseColor_Normal_Roughness				2
+#define ERuntimeVirtualTextureMaterialType_BaseColor_Normal_Specular				3
+#define ERuntimeVirtualTextureMaterialType_BaseColor_Normal_Specular_YCoCg			4
+#define ERuntimeVirtualTextureMaterialType_BaseColor_Normal_Specular_Mask_YCoCg		5
+#define ERuntimeVirtualTextureMaterialType_WorldHeight								6
+#define ERuntimeVirtualTextureMaterialType_Displacement								7
+#define ERuntimeVirtualTextureMaterialType_Count									8
+
+// Mirror of PCGVirtualTextureCommon::ENormalUnpackType
+#define VIRTUAL_TEXTURE_UNPACK_NORMAL_NONE		0
+#define VIRTUAL_TEXTURE_UNPACK_NORMAL_BC3BC3	1
+#define VIRTUAL_TEXTURE_UNPACK_NORMAL_BC5BC1	2
+#define VIRTUAL_TEXTURE_UNPACK_NORMAL_B5G6R5	3
diff --git a/Engine/Plugins/PCG/Shaders/Private/PCGVirtualTextureDataInterface.ush b/Engine/Plugins/PCG/Shaders/Private/PCGVirtualTextureDataInterface.ush
new file mode 100644
--- /dev/null
+++ b/Engine/Plugins/PCG/Shaders/Private/PCGVirtualTextureDataInterface.ush
@@ -0,0 +1,243 @@
+// Copyright Epic Games, Inc. All Rights Reserved.
+
+#include "/Plugin/PCG/Private/PCGVirtualTextureCommon.ush"
+
+#define PCG_VIRTUAL_TEXTURE_DECLARATION(DataIndex) \
+	Texture2D {DataInterfaceName}_VirtualTexture0_##DataIndex; \
+	Texture2D {DataInterfaceName}_VirtualTexture1_##DataIndex; \
+	Texture2D {DataInterfaceName}_VirtualTexture2_##DataIndex; \
+	Texture2D<uint4> {DataInterfaceName}_PageTable0_##DataIndex; \
+	Texture2D<uint4> {DataInterfaceName}_PageTable1_##DataIndex; \
+	Texture2D<uint4> {DataInterfaceName}_PageTable2_##DataIndex; \
+	Texture2D<uint> {DataInterfaceName}_PageTableIndirection0_##DataIndex; \
+	Texture2D<uint> {DataInterfaceName}_PageTableIndirection1_##DataIndex; \
+	Texture2D<uint> {DataInterfaceName}_PageTableIndirection2_##DataIndex;
+
+PCG_VIRTUAL_TEXTURE_DECLARATION(0)
+PCG_VIRTUAL_TEXTURE_DECLARATION(1)
+PCG_VIRTUAL_TEXTURE_DECLARATION(2)
+PCG_VIRTUAL_TEXTURE_DECLARATION(3)
+
+#undef PCG_VIRTUAL_TEXTURE_DECLARATION
+
+// Layer infos
+uint4 {DataInterfaceName}_MaterialTypes[MAX_NUM_BINDINGS];
+uint4 {DataInterfaceName}_ValidLayerMasks[MAX_NUM_BINDINGS];
+float4 {DataInterfaceName}_LWCTiles[MAX_NUM_BINDINGS];
+float4x4 {DataInterfaceName}_WorldToUVTransforms[MAX_NUM_BINDINGS];
+uint4 {DataInterfaceName}_Adaptive[MAX_NUM_BINDINGS * MAX_NUM_LAYERS];
+uint4 {DataInterfaceName}_Uniforms[MAX_NUM_BINDINGS * MAX_NUM_LAYERS];
+uint4 {DataInterfaceName}_PageTableUniforms[MAX_NUM_BINDINGS * 2];
+
+// Shared
+uint {DataInterfaceName}_NumVirtualTextures;
+SamplerState {DataInterfaceName}_Sampler;
+
+float4 SampleLayer_{DataInterfaceName}(
+	uint InDataIndex,
+	float2 InSampleUV,
+	Texture2D InTexture,
+	Texture2D<uint4> InPageTable,
+	Texture2D<uint> InPageTableIndirection,
+	uint InAdaptive,
+	uint4 InUniforms,
+	uint4 PageTableUniforms[2])
+{
+	VTPageTableResult PageTable;
+
+	if (InAdaptive == 0)
+	{
+		PageTable = TextureLoadVirtualPageTableLevel(
+			InPageTable,
+			VTPageTableUniform_Unpack(PageTableUniforms[0], PageTableUniforms[1]),
+			InSampleUV,
+			VTADDRESSMODE_CLAMP,
+			VTADDRESSMODE_CLAMP,
+			0.0f);
+	}
+	else
+	{
+		PageTable = TextureLoadVirtualPageTableAdaptiveLevel(
+			InPageTable,
+			InPageTableIndirection,
+			VTPageTableUniform_Unpack(PageTableUniforms[0], PageTableUniforms[1]),
+			InSampleUV,
+			VTADDRESSMODE_CLAMP,
+			VTADDRESSMODE_CLAMP,
+			0.0f);
+	}
+
+	return TextureVirtualSample(InTexture, {DataInterfaceName}_Sampler, PageTable, 0, VTUniform_Unpack(InUniforms));
+}
+
+#define PCG_CONDITIONAL_VIRTUAL_TEXTURE_SAMPLE(DataIndex) \
+	if (InDataIndex == DataIndex) \
+	{ \
+		LayerSample[0] = (ValidLayerMask & 0x1) != 0 \
+			? SampleLayer_{DataInterfaceName}( \
+				InDataIndex, \
+				SampleUV, \
+				{DataInterfaceName}_VirtualTexture0_##DataIndex, \
+				{DataInterfaceName}_PageTable0_##DataIndex, \
+				{DataInterfaceName}_PageTableIndirection0_##DataIndex, \
+				Adaptive0, \
+				Uniforms0, \
+				PageTableUniforms) \
+			: 0; \
+		LayerSample[1] = (ValidLayerMask & 0x2) != 0 \
+			? SampleLayer_{DataInterfaceName}( \
+				InDataIndex, \
+				SampleUV, \
+				{DataInterfaceName}_VirtualTexture1_##DataIndex, \
+				{DataInterfaceName}_PageTable1_##DataIndex, \
+				{DataInterfaceName}_PageTableIndirection1_##DataIndex, \
+				Adaptive1, \
+				Uniforms1, \
+				PageTableUniforms) \
+			: 0; \
+		LayerSample[2] = (ValidLayerMask & 0x4) != 0 \
+			? SampleLayer_{DataInterfaceName}( \
+				InDataIndex, \
+				SampleUV, \
+				{DataInterfaceName}_VirtualTexture2_##DataIndex, \
+				{DataInterfaceName}_PageTable2_##DataIndex, \
+				{DataInterfaceName}_PageTableIndirection2_##DataIndex, \
+				Adaptive2, \
+				Uniforms2, \
+				PageTableUniforms) \
+			: 0; \
+	}
+
+void SampleVirtualTexture_{DataInterfaceName}(
+	uint InDataIndex,
+	float3 InWorldPos,
+	out bool bOutInsideVolume,
+	out float3 OutBaseColor,
+	out float OutSpecular,
+	out float OutRoughness,
+	out float OutWorldHeight,
+	out float3 OutNormal,
+	out float OutDisplacement,
+	out float OutMask,
+	out float4 OutMask4)
+{
+	bOutInsideVolume = false;
+	OutBaseColor = float3(0.0f, 0.0f, 0.0f);
+	OutSpecular = 0.5f;
+	OutRoughness = 0.5f;
+	OutWorldHeight = 0.0f;
+	OutNormal = float3(0.0f, 0.0f, 1.0f);
+	OutDisplacement = 0.0f;
+	OutMask = 1.0f;
+	OutMask4 = float4(0.0f, 0.0f, 0.0f, 0.0f);
+
+	if (InDataIndex < 0 || InDataIndex >= {DataInterfaceName}_NumVirtualTextures)
+	{
+		return;
+	}
+
+	// Get sample location.
+	const float4x4 WorldToUVTransform = {DataInterfaceName}_WorldToUVTransforms[InDataIndex];
+	const float3 UnpackOrigin = WorldToUVTransform[0].xyz;
+	const float3 UnpackU = WorldToUVTransform[1].xyz;
+	const float3 UnpackV = WorldToUVTransform[2].xyz;
+	const float3 LWCTile = (float3)0;
+
+	const FDFVector3 LWCWorldPos = DFFromTileOffset_Hack(MakeLWCVector3(LWCTile, InWorldPos));
+	const FDFVector3 LWCUVOrigin = DFFromTileOffset_Hack(MakeLWCVector3({DataInterfaceName}_LWCTiles[InDataIndex].xyz, UnpackOrigin));
+	const float2 SampleUV = VirtualTextureWorldToUV(LWCWorldPos, LWCUVOrigin, UnpackU, UnpackV);
+
+	// Test to see if we are inside the volume, but still take the samples as it will clamp to the edge.
+	bOutInsideVolume = all(SampleUV >= 0.0f) && all(SampleUV < 1.0f);
+
+	// Sample RVTs.
+	float4 LayerSample[3];
+	const uint ValidLayerMask = {DataInterfaceName}_ValidLayerMasks[InDataIndex].x;
+
+	{
+		const uint PackedLayerIndex = InDataIndex * MAX_NUM_LAYERS;
+
+		uint Adaptive0 = {DataInterfaceName}_Adaptive[PackedLayerIndex + 0].x;
+		uint Adaptive1 = {DataInterfaceName}_Adaptive[PackedLayerIndex + 1].x;
+		uint Adaptive2 = {DataInterfaceName}_Adaptive[PackedLayerIndex + 2].x;
+
+		uint4 Uniforms0 = {DataInterfaceName}_Uniforms[PackedLayerIndex + 0];
+		uint4 Uniforms1 = {DataInterfaceName}_Uniforms[PackedLayerIndex + 1];
+		uint4 Uniforms2 = {DataInterfaceName}_Uniforms[PackedLayerIndex + 2];
+
+		uint4 PageTableUniforms[2] =
+		{
+			{DataInterfaceName}_PageTableUniforms[InDataIndex * 2 + 0],
+			{DataInterfaceName}_PageTableUniforms[InDataIndex * 2 + 1],
+		};
+
+		PCG_CONDITIONAL_VIRTUAL_TEXTURE_SAMPLE(0)
+#if ENABLE_MULTIPLE_VIRTUAL_TEXTURES
+		PCG_CONDITIONAL_VIRTUAL_TEXTURE_SAMPLE(1)
+		PCG_CONDITIONAL_VIRTUAL_TEXTURE_SAMPLE(2)
+		PCG_CONDITIONAL_VIRTUAL_TEXTURE_SAMPLE(3)
+#endif
+	}
+
+	// Unpack layer samples based on material type.
+	const uint MaterialType = {DataInterfaceName}_MaterialTypes[InDataIndex].x;
+
+	switch (MaterialType)
+	{
+		case ERuntimeVirtualTextureMaterialType_BaseColor:
+		{
+			OutBaseColor = LayerSample[0].rgb;
+			break;
+		}
+		case ERuntimeVirtualTextureMaterialType_Mask4:
+		{
+			OutMask4 = LayerSample[0];
+			break;
+		}
+		case ERuntimeVirtualTextureMaterialType_BaseColor_Normal_Roughness:
+		{
+			OutBaseColor = VirtualTextureUnpackBaseColorSRGB(LayerSample[0]);
+			OutRoughness = LayerSample[1].y;
+			OutNormal = VirtualTextureUnpackNormalBGR565(LayerSample[1]);
+			break;
+		}
+		case ERuntimeVirtualTextureMaterialType_BaseColor_Normal_Specular:
+		{
+			OutBaseColor = LayerSample[0].rgb;
+			OutSpecular = LayerSample[1].x;
+			OutRoughness = LayerSample[1].y;
+			OutNormal = VirtualTextureUnpackNormalBC3BC3(LayerSample[0], LayerSample[1]);
+			break;
+		}
+		case ERuntimeVirtualTextureMaterialType_BaseColor_Normal_Specular_YCoCg:
+		{
+			OutBaseColor = VirtualTextureUnpackBaseColorYCoCg(LayerSample[0]);
+			OutSpecular = LayerSample[2].x;
+			OutRoughness = LayerSample[2].y;
+			OutNormal = VirtualTextureUnpackNormalBC5BC1(LayerSample[1], LayerSample[2]);
+			break;
+		}
+		case ERuntimeVirtualTextureMaterialType_BaseColor_Normal_Specular_Mask_YCoCg:
+		{
+			OutBaseColor = VirtualTextureUnpackBaseColorYCoCg(LayerSample[0]);
+			OutSpecular = LayerSample[2].x;
+			OutRoughness = LayerSample[2].y;
+			OutNormal = VirtualTextureUnpackNormalBC5BC1(LayerSample[1], LayerSample[2]);
+			OutMask = LayerSample[2].w;
+			break;
+		}
+		case ERuntimeVirtualTextureMaterialType_WorldHeight:
+		{
+			const float2 UnpackScaleBias = WorldToUVTransform[3].xy;
+			OutWorldHeight = VirtualTextureUnpackHeight(LayerSample[0], UnpackScaleBias);
+			break;
+		}
+		case ERuntimeVirtualTextureMaterialType_Displacement:
+		{
+			OutDisplacement = LayerSample[0].x;
+			break;
+		}
+	}
+}
+
+#undef PCG_CONDITIONAL_VIRTUAL_TEXTURE_SAMPLE
diff --git a/Engine/Plugins/PCG/Source/PCG/Internal/Compute/PCGComputeCommon.h b/Engine/Plugins/PCG/Source/PCG/Internal/Compute/PCGComputeCommon.h
--- a/Engine/Plugins/PCG/Source/PCG/Internal/Compute/PCGComputeCommon.h
+++ b/Engine/Plugins/PCG/Source/PCG/Internal/Compute/PCGComputeCommon.h
@@ -74,6 +74,7 @@ namespace PCGComputeConstants
 		| EPCGDataType::Param
 		| EPCGDataType::Landscape
 		| EPCGDataType::BaseTexture
+		| EPCGDataType::VirtualTexture
 		| EPCGDataType::Resource
 		| EPCGDataType::ProxyForGPU;
 
diff --git a/Engine/Plugins/PCG/Source/PCG/Private/Compute/Data/PCGVirtualTextureCommon.cpp b/Engine/Plugins/PCG/Source/PCG/Private/Compute/Data/PCGVirtualTextureCommon.cpp
new file mode 100644
--- /dev/null
+++ b/Engine/Plugins/PCG/Source/PCG/Private/Compute/Data/PCGVirtualTextureCommon.cpp
@@ -0,0 +1,93 @@
+// Copyright Epic Games, Inc. All Rights Reserved.
+
+#include "PCGVirtualTextureCommon.h"
+
+#include "RHIResources.h"
+#include "VT/RuntimeVirtualTexture.h"
+
+namespace PCGVirtualTextureCommon
+{
+	void FVirtualTextureLayer::Initialize(const URuntimeVirtualTexture* VirtualTexture, uint32 LayerIndex, bool bSRGB)
+	{
+		Reset();
+	
+		if (::IsValid(VirtualTexture))
+		{
+			if (IAllocatedVirtualTexture* AllocatedTexture = VirtualTexture->GetAllocatedVirtualTexture())
+			{
+				if (FRHIShaderResourceView* PhysicalTextureSRV = AllocatedTexture->GetPhysicalTextureSRV(LayerIndex, bSRGB))
+				{
+					TextureSRV = PhysicalTextureSRV;
+					AllocatedTexture->GetPackedUniform(&TextureUniforms, LayerIndex);
+				}
+			}
+		}
+	}
+	
+	void FVirtualTextureLayer::Reset()
+	{
+		TextureSRV = nullptr;
+		TextureUniforms = FUintVector4::ZeroValue;
+	}
+	
+	bool FVirtualTextureLayer::IsValid() const
+	{
+		return TextureSRV.IsValid();
+	}
+	
+	void FVirtualTexturePageTable::Initialize(const URuntimeVirtualTexture* VirtualTexture, uint32 PageTableIndex, bool bIncludeWorldToUV, bool bIncludeHeightUnpack)
+	{
+		Reset();
+	
+		if (::IsValid(VirtualTexture))
+		{
+			if (IAllocatedVirtualTexture* AllocatedTexture = VirtualTexture->GetAllocatedVirtualTexture())
+			{
+				PageTableRef = PageTableIndex < AllocatedTexture->GetNumPageTableTextures() ? AllocatedTexture->GetPageTableTexture(PageTableIndex) : nullptr;
+				PageTableIndirectionRef = AllocatedTexture->GetPageTableIndirectionTexture();
+				bIsAdaptive = VirtualTexture->GetAdaptivePageTable();
+	
+				if (PageTableRef.IsValid())
+				{
+					AllocatedTexture->GetPackedPageTableUniform(PageTableUniforms);
+	
+					if (bIncludeWorldToUV)
+					{
+						WorldToUVParameters[0] = VirtualTexture->GetUniformParameter(ERuntimeVirtualTextureShaderUniform_WorldToUVTransform0);
+						WorldToUVParameters[1] = VirtualTexture->GetUniformParameter(ERuntimeVirtualTextureShaderUniform_WorldToUVTransform1);
+						WorldToUVParameters[2] = VirtualTexture->GetUniformParameter(ERuntimeVirtualTextureShaderUniform_WorldToUVTransform2);
+					}
+	
+					if (bIncludeHeightUnpack)
+					{
+						WorldToUVParameters[3] = VirtualTexture->GetUniformParameter(ERuntimeVirtualTextureShaderUniform_WorldHeightUnpack);
+					}
+				}
+			}
+		}
+	}
+	
+	void FVirtualTexturePageTable::Reset()
+	{
+		PageTableRef = nullptr;
+		PageTableIndirectionRef = nullptr;
+		bIsAdaptive = false;
+		PageTableUniforms[0] = FUintVector4::ZeroValue;
+		PageTableUniforms[1] = FUintVector4::ZeroValue;
+		WorldToUVParameters[0] = FVector4::Zero();
+		WorldToUVParameters[1] = FVector4::Zero();
+		WorldToUVParameters[2] = FVector4::Zero();
+		WorldToUVParameters[3] = FVector4::Zero();
+	}
+	
+	bool FVirtualTexturePageTable::IsValid() const
+	{
+		// If this is an adaptive RVT then we require the indirection page table.
+		if (bIsAdaptive && !PageTableIndirectionRef.IsValid())
+		{
+			return false;
+		}
+	
+		return PageTableRef.IsValid();
+	}
+}
\ No newline at end of file
diff --git a/Engine/Plugins/PCG/Source/PCG/Private/Compute/Data/PCGVirtualTextureCommon.h b/Engine/Plugins/PCG/Source/PCG/Private/Compute/Data/PCGVirtualTextureCommon.h
new file mode 100644
--- /dev/null
+++ b/Engine/Plugins/PCG/Source/PCG/Private/Compute/Data/PCGVirtualTextureCommon.h
@@ -0,0 +1,53 @@
+// Copyright Epic Games, Inc. All Rights Reserved.
+
+#pragma once
+
+#include "CoreMinimal.h"
+#include "RHIFwd.h"
+#include "VT/RuntimeVirtualTextureEnum.h"
+
+class URuntimeVirtualTexture;
+
+namespace PCGVirtualTextureCommon
+{
+	enum class ENormalUnpackType : int32
+	{
+		None = 0,
+		BC3BC3,
+		BC5BC1,
+		B5G6R5,
+	};
+
+	enum class EBaseColorUnpackType : int32
+	{
+		/** No unpacking is required. */
+		None = 0,
+		/** Base color is manually packed as SRGB. */
+		SRGBUnpack,
+		/** Base color is manually packed as YCoCg. */
+		YCoCgUnpack,
+	};
+
+	struct FVirtualTextureLayer
+	{
+		void Initialize(const URuntimeVirtualTexture* VirtualTexture, uint32 LayerIndex, bool bSRGB);
+		void Reset();
+		bool IsValid() const;
+
+		FShaderResourceViewRHIRef TextureSRV;
+		FUintVector4 TextureUniforms = FUintVector4::ZeroValue;
+	};
+
+	struct FVirtualTexturePageTable
+	{
+		void Initialize(const URuntimeVirtualTexture* VirtualTexture, uint32 PageTableIndex, bool bIncludeWorldToUV, bool bIncludeHeightUnpack);
+		void Reset();
+		bool IsValid() const;
+
+		FTextureRHIRef PageTableRef;
+		FTextureRHIRef PageTableIndirectionRef;
+		bool bIsAdaptive = false;
+		FUintVector4 PageTableUniforms[2];
+		FVector4 WorldToUVParameters[4];
+	};
+}
diff --git a/Engine/Plugins/PCG/Source/PCG/Private/Compute/DataInterfaces/PCGLandscapeDataInterface.cpp b/Engine/Plugins/PCG/Source/PCG/Private/Compute/DataInterfaces/PCGLandscapeDataInterface.cpp
--- a/Engine/Plugins/PCG/Source/PCG/Private/Compute/DataInterfaces/PCGLandscapeDataInterface.cpp
+++ b/Engine/Plugins/PCG/Source/PCG/Private/Compute/DataInterfaces/PCGLandscapeDataInterface.cpp
@@ -128,13 +128,13 @@ FPCGLandscapeResource::FPCGLandscapeResource(const FResourceKey& InKey)
 	CellCount = CaptureVertexSpan;
 
 	// Mapping to get the UV from 'cell space' which is relative to the entire terrain (not just the captured regions).
-	FVector2D UvScale(1.0f / CaptureVertexSpan.X, 1.0f / CaptureVertexSpan.Y);
+	FVector2D UVScale(1.0f / CaptureVertexSpan.X, 1.0f / CaptureVertexSpan.Y);
 
-	UvScaleBias = FVector4(
-		UvScale.X,
-		UvScale.Y,
-		(0.5f - RegionVertexBase.X) * UvScale.X,
-		(0.5f - RegionVertexBase.Y) * UvScale.Y);
+	UVScaleBias = FVector4(
+		UVScale.X,
+		UVScale.Y,
+		(0.5f - RegionVertexBase.X) * UVScale.X,
+		(0.5f - RegionVertexBase.Y) * UVScale.Y);
 
 	FTransform LandscapeTransform = ResourceKey.Source->GetTransform();
 	FLargeWorldRenderPosition LandscapeTransformOrigin(LandscapeTransform.GetLocation());
@@ -194,7 +194,7 @@ BEGIN_SHADER_PARAMETER_STRUCT(FPCGLandscapeDataInterfaceParameters,)
 	SHADER_PARAMETER(uint32,					HeightVirtualTextureAdaptive)
 	SHADER_PARAMETER_SAMPLER(SamplerState,		HeightVirtualTextureSampler)
 	SHADER_PARAMETER(FVector3f,					HeightVirtualTextureLWCTile)
-	SHADER_PARAMETER(FMatrix44f,				HeightVirtualTextureWorldToUvTransform)
+	SHADER_PARAMETER(FMatrix44f,				HeightVirtualTextureWorldToUVTransform)
 	SHADER_PARAMETER(uint32,					HeightVirtualTextureEnabled)
 	SHADER_PARAMETER(FUintVector4,				HeightVirtualTexturePackedUniform0)
 	SHADER_PARAMETER(FUintVector4,				HeightVirtualTexturePackedUniform1)
@@ -209,7 +209,7 @@ BEGIN_SHADER_PARAMETER_STRUCT(FPCGLandscapeDataInterfaceParameters,)
 	SHADER_PARAMETER_SAMPLER(SamplerState,		NormalVirtualTexture0Sampler)
 	SHADER_PARAMETER_SAMPLER(SamplerState,		NormalVirtualTexture1Sampler)
 	SHADER_PARAMETER(FVector3f,					NormalVirtualTextureLWCTile)
-	SHADER_PARAMETER(FMatrix44f,				NormalVirtualTextureWorldToUvTransform)
+	SHADER_PARAMETER(FMatrix44f,				NormalVirtualTextureWorldToUVTransform)
 	SHADER_PARAMETER(FUintVector4,				NormalVirtualTexturePackedUniform0)
 	SHADER_PARAMETER(FUintVector4,				NormalVirtualTexturePackedUniform1)
 	SHADER_PARAMETER(FUintVector4,				NormalVirtualTextureUniforms0)
@@ -224,7 +224,7 @@ BEGIN_SHADER_PARAMETER_STRUCT(FPCGLandscapeDataInterfaceParameters,)
 	SHADER_PARAMETER(uint32,					BaseColorVirtualTextureAdaptive)
 	SHADER_PARAMETER_SAMPLER(SamplerState,		BaseColorVirtualTextureSampler)
 	SHADER_PARAMETER(FVector3f,					BaseColorVirtualTextureLWCTile)
-	SHADER_PARAMETER(FMatrix44f,				BaseColorVirtualTextureWorldToUvTransform)
+	SHADER_PARAMETER(FMatrix44f,				BaseColorVirtualTextureWorldToUVTransform)
 	SHADER_PARAMETER(uint32,					BaseColorVirtualTextureUnpackSRGB)
 	SHADER_PARAMETER(uint32,					BaseColorVirtualTextureUnpackYCoCg)
 	SHADER_PARAMETER(uint32,					BaseColorVirtualTextureEnabled)
@@ -237,9 +237,9 @@ BEGIN_SHADER_PARAMETER_STRUCT(FPCGLandscapeDataInterfaceParameters,)
 	SHADER_PARAMETER_SAMPLER(SamplerState,		CollisionHeightTextureSampler)
 	SHADER_PARAMETER(uint32,					CollisionHeightTextureEnabled)
 	SHADER_PARAMETER(FVector3f,					CollisionHeightTextureLWCTile)
-	SHADER_PARAMETER(FMatrix44f,				CollisionHeightTextureWorldToUvTransform)
-	SHADER_PARAMETER(FMatrix44f,				CollisionHeightTextureUvToWorldTransform)
-	SHADER_PARAMETER(FVector4f,					CollisionHeightTextureUvScaleBias)
+	SHADER_PARAMETER(FMatrix44f,				CollisionHeightTextureWorldToUVTransform)
+	SHADER_PARAMETER(FMatrix44f,				CollisionHeightTextureUVToWorldTransform)
+	SHADER_PARAMETER(FVector4f,					CollisionHeightTextureUVScaleBias)
 	SHADER_PARAMETER(FIntPoint,					CollisionHeightTextureDimension)
 
 	// Misc parameters
@@ -261,6 +261,8 @@ TCHAR const* UPCGLandscapeDataInterface::GetShaderVirtualPath() const
 void UPCGLandscapeDataInterface::GetShaderHash(FString& InOutKey) const
 {
 	GetShaderFileHash(TemplateFilePath, EShaderPlatform::SP_PCD3D_SM5).AppendString(InOutKey);
+	GetShaderFileHash(TEXT("/Engine/Private/VirtualTextureCommon.ush"), EShaderPlatform::SP_PCD3D_SM5).AppendString(InOutKey);
+	GetShaderFileHash(TEXT("/Plugin/PCG/Private/PCGVirtualTextureCommon.ush"), EShaderPlatform::SP_PCD3D_SM5).AppendString(InOutKey);
 }
 
 void UPCGLandscapeDataInterface::GetHLSL(FString& OutHLSL, FString const& InDataInterfaceName) const
@@ -423,11 +425,11 @@ void UPCGLandscapeDataProvider::InitializeRuntimeVirtualTextures(ALandscape* InL
 
 					if (RVT->IsLayerYCoCg(0))
 					{
-						BaseColorVirtualTextureUnpackType = EPCGVirtualTextureBaseColorUnpackType::YCoCgUnpack;
+						BaseColorVirtualTextureUnpackType = PCGVirtualTextureCommon::EBaseColorUnpackType::YCoCgUnpack;
 					}
 					else if (VirtualMaterialType == ERuntimeVirtualTextureMaterialType::BaseColor_Normal_Roughness)
 					{
-						BaseColorVirtualTextureUnpackType = EPCGVirtualTextureBaseColorUnpackType::SRGBUnpack;
+						BaseColorVirtualTextureUnpackType = PCGVirtualTextureCommon::EBaseColorUnpackType::SRGBUnpack;
 					}
 				}
 			}
@@ -569,104 +571,20 @@ void UPCGLandscapeDataProvider::Reset()
 	Resource.Reset();
 	bCanSampleVirtualTextures = false;
 	bBaseColorSRGB = false;
-	BaseColorVirtualTextureUnpackType = EPCGVirtualTextureBaseColorUnpackType::None;
+	BaseColorVirtualTextureUnpackType = PCGVirtualTextureCommon::EBaseColorUnpackType::None;
 	NormalVirtualTextureMode = ERuntimeVirtualTextureMaterialType::Count;
 	LandscapeGridSize = FVector2D(1.0f, 1.0f);
 
 	Super::Reset();
 }
 
-void FPCGLandscapeDataProviderProxy::FVirtualTextureLayer::Initialize(const URuntimeVirtualTexture* VirtualTexture, uint32 LayerIndex, bool bSRGB)
-{
-	Reset();
-
-	if (VirtualTexture && ::IsValid(VirtualTexture))
-	{
-		if (IAllocatedVirtualTexture* AllocatedTexture = VirtualTexture->GetAllocatedVirtualTexture())
-		{
-			if (FRHIShaderResourceView* PhysicalTextureSRV = AllocatedTexture->GetPhysicalTextureSRV(LayerIndex, bSRGB))
-			{
-				TextureSRV = PhysicalTextureSRV;
-				AllocatedTexture->GetPackedUniform(&TextureUniforms, LayerIndex);
-			}
-		}
-	}
-}
-
-void FPCGLandscapeDataProviderProxy::FVirtualTextureLayer::Reset()
-{
-	TextureSRV = nullptr;
-	TextureUniforms = FUintVector4(0, 0, 0, 0);
-}
-
-bool FPCGLandscapeDataProviderProxy::FVirtualTextureLayer::IsValid() const
-{
-	return TextureSRV.IsValid();
-}
-
-void FPCGLandscapeDataProviderProxy::FVirtualTexturePageTable::Initialize(const URuntimeVirtualTexture* VirtualTexture, uint32 PageTableIndex, bool bIncludeWorldToUv, bool bIncludeHeightUnpack)
-{
-	Reset();
-
-	if (VirtualTexture && ::IsValid(VirtualTexture))
-	{
-		if (IAllocatedVirtualTexture* AllocatedTexture = VirtualTexture->GetAllocatedVirtualTexture())
-		{
-			PageTableRef = AllocatedTexture->GetPageTableTexture(PageTableIndex);
-			PageTableIndirectionRef = AllocatedTexture->GetPageTableIndirectionTexture();
-			bIsAdaptive = VirtualTexture->GetAdaptivePageTable();
-
-			if (PageTableRef.IsValid())
-			{
-				AllocatedTexture->GetPackedPageTableUniform(PageTableUniforms);
-
-				if (bIncludeWorldToUv)
-				{
-					WorldToUvParameters[0] = VirtualTexture->GetUniformParameter(ERuntimeVirtualTextureShaderUniform_WorldToUVTransform0);
-					WorldToUvParameters[1] = VirtualTexture->GetUniformParameter(ERuntimeVirtualTextureShaderUniform_WorldToUVTransform1);
-					WorldToUvParameters[2] = VirtualTexture->GetUniformParameter(ERuntimeVirtualTextureShaderUniform_WorldToUVTransform2);
-				}
-
-				if (bIncludeHeightUnpack)
-				{
-					WorldToUvParameters[3] = VirtualTexture->GetUniformParameter(ERuntimeVirtualTextureShaderUniform_WorldHeightUnpack);
-				}
-			}
-		}
-	}
-}
-
-void FPCGLandscapeDataProviderProxy::FVirtualTexturePageTable::Reset()
-{
-	PageTableRef = nullptr;
-	PageTableIndirectionRef = nullptr;
-	bIsAdaptive = false;
-	PageTableUniforms[0] = FUintVector4(0, 0, 0, 0);
-	PageTableUniforms[1] = FUintVector4(0, 0, 0, 0);
-	WorldToUvParameters[0] = FVector4::Zero();
-	WorldToUvParameters[1] = FVector4::Zero();
-	WorldToUvParameters[2] = FVector4::Zero();
-	WorldToUvParameters[3] = FVector4::Zero();
-}
-
-bool FPCGLandscapeDataProviderProxy::FVirtualTexturePageTable::IsValid() const
-{
-	// If this is an adaptive RVT then we require the indirection page table.
-	if (bIsAdaptive && !PageTableIndirectionRef.IsValid())
-	{
-		return false;
-	}
-
-	return PageTableRef.IsValid();
-}
-
 FPCGLandscapeDataProviderProxy::FPCGLandscapeDataProviderProxy(
 	TSharedPtr<const FPCGLandscapeResource> InResource,
 	const URuntimeVirtualTexture* InBaseColorVirtualTexture,
 	const URuntimeVirtualTexture* InHeightVirtualTexture,
 	const URuntimeVirtualTexture* InNormalVirtualTexture,
 	bool bInBaseColorSRGB,
-	EPCGVirtualTextureBaseColorUnpackType InBaseColorVirtualTextureUnpackType,
+	PCGVirtualTextureCommon::EBaseColorUnpackType InBaseColorVirtualTextureUnpackType,
 	ERuntimeVirtualTextureMaterialType InNormalVirtualTextureMode,
 	FVector2D InLandscapeGridSize)
 	: Resource(InResource)
@@ -674,37 +592,36 @@ FPCGLandscapeDataProviderProxy::FPCGLandscapeDataProviderProxy(
 	, LandscapeGridSize(InLandscapeGridSize)
 {
 	BaseColorVirtualLayer.Initialize(InBaseColorVirtualTexture, /*PageTableIndex=*/0, bInBaseColorSRGB);
-	BaseColorVirtualPage.Initialize(InBaseColorVirtualTexture, /*PageTableIndex=*/0, /*bIncludeWorldToUv=*/true, /*bIncludeHeightUnpack=*/false);
+	BaseColorVirtualPage.Initialize(InBaseColorVirtualTexture, /*PageTableIndex=*/0, /*bIncludeWorldToUV=*/true, /*bIncludeHeightUnpack=*/false);
 
 	HeightVirtualLayer.Initialize(InHeightVirtualTexture, /*PageTableIndex=*/0, /*bSRGB=*/false);
-	HeightVirtualPage.Initialize(InHeightVirtualTexture, /*PageTableIndex=*/0, /*bIncludeWorldToUv=*/true, /*bIncludeHeightUnpack=*/true);
+	HeightVirtualPage.Initialize(InHeightVirtualTexture, /*PageTableIndex=*/0, /*bIncludeWorldToUV=*/true, /*bIncludeHeightUnpack=*/true);
 
-	NormalVirtualPage.Initialize(InNormalVirtualTexture, /*PageTableIndex=*/0, /*bIncludeWorldToUv=*/true, /*bIncludeHeightUnpack=*/false);
+	NormalVirtualPage.Initialize(InNormalVirtualTexture, /*PageTableIndex=*/0, /*bIncludeWorldToUV=*/true, /*bIncludeHeightUnpack=*/false);
 
 	switch (InNormalVirtualTextureMode)
 	{
 	case ERuntimeVirtualTextureMaterialType::BaseColor_Normal_Roughness:
 		NormalVirtualLayer0.Initialize(InNormalVirtualTexture, /*PageTableIndex=*/0, /*bSRGB=*/false);
 		NormalVirtualLayer1.Initialize(InNormalVirtualTexture, /*PageTableIndex=*/1, /*bSRGB=*/false);
-		NormalUnpackMode = ENormalUnpackType::B5G6R5;
+		NormalUnpackMode = PCGVirtualTextureCommon::ENormalUnpackType::B5G6R5;
 		break;
 	case ERuntimeVirtualTextureMaterialType::BaseColor_Normal_Specular:
 		NormalVirtualLayer0.Initialize(InNormalVirtualTexture, /*PageTableIndex=*/0, /*bSRGB=*/false);
 		NormalVirtualLayer1.Initialize(InNormalVirtualTexture, /*PageTableIndex=*/1, /*bSRGB=*/false);
-		NormalUnpackMode = ENormalUnpackType::BC3BC3;
+		NormalUnpackMode = PCGVirtualTextureCommon::ENormalUnpackType::BC3BC3;
 		break;
-
 	case ERuntimeVirtualTextureMaterialType::BaseColor_Normal_Specular_YCoCg:
 	case ERuntimeVirtualTextureMaterialType::BaseColor_Normal_Specular_Mask_YCoCg:
 		NormalVirtualLayer0.Initialize(InNormalVirtualTexture, /*PageTableIndex=*/1, /*bSRGB=*/false);
 		NormalVirtualLayer1.Initialize(InNormalVirtualTexture, /*PageTableIndex=*/2, /*bSRGB=*/false);
-		NormalUnpackMode = ENormalUnpackType::BC5BC1;
+		NormalUnpackMode = PCGVirtualTextureCommon::ENormalUnpackType::BC5BC1;
 		break;
 	default:
 		NormalVirtualPage.Reset();
 		NormalVirtualLayer0.Reset();
 		NormalVirtualLayer1.Reset();
-		NormalUnpackMode = ENormalUnpackType::None;
+		NormalUnpackMode = PCGVirtualTextureCommon::ENormalUnpackType::None;
 		break;
 	}
 }
@@ -777,16 +694,16 @@ bool FPCGLandscapeDataProviderProxy::SetBaseColorVirtualTextureParameters(FParam
 	OutShaderParameters.BaseColorVirtualTexturePageTableIndirection = BaseColorVirtualPage.PageTableIndirectionRef;
 	OutShaderParameters.BaseColorVirtualTextureAdaptive = BaseColorVirtualPage.bIsAdaptive;
 
-	FLargeWorldRenderPosition BaseColorVirtualTextureOrigin(BaseColorVirtualPage.WorldToUvParameters[0]);
+	FLargeWorldRenderPosition BaseColorVirtualTextureOrigin(BaseColorVirtualPage.WorldToUVParameters[0]);
 	OutShaderParameters.BaseColorVirtualTextureLWCTile = BaseColorVirtualTextureOrigin.GetTile();
 
-	OutShaderParameters.BaseColorVirtualTextureWorldToUvTransform = FMatrix44f(
+	OutShaderParameters.BaseColorVirtualTextureWorldToUVTransform = FMatrix44f(
 		BaseColorVirtualTextureOrigin.GetOffset(),
-		FVector3f((FVector4f)BaseColorVirtualPage.WorldToUvParameters[1]),
-		FVector3f((FVector4f)BaseColorVirtualPage.WorldToUvParameters[2]),
+		FVector3f((FVector4f)BaseColorVirtualPage.WorldToUVParameters[1]),
+		FVector3f((FVector4f)BaseColorVirtualPage.WorldToUVParameters[2]),
 		FVector3f(0, 0, 0));
-	OutShaderParameters.BaseColorVirtualTextureUnpackSRGB = BaseColorVirtualTextureUnpackType == EPCGVirtualTextureBaseColorUnpackType::SRGBUnpack;
-	OutShaderParameters.BaseColorVirtualTextureUnpackYCoCg = BaseColorVirtualTextureUnpackType == EPCGVirtualTextureBaseColorUnpackType::YCoCgUnpack;
+	OutShaderParameters.BaseColorVirtualTextureUnpackSRGB = BaseColorVirtualTextureUnpackType == PCGVirtualTextureCommon::EBaseColorUnpackType::SRGBUnpack;
+	OutShaderParameters.BaseColorVirtualTextureUnpackYCoCg = BaseColorVirtualTextureUnpackType == PCGVirtualTextureCommon::EBaseColorUnpackType::YCoCgUnpack;
 	OutShaderParameters.BaseColorVirtualTextureEnabled = 1;
 	OutShaderParameters.BaseColorVirtualTexturePackedUniform0 = BaseColorVirtualPage.PageTableUniforms[0];
 	OutShaderParameters.BaseColorVirtualTexturePackedUniform1 = BaseColorVirtualPage.PageTableUniforms[1];
@@ -807,14 +724,14 @@ bool FPCGLandscapeDataProviderProxy::SetHeightVirtualTextureParameters(FParamete
 	OutShaderParameters.HeightVirtualTexturePageTableIndirection = HeightVirtualPage.PageTableIndirectionRef;
 	OutShaderParameters.HeightVirtualTextureAdaptive = HeightVirtualPage.bIsAdaptive;
 
-	FLargeWorldRenderPosition HeightVirtualTextureOrigin(HeightVirtualPage.WorldToUvParameters[0]);
+	FLargeWorldRenderPosition HeightVirtualTextureOrigin(HeightVirtualPage.WorldToUVParameters[0]);
 	OutShaderParameters.HeightVirtualTextureLWCTile = HeightVirtualTextureOrigin.GetTile();
 
-	OutShaderParameters.HeightVirtualTextureWorldToUvTransform = FMatrix44f(
+	OutShaderParameters.HeightVirtualTextureWorldToUVTransform = FMatrix44f(
 		HeightVirtualTextureOrigin.GetOffset(),
-		FVector3f((FVector4f)HeightVirtualPage.WorldToUvParameters[1]),
-		FVector3f((FVector4f)HeightVirtualPage.WorldToUvParameters[2]),
-		FVector3f((FVector4f)HeightVirtualPage.WorldToUvParameters[3]));
+		FVector3f((FVector4f)HeightVirtualPage.WorldToUVParameters[1]),
+		FVector3f((FVector4f)HeightVirtualPage.WorldToUVParameters[2]),
+		FVector3f((FVector4f)HeightVirtualPage.WorldToUVParameters[3]));
 
 	OutShaderParameters.HeightVirtualTextureEnabled = 1;
 	OutShaderParameters.HeightVirtualTexturePackedUniform0 = HeightVirtualPage.PageTableUniforms[0];
@@ -837,13 +754,13 @@ bool FPCGLandscapeDataProviderProxy::SetNormalVirtualTextureParameters(FParamete
 	OutShaderParameters.NormalVirtualTexturePageTableIndirection = NormalVirtualPage.PageTableIndirectionRef;
 	OutShaderParameters.NormalVirtualTextureAdaptive = NormalVirtualPage.bIsAdaptive;
 
-	FLargeWorldRenderPosition NormalVirtualTextureOrigin(NormalVirtualPage.WorldToUvParameters[0]);
+	FLargeWorldRenderPosition NormalVirtualTextureOrigin(NormalVirtualPage.WorldToUVParameters[0]);
 	OutShaderParameters.NormalVirtualTextureLWCTile = NormalVirtualTextureOrigin.GetTile();
 
-	OutShaderParameters.NormalVirtualTextureWorldToUvTransform = FMatrix44f(
+	OutShaderParameters.NormalVirtualTextureWorldToUVTransform = FMatrix44f(
 		FPlane4f(NormalVirtualTextureOrigin.GetOffset()),
-		FPlane4f((FVector4f)NormalVirtualPage.WorldToUvParameters[1]),
-		FPlane4f((FVector4f)NormalVirtualPage.WorldToUvParameters[2]),
+		FPlane4f((FVector4f)NormalVirtualPage.WorldToUVParameters[1]),
+		FPlane4f((FVector4f)NormalVirtualPage.WorldToUVParameters[2]),
 		FPlane4f(0.0f, 0.0f, 0.0f, 1.0f));
 
 	OutShaderParameters.NormalVirtualTextureEnabled = 1;
@@ -868,9 +785,9 @@ bool FPCGLandscapeDataProviderProxy::SetCollisionHeightTextureParameters(FParame
 	const FIntPoint TextureDimensions(Resource->LandscapeTexture->GetDimensions());
 
 	OutShaderParameters.CollisionHeightTextureLWCTile = Resource->LandscapeLWCTile;
-	OutShaderParameters.CollisionHeightTextureWorldToUvTransform = (FMatrix44f)Resource->WorldToActorTransform;
-	OutShaderParameters.CollisionHeightTextureUvToWorldTransform = (FMatrix44f)Resource->ActorToWorldTransform;
-	OutShaderParameters.CollisionHeightTextureUvScaleBias = (FVector4f)Resource->UvScaleBias;
+	OutShaderParameters.CollisionHeightTextureWorldToUVTransform = (FMatrix44f)Resource->WorldToActorTransform;
+	OutShaderParameters.CollisionHeightTextureUVToWorldTransform = (FMatrix44f)Resource->ActorToWorldTransform;
+	OutShaderParameters.CollisionHeightTextureUVScaleBias = (FVector4f)Resource->UVScaleBias;
 	OutShaderParameters.CollisionHeightTexture = CollisionHeightTexture;
 	OutShaderParameters.CollisionHeightTextureEnabled = 1;
 	OutShaderParameters.CollisionHeightTextureDimension = TextureDimensions;
@@ -886,7 +803,8 @@ void FPCGLandscapeDataProviderProxy::SetBaseColorVirtualTextureParameters_Defaul
 	OutShaderParameters.BaseColorVirtualTexturePageTable = GBlackUintTexture->TextureRHI;
 	OutShaderParameters.BaseColorVirtualTexturePageTableIndirection = GBlackUintTexture->TextureRHI;
 	OutShaderParameters.BaseColorVirtualTextureAdaptive = 0;
-	OutShaderParameters.BaseColorVirtualTextureWorldToUvTransform = FMatrix44f::Identity;
+	OutShaderParameters.BaseColorVirtualTextureLWCTile = FVector3f::ZeroVector;
+	OutShaderParameters.BaseColorVirtualTextureWorldToUVTransform = FMatrix44f::Identity;
 	OutShaderParameters.BaseColorVirtualTextureUnpackSRGB = 0;
 	OutShaderParameters.BaseColorVirtualTextureUnpackYCoCg = 0;
 	OutShaderParameters.BaseColorVirtualTextureEnabled = 0;
@@ -903,7 +821,8 @@ void FPCGLandscapeDataProviderProxy::SetHeightVirtualTextureParameters_Default(F
 	OutShaderParameters.HeightVirtualTexturePageTable = GBlackUintTexture->TextureRHI;
 	OutShaderParameters.HeightVirtualTexturePageTableIndirection = GBlackUintTexture->TextureRHI;
 	OutShaderParameters.HeightVirtualTextureAdaptive = 0;
-	OutShaderParameters.HeightVirtualTextureWorldToUvTransform = FMatrix44f::Identity;
+	OutShaderParameters.HeightVirtualTextureLWCTile = FVector3f::ZeroVector;
+	OutShaderParameters.HeightVirtualTextureWorldToUVTransform = FMatrix44f::Identity;
 	OutShaderParameters.HeightVirtualTextureEnabled = 0;
 	OutShaderParameters.HeightVirtualTexturePackedUniform0 = DummyUint4;
 	OutShaderParameters.HeightVirtualTexturePackedUniform1 = DummyUint4;
@@ -919,9 +838,10 @@ void FPCGLandscapeDataProviderProxy::SetNormalVirtualTextureParameters_Default(F
 	OutShaderParameters.NormalVirtualTexturePageTable = GBlackUintTexture->TextureRHI;
 	OutShaderParameters.NormalVirtualTexturePageTableIndirection = GBlackUintTexture->TextureRHI;
 	OutShaderParameters.NormalVirtualTextureAdaptive = 0;
-	OutShaderParameters.NormalVirtualTextureWorldToUvTransform = FMatrix44f::Identity;
+	OutShaderParameters.NormalVirtualTextureLWCTile = FVector3f::ZeroVector;
+	OutShaderParameters.NormalVirtualTextureWorldToUVTransform = FMatrix44f::Identity;
 	OutShaderParameters.NormalVirtualTextureEnabled = 0;
-	OutShaderParameters.NormalVirtualTextureUnpackMode = (int32)ENormalUnpackType::None;
+	OutShaderParameters.NormalVirtualTextureUnpackMode = (int32)PCGVirtualTextureCommon::ENormalUnpackType::None;
 	OutShaderParameters.NormalVirtualTexturePackedUniform0 = DummyUint4;
 	OutShaderParameters.NormalVirtualTexturePackedUniform1 = DummyUint4;
 	OutShaderParameters.NormalVirtualTextureUniforms0 = DummyUint4;
@@ -934,9 +854,9 @@ void FPCGLandscapeDataProviderProxy::SetCollisionHeightTextureParameters_Default
 
 	OutShaderParameters.CollisionHeightTexture = GBlackTexture->TextureRHI;
 	OutShaderParameters.CollisionHeightTextureLWCTile = FVector3f::ZeroVector;
-	OutShaderParameters.CollisionHeightTextureWorldToUvTransform = FMatrix44f::Identity;
-	OutShaderParameters.CollisionHeightTextureUvToWorldTransform = FMatrix44f::Identity;
-	OutShaderParameters.CollisionHeightTextureUvScaleBias = DummyVector4;
+	OutShaderParameters.CollisionHeightTextureWorldToUVTransform = FMatrix44f::Identity;
+	OutShaderParameters.CollisionHeightTextureUVToWorldTransform = FMatrix44f::Identity;
+	OutShaderParameters.CollisionHeightTextureUVScaleBias = DummyVector4;
 	OutShaderParameters.CollisionHeightTextureEnabled = 0;
 }
 
diff --git a/Engine/Plugins/PCG/Source/PCG/Private/Compute/DataInterfaces/PCGLandscapeDataInterface.h b/Engine/Plugins/PCG/Source/PCG/Private/Compute/DataInterfaces/PCGLandscapeDataInterface.h
--- a/Engine/Plugins/PCG/Source/PCG/Private/Compute/DataInterfaces/PCGLandscapeDataInterface.h
+++ b/Engine/Plugins/PCG/Source/PCG/Private/Compute/DataInterfaces/PCGLandscapeDataInterface.h
@@ -2,11 +2,10 @@
 
 #pragma once
 
+#include "Compute/Data/PCGVirtualTextureCommon.h"
 #include "Compute/DataInterfaces/PCGComputeDataInterface.h"
-#include "ComputeFramework/ComputeDataProvider.h"
 
-#include "RenderResource.h"
-#include "VT/RuntimeVirtualTextureEnum.h"
+#include "ComputeFramework/ComputeDataProvider.h"
 
 #include "PCGLandscapeDataInterface.generated.h"
 
@@ -15,16 +14,6 @@ class FPCGLandscapeDataInterfaceParameters;
 class FPCGLandscapeTextureResource;
 class URuntimeVirtualTexture;
 
-enum class EPCGVirtualTextureBaseColorUnpackType : int32
-{
-	/** No unpacking is required. */
-	None = 0,
-	/** Base color is manually packed as SRGB. */
-	SRGBUnpack,
-	/** Base color is manually packed as YCoCg. */
-	YCoCgUnpack,
-};
-
 /** Manages the textures created from landscape height data. */
 class FPCGLandscapeResource
 {
@@ -45,7 +34,7 @@ public:
 	FVector3f LandscapeLWCTile = FVector3f::ZeroVector;
 	FMatrix ActorToWorldTransform = FMatrix::Identity;
 	FMatrix WorldToActorTransform = FMatrix::Identity;
-	FVector4 UvScaleBias = FVector4(1.0f, 1.0f, 0.0f, 0.0f);
+	FVector4 UVScaleBias = FVector4(1.0f, 1.0f, 0.0f, 0.0f);
 	FIntPoint CellCount = FIntPoint(ForceInitToZero);
 
 private:
@@ -104,7 +93,7 @@ private:
 	TSharedPtr<FPCGLandscapeResource> Resource;
 	bool bCanSampleVirtualTextures = false;
 	bool bBaseColorSRGB = false;
-	EPCGVirtualTextureBaseColorUnpackType BaseColorVirtualTextureUnpackType = EPCGVirtualTextureBaseColorUnpackType::None;
+	PCGVirtualTextureCommon::EBaseColorUnpackType BaseColorVirtualTextureUnpackType = PCGVirtualTextureCommon::EBaseColorUnpackType::None;
 	ERuntimeVirtualTextureMaterialType NormalVirtualTextureMode = ERuntimeVirtualTextureMaterialType::Count;
 	FVector2D LandscapeGridSize = FVector2D(1.0f, 1.0f);
 };
@@ -117,7 +106,7 @@ public:
 		const URuntimeVirtualTexture* InHeightVirtualTexture,
 		const URuntimeVirtualTexture* InNormalVirtualTexture,
 		bool bBaseColorSRGB,
-		EPCGVirtualTextureBaseColorUnpackType BaseColorVirtualTextureUnpackType,
+		PCGVirtualTextureCommon::EBaseColorUnpackType BaseColorVirtualTextureUnpackType,
 		ERuntimeVirtualTextureMaterialType NormalVirtualTextureMode,
 		FVector2D InWorldGridSize);
 
@@ -130,37 +119,6 @@ public:
 protected:
 	using FParameters = FPCGLandscapeDataInterfaceParameters;
 
-	enum class ENormalUnpackType : int32
-	{
-		None = 0,
-		BC3BC3,
-		BC5BC1,
-		B5G6R5
-	};
-
-	struct FVirtualTextureLayer
-	{
-		void Initialize(const URuntimeVirtualTexture* VirtualTexture, uint32 LayerIndex, bool bSRGB);
-		void Reset();
-		bool IsValid() const;
-	
-		FShaderResourceViewRHIRef TextureSRV;
-		FUintVector4 TextureUniforms = FUintVector4(0, 0, 0, 0);
-	};
-	
-	struct FVirtualTexturePageTable
-	{
-		void Initialize(const URuntimeVirtualTexture* VirtualTexture, uint32 PageTableIndex, bool bIncludeWorldToUv, bool bIncludeHeightUnpack);
-		void Reset();
-		bool IsValid() const;
-	
-		FTextureRHIRef PageTableRef;
-		FTextureRHIRef PageTableIndirectionRef;
-		bool bIsAdaptive = false;
-		FUintVector4 PageTableUniforms[2];
-		FVector4 WorldToUvParameters[4];
-	};
-
 	/** Sets all shader parameters. */
 	void SetShaderParameters(FParameters& OutShaderParameters);
 
@@ -179,17 +137,17 @@ protected:
 protected:
 	TSharedPtr<const FPCGLandscapeResource> Resource;
 
-	FVirtualTexturePageTable BaseColorVirtualPage;
-	FVirtualTextureLayer BaseColorVirtualLayer;
-	EPCGVirtualTextureBaseColorUnpackType BaseColorVirtualTextureUnpackType = EPCGVirtualTextureBaseColorUnpackType::None;
+	PCGVirtualTextureCommon::FVirtualTexturePageTable BaseColorVirtualPage;
+	PCGVirtualTextureCommon::FVirtualTextureLayer BaseColorVirtualLayer;
+	PCGVirtualTextureCommon::EBaseColorUnpackType BaseColorVirtualTextureUnpackType = PCGVirtualTextureCommon::EBaseColorUnpackType::None;
 	
-	FVirtualTexturePageTable HeightVirtualPage;
-	FVirtualTextureLayer HeightVirtualLayer;
+	PCGVirtualTextureCommon::FVirtualTexturePageTable HeightVirtualPage;
+	PCGVirtualTextureCommon::FVirtualTextureLayer HeightVirtualLayer;
 	
-	FVirtualTexturePageTable NormalVirtualPage;
-	FVirtualTextureLayer NormalVirtualLayer0;
-	FVirtualTextureLayer NormalVirtualLayer1;
-	ENormalUnpackType NormalUnpackMode;
+	PCGVirtualTextureCommon::FVirtualTexturePageTable NormalVirtualPage;
+	PCGVirtualTextureCommon::FVirtualTextureLayer NormalVirtualLayer0;
+	PCGVirtualTextureCommon::FVirtualTextureLayer NormalVirtualLayer1;
+	PCGVirtualTextureCommon::ENormalUnpackType NormalUnpackMode;
 
 	FVector2D LandscapeGridSize = FVector2D(1.0f, 1.0f);
 };
diff --git a/Engine/Plugins/PCG/Source/PCG/Private/Compute/DataInterfaces/PCGVirtualTextureDataInterface.cpp b/Engine/Plugins/PCG/Source/PCG/Private/Compute/DataInterfaces/PCGVirtualTextureDataInterface.cpp
new file mode 100644
--- /dev/null
+++ b/Engine/Plugins/PCG/Source/PCG/Private/Compute/DataInterfaces/PCGVirtualTextureDataInterface.cpp
@@ -0,0 +1,337 @@
+// Copyright Epic Games, Inc. All Rights Reserved.
+
+#include "PCGVirtualTextureDataInterface.h"
+
+#include "PCGComponent.h"
+#include "PCGData.h"
+#include "Compute/PCGDataBinding.h"
+#include "Data/PCGVirtualTextureData.h"
+
+#include "GlobalRenderResources.h"
+#include "RHIResources.h"
+#include "RHIStaticStates.h"
+#include "ShaderCompilerCore.h"
+#include "ShaderParameterMetadataBuilder.h"
+#include "ComputeFramework/ComputeKernelPermutationSet.h"
+#include "ComputeFramework/ComputeKernelPermutationVector.h"
+#include "ComputeFramework/ShaderParamTypeDefinition.h"
+#include "Misc/LargeWorldRenderPosition.h"
+#include "VT/RuntimeVirtualTexture.h"
+
+#include UE_INLINE_GENERATED_CPP_BY_NAME(PCGVirtualTextureDataInterface)
+
+#define LOCTEXT_NAMESPACE "PCGVirtualTextureDataInterface"
+
+using namespace PCGVirtualTextureDataInterfaceConstants;
+static const TCHAR* EnableMultipleVirtualTexturePermutationName = TEXT("ENABLE_MULTIPLE_VIRTUAL_TEXTURES");
+
+void UPCGVirtualTextureDataInterface::GetSupportedInputs(TArray<FShaderFunctionDefinition>& OutFunctions) const
+{
+	OutFunctions.AddDefaulted_GetRef()
+		.SetName(TEXT("SampleVirtualTexture"))
+		.AddParam(EShaderFundamentalType::Uint) // InDataIndex
+		.AddParam(EShaderFundamentalType::Float, 3) // InWorldPos
+		.AddParam(EShaderFundamentalType::Bool, 0, 0, EShaderParamModifier::Out) // bOutInsideVolume
+		.AddParam(EShaderFundamentalType::Float, 3, 0, EShaderParamModifier::Out) // OutBaseColor
+		.AddParam(EShaderFundamentalType::Float, 0, 0, EShaderParamModifier::Out) // OutSpecular
+		.AddParam(EShaderFundamentalType::Float, 0, 0, EShaderParamModifier::Out) // OutRoughness
+		.AddParam(EShaderFundamentalType::Float, 0, 0, EShaderParamModifier::Out) // OutWorldHeight
+		.AddParam(EShaderFundamentalType::Float, 3, 0, EShaderParamModifier::Out) // OutNormal
+		.AddParam(EShaderFundamentalType::Float, 0, 0, EShaderParamModifier::Out) // OutDisplacement
+		.AddParam(EShaderFundamentalType::Float, 0, 0, EShaderParamModifier::Out) // OutMask
+		.AddParam(EShaderFundamentalType::Float, 4, 0, EShaderParamModifier::Out); // OutMask4
+}
+
+BEGIN_SHADER_PARAMETER_STRUCT(FPCGVirtualTextureDataInterfaceParameters, )
+	// Layer 0
+	SHADER_PARAMETER_SRV_ARRAY(Texture2D, VirtualTexture0, [MAX_NUM_BINDINGS])
+	SHADER_PARAMETER_TEXTURE_ARRAY(Texture2D<uint4>, PageTable0, [MAX_NUM_BINDINGS])
+	SHADER_PARAMETER_TEXTURE_ARRAY(Texture2D<uint>, PageTableIndirection0, [MAX_NUM_BINDINGS])
+
+	// Layer 1
+	SHADER_PARAMETER_SRV_ARRAY(Texture2D, VirtualTexture1, [MAX_NUM_BINDINGS])
+	SHADER_PARAMETER_TEXTURE_ARRAY(Texture2D<uint4>, PageTable1, [MAX_NUM_BINDINGS])
+	SHADER_PARAMETER_TEXTURE_ARRAY(Texture2D<uint>, PageTableIndirection1, [MAX_NUM_BINDINGS])
+
+	// Layer 2
+	SHADER_PARAMETER_SRV_ARRAY(Texture2D, VirtualTexture2, [MAX_NUM_BINDINGS])
+	SHADER_PARAMETER_TEXTURE_ARRAY(Texture2D<uint4>, PageTable2, [MAX_NUM_BINDINGS])
+	SHADER_PARAMETER_TEXTURE_ARRAY(Texture2D<uint>, PageTableIndirection2, [MAX_NUM_BINDINGS])
+
+	// Layer info
+	SHADER_PARAMETER_ARRAY(FUintVector4, MaterialTypes, [MAX_NUM_BINDINGS])
+	SHADER_PARAMETER_ARRAY(FUintVector4, ValidLayerMasks, [MAX_NUM_BINDINGS])
+	SHADER_PARAMETER_ARRAY(FVector4f, LWCTiles, [MAX_NUM_BINDINGS])
+	SHADER_PARAMETER_ARRAY(FMatrix44f, WorldToUVTransforms, [MAX_NUM_BINDINGS])
+	SHADER_PARAMETER_ARRAY(FUintVector4, Adaptive, [MAX_NUM_BINDINGS * MAX_NUM_LAYERS])
+	SHADER_PARAMETER_ARRAY(FUintVector4, Uniforms, [MAX_NUM_BINDINGS * MAX_NUM_LAYERS])
+	SHADER_PARAMETER_ARRAY(FUintVector4, PageTableUniforms, [MAX_NUM_BINDINGS * 2])
+
+	// Shared
+	SHADER_PARAMETER(uint32, NumVirtualTextures)
+	SHADER_PARAMETER_SAMPLER(SamplerState, Sampler)
+END_SHADER_PARAMETER_STRUCT()
+
+void UPCGVirtualTextureDataInterface::GetShaderParameters(TCHAR const* UID, FShaderParametersMetadataBuilder& InOutBuilder, FShaderParametersMetadataAllocations& InOutAllocations) const
+{
+	InOutBuilder.AddNestedStruct<FPCGVirtualTextureDataInterfaceParameters>(UID);
+}
+
+TCHAR const* UPCGVirtualTextureDataInterface::TemplateFilePath = TEXT("/Plugin/PCG/Private/PCGVirtualTextureDataInterface.ush");
+
+TCHAR const* UPCGVirtualTextureDataInterface::GetShaderVirtualPath() const
+{
+	return TemplateFilePath;
+}
+
+void UPCGVirtualTextureDataInterface::GetShaderHash(FString& InOutKey) const
+{
+	GetShaderFileHash(TemplateFilePath, EShaderPlatform::SP_PCD3D_SM5).AppendString(InOutKey);
+	GetShaderFileHash(TEXT("/Engine/Private/VirtualTextureCommon.ush"), EShaderPlatform::SP_PCD3D_SM5).AppendString(InOutKey);
+	GetShaderFileHash(TEXT("/Plugin/PCG/Private/PCGVirtualTextureCommon.ush"), EShaderPlatform::SP_PCD3D_SM5).AppendString(InOutKey);
+}
+
+void UPCGVirtualTextureDataInterface::GetHLSL(FString& OutHLSL, FString const& InDataInterfaceName) const
+{
+	TMap<FString, FStringFormatArg> TemplateArgs =
+	{
+		{ TEXT("DataInterfaceName"), InDataInterfaceName },
+	};
+
+	FString TemplateFile;
+	if (ensure(LoadShaderSourceFile(TemplateFilePath, EShaderPlatform::SP_PCD3D_SM5, &TemplateFile, nullptr)))
+	{
+		OutHLSL += FString::Format(*TemplateFile, TemplateArgs);
+	}
+}
+
+void UPCGVirtualTextureDataInterface::GetDefines(FComputeKernelDefinitionSet& OutDefinitionSet) const
+{
+	OutDefinitionSet.Defines.Add(FComputeKernelDefinition(TEXT("MAX_NUM_BINDINGS"), FString::FromInt(MAX_NUM_BINDINGS)));
+	OutDefinitionSet.Defines.Add(FComputeKernelDefinition(TEXT("MAX_NUM_LAYERS"), FString::FromInt(MAX_NUM_LAYERS)));
+}
+
+void UPCGVirtualTextureDataInterface::GetPermutations(FComputeKernelPermutationVector& OutPermutationVector) const
+{
+	OutPermutationVector.AddPermutation(EnableMultipleVirtualTexturePermutationName, /*NumValues=*/2);
+}
+
+UComputeDataProvider* UPCGVirtualTextureDataInterface::CreateDataProvider() const
+{
+	return NewObject<UPCGVirtualTextureDataProvider>();
+}
+
+void UPCGVirtualTextureDataProvider::Initialize(const UComputeDataInterface* InDataInterface, UObject* InBinding, uint64 InInputMask, uint64 InOutputMask)
+{
+	TRACE_CPUPROFILER_EVENT_SCOPE(UPCGVirtualTextureDataProvider::Initialize);
+
+	Super::Initialize(InDataInterface, InBinding, InInputMask, InOutputMask);
+
+	const UPCGVirtualTextureDataInterface* DataInterface = CastChecked<UPCGVirtualTextureDataInterface>(InDataInterface);
+	UPCGDataBinding* Binding = CastChecked<UPCGDataBinding>(InBinding);
+	const UPCGComponent* SourceComponent = Cast<UPCGComponent>(Binding->GetExecutionSource());
+
+	// Take any input pin label alias to obtain the data from the input data collection.
+	check(!DataInterface->DownstreamInputPinLabelAliases.IsEmpty());
+	const TArray<FPCGTaggedData> InputTaggedData = Binding->InputDataCollection.GetInputsByPin(DataInterface->DownstreamInputPinLabelAliases[0]);
+
+	// Sampling of virtual textures is not supported outside of runtime generation at this time, because the RuntimeGenScheduler is responsible for priming the virtual textures.
+	if (!SourceComponent || !SourceComponent->IsManagedByRuntimeGenSystem())
+	{
+		UE_LOG(LogPCG, Error, TEXT("Virtual texture data interface on pin '%s' failed to initialize. Virtual texture sampling is only supported for runtime generation."),
+			*DataInterface->DownstreamInputPinLabelAliases[0].ToString());
+		return;
+	}
+
+	for (const FPCGTaggedData& TaggedData : InputTaggedData)
+	{
+		const UPCGVirtualTextureData* VirtualTextureData = Cast<UPCGVirtualTextureData>(TaggedData.Data);
+
+		if (!ensure(VirtualTextureData))
+		{
+			UE_LOG(LogPCG, Error, TEXT("Virtual texture data interface on pin '%s' receieved unsupported data type: '%s'"),
+				*DataInterface->DownstreamInputPinLabelAliases[0].ToString(),
+				TaggedData.Data ? *TaggedData.Data->GetName() : TEXT("NULL"));
+			continue;
+		}
+
+		if (VirtualTextures.Num() < MAX_NUM_BINDINGS)
+		{
+			VirtualTextures.Add(VirtualTextureData->GetRuntimeVirtualTexture());
+		}
+		else
+		{
+			UE_LOG(LogPCG, Warning, TEXT("Texture data interface on pin '%s' received too many textures to bind. Only the first %d textures will be bound."),
+				*DataInterface->DownstreamInputPinLabelAliases[0].ToString(),
+				MAX_NUM_BINDINGS);
+			break;
+		}
+	}
+}
+
+FComputeDataProviderRenderProxy* UPCGVirtualTextureDataProvider::GetRenderProxy()
+{
+	return new FPCGVirtualTextureDataProviderProxy(VirtualTextures);
+}
+
+void UPCGVirtualTextureDataProvider::Reset()
+{
+	Super::Reset();
+
+	VirtualTextures.Empty();
+}
+
+FPCGVirtualTextureDataProviderProxy::FPCGVirtualTextureDataProviderProxy(const TArray<TObjectPtr<const URuntimeVirtualTexture>>& InVirtualTextures)
+{
+	NumVirtualTextures = InVirtualTextures.Num();
+	for (int32 VirtualTextureIndex = 0; VirtualTextureIndex < NumVirtualTextures; ++VirtualTextureIndex)
+	{
+		ERuntimeVirtualTextureMaterialType MaterialType = ERuntimeVirtualTextureMaterialType::Count;
+
+		if (const URuntimeVirtualTexture* VirtualTexture = InVirtualTextures[VirtualTextureIndex])
+		{
+			MaterialType = VirtualTexture->GetMaterialType();
+			const int32 NumVirtualTextureLayers = FMath::Min(URuntimeVirtualTexture::GetLayerCount(MaterialType), MAX_NUM_LAYERS);
+			const int32 PackedLayerBaseIndex = VirtualTextureIndex * MAX_NUM_LAYERS;
+
+			for (int32 LayerIndex = 0; LayerIndex < NumVirtualTextureLayers; ++LayerIndex)
+			{
+				const int32 PackedLayerIndex = PackedLayerBaseIndex + LayerIndex;
+
+				// Note, we always use the first layer to get the world to uv transform because every material type uses at least the first layer.
+				const bool bIncludeWorldToUVTransform = LayerIndex == 0;
+
+				Layers[PackedLayerIndex].Initialize(VirtualTexture, LayerIndex, VirtualTexture->IsLayerSRGB(LayerIndex));
+				PageTables[PackedLayerIndex].Initialize(VirtualTexture, LayerIndex, bIncludeWorldToUVTransform, bIncludeWorldToUVTransform);
+			}
+		}
+
+		MaterialTypes[VirtualTextureIndex] = MaterialType;
+	}
+}
+
+bool FPCGVirtualTextureDataProviderProxy::IsValid(FValidationData const& InValidationData) const
+{
+	return InValidationData.ParameterStructSize == sizeof(FParameters);
+}
+
+struct FPCGVirtualTextureDataInterfacePermutationIds
+{
+	uint32 EnableMultipleVirtualTextures = 0;
+ 
+	FPCGVirtualTextureDataInterfacePermutationIds(const FComputeKernelPermutationVector& PermutationVector)
+	{
+		static FString Name(EnableMultipleVirtualTexturePermutationName);
+		static uint32 Hash = GetTypeHash(Name);
+		EnableMultipleVirtualTextures = PermutationVector.GetPermutationBits(Name, Hash, /*Value=*/1);
+	}
+};
+
+void FPCGVirtualTextureDataProviderProxy::GatherPermutations(FPermutationData& InOutPermutationData) const
+{
+	FPCGVirtualTextureDataInterfacePermutationIds PermutationIds(InOutPermutationData.PermutationVector);
+	for (int32 InvocationIndex = 0; InvocationIndex < InOutPermutationData.NumInvocations; ++InvocationIndex)
+	{
+		InOutPermutationData.PermutationIds[InvocationIndex] |= (NumVirtualTextures > 1 ? PermutationIds.EnableMultipleVirtualTextures : 0);
+	}
+}
+
+void FPCGVirtualTextureDataProviderProxy::GatherDispatchData(FDispatchData const& InDispatchData)
+{
+	const TStridedView<FParameters> ParameterArray = MakeStridedParameterView<FParameters>(InDispatchData);
+	for (int32 InvocationIndex = 0; InvocationIndex < ParameterArray.Num(); ++InvocationIndex)
+	{
+		FParameters& Parameters = ParameterArray[InvocationIndex];
+		Parameters.NumVirtualTextures = NumVirtualTextures;
+		Parameters.Sampler = TStaticSamplerState<SF_Bilinear, AM_Clamp, AM_Clamp, AM_Clamp>::GetRHI();
+
+		for (int32 BindingIndex = 0; BindingIndex < MAX_NUM_BINDINGS; ++BindingIndex)
+		{
+			const int32 PackedLayerBaseIndex = BindingIndex * MAX_NUM_LAYERS;
+			const PCGVirtualTextureCommon::FVirtualTextureLayer& Layer0 = Layers[PackedLayerBaseIndex + 0];
+			const PCGVirtualTextureCommon::FVirtualTextureLayer& Layer1 = Layers[PackedLayerBaseIndex + 1];
+			const PCGVirtualTextureCommon::FVirtualTextureLayer& Layer2 = Layers[PackedLayerBaseIndex + 2];
+			const PCGVirtualTextureCommon::FVirtualTexturePageTable& PageTable0 = PageTables[PackedLayerBaseIndex + 0];
+			const PCGVirtualTextureCommon::FVirtualTexturePageTable& PageTable1 = PageTables[PackedLayerBaseIndex + 1];
+			const PCGVirtualTextureCommon::FVirtualTexturePageTable& PageTable2 = PageTables[PackedLayerBaseIndex + 2];
+
+			Parameters.MaterialTypes[BindingIndex] = FUintVector4(static_cast<uint32>(MaterialTypes[BindingIndex]), 0, 0, 0);
+			Parameters.ValidLayerMasks[BindingIndex] = FUintVector4::ZeroValue;
+
+			if (Layer0.IsValid() && PageTable0.IsValid())
+			{
+				Parameters.ValidLayerMasks[BindingIndex].X |= 1;
+				Parameters.VirtualTexture0[BindingIndex] = Layer0.TextureSRV;
+				Parameters.PageTable0[BindingIndex] = PageTable0.PageTableRef;
+				Parameters.PageTableIndirection0[BindingIndex] = PageTable0.PageTableIndirectionRef;
+				Parameters.Adaptive[PackedLayerBaseIndex + 0] = static_cast<FUintVector4>(PageTable0.bIsAdaptive);
+				Parameters.Uniforms[PackedLayerBaseIndex + 0] = Layer0.TextureUniforms;
+
+				// Layer 0 should always be valid, so we can use it to set the shared parameters.
+				FLargeWorldRenderPosition TextureOrigin(PageTable0.WorldToUVParameters[0]);
+				Parameters.LWCTiles[BindingIndex] = TextureOrigin.GetTile();
+
+				Parameters.WorldToUVTransforms[BindingIndex] = FMatrix44f(
+					TextureOrigin.GetOffset(),
+					FVector3f((FVector4f)PageTable0.WorldToUVParameters[1]),
+					FVector3f((FVector4f)PageTable0.WorldToUVParameters[2]),
+					FVector3f((FVector4f)PageTable0.WorldToUVParameters[3]));
+
+				Parameters.PageTableUniforms[BindingIndex * 2 + 0] = PageTable0.PageTableUniforms[0];
+				Parameters.PageTableUniforms[BindingIndex * 2 + 1] = PageTable0.PageTableUniforms[1];
+			}
+			else
+			{
+				Parameters.VirtualTexture0[BindingIndex] = GBlackTextureWithSRV->ShaderResourceViewRHI;
+				Parameters.PageTable0[BindingIndex] = GBlackUintTexture->TextureRHI;
+				Parameters.PageTableIndirection0[BindingIndex] = GBlackUintTexture->TextureRHI;
+				Parameters.Adaptive[PackedLayerBaseIndex + 0] = FUintVector4::ZeroValue;
+				Parameters.Uniforms[PackedLayerBaseIndex + 0] = FUintVector4::ZeroValue;
+
+				// If Layer 0 was invalid, default initialize the shared parameters as well.
+				Parameters.LWCTiles[BindingIndex] = FVector3f::ZeroVector;
+				Parameters.WorldToUVTransforms[BindingIndex] = FMatrix44f::Identity;
+				Parameters.PageTableUniforms[BindingIndex * 2 + 0] = FUintVector4::ZeroValue;
+				Parameters.PageTableUniforms[BindingIndex * 2 + 1] = FUintVector4::ZeroValue;
+			}
+
+			if (Layer1.IsValid() && PageTable1.IsValid())
+			{
+				Parameters.ValidLayerMasks[BindingIndex].X |= 2;
+				Parameters.VirtualTexture1[BindingIndex] = Layer1.TextureSRV;
+				Parameters.PageTable1[BindingIndex] = PageTable1.PageTableRef;
+				Parameters.PageTableIndirection1[BindingIndex] = PageTable1.PageTableIndirectionRef;
+				Parameters.Adaptive[PackedLayerBaseIndex + 1] = static_cast<FUintVector4>(PageTable1.bIsAdaptive);
+				Parameters.Uniforms[PackedLayerBaseIndex + 1] = Layer1.TextureUniforms;
+			}
+			else
+			{
+				Parameters.VirtualTexture1[BindingIndex] = GBlackTextureWithSRV->ShaderResourceViewRHI;
+				Parameters.PageTable1[BindingIndex] = GBlackUintTexture->TextureRHI;
+				Parameters.PageTableIndirection1[BindingIndex] = GBlackUintTexture->TextureRHI;
+				Parameters.Adaptive[PackedLayerBaseIndex + 1] = FUintVector4::ZeroValue;
+				Parameters.Uniforms[PackedLayerBaseIndex + 1] = FUintVector4::ZeroValue;
+			}
+
+			if (Layer2.IsValid() && PageTable2.IsValid())
+			{
+				Parameters.ValidLayerMasks[BindingIndex].X |= 4;
+				Parameters.VirtualTexture2[BindingIndex] = Layer2.TextureSRV;
+				Parameters.PageTable2[BindingIndex] = PageTable2.PageTableRef;
+				Parameters.PageTableIndirection2[BindingIndex] = PageTable2.PageTableIndirectionRef;
+				Parameters.Adaptive[PackedLayerBaseIndex + 2] = static_cast<FUintVector4>(PageTable2.bIsAdaptive);
+				Parameters.Uniforms[PackedLayerBaseIndex + 2] = Layer2.TextureUniforms;
+			}
+			else
+			{
+				Parameters.VirtualTexture2[BindingIndex] = GBlackTextureWithSRV->ShaderResourceViewRHI;
+				Parameters.PageTable2[BindingIndex] = GBlackUintTexture->TextureRHI;
+				Parameters.PageTableIndirection2[BindingIndex] = GBlackUintTexture->TextureRHI;
+				Parameters.Adaptive[PackedLayerBaseIndex + 2] = FUintVector4::ZeroValue;
+				Parameters.Uniforms[PackedLayerBaseIndex + 2] = FUintVector4::ZeroValue;
+			}
+		}
+	}
+}
+
+#undef LOCTEXT_NAMESPACE
diff --git a/Engine/Plugins/PCG/Source/PCG/Private/Compute/DataInterfaces/PCGVirtualTextureDataInterface.h b/Engine/Plugins/PCG/Source/PCG/Private/Compute/DataInterfaces/PCGVirtualTextureDataInterface.h
new file mode 100644
--- /dev/null
+++ b/Engine/Plugins/PCG/Source/PCG/Private/Compute/DataInterfaces/PCGVirtualTextureDataInterface.h
@@ -0,0 +1,88 @@
+// Copyright Epic Games, Inc. All Rights Reserved.
+
+#pragma once
+
+#include "Compute/Data/PCGVirtualTextureCommon.h"
+#include "Compute/DataInterfaces/PCGComputeDataInterface.h"
+
+#include "PCGVirtualTextureDataInterface.generated.h"
+
+class FPCGVirtualTextureDataInterfaceParameters;
+class URuntimeVirtualTexture;
+
+namespace PCGVirtualTextureDataInterfaceConstants
+{
+	/**
+	 * The maximum number of virtual textures which can be bound to a PCGVirtualTextureDataInterface.
+	 * We can't use full bindless on all platforms, so fallback to emulating.
+	 * If you change this number, make sure update PCGVirtualTextureDataInterface.ush as well.
+	 */
+	constexpr uint32 MAX_NUM_BINDINGS = 4;
+
+	/**
+	 * The maximum number of layers in a virtual texture.
+	 * If you change this number, make sure update PCGVirtualTextureDataInterface.ush as well.
+	 */
+	constexpr const int32 MAX_NUM_LAYERS = 3;
+}
+
+/** Data Interface allowing sampling of a virtual texture. */
+UCLASS(ClassGroup = (Procedural))
+class UPCGVirtualTextureDataInterface : public UPCGComputeDataInterface
+{
+	GENERATED_BODY()
+
+public:
+	//~ Begin UComputeDataInterface Interface
+	TCHAR const* GetClassName() const override { return TEXT("PCGVirtualTexture"); }
+	void GetSupportedInputs(TArray<FShaderFunctionDefinition>& OutFunctions) const override;
+	void GetShaderParameters(TCHAR const* UID, FShaderParametersMetadataBuilder& InOutBuilder, FShaderParametersMetadataAllocations& InOutAllocations) const override;
+	TCHAR const* GetShaderVirtualPath() const override;
+	void GetShaderHash(FString& InOutKey) const override;
+	void GetHLSL(FString& OutHLSL, FString const& InDataInterfaceName) const override;
+	void GetDefines(FComputeKernelDefinitionSet& OutDefinitionSet) const override;
+	void GetPermutations(FComputeKernelPermutationVector& OutPermutationVector) const override;
+	UComputeDataProvider* CreateDataProvider() const override;
+	//~ End UComputeDataInterface Interface
+
+private:
+	static TCHAR const* TemplateFilePath;
+};
+
+/** Compute Framework Data Provider for reading a virtual texture. */
+UCLASS()
+class UPCGVirtualTextureDataProvider : public UComputeDataProvider
+{
+	GENERATED_BODY()
+
+public:
+	//~ Begin UComputeDataProvider Interface
+	virtual void Initialize(const UComputeDataInterface* InDataInterface, UObject* InBinding, uint64 InInputMask, uint64 InOutputMask) override;
+	FComputeDataProviderRenderProxy* GetRenderProxy() override;
+	virtual void Reset() override;
+	//~ End UComputeDataProvider Interface
+
+private:
+	UPROPERTY()
+	TArray<TObjectPtr<const URuntimeVirtualTexture>> VirtualTextures;
+};
+
+class FPCGVirtualTextureDataProviderProxy : public FComputeDataProviderRenderProxy
+{
+public:
+	FPCGVirtualTextureDataProviderProxy(const TArray<TObjectPtr<const URuntimeVirtualTexture>>& InVirtualTextures);
+
+	//~ Begin FComputeDataProviderRenderProxy Interface
+	bool IsValid(FValidationData const& InValidationData) const override;
+	void GatherPermutations(FPermutationData& InOutPermutationData) const override;
+	void GatherDispatchData(FDispatchData const& InDispatchData) override;
+	//~ End FComputeDataProviderRenderProxy Interface
+
+protected:
+	using FParameters = FPCGVirtualTextureDataInterfaceParameters;
+
+	ERuntimeVirtualTextureMaterialType MaterialTypes[PCGVirtualTextureDataInterfaceConstants::MAX_NUM_BINDINGS];
+	PCGVirtualTextureCommon::FVirtualTextureLayer Layers[PCGVirtualTextureDataInterfaceConstants::MAX_NUM_LAYERS * PCGVirtualTextureDataInterfaceConstants::MAX_NUM_BINDINGS];
+	PCGVirtualTextureCommon::FVirtualTexturePageTable PageTables[PCGVirtualTextureDataInterfaceConstants::MAX_NUM_LAYERS * PCGVirtualTextureDataInterfaceConstants::MAX_NUM_BINDINGS];
+	int32 NumVirtualTextures = 0;
+};
diff --git a/Engine/Plugins/PCG/Source/PCG/Private/Compute/Elements/PCGCustomHLSL.cpp b/Engine/Plugins/PCG/Source/PCG/Private/Compute/Elements/PCGCustomHLSL.cpp
--- a/Engine/Plugins/PCG/Source/PCG/Private/Compute/Elements/PCGCustomHLSL.cpp
+++ b/Engine/Plugins/PCG/Source/PCG/Private/Compute/Elements/PCGCustomHLSL.cpp
@@ -408,6 +408,7 @@ void UPCGCustomHLSLSettings::UpdateInputDeclarations()
 	TArray<FString> SplineDataPins;
 	TArray<FString> LandscapeDataPins;
 	TArray<FString> TextureDataPins;
+	TArray<FString> VirtualTextureDataPins;
 	TArray<FString> RawBufferDataPins;
 	TArray<FString> StaticMeshDataPins;
 
@@ -438,6 +439,11 @@ void UPCGCustomHLSLSettings::UpdateInputDeclarations()
 			TextureDataPins.Add(Pin.Label.ToString());
 		}
 
+		if (!!(Pin.AllowedTypes & EPCGDataType::VirtualTexture))
+		{
+			VirtualTextureDataPins.Add(Pin.Label.ToString());
+		}
+
 		if (!!(Pin.AllowedTypes & EPCGDataType::StaticMeshResource))
 		{
 			StaticMeshDataPins.Add(Pin.Label.ToString());
@@ -540,6 +546,36 @@ void UPCGCustomHLSLSettings::UpdateInputDeclarations()
 		InputDeclarations += TEXT("\n");
 	}
 
+	if (!VirtualTextureDataPins.IsEmpty())
+	{
+		InputDeclarations += TEXT("/*** INPUT VIRTUAL TEXTURE DATA FUNCTIONS ***/\n\n");
+
+		const bool bMultiPin = VirtualTextureDataPins.Num() > 1;
+
+		if (bMultiPin)
+		{
+			InputDeclarations += TEXT("// Valid pins: ") + FString::Join(VirtualTextureDataPins, TEXT(", ")) + TEXT("\n\n");
+		}
+
+		InputDeclarations += FString::Format(TEXT(
+			"// Samples a virtual texture and gets all values that are available. Otherwise returns default values."
+			"void {0}_SampleVirtualTexture(\n"
+			"	uint InDataIndex,\n"
+			"	float3 InWorldPos,\n"
+			"	out bool bOutInsideVolume,\n"
+			"	out float3 OutBaseColor,\n"
+			"	out float OutSpecular,\n"
+			"	out float OutRoughness,\n"
+			"	out float OutWorldHeight,\n"
+			"	out float3 OutNormal,\n"
+			"	out float OutDisplacement,\n"
+			"	out float OutMask,\n"
+			"	out float4 OutMask4);\n"),
+			{ bMultiPin ? PCGHLSLElement::PinDeclTemplateStr : VirtualTextureDataPins[0] });
+
+		InputDeclarations += TEXT("\n");
+	}
+
 	if (!RawBufferDataPins.IsEmpty())
 	{
 		InputDeclarations += TEXT("/*** INPUT BYTE ADDRESS BUFFER DATA FUNCTIONS ***/\n\n");
diff --git a/Engine/Plugins/PCG/Source/PCG/Private/Data/PCGGetDataFunctionRegistry.cpp b/Engine/Plugins/PCG/Source/PCG/Private/Data/PCGGetDataFunctionRegistry.cpp
--- a/Engine/Plugins/PCG/Source/PCG/Private/Data/PCGGetDataFunctionRegistry.cpp
+++ b/Engine/Plugins/PCG/Source/PCG/Private/Data/PCGGetDataFunctionRegistry.cpp
@@ -7,6 +7,7 @@
 #include "PCGContext.h"
 #include "PCGData.h"
 #include "PCGGraph.h"
+#include "Compute/PCGComputeCommon.h"
 #include "Data/PCGBasePointData.h"
 #include "Data/PCGCollisionShapeData.h"
 #include "Data/PCGIntersectionData.h"
@@ -17,6 +18,7 @@
 #include "Data/PCGSpatialData.h"
 #include "Data/PCGSplineData.h"
 #include "Data/PCGUnionData.h"
+#include "Data/PCGVirtualTextureData.h"
 #include "Data/PCGVolumeData.h"
 #include "Grid/PCGPartitionActor.h"
 #include "Helpers/PCGActorHelpers.h"
@@ -29,11 +31,13 @@
 #include "Components/ActorComponent.h"
 #include "Components/BillboardComponent.h"
 #include "Components/InstancedStaticMeshComponent.h"
+#include "Components/RuntimeVirtualTextureComponent.h"
 #include "Components/ShapeComponent.h"
 #include "Components/SplineComponent.h"
 #include "Components/SplineMeshComponent.h"
 #include "GameFramework/Actor.h"
 #include "GameFramework/Volume.h"
+#include "VT/RuntimeVirtualTexture.h"
 
 FPCGGetDataFunctionRegistry::FFunctionHandle FPCGGetDataFunctionRegistry::RegisterDataFromActorFunction(const FPCGGetDataFunctionRegistry::FDataFromActorFunction& InFunction)
 {
@@ -242,6 +246,7 @@ int FPCGGetDataFunctionRegistry::DefaultDataFromComponent(FPCGContext* InContext
 	// Splines component
 	// Shape component
 	// Primitive component, but ignore billboard components
+	// RuntimeVirtualTexture component
 	FPCGTaggedData TaggedData;
 
 	if (!TaggedData.Data && !!(InDataFilter & EPCGDataType::Spline))
@@ -280,6 +285,25 @@ int FPCGGetDataFunctionRegistry::DefaultDataFromComponent(FPCGContext* InContext
 		}
 	}
 
+	if (!TaggedData.Data && !!(InDataFilter & EPCGDataType::VirtualTexture))
+	{
+		if (URuntimeVirtualTextureComponent* VirtualTextureComponent = Cast<URuntimeVirtualTextureComponent>(InComponent))
+		{
+			UPCGVirtualTextureData* VirtualTextureData = FPCGContext::NewObject_AnyThread<UPCGVirtualTextureData>(InContext);
+			VirtualTextureData->Initialize(VirtualTextureComponent);
+			TaggedData.Data = VirtualTextureData;
+
+			FString VirtualTextureName;
+
+			if (const URuntimeVirtualTexture* VirtualTexture = VirtualTextureData->GetRuntimeVirtualTexture())
+			{
+				VirtualTexture->GetName(VirtualTextureName);
+			}
+
+			TaggedData.Tags.Add(PCGComputeHelpers::GetPrefixedDataLabel(VirtualTextureName));
+		}
+	}
+
 	if (TaggedData.Data)
 	{
 		Algo::Transform(InComponent->ComponentTags, TaggedData.Tags, NameTagsToStringTags);
diff --git a/Engine/Plugins/PCG/Source/PCG/Private/Data/PCGVirtualTextureData.cpp b/Engine/Plugins/PCG/Source/PCG/Private/Data/PCGVirtualTextureData.cpp
new file mode 100644
--- /dev/null
+++ b/Engine/Plugins/PCG/Source/PCG/Private/Data/PCGVirtualTextureData.cpp
@@ -0,0 +1,45 @@
+// Copyright Epic Games, Inc. All Rights Reserved.
+
+#include "Data/PCGVirtualTextureData.h"
+
+#include "PCGContext.h"
+#include "Data/PCGPointArrayData.h"
+#include "Data/PCGPointData.h"
+
+#include "Components/RuntimeVirtualTextureComponent.h"
+#include "VT/RuntimeVirtualTexture.h"
+
+#include UE_INLINE_GENERATED_CPP_BY_NAME(PCGVirtualTextureData)
+
+void UPCGVirtualTextureData::Initialize(const URuntimeVirtualTextureComponent* InVirtualTextureComponent)
+{
+	check(InVirtualTextureComponent);
+	RuntimeVirtualTexture = InVirtualTextureComponent->GetVirtualTexture();
+}
+
+UPCGSpatialData* UPCGVirtualTextureData::CopyInternal(FPCGContext* Context) const
+{
+	UPCGVirtualTextureData* NewTextureData = FPCGContext::NewObject_AnyThread<UPCGVirtualTextureData>(Context);
+
+	CopyBaseSurfaceData(NewTextureData);
+
+	NewTextureData->RuntimeVirtualTexture = RuntimeVirtualTexture;
+
+	return NewTextureData;
+}
+
+const UPCGPointData* UPCGVirtualTextureData::CreatePointData(FPCGContext* Context) const
+{
+	UPCGPointData* Data = FPCGContext::NewObject_AnyThread<UPCGPointData>(Context);
+	Data->InitializeFromData(this);
+
+	return Data;
+}
+
+const UPCGPointArrayData* UPCGVirtualTextureData::CreatePointArrayData(FPCGContext* Context, const FBox& InBounds) const
+{
+	UPCGPointArrayData* Data = FPCGContext::NewObject_AnyThread<UPCGPointArrayData>(Context);
+	Data->InitializeFromData(this);
+
+	return Data;
+}
diff --git a/Engine/Plugins/PCG/Source/PCG/Private/Elements/PCGTypedGetter.cpp b/Engine/Plugins/PCG/Source/PCG/Private/Elements/PCGTypedGetter.cpp
--- a/Engine/Plugins/PCG/Source/PCG/Private/Elements/PCGTypedGetter.cpp
+++ b/Engine/Plugins/PCG/Source/PCG/Private/Elements/PCGTypedGetter.cpp
@@ -7,6 +7,8 @@
 #include "Helpers/PCGHelpers.h"
 
 #include "Landscape.h"
+#include "Components/RuntimeVirtualTextureComponent.h"
+#include "VT/RuntimeVirtualTextureVolume.h"
 
 #include UE_INLINE_GENERATED_CPP_BY_NAME(PCGTypedGetter)
 
@@ -222,4 +224,48 @@ FText UPCGGetPCGComponentSettings::GetNodeTooltipText() const
 }
 #endif
 
+UPCGGetVirtualTextureSettings::UPCGGetVirtualTextureSettings()
+{
+	Mode = EPCGGetDataFromActorMode::ParseActorComponents;
+
+	ActorSelector.ActorFilter = EPCGActorFilter::AllWorldActors;
+	ActorSelector.bShowActorFilter = false;
+	ActorSelector.ActorSelection = EPCGActorSelection::ByClass;
+	ActorSelector.ActorSelectionClass = ARuntimeVirtualTextureVolume::StaticClass();
+	ActorSelector.bShowActorSelectionClass = false;
+	ActorSelector.bIncludeChildren = false;
+	ActorSelector.bSelectMultiple = true;
+	ActorSelector.bShowSelectMultiple = false;
+
+	ComponentSelector.ComponentSelection = EPCGComponentSelection::ByClass;
+	ComponentSelector.ComponentSelectionClass = URuntimeVirtualTextureComponent::StaticClass();
+	ComponentSelector.bShowComponentSelectionClass = false;
+}
+
+#if WITH_EDITOR
+FText UPCGGetVirtualTextureSettings::GetNodeTooltipText() const
+{
+	return LOCTEXT("GetVirtualTextureTooltip", "Builds a collection of virtual texture data from the selected actors.");
+}
+#endif
+
+FString UPCGGetVirtualTextureSettings::GetAdditionalTitleInformation() const
+{
+	// Do not use the version from data from actor otherwise we'll show the selected actor class, which serves no purpose
+	return UPCGSettings::GetAdditionalTitleInformation();
+}
+
+TArray<FPCGPinProperties> UPCGGetVirtualTextureSettings::OutputPinProperties() const
+{
+	TArray<FPCGPinProperties> PinProperties;
+	PinProperties.Emplace(PCGPinConstants::DefaultOutputLabel, EPCGDataType::VirtualTexture);
+
+	return PinProperties;
+}
+
+TSubclassOf<AActor> UPCGGetVirtualTextureSettings::GetDefaultActorSelectorClass() const
+{
+	return ARuntimeVirtualTextureVolume::StaticClass();
+}
+
 #undef LOCTEXT_NAMESPACE
diff --git a/Engine/Plugins/PCG/Source/PCG/Private/Graph/PCGGraphCompilerGPU.cpp b/Engine/Plugins/PCG/Source/PCG/Private/Graph/PCGGraphCompilerGPU.cpp
--- a/Engine/Plugins/PCG/Source/PCG/Private/Graph/PCGGraphCompilerGPU.cpp
+++ b/Engine/Plugins/PCG/Source/PCG/Private/Graph/PCGGraphCompilerGPU.cpp
@@ -18,6 +18,7 @@
 #include "Compute/DataInterfaces/PCGLandscapeDataInterface.h"
 #include "Compute/DataInterfaces/PCGStaticMeshDataInterface.h"
 #include "Compute/DataInterfaces/PCGTextureDataInterface.h"
+#include "Compute/DataInterfaces/PCGVirtualTextureDataInterface.h"
 #include "Compute/Elements/PCGComputeGraphElement.h"
 #include "Elements/PCGStaticMeshSpawnerKernel.h"
 #include "Graph/PCGGPUGraphCompilationContext.h"
@@ -903,6 +904,11 @@ void FPCGGraphCompilerGPU::CreateDataInterfaces(
 			DataInterface = DataInterfacePCGData;
 			break;
 		}
+		case EPCGDataType::VirtualTexture:
+		{
+			DataInterface = InOutContext.NewObject_AnyThread<UPCGVirtualTextureDataInterface>(InOutComputeGraph);
+			break;
+		}
 		case EPCGDataType::BaseTexture:
 		case EPCGDataType::Texture:
 		case EPCGDataType::RenderTarget:
diff --git a/Engine/Plugins/PCG/Source/PCG/Public/Data/PCGVirtualTextureData.h b/Engine/Plugins/PCG/Source/PCG/Public/Data/PCGVirtualTextureData.h
new file mode 100644
--- /dev/null
+++ b/Engine/Plugins/PCG/Source/PCG/Public/Data/PCGVirtualTextureData.h
@@ -0,0 +1,46 @@
+// Copyright Epic Games, Inc. All Rights Reserved.
+
+#pragma once
+
+#include "Data/PCGSurfaceData.h"
+
+#include "PCGVirtualTextureData.generated.h"
+
+class URuntimeVirtualTexture;
+class URuntimeVirtualTextureComponent;
+
+UCLASS()
+class UPCGVirtualTextureData : public UPCGSurfaceData
+{
+	GENERATED_BODY()
+
+public:
+	void Initialize(const URuntimeVirtualTextureComponent* InVirtualTextureComponent);
+
+	//~Begin UPCGData interface
+	virtual EPCGDataType GetDataType() const override { return EPCGDataType::VirtualTexture; }
+	//~End UPCGData interface
+
+	//~Begin UPCGSpatialData interface
+	virtual FBox GetBounds() const override { return FBox(EForceInit::ForceInit); }
+	virtual bool SamplePoint(const FTransform& InTransform, const FBox& InBounds, FPCGPoint& OutPoint, UPCGMetadata* OutMetadata) const override { return false; }
+protected:
+	virtual UPCGSpatialData* CopyInternal(FPCGContext* Context) const override;
+	//~End UPCGSpatialData interface
+
+public:
+	//~Begin UPCGSpatialDataWithPointCache interface
+	virtual const UPCGPointData* CreatePointData(FPCGContext* Context) const override;
+	virtual const UPCGPointArrayData* CreatePointArrayData(FPCGContext* Context, const FBox& InBounds) const override;
+	//~End UPCGSpatialDataWithPointCache interface
+
+	TObjectPtr<const URuntimeVirtualTexture> GetRuntimeVirtualTexture() const { return RuntimeVirtualTexture; }
+
+protected:
+	UPROPERTY()
+	TObjectPtr<URuntimeVirtualTexture> RuntimeVirtualTexture = nullptr;
+
+	// @todo_pcg: Support SVT as well?
+	//UPROPERTY()
+	//TObjectPtr<class UVirtualTextureBuilder> StreamingVirtualTexture = nullptr;
+};
diff --git a/Engine/Plugins/PCG/Source/PCG/Public/Elements/PCGTypedGetter.h b/Engine/Plugins/PCG/Source/PCG/Public/Elements/PCGTypedGetter.h
--- a/Engine/Plugins/PCG/Source/PCG/Public/Elements/PCGTypedGetter.h
+++ b/Engine/Plugins/PCG/Source/PCG/Public/Elements/PCGTypedGetter.h
@@ -193,3 +193,37 @@ protected:
 #endif
 	//~End UPCGDataFromActorSettings
 };
+
+/** Builds a collection of virtual texture data from the selected actors. */
+UCLASS(MinimalAPI, BlueprintType, ClassGroup = (Procedural))
+class UPCGGetVirtualTextureSettings : public UPCGDataFromActorSettings
+{
+	GENERATED_BODY()
+
+public:
+	UPCGGetVirtualTextureSettings();
+
+	//~Begin UPCGSettings interface
+#if WITH_EDITOR
+	virtual FName GetDefaultNodeName() const override { return FName(TEXT("GetVirtualTextureData")); }
+	virtual FText GetDefaultNodeTitle() const override { return NSLOCTEXT("PCGGetVirtualTextureElement", "NodeTitle", "Get Virtual Texture Data"); }
+	virtual FText GetNodeTooltipText() const override;
+#endif
+
+	virtual FString GetAdditionalTitleInformation() const override;
+
+protected:
+	virtual TArray<FPCGPinProperties> OutputPinProperties() const override;
+	//~End UPCGSettings
+
+	//~Begin UPCGDataFromActorSettings interface
+public:
+	virtual EPCGDataType GetDataFilter() const override { return EPCGDataType::VirtualTexture; }
+	virtual TSubclassOf<AActor> GetDefaultActorSelectorClass() const override;
+
+protected:
+#if WITH_EDITOR
+	virtual bool DisplayModeSettings() const override { return false; }
+#endif
+	//~End UPCGDataFromActorSettings
+};
diff --git a/Engine/Plugins/PCG/Source/PCG/Public/PCGCommon.h b/Engine/Plugins/PCG/Source/PCG/Public/PCGCommon.h
--- a/Engine/Plugins/PCG/Source/PCG/Public/PCGCommon.h
+++ b/Engine/Plugins/PCG/Source/PCG/Public/PCGCommon.h
@@ -116,8 +116,10 @@ enum class EPCGDataType : uint32
 	Landscape = 1 << 4,
 	Texture = 1 << 5,
 	RenderTarget = 1 << 6,
+	// VirtualTexture is not a subtype of BaseTexture because they share no common functionality, particularly with respect to sampling.
+	VirtualTexture = 1 << 12,
 	BaseTexture = Texture | RenderTarget UMETA(Tooltip = "Common base type for both textures and render targets."),
-	Surface = Landscape | BaseTexture,
+	Surface = Landscape | BaseTexture | VirtualTexture,
 
 	Volume = 1 << 7,
 	Primitive = 1 << 8,
@@ -163,6 +165,7 @@ enum class EPCGExclusiveDataType : uint8
 	Landscape,
 	Texture,
 	RenderTarget,
+	VirtualTexture,
 	BaseTexture UMETA(Hidden),
 	Surface,
 	Volume,
diff --git a/Engine/Plugins/PCG/Source/PCGEditor/Private/Nodes/PCGEditorGraphNodeBase.cpp b/Engine/Plugins/PCG/Source/PCGEditor/Private/Nodes/PCGEditorGraphNodeBase.cpp
--- a/Engine/Plugins/PCG/Source/PCGEditor/Private/Nodes/PCGEditorGraphNodeBase.cpp
+++ b/Engine/Plugins/PCG/Source/PCGEditor/Private/Nodes/PCGEditorGraphNodeBase.cpp
@@ -1132,6 +1132,10 @@ FEdGraphPinType UPCGEditorGraphNodeBase::GetPinType(const UPCGPin* InPin)
 		{
 			EdPinType.PinSubCategory = FPCGEditorCommon::LandscapeDataType;
 		}
+		else if (CheckType(EPCGDataType::VirtualTexture))
+		{
+			EdPinType.PinSubCategory = FPCGEditorCommon::VirtualTextureDataType;
+		}
 		else if (CheckType(EPCGDataType::BaseTexture))
 		{
 			EdPinType.PinSubCategory = FPCGEditorCommon::BaseTextureDataType;
diff --git a/Engine/Plugins/PCG/Source/PCGEditor/Private/PCGEditorCommon.h b/Engine/Plugins/PCG/Source/PCGEditor/Private/PCGEditorCommon.h
--- a/Engine/Plugins/PCG/Source/PCGEditor/Private/PCGEditorCommon.h
+++ b/Engine/Plugins/PCG/Source/PCGEditor/Private/PCGEditorCommon.h
@@ -46,6 +46,7 @@ namespace FPCGEditorCommon
 	const FName BaseTextureDataType = FName(TEXT("Base Texture Data"));
 	const FName TextureDataType = FName(TEXT("Texture Data"));
 	const FName RenderTargetDataType = FName(TEXT("Render Target Data"));
+	const FName VirtualTextureDataType = FName(TEXT("Virtual Texture Data"));
 	const FName VolumeDataType = FName(TEXT("Volume Data"));
 	const FName PrimitiveDataType = FName(TEXT("Primitive Data"));
 	const FName DynamicMeshDataType = FName(TEXT("Dynamic Mesh Data"));
diff --git a/Engine/Plugins/PCG/Source/PCGEditor/Private/PCGEditorSettings.cpp b/Engine/Plugins/PCG/Source/PCGEditor/Private/PCGEditorSettings.cpp
--- a/Engine/Plugins/PCG/Source/PCGEditor/Private/PCGEditorSettings.cpp
+++ b/Engine/Plugins/PCG/Source/PCGEditor/Private/PCGEditorSettings.cpp
@@ -45,6 +45,7 @@ UPCGEditorSettings::UPCGEditorSettings(const FObjectInitializer& ObjectInitializ
 	BaseTextureDataPinColor = FLinearColor(0.81f, 0.28f, 0.22f);
 	TextureDataPinColor = FLinearColor(0.79f, 0.08f, 0.01f);
 	RenderTargetDataPinColor = FLinearColor(0.8f, 0.18f, 0.12f);
+	VirtualTextureDataPinColor = FLinearColor(0.8f, 0.48f, 0.22f);
 	VolumeDataPinColor = FLinearColor(0.79f, 0.06f, 0.5f);
 	PrimitiveDataPinColor = FLinearColor(0.22f, 0.05f, 1.0f);
 	DynamicMeshPinColor = FLinearColor(0.42f, 0.25f, 1.0f);
@@ -142,6 +143,10 @@ FLinearColor UPCGEditorSettings::GetPinColor(const FEdGraphPinType& PinType) con
 		{
 			return RenderTargetDataPinColor;
 		}
+		else if (PinType.PinSubCategory == FPCGEditorCommon::VirtualTextureDataType)
+		{
+			return VirtualTextureDataPinColor;
+		}
 		else if (PinType.PinSubCategory == FPCGEditorCommon::SurfaceDataType)
 		{
 			return SurfaceDataPinColor;
diff --git a/Engine/Plugins/PCG/Source/PCGEditor/Public/PCGEditorSettings.h b/Engine/Plugins/PCG/Source/PCGEditor/Public/PCGEditorSettings.h
--- a/Engine/Plugins/PCG/Source/PCGEditor/Public/PCGEditorSettings.h
+++ b/Engine/Plugins/PCG/Source/PCGEditor/Public/PCGEditorSettings.h
@@ -172,6 +172,10 @@ public:
 	UPROPERTY(EditAnywhere, config, Category = Node, meta = (HideAlphaChannel))
 	FLinearColor RenderTargetDataPinColor;
 
+	/** Color used for data pins of type Virtual Texture */
+	UPROPERTY(EditAnywhere, config, Category = Node, meta = (HideAlphaChannel))
+	FLinearColor VirtualTextureDataPinColor;
+
 	/** Color used for data pins of type Volume */
 	UPROPERTY(EditAnywhere, config, Category = Node, meta = (HideAlphaChannel))
 	FLinearColor VolumeDataPinColor;
