[PCG] Builder: Fix missing outputs when using "One Component at a Time" + Hierarchical Generation
- Add a scope where the GraphExecutor is not allowed to clear out its results

diff --git a/Engine/Plugins/PCG/Source/PCG/Private/Graph/PCGGraphExecutor.cpp b/Engine/Plugins/PCG/Source/PCG/Private/Graph/PCGGraphExecutor.cpp
--- a/Engine/Plugins/PCG/Source/PCG/Private/Graph/PCGGraphExecutor.cpp
+++ b/Engine/Plugins/PCG/Source/PCG/Private/Graph/PCGGraphExecutor.cpp
@@ -2318,6 +2318,13 @@ void FPCGGraphExecutor::StoreResults(FPCGTaskId InTaskId, const FPCGDataCollecti
 
 void FPCGGraphExecutor::ClearResults()
 {
+#if WITH_EDITOR
+	if (bDisableClearResults)
+	{
+		return;
+	}
+#endif
+
 	{
 		// Only reset if we have no more scheduled tasks, to avoid breaking dependencies
 		if (ScheduledTasks.IsEmpty())
diff --git a/Engine/Plugins/PCG/Source/PCG/Private/Graph/PCGGraphExecutor.h b/Engine/Plugins/PCG/Source/PCG/Private/Graph/PCGGraphExecutor.h
--- a/Engine/Plugins/PCG/Source/PCG/Private/Graph/PCGGraphExecutor.h
+++ b/Engine/Plugins/PCG/Source/PCG/Private/Graph/PCGGraphExecutor.h
@@ -270,6 +270,8 @@ private:
 	void ReleaseGenerationNotification();
 	void OnNotificationCancel();
 	static FTextFormat GetNotificationTextFormat();
+
+	void SetDisableClearResults(bool bInDisableClearResults) { bDisableClearResults = bInDisableClearResults; }
 #endif
 
 	/** Graph compiler that turns a graph into tasks */
@@ -370,6 +372,8 @@ private:
 
 	/** Set true when a non-runtime-gen component is scheduled and reset to false when all tasks are done executing. */
 	std::atomic<bool> bAnyNonRuntimeGenComponentScheduled = false;
+
+	bool bDisableClearResults = false;
 #endif
 
 	TObjectPtr<UWorld> World = nullptr;
diff --git a/Engine/Plugins/PCG/Source/PCG/Private/PCGSubsystem.cpp b/Engine/Plugins/PCG/Source/PCG/Private/PCGSubsystem.cpp
--- a/Engine/Plugins/PCG/Source/PCG/Private/PCGSubsystem.cpp
+++ b/Engine/Plugins/PCG/Source/PCG/Private/PCGSubsystem.cpp
@@ -2061,6 +2061,14 @@ void UPCGSubsystem::OnPCGGraphCleaned(UPCGComponent* InComponent)
 	ActorAndComponentMapping.OnPCGGraphGeneratedOrCleaned(InComponent);
 }
 
+void UPCGSubsystem::SetDisableClearResults(bool bInDisableClearResults)
+{
+	if (GraphExecutor)
+	{
+		GraphExecutor->SetDisableClearResults(bInDisableClearResults);
+	}
+}
+
 #endif // WITH_EDITOR
 
 UPCGData* UPCGSubsystem::GetPCGData(FPCGTaskId InGraphExecutionTaskId)
diff --git a/Engine/Plugins/PCG/Source/PCG/Public/PCGSubsystem.h b/Engine/Plugins/PCG/Source/PCG/Public/PCGSubsystem.h
--- a/Engine/Plugins/PCG/Source/PCG/Public/PCGSubsystem.h
+++ b/Engine/Plugins/PCG/Source/PCG/Public/PCGSubsystem.h
@@ -467,5 +467,7 @@ private:
 	{
 		return DisablePartitionActorCreationForWorld.Contains(InWorld);
 	}
+
+	void SetDisableClearResults(bool bInDisableClearResults);
 #endif
 };
diff --git a/Engine/Plugins/PCG/Source/PCGEditor/Private/WorldPartitionBuilder/PCGWorldPartitionBuilder.cpp b/Engine/Plugins/PCG/Source/PCGEditor/Private/WorldPartitionBuilder/PCGWorldPartitionBuilder.cpp
--- a/Engine/Plugins/PCG/Source/PCGEditor/Private/WorldPartitionBuilder/PCGWorldPartitionBuilder.cpp
+++ b/Engine/Plugins/PCG/Source/PCGEditor/Private/WorldPartitionBuilder/PCGWorldPartitionBuilder.cpp
@@ -92,6 +92,9 @@ struct FPCGWorldPartitionBuilder
 	/** Returns when GraphExecutor is done running */
 	static void WaitOnGraphExecutor(UWorld* InWorld);
 
+	/** Allows to temporarly disable the clearing of results */
+	static void SetDisableClearResults(UPCGSubsystem* InSubsystem, bool bInDisableClearResults);
+
 private:
 	static void CollectSubGridComponentsToGenerateInternal(UPCGSubsystem* InPCGSubsystem, int32 InGridSizeIndex, const PCGHiGenGrid::FSizeArray& InGridSizes, UPCGComponent* InOriginalComponent, APCGPartitionActor* InPartitionActor, TArray<UPCGComponent*>& OutAdditionalComponentsToGenerate, TMap<UPCGComponent*, UPCGComponent*>& OutDependencies);
 
@@ -225,6 +228,30 @@ private:
 	std::atomic<bool> bErrorOccurred = false;
 };
 
+/** Disable clearing of results while in this scope */
+struct FPCGDisableClearResults
+{
+	FPCGDisableClearResults(UWorld* InWorld)
+	{
+		Subsystem = UWorld::GetSubsystem<UPCGSubsystem>(InWorld);
+
+		if(Subsystem)
+		{
+			FPCGWorldPartitionBuilder::SetDisableClearResults(Subsystem, true);
+		}
+	}
+
+	virtual ~FPCGDisableClearResults()
+	{
+		if (Subsystem)
+		{
+			FPCGWorldPartitionBuilder::SetDisableClearResults(Subsystem, false);
+		}
+	}
+
+	UPCGSubsystem* Subsystem = nullptr;
+};
+
 bool FPCGWorldPartitionBuilderArgs::InitializeFrom(const FPCGWorldPartitionCommandlineArgs& CommandlineArgs, FPCGWorldPartitionBuilderArgs& OutArgs)
 {
 	FPCGWorldPartitionBuilderArgs BuilderArgs;
@@ -1054,33 +1081,43 @@ bool FPCGWorldPartitionBuilder::GenerateComponents(
 
 	TMap<TObjectKey<UPCGComponent>, FPCGTaskId> GeneratedComponents;
 
-	for (TWeakObjectPtr<UPCGComponent> ComponentWeakPtr : Components)
 	{
-		UPCGComponent* Component = ComponentWeakPtr.Get();
-		if (!Component)
-		{
-			UE_LOG(LogPCGWorldPartitionBuilder, Warning, TEXT("Lost a PCG component weak pointer, component will not be generated."));
-			continue;
-		}
-		
-		FPCGTaskId TaskIdDependency = InvalidPCGTaskId;
-		if (const UPCGComponent* const* TaskDependency = InTaskDependencies.Find(Component))
+		// While we process a batch of inter-dependent components (InTaskDependencies) we want to prevent the Subsystem's GraphExecutor from clearing its outputs.
+		// This is specially important when InArgs.bOneComponentAtATime is true which waits for a component generation to end before scheduling another one
+		// which will tick the world and the Subsystem's GraphExecutor will become idle and clear it's outputs.
+		// If that following component generation is dependant on the previous one through its InTaskDependencies then it will never be allowed to run as the outputs will no longer contain
+		// the previous component's results.
+		// Important: the list of component created sorted by dependency so that we can rely on the order of scheduling here
+		const FPCGDisableClearResults DisableClearResults(InWorld);
+
+		for (TWeakObjectPtr<UPCGComponent> ComponentWeakPtr : Components)
 		{
-			if (FPCGTaskId* TaskIdDependencyPtr = GeneratedComponents.Find(*TaskDependency))
+			UPCGComponent* Component = ComponentWeakPtr.Get();
+			if (!Component)
 			{
-				TaskIdDependency = *TaskIdDependencyPtr;
+				UE_LOG(LogPCGWorldPartitionBuilder, Warning, TEXT("Lost a PCG component weak pointer, component will not be generated."));
+				continue;
 			}
-		}
 
-		FPCGTaskId GenerateTaskId = FPCGWorldPartitionBuilder::GenerateComponent(Component, InWorld, InArgs, TaskIdDependency);
-		if (GenerateTaskId != InvalidPCGTaskId)
-		{
-			GeneratedComponents.Add(Component, GenerateTaskId);
-		}
+			FPCGTaskId TaskIdDependency = InvalidPCGTaskId;
+			if (const UPCGComponent* const* TaskDependency = InTaskDependencies.Find(Component))
+			{
+				if (FPCGTaskId* TaskIdDependencyPtr = GeneratedComponents.Find(*TaskDependency))
+				{
+					TaskIdDependency = *TaskIdDependencyPtr;
+				}
+			}
 
-		if (InArgs.bOneComponentAtATime)
-		{
-			WaitForComponentGeneration(Component);
+			FPCGTaskId GenerateTaskId = FPCGWorldPartitionBuilder::GenerateComponent(Component, InWorld, InArgs, TaskIdDependency);
+			if (GenerateTaskId != InvalidPCGTaskId)
+			{
+				GeneratedComponents.Add(Component, GenerateTaskId);
+			}
+
+			if (InArgs.bOneComponentAtATime)
+			{
+				WaitForComponentGeneration(Component);
+			}
 		}
 	}
 
@@ -1307,6 +1344,14 @@ void FPCGWorldPartitionBuilder::ResetPackageDirtyFlags()
 	}
 }
 
+void FPCGWorldPartitionBuilder::SetDisableClearResults(UPCGSubsystem* InSubsystem, bool bInDisableClearResults)
+{
+	if (InSubsystem)
+	{
+		InSubsystem->SetDisableClearResults(bInDisableClearResults);
+	}
+}
+
 TArray<UPackage*> FPCGWorldPartitionBuilder::GetDirtyPackages(const TSet<FString>& AllowedAssetPackages)
 {
 	TArray<UPackage*> DirtyPackages;
