[PCG] Enable UPCGPointArrayData by default (pcg.EnablePointArrayData)
-Rename experimental cvars

diff --git a/Engine/Plugins/PCG/Source/PCG/Private/Data/PCGPointArrayData.cpp b/Engine/Plugins/PCG/Source/PCG/Private/Data/PCGPointArrayData.cpp
--- a/Engine/Plugins/PCG/Source/PCG/Private/Data/PCGPointArrayData.cpp
+++ b/Engine/Plugins/PCG/Source/PCG/Private/Data/PCGPointArrayData.cpp
@@ -7,7 +7,7 @@
 #include "Data/PCGPointData.h"
 
 TAutoConsoleVariable<bool> CVarPCGEnablePointArrayDataParenting(
-	TEXT("pcg.Experimental.EnablePointArrayDataParenting"),
+	TEXT("pcg.EnablePointArrayDataParenting"),
 	true,
 	TEXT("Whether to enable inheritance of data on PointArrayData (memory savings)"));
 
diff --git a/Engine/Plugins/PCG/Source/PCG/Private/Data/PCGSpatialData.cpp b/Engine/Plugins/PCG/Source/PCG/Private/Data/PCGSpatialData.cpp
--- a/Engine/Plugins/PCG/Source/PCG/Private/Data/PCGSpatialData.cpp
+++ b/Engine/Plugins/PCG/Source/PCG/Private/Data/PCGSpatialData.cpp
@@ -18,7 +18,7 @@
 namespace PCGSpatialData
 {
 	TAutoConsoleVariable<bool> CVarEnablePrepareForSpatialQuery(
-		TEXT("pcg.Experimental.SpatialData.EnablePrepareForSpatialQuery"),
+		TEXT("pcg.SpatialData.EnablePrepareForSpatialQuery"),
 		true,
 		TEXT("Enable UPCGSpatialData subclass PrepareForSpatialQuery task scheduling"));
 }
diff --git a/Engine/Plugins/PCG/Source/PCG/Private/Graph/PCGGraphExecutor.cpp b/Engine/Plugins/PCG/Source/PCG/Private/Graph/PCGGraphExecutor.cpp
--- a/Engine/Plugins/PCG/Source/PCG/Private/Graph/PCGGraphExecutor.cpp
+++ b/Engine/Plugins/PCG/Source/PCG/Private/Graph/PCGGraphExecutor.cpp
@@ -97,7 +97,7 @@ namespace PCGGraphExecutor
 		TEXT("Whether proxies for GPU data are cached in per pin output data and passed through grid links. If false data is read back to CPU."));
 
 	TAutoConsoleVariable<bool> CVarEnableDynamicDependencies(
-		TEXT("pcg.Experimental.EnableDynamicDependencies"),
+		TEXT("pcg.EnableDynamicDependencies"),
 		true,
 		TEXT("Support re-scheduling of paused tasks through dynamic dependencies"));
 
diff --git a/Engine/Plugins/PCG/Source/PCG/Private/PCGContext.cpp b/Engine/Plugins/PCG/Source/PCG/Private/PCGContext.cpp
--- a/Engine/Plugins/PCG/Source/PCG/Private/PCGContext.cpp
+++ b/Engine/Plugins/PCG/Source/PCG/Private/PCGContext.cpp
@@ -29,8 +29,8 @@
 #define LOCTEXT_NAMESPACE "PCGContext"
 
 TAutoConsoleVariable<bool> CVarPCGEnablePointArrayData(
-	TEXT("pcg.Experimental.EnablePointArrayData"),
-	false,
+	TEXT("pcg.EnablePointArrayData"),
+	true,
 	TEXT("Whether to use the new UPCGPointArrayData when calling FPCGContext::CreatePointData_AnyThread"));
 
 namespace PCGContextHelpers
