[PCG] Remove dependency on ProceduralMeshComponent plugin

#rnx

diff --git a/Engine/Plugins/PCG/PCG.uplugin b/Engine/Plugins/PCG/PCG.uplugin
--- a/Engine/Plugins/PCG/PCG.uplugin
+++ b/Engine/Plugins/PCG/PCG.uplugin
@@ -39,10 +39,6 @@
 		{
 			"Name": "ComputeFramework",
 			"Enabled": true
-		},
-		{
-			"Name": "ProceduralMeshComponent",
-			"Enabled": true
 		}
 	]
 }
diff --git a/Engine/Plugins/PCG/Source/PCGEditor/PCGEditor.Build.cs b/Engine/Plugins/PCG/Source/PCGEditor/PCGEditor.Build.cs
--- a/Engine/Plugins/PCG/Source/PCGEditor/PCGEditor.Build.cs
+++ b/Engine/Plugins/PCG/Source/PCGEditor/PCGEditor.Build.cs
@@ -53,7 +53,6 @@ namespace UnrealBuildTool.Rules
 					"Kismet",
 					"KismetWidgets",
 					"PCG",
-					"ProceduralMeshComponent",
 					"PropertyEditor",
 					"RenderCore",
 					"Slate",
diff --git a/Engine/Plugins/PCG/Source/PCGEditor/Private/DataVisualizations/PCGBaseTextureDataVisualization.cpp b/Engine/Plugins/PCG/Source/PCGEditor/Private/DataVisualizations/PCGBaseTextureDataVisualization.cpp
--- a/Engine/Plugins/PCG/Source/PCGEditor/Private/DataVisualizations/PCGBaseTextureDataVisualization.cpp
+++ b/Engine/Plugins/PCG/Source/PCGEditor/Private/DataVisualizations/PCGBaseTextureDataVisualization.cpp
@@ -5,50 +5,33 @@
 #include "Data/PCGTextureData.h"
 
 #include "AdvancedPreviewScene.h"
-#include "ProceduralMeshComponent.h"
+#include "Components/StaticMeshComponent.h"
 #include "Engine/AssetManager.h"
 #include "Engine/Texture.h"
 #include "Materials/MaterialInstanceDynamic.h"
 
 namespace PCGBaseTextureVisualizationConstants
 {
-	static constexpr float DebugQuadSize = 500.0f;
-
-	static const TArray<FVector> DebugMeshVertices =
-	{
-		FVector(-1, 0, -1) * DebugQuadSize,
-		FVector( 1, 0, -1) * DebugQuadSize,
-		FVector(-1, 0,  1) * DebugQuadSize,
-		FVector( 1, 0,  1) * DebugQuadSize
-	};
-
-	static const TArray<int32> DebugMeshTriangles = { 0, 1, 2, 1, 3, 2 };
-
-	static const TArray<FVector2D> DebugMeshUVs =
-	{
-		FVector2D(0, 0),
-		FVector2D(1, 0),
-		FVector2D(0, 1),
-		FVector2D(1, 1)
-	};
-
+	static const FSoftObjectPath PlaneMeshPath(TEXT("/Engine/BasicShapes/Plane.Plane"));
 	static const FSoftObjectPath DebugMaterialPath(TEXT("Material'/PCG/DebugObjects/PCG_DebugMaterialTexture.PCG_DebugMaterialTexture'"));
 }
 
 TArray<TSharedPtr<FStreamableHandle>> FPCGBaseTextureDataVisualization::LoadRequiredResources(const UPCGData* Data) const
 {
-	return { UAssetManager::GetStreamableManager().RequestAsyncLoad(PCGBaseTextureVisualizationConstants::DebugMaterialPath, nullptr) };
+	TArray<TSharedPtr<FStreamableHandle>> LoadHandles;
+	LoadHandles.Add(UAssetManager::GetStreamableManager().RequestAsyncLoad(PCGBaseTextureVisualizationConstants::PlaneMeshPath, nullptr));
+	LoadHandles.Add(UAssetManager::GetStreamableManager().RequestAsyncLoad(PCGBaseTextureVisualizationConstants::DebugMaterialPath, nullptr));
+
+	return LoadHandles;
 }
 
 FPCGSetupSceneFunc FPCGBaseTextureDataVisualization::GetViewportSetupFunc(const UPCGData* Data) const
 {
-	using namespace PCGBaseTextureVisualizationConstants;
-
 	return [this, WeakData=TWeakObjectPtr<const UPCGBaseTextureData>(Cast<UPCGBaseTextureData>(Data))](FAdvancedPreviewScene* InScene, const TArray<UObject*>& InResources, FVector& OutViewLocation, float& OutOrbitDistance, TArray<TObjectPtr<UObject>>& OutManagedResources)
 	{
 		check(InScene);
 
-		if (InResources.IsEmpty())
+		if (InResources.Num() != 2)
 		{
 			return;
 		}
@@ -65,13 +48,13 @@ FPCGSetupSceneFunc FPCGBaseTextureDataVisualization::GetViewportSetupFunc(const
 			UE_LOG(LogPCG, Verbose, TEXT("Texture data uses an unsupported resource type for data viewport visualization."));
 		}
 
-		UMaterialInterface* DebugMaterial = Cast<UMaterialInterface>(InResources[0]);
+		UMaterialInterface* DebugMaterial = Cast<UMaterialInterface>(InResources[1]);
 		UMaterialInstanceDynamic* MaterialInstance = UMaterialInstanceDynamic::Create(DebugMaterial, GetTransientPackage());
 		MaterialInstance->SetTextureParameterValue(FName(TEXT("DebugTexture")), WeakData->GetTexture());
 		OutManagedResources.Add(MaterialInstance);
 
-		UProceduralMeshComponent* MeshComponent = NewObject<UProceduralMeshComponent>(GetTransientPackage(), NAME_None, RF_Transient);
-		MeshComponent->CreateMeshSection(/*SectionIndex=*/0, DebugMeshVertices, DebugMeshTriangles, /*Normals=*/{}, DebugMeshUVs, /*UV1=*/{}, /*UV2=*/{}, /*UV3=*/{}, /*VertexColors=*/{}, /*Tangents=*/{}, /*bCreateCollision=*/false);
+		TObjectPtr<UStaticMeshComponent> MeshComponent = NewObject<UStaticMeshComponent>(GetTransientPackage(), NAME_None, RF_Transient);
+		MeshComponent->SetStaticMesh(Cast<UStaticMesh>(InResources[0]));
 		MeshComponent->OverrideMaterials.Add(MaterialInstance);
 		OutManagedResources.Add(MeshComponent);
 
@@ -80,10 +63,17 @@ FPCGSetupSceneFunc FPCGBaseTextureDataVisualization::GetViewportSetupFunc(const
 			MeshComponent->SetMobility(EComponentMobility::Static);
 		}
 
-		InScene->AddComponent(MeshComponent, FTransform::Identity);
-		InScene->SetFloorOffset(DebugQuadSize * 2);
+		static const FTransform MeshTransform(FRotator(0, 0, 90), FVector::ZeroVector, FVector::OneVector);
+		InScene->AddComponent(MeshComponent, MeshTransform);
 
-		OutViewLocation = FVector::ZeroVector;
-		OutOrbitDistance = DebugQuadSize * 2;
+		// Bounds will be updated already by SetStaticMesh() call.
+		const FBoxSphereBounds Bounds = MeshComponent->Bounds;
+		InScene->SetFloorOffset(-Bounds.Origin.Z + Bounds.BoxExtent.Z);
+		OutViewLocation = Bounds.Origin;
+
+		if (Bounds.SphereRadius > 0.0f)
+		{
+			OutOrbitDistance = Bounds.SphereRadius * 2.0f;
+		}
 	};
 }
