[PCG] Change the default names of ArriveTangent and LeaveTangent for Create Spline, in new nodes

* Make sure they match with the spline sampler

#rnx

diff --git a/Engine/Plugins/PCG/Source/PCG/Private/Elements/PCGCreateSpline.cpp b/Engine/Plugins/PCG/Source/PCG/Private/Elements/PCGCreateSpline.cpp
--- a/Engine/Plugins/PCG/Source/PCG/Private/Elements/PCGCreateSpline.cpp
+++ b/Engine/Plugins/PCG/Source/PCG/Private/Elements/PCGCreateSpline.cpp
@@ -9,6 +9,7 @@
 #include "PCGPin.h"
 #include "Data/PCGBasePointData.h"
 #include "Data/PCGSplineData.h"
+#include "Elements/PCGSplineSampler.h"
 #include "Helpers/PCGActorHelpers.h"
 #include "Helpers/PCGHelpers.h"
 #include "Metadata/PCGMetadata.h"
@@ -21,6 +22,16 @@
 
 #define LOCTEXT_NAMESPACE "PCGCreateSpline"
 
+UPCGCreateSplineSettings::UPCGCreateSplineSettings()
+{
+	// Change the default for Arrive and Leave tangent
+	if (PCGHelpers::IsNewObjectAndNotDefault(this))
+	{
+		ArriveTangentAttribute = PCGSplineSamplerConstants::ArriveTangentName;
+		LeaveTangentAttribute = PCGSplineSamplerConstants::LeaveTangentName;
+	}
+}
+
 #if WITH_EDITOR
 FText UPCGCreateSplineSettings::GetNodeTooltipText() const
 {
diff --git a/Engine/Plugins/PCG/Source/PCG/Public/Elements/PCGCreateSpline.h b/Engine/Plugins/PCG/Source/PCG/Public/Elements/PCGCreateSpline.h
--- a/Engine/Plugins/PCG/Source/PCG/Public/Elements/PCGCreateSpline.h
+++ b/Engine/Plugins/PCG/Source/PCG/Public/Elements/PCGCreateSpline.h
@@ -24,6 +24,8 @@ class UPCGCreateSplineSettings : public UPCGSettings
 	GENERATED_BODY()
 
 public:
+	UPCGCreateSplineSettings();
+
 	//~Begin UPCGSettings interface
 #if WITH_EDITOR
 	virtual FName GetDefaultNodeName() const override { return FName(TEXT("CreateSpline")); }
diff --git a/Engine/Plugins/PCG/Source/PCG/Public/Elements/PCGSplineSampler.h b/Engine/Plugins/PCG/Source/PCG/Public/Elements/PCGSplineSampler.h
--- a/Engine/Plugins/PCG/Source/PCG/Public/Elements/PCGSplineSampler.h
+++ b/Engine/Plugins/PCG/Source/PCG/Public/Elements/PCGSplineSampler.h
@@ -18,6 +18,8 @@ namespace PCGSplineSamplerConstants
 {
 	const FName SplineLabel = TEXT("Spline");
 	const FName BoundingShapeLabel = TEXT("Bounding Shape");
+	const FName ArriveTangentName = TEXT("ArriveTangent");
+	const FName LeaveTangentName = TEXT("LeaveTangent");
 }
 
 UENUM()
@@ -170,11 +172,11 @@ struct FPCGSplineSamplerParams
 
 	/** Attribute that will contain the arrive tangent vector. For control points, this will be the actual arrive tangent. For non-control points, this will only be the normalized tangent at this point. */
 	UPROPERTY(BlueprintReadWrite, EditAnywhere, Category = "Settings|Attributes", meta = (EditCondition = "bComputeTangents&&Dimension!=EPCGSplineSamplingDimension::OnInterior", EditConditionHides))
-	FName ArriveTangentAttribute = TEXT("ArriveTangent");
+	FName ArriveTangentAttribute = PCGSplineSamplerConstants::ArriveTangentName;
 
 	/** Attribute that will contain the leave tangent vector. For control points, this will be the actual leave tangent. For non-control points, this will only be the normalized tangent at this point. */
 	UPROPERTY(BlueprintReadWrite, EditAnywhere, Category = "Settings|Attributes", meta = (EditCondition = "bComputeTangents&&Dimension!=EPCGSplineSamplingDimension::OnInterior", EditConditionHides))
-	FName LeaveTangentAttribute = TEXT("LeaveTangent");
+	FName LeaveTangentAttribute = PCGSplineSamplerConstants::LeaveTangentName;
 
 	/** Compute an alpha value along the spline and write it to an attribute. */
 	UPROPERTY(BlueprintReadWrite, EditAnywhere, Category = "Settings|Attributes", meta = (EditCondition = "Dimension!=EPCGSplineSamplingDimension::OnInterior", EditConditionHides))
