Change PCG RVT preload bounds to be calculated directly from sphere with the generation preload radius.
This avoids bounds expanding due to component z bounds and box->bounding sphere calculations.
This in turn minimizes amount of RVT pages to preload.
#rnx

diff --git a/Engine/Plugins/PCG/Source/PCG/Private/RuntimeGen/PCGRuntimeGenScheduler.cpp b/Engine/Plugins/PCG/Source/PCG/Private/RuntimeGen/PCGRuntimeGenScheduler.cpp
--- a/Engine/Plugins/PCG/Source/PCG/Private/RuntimeGen/PCGRuntimeGenScheduler.cpp
+++ b/Engine/Plugins/PCG/Source/PCG/Private/RuntimeGen/PCGRuntimeGenScheduler.cpp
@@ -999,20 +999,15 @@ void FPCGRuntimeGenScheduler::TickRequestVirtualTexturePriming(const TSet<IPCGGe
 						continue;
 					}
 
-					const FVector GenSourcePosition = GenSourcePositionOptional.GetValue();
-					FBox GenSourcePrimingBounds = FBox(FSphere(GenSourcePosition, PrimingRadius));
-					FVector ModifiedGenSourcePosition = GenSourcePosition;
-
+					FVector GenSourcePosition = GenSourcePositionOptional.GetValue();
 					if (OriginalComponent->Use2DGrid())
 					{
-						ModifiedGenSourcePosition.Z = OriginalComponentBounds.Min.Z;
-						GenSourcePrimingBounds.Min.Z = OriginalComponentBounds.Min.Z;
-						GenSourcePrimingBounds.Max.Z = OriginalComponentBounds.Max.Z;
+						GenSourcePosition.Z = OriginalComponentBounds.GetCenter().Z;
 					}
+					const FSphere PrimingBounds = FSphere(GenSourcePosition, PrimingRadius);
 
-					const FBox PrimingBounds = OriginalComponentBounds.Overlap(GenSourcePrimingBounds);
-
-					if (PrimingBounds.GetVolume() > 0)
+					const FBox OverlapBox = OriginalComponentBounds.Overlap(FBox(PrimingBounds));
+					if (OverlapBox.GetVolume() > 0)
 					{
 						TRACE_CPUPROFILER_EVENT_SCOPE(RequestVirtualTexturePreload);
 
@@ -1022,9 +1017,10 @@ void FPCGRuntimeGenScheduler::TickRequestVirtualTexturePriming(const TSet<IPCGGe
 						{
 							check(World);
 
-							DrawDebugBox(World,
-								PrimingBounds.GetCenter(),
-								PrimingBounds.GetExtent(),
+							DrawDebugSphere(World,
+								PrimingBounds.Center,
+								PrimingBounds.W,
+								8,
 								FColor::Red,
 								/*bPersistentLines=*/false,
 								/*LifeTime=*/0.02f);
