[PCG - GPU] Groundwork for supporting data attribute on GPU

* Replaced attribute names by attribute identifiers
* Add selector with identifier parsing for custom attributes in custom HLSL nodes.
* For now, discard any attribute on the data domain (which mean they are lost on output)
* Fixes crash with data that had data domain attributes.

#rnx

diff --git a/Engine/Plugins/PCG/Source/PCG/Internal/Compute/PCGDataForGPU.h b/Engine/Plugins/PCG/Source/PCG/Internal/Compute/PCGDataForGPU.h
--- a/Engine/Plugins/PCG/Source/PCG/Internal/Compute/PCGDataForGPU.h
+++ b/Engine/Plugins/PCG/Source/PCG/Internal/Compute/PCGDataForGPU.h
@@ -53,20 +53,35 @@ struct FPCGKernelAttributeKey
 
 	FPCGKernelAttributeKey() = default;
 
-	explicit FPCGKernelAttributeKey(FName InName, EPCGKernelAttributeType InType)
-		: Name(InName)
+	explicit FPCGKernelAttributeKey(FPCGAttributeIdentifier InIdentifier, EPCGKernelAttributeType InType)
+		: Identifier(InIdentifier)
 		, Type(InType)
 	{}
 
-	UPROPERTY(EditAnywhere, Category = "Settings")
-	FName Name = NAME_None;
+	explicit FPCGKernelAttributeKey(const FPCGAttributePropertySelector& InSelector, EPCGKernelAttributeType InType);
 
-	UPROPERTY(EditAnywhere, Category = "Settings")
+	/** To be called everytime the Selector changes, to update the Identifier. Return true if it has changed. */
+	bool UpdateIdentifierFromSelector();
+	
+	void SetSelector(const FPCGAttributePropertySelector& InSelector);
+
+	/** Cached identifier. Need to be updated if the Selector ever change. */
+	UPROPERTY()
+	FPCGAttributeIdentifier Identifier;
+
+	UPROPERTY(EditAnywhere, Category = "Settings", meta = (DisplayAfter = "Attribute"))
 	EPCGKernelAttributeType Type = EPCGKernelAttributeType::Float;
 
+	bool IsValid() const;
+
 	bool operator==(const FPCGKernelAttributeKey& Other) const;
 
 	friend uint32 GetTypeHash(const FPCGKernelAttributeKey& In);
+
+protected:
+	/** Selector to specify which attribute to create and on which domain. At the moment, only support `@Data` domains or no domain (default one). */
+	UPROPERTY(EditAnywhere, Category = "Settings", meta = (PCG_DiscardPropertySelection, PCG_DiscardExtraSelection))
+	FPCGAttributePropertyOutputNoSourceSelector Name;
 };
 
 /** Table of attributes used in compute graph with helpers to get unique attribute ID used to read/write attribute in data collection buffers. */
@@ -79,10 +94,10 @@ struct FPCGKernelAttributeTable
 
 	int32 GetAttributeId(const FPCGKernelAttributeKey& InAttribute) const;
 
-	int32 GetAttributeId(FName InName, EPCGKernelAttributeType InType) const;
+	int32 GetAttributeId(FPCGAttributeIdentifier InIdentifier, EPCGKernelAttributeType InType) const;
 
 	/** Adds an attribute of given name and type. Returns added attribute ID. */
-	int32 AddAttribute(FName InName, EPCGKernelAttributeType InType);
+	int32 AddAttribute(FPCGAttributeIdentifier InIdentifier, EPCGKernelAttributeType InType);
 
 #if PCG_KERNEL_LOGGING_ENABLED
 	void DebugLog() const;
@@ -101,17 +116,15 @@ struct FPCGKernelAttributeDesc
 public:
 	FPCGKernelAttributeDesc() = default;
 
-	explicit FPCGKernelAttributeDesc(int32 InIndex, EPCGKernelAttributeType InType, FName InName)
-		: Index(InIndex)
-		, Type(InType)
-		, Name(InName)
+	explicit FPCGKernelAttributeDesc(int32 InIndex, EPCGKernelAttributeType InType, FPCGAttributeIdentifier InIdentifier)
+		: AttributeId(InIndex)
+		, AttributeKey(InIdentifier, InType)
 	{
 	}
 
-	explicit FPCGKernelAttributeDesc(int32 InIndex, EPCGKernelAttributeType InType, FName InName, const TArray<int32>* InUniqueStringKeys)
-		: Index(InIndex)
-		, Type(InType)
-		, Name(InName)
+	explicit FPCGKernelAttributeDesc(int32 InIndex, EPCGKernelAttributeType InType, FPCGAttributeIdentifier InIdentifier, const TArray<int32>* InUniqueStringKeys)
+		: AttributeId(InIndex)
+		, AttributeKey(InIdentifier, InType)
 	{
 		if (InUniqueStringKeys)
 		{
@@ -128,13 +141,9 @@ public:
 	bool operator==(const FPCGKernelAttributeDesc& Other) const;
 
 public:
-	// Attribute Id
-	// TODO: Rename to Id to disambiguate absolute attribute Id from index into metadata attribute table.
-	int32 Index = INDEX_NONE;
+	int32 AttributeId = INDEX_NONE;
 
-	EPCGKernelAttributeType Type = EPCGKernelAttributeType::Float;
-
-	FName Name = NAME_None;
+	FPCGKernelAttributeKey AttributeKey{};
 
 private:
 	/* All possible string keys arriving on this attribute (string keys are indices into the string table in the data binding). */
@@ -151,14 +160,16 @@ public:
 
 	uint64 ComputePackedSize() const;
 
-	bool HasMetadataAttributes() const;
+	bool HasElementsMetadataDomainAttributes() const;
 
-	bool ContainsAttribute(FName InAttributeName) const;
+	bool ContainsAttribute(FPCGAttributeIdentifier InAttributeIdentifier) const;
 
-	bool ContainsAttribute(FName InAttributeName, EPCGKernelAttributeType InAttributeType) const;
+	bool ContainsAttribute(FPCGAttributeIdentifier InAttributeIdentifier, EPCGKernelAttributeType InAttributeType) const;
 
 	void AddAttribute(FPCGKernelAttributeKey InAttribute, const UPCGDataBinding* InBinding, const TArray<int32>* InOptionalUniqueStringKeys = nullptr);
 
+	int32 GetElementCountForAttribute(const FPCGKernelAttributeDesc& AttributeDesc) const;
+
 public:
 	EPCGDataType Type = EPCGDataType::Point;
 	TArray<FPCGKernelAttributeDesc> AttributeDescs;
@@ -196,11 +207,11 @@ struct FPCGDataCollectionDesc
 	/** Aggregate another data description. */
 	void Combine(const FPCGDataCollectionDesc& Other);
 
-	/** Get description of first attribute with matching name in input data. Returns true if such an attribute found and also signals whether multiple matching attributes
+	/** Get description of first attribute with matching identifier in input data. Returns true if such an attribute found and also signals whether multiple matching attributes
 	* with conflicting names are present. */
-	bool GetAttributeDesc(FName InAttributeName, FPCGKernelAttributeDesc& OutAttributeDesc, bool &bOutConflictingTypesFound) const;
+	bool GetAttributeDesc(FPCGAttributeIdentifier InAttributeIdentifier, FPCGKernelAttributeDesc& OutAttributeDesc, bool &bOutConflictingTypesFound) const;
 
-	bool ContainsAttributeOnAnyData(FName InAttributeName) const;
+	bool ContainsAttributeOnAnyData(FPCGAttributeIdentifier InAttributeIdentifier) const;
 
 	/** Makes attribute present on all data. If data has existing attribute with same name then the given type will be applied. */
 	void AddAttributeToAllData(FPCGKernelAttributeKey InAttribute, const UPCGDataBinding* InBinding, const TArray<int32>* InOptionalUniqueStringKeys = nullptr);
diff --git a/Engine/Plugins/PCG/Source/PCG/Private/Compute/BuiltInKernels/PCGCountUniqueAttributeValuesKernel.cpp b/Engine/Plugins/PCG/Source/PCG/Private/Compute/BuiltInKernels/PCGCountUniqueAttributeValuesKernel.cpp
--- a/Engine/Plugins/PCG/Source/PCG/Private/Compute/BuiltInKernels/PCGCountUniqueAttributeValuesKernel.cpp
+++ b/Engine/Plugins/PCG/Source/PCG/Private/Compute/BuiltInKernels/PCGCountUniqueAttributeValuesKernel.cpp
@@ -80,7 +80,7 @@ bool UPCGCountUniqueAttributeValuesKernel::IsKernelDataValid(FPCGContext* InCont
 			return false;
 		}
 
-		if (AttributeDesc.Type != EPCGKernelAttributeType::StringKey)
+		if (AttributeDesc.AttributeKey.Type != EPCGKernelAttributeType::StringKey)
 		{
 			// Attribute value counting only currently supported for attributes of type StringKey.
 			PCG_KERNEL_VALIDATION_ERR(Context, Settings, FText::Format(
diff --git a/Engine/Plugins/PCG/Source/PCG/Private/Compute/BuiltInKernels/PCGDataNumKernel.cpp b/Engine/Plugins/PCG/Source/PCG/Private/Compute/BuiltInKernels/PCGDataNumKernel.cpp
--- a/Engine/Plugins/PCG/Source/PCG/Private/Compute/BuiltInKernels/PCGDataNumKernel.cpp
+++ b/Engine/Plugins/PCG/Source/PCG/Private/Compute/BuiltInKernels/PCGDataNumKernel.cpp
@@ -57,13 +57,13 @@ bool UPCGDataNumKernel::IsKernelDataValid(FPCGContext* InContext) const
 		{
 			FPCGKernelAttributeDesc AttributeDesc;
 			bool bConflictingTypesFound = false;
-			if (OutputDataDesc->GetAttributeDesc(OutCountAttributeKey.Name, AttributeDesc, bConflictingTypesFound))
+			if (OutputDataDesc->GetAttributeDesc(OutCountAttributeKey.Identifier, AttributeDesc, bConflictingTypesFound))
 			{
 				if (OutCountAttributeKey.Type == EPCGKernelAttributeType::Invalid)
 				{
 					PCG_KERNEL_VALIDATION_ERR(InContext, Settings, FText::Format(
 						LOCTEXT("InvalidAttributeKeyType", "Output attribute for the count could not be created with an 'invalid' type."),
-						FText::FromName(OutCountAttributeKey.Name)));
+						FText::FromString(OutCountAttributeKey.Identifier.ToString())));
 
 					return false;
 				}
@@ -74,7 +74,7 @@ bool UPCGDataNumKernel::IsKernelDataValid(FPCGContext* InContext) const
 			{
 				PCG_KERNEL_VALIDATION_ERR(InContext, Settings, FText::Format(
 					LOCTEXT("AttributeMissing", "Output attribute '{0}' not found."),
-					FText::FromName(OutCountAttributeKey.Name)));
+					FText::FromString(OutCountAttributeKey.Identifier.ToString())));
 
 				return false;
 			}
diff --git a/Engine/Plugins/PCG/Source/PCG/Private/Compute/DataInterfaces/Elements/PCGStaticMeshSpawnerDataInterface.cpp b/Engine/Plugins/PCG/Source/PCG/Private/Compute/DataInterfaces/Elements/PCGStaticMeshSpawnerDataInterface.cpp
--- a/Engine/Plugins/PCG/Source/PCG/Private/Compute/DataInterfaces/Elements/PCGStaticMeshSpawnerDataInterface.cpp
+++ b/Engine/Plugins/PCG/Source/PCG/Private/Compute/DataInterfaces/Elements/PCGStaticMeshSpawnerDataInterface.cpp
@@ -336,9 +336,9 @@ bool UPCGStaticMeshSpawnerDataProvider::PrepareForExecute_GameThread(UPCGDataBin
 
 			for (const FPCGKernelAttributeDesc& AttributeDesc : Desc.AttributeDescs)
 			{
-				if (AttributeDesc.Name == SelectorName && AttributeDesc.Type == EPCGKernelAttributeType::StringKey)
+				if (AttributeDesc.AttributeKey.Identifier.Name == SelectorName && AttributeDesc.AttributeKey.Type == EPCGKernelAttributeType::StringKey)
 				{
-					SelectorAttributeId = AttributeDesc.Index;
+					SelectorAttributeId = AttributeDesc.AttributeId;
 					break;
 				}
 			}
@@ -542,15 +542,15 @@ void UPCGStaticMeshSpawnerDataProvider::CreatePrimitiveDescriptors(FPCGContext*
 
 				for (const FPCGKernelAttributeDesc& AttributeDesc : Desc.AttributeDescs)
 				{
-					if (AttributeDesc.Name == SelectorName)
+					if (AttributeDesc.AttributeKey.Identifier.Name == SelectorName)
 					{
-						if (AttributeDesc.Type == EPCGKernelAttributeType::StringKey)
+						if (AttributeDesc.AttributeKey.Type == EPCGKernelAttributeType::StringKey)
 						{
 							for (int32 AttributeStringKeyValue : AttributeDesc.GetUniqueStringKeys())
 							{
 								if (AttributeStringKeyValue > 0)
 								{
-									StringKeyToInstanceCount.FindOrAdd(AttributeStringKeyValue, 0) += Desc.ElementCount;
+									StringKeyToInstanceCount.FindOrAdd(AttributeStringKeyValue, 0) += Desc.GetElementCountForAttribute(AttributeDesc);
 								}
 							}
 
diff --git a/Engine/Plugins/PCG/Source/PCG/Private/Compute/Elements/PCGCustomHLSL.cpp b/Engine/Plugins/PCG/Source/PCG/Private/Compute/Elements/PCGCustomHLSL.cpp
--- a/Engine/Plugins/PCG/Source/PCG/Private/Compute/Elements/PCGCustomHLSL.cpp
+++ b/Engine/Plugins/PCG/Source/PCG/Private/Compute/Elements/PCGCustomHLSL.cpp
@@ -880,6 +880,23 @@ void UPCGCustomHLSLSettings::UpdatePinSettings()
 
 void UPCGCustomHLSLSettings::UpdateAttributeKeys()
 {
+	// Make sure all the output attributes are up-to-date on their identifiers.
+	// Need to modify it if it ever changes for undo/redo support.
+	
+	Modify(/*bAlwaysMarkDirty=*/false);
+	bool bMarkedDirty = false;
+	for (FPCGPinPropertiesGPU& OutputPin : OutputPins)
+	{
+		for (FPCGKernelAttributeKey& AttributeKey : OutputPin.PropertiesGPU.CreatedKernelAttributeKeys)
+		{
+			if (AttributeKey.UpdateIdentifierFromSelector() && !bMarkedDirty)
+			{
+				bMarkedDirty = true;
+				MarkPackageDirty();
+			}
+		}
+	}
+	
 	KernelAttributeKeys.Reset();
 	PinToAttributeKeys.Reset();
 
@@ -923,12 +940,16 @@ void UPCGCustomHLSLSettings::UpdateAttributeKeys()
 		// Add the attribute if it hasn't already been referenced.
 		FPCGKernelAttributeKey Key;
 		Key.Type = static_cast<EPCGKernelAttributeType>(AttributeType);
-		Key.Name = FName(*NameStr);
-		KernelAttributeKeys.AddUnique(Key);
+		Key.SetSelector(FPCGAttributePropertySelector::CreateSelectorFromString(NameStr));
+		
+		if (Key.IsValid())
+		{
+			KernelAttributeKeys.AddUnique(Key);
 
-		// Add an entry mapping this pin to the referenced attribute, if the entry doesn't already exist.
-		TArray<TTuple<FPCGKernelAttributeKey, bool>>& Keys = PinToAttributeKeys.FindOrAdd(FName(*PinStr));
-		Keys.AddUnique(MakeTuple(Key, false));
+			// Add an entry mapping this pin to the referenced attribute, if the entry doesn't already exist.
+			TArray<TTuple<FPCGKernelAttributeKey, bool>>& Keys = PinToAttributeKeys.FindOrAdd(FName(*PinStr));
+			Keys.AddUnique(MakeTuple(Key, false));
+		}
 	}
 
 	// Process each output pin for any new attributes they want to create.
diff --git a/Engine/Plugins/PCG/Source/PCG/Private/Compute/Elements/PCGCustomHLSLKernel.cpp b/Engine/Plugins/PCG/Source/PCG/Private/Compute/Elements/PCGCustomHLSLKernel.cpp
--- a/Engine/Plugins/PCG/Source/PCG/Private/Compute/Elements/PCGCustomHLSLKernel.cpp
+++ b/Engine/Plugins/PCG/Source/PCG/Private/Compute/Elements/PCGCustomHLSLKernel.cpp
@@ -22,11 +22,6 @@ namespace PCGCustomHLSLKernel
 	constexpr TCHAR AttributeFunctionGet[] = { TEXT("Get") };
 	constexpr TCHAR AttributeFunctionSet[] = { TEXT("Set") };
 
-	FString GetKernelAttributeKeyAsString(const FPCGKernelAttributeKey& Key)
-	{
-		return FString::Format(TEXT("'{0}'"), { Key.Name.ToString() });
-	}
-
 	FString GetDataTypeString(EPCGDataType Type)
 	{
 		const UEnum* DataTypeEnum = StaticEnum<EPCGDataType>();
@@ -191,7 +186,7 @@ bool UPCGCustomHLSLKernel::ComputeOutputBindingDataDesc(const UPCGComputeGraph*
 
 				if (AttributeId != INDEX_NONE)
 				{
-					const FPCGKernelAttributeDesc AttributeDesc(AttributeId, AttributeKey.Type, AttributeKey.Name);
+					const FPCGKernelAttributeDesc AttributeDesc(AttributeId, AttributeKey.Type, AttributeKey.Identifier);
 					DataDesc.AttributeDescs.AddUnique(AttributeDesc);
 				}
 			}
@@ -207,7 +202,7 @@ bool UPCGCustomHLSLKernel::ComputeOutputBindingDataDesc(const UPCGComputeGraph*
 	{
 		auto HasStringKey = [](const FPCGKernelAttributeDesc& InAttributeDesc)
 			{
-				return InAttributeDesc.Type == EPCGKernelAttributeType::StringKey || InAttributeDesc.Type == EPCGKernelAttributeType::Name;
+				return InAttributeDesc.AttributeKey.Type == EPCGKernelAttributeType::StringKey || InAttributeDesc.AttributeKey.Type == EPCGKernelAttributeType::Name;
 			};
 
 		if (DataDesc.AttributeDescs.FindByPredicate(HasStringKey))
@@ -234,7 +229,7 @@ bool UPCGCustomHLSLKernel::ComputeOutputBindingDataDesc(const UPCGComputeGraph*
 			{
 				for (FPCGKernelAttributeDesc& AttributeDesc : DataDesc.AttributeDescs)
 				{
-					bFoundStringKeyAttribute |= (AttributeDesc.Type == EPCGKernelAttributeType::StringKey || AttributeDesc.Type == EPCGKernelAttributeType::Name);
+					bFoundStringKeyAttribute |= (AttributeDesc.AttributeKey.Type == EPCGKernelAttributeType::StringKey || AttributeDesc.AttributeKey.Type == EPCGKernelAttributeType::Name);
 				}
 			}
 
@@ -250,7 +245,7 @@ bool UPCGCustomHLSLKernel::ComputeOutputBindingDataDesc(const UPCGComputeGraph*
 			{
 				for (FPCGKernelAttributeDesc& AttributeDesc : DataDesc.AttributeDescs)
 				{
-					if (AttributeDesc.Type != EPCGKernelAttributeType::StringKey && AttributeDesc.Type != EPCGKernelAttributeType::Name)
+					if (AttributeDesc.AttributeKey.Type != EPCGKernelAttributeType::StringKey && AttributeDesc.AttributeKey.Type != EPCGKernelAttributeType::Name)
 					{
 						continue;
 					}
@@ -265,8 +260,8 @@ bool UPCGCustomHLSLKernel::ComputeOutputBindingDataDesc(const UPCGComputeGraph*
 						{
 							for (const FPCGKernelAttributeDesc& InputAttributeDesc : InputDataDesc.AttributeDescs)
 							{
-								const bool bIsString = InputAttributeDesc.Type == EPCGKernelAttributeType::StringKey || InputAttributeDesc.Type == EPCGKernelAttributeType::Name;
-								if (bIsString && InputAttributeDesc.Name == AttributeDesc.Name)
+								const bool bIsString = InputAttributeDesc.AttributeKey.Type == EPCGKernelAttributeType::StringKey || InputAttributeDesc.AttributeKey.Type == EPCGKernelAttributeType::Name;
+								if (bIsString && InputAttributeDesc.AttributeKey.Identifier == AttributeDesc.AttributeKey.Identifier)
 								{
 									AttributeDesc.AddUniqueStringKeys(InputAttributeDesc.GetUniqueStringKeys());
 									bFoundMatchingAttribute = true;
@@ -288,7 +283,7 @@ bool UPCGCustomHLSLKernel::ComputeOutputBindingDataDesc(const UPCGComputeGraph*
 							{
 								for (const FPCGKernelAttributeDesc& InputAttributeDesc : InputDataDesc.AttributeDescs)
 								{
-									if (InputAttributeDesc.Type == EPCGKernelAttributeType::StringKey || InputAttributeDesc.Type == EPCGKernelAttributeType::Name)
+									if (InputAttributeDesc.AttributeKey.Type == EPCGKernelAttributeType::StringKey || InputAttributeDesc.AttributeKey.Type == EPCGKernelAttributeType::Name)
 									{
 										AttributeDesc.AddUniqueStringKeys(InputAttributeDesc.GetUniqueStringKeys());
 									}
@@ -879,7 +874,7 @@ bool UPCGCustomHLSLKernel::PerformStaticValidation()
 #if PCG_KERNEL_LOGGING_ENABLED
 			StaticLogEntries.Emplace(FText::Format(
 				LOCTEXT("InvalidAttributeTypeNone", "Attribute '{0}' has invalid GPU attribute type '{1}', check the 'Attributes to Create' array on your pins."),
-				FText::FromName(AttributeKey.Name),
+				FText::FromName(AttributeKey.Identifier.Name),
 				FText::FromString(AttributeTypeEnum->GetNameStringByValue(static_cast<int64>(AttributeKey.Type)))),
 				EPCGKernelLogVerbosity::Error);
 #endif
@@ -975,8 +970,8 @@ void UPCGCustomHLSLKernel::ProcessShaderSource(const UPCGComputeGraph* InCompute
 				break;
 			}
 
-			const FPCGKernelAttributeKey AttributeKey(FName(*ParsedFunction.AttributeName), static_cast<EPCGKernelAttributeType>(AttrType));
-			const FString SourceDefinition = PCGCustomHLSLKernel::GetKernelAttributeKeyAsString(AttributeKey);
+			const FPCGKernelAttributeKey AttributeKey(FPCGAttributePropertySelector::CreateSelectorFromString(*ParsedFunction.AttributeName), static_cast<EPCGKernelAttributeType>(AttrType));
+			const FString SourceDefinition = FString::Format(TEXT("'{0}'"), { ParsedFunction.AttributeName });
 
 			const int32 AttributeIndex = InComputeGraph->GetStaticAttributeTable().GetAttributeId(AttributeKey);
 			const int32 ReplaceStartIndex = InOutSource.Find(SourceDefinition, ESearchCase::CaseSensitive, ESearchDir::FromStart, ParsedFunction.MatchBeginning);
@@ -1248,18 +1243,18 @@ bool UPCGCustomHLSLKernel::AreAttributesValid(FPCGContext* InContext, FText* Out
 				}
 
 				const FPCGKernelAttributeDesc* AttrDesc = nullptr;
-				const FName AttrName = FName(*AttributeName);
+				const FPCGKernelAttributeKey AttrKey(FPCGAttributePropertySelector::CreateSelectorFromString(AttributeName), static_cast<EPCGKernelAttributeType>(AttrType));
 				bool bFoundMatchingAttributeName = false;
 
 				// Verify that the attribute exists on at least one data in pin data collection.
 				for (const FPCGDataDesc& DataDesc : PinDesc->DataDescs)
 				{
-					AttrDesc = DataDesc.AttributeDescs.FindByPredicate([AttrName, AttrType, &bFoundMatchingAttributeName](const FPCGKernelAttributeDesc& Desc)
+					AttrDesc = DataDesc.AttributeDescs.FindByPredicate([&AttrKey, &bFoundMatchingAttributeName](const FPCGKernelAttributeDesc& Desc)
 					{
-						const bool bAttributeNameMatches = Desc.Name == AttrName;
+						const bool bAttributeNameMatches = Desc.AttributeKey.Identifier == AttrKey.Identifier;
 						bFoundMatchingAttributeName |= bAttributeNameMatches;
 
-						return bAttributeNameMatches && Desc.Type == static_cast<EPCGKernelAttributeType>(AttrType);
+						return bAttributeNameMatches && Desc.AttributeKey.Type == AttrKey.Type;
 					});
 
 					if (AttrDesc)
@@ -1278,7 +1273,7 @@ bool UPCGCustomHLSLKernel::AreAttributesValid(FPCGContext* InContext, FText* Out
 							*OutErrorText = FText::Format(
 								LOCTEXT("InvalidAttributeType", "Tried to call attribute function '{0}' on attribute '{1}' which is not of type '{2}'."),
 								ConstructFunctionText(),
-								FText::FromName(AttrName),
+								FText::FromString(AttributeName),
 								FText::FromString(TypeStr));
 						}
 						else
@@ -1286,7 +1281,7 @@ bool UPCGCustomHLSLKernel::AreAttributesValid(FPCGContext* InContext, FText* Out
 							*OutErrorText = FText::Format(
 								LOCTEXT("InvalidAttributeDNE", "Tried to call attribute function '{0}' on attribute '{1}' which does not exist."),
 								ConstructFunctionText(),
-								FText::FromName(AttrName),
+								FText::FromString(AttributeName),
 								FText::FromString(TypeStr));
 						}
 					}
diff --git a/Engine/Plugins/PCG/Source/PCG/Private/Compute/PCGComputeGraph.cpp b/Engine/Plugins/PCG/Source/PCG/Private/Compute/PCGComputeGraph.cpp
--- a/Engine/Plugins/PCG/Source/PCG/Private/Compute/PCGComputeGraph.cpp
+++ b/Engine/Plugins/PCG/Source/PCG/Private/Compute/PCGComputeGraph.cpp
@@ -223,9 +223,9 @@ void UPCGComputeGraph::DebugLogDataDescriptions(const UPCGDataBinding* InBinding
 				}
 
 				UE_LOG(LogPCG, Display, TEXT("\t\t\t\tID: %d\t\tName: %s\t\tType: %s\t\tUniqueStringKeys%s"),
-					AttributeDesc.Index,
-					*AttributeDesc.Name.ToString(),
-					*PCGKernelAttributeTypeEnum->GetDisplayNameTextByValue(static_cast<int64>(AttributeDesc.Type)).ToString(),
+					AttributeDesc.AttributeId,
+					*AttributeDesc.AttributeKey.Identifier.ToString(),
+					*PCGKernelAttributeTypeEnum->GetDisplayNameTextByValue(static_cast<int64>(AttributeDesc.AttributeKey.Type)).ToString(),
 					UniqueStringKeys.IsEmpty() ? TEXT(": ") : *UniqueStringKeys);
 			}
 
@@ -314,7 +314,7 @@ bool UPCGComputeGraph::ComputeTaggedDataPinDesc(const FPCGTaggedData& InTaggedDa
 		{
 			for (FPCGKernelAttributeDesc& AttrDesc : OutDescription.AttributeDescs)
 			{
-				const int32 Index = PCGComputeHelpers::GetMetadataAttributeIndexFromAttributeId(AttrDesc.Index);
+				const int32 Index = PCGComputeHelpers::GetMetadataAttributeIndexFromAttributeId(AttrDesc.AttributeId);
 				if (Index == INDEX_NONE)
 				{
 					continue;
@@ -324,7 +324,7 @@ bool UPCGComputeGraph::ComputeTaggedDataPinDesc(const FPCGTaggedData& InTaggedDa
 				{
 					const int32 NewId = PCGComputeHelpers::GetAttributeIdFromMetadataAttributeIndex(*AttributeIndex);
 					//UE_LOG(LogPCG, Warning, TEXT("Remapped attribute '%s' ID: %d -> %d"), *AttrDesc.Name.ToString(), AttrDesc.Index, NewId);
-					AttrDesc.Index = NewId;
+					AttrDesc.AttributeId = NewId;
 				}
 			}
 		}
diff --git a/Engine/Plugins/PCG/Source/PCG/Private/Compute/PCGComputeKernel.cpp b/Engine/Plugins/PCG/Source/PCG/Private/Compute/PCGComputeKernel.cpp
--- a/Engine/Plugins/PCG/Source/PCG/Private/Compute/PCGComputeKernel.cpp
+++ b/Engine/Plugins/PCG/Source/PCG/Private/Compute/PCGComputeKernel.cpp
@@ -118,12 +118,12 @@ void UPCGComputeKernel::ComputeDataDescFromPinProperties(const FPCGPinProperties
 		{
 			for (const FPCGKernelAttributeDesc& InAttrDesc : InDataDesc.AttributeDescs)
 			{
-				if (InAttrDesc.Index >= PCGComputeConstants::NUM_RESERVED_ATTRS
+				if (InAttrDesc.AttributeId >= PCGComputeConstants::NUM_RESERVED_ATTRS
 					&& !OutDataDesc.AttributeDescs.ContainsByPredicate([InAttrDesc](const FPCGKernelAttributeDesc& OutAttrDesc)
 						{
 							// Note: We shouldn't need to check for uniqueness of attr Index, since attributes should all have unique index via
 							// the GlobalAttributeLookupTable
-							return InAttrDesc.Name == OutAttrDesc.Name;
+							return InAttrDesc.AttributeKey.Identifier == OutAttrDesc.AttributeKey.Identifier;
 						}))
 				{
 					OutDataDesc.AttributeDescs.Emplace(InAttrDesc);
diff --git a/Engine/Plugins/PCG/Source/PCG/Private/Compute/PCGDataBinding.cpp b/Engine/Plugins/PCG/Source/PCG/Private/Compute/PCGDataBinding.cpp
--- a/Engine/Plugins/PCG/Source/PCG/Private/Compute/PCGDataBinding.cpp
+++ b/Engine/Plugins/PCG/Source/PCG/Private/Compute/PCGDataBinding.cpp
@@ -11,7 +11,9 @@
 #include "Compute/PCGComputeKernel.h"
 #include "Compute/Data/PCGProxyForGPUData.h"
 #include "Data/PCGPointData.h"
+#include "Elements/Metadata/PCGMetadataElementCommon.h"
 #include "Metadata/PCGMetadata.h"
+#include "Metadata/Accessors/PCGAttributeAccessorHelpers.h"
 
 void UPCGDataBinding::Initialize(const UPCGComputeGraph* InComputeGraph, FPCGContext* InContext)
 {
@@ -299,14 +301,14 @@ void UPCGDataBinding::AddInputDataAttributesToTable()
 				for (FPCGKernelAttributeDesc& SourceAttrDesc : SourceDesc.AttributeDescs)
 				{
 					// Convert from absolute ID to index into metadata attributes.
-					const int32 SourceAttributeTableIndex = PCGComputeHelpers::GetMetadataAttributeIndexFromAttributeId(SourceAttrDesc.Index);
+					const int32 SourceAttributeTableIndex = PCGComputeHelpers::GetMetadataAttributeIndexFromAttributeId(SourceAttrDesc.AttributeId);
 					if (SourceAttributeTableIndex == INDEX_NONE)
 					{
 						// Invalid index means the attribute is a property rather than metadata.
 						continue;
 					}
 
-					const int32 ExistingAttributeIndex = AttributeTable.AddAttribute(SourceAttrDesc.Name, SourceAttrDesc.Type);
+					const int32 ExistingAttributeIndex = AttributeTable.AddAttribute(SourceAttrDesc.AttributeKey.Identifier, SourceAttrDesc.AttributeKey.Type);
 
 					if (ExistingAttributeIndex != SourceAttributeTableIndex)
 					{
@@ -317,16 +319,27 @@ void UPCGDataBinding::AddInputDataAttributesToTable()
 		}
 		else if (const UPCGMetadata* Metadata = Data.Data ? Data.Data->ConstMetadata() : nullptr)
 		{
-			TArray<FName> AttributeNames;
+			const FPCGMetadataDomainID DefaultDomain = Metadata->GetConstDefaultMetadataDomain()->GetDomainID();
+			
+			TArray<FPCGAttributeIdentifier> AttributeIdentifiers;
 			TArray<EPCGMetadataTypes> AttributeTypes;
-			Metadata->GetAttributes(AttributeNames, AttributeTypes);
+			Metadata->GetAllAttributes(AttributeIdentifiers, AttributeTypes);
 
-			for (int Index = 0; Index < AttributeNames.Num(); ++Index)
+			// @todo_pcg: Attributes on other domains than the default are ignored at the moment, until we have a better way of representing
+			// different domains in the GPU header.
+			// It means those are lost.
+
+			for (int Index = 0; Index < AttributeIdentifiers.Num(); ++Index)
 			{
-				const FName AttributeName = AttributeNames[Index];
+				const FPCGAttributeIdentifier AttributeIdentifier = AttributeIdentifiers[Index];
+				if (!AttributeIdentifier.MetadataDomain.IsDefault() && AttributeIdentifier.MetadataDomain != DefaultDomain)
+				{
+					continue;
+				}
+					
 				const EPCGKernelAttributeType AttributeType = PCGDataForGPUHelpers::GetAttributeTypeFromMetadataType(AttributeTypes[Index]);
 
-				AttributeTable.AddAttribute(AttributeName, AttributeType);
+				AttributeTable.AddAttribute(AttributeIdentifier, AttributeType);
 			}
 		}
 	}
@@ -382,66 +395,50 @@ void UPCGDataBinding::AddInputDataStringsToTable()
 			{
 				continue;
 			}
+			
+			const FPCGMetadataDomainID DefaultDomain = Metadata->GetConstDefaultMetadataDomain()->GetDomainID();
 
-			TArray<FName> AttributeNames;
+			TArray<FPCGAttributeIdentifier> AttributeIdentifiers;
 			TArray<EPCGMetadataTypes> AttributeTypes;
-			Metadata->GetAttributes(AttributeNames, AttributeTypes);
+			Metadata->GetAllAttributes(AttributeIdentifiers, AttributeTypes);
 
-			for (int AttributeIndex = 0; AttributeIndex < AttributeNames.Num(); ++AttributeIndex)
+			// Cache the keys to a given domain, so we don't recreate them
+			TMap<FPCGMetadataDomainID, TUniquePtr<const IPCGAttributeAccessorKeys>> AllKeys;
+
+			for (int AttributeIndex = 0; AttributeIndex < AttributeIdentifiers.Num(); ++AttributeIndex)
 			{
-				const FName AttributeName = AttributeNames[AttributeIndex];
+				// @todo_pcg: Attributes on other domains than the default are ignored at the moment, until we have a better way of representing
+				// different domains in the GPU header.
+				// It means those are lost.
+				const FPCGAttributeIdentifier AttributeIdentifier = AttributeIdentifiers[AttributeIndex];
+				if (!AttributeIdentifier.MetadataDomain.IsDefault() && AttributeIdentifier.MetadataDomain != DefaultDomain)
+				{
+					continue;
+				}
+				
 				const EPCGKernelAttributeType AttributeType = PCGDataForGPUHelpers::GetAttributeTypeFromMetadataType(AttributeTypes[AttributeIndex]);
 
 				if (AttributeType == EPCGKernelAttributeType::StringKey || AttributeType == EPCGKernelAttributeType::Name)
 				{
-					const FPCGMetadataAttributeBase* AttributeBase = Metadata->GetConstAttribute(AttributeName);
-					if (!AttributeBase)
-					{
-						UE_LOG(LogPCG, Warning, TEXT("Attribute '%s' not found."), *AttributeName.ToString());
-						continue;
-					}
+					const FPCGMetadataDomain* MetadataDomain = Metadata->GetConstMetadataDomain(AttributeIdentifier.MetadataDomain);
+					const FPCGMetadataAttributeBase* AttributeBase = MetadataDomain->GetConstAttribute(AttributeIdentifier.Name);
+					check(MetadataDomain && AttributeBase);
 
-					auto AddEntryToStringTable = [AttributeBase, &StringTable=StringTable](int64 InValueKey)
-					{
-						if (AttributeBase->GetTypeId() == PCG::Private::MetadataTypes<FSoftObjectPath>::Id)
-						{
-							StringTable.AddUnique(static_cast<const FPCGMetadataAttribute<FSoftObjectPath>*>(AttributeBase)->GetValue(InValueKey).ToString());
-						}
-						else if (AttributeBase->GetTypeId() == PCG::Private::MetadataTypes<FString>::Id)
-						{
-							StringTable.AddUnique(static_cast<const FPCGMetadataAttribute<FString>*>(AttributeBase)->GetValue(InValueKey));
-						}
-						else if (AttributeBase->GetTypeId() == PCG::Private::MetadataTypes<FSoftClassPath>::Id)
-						{
-							StringTable.AddUnique(static_cast<const FPCGMetadataAttribute<FSoftClassPath>*>(AttributeBase)->GetValue(InValueKey).ToString());
-						}
-						else if (AttributeBase->GetTypeId() == PCG::Private::MetadataTypes<FName>::Id)
-						{
-							StringTable.AddUnique(static_cast<const FPCGMetadataAttribute<FName>*>(AttributeBase)->GetValue(InValueKey).ToString());
-						}
-						else
-						{
-							checkNoEntry();
-						}
-					};
-
-					if (const UPCGPointData* PointData = Cast<UPCGPointData>(Data.Data))
+					const TUniquePtr<const IPCGAttributeAccessor> Accessor = PCGAttributeAccessorHelpers::CreateConstAccessor(AttributeBase, MetadataDomain);
+					TUniquePtr<const IPCGAttributeAccessorKeys>& Keys = AllKeys.FindOrAdd(AttributeIdentifier.MetadataDomain);
+					if (!Keys.IsValid())
 					{
-						for (const FPCGPoint& Point : PointData->GetPoints())
-						{
-							AddEntryToStringTable(AttributeBase->GetValueKey(Point.MetadataEntry));
-						}
+						FPCGAttributePropertySelector Selector = FPCGAttributePropertySelector::CreateAttributeSelector(AttributeIdentifier.Name);
+						Data.Data->SetDomainFromDomainID(AttributeIdentifier.MetadataDomain, Selector);
+						Keys = PCGAttributeAccessorHelpers::CreateConstKeys(Data.Data, Selector);
 					}
-					else if (const UPCGParamData* ParamData = Cast<UPCGParamData>(Data.Data))
-					{
-						const int32 NumElements = Metadata->GetItemCountForChild();
 
-						for (int64 MetadataKey = 0; MetadataKey < NumElements; ++MetadataKey)
-						{
-							AddEntryToStringTable(AttributeBase->GetValueKey(MetadataKey));
-						}
-					}
-					else { /** TODO: More attribute types! */ }
+					check(Accessor && Keys);
+
+					PCGMetadataElementCommon::ApplyOnAccessor<FString>(*Keys, *Accessor, [&StringTable = StringTable](FString&& InValue, int32)
+					{
+						StringTable.AddUnique(std::move(InValue));
+					});
 				}
 			}
 		}
diff --git a/Engine/Plugins/PCG/Source/PCG/Private/Compute/PCGDataBinding.h b/Engine/Plugins/PCG/Source/PCG/Private/Compute/PCGDataBinding.h
--- a/Engine/Plugins/PCG/Source/PCG/Private/Compute/PCGDataBinding.h
+++ b/Engine/Plugins/PCG/Source/PCG/Private/Compute/PCGDataBinding.h
@@ -86,7 +86,7 @@ public:
 	void ReleaseTransientResources();
 
 	int32 GetAttributeId(const FPCGKernelAttributeKey& InAttribute) const { ensure(bTablesInitialized); return AttributeTable.GetAttributeId(InAttribute); }
-	int32 GetAttributeId(FName InName, EPCGKernelAttributeType InType) const { ensure(bTablesInitialized); return AttributeTable.GetAttributeId(InName, InType); }
+	int32 GetAttributeId(FPCGAttributeIdentifier InIdentifier, EPCGKernelAttributeType InType) const { ensure(bTablesInitialized); return AttributeTable.GetAttributeId(InIdentifier, InType); }
 
 	/** Get the set of all unique strings entering the compute graph. */
 	const TArray<FString>& GetStringTable() const { ensure(bTablesInitialized); return StringTable; }
diff --git a/Engine/Plugins/PCG/Source/PCG/Private/Compute/PCGDataForGPU.cpp b/Engine/Plugins/PCG/Source/PCG/Private/Compute/PCGDataForGPU.cpp
--- a/Engine/Plugins/PCG/Source/PCG/Private/Compute/PCGDataForGPU.cpp
+++ b/Engine/Plugins/PCG/Source/PCG/Private/Compute/PCGDataForGPU.cpp
@@ -14,6 +14,7 @@
 #include "Compute/PCGDataBinding.h"
 #include "Compute/Data/PCGProxyForGPUData.h"
 #include "Data/PCGPointData.h"
+#include "Elements/Metadata/PCGMetadataElementCommon.h"
 #include "Helpers/PCGAsync.h"
 #include "Metadata/PCGMetadata.h"
 #include "Metadata/Accessors/PCGAttributeAccessorHelpers.h"
@@ -118,7 +119,7 @@ namespace PCGDataForGPUHelpers
 
 		const PCGMetadataValueKey ValueKey = InAttributeBase->GetValueKey(InEntryKey);
 		const int16 TypeId = InAttributeBase->GetTypeId();
-		const int StrideBytes = PCGDataForGPUHelpers::GetAttributeTypeStrideBytes(InAttributeDesc.Type);
+		const int StrideBytes = PCGDataForGPUHelpers::GetAttributeTypeStrideBytes(InAttributeDesc.AttributeKey.Type);
 
 		switch (TypeId)
 		{
@@ -293,51 +294,58 @@ namespace PCGDataForGPUHelpers
 
 	FPCGMetadataAttributeBase* CreateAttributeFromAttributeDesc(UPCGMetadata* Metadata, const FPCGKernelAttributeDesc& AttributeDesc)
 	{
-		switch (AttributeDesc.Type)
+		check(Metadata);
+		FPCGMetadataDomain* MetadataDomain = Metadata->GetMetadataDomain(AttributeDesc.AttributeKey.Identifier.MetadataDomain);
+		if (!MetadataDomain)
+		{
+			return nullptr;
+		}
+		
+		switch (AttributeDesc.AttributeKey.Type)
 		{
 		case EPCGKernelAttributeType::Bool:
 		{
-			return Metadata->FindOrCreateAttribute<bool>(AttributeDesc.Name);
+			return MetadataDomain->FindOrCreateAttribute<bool>(AttributeDesc.AttributeKey.Identifier.Name);
 		}
 		case EPCGKernelAttributeType::Int:
 		{
-			return Metadata->FindOrCreateAttribute<int>(AttributeDesc.Name);
+			return MetadataDomain->FindOrCreateAttribute<int>(AttributeDesc.AttributeKey.Identifier.Name);
 		}
 		case EPCGKernelAttributeType::Float:
 		{
-			return Metadata->FindOrCreateAttribute<float>(AttributeDesc.Name);
+			return MetadataDomain->FindOrCreateAttribute<float>(AttributeDesc.AttributeKey.Identifier.Name);
 		}
 		case EPCGKernelAttributeType::Float2:
 		{
-			return Metadata->FindOrCreateAttribute<FVector2D>(AttributeDesc.Name);
+			return MetadataDomain->FindOrCreateAttribute<FVector2D>(AttributeDesc.AttributeKey.Identifier.Name);
 		}
 		case EPCGKernelAttributeType::Float3:
 		{
-			return Metadata->FindOrCreateAttribute<FVector>(AttributeDesc.Name);
+			return MetadataDomain->FindOrCreateAttribute<FVector>(AttributeDesc.AttributeKey.Identifier.Name);
 		}
 		case EPCGKernelAttributeType::Float4:
 		{
-			return Metadata->FindOrCreateAttribute<FVector4>(AttributeDesc.Name);
+			return MetadataDomain->FindOrCreateAttribute<FVector4>(AttributeDesc.AttributeKey.Identifier.Name);
 		}
 		case EPCGKernelAttributeType::Rotator:
 		{
-			return Metadata->FindOrCreateAttribute<FRotator>(AttributeDesc.Name);
+			return MetadataDomain->FindOrCreateAttribute<FRotator>(AttributeDesc.AttributeKey.Identifier.Name);
 		}
 		case EPCGKernelAttributeType::Quat:
 		{
-			return Metadata->FindOrCreateAttribute<FQuat>(AttributeDesc.Name);
+			return MetadataDomain->FindOrCreateAttribute<FQuat>(AttributeDesc.AttributeKey.Identifier.Name);
 		}
 		case EPCGKernelAttributeType::Transform:
 		{
-			return Metadata->FindOrCreateAttribute<FTransform>(AttributeDesc.Name);
+			return MetadataDomain->FindOrCreateAttribute<FTransform>(AttributeDesc.AttributeKey.Identifier.Name);
 		}
 		case EPCGKernelAttributeType::StringKey:
 		{
-			return Metadata->FindOrCreateAttribute<FString>(AttributeDesc.Name);
+			return MetadataDomain->FindOrCreateAttribute<FString>(AttributeDesc.AttributeKey.Identifier.Name);
 		}
 		case EPCGKernelAttributeType::Name:
 		{
-			return Metadata->FindOrCreateAttribute<FName>(AttributeDesc.Name);
+			return MetadataDomain->FindOrCreateAttribute<FName>(AttributeDesc.AttributeKey.Identifier.Name);
 		}
 		default:
 			return nullptr;
@@ -352,9 +360,10 @@ namespace PCGDataForGPUHelpers
 		const float* DataAsFloat = static_cast<const float*>(InPackedData);
 		const int32* DataAsInt = static_cast<const int32*>(InPackedData);
 
-		FPCGAttributePropertyOutputSelector Selector = FPCGAttributePropertySelector::CreateAttributeSelector<FPCGAttributePropertyOutputSelector>(InAttributeDesc.Name);
+		FPCGAttributePropertyOutputSelector Selector = FPCGAttributePropertySelector::CreateAttributeSelector<FPCGAttributePropertyOutputSelector>(InAttributeDesc.AttributeKey.Identifier.Name);
+		OutData->SetDomainFromDomainID(InAttributeDesc.AttributeKey.Identifier.MetadataDomain, Selector);
 
-		switch (InAttributeDesc.Type)
+		switch (InAttributeDesc.AttributeKey.Type)
 		{
 		case EPCGKernelAttributeType::Bool:
 		{
@@ -590,9 +599,9 @@ namespace PCGDataForGPUHelpers
 				return;
 			}
 
-			const uint32 StrideFloats = GetAttributeTypeStrideBytes(AttributeDesc.Type) / sizeof(float);
+			const uint32 StrideFloats = GetAttributeTypeStrideBytes(AttributeDesc.AttributeKey.Type) / sizeof(float);
 
-			OutAttributeIdOffsetStrides.Emplace(static_cast<uint32>(AttributeDesc.Index), OffsetFloats, StrideFloats, /*Unused*/0);
+			OutAttributeIdOffsetStrides.Emplace(static_cast<uint32>(AttributeDesc.AttributeId), OffsetFloats, StrideFloats, /*Unused*/0);
 
 			OffsetFloats += StrideFloats;
 		}
@@ -601,14 +610,61 @@ namespace PCGDataForGPUHelpers
 	}
 }
 
+FPCGKernelAttributeKey::FPCGKernelAttributeKey(const FPCGAttributePropertySelector& InSelector, EPCGKernelAttributeType InType)
+	: Type(InType)
+{
+	SetSelector(InSelector);
+}
+
+void FPCGKernelAttributeKey::SetSelector(const FPCGAttributePropertySelector& InSelector)
+{
+	Name.ImportFromOtherSelector(InSelector);
+	UpdateIdentifierFromSelector();
+}
+
+bool FPCGKernelAttributeKey::UpdateIdentifierFromSelector()
+{
+	// @todo_pcg: When we support more domains, we'll need to have a way to convert it to a proper identifier.
+	// For now, only support default metadata domain, and mark the identifier invalid if it is anything else.
+	if (Name.GetSelection() != EPCGAttributePropertySelection::Attribute)
+	{
+		UE_LOG(LogPCG, Error, TEXT("While updating FPCGKernelAttributeKey with %s, the selector was not targeting an attribute. Discarded."), *Name.ToString());
+		Identifier.MetadataDomain = PCGMetadataDomainID::Invalid;
+		return true;
+	}
+
+	const FName DomainName = Name.GetDomainName();
+	if (DomainName != NAME_None)// && DomainName != PCGDataConstants::DataDomainName)
+	{
+		//UE_LOG(LogPCG, Error, TEXT("While updating FPCGKernelAttributeKey with %s, the selector was targeting a domain that is not the @Data domain which is unsupported at the moment. Discarded."), *Name.ToString());
+		UE_LOG(LogPCG, Error, TEXT("While updating FPCGKernelAttributeKey with %s, the selector was targeting a domain which is unsupported at the moment. Discarded."), *Name.ToString());
+		Identifier.MetadataDomain = PCGMetadataDomainID::Invalid;
+		return true;
+	}
+
+	FPCGAttributeIdentifier NewIdentifier{Name.GetAttributeName(), DomainName == PCGDataConstants::DataDomainName ? PCGMetadataDomainID::Data : PCGMetadataDomainID::Default};
+	bool bHasChanged = NewIdentifier != Identifier;
+	if (bHasChanged)
+	{
+		Identifier = std::move(NewIdentifier);
+	}
+	
+	return bHasChanged;
+}
+
+bool FPCGKernelAttributeKey::IsValid() const
+{
+	return Identifier.MetadataDomain.IsValid();
+}
+
 bool FPCGKernelAttributeKey::operator==(const FPCGKernelAttributeKey& Other) const
 {
-	return Type == Other.Type && Name == Other.Name;
+	return Type == Other.Type && Identifier == Other.Identifier;
 }
 
 uint32 GetTypeHash(const FPCGKernelAttributeKey& In)
 {
-	return HashCombine(GetTypeHash(In.Type), GetTypeHash(In.Name));
+	return HashCombine(GetTypeHash(In.Type), GetTypeHash(In.Identifier));
 }
 
 int32 FPCGKernelAttributeTable::GetAttributeId(const FPCGKernelAttributeKey& InAttribute) const
@@ -626,14 +682,14 @@ int32 FPCGKernelAttributeTable::GetAttributeId(const FPCGKernelAttributeKey& InA
 	}
 }
 
-int32 FPCGKernelAttributeTable::GetAttributeId(FName InName, EPCGKernelAttributeType InType) const
+int32 FPCGKernelAttributeTable::GetAttributeId(FPCGAttributeIdentifier InIdentifier, EPCGKernelAttributeType InType) const
 {
-	return GetAttributeId(FPCGKernelAttributeKey(InName, InType));
+	return GetAttributeId(FPCGKernelAttributeKey(InIdentifier, InType));
 }
 
-int32 FPCGKernelAttributeTable::AddAttribute(FName InName, EPCGKernelAttributeType InType)
+int32 FPCGKernelAttributeTable::AddAttribute(FPCGAttributeIdentifier InIdentifier, EPCGKernelAttributeType InType)
 {
-	return AttributeTable.AddUnique(FPCGKernelAttributeKey(InName, InType));
+	return AttributeTable.AddUnique(FPCGKernelAttributeKey(InIdentifier, InType));
 }
 
 #if PCG_KERNEL_LOGGING_ENABLED
@@ -644,7 +700,7 @@ void FPCGKernelAttributeTable::DebugLog() const
 	for (const FPCGKernelAttributeKey& Attribute : AttributeTable)
 	{
 		UE_LOG(LogPCG, Display, TEXT("\tName: %s\t\tID: %d\t\tType: %s"),
-			*Attribute.Name.ToString(),
+			*Attribute.Identifier.ToString(),
 			GetAttributeId(Attribute),
 			*PCGKernelAttributeTypeEnum->GetDisplayNameTextByValue(static_cast<int64>(Attribute.Type)).ToString());
 	}
@@ -666,7 +722,7 @@ void FPCGKernelAttributeDesc::SetStringKeys(const TArrayView<int32>& InStringKey
 
 bool FPCGKernelAttributeDesc::operator==(const FPCGKernelAttributeDesc& Other) const
 {
-	return Index == Other.Index && Type == Other.Type && Name == Other.Name;
+	return AttributeId == Other.AttributeId && AttributeKey == Other.AttributeKey;
 }
 
 FPCGDataDesc::FPCGDataDesc(EPCGDataType InType, int InElementCount)
@@ -705,33 +761,34 @@ uint64 FPCGDataDesc::ComputePackedSize() const
 
 	for (const FPCGKernelAttributeDesc& AttributeDesc : AttributeDescs)
 	{
-		DataSizeBytes += static_cast<uint64>(PCGDataForGPUHelpers::GetAttributeTypeStrideBytes(AttributeDesc.Type)) * static_cast<uint64>(ElementCount);
+		const uint64 NumValues = static_cast<uint64>(GetElementCountForAttribute(AttributeDesc));
+		DataSizeBytes += static_cast<uint64>(PCGDataForGPUHelpers::GetAttributeTypeStrideBytes(AttributeDesc.AttributeKey.Type)) * NumValues;
 	}
 
 	return DataSizeBytes;
 }
 
-bool FPCGDataDesc::HasMetadataAttributes() const
+bool FPCGDataDesc::HasElementsMetadataDomainAttributes() const
 {
 	return AttributeDescs.ContainsByPredicate([](const FPCGKernelAttributeDesc& AttributeDesc)
 	{
-		return AttributeDesc.Index >= NUM_RESERVED_ATTRS;
+		return AttributeDesc.AttributeId >= NUM_RESERVED_ATTRS && AttributeDesc.AttributeKey.Identifier.MetadataDomain != PCGMetadataDomainID::Data;
 	});
 }
 
-bool FPCGDataDesc::ContainsAttribute(FName InAttributeName) const
+bool FPCGDataDesc::ContainsAttribute(FPCGAttributeIdentifier InAttributeIdentifier) const
 {
-	return AttributeDescs.ContainsByPredicate([InAttributeName](const FPCGKernelAttributeDesc& AttributeDesc)
+	return AttributeDescs.ContainsByPredicate([InAttributeIdentifier](const FPCGKernelAttributeDesc& AttributeDesc)
 	{
-		return AttributeDesc.Name == InAttributeName;
+		return AttributeDesc.AttributeKey.Identifier == InAttributeIdentifier;
 	});
 }
 
-bool FPCGDataDesc::ContainsAttribute(FName InAttributeName, EPCGKernelAttributeType InAttributeType) const
+bool FPCGDataDesc::ContainsAttribute(FPCGAttributeIdentifier InAttributeIdentifier, EPCGKernelAttributeType InAttributeType) const
 {
-	return AttributeDescs.ContainsByPredicate([InAttributeName, InAttributeType](const FPCGKernelAttributeDesc& AttributeDesc)
+	return AttributeDescs.ContainsByPredicate([InAttributeKey = FPCGKernelAttributeKey(InAttributeIdentifier, InAttributeType)](const FPCGKernelAttributeDesc& AttributeDesc)
 	{
-		return AttributeDesc.Name == InAttributeName && AttributeDesc.Type == InAttributeType;
+		return AttributeDesc.AttributeKey == InAttributeKey;
 	});
 }
 
@@ -747,21 +804,27 @@ void FPCGDataDesc::AddAttribute(FPCGKernelAttributeKey InAttribute, const UPCGDa
 
 	for (FPCGKernelAttributeDesc& AttributeDesc : AttributeDescs)
 	{
-		if (AttributeDesc.Name == InAttribute.Name)
+		if (AttributeDesc.AttributeKey == InAttribute)
 		{
-			ensureMsgf(!bAttributeFound, TEXT("Attribute '%s' encountered multiple times in a single data."), *AttributeDesc.Name.ToString());
+			ensureMsgf(!bAttributeFound, TEXT("Attribute '%s' encountered multiple times in a single data."), *AttributeDesc.AttributeKey.Identifier.ToString());
 			bAttributeFound = true;
 
-			AttributeDesc.Type = InAttribute.Type;
+			AttributeDesc.AttributeKey = InAttribute;
 		}
 	}
 
 	if (!bAttributeFound)
 	{
-		AttributeDescs.Emplace(AttributeId, InAttribute.Type, InAttribute.Name, InOptionalUniqueStringKeys);
+		AttributeDescs.Emplace(AttributeId, InAttribute.Type, InAttribute.Identifier, InOptionalUniqueStringKeys);
 	}
 }
 
+int32 FPCGDataDesc::GetElementCountForAttribute(const FPCGKernelAttributeDesc& AttributeDesc) const
+{
+	// Data domain only have a single element.
+	return AttributeDesc.AttributeKey.Identifier.MetadataDomain == PCGMetadataDomainID::Data ? 1 : ElementCount;
+}
+
 void FPCGDataDesc::InitializeAttributeDescs(const UPCGData* InData, const UPCGDataBinding* InBinding)
 {
 	TRACE_CPUPROFILER_EVENT_SCOPE(FPCGDataDesc::InitializeAttributeDescs);
@@ -776,13 +839,26 @@ void FPCGDataDesc::InitializeAttributeDescs(const UPCGData* InData, const UPCGDa
 
 	if (InBinding && Metadata)
 	{
-		TArray<FName> AttributeNames;
+		const FPCGMetadataDomainID DefaultDomain = Metadata->GetConstDefaultMetadataDomain()->GetDomainID();
+		
+		TArray<FPCGAttributeIdentifier> AttributeIdentifiers;
 		TArray<EPCGMetadataTypes> AttributeTypes;
-		Metadata->GetAttributes(AttributeNames, AttributeTypes);
+		Metadata->GetAllAttributes(AttributeIdentifiers, AttributeTypes);
+
+		// Cache the keys to a given domain, so we don't recreate them
+		TMap<FPCGMetadataDomainID, TUniquePtr<const IPCGAttributeAccessorKeys>> AllKeys;
 
-		for (int CustomAttributeIndex = 0; CustomAttributeIndex < AttributeNames.Num(); ++CustomAttributeIndex)
+		for (int CustomAttributeIndex = 0; CustomAttributeIndex < AttributeIdentifiers.Num(); ++CustomAttributeIndex)
 		{
-			const FName AttributeName = AttributeNames[CustomAttributeIndex];
+			// @todo_pcg: Attributes on other domains than the default are ignored at the moment, until we have a better way of representing
+			// different domains in the GPU header.
+			// It means those are lost.
+			const FPCGAttributeIdentifier AttributeIdentifier = AttributeIdentifiers[CustomAttributeIndex];
+			if (!AttributeIdentifier.MetadataDomain.IsDefault() && AttributeIdentifier.MetadataDomain != DefaultDomain)
+			{
+				continue;
+			}
+			
 			const EPCGKernelAttributeType AttributeType = PCGDataForGPUHelpers::GetAttributeTypeFromMetadataType(AttributeTypes[CustomAttributeIndex]);
 
 			if (AttributeType == EPCGKernelAttributeType::Invalid)
@@ -790,8 +866,11 @@ void FPCGDataDesc::InitializeAttributeDescs(const UPCGData* InData, const UPCGDa
 				const UEnum* EnumClass = StaticEnum<EPCGMetadataTypes>();
 				check(EnumClass);
 
+				FPCGAttributePropertySelector Selector = FPCGAttributePropertySelector::CreateAttributeSelector(AttributeIdentifier.Name);
+				InData->SetDomainFromDomainID(AttributeIdentifier.MetadataDomain, Selector);
+
 				UE_LOG(LogPCG, Warning, TEXT("Skipping attribute '%s'. '%s' type attributes are not supported on GPU."),
-					*AttributeName.ToString(),
+					*Selector.ToString(),
 					*EnumClass->GetNameStringByValue(static_cast<int64>(AttributeTypes[CustomAttributeIndex])));
 
 				continue;
@@ -809,94 +888,37 @@ void FPCGDataDesc::InitializeAttributeDescs(const UPCGData* InData, const UPCGDa
 
 			if (AttributeType == EPCGKernelAttributeType::StringKey || AttributeType == EPCGKernelAttributeType::Name)
 			{
-				const FPCGMetadataAttributeBase* AttributeBase = Metadata->GetConstAttribute(AttributeName);
-				check(AttributeBase);
-
-				const TArray<FString>& StringTable = InBinding->GetStringTable();
+				const FPCGMetadataDomain* MetadataDomain = Metadata->GetConstMetadataDomain(AttributeIdentifier.MetadataDomain);
+				const FPCGMetadataAttributeBase* AttributeBase = MetadataDomain->GetConstAttribute(AttributeIdentifier.Name);
+				check(MetadataDomain && AttributeBase);
 
-				if (Type == EPCGDataType::Point && ensure(InData->IsA<UPCGPointData>()))
+				const TUniquePtr<const IPCGAttributeAccessor> Accessor = PCGAttributeAccessorHelpers::CreateConstAccessor(AttributeBase, MetadataDomain);
+				TUniquePtr<const IPCGAttributeAccessorKeys>& Keys = AllKeys.FindOrAdd(AttributeIdentifier.MetadataDomain);
+				if (!Keys.IsValid())
 				{
-					const UPCGPointData* PointData = CastChecked<UPCGPointData>(InData);
-
-					for (int32 PointIndex = 0; PointIndex < PointData->GetPoints().Num(); ++PointIndex)
-					{
-						const PCGMetadataValueKey ValueKey = AttributeBase->GetValueKey(PointData->GetPoints()[PointIndex].MetadataEntry);
+					FPCGAttributePropertySelector Selector = FPCGAttributePropertySelector::CreateAttributeSelector(AttributeIdentifier.Name);
+					InData->SetDomainFromDomainID(AttributeIdentifier.MetadataDomain, Selector);
+					Keys = PCGAttributeAccessorHelpers::CreateConstKeys(InData, Selector);
+				}
 
-						int StringTableIndex = INDEX_NONE;
+				check(Accessor && Keys);
 
-						if (AttributeBase->GetTypeId() == PCG::Private::MetadataTypes<FSoftObjectPath>::Id)
-						{
-							StringTableIndex = StringTable.IndexOfByKey(static_cast<const FPCGMetadataAttribute<FSoftObjectPath>*>(AttributeBase)->GetValue(ValueKey).ToString());
-						}
-						else if (AttributeBase->GetTypeId() == PCG::Private::MetadataTypes<FString>::Id)
-						{
-							StringTableIndex = StringTable.IndexOfByKey(static_cast<const FPCGMetadataAttribute<FString>*>(AttributeBase)->GetValue(ValueKey));
-						}
-						else if (AttributeBase->GetTypeId() == PCG::Private::MetadataTypes<FSoftClassPath>::Id)
-						{
-							StringTableIndex = StringTable.IndexOfByKey(static_cast<const FPCGMetadataAttribute<FSoftClassPath>*>(AttributeBase)->GetValue(ValueKey).ToString());
-						}
-						else if (AttributeBase->GetTypeId() == PCG::Private::MetadataTypes<FName>::Id)
-						{
-							StringTableIndex = StringTable.IndexOfByKey(static_cast<const FPCGMetadataAttribute<FName>*>(AttributeBase)->GetValue(ValueKey).ToString());
-						}
-						else
-						{
-							// Should not get here if attribute type is string key.
-							checkNoEntry();
-						}
-
-						if (StringTableIndex != INDEX_NONE)
-						{
-							UniqueStringKeys.AddUnique(StringTableIndex);
-						}
-					}
-				}
-				else if (Type == EPCGDataType::Param && ensure(InData->IsA<UPCGParamData>()))
+				PCGMetadataElementCommon::ApplyOnAccessor<FString>(*Keys, *Accessor, [&UniqueStringKeys, &StringTable = InBinding->GetStringTable()](const FString& InValue, int32)
 				{
-					const int32 NumElements = Metadata->GetItemCountForChild();
-
-					for (int64 MetadataKey = 0; MetadataKey < NumElements; ++MetadataKey)
+					const int StringTableIndex = StringTable.IndexOfByKey(InValue);
+					if (StringTableIndex != INDEX_NONE)
 					{
-						int StringTableIndex = INDEX_NONE;
-
-						if (AttributeBase->GetTypeId() == PCG::Private::MetadataTypes<FSoftObjectPath>::Id)
-						{
-							StringTableIndex = StringTable.IndexOfByKey(static_cast<const FPCGMetadataAttribute<FSoftObjectPath>*>(AttributeBase)->GetValue(MetadataKey).ToString());
-						}
-						else if (AttributeBase->GetTypeId() == PCG::Private::MetadataTypes<FString>::Id)
-						{
-							StringTableIndex = StringTable.IndexOfByKey(static_cast<const FPCGMetadataAttribute<FString>*>(AttributeBase)->GetValue(MetadataKey));
-						}
-						else if (AttributeBase->GetTypeId() == PCG::Private::MetadataTypes<FSoftClassPath>::Id)
-						{
-							StringTableIndex = StringTable.IndexOfByKey(static_cast<const FPCGMetadataAttribute<FSoftClassPath>*>(AttributeBase)->GetValue(MetadataKey).ToString());
-						}
-						else if (AttributeBase->GetTypeId() == PCG::Private::MetadataTypes<FName>::Id)
-						{
-							StringTableIndex = StringTable.IndexOfByKey(static_cast<const FPCGMetadataAttribute<FName>*>(AttributeBase)->GetValue(MetadataKey).ToString());
-						}
-						else
-						{
-							// Should not get here if attribute type is string key.
-							checkNoEntry();
-						}
-
-						if (StringTableIndex != INDEX_NONE)
-						{
-							UniqueStringKeys.AddUnique(StringTableIndex);
-						}
+						UniqueStringKeys.AddUnique(StringTableIndex);
 					}
-				}
-				else { /* TODO: More types! */ }
+				});
 			}
 
-			const int32 AttributeId = InBinding->GetAttributeId(AttributeName, AttributeType);
-			ensureMsgf(AttributeId != INDEX_NONE, TEXT("Attribute '%s' type %d was missing from attribute table."), *AttributeName.ToString(), (int32)AttributeType);
+			const int32 AttributeId = InBinding->GetAttributeId(AttributeIdentifier, AttributeType);
+			ensureMsgf(AttributeId != INDEX_NONE, TEXT("Attribute '%s' type %d was missing from attribute table."), *AttributeIdentifier.Name.ToString(), (int32)AttributeType);
 
 			if (AttributeId != INDEX_NONE)
 			{
-				AttributeDescs.Emplace(AttributeId, AttributeType, AttributeName, &UniqueStringKeys);
+				AttributeDescs.Emplace(AttributeId, AttributeType, AttributeIdentifier, &UniqueStringKeys);
 			}
 		}
 	}
@@ -990,19 +1012,20 @@ void FPCGDataCollectionDesc::WriteHeader(TArray<uint32>& OutPackedDataCollection
 		for (int32 AttrIndex = 0; AttrIndex < DataDesc.AttributeDescs.Num(); ++AttrIndex)
 		{
 			const FPCGKernelAttributeDesc& AttributeDesc = DataDesc.AttributeDescs[AttrIndex];
-			const int AttributeStride = PCGDataForGPUHelpers::GetAttributeTypeStrideBytes(AttributeDesc.Type);
+			const int AttributeStride = PCGDataForGPUHelpers::GetAttributeTypeStrideBytes(AttributeDesc.AttributeKey.Type);
+			const int AttributeElementsCount = DataDesc.GetElementCountForAttribute(AttributeDesc);
 
 			// Scatter from attributes that are present into header which has slots for all possible attributes.
-			WriteAddressUints = (AttributeDesc.Index * ATTRIBUTE_HEADER_SIZE_BYTES + DataAttributesHeaderStartAddressBytes) >> 2;
+			WriteAddressUints = (AttributeDesc.AttributeId * ATTRIBUTE_HEADER_SIZE_BYTES + DataAttributesHeaderStartAddressBytes) >> 2;
 
 			// Data i element j: packed ID and stride
-			const uint32 AttributeId = AttributeDesc.Index;
+			const uint32 AttributeId = AttributeDesc.AttributeId;
 			const uint32 PackedIdAndStride = (AttributeId << 8) + AttributeStride;
 			OutPackedDataCollectionHeader[WriteAddressUints++] = PackedIdAndStride;
 			OutPackedDataCollectionHeader[WriteAddressUints++] = DataStartAddressBytes;
 
 			// Move the DataStartAddress to the beginning of the next attribute by skipping all the bytes of the current attribute.
-			DataStartAddressBytes += DataDesc.ElementCount * AttributeStride;
+			DataStartAddressBytes += AttributeElementsCount * AttributeStride;
 		}
 
 		// After scattering in attribute headers, fast forward to end of section.
@@ -1072,9 +1095,10 @@ void FPCGDataCollectionDesc::PackDataCollection(const FPCGDataCollection& InData
 
 			for (const FPCGKernelAttributeDesc& AttributeDesc : DataDesc.AttributeDescs)
 			{
-				const uint32 AttributeId = AttributeDesc.Index;
+				const uint32 AttributeId = AttributeDesc.AttributeId;
 
-				const FPCGMetadataAttributeBase* AttributeBase = (AttributeId >= NUM_RESERVED_ATTRS) ? Metadata->GetConstAttribute(AttributeDesc.Name) : nullptr;
+				const FPCGMetadataDomain* MetadataDomain = Metadata->GetConstMetadataDomain(AttributeDesc.AttributeKey.Identifier.MetadataDomain);
+				const FPCGMetadataAttributeBase* AttributeBase = (AttributeId >= NUM_RESERVED_ATTRS && MetadataDomain) ? MetadataDomain->GetConstAttribute(AttributeDesc.AttributeKey.Identifier.Name) : nullptr;
 
 				uint32 AddressUints = GetElementDataStartAddressUints(OutPackedDataCollection.GetData(), DataIndex, AttributeId);
 
@@ -1186,10 +1210,17 @@ void FPCGDataCollectionDesc::PackDataCollection(const FPCGDataCollection& InData
 				else
 				{
 					// Pack attribute. Validate first element only for perf.
-					ensure(PCGDataForGPUHelpers::PackAttributeHelper(AttributeBase, AttributeDesc, Points[0].MetadataEntry, InDataBinding->GetStringTable(), OutPackedDataCollection, AddressUints));
-					for (uint32 ElementIndex = 1; ElementIndex < NumElements; ++ElementIndex)
+					if (AttributeBase->GetMetadataDomain()->GetDomainID() == PCGMetadataDomainID::Data)
 					{
-						PCGDataForGPUHelpers::PackAttributeHelper(AttributeBase, AttributeDesc, Points[ElementIndex].MetadataEntry, InDataBinding->GetStringTable(), OutPackedDataCollection, AddressUints);
+						ensure(PCGDataForGPUHelpers::PackAttributeHelper(AttributeBase, AttributeDesc, /*InEntryKey*/0, InDataBinding->GetStringTable(), OutPackedDataCollection, AddressUints));
+					}
+					else
+					{
+						ensure(PCGDataForGPUHelpers::PackAttributeHelper(AttributeBase, AttributeDesc, Points[0].MetadataEntry, InDataBinding->GetStringTable(), OutPackedDataCollection, AddressUints));
+						for (uint32 ElementIndex = 1; ElementIndex < NumElements; ++ElementIndex)
+						{
+							PCGDataForGPUHelpers::PackAttributeHelper(AttributeBase, AttributeDesc, Points[ElementIndex].MetadataEntry, InDataBinding->GetStringTable(), OutPackedDataCollection, AddressUints);
+						}
 					}
 				}
 			}
@@ -1198,17 +1229,18 @@ void FPCGDataCollectionDesc::PackDataCollection(const FPCGDataCollection& InData
 		{
 			for (const FPCGKernelAttributeDesc& AttributeDesc : DataDesc.AttributeDescs)
 			{
-				const FPCGMetadataAttributeBase* AttributeBase = Metadata->GetConstAttribute(AttributeDesc.Name);
+				const FPCGMetadataDomain* MetadataDomain = Metadata->GetConstMetadataDomain(AttributeDesc.AttributeKey.Identifier.MetadataDomain);
+				const FPCGMetadataAttributeBase* AttributeBase = MetadataDomain ? MetadataDomain->GetConstAttribute(AttributeDesc.AttributeKey.Identifier.Name) : nullptr;
 				if (!AttributeBase)
 				{
 					continue;
 				}
 
-				uint32 AddressUints = GetElementDataStartAddressUints(OutPackedDataCollection.GetData(), DataIndex, AttributeDesc.Index);
+				uint32 AddressUints = GetElementDataStartAddressUints(OutPackedDataCollection.GetData(), DataIndex, AttributeDesc.AttributeId);
 
 				// Pack attribute. Validate first element only for perf.
 				ensure(PCGDataForGPUHelpers::PackAttributeHelper(AttributeBase, AttributeDesc, /*InEntryKey*/0, InDataBinding->GetStringTable(), OutPackedDataCollection, AddressUints));
-				for (int32 ElementIndex = 1; ElementIndex < DataDesc.ElementCount; ++ElementIndex)
+				for (int32 ElementIndex = 1; ElementIndex < DataDesc.GetElementCountForAttribute(AttributeDesc); ++ElementIndex)
 				{
 					PCGDataForGPUHelpers::PackAttributeHelper(AttributeBase, AttributeDesc, /*InEntryKey*/ElementIndex, InDataBinding->GetStringTable(), OutPackedDataCollection, AddressUints);
 				}
@@ -1278,7 +1310,7 @@ EPCGUnpackDataCollectionResult FPCGDataCollectionDesc::UnpackDataCollection(FPCG
 				TRACE_CPUPROFILER_EVENT_SCOPE(InitializeOutput);
 
 				// We only need to add the entry keys if there are actually attributes to unpack.
-				if (DataDesc.HasMetadataAttributes())
+				if (DataDesc.HasElementsMetadataDomainAttributes())
 				{
 					TPCGValueRange<int64> MetadataEntryRange = OutPointData->GetMetadataEntryValueRange(/*bAllocate=*/true);
 					TArray<PCGMetadataEntryKey*> ParentEntryKeys;
@@ -1322,9 +1354,9 @@ EPCGUnpackDataCollectionResult FPCGDataCollectionDesc::UnpackDataCollection(FPCG
 
 			for (const FPCGKernelAttributeDesc& AttributeDesc : AttributeDescs)
 			{
-				if (AttributeDesc.Index >= NUM_RESERVED_ATTRS)
+				if (AttributeDesc.AttributeId >= NUM_RESERVED_ATTRS)
 				{
-					MetadataAttributes[AttributeDesc.Index] = PCGDataForGPUHelpers::CreateAttributeFromAttributeDesc(Metadata, AttributeDesc);
+					MetadataAttributes[AttributeDesc.AttributeId] = PCGDataForGPUHelpers::CreateAttributeFromAttributeDesc(Metadata, AttributeDesc);
 				}
 			}
 
@@ -1337,8 +1369,8 @@ EPCGUnpackDataCollectionResult FPCGDataCollectionDesc::UnpackDataCollection(FPCG
 			// Loop over attributes.
 			for (const FPCGKernelAttributeDesc& AttributeDesc : AttributeDescs)
 			{
-				const uint32 AttributeId = AttributeDesc.Index;
-				const uint32 AttributeStrideUints = PCGDataForGPUHelpers::GetAttributeTypeStrideBytes(AttributeDesc.Type) >> 2;
+				const uint32 AttributeId = AttributeDesc.AttributeId;
+				const uint32 AttributeStrideUints = PCGDataForGPUHelpers::GetAttributeTypeStrideBytes(AttributeDesc.AttributeKey.Type) >> 2;
 				const uint32 AddressUints = GetElementDataStartAddressUints(DataAsUint, DataIndex, AttributeId);
 
 				if (AttributeId < NUM_RESERVED_ATTRS)
@@ -1472,9 +1504,9 @@ EPCGUnpackDataCollectionResult FPCGDataCollectionDesc::UnpackDataCollection(FPCG
 				{
 					TRACE_CPUPROFILER_EVENT_SCOPE(UnpackAttribute);
 
-					if (FPCGMetadataAttributeBase* AttributeBase = MetadataAttributes[AttributeDesc.Index])
+					if (FPCGMetadataAttributeBase* AttributeBase = MetadataAttributes[AttributeDesc.AttributeId])
 					{
-						ensure(PCGDataForGPUHelpers::UnpackAttributeHelper(InContext, PackedData, AttributeDesc, InStringTable, AddressUints, NumElements, OutPointData));
+						ensure(PCGDataForGPUHelpers::UnpackAttributeHelper(InContext, PackedData, AttributeDesc, InStringTable, AddressUints, DataDesc.GetElementCountForAttribute(AttributeDesc), OutPointData));
 					}
 				}
 			}
@@ -1552,7 +1584,7 @@ EPCGUnpackDataCollectionResult FPCGDataCollectionDesc::UnpackDataCollection(FPCG
 
 			for (const FPCGKernelAttributeDesc& AttributeDesc : AttributeDescs)
 			{
-				MetadataAttributes[AttributeDesc.Index] = PCGDataForGPUHelpers::CreateAttributeFromAttributeDesc(Metadata, AttributeDesc);
+				MetadataAttributes[AttributeDesc.AttributeId] = PCGDataForGPUHelpers::CreateAttributeFromAttributeDesc(Metadata, AttributeDesc);
 			}
 
 			// No work to do if there are no elements to process.
@@ -1566,9 +1598,9 @@ EPCGUnpackDataCollectionResult FPCGDataCollectionDesc::UnpackDataCollection(FPCG
 			{
 				TRACE_CPUPROFILER_EVENT_SCOPE(UnpackAttribute);
 
-				if (FPCGMetadataAttributeBase* AttributeBase = MetadataAttributes[AttributeDesc.Index])
+				if (FPCGMetadataAttributeBase* AttributeBase = MetadataAttributes[AttributeDesc.AttributeId])
 				{
-					const uint32 AddressUints = GetElementDataStartAddressUints(DataAsUint, DataIndex, AttributeDesc.Index);
+					const uint32 AddressUints = GetElementDataStartAddressUints(DataAsUint, DataIndex, AttributeDesc.AttributeId);
 					ensure(PCGDataForGPUHelpers::UnpackAttributeHelper(InContext, PackedData, AttributeDesc, InStringTable, AddressUints, NumElements, OutParamData));
 				}
 			}
@@ -1599,7 +1631,7 @@ void FPCGDataCollectionDesc::Combine(const FPCGDataCollectionDesc& Other)
 	DataDescs.Append(Other.DataDescs);
 }
 
-bool FPCGDataCollectionDesc::GetAttributeDesc(FName InAttributeName, FPCGKernelAttributeDesc& OutAttributeDesc, bool& bOutConflictingTypesFound) const
+bool FPCGDataCollectionDesc::GetAttributeDesc(FPCGAttributeIdentifier InAttributeIdentifier, FPCGKernelAttributeDesc& OutAttributeDesc, bool& bOutConflictingTypesFound) const
 {
 	// Will be set to the type of the first attribute across all the data which has a matching name, and then is used to detect if subsequent
 	// attributes are found with conflicting type.
@@ -1609,15 +1641,15 @@ bool FPCGDataCollectionDesc::GetAttributeDesc(FName InAttributeName, FPCGKernelA
 	{
 		for (const FPCGKernelAttributeDesc& AttributeDesc : DataDesc.AttributeDescs)
 		{
-			if (AttributeDesc.Name == InAttributeName && ensure(AttributeDesc.Type != EPCGKernelAttributeType::Invalid))
+			if (AttributeDesc.AttributeKey.Identifier == InAttributeIdentifier && ensure(AttributeDesc.AttributeKey.Type != EPCGKernelAttributeType::Invalid))
 			{
 				if (FoundType == EPCGKernelAttributeType::Invalid)
 				{
 					// Take the first matching attribute.
 					OutAttributeDesc = AttributeDesc;
-					FoundType = AttributeDesc.Type;
+					FoundType = AttributeDesc.AttributeKey.Type;
 				}
-				else if (FoundType != AttributeDesc.Type)
+				else if (FoundType != AttributeDesc.AttributeKey.Type)
 				{
 					// Signal conflict found.
 					bOutConflictingTypesFound = true;
@@ -1636,11 +1668,11 @@ bool FPCGDataCollectionDesc::GetAttributeDesc(FName InAttributeName, FPCGKernelA
 	return FoundType != EPCGKernelAttributeType::Invalid;
 }
 
-bool FPCGDataCollectionDesc::ContainsAttributeOnAnyData(FName InAttributeName) const
+bool FPCGDataCollectionDesc::ContainsAttributeOnAnyData(FPCGAttributeIdentifier InAttributeIdentifier) const
 {
-	return DataDescs.ContainsByPredicate([InAttributeName](const FPCGDataDesc& DataDesc)
+	return DataDescs.ContainsByPredicate([InAttributeIdentifier](const FPCGDataDesc& DataDesc)
 	{
-		return DataDesc.ContainsAttribute(InAttributeName);
+		return DataDesc.ContainsAttribute(InAttributeIdentifier);
 	});
 }
 
@@ -1658,7 +1690,7 @@ void FPCGDataCollectionDesc::GetUniqueStringKeyValues(int32 InAttributeId, TArra
 	{
 		for (const FPCGKernelAttributeDesc& AttributeDesc : DataDesc.AttributeDescs)
 		{
-			if (AttributeDesc.Index == InAttributeId)
+			if (AttributeDesc.AttributeId == InAttributeId)
 			{
 				for (int32 StringKey : AttributeDesc.GetUniqueStringKeys())
 				{
diff --git a/Engine/Plugins/PCG/Source/PCG/Private/Elements/Metadata/PCGMetadataPartitionKernel.cpp b/Engine/Plugins/PCG/Source/PCG/Private/Elements/Metadata/PCGMetadataPartitionKernel.cpp
--- a/Engine/Plugins/PCG/Source/PCG/Private/Elements/Metadata/PCGMetadataPartitionKernel.cpp
+++ b/Engine/Plugins/PCG/Source/PCG/Private/Elements/Metadata/PCGMetadataPartitionKernel.cpp
@@ -70,7 +70,7 @@ bool UPCGMetadataPartitionKernel::IsKernelDataValid(FPCGContext* InContext) cons
 			return false;
 		}
 
-		if (AttributeDesc.Type != EPCGKernelAttributeType::StringKey)
+		if (AttributeDesc.AttributeKey.Type != EPCGKernelAttributeType::StringKey)
 		{
 			PCG_KERNEL_VALIDATION_ERR(Context, Settings, FText::Format(
 				LOCTEXT("PartitionAttributeTypeInvalid", "Attribute '{0}' not usable for partitioning, only attributes of type String Key are currently supported."),
@@ -119,7 +119,7 @@ bool UPCGMetadataPartitionKernel::ComputeOutputBindingDataDesc(const UPCGCompute
 				// Set the unique string keys to reflect the partition operation.
 				for (FPCGKernelAttributeDesc& AttrDesc : Desc.AttributeDescs)
 				{
-					if (AttrDesc.Index == PartitionAttributeId)
+					if (AttrDesc.AttributeId == PartitionAttributeId)
 					{
 						AttrDesc.SetStringKeys(MakeArrayView(&StringKeyValue, 1));
 						break;
diff --git a/Engine/Plugins/PCG/Source/PCG/Private/Elements/PCGStaticMeshSpawnerKernel.cpp b/Engine/Plugins/PCG/Source/PCG/Private/Elements/PCGStaticMeshSpawnerKernel.cpp
--- a/Engine/Plugins/PCG/Source/PCG/Private/Elements/PCGStaticMeshSpawnerKernel.cpp
+++ b/Engine/Plugins/PCG/Source/PCG/Private/Elements/PCGStaticMeshSpawnerKernel.cpp
@@ -88,7 +88,7 @@ bool UPCGStaticMeshSpawnerKernel::ComputeOutputBindingDataDesc(const UPCGCompute
 
 					for (const FPCGKernelAttributeDesc& AttrDesc : DataDesc.AttributeDescs)
 					{
-						if (AttrDesc.Name == SelectorByAttribute->AttributeName)
+						if (AttrDesc.AttributeKey.Identifier.Name == SelectorByAttribute->AttributeName)
 						{
 							StringKeys = AttrDesc.GetUniqueStringKeys();
 							break;
diff --git a/Engine/Plugins/PCG/Source/PCG/Private/Metadata/PCGObjectPropertyOverride.cpp b/Engine/Plugins/PCG/Source/PCG/Private/Metadata/PCGObjectPropertyOverride.cpp
--- a/Engine/Plugins/PCG/Source/PCG/Private/Metadata/PCGObjectPropertyOverride.cpp
+++ b/Engine/Plugins/PCG/Source/PCG/Private/Metadata/PCGObjectPropertyOverride.cpp
@@ -170,7 +170,7 @@ void FPCGObjectSingleOverride::GatherAllOverridesToLoad(TArray<FSoftObjectPath>&
 		return;
 	}
 
-	PCGMetadataElementCommon::ApplyOnAccessor<FSoftObjectPath>(*InputKeys, *ObjectOverrideInputAccessor, [&OutObjectsToLoad](FSoftObjectPath& Value, int32 Index)
+	PCGMetadataElementCommon::ApplyOnAccessor<FSoftObjectPath>(*InputKeys, *ObjectOverrideInputAccessor, [&OutObjectsToLoad](FSoftObjectPath&& Value, int32 Index)
 	{
 		// Only keep valid paths
 		if (!Value.IsNull())
diff --git a/Engine/Plugins/PCG/Source/PCG/Public/Elements/Metadata/PCGMetadataElementCommon.h b/Engine/Plugins/PCG/Source/PCG/Public/Elements/Metadata/PCGMetadataElementCommon.h
--- a/Engine/Plugins/PCG/Source/PCG/Public/Elements/Metadata/PCGMetadataElementCommon.h
+++ b/Engine/Plugins/PCG/Source/PCG/Public/Elements/Metadata/PCGMetadataElementCommon.h
@@ -130,7 +130,7 @@ namespace PCGMetadataElementCommon
 	/**
 	* Iterate over the full range of the keys (if Count is negative, otherwise, as many times as Count), calling the callback with values get from the accessor.
 	* ChunkSize influence the max number of values to get in one go with GetRange.
-	* Callback should have a signature: void(const T& Value, int32 Index)
+	* Callback should have a signature: void(const T& Value, int32 Index) / void(T Value, int32 Index)
 	* Return false if it process nothing.
 	*/
 	template <typename T, typename Func>
@@ -140,7 +140,8 @@ namespace PCGMetadataElementCommon
 		{
 			for (int32 j = 0; j < Range; ++j)
 			{
-				Callback(View[j], Start + j);
+				// We can move the value since we don't use it afterward. If the callback is taking a const ref, it's fine.
+				Callback(std::move(View[j]), Start + j);
 			}
 		};
 
@@ -252,7 +253,7 @@ namespace PCGMetadataElementCommon
 	/**
 	* Iterate over the full range of the keys (if Count is negative, otherwise, as many times as Count), calling the callback with values get from multiple accessor.
 	* ChunkSize influence the max number of values to get in one go with GetRange.
-	* Callback should have a signature: void(const T&... Values, int32 Index)
+	* Callback should have a signature: void(const T&... Values, int32 Index) / void(T... Values, int32 Index)
 	* Return false if it process nothing.
 	*/
 	template <typename ...T, typename Func = TFunction<void(const T&..., int32 Index)>>
@@ -262,7 +263,8 @@ namespace PCGMetadataElementCommon
 		{
 			for (int32 j = 0; j < Range; ++j)
 			{
-				Callback(Views[j]..., Start + j);
+				// We can move the value since we don't use it afterward. If the callback is taking a const ref, it's fine.
+				Callback(std::move(Views[j])..., Start + j);
 			}
 		};
 
diff --git a/Engine/Plugins/PCG/Source/PCG/Public/Metadata/PCGMetadataCommon.h b/Engine/Plugins/PCG/Source/PCG/Public/Metadata/PCGMetadataCommon.h
--- a/Engine/Plugins/PCG/Source/PCG/Public/Metadata/PCGMetadataCommon.h
+++ b/Engine/Plugins/PCG/Source/PCG/Public/Metadata/PCGMetadataCommon.h
@@ -193,6 +193,22 @@ struct FPCGAttributeIdentifier
 		return OutMap;
 	}
 
+	FString ToString() const
+	{
+		if (!MetadataDomain.IsValid())
+		{
+			return TEXT("INVALID");
+		}
+		else if (MetadataDomain.IsDefault())
+		{
+			return Name.ToString();
+		}
+		else
+		{
+			return FString::Printf(TEXT("%s.%s"), *MetadataDomain.DebugName.ToString(), *Name.ToString());
+		}
+	}
+
 	UPROPERTY(EditAnywhere, BlueprintReadWrite, Category = "")
 	FName Name = NAME_None;
 
