[PCG] Fix missing points on Projection element when projecting against Landscape (Editor)
- PCGDataFromActor ProcessActors can now wait on optional dependencies
- FPCGGetLandscapeDataElement implements this new ProcessActors and in editor will call PrepareForSpatialQuery to create those additional dependencies (needed because Landscape Cache entry creation doesn't support multithreading)

#rnx

diff --git a/Engine/Plugins/PCG/Source/PCG/Private/Elements/PCGDataFromActor.cpp b/Engine/Plugins/PCG/Source/PCG/Private/Elements/PCGDataFromActor.cpp
--- a/Engine/Plugins/PCG/Source/PCG/Private/Elements/PCGDataFromActor.cpp
+++ b/Engine/Plugins/PCG/Source/PCG/Private/Elements/PCGDataFromActor.cpp
@@ -350,6 +350,12 @@ bool FPCGDataFromActorElement::ExecuteInternal(FPCGContext* InContext) const
 	check(InContext);
 	FPCGDataFromActorContext* Context = static_cast<FPCGDataFromActorContext*>(InContext);
 
+	// Done waiting
+	if (Context->bWaitingOnProcessActors)
+	{
+		return true;
+	}
+
 	const UPCGDataFromActorSettings* Settings = Context->GetInputSettings<UPCGDataFromActorSettings>();
 	check(Settings);
 
@@ -580,7 +586,15 @@ bool FPCGDataFromActorElement::ExecuteInternal(FPCGContext* InContext) const
 		}
 #endif
 
-		ProcessActors(Context, Settings, Context->FoundActors);
+		TArray<FPCGTaskId> OutDynamicDependencies;
+		ProcessActors(Context, Settings, Context->FoundActors, OutDynamicDependencies);
+		if (!OutDynamicDependencies.IsEmpty())
+		{
+			Context->bWaitingOnProcessActors = true;
+			Context->bIsPaused = true;
+			Context->DynamicDependencies.Append(OutDynamicDependencies);
+			return false;
+		}
 	}
 
 	return true;
@@ -672,6 +686,11 @@ void FPCGDataFromActorElement::GatherWaitTasks(AActor* FoundActor, FPCGContext*
 	}
 }
 
+void FPCGDataFromActorElement::ProcessActors(FPCGContext* Context, const UPCGDataFromActorSettings* Settings, const TArray<AActor*>& FoundActors, TArray<FPCGTaskId>& OutDynamicDependencies) const
+{
+	ProcessActors(Context, Settings, FoundActors);
+}
+
 void FPCGDataFromActorElement::ProcessActors(FPCGContext* Context, const UPCGDataFromActorSettings* Settings, const TArray<AActor*>& FoundActors) const
 {
 	// Special case:
diff --git a/Engine/Plugins/PCG/Source/PCG/Private/Elements/PCGTypedGetter.cpp b/Engine/Plugins/PCG/Source/PCG/Private/Elements/PCGTypedGetter.cpp
--- a/Engine/Plugins/PCG/Source/PCG/Private/Elements/PCGTypedGetter.cpp
+++ b/Engine/Plugins/PCG/Source/PCG/Private/Elements/PCGTypedGetter.cpp
@@ -85,7 +85,7 @@ FPCGElementPtr UPCGGetLandscapeSettings::CreateElement() const
 	return MakeShared<FPCGGetLandscapeDataElement>();
 }
 
-void FPCGGetLandscapeDataElement::ProcessActors(FPCGContext* Context, const UPCGDataFromActorSettings* InSettings, const TArray<AActor*>& FoundActors) const
+void FPCGGetLandscapeDataElement::ProcessActors(FPCGContext* Context, const UPCGDataFromActorSettings* InSettings, const TArray<AActor*>& FoundActors, TArray<FPCGTaskId>& OutDynamicDependencies) const
 {
 	check(Context);
 	check(InSettings);
@@ -125,12 +125,23 @@ void FPCGGetLandscapeDataElement::ProcessActors(FPCGContext* Context, const UPCG
 		UPCGLandscapeData* LandscapeData = FPCGContext::NewObject_AnyThread<UPCGLandscapeData>(Context);
 		LandscapeData->Initialize(Landscapes, LandscapeBounds, Settings->SamplingProperties);
 		
+		// Prepare data for the PCG Bounds in Editor only as it can create cache entries which isn't thread safe
+#if WITH_EDITOR
+		const FBox PrepareBounds = Context->ExecutionSource.IsValid() ? Context->ExecutionSource->GetExecutionState().GetBounds() : LandscapeBounds;
+		OutDynamicDependencies.Append(LandscapeData->PrepareForSpatialQuery(Context, PrepareBounds));
+#endif
+
 		FPCGTaggedData& TaggedData = Context->OutputData.TaggedData.Emplace_GetRef();
 		TaggedData.Data = LandscapeData;
 		TaggedData.Tags = LandscapeTags;
 	}
 }
 
+void FPCGGetLandscapeDataElement::ProcessActors(FPCGContext* Context, const UPCGDataFromActorSettings* InSettings, const TArray<AActor*>& FoundActors) const
+{
+	checkf(false, TEXT("This should never be called directly"));
+}
+
 void FPCGGetLandscapeDataElement::ProcessActor(FPCGContext* Context, const UPCGDataFromActorSettings* Settings, AActor* FoundActor) const
 {
 	checkf(false, TEXT("This should never be called directly"));
diff --git a/Engine/Plugins/PCG/Source/PCG/Public/Elements/PCGDataFromActor.h b/Engine/Plugins/PCG/Source/PCG/Public/Elements/PCGDataFromActor.h
--- a/Engine/Plugins/PCG/Source/PCG/Public/Elements/PCGDataFromActor.h
+++ b/Engine/Plugins/PCG/Source/PCG/Public/Elements/PCGDataFromActor.h
@@ -149,6 +149,7 @@ struct FPCGDataFromActorContext : public FPCGLoadObjectsFromPathContext
 	TArray<AActor*> FoundActors;
 	FPCGComponentSelectorSettings ComponentSelector;
 	bool bPerformedQuery = false;
+	bool bWaitingOnProcessActors = false;
 
 #if WITH_EDITOR
 	/** Any change origin ignores we added, to solve dependency issues (like upstream execution cancelling downstream graph). */
@@ -169,6 +170,8 @@ protected:
 	virtual bool ExecuteInternal(FPCGContext* Context) const override;
 	virtual bool SupportsBasePointDataInputs(FPCGContext* InContext) const override { return true; }
 	void GatherWaitTasks(AActor* FoundActor, FPCGContext* InContext, TArray<FPCGTaskId>& OutWaitTasks) const;
+
+	virtual void ProcessActors(FPCGContext* Context, const UPCGDataFromActorSettings* Settings, const TArray<AActor*>& FoundActors, TArray<FPCGTaskId>& OutDynamicDependencies) const;
 	virtual void ProcessActors(FPCGContext* Context, const UPCGDataFromActorSettings* Settings, const TArray<AActor*>& FoundActors) const;
 	virtual void ProcessActor(FPCGContext* Context, const UPCGDataFromActorSettings* Settings, AActor* FoundActor) const;
 
diff --git a/Engine/Plugins/PCG/Source/PCG/Public/Elements/PCGTypedGetter.h b/Engine/Plugins/PCG/Source/PCG/Public/Elements/PCGTypedGetter.h
--- a/Engine/Plugins/PCG/Source/PCG/Public/Elements/PCGTypedGetter.h
+++ b/Engine/Plugins/PCG/Source/PCG/Public/Elements/PCGTypedGetter.h
@@ -61,6 +61,7 @@ public:
 class FPCGGetLandscapeDataElement : public FPCGDataFromActorElement
 {
 protected:
+	virtual void ProcessActors(FPCGContext* Context, const UPCGDataFromActorSettings* Settings, const TArray<AActor*>& FoundActors, TArray<FPCGTaskId>& OutDynamicDependencies) const override;
 	virtual void ProcessActors(FPCGContext* Context, const UPCGDataFromActorSettings* Settings, const TArray<AActor*>& FoundActors) const override;
 	virtual void ProcessActor(FPCGContext* Context, const UPCGDataFromActorSettings* Settings, AActor* FoundActor) const override;
 };
