PCG: Expose default point data type that is currently set by cvar.

diff --git a/Engine/Plugins/PCG/Source/PCG/Private/PCGContext.cpp b/Engine/Plugins/PCG/Source/PCG/Private/PCGContext.cpp
--- a/Engine/Plugins/PCG/Source/PCG/Private/PCGContext.cpp
+++ b/Engine/Plugins/PCG/Source/PCG/Private/PCGContext.cpp
@@ -97,6 +97,12 @@ UPCGBasePointData* FPCGContext::NewPointData_AnyThread(FPCGContext* Context)
 	}
 }
 
+TSubclassOf<UPCGBasePointData> FPCGContext::GetDefaultPointDataClass()
+{
+	return CVarPCGEnablePointArrayData.GetValueOnAnyThread() ? UPCGPointArrayData::StaticClass()
+		: UPCGPointData::StaticClass();
+}
+
 FPCGTaskId FPCGContext::GetGraphExecutionTaskId() const
 {
 	const FPCGStack* StackPtr = GetStack();
diff --git a/Engine/Plugins/PCG/Source/PCG/Public/PCGContext.h b/Engine/Plugins/PCG/Source/PCG/Public/PCGContext.h
--- a/Engine/Plugins/PCG/Source/PCG/Public/PCGContext.h
+++ b/Engine/Plugins/PCG/Source/PCG/Public/PCGContext.h
@@ -6,6 +6,7 @@
 #include "PCGNode.h" // IWYU pragma: keep
 #include "PCGGraphExecutionStateInterface.h"
 #include "Helpers/PCGAsyncState.h"
+#include "Templates/SubclassOf.h"
 #include "Utils/PCGExtraCapture.h"
 
 #include "UObject/GCObject.h"
@@ -154,6 +155,8 @@ PRAGMA_ENABLE_DEPRECATION_WARNINGS
 
 	// Temporary so that we can toggle between UPCGPointData and UPCGPointArrayData with a CVar
 	static UPCGBasePointData* NewPointData_AnyThread(FPCGContext* Context);
+	// Like NewPointData_AnyThread, gets the class used for point data depending on the same CVar
+	static TSubclassOf<UPCGBasePointData> GetDefaultPointDataClass();
 
 	bool ContainsAsyncObject(const UObject* InAsyncObject) { return AsyncObjects.Contains(InAsyncObject); }
 
