[PCG] Remove pcg.EnableDynamicDependencies and make it the only code path (cleanup old code path)

#rnx

diff --git a/Engine/Plugins/PCG/Source/PCG/Private/Elements/PCGLoopElement.cpp b/Engine/Plugins/PCG/Source/PCG/Private/Elements/PCGLoopElement.cpp
--- a/Engine/Plugins/PCG/Source/PCG/Private/Elements/PCGLoopElement.cpp
+++ b/Engine/Plugins/PCG/Source/PCG/Private/Elements/PCGLoopElement.cpp
@@ -16,11 +16,6 @@
 
 #define LOCTEXT_NAMESPACE "PCGLoopElement"
 
-namespace PCGGraphExecutor
-{
-	extern TAutoConsoleVariable<bool> CVarEnableDynamicDependencies;
-}
-
 #if WITH_EDITOR
 FText UPCGLoopSettings::GetDefaultNodeTitle() const
 {
@@ -321,41 +316,13 @@ bool FPCGLoopElement::ExecuteInternal(FPCGContext* InContext) const
 			Context->bIsPaused = true;
 
 			// In the case of a feedback loop we only need to depend on the last scheduled task since all other tasks depend on each other
-			if (PCGGraphExecutor::CVarEnableDynamicDependencies.GetValueOnAnyThread())
+			if (bIsFeedbackLoop)
 			{
-				if (bIsFeedbackLoop)
-				{
-					Context->DynamicDependencies.Add(Context->SubgraphTaskIds.Last());
-				}
-				else
-				{
-					Context->DynamicDependencies.Append(Context->SubgraphTaskIds);
-				}
+				Context->DynamicDependencies.Add(Context->SubgraphTaskIds.Last());
 			}
 			else
 			{
-				Context->ScheduleGeneric(FPCGScheduleGenericParams(
-					[ContextHandle = Context->GetOrCreateHandle()](FPCGContext* NewContext) // Normal execution: Wake up the current task
-					{
-						FPCGContext::FSharedContext<FPCGSubgraphContext> SharedContext(ContextHandle);
-						if (FPCGSubgraphContext* ContextPtr = SharedContext.Get())
-						{
-							ContextPtr->bIsPaused = false;
-						}
-						return true;
-					},
-					[ContextHandle = Context->GetOrCreateHandle()](FPCGContext* NewContext) // On abort: wakeup and cancel, forget subgraphs
-					{
-						FPCGContext::FSharedContext<FPCGSubgraphContext> SharedContext(ContextHandle);
-						if (FPCGSubgraphContext* ContextPtr = SharedContext.Get())
-						{
-							ContextPtr->bIsPaused = false;
-							ContextPtr->SubgraphTaskIds.Reset();
-							ContextPtr->OutputData.bCancelExecution = true;
-						}
-					},
-					Context->ExecutionSource.Get(),
-					Context->SubgraphTaskIds));
+				Context->DynamicDependencies.Append(Context->SubgraphTaskIds);
 			}
 
 			return false;
diff --git a/Engine/Plugins/PCG/Source/PCG/Private/Elements/PCGSpawnActor.cpp b/Engine/Plugins/PCG/Source/PCG/Private/Elements/PCGSpawnActor.cpp
--- a/Engine/Plugins/PCG/Source/PCG/Private/Elements/PCGSpawnActor.cpp
+++ b/Engine/Plugins/PCG/Source/PCG/Private/Elements/PCGSpawnActor.cpp
@@ -38,11 +38,6 @@ static TAutoConsoleVariable<bool> CVarAllowActorReuse(
 	true,
 	TEXT("Controls whether PCG spawned actors can be reused and skipped when re-executing"));
 
-namespace PCGGraphExecutor
-{
-	extern TAutoConsoleVariable<bool> CVarEnableDynamicDependencies;
-}
-
 class FPCGSpawnActorPartitionByAttribute : public FPCGDataPartitionBase<FPCGSpawnActorPartitionByAttribute, TSubclassOf<AActor>>
 {
 public:
@@ -685,35 +680,7 @@ bool FPCGSpawnActorElement::SpawnAndPrepareSubgraphs(FPCGSubgraphContext* Contex
 	{
 		Context->bScheduledSubgraph = true;
 		Context->bIsPaused = true;
-
-		if (PCGGraphExecutor::CVarEnableDynamicDependencies.GetValueOnAnyThread())
-		{
-			Context->DynamicDependencies.Append(Context->SubgraphTaskIds);
-		}
-		else
-		{
-			Context->ScheduleGeneric(FPCGScheduleGenericParams(
-				[ContextHandle = Context->GetOrCreateHandle()](FPCGContext*) // Normal execution: Wake up the current task
-				{
-					FPCGContext::FSharedContext<FPCGSubgraphContext> SharedContext(ContextHandle);
-					if (FPCGSubgraphContext* ContextPtr = SharedContext.Get())
-					{
-						ContextPtr->bIsPaused = false;
-					}
-					return true;
-				},
-				[ContextHandle = Context->GetOrCreateHandle()](FPCGContext*) // On Abort: Wake up & cancel
-				{
-					FPCGContext::FSharedContext<FPCGSubgraphContext> SharedContext(ContextHandle);
-					if (FPCGSubgraphContext* ContextPtr = SharedContext.Get())
-					{
-						ContextPtr->bIsPaused = false;
-						ContextPtr->OutputData.bCancelExecution = true;
-					}
-				},
-				Context->ExecutionSource.Get(),
-				Context->SubgraphTaskIds));
-		}
+		Context->DynamicDependencies.Append(Context->SubgraphTaskIds);
 
 		return false;
 	}
diff --git a/Engine/Plugins/PCG/Source/PCG/Private/Elements/PCGSurfaceSampler.cpp b/Engine/Plugins/PCG/Source/PCG/Private/Elements/PCGSurfaceSampler.cpp
--- a/Engine/Plugins/PCG/Source/PCG/Private/Elements/PCGSurfaceSampler.cpp
+++ b/Engine/Plugins/PCG/Source/PCG/Private/Elements/PCGSurfaceSampler.cpp
@@ -22,11 +22,6 @@
 
 #define LOCTEXT_NAMESPACE "PCGSurfaceSamplerElement"
 
-namespace PCGGraphExecutor
-{
-	extern TAutoConsoleVariable<bool> CVarEnableDynamicDependencies;
-}
-
 namespace PCGSurfaceSampler
 {
 	void FSurfaceSamplerParams::InitializeFromSettings(const UPCGSurfaceSamplerSettings* Settings)
@@ -565,18 +560,15 @@ PRAGMA_ENABLE_DEPRECATION_WARNINGS
 				return EPCGTimeSliceInitResult::NoOperation;
 			}
 
-			if (PCGGraphExecutor::CVarEnableDynamicDependencies.GetValueOnAnyThread())
+			FBox PrepareSpatialQueryBounds = EffectiveGridBounds;
+			if (OutState.SamplerData.bNeedsLocalTransformation)
 			{
-				FBox PrepareSpatialQueryBounds = EffectiveGridBounds;
-				if (OutState.SamplerData.bNeedsLocalTransformation)
-				{
-					// Transform the pre-projected bounds around the surface's origin from local to world space
-					PrepareSpatialQueryBounds = EffectiveGridBounds.TransformBy(OutState.SamplerData.PreProjectionTransform);
-				}
-
-				PrepareTasks.Append(GeneratingShape->PrepareForSpatialQuery(Context, PrepareSpatialQueryBounds));
+				// Transform the pre-projected bounds around the surface's origin from local to world space
+				PrepareSpatialQueryBounds = EffectiveGridBounds.TransformBy(OutState.SamplerData.PreProjectionTransform);
 			}
 
+			PrepareTasks.Append(GeneratingShape->PrepareForSpatialQuery(Context, PrepareSpatialQueryBounds));
+
 			// Assigning this here prevents the need to root
 			Outputs[IterationIndex].Data = OutState.OutputData;
 			return EPCGTimeSliceInitResult::Success;
diff --git a/Engine/Plugins/PCG/Source/PCG/Private/Graph/PCGGraphExecutor.cpp b/Engine/Plugins/PCG/Source/PCG/Private/Graph/PCGGraphExecutor.cpp
--- a/Engine/Plugins/PCG/Source/PCG/Private/Graph/PCGGraphExecutor.cpp
+++ b/Engine/Plugins/PCG/Source/PCG/Private/Graph/PCGGraphExecutor.cpp
@@ -96,11 +96,6 @@ namespace PCGGraphExecutor
 		true,
 		TEXT("Whether proxies for GPU data are cached in per pin output data and passed through grid links. If false data is read back to CPU."));
 
-	TAutoConsoleVariable<bool> CVarEnableDynamicDependencies(
-		TEXT("pcg.EnableDynamicDependencies"),
-		true,
-		TEXT("Support re-scheduling of paused tasks through dynamic dependencies"));
-
 	void ClearAsyncFlags(TSet<TObjectPtr<UObject>>& AsyncObjects)
 	{
 		// Remove Async Flags from Objects created in Async PCG Tasks so that they can get tracked by the GarbageCollector properly
@@ -3159,18 +3154,7 @@ namespace PCGGraphExecutor
 				FPCGTaskId GenerationTask = ComponentWithData->GetGenerationTaskId();
 				if (ensure(GenerationTask != InvalidPCGTaskId))
 				{
-					if(PCGGraphExecutor::CVarEnableDynamicDependencies.GetValueOnAnyThread())
-					{
-						InContext->DynamicDependencies.Add(ComponentWithData->GetGenerationTaskId());
-					}
-					else
-					{
-						FPCGScheduleGenericParams Params(MoveTemp(WakeUpLambda), InContext->ExecutionSource.Get());
-						Params.bCanExecuteOnlyOnMainThread = false;
-						Params.ExecutionDependencies = { ComponentWithData->GetGenerationTaskId() };
-
-						InContext->ScheduleGeneric(Params);
-					}
+					InContext->DynamicDependencies.Add(ComponentWithData->GetGenerationTaskId());
 				}
 
 				return false;
@@ -3187,21 +3171,9 @@ namespace PCGGraphExecutor
 
 				// Wake up this task after graph has generated.
 				const FPCGTaskId GraphTaskId = ComponentWithData->GenerateLocalGetTaskId(GenTrigger, /*bForce=*/true);
-				
-				if (PCGGraphExecutor::CVarEnableDynamicDependencies.GetValueOnAnyThread())
-				{
-					InContext->DynamicDependencies.Add(GraphTaskId);
-				}
-				else
-				{
-					FPCGScheduleGenericParams Params(MoveTemp(WakeUpLambda), InContext->ExecutionSource.Get());
-					Params.bCanExecuteOnlyOnMainThread = false;
-					Params.ExecutionDependencies = { GraphTaskId };
-
-					Subsystem->ScheduleGeneric(Params);
-				}
 
 				// Update state and go to paused.
+				InContext->DynamicDependencies.Add(GraphTaskId);
 				InContext->bScheduledGraph = true;
 				InContext->bIsPaused = true;
 				return false;
diff --git a/Engine/Plugins/PCG/Source/PCG/Private/PCGSubgraph.cpp b/Engine/Plugins/PCG/Source/PCG/Private/PCGSubgraph.cpp
--- a/Engine/Plugins/PCG/Source/PCG/Private/PCGSubgraph.cpp
+++ b/Engine/Plugins/PCG/Source/PCG/Private/PCGSubgraph.cpp
@@ -18,11 +18,6 @@
 
 #define LOCTEXT_NAMESPACE "PCGSubgraphElement"
 
-namespace PCGGraphExecutor
-{
-	extern TAutoConsoleVariable<bool> CVarEnableDynamicDependencies;
-}
-
 namespace PCGSubgraphSettings
 {
 	void RemoveAdvancedAndInvisibleOnConnectedPins(const UPCGGraph* Subgraph, TArray<FPCGPinProperties>& InOutPinProperties, const bool bIsInput)
@@ -794,38 +789,7 @@ bool FPCGSubgraphElement::ExecuteInternal(FPCGContext* InContext) const
 					Context->SubgraphTaskIds.Add(SubgraphTaskId);
 					Context->bScheduledSubgraph = true;
 					Context->bIsPaused = true;
-
-					if (PCGGraphExecutor::CVarEnableDynamicDependencies.GetValueOnAnyThread())
-					{
-						Context->DynamicDependencies.Add(SubgraphTaskId);
-					}
-					else
-					{
-						// add a trivial task after the output task that wakes up this task
-						Context->ScheduleGeneric(FPCGScheduleGenericParams(
-							[ContextHandle = Context->GetOrCreateHandle()](FPCGContext* NewContext) // Normal execution: Wake up the current task
-							{
-								FPCGContext::FSharedContext<FPCGSubgraphContext> SharedContext(ContextHandle);
-								if (FPCGSubgraphContext* ContextPtr = SharedContext.Get())
-								{
-									ContextPtr->bIsPaused = false;
-								}
-								return true;
-							},
-							[ContextHandle = Context->GetOrCreateHandle()](FPCGContext* NewContext) // On Abort: Wake up and cancel the execution
-							{
-								FPCGContext::FSharedContext<FPCGSubgraphContext> SharedContext(ContextHandle);
-								if (FPCGSubgraphContext* ContextPtr = SharedContext.Get())
-								{
-									ContextPtr->bIsPaused = false;
-									ContextPtr->OutputData.bCancelExecution = true;
-									ContextPtr->SubgraphTaskIds.Reset();
-								}
-							},
-							Context->ExecutionSource.Get(),
-							Context->SubgraphTaskIds));
-					}
-
+					Context->DynamicDependencies.Add(SubgraphTaskId);
 					return false;
 				}
 				else
