[PCG] Fix Compute element hang (paused forever) because of 0 instance provider

#rnx

diff --git a/Engine/Plugins/PCG/Source/PCG/Private/Compute/DataInterfaces/PCGInstanceDataInterface.cpp b/Engine/Plugins/PCG/Source/PCG/Private/Compute/DataInterfaces/PCGInstanceDataInterface.cpp
--- a/Engine/Plugins/PCG/Source/PCG/Private/Compute/DataInterfaces/PCGInstanceDataInterface.cpp
+++ b/Engine/Plugins/PCG/Source/PCG/Private/Compute/DataInterfaces/PCGInstanceDataInterface.cpp
@@ -222,6 +222,11 @@ FPCGInstanceDataProviderProxy::FPCGInstanceDataProviderProxy(TArray<FPCGPrimitiv
 	if (ensure(InDataProvider.IsValid()))
 	{
 		DataProviderGeneration = InDataProvider->GenerationCounter;
+
+		if (NumInstancesAllPrimitives == 0)
+		{
+			InDataProvider->bWroteInstances = true;
+		}
 	}
 	else
 	{
