[PCG] Improve Custom HLSL code comprehension
* Add PCGHLSLSyntaxTokenizer to properly parse HLSL source code. Now we are more robust in our matching and also ignore would-be matches in comments/strings/preprocessor statements.
* Only parse source during kernel initialization, instead of in standalone builds
* Move some data from Custom HLSL Settings to Custom HLSL Kernel (shader source and attribute keys)
* TODO in a follow up: Properly use our parsed matches for find/replace, rather than relying on a combination of regex and string operations

#rnx

diff --git a/Engine/Plugins/PCG/Source/PCG/Internal/Compute/PCGComputeCommon.h b/Engine/Plugins/PCG/Source/PCG/Internal/Compute/PCGComputeCommon.h
--- a/Engine/Plugins/PCG/Source/PCG/Internal/Compute/PCGComputeCommon.h
+++ b/Engine/Plugins/PCG/Source/PCG/Internal/Compute/PCGComputeCommon.h
@@ -190,3 +190,23 @@ struct FPCGKernelEdge
 	FPCGPinReference UpstreamPin;
 	FPCGPinReference DownstreamPin;
 };
+
+/** Helper struct for serializing data labels. */
+USTRUCT()
+struct FPCGDataLabels
+{
+	GENERATED_BODY()
+
+	UPROPERTY()
+	TArray<FString> Labels;
+};
+
+/** Helper struct for serializing map of pin name to data labels. */
+USTRUCT()
+struct FPCGPinDataLabels
+{
+	GENERATED_BODY()
+
+	UPROPERTY()
+	TMap</*PinLabel*/FName, FPCGDataLabels> PinToDataLabels;
+};
diff --git a/Engine/Plugins/PCG/Source/PCG/Private/Compute/BuiltInKernels/PCGCountUniqueAttributeValuesKernel.cpp b/Engine/Plugins/PCG/Source/PCG/Private/Compute/BuiltInKernels/PCGCountUniqueAttributeValuesKernel.cpp
--- a/Engine/Plugins/PCG/Source/PCG/Private/Compute/BuiltInKernels/PCGCountUniqueAttributeValuesKernel.cpp
+++ b/Engine/Plugins/PCG/Source/PCG/Private/Compute/BuiltInKernels/PCGCountUniqueAttributeValuesKernel.cpp
@@ -21,11 +21,11 @@
 #define LOCTEXT_NAMESPACE "PCGCountUniqueAttributeValuesKernel"
 
 #if WITH_EDITOR
-void UPCGCountUniqueAttributeValuesKernel::Initialize(const UPCGSettings* InSettings, bool bLogDescriptions)
+void UPCGCountUniqueAttributeValuesKernel::InitializeInternal()
 {
-	Super::Initialize(InSettings, bLogDescriptions);
+	Super::InitializeInternal();
 
-	const UPCGStaticMeshSpawnerSettings* SMSettings = CastChecked<UPCGStaticMeshSpawnerSettings>(InSettings);
+	const UPCGStaticMeshSpawnerSettings* SMSettings = CastChecked<UPCGStaticMeshSpawnerSettings>(Settings);
 	const UPCGMeshSelectorByAttribute* Selector = CastChecked<UPCGMeshSelectorByAttribute>(SMSettings->MeshSelectorParameters);
 	AttributeName = Selector->AttributeName;
 }
diff --git a/Engine/Plugins/PCG/Source/PCG/Private/Compute/BuiltInKernels/PCGCountUniqueAttributeValuesKernel.h b/Engine/Plugins/PCG/Source/PCG/Private/Compute/BuiltInKernels/PCGCountUniqueAttributeValuesKernel.h
--- a/Engine/Plugins/PCG/Source/PCG/Private/Compute/BuiltInKernels/PCGCountUniqueAttributeValuesKernel.h
+++ b/Engine/Plugins/PCG/Source/PCG/Private/Compute/BuiltInKernels/PCGCountUniqueAttributeValuesKernel.h
@@ -18,10 +18,6 @@ class UPCGCountUniqueAttributeValuesKernel : public UPCGComputeKernel
 	GENERATED_BODY()
 
 public:
-#if WITH_EDITOR
-	virtual void Initialize(const UPCGSettings* InSettings, bool bLogDescriptions) override;
-#endif
-
 	virtual bool IsKernelDataValid(FPCGContext* InContext) const override;
 	virtual bool ComputeOutputBindingDataDesc(const UPCGComputeGraph* InGraph, FName InOutputPinLabel, UPCGDataBinding* InBinding, FPCGDataCollectionDesc& OutDataDesc) const override;
 	virtual int ComputeThreadCount(const UPCGDataBinding* Binding) const override;
@@ -38,6 +34,11 @@ public:
 	virtual void GetInputPins(TArray<FPCGPinProperties>& OutPins) const override;
 	virtual void GetOutputPins(TArray<FPCGPinPropertiesGPU>& OutPins) const override;
 
+#if WITH_EDITOR
+protected:
+	virtual void InitializeInternal() override;
+#endif
+
 public:
 	UPROPERTY()
 	FName AttributeName;
diff --git a/Engine/Plugins/PCG/Source/PCG/Private/Compute/Elements/PCGCustomHLSL.cpp b/Engine/Plugins/PCG/Source/PCG/Private/Compute/Elements/PCGCustomHLSL.cpp
--- a/Engine/Plugins/PCG/Source/PCG/Private/Compute/Elements/PCGCustomHLSL.cpp
+++ b/Engine/Plugins/PCG/Source/PCG/Private/Compute/Elements/PCGCustomHLSL.cpp
@@ -23,11 +23,17 @@
 
 #define LOCTEXT_NAMESPACE "PCGCustomHLSLElement"
 
+namespace PCGHLSLElement
+{
+	const FString PinDeclTemplateStr = TEXT("{pin}");
+}
+
 UPCGCustomHLSLSettings::UPCGCustomHLSLSettings()
 {
 	bExecuteOnGPU = true;
 }
 
+#if WITH_EDITOR
 void UPCGCustomHLSLSettings::PostLoad()
 {
 	Super::PostLoad();
@@ -35,10 +41,7 @@ void UPCGCustomHLSLSettings::PostLoad()
 	// Note: We update here so that Custom HLSL nodes will have the correct pin settings & declarations on load.
 	UpdatePinSettings();
 	UpdateAttributeKeys();
-
-#if WITH_EDITOR
 	UpdateDeclarations();
-#endif
 }
 
 void UPCGCustomHLSLSettings::PostInitProperties()
@@ -47,11 +50,10 @@ void UPCGCustomHLSLSettings::PostInitProperties()
 
 	// Note: We update here so that Custom HLSL nodes will have the correct pin settings & declarations on creation.
 	UpdatePinSettings();
-
-#if WITH_EDITOR
+	UpdateAttributeKeys();
 	UpdateDeclarations();
-#endif
 }
+#endif
 
 TArray<FPCGPinProperties> UPCGCustomHLSLSettings::OutputPinProperties() const
 {
@@ -81,6 +83,7 @@ void UPCGCustomHLSLSettings::PostEditChangeProperty(FPropertyChangedEvent& Prope
 {
 	// Apply any pin setup before refreshing the node.
 	UpdatePinSettings();
+	UpdateAttributeKeys();
 
 	const FName MemberProperty = PropertyChangedEvent.MemberProperty ? PropertyChangedEvent.MemberProperty->GetFName() : NAME_None;
 	const FName Property = PropertyChangedEvent.Property ? PropertyChangedEvent.Property->GetFName() : NAME_None;
@@ -130,7 +133,6 @@ void UPCGCustomHLSLSettings::PostEditChangeProperty(FPropertyChangedEvent& Prope
 	Super::PostEditChangeProperty(PropertyChangedEvent);
 
 	UpdateDeclarations();
-	UpdateAttributeKeys();
 }
 #endif
 
@@ -808,7 +810,6 @@ void UPCGCustomHLSLSettings::UpdateHelperDeclarations()
 
 	HelperDeclarations.TrimStartAndEndInline();
 }
-#endif // WITH_EDITOR
 
 void UPCGCustomHLSLSettings::UpdatePinSettings()
 {
@@ -829,24 +830,18 @@ void UPCGCustomHLSLSettings::UpdatePinSettings()
 			// compute kernel.
 			Properties.bAllowMultipleData = false;
 
-#if WITH_EDITOR
 			Properties.bAllowEditMultipleData = false;
-#endif
 		}
 		else
 		{
-#if WITH_EDITOR
 			Properties.bAllowEditMultipleData = true;
-#endif
 		}
 
 		// TODO: We have work to do to allow dynamic merging of data. Also we will likely inject Gather
 		// nodes on the CPU side so that merging is handled CPU side where possible.
 		Properties.SetAllowMultipleConnections(false);
 
-#if WITH_EDITOR
 		Properties.bAllowEditMultipleConnections = false;
-#endif
 	}
 
 	// Setup output pins.
@@ -860,7 +855,6 @@ void UPCGCustomHLSLSettings::UpdatePinSettings()
 			Properties.AllowedTypes = EPCGDataType::Point;
 		}
 
-#if WITH_EDITOR
 		// Only allow editing the initialization mode if it's not driven by the kernel type.
 		const bool bInitModeDrivenByKernel = PinIndex == 0 && (KernelType == EPCGKernelType::PointProcessor || KernelType == EPCGKernelType::PointGenerator);
 		Properties.PropertiesGPU.bAllowEditInitMode = !bInitModeDrivenByKernel;
@@ -869,21 +863,15 @@ void UPCGCustomHLSLSettings::UpdatePinSettings()
 		// Output pins should always allow multiple connections.
 		// TODO this could be hoisted up somewhere in the future.
 		Properties.bAllowEditMultipleConnections = false;
-#endif
 
-#if WITH_EDITOR
 		Properties.bAllowEditMultipleData = true;
 		Properties.PropertiesGPU.bAllowEditDataCount = true;
-#endif
 	}
 }
 
 void UPCGCustomHLSLSettings::UpdateAttributeKeys()
 {
 	// Make sure all the output attributes are up-to-date on their identifiers.
-	// Need to modify it if it ever changes for undo/redo support.
-	
-	Modify(/*bAlwaysMarkDirty=*/false);
 	bool bMarkedDirty = false;
 	for (FPCGPinPropertiesGPU& OutputPin : OutputPins)
 	{
@@ -896,86 +884,8 @@ void UPCGCustomHLSLSettings::UpdateAttributeKeys()
 			}
 		}
 	}
-	
-	KernelAttributeKeys.Reset();
-	PinToAttributeKeys.Reset();
-
-	const FString Source = ShaderFunctions + ShaderSource;
-	FRegexMatcher ModuleMatcher(FRegexPattern(PCGHLSLElement::AttributeFunctionPattern), Source);
-
-	while (ModuleMatcher.FindNext())
-	{
-		const FString PinStr = ModuleMatcher.GetCaptureGroup(PCGHLSLElement::AttributePinCaptureGroup);
-		const FString FuncStr = ModuleMatcher.GetCaptureGroup(PCGHLSLElement::AttributeFunctionCaptureGroup);
-		const FString TypeStr = ModuleMatcher.GetCaptureGroup(PCGHLSLElement::AttributeTypeCaptureGroup);
-		const FString NameStr = ModuleMatcher.GetCaptureGroup(PCGHLSLElement::AttributeNameCaptureGroup);
-
-		const int LineStartIndex = ModuleMatcher.GetMatchBeginning();
-		int CurrentSourceIndex = Source.Find("\n", ESearchCase::IgnoreCase, ESearchDir::FromStart, 0);
-		int LineNumber = 0;
-
-		while (CurrentSourceIndex < LineStartIndex && CurrentSourceIndex != INDEX_NONE)
-		{
-			++LineNumber;
-			CurrentSourceIndex = Source.Find("\n", ESearchCase::IgnoreCase, ESearchDir::FromStart, CurrentSourceIndex + 1);
-		}
-
-		if (PinStr.IsEmpty() || FuncStr.IsEmpty() || TypeStr.IsEmpty() || NameStr.IsEmpty())
-		{
-			UE_LOG(LogPCG, Error, TEXT("Invalid attribute usage in shader source, line %d."), LineNumber);
-			continue;
-		}
-
-		const UEnum* AttributeTypeEnum = StaticEnum<EPCGKernelAttributeType>();
-		check(AttributeTypeEnum);
-
-		const int64 AttributeType = AttributeTypeEnum->GetValueByName(FName(*TypeStr));
-
-		if (AttributeType == INDEX_NONE)
-		{
-			UE_LOG(LogPCG, Error, TEXT("Invalid attribute type in shader source, line %d."), LineNumber);
-			continue;
-		}
-
-		// Add the attribute if it hasn't already been referenced.
-		FPCGKernelAttributeKey Key;
-		Key.Type = static_cast<EPCGKernelAttributeType>(AttributeType);
-		Key.SetSelector(FPCGAttributePropertySelector::CreateSelectorFromString(NameStr));
-		
-		if (Key.IsValid())
-		{
-			KernelAttributeKeys.AddUnique(Key);
-
-			// Add an entry mapping this pin to the referenced attribute, if the entry doesn't already exist.
-			TArray<TTuple<FPCGKernelAttributeKey, bool>>& Keys = PinToAttributeKeys.FindOrAdd(FName(*PinStr));
-			Keys.AddUnique(MakeTuple(Key, false));
-		}
-	}
-
-	// Process each output pin for any new attributes they want to create.
-	for (const FPCGPinPropertiesGPU& OutputPin : OutputPins)
-	{
-		for (const FPCGKernelAttributeKey& Key : OutputPin.PropertiesGPU.CreatedKernelAttributeKeys)
-		{
-			KernelAttributeKeys.AddUnique(Key);
-
-			TArray<TTuple<FPCGKernelAttributeKey, bool>>& Keys = PinToAttributeKeys.FindOrAdd(OutputPin.Label);
-			TTuple<FPCGKernelAttributeKey, bool>* Pair = Keys.FindByPredicate([Key](const TTuple<FPCGKernelAttributeKey, bool>& Pair) { return Pair.Key == Key; });
-
-			// Mark as created on GPU
-			if (Pair)
-			{
-				Pair->Value = true;
-			}
-			else
-			{
-				Keys.Add(MakeTuple(Key, /*bCreatedOnGPU=*/true));
-			}
-		}
-	}
 }
 
-#if WITH_EDITOR
 TArray<FName> UPCGCustomHLSLSettings::GetInputPinNames() const
 {
 	TArray<FName> PinNames;
diff --git a/Engine/Plugins/PCG/Source/PCG/Private/Compute/Elements/PCGCustomHLSL.h b/Engine/Plugins/PCG/Source/PCG/Private/Compute/Elements/PCGCustomHLSL.h
--- a/Engine/Plugins/PCG/Source/PCG/Private/Compute/Elements/PCGCustomHLSL.h
+++ b/Engine/Plugins/PCG/Source/PCG/Private/Compute/Elements/PCGCustomHLSL.h
@@ -41,10 +41,12 @@ class UPCGCustomHLSLSettings
 public:
 	UPCGCustomHLSLSettings();
 
+#if WITH_EDITOR
 	//~Begin UObject interface
 	virtual void PostLoad() override;
 	virtual void PostInitProperties() override;
 	//~End UObject interface
+#endif
 
 	//~Begin UPCGSettings interface
 	virtual TArray<FPCGPinProperties> InputPinProperties() const override { return InputPins; }
@@ -94,20 +96,16 @@ protected:
 	/** Gets the GPU pin properties for the output pin with the given label. */
 	const FPCGPinPropertiesGPU* GetOutputPinPropertiesGPU(const FName& InPinLabel) const;
 
-	/** Enforce required pin settings and set display toggles to drive UI. */
-	void UpdatePinSettings();
-
-	/** Update the kernel keys read/written/created by this node. In a pure sense this could register only the created attributes and instead let
-	* read/written attributes get registered by upstream nodes or initialized from input data, but this way we can statically know the attribute IDs
-	* and we can avoid resolving them at execution time. */
-	void UpdateAttributeKeys();
-
 #if WITH_EDITOR
 	void UpdateDeclarations();
 	void UpdateInputDeclarations();
 	void UpdateOutputDeclarations();
 	void UpdateHelperDeclarations();
 
+	/** Enforce required pin settings and set display toggles to drive UI. */
+	void UpdatePinSettings();
+	void UpdateAttributeKeys();
+
 	/** List of all non-advanced input pin names. */
 	UFUNCTION()
 	TArray<FName> GetInputPinNames() const;
@@ -163,6 +161,9 @@ protected:
 	UPROPERTY(EditAnywhere, Category = "Settings")
 	bool bMuteUnwrittenPinDataErrors = false;
 
+#if WITH_EDITORONLY_DATA
+	// Shader source and declarations are entirely editor-only, and should never be serialized outside of the editor.
+
 	/** Optional functions that can be called from the source. Intended to be edited using the Node Source Editor window. */
 	UPROPERTY(EditAnywhere, Category = "Settings", AdvancedDisplay, meta = (MultiLine = true))
 	FString ShaderFunctions = "/** CUSTOM SHADER FUNCTIONS **/\n";
@@ -182,12 +183,7 @@ protected:
 	/** Helper data and functions that can be used from the shader code. Intended to be viewed using the Node Source Editor window. */
 	UPROPERTY(Transient, VisibleAnywhere, Category = "Settings", AdvancedDisplay, meta = (MultiLine = true))
 	FString HelperDeclarations;
-
-	/** Attributes statically detected as being read, written, or created by this node. */
-	TArray<FPCGKernelAttributeKey> KernelAttributeKeys;
-
-	/** Maps pins to their attribute keys and whether or not they were created on the GPU. */
-	TMap<FName, TArray<TTuple<FPCGKernelAttributeKey, bool /*bCreatedOnGPU*/>>> PinToAttributeKeys;
+#endif
 
 	friend class UPCGCustomHLSLKernel;
 };
@@ -198,23 +194,3 @@ protected:
 	// This will only be called if the custom HLSL node is not set up correctly (valid nodes are replaced with a compute graph element).
 	virtual bool ExecuteInternal(FPCGContext* Context) const override;
 };
-
-namespace PCGHLSLElement
-{
-	const FString PinDeclTemplateStr = TEXT("{pin}");
-
-	/** First capture: Pin name (supports a - z, A - Z, and 0 - 9) */
-	constexpr int AttributePinCaptureGroup = 1;
-
-	/** Second capture: Function name (Get or Set) */
-	constexpr int AttributeFunctionCaptureGroup = 2;
-
-	/** Third capture: Attribute type (e.g. Int, Float, Rotator, etc.) */
-	constexpr int AttributeTypeCaptureGroup = 3;
-
-	/** Fourth capture: Attribute name (supports a-z, A-Z, 0-9, ' ', '-', '_', and '/') */
-	constexpr int AttributeNameCaptureGroup = 4;
-
-	/** Regex pattern used to detect and parse attribute function usage in kernels. */
-	constexpr TCHAR AttributeFunctionPattern[] = { TEXT("([a-zA-Z0-9]+)_(Get|Set)(.*?)[\\s]*?\\(.*?,.*?,.*?'([a-zA-Z0-9 -_\\/]*?)'") };
-}
diff --git a/Engine/Plugins/PCG/Source/PCG/Private/Compute/Elements/PCGCustomHLSLKernel.cpp b/Engine/Plugins/PCG/Source/PCG/Private/Compute/Elements/PCGCustomHLSLKernel.cpp
--- a/Engine/Plugins/PCG/Source/PCG/Private/Compute/Elements/PCGCustomHLSLKernel.cpp
+++ b/Engine/Plugins/PCG/Source/PCG/Private/Compute/Elements/PCGCustomHLSLKernel.cpp
@@ -13,14 +13,31 @@
 #include "Containers/StaticArray.h"
 #include "Internationalization/Regex.h"
 
+#if WITH_EDITOR
+#include "Editor/IPCGEditorModule.h"
+#include "Framework/Text/SyntaxTokenizer.h"
+#endif
+
 #include UE_INLINE_GENERATED_CPP_BY_NAME(PCGCustomHLSLKernel)
 
 #define LOCTEXT_NAMESPACE "PCGCustomHLSLKernel"
 
 namespace PCGCustomHLSLKernel
 {
-	constexpr TCHAR AttributeFunctionGet[] = { TEXT("Get") };
-	constexpr TCHAR AttributeFunctionSet[] = { TEXT("Set") };
+	constexpr TCHAR AttributeFunctionGetKeyword[] = { TEXT("Get") };
+	constexpr TCHAR AttributeFunctionSetKeyword[] = { TEXT("Set") };
+	constexpr TCHAR CopyElementFunctionKeyword[] = { TEXT("CopyElementFrom") };
+	constexpr TCHAR InitializeFunctionKeyword[] = { TEXT("Initialize") };
+
+#if WITH_EDITOR
+	enum class EParseState : uint8
+	{
+		None,
+		LookingForDoubleQuotedString,
+		LookingForSingleQuotedString,
+		LookingForSingleLineComment,
+		LookingForMultiLineComment,
+	};
 
 	FString GetDataTypeString(EPCGDataType Type)
 	{
@@ -29,69 +46,156 @@ namespace PCGCustomHLSLKernel
 
 		return DataTypeEnum->GetValueOrBitfieldAsString(static_cast<int64>(Type));
 	}
-}
 
-#if WITH_EDITOR
-void UPCGCustomHLSLKernel::Initialize(const UPCGSettings* InSettings, bool bLogDescriptions)
-{
-	Super::Initialize(InSettings, bLogDescriptions);
-
-	InitEntryPoint();
-
-	const UPCGCustomHLSLSettings* CustomHLSLSettings = CastChecked<UPCGCustomHLSLSettings>(Settings);
+	// Reference implementation in FHLSLSyntaxHighlighterMarshaller::ProcessTokenizedLine()
+	void ProcessTokenizedLine(const FString& InSourceString, const ISyntaxTokenizer::FTokenizedLine& InTokenizedLine, EParseState& InOutParseState, TArray<FPCGCustomHLSLParsedSource::FToken>& OutTokens)
+	{
+		for(const ISyntaxTokenizer::FToken& Token : InTokenizedLine.Tokens)
+		{
+			FPCGCustomHLSLParsedSource::FToken& Run = OutTokens.Emplace_GetRef();
+			Run.Range = Token.Range;
+	
+			const FString TokenText = InSourceString.Mid(Token.Range.BeginIndex, Token.Range.Len());
+			const bool bIsWhitespace = TokenText.TrimEnd().IsEmpty();
 
-	FRegexMatcher FunctionsModuleMatcher(FRegexPattern(PCGHLSLElement::AttributeFunctionPattern), CustomHLSLSettings->ShaderFunctions);
+			if(!bIsWhitespace)
+			{
+				bool bHasMatchedSyntax = false;
+				if(Token.Type == ISyntaxTokenizer::ETokenType::Syntax)
+				{
+					if(InOutParseState == EParseState::None && TokenText == TEXT("\""))
+					{
+						Run.Type = FPCGCustomHLSLParsedSource::ETokenType::DoubleQuotedString;
+						InOutParseState = EParseState::LookingForDoubleQuotedString;
+						bHasMatchedSyntax = true;
+					}
+					else if(InOutParseState == EParseState::LookingForDoubleQuotedString && TokenText == TEXT("\""))
+					{
+						Run.Type = FPCGCustomHLSLParsedSource::ETokenType::Normal;
+						InOutParseState = EParseState::None;
+					}
+					else if(InOutParseState == EParseState::None && TokenText == TEXT("\'"))
+					{
+						Run.Type = FPCGCustomHLSLParsedSource::ETokenType::SingleQuotedString;
+						InOutParseState = EParseState::LookingForSingleQuotedString;
+						bHasMatchedSyntax = true;
+					}
+					else if(InOutParseState == EParseState::LookingForSingleQuotedString && TokenText == TEXT("\'"))
+					{
+						Run.Type = FPCGCustomHLSLParsedSource::ETokenType::Normal;
+						InOutParseState = EParseState::None;
+					}
+					else if(InOutParseState == EParseState::None && TokenText.StartsWith(TEXT("#")))
+					{
+						Run.Type = FPCGCustomHLSLParsedSource::ETokenType::PreProcessorKeyword;
+						InOutParseState = EParseState::None;
+					}
+					else if(InOutParseState == EParseState::None && TokenText == TEXT("//"))
+					{
+						Run.Type = FPCGCustomHLSLParsedSource::ETokenType::Comment;
+						InOutParseState = EParseState::LookingForSingleLineComment;
+					}
+					else if(InOutParseState == EParseState::None && TokenText == TEXT("/*"))
+					{
+						Run.Type = FPCGCustomHLSLParsedSource::ETokenType::Comment;
+						InOutParseState = EParseState::LookingForMultiLineComment;
+					}
+					else if(InOutParseState == EParseState::LookingForMultiLineComment && TokenText == TEXT("*/"))
+					{
+						Run.Type = FPCGCustomHLSLParsedSource::ETokenType::Comment;
+						InOutParseState = EParseState::None;
+					}
+					else if(InOutParseState == EParseState::None && TChar<TCHAR>::IsIdentifier(TokenText[0]))
+					{
+						Run.Type = FPCGCustomHLSLParsedSource::ETokenType::Keyword;
+						InOutParseState = EParseState::None;
+					}
+					else if(InOutParseState == EParseState::None && !TChar<TCHAR>::IsIdentifier(TokenText[0]))
+					{
+						Run.Type = FPCGCustomHLSLParsedSource::ETokenType::Operator;
+						InOutParseState = EParseState::None;
+					}
+				}
+				
+				// It's possible that we fail to match a syntax token if we're in a state where it isn't parsed
+				// In this case, we treat it as a literal token
+				if(Token.Type == ISyntaxTokenizer::ETokenType::Literal || !bHasMatchedSyntax)
+				{
+					if(InOutParseState == EParseState::LookingForDoubleQuotedString)
+					{
+						Run.Type = FPCGCustomHLSLParsedSource::ETokenType::DoubleQuotedString;
+					}
+					else if(InOutParseState == EParseState::LookingForSingleQuotedString)
+					{
+						Run.Type = FPCGCustomHLSLParsedSource::ETokenType::SingleQuotedString;
+					}
+					else if(InOutParseState == EParseState::LookingForSingleLineComment)
+					{
+						Run.Type = FPCGCustomHLSLParsedSource::ETokenType::Comment;
+					}
+					else if(InOutParseState == EParseState::LookingForMultiLineComment)
+					{
+						Run.Type = FPCGCustomHLSLParsedSource::ETokenType::Comment;
+					}
+				}
+			}
+			else
+			{
+				Run.Type = FPCGCustomHLSLParsedSource::ETokenType::Whitespace;
+			}
+		}
+	
+		if (InOutParseState != EParseState::LookingForMultiLineComment)
+		{
+			InOutParseState = EParseState::None;
+		}
+	}
 
-	auto ParseAttributeFunctions = [](const FString& SourceToParse, TArray<FPCGParsedAttributeFunction>& OutParsedAttributeFunctions)
+	void ParseTokens(const FString& InSourceString, const TArray<ISyntaxTokenizer::FTokenizedLine>& InTokenizedLines, TArray<FPCGCustomHLSLParsedSource::FToken>& OutTokens)
 	{
-		const UEnum* AttributeTypeEnum = StaticEnum<EPCGKernelAttributeType>();
-		check(AttributeTypeEnum);
-
-		FRegexMatcher ModuleMatcher(FRegexPattern(PCGHLSLElement::AttributeFunctionPattern), SourceToParse);
+		EParseState ParseState = EParseState::None;
 
-		while (ModuleMatcher.FindNext())
+		for(const ISyntaxTokenizer::FTokenizedLine& TokenizedLine : InTokenizedLines)
 		{
-			const int64 AttributeType = AttributeTypeEnum->GetValueByNameString(ModuleMatcher.GetCaptureGroup(PCGHLSLElement::AttributeTypeCaptureGroup));
-
-			OutParsedAttributeFunctions.Emplace(
-				ModuleMatcher.GetCaptureGroup(PCGHLSLElement::AttributePinCaptureGroup),
-				ModuleMatcher.GetCaptureGroup(PCGHLSLElement::AttributeFunctionCaptureGroup),
-				AttributeType,
-				ModuleMatcher.GetCaptureGroup(PCGHLSLElement::AttributeNameCaptureGroup),
-				ModuleMatcher.GetMatchBeginning());
+			ProcessTokenizedLine(InSourceString, TokenizedLine, ParseState, OutTokens);
 		}
-	};
+	}
+#endif
+}
 
-	ParseAttributeFunctions(CustomHLSLSettings->ShaderSource, SourceParsedAttributeFunctions);
-	ParseAttributeFunctions(CustomHLSLSettings->ShaderFunctions, FunctionsParsedAttributeFunctions);
+#if WITH_EDITOR
+void UPCGCustomHLSLKernel::InitializeInternal()
+{
+	Super::InitializeInternal();
 
-	auto ParseCopyElementFunctions = [](const FString& SourceToParse, TArray<FPCGParsedCopyElementFunction>& OutParsedCopyElementFunctions)
-	{
-		/** Regex for parsing usages of CopyElementFrom function. */
-		constexpr TCHAR CopyElementPattern[] = { TEXT(/*TargetPin*/"([a-zA-Z0-9]+)" "_CopyElementFrom_" /*SourcePin*/"([a-zA-Z0-9]+)\\s*\\(") };
-		constexpr int CopyElementTargetPinCaptureGroup = 1;
-		constexpr int CopyElementSourcePinCaptureGroup = 2;
+	InitEntryPoint();
+	ParseShaderSource();
 
-		FRegexMatcher ModuleMatcher(FRegexPattern(CopyElementPattern), SourceToParse);
+	const UPCGCustomHLSLSettings* CustomHLSLSettings = CastChecked<UPCGCustomHLSLSettings>(GetSettings());
 
-		while (ModuleMatcher.FindNext())
+	// Process each output pin for any new attributes they want to create.
+	for (const FPCGPinPropertiesGPU& OutputPin : CustomHLSLSettings->OutputPins)
+	{
+		for (const FPCGKernelAttributeKey& AuthoredKey : OutputPin.PropertiesGPU.CreatedKernelAttributeKeys)
 		{
-			OutParsedCopyElementFunctions.Emplace(
-				ModuleMatcher.GetCaptureGroup(CopyElementSourcePinCaptureGroup),
-				ModuleMatcher.GetCaptureGroup(CopyElementTargetPinCaptureGroup));
-		}
-	};
+			const uint32 AttributeKeyIndex = KernelAttributeKeys.AddUnique(AuthoredKey);
+			TArray<FPCGCustomHLSLAttributeInfo>& Infos = PinToAttributeInfos.FindOrAdd(OutputPin.Label).Infos;
+			FPCGCustomHLSLAttributeInfo* Info = Infos.FindByPredicate([AttributeKeyIndex](const FPCGCustomHLSLAttributeInfo& Info) { return Info.Index == AttributeKeyIndex; });
 
-	ParseCopyElementFunctions(CustomHLSLSettings->ShaderSource, SourceParsedCopyElementFunctions);
-	ParseCopyElementFunctions(CustomHLSLSettings->ShaderFunctions, FunctionsParsedCopyElementFunctions);
+			// Mark as created on GPU.
+			if (Info)
+			{
+				Info->bCreatedOnGPU = true;
+			}
+			else
+			{
+				Infos.Emplace(AttributeKeyIndex, /*bCreatedOnGPU=*/true);
+			}
+		}
+	}
 }
 #endif // WITH_EDITOR
 
-const UPCGSettings* UPCGCustomHLSLKernel::GetSettings() const
-{
-	return Settings;
-}
-
 bool UPCGCustomHLSLKernel::IsKernelDataValid(FPCGContext* InContext) const
 {
 	TRACE_CPUPROFILER_EVENT_SCOPE(UPCGCustomHLSLKernel::IsKernelDataValid);
@@ -168,26 +272,24 @@ bool UPCGCustomHLSLKernel::ComputeOutputBindingDataDesc(const UPCGComputeGraph*
 		ComputeDataDescFromPinProperties(*OutputPinProperties, MakeArrayView(CustomHLSLSettings->InputPins), InBinding, OutputDataDesc);
 	}
 
-	for (const FPCGKernelAttributeKey& AttributeKey : CustomHLSLSettings->KernelAttributeKeys)
+	// Add attributes that will be created for this pin on the GPU.
+	if (const FPCGCustomHLSLAttributeInfos* Infos = PinToAttributeInfos.Find(InOutputPinLabel))
 	{
-		// Add attributes that will be created for this pin on the GPU.
-		if (const TArray<TTuple<FPCGKernelAttributeKey, bool>>* Keys = CustomHLSLSettings->PinToAttributeKeys.Find(InOutputPinLabel))
+		for (const FPCGCustomHLSLAttributeInfo& Info : Infos->Infos)
 		{
-			const TTuple<FPCGKernelAttributeKey, bool>* Pair = Keys->FindByPredicate([AttributeKey](const TTuple<FPCGKernelAttributeKey, bool>& Pair) { return Pair.Key == AttributeKey; });
-			if (!Pair || !Pair->Value)
+			if (Info.bCreatedOnGPU && KernelAttributeKeys.IsValidIndex(Info.Index))
 			{
-				// Not created on GPU.
-				continue;
-			}
+				const FPCGKernelAttributeKey& AttributeKey = KernelAttributeKeys[Info.Index];
 
-			for (FPCGDataDesc& DataDesc : OutputDataDesc.DataDescs)
-			{
-				const int32 AttributeId = InBinding->GetAttributeId(AttributeKey);
-
-				if (AttributeId != INDEX_NONE)
+				for (FPCGDataDesc& DataDesc : OutputDataDesc.DataDescs)
 				{
-					const FPCGKernelAttributeDesc AttributeDesc(AttributeId, AttributeKey.Type, AttributeKey.Identifier);
-					DataDesc.AttributeDescs.AddUnique(AttributeDesc);
+					const int32 AttributeId = InBinding->GetAttributeId(AttributeKey);
+
+					if (AttributeId != INDEX_NONE)
+					{
+						const FPCGKernelAttributeDesc AttributeDesc(AttributeId, AttributeKey.Type, AttributeKey.Identifier);
+						DataDesc.AttributeDescs.AddUnique(AttributeDesc);
+					}
 				}
 			}
 		}
@@ -313,8 +415,8 @@ FString UPCGCustomHLSLKernel::GetCookedSource(const UPCGComputeGraph* InComputeG
 
 	// @todo_pcg: Allow overriding kernel source from an external kernel source
 	const UPCGCustomHLSLSettings* CustomHLSLSettings = CastChecked<UPCGCustomHLSLSettings>(Settings);
-	FString Source = CustomHLSLSettings->ShaderSource;
-	FString Functions = TEXT("\n\n// User kernel functions\n") + CustomHLSLSettings->ShaderFunctions;
+	FString Source = ParsedSource.Source;
+	FString Functions = TEXT("\n\n// User kernel functions\n") + ParsedFunctions.Source;
 
 	// @todo_pcg: Include additional source files in custom HLSL parsing pass, so that they support attribute names, data labels, etc.
 	ProcessShaderSource(InComputeGraph, Source, Functions);
@@ -495,26 +597,15 @@ int UPCGCustomHLSLKernel::ComputeThreadCount(const UPCGDataBinding* InBinding) c
 
 void UPCGCustomHLSLKernel::GetDataLabels(FName InPinLabel, TArray<FString>& OutDataLabels) const
 {
-	const UPCGCustomHLSLSettings* CustomHLSLSettings = CastChecked<UPCGCustomHLSLSettings>(Settings);
-
-	// Matches against {PinName}_AnyFunction('{DataLabel}'...
-	const FString Pattern = FString::Format(TEXT("{0}_.*?[\\s]*?\\([\\s]*?'([a-zA-Z0-9_].*?)'"), { *InPinLabel.ToString() });
-	const FString Source = CustomHLSLSettings->ShaderFunctions + TEXT("\n") + CustomHLSLSettings->ShaderSource;
-
-	// First capture: Data label (supports a - z, A - Z, 0 - 9, and underscores).
-	constexpr int LabelCaptureGroup = 1;
-
-	FRegexMatcher ModuleMatcher(FRegexPattern(Pattern), Source);
-	while (ModuleMatcher.FindNext())
+	if (const FPCGDataLabels* DataLabels = PinDataLabels.PinToDataLabels.Find(InPinLabel))
 	{
-		OutDataLabels.AddUnique(ModuleMatcher.GetCaptureGroup(LabelCaptureGroup));
+		OutDataLabels = DataLabels->Labels;
 	}
 }
 
 void UPCGCustomHLSLKernel::GetKernelAttributeKeys(TArray<FPCGKernelAttributeKey>& OutKeys) const
 {
-	const UPCGCustomHLSLSettings* CustomHLSLSettings = CastChecked<UPCGCustomHLSLSettings>(Settings);
-	OutKeys.Append(CustomHLSLSettings->KernelAttributeKeys);
+	OutKeys.Append(KernelAttributeKeys);
 }
 
 uint32 UPCGCustomHLSLKernel::GetThreadCountMultiplier() const
@@ -584,6 +675,301 @@ void UPCGCustomHLSLKernel::InitEntryPoint()
 	}
 }
 
+void UPCGCustomHLSLKernel::ParseShaderSource()
+{
+	const UPCGCustomHLSLSettings* CustomHLSLSettings = CastChecked<UPCGCustomHLSLSettings>(Settings);
+	ParsedSource.Source = CustomHLSLSettings->ShaderSource;
+	ParsedFunctions.Source = CustomHLSLSettings->ShaderFunctions;
+
+	static const TArray<FString> AttributeTypeStrings =
+	{
+		TEXT("Bool"),
+		TEXT("Int"),
+		TEXT("Uint"),
+		TEXT("Float"),
+		TEXT("Float2"),
+		TEXT("Float3"),
+		TEXT("Float4"),
+		TEXT("Rotator"),
+		TEXT("Quat"),
+		TEXT("Transform"),
+		TEXT("StringKey"),
+		TEXT("Name"),
+	};
+
+	// Collect additional keywords, such as function getters and setters.
+	TArray<FString> AdditionalKeywords;
+	TArray<FString> AttributeKeywords;
+	TArray<FString> CopyElementKeywords;
+	TArray<FString> InitializeKeywords;
+
+	for (const FPCGPinProperties& PinProps : CustomHLSLSettings->InputPins)
+	{
+		if (!!(PinProps.AllowedTypes & (EPCGDataType::Point | EPCGDataType::Param)))
+		{
+			for (const FString& AttributeTypeString : AttributeTypeStrings)
+			{
+				AttributeKeywords.Add(PinProps.Label.ToString() + "_" + PCGCustomHLSLKernel::AttributeFunctionGetKeyword + AttributeTypeString);
+			}
+		}
+	}
+
+	for (const FPCGPinProperties& PinProps : CustomHLSLSettings->OutputPins)
+	{
+		const FString PinStr = *PinProps.Label.ToString();
+
+		for (const FString& AttributeTypeString : AttributeTypeStrings)
+		{
+			AttributeKeywords.Add(PinStr + "_" + PCGCustomHLSLKernel::AttributeFunctionSetKeyword + AttributeTypeString);
+		}
+
+		for (const FPCGPinProperties& InputPinProps : CustomHLSLSettings->InputPins)
+		{
+			if (InputPinProps.AllowedTypes == PinProps.AllowedTypes)
+			{
+				CopyElementKeywords.Add(PinStr + "_" + PCGCustomHLSLKernel::CopyElementFunctionKeyword + "_" + InputPinProps.Label.ToString());
+			}
+		}
+
+		InitializeKeywords.Add(PinStr + "_" + PCGCustomHLSLKernel::InitializeFunctionKeyword);
+	}
+
+	AdditionalKeywords.Append(AttributeKeywords);
+	AdditionalKeywords.Append(CopyElementKeywords);
+	AdditionalKeywords.Append(InitializeKeywords);
+
+	IPCGEditorModule* PCGEditorModule = IPCGEditorModule::Get();
+	check(PCGEditorModule);
+
+	FPCGSyntaxTokenizerParams TokenizerParams;
+	TokenizerParams.AdditionalKeywords = MoveTemp(AdditionalKeywords);
+	TSharedPtr<ISyntaxTokenizer> Tokenizer = PCGEditorModule->CreateHLSLSyntaxTokenizer(TokenizerParams);
+	check(Tokenizer.IsValid());
+
+	auto ParseHelper = [&](FPCGCustomHLSLParsedSource& InOutParsedSource)
+	{
+		TArray<ISyntaxTokenizer::FTokenizedLine> TokenizedLines;
+		Tokenizer->Process(TokenizedLines, InOutParsedSource.Source);
+		PCGCustomHLSLKernel::ParseTokens(InOutParsedSource.Source, TokenizedLines, InOutParsedSource.Tokens);
+
+		// Information about an attribute function looking for a match.
+		struct FAttributeFunctionMatch
+		{
+			FString PinStr;
+			FString FuncStr;
+			FString TypeStr;
+			FString NameStr;
+			int32 MatchBeginning = INDEX_NONE;
+			int32 EncounteredCommas = 0;
+			const int32 RequiredCommas = 2;
+
+			void Reset()
+			{
+				PinStr = "";
+				FuncStr = "";
+				TypeStr = "";
+				NameStr = "";
+				MatchBeginning = INDEX_NONE;
+				EncounteredCommas = 0;
+			}
+		};
+
+		FAttributeFunctionMatch AttributeFunctionMatch;
+		bool bLookingForAttributeFunctionMatch = false;
+
+		FString SingleQuoteString = "";
+		bool bLookingForSingleQuoteMatch = false;
+
+		auto AddCompletedAttributeFunction = [&KernelAttributeKeys=KernelAttributeKeys, &PinToAttributeInfos=PinToAttributeInfos, &AttributeFunctions=InOutParsedSource.AttributeFunctions](const FAttributeFunctionMatch& InAttributeFunction)
+		{
+			if (InAttributeFunction.EncounteredCommas != InAttributeFunction.RequiredCommas)
+			{
+				return;
+			}
+
+			// @todo_pcg: Validate NameStr in [a-zA-Z0-9 -_\/] ?
+			const FString& PinStr = InAttributeFunction.PinStr;
+			const FString& FuncStr = InAttributeFunction.FuncStr;
+			const FString& TypeStr = InAttributeFunction.TypeStr;
+			const FString& NameStr = InAttributeFunction.NameStr;
+			const FString UsageString = PinStr + "_" + FuncStr + TypeStr;
+
+			if (PinStr.IsEmpty() || FuncStr.IsEmpty() || TypeStr.IsEmpty() || NameStr.IsEmpty())
+			{
+				UE_LOG(LogPCG, Error, TEXT("Invalid attribute usage in shader source: '%s' on attribute name '%s'."), *UsageString, *NameStr);
+				return;
+			}
+
+			const UEnum* AttributeTypeEnum = StaticEnum<EPCGKernelAttributeType>();
+			check(AttributeTypeEnum);
+
+			const int64 AttributeType = AttributeTypeEnum->GetValueByName(FName(*TypeStr));
+
+			if (AttributeType == INDEX_NONE)
+			{
+				UE_LOG(LogPCG, Error, TEXT("Invalid attribute type in shader source: '%s' on attribute name '%s'."), *UsageString, *NameStr);
+				return;
+			}
+
+			// Add the attribute if it hasn't already been referenced.
+			FPCGKernelAttributeKey Key;
+			Key.Type = static_cast<EPCGKernelAttributeType>(AttributeType);
+			Key.SetSelector(FPCGAttributePropertySelector::CreateSelectorFromString(NameStr));
+
+			if (Key.IsValid())
+			{
+				const uint32 AttributeKeyIndex = KernelAttributeKeys.AddUnique(Key);
+
+				// Add an entry mapping this pin to the referenced attribute, if the entry doesn't already exist.
+				FPCGCustomHLSLAttributeInfos& Infos = PinToAttributeInfos.FindOrAdd(FName(*PinStr));
+				Infos.Infos.AddUnique({ AttributeKeyIndex, /*bCreatedOnGPU=*/false });
+			
+				AttributeFunctions.Emplace(PinStr, FuncStr, AttributeType, NameStr, InAttributeFunction.MatchBeginning);
+			}
+		};
+
+		const int32 NumTokens = InOutParsedSource.Tokens.Num();
+
+		for (int32 TokenIndex = 0; TokenIndex < NumTokens; ++TokenIndex)
+		{
+			const FPCGCustomHLSLParsedSource::FToken& Token = InOutParsedSource.Tokens[TokenIndex];
+			const FString TokenString = InOutParsedSource.Source.Mid(Token.Range.BeginIndex, Token.Range.Len());
+
+			if (Token.Type == FPCGCustomHLSLParsedSource::ETokenType::Keyword)
+			{
+				if (AttributeKeywords.Contains(TokenString))
+				{
+					AttributeFunctionMatch.Reset();
+
+					const int32 DelimiterIndex = TokenString.Find("_");
+					AttributeFunctionMatch.PinStr = TokenString.Left(DelimiterIndex);
+					AttributeFunctionMatch.FuncStr = TokenString.Mid(DelimiterIndex + 1, 3);
+					AttributeFunctionMatch.TypeStr = TokenString.Mid(DelimiterIndex + 4, Token.Range.Len() - DelimiterIndex + 1);
+					AttributeFunctionMatch.MatchBeginning = Token.Range.BeginIndex;
+					bLookingForAttributeFunctionMatch = true;
+
+					if (AttributeFunctionMatch.FuncStr == PCGCustomHLSLKernel::AttributeFunctionSetKeyword)
+					{
+						InOutParsedSource.InitializedOutputPins.Add(AttributeFunctionMatch.PinStr);
+					}
+				}
+				else if (CopyElementKeywords.Contains(TokenString))
+				{
+					const int32 FirstDelimiterIndex = TokenString.Find("_");
+					const int32 SecondDelimiterIndex = TokenString.Find("_", ESearchCase::IgnoreCase, ESearchDir::FromEnd);
+					
+					FString TargetPin = TokenString.Left(FirstDelimiterIndex);
+					FString SourcePin = TokenString.RightChop(SecondDelimiterIndex + 1);
+
+					InOutParsedSource.InitializedOutputPins.Add(TargetPin);
+					InOutParsedSource.CopyElementFunctions.Emplace(MoveTemp(SourcePin), MoveTemp(TargetPin));
+				}
+				else if (InitializeKeywords.Contains(TokenString))
+				{
+					const int32 DelimiterIndex = TokenString.Find("_");
+					FString PinStr = TokenString.Left(DelimiterIndex);
+
+					InOutParsedSource.InitializedOutputPins.Add(MoveTemp(PinStr));
+				}
+			}
+			else if (Token.Type == FPCGCustomHLSLParsedSource::ETokenType::Normal)
+			{
+				if (bLookingForSingleQuoteMatch && TokenString == "\'")
+				{
+					if (bLookingForAttributeFunctionMatch && AttributeFunctionMatch.EncounteredCommas == AttributeFunctionMatch.RequiredCommas)
+					{
+						AttributeFunctionMatch.NameStr = SingleQuoteString;
+						AddCompletedAttributeFunction(AttributeFunctionMatch);
+						bLookingForAttributeFunctionMatch = false;
+					}
+
+					SingleQuoteString = "";
+					bLookingForSingleQuoteMatch = false;
+				}
+				else if (bLookingForAttributeFunctionMatch && TokenString == ",")
+				{
+					++AttributeFunctionMatch.EncounteredCommas;
+
+					if (AttributeFunctionMatch.EncounteredCommas > AttributeFunctionMatch.RequiredCommas)
+					{
+						AttributeFunctionMatch.Reset();
+						bLookingForAttributeFunctionMatch = false;
+					}
+				}
+			}
+			else if (Token.Type == FPCGCustomHLSLParsedSource::ETokenType::SingleQuotedString)
+			{
+				if (!bLookingForSingleQuoteMatch)
+				{
+					bLookingForSingleQuoteMatch = true;
+
+					// Chop the leading single quote.
+					SingleQuoteString = TokenString.RightChop(1);
+				}
+				else
+				{
+					SingleQuoteString += TokenString;
+				}
+			}
+		}
+
+		// @todo_pcg: Maybe this should also be parsed instead of regex'd, but it's not as trivial to create tokens to detect the data labels pattern.
+		CollectDataLabels(InOutParsedSource);
+
+		// Capture any pin inititialization that won't have been caught by parsing.
+		// @todo_pcg: We should just add the point property setters as keywords to the parsing.
+		for (const FPCGPinPropertiesGPU& Pin : CustomHLSLSettings->OutputPins)
+		{
+			if (Pin.Label == NAME_None)
+			{
+				continue;
+			}
+
+			const FString PinStr = Pin.Label.ToString();
+
+			if (InOutParsedSource.Source.Contains(PinStr + TEXT("_Set")))
+			{
+				InOutParsedSource.InitializedOutputPins.Add(PinStr);
+			}
+		}
+	};
+
+	ParseHelper(ParsedSource);
+	ParseHelper(ParsedFunctions);
+}
+
+void UPCGCustomHLSLKernel::CollectDataLabels(const FPCGCustomHLSLParsedSource& InParsedSource)
+{
+	auto CollectDataLabelsForPin = [&InParsedSource=InParsedSource, &PinDataLabels=PinDataLabels](FName InPinLabel)
+	{
+		FPCGDataLabels& DataLabels = PinDataLabels.PinToDataLabels.FindOrAdd(InPinLabel);
+
+		// Matches against {PinName}_AnyFunction('{DataLabel}'...
+		const FString Pattern = FString::Format(TEXT("{0}_.*?[\\s]*?\\([\\s]*?'([a-zA-Z0-9_].*?)'"), { *InPinLabel.ToString() });
+
+		// First capture: Data label (supports a - z, A - Z, 0 - 9, and underscores).
+		constexpr int LabelCaptureGroup = 1;
+
+		FRegexMatcher ModuleMatcher(FRegexPattern(Pattern), InParsedSource.Source);
+		while (ModuleMatcher.FindNext())
+		{
+			DataLabels.Labels.AddUnique(ModuleMatcher.GetCaptureGroup(LabelCaptureGroup));
+		}
+	};
+
+	const UPCGCustomHLSLSettings* CustomHLSLSettings = CastChecked<UPCGCustomHLSLSettings>(Settings);
+
+	for (const FPCGPinProperties& PinProps : CustomHLSLSettings->InputPins)
+	{
+		CollectDataLabelsForPin(PinProps.Label);
+	}
+
+	for (const FPCGPinProperties& PinProps : CustomHLSLSettings->OutputPins)
+	{
+		CollectDataLabelsForPin(PinProps.Label);
+	}
+}
+
 bool UPCGCustomHLSLKernel::PerformStaticValidation()
 {
 	TRACE_CPUPROFILER_EVENT_SCOPE(UPCGCustomHLSLKernel::PerformStaticValidation);
@@ -864,7 +1250,7 @@ bool UPCGCustomHLSLKernel::PerformStaticValidation()
 		}
 	}
 
-	for (const FPCGKernelAttributeKey& AttributeKey : CustomHLSLSettings->KernelAttributeKeys)
+	for (const FPCGKernelAttributeKey& AttributeKey : KernelAttributeKeys)
 	{
 		if (AttributeKey.Type == EPCGKernelAttributeType::Invalid)
 		{
@@ -887,64 +1273,19 @@ bool UPCGCustomHLSLKernel::PerformStaticValidation()
 		return false;
 	}
 
-	if (!CustomHLSLSettings->bMuteUnwrittenPinDataErrors && !AreAllOutputPinsWritten(StaticLogEntries))
-	{
-		return false;
-	}
-
 	return true;
 }
 
 bool UPCGCustomHLSLKernel::ValidateShaderSource(TArray<FPCGKernelLogEntry>& OutLogEntries) const
 {
-	auto ParseCopyElementFunctions = [Settings=Settings, &OutLogEntries](const TArray<FPCGParsedCopyElementFunction>& ParsedFunctions)
-	{
-		const UPCGCustomHLSLSettings* CustomHLSLSettings = CastChecked<UPCGCustomHLSLSettings>(Settings);
-
-		for (const FPCGParsedCopyElementFunction& ParsedFunction : ParsedFunctions)
-		{
-			const FPCGPinProperties* SourcePinProps = CustomHLSLSettings->InputPins.FindByPredicate([Label=FName(*ParsedFunction.SourcePin)](const FPCGPinProperties& InProps) { return InProps.Label == Label; });
-			const FPCGPinPropertiesGPU* TargetPinProps = CustomHLSLSettings->OutputPins.FindByPredicate([Label=FName(*ParsedFunction.TargetPin)](const FPCGPinPropertiesGPU& InProps) { return InProps.Label == Label; });
-
-			if (!SourcePinProps)
-			{
-#if PCG_KERNEL_LOGGING_ENABLED
-				OutLogEntries.Emplace(FText::Format(LOCTEXT("InvalidCopySourcePin", "Invalid input pin '{0}' used in copy element expression."), FText::FromString(ParsedFunction.SourcePin)), EPCGKernelLogVerbosity::Error);
-#endif
-				return false;
-			}
-
-			if (!TargetPinProps)
-			{
-#if PCG_KERNEL_LOGGING_ENABLED
-				OutLogEntries.Emplace(FText::Format(LOCTEXT("InvalidCopyTargetPin", "Invalid output pin '{0}' used in copy element expression."), FText::FromString(ParsedFunction.TargetPin)), EPCGKernelLogVerbosity::Error);
-#endif
-				return false;
-			}
-
-			if (SourcePinProps && TargetPinProps && SourcePinProps->AllowedTypes != TargetPinProps->AllowedTypes)
-			{
-#if PCG_KERNEL_LOGGING_ENABLED
-				OutLogEntries.Emplace(LOCTEXT("InvalidCopyPinTypes", "CopyFromElement only allowed across pins of matching type."), EPCGKernelLogVerbosity::Error);
-#endif
-				return false;
-			}
-		}
-
-		return true;
-	};
-
-	if (!ParseCopyElementFunctions(SourceParsedCopyElementFunctions))
-	{
-		return false;
-	}
+	const UPCGCustomHLSLSettings* CustomHLSLSettings = CastChecked<UPCGCustomHLSLSettings>(Settings);
 
-	if (!ParseCopyElementFunctions(FunctionsParsedCopyElementFunctions))
+	if (!CustomHLSLSettings->bMuteUnwrittenPinDataErrors && !AreAllOutputPinsWritten(OutLogEntries))
 	{
 		return false;
 	}
 
-	// @todo_pcg: Validation of parsed attribute functions could be done here?
+	// @todo_pcg: Validation of parsed attribute functions could be done here instead of during parsing?
 
 	return true;
 }
@@ -992,8 +1333,8 @@ void UPCGCustomHLSLKernel::ProcessShaderSource(const UPCGComputeGraph* InCompute
 
 	// Implementation note: Attribute name replacement relies on precomputed indices into the source strings, therefore the attribute name replacement must take place
 	// before any other modifications. Otherwise, the indices will be incorrect and the source will become gibberish.
-	ReplaceAttributeNamesWithIds(OutShaderSource, SourceParsedAttributeFunctions);
-	ReplaceAttributeNamesWithIds(OutShaderFunctions, FunctionsParsedAttributeFunctions);
+	ReplaceAttributeNamesWithIds(OutShaderSource, ParsedSource.AttributeFunctions);
+	ReplaceAttributeNamesWithIds(OutShaderFunctions, ParsedFunctions.AttributeFunctions);
 
 #if PLATFORM_WINDOWS
 	// Remove old-school stuff.
@@ -1001,57 +1342,55 @@ void UPCGCustomHLSLKernel::ProcessShaderSource(const UPCGComputeGraph* InCompute
 	OutShaderFunctions.ReplaceInline(TEXT("\r"), TEXT(""));
 #endif
 
+	// @todo_pcg: Replace using token ranges instead of find/replace, similar to what we do with the parsed attribute functions.
 	auto ReplaceCopyElementFrom = [](FString& OutSource, const TArray<FPCGParsedCopyElementFunction>& ParsedFunctions)
 	{
 		// Replace function calls like Out_CopyElementFrom_In(...) with macro PCG_COPY_ALL_ATTRIBUTES_TO_OUTPUT(Out, In, ...).
 		for (const FPCGParsedCopyElementFunction& ParsedFunction : ParsedFunctions)
 		{
 			OutSource.ReplaceInline(
-				*FString::Format(TEXT("{1}_CopyElementFrom_{0}("), { ParsedFunction.SourcePin, ParsedFunction.TargetPin }),
+				*FString::Format(TEXT("{2}_{0}_{1}("), { PCGCustomHLSLKernel::CopyElementFunctionKeyword, ParsedFunction.SourcePin, ParsedFunction.TargetPin }),
 				*FString::Format(TEXT("PCG_COPY_ALL_ATTRIBUTES_TO_OUTPUT({1}, {0}, "), { ParsedFunction.SourcePin, ParsedFunction.TargetPin })
 			);
 		}
 	};
 
-	ReplaceCopyElementFrom(OutShaderSource, SourceParsedCopyElementFunctions);
-	ReplaceCopyElementFrom(OutShaderSource, FunctionsParsedCopyElementFunctions);
+	ReplaceCopyElementFrom(OutShaderSource, ParsedSource.CopyElementFunctions);
+	ReplaceCopyElementFrom(OutShaderSource, ParsedFunctions.CopyElementFunctions);
 
 	// We inject data IDs directly into the source. They will get remapped to data indices using a label resolver data interface.
-	if (const FPCGPinDataLabels* PinToDataLabels = InComputeGraph->GetStaticDataLabelsTable().Find(KernelIndex))
+	for (const TPair<FName, FPCGDataLabels>& Pair : PinDataLabels.PinToDataLabels)
 	{
-		for (const TPair<FName, FPCGDataLabels>& Pair : PinToDataLabels->PinToDataLabels)
-		{
-			const FName PinLabel = Pair.Key;
-			const TArray<FString>& DataLabels = Pair.Value.Labels;
+		const FName PinLabel = Pair.Key;
+		const TArray<FString>& DataLabels = Pair.Value.Labels;
 
-			for (int DataId = 0; DataId < DataLabels.Num(); ++DataId)
+		for (int DataId = 0; DataId < DataLabels.Num(); ++DataId)
+		{
+			const FString& DataLabel = DataLabels[DataId];
+			const FString ReplacementStr = FString::Format(TEXT("{0}_GetDataIndexFromIdInternal(/*DataId=*/{1})"),
 			{
-				const FString& DataLabel = DataLabels[DataId];
-				const FString ReplacementStr = FString::Format(TEXT("{0}_GetDataIndexFromIdInternal(/*DataId=*/{1})"),
-				{
-					PCGComputeHelpers::GetDataLabelResolverName(PinLabel),
-					FString::FromInt(DataId)
-				});
+				PCGComputeHelpers::GetDataLabelResolverName(PinLabel),
+				FString::FromInt(DataId)
+			});
 
-				// Matches against {PinName}_AnyFunction('{DataLabel}'...
-				// First capture group is the data label, so that we can find & replace it by index.
-				const FString Pattern = FString::Format(TEXT("{0}_.*?[\\s]*?\\([\\s]*?('{1}')"), { *PinLabel.ToString(), DataLabel });
+			// Matches against {PinName}_AnyFunction('{DataLabel}'...
+			// First capture group is the data label, so that we can find & replace it by index.
+			const FString Pattern = FString::Format(TEXT("{0}_.*?[\\s]*?\\([\\s]*?('{1}')"), { *PinLabel.ToString(), DataLabel });
 
-				auto ReplaceDataLabels = [&Pattern, &ReplacementStr](FString& InOutSource)
+			auto ReplaceDataLabels = [&Pattern, &ReplacementStr](FString& InOutSource)
+			{
+				FRegexMatcher ModuleMatcher(FRegexPattern(Pattern), InOutSource);
+				while (ModuleMatcher.FindNext())
 				{
-					FRegexMatcher ModuleMatcher(FRegexPattern(Pattern), InOutSource);
-					while (ModuleMatcher.FindNext())
-					{
-						const int32 DataLabelCharIndexStart = ModuleMatcher.GetCaptureGroupBeginning(1);
-						const int32 DataLabelCharIndexEnd = ModuleMatcher.GetCaptureGroupEnding(1);
+					const int32 DataLabelCharIndexStart = ModuleMatcher.GetCaptureGroupBeginning(1);
+					const int32 DataLabelCharIndexEnd = ModuleMatcher.GetCaptureGroupEnding(1);
 
-						InOutSource = InOutSource.Left(DataLabelCharIndexStart) + ReplacementStr + InOutSource.RightChop(DataLabelCharIndexEnd);
-					}
-				};
+					InOutSource = InOutSource.Left(DataLabelCharIndexStart) + ReplacementStr + InOutSource.RightChop(DataLabelCharIndexEnd);
+				}
+			};
 
-				ReplaceDataLabels(OutShaderSource);
-				ReplaceDataLabels(OutShaderFunctions);
-			}
+			ReplaceDataLabels(OutShaderSource);
+			ReplaceDataLabels(OutShaderFunctions);
 		}
 	}
 }
@@ -1060,48 +1399,34 @@ bool UPCGCustomHLSLKernel::AreAllOutputPinsWritten(TArray<FPCGKernelLogEntry>& O
 {
 	TRACE_CPUPROFILER_EVENT_SCOPE(UPCGCustomHLSLKernel::AreAllOutputPinsWritten);
 
-	const UPCGCustomHLSLSettings* CustomHLSLSettings = CastChecked<UPCGCustomHLSLSettings>(Settings);
-	const FString AllSource = CustomHLSLSettings->ShaderFunctions + CustomHLSLSettings->ShaderSource;
-
-	// Non-custom kernels initialize the first output pin data automatically.
-	if (CustomHLSLSettings->KernelType != EPCGKernelType::Custom && !CustomHLSLSettings->OutputPins.IsEmpty())
+	auto IsPinInitializedBySource = [](const FPCGCustomHLSLParsedSource& InParsedSource, const FString& PinStr)
 	{
-		const FName AutomaticallyInitializedPin = CustomHLSLSettings->OutputPins[0].Label;
-
-		for (const FPCGPinPropertiesGPU& Pin : CustomHLSLSettings->OutputPins)
-		{
-			if (Pin.Label == AutomaticallyInitializedPin || Pin.Label == NAME_None)
-			{
-				continue;
-			}
+		return InParsedSource.InitializedOutputPins.Contains(PinStr);
+	};
 
-			const TStaticArray<FString, 3> WriteStrings = {
-				Pin.Label.ToString() + TEXT("_Set"),
-				Pin.Label.ToString() + TEXT("_CopyElementFrom_"),
-				Pin.Label.ToString() + TEXT("_Initialize"),
-			};
+	const UPCGCustomHLSLSettings* CustomHLSLSettings = CastChecked<UPCGCustomHLSLSettings>(Settings);
 
-			bool bFoundDataWrite = false;
+	// Processor/Generator kernels initialize the first output pin data automatically.
+	const bool bSkipFirstPin = CustomHLSLSettings->KernelType == EPCGKernelType::PointProcessor || CustomHLSLSettings->KernelType == EPCGKernelType::PointGenerator;
 
-			for (const FString& WriteString : WriteStrings)
-			{
-				if (AllSource.Contains(WriteString))
-				{
-					bFoundDataWrite = true;
-					break;
-				}
-			}
+	for (int32 I = 0; I < CustomHLSLSettings->OutputPins.Num(); ++I)
+	{
+		if (I == 0 && bSkipFirstPin)
+		{
+			continue;
+		}
+		
+		const FString& PinStr = CustomHLSLSettings->OutputPins[I].Label.ToString();
 
-			if (!bFoundDataWrite)
-			{
+		if (!IsPinInitializedBySource(ParsedSource, PinStr) && !IsPinInitializedBySource(ParsedFunctions, PinStr))
+		{
 #if PCG_KERNEL_LOGGING_ENABLED
-				OutLogEntries.Emplace(FText::Format(
-					LOCTEXT("PinMayNotBeWritten", "Data on pin '{0}' may be uninitialized. Add code to write to this data, or mute this error in the node settings."),
-					{ FText::FromName(Pin.Label) }),
-					EPCGKernelLogVerbosity::Error);
+			OutLogEntries.Emplace(FText::Format(
+				LOCTEXT("PinMayNotBeWritten", "Data on pin '{0}' may be uninitialized. Add code to write to this data, or mute this error in the node settings."),
+				{ FText::FromString(PinStr) }),
+				EPCGKernelLogVerbosity::Error);
 #endif
-				return false;
-			}
+			return false;
 		}
 	}
 
@@ -1179,7 +1504,7 @@ bool UPCGCustomHLSLKernel::AreAttributesValid(FPCGContext* InContext, FText* Out
 					return FText::FromString(ParsedFunction.PinLabel + TEXT("_") + ParsedFunction.FunctionName + TypeStr);
 				};
 
-				if (FunctionName == PCGCustomHLSLKernel::AttributeFunctionSet)
+				if (FunctionName == PCGCustomHLSLKernel::AttributeFunctionSetKeyword)
 				{
 					if (const FPCGDataCollectionDesc** PinDescPtr = OutputPinDescs.Find(PinLabel))
 					{
@@ -1201,7 +1526,7 @@ bool UPCGCustomHLSLKernel::AreAttributesValid(FPCGContext* InContext, FText* Out
 						return false;
 					}
 				}
-				else if (ensure(FunctionName == PCGCustomHLSLKernel::AttributeFunctionGet))
+				else if (ensure(FunctionName == PCGCustomHLSLKernel::AttributeFunctionGetKeyword))
 				{
 					if (const FPCGDataCollectionDesc** PinDescPtr = InputPinDescs.Find(PinLabel))
 					{
@@ -1293,12 +1618,12 @@ bool UPCGCustomHLSLKernel::AreAttributesValid(FPCGContext* InContext, FText* Out
 			return true;
 		};
 
-		if (!ValidateParsedAttributeFunctions(SourceParsedAttributeFunctions))
+		if (!ValidateParsedAttributeFunctions(ParsedSource.AttributeFunctions))
 		{
 			return false;
 		}
 
-		if (!ValidateParsedAttributeFunctions(FunctionsParsedAttributeFunctions))
+		if (!ValidateParsedAttributeFunctions(ParsedFunctions.AttributeFunctions))
 		{
 			return false;
 		}
diff --git a/Engine/Plugins/PCG/Source/PCG/Private/Compute/Elements/PCGCustomHLSLKernel.h b/Engine/Plugins/PCG/Source/PCG/Private/Compute/Elements/PCGCustomHLSLKernel.h
--- a/Engine/Plugins/PCG/Source/PCG/Private/Compute/Elements/PCGCustomHLSLKernel.h
+++ b/Engine/Plugins/PCG/Source/PCG/Private/Compute/Elements/PCGCustomHLSLKernel.h
@@ -2,14 +2,13 @@
 
 #pragma once
 
+#include "Compute/PCGComputeCommon.h"
 #include "Compute/PCGComputeKernel.h"
 
 #include "Compute/PCGPinPropertiesGPU.h"
 
 #include "PCGCustomHLSLKernel.generated.h"
 
-class UPCGCustomHLSLSettings;
-
 USTRUCT()
 struct FPCGParsedAttributeFunction
 {
@@ -60,17 +59,85 @@ public:
 	FString TargetPin;
 };
 
+/** Helper struct for serializing attribute info. */
+USTRUCT()
+struct FPCGCustomHLSLAttributeInfo
+{
+	GENERATED_BODY()
+
+	FPCGCustomHLSLAttributeInfo() {}
+	FPCGCustomHLSLAttributeInfo(uint32 InIndex, bool bInCreatedOnGPU)
+		: Index(InIndex)
+		, bCreatedOnGPU(bInCreatedOnGPU)
+	{}
+
+	bool operator==(const FPCGCustomHLSLAttributeInfo& Other) const { return Index == Other.Index && bCreatedOnGPU == Other.bCreatedOnGPU; }
+
+	/** Index into KernelAttributeKeys. */
+	UPROPERTY()
+	uint32 Index = INDEX_NONE;
+
+	/** Whether or not this attribute was created on the GPU. */
+	UPROPERTY()
+	bool bCreatedOnGPU = false;
+};
+
+/** Helper struct for serializing map of pin name to attribute infos. */
+USTRUCT()
+struct FPCGCustomHLSLAttributeInfos
+{
+	GENERATED_BODY()
+
+	UPROPERTY()
+	TArray<FPCGCustomHLSLAttributeInfo> Infos;
+};
+
+/** Holds the results of parsing for some HLSL source. */
+USTRUCT()
+struct FPCGCustomHLSLParsedSource
+{
+	GENERATED_BODY()
+
+	UPROPERTY()
+	TArray<FPCGParsedAttributeFunction> AttributeFunctions;
+
+#if WITH_EDITORONLY_DATA
+	UPROPERTY()
+	TArray<FPCGParsedCopyElementFunction> CopyElementFunctions;
+
+	enum class ETokenType : uint8
+	{
+		Normal,
+		Keyword,
+		PreProcessorKeyword,
+		Operator,
+		DoubleQuotedString,
+		SingleQuotedString,
+		Comment,
+		Whitespace
+	};
+
+	struct FToken
+	{
+		ETokenType Type = ETokenType::Normal;
+		FTextRange Range;
+	};
+
+	FString Source;
+	TArray<FToken> Tokens;
+
+	/** Pins identified as being written to. Used to validate that output pins are initialized in some way. */
+	TArray<FString> InitializedOutputPins;
+#endif
+};
+
 UCLASS()
 class UPCGCustomHLSLKernel : public UPCGComputeKernel
 {
 	GENERATED_BODY()
 
+	// ~Begin UPCGComputeKernel interface
 public:
-#if WITH_EDITOR
-	virtual void Initialize(const UPCGSettings* InSettings, bool bLogDescriptions) override;
-#endif
-
-	virtual const UPCGSettings* GetSettings() const override;
 	virtual bool IsKernelDataValid(FPCGContext* InContext) const override;
 	virtual bool ComputeOutputBindingDataDesc(const UPCGComputeGraph* InGraph, FName InOutputPinLabel, UPCGDataBinding* InBinding, FPCGDataCollectionDesc& OutDataDesc) const override;
 #if WITH_EDITOR
@@ -87,8 +154,14 @@ public:
 
 protected:
 #if WITH_EDITOR
-	void InitEntryPoint();
+	virtual void InitializeInternal() override;
 	virtual bool PerformStaticValidation() override;
+	// ~End UPCGComputeKernel interface
+
+	void InitEntryPoint();
+	void ParseShaderSource();
+	void CollectDataLabels(const FPCGCustomHLSLParsedSource& InParsedSource);
+
 	bool ValidateShaderSource(TArray<FPCGKernelLogEntry>& OutLogEntries) const;
 	bool AreAllOutputPinsWritten(TArray<FPCGKernelLogEntry>& OutLogEntries) const;
 
@@ -113,19 +186,19 @@ public:
 #endif
 
 protected:
-	/** Attribute functions parsed and cached from HLSL shader source. */
+	// @todo_pcg: To support additional external sources, we should probably have an array/map of parsed sources instead.
+	UPROPERTY()
+	FPCGCustomHLSLParsedSource ParsedSource;
+
 	UPROPERTY()
-	TArray<FPCGParsedAttributeFunction> SourceParsedAttributeFunctions;
+	FPCGCustomHLSLParsedSource ParsedFunctions;
 
-	/** Attribute functions parsed and cached from HLSL shader functions source. */
 	UPROPERTY()
-	TArray<FPCGParsedAttributeFunction> FunctionsParsedAttributeFunctions;
+	TArray<FPCGKernelAttributeKey> KernelAttributeKeys;
 
-	/** CopyElementFrom functions parsed and cached from HLSL shader source. */
 	UPROPERTY()
-	TArray<FPCGParsedCopyElementFunction> SourceParsedCopyElementFunctions;
+	TMap<FName, FPCGCustomHLSLAttributeInfos> PinToAttributeInfos;
 
-	/** CopyElementFrom functions parsed and cached from HLSL shader functions source. */
 	UPROPERTY()
-	TArray<FPCGParsedCopyElementFunction> FunctionsParsedCopyElementFunctions;
+	FPCGPinDataLabels PinDataLabels;
 };
diff --git a/Engine/Plugins/PCG/Source/PCG/Private/Compute/PCGComputeGraph.h b/Engine/Plugins/PCG/Source/PCG/Private/Compute/PCGComputeGraph.h
--- a/Engine/Plugins/PCG/Source/PCG/Private/Compute/PCGComputeGraph.h
+++ b/Engine/Plugins/PCG/Source/PCG/Private/Compute/PCGComputeGraph.h
@@ -2,6 +2,7 @@
 
 #pragma once
 
+#include "Compute/PCGComputeCommon.h"
 #include "Compute/PCGDataForGPU.h"
 
 #include "ComputeFramework/ComputeGraph.h"
@@ -18,26 +19,6 @@ class UPCGNode;
 class UPCGPin;
 class UPCGStaticMeshSpawnerKernel;
 
-/** Helper struct for serializing data labels. */
-USTRUCT()
-struct FPCGDataLabels
-{
-	GENERATED_BODY()
-
-	UPROPERTY()
-	TArray<FString> Labels;
-};
-
-/** Helper struct for serializing map of pin name to data labels. */
-USTRUCT()
-struct FPCGPinDataLabels
-{
-	GENERATED_BODY()
-
-	UPROPERTY()
-	TMap</*PinLabel*/FName, FPCGDataLabels> PinToDataLabels;
-};
-
 /** An input or output pin of a kernel. Compute graph does not internally have 'pins' so this is useful for mapping between kernel data and PCG pins. */
 USTRUCT()
 struct FPCGKernelPin
diff --git a/Engine/Plugins/PCG/Source/PCG/Private/Compute/PCGComputeKernel.cpp b/Engine/Plugins/PCG/Source/PCG/Private/Compute/PCGComputeKernel.cpp
--- a/Engine/Plugins/PCG/Source/PCG/Private/Compute/PCGComputeKernel.cpp
+++ b/Engine/Plugins/PCG/Source/PCG/Private/Compute/PCGComputeKernel.cpp
@@ -34,6 +34,8 @@ void UPCGComputeKernel::Initialize(const UPCGSettings* InSettings, bool bLogDesc
 	bLogDataDescriptions = bLogDescriptions;
 	bInitialized = true;
 
+	InitializeInternal();
+
 	bHasStaticValidationErrors = !PerformStaticValidation();
 }
 #endif
diff --git a/Engine/Plugins/PCG/Source/PCG/Private/Compute/PCGComputeKernel.h b/Engine/Plugins/PCG/Source/PCG/Private/Compute/PCGComputeKernel.h
--- a/Engine/Plugins/PCG/Source/PCG/Private/Compute/PCGComputeKernel.h
+++ b/Engine/Plugins/PCG/Source/PCG/Private/Compute/PCGComputeKernel.h
@@ -55,7 +55,7 @@ class UPCGComputeKernel : public UComputeKernel
 public:
 #if WITH_EDITOR
 	/** Initialize kernel. Editor only as shader compilation only available in editor. */
-	virtual void Initialize(const UPCGSettings* InSettings, bool bLogDescriptions);
+	void Initialize(const UPCGSettings* InSettings, bool bLogDescriptions);
 #endif
 
 	/** Get/set index of kernel in kernels array on owning compute graph. */
@@ -130,6 +130,9 @@ public:
 
 protected:
 #if WITH_EDITOR
+	/** Implement to do any additional initialization specific to your kernel. */
+	virtual void InitializeInternal() {}
+
 	/** Performs validation on compile time information (e.g. Node, Settings, etc.). Caches validation errors/warnings. */
 	virtual bool PerformStaticValidation();
 
diff --git a/Engine/Plugins/PCG/Source/PCG/Public/Editor/IPCGEditorModule.h b/Engine/Plugins/PCG/Source/PCG/Public/Editor/IPCGEditorModule.h
--- a/Engine/Plugins/PCG/Source/PCG/Public/Editor/IPCGEditorModule.h
+++ b/Engine/Plugins/PCG/Source/PCG/Public/Editor/IPCGEditorModule.h
@@ -3,6 +3,7 @@
 #pragma once
 
 #include "Editor/IPCGEditorProgressNotification.h"
+#include "Editor/PCGSyntaxTokenizerParams.h"
 #include "Utils/PCGNodeVisualLogs.h"
 
 #include "Internationalization/Text.h"
@@ -24,6 +25,7 @@ public:
 	virtual const FPCGNodeVisualLogs& GetNodeVisualLogs() const = 0;
 	virtual FPCGNodeVisualLogs& GetNodeVisualLogsMutable() = 0;
 	virtual bool CanSelectPartitionActors() const = 0;
+	virtual TSharedPtr<class ISyntaxTokenizer> CreateHLSLSyntaxTokenizer(const FPCGSyntaxTokenizerParams& InParams) const = 0;
 #endif
 
 protected:
diff --git a/Engine/Plugins/PCG/Source/PCG/Public/Editor/PCGSyntaxTokenizerParams.h b/Engine/Plugins/PCG/Source/PCG/Public/Editor/PCGSyntaxTokenizerParams.h
new file mode 100644
--- /dev/null
+++ b/Engine/Plugins/PCG/Source/PCG/Public/Editor/PCGSyntaxTokenizerParams.h
@@ -0,0 +1,13 @@
+// Copyright Epic Games, Inc. All Rights Reserved.
+
+#pragma once
+
+#if WITH_EDITOR
+#include "CoreMinimal.h"
+
+/** Parameters to create a PCG syntax tokenizer. */
+struct FPCGSyntaxTokenizerParams
+{
+	TArray<FString> AdditionalKeywords;
+};
+#endif
diff --git a/Engine/Plugins/PCG/Source/PCGEditor/Private/PCGEditorModule.cpp b/Engine/Plugins/PCG/Source/PCGEditor/Private/PCGEditorModule.cpp
--- a/Engine/Plugins/PCG/Source/PCGEditor/Private/PCGEditorModule.cpp
+++ b/Engine/Plugins/PCG/Source/PCGEditor/Private/PCGEditorModule.cpp
@@ -26,6 +26,7 @@
 #include "PCGEditorSettings.h"
 #include "PCGEditorStyle.h"
 #include "PCGEditorUtils.h"
+#include "PCGHLSLSyntaxTokenizer.h"
 #include "DataVisualizations/PCGBaseTextureDataVisualization.h"
 #include "DataVisualizations/PCGLandscapeDataVisualization.h"
 #include "DataVisualizations/PCGParamDataVisualization.h"
@@ -258,6 +259,11 @@ bool FPCGEditorModule::CanSelectPartitionActors() const
 	return GetDefault<UPCGEditorSettings>()->bCanSelectPartitionActors;
 }
 
+TSharedPtr<ISyntaxTokenizer> FPCGEditorModule::CreateHLSLSyntaxTokenizer(const FPCGSyntaxTokenizerParams& InParams) const
+{
+	return MakeShared<FPCGHLSLSyntaxTokenizer>(InParams);
+}
+
 void FPCGEditorModule::RegisterOnEditorModeChange()
 {
 	// Have a callback that catches changes in the Editor modes, to catch when we exit the landscape edit mode.
diff --git a/Engine/Plugins/PCG/Source/PCGEditor/Private/PCGHLSLSyntaxHighlighter.cpp b/Engine/Plugins/PCG/Source/PCGEditor/Private/PCGHLSLSyntaxHighlighter.cpp
--- a/Engine/Plugins/PCG/Source/PCGEditor/Private/PCGHLSLSyntaxHighlighter.cpp
+++ b/Engine/Plugins/PCG/Source/PCGEditor/Private/PCGHLSLSyntaxHighlighter.cpp
@@ -1,6 +1,7 @@
 // Copyright Epic Games, Inc. All Rights Reserved.
 
 #include "PCGHLSLSyntaxHighlighter.h"
+#include "PCGHLSLSyntaxTokenizer.h"
 #include "PCGEditorStyle.h"
 #include "Styling/SlateTypes.h"
 #include "Framework/Text/IRun.h"
@@ -9,7 +10,7 @@
 
 TSharedRef<FPCGHLSLSyntaxHighlighter> FPCGHLSLSyntaxHighlighter::Create(const FSyntaxTextStyle& InSyntaxTextStyle)
 {
-	return MakeShareable(new FPCGHLSLSyntaxHighlighter(CreateTokenizer(), InSyntaxTextStyle));
+	return MakeShareable(new FPCGHLSLSyntaxHighlighter(MakeShareable(new FPCGHLSLSyntaxTokenizer()), InSyntaxTextStyle));
 }
 
 void FPCGHLSLSyntaxHighlighter::SetCompilerMessages(const FPCGCompilerDiagnostics& InCompilerMessages)
diff --git a/Engine/Plugins/PCG/Source/PCGEditor/Private/PCGHLSLSyntaxTokenizer.cpp b/Engine/Plugins/PCG/Source/PCGEditor/Private/PCGHLSLSyntaxTokenizer.cpp
new file mode 100644
--- /dev/null
+++ b/Engine/Plugins/PCG/Source/PCGEditor/Private/PCGHLSLSyntaxTokenizer.cpp
@@ -0,0 +1,9 @@
+// Copyright Epic Games, Inc. All Rights Reserved.
+
+#include "PCGHLSLSyntaxTokenizer.h"
+
+FPCGHLSLSyntaxTokenizer::FPCGHLSLSyntaxTokenizer(const FPCGSyntaxTokenizerParams& InParams)
+	: FHlslSyntaxTokenizer()
+{
+	Keywords.Append(InParams.AdditionalKeywords);
+}
diff --git a/Engine/Plugins/PCG/Source/PCGEditor/Private/PCGHLSLSyntaxTokenizer.h b/Engine/Plugins/PCG/Source/PCGEditor/Private/PCGHLSLSyntaxTokenizer.h
new file mode 100644
--- /dev/null
+++ b/Engine/Plugins/PCG/Source/PCGEditor/Private/PCGHLSLSyntaxTokenizer.h
@@ -0,0 +1,16 @@
+// Copyright Epic Games, Inc. All Rights Reserved.
+
+#pragma once
+
+#include "Editor/PCGSyntaxTokenizerParams.h"
+
+#include "Text/HLSLSyntaxTokenizer.h"
+
+/**
+ * Tokenize the text based on HLSL keywords and operators. Also considers PCG specific keywords as tokens if provided through AddAdditionalKeywords().
+ */
+class FPCGHLSLSyntaxTokenizer : public FHlslSyntaxTokenizer
+{
+public:
+	FPCGHLSLSyntaxTokenizer(const FPCGSyntaxTokenizerParams& InParams = {});
+};
diff --git a/Engine/Plugins/PCG/Source/PCGEditor/Public/PCGEditorModule.h b/Engine/Plugins/PCG/Source/PCGEditor/Public/PCGEditorModule.h
--- a/Engine/Plugins/PCG/Source/PCGEditor/Public/PCGEditorModule.h
+++ b/Engine/Plugins/PCG/Source/PCGEditor/Public/PCGEditorModule.h
@@ -34,6 +34,7 @@ public:
 	virtual const FPCGNodeVisualLogs& GetNodeVisualLogs() const { return NodeVisualLogs; }
 	virtual FPCGNodeVisualLogs& GetNodeVisualLogsMutable() { return NodeVisualLogs; }
 	virtual bool CanSelectPartitionActors() const override;
+	virtual TSharedPtr<class ISyntaxTokenizer> CreateHLSLSyntaxTokenizer(const FPCGSyntaxTokenizerParams& InParams) const override;
 
 protected:
 	virtual void OnScheduleGraph(const FPCGStackContext& StackContext);
diff --git a/Engine/Source/Editor/UnrealEd/Private/Text/HLSLSyntaxHighlighterMarshaller.cpp b/Engine/Source/Editor/UnrealEd/Private/Text/HLSLSyntaxHighlighterMarshaller.cpp
--- a/Engine/Source/Editor/UnrealEd/Private/Text/HLSLSyntaxHighlighterMarshaller.cpp
+++ b/Engine/Source/Editor/UnrealEd/Private/Text/HLSLSyntaxHighlighterMarshaller.cpp
@@ -1,6 +1,7 @@
 ﻿// Copyright Epic Games, Inc. All Rights Reserved.
 
 #include "Text/HLSLSyntaxHighlighterMarshaller.h"
+#include "Text/HLSLSyntaxTokenizer.h"
 #include "Framework/Text/IRun.h"
 #include "Framework/Text/TextLayout.h"
 #include "Framework/Text/ISlateRun.h"
@@ -8,577 +9,6 @@
 #include "Fonts/FontMeasure.h"
 #include "Framework/Application/SlateApplication.h"
 
-// NOTE: Since SyntaxTokenizer matches on a first-token-encountered basis, it's important that
-// tokens with the same prefix are ordered by longest-prefix-first. Ideally SyntaxTokenizer 
-// should be using a prefix tree structure for longest prefix matching.
-
-// Type Keywords are copied from CrossCompiler::EHlslToken
-const TCHAR* HlslKeywords[] =
-{
-	TEXT("while"),
-	TEXT("volatile"),
-	TEXT("void"),
-	TEXT("vector"),
-	TEXT("unorm"),
-	TEXT("uniform"),
-	TEXT("uint4x4"),
-	TEXT("uint4x3"),
-	TEXT("uint4x2"),
-	TEXT("uint4x1"),
-	TEXT("uint4"),
-	TEXT("uint3x4"),
-	TEXT("uint3x3"),
-	TEXT("uint3x2"),
-	TEXT("uint3x1"),
-	TEXT("uint3"),
-	TEXT("uint2x4"),
-	TEXT("uint2x3"),
-	TEXT("uint2x2"),
-	TEXT("uint2x1"),
-	TEXT("uint2"),
-	TEXT("uint1x4"),
-	TEXT("uint1x3"),
-	TEXT("uint1x2"),
-	TEXT("uint1x1"),
-	TEXT("uint1"),
-	TEXT("uint"),
-	TEXT("true"),
-	TEXT("switch"),
-	TEXT("struct"),
-	TEXT("static"),
-	TEXT("snorm"),
-	TEXT("shared"),
-	TEXT("row_major"),
-	TEXT("return"),
-	TEXT("register"),
-	TEXT("precise"),
-	TEXT("packoffset"),
-	TEXT("numthreads"),
-	TEXT("nointerpolation"),
-	TEXT("namespace"),
-	TEXT("matrix"),
-	TEXT("int4x4"),
-	TEXT("int4x3"),
-	TEXT("int4x2"),
-	TEXT("int4x1"),
-	TEXT("int4"),
-	TEXT("int3x4"),
-	TEXT("int3x3"),
-	TEXT("int3x2"),
-	TEXT("int3x1"),
-	TEXT("int3"),
-	TEXT("int2x4"),
-	TEXT("int2x3"),
-	TEXT("int2x2"),
-	TEXT("int2x1"),
-	TEXT("int2"),
-	TEXT("int1x4"),
-	TEXT("int1x3"),
-	TEXT("int1x2"),
-	TEXT("int1x1"),
-	TEXT("int1"),
-	TEXT("int"),
-	TEXT("if"),
-	TEXT("half4x4"),
-	TEXT("half4x3"),
-	TEXT("half4x2"),
-	TEXT("half4x1"),
-	TEXT("half4"),
-	TEXT("half3x4"),
-	TEXT("half3x3"),
-	TEXT("half3x2"),
-	TEXT("half3x1"),
-	TEXT("half3"),
-	TEXT("half2x4"),
-	TEXT("half2x3"),
-	TEXT("half2x2"),
-	TEXT("half2x1"),
-	TEXT("half2"),
-	TEXT("half1x4"),
-	TEXT("half1x3"),
-	TEXT("half1x2"),
-	TEXT("half1x1"),
-	TEXT("half1"),
-	TEXT("half"),
-	TEXT("groupshared"),
-	TEXT("goto"),
-	TEXT("for"),
-	TEXT("float4x4"),
-	TEXT("float4x3"),
-	TEXT("float4x2"),
-	TEXT("float4x1"),
-	TEXT("float4"),
-	TEXT("float3x4"),
-	TEXT("float3x3"),
-	TEXT("float3x2"),
-	TEXT("float3x1"),
-	TEXT("float3"),
-	TEXT("float2x4"),
-	TEXT("float2x3"),
-	TEXT("float2x2"),
-	TEXT("float2x1"),
-	TEXT("float2"),
-	TEXT("float1x4"),
-	TEXT("float1x3"),
-	TEXT("float1x2"),
-	TEXT("float1x1"),
-	TEXT("float1"),
-	TEXT("float"),
-	TEXT("false"),
-	TEXT("extern"),
-	TEXT("export"),
-	TEXT("enum"),
-	TEXT("else"),
-	TEXT("dword"),
-	TEXT("double"),
-	TEXT("do"),
-	TEXT("default"),
-	TEXT("continue"),
-	TEXT("const"),
-	TEXT("column_major"),
-	TEXT("case"),
-	TEXT("break"),
-	TEXT("bool4x4"),
-	TEXT("bool4x3"),
-	TEXT("bool4x2"),
-	TEXT("bool4x1"),
-	TEXT("bool4"),
-	TEXT("bool3x4"),
-	TEXT("bool3x3"),
-	TEXT("bool3x2"),
-	TEXT("bool3x1"),
-	TEXT("bool3"),
-	TEXT("bool2x4"),
-	TEXT("bool2x3"),
-	TEXT("bool2x2"),
-	TEXT("bool2x1"),
-	TEXT("bool2"),
-	TEXT("bool1x4"),
-	TEXT("bool1x3"),
-	TEXT("bool1x2"),
-	TEXT("bool1x1"),
-	TEXT("bool1"),
-	TEXT("bool"),
-	TEXT("Buffer"),
-	TEXT("in"),
-	TEXT("out"),
-	TEXT("inout"),
-};
-
-const TCHAR* HlslOperators[] =
-{
-	TEXT("/*"),
-	TEXT("*/"),
-	TEXT("//"),
-	TEXT("\""),
-	TEXT("\'"),
-	TEXT("::"),
-	TEXT(":"),
-	TEXT("+="),
-	TEXT("++"),
-	TEXT("+"),
-	TEXT("--"),
-	TEXT("-="),
-	TEXT("-"),
-	TEXT("("),
-	TEXT(")"),
-	TEXT("["),
-	TEXT("]"),
-	TEXT("."),
-	TEXT("->"),
-	TEXT("!="),
-	TEXT("!"),
-	TEXT("&="),
-	TEXT("~"),
-	TEXT("&"),
-	TEXT("*="),
-	TEXT("*"),
-	TEXT("->"),
-	TEXT("/="),
-	TEXT("/"),
-	TEXT("%="),
-	TEXT("%"),
-	TEXT("<<="),
-	TEXT("<<"),
-	TEXT("<="),
-	TEXT("<"),
-	TEXT(">>="),
-	TEXT(">>"),
-	TEXT(">="),
-	TEXT(">"),
-	TEXT("=="),
-	TEXT("&&"),
-	TEXT("&"),
-	TEXT("^="),
-	TEXT("^"),
-	TEXT("|="),
-	TEXT("||"),
-	TEXT("|"),
-	TEXT("?"),
-	TEXT("="),
-};
-
-const TCHAR* HlslPreProcessorKeywords[] =
-{
-	TEXT("#include"),
-	TEXT("#define"),
-	TEXT("#ifndef"),
-	TEXT("#ifdef"),
-	TEXT("#if"),
-	TEXT("#else"),
-	TEXT("#endif"),
-	TEXT("#pragma"),
-	TEXT("#undef"),
-};
-
-const TCHAR* HlslSymbols[] =
-{
-	TEXT("abort"),
-	TEXT("abs"),
-	TEXT("acos"),
-	TEXT("all"),
-	TEXT("AllMemoryBarrier"),
-	TEXT("AllMemoryBarrierWithGroupSync"),
-	TEXT("any"),
-	TEXT("asdouble"),
-	TEXT("asfloat"),
-	TEXT("asin"),
-	TEXT("asint"),
-	TEXT("asuint"),
-	TEXT("asuint"),
-	TEXT("atan"),
-	TEXT("atan2"),
-	TEXT("ceil"),
-	TEXT("CheckAccessFullyMapped"),
-	TEXT("clamp"),
-	TEXT("clip"),
-	TEXT("cos"),
-	TEXT("cosh"),
-	TEXT("countbits"),
-	TEXT("cross"),
-	TEXT("D3DCOLORtoUBYTE4"),
-	TEXT("ddx"),
-	TEXT("ddx_coarse"),
-	TEXT("ddx_fine"),
-	TEXT("ddy"),
-	TEXT("ddy_coarse"),
-	TEXT("ddy_fine"),
-	TEXT("degrees"),
-	TEXT("determinant"),
-	TEXT("DeviceMemoryBarrier"),
-	TEXT("DeviceMemoryBarrierWithGroupSync"),
-	TEXT("distance"),
-	TEXT("dot"),
-	TEXT("dst"),
-	TEXT("errorf"),
-	TEXT("EvaluateAttributeCentroid"),
-	TEXT("EvaluateAttributeAtSample"),
-	TEXT("EvaluateAttributeSnapped"),
-	TEXT("exp"),
-	TEXT("exp2"),
-	TEXT("f16tof32"),
-	TEXT("f32tof16"),
-	TEXT("faceforward"),
-	TEXT("firstbithigh"),
-	TEXT("firstbitlow"),
-	TEXT("floor"),
-	TEXT("fma"),
-	TEXT("fmod"),
-	TEXT("frac"),
-	TEXT("frexp"),
-	TEXT("fwidth"),
-	TEXT("GetRenderTargetSampleCount"),
-	TEXT("GetRenderTargetSamplePosition"),
-	TEXT("GroupMemoryBarrier"),
-	TEXT("GroupMemoryBarrierWithGroupSync"),
-	TEXT("InterlockedAdd"),
-	TEXT("InterlockedAnd"),
-	TEXT("InterlockedCompareExchange"),
-	TEXT("InterlockedCompareStore"),
-	TEXT("InterlockedExchange"),
-	TEXT("InterlockedMax"),
-	TEXT("InterlockedMin"),
-	TEXT("InterlockedOr"),
-	TEXT("InterlockedXor"),
-	TEXT("isfinite"),
-	TEXT("isinf"),
-	TEXT("isnan"),
-	TEXT("ldexp"),
-	TEXT("length"),
-	TEXT("lerp"),
-	TEXT("lit"),
-	TEXT("log"),
-	TEXT("log10"),
-	TEXT("log2"),
-	TEXT("mad"),
-	TEXT("max"),
-	TEXT("min"),
-	TEXT("modf"),
-	TEXT("msad4"),
-	TEXT("mul"),
-	TEXT("noise"),
-	TEXT("normalize"),
-	TEXT("pow"),
-	TEXT("printf"),
-	TEXT("Process2DQuadTessFactorsAvg"),
-	TEXT("Process2DQuadTessFactorsMax"),
-	TEXT("Process2DQuadTessFactorsMin"),
-	TEXT("ProcessIsolineTessFactors"),
-	TEXT("ProcessQuadTessFactorsAvg"),
-	TEXT("ProcessQuadTessFactorsMax"),
-	TEXT("ProcessQuadTessFactorsMin"),
-	TEXT("ProcessTriTessFactorsAvg"),
-	TEXT("ProcessTriTessFactorsMax"),
-	TEXT("ProcessTriTessFactorsMin"),
-	TEXT("radians"),
-	TEXT("rcp"),
-	TEXT("reflect"),
-	TEXT("refract"),
-	TEXT("reversebits"),
-	TEXT("round"),
-	TEXT("rsqrt"),
-	TEXT("saturate"),
-	TEXT("sign"),
-	TEXT("sin"),
-	TEXT("sincos"),
-	TEXT("sinh"),
-	TEXT("smoothstep"),
-	TEXT("sqrt"),
-	TEXT("step"),
-	TEXT("tan"),
-	TEXT("tanh"),
-	TEXT("tex1D"),
-	TEXT("tex1D"),
-	TEXT("tex1Dbias"),
-	TEXT("tex1Dgrad"),
-	TEXT("tex1Dlod"),
-	TEXT("tex1Dproj"),
-	TEXT("tex2D"),
-	TEXT("tex2D"),
-	TEXT("tex2Dbias"),
-	TEXT("tex2Dgrad"),
-	TEXT("tex2Dlod"),
-	TEXT("tex2Dproj"),
-	TEXT("tex3D"),
-	TEXT("tex3D"),
-	TEXT("tex3Dbias"),
-	TEXT("tex3Dgrad"),
-	TEXT("tex3Dlod"),
-	TEXT("tex3Dproj"),
-	TEXT("texCUBE"),
-	TEXT("texCUBE"),
-	TEXT("texCUBEbias"),
-	TEXT("texCUBEgrad"),
-	TEXT("texCUBElod"),
-	TEXT("texCUBEproj"),
-	TEXT("transpose"),
-	TEXT("trunc"),
-	TEXT("SV_ClipDistance"),
-	TEXT("SV_CullDistance"),
-	TEXT("SV_Coverage"),
-	TEXT("SV_Depth"),
-	TEXT("SV_DepthGreaterEqual"),
-	TEXT("SV_DepthLessEqual"),
-	TEXT("SV_DispatchThreadID"),
-	TEXT("SV_DomainLocation"),
-	TEXT("SV_GroupID"),
-	TEXT("SV_GroupIndex"),
-	TEXT("SV_GroupThreadID"),
-	TEXT("SV_GSInstanceID"),
-	TEXT("SV_InnerCoverage"),
-	TEXT("SV_InsideTessFactor"),
-	TEXT("SV_InstanceID"),
-	TEXT("SV_IsFrontFace"),
-	TEXT("SV_OutputControlPointID"),
-	TEXT("SV_Position"),
-	TEXT("SV_PrimitiveID"),
-	TEXT("SV_RenderTargetArrayIndex"),
-	TEXT("SV_SampleIndex"),
-	TEXT("SV_StencilRef"),
-	TEXT("SV_Target"),
-	TEXT("SV_TessFactor"),
-	TEXT("SV_VertexID"),
-	TEXT("SV_ViewportArrayIndex"),
-	TEXT("SV_ShadingRate"),
-};
-
-/**
- * Tokenize the text based on Hlsl tokens
- */
-class FHlslSyntaxTokenizer : public ISyntaxTokenizer
-{
-public:
-	/** 
-	 * Create a new tokenizer
-	 */
-	static TSharedRef< FHlslSyntaxTokenizer > Create()
-	{
-		return MakeShareable(new FHlslSyntaxTokenizer());
-	};
-
-	virtual ~FHlslSyntaxTokenizer(){};
-
-	virtual void Process(TArray<FTokenizedLine>& OutTokenizedLines, const FString& Input) override
-	{
-#if UE_ENABLE_ICU
-		TArray<FTextRange> LineRanges;
-		FTextRange::CalculateLineRangesFromString(Input, LineRanges);
-		TokenizeLineRanges(Input, LineRanges, OutTokenizedLines);
-#else
-		FTokenizedLine FakeTokenizedLine;
-		FakeTokenizedLine.Range = FTextRange(0, Input.Len());
-		FakeTokenizedLine.Tokens.Emplace(FToken(ETokenType::Literal, FakeTokenizedLine.Range));
-		OutTokenizedLines.Add(FakeTokenizedLine);
-#endif
-	};
-
-private:
-	FHlslSyntaxTokenizer()
-	{
-		// operators
-		for(const auto& Operator : HlslOperators)
-		{
-			Operators.Emplace(Operator);
-		}	
-
-		// keywords
-		for(const auto& Keyword : HlslKeywords)
-		{
-			Keywords.Emplace(Keyword);
-		}
-
-		// Pre-processor Keywords
-		for(const auto& PreProcessorKeyword : HlslPreProcessorKeywords)
-		{
-			Keywords.Emplace(PreProcessorKeyword);
-		}
-		
-		// Symbols
-		for (const auto& Function : HlslSymbols)
-		{
-			Keywords.Emplace(Function);
-		}
-	};
-
-	void TokenizeLineRanges(const FString& Input, const TArray<FTextRange>& LineRanges, TArray<FTokenizedLine>& OutTokenizedLines)
-	{
-		// Tokenize line ranges
-		for(const FTextRange& LineRange : LineRanges)
-		{
-			FTokenizedLine TokenizedLine;
-			TokenizedLine.Range = LineRange;
-
-			if(TokenizedLine.Range.IsEmpty())
-			{
-				TokenizedLine.Tokens.Emplace(FToken(ETokenType::Literal, TokenizedLine.Range));
-			}
-			else
-			{
-				int32 CurrentOffset = LineRange.BeginIndex;
-				
-				while(CurrentOffset < LineRange.EndIndex)
-				{
-					const TCHAR* CurrentString = &Input[CurrentOffset];
-					const TCHAR CurrentChar = Input[CurrentOffset];
-
-					bool bHasMatchedSyntax = false;
-
-					// Greedy matching for operators
-					for(const FString& Operator : Operators)
-					{
-						if(FCString::Strncmp(CurrentString, *Operator, Operator.Len()) == 0)
-						{
-							const int32 SyntaxTokenEnd = CurrentOffset + Operator.Len();
-							TokenizedLine.Tokens.Emplace(FToken(ETokenType::Syntax, FTextRange(CurrentOffset, SyntaxTokenEnd)));
-						
-							check(SyntaxTokenEnd <= LineRange.EndIndex);
-						
-							bHasMatchedSyntax = true;
-							CurrentOffset = SyntaxTokenEnd;
-							break;
-						}
-					}
-				
-					if(bHasMatchedSyntax)
-					{
-						continue;
-					}
-
-					int32 PeekOffset = CurrentOffset + 1;
-					if (CurrentChar == TEXT('#'))
-					{
-						// Match PreProcessorKeywords
-						// They only contain letters
-						while(PeekOffset < LineRange.EndIndex)
-						{
-							const TCHAR PeekChar = Input[PeekOffset];
-
-							if (!TChar<TCHAR>::IsAlpha(PeekChar))
-							{
-								break;
-							}
-							
-							PeekOffset++;
-						}
-					}
-					else if (TChar<TCHAR>::IsAlpha(CurrentChar))
-					{
-						// Match Identifiers,
-						// They start with a letter and contain
-						// letters or numbers
-						while(PeekOffset < LineRange.EndIndex)
-						{
-							const TCHAR PeekChar = Input[PeekOffset];
-
-							if (!TChar<TCHAR>::IsIdentifier(PeekChar))
-							{
-								break;
-							}
-							
-							PeekOffset++;
-						}
-					}
-
-					const int32 CurrentStringLength = PeekOffset - CurrentOffset;
-					
-					// Check if it is an reserved keyword
-					for(const FString& Keyword : Keywords)
-					{
-						if (FCString::Strncmp(CurrentString, *Keyword, FMath::Max(CurrentStringLength, Keyword.Len())) == 0)
-						{
-							const int32 SyntaxTokenEnd = CurrentOffset + CurrentStringLength;
-							TokenizedLine.Tokens.Emplace(FToken(ETokenType::Syntax, FTextRange(CurrentOffset, SyntaxTokenEnd)));
-						
-							check(SyntaxTokenEnd <= LineRange.EndIndex);
-						
-							bHasMatchedSyntax = true;
-							CurrentOffset = SyntaxTokenEnd;
-							break;
-						}
-					}
-
-					if (bHasMatchedSyntax)
-					{
-						continue;
-					}
-
-					// If none matched, consume the character(s) as text
-					const int32 TextTokenEnd = CurrentOffset + CurrentStringLength;
-					TokenizedLine.Tokens.Emplace(FToken(ETokenType::Literal, FTextRange(CurrentOffset, TextTokenEnd)));
-					CurrentOffset = TextTokenEnd;
-				}
-			}
-
-			OutTokenizedLines.Add(TokenizedLine);
-		}
-	};
-
-	
-	TArray<FString> Keywords;
-	TArray<FString> Operators;
-};
-
-
 class FWhiteSpaceTextRun : public FSlateTextRun
 {
 public:
diff --git a/Engine/Source/Editor/UnrealEd/Private/Text/HLSLSyntaxTokenizer.cpp b/Engine/Source/Editor/UnrealEd/Private/Text/HLSLSyntaxTokenizer.cpp
new file mode 100644
--- /dev/null
+++ b/Engine/Source/Editor/UnrealEd/Private/Text/HLSLSyntaxTokenizer.cpp
@@ -0,0 +1,556 @@
+// Copyright Epic Games, Inc. All Rights Reserved.
+
+#include "Text/HLSLSyntaxTokenizer.h"
+
+// NOTE: Since SyntaxTokenizer matches on a first-token-encountered basis, it's important that
+// tokens with the same prefix are ordered by longest-prefix-first. Ideally SyntaxTokenizer 
+// should be using a prefix tree structure for longest prefix matching.
+
+// Type Keywords are copied from CrossCompiler::EHlslToken
+const TCHAR* HlslKeywords[] =
+{
+	TEXT("while"),
+	TEXT("volatile"),
+	TEXT("void"),
+	TEXT("vector"),
+	TEXT("unorm"),
+	TEXT("uniform"),
+	TEXT("uint4x4"),
+	TEXT("uint4x3"),
+	TEXT("uint4x2"),
+	TEXT("uint4x1"),
+	TEXT("uint4"),
+	TEXT("uint3x4"),
+	TEXT("uint3x3"),
+	TEXT("uint3x2"),
+	TEXT("uint3x1"),
+	TEXT("uint3"),
+	TEXT("uint2x4"),
+	TEXT("uint2x3"),
+	TEXT("uint2x2"),
+	TEXT("uint2x1"),
+	TEXT("uint2"),
+	TEXT("uint1x4"),
+	TEXT("uint1x3"),
+	TEXT("uint1x2"),
+	TEXT("uint1x1"),
+	TEXT("uint1"),
+	TEXT("uint"),
+	TEXT("true"),
+	TEXT("switch"),
+	TEXT("struct"),
+	TEXT("static"),
+	TEXT("snorm"),
+	TEXT("shared"),
+	TEXT("row_major"),
+	TEXT("return"),
+	TEXT("register"),
+	TEXT("precise"),
+	TEXT("packoffset"),
+	TEXT("numthreads"),
+	TEXT("nointerpolation"),
+	TEXT("namespace"),
+	TEXT("matrix"),
+	TEXT("int4x4"),
+	TEXT("int4x3"),
+	TEXT("int4x2"),
+	TEXT("int4x1"),
+	TEXT("int4"),
+	TEXT("int3x4"),
+	TEXT("int3x3"),
+	TEXT("int3x2"),
+	TEXT("int3x1"),
+	TEXT("int3"),
+	TEXT("int2x4"),
+	TEXT("int2x3"),
+	TEXT("int2x2"),
+	TEXT("int2x1"),
+	TEXT("int2"),
+	TEXT("int1x4"),
+	TEXT("int1x3"),
+	TEXT("int1x2"),
+	TEXT("int1x1"),
+	TEXT("int1"),
+	TEXT("int"),
+	TEXT("if"),
+	TEXT("half4x4"),
+	TEXT("half4x3"),
+	TEXT("half4x2"),
+	TEXT("half4x1"),
+	TEXT("half4"),
+	TEXT("half3x4"),
+	TEXT("half3x3"),
+	TEXT("half3x2"),
+	TEXT("half3x1"),
+	TEXT("half3"),
+	TEXT("half2x4"),
+	TEXT("half2x3"),
+	TEXT("half2x2"),
+	TEXT("half2x1"),
+	TEXT("half2"),
+	TEXT("half1x4"),
+	TEXT("half1x3"),
+	TEXT("half1x2"),
+	TEXT("half1x1"),
+	TEXT("half1"),
+	TEXT("half"),
+	TEXT("groupshared"),
+	TEXT("goto"),
+	TEXT("for"),
+	TEXT("float4x4"),
+	TEXT("float4x3"),
+	TEXT("float4x2"),
+	TEXT("float4x1"),
+	TEXT("float4"),
+	TEXT("float3x4"),
+	TEXT("float3x3"),
+	TEXT("float3x2"),
+	TEXT("float3x1"),
+	TEXT("float3"),
+	TEXT("float2x4"),
+	TEXT("float2x3"),
+	TEXT("float2x2"),
+	TEXT("float2x1"),
+	TEXT("float2"),
+	TEXT("float1x4"),
+	TEXT("float1x3"),
+	TEXT("float1x2"),
+	TEXT("float1x1"),
+	TEXT("float1"),
+	TEXT("float"),
+	TEXT("false"),
+	TEXT("extern"),
+	TEXT("export"),
+	TEXT("enum"),
+	TEXT("else"),
+	TEXT("dword"),
+	TEXT("double"),
+	TEXT("do"),
+	TEXT("default"),
+	TEXT("continue"),
+	TEXT("const"),
+	TEXT("column_major"),
+	TEXT("case"),
+	TEXT("break"),
+	TEXT("bool4x4"),
+	TEXT("bool4x3"),
+	TEXT("bool4x2"),
+	TEXT("bool4x1"),
+	TEXT("bool4"),
+	TEXT("bool3x4"),
+	TEXT("bool3x3"),
+	TEXT("bool3x2"),
+	TEXT("bool3x1"),
+	TEXT("bool3"),
+	TEXT("bool2x4"),
+	TEXT("bool2x3"),
+	TEXT("bool2x2"),
+	TEXT("bool2x1"),
+	TEXT("bool2"),
+	TEXT("bool1x4"),
+	TEXT("bool1x3"),
+	TEXT("bool1x2"),
+	TEXT("bool1x1"),
+	TEXT("bool1"),
+	TEXT("bool"),
+	TEXT("Buffer"),
+	TEXT("in"),
+	TEXT("out"),
+	TEXT("inout"),
+};
+
+const TCHAR* HlslOperators[] =
+{
+	TEXT("/*"),
+	TEXT("*/"),
+	TEXT("//"),
+	TEXT("\""),
+	TEXT("\'"),
+	TEXT("::"),
+	TEXT(":"),
+	TEXT("+="),
+	TEXT("++"),
+	TEXT("+"),
+	TEXT("--"),
+	TEXT("-="),
+	TEXT("-"),
+	TEXT("("),
+	TEXT(")"),
+	TEXT("["),
+	TEXT("]"),
+	TEXT("."),
+	TEXT("->"),
+	TEXT("!="),
+	TEXT("!"),
+	TEXT("&="),
+	TEXT("~"),
+	TEXT("&"),
+	TEXT("*="),
+	TEXT("*"),
+	TEXT("->"),
+	TEXT("/="),
+	TEXT("/"),
+	TEXT("%="),
+	TEXT("%"),
+	TEXT("<<="),
+	TEXT("<<"),
+	TEXT("<="),
+	TEXT("<"),
+	TEXT(">>="),
+	TEXT(">>"),
+	TEXT(">="),
+	TEXT(">"),
+	TEXT("=="),
+	TEXT("&&"),
+	TEXT("&"),
+	TEXT("^="),
+	TEXT("^"),
+	TEXT("|="),
+	TEXT("||"),
+	TEXT("|"),
+	TEXT("?"),
+	TEXT("="),
+};
+
+const TCHAR* HlslPreProcessorKeywords[] =
+{
+	TEXT("#include"),
+	TEXT("#define"),
+	TEXT("#ifndef"),
+	TEXT("#ifdef"),
+	TEXT("#if"),
+	TEXT("#else"),
+	TEXT("#endif"),
+	TEXT("#pragma"),
+	TEXT("#undef"),
+};
+
+const TCHAR* HlslSymbols[] =
+{
+	TEXT("abort"),
+	TEXT("abs"),
+	TEXT("acos"),
+	TEXT("all"),
+	TEXT("AllMemoryBarrier"),
+	TEXT("AllMemoryBarrierWithGroupSync"),
+	TEXT("any"),
+	TEXT("asdouble"),
+	TEXT("asfloat"),
+	TEXT("asin"),
+	TEXT("asint"),
+	TEXT("asuint"),
+	TEXT("asuint"),
+	TEXT("atan"),
+	TEXT("atan2"),
+	TEXT("ceil"),
+	TEXT("CheckAccessFullyMapped"),
+	TEXT("clamp"),
+	TEXT("clip"),
+	TEXT("cos"),
+	TEXT("cosh"),
+	TEXT("countbits"),
+	TEXT("cross"),
+	TEXT("D3DCOLORtoUBYTE4"),
+	TEXT("ddx"),
+	TEXT("ddx_coarse"),
+	TEXT("ddx_fine"),
+	TEXT("ddy"),
+	TEXT("ddy_coarse"),
+	TEXT("ddy_fine"),
+	TEXT("degrees"),
+	TEXT("determinant"),
+	TEXT("DeviceMemoryBarrier"),
+	TEXT("DeviceMemoryBarrierWithGroupSync"),
+	TEXT("distance"),
+	TEXT("dot"),
+	TEXT("dst"),
+	TEXT("errorf"),
+	TEXT("EvaluateAttributeCentroid"),
+	TEXT("EvaluateAttributeAtSample"),
+	TEXT("EvaluateAttributeSnapped"),
+	TEXT("exp"),
+	TEXT("exp2"),
+	TEXT("f16tof32"),
+	TEXT("f32tof16"),
+	TEXT("faceforward"),
+	TEXT("firstbithigh"),
+	TEXT("firstbitlow"),
+	TEXT("floor"),
+	TEXT("fma"),
+	TEXT("fmod"),
+	TEXT("frac"),
+	TEXT("frexp"),
+	TEXT("fwidth"),
+	TEXT("GetRenderTargetSampleCount"),
+	TEXT("GetRenderTargetSamplePosition"),
+	TEXT("GroupMemoryBarrier"),
+	TEXT("GroupMemoryBarrierWithGroupSync"),
+	TEXT("InterlockedAdd"),
+	TEXT("InterlockedAnd"),
+	TEXT("InterlockedCompareExchange"),
+	TEXT("InterlockedCompareStore"),
+	TEXT("InterlockedExchange"),
+	TEXT("InterlockedMax"),
+	TEXT("InterlockedMin"),
+	TEXT("InterlockedOr"),
+	TEXT("InterlockedXor"),
+	TEXT("isfinite"),
+	TEXT("isinf"),
+	TEXT("isnan"),
+	TEXT("ldexp"),
+	TEXT("length"),
+	TEXT("lerp"),
+	TEXT("lit"),
+	TEXT("log"),
+	TEXT("log10"),
+	TEXT("log2"),
+	TEXT("mad"),
+	TEXT("max"),
+	TEXT("min"),
+	TEXT("modf"),
+	TEXT("msad4"),
+	TEXT("mul"),
+	TEXT("noise"),
+	TEXT("normalize"),
+	TEXT("pow"),
+	TEXT("printf"),
+	TEXT("Process2DQuadTessFactorsAvg"),
+	TEXT("Process2DQuadTessFactorsMax"),
+	TEXT("Process2DQuadTessFactorsMin"),
+	TEXT("ProcessIsolineTessFactors"),
+	TEXT("ProcessQuadTessFactorsAvg"),
+	TEXT("ProcessQuadTessFactorsMax"),
+	TEXT("ProcessQuadTessFactorsMin"),
+	TEXT("ProcessTriTessFactorsAvg"),
+	TEXT("ProcessTriTessFactorsMax"),
+	TEXT("ProcessTriTessFactorsMin"),
+	TEXT("radians"),
+	TEXT("rcp"),
+	TEXT("reflect"),
+	TEXT("refract"),
+	TEXT("reversebits"),
+	TEXT("round"),
+	TEXT("rsqrt"),
+	TEXT("saturate"),
+	TEXT("sign"),
+	TEXT("sin"),
+	TEXT("sincos"),
+	TEXT("sinh"),
+	TEXT("smoothstep"),
+	TEXT("sqrt"),
+	TEXT("step"),
+	TEXT("tan"),
+	TEXT("tanh"),
+	TEXT("tex1D"),
+	TEXT("tex1D"),
+	TEXT("tex1Dbias"),
+	TEXT("tex1Dgrad"),
+	TEXT("tex1Dlod"),
+	TEXT("tex1Dproj"),
+	TEXT("tex2D"),
+	TEXT("tex2D"),
+	TEXT("tex2Dbias"),
+	TEXT("tex2Dgrad"),
+	TEXT("tex2Dlod"),
+	TEXT("tex2Dproj"),
+	TEXT("tex3D"),
+	TEXT("tex3D"),
+	TEXT("tex3Dbias"),
+	TEXT("tex3Dgrad"),
+	TEXT("tex3Dlod"),
+	TEXT("tex3Dproj"),
+	TEXT("texCUBE"),
+	TEXT("texCUBE"),
+	TEXT("texCUBEbias"),
+	TEXT("texCUBEgrad"),
+	TEXT("texCUBElod"),
+	TEXT("texCUBEproj"),
+	TEXT("transpose"),
+	TEXT("trunc"),
+	TEXT("SV_ClipDistance"),
+	TEXT("SV_CullDistance"),
+	TEXT("SV_Coverage"),
+	TEXT("SV_Depth"),
+	TEXT("SV_DepthGreaterEqual"),
+	TEXT("SV_DepthLessEqual"),
+	TEXT("SV_DispatchThreadID"),
+	TEXT("SV_DomainLocation"),
+	TEXT("SV_GroupID"),
+	TEXT("SV_GroupIndex"),
+	TEXT("SV_GroupThreadID"),
+	TEXT("SV_GSInstanceID"),
+	TEXT("SV_InnerCoverage"),
+	TEXT("SV_InsideTessFactor"),
+	TEXT("SV_InstanceID"),
+	TEXT("SV_IsFrontFace"),
+	TEXT("SV_OutputControlPointID"),
+	TEXT("SV_Position"),
+	TEXT("SV_PrimitiveID"),
+	TEXT("SV_RenderTargetArrayIndex"),
+	TEXT("SV_SampleIndex"),
+	TEXT("SV_StencilRef"),
+	TEXT("SV_Target"),
+	TEXT("SV_TessFactor"),
+	TEXT("SV_VertexID"),
+	TEXT("SV_ViewportArrayIndex"),
+	TEXT("SV_ShadingRate"),
+};
+
+TSharedRef<FHlslSyntaxTokenizer> FHlslSyntaxTokenizer::Create()
+{
+	return MakeShareable(new FHlslSyntaxTokenizer());
+}
+
+void FHlslSyntaxTokenizer::Process(TArray<FTokenizedLine>& OutTokenizedLines, const FString& Input)
+{
+#if UE_ENABLE_ICU
+	TArray<FTextRange> LineRanges;
+	FTextRange::CalculateLineRangesFromString(Input, LineRanges);
+	TokenizeLineRanges(Input, LineRanges, OutTokenizedLines);
+#else
+	FTokenizedLine FakeTokenizedLine;
+	FakeTokenizedLine.Range = FTextRange(0, Input.Len());
+	FakeTokenizedLine.Tokens.Emplace(FToken(ETokenType::Literal, FakeTokenizedLine.Range));
+	OutTokenizedLines.Add(FakeTokenizedLine);
+#endif
+}
+
+FHlslSyntaxTokenizer::FHlslSyntaxTokenizer()
+{
+	// operators
+	for(const auto& Operator : HlslOperators)
+	{
+		Operators.Emplace(Operator);
+	}	
+
+	// keywords
+	for(const auto& Keyword : HlslKeywords)
+	{
+		Keywords.Emplace(Keyword);
+	}
+
+	// Pre-processor Keywords
+	for(const auto& PreProcessorKeyword : HlslPreProcessorKeywords)
+	{
+		Keywords.Emplace(PreProcessorKeyword);
+	}
+	
+	// Symbols
+	for (const auto& Function : HlslSymbols)
+	{
+		Keywords.Emplace(Function);
+	}
+}
+
+void FHlslSyntaxTokenizer::TokenizeLineRanges(const FString& Input, const TArray<FTextRange>& LineRanges, TArray<FTokenizedLine>& OutTokenizedLines)
+{
+	// Tokenize line ranges
+	for(const FTextRange& LineRange : LineRanges)
+	{
+		FTokenizedLine TokenizedLine;
+		TokenizedLine.Range = LineRange;
+
+		if(TokenizedLine.Range.IsEmpty())
+		{
+			TokenizedLine.Tokens.Emplace(FToken(ETokenType::Literal, TokenizedLine.Range));
+		}
+		else
+		{
+			int32 CurrentOffset = LineRange.BeginIndex;
+			
+			while(CurrentOffset < LineRange.EndIndex)
+			{
+				const TCHAR* CurrentString = &Input[CurrentOffset];
+				const TCHAR CurrentChar = Input[CurrentOffset];
+
+				bool bHasMatchedSyntax = false;
+
+				// Greedy matching for operators
+				for(const FString& Operator : Operators)
+				{
+					if(FCString::Strncmp(CurrentString, *Operator, Operator.Len()) == 0)
+					{
+						const int32 SyntaxTokenEnd = CurrentOffset + Operator.Len();
+						TokenizedLine.Tokens.Emplace(FToken(ETokenType::Syntax, FTextRange(CurrentOffset, SyntaxTokenEnd)));
+					
+						check(SyntaxTokenEnd <= LineRange.EndIndex);
+					
+						bHasMatchedSyntax = true;
+						CurrentOffset = SyntaxTokenEnd;
+						break;
+					}
+				}
+			
+				if(bHasMatchedSyntax)
+				{
+					continue;
+				}
+
+				int32 PeekOffset = CurrentOffset + 1;
+				if (CurrentChar == TEXT('#'))
+				{
+					// Match PreProcessorKeywords
+					// They only contain letters
+					while(PeekOffset < LineRange.EndIndex)
+					{
+						const TCHAR PeekChar = Input[PeekOffset];
+
+						if (!TChar<TCHAR>::IsAlpha(PeekChar))
+						{
+							break;
+						}
+						
+						PeekOffset++;
+					}
+				}
+				else if (TChar<TCHAR>::IsAlpha(CurrentChar))
+				{
+					// Match Identifiers,
+					// They start with a letter and contain
+					// letters or numbers
+					while(PeekOffset < LineRange.EndIndex)
+					{
+						const TCHAR PeekChar = Input[PeekOffset];
+
+						if (!TChar<TCHAR>::IsIdentifier(PeekChar))
+						{
+							break;
+						}
+						
+						PeekOffset++;
+					}
+				}
+
+				const int32 CurrentStringLength = PeekOffset - CurrentOffset;
+				
+				// Check if it is an reserved keyword
+				for(const FString& Keyword : Keywords)
+				{
+					if (FCString::Strncmp(CurrentString, *Keyword, FMath::Max(CurrentStringLength, Keyword.Len())) == 0)
+					{
+						const int32 SyntaxTokenEnd = CurrentOffset + CurrentStringLength;
+						TokenizedLine.Tokens.Emplace(FToken(ETokenType::Syntax, FTextRange(CurrentOffset, SyntaxTokenEnd)));
+					
+						check(SyntaxTokenEnd <= LineRange.EndIndex);
+					
+						bHasMatchedSyntax = true;
+						CurrentOffset = SyntaxTokenEnd;
+						break;
+					}
+				}
+
+				if (bHasMatchedSyntax)
+				{
+					continue;
+				}
+
+				// If none matched, consume the character(s) as text
+				const int32 TextTokenEnd = CurrentOffset + CurrentStringLength;
+				TokenizedLine.Tokens.Emplace(FToken(ETokenType::Literal, FTextRange(CurrentOffset, TextTokenEnd)));
+				CurrentOffset = TextTokenEnd;
+			}
+		}
+
+		OutTokenizedLines.Add(TokenizedLine);
+	}
+}
diff --git a/Engine/Source/Editor/UnrealEd/Public/Text/HLSLSyntaxTokenizer.h b/Engine/Source/Editor/UnrealEd/Public/Text/HLSLSyntaxTokenizer.h
new file mode 100644
--- /dev/null
+++ b/Engine/Source/Editor/UnrealEd/Public/Text/HLSLSyntaxTokenizer.h
@@ -0,0 +1,29 @@
+// Copyright Epic Games, Inc. All Rights Reserved.
+
+#pragma once
+
+#include "Framework/Text/SyntaxTokenizer.h"
+
+/**
+ * Tokenize the text based on Hlsl tokens
+ */
+class FHlslSyntaxTokenizer : public ISyntaxTokenizer
+{
+public:
+	/** 
+	 * Create a new tokenizer
+	 */
+	static TSharedRef<FHlslSyntaxTokenizer> Create();
+
+	virtual ~FHlslSyntaxTokenizer(){};
+
+	UNREALED_API virtual void Process(TArray<FTokenizedLine>& OutTokenizedLines, const FString& Input) override;
+
+protected:
+	UNREALED_API FHlslSyntaxTokenizer();
+
+	void TokenizeLineRanges(const FString& Input, const TArray<FTextRange>& LineRanges, TArray<FTokenizedLine>& OutTokenizedLines);
+
+	TArray<FString> Keywords;
+	TArray<FString> Operators;
+};
