[PCG] Fix Flatten on data attributes for points

* We were copying the entries on all the domains
* Changed the signature for array views
* Add safeguard to not add more entries that we can with an error message on data domain
* Deprecated FlattenAndCompress using an array, and add a new one that uses a mapping

#rnx

diff --git a/Engine/Plugins/PCG/Source/PCG/Private/Data/PCGBasePointData.cpp b/Engine/Plugins/PCG/Source/PCG/Private/Data/PCGBasePointData.cpp
--- a/Engine/Plugins/PCG/Source/PCG/Private/Data/PCGBasePointData.cpp
+++ b/Engine/Plugins/PCG/Source/PCG/Private/Data/PCGBasePointData.cpp
@@ -1332,8 +1332,9 @@ void UPCGBasePointData::Flatten()
 		}
 	}
 
-	// Then flatten and compress the Metadata for all invalid entry keys. Return true if something changed.
-	if (Metadata->FlattenAndCompress(EntryKeys))
+	// Then flatten and compress the Metadata for all valid entry keys. Return true if something changed.
+	// For the data domain, it will do a normal flatten.
+	if (Metadata->FlattenAndCompress({{PCGMetadataDomainID::Elements, EntryKeys}}))
 	{
 		Modify();
 				
diff --git a/Engine/Plugins/PCG/Source/PCG/Private/Metadata/PCGMetadata.cpp b/Engine/Plugins/PCG/Source/PCG/Private/Metadata/PCGMetadata.cpp
--- a/Engine/Plugins/PCG/Source/PCG/Private/Metadata/PCGMetadata.cpp
+++ b/Engine/Plugins/PCG/Source/PCG/Private/Metadata/PCGMetadata.cpp
@@ -726,7 +726,33 @@ bool UPCGMetadata::FlattenAndCompress(const TArray<PCGMetadataEntryKey>& InEntry
 {
 	TRACE_CPUPROFILER_EVENT_SCOPE(UPCGMetadata::FlattenAndCompress);
 	
-	ForEachValidUniqueMetadataDomain([&InEntryKeysToKeep](const FPCGMetadataDomainID& DomainID, FPCGMetadataDomain& MetadataDomain) { MetadataDomain.FlattenAndCompress(InEntryKeysToKeep);});
+	ForEachValidUniqueMetadataDomain([&InEntryKeysToKeep](const FPCGMetadataDomainID& DomainID, FPCGMetadataDomain& MetadataDomain){ MetadataDomain.FlattenAndCompress(InEntryKeysToKeep); });
+
+	Parent = nullptr;
+	return true;
+}
+
+bool UPCGMetadata::FlattenAndCompress(const TMap<FPCGMetadataDomainID, TArrayView<const PCGMetadataEntryKey>>& InEntryKeysToKeepMapping)
+{
+	TRACE_CPUPROFILER_EVENT_SCOPE(UPCGMetadata::FlattenAndCompress);
+	
+	ForEachValidUniqueMetadataDomain([this, &InEntryKeysToKeepMapping](const FPCGMetadataDomainID& DomainID, FPCGMetadataDomain& MetadataDomain)
+	{
+		const TArrayView<const PCGMetadataEntryKey>* EntryKeysToKeep = InEntryKeysToKeepMapping.Find(MetadataDomain.GetDomainID());
+		if (!EntryKeysToKeep && MetadataDomain.GetDomainID() == DefaultDomain)
+		{
+			EntryKeysToKeep = InEntryKeysToKeepMapping.Find(PCGMetadataDomainID::Default);
+		}
+		
+		if (EntryKeysToKeep)
+		{
+			MetadataDomain.FlattenAndCompress(*EntryKeysToKeep);
+		}
+		else
+		{
+			MetadataDomain.FlattenImpl();
+		}
+	});
 
 	Parent = nullptr;
 	return true;
diff --git a/Engine/Plugins/PCG/Source/PCG/Private/Metadata/PCGMetadataDomain.cpp b/Engine/Plugins/PCG/Source/PCG/Private/Metadata/PCGMetadataDomain.cpp
--- a/Engine/Plugins/PCG/Source/PCG/Private/Metadata/PCGMetadataDomain.cpp
+++ b/Engine/Plugins/PCG/Source/PCG/Private/Metadata/PCGMetadataDomain.cpp
@@ -574,7 +574,7 @@ void FPCGMetadataDomain::FlattenImpl()
 	ItemKeyOffset = 0;
 }
 
-bool FPCGMetadataDomain::FlattenAndCompress(const TArray<PCGMetadataEntryKey>& InEntryKeysToKeep)
+bool FPCGMetadataDomain::FlattenAndCompress(const TArrayView<const PCGMetadataEntryKey>& InEntryKeysToKeep)
 {
 	TRACE_CPUPROFILER_EVENT_SCOPE(FPCGSubMetadata::FlattenAndCompress);
 
@@ -584,19 +584,27 @@ bool FPCGMetadataDomain::FlattenAndCompress(const TArray<PCGMetadataEntryKey>& I
 		return false;
 	}
 
+	TArrayView<const PCGMetadataEntryKey> EntryKeysToKeep = InEntryKeysToKeep;
+		
+	if (InEntryKeysToKeep.Num() > 1 && !SupportsMultiEntries())
+	{
+		EntryKeysToKeep = InEntryKeysToKeep.Left(1);
+		UE_LOG(LogPCG, Warning, TEXT("Tried to flatten and compress a domain '%ls', which doesn't support multiple entries, with %d entries to keep. Will only keep the first one."), *DomainID.DebugName.ToString(), InEntryKeysToKeep.Num());
+	}
+
 	AttributeLock.WriteLock();
 	for (auto& AttributePair : Attributes)
 	{
 		FPCGMetadataAttributeBase* Attribute = AttributePair.Value;
 		check(Attribute);
 
-		Attribute->FlattenAndCompress(InEntryKeysToKeep);
+		Attribute->FlattenAndCompress(EntryKeysToKeep);
 	}
 	AttributeLock.WriteUnlock();
 
 	Parent = nullptr;
 	ParentKeys.Reset();
-	ParentKeys.Init(PCGInvalidEntryKey, InEntryKeysToKeep.Num());
+	ParentKeys.Init(PCGInvalidEntryKey, EntryKeysToKeep.Num());
 	ItemKeyOffset = 0;
 
 	return true;
diff --git a/Engine/Plugins/PCG/Source/PCG/Public/Metadata/PCGMetadata.h b/Engine/Plugins/PCG/Source/PCG/Public/Metadata/PCGMetadata.h
--- a/Engine/Plugins/PCG/Source/PCG/Public/Metadata/PCGMetadata.h
+++ b/Engine/Plugins/PCG/Source/PCG/Public/Metadata/PCGMetadata.h
@@ -147,9 +147,12 @@ public:
 	/** Unparents current metadata by flattening the attributes (values, entries, etc.) */
 	void FlattenImpl();
 
-	/** Unparents current metadata, flatten attribute and only keep the entries specified. Return true if something has changed and keys needs be updated. */
+	UE_DEPRECATED(5.6, "Use the version with the mapping")
 	bool FlattenAndCompress(const TArray<PCGMetadataEntryKey>& InEntryKeysToKeep);
 
+	/** Unparents current metadata, flatten attribute and only keep the entries specified. Return true if something has changed and keys needs be updated. */
+	bool FlattenAndCompress(const TMap<FPCGMetadataDomainID, TArrayView<const PCGMetadataEntryKey>>& InEntryKeysToKeepMapping);
+	
 	// TODO: Will be moved at the end of the class before commit, as they are functions planned to be deprecated.
 	UFUNCTION(BlueprintCallable, Category = "PCG|Metadata")
 	UPARAM(DisplayName = "Metadata") UPCGMetadata* CreateInteger32Attribute(FName AttributeName, int32 DefaultValue, bool bAllowsInterpolation, bool bOverrideParent = true);
diff --git a/Engine/Plugins/PCG/Source/PCG/Public/Metadata/PCGMetadataAttribute.h b/Engine/Plugins/PCG/Source/PCG/Public/Metadata/PCGMetadataAttribute.h
--- a/Engine/Plugins/PCG/Source/PCG/Public/Metadata/PCGMetadataAttribute.h
+++ b/Engine/Plugins/PCG/Source/PCG/Public/Metadata/PCGMetadataAttribute.h
@@ -36,7 +36,7 @@ public:
 	/** Unparents current attribute by flattening the values, entries, etc. */
 	virtual void Flatten() = 0;
 	/** Unparents current attribute by flattening the values, entries, etc while only keeping the entries referenced in InEntryKeysToKeep. There must be NO invalid entry keys. */
-	virtual void FlattenAndCompress(const TArray<PCGMetadataEntryKey>& InEntryKeysToKeep) = 0;
+	virtual void FlattenAndCompress(const TArrayView<const PCGMetadataEntryKey>& InEntryKeysToKeep) = 0;
 
 	/** Remove all entries, values and parenting. */
 	virtual void Reset() = 0;
diff --git a/Engine/Plugins/PCG/Source/PCG/Public/Metadata/PCGMetadataAttributeTpl.h b/Engine/Plugins/PCG/Source/PCG/Public/Metadata/PCGMetadataAttributeTpl.h
--- a/Engine/Plugins/PCG/Source/PCG/Public/Metadata/PCGMetadataAttributeTpl.h
+++ b/Engine/Plugins/PCG/Source/PCG/Public/Metadata/PCGMetadataAttributeTpl.h
@@ -99,7 +99,7 @@ public:
 		Parent = nullptr;
 	}
 
-	virtual void FlattenAndCompress(const TArray<PCGMetadataEntryKey>& InEntryKeysToKeep) override
+	virtual void FlattenAndCompress(const TArrayView<const PCGMetadataEntryKey>& InEntryKeysToKeep) override
 	{
 		// No entries, we can just delete everything in the attribute.
 		if (InEntryKeysToKeep.IsEmpty())
diff --git a/Engine/Plugins/PCG/Source/PCG/Public/Metadata/PCGMetadataDomain.h b/Engine/Plugins/PCG/Source/PCG/Public/Metadata/PCGMetadataDomain.h
--- a/Engine/Plugins/PCG/Source/PCG/Public/Metadata/PCGMetadataDomain.h
+++ b/Engine/Plugins/PCG/Source/PCG/Public/Metadata/PCGMetadataDomain.h
@@ -87,7 +87,7 @@ public:
 	void FlattenImpl();
 
 	/** Unparents current metadata, flatten attribute and only keep the entries specified. Return true if something has changed and keys needs be updated. */
-	bool FlattenAndCompress(const TArray<PCGMetadataEntryKey>& InEntryKeysToKeep);
+	bool FlattenAndCompress(const TArrayView<const PCGMetadataEntryKey>& InEntryKeysToKeep);
 
 	/** Creates an attribute given a property.
 	* @param AttributeName: Target attribute to create
