[PCG] Fix Union property allocation
- Wasn't handling properly different default values from inputs
- Avoid setting color if it isn't allocated
- Add Density as a default allocated property

#rnx

diff --git a/Engine/Plugins/PCG/Source/PCG/Private/Data/PCGUnionData.cpp b/Engine/Plugins/PCG/Source/PCG/Private/Data/PCGUnionData.cpp
--- a/Engine/Plugins/PCG/Source/PCG/Private/Data/PCGUnionData.cpp
+++ b/Engine/Plugins/PCG/Source/PCG/Private/Data/PCGUnionData.cpp
@@ -39,6 +39,85 @@ namespace PCGUnionDataMaths
 	}
 }
 
+
+namespace PCGUnionHelper
+{
+	template<class T>
+	bool NeedToAllocateSingleValueProperty(EPCGPointNativeProperties InProperty, const TArray<const UPCGBasePointData*>& PointInputDatas)
+	{	
+		// Here we are comparing single values (non-allocated properties) to see if we have single values that differ between inputs.
+		// If they do differ, we need to allocate those properties because the output will have multiple values.
+		TOptional<const T> DefaultValue;
+		for (int32 i = 0; i < PointInputDatas.Num(); ++i)
+		{
+			const TConstPCGValueRange<T> ValueRange = PointInputDatas[i]->GetConstValueRange<T>(InProperty);
+			TOptional<const T> SingleValue = ValueRange.GetSingleValue();
+			if (!DefaultValue.IsSet())
+			{
+				DefaultValue = SingleValue;
+			}
+			else if (DefaultValue.IsSet() && SingleValue.IsSet() && !PCG::Private::MetadataTraits<T>::Equal(DefaultValue.GetValue(), SingleValue.GetValue()))
+			{
+				return true;
+			}
+		}
+
+		return false;
+	}
+
+	EPCGPointNativeProperties GetPropertiesToAllocateFromInputs(const TArray<const UPCGBasePointData*>& PointInputDatas)
+	{
+		// Start by doing a union of all input allocated properties
+		EPCGPointNativeProperties PropertiesToAllocate = EPCGPointNativeProperties::None;
+		for (const UPCGBasePointData* PointInputData : PointInputDatas)
+		{
+			PropertiesToAllocate |= PointInputData->GetAllocatedProperties();
+		}
+
+		if (!EnumHasAnyFlags(PropertiesToAllocate, EPCGPointNativeProperties::Transform) && NeedToAllocateSingleValueProperty<FTransform>(EPCGPointNativeProperties::Transform, PointInputDatas))
+		{
+			PropertiesToAllocate |= EPCGPointNativeProperties::Transform;
+		}
+
+		if (!EnumHasAnyFlags(PropertiesToAllocate, EPCGPointNativeProperties::Density) && NeedToAllocateSingleValueProperty<float>(EPCGPointNativeProperties::Density, PointInputDatas))
+		{
+			PropertiesToAllocate |= EPCGPointNativeProperties::Density;
+		}
+
+		if (!EnumHasAnyFlags(PropertiesToAllocate, EPCGPointNativeProperties::BoundsMin) && NeedToAllocateSingleValueProperty<FVector>(EPCGPointNativeProperties::BoundsMin, PointInputDatas))
+		{
+			PropertiesToAllocate |= EPCGPointNativeProperties::BoundsMin;
+		}
+
+		if (!EnumHasAnyFlags(PropertiesToAllocate, EPCGPointNativeProperties::BoundsMax) && NeedToAllocateSingleValueProperty<FVector>(EPCGPointNativeProperties::BoundsMax, PointInputDatas))
+		{
+			PropertiesToAllocate |= EPCGPointNativeProperties::BoundsMax;
+		}
+
+		if (!EnumHasAnyFlags(PropertiesToAllocate, EPCGPointNativeProperties::Color) && NeedToAllocateSingleValueProperty<FVector4>(EPCGPointNativeProperties::Color, PointInputDatas))
+		{
+			PropertiesToAllocate |= EPCGPointNativeProperties::Color;
+		}
+
+		if (!EnumHasAnyFlags(PropertiesToAllocate, EPCGPointNativeProperties::Steepness) && NeedToAllocateSingleValueProperty<float>(EPCGPointNativeProperties::Steepness, PointInputDatas))
+		{
+			PropertiesToAllocate |= EPCGPointNativeProperties::Steepness;
+		}
+
+		if (!EnumHasAnyFlags(PropertiesToAllocate, EPCGPointNativeProperties::Seed) && NeedToAllocateSingleValueProperty<int32>(EPCGPointNativeProperties::Seed, PointInputDatas))
+		{
+			PropertiesToAllocate |= EPCGPointNativeProperties::Seed;
+		}
+
+		if (!EnumHasAnyFlags(PropertiesToAllocate, EPCGPointNativeProperties::MetadataEntry) && NeedToAllocateSingleValueProperty<int64>(EPCGPointNativeProperties::MetadataEntry, PointInputDatas))
+		{
+			PropertiesToAllocate |= EPCGPointNativeProperties::MetadataEntry;
+		}
+
+		return PropertiesToAllocate;
+	}
+}
+
 void UPCGUnionData::Initialize(const UPCGSpatialData* InA, const UPCGSpatialData* InB)
 {
 	check(InA && InB);
@@ -297,14 +376,14 @@ const UPCGBasePointData* UPCGUnionData::CreateBasePointData(FPCGContext* Context
 			int32 PointInputCount = 0;
 			TArray<const UPCGBasePointData*> PointInputDatas;
 			PointInputDatas.SetNumUninitialized(DataRawPtr.Num());
-			EPCGPointNativeProperties PropertiesToAllocate = EPCGPointNativeProperties::None;
 			for (int32 Index = 0; Index < PointInputDatas.Num(); ++Index)
 			{
 				PointInputDatas[Index] = ToPointDataFunc(Context, DataRawPtr[Index]);
 				PointInputCount += PointInputDatas[Index]->GetNumPoints();
-				PropertiesToAllocate |= PointInputDatas[Index]->GetAllocatedProperties();
 			}
 
+			const EPCGPointNativeProperties PropertiesToAllocate = PCGUnionHelper::GetPropertiesToAllocateFromInputs(PointInputDatas);
+
 			if (PointInputCount > 0)
 			{
 				PointData->SetNumPoints(PointInputCount);
@@ -401,16 +480,17 @@ void UPCGUnionData::CreateSequentialPointData(FPCGContext* Context, TArray<const
 	
 	TArray<const UPCGBasePointData*> PointInputDatas;
 	PointInputDatas.SetNumUninitialized(InputDatas.Num());
-	EPCGPointNativeProperties PropertiesToAllocate = EPCGPointNativeProperties::None;
 	for (int32 Index = 0; Index < PointInputDatas.Num(); ++Index)
 	{
 		PointInputDatas[Index] = ToPointDataFunc(Context, InputDatas[Index]);
 		PointInputCount += PointInputDatas[Index]->GetNumPoints();
-		PropertiesToAllocate |= PointInputDatas[Index]->GetAllocatedProperties();
 	}
 	
 	OutPointData->SetNumPoints(PointInputCount, /*bInitializeValues=*/false);
-	OutPointData->AllocateProperties(PropertiesToAllocate | EPCGPointNativeProperties::MetadataEntry);
+
+	const EPCGPointNativeProperties PropertiesToAllocate = PCGUnionHelper::GetPropertiesToAllocateFromInputs(PointInputDatas);
+	OutPointData->AllocateProperties(PropertiesToAllocate | EPCGPointNativeProperties::MetadataEntry | EPCGPointNativeProperties::Density);
+	const bool bSetColor = EnumHasAnyFlags(PropertiesToAllocate, EPCGPointNativeProperties::Color);
 
 	struct FIndexParams
 	{
@@ -440,13 +520,16 @@ void UPCGUnionData::CreateSequentialPointData(FPCGContext* Context, TArray<const
 		PointOffset += Count;
 	};
 		
-	auto ProcessRange = [this, &PointOffset, &PointInputDatas, &InputDatas, &InputMetadatas, OutPointData, OutMetadata](int32 StartReadIndex, int32 StartWriteIndex, int32 Count, FIndexParams& IndexParams)
+	auto ProcessRange = [this, &PointOffset, &PointInputDatas, &InputDatas, &InputMetadatas, OutPointData, OutMetadata, bSetColor](int32 StartReadIndex, int32 StartWriteIndex, int32 Count, FIndexParams& IndexParams)
 	{
 		const UPCGBasePointData* CurrentPointData = PointInputDatas[IndexParams.CurrentIndex];
 				
-		FConstPCGPointValueRanges InRanges(CurrentPointData);
+		const FConstPCGPointValueRanges InRanges(CurrentPointData);
 		FPCGPointValueRanges OutRanges(OutPointData, /*bAllocate=*/false);
 
+		// Use const range for values we are not explicitely writing to
+		const FConstPCGPointValueRanges ConstOutRanges(OutPointData);
+
 		int32 NumWritten = 0;
 
 		for (int32 ReadIndex = StartReadIndex; ReadIndex < StartReadIndex + Count; ++ReadIndex)
@@ -500,21 +583,24 @@ void UPCGUnionData::CreateSequentialPointData(FPCGContext* Context, TArray<const
 
 				FPCGPoint PointInData;
 
-				const FBox LocalBounds = PCGPointHelpers::GetLocalBounds(OutRanges.BoundsMinRange[WriteIndex], OutRanges.BoundsMaxRange[WriteIndex]);
-				if (InputDatas[FollowingDataIndex]->SamplePoint(OutRanges.TransformRange[WriteIndex], LocalBounds, PointInData, OutMetadata))
+				const FBox LocalBounds = PCGPointHelpers::GetLocalBounds(ConstOutRanges.BoundsMinRange[WriteIndex], ConstOutRanges.BoundsMaxRange[WriteIndex]);
+				if (InputDatas[FollowingDataIndex]->SamplePoint(ConstOutRanges.TransformRange[WriteIndex], LocalBounds, PointInData, OutMetadata))
 				{
 					// Update density
 					PCGUnionDataMaths::UpdateDensity(OutRanges.DensityRange[WriteIndex], PointInData.Density, DensityFunction);
 
-					const FVector4& InColor = InRanges.ColorRange[ReadIndex];
-					FVector4& OutColor = OutRanges.ColorRange[WriteIndex];
-
-					// Update color
-					OutColor = FVector4(
-						FMath::Max(OutColor.X, InColor.X),
-						FMath::Max(OutColor.Y, InColor.Y),
-						FMath::Max(OutColor.Z, InColor.Z),
-						FMath::Max(OutColor.W, InColor.W));
+					if (bSetColor)
+					{
+						const FVector4& InColor = InRanges.ColorRange[ReadIndex];
+						FVector4& OutColor = OutRanges.ColorRange[WriteIndex];
+
+						// Update color
+						OutColor = FVector4(
+							FMath::Max(OutColor.X, InColor.X),
+							FMath::Max(OutColor.Y, InColor.Y),
+							FMath::Max(OutColor.Z, InColor.Z),
+							FMath::Max(OutColor.W, InColor.W));
+					}
 
 					if (OutRanges.MetadataEntryRange[WriteIndex] != PCGInvalidEntryKey && PointInData.MetadataEntry != PCGInvalidEntryKey)
 					{
