[PCG] Initialize compute kernels with a parameters struct

#rnx

diff --git a/Engine/Plugins/PCG/Source/PCG/Private/Compute/Elements/PCGCustomHLSL.cpp b/Engine/Plugins/PCG/Source/PCG/Private/Compute/Elements/PCGCustomHLSL.cpp
--- a/Engine/Plugins/PCG/Source/PCG/Private/Compute/Elements/PCGCustomHLSL.cpp
+++ b/Engine/Plugins/PCG/Source/PCG/Private/Compute/Elements/PCGCustomHLSL.cpp
@@ -328,8 +328,12 @@ bool UPCGCustomHLSLSettings::IsShaderTextReadOnly() const
 
 void UPCGCustomHLSLSettings::CreateKernels(FPCGGPUCompilationContext& InOutContext, UObject* InObjectOuter, TArray<UPCGComputeKernel*>& OutKernels, TArray<FPCGKernelEdge>& OutEdges) const
 {
+	FPCGComputeKernelParams KernelParams;
+	KernelParams.Settings = this;
+	KernelParams.bLogDescriptions = bDumpDataDescriptions;
+
 	UPCGCustomHLSLKernel* Kernel = InOutContext.NewObject_AnyThread<UPCGCustomHLSLKernel>(InObjectOuter);
-	Kernel->Initialize(this, bDumpDataDescriptions);
+	Kernel->Initialize(KernelParams);
 	OutKernels.Add(Kernel);
 
 	// Connect node pins to kernel pins
diff --git a/Engine/Plugins/PCG/Source/PCG/Private/Compute/PCGComputeKernel.cpp b/Engine/Plugins/PCG/Source/PCG/Private/Compute/PCGComputeKernel.cpp
--- a/Engine/Plugins/PCG/Source/PCG/Private/Compute/PCGComputeKernel.cpp
+++ b/Engine/Plugins/PCG/Source/PCG/Private/Compute/PCGComputeKernel.cpp
@@ -28,10 +28,10 @@ namespace PCGComputeKernel
 }
 
 #if WITH_EDITOR
-void UPCGComputeKernel::Initialize(const UPCGSettings* InSettings, bool bLogDescriptions)
+void UPCGComputeKernel::Initialize(const FPCGComputeKernelParams& InParams)
 {
-	Settings = InSettings;
-	bLogDataDescriptions = bLogDescriptions;
+	Settings = InParams.Settings;
+	bLogDataDescriptions = InParams.bLogDescriptions;
 	bInitialized = true;
 
 	InitializeInternal();
diff --git a/Engine/Plugins/PCG/Source/PCG/Private/Compute/PCGComputeKernel.h b/Engine/Plugins/PCG/Source/PCG/Private/Compute/PCGComputeKernel.h
--- a/Engine/Plugins/PCG/Source/PCG/Private/Compute/PCGComputeKernel.h
+++ b/Engine/Plugins/PCG/Source/PCG/Private/Compute/PCGComputeKernel.h
@@ -47,6 +47,14 @@ struct FPCGKernelLogEntry
 	EPCGKernelLogVerbosity Verbosity = EPCGKernelLogVerbosity::Verbose;
 };
 
+#if WITH_EDITOR
+struct FPCGComputeKernelParams
+{
+	const UPCGSettings* Settings = nullptr;
+	bool bLogDescriptions = false;
+};
+#endif
+
 UCLASS(Abstract)
 class UPCGComputeKernel : public UComputeKernel
 {
@@ -55,7 +63,7 @@ class UPCGComputeKernel : public UComputeKernel
 public:
 #if WITH_EDITOR
 	/** Initialize kernel. Editor only as shader compilation only available in editor. */
-	void Initialize(const UPCGSettings* InSettings, bool bLogDescriptions);
+	void Initialize(const FPCGComputeKernelParams& InParams);
 #endif
 
 	/** Get/set index of kernel in kernels array on owning compute graph. */
diff --git a/Engine/Plugins/PCG/Source/PCG/Private/Elements/Metadata/PCGMetadataPartition.cpp b/Engine/Plugins/PCG/Source/PCG/Private/Elements/Metadata/PCGMetadataPartition.cpp
--- a/Engine/Plugins/PCG/Source/PCG/Private/Elements/Metadata/PCGMetadataPartition.cpp
+++ b/Engine/Plugins/PCG/Source/PCG/Private/Elements/Metadata/PCGMetadataPartition.cpp
@@ -59,8 +59,12 @@ void UPCGMetadataPartitionSettings::PostLoad()
 #if WITH_EDITOR
 void UPCGMetadataPartitionSettings::CreateKernels(FPCGGPUCompilationContext& InOutContext, UObject* InObjectOuter, TArray<UPCGComputeKernel*>& OutKernels, TArray<FPCGKernelEdge>& OutEdges) const
 {
+	FPCGComputeKernelParams KernelParams;
+	KernelParams.Settings = this;
+	KernelParams.bLogDescriptions = bDumpDataDescriptions;
+
 	UPCGMetadataPartitionKernel* Kernel = InOutContext.NewObject_AnyThread<UPCGMetadataPartitionKernel>(InObjectOuter);
-	Kernel->Initialize(this, bDumpDataDescriptions);
+	Kernel->Initialize(KernelParams);
 	OutKernels.Add(Kernel);
 
 	// Connect node pins to kernel pins
diff --git a/Engine/Plugins/PCG/Source/PCG/Private/Elements/PCGCopyPoints.cpp b/Engine/Plugins/PCG/Source/PCG/Private/Elements/PCGCopyPoints.cpp
--- a/Engine/Plugins/PCG/Source/PCG/Private/Elements/PCGCopyPoints.cpp
+++ b/Engine/Plugins/PCG/Source/PCG/Private/Elements/PCGCopyPoints.cpp
@@ -25,8 +25,12 @@ FText UPCGCopyPointsSettings::GetNodeTooltipText() const
 
 void UPCGCopyPointsSettings::CreateKernels(FPCGGPUCompilationContext& InOutContext, UObject* InObjectOuter, TArray<UPCGComputeKernel*>& OutKernels, TArray<FPCGKernelEdge>& OutEdges) const
 {
+	FPCGComputeKernelParams KernelParams;
+	KernelParams.Settings = this;
+	KernelParams.bLogDescriptions = bDumpDataDescriptions;
+
 	UPCGCopyPointsKernel* CopyKernel = InOutContext.NewObject_AnyThread<UPCGCopyPointsKernel>(InObjectOuter);
-	CopyKernel->Initialize(this, bDumpDataDescriptions);
+	CopyKernel->Initialize(KernelParams);
 	OutKernels.Add(CopyKernel);
 
 	OutEdges.Emplace(FPCGPinReference(PCGCopyPointsConstants::SourcePointsLabel), FPCGPinReference(CopyKernel, PCGCopyPointsConstants::SourcePointsLabel));
@@ -36,7 +40,7 @@ void UPCGCopyPointsSettings::CreateKernels(FPCGGPUCompilationContext& InOutConte
 	if (bMatchBasedOnAttribute)
 	{
 		UPCGCopyPointsAnalysisKernel* AnalysisKernel = InOutContext.NewObject_AnyThread<UPCGCopyPointsAnalysisKernel>(InObjectOuter);
-		AnalysisKernel->Initialize(this, bDumpDataDescriptions);
+		AnalysisKernel->Initialize(KernelParams);
 		OutKernels.Add(AnalysisKernel);
 
 		// Connect to node pins and the copy kernel.
diff --git a/Engine/Plugins/PCG/Source/PCG/Private/Elements/PCGDataNum.cpp b/Engine/Plugins/PCG/Source/PCG/Private/Elements/PCGDataNum.cpp
--- a/Engine/Plugins/PCG/Source/PCG/Private/Elements/PCGDataNum.cpp
+++ b/Engine/Plugins/PCG/Source/PCG/Private/Elements/PCGDataNum.cpp
@@ -30,9 +30,13 @@ FText UPCGDataNumSettings::GetNodeTooltipText() const
 
 void UPCGDataNumSettings::CreateKernels(FPCGGPUCompilationContext& InOutContext, UObject* InObjectOuter, TArray<UPCGComputeKernel*>& OutKernels, TArray<FPCGKernelEdge>& OutEdges) const
 {
+	FPCGComputeKernelParams KernelParams;
+	KernelParams.Settings = this;
+	KernelParams.bLogDescriptions = bDumpDataDescriptions;
+
 	UPCGDataNumKernel* Kernel = InOutContext.NewObject_AnyThread<UPCGDataNumKernel>(InObjectOuter);
 	Kernel->SetOutputAttribute(OutputAttributeName);
-	Kernel->Initialize(this, bDumpDataDescriptions);
+	Kernel->Initialize(KernelParams);
 	OutKernels.Add(Kernel);
 
 	OutEdges.Emplace(FPCGPinReference(PCGPinConstants::DefaultInputLabel), FPCGPinReference(Kernel, PCGPinConstants::DefaultInputLabel));
diff --git a/Engine/Plugins/PCG/Source/PCG/Private/Elements/PCGNormalToDensity.cpp b/Engine/Plugins/PCG/Source/PCG/Private/Elements/PCGNormalToDensity.cpp
--- a/Engine/Plugins/PCG/Source/PCG/Private/Elements/PCGNormalToDensity.cpp
+++ b/Engine/Plugins/PCG/Source/PCG/Private/Elements/PCGNormalToDensity.cpp
@@ -13,8 +13,12 @@
 #if WITH_EDITOR
 void UPCGNormalToDensitySettings::CreateKernels(FPCGGPUCompilationContext& InOutContext, UObject* InObjectOuter, TArray<UPCGComputeKernel*>& OutKernels, TArray<FPCGKernelEdge>& OutEdges) const
 {
+	FPCGComputeKernelParams KernelParams;
+	KernelParams.Settings = this;
+	KernelParams.bLogDescriptions = bDumpDataDescriptions;
+
 	UPCGNormalToDensityKernel* Kernel = InOutContext.NewObject_AnyThread<UPCGNormalToDensityKernel>(InObjectOuter);
-	Kernel->Initialize(this, bDumpDataDescriptions);
+	Kernel->Initialize(KernelParams);
 	OutKernels.Add(Kernel);
 
 	// Connect node pins to kernel pins
diff --git a/Engine/Plugins/PCG/Source/PCG/Private/Elements/PCGStaticMeshSpawner.cpp b/Engine/Plugins/PCG/Source/PCG/Private/Elements/PCGStaticMeshSpawner.cpp
--- a/Engine/Plugins/PCG/Source/PCG/Private/Elements/PCGStaticMeshSpawner.cpp
+++ b/Engine/Plugins/PCG/Source/PCG/Private/Elements/PCGStaticMeshSpawner.cpp
@@ -49,8 +49,12 @@ UPCGStaticMeshSpawnerSettings::UPCGStaticMeshSpawnerSettings(const FObjectInitia
 #if WITH_EDITOR
 void UPCGStaticMeshSpawnerSettings::CreateKernels(FPCGGPUCompilationContext& InOutContext, UObject* InObjectOuter, TArray<UPCGComputeKernel*>& OutKernels, TArray<FPCGKernelEdge>& OutEdges) const
 {
+	FPCGComputeKernelParams KernelParams;
+	KernelParams.Settings = this;
+	KernelParams.bLogDescriptions = bDumpDataDescriptions;
+
 	UPCGStaticMeshSpawnerKernel* Kernel = InOutContext.NewObject_AnyThread<UPCGStaticMeshSpawnerKernel>(InObjectOuter);
-	Kernel->Initialize(this, bDumpDataDescriptions);
+	Kernel->Initialize(KernelParams);
 	OutKernels.Add(Kernel);
 
 	// Connect node pins to kernel pins
@@ -61,7 +65,7 @@ void UPCGStaticMeshSpawnerSettings::CreateKernels(FPCGGPUCompilationContext& InO
 	if (const UPCGMeshSelectorByAttribute* Selector = Cast<UPCGMeshSelectorByAttribute>(MeshSelectorParameters))
 	{
 		UPCGCountUniqueAttributeValuesKernel* CountKernel = InOutContext.NewObject_AnyThread<UPCGCountUniqueAttributeValuesKernel>(InObjectOuter);
-		CountKernel->Initialize(this, bDumpDataDescriptions);
+		CountKernel->Initialize(KernelParams);
 		OutKernels.Add(CountKernel);
 
 		OutEdges.Emplace(FPCGPinReference(PCGPinConstants::DefaultInputLabel), FPCGPinReference(CountKernel, PCGPinConstants::DefaultInputLabel));
