[PCG - Core] Fix parameters override not updating in BP instance when changed in BP template

* EnumerateObjectsToModify will look for objects owners even if the closest complex node is a Struct
* This will allow the propagation of the changes
* Also use ImportText to change the property override IDs, to follow the same propagation path than value changes
* This will effectively allow to change values for Graph Parameters in a BP template, and it will propagate to BP instances (if they are not different than the template, same rule than the other properties on the BP).

#rnx

diff --git a/Engine/Plugins/PCG/Source/PCG/Private/PCGGraph.cpp b/Engine/Plugins/PCG/Source/PCG/Private/PCGGraph.cpp
--- a/Engine/Plugins/PCG/Source/PCG/Private/PCGGraph.cpp
+++ b/Engine/Plugins/PCG/Source/PCG/Private/PCGGraph.cpp
@@ -2171,7 +2171,9 @@ void UPCGGraphInstance::PostEditChangeChainProperty(FPropertyChangedChainEvent&
 	}
 	else if (PropertyName == GET_MEMBER_NAME_CHECKED(FPCGOverrideInstancedPropertyBag, PropertiesIDsOverridden))
 	{
-		// Changes to this property will be caught by the update on override status, and "OnGraphParametersChanged" will be called accordingly.
+		// A change on the overriden properties can come from multiple places (like ImportText), and we don't know what changed or caused the change
+		// so treat it as a Undo/Redo.
+		OnGraphParametersChanged(Graph, EPCGGraphParameterEvent::UndoRedo, NAME_None);
 	}
 	else if (VerifyIfGraphCustomizationChanged(PropertyChangedEvent))
 	{
@@ -2463,6 +2465,43 @@ bool UPCGGraphInstance::IsPropertyOverriddenAndNotDefault(const FProperty* InPro
 	return Graph ? ParametersOverrides.IsPropertyOverriddenAndNotDefault(InProperty, Graph->GetUserParametersStruct()) : false;
 }
 
+#if WITH_EDITOR
+FString UPCGGraphInstance::ExportOverriddenPropertyIdsChangeForEditor(const FProperty* InProperty, bool bMarkAsOverridden, bool& bHasChanged)
+{
+	TSet<FGuid> CopyOfPropertyIDsOverridden;
+	bHasChanged = false;
+	
+	if (const FPropertyBagPropertyDesc* PropertyDesc = ParametersOverrides.Parameters.FindPropertyDescByName(InProperty->GetFName()))
+	{
+		if (bMarkAsOverridden && !ParametersOverrides.PropertiesIDsOverridden.Contains(PropertyDesc->ID))
+		{
+			CopyOfPropertyIDsOverridden = ParametersOverrides.PropertiesIDsOverridden;
+			CopyOfPropertyIDsOverridden.Add(PropertyDesc->ID);
+			bHasChanged = true;
+		}
+		else if (!bMarkAsOverridden && ParametersOverrides.PropertiesIDsOverridden.Contains(PropertyDesc->ID))
+		{
+			CopyOfPropertyIDsOverridden = ParametersOverrides.PropertiesIDsOverridden;
+			CopyOfPropertyIDsOverridden.Remove(PropertyDesc->ID);
+			bHasChanged = true;
+		}
+	}
+
+	if (bHasChanged)
+	{
+		const FProperty* PropertyIDsOverriddenProperty = FPCGOverrideInstancedPropertyBag::StaticStruct()->FindPropertyByName(GET_MEMBER_NAME_CHECKED(FPCGOverrideInstancedPropertyBag, PropertiesIDsOverridden));
+		check(PropertyIDsOverriddenProperty);
+		FString Result{};
+		PropertyIDsOverriddenProperty->ExportTextItem_Direct(Result, &CopyOfPropertyIDsOverridden, nullptr, nullptr, PPF_None);
+		return Result;
+	}
+	else
+	{
+		return FString{};
+	}
+}
+#endif // WITH_EDITOR
+
 bool UPCGGraphInstance::IsGraphParameterOverridden(const FName PropertyName) const
 {
 	return (ParametersOverrides.Parameters.FindPropertyDescByName(PropertyName) != nullptr);
diff --git a/Engine/Plugins/PCG/Source/PCG/Public/PCGGraph.h b/Engine/Plugins/PCG/Source/PCG/Public/PCGGraph.h
--- a/Engine/Plugins/PCG/Source/PCG/Public/PCGGraph.h
+++ b/Engine/Plugins/PCG/Source/PCG/Public/PCGGraph.h
@@ -623,6 +623,14 @@ public:
 	void ResetPropertyToDefault(const FProperty* InProperty);
 	bool IsPropertyOverridden(const FProperty* InProperty) const { return ParametersOverrides.IsPropertyOverridden(InProperty); }
 	bool IsPropertyOverriddenAndNotDefault(const FProperty* InProperty) const;
+	
+#if WITH_EDITOR
+	/** For properties to be propagated correctly, changes need to happen with ImportText.
+	 * This function will make a copy of the currently overridden property ids, add/remove the InProperty (depending on bMarkAsOverridden)
+	 * and return the ExportText of the set of overridden property ids, for it to be re-imported by the Editor to effectively make the change.
+	 */
+	FString ExportOverriddenPropertyIdsChangeForEditor(const FProperty* InProperty, bool bMarkAsOverridden, bool& bHasChanged);
+#endif // WITH_EDITOR
 
 	/** Because subgraph nodes hold graph instances, delegation of the selection can't be done elsewhere, so this function gets called when the dialog is opened and will forward to the graph customization if filtering is applied. */
 	UFUNCTION()
diff --git a/Engine/Plugins/PCG/Source/PCGEditor/Private/Details/PCGInstancedPropertyBagOverrideDetails.cpp b/Engine/Plugins/PCG/Source/PCGEditor/Private/Details/PCGInstancedPropertyBagOverrideDetails.cpp
--- a/Engine/Plugins/PCG/Source/PCGEditor/Private/Details/PCGInstancedPropertyBagOverrideDetails.cpp
+++ b/Engine/Plugins/PCG/Source/PCGEditor/Private/Details/PCGInstancedPropertyBagOverrideDetails.cpp
@@ -213,15 +213,17 @@ void FPCGOverrideInstancedPropertyBagDataDetails::OnChildRowAdded(IDetailPropert
 	TAttribute<bool> EditConditionValue = TAttribute<bool>::CreateLambda([this, ChildPropertyHandle]() -> bool { return Owner.IsValid() && Owner->IsPropertyOverridden(ChildPropertyHandle->GetProperty()) && !IsPropertyOverriddenByPin(ChildPropertyHandle); });
 	FOnBooleanValueChanged OnEditConditionChanged = FOnBooleanValueChanged::CreateLambda([this, ChildPropertyHandle](bool bNewValue)
 	{
-		PropertiesIDsOverriddenHandle->NotifyPreChange();
-
-		FScopedTransaction Transaction(FText::Format(LOCTEXT("OnCheckStateChanged", "Change Override for {0}"), FText::FromName(ChildPropertyHandle->GetProperty()->GetFName())));
+		//We have to go through ImportText for the whole propagation to happen
 		if (Owner.IsValid())
 		{
-			Owner->UpdatePropertyOverride(ChildPropertyHandle->GetProperty(), bNewValue);
+			bool bHasChanged = false;
+			FString NewValue = Owner->ExportOverriddenPropertyIdsChangeForEditor(ChildPropertyHandle->GetProperty(), bNewValue, bHasChanged);
+			if (bHasChanged)
+			{
+				FScopedTransaction Transaction(FText::Format(LOCTEXT("OnCheckStateChanged", "Change Override for {0}"), FText::FromName(ChildPropertyHandle->GetProperty()->GetFName())));
+				PropertiesIDsOverriddenHandle->SetPerObjectValues({std::move(NewValue)});
+			}
 		}
-
-		PropertiesIDsOverriddenHandle->NotifyPostChange(EPropertyChangeType::ValueSet);
 	});
 
 	ChildRow.EditCondition(std::move(EditConditionValue), std::move(OnEditConditionChanged));
diff --git a/Engine/Source/Editor/PropertyEditor/Private/PropertyHandleImpl.cpp b/Engine/Source/Editor/PropertyEditor/Private/PropertyHandleImpl.cpp
--- a/Engine/Source/Editor/PropertyEditor/Private/PropertyHandleImpl.cpp
+++ b/Engine/Source/Editor/PropertyEditor/Private/PropertyHandleImpl.cpp
@@ -70,8 +70,13 @@ FPropertyValueImpl::FPropertyValueImpl( TSharedPtr<FPropertyNode> InPropertyNode
 
 void FPropertyValueImpl::EnumerateObjectsToModify( FPropertyNode* InPropertyNode, const EnumerateObjectsToModifyFuncRef& InObjectsToModifyCallback ) const
 {
-	// Find the parent object node which contains offset addresses for reading a property value on an object
+	// EnumerateObjectsToModify catch the first ComplexNode and iterates on its number of instances to get the property address in memory.
+	// A complex node can be an object node or a struct node (at this time of writing).
+	// If the complex node is an object, we also extract the address of this object.
+	// If the complex node is a struct, we still want to know if it is owned by an object, so we look further up to find the first object node (if it exists).
+	// If it exists, we have a complex struct owned by an object, so we extract the object address.
 	FComplexPropertyNode* ComplexNode = InPropertyNode->FindComplexParent();
+	FObjectPropertyNode* ObjectNode = InPropertyNode->FindObjectItemParent();
 	if (ComplexNode)
 	{
 		const bool bIsStruct = FComplexPropertyNode::EPT_StandaloneStructure == ComplexNode->GetPropertyType();
@@ -83,12 +88,17 @@ void FPropertyValueImpl::EnumerateObjectsToModify( FPropertyNode* InPropertyNode
 			if (bIsStruct)
 			{
 				StructAddress = ComplexNode->GetMemoryOfInstance(Index);
+				if (ObjectNode)
+				{
+					Object = ObjectNode->GetInstanceAsUObject(Index).Get();
+				}
 			}
 			else
 			{
 				Object = ComplexNode->GetInstanceAsUObject(Index).Get();
 				StructAddress = InPropertyNode->GetStartAddressFromObject(Object);
 			}
+
 			uint8* BaseAddress = InPropertyNode->GetValueBaseAddress(StructAddress, InPropertyNode->HasNodeFlags(EPropertyNodeFlags::IsSparseClassData), bIsStruct);
 
 			if (!InObjectsToModifyCallback(FObjectBaseAddress(Object, StructAddress, BaseAddress), Index, NumInstances))
