[PCG] Enable per-data crc caching by default

diff --git a/Engine/Plugins/PCG/Source/PCG/Private/PCGElement.cpp b/Engine/Plugins/PCG/Source/PCG/Private/PCGElement.cpp
--- a/Engine/Plugins/PCG/Source/PCG/Private/PCGElement.cpp
+++ b/Engine/Plugins/PCG/Source/PCG/Private/PCGElement.cpp
@@ -29,7 +29,7 @@ static TAutoConsoleVariable<bool> CVarPCGValidatePointMetadata(
 
 static TAutoConsoleVariable<bool> CVarPCGAllowPerDataCaching(
 	TEXT("pcg.AllowPerDataCaching"),
-	false,
+	true,
 	TEXT("Controls whether we test & split down inputs to check caching per input on primary loop nodes."));
 
 static TAutoConsoleVariable<bool> CVarPCGShouldVerifyIfOutputsAreUsedMultipleTimes(
