[PCG] Fix PCGDataPtrWrapper GC softlock caused by Pinning a WeakPtr while inside a GraphExecutor lock that is also locked from within FPCGGraphExecutor::AddReferencedObjects
- Instead use TryPin and if Pin fails, dispatch to GameThread

#rnx

diff --git a/Engine/Plugins/PCG/Source/PCG/Private/Data/PCGDataPtrWrapper.cpp b/Engine/Plugins/PCG/Source/PCG/Private/Data/PCGDataPtrWrapper.cpp
--- a/Engine/Plugins/PCG/Source/PCG/Private/Data/PCGDataPtrWrapper.cpp
+++ b/Engine/Plugins/PCG/Source/PCG/Private/Data/PCGDataPtrWrapper.cpp
@@ -3,6 +3,8 @@
 #include "Data/PCGDataPtrWrapper.h"
 #include "PCGData.h"
 
+#include "Containers/Ticker.h"
+
 FPCGDataPtrWrapper::~FPCGDataPtrWrapper()
 {
 	DecRefCount();
@@ -84,10 +86,21 @@ void FPCGDataPtrWrapper::IncRefCount()
 
 void FPCGDataPtrWrapper::DecRefCount(const TWeakObjectPtr<const UPCGData>& InWeakData)
 {
-	if (TStrongObjectPtr<const UPCGData> InData = InWeakData.Pin())
+	bool bOutPinValid;
+	if (TStrongObjectPtr<const UPCGData> InData = InWeakData.TryPin(bOutPinValid))
 	{
 		InData->DecCollectionRefCount();
 	}
+	else if(!bOutPinValid)
+	{
+		ExecuteOnGameThread(UE_SOURCE_LOCATION, [InWeakData]()
+		{
+			if (const UPCGData* InData = InWeakData.Get())
+			{
+				InData->DecCollectionRefCount();
+			}
+		});
+	}
 }
 
 void FPCGDataPtrWrapper::DecRefCount() 
