[PCG] Add runtime gen startup delay cvar

#rnx

diff --git a/Engine/Plugins/PCG/Source/PCG/Private/RuntimeGen/PCGRuntimeGenScheduler.cpp b/Engine/Plugins/PCG/Source/PCG/Private/RuntimeGen/PCGRuntimeGenScheduler.cpp
--- a/Engine/Plugins/PCG/Source/PCG/Private/RuntimeGen/PCGRuntimeGenScheduler.cpp
+++ b/Engine/Plugins/PCG/Source/PCG/Private/RuntimeGen/PCGRuntimeGenScheduler.cpp
@@ -91,6 +91,11 @@ namespace PCGRuntimeGenSchedulerHelpers
 		false,
 		TEXT("Checks that the world is streamed in before triggering generation of local (partitioned) components."));
 
+	static TAutoConsoleVariable<int32> CVarFramesBeforeFirstGenerate(
+		TEXT("pcg.RuntimeGeneration.FramesBeforeFirstGenerate"),
+		0,
+		TEXT("Waits this many engine ticks before allowing runtime gen to schedule generation."));
+
 	static TAutoConsoleVariable<bool> CVarEnableVirtualTexturePriming(
 		TEXT("pcg.VirtualTexturePriming.Enable"),
 		false,
@@ -117,6 +122,7 @@ FPCGRuntimeGenScheduler::FPCGRuntimeGenScheduler(UWorld* InWorld, FPCGActorAndCo
 	GenSourceManager = new FPCGGenSourceManager(InWorld);
 	bPoolingWasEnabledLastFrame = PCGRuntimeGenSchedulerHelpers::CVarRuntimeGenerationEnablePooling.GetValueOnAnyThread();
 	BasePoolSizeLastFrame = PCGRuntimeGenSchedulerHelpers::CVarRuntimeGenerationBasePoolSize.GetValueOnAnyThread();
+	FramesUntilGeneration = PCGRuntimeGenSchedulerHelpers::CVarFramesBeforeFirstGenerate.GetValueOnGameThread();
 }
 
 FPCGRuntimeGenScheduler::~FPCGRuntimeGenScheduler()
@@ -162,6 +168,14 @@ void FPCGRuntimeGenScheduler::Tick(APCGWorldActor* InPCGWorldActor)
 		TickRequestVirtualTexturePriming(GenSources);
 	}
 
+	// Allow virtual texture priming to tick even when generation has not begun. This helps alleviate issues where we
+	// generate before the virtual textures have finished streaming in, which is particularly problematic on load.
+	if (FramesUntilGeneration > 0)
+	{
+		--FramesUntilGeneration;
+		return;
+	}
+
 	// Initialize RuntimeGen PA pool if necessary. If PoolSize is 0, then we have not initialized the pool yet.
 	if (!GenSources.IsEmpty() || !GeneratedComponents.IsEmpty())
 	{
diff --git a/Engine/Plugins/PCG/Source/PCG/Public/RuntimeGen/PCGRuntimeGenScheduler.h b/Engine/Plugins/PCG/Source/PCG/Public/RuntimeGen/PCGRuntimeGenScheduler.h
--- a/Engine/Plugins/PCG/Source/PCG/Public/RuntimeGen/PCGRuntimeGenScheduler.h
+++ b/Engine/Plugins/PCG/Source/PCG/Public/RuntimeGen/PCGRuntimeGenScheduler.h
@@ -172,4 +172,6 @@ private:
 	* this OC we early out. Basically we don't respond to refresh calls for a component we are midway through setting up.
 	*/
 	const UPCGComponent* OriginalComponentBeingGenerated = nullptr;
+
+	int32 FramesUntilGeneration = 0;
 };
