[PCG] Data Viewport quality of life improvements
* Add button to the the ALV to open the data viewport. Only shown when the viewport is closed
* Add toolbar to the data viewport to control things like camera speed, lighting mode, etc.

#rnx

diff --git a/Engine/Plugins/PCG/Source/PCGEditor/Private/Widgets/AssetEditorViewport/SPCGEditorViewport.cpp b/Engine/Plugins/PCG/Source/PCGEditor/Private/Widgets/AssetEditorViewport/SPCGEditorViewport.cpp
--- a/Engine/Plugins/PCG/Source/PCGEditor/Private/Widgets/AssetEditorViewport/SPCGEditorViewport.cpp
+++ b/Engine/Plugins/PCG/Source/PCGEditor/Private/Widgets/AssetEditorViewport/SPCGEditorViewport.cpp
@@ -5,11 +5,14 @@
 #include "AdvancedPreviewScene.h"
 #include "AssetViewerSettings.h"
 #include "EditorViewportClient.h"
+#include "ToolMenus.h"
+#include "ViewportToolbar/UnrealEdViewportToolbar.h"
+#include "ViewportToolbar/UnrealEdViewportToolbarContext.h"
 
 namespace PCGEditorViewportConstants
 {
-	const FVector DefaultViewLocation = FVector::ZeroVector;
-	const float DefaultOrbitDistance = 500.0f;
+	static const FVector DefaultViewLocation = FVector::ZeroVector;
+	static const float DefaultOrbitDistance = 500.0f;
 }
 
 class FPCGEditorViewportClient : public FEditorViewportClient
@@ -148,6 +151,17 @@ void SPCGEditorViewport::Construct(const FArguments& InArgs)
 	SEditorViewport::Construct(SEditorViewport::FArguments());
 }
 
+
+TSharedRef<class SEditorViewport> SPCGEditorViewport::GetViewportWidget()
+{
+	return SharedThis(this);
+}
+
+TSharedPtr<FExtender> SPCGEditorViewport::GetExtenders() const
+{
+	return MakeShareable(new FExtender);
+}
+
 void SPCGEditorViewport::SetupScene(const TArray<UObject*>& InResources, const FPCGSetupSceneFunc& SetupFunc)
 {
 	ResetScene();
@@ -203,6 +217,31 @@ TSharedRef<FEditorViewportClient> SPCGEditorViewport::MakeEditorViewportClient()
 	return EditorViewportClient.ToSharedRef();
 }
 
+TSharedPtr<SWidget> SPCGEditorViewport::BuildViewportToolbar()
+{
+	const FName ToolbarName = "PCG.DataViewportToolbar";
+	
+	if (!UToolMenus::Get()->IsMenuRegistered(ToolbarName))
+	{
+		UToolMenu* Menu = UToolMenus::Get()->RegisterMenu(ToolbarName, /*Parent=*/NAME_None, EMultiBoxType::SlimHorizontalToolBar);
+		Menu->StyleName = "ViewportToolbar";
+		Menu->AddSection("Left");
+		
+		FToolMenuSection& RightSection = Menu->AddSection("Right");
+		RightSection.Alignment = EToolMenuSectionAlign::Last;
+		RightSection.AddEntry(UE::UnrealEd::CreateCameraSubmenu(UE::UnrealEd::FViewportCameraMenuOptions().ShowAll()));
+		RightSection.AddEntry(UE::UnrealEd::CreateViewModesSubmenu());
+		RightSection.AddEntry(UE::UnrealEd::CreateDefaultShowSubmenu());
+	}
+	
+	FToolMenuContext Context;
+	Context.AppendCommandList(GetCommandList());
+	Context.AddExtender(GetExtenders());
+	Context.AddObject(UE::UnrealEd::CreateViewportToolbarDefaultContext(GetViewportWidget()));
+	
+	return UToolMenus::Get()->GenerateWidget(ToolbarName, Context);
+}
+
 void SPCGEditorViewport::ReleaseManagedResources()
 {
 	ManagedResources.Empty();
diff --git a/Engine/Plugins/PCG/Source/PCGEditor/Private/Widgets/AssetEditorViewport/SPCGEditorViewport.h b/Engine/Plugins/PCG/Source/PCGEditor/Private/Widgets/AssetEditorViewport/SPCGEditorViewport.h
--- a/Engine/Plugins/PCG/Source/PCGEditor/Private/Widgets/AssetEditorViewport/SPCGEditorViewport.h
+++ b/Engine/Plugins/PCG/Source/PCGEditor/Private/Widgets/AssetEditorViewport/SPCGEditorViewport.h
@@ -5,10 +5,11 @@
 #include "PCGDataVisualization.h"
 
 #include "SAssetEditorViewport.h"
+#include "SCommonEditorViewportToolbarBase.h"
 
 class FPCGEditorViewportClient;
 
-class SPCGEditorViewport : public SAssetEditorViewport
+class SPCGEditorViewport : public SAssetEditorViewport, public ICommonEditorViewportToolbarInfoProvider
 {
 public:
 	SLATE_BEGIN_ARGS(SPCGEditorViewport) {}
@@ -18,6 +19,12 @@ public:
 
 	void Construct(const FArguments& InArgs);
 
+	//~ Begin ICommonEditorViewportToolbarInfoProvider Interface
+	virtual TSharedRef<class SEditorViewport> GetViewportWidget() override;
+	virtual TSharedPtr<FExtender> GetExtenders() const override;
+	virtual void OnFloatingButtonClicked() override {}
+	//~ End ICommonEditorViewportToolbarInfoProvider Interface
+
 	/** Rebuilds the scene using the provided resources and setup callback. */
 	void SetupScene(const TArray<UObject*>& InResources, const FPCGSetupSceneFunc& SetupFunc);
 	void ResetScene();
@@ -27,6 +34,7 @@ public:
 protected:
 	//~ Begin SAssetEditorViewport Interface
 	virtual TSharedRef<FEditorViewportClient> MakeEditorViewportClient() override;
+	virtual TSharedPtr<SWidget> BuildViewportToolbar() override;
 	//~ End SAssetEditorViewport Interface
 
 	void ReleaseManagedResources();
diff --git a/Engine/Plugins/PCG/Source/PCGEditor/Private/Widgets/SPCGEditorGraphAttributeListView.cpp b/Engine/Plugins/PCG/Source/PCGEditor/Private/Widgets/SPCGEditorGraphAttributeListView.cpp
--- a/Engine/Plugins/PCG/Source/PCGEditor/Private/Widgets/SPCGEditorGraphAttributeListView.cpp
+++ b/Engine/Plugins/PCG/Source/PCGEditor/Private/Widgets/SPCGEditorGraphAttributeListView.cpp
@@ -418,7 +418,34 @@ void SPCGEditorGraphAttributeListView::Construct(const FArguments& InArgs, TShar
 				.ColorAndOpacity(FSlateColor::UseForeground())
 				.Image(FPCGEditorStyle::Get().GetBrush("PCG.Editor.ZoomToSelection"))
 		];
-		
+
+	TSharedPtr<SButton> OpenViewportButton = SNew(SButton)
+		.ButtonStyle(FAppStyle::Get(), "SimpleButton")
+		.OnClicked_Lambda([this]() -> FReply
+			{
+				if (const TSharedPtr<const FPCGEditor> SharedPtr = PCGEditorPtr.Pin())
+				{
+					SharedPtr->BringFocusToPanel(ViewportEditorPanel);
+					bViewportNeedsRefresh = true;
+				}
+
+				return FReply::Handled();
+			})
+		.Visibility_Lambda([this]()
+		{
+			const TSharedPtr<const FPCGEditor> SharedPtr = PCGEditorPtr.Pin();
+			return (SharedPtr.IsValid() && SharedPtr->IsPanelCurrentlyOpen(ViewportEditorPanel)) ? EVisibility::Collapsed : EVisibility::Visible;
+		})
+		.ContentPadding(FMargin(4, 2))
+		.HAlign(HAlign_Center)
+		.VAlign(VAlign_Center)
+		.ToolTipText(LOCTEXT("OpenDataViewport_ToolTip", "Opens Data Viewport Panel."))
+		[
+			SNew(SImage)
+				.ColorAndOpacity(FSlateColor::UseForeground())
+				.Image(FAppStyle::GetBrush("ClassIcon.CameraComponent"))
+		];
+
 	TSharedPtr<SComboButton> FilterButton = SNew(SComboButton)
 		.ForegroundColor(FSlateColor::UseStyle())
 		.HasDownArrow(false)
@@ -472,6 +499,12 @@ void SPCGEditorGraphAttributeListView::Construct(const FArguments& InArgs, TShar
 			+SHorizontalBox::Slot()
 			.AutoWidth()
 			.Padding(1.0f, 0.0f)
+			[
+				OpenViewportButton->AsShared()
+			]
+			+SHorizontalBox::Slot()
+			.AutoWidth()
+			.Padding(1.0f, 0.0f)
 			[
 				FilterButton->AsShared()
 			]
