[PCG] Add option to show the default value of attributes

* In the drop down menu of the Attribute list view.
* This default value line will stick at the top (not filtered out, not sorted)
* Default line will be in Italic

#rnx

diff --git a/Engine/Plugins/PCG/Source/PCGEditor/Private/DataVisualizations/PCGDataVisualizationHelpers.cpp b/Engine/Plugins/PCG/Source/PCGEditor/Private/DataVisualizations/PCGDataVisualizationHelpers.cpp
--- a/Engine/Plugins/PCG/Source/PCGEditor/Private/DataVisualizations/PCGDataVisualizationHelpers.cpp
+++ b/Engine/Plugins/PCG/Source/PCGEditor/Private/DataVisualizations/PCGDataVisualizationHelpers.cpp
@@ -13,9 +13,6 @@ namespace PCGDataVisualizationConstants
 {
 	const FString NoneAttributeId = TEXT("@None");
 
-	/** Names of the columns in the attribute list. */
-	const FName NAME_Index = FName(TEXT("$Index"));
-
 	/** Labels of the columns. */
 	const FText TEXT_Index = LOCTEXT("Index", "$Index");
 }
diff --git a/Engine/Plugins/PCG/Source/PCGEditor/Private/Widgets/SPCGEditorGraphAttributeListView.cpp b/Engine/Plugins/PCG/Source/PCGEditor/Private/Widgets/SPCGEditorGraphAttributeListView.cpp
--- a/Engine/Plugins/PCG/Source/PCGEditor/Private/Widgets/SPCGEditorGraphAttributeListView.cpp
+++ b/Engine/Plugins/PCG/Source/PCGEditor/Private/Widgets/SPCGEditorGraphAttributeListView.cpp
@@ -22,6 +22,7 @@
 #include "PCGEditorGraph.h"
 #include "PCGEditorStyle.h"
 #include "PCGEditorUtils.h"
+#include "DataVisualizations/PCGDataVisualizationHelpers.h"
 #include "Nodes/PCGEditorGraphNodeBase.h"
 #include "Widgets/AssetEditorViewport/SPCGEditorViewport.h"
 
@@ -144,12 +145,20 @@ TSharedRef<SWidget> SPCGListViewItemRow::GenerateWidgetForColumn(const FName& Co
 		// We have to make sure the data is still valid otherwise this will crash.
 		if (ensure(SharedAttributeListView->DataPtr) && PCGColumnData->DataAccessor.IsValid() && PCGColumnData->DataKeys.IsValid())
 		{
-			int32 Index = InternalItem->Index;
-			auto Callback = [&PCGColumnData, Index, &RowText, &bShouldComputeCellWidth, &bIsClickable] (auto Dummy)
+			const int32 Index = InternalItem->Index;
+			auto Callback = [&PCGColumnData, Index, &RowText, &bShouldComputeCellWidth, &bIsClickable, ColumnId] (auto Dummy)
 			{
 				using ValueType = decltype(Dummy);
-				ValueType Value{};
-				if (PCGColumnData->DataAccessor->Get<ValueType>(Value, Index, *PCGColumnData->DataKeys))
+				ValueType Value = PCG::Private::MetadataTraits<ValueType>::ZeroValue();
+
+				static const FPCGAttributeAccessorKeysEntries DefaultKeys(PCGInvalidEntryKey);
+				const IPCGAttributeAccessorKeys& DataKeys = (Index == -1 && PCGColumnData->DataAccessor->IsAttribute()) ? DefaultKeys : *PCGColumnData->DataKeys;
+
+				if (Index == -1 && !PCGColumnData->DataAccessor->IsAttribute())
+				{
+					RowText = ColumnId == PCGDataVisualizationConstants::NAME_Index ? LOCTEXT("DefaultIndex", "Default") : FText{};
+				}
+				else if (PCGColumnData->DataAccessor->Get<ValueType>(Value, Index == -1 ? 0 : Index, DataKeys))
 				{
 					if constexpr (PCG::Private::IsOfTypes<ValueType, bool>())
 					{
@@ -206,6 +215,14 @@ TSharedRef<SWidget> SPCGListViewItemRow::GenerateWidgetForColumn(const FName& Co
 	const FMargin Margin = FMargin(2.0f, 0.0f);
 	const ETextOverflowPolicy OverflowPolicy = ETextOverflowPolicy::Ellipsis;
 
+	// TextStyles are pointers so they need to stay alive. Create 2 static copies for each of our style: Normal and Italic.
+	static const FTextBlockStyle NormalTextStyle = FCoreStyle::Get().GetWidgetStyle<FTextBlockStyle>("NormalText");
+	static const FTextBlockStyle ItalicTextStyle = FTextBlockStyle(NormalTextStyle)
+		.SetFont(FCoreStyle::Get().GetFontStyle("NormalFontItalic"));
+
+	// Set the default line (-1) in Italic, Normal for the rest.
+	const FTextBlockStyle* TextStyle = InternalItem->Index == -1 ? &ItalicTextStyle : &NormalTextStyle;
+
 	if (bIsClickable)
 	{
 		// @todo_pcg: We're missing Margin and Overflow policy with this widget, maybe something we can re-add in the future.
@@ -214,6 +231,7 @@ TSharedRef<SWidget> SPCGListViewItemRow::GenerateWidgetForColumn(const FName& Co
 			.Text(RowText)
 			.ToolTipText(LOCTEXT("Hyperlink", "Ctrl + Click to jump to the asset."))
 			.Style(FAppStyle::Get(), "HoverOnlyHyperlink")
+			.TextStyle(TextStyle)
 			.OnNavigate(FSimpleDelegate::CreateLambda([Text = std::move(RowText)](){ SPCGListViewItemRow::OnSoftObjectPathHyperlinkClicked(Text); }));
 
 		return Hyperlink.ToSharedRef();
@@ -223,7 +241,8 @@ TSharedRef<SWidget> SPCGListViewItemRow::GenerateWidgetForColumn(const FName& Co
 		return SNew(STextBlock)
 			.Text(RowText)
 			.OverflowPolicy(OverflowPolicy)
-			.Margin(Margin);
+			.Margin(Margin)
+			.TextStyle(TextStyle);
 	}
 }
 
@@ -636,13 +655,27 @@ void SPCGEditorGraphAttributeListView::Tick(const FGeometry& AllottedGeometry, c
 
 	if (CurrentUpdateTask.IsValid() && CurrentUpdateTask->IsCompleted())
 	{
-		FilteredListViewItems = MoveTemp(CurrentUpdateTask->ListViewItems);
+		// In case of default value, we stick the first line (default line) at the top.
+		if (bShowDefaultValue && !ListViewItems.IsEmpty())
+		{
+			FilteredListViewItems = {ListViewItems[0]};
+			FilteredListViewItems.Append(MoveTemp(CurrentUpdateTask->ListViewItems));
+		}
+		else
+		{
+			FilteredListViewItems = CurrentUpdateTask->ListViewItems;
+		}
+		
 		if (ListView.IsValid())
 		{
 			ListView->SetItemsSource(&FilteredListViewItems);
 			ListView->RequestListRefresh();
 
-			InfoTextBlock->SetText(FText::Format(LOCTEXT("InfoTextBlockFmt", "Showing {0}/{1} entries"), FilteredListViewItems.Num(), ListViewItems.Num()));
+			// Don't take into account the default value in the number of entries.
+			const int32 FilteredNum = bShowDefaultValue ? FMath::Max(FilteredListViewItems.Num() - 1, 0) : FilteredListViewItems.Num();
+			const int32 NonFilteredNum = bShowDefaultValue ? FMath::Max(ListViewItems.Num() - 1, 0) : ListViewItems.Num();
+
+			InfoTextBlock->SetText(FText::Format(LOCTEXT("InfoTextBlockFmt", "Showing {0}/{1} entries"), FilteredNum, NonFilteredNum));
 
 			if (PCGEditorGraphAttributeListView::IsGraphCacheDebuggingEnabled())
 			{
@@ -899,7 +932,9 @@ void SPCGEditorGraphAttributeListView::RefreshAttributeList()
 
 		ListViewItems.Reserve(NumEntries);
 
-		for (int32 Index = 0; Index < NumEntries; ++Index)
+		int32 StartIndex = bShowDefaultValue ? -1 : 0;
+
+		for (int32 Index = StartIndex; Index < NumEntries; ++Index)
 		{
 			PCGListViewItemPtr ListViewItem = MakeShared<FPCGListViewItem>();
 			ListViewItem->Index = Index;
@@ -1223,8 +1258,10 @@ void SPCGEditorGraphAttributeListView::LaunchUpdateTask()
 {
 	// Discarding any currently running updater, task will still run and keep the old object alive but we wont care about the result.
 	// This is done because we cant afford to wait for task completion before starting a new task.
+	// Note that in the case we are showing the default value, skip the first element (index -1).
+	// We will not use the default value for sorting/filtering.
 	CurrentUpdateTask.Reset();
-	CurrentUpdateTask = MakeShared<FPCGListViewUpdater>(ListViewItems,
+	CurrentUpdateTask = MakeShared<FPCGListViewUpdater>(bShowDefaultValue ? MakeArrayView(ListViewItems).RightChop(1) : MakeArrayView(ListViewItems),
 		PCGColumnData,
 		SortMode,
 		SortingColumn,
@@ -1329,6 +1366,20 @@ TSharedRef<SWidget> SPCGEditorGraphAttributeListView::OnGenerateAdditionalOperat
 		NAME_None,
 		EUserInterfaceActionType::Button);
 
+	MenuBuilder.AddSeparator();
+
+	MenuBuilder.AddMenuEntry(
+		LOCTEXT("ShowDefaultValue", "Show Default Value for Attributes"),
+		LOCTEXT("ShowDefaultValueTooltip", "[ADVANCED] For all non-$ attributes, show the underlying default value. For all $ attributes, it will be 0. Stick at the top"),
+		FSlateIcon(),
+		FUIAction(
+	FExecuteAction::CreateSP(this, &SPCGEditorGraphAttributeListView::OnToggleShowDefaultValue),
+			FCanExecuteAction(),
+			FIsActionChecked::CreateSP(this, &SPCGEditorGraphAttributeListView::IsShowingDefaultValue)
+		),
+		NAME_None,
+		EUserInterfaceActionType::ToggleButton);
+
 	return MenuBuilder.MakeWidget();
 }
 
@@ -1941,4 +1992,16 @@ bool SPCGEditorGraphAttributeListView::IsViewportOpen() const
 	return PCGEditorPtr.IsValid() && PCGEditorPtr.Pin()->IsPanelCurrentlyOpen(ViewportEditorPanel);
 }
 
+
+void SPCGEditorGraphAttributeListView::OnToggleShowDefaultValue()
+{
+	bShowDefaultValue = !bShowDefaultValue;
+	RequestRefresh();
+}
+
+bool SPCGEditorGraphAttributeListView::IsShowingDefaultValue() const
+{
+	return bShowDefaultValue;
+}
+
 #undef LOCTEXT_NAMESPACE
diff --git a/Engine/Plugins/PCG/Source/PCGEditor/Private/Widgets/SPCGEditorGraphAttributeListView.h b/Engine/Plugins/PCG/Source/PCGEditor/Private/Widgets/SPCGEditorGraphAttributeListView.h
--- a/Engine/Plugins/PCG/Source/PCGEditor/Private/Widgets/SPCGEditorGraphAttributeListView.h
+++ b/Engine/Plugins/PCG/Source/PCGEditor/Private/Widgets/SPCGEditorGraphAttributeListView.h
@@ -72,7 +72,7 @@ class FPCGListViewUpdater : public TSharedFromThis<FPCGListViewUpdater>
 {
 public:
 	FPCGListViewUpdater(
-		const TArray<PCGListViewItemPtr>& InListViewItems,
+		const TArrayView<const PCGListViewItemPtr>& InListViewItems,
 		const TMap<FName, FPCGColumnData>& InColumnData,
 		const EColumnSortMode::Type InSortMode,
 		const FName InSortingColumn,
@@ -230,6 +230,9 @@ private:
 	void SaveData(bool bUsePinIndex, bool bUseDataIndex);
 	bool CanSaveData(bool bUsePinIndex, bool bUseDataIndex) const;
 
+	void OnToggleShowDefaultValue();
+	bool IsShowingDefaultValue() const;
+
 	TSharedRef<ITableRow> OnGenerateRow(PCGListViewItemPtr Item, const TSharedRef<STableViewBase>& OwnerTable);
 	void OnItemDoubleClicked(PCGListViewItemPtr Item) const;
 	TSharedPtr<SWidget> OnItemsContextMenu();
@@ -332,6 +335,8 @@ private:
 	/** True if this property view is currently locked (i.e. The objects being observed are not changed automatically due to user selection)*/
 	bool bIsLocked : 1 = false;
 
+	bool bShowDefaultValue : 1 = false;
+
 	TSharedPtr<FPCGListViewUpdater> CurrentUpdateTask = nullptr;
 
 	/** In some cases, like when inspecting temporary collapsed point data, there is no owner of the data.
diff --git a/Engine/Plugins/PCG/Source/PCGEditor/Public/DataVisualizations/PCGDataVisualizationHelpers.h b/Engine/Plugins/PCG/Source/PCGEditor/Public/DataVisualizations/PCGDataVisualizationHelpers.h
--- a/Engine/Plugins/PCG/Source/PCGEditor/Public/DataVisualizations/PCGDataVisualizationHelpers.h
+++ b/Engine/Plugins/PCG/Source/PCGEditor/Public/DataVisualizations/PCGDataVisualizationHelpers.h
@@ -7,6 +7,12 @@
 
 class UPCGData;
 
+namespace PCGDataVisualizationConstants
+{
+	/** Names of the columns in the attribute list. */
+	const FName NAME_Index = FName(TEXT("$Index"));
+}
+
 namespace PCGDataVisualizationHelpers
 {
 	struct FColumnInfoOverrides
