[PCG] Virtual texture priming improvements
* Allow landscape data to opt out of virtual texture sampling and virtual texture normal sampling
* Enable virtual texture priming by default
* Add desired texel size to FPCGVirtualTexturePrimingInfo to allow computing the correct mip to sample
* Draw more accurate debug bounds for the preload region

#rnx

diff --git a/Engine/Plugins/PCG/Source/PCG/Private/Compute/DataInterfaces/PCGLandscapeDataInterface.cpp b/Engine/Plugins/PCG/Source/PCG/Private/Compute/DataInterfaces/PCGLandscapeDataInterface.cpp
--- a/Engine/Plugins/PCG/Source/PCG/Private/Compute/DataInterfaces/PCGLandscapeDataInterface.cpp
+++ b/Engine/Plugins/PCG/Source/PCG/Private/Compute/DataInterfaces/PCGLandscapeDataInterface.cpp
@@ -28,20 +28,6 @@
 
 #define LOCTEXT_NAMESPACE "UPCGLandscapeDataInterface"
 
-namespace PCGLandscapeDataInterface
-{
-	static TAutoConsoleVariable<bool> CVarUseLandscapeVTs(
-		TEXT("pcg.GPU.UseLandscapeVTs"),
-		true,
-		TEXT("Allow the landscape data interface to use landscape virtual textures if available."));
-
-	/** Normal VTs may not match the actual landscape surface normals, so generally prefer computing normals from multiple height samples. */
-	static TAutoConsoleVariable<bool> CVarUseLandscapeNormalVTs(
-		TEXT("pcg.GPU.VirtualTextures.EnableNormals"),
-		false,
-		TEXT("Allow the landscape data interface to use landscape normal virtual textures if available."));
-}
-
 /** RenderResource used to hold textures generated by this DI, pulled from the collision geometry of the terrain. */
 class FPCGLandscapeTextureResource : public FRenderResource
 {
@@ -300,7 +286,7 @@ void UPCGLandscapeDataProvider::Initialize(const UComputeDataInterface* InDataIn
 	const UPCGComponent* SourceComponent = Cast<UPCGComponent>(ExecutionSource);
 
 	// Sampling of virtual textures is not supported outside of runtime generation at this time, because the RuntimeGenScheduler is responsible for priming the virtual textures.
-	bCanSampleVirtualTextures = SourceComponent ? SourceComponent->IsManagedByRuntimeGenSystem() : false;
+	const bool bCanSampleVirtualTextures = SourceComponent ? SourceComponent->IsManagedByRuntimeGenSystem() : false;
 
 	auto TestLandscape = [World, &ComponentBounds](const ALandscape* InLandscape)
 	{
@@ -340,7 +326,7 @@ void UPCGLandscapeDataProvider::Initialize(const UComputeDataInterface* InDataIn
 					{
 						if (TestLandscape(Landscape))
 						{
-							InitFromLandscape(Landscape, ComponentBounds);
+							InitFromLandscape(Landscape, ComponentBounds, bCanSampleVirtualTextures && LandscapeData->CanSampleVirtualTextures(), LandscapeData->CanSampleVirtualTextureNormals());
 							break;
 						}
 					}
@@ -350,7 +336,7 @@ void UPCGLandscapeDataProvider::Initialize(const UComputeDataInterface* InDataIn
 	}
 }
 
-void UPCGLandscapeDataProvider::InitFromLandscape(ALandscape* InLandscape, const FBox& Bounds)
+void UPCGLandscapeDataProvider::InitFromLandscape(ALandscape* InLandscape, const FBox& Bounds, bool bAllowSampleVirtualTexture, bool bAllowSampleVirtualTextureNormal)
 {
 	TRACE_CPUPROFILER_EVENT_SCOPE(UPCGLandscapeDataProvider::InitFromLandscape);
 	check(IsInGameThread());
@@ -360,9 +346,9 @@ void UPCGLandscapeDataProvider::InitFromLandscape(ALandscape* InLandscape, const
 		return;
 	}
 
-	if (bCanSampleVirtualTextures && PCGLandscapeDataInterface::CVarUseLandscapeVTs.GetValueOnAnyThread()) 
+	if (bAllowSampleVirtualTexture)
 	{
-		InitializeRuntimeVirtualTextures(InLandscape);
+		InitializeRuntimeVirtualTextures(InLandscape, bAllowSampleVirtualTextureNormal);
 	}
 
 	// If the height can't be accessed via RVT, build a physics based height texture instead.
@@ -374,7 +360,7 @@ void UPCGLandscapeDataProvider::InitFromLandscape(ALandscape* InLandscape, const
 	LandscapeGridSize = FVector2D(InLandscape->GetTransform().GetScale3D());
 }
 
-void UPCGLandscapeDataProvider::InitializeRuntimeVirtualTextures(ALandscape* InLandscape)
+void UPCGLandscapeDataProvider::InitializeRuntimeVirtualTextures(ALandscape* InLandscape, bool bAllowSampleVirtualTextureNormal)
 {
 	TRACE_CPUPROFILER_EVENT_SCOPE(UPCGLandscapeDataProvider::InitializeRuntimeVirtualTextures);
 	check(InLandscape);
@@ -444,7 +430,7 @@ void UPCGLandscapeDataProvider::InitializeRuntimeVirtualTextures(ALandscape* InL
 		?  InLandscape->RuntimeVirtualTextures[HeightVirtualTextureIndex]
 		: nullptr;
 
-	if (PCGLandscapeDataInterface::CVarUseLandscapeNormalVTs.GetValueOnAnyThread())
+	if (bAllowSampleVirtualTextureNormal)
 	{
 		NormalVirtualTexture = InLandscape->RuntimeVirtualTextures.IsValidIndex(NormalVirtualTextureIndex)
 			? InLandscape->RuntimeVirtualTextures[NormalVirtualTextureIndex]
@@ -569,7 +555,6 @@ void UPCGLandscapeDataProvider::Reset()
 	NormalVirtualTexture = nullptr;
 
 	Resource.Reset();
-	bCanSampleVirtualTextures = false;
 	bBaseColorSRGB = false;
 	BaseColorVirtualTextureUnpackType = PCGVirtualTextureCommon::EBaseColorUnpackType::None;
 	NormalVirtualTextureMode = ERuntimeVirtualTextureMaterialType::Count;
diff --git a/Engine/Plugins/PCG/Source/PCG/Private/Compute/DataInterfaces/PCGLandscapeDataInterface.h b/Engine/Plugins/PCG/Source/PCG/Private/Compute/DataInterfaces/PCGLandscapeDataInterface.h
--- a/Engine/Plugins/PCG/Source/PCG/Private/Compute/DataInterfaces/PCGLandscapeDataInterface.h
+++ b/Engine/Plugins/PCG/Source/PCG/Private/Compute/DataInterfaces/PCGLandscapeDataInterface.h
@@ -76,8 +76,8 @@ public:
 	//~ End UComputeDataProvider Interface
 
 private:
-	void InitFromLandscape(ALandscape* InLandscape, const FBox& Bounds);
-	void InitializeRuntimeVirtualTextures(ALandscape* InLandscape);
+	void InitFromLandscape(ALandscape* InLandscape, const FBox& Bounds, bool bAllowSampleVirtualTexture, bool bAllowSampleVirtualTextureNormal);
+	void InitializeRuntimeVirtualTextures(ALandscape* InLandscape, bool bAllowSampleVirtualTextureNormal);
 	void InitializeFromLandscapeCollision(ALandscape* InLandscape, const FBox& Bounds);
 
 private:
@@ -91,7 +91,6 @@ private:
 	TObjectPtr<const URuntimeVirtualTexture> NormalVirtualTexture = nullptr;
 
 	TSharedPtr<FPCGLandscapeResource> Resource;
-	bool bCanSampleVirtualTextures = false;
 	bool bBaseColorSRGB = false;
 	PCGVirtualTextureCommon::EBaseColorUnpackType BaseColorVirtualTextureUnpackType = PCGVirtualTextureCommon::EBaseColorUnpackType::None;
 	ERuntimeVirtualTextureMaterialType NormalVirtualTextureMode = ERuntimeVirtualTextureMaterialType::Count;
diff --git a/Engine/Plugins/PCG/Source/PCG/Private/RuntimeGen/PCGRuntimeGenScheduler.cpp b/Engine/Plugins/PCG/Source/PCG/Private/RuntimeGen/PCGRuntimeGenScheduler.cpp
--- a/Engine/Plugins/PCG/Source/PCG/Private/RuntimeGen/PCGRuntimeGenScheduler.cpp
+++ b/Engine/Plugins/PCG/Source/PCG/Private/RuntimeGen/PCGRuntimeGenScheduler.cpp
@@ -34,6 +34,7 @@
 namespace PCGRuntimeGenSchedulerConstants
 {
 	const FString PooledPartitionActorName = TEXT("PCGRuntimePartitionGridActor_POOLED");
+	static constexpr float MinWorldVirtualTextureTexelSize = 0.1;
 }
 
 namespace PCGRuntimeGenSchedulerHelpers
@@ -98,7 +99,7 @@ namespace PCGRuntimeGenSchedulerHelpers
 
 	static TAutoConsoleVariable<bool> CVarEnableVirtualTexturePriming(
 		TEXT("pcg.VirtualTexturePriming.Enable"),
-		false,
+		true,
 		TEXT("Enable priming of virtual textures for PCG Components which request it."));
 
 	static TAutoConsoleVariable<bool> CVarDebugDrawTexturePrimingBounds(
@@ -986,7 +987,7 @@ void FPCGRuntimeGenScheduler::TickRequestVirtualTexturePriming(const TSet<IPCGGe
 
 		for (const FPCGVirtualTexturePrimingInfo* PrimingInfo : PrimingInfos)
 		{
-			if (!PrimingInfo || !PrimingInfo->VirtualTexture)
+			if (!PrimingInfo || !PrimingInfo->VirtualTexture || PrimingInfo->WorldTexelSize < PCGRuntimeGenSchedulerConstants::MinWorldVirtualTextureTexelSize)
 			{
 				continue;
 			}
@@ -1014,28 +1015,38 @@ void FPCGRuntimeGenScheduler::TickRequestVirtualTexturePriming(const TSet<IPCGGe
 					}
 
 					FVector GenSourcePosition = GenSourcePositionOptional.GetValue();
+
 					if (OriginalComponent->Use2DGrid())
 					{
 						GenSourcePosition.Z = OriginalComponentBounds.GetCenter().Z;
 					}
+
 					const FSphere PrimingBounds = FSphere(GenSourcePosition, PrimingRadius);
 
-					const FBox OverlapBox = OriginalComponentBounds.Overlap(FBox(PrimingBounds));
-					if (OverlapBox.GetVolume() > 0)
+					if (OriginalComponentBounds.Intersect(FBox(PrimingBounds)))
 					{
 						TRACE_CPUPROFILER_EVENT_SCOPE(RequestVirtualTexturePreload);
 
-						VirtualTextureComponent->RequestPreload(PrimingBounds, /*Level=*/0);
+						const double BoundsMaxExtent = FMath::Max(VirtualTextureComponent->Bounds.BoxExtent.X, VirtualTextureComponent->Bounds.BoxExtent.Y);
+						const int32 VirtualTextureSizeTexels = FMath::Max(1, PrimingInfo->VirtualTexture->GetSize());
+						const int32 SizeTexelsLog2 = FMath::FloorLog2(VirtualTextureSizeTexels);
+						const int32 RequestedNumTexels = FMath::Max(1, BoundsMaxExtent / PrimingInfo->WorldTexelSize);
+						const int32 RequestedTexelsLog2 = FMath::FloorLog2(RequestedNumTexels);
+						const int32 MipLevel = FMath::Max(0, SizeTexelsLog2 - RequestedTexelsLog2);
+
+						VirtualTextureComponent->RequestPreload(PrimingBounds, MipLevel);
 
 						if (PCGRuntimeGenSchedulerHelpers::CVarDebugDrawTexturePrimingBounds.GetValueOnAnyThread())
 						{
 							check(World);
 
-							DrawDebugSphere(World,
-								PrimingBounds.Center,
-								PrimingBounds.W,
-								8,
-								FColor::Red,
+							DrawDebugCylinder(
+								World,
+								FVector(GenSourcePosition.X, GenSourcePosition.Y, OriginalComponentBounds.Min.Z),
+								FVector(GenSourcePosition.X, GenSourcePosition.Y, OriginalComponentBounds.Max.Z),
+								/*Radius=*/PrimingBounds.W,
+								/*Segments=*/8,
+								/*Color=*/FColor::Red,
 								/*bPersistentLines=*/false,
 								/*LifeTime=*/0.02f);
 						}
diff --git a/Engine/Plugins/PCG/Source/PCG/Public/Data/PCGLandscapeData.h b/Engine/Plugins/PCG/Source/PCG/Public/Data/PCGLandscapeData.h
--- a/Engine/Plugins/PCG/Source/PCG/Public/Data/PCGLandscapeData.h
+++ b/Engine/Plugins/PCG/Source/PCG/Public/Data/PCGLandscapeData.h
@@ -37,6 +37,17 @@ struct FPCGLandscapeDataProps
 	/** Controls whether the component coordinates will be added the point as attributes ('CoordinateX', 'CoordinateY') */
 	UPROPERTY(BlueprintReadWrite, EditAnywhere, Category = Settings, meta = (PCG_Overridable))
 	bool bGetComponentCoordinates = false;
+
+	/** Controls whether the landscape will try to sample from the landscape virtual textures (if they exist). Only relevant to GPU sampling. */
+	UPROPERTY(BlueprintReadWrite, EditAnywhere, Category = GPU, meta = (PCG_Overridable))
+	bool bSampleVirtualTextures = true;
+
+	/**
+	 * Controls whether the landscape will try to sample normals from a normals virtual texture (if it exists), otherwise computes normals from multiple height samples. Only relevant to GPU sampling.
+	 * Note that normal virtual textures may be detail normals and not match the actual landscape surface normals, so enable this with caution. Requires bSampleVirtualTextures to be true.
+	 */
+	UPROPERTY(BlueprintReadWrite, EditAnywhere, Category = GPU, meta = (PCG_Overridable, EditCondition = "bSampleVirtualTextures", EditConditionHides))
+	bool bSampleVirtualTextureNormals = false;
 };
 
 /**
@@ -84,6 +95,8 @@ public:
 	TArray<TSoftObjectPtr<ALandscapeProxy>> Landscapes;
 
 	bool IsUsingMetadata() const { return DataProps.bGetLayerWeights; }
+	bool CanSampleVirtualTextures() const { return DataProps.bSampleVirtualTextures; }
+	bool CanSampleVirtualTextureNormals() const { return CanSampleVirtualTextures() && DataProps.bSampleVirtualTextureNormals; }
 
 protected:
 	/** Returns the landscape info associated to the first landscape that contains the given position
diff --git a/Engine/Plugins/PCG/Source/PCG/Public/RuntimeGen/PCGRuntimeGenScheduler.h b/Engine/Plugins/PCG/Source/PCG/Public/RuntimeGen/PCGRuntimeGenScheduler.h
--- a/Engine/Plugins/PCG/Source/PCG/Public/RuntimeGen/PCGRuntimeGenScheduler.h
+++ b/Engine/Plugins/PCG/Source/PCG/Public/RuntimeGen/PCGRuntimeGenScheduler.h
@@ -31,6 +31,10 @@ struct FPCGVirtualTexturePrimingInfo
 	/** Largest grid on which this virtual texture is sampled. */
 	UPROPERTY(EditAnywhere, Category = "")
 	EPCGHiGenGrid Grid = EPCGHiGenGrid::Grid32;
+
+	/** Desired world size (cm) of a texel in the primed virtual texture. Determines what mip level will be primed. */
+	UPROPERTY(EditAnywhere, Category = "", meta = (ClampMin=0.1))
+	float WorldTexelSize = 100.0f;
 };
 
 /**
