[PCG] Make a pass on default value changes

* First value is no more set as default value
* Make a pass where we were getting the default value, to get the first value.
* Changed the metadata accessor keys to always starts from 0, since we enforce it correctly now.

#rnx

diff --git a/Engine/Plugins/PCG/Source/PCG/Private/Elements/Grammar/PCGDuplicateCrossSections.cpp b/Engine/Plugins/PCG/Source/PCG/Private/Elements/Grammar/PCGDuplicateCrossSections.cpp
--- a/Engine/Plugins/PCG/Source/PCG/Private/Elements/Grammar/PCGDuplicateCrossSections.cpp
+++ b/Engine/Plugins/PCG/Source/PCG/Private/Elements/Grammar/PCGDuplicateCrossSections.cpp
@@ -96,7 +96,7 @@ bool FPCGDuplicateCrossSectionsElement::ExecuteInternal(FPCGContext* InContext)
 				return false;
 			}
 
-			if (!Accessor->Get(OutValue, FPCGAttributeAccessorKeysEntries(PCGInvalidEntryKey), EPCGAttributeAccessorFlags::AllowBroadcastAndConstructible))
+			if (!Accessor->Get(OutValue, FPCGAttributeAccessorKeysEntries(PCGFirstEntryKey), EPCGAttributeAccessorFlags::AllowBroadcastAndConstructible))
 			{
 				PCGLog::Metadata::LogFailToGetAttributeError<T>(Selector, Accessor.Get(), InContext);
 				return false;
@@ -128,7 +128,7 @@ bool FPCGDuplicateCrossSectionsElement::ExecuteInternal(FPCGContext* InContext)
 				PCGLog::Metadata::LogFailToCreateAccessorError(Selector, InContext);
 			}
 			// Otherwise, get the value, if it fails, the attribute wasn't compatible
-			else if (!SeedAccessor->Get(AdditionalSeed, FPCGAttributeAccessorKeysEntries(PCGInvalidEntryKey), EPCGAttributeAccessorFlags::AllowBroadcastAndConstructible))
+			else if (!SeedAccessor->Get(AdditionalSeed, FPCGAttributeAccessorKeysEntries(PCGFirstEntryKey), EPCGAttributeAccessorFlags::AllowBroadcastAndConstructible))
 			{
 				PCGLog::Metadata::LogFailToGetAttributeError<int32>(Selector, SeedAccessor.Get(), InContext);
 			}
diff --git a/Engine/Plugins/PCG/Source/PCG/Private/Elements/Grammar/PCGSubdivideSpline.cpp b/Engine/Plugins/PCG/Source/PCG/Private/Elements/Grammar/PCGSubdivideSpline.cpp
--- a/Engine/Plugins/PCG/Source/PCG/Private/Elements/Grammar/PCGSubdivideSpline.cpp
+++ b/Engine/Plugins/PCG/Source/PCG/Private/Elements/Grammar/PCGSubdivideSpline.cpp
@@ -362,7 +362,7 @@ bool FPCGSubdivideSplineElement::ExecuteInternal(FPCGContext* InContext) const
 				continue;
 			}
 
-			if (!HeightAccessor->Get<double>(ModuleHeight, FPCGAttributeAccessorKeysEntries(PCGInvalidEntryKey), EPCGAttributeAccessorFlags::AllowBroadcastAndConstructible))
+			if (!HeightAccessor->Get<double>(ModuleHeight, FPCGAttributeAccessorKeysEntries(PCGFirstEntryKey), EPCGAttributeAccessorFlags::AllowBroadcastAndConstructible))
 			{
 				PCGLog::Metadata::LogFailToGetAttributeError(Selector, InContext);
 				continue;
@@ -417,7 +417,7 @@ bool FPCGSubdivideSplineElement::ExecuteInternal(FPCGContext* InContext) const
 				PCGLog::Metadata::LogFailToCreateAccessorError(Selector, InContext);
 			}
 			// Otherwise, get the value, if it fails, the attribute wasn't compatible
-			else if (!SeedAccessor->Get(Parameters.AdditionalSeed, FPCGAttributeAccessorKeysEntries(PCGInvalidEntryKey), EPCGAttributeAccessorFlags::AllowBroadcastAndConstructible))
+			else if (!SeedAccessor->Get(Parameters.AdditionalSeed, FPCGAttributeAccessorKeysEntries(PCGFirstEntryKey), EPCGAttributeAccessorFlags::AllowBroadcastAndConstructible))
 			{
 				PCGLog::Metadata::LogFailToGetAttributeError<int32>(Selector, SeedAccessor.Get(), InContext);
 			}
@@ -427,7 +427,7 @@ bool FPCGSubdivideSplineElement::ExecuteInternal(FPCGContext* InContext) const
 		{
 			const FPCGAttributePropertyInputSelector Selector = Settings->GrammarSelection.GrammarAttribute.CopyAndFixLast(InputPolyLineData);
 			const TUniquePtr<const IPCGAttributeAccessor> GrammarAccessor = PCGAttributeAccessorHelpers::CreateConstAccessor(InputPolyLineData, Selector);
-			const FPCGAttributeAccessorKeysEntries Keys(PCGInvalidEntryKey);
+			const FPCGAttributeAccessorKeysEntries Keys(PCGFirstEntryKey);
 			if (!GrammarAccessor)
 			{
 				PCGLog::Metadata::LogFailToCreateAccessorError(Selector, InContext);
diff --git a/Engine/Plugins/PCG/Source/PCG/Private/Elements/Grammar/PCGSubdivisionBase.cpp b/Engine/Plugins/PCG/Source/PCG/Private/Elements/Grammar/PCGSubdivisionBase.cpp
--- a/Engine/Plugins/PCG/Source/PCG/Private/Elements/Grammar/PCGSubdivisionBase.cpp
+++ b/Engine/Plugins/PCG/Source/PCG/Private/Elements/Grammar/PCGSubdivisionBase.cpp
@@ -132,7 +132,7 @@ PCGGrammar::FTokenizedGrammar FPCGSubdivisionBaseElement::GetTokenizedGrammar(FP
 			return {};
 		}
 
-		if (!Accessor->Get(Grammar, FPCGAttributeAccessorKeysEntries(PCGInvalidEntryKey), EPCGAttributeAccessorFlags::AllowBroadcastAndConstructible))
+		if (!Accessor->Get(Grammar, FPCGAttributeAccessorKeysEntries(PCGFirstEntryKey), EPCGAttributeAccessorFlags::AllowBroadcastAndConstructible))
 		{
 			PCGLog::Metadata::LogFailToGetAttributeError<FString>(Selector, Accessor.Get(), InContext);
 			return {};
diff --git a/Engine/Plugins/PCG/Source/PCG/Private/Elements/PCGAttributeGetFromIndexElement.cpp b/Engine/Plugins/PCG/Source/PCG/Private/Elements/PCGAttributeGetFromIndexElement.cpp
--- a/Engine/Plugins/PCG/Source/PCG/Private/Elements/PCGAttributeGetFromIndexElement.cpp
+++ b/Engine/Plugins/PCG/Source/PCG/Private/Elements/PCGAttributeGetFromIndexElement.cpp
@@ -93,14 +93,6 @@ bool FPCGAttributeGetFromIndexElement::ExecuteInternal(FPCGContext* Context) con
 		TArray<EPCGMetadataTypes> AttributeTypes;
 		SubParam->Metadata->GetAttributes(AttributeNames, AttributeTypes);
 
-		for (const FName AttributeName : AttributeNames)
-		{
-			if (FPCGMetadataAttributeBase* Attribute = SubParam->Metadata->GetMutableAttribute(AttributeName))
-			{
-				Attribute->SetDefaultValueToFirstEntry();
-			}
-		}
-
 		FPCGTaggedData& TaggedData = Context->OutputData.TaggedData.Emplace_GetRef(Input);
 		TaggedData.Data = SubParam;
 	}
diff --git a/Engine/Plugins/PCG/Source/PCG/Private/Helpers/PCGMetadataHelpers.cpp b/Engine/Plugins/PCG/Source/PCG/Private/Helpers/PCGMetadataHelpers.cpp
--- a/Engine/Plugins/PCG/Source/PCG/Private/Helpers/PCGMetadataHelpers.cpp
+++ b/Engine/Plugins/PCG/Source/PCG/Private/Helpers/PCGMetadataHelpers.cpp
@@ -192,7 +192,7 @@ namespace PCGMetadataHelpers
 				if (FPCGMetadataAttributeBase* NewAttr = TargetMetadataDomain->CopyAttribute(SourceAttribute, LocalDestinationAttribute.Name, /*bKeepParent=*/false, /*bCopyEntries=*/bCopyEntries, /*bCopyValues=*/true))
 				{
 					// And finally, we create our source and target keys, to get all the metadata entry keys, to remap if we are not the same origin and that we have multiple entries.
-					if (!bCopyEntries && bSourceSupportsMultiEntries)
+					if (!bCopyEntries)
 					{
 						TUniquePtr<const IPCGAttributeAccessorKeys> SourceKeys = PCGAttributeAccessorHelpers::CreateConstKeys(SourceData, InputSource);
 						TUniquePtr<IPCGAttributeAccessorKeys> TargetKeys = PCGAttributeAccessorHelpers::CreateKeys(TargetData, OutputTarget);
@@ -205,8 +205,10 @@ namespace PCGMetadataHelpers
 							TArray<const PCGMetadataEntryKey*> AllSourceEntryKeysPtr;
 							TArray<PCGMetadataEntryKey*> AllTargetEntryKeysPtr;
 
-							AllSourceEntryKeysPtr.SetNumUninitialized(FMath::Min(SourceKeys->GetNum(), TargetKeys->GetNum()));
-							AllTargetEntryKeysPtr.SetNumUninitialized(TargetKeys->GetNum());
+							int32 NumTargetKeys = bSourceSupportsMultiEntries ? TargetKeys->GetNum() : 1;
+
+							AllSourceEntryKeysPtr.SetNumUninitialized(FMath::Min(SourceKeys->GetNum(), NumTargetKeys));
+							AllTargetEntryKeysPtr.SetNumUninitialized(NumTargetKeys);
 
 							if (SourceKeys->GetKeys<PCGMetadataEntryKey>(0, AllSourceEntryKeysPtr) && TargetKeys->GetKeys<PCGMetadataEntryKey>(0, AllTargetEntryKeysPtr))
 							{
@@ -240,7 +242,7 @@ namespace PCGMetadataHelpers
 								for (int EntryIndex = 0; EntryIndex < AllTargetEntryKeysPtr.Num(); ++EntryIndex)
 								{
 									PCGMetadataEntryKey& EntryKey = *AllTargetEntryKeysPtr[EntryIndex];
-									if (EntryKey == PCGInvalidEntryKey || (EntryKey < TargetMetadataDomain->GetItemKeyCountForParent()))
+									if (EntryKey == PCGInvalidEntryKey || (bSourceSupportsMultiEntries && EntryKey < TargetMetadataDomain->GetItemKeyCountForParent()))
 									{
 										AllTargetEntryKeysPtrTemp.Add(&EntryKey);
 									}
@@ -253,19 +255,9 @@ namespace PCGMetadataHelpers
 								}
 
 								NewAttr->SetValuesFromValueKeys(AllTargetEntryKeysPtr, FinalValueKeys, /*bResetValueOnDefaultValueKey=*/true);
-
-								// To keep the previous behavior for single value attributes, we force the copied attribute to have its default value set to the first entry (only done if there is just a single valid entry key).
-								if (AllSourceEntryKeys.Num() == 1 && AllSourceEntryKeys[0] != PCGInvalidEntryKey)
-								{
-									NewAttr->SetDefaultValueToFirstEntry();
-								}
 							}
 						}
 					}
-					else if (!bCopyEntries && !bSourceSupportsMultiEntries)
-					{
-						NewAttr->SetDefaultValueToFirstEntry();
-					}
 				}
 				else
 				{
diff --git a/Engine/Plugins/PCG/Source/PCG/Private/Metadata/Accessors/PCGAttributeAccessorKeys.cpp b/Engine/Plugins/PCG/Source/PCG/Private/Metadata/Accessors/PCGAttributeAccessorKeys.cpp
--- a/Engine/Plugins/PCG/Source/PCG/Private/Metadata/Accessors/PCGAttributeAccessorKeys.cpp
+++ b/Engine/Plugins/PCG/Source/PCG/Private/Metadata/Accessors/PCGAttributeAccessorKeys.cpp
@@ -67,16 +67,13 @@ void FPCGAttributeAccessorKeysEntries::InitializeFromMetadata(const FPCGMetadata
 	}
 
 	check(ExtractedEntries.IsEmpty());
-
-	const PCGMetadataEntryKey ItemKeyLowerBound = Metadata->GetItemKeyCountForParent();
 	const PCGMetadataEntryKey ItemKeyUpperBound = Metadata->GetItemCountForChild();
-	const int64 Count = ItemKeyUpperBound - ItemKeyLowerBound;
 
-	if (Count > 0)
+	if (ItemKeyUpperBound > 0)
 	{
-		ExtractedEntries.Reserve(Count);
+		ExtractedEntries.Reserve(ItemKeyUpperBound);
 
-		for (PCGMetadataEntryKey Entry = ItemKeyLowerBound; Entry < ItemKeyUpperBound; ++Entry)
+		for (PCGMetadataEntryKey Entry = 0; Entry < ItemKeyUpperBound; ++Entry)
 		{
 			ExtractedEntries.Add(Entry);
 		}
diff --git a/Engine/Plugins/PCG/Source/PCG/Private/Metadata/PCGMetadataPartitionCommon.cpp b/Engine/Plugins/PCG/Source/PCG/Private/Metadata/PCGMetadataPartitionCommon.cpp
--- a/Engine/Plugins/PCG/Source/PCG/Private/Metadata/PCGMetadataPartitionCommon.cpp
+++ b/Engine/Plugins/PCG/Source/PCG/Private/Metadata/PCGMetadataPartitionCommon.cpp
@@ -553,11 +553,6 @@ namespace PCGMetadataPartitionCommon
 				{
 					NewAttribute->SetValue(EntryKeys[i], OriginalAttribute, Indices[i]);
 				}
-
-				if (Indices.Num() == 1)
-				{
-					NewAttribute->SetDefaultValueToFirstEntry();
-				}
 			}
 		}
 
@@ -625,11 +620,6 @@ namespace PCGMetadataPartitionCommon
 				check(OriginalAttribute && NewAttribute);
 
 				NewAttribute->SetValue(CurrentEntryKey, OriginalAttribute, Indices[0]);
-
-				if (Partition.Num() == 1)
-				{
-					NewAttribute->SetDefaultValueToFirstEntry();
-				}
 			}
 		}
 
diff --git a/Engine/Plugins/PCG/Source/PCG/Private/PCGContext.cpp b/Engine/Plugins/PCG/Source/PCG/Private/PCGContext.cpp
--- a/Engine/Plugins/PCG/Source/PCG/Private/PCGContext.cpp
+++ b/Engine/Plugins/PCG/Source/PCG/Private/PCGContext.cpp
@@ -41,7 +41,7 @@ namespace PCGContextHelpers
 	bool GetOverrideParamValue(const IPCGAttributeAccessor& InAccessor, T& OutValue)
 	{
 		// Override were using the first entry (0) by default.
-		FPCGAttributeAccessorKeysEntries FirstEntry(PCGMetadataEntryKey(0));
+		FPCGAttributeAccessorKeysEntries FirstEntry(PCGFirstEntryKey);
 		return InAccessor.Get<T>(OutValue, FirstEntry, EPCGAttributeAccessorFlags::AllowBroadcastAndConstructible);
 	}
 }
@@ -414,7 +414,7 @@ void FPCGContext::OverrideSettings()
 			using PropertyType = decltype(Dummy);
 
 			// Override were using the first entry (0) by default.
-			FPCGAttributeAccessorKeysEntries FirstEntry(PCGMetadataEntryKey(0));
+			FPCGAttributeAccessorKeysEntries FirstEntry(PCGFirstEntryKey);
 					
 			PropertyType Value{};
 			if (!PCGContextHelpers::GetOverrideParamValue(*AttributeAccessor, Value))
diff --git a/Engine/Plugins/PCG/Source/PCG/Private/Tests/Accessors/PCGAttributeAccessorTest.cpp b/Engine/Plugins/PCG/Source/PCG/Private/Tests/Accessors/PCGAttributeAccessorTest.cpp
--- a/Engine/Plugins/PCG/Source/PCG/Private/Tests/Accessors/PCGAttributeAccessorTest.cpp
+++ b/Engine/Plugins/PCG/Source/PCG/Private/Tests/Accessors/PCGAttributeAccessorTest.cpp
@@ -593,7 +593,7 @@ bool FPCGAttributeAccessorSingleSetPointsTest::RunTest(const FString& Parameters
 	// Points should not have a different entry key, if it was not invalid in the first place
 	if (!TestEqual("Point 0 should have entry 2", MetadataEntryRange[0], PCGMetadataEntryKey(2))
 		|| !TestEqual("Point 1 should have entry 1", MetadataEntryRange[1], PCGMetadataEntryKey(1))
-		|| !TestEqual("Point 2 should have entry 0", MetadataEntryRange[2], PCGMetadataEntryKey(0)))
+		|| !TestEqual("Point 2 should have entry 0", MetadataEntryRange[2], PCGFirstEntryKey))
 	{
 		return false;
 	}
@@ -634,7 +634,7 @@ bool FPCGAttributeAccessorSetRangePointsTest::RunTest(const FString& Parameters)
 	// Points should not have a different entry key, if it was not invalid in the first place
 	if (!TestEqual("Point 0 should have entry 2", MetadataEntryRange[0], PCGMetadataEntryKey(2))
 		|| !TestEqual("Point 1 should have entry 1", MetadataEntryRange[1], PCGMetadataEntryKey(1))
-		|| !TestEqual("Point 2 should have entry 0", MetadataEntryRange[2], PCGMetadataEntryKey(0)))
+		|| !TestEqual("Point 2 should have entry 0", MetadataEntryRange[2], PCGFirstEntryKey))
 	{
 		return false;
 	}
diff --git a/Engine/Plugins/PCG/Source/PCG/Private/Tests/Elements/PCGCopyAttributeTest.cpp b/Engine/Plugins/PCG/Source/PCG/Private/Tests/Elements/PCGCopyAttributeTest.cpp
--- a/Engine/Plugins/PCG/Source/PCG/Private/Tests/Elements/PCGCopyAttributeTest.cpp
+++ b/Engine/Plugins/PCG/Source/PCG/Private/Tests/Elements/PCGCopyAttributeTest.cpp
@@ -362,7 +362,7 @@ bool FPCGCopyAttributeTests_Params_SingleValue::RunTest(const FString& Parameter
 	check(InputAttribute && OutputAttribute);
 
 	UTEST_EQUAL("Output Attribute default value is the same as default of Input Attribute", OutputAttribute->GetValueFromItemKey(PCGInvalidEntryKey), InputAttribute->GetValueFromItemKey(PCGInvalidEntryKey));
-	UTEST_EQUAL("Output Attribute has the same value as Input Attribute for entry 0", OutputAttribute->GetValueFromItemKey(PCGMetadataEntryKey(0)), InputAttribute->GetValueFromItemKey(PCGMetadataEntryKey(0)));
+	UTEST_EQUAL("Output Attribute has the same value as Input Attribute for entry 0", OutputAttribute->GetValueFromItemKey(PCGFirstEntryKey), InputAttribute->GetValueFromItemKey(PCGFirstEntryKey));
 
 	return true;
 }
@@ -738,7 +738,7 @@ bool FPCGCopyAttributeTests_MultiDomain_CopyAll_DataToData::RunTest(const FStrin
 	check(Attribute);
 
 	UTEST_EQUAL("Target Attribute default value is the same as default value of Source Attribute", Attribute->GetValueFromItemKey(PCGInvalidEntryKey), SourceAttribute->GetValueFromItemKey(PCGInvalidEntryKey));
-	UTEST_EQUAL("Target Attribute value is the same as value of Source Attribute", Attribute->GetValueFromItemKey(PCGMetadataEntryKey(0)), SourceAttribute->GetValueFromItemKey(PCGMetadataEntryKey(0)));
+	UTEST_EQUAL("Target Attribute value is the same as value of Source Attribute", Attribute->GetValueFromItemKey(PCGFirstEntryKey), SourceAttribute->GetValueFromItemKey(PCGFirstEntryKey));
 	
 	return true;
 }
@@ -796,7 +796,7 @@ bool FPCGCopyAttributeTests_MultiDomain_CopyAll_AllToAll::RunTest(const FString&
 	check(Attribute);
 
 	UTEST_EQUAL("Target Attribute default value is the same as default value of Source Attribute", Attribute->GetValueFromItemKey(PCGInvalidEntryKey), SourceAttribute->GetValueFromItemKey(PCGInvalidEntryKey));
-	UTEST_EQUAL("Target Attribute value is the same as value of Source Attribute", Attribute->GetValueFromItemKey(PCGMetadataEntryKey(0)), SourceAttribute->GetValueFromItemKey(PCGMetadataEntryKey(0)));
+	UTEST_EQUAL("Target Attribute value is the same as value of Source Attribute", Attribute->GetValueFromItemKey(PCGFirstEntryKey), SourceAttribute->GetValueFromItemKey(PCGFirstEntryKey));
 
 	const FPCGMetadataAttribute<double>* ExistingElementsAttribute = ElementsDomain->GetConstTypedAttribute<double>(TEXT("MyAttr"));
 	UTEST_NOT_NULL("Existing attribute is still there in the output elements metadata domain", ExistingElementsAttribute);
@@ -902,7 +902,7 @@ bool FPCGCopyAttributeTests_MultiDomain_CopyAll_ElementsToData::RunTest(const FS
 
 	check(InputAttribute && OutputAttribute);
 
-	UTEST_EQUAL("Output Attribute default value is the same as the first value of Input Attribute", OutputAttribute->GetValueFromItemKey(PCGInvalidEntryKey), InputAttribute->GetValueFromItemKey(0));
+	UTEST_EQUAL("Output Attribute default value is the same as default of Input Attribute", OutputAttribute->GetValueFromItemKey(PCGInvalidEntryKey), InputAttribute->GetValueFromItemKey(PCGInvalidEntryKey));
 	UTEST_EQUAL("Output Attribute value is the same as value of Input Attribute", OutputAttribute->GetValueFromItemKey(0), InputAttribute->GetValueFromItemKey(0));
 	
 	return true;
diff --git a/Engine/Plugins/PCG/Source/PCG/Private/Tests/Elements/PCGGetAttributeSetByIndexTest.cpp b/Engine/Plugins/PCG/Source/PCG/Private/Tests/Elements/PCGGetAttributeSetByIndexTest.cpp
--- a/Engine/Plugins/PCG/Source/PCG/Private/Tests/Elements/PCGGetAttributeSetByIndexTest.cpp
+++ b/Engine/Plugins/PCG/Source/PCG/Private/Tests/Elements/PCGGetAttributeSetByIndexTest.cpp
@@ -68,11 +68,11 @@ bool FPCGGetAttributeSetByIndexTest_Basic::RunTest(const FString& Parameters)
 
 	check(OutDoubleAttr && OutIntAttr);
 
-	UTEST_EQUAL("Double attribute first entry", OutDoubleAttr->GetValueFromItemKey(PCGMetadataEntryKey(0)), 2.0);
-	UTEST_EQUAL("Double attribute default value was set to the first entry", OutDoubleAttr->GetValueFromItemKey(PCGInvalidEntryKey), 2.0);
+	UTEST_EQUAL("Double attribute first entry", OutDoubleAttr->GetValueFromItemKey(PCGFirstEntryKey), 2.0);
+	UTEST_EQUAL("Double attribute default value was not changed", OutDoubleAttr->GetValueFromItemKey(PCGInvalidEntryKey), 0.0);
 
-	UTEST_EQUAL("Int attribute first entry", OutIntAttr->GetValueFromItemKey(PCGMetadataEntryKey(0)), 3);
-	UTEST_EQUAL("Int attribute default value was set to the first entry", OutIntAttr->GetValueFromItemKey(PCGInvalidEntryKey), 3);
+	UTEST_EQUAL("Int attribute first entry", OutIntAttr->GetValueFromItemKey(PCGFirstEntryKey), 3);
+	UTEST_EQUAL("Int attribute default value was not changed", OutIntAttr->GetValueFromItemKey(PCGInvalidEntryKey), 0);
 
 	return true;
 }
\ No newline at end of file
diff --git a/Engine/Plugins/PCG/Source/PCG/Private/Tests/Metadata/PCGMetadataTest.cpp b/Engine/Plugins/PCG/Source/PCG/Private/Tests/Metadata/PCGMetadataTest.cpp
--- a/Engine/Plugins/PCG/Source/PCG/Private/Tests/Metadata/PCGMetadataTest.cpp
+++ b/Engine/Plugins/PCG/Source/PCG/Private/Tests/Metadata/PCGMetadataTest.cpp
@@ -242,7 +242,7 @@ bool FPCGMetadataTest_TwoAddEntryToData::RunTest(const FString& Parameters)
 	}
 
 	UTEST_EQUAL("Just have a single entry", ParamMetadataDomain->GetItemCountForChild(), 1);
-	UTEST_EQUAL("Entry has the latest value", Attribute->GetValueFromItemKey(PCGMetadataEntryKey(0)), 1);
+	UTEST_EQUAL("Entry has the latest value", Attribute->GetValueFromItemKey(PCGFirstEntryKey), 1);
 	
 	return true;
 }
@@ -264,7 +264,7 @@ bool FPCGMetadataTest_AddEntriesToData::RunTest(const FString& Parameters)
 	Attribute->SetValues(NewKeys, Values);
 
 	UTEST_EQUAL("Just have a single entry", ParamMetadataDomain->GetItemCountForChild(), 1);
-	UTEST_EQUAL("Entry has the latest value", Attribute->GetValueFromItemKey(PCGMetadataEntryKey(0)), Values[1]);
+	UTEST_EQUAL("Entry has the latest value", Attribute->GetValueFromItemKey(PCGFirstEntryKey), Values[1]);
 	
 	return true;
 }
@@ -288,7 +288,7 @@ bool FPCGMetadataTest_AddEntriesInPlaceToData::RunTest(const FString& Parameters
 	Attribute->SetValues(Keys, Values);
 
 	UTEST_EQUAL("Just have a single entry", ParamMetadataDomain->GetItemCountForChild(), 1);
-	UTEST_EQUAL("Entry has the latest value", Attribute->GetValueFromItemKey(PCGMetadataEntryKey(0)), Values[1]);
+	UTEST_EQUAL("Entry has the latest value", Attribute->GetValueFromItemKey(PCGFirstEntryKey), Values[1]);
 	
 	return true;
 }
@@ -319,7 +319,7 @@ bool FPCGMetadataTest_DelayedAddEntriesToData::RunTest(const FString& Parameters
 	ParamMetadataDomain->AddDelayedEntries(Mapping);
 
 	UTEST_EQUAL("Just have a single entry", ParamMetadataDomain->GetItemCountForChild(), 1);
-	UTEST_EQUAL("Entry has the latest value", Attribute->GetValueFromItemKey(PCGMetadataEntryKey(0)), Values[1]);
+	UTEST_EQUAL("Entry has the latest value", Attribute->GetValueFromItemKey(PCGFirstEntryKey), Values[1]);
 	
 	return true;
 }
diff --git a/Engine/Plugins/PCG/Source/PCG/Public/Helpers/PCGSettingsHelpers.h b/Engine/Plugins/PCG/Source/PCG/Public/Helpers/PCGSettingsHelpers.h
--- a/Engine/Plugins/PCG/Source/PCG/Public/Helpers/PCGSettingsHelpers.h
+++ b/Engine/Plugins/PCG/Source/PCG/Public/Helpers/PCGSettingsHelpers.h
@@ -69,7 +69,7 @@ namespace PCGSettingsHelpers
 			using PropertyType = decltype(Dummy);
 
 			// Override were using the first entry (0) by default.
-			FPCGAttributeAccessorKeysEntries FirstEntry(PCGMetadataEntryKey(0));
+			FPCGAttributeAccessorKeysEntries FirstEntry(PCGFirstEntryKey);
 
 			if (!AttributeAccessor->Get<T>(OutValue, FirstEntry, EPCGAttributeAccessorFlags::AllowBroadcast | EPCGAttributeAccessorFlags::AllowConstructible))
 			{
diff --git a/Engine/Plugins/PCG/Source/PCG/Public/Metadata/Accessors/PCGAttributeAccessor.h b/Engine/Plugins/PCG/Source/PCG/Public/Metadata/Accessors/PCGAttributeAccessor.h
--- a/Engine/Plugins/PCG/Source/PCG/Public/Metadata/Accessors/PCGAttributeAccessor.h
+++ b/Engine/Plugins/PCG/Source/PCG/Public/Metadata/Accessors/PCGAttributeAccessor.h
@@ -162,7 +162,7 @@ public:
 			}
 			else if (!bSupportsMultiEntries && ensure(EntriesToSet.Num() == 1))
 			{
-				*(EntriesToSet[0]) = PCGMetadataEntryKey(0);
+				*(EntriesToSet[0]) = PCGFirstEntryKey;
 			}
 		}
 
diff --git a/Engine/Plugins/PCG/Source/PCG/Public/Metadata/PCGMetadataCommon.h b/Engine/Plugins/PCG/Source/PCG/Public/Metadata/PCGMetadataCommon.h
--- a/Engine/Plugins/PCG/Source/PCG/Public/Metadata/PCGMetadataCommon.h
+++ b/Engine/Plugins/PCG/Source/PCG/Public/Metadata/PCGMetadataCommon.h
@@ -10,6 +10,7 @@ typedef int32 PCGMetadataAttributeKey;
 typedef int32 PCGMetadataValueKey;
 
 const PCGMetadataEntryKey PCGInvalidEntryKey = -1;
+const PCGMetadataEntryKey PCGFirstEntryKey = 0;
 const PCGMetadataValueKey PCGDefaultValueKey = -1;
 const PCGMetadataValueKey PCGNotFoundValueKey = -2;
 
