[PCG] Projection: always allocate MetadataEntry property
- could probably be avoided sometimes but would need a more complex condition to figure it out

#rnx

diff --git a/Engine/Plugins/PCG/Source/PCG/Private/Data/PCGProjectionData.cpp b/Engine/Plugins/PCG/Source/PCG/Private/Data/PCGProjectionData.cpp
--- a/Engine/Plugins/PCG/Source/PCG/Private/Data/PCGProjectionData.cpp
+++ b/Engine/Plugins/PCG/Source/PCG/Private/Data/PCGProjectionData.cpp
@@ -310,7 +310,7 @@ const UPCGBasePointData* UPCGProjectionData::CreateBasePointData(FPCGContext* Co
 		PointData->SetNumPoints(SourcePointData->GetNumPoints(), /*bInitializeValues=*/false);
 
 		// @todo_pcg: try to optimize allocations here
-		PointData->AllocateProperties(SourcePointData->GetAllocatedProperties() | EPCGPointNativeProperties::Transform | EPCGPointNativeProperties::Color | EPCGPointNativeProperties::Density);
+		PointData->AllocateProperties(SourcePointData->GetAllocatedProperties() | EPCGPointNativeProperties::Transform | EPCGPointNativeProperties::Color | EPCGPointNativeProperties::Density | EPCGPointNativeProperties::MetadataEntry);
 		PointData->CopyUnallocatedPropertiesFrom(SourcePointData);
 	};
 
