[PCG] Fix crash with GetSingleValue() when range is empty but view has a single element (the default value)

#rnx

diff --git a/Engine/Plugins/PCG/Source/PCG/Public/Utils/PCGValueRange.h b/Engine/Plugins/PCG/Source/PCG/Public/Utils/PCGValueRange.h
--- a/Engine/Plugins/PCG/Source/PCG/Public/Utils/PCGValueRange.h
+++ b/Engine/Plugins/PCG/Source/PCG/Public/Utils/PCGValueRange.h
@@ -65,7 +65,7 @@ public:
 		
 	[[nodiscard]] FORCEINLINE TOptional<ElementType> GetSingleValue() const
 	{
-		return ViewNum() == 1 ? TOptional<ElementType>(GetElement(0)) : TOptional<ElementType>();
+		return ViewNum() == 1 && Num() > 0 ? TOptional<ElementType>(GetElement(0)) : TOptional<ElementType>();
 	}
 
 	[[nodiscard]] FORCEINLINE bool IsValidIndex(SizeType Index) const
