[PCG] GetLandscapeData: Add a Editor only 'bUnbounded' (default to true)
- When true the PrepareForSpatialQuery will use the intersected landscape bounds instead of the PCG Component Grid Bounds

#rnx

diff --git a/Engine/Plugins/PCG/Source/PCG/Private/Elements/PCGTypedGetter.cpp b/Engine/Plugins/PCG/Source/PCG/Private/Elements/PCGTypedGetter.cpp
--- a/Engine/Plugins/PCG/Source/PCG/Private/Elements/PCGTypedGetter.cpp
+++ b/Engine/Plugins/PCG/Source/PCG/Private/Elements/PCGTypedGetter.cpp
@@ -127,7 +127,7 @@ void FPCGGetLandscapeDataElement::ProcessActors(FPCGContext* Context, const UPCG
 		
 		// Prepare data for the PCG Bounds in Editor only as it can create cache entries which isn't thread safe
 #if WITH_EDITOR
-		const FBox PrepareBounds = Context->ExecutionSource.IsValid() ? Context->ExecutionSource->GetExecutionState().GetBounds() : LandscapeBounds;
+		const FBox PrepareBounds = Context->ExecutionSource.IsValid() && !Settings->bUnbounded ? Context->ExecutionSource->GetExecutionState().GetBounds() : LandscapeBounds;
 		OutDynamicDependencies.Append(LandscapeData->PrepareForSpatialQuery(Context, PrepareBounds));
 #endif
 
diff --git a/Engine/Plugins/PCG/Source/PCG/Public/Elements/PCGTypedGetter.h b/Engine/Plugins/PCG/Source/PCG/Public/Elements/PCGTypedGetter.h
--- a/Engine/Plugins/PCG/Source/PCG/Public/Elements/PCGTypedGetter.h
+++ b/Engine/Plugins/PCG/Source/PCG/Public/Elements/PCGTypedGetter.h
@@ -52,6 +52,10 @@ public:
 	FPCGLandscapeDataProps SamplingProperties;
 
 #if WITH_EDITORONLY_DATA
+	// Editor only: If true, the intersected landscape bounds are going to be used to prepare the landscape cache, otherwise the PCG Component's grid bounds will be used.
+	UPROPERTY(BlueprintReadWrite, EditAnywhere, Category = Settings, meta = (PCG_Overridable))
+	bool bUnbounded = true;
+
 	UPROPERTY()
 	bool bGetHeightOnly_DEPRECATED = false;
 
