[PCG] Add spatial data support to viewport

#rnx

diff --git a/Engine/Plugins/PCG/Source/PCGEditor/Private/DataVisualizations/PCGSpatialDataVisualization.cpp b/Engine/Plugins/PCG/Source/PCGEditor/Private/DataVisualizations/PCGSpatialDataVisualization.cpp
--- a/Engine/Plugins/PCG/Source/PCGEditor/Private/DataVisualizations/PCGSpatialDataVisualization.cpp
+++ b/Engine/Plugins/PCG/Source/PCGEditor/Private/DataVisualizations/PCGSpatialDataVisualization.cpp
@@ -18,7 +18,9 @@
 #include "Metadata/Accessors/PCGAttributeAccessorHelpers.h"
 #include "Metadata/Accessors/PCGCustomAccessor.h"
 
+#include "AdvancedPreviewScene.h"
 #include "Components/InstancedStaticMeshComponent.h"
+#include "Engine/AssetManager.h"
 
 #define LOCTEXT_NAMESPACE "PCGSpatialDataVisualization"
 
@@ -41,18 +43,22 @@ void IPCGSpatialDataVisualization::ExecuteDebugDisplay(FPCGContext* Context, con
 		return;
 	}
 
-	if (!SettingsInterface)
+	if (!SettingsInterface || !SettingsInterface->GetSettings())
 	{
 		return;
 	}
 
-	UPCGComponent* SourceComponent = Cast<UPCGComponent>(Context->ExecutionSource.Get());
-	if (!SourceComponent)
-	{
-		return;
-	}
+	ExecuteDebugDisplayHelper(Data, SettingsInterface->DebugSettings, Context, TargetActor, SettingsInterface->GetSettings()->GetSettingsCrc(), [](UInstancedStaticMeshComponent*) {});
+}
 
-	const FPCGDebugVisualizationSettings& DebugSettings = SettingsInterface->DebugSettings;
+void IPCGSpatialDataVisualization::ExecuteDebugDisplayHelper(
+	const UPCGData* Data,
+	const FPCGDebugVisualizationSettings& DebugSettings,
+	FPCGContext* Context,
+	AActor* TargetActor,
+	const FPCGCrc& Crc,
+	const TFunction<void(UInstancedStaticMeshComponent*)>& OnISMCCreatedCallback) const
+{
 	UStaticMesh* Mesh = DebugSettings.PointMesh.LoadSynchronous();
 
 	if (!Mesh)
@@ -131,7 +137,7 @@ void IPCGSpatialDataVisualization::ExecuteDebugDisplay(FPCGContext* Context, con
 	}
 
 	FPCGISMComponentBuilderParams Params[2];
-	Params[0].SettingsCrc = SettingsInterface->GetSettings()->GetSettingsCrc();
+	Params[0].SettingsCrc = Crc;
 	Params[0].bTransient = true;
 	Params[0].NumCustomDataFloats = NumCustomData;
 	Params[0].Descriptor.StaticMesh = Mesh;
@@ -155,7 +161,8 @@ void IPCGSpatialDataVisualization::ExecuteDebugDisplay(FPCGContext* Context, con
 	Params[0].Descriptor.bVisibleInRayTracing = false;
 	
 	// If the root actor we're binding to is movable, then the ISMC should be movable by default
-	if (USceneComponent* SceneComponent = TargetActor->GetRootComponent())
+	USceneComponent* SceneComponent = TargetActor ? TargetActor->GetRootComponent() : nullptr;
+	if (SceneComponent)
 	{
 		Params[0].Descriptor.Mobility = SceneComponent->Mobility;
 	}
@@ -163,6 +170,8 @@ void IPCGSpatialDataVisualization::ExecuteDebugDisplay(FPCGContext* Context, con
 	Params[1] = Params[0];
 	Params[1].Descriptor.bReverseCulling = true;
 
+	UPCGComponent* SourceComponent = Cast<UPCGComponent>(Context ? Context->ExecutionSource.Get() : nullptr);
+
 	// Since the instance count is global, track the current instances applied and previously belong to the ISMCs.
 	int32 NumCurrentInstances = 0;
 	for (int32 Direction = 0; Direction < 2; ++Direction)
@@ -174,7 +183,22 @@ void IPCGSpatialDataVisualization::ExecuteDebugDisplay(FPCGContext* Context, con
 			continue;
 		}
 
-		UInstancedStaticMeshComponent* ISMC = UPCGActorHelpers::GetOrCreateISMC(TargetActor, SourceComponent, Params[Direction], Context);
+		UInstancedStaticMeshComponent* ISMC = nullptr;
+		
+		if (TargetActor && SourceComponent)
+		{
+			ISMC = UPCGActorHelpers::GetOrCreateISMC(TargetActor, SourceComponent, Params[Direction], Context);
+		}
+		else
+		{
+			// If no target actor/source component were provided, create an ISMC directly instead.
+			ISMC = NewObject<UInstancedStaticMeshComponent>(GetTransientPackage(), NAME_None, RF_Transient);
+
+			FISMComponentDescriptor Descriptor(Params[Direction].Descriptor);
+			Descriptor.InitComponent(ISMC);
+			ISMC->SetNumCustomDataFloats(Params[Direction].NumCustomDataFloats);
+		}
+
 		check(ISMC && ISMC->NumCustomDataFloats == NumCustomData);
 
 		ISMC->ComponentTags.AddUnique(PCGHelpers::DefaultPCGDebugTag);
@@ -226,6 +250,7 @@ void IPCGSpatialDataVisualization::ExecuteDebugDisplay(FPCGContext* Context, con
 		}
 
 		ISMC->UpdateBounds();
+		OnISMCCreatedCallback(ISMC);
 	}
 }
 
@@ -379,4 +404,61 @@ TArray<FPCGMetadataDomainID> IPCGSpatialDataVisualization::GetAllSupportedDomain
 	return GetDefault<UPCGBasePointData>()->GetAllSupportedMetadataDomainIDs();
 }
 
+FPCGSetupSceneFunc IPCGSpatialDataVisualization::GetViewportSetupFunc(const UPCGData* Data) const
+{
+	return [this, WeakData=TWeakObjectPtr<const UPCGData>(Data)](FAdvancedPreviewScene* InScene, const TArray<UObject*>& InResources, FVector& OutViewLocation, float& OutOrbitDistance, TArray<TObjectPtr<UObject>>& OutManagedResources)
+	{
+		check(InScene);
+
+		if (!WeakData.IsValid())
+		{
+			UE_LOG(LogPCG, Error, TEXT("Failed to setup data viewport, the data was lost or invalid."));
+			return;
+		}
+
+		FVector BoundsMin, BoundsMax;
+		bool bInitializedBounds = false;
+
+		ExecuteDebugDisplayHelper(
+			WeakData.Get(),
+			FPCGDebugVisualizationSettings(),
+			nullptr,
+			nullptr,
+			FPCGCrc(),
+			[&OutManagedResources, InScene, &bInitializedBounds, &BoundsMin, &BoundsMax](UInstancedStaticMeshComponent* ISMC)
+			{
+				check(ISMC);
+
+				OutManagedResources.Add(ISMC);
+				InScene->AddComponent(ISMC, FTransform::Identity);
+
+				const FVector CurrentBoundsMin = ISMC->Bounds.Origin - ISMC->Bounds.BoxExtent;
+				const FVector CurrentBoundsMax = ISMC->Bounds.Origin + ISMC->Bounds.BoxExtent;
+
+				if (!bInitializedBounds)
+				{
+					BoundsMin = CurrentBoundsMin;
+					BoundsMax = CurrentBoundsMax;
+					bInitializedBounds = true;
+				}
+				else
+				{
+					BoundsMin = BoundsMin.ComponentMin(CurrentBoundsMin);
+					BoundsMax = BoundsMax.ComponentMax(CurrentBoundsMax);
+				}
+			}
+		);
+
+		const FBoxSphereBounds Bounds = FBoxSphereBounds(FBox(BoundsMin, BoundsMax));
+		InScene->SetFloorOffset(-Bounds.Origin.Z + Bounds.BoxExtent.Z);
+
+		OutViewLocation = Bounds.Origin;
+
+		if (Bounds.SphereRadius > 0.0f)
+		{
+			OutOrbitDistance = Bounds.SphereRadius * 2.0f;
+		}
+	};
+}
+
 #undef LOCTEXT_NAMESPACE
diff --git a/Engine/Plugins/PCG/Source/PCGEditor/Public/DataVisualizations/PCGSpatialDataVisualization.h b/Engine/Plugins/PCG/Source/PCGEditor/Public/DataVisualizations/PCGSpatialDataVisualization.h
--- a/Engine/Plugins/PCG/Source/PCGEditor/Public/DataVisualizations/PCGSpatialDataVisualization.h
+++ b/Engine/Plugins/PCG/Source/PCGEditor/Public/DataVisualizations/PCGSpatialDataVisualization.h
@@ -10,6 +10,7 @@ class UPCGData;
 class UPCGBasePointData;
 class UPCGPointData;
 struct FPCGContext;
+struct FPCGDebugVisualizationSettings;
 
 /** Default implementation for spatial data. Collapses to a PointData representation. */
 class PCGEDITOR_API IPCGSpatialDataVisualization : public IPCGDataVisualization
@@ -23,10 +24,19 @@ public:
 	virtual FPCGMetadataDomainID GetDefaultDomainForInspection(const UPCGData* Data) const override { return PCGMetadataDomainID::Elements; }
 	virtual FString GetDomainDisplayNameForInspection(const UPCGData* Data, const FPCGMetadataDomainID& DomainID) const override;
 	virtual TArray<FPCGMetadataDomainID> GetAllSupportedDomainsForInspection(const UPCGData* Data) const override;
+	virtual FPCGSetupSceneFunc GetViewportSetupFunc(const UPCGData* Data) const override;
 	// ~End IPCGDataVisualization interface
 
 	UE_DEPRECATED(5.6, "Implement CollapseToDebugBasePointData instead")
 	virtual const UPCGPointData* CollapseToDebugPointData(FPCGContext* Context, const UPCGData* Data) const;
 
 	virtual const UPCGBasePointData* CollapseToDebugBasePointData(FPCGContext* Context, const UPCGData* Data) const;
+
+	virtual void ExecuteDebugDisplayHelper(
+		const UPCGData* Data,
+		const FPCGDebugVisualizationSettings& DebugSettings,
+		FPCGContext* Context,
+		AActor* TargetActor,
+		const FPCGCrc& Crc,
+		const TFunction<void(class UInstancedStaticMeshComponent*)>& OnISMCCreatedCallback) const;
 };
