[PCG] [ISKM] Add support to PCG for GPU-instancing of ISKMs.

#rnx

diff --git a/Engine/Plugins/PCG/Source/PCG/Private/Compute/BuiltInKernels/PCGCountUniqueAttributeValuesKernel.cpp b/Engine/Plugins/PCG/Source/PCG/Private/Compute/BuiltInKernels/PCGCountUniqueAttributeValuesKernel.cpp
--- a/Engine/Plugins/PCG/Source/PCG/Private/Compute/BuiltInKernels/PCGCountUniqueAttributeValuesKernel.cpp
+++ b/Engine/Plugins/PCG/Source/PCG/Private/Compute/BuiltInKernels/PCGCountUniqueAttributeValuesKernel.cpp
@@ -20,17 +20,6 @@
 
 #define LOCTEXT_NAMESPACE "PCGCountUniqueAttributeValuesKernel"
 
-#if WITH_EDITOR
-void UPCGCountUniqueAttributeValuesKernel::InitializeInternal()
-{
-	Super::InitializeInternal();
-
-	const UPCGStaticMeshSpawnerSettings* SMSettings = CastChecked<UPCGStaticMeshSpawnerSettings>(Settings);
-	const UPCGMeshSelectorByAttribute* Selector = CastChecked<UPCGMeshSelectorByAttribute>(SMSettings->MeshSelectorParameters);
-	AttributeName = Selector->AttributeName;
-}
-#endif
-
 bool UPCGCountUniqueAttributeValuesKernel::IsKernelDataValid(FPCGContext* InContext) const
 {
 	TRACE_CPUPROFILER_EVENT_SCOPE(UPCGCountUniqueAttributeValuesKernel::IsKernelDataValid);
diff --git a/Engine/Plugins/PCG/Source/PCG/Private/Compute/BuiltInKernels/PCGCountUniqueAttributeValuesKernel.h b/Engine/Plugins/PCG/Source/PCG/Private/Compute/BuiltInKernels/PCGCountUniqueAttributeValuesKernel.h
--- a/Engine/Plugins/PCG/Source/PCG/Private/Compute/BuiltInKernels/PCGCountUniqueAttributeValuesKernel.h
+++ b/Engine/Plugins/PCG/Source/PCG/Private/Compute/BuiltInKernels/PCGCountUniqueAttributeValuesKernel.h
@@ -18,6 +18,8 @@ class UPCGCountUniqueAttributeValuesKernel : public UPCGComputeKernel
 	GENERATED_BODY()
 
 public:
+	void SetAttributeName(FName InAttributeName) { AttributeName = InAttributeName; }
+
 	virtual bool IsKernelDataValid(FPCGContext* InContext) const override;
 	virtual bool ComputeOutputBindingDataDesc(const UPCGComputeGraph* InGraph, FName InOutputPinLabel, UPCGDataBinding* InBinding, FPCGDataCollectionDesc& OutDataDesc) const override;
 	virtual int ComputeThreadCount(const UPCGDataBinding* Binding) const override;
@@ -34,11 +36,6 @@ public:
 	virtual void GetInputPins(TArray<FPCGPinProperties>& OutPins) const override;
 	virtual void GetOutputPins(TArray<FPCGPinPropertiesGPU>& OutPins) const override;
 
-#if WITH_EDITOR
-protected:
-	virtual void InitializeInternal() override;
-#endif
-
 public:
 	UPROPERTY()
 	FName AttributeName;
diff --git a/Engine/Plugins/PCG/Source/PCG/Private/Compute/DataInterfaces/Elements/PCGSkinnedMeshSpawnerDataInterface.cpp b/Engine/Plugins/PCG/Source/PCG/Private/Compute/DataInterfaces/Elements/PCGSkinnedMeshSpawnerDataInterface.cpp
new file mode 100644
--- /dev/null
+++ b/Engine/Plugins/PCG/Source/PCG/Private/Compute/DataInterfaces/Elements/PCGSkinnedMeshSpawnerDataInterface.cpp
@@ -0,0 +1,821 @@
+// Copyright Epic Games, Inc. All Rights Reserved.
+
+#include "PCGSkinnedMeshSpawnerDataInterface.h"
+
+#include "PCGComponent.h"
+#include "PCGContext.h"
+#include "PCGManagedResource.h"
+#include "PCGParamData.h"
+#include "Compute/PCGComputeCommon.h"
+#include "Compute/PCGComputeGraph.h"
+#include "Compute/PCGDataBinding.h"
+#include "Compute/BuiltInKernels/PCGCountUniqueAttributeValuesKernel.h"
+#include "Compute/Data/PCGProxyForGPUData.h"
+#include "Elements/PCGSkinnedMeshSpawner.h"
+#include "Elements/PCGSkinnedMeshSpawnerKernel.h"
+#include "InstanceDataPackers/PCGSkinnedMeshInstanceDataPackerBase.h"
+#include "MeshSelectors/PCGSkinnedMeshSelector.h"
+#include "Helpers/PCGActorHelpers.h"
+
+#include "GlobalRenderResources.h"
+#include "RenderGraphBuilder.h"
+#include "RenderGraphResources.h"
+#include "RHIResources.h"
+#include "ShaderCompilerCore.h"
+#include "ShaderCore.h"
+#include "ShaderParameterMetadataBuilder.h"
+#include "SystemTextures.h"
+#include "Engine/SkinnedAsset.h"
+#include "ComputeFramework/ShaderParamTypeDefinition.h"
+#include "Components/InstancedSkinnedMeshComponent.h"
+
+#include UE_INLINE_GENERATED_CPP_BY_NAME(PCGSkinnedMeshSpawnerDataInterface)
+
+#define LOCTEXT_NAMESPACE "PCGSkinnedMeshSpawnerDataInterface"
+
+namespace PCGSkinnedMeshSpawnerDataInterface
+{
+	static const FText CouldNotLoadStaticMeshFormat = LOCTEXT("CouldNotLoadSkinnedMesh", "Could not load Skinned mesh from path '{0}'.");
+	static const FText TooManyPrimitivesFormat = LOCTEXT("TooManyPrimitives", "Attempted to emit too many primitive components, terminated after creating '{0}'.");
+	static const FText NoMeshEntriesFormat = LOCTEXT("NoMeshEntries", "No mesh entries provided in weighted mesh selector.");
+}
+
+void UPCGSkinnedMeshSpawnerDataInterface::GetSupportedInputs(TArray<FShaderFunctionDefinition>& OutFunctions) const
+{
+	OutFunctions.AddDefaulted_GetRef()
+		.SetName(TEXT("SMSpawner_GetSelectorAttributeId"))
+		.AddReturnType(EShaderFundamentalType::Uint); // Attribute id to get mesh path string key from, or invalid if we should use CDF instead.
+
+	OutFunctions.AddDefaulted_GetRef()
+		.SetName(TEXT("SMSpawner_GetNumAttributes"))
+		.AddReturnType(EShaderFundamentalType::Uint); // Num attributes
+
+	OutFunctions.AddDefaulted_GetRef()
+		.SetName(TEXT("SMSpawner_GetNumPrimitives"))
+		.AddReturnType(EShaderFundamentalType::Uint); // Num primitives
+
+	OutFunctions.AddDefaulted_GetRef()
+		.SetName(TEXT("SMSpawner_ShouldApplyBounds"))
+		.AddReturnType(EShaderFundamentalType::Bool);
+
+	OutFunctions.AddDefaulted_GetRef()
+		.SetName(TEXT("SMSpawner_GetPrimitiveMeshBoundsMin"))
+		.AddReturnType(EShaderFundamentalType::Float, 3) // Local bounds min
+		.AddParam(EShaderFundamentalType::Uint); // InPrimitiveIndex
+
+	OutFunctions.AddDefaulted_GetRef()
+		.SetName(TEXT("SMSpawner_GetPrimitiveMeshBoundsMax"))
+		.AddReturnType(EShaderFundamentalType::Float, 3) // Local bounds max
+		.AddParam(EShaderFundamentalType::Uint); // InPrimitiveIndex
+
+	OutFunctions.AddDefaulted_GetRef()
+		.SetName(TEXT("SMSpawner_GetAttributeIdOffsetStride"))
+		.AddReturnType(EShaderFundamentalType::Uint, 4)
+		.AddParam(EShaderFundamentalType::Uint); // InAttributeIndex
+
+	OutFunctions.AddDefaulted_GetRef()
+		.SetName(TEXT("SMSpawner_GetPrimitiveStringKey"))
+		.AddReturnType(EShaderFundamentalType::Int) // String key
+		.AddParam(EShaderFundamentalType::Uint); // InPrimitiveIndex
+
+	OutFunctions.AddDefaulted_GetRef()
+		.SetName(TEXT("SMSpawner_GetPrimitiveIndexFromStringKey"))
+		.AddReturnType(EShaderFundamentalType::Uint) // Primitive index
+		.AddParam(EShaderFundamentalType::Int); // InMeshPathStringKey
+
+	OutFunctions.AddDefaulted_GetRef()
+		.SetName(TEXT("SMSpawner_GetPrimitiveSelectionCDF"))
+		.AddReturnType(EShaderFundamentalType::Float) // CDF value
+		.AddParam(EShaderFundamentalType::Uint); // InPrimitiveIndex
+
+	OutFunctions.AddDefaulted_GetRef()
+		.SetName(TEXT("SMSpawner_GetSelectedMeshAttributeId"))
+		.AddReturnType(EShaderFundamentalType::Uint); // Attribute id to output mesh path string key to
+}
+
+BEGIN_SHADER_PARAMETER_STRUCT(FPCGSkinnedMeshSpawnerDataInterfaceParameters,)
+	SHADER_PARAMETER_ARRAY(FUintVector4, AttributeIdOffsetStrides, [UPCGSkinnedMeshSpawnerDataInterface::MAX_ATTRIBUTES])
+	SHADER_PARAMETER_RDG_BUFFER_SRV(StructuredBuffer<int32>, PrimitiveStringKeys)
+	SHADER_PARAMETER_RDG_BUFFER_SRV(StructuredBuffer<FVector4f>, PrimitiveMeshBoundsMin)
+	SHADER_PARAMETER_RDG_BUFFER_SRV(StructuredBuffer<FVector4f>, PrimitiveMeshBoundsMax)
+	SHADER_PARAMETER_SCALAR_ARRAY(float, SelectionCDF, [PCGComputeConstants::MAX_PRIMITIVE_COMPONENTS_PER_SPAWNER])
+	SHADER_PARAMETER(uint32, NumAttributes)
+	SHADER_PARAMETER(uint32, NumPrimitives)
+	SHADER_PARAMETER(int32, SelectorAttributeId)
+	SHADER_PARAMETER(uint32, ApplyBounds)
+END_SHADER_PARAMETER_STRUCT()
+
+void UPCGSkinnedMeshSpawnerDataInterface::GetShaderParameters(TCHAR const* UID, FShaderParametersMetadataBuilder& InOutBuilder, FShaderParametersMetadataAllocations& InOutAllocations) const
+{
+	InOutBuilder.AddNestedStruct<FPCGSkinnedMeshSpawnerDataInterfaceParameters>(UID);
+}
+
+void UPCGSkinnedMeshSpawnerDataInterface::GetHLSL(FString& OutHLSL, FString const& InDataInterfaceName) const
+{
+	TMap<FString, FStringFormatArg> TemplateArgs =
+	{
+		{ TEXT("DataInterfaceName"), InDataInterfaceName },
+		{ TEXT("MaxAttributes"), MAX_ATTRIBUTES },
+		{ TEXT("MaxPrimitives"), PCGComputeConstants::MAX_PRIMITIVE_COMPONENTS_PER_SPAWNER },
+	};
+
+	OutHLSL += FString::Format(TEXT(
+		"int {DataInterfaceName}_SelectorAttributeId;\n"
+		"uint {DataInterfaceName}_NumAttributes;\n"
+		"uint {DataInterfaceName}_NumPrimitives;\n"
+		"uint {DataInterfaceName}_ApplyBounds;\n"
+		"uint4 {DataInterfaceName}_AttributeIdOffsetStrides[{MaxAttributes}];\n"
+		"StructuredBuffer<float4> {DataInterfaceName}_PrimitiveMeshBoundsMin;\n"
+		"StructuredBuffer<float4> {DataInterfaceName}_PrimitiveMeshBoundsMax;\n"
+		"StructuredBuffer<int> {DataInterfaceName}_PrimitiveStringKeys;\n"
+		"DECLARE_SCALAR_ARRAY(float, {DataInterfaceName}_SelectionCDF, {MaxPrimitives});\n"
+		"\n"
+		"int SMSpawner_GetSelectorAttributeId_{DataInterfaceName}()\n"
+		"{\n"
+		"	return {DataInterfaceName}_SelectorAttributeId;\n"
+		"}\n"
+		"\n"
+		"uint SMSpawner_GetNumAttributes_{DataInterfaceName}()\n"
+		"{\n"
+		"	return {DataInterfaceName}_NumAttributes;\n"
+		"}\n"
+		"\n"
+		"uint4 SMSpawner_GetAttributeIdOffsetStride_{DataInterfaceName}(uint InAttributeIndex)\n"
+		"{\n"
+		"	return {DataInterfaceName}_AttributeIdOffsetStrides[InAttributeIndex];\n"
+		"}\n"
+		"\n"
+		"uint SMSpawner_GetNumPrimitives_{DataInterfaceName}()\n"
+		"{\n"
+		"	return {DataInterfaceName}_NumPrimitives;\n"
+		"}\n"
+		"\n"
+		"bool SMSpawner_ShouldApplyBounds_{DataInterfaceName}()\n"
+		"{\n"
+		"	return {DataInterfaceName}_ApplyBounds > 0;\n"
+		"}\n"
+		"\n"
+		"float3 SMSpawner_GetPrimitiveMeshBoundsMin_{DataInterfaceName}(uint InPrimitiveIndex)\n"
+		"{\n"
+		"	return {DataInterfaceName}_PrimitiveMeshBoundsMin[InPrimitiveIndex].xyz;\n"
+		"}\n"
+		"\n"
+		"float3 SMSpawner_GetPrimitiveMeshBoundsMax_{DataInterfaceName}(uint InPrimitiveIndex)\n"
+		"{\n"
+		"	return {DataInterfaceName}_PrimitiveMeshBoundsMax[InPrimitiveIndex].xyz;\n"
+		"}\n"
+		"\n"
+		"int SMSpawner_GetPrimitiveStringKey_{DataInterfaceName}(uint InPrimitiveIndex)\n"
+		"{\n"
+		"	return {DataInterfaceName}_PrimitiveStringKeys[InPrimitiveIndex];\n"
+		"}\n"
+		"\n"
+		"uint SMSpawner_GetPrimitiveIndexFromStringKey_{DataInterfaceName}(int InMeshPathStringKey)\n"
+		"{\n"
+		"	for (uint Index = 0; Index < {DataInterfaceName}_NumPrimitives; ++Index)\n"
+		"	{\n"
+		"		if ({DataInterfaceName}_PrimitiveStringKeys[Index] == InMeshPathStringKey)\n"
+		"		{\n"
+		"			return Index;\n"
+		"		}\n"
+		"	}\n"
+		"	\n"
+		"	return (uint)-1;\n"
+		"}\n"
+		"\n"
+		"float SMSpawner_GetPrimitiveSelectionCDF_{DataInterfaceName}(uint InPrimitiveIndex)\n"
+		"{\n"
+		"	return GET_SCALAR_ARRAY_ELEMENT({DataInterfaceName}_SelectionCDF, InPrimitiveIndex);\n"
+		"}\n"
+		"\n"
+		"int SMSpawner_GetSelectedMeshAttributeId_{DataInterfaceName}()\n"
+		"{\n"
+		"	return -1; \n" // Not currently supported by skinned mesh spawner node. Stubbed here so that single spawner usf can be used.
+		"}\n\n"
+	), TemplateArgs);
+}
+
+UComputeDataProvider* UPCGSkinnedMeshSpawnerDataInterface::CreateDataProvider() const
+{
+	return NewObject<UPCGSkinnedMeshSpawnerDataProvider>();
+}
+
+void UPCGSkinnedMeshSpawnerDataProvider::Initialize(const UComputeDataInterface* InDataInterface, UObject* InBinding, uint64 InInputMask, uint64 InOutputMask)
+{
+	TRACE_CPUPROFILER_EVENT_SCOPE(UPCGSkinnedMeshSpawnerDataProvider::Initialize);
+
+	Super::Initialize(InDataInterface, InBinding, InInputMask, InOutputMask);
+
+	const UPCGSkinnedMeshSpawnerDataInterface* DataInterface = CastChecked<UPCGSkinnedMeshSpawnerDataInterface>(InDataInterface);
+	ProducerKernel = CastChecked<UPCGSkinnedMeshSpawnerKernel>(DataInterface->ProducerKernel);
+}
+
+bool UPCGSkinnedMeshSpawnerDataProvider::PerformPreExecuteReadbacks_GameThread(UPCGDataBinding* InBinding)
+{
+	TRACE_CPUPROFILER_EVENT_SCOPE(UPCGSkinnedMeshSpawnerDataProvider::PerformPreExecuteReadbacks_GameThread);
+	check(InBinding);
+
+	if (!Super::PerformPreExecuteReadbacks_GameThread(InBinding))
+	{
+		return false;
+	}
+
+	const UPCGSkinnedMeshSpawnerSettings* Settings = CastChecked<UPCGSkinnedMeshSpawnerSettings>(ProducerKernel->GetSettings());
+
+	// Obtain index of analysis data that we need to read back.
+	if (AnalysisDataIndex == INDEX_NONE)
+	{
+		if (!ensure(Settings->MeshSelectorParameters))
+		{
+			// Non by-attribute selection does not need to do any readbacks.
+			return true;
+		}
+
+		AnalysisDataIndex = ProducerKernel ? InBinding->GetFirstInputDataIndex(ProducerKernel, PCGSkinnedMeshSpawnerConstants::InstanceCountsPinLabel) : INDEX_NONE;
+
+		if (AnalysisDataIndex == INDEX_NONE)
+		{
+			// No analysis data to read back.
+			return true;
+		}
+	}
+
+	// Readback analysis data - poll until readback complete (return true).
+	if (!InBinding->ReadbackInputDataToCPU(AnalysisDataIndex))
+	{
+		return false;
+	}
+
+	const UPCGParamData* AnalysisResultsData = Cast<UPCGParamData>(InBinding->InputDataCollection.TaggedData[AnalysisDataIndex].Data);
+	const UPCGMetadata* AnalysisMetadata = AnalysisResultsData ? AnalysisResultsData->ConstMetadata() : nullptr;
+
+	const FPCGMetadataAttributeBase* AnalysisValueAttributeBase = AnalysisMetadata ? AnalysisMetadata->GetConstAttribute(PCGCountUniqueAttributeValuesConstants::ValueAttributeName) : nullptr;
+	const FPCGMetadataAttributeBase* AnalysisCountAttributeBase = AnalysisMetadata ? AnalysisMetadata->GetConstAttribute(PCGCountUniqueAttributeValuesConstants::ValueCountAttributeName) : nullptr;
+
+	if (AnalysisValueAttributeBase && AnalysisValueAttributeBase->GetTypeId() == PCG::Private::MetadataTypes<int32>::Id &&
+		AnalysisCountAttributeBase && AnalysisCountAttributeBase->GetTypeId() == PCG::Private::MetadataTypes<int32>::Id)
+	{
+		const FPCGMetadataAttribute<int32>* ValueAttribute = static_cast<const FPCGMetadataAttribute<int32>*>(AnalysisValueAttributeBase);
+		const FPCGMetadataAttribute<int32>* CountAttribute = static_cast<const FPCGMetadataAttribute<int32>*>(AnalysisCountAttributeBase);
+
+		const int32 NumElements = AnalysisMetadata->GetItemCountForChild();
+
+		StringKeyToInstanceCount.Reserve(NumElements);
+
+		// TODO: Range based get would scale better.
+		for (int64 MetadataKey = 0; MetadataKey < NumElements; ++MetadataKey)
+		{
+			StringKeyToInstanceCount.Add(ValueAttribute->GetValue(MetadataKey), CountAttribute->GetValue(MetadataKey));
+		}
+	}
+	else
+	{
+		UE_LOG(LogPCG, Warning, TEXT("No analysis data received by static mesh spawner kernel, worst case instance allocations will be made."));
+
+		if (InBinding->InputDataCollection.TaggedData[AnalysisDataIndex].Data->IsA<UPCGProxyForGPUData>())
+		{
+			UE_LOG(LogPCG, Error, TEXT("Data was not read back."));
+		}
+
+		return true;
+	}
+
+	return true;
+}
+
+bool UPCGSkinnedMeshSpawnerDataProvider::PrepareForExecute_GameThread(UPCGDataBinding* InBinding)
+{
+	TRACE_CPUPROFILER_EVENT_SCOPE(UPCGSkinnedMeshSpawnerDataProvider::PrepareForExecute_GameThread);
+	check(InBinding);
+	const UPCGSkinnedMeshSpawnerSettings* Settings = CastChecked<UPCGSkinnedMeshSpawnerSettings>(ProducerKernel->GetSettings());
+
+	if (!Super::PrepareForExecute_GameThread(InBinding))
+	{
+		return false;
+	}
+
+	if (!bRegisteredPrimitives)
+	{
+		InBinding->MeshSpawnersToPrimitives.Add(ProducerKernel, {});
+		bRegisteredPrimitives = true;
+	}
+
+	TSharedPtr<FPCGContextHandle> ContextHandle = InBinding->ContextHandle.Pin();
+	FPCGContext* Context = ContextHandle ? ContextHandle->GetContext() : nullptr;
+	if (!ensure(Context))
+	{
+		InBinding->CompletedMeshSpawners.Add(ProducerKernel);
+		return true;
+	}
+
+	UPCGComponent* SourceComponent = Cast<UPCGComponent>(Context->ExecutionSource.Get());
+	if (!ensure(SourceComponent))
+	{
+		InBinding->CompletedMeshSpawners.Add(ProducerKernel);
+		return true;
+	}
+
+	if (SelectorAttributeId == INDEX_NONE && ensure(Settings->MeshSelectorParameters))
+	{
+		const FName SelectorName = Settings->MeshSelectorParameters->MeshAttribute.GetName();
+		const FPCGDataCollectionDesc* InputDataDesc = InBinding->GetCachedKernelPinDataDesc(ProducerKernel, PCGPinConstants::DefaultInputLabel, /*bIsInput=*/true);
+
+		if (!ensure(InputDataDesc))
+		{
+			InBinding->CompletedMeshSpawners.Add(ProducerKernel);
+			return true;
+		}
+
+		bool bAnyPointsPresent = false;
+
+		for (const FPCGDataDesc& Desc : InputDataDesc->DataDescs)
+		{
+			if (Desc.ElementCount <= 0)
+			{
+				continue;
+			}
+
+			bAnyPointsPresent = true;
+
+			for (const FPCGKernelAttributeDesc& AttributeDesc : Desc.AttributeDescs)
+			{
+				if (AttributeDesc.AttributeKey.Identifier.Name == SelectorName && AttributeDesc.AttributeKey.Type == EPCGKernelAttributeType::StringKey)
+				{
+					SelectorAttributeId = AttributeDesc.AttributeId;
+					break;
+				}
+			}
+
+			if (SelectorAttributeId != INDEX_NONE)
+			{
+				break;
+			}
+		}
+
+		if (SelectorAttributeId == INDEX_NONE)
+		{
+			// Mute this error if the point data is empty.
+			if (bAnyPointsPresent)
+			{
+				PCG_KERNEL_VALIDATION_ERR(Context, Settings, FText::Format(
+					LOCTEXT("MeshSelectorAttributeNotFound", "Mesh selector attribute '{0}' not found."),
+					FText::FromName(SelectorName)));
+			}
+
+			InBinding->CompletedMeshSpawners.Add(ProducerKernel);
+			return true;
+		}
+
+		if (Context->AsyncState.ShouldStop())
+		{
+			return false;
+		}
+	}
+
+	if (!bPrimitiveDescriptorsCreated)
+	{
+		CreatePrimitiveDescriptors(Context, InBinding);
+
+		for (const FPCGSoftSkinnedMeshComponentDescriptor& PrimitiveDescriptor : PrimitiveDescriptors)
+		{
+			UE_LOG(LogPCG, Verbose, TEXT("Request '%s' to load."), *PrimitiveDescriptor.SkinnedAsset.ToString());
+
+			PrimitiveDescriptor.SkinnedAsset.LoadAsync({});
+		}
+
+		bPrimitiveDescriptorsCreated = true;
+
+		if (Context->AsyncState.ShouldStop())
+		{
+			return false;
+		}
+	}
+
+	if (!bStaticMeshesLoaded)
+	{
+		for (const FPCGSoftSkinnedMeshComponentDescriptor& PrimitiveDescriptor : PrimitiveDescriptors)
+		{
+			if (PrimitiveDescriptor.SkinnedAsset.IsPending())
+			{
+				UE_LOG(LogPCG, Verbose, TEXT("Waiting for '%s' to load."), *PrimitiveDescriptor.SkinnedAsset.ToString());
+
+				return false;
+			}
+		}
+
+		bStaticMeshesLoaded = true;
+	}
+
+	if (!bPrimitivesSetUp)
+	{
+		if (!SetupPrimitives(Context, InBinding))
+		{
+			return false;
+		}
+
+		bPrimitivesSetUp = true;
+
+		if (NumPrimitivesSetup > 0)
+		{
+			SourceComponent->NotifyProceduralInstancesInUse();
+
+			// Signal not finished. Compute graph element will wait a frame, giving time for GPU Scene to be updated.
+			return false;
+		}
+		else
+		{
+			// No component set up means we have no more work to do.
+			InBinding->CompletedMeshSpawners.Add(ProducerKernel);
+			return true;
+		}
+	}
+
+	InBinding->CompletedMeshSpawners.Add(ProducerKernel);
+	return true;
+}
+
+FComputeDataProviderRenderProxy* UPCGSkinnedMeshSpawnerDataProvider::GetRenderProxy()
+{
+	return new FPCGSkinnedMeshSpawnerDataProviderProxy(AttributeIdOffsetStrides, SelectorAttributeId, PrimitiveStringKeys, PrimitiveSelectionCDF, PrimitiveMeshBounds);
+}
+
+void UPCGSkinnedMeshSpawnerDataProvider::Reset()
+{
+	ProducerKernel = nullptr;
+	AttributeIdOffsetStrides.Empty();
+	PrimitiveStringKeys.Empty();
+	PrimitiveMeshBounds.Empty();
+	PrimitiveSelectionCDF.Empty();
+	SelectorAttributeId = INDEX_NONE;
+	NumInputPoints = 0;
+	StringKeyToInstanceCount.Empty();
+	AnalysisDataIndex = INDEX_NONE;
+	bPrimitiveDescriptorsCreated = false;
+	PrimitiveDescriptors.Empty();
+	CustomFloatCount = 0;
+	bPrimitivesSetUp = false;
+	NumPrimitivesProcessed = 0;
+	NumPrimitivesSetup = 0;
+	bRegisteredPrimitives = false;
+	bStaticMeshesLoaded = false;
+
+	Super::Reset();
+}
+
+void UPCGSkinnedMeshSpawnerDataProvider::CreatePrimitiveDescriptors(FPCGContext* InContext, UPCGDataBinding* InBinding)
+{
+	PrimitiveDescriptors.Empty();
+
+	const UPCGSkinnedMeshSpawnerSettings* Settings = CastChecked<UPCGSkinnedMeshSpawnerSettings>(ProducerKernel->GetSettings());
+
+	AActor* TargetActor = Settings->TargetActor.Get() ? Settings->TargetActor.Get() : InContext->GetTargetActor(nullptr);
+	if (!ensure(Settings->MeshSelectorParameters) || !ensure(TargetActor))
+	{
+		return;
+	}
+
+	const FPCGDataCollectionDesc* InputDataDesc = InBinding->GetCachedKernelPinDataDesc(ProducerKernel, PCGPinConstants::DefaultInputLabel, /*bIsInput=*/true);
+
+	if (!ensure(InputDataDesc))
+	{
+		return;
+	}
+
+	const int32 TotalInputPointCount = InputDataDesc->ComputeDataElementCount(EPCGDataType::Point);
+	if (TotalInputPointCount == 0)
+	{
+		return;
+	}
+
+	UPCGComponent* SourceComponent = Cast<UPCGComponent>(InContext->ExecutionSource.Get());
+	if (!SourceComponent)
+	{
+		return;
+	}
+
+	
+	const FBox LocalBounds = SourceComponent->GetGridBounds().InverseTransformBy(SourceComponent->GetOwner()->GetActorTransform());
+
+	if (Settings->InstanceDataPackerParameters && Settings->InstanceDataPackerParameters->GetAttributeNames(/*OutNames=*/nullptr))
+	{
+		TArray<FName> AttributeNames;
+		Settings->InstanceDataPackerParameters->GetAttributeNames(&AttributeNames);
+
+		for (FName AttributeName : AttributeNames)
+		{
+			if (AttributeName == NAME_None)
+			{
+				PCG_KERNEL_VALIDATION_ERR(InContext, Settings,
+					FText::Format(LOCTEXT("InstanceDataAttributeInvalid", "Invalid instance data attribute specified '{0}'."), FText::FromName(AttributeName)));
+
+				continue;
+			}
+
+			if (!InputDataDesc->ContainsAttributeOnAnyData(AttributeName))
+			{
+				PCG_KERNEL_VALIDATION_ERR(InContext, Settings,
+					FText::Format(LOCTEXT("InstanceDataAttributeNotFound", "Instance data attribute '{0}' not found."), FText::FromName(AttributeName)));
+
+				continue;
+			}
+		}
+
+		PCGDataForGPUHelpers::ComputeCustomFloatPacking(InContext, Settings, AttributeNames, InBinding, InputDataDesc, CustomFloatCount, AttributeIdOffsetStrides);
+	}
+
+	if (Settings->MeshSelectorParameters)
+	{
+		const FName SelectorName = Settings->MeshSelectorParameters->MeshAttribute.GetName(); // getname ok?
+
+		// Compute how many instances we expect for each mesh, if we do not already have the answer from analysis.
+		if (StringKeyToInstanceCount.IsEmpty())
+		{
+			for (const FPCGDataDesc& Desc : InputDataDesc->DataDescs)
+			{
+				if (Desc.ElementCount <= 0)
+				{
+					continue;
+				}
+
+				for (const FPCGKernelAttributeDesc& AttributeDesc : Desc.AttributeDescs)
+				{
+					if (AttributeDesc.AttributeKey.Identifier.Name == SelectorName)
+					{
+						if (AttributeDesc.AttributeKey.Type == EPCGKernelAttributeType::StringKey)
+						{
+							for (int32 AttributeStringKeyValue : AttributeDesc.GetUniqueStringKeys())
+							{
+								if (AttributeStringKeyValue > 0)
+								{
+									StringKeyToInstanceCount.FindOrAdd(AttributeStringKeyValue, 0) += Desc.ElementCount;
+								}
+							}
+
+							// We could early out but we currently continue looping to generate the below warning for mismatched types.
+						}
+						else
+						{
+							// Currently only a single type per attribute name is supported (the name/type in the attribute table). It's possible to wire a graph
+							// for which an attribute is present with multiple types. Warn if this is encountered.
+							PCG_KERNEL_VALIDATION_WARN(InContext, Settings, FText::Format(
+								LOCTEXT("MeshSelectorAttributeNotUsable", "Attribute '{0}' not usable for mesh selection, only attributes of type String Key are supported."),
+								FText::FromName(SelectorName)));
+						}
+					}
+				}
+			}
+		}
+
+		PrimitiveDescriptors.Reserve(StringKeyToInstanceCount.Num());
+
+		for (const TPair<int32, uint32>& StringKeyAndInstanceCount : StringKeyToInstanceCount)
+		{
+			if (StringKeyAndInstanceCount.Value == 0)
+			{
+				continue;
+			}
+
+			if (!ensure(InBinding->GetStringTable().IsValidIndex(StringKeyAndInstanceCount.Key)))
+			{
+				continue;
+			}
+
+			const FString& MeshPathString = InBinding->GetStringTable()[StringKeyAndInstanceCount.Key];
+			if (MeshPathString.IsEmpty())
+			{
+				continue;
+			}
+
+			if (PrimitiveDescriptors.Num() >= PCGComputeConstants::MAX_PRIMITIVE_COMPONENTS_PER_SPAWNER)
+			{
+				PCG_KERNEL_VALIDATION_WARN(InContext, Settings,
+					FText::Format(PCGSkinnedMeshSpawnerDataInterface::TooManyPrimitivesFormat, FText::FromString(FString::FromInt(PCGComputeConstants::MAX_PRIMITIVE_COMPONENTS_PER_SPAWNER))));
+
+				break;
+			}
+
+			FPCGSoftSkinnedMeshComponentDescriptor Descriptor{};
+			
+			Descriptor = Settings->MeshSelectorParameters->TemplateDescriptor;
+			Descriptor.bIsInstanceDataGPUOnly = true;
+			Descriptor.PrimitiveLocalBoundsOverride = LocalBounds;
+			Descriptor.NumInstancesGPUOnly = StringKeyAndInstanceCount.Value;
+			Descriptor.NumCustomDataFloatsGPUOnly = CustomFloatCount;
+			Descriptor.SkinnedAsset = FSoftObjectPath(MeshPathString);
+			Descriptor.bAffectDynamicIndirectLighting = false; // Not supported
+			Descriptor.bAffectDistanceFieldLighting = false; // Not supported
+
+			// Sanity check instance count.
+			if (!ensure(Descriptor.NumInstancesGPUOnly <= TotalInputPointCount))
+			{
+				Descriptor.NumInstancesGPUOnly = TotalInputPointCount;
+			}
+
+			PrimitiveStringKeys.Add(StringKeyAndInstanceCount.Key);
+			PrimitiveDescriptors.Add(MoveTemp(Descriptor));
+		}
+
+		PrimitiveSelectionCDF.SetNumZeroed(PrimitiveDescriptors.Num());
+	}
+
+	// Validate instance count is not too large.
+	uint64 TotalInstanceCount = 0;
+	for (const FPCGSoftSkinnedMeshComponentDescriptor& Desc : PrimitiveDescriptors)
+	{
+		TotalInstanceCount += Desc.NumInstancesGPUOnly;
+	}
+
+	if (TotalInstanceCount >= MAX_INSTANCE_ID)
+	{
+		PCG_KERNEL_VALIDATION_ERR(InContext, Settings, FText::Format(
+			LOCTEXT("TooManyInstances", "Tried to spawn too many instances ({0}), procedural ISM component creation skipped and instances will not be rendered."),
+			FText::FromString(FString::FromInt(TotalInstanceCount))));
+
+		return;
+	}
+
+	FPCGSpawnerPrimitives& Primitives = InBinding->MeshSpawnersToPrimitives.FindOrAdd(ProducerKernel);
+	Primitives.NumCustomFloats = CustomFloatCount;
+	Primitives.AttributeIdOffsetStrides = AttributeIdOffsetStrides;
+	Primitives.SelectorAttributeId = SelectorAttributeId;
+	Primitives.SelectionCDF = PrimitiveSelectionCDF;
+	Primitives.PrimitiveStringKeys = PrimitiveStringKeys;
+	Primitives.PrimitiveMeshBounds = PrimitiveMeshBounds;
+
+	// Useful for debugging instance counts.
+	//UE_LOG(LogPCG, Log, TEXT("Input points %d, total instance count %d (+ %.2f%%)"), TotalInputPointCount, Primitives.NumInstancesAllPrimitives, (float(Primitives.NumInstancesAllPrimitives - TotalInputPointCount) / TotalInputPointCount) * 100.0f);
+}
+
+bool UPCGSkinnedMeshSpawnerDataProvider::SetupPrimitives(FPCGContext* InContext, UPCGDataBinding* InBinding)
+{
+	const UPCGSkinnedMeshSpawnerSettings* Settings = CastChecked<UPCGSkinnedMeshSpawnerSettings>(ProducerKernel->GetSettings());
+
+	AActor* TargetActor = Settings->TargetActor.Get() ? Settings->TargetActor.Get() : InContext->GetTargetActor(nullptr);
+	if (!ensure(TargetActor))
+	{
+		return true;
+	}
+
+	FPCGSpawnerPrimitives* Primitives = InBinding->MeshSpawnersToPrimitives.Find(ProducerKernel);
+	if (!Primitives)
+	{
+		return true;
+	}
+
+	while (NumPrimitivesProcessed < PrimitiveDescriptors.Num())
+	{
+		const FPCGSoftSkinnedMeshComponentDescriptor& Desc = PrimitiveDescriptors[NumPrimitivesProcessed];
+
+		FPCGSkinnedMeshComponentBuilderParams Params;
+		Params.Descriptor = Desc;
+		Params.NumCustomDataFloats = 0; // This is the num custom floats allocated in CPU memory.
+
+		// If the root actor we're binding to is movable, then the component should be movable by default
+		if (USceneComponent* SceneComponent = TargetActor->GetRootComponent())
+		{
+			Params.Descriptor.Mobility = SceneComponent->Mobility;
+		}
+
+		Params.SettingsCrc = Settings->GetSettingsCrc();
+		ensure(Params.SettingsCrc.IsValid());
+
+		UPCGComponent* SourceComponent = Cast<UPCGComponent>(InContext->ExecutionSource.Get());
+		check(SourceComponent);
+
+		// TODO - not sure this is a good managed resource object to use, could make procedural variant?
+		UPCGManagedISKMComponent* ManagedComponent = UPCGActorHelpers::GetOrCreateManagedABMC(TargetActor, SourceComponent, Params, InContext);
+
+		if (ManagedComponent)
+		{
+			ManagedComponent->SetCrc(InContext->DependenciesCrc);
+
+			// Don't bother registering the resource change as we're transient anyway.
+			//InContext->TouchedResources.Emplace(MISMC);
+
+			FPCGSpawnerPrimitive Primitive;
+			Primitive.Component = ManagedComponent->GetComponent();
+			Primitive.NumInstances = Desc.NumInstancesGPUOnly;
+
+			Primitives->Primitives.Add(MoveTemp(Primitive));
+
+			if (Settings->bApplyMeshBoundsToPoints)
+			{
+				PrimitiveMeshBounds.Add(Desc.SkinnedAsset->GetBounds().GetBox());
+			}
+
+			++NumPrimitivesSetup;
+		}
+
+		++NumPrimitivesProcessed;
+
+		if (NumPrimitivesProcessed < PrimitiveDescriptors.Num() && InContext->AsyncState.ShouldStop())
+		{
+			return false;
+		}
+	}
+
+	return true;
+}
+
+FPCGSkinnedMeshSpawnerDataProviderProxy::FPCGSkinnedMeshSpawnerDataProviderProxy(
+	const TArray<FUintVector4>& InAttributeIdOffsetStrides,
+	int32 InSelectorAttributeId,
+	const TArray<int32>& InPrimitiveStringKeys,
+	TArray<float> InSelectionCDF,
+	const TArray<FBox>& InPrimitiveMeshBounds)
+	: AttributeIdOffsetStrides(InAttributeIdOffsetStrides)
+	, SelectionCDF(InSelectionCDF)
+	, SelectorAttributeId(InSelectorAttributeId)
+	, PrimitiveStringKeys(InPrimitiveStringKeys)
+{
+	PrimitiveMeshBoundsMin.Reserve(InPrimitiveMeshBounds.Num());
+	PrimitiveMeshBoundsMax.Reserve(InPrimitiveMeshBounds.Num());
+
+	for (int Index = 0; Index < InPrimitiveMeshBounds.Num(); ++Index)
+	{
+		PrimitiveMeshBoundsMin.Add(FVector4f(InPrimitiveMeshBounds[Index].Min.X, InPrimitiveMeshBounds[Index].Min.Y, InPrimitiveMeshBounds[Index].Min.Z, /*Unused*/0.0f));
+		PrimitiveMeshBoundsMax.Add(FVector4f(InPrimitiveMeshBounds[Index].Max.X, InPrimitiveMeshBounds[Index].Max.Y, InPrimitiveMeshBounds[Index].Max.Z, /*Unused*/0.0f));
+	}
+}
+
+bool FPCGSkinnedMeshSpawnerDataProviderProxy::IsValid(FValidationData const& InValidationData) const
+{
+	return InValidationData.ParameterStructSize == sizeof(FParameters);
+}
+
+void FPCGSkinnedMeshSpawnerDataProviderProxy::AllocateResources(FRDGBuilder& GraphBuilder, FAllocationData const& InAllocationData)
+{
+	if (!PrimitiveStringKeys.IsEmpty())
+	{
+		const FRDGBufferDesc BufferDesc = FRDGBufferDesc::CreateStructuredDesc(PrimitiveStringKeys.GetTypeSize(), PrimitiveStringKeys.Num());
+
+		FRDGBufferRef Buffer = GraphBuilder.CreateBuffer(BufferDesc, TEXT("PCGSkinnedMeshSpawner_PrimitiveStringKeys"));
+		PrimitiveStringKeysBufferSRV = GraphBuilder.CreateSRV(Buffer);
+
+		GraphBuilder.QueueBufferUpload(Buffer, MakeArrayView(PrimitiveStringKeys));
+	}
+	else
+	{
+		PrimitiveStringKeysBufferSRV = GraphBuilder.CreateSRV(FRDGBufferSRVDesc(GSystemTextures.GetDefaultStructuredBuffer(GraphBuilder, PrimitiveStringKeys.GetTypeSize())));
+	}
+
+	if (!PrimitiveMeshBoundsMin.IsEmpty())
+	{
+		const FRDGBufferDesc BufferDesc = FRDGBufferDesc::CreateStructuredDesc(PrimitiveMeshBoundsMin.GetTypeSize(), PrimitiveMeshBoundsMin.Num());
+
+		FRDGBufferRef Buffer = GraphBuilder.CreateBuffer(BufferDesc, TEXT("PCGSkinnedMeshSpawner_PrimitiveMeshBoundsMin"));
+		PrimitiveMeshBoundsMinBufferSRV = GraphBuilder.CreateSRV(Buffer);
+
+		GraphBuilder.QueueBufferUpload(Buffer, MakeArrayView(PrimitiveMeshBoundsMin));
+	}
+	else
+	{
+		PrimitiveMeshBoundsMinBufferSRV = GraphBuilder.CreateSRV(FRDGBufferSRVDesc(GSystemTextures.GetDefaultStructuredBuffer(GraphBuilder, PrimitiveMeshBoundsMin.GetTypeSize())));
+	}
+
+	if (!PrimitiveMeshBoundsMax.IsEmpty())
+	{
+		const FRDGBufferDesc BufferDesc = FRDGBufferDesc::CreateStructuredDesc(PrimitiveMeshBoundsMax.GetTypeSize(), PrimitiveMeshBoundsMax.Num());
+
+		FRDGBufferRef Buffer = GraphBuilder.CreateBuffer(BufferDesc, TEXT("PCGSkinnedMeshSpawner_PrimitiveMeshBoundsMax"));
+		PrimitiveMeshBoundsMaxBufferSRV = GraphBuilder.CreateSRV(Buffer);
+
+		GraphBuilder.QueueBufferUpload(Buffer, MakeArrayView(PrimitiveMeshBoundsMax));
+	}
+	else
+	{
+		PrimitiveMeshBoundsMaxBufferSRV = GraphBuilder.CreateSRV(FRDGBufferSRVDesc(GSystemTextures.GetDefaultStructuredBuffer(GraphBuilder, PrimitiveMeshBoundsMax.GetTypeSize())));
+	}
+}
+
+void FPCGSkinnedMeshSpawnerDataProviderProxy::GatherDispatchData(FDispatchData const& InDispatchData)
+{
+	const TStridedView<FParameters> ParameterArray = MakeStridedParameterView<FParameters>(InDispatchData);
+	for (int32 InvocationIndex = 0; InvocationIndex < ParameterArray.Num(); ++InvocationIndex)
+	{
+		FParameters& Parameters = ParameterArray[InvocationIndex];
+
+		Parameters.NumAttributes = AttributeIdOffsetStrides.Num();
+		Parameters.NumPrimitives = SelectionCDF.Num();
+		Parameters.SelectorAttributeId = SelectorAttributeId;
+
+		for (int32 Index = 0; Index < AttributeIdOffsetStrides.Num(); ++Index)
+		{
+			Parameters.AttributeIdOffsetStrides[Index] = AttributeIdOffsetStrides[Index];
+		}
+
+		Parameters.PrimitiveStringKeys = PrimitiveStringKeysBufferSRV;
+
+		for (int32 Index = 0; Index < SelectionCDF.Num(); ++Index)
+		{
+			GET_SCALAR_ARRAY_ELEMENT(Parameters.SelectionCDF, Index) = SelectionCDF[Index];
+		}
+
+		Parameters.ApplyBounds = PrimitiveMeshBoundsMin.IsEmpty() ? 0 : 1;
+
+		Parameters.PrimitiveMeshBoundsMin = PrimitiveMeshBoundsMinBufferSRV;
+		Parameters.PrimitiveMeshBoundsMax = PrimitiveMeshBoundsMaxBufferSRV;
+	}
+}
+
+#undef LOCTEXT_NAMESPACE
diff --git a/Engine/Plugins/PCG/Source/PCG/Private/Compute/DataInterfaces/Elements/PCGSkinnedMeshSpawnerDataInterface.h b/Engine/Plugins/PCG/Source/PCG/Private/Compute/DataInterfaces/Elements/PCGSkinnedMeshSpawnerDataInterface.h
new file mode 100644
--- /dev/null
+++ b/Engine/Plugins/PCG/Source/PCG/Private/Compute/DataInterfaces/Elements/PCGSkinnedMeshSpawnerDataInterface.h
@@ -0,0 +1,133 @@
+// Copyright Epic Games, Inc. All Rights Reserved.
+
+#pragma once
+
+#include "Compute/DataInterfaces/PCGComputeDataInterface.h"
+#include "ComputeFramework/ComputeDataProvider.h"
+
+#include "MeshSelectors/PCGSkinnedMeshDescriptor.h"
+
+#include "PCGSkinnedMeshSpawnerDataInterface.generated.h"
+
+class FPCGSkinnedMeshSpawnerDataInterfaceParameters;
+class UPCGSkinnedMeshSpawnerSettings;
+class UPCGSettings;
+class UPCGSkinnedMeshSpawnerKernel;
+struct FPCGContext;
+
+UCLASS(ClassGroup = (Procedural))
+class UPCGSkinnedMeshSpawnerDataInterface : public UPCGComputeDataInterface
+{
+	GENERATED_BODY()
+
+public:
+	//~ Begin UComputeDataInterface Interface
+	TCHAR const* GetClassName() const override { return TEXT("PCGSkinnedMeshSpawner"); }
+	void GetSupportedInputs(TArray<FShaderFunctionDefinition>& OutFunctions) const override;
+	void GetShaderParameters(TCHAR const* UID, FShaderParametersMetadataBuilder& InOutBuilder, FShaderParametersMetadataAllocations& InOutAllocations) const override;
+	void GetHLSL(FString& OutHLSL, FString const& InDataInterfaceName) const override;
+	UComputeDataProvider* CreateDataProvider() const override;
+	//~ End UComputeDataInterface Interface
+
+	static constexpr uint32 MAX_ATTRIBUTES = 64;
+};
+
+UCLASS()
+class UPCGSkinnedMeshSpawnerDataProvider : public UPCGComputeDataProvider
+{
+	GENERATED_BODY()
+
+public:
+	//~ Begin UComputeDataProvider Interface
+	virtual void Initialize(const UComputeDataInterface* InDataInterface, UObject* InBinding, uint64 InInputMask, uint64 InOutputMask) override;
+	virtual bool PerformPreExecuteReadbacks_GameThread(UPCGDataBinding* InBinding) override;
+	virtual bool PrepareForExecute_GameThread(UPCGDataBinding* InBinding) override;
+	FComputeDataProviderRenderProxy* GetRenderProxy() override;
+	virtual void Reset() override;
+	//~ End UComputeDataProvider Interface
+
+	void CreatePrimitiveDescriptors(FPCGContext* InContext, UPCGDataBinding* InBinding);
+
+	/** Setup Procedural ISMCs. Returns true if at least one component was set up. */
+	virtual bool SetupPrimitives(FPCGContext* InContext, UPCGDataBinding* InBinding);
+
+public:
+	UPROPERTY()
+	TObjectPtr<const UPCGSkinnedMeshSpawnerKernel> ProducerKernel;
+
+	/** Attributes to use for writing per-instance custom floats. */
+	UPROPERTY()
+	TArray<FUintVector4> AttributeIdOffsetStrides;
+
+	UPROPERTY()
+	TArray<int32> PrimitiveStringKeys;
+
+	UPROPERTY()
+	TArray<FBox> PrimitiveMeshBounds;
+
+	UPROPERTY()
+	TArray<float> PrimitiveSelectionCDF;
+
+	/** Attribute Id for mesh selector. */
+	UPROPERTY()
+	int32 SelectorAttributeId = INDEX_NONE;
+
+	UPROPERTY()
+	int32 NumInputPoints = 0;
+
+	/** The number of instances per string key value, used for by-attribute spawning. */
+	TMap<int32, uint32> StringKeyToInstanceCount;
+
+	int32 AnalysisDataIndex = INDEX_NONE;
+
+	bool bPrimitiveDescriptorsCreated = false;
+
+	TArray<FPCGSoftSkinnedMeshComponentDescriptor> PrimitiveDescriptors;
+
+	uint32 CustomFloatCount = 0;
+
+	bool bRegisteredPrimitives = false;
+
+	bool bStaticMeshesLoaded = false;
+
+	int32 NumPrimitivesProcessed = 0;
+
+	bool bPrimitivesSetUp = false;
+
+	/** Number of primitives created during primitive setup. */
+	int32 NumPrimitivesSetup = 0;
+};
+
+class FPCGSkinnedMeshSpawnerDataProviderProxy : public FComputeDataProviderRenderProxy
+{
+public:
+	FPCGSkinnedMeshSpawnerDataProviderProxy(
+		const TArray<FUintVector4>& InAttributeIdOffsetStrides,
+		int32 InSelectorAttributeId,
+		const TArray<int32>& InPrimitiveStringKeys,
+		TArray<float> InSelectionCDF,
+		const TArray<FBox>& InPrimitiveMeshBounds);
+
+	//~ Begin FComputeDataProviderRenderProxy Interface
+	bool IsValid(FValidationData const& InValidationData) const override;
+	void AllocateResources(FRDGBuilder& GraphBuilder, FAllocationData const& InAllocationData) override;
+	void GatherDispatchData(FDispatchData const& InDispatchData) override;
+	//~ End FComputeDataProviderRenderProxy Interface
+
+protected:
+	using FParameters = FPCGSkinnedMeshSpawnerDataInterfaceParameters;
+
+	TArray<FUintVector4> AttributeIdOffsetStrides;
+	TArray<float> SelectionCDF;
+
+	int32 SelectorAttributeId = INDEX_NONE;
+
+	TArray<int32> PrimitiveStringKeys;
+	FRDGBufferSRVRef PrimitiveStringKeysBufferSRV = nullptr;
+
+	TArray<FVector4f> PrimitiveMeshBoundsMin;
+	FRDGBufferSRVRef PrimitiveMeshBoundsMinBufferSRV = nullptr;
+
+	TArray<FVector4f> PrimitiveMeshBoundsMax;
+	FRDGBufferSRVRef PrimitiveMeshBoundsMaxBufferSRV = nullptr;
+};
diff --git a/Engine/Plugins/PCG/Source/PCG/Private/Compute/DataInterfaces/Elements/PCGStaticMeshSpawnerDataInterface.cpp b/Engine/Plugins/PCG/Source/PCG/Private/Compute/DataInterfaces/Elements/PCGStaticMeshSpawnerDataInterface.cpp
--- a/Engine/Plugins/PCG/Source/PCG/Private/Compute/DataInterfaces/Elements/PCGStaticMeshSpawnerDataInterface.cpp
+++ b/Engine/Plugins/PCG/Source/PCG/Private/Compute/DataInterfaces/Elements/PCGStaticMeshSpawnerDataInterface.cpp
@@ -327,6 +327,8 @@ bool UPCGStaticMeshSpawnerDataProvider::PrepareForExecute_GameThread(UPCGDataBin
 			return true;
 		}
 
+		bool bAnyPointsPresent = false;
+
 		for (const FPCGDataDesc& Desc : InputDataDesc->DataDescs)
 		{
 			if (Desc.ElementCount <= 0)
@@ -334,6 +336,8 @@ bool UPCGStaticMeshSpawnerDataProvider::PrepareForExecute_GameThread(UPCGDataBin
 				continue;
 			}
 
+			bAnyPointsPresent = true;
+
 			for (const FPCGKernelAttributeDesc& AttributeDesc : Desc.AttributeDescs)
 			{
 				if (AttributeDesc.AttributeKey.Identifier.Name == SelectorName && AttributeDesc.AttributeKey.Type == EPCGKernelAttributeType::StringKey)
@@ -351,7 +355,8 @@ bool UPCGStaticMeshSpawnerDataProvider::PrepareForExecute_GameThread(UPCGDataBin
 
 		if (SelectorAttributeId == INDEX_NONE)
 		{
-			if (!InputDataDesc->DataDescs.IsEmpty())
+			// Mute this error if the point data is empty.
+			if (!InputDataDesc->DataDescs.IsEmpty() && bAnyPointsPresent)
 			{
 				PCG_KERNEL_VALIDATION_ERR(Context, Settings, FText::Format(
 					LOCTEXT("MeshSelectorAttributeNotFound", "Mesh selector attribute '{0}' not found."),
@@ -497,7 +502,7 @@ void UPCGStaticMeshSpawnerDataProvider::CreatePrimitiveDescriptors(FPCGContext*
 		return;
 	}
 
-	const FBox LocalBounds = SourceComponent->GetGridBounds().ShiftBy(-SourceComponent->GetOwner()->GetActorLocation());
+	const FBox LocalBounds = SourceComponent->GetGridBounds().InverseTransformBy(SourceComponent->GetOwner()->GetActorTransform());
 
 	if (Settings->InstanceDataPackerParameters && Settings->InstanceDataPackerParameters->GetAttributeNames(/*OutNames=*/nullptr))
 	{
diff --git a/Engine/Plugins/PCG/Source/PCG/Private/Compute/DataInterfaces/PCGInstanceDataInterface.cpp b/Engine/Plugins/PCG/Source/PCG/Private/Compute/DataInterfaces/PCGInstanceDataInterface.cpp
--- a/Engine/Plugins/PCG/Source/PCG/Private/Compute/DataInterfaces/PCGInstanceDataInterface.cpp
+++ b/Engine/Plugins/PCG/Source/PCG/Private/Compute/DataInterfaces/PCGInstanceDataInterface.cpp
@@ -6,13 +6,14 @@
 #include "PCGSettings.h"
 #include "PCGSceneWriterCS.h"
 #include "Components/PCGProceduralISMComponent.h"
+#include "Compute/PCGComputeKernel.h"
 #include "Compute/PCGDataBinding.h"
-#include "Elements/PCGStaticMeshSpawnerKernel.h"
 
 #include "GlobalRenderResources.h"
 #include "GlobalShader.h"
 #include "GPUSceneWriter.h"
 #include "InstanceDataSceneProxy.h"
+#include "InstancedSkinnedMeshSceneProxyDesc.h"
 #include "InstancedStaticMeshSceneProxyDesc.h"
 #include "PrimitiveSceneDesc.h"
 #include "PrimitiveSceneInfo.h"
@@ -29,6 +30,7 @@
 #include "ShaderParameterMetadataBuilder.h"
 #include "ShaderParameterStruct.h"
 #include "Algo/AnyOf.h"
+#include "Components/InstancedSkinnedMeshComponent.h"
 #include "ComputeFramework/ComputeKernelPermutationSet.h"
 #include "ComputeFramework/ShaderParamTypeDefinition.h"
 #include "Engine/Texture.h"
@@ -139,7 +141,7 @@ void UPCGInstanceDataProvider::Initialize(const UComputeDataInterface* InDataInt
 	const UPCGInstanceDataInterface* DataInterface = CastChecked<UPCGInstanceDataInterface>(InDataInterface);
 	UPCGDataBinding* Binding = CastChecked<UPCGDataBinding>(InBinding);
 
-	ProducerKernel = CastChecked<UPCGStaticMeshSpawnerKernel>(DataInterface->ProducerKernel);
+	ProducerKernel = CastChecked<UComputeKernel>(DataInterface->ProducerKernel);
 	check(ProducerKernel);
 }
 
@@ -154,7 +156,7 @@ bool UPCGInstanceDataProvider::PrepareForExecute_GameThread(UPCGDataBinding* InB
 		return false;
 	}
 
-	FPCGSpawnerPrimitives* FoundPrimitives = Binding->MeshSpawnersToPrimitives.Find(Cast<UPCGStaticMeshSpawnerKernel>(ProducerKernel));
+	FPCGSpawnerPrimitives* FoundPrimitives = Binding->MeshSpawnersToPrimitives.Find(ProducerKernel);
 
 	if (!FoundPrimitives)
 	{
@@ -324,6 +326,13 @@ void FPCGInstanceDataProviderProxy::PostSubmit(FRDGBuilder& InGraphBuilder) cons
 		{
 			FPCGPrimitiveComponentAndNumInstances& PrimitiveAndInstanceCount = DataProviderInner->Primitives[PrimitiveIndex];
 
+			UPrimitiveComponent* PrimitiveComponent = PrimitiveAndInstanceCount.Get<0>().Get();
+			if (!PrimitiveComponent)
+			{
+				UE_LOG(LogPCG, Warning, TEXT("Null primitive component encountered during instancing, should not happen."));
+				continue;
+			}
+
 			const int32 PrimitiveNumInstances = PrimitiveAndInstanceCount.Get<1>();
 			if (PrimitiveNumInstances <= 0)
 			{
@@ -331,10 +340,43 @@ void FPCGInstanceDataProviderProxy::PostSubmit(FRDGBuilder& InGraphBuilder) cons
 				continue;
 			}
 
-			UPCGProceduralISMComponent* Primitive = Cast<UPCGProceduralISMComponent>(PrimitiveAndInstanceCount.Get<0>().Get());
-			UWorld* World = Primitive ? Primitive->GetWorld() : nullptr;
+			FPrimitiveSceneProxyDesc* SceneProxyDesc = nullptr;
+
+			// TODO: Slight mess, could add a GetSceneProxyDesc function to the base descriptor class. However we plan
+			// to move away from components and this code is temporary.
+			FInstancedStaticMeshSceneProxyDesc PISMCSceneProxyDesc;
+			FInstancedSkinnedMeshSceneProxyDesc ISKMSceneProxyDesc;
+
+			if (UPCGProceduralISMComponent* PISM = Cast<UPCGProceduralISMComponent>(PrimitiveComponent))
+			{
+				PISM->GetSceneProxyDesc(PISMCSceneProxyDesc);
+				SceneProxyDesc = &PISMCSceneProxyDesc;
+			}
+			else if (UInstancedSkinnedMeshComponent* ISKM = Cast<UInstancedSkinnedMeshComponent>(PrimitiveComponent))
+			{
+				ISKMSceneProxyDesc = FInstancedSkinnedMeshSceneProxyDesc(ISKM);
+				SceneProxyDesc = &ISKMSceneProxyDesc;
+			}
+
+			if (!SceneProxyDesc)
+			{
+				UE_LOG(LogPCG, Warning, TEXT("Primitive '%s' of unsupported type, instancing aborted."), *PrimitiveComponent->GetName());
+				continue;
+			}
+
+			FPrimitiveSceneDesc PrimitiveSceneDesc;
 
-			if (FSceneInterface* Scene = World ? World->Scene : nullptr)
+			// TODO try add FPrimitiveSceneDesc::InitFromPrimitiveComponent
+			if (UPCGProceduralISMComponent* PISM = Cast<UPCGProceduralISMComponent>(PrimitiveComponent))
+			{
+				PISM->BuildSceneDesc(SceneProxyDesc, PrimitiveSceneDesc);
+			}
+			else if (UInstancedSkinnedMeshComponent* ISKM = Cast<UInstancedSkinnedMeshComponent>(PrimitiveComponent))
+			{
+				ISKM->BuildSceneDesc(SceneProxyDesc, PrimitiveSceneDesc);
+			}
+
+			if (FSceneInterface* Scene = PrimitiveComponent->GetWorld() ? PrimitiveComponent->GetWorld()->Scene : nullptr)
 			{
 #if UE_BUILD_SHIPPING
 				const bool bTriggerCapture = false;
@@ -343,13 +385,6 @@ void FPCGInstanceDataProviderProxy::PostSubmit(FRDGBuilder& InGraphBuilder) cons
 				PCGInstanceDataInterface::TriggerGPUCaptureDispatchIndex = FMath::Max(PCGInstanceDataInterface::TriggerGPUCaptureDispatchIndex - 1, 0);
 #endif // UE_BUILD_SHIPPING
 
-				// todo_pcg: Can produce these upstream and move away from requiring PISMC components.
-				FInstancedStaticMeshSceneProxyDesc SceneProxyDesc;
-				Primitive->GetSceneProxyDesc(SceneProxyDesc);
-
-				FPrimitiveSceneDesc PrimitiveSceneDesc;
-				Primitive->BuildSceneDesc(&SceneProxyDesc, PrimitiveSceneDesc);
-
 				Scene->UpdatePrimitiveInstancesFromCompute(&PrimitiveSceneDesc, FGPUSceneWriteDelegate::CreateLambda([PrimitiveIndex, PrimitiveNumInstances, CumulativeInstanceCount, bTriggerCapture, InstanceDataExported, InstanceCustomFloatDataExported, WriteCountersExported](FRDGBuilder& InGraphBuilder, const FGPUSceneWriteDelegateParams& Params)
 				{
 					RDG_EVENT_SCOPE(InGraphBuilder, "SceneComputeUpdateInterface->EnqueueUpdate");
diff --git a/Engine/Plugins/PCG/Source/PCG/Private/Compute/DataInterfaces/PCGInstanceDataInterface.h b/Engine/Plugins/PCG/Source/PCG/Private/Compute/DataInterfaces/PCGInstanceDataInterface.h
--- a/Engine/Plugins/PCG/Source/PCG/Private/Compute/DataInterfaces/PCGInstanceDataInterface.h
+++ b/Engine/Plugins/PCG/Source/PCG/Private/Compute/DataInterfaces/PCGInstanceDataInterface.h
@@ -11,7 +11,7 @@
 class FPCGInstanceDataInterfaceParameters;
 class FPrimitiveSceneProxy;
 class UPCGSettings;
-class UPCGStaticMeshSpawnerKernel;
+class UComputeKernel;
 
 using FPCGPrimitiveComponentAndNumInstances = TPair<TWeakObjectPtr<UPrimitiveComponent>, int>;
 
@@ -57,7 +57,7 @@ public:
 	//~ End UComputeDataProvider Interface
 
 	UPROPERTY()
-	TObjectPtr<const UPCGStaticMeshSpawnerKernel> ProducerKernel;
+	TObjectPtr<const UComputeKernel> ProducerKernel;
 
 	TArray<FPCGPrimitiveComponentAndNumInstances> Primitives;
 
diff --git a/Engine/Plugins/PCG/Source/PCG/Private/Compute/PCGDataBinding.h b/Engine/Plugins/PCG/Source/PCG/Private/Compute/PCGDataBinding.h
--- a/Engine/Plugins/PCG/Source/PCG/Private/Compute/PCGDataBinding.h
+++ b/Engine/Plugins/PCG/Source/PCG/Private/Compute/PCGDataBinding.h
@@ -7,11 +7,11 @@
 
 #include "PCGDataBinding.generated.h"
 
+class UComputeKernel;
 class UPCGComponent;
 class UPCGComputeGraph;
 class UPCGComputeKernel;
 class UPCGProxyForGPUData;
-class UPCGStaticMeshSpawnerKernel;
 struct FPCGComputeGraphContext;
 struct FPCGContextHandle;
 struct FPCGKernelPin;
@@ -144,10 +144,10 @@ public:
 	FPCGDataCollection InputDataCollection;
 
 	UPROPERTY()
-	TMap<TObjectPtr<const UPCGStaticMeshSpawnerKernel>, FPCGSpawnerPrimitives> MeshSpawnersToPrimitives;
+	TMap<TObjectPtr<const UComputeKernel>, FPCGSpawnerPrimitives> MeshSpawnersToPrimitives;
 
 	UPROPERTY()	
-	TArray<TObjectPtr<const UPCGStaticMeshSpawnerKernel>> CompletedMeshSpawners;
+	TArray<TObjectPtr<const UComputeKernel>> CompletedMeshSpawners;
 
 	UPROPERTY()
 	TArray<FPCGDataToDebug> DataToDebug;
diff --git a/Engine/Plugins/PCG/Source/PCG/Private/Elements/PCGSkinnedMeshSpawner.cpp b/Engine/Plugins/PCG/Source/PCG/Private/Elements/PCGSkinnedMeshSpawner.cpp
--- a/Engine/Plugins/PCG/Source/PCG/Private/Elements/PCGSkinnedMeshSpawner.cpp
+++ b/Engine/Plugins/PCG/Source/PCG/Private/Elements/PCGSkinnedMeshSpawner.cpp
@@ -5,9 +5,12 @@
 #include "PCGComponent.h"
 #include "PCGCustomVersion.h"
 #include "PCGManagedResource.h"
+#include "Compute/BuiltInKernels/PCGCountUniqueAttributeValuesKernel.h"
 #include "Data/PCGPointData.h"
 #include "Data/PCGSpatialData.h"
 #include "Elements/PCGSkinnedMeshSpawnerContext.h"
+#include "Elements/PCGSkinnedMeshSpawnerKernel.h"
+#include "Graph/PCGGPUGraphCompilationContext.h"
 #include "Helpers/PCGActorHelpers.h"
 #include "Helpers/PCGHelpers.h"
 #include "InstanceDataPackers/PCGSkinnedMeshInstanceDataPackerBase.h"
@@ -34,6 +37,36 @@ UPCGSkinnedMeshSpawnerSettings::UPCGSkinnedMeshSpawnerSettings(const FObjectInit
 }
 
 #if WITH_EDITOR
+void UPCGSkinnedMeshSpawnerSettings::CreateKernels(FPCGGPUCompilationContext& InOutContext, UObject* InObjectOuter, TArray<UPCGComputeKernel*>& OutKernels, TArray<FPCGKernelEdge>& OutEdges) const
+{
+	FPCGComputeKernelParams KernelParams;
+	KernelParams.Settings = this;
+	KernelParams.bLogDescriptions = bDumpDataDescriptions;
+
+	UPCGSkinnedMeshSpawnerKernel* Kernel = InOutContext.NewObject_AnyThread<UPCGSkinnedMeshSpawnerKernel>(InObjectOuter);
+	Kernel->Initialize(KernelParams);
+
+	OutKernels.Add(Kernel);
+
+	// Connect node pins to kernel pins
+	OutEdges.Emplace(FPCGPinReference(PCGPinConstants::DefaultInputLabel), FPCGPinReference(Kernel, PCGPinConstants::DefaultInputLabel));
+	OutEdges.Emplace(FPCGPinReference(Kernel, PCGPinConstants::DefaultOutputLabel), FPCGPinReference(PCGPinConstants::DefaultOutputLabel));
+
+	// Add analysis kernel that will count how many instances of each mesh are present.
+	const UPCGSkinnedMeshSelector* Selector = Cast<UPCGSkinnedMeshSelector>(MeshSelectorParameters);
+	if (ensure(Selector))
+	{
+		UPCGCountUniqueAttributeValuesKernel* CountKernel = InOutContext.NewObject_AnyThread<UPCGCountUniqueAttributeValuesKernel>(InObjectOuter);
+		CountKernel->Initialize(KernelParams);
+		CountKernel->SetAttributeName(MeshSelectorParameters->MeshAttribute.GetName());
+
+		OutKernels.Add(CountKernel);
+
+		OutEdges.Emplace(FPCGPinReference(PCGPinConstants::DefaultInputLabel), FPCGPinReference(CountKernel, PCGPinConstants::DefaultInputLabel));
+		OutEdges.Emplace(FPCGPinReference(CountKernel, PCGPinConstants::DefaultOutputLabel), FPCGPinReference(Kernel, PCGSkinnedMeshSpawnerConstants::InstanceCountsPinLabel));
+	}
+}
+
 FText UPCGSkinnedMeshSpawnerSettings::GetDefaultNodeTitle() const
 {
 	return LOCTEXT("NodeTitle", "Instanced Skinned Mesh Spawner");
diff --git a/Engine/Plugins/PCG/Source/PCG/Private/Elements/PCGSkinnedMeshSpawnerKernel.cpp b/Engine/Plugins/PCG/Source/PCG/Private/Elements/PCGSkinnedMeshSpawnerKernel.cpp
new file mode 100644
--- /dev/null
+++ b/Engine/Plugins/PCG/Source/PCG/Private/Elements/PCGSkinnedMeshSpawnerKernel.cpp
@@ -0,0 +1,130 @@
+// Copyright Epic Games, Inc. All Rights Reserved.
+
+#include "Elements/PCGSkinnedMeshSpawnerKernel.h"
+
+#include "PCGComponent.h"
+#include "PCGContext.h"
+#include "PCGModule.h"
+#include "PCGParamData.h"
+#include "Components/PCGProceduralISMComponent.h"
+#include "Compute/PCGComputeGraph.h"
+#include "Compute/PCGDataBinding.h"
+#include "Compute/PCGPinPropertiesGPU.h"
+#include "Compute/BuiltInKernels/PCGCountUniqueAttributeValuesKernel.h"
+#include "Compute/Data/PCGProxyForGPUData.h"
+#include "Compute/DataInterfaces/PCGInstanceDataInterface.h"
+#include "Compute/DataInterfaces/Elements/PCGSkinnedMeshSpawnerDataInterface.h"
+#include "Elements/PCGSkinnedMeshSpawner.h"
+#include "Graph/PCGGPUGraphCompilationContext.h"
+#include "InstanceDataPackers/PCGSkinnedMeshInstanceDataPackerBase.h"
+#include "MeshSelectors/PCGSkinnedMeshSelector.h"
+
+#include "ShaderCompilerCore.h"
+#include "Engine/StaticMesh.h"
+
+#include UE_INLINE_GENERATED_CPP_BY_NAME(PCGSkinnedMeshSpawnerKernel)
+
+#define LOCTEXT_NAMESPACE "PCGSkinnedMeshSpawnerKernel"
+
+namespace PCGSkinnedMeshSpawnerKernel
+{
+	static const FText NoMeshEntriesFormat = LOCTEXT("NoMeshEntries", "No mesh entries provided in weighted mesh selector.");
+}
+
+const UPCGSettings* UPCGSkinnedMeshSpawnerKernel::GetSettings() const
+{
+	return Settings;
+}
+
+bool UPCGSkinnedMeshSpawnerKernel::ComputeOutputBindingDataDesc(const UPCGComputeGraph* InGraph, FName InOutputPinLabel, UPCGDataBinding* InBinding, FPCGDataCollectionDesc& OutDataDesc) const
+{
+	check(InGraph);
+	check(InBinding);
+
+	// Code assumes single output pin.
+	if (!ensure(InOutputPinLabel == PCGPinConstants::DefaultOutputLabel))
+	{
+		return false;
+	}
+
+	// Forward data from In to Out.
+	FPCGDataCollectionDesc PinDesc;
+	const FPCGKernelPin InputKernelPin(KernelIndex, PCGPinConstants::DefaultInputLabel, /*bIsInput=*/true);
+	ensure(InBinding->ComputeKernelPinDataDesc(InputKernelPin, &PinDesc));
+
+	OutDataDesc = MoveTemp(PinDesc);
+	return true;
+}
+
+int UPCGSkinnedMeshSpawnerKernel::ComputeThreadCount(const UPCGDataBinding* InBinding) const
+{
+	const FPCGDataCollectionDesc* InputPinDesc = InBinding->GetCachedKernelPinDataDesc(this, PCGPinConstants::DefaultInputLabel, /*bIsInput=*/true);
+	return ensure(InputPinDesc) ? InputPinDesc->ComputeDataElementCount(EPCGDataType::Point) : 0;
+}
+
+#if WITH_EDITOR
+FString UPCGSkinnedMeshSpawnerKernel::GetCookedSource(const UPCGComputeGraph* InComputeGraph) const
+{
+	FString TemplateFile;
+	ensure(LoadShaderSourceFile(TEXT("/Plugin/PCG/Private/Elements/PCGStaticMeshSpawner.usf"), EShaderPlatform::SP_PCD3D_SM5, &TemplateFile, nullptr));
+	return TemplateFile;
+}
+
+void UPCGSkinnedMeshSpawnerKernel::CreateAdditionalInputDataInterfaces(FPCGGPUCompilationContext& InOutContext, UObject* InObjectOuter, TArray<TObjectPtr<UComputeDataInterface>>& OutDataInterfaces) const
+{
+	Super::CreateAdditionalInputDataInterfaces(InOutContext, InObjectOuter, OutDataInterfaces);
+
+	TObjectPtr<UPCGSkinnedMeshSpawnerDataInterface> NodeDI = InOutContext.NewObject_AnyThread<UPCGSkinnedMeshSpawnerDataInterface>(InObjectOuter);
+	NodeDI->ProducerKernel = this;
+	OutDataInterfaces.Add(NodeDI);
+}
+
+void UPCGSkinnedMeshSpawnerKernel::CreateAdditionalOutputDataInterfaces(FPCGGPUCompilationContext& InOutContext, UObject* InObjectOuter, TArray<TObjectPtr<UComputeDataInterface>>& OutDataInterfaces) const
+{
+	Super::CreateAdditionalOutputDataInterfaces(InOutContext, InObjectOuter, OutDataInterfaces);
+
+	UPCGInstanceDataInterface* InstanceDI = InOutContext.NewObject_AnyThread<UPCGInstanceDataInterface>(InObjectOuter);
+	InstanceDI->ProducerKernel = this;
+	InstanceDI->InputPinProvidingData = PCGPinConstants::DefaultInputLabel;
+	OutDataInterfaces.Add(InstanceDI);
+}
+#endif
+
+void UPCGSkinnedMeshSpawnerKernel::GetInputPins(TArray<FPCGPinProperties>& OutPins) const
+{
+	OutPins.Emplace(PCGPinConstants::DefaultInputLabel, EPCGDataType::Point);
+	OutPins.Emplace(PCGSkinnedMeshSpawnerConstants::InstanceCountsPinLabel, EPCGDataType::Param);
+}
+
+void UPCGSkinnedMeshSpawnerKernel::GetOutputPins(TArray<FPCGPinPropertiesGPU>& OutPins) const
+{
+	OutPins.Emplace(PCGPinConstants::DefaultOutputLabel, EPCGDataType::Point);
+}
+
+#if WITH_EDITOR
+bool UPCGSkinnedMeshSpawnerKernel::PerformStaticValidation()
+{
+	TRACE_CPUPROFILER_EVENT_SCOPE(UPCGSkinnedMeshSpawnerKernel::PerformStaticValidation);
+
+	if (!Super::PerformStaticValidation())
+	{
+		return false;
+	}
+
+	const UPCGSkinnedMeshSpawnerSettings* SMSettings = CastChecked<UPCGSkinnedMeshSpawnerSettings>(Settings);
+
+	// Currently instance packers must be able to specify a full list of attribute names upfront, to build the attribute table at compile time.
+	// TODO: We should be able to augment a static attribute table with new attributes at execution time, which will allow other types like regex.
+	if (SMSettings->InstanceDataPackerParameters && !SMSettings->InstanceDataPackerParameters->GetAttributeNames(/*OutNames=*/nullptr))
+	{
+#if PCG_KERNEL_LOGGING_ENABLED
+		StaticLogEntries.Emplace(LOCTEXT("InvalidInstancePacker", "Selected instance packer does not support GPU execution."), EPCGKernelLogVerbosity::Error);
+#endif
+		return false;
+	}
+	
+	return true;
+}
+#endif
+
+#undef LOCTEXT_NAMESPACE
diff --git a/Engine/Plugins/PCG/Source/PCG/Private/Elements/PCGStaticMeshSpawner.cpp b/Engine/Plugins/PCG/Source/PCG/Private/Elements/PCGStaticMeshSpawner.cpp
--- a/Engine/Plugins/PCG/Source/PCG/Private/Elements/PCGStaticMeshSpawner.cpp
+++ b/Engine/Plugins/PCG/Source/PCG/Private/Elements/PCGStaticMeshSpawner.cpp
@@ -66,6 +66,8 @@ void UPCGStaticMeshSpawnerSettings::CreateKernels(FPCGGPUCompilationContext& InO
 	{
 		UPCGCountUniqueAttributeValuesKernel* CountKernel = InOutContext.NewObject_AnyThread<UPCGCountUniqueAttributeValuesKernel>(InObjectOuter);
 		CountKernel->Initialize(KernelParams);
+		CountKernel->SetAttributeName(Selector->AttributeName);
+
 		OutKernels.Add(CountKernel);
 
 		OutEdges.Emplace(FPCGPinReference(PCGPinConstants::DefaultInputLabel), FPCGPinReference(CountKernel, PCGPinConstants::DefaultInputLabel));
diff --git a/Engine/Plugins/PCG/Source/PCG/Public/Elements/PCGSkinnedMeshSpawner.h b/Engine/Plugins/PCG/Source/PCG/Public/Elements/PCGSkinnedMeshSpawner.h
--- a/Engine/Plugins/PCG/Source/PCG/Public/Elements/PCGSkinnedMeshSpawner.h
+++ b/Engine/Plugins/PCG/Source/PCG/Public/Elements/PCGSkinnedMeshSpawner.h
@@ -29,6 +29,8 @@ public:
 
 	// ~Begin UPCGSettings interface
 #if WITH_EDITOR
+	//virtual bool DisplayExecuteOnGPUSetting() const override { return true; } // Disabled pending testing.
+	virtual void CreateKernels(FPCGGPUCompilationContext& InOutContext, UObject* InObjectOuter, TArray<UPCGComputeKernel*>& OutKernels, TArray<FPCGKernelEdge>& OutEdges) const override;
 	virtual FName GetDefaultNodeName() const override { return FName(TEXT("SkinnedMeshSpawner")); }
 	UE_API virtual FText GetDefaultNodeTitle() const override;
 	virtual EPCGSettingsType GetType() const override { return EPCGSettingsType::Spawner; }
@@ -98,10 +100,7 @@ public:
 	bool bWarnOnIdenticalSpawn = true;
 
 protected:
-	UE_API void RefreshInstancePacker();
-
-private:
-	static UE_API TCHAR const* TemplateFilePath;
+	void RefreshInstancePacker();
 };
 
 class FPCGSkinnedMeshSpawnerElement : public IPCGElement
diff --git a/Engine/Plugins/PCG/Source/PCG/Public/Elements/PCGSkinnedMeshSpawnerKernel.h b/Engine/Plugins/PCG/Source/PCG/Public/Elements/PCGSkinnedMeshSpawnerKernel.h
new file mode 100644
--- /dev/null
+++ b/Engine/Plugins/PCG/Source/PCG/Public/Elements/PCGSkinnedMeshSpawnerKernel.h
@@ -0,0 +1,39 @@
+// Copyright Epic Games, Inc. All Rights Reserved.
+
+#pragma once
+
+#include "Compute/PCGComputeKernel.h"
+
+#include "PCGSkinnedMeshSpawnerKernel.generated.h"
+
+class UPCGSkinnedMeshSpawnerSettings;
+
+namespace PCGSkinnedMeshSpawnerConstants
+{
+	const FName InstanceCountsPinLabel = TEXT("InstanceCounts");
+}
+
+UCLASS()
+class UPCGSkinnedMeshSpawnerKernel : public UPCGComputeKernel
+{
+	GENERATED_BODY()
+
+public:
+	virtual const UPCGSettings* GetSettings() const override;
+	virtual bool ComputeOutputBindingDataDesc(const UPCGComputeGraph* InGraph, FName InOutputPinLabel, UPCGDataBinding* InBinding, FPCGDataCollectionDesc& OutDataDesc) const override;
+	virtual int ComputeThreadCount(const UPCGDataBinding* Binding) const override;
+#if WITH_EDITOR
+	virtual FString GetCookedSource(const UPCGComputeGraph* InComputeGraph) const override;
+	virtual FString GetEntryPoint() const override { return TEXT("PCGStaticMeshSpawnerCS"); }
+	void CreateAdditionalInputDataInterfaces(FPCGGPUCompilationContext& InOutContext, UObject* InObjectOuter, TArray<TObjectPtr<UComputeDataInterface>>& OutDataInterfaces) const override;
+	void CreateAdditionalOutputDataInterfaces(FPCGGPUCompilationContext& InOutContext, UObject* InObjectOuter, TArray<TObjectPtr<UComputeDataInterface>>& OutDataInterfaces) const override;
+#endif
+
+	virtual void GetInputPins(TArray<FPCGPinProperties>& OutPins) const override;
+	virtual void GetOutputPins(TArray<FPCGPinPropertiesGPU>& OutPins) const override;
+
+protected:
+#if WITH_EDITOR
+	virtual bool PerformStaticValidation() override;
+#endif
+};
diff --git a/Engine/Source/Runtime/Engine/Classes/Animation/AnimBank.h b/Engine/Source/Runtime/Engine/Classes/Animation/AnimBank.h
--- a/Engine/Source/Runtime/Engine/Classes/Animation/AnimBank.h
+++ b/Engine/Source/Runtime/Engine/Classes/Animation/AnimBank.h
@@ -386,6 +386,27 @@ public:
 	UPROPERTY(EditAnywhere, Category = "Component Settings")
 	uint8 bCastShadow : 1;
 
+	UPROPERTY(EditAnywhere, Category = "Component Settings")
+	bool bVisibleInRayTracing = true;
+
+	UPROPERTY(EditAnywhere, Category = "Component Settings")
+	bool bAffectDynamicIndirectLighting = true;
+
+	UPROPERTY(EditAnywhere, Category = "Component Settings")
+	bool bAffectDistanceFieldLighting = true;
+
+	UPROPERTY()
+	FBox PrimitiveLocalBoundsOverride;
+
+	UPROPERTY()
+	bool bIsInstanceDataGPUOnly = false;
+
+	UPROPERTY()
+	int32 NumInstancesGPUOnly = 0;
+
+	UPROPERTY()
+	int32 NumCustomDataFloatsGPUOnly = 0;
+
 #if WITH_EDITORONLY_DATA
 	UPROPERTY(EditAnywhere, Category = "Component Settings")
 	uint8 bIncludeInHLOD : 1;
diff --git a/Engine/Source/Runtime/Engine/Classes/Components/InstancedSkinnedMeshComponent.h b/Engine/Source/Runtime/Engine/Classes/Components/InstancedSkinnedMeshComponent.h
--- a/Engine/Source/Runtime/Engine/Classes/Components/InstancedSkinnedMeshComponent.h
+++ b/Engine/Source/Runtime/Engine/Classes/Components/InstancedSkinnedMeshComponent.h
@@ -17,6 +17,7 @@ class UMaterialInterface;
 class UTexture;
 class UAnimBank;
 struct FAnimBankRecordHandle;
+struct FPrimitiveSceneDesc;
 
 USTRUCT()
 struct FSkinnedMeshInstanceData
@@ -95,6 +96,15 @@ protected:
 	UPROPERTY(EditAnywhere, BlueprintReadOnly, Category = Animation, meta = (ShowOnlyInnerProperties))
 	TArray<struct FAnimBankItem> AnimBankItems;
 
+	UPROPERTY()
+	FBox PrimitiveLocalBoundsOverride;
+
+	UPROPERTY()
+	bool bIsInstanceDataGPUOnly = false;
+
+	UPROPERTY()
+	int32 NumInstancesGPUOnly = 0;
+
 public:
 
 	ENGINE_API UInstancedSkinnedMeshComponent(FVTableHelper& Helper);
@@ -144,7 +154,16 @@ public:
 #endif
 	//~ End USkinnedMeshComponent Interface
 
+	ENGINE_API void BuildSceneDesc(FPrimitiveSceneProxyDesc* InSceneProxyDesc, FPrimitiveSceneDesc& OutPrimitiveSceneDesc);
+
 public:
+	void SetNumGPUInstances(int32 InCount);
+
+	bool UsesGPUOnlyInstances() const { return bIsInstanceDataGPUOnly; }
+
+	FBox GetPrimitiveLocalBoundsOverride() const { return PrimitiveLocalBoundsOverride; }
+	void SetPrimitiveLocalBoundsOverride(const FBox& InBounds) { PrimitiveLocalBoundsOverride = InBounds; }
+
 	/** Add an instance to this component. Transform is given in local space of this component unless bWorldSpace is set. */
 	UFUNCTION(BlueprintCallable, Category = "Components|InstancedSkinnedMesh")
 	ENGINE_API FPrimitiveInstanceId AddInstance(const FTransform& InstanceTransform, int32 BankIndex, bool bWorldSpace = false);
@@ -208,6 +227,8 @@ public:
 	const TArray<float>& GetInstanceCustomData() const { return InstanceCustomData; }
 	const TArray<FSkinnedMeshInstanceData>& GetInstanceData() const { return InstanceData; }
 
+	int32 GetInstanceCountGPUOnly() const { return NumInstancesGPUOnly; }
+
 	int32 GetMinDrawDistance() const { return InstanceMinDrawDistance; }
 	void GetCullDistances(int32& OutStartCullDistance, int32& OutEndCullDistance) const { OutStartCullDistance = InstanceStartCullDistance; OutEndCullDistance = InstanceEndCullDistance; }
 	ENGINE_API void SetCullDistances(int32 StartCullDistance, int32 EndCullDistance);
@@ -243,6 +264,8 @@ private:
 	bool ShouldInheritPerInstanceData(const UInstancedSkinnedMeshComponent* InArchetype) const;
 	bool ShouldInheritPerInstanceData() const;
 
+	void SetInstanceDataGPUOnly(bool bInInstancesGPUOnly);
+
 	/** Sets up new instance data to sensible defaults, creates physics counterparts if possible. */
 	void SetupNewInstanceData(FSkinnedMeshInstanceData& InOutNewInstanceData, int32 InInstanceIndex, const FTransform3f& InInstanceTransform, int32 BankIndex);
 
@@ -252,15 +275,11 @@ private:
 	static ENGINE_API FSkeletalMeshObject* CreateMeshObject(const FInstancedSkinnedMeshSceneProxyDesc& InMeshDesc, TConstArrayView<FAnimBankItem> InAnimBankItems, FSkeletalMeshRenderData* InRenderData, ERHIFeatureLevel::Type InFeatureLevel);
 	static ENGINE_API FPrimitiveSceneProxy* CreateSceneProxy(const FInstancedSkinnedMeshSceneProxyDesc& Desc, bool bHideSkin, bool bShouldNaniteSkin, bool bIsEnabled, int32 MinLODIndex);
 
-	TSharedPtr<FInstanceDataSceneProxy, ESPMode::ThreadSafe> GetOrCreateInstanceDataSceneProxy()
-	{
-		return InstanceDataManager.GetOrCreateProxy();
-	}
+	TSharedPtr<FInstanceDataSceneProxy, ESPMode::ThreadSafe> GetOrCreateInstanceDataSceneProxy();
 
-	TSharedPtr<FInstanceDataSceneProxy, ESPMode::ThreadSafe> GetInstanceDataSceneProxy() const
-	{
-		return const_cast<UInstancedSkinnedMeshComponent*>(this)->InstanceDataManager.GetProxy();
-	}
+	TSharedPtr<FInstanceDataSceneProxy, ESPMode::ThreadSafe> GetInstanceDataSceneProxy() const;
+
+	TSharedPtr<FInstanceDataSceneProxy, ESPMode::ThreadSafe> CreateInstanceDataProxyGPUOnly() const;
 
 	FInstanceDataManager InstanceDataManager;
 
@@ -269,6 +288,7 @@ private:
 	friend class FInstancedSkinnedMeshComponentHelper;
 	friend struct FInstancedSkinnedMeshComponentInstanceData;
 	friend struct FInstancedSkinnedMeshSceneProxyDesc;
+	friend struct FSkinnedMeshComponentDescriptorBase;
 
 private:
 	void SetSkinnedAssetCallback();
@@ -284,8 +304,11 @@ public:
 	FInstancedSkinnedMeshComponentInstanceData() = default;
 	
 	FInstancedSkinnedMeshComponentInstanceData(const UInstancedSkinnedMeshComponent* InComponent)
-	: FSceneComponentInstanceData(InComponent)
-	, SkinnedAsset(InComponent->GetSkinnedAsset())
+		: FSceneComponentInstanceData(InComponent)
+		, SkinnedAsset(InComponent->GetSkinnedAsset())
+		, PrimitiveLocalBoundsOverride(InComponent->GetPrimitiveLocalBoundsOverride())
+		, bIsInstanceDataGPUOnly(InComponent->UsesGPUOnlyInstances())
+		, NumInstancesGPUOnly(InComponent->GetInstanceCountGPUOnly())
 	{
 	}
 
@@ -321,4 +344,13 @@ public:
 
 	UPROPERTY()
 	bool bHasPerInstanceHitProxies = false;
+
+	UPROPERTY()
+	FBox PrimitiveLocalBoundsOverride;
+
+	UPROPERTY()
+	bool bIsInstanceDataGPUOnly = false;
+
+	UPROPERTY()
+	int32 NumInstancesGPUOnly = 0;
 };
diff --git a/Engine/Source/Runtime/Engine/Private/Animation/AnimBank.cpp b/Engine/Source/Runtime/Engine/Private/Animation/AnimBank.cpp
--- a/Engine/Source/Runtime/Engine/Private/Animation/AnimBank.cpp
+++ b/Engine/Source/Runtime/Engine/Private/Animation/AnimBank.cpp
@@ -1404,10 +1404,21 @@ void FSkinnedMeshComponentDescriptorBase::InitFrom(const UInstancedSkinnedMeshCo
 	Template->GetCullDistances(InstanceStartCullDistance, InstanceEndCullDistance);
 	ComponentClass = Template->GetClass();
 	bCastShadow = Template->CastShadow;
+	bVisibleInRayTracing = Template->bVisibleInRayTracing;
+	bAffectDynamicIndirectLighting = Template->bAffectDynamicIndirectLighting;
+	bAffectDistanceFieldLighting = Template->bAffectDistanceFieldLighting;
+	PrimitiveLocalBoundsOverride = Template->GetPrimitiveLocalBoundsOverride();
 #if WITH_EDITOR
 	HLODBatchingPolicy = Template->HLODBatchingPolicy;
 	bIncludeInHLOD = Template->bEnableAutoLODGeneration;
-#endif	
+#endif
+
+	bIsInstanceDataGPUOnly = Template->UsesGPUOnlyInstances();
+	if (bIsInstanceDataGPUOnly)
+	{
+		NumInstancesGPUOnly = Template->GetInstanceCountGPUOnly();
+		NumCustomDataFloatsGPUOnly = Template->GetNumCustomDataFloats();
+	}
 }
 
 void FSkinnedMeshComponentDescriptorBase::InitComponent(UInstancedSkinnedMeshComponent* Component) const
@@ -1416,10 +1427,21 @@ void FSkinnedMeshComponentDescriptorBase::InitComponent(UInstancedSkinnedMeshCom
 	Component->MinDrawDistance = InstanceMinDrawDistance;
 	Component->SetCullDistances(InstanceStartCullDistance, InstanceEndCullDistance);
 	Component->CastShadow = bCastShadow;
+	Component->bVisibleInRayTracing = bVisibleInRayTracing;
+	Component->bAffectDynamicIndirectLighting = bAffectDynamicIndirectLighting;
+	Component->bAffectDistanceFieldLighting = bAffectDistanceFieldLighting;
+	Component->SetPrimitiveLocalBoundsOverride(PrimitiveLocalBoundsOverride);
 #if WITH_EDITOR
 	Component->HLODBatchingPolicy = HLODBatchingPolicy;
 	Component->bEnableAutoLODGeneration = bIncludeInHLOD;
 #endif
+
+	Component->SetInstanceDataGPUOnly(bIsInstanceDataGPUOnly);
+	if (bIsInstanceDataGPUOnly)
+	{
+		Component->SetNumGPUInstances(NumInstancesGPUOnly);
+		Component->SetNumCustomDataFloats(NumCustomDataFloatsGPUOnly);
+	}
 }
 
 bool FSkinnedMeshComponentDescriptorBase::operator!=(const FSkinnedMeshComponentDescriptorBase& Other) const
@@ -1439,7 +1461,14 @@ bool FSkinnedMeshComponentDescriptorBase::operator==(const FSkinnedMeshComponent
 		HLODBatchingPolicy == Other.HLODBatchingPolicy &&
 		bIncludeInHLOD == Other.bIncludeInHLOD &&
 #endif 
-		bCastShadow == Other.bCastShadow;
+		bCastShadow == Other.bCastShadow &&
+		bIsInstanceDataGPUOnly == Other.bIsInstanceDataGPUOnly &&
+		PrimitiveLocalBoundsOverride == Other.PrimitiveLocalBoundsOverride &&
+		NumInstancesGPUOnly == Other.NumInstancesGPUOnly &&
+		NumCustomDataFloatsGPUOnly == Other.NumCustomDataFloatsGPUOnly &&
+		bVisibleInRayTracing == Other.bVisibleInRayTracing &&
+		bAffectDynamicIndirectLighting == Other.bAffectDynamicIndirectLighting &&
+		bAffectDistanceFieldLighting == Other.bAffectDistanceFieldLighting;
 }
 
 FSkinnedMeshComponentDescriptor::FSkinnedMeshComponentDescriptor()
diff --git a/Engine/Source/Runtime/Engine/Private/Components/InstancedSkinnedMeshComponent.cpp b/Engine/Source/Runtime/Engine/Private/Components/InstancedSkinnedMeshComponent.cpp
--- a/Engine/Source/Runtime/Engine/Private/Components/InstancedSkinnedMeshComponent.cpp
+++ b/Engine/Source/Runtime/Engine/Private/Components/InstancedSkinnedMeshComponent.cpp
@@ -22,6 +22,7 @@
 #include "InstanceData/InstanceUpdateChangeSet.h"
 #include "InstanceData/InstanceDataUpdateUtils.h"
 #include "InstancedSkinnedMeshComponentHelper.h"
+#include "PrimitiveSceneDesc.h"
 
 #if WITH_EDITOR
 #include "WorldPartition/HLOD/HLODHashBuilder.h"
@@ -849,6 +850,8 @@ void UInstancedSkinnedMeshComponent::BeginDestroy()
 void UInstancedSkinnedMeshComponent::OnRegister()
 {
 	Super::OnRegister();
+
+	ensure(!PrimitiveLocalBoundsOverride.IsValid);
 }
 
 void UInstancedSkinnedMeshComponent::TickComponent(float DeltaTime, ELevelTick TickType, FActorComponentTickFunction* ThisTickFunction)
@@ -863,7 +866,7 @@ bool UInstancedSkinnedMeshComponent::IsEnabled() const
 
 int32 UInstancedSkinnedMeshComponent::GetInstanceCount() const
 {
-	return InstanceData.Num();
+	return bIsInstanceDataGPUOnly ? NumInstancesGPUOnly : InstanceData.Num();
 }
 
 TArrayView<const struct FAnimBankItem> UInstancedSkinnedMeshComponent::GetAnimBankItems() const
@@ -930,6 +933,19 @@ bool UInstancedSkinnedMeshComponent::ShouldInheritPerInstanceData(const UInstanc
 	return (bInheritPerInstanceData || !bEditableWhenInherited) && InArchetype && InArchetype->IsInBlueprint() && !IsTemplate();
 }
 
+void UInstancedSkinnedMeshComponent::SetInstanceDataGPUOnly(bool bInInstancesGPUOnly)
+{
+	if (bIsInstanceDataGPUOnly != bInInstancesGPUOnly)
+	{
+		bIsInstanceDataGPUOnly = bInInstancesGPUOnly;
+
+		if (bIsInstanceDataGPUOnly)
+		{
+			ClearInstances();
+		}
+	}
+}
+
 void UInstancedSkinnedMeshComponent::SetupNewInstanceData(FSkinnedMeshInstanceData& InOutNewInstanceData, int32 InInstanceIndex, const FTransform3f& InInstanceTransform, int32 InBankIndex)
 {
 	InOutNewInstanceData.Transform = InInstanceTransform;
@@ -956,6 +972,20 @@ void UInstancedSkinnedMeshComponent::PostAssetCompilation()
 
 #endif 
 
+void UInstancedSkinnedMeshComponent::BuildSceneDesc(FPrimitiveSceneProxyDesc* InSceneProxyDesc, FPrimitiveSceneDesc& OutPrimitiveSceneDesc) // TODO try add FPrimitiveSceneDesc::InitFromPrimitiveComponent
+{
+	check(InSceneProxyDesc);
+
+	OutPrimitiveSceneDesc.SceneProxy = GetSceneProxy();
+	OutPrimitiveSceneDesc.ProxyDesc = InSceneProxyDesc;
+	OutPrimitiveSceneDesc.PrimitiveSceneData = &GetSceneData();
+	OutPrimitiveSceneDesc.RenderMatrix = GetRenderMatrix();
+	OutPrimitiveSceneDesc.AttachmentRootPosition = GetComponentLocation();
+	OutPrimitiveSceneDesc.LocalBounds = CalcBounds(FTransform::Identity);
+	OutPrimitiveSceneDesc.Bounds = CalcBounds(GetComponentToWorld());
+	OutPrimitiveSceneDesc.Mobility = InSceneProxyDesc->Mobility;
+}
+
 FInstanceDataManagerSourceDataDesc UInstancedSkinnedMeshComponent::GetComponentDesc(ERHIFeatureLevel::Type FeatureLevel)
 {
 	FInstanceDataManagerSourceDataDesc ComponentDesc;
@@ -1046,6 +1076,12 @@ void UInstancedSkinnedMeshComponent::SendRenderInstanceData_Concurrent()
 {
 	Super::SendRenderInstanceData_Concurrent();
 
+	// If instance data is entirely GPU driven, don't upload from CPU.
+	if (bIsInstanceDataGPUOnly)
+	{
+		return;
+	}
+
 	// If the primitive isn't hidden update its instances.
 	const bool bDetailModeAllowsRendering = true;//DetailMode <= GetCachedScalabilityCVars().DetailMode;
 	// The proxy may not be created, this can happen when a SM is async loading for example.
@@ -1157,7 +1193,7 @@ FPrimitiveSceneProxy* UInstancedSkinnedMeshComponent::CreateSceneProxy()
 #endif
 
 	const USkinnedAsset* SkinnedAssetPtr = GetSkinnedAsset();
-	if (InstanceData.Num() == 0 || SkinnedAssetPtr == nullptr || SkinnedAssetPtr->IsCompiling())
+	if (GetInstanceCount() == 0 || SkinnedAssetPtr == nullptr || SkinnedAssetPtr->IsCompiling())
 	{
 		return nullptr;
 	}
@@ -1181,7 +1217,8 @@ FPrimitiveSceneProxy* UInstancedSkinnedMeshComponent::CreateSceneProxy()
 	Result = CreateSceneProxy(FInstancedSkinnedMeshSceneProxyDesc(this), bHideSkin, ShouldNaniteSkin(), IsEnabled(), ComputeMinLOD());
 
 	// Unclear exactly how this is supposed to work with a non-instanced proxy - will be interesting...
-	if (Result)
+	// If GPU-only flag set, instance data is entirely GPU driven, don't upload from CPU.
+	if (Result && !bIsInstanceDataGPUOnly)
 	{
 		InstanceDataManager.FlushChanges(GetComponentDesc(Result->GetScene().GetFeatureLevel()));
 	}
@@ -1395,13 +1432,23 @@ void UInstancedSkinnedMeshComponent::ApplyComponentInstanceData(struct FInstance
 
 	SelectedInstances = InstancedMeshData->SelectedInstances;
 	bHasPerInstanceHitProxies = InstancedMeshData->bHasPerInstanceHitProxies;
+	PrimitiveLocalBoundsOverride = InstancedMeshData->PrimitiveLocalBoundsOverride;
+	bIsInstanceDataGPUOnly = InstancedMeshData->bIsInstanceDataGPUOnly;
+	NumInstancesGPUOnly = InstancedMeshData->NumInstancesGPUOnly;
 #endif
 }
 
 FBoxSphereBounds UInstancedSkinnedMeshComponent::CalcBounds(const FTransform& LocalToWorld) const
 {
 	//SCOPE_CYCLE_COUNTER(STAT_CalcSkelMeshBounds);
-	return FInstancedSkinnedMeshComponentHelper::CalcBounds(*this, LocalToWorld);
+	if (PrimitiveLocalBoundsOverride.IsValid)
+	{
+		return PrimitiveLocalBoundsOverride.TransformBy(LocalToWorld);
+	}
+	else
+	{
+		return FInstancedSkinnedMeshComponentHelper::CalcBounds(*this, LocalToWorld);
+	}
 }
 
 void UInstancedSkinnedMeshComponent::SetSkinnedAssetCallback()
@@ -1437,6 +1484,11 @@ void UInstancedSkinnedMeshComponent::RefreshBoneTransforms(FActorComponentTickFu
 	//bNeedsRefreshTransform = false;
 }
 
+void UInstancedSkinnedMeshComponent::SetNumGPUInstances(int32 InCount)
+{
+	NumInstancesGPUOnly = InCount;
+}
+
 FPrimitiveInstanceId UInstancedSkinnedMeshComponent::AddInstance(const FTransform& InstanceTransform, int32 BankIndex, bool bWorldSpace)
 {
 	return AddInstanceInternal(InstanceData.Num(), InstanceTransform, BankIndex, bWorldSpace);
@@ -1786,4 +1838,51 @@ FPrimitiveSceneProxy* UInstancedSkinnedMeshComponent::CreateSceneProxy(const FIn
 	}
 
 	return Result;
-}
\ No newline at end of file
+}
+
+TSharedPtr<FInstanceDataSceneProxy, ESPMode::ThreadSafe> UInstancedSkinnedMeshComponent::GetOrCreateInstanceDataSceneProxy()
+{
+	if (bIsInstanceDataGPUOnly)
+	{
+		return CreateInstanceDataProxyGPUOnly();
+	}
+	else
+	{
+		return InstanceDataManager.GetOrCreateProxy();
+	}
+}
+
+TSharedPtr<FInstanceDataSceneProxy, ESPMode::ThreadSafe> UInstancedSkinnedMeshComponent::GetInstanceDataSceneProxy() const
+{
+	if (bIsInstanceDataGPUOnly)
+	{
+		return CreateInstanceDataProxyGPUOnly();
+	}
+	else
+	{
+		return const_cast<UInstancedSkinnedMeshComponent*>(this)->InstanceDataManager.GetProxy();
+	}
+}
+
+TSharedPtr<FInstanceDataSceneProxy, ESPMode::ThreadSafe> UInstancedSkinnedMeshComponent::CreateInstanceDataProxyGPUOnly() const
+{
+	FInstanceSceneDataBuffers InstanceSceneDataBuffers(/*InbInstanceDataIsGPUOnly=*/true);
+	{
+		FInstanceSceneDataBuffers::FAccessTag AccessTag(PointerHash(this));
+		FInstanceSceneDataBuffers::FWriteView ProxyData = InstanceSceneDataBuffers.BeginWriteAccess(AccessTag);
+
+		InstanceSceneDataBuffers.SetPrimitiveLocalToWorld(GetRenderMatrix(), AccessTag);
+
+		ProxyData.NumInstancesGPUOnly = GetInstanceCountGPUOnly();
+		ProxyData.NumCustomDataFloats = NumCustomDataFloats;
+		ProxyData.InstanceLocalBounds.SetNum(1);
+		ProxyData.InstanceLocalBounds[0] = ensure(GetSkinnedAsset()) ? GetSkinnedAsset()->GetBounds() : FBox();
+
+		ProxyData.Flags.bHasPerInstanceCustomData = ProxyData.NumCustomDataFloats > 0;
+
+		InstanceSceneDataBuffers.EndWriteAccess(AccessTag);
+		InstanceSceneDataBuffers.ValidateData();
+	}
+
+	return MakeShared<FInstanceDataSceneProxy, ESPMode::ThreadSafe>(MoveTemp(InstanceSceneDataBuffers));
+}
