[PCG] Add additional sources and kernel sources to Custom HLSL node
* These new sources will be correctly parsed and injected into the cooked HLSL
* Add an asset definition and custom details for PCGComputeSource, so that we can easily share source assets and modify them with our HLSL source text box.

diff --git a/Engine/Plugins/PCG/Source/PCG/Private/Compute/BuiltInKernels/PCGCountUniqueAttributeValuesKernel.cpp b/Engine/Plugins/PCG/Source/PCG/Private/Compute/BuiltInKernels/PCGCountUniqueAttributeValuesKernel.cpp
--- a/Engine/Plugins/PCG/Source/PCG/Private/Compute/BuiltInKernels/PCGCountUniqueAttributeValuesKernel.cpp
+++ b/Engine/Plugins/PCG/Source/PCG/Private/Compute/BuiltInKernels/PCGCountUniqueAttributeValuesKernel.cpp
@@ -141,7 +141,7 @@ bool UPCGCountUniqueAttributeValuesKernel::DoesOutputPinRequireZeroInitializatio
 }
 
 #if WITH_EDITOR
-FString UPCGCountUniqueAttributeValuesKernel::GetCookedSource(const UPCGComputeGraph* InComputeGraph) const
+FString UPCGCountUniqueAttributeValuesKernel::GetCookedSource(FPCGGPUCompilationContext& InOutContext) const
 {
 	FString TemplateFile;
 	ensure(LoadShaderSourceFile(TEXT("/Plugin/PCG/Private/BuiltInKernels/PCGCountUniqueAttributeValues.usf"), EShaderPlatform::SP_PCD3D_SM5, &TemplateFile, nullptr));
diff --git a/Engine/Plugins/PCG/Source/PCG/Private/Compute/BuiltInKernels/PCGCountUniqueAttributeValuesKernel.h b/Engine/Plugins/PCG/Source/PCG/Private/Compute/BuiltInKernels/PCGCountUniqueAttributeValuesKernel.h
--- a/Engine/Plugins/PCG/Source/PCG/Private/Compute/BuiltInKernels/PCGCountUniqueAttributeValuesKernel.h
+++ b/Engine/Plugins/PCG/Source/PCG/Private/Compute/BuiltInKernels/PCGCountUniqueAttributeValuesKernel.h
@@ -25,7 +25,7 @@ public:
 	virtual int ComputeThreadCount(const UPCGDataBinding* Binding) const override;
 	virtual bool DoesOutputPinRequireZeroInitialization(FName InOutputPinLabel) const override;
 #if WITH_EDITOR
-	virtual FString GetCookedSource(const UPCGComputeGraph* InComputeGraph) const override;
+	virtual FString GetCookedSource(FPCGGPUCompilationContext& InOutContext) const override;
 	virtual FString GetEntryPoint() const override { return TEXT("PCGCountUniqueAttributeValuesCS"); }
 	void CreateAdditionalInputDataInterfaces(FPCGGPUCompilationContext& InOutContext, UObject* InObjectOuter, TArray<TObjectPtr<UComputeDataInterface>>& OutDataInterfaces) const override;
 	// Split graph to read back analysis results.
diff --git a/Engine/Plugins/PCG/Source/PCG/Private/Compute/BuiltInKernels/PCGDataNumKernel.cpp b/Engine/Plugins/PCG/Source/PCG/Private/Compute/BuiltInKernels/PCGDataNumKernel.cpp
--- a/Engine/Plugins/PCG/Source/PCG/Private/Compute/BuiltInKernels/PCGDataNumKernel.cpp
+++ b/Engine/Plugins/PCG/Source/PCG/Private/Compute/BuiltInKernels/PCGDataNumKernel.cpp
@@ -85,7 +85,7 @@ bool UPCGDataNumKernel::IsKernelDataValid(FPCGContext* InContext) const
 }
 
 #if WITH_EDITOR
-FString UPCGDataNumKernel::GetCookedSource(const UPCGComputeGraph* InComputeGraph) const
+FString UPCGDataNumKernel::GetCookedSource(FPCGGPUCompilationContext& InOutContext) const
 {
 	FString DataNumFile;
 	ensure(LoadShaderSourceFile(TEXT("/Plugin/PCG/Private/BuiltInKernels/PCGDataNum.usf"), EShaderPlatform::SP_PCD3D_SM5, &DataNumFile, nullptr));
diff --git a/Engine/Plugins/PCG/Source/PCG/Private/Compute/BuiltInKernels/PCGDataNumKernel.h b/Engine/Plugins/PCG/Source/PCG/Private/Compute/BuiltInKernels/PCGDataNumKernel.h
--- a/Engine/Plugins/PCG/Source/PCG/Private/Compute/BuiltInKernels/PCGDataNumKernel.h
+++ b/Engine/Plugins/PCG/Source/PCG/Private/Compute/BuiltInKernels/PCGDataNumKernel.h
@@ -19,7 +19,7 @@ public:
 	virtual bool IsKernelDataValid(FPCGContext* InContext) const override;
 
 #if WITH_EDITOR
-	virtual FString GetCookedSource(const UPCGComputeGraph* InComputeGraph) const override;
+	virtual FString GetCookedSource(FPCGGPUCompilationContext& InOutContext) const override;
 	virtual FString GetEntryPoint() const override { return TEXT("PCGDataNumCS"); }
 	virtual void CreateAdditionalInputDataInterfaces(FPCGGPUCompilationContext& InOutContext, UObject* InObjectOuter, TArray<TObjectPtr<UComputeDataInterface>>& OutDataInterfaces) const override;
 #endif
diff --git a/Engine/Plugins/PCG/Source/PCG/Private/Compute/Elements/PCGCustomHLSL.h b/Engine/Plugins/PCG/Source/PCG/Private/Compute/Elements/PCGCustomHLSL.h
--- a/Engine/Plugins/PCG/Source/PCG/Private/Compute/Elements/PCGCustomHLSL.h
+++ b/Engine/Plugins/PCG/Source/PCG/Private/Compute/Elements/PCGCustomHLSL.h
@@ -12,6 +12,8 @@
 class UPCGComputeKernel;
 class UPCGPin;
 
+class UComputeSource;
+
 /** Type of kernel allows us to make decisions about execution automatically, streamlining authoring. */
 UENUM()
 enum class EPCGKernelType : uint8
@@ -157,6 +159,16 @@ public:
 #endif
 
 protected:
+#if WITH_EDITORONLY_DATA
+	/** Override your kernel with a PCG compute source asset. */
+	UPROPERTY(EditAnywhere, Category = "Settings", meta = (AllowedClasses = "PCGComputeSource"))
+	TObjectPtr<UComputeSource> KernelSourceOverride;
+
+	/** Additional source files to use in your kernel. */
+	UPROPERTY(EditAnywhere, Category = "Settings", meta = (AllowedClasses = "PCGComputeSource"))
+	TArray<TObjectPtr<UComputeSource>> AdditionalSources;
+#endif
+
 	/** Mute uninitialized data warnings. */
 	UPROPERTY(EditAnywhere, Category = "Settings")
 	bool bMuteUnwrittenPinDataErrors = false;
@@ -165,23 +177,23 @@ protected:
 	// Shader source and declarations are entirely editor-only, and should never be serialized outside of the editor.
 
 	/** Optional functions that can be called from the source. Intended to be edited using the Node Source Editor window. */
-	UPROPERTY(EditAnywhere, Category = "Settings", AdvancedDisplay, meta = (MultiLine = true))
+	UPROPERTY()
 	FString ShaderFunctions = "/** CUSTOM SHADER FUNCTIONS **/\n";
 
 	/** Shader code that forms the body of the kernel. Intended to be edited using the Node Source Editor window. */
-	UPROPERTY(EditAnywhere, Category = "Settings", AdvancedDisplay, meta = (MultiLine = true))
+	UPROPERTY()
 	FString ShaderSource;
 
 	/** Inputs data accessors that can be used from the shader code. Intended to be viewed using the Node Source Editor window. */
-	UPROPERTY(Transient, VisibleAnywhere, Category = "Settings", AdvancedDisplay, meta = (MultiLine = true))
+	UPROPERTY(Transient)
 	FString InputDeclarations;
 
 	/** Output data accessors that can be used from the shader code. Intended to be viewed using the Node Source Editor window. */
-	UPROPERTY(Transient, VisibleAnywhere, Category = "Settings", AdvancedDisplay, meta = (MultiLine = true))
+	UPROPERTY(Transient)
 	FString OutputDeclarations;
 
 	/** Helper data and functions that can be used from the shader code. Intended to be viewed using the Node Source Editor window. */
-	UPROPERTY(Transient, VisibleAnywhere, Category = "Settings", AdvancedDisplay, meta = (MultiLine = true))
+	UPROPERTY(Transient)
 	FString HelperDeclarations;
 #endif
 
diff --git a/Engine/Plugins/PCG/Source/PCG/Private/Compute/Elements/PCGCustomHLSLKernel.cpp b/Engine/Plugins/PCG/Source/PCG/Private/Compute/Elements/PCGCustomHLSLKernel.cpp
--- a/Engine/Plugins/PCG/Source/PCG/Private/Compute/Elements/PCGCustomHLSLKernel.cpp
+++ b/Engine/Plugins/PCG/Source/PCG/Private/Compute/Elements/PCGCustomHLSLKernel.cpp
@@ -10,12 +10,13 @@
 #include "Graph/PCGGPUGraphCompilationContext.h"
 #include "Compute/Elements/PCGCustomHLSL.h"
 
+#include "ComputeFramework/ComputeSource.h"
 #include "Containers/StaticArray.h"
-#include "Internationalization/Regex.h"
 
 #if WITH_EDITOR
 #include "Editor/IPCGEditorModule.h"
 #include "Framework/Text/SyntaxTokenizer.h"
+#include "Internationalization/Regex.h"
 #endif
 
 #include UE_INLINE_GENERATED_CPP_BY_NAME(PCGCustomHLSLKernel)
@@ -409,21 +410,14 @@ bool UPCGCustomHLSLKernel::ComputeOutputBindingDataDesc(const UPCGComputeGraph*
 }
 
 #if WITH_EDITOR
-FString UPCGCustomHLSLKernel::GetCookedSource(const UPCGComputeGraph* InComputeGraph) const
+FString UPCGCustomHLSLKernel::GetCookedSource(FPCGGPUCompilationContext& InOutContext) const
 {
-	check(InComputeGraph);
-
-	// @todo_pcg: Allow overriding kernel source from an external kernel source
 	const UPCGCustomHLSLSettings* CustomHLSLSettings = CastChecked<UPCGCustomHLSLSettings>(Settings);
-	FString Source = ParsedSource.Source;
-	FString Functions = TEXT("\n\n// User kernel functions\n") + ParsedFunctions.Source;
-
-	// @todo_pcg: Include additional source files in custom HLSL parsing pass, so that they support attribute names, data labels, etc.
-	ProcessShaderSource(InComputeGraph, Source, Functions);
-
 	FString ShaderPathName = CustomHLSLSettings->GetPathName();
 	PCGComputeHelpers::ConvertObjectPathToShaderFilePath(ShaderPathName);
 
+	const FString Source = ProcessShaderSource(InOutContext, ParsedSources[0]);
+	const FString AdditionalSources = ProcessAdditionalShaderSources(InOutContext);
 	//const bool bHasKernelKeyword = Source.Contains(TEXT("KERNEL"), ESearchCase::CaseSensitive);
 
 	const FIntVector GroupSize(PCGComputeConstants::THREAD_GROUP_SIZE, 1, 1);
@@ -505,6 +499,9 @@ FString UPCGCustomHLSLKernel::GetCookedSource(const UPCGComputeGraph* InComputeG
 
 	FString Result;
 
+	// Note, it would be preferable to have the AdditionalSources included via the kernel CreateAdditionalSources(), but when the HLSL is composed, those additional sources are
+	// placed above the data interfaces, so any additional sources would be unable to utilize functions provided by the data interfaces. Therefore we just inject them by hand here.
+
 	// TODO: Support KERNEL keyword in shader source. Could be handy for external source assets and breaking kernels into sections to
 	// support pin/attribute declarations, etc.
 	/*if (bHasKernelKeyword)
@@ -513,16 +510,15 @@ FString UPCGCustomHLSLKernel::GetCookedSource(const UPCGComputeGraph* InComputeG
 
 		Result = FString::Printf(TEXT(
 			"#line 0 \"%s\"\n" // ShaderPathName
-			"%s\n" // Includes
-			"%s\n" // Functions
+			"%s\n" // AdditionalSources
 			"%s\n" // Source
 			"%s { __kernel_func(%s); }\n"), // KernelFunc, UnWrappedDispatchThreadId
-			*ShaderPathName, *Includes, *Functions, *Source, *KernelFunc, *UnWrappedDispatchThreadId);
+			*ShaderPathName, *AdditionalSources, *Source, *KernelFunc, *UnWrappedDispatchThreadId);
 	}
 	else*/
 	{
 		Result = FString::Printf(TEXT(
-			"%s\n\n" // Functions
+			"%s\n\n" // AdditionalSources
 			"%s\n" // KernelFunc
 			"{\n"
 			"%s\n" // SetAsExecuted
@@ -532,7 +528,7 @@ FString UPCGCustomHLSLKernel::GetCookedSource(const UPCGComputeGraph* InComputeG
 			"#line 0 \"%s\"\n" // ShaderPathName
 			"%s\n" // Source
 			"}\n"),
-			*Functions, *KernelFunc, *SetAsExecuted, *UnWrappedDispatchThreadId, *KernelSpecificPreamble, *ShaderPathName, *Source);
+			*AdditionalSources, *KernelFunc, *SetAsExecuted, *UnWrappedDispatchThreadId, *KernelSpecificPreamble, *ShaderPathName, *Source);
 	}
 
 	return Result;
@@ -677,9 +673,7 @@ void UPCGCustomHLSLKernel::InitEntryPoint()
 
 void UPCGCustomHLSLKernel::ParseShaderSource()
 {
-	const UPCGCustomHLSLSettings* CustomHLSLSettings = CastChecked<UPCGCustomHLSLSettings>(Settings);
-	ParsedSource.Source = CustomHLSLSettings->ShaderSource;
-	ParsedFunctions.Source = CustomHLSLSettings->ShaderFunctions;
+	CreateParsedSources();
 
 	static const TArray<FString> AttributeTypeStrings =
 	{
@@ -703,6 +697,8 @@ void UPCGCustomHLSLKernel::ParseShaderSource()
 	TArray<FString> CopyElementKeywords;
 	TArray<FString> InitializeKeywords;
 
+	const UPCGCustomHLSLSettings* CustomHLSLSettings = CastChecked<UPCGCustomHLSLSettings>(Settings);
+
 	for (const FPCGPinProperties& PinProps : CustomHLSLSettings->InputPins)
 	{
 		if (!!(PinProps.AllowedTypes & (EPCGDataType::Point | EPCGDataType::Param)))
@@ -934,8 +930,61 @@ void UPCGCustomHLSLKernel::ParseShaderSource()
 		}
 	};
 
-	ParseHelper(ParsedSource);
-	ParseHelper(ParsedFunctions);
+	for (FPCGCustomHLSLParsedSource& ParsedSource : ParsedSources)
+	{
+		ParseHelper(ParsedSource);
+	}
+}
+
+void UPCGCustomHLSLKernel::CreateParsedSources()
+{
+	TArray<TObjectPtr<UComputeSource>> AdditionalSources;
+	TSet<TObjectPtr<UComputeSource>> VisitedAdditionalSources;
+
+	auto TraverseAdditionalSources = [&AdditionalSources, &VisitedAdditionalSources](TObjectPtr<UComputeSource> AdditionalSource, auto&& RecursiveCall)
+	{
+		if (!AdditionalSource || VisitedAdditionalSources.Contains(AdditionalSource))
+		{
+			return;
+		}
+
+		VisitedAdditionalSources.Add(AdditionalSource);
+
+		// We do a postfix traversal of the nested additional sources because we need them to be pasted higher in the resulting HLSL, since presumably a source depends on its additional sources.
+		for (TObjectPtr<UComputeSource> NestedSource : AdditionalSource->AdditionalSources)
+		{
+			RecursiveCall(NestedSource, RecursiveCall);
+		}
+
+		AdditionalSources.Add(AdditionalSource);
+	};
+
+	const UPCGCustomHLSLSettings* CustomHLSLSettings = CastChecked<UPCGCustomHLSLSettings>(Settings);
+	TArray<TObjectPtr<UComputeSource>> AdditionalSourcesToProcess = CustomHLSLSettings->AdditionalSources;
+
+	if (TObjectPtr<UComputeSource> KernelSourceOverride = CustomHLSLSettings->KernelSourceOverride)
+	{
+		ParsedSources.Emplace(KernelSourceOverride->GetSource());
+		VisitedAdditionalSources.Add(KernelSourceOverride);
+		AdditionalSourcesToProcess.Append(KernelSourceOverride->AdditionalSources);
+	}
+	else
+	{
+		ParsedSources.Emplace(CustomHLSLSettings->ShaderSource);
+		ParsedSources.Emplace(CustomHLSLSettings->ShaderFunctions);
+	}
+ 
+	for (TObjectPtr<UComputeSource> RootAdditionalSource : AdditionalSourcesToProcess)
+	{
+		TraverseAdditionalSources(RootAdditionalSource, TraverseAdditionalSources);
+	}
+
+	// Now that the additional sources are in post-fix order, we can begin to parse them.
+	for (TObjectPtr<UComputeSource> AdditionalSource : AdditionalSources)
+	{
+		check(AdditionalSource);
+		ParsedSources.Emplace(AdditionalSource->GetSource());
+	}
 }
 
 void UPCGCustomHLSLKernel::CollectDataLabels(const FPCGCustomHLSLParsedSource& InParsedSource)
@@ -1290,73 +1339,70 @@ bool UPCGCustomHLSLKernel::ValidateShaderSource(TArray<FPCGKernelLogEntry>& OutL
 	return true;
 }
 
-void UPCGCustomHLSLKernel::ProcessShaderSource(const UPCGComputeGraph* InComputeGraph, FString& OutShaderSource, FString& OutShaderFunctions) const
+FString UPCGCustomHLSLKernel::ProcessShaderSource(FPCGGPUCompilationContext& InOutContext, const FPCGCustomHLSLParsedSource& InParsedSource) const
 {
+	FString OutShaderSource = InParsedSource.Source;
+
+	const FPCGKernelAttributeTable* StaticAttributeTable = InOutContext.GetStaticAttributeTable();
+
+	if (!ensure(StaticAttributeTable))
+	{
+		return OutShaderSource;
+	}
+
 	// We inject attribute IDs directly into the source. This is most efficient and saves us passing them into the kernel. However
 	// the trade off is that a shader with a shared source will generate multiple variants if used in different compute graphs with
 	// different attribute IDs.
-	auto ReplaceAttributeNamesWithIds = [InComputeGraph](FString& InOutSource, const TArray<FPCGParsedAttributeFunction>& ParsedAttributeFunctions)
-	{
-		TArray<TTuple<int32, int32, int32>> Replacements;
+	//
+	// Attribute name replacement relies on precomputed indices into the source strings, therefore the attribute name replacement must take place
+	// before any other modifications. Otherwise, the indices will be incorrect and the source will become gibberish.
+	TArray<TTuple</*AttributeIndex=*/int32, /*ReplaceStartIndex=*/int32, /*ReplaceEndIndex=*/int32>> Replacements;
 
-		const UEnum* AttributeTypeEnum = StaticEnum<EPCGKernelAttributeType>();
-		check(AttributeTypeEnum);
+	const UEnum* AttributeTypeEnum = StaticEnum<EPCGKernelAttributeType>();
+	check(AttributeTypeEnum);
 
-		for (const FPCGParsedAttributeFunction& ParsedFunction : ParsedAttributeFunctions)
-		{
-			const int64 AttrType = ParsedFunction.AttributeType;
-
-			if (!ensure(AttrType != INDEX_NONE))
-			{
-				break;
-			}
+	for (const FPCGParsedAttributeFunction& ParsedFunction : InParsedSource.AttributeFunctions)
+	{
+		const int64 AttrType = ParsedFunction.AttributeType;
 
-			const FPCGKernelAttributeKey AttributeKey(FPCGAttributePropertySelector::CreateSelectorFromString(*ParsedFunction.AttributeName), static_cast<EPCGKernelAttributeType>(AttrType));
-			const FString SourceDefinition = FString::Format(TEXT("'{0}'"), { ParsedFunction.AttributeName });
+		if (!ensure(AttrType != INDEX_NONE))
+		{
+			break;
+		}
 
-			const int32 AttributeIndex = InComputeGraph->GetStaticAttributeTable().GetAttributeId(AttributeKey);
-			const int32 ReplaceStartIndex = InOutSource.Find(SourceDefinition, ESearchCase::CaseSensitive, ESearchDir::FromStart, ParsedFunction.MatchBeginning);
-			const int32 ReplaceEndIndex = ReplaceStartIndex + SourceDefinition.Len();
+		const FPCGKernelAttributeKey AttributeKey(FPCGAttributePropertySelector::CreateSelectorFromString(*ParsedFunction.AttributeName), static_cast<EPCGKernelAttributeType>(AttrType));
+		const FString SourceDefinition = FString::Format(TEXT("'{0}'"), { ParsedFunction.AttributeName });
 
-			Replacements.Emplace(AttributeIndex, ReplaceStartIndex, ReplaceEndIndex);
-		}
+		const int32 AttributeIndex = StaticAttributeTable->GetAttributeId(AttributeKey);
+		const int32 ReplaceStartIndex = OutShaderSource.Find(SourceDefinition, ESearchCase::CaseSensitive, ESearchDir::FromStart, ParsedFunction.MatchBeginning);
+		const int32 ReplaceEndIndex = ReplaceStartIndex + SourceDefinition.Len();
 
-		for (int I = Replacements.Num() - 1; I >= 0; --I)
-		{
-			const int32 AttributeIndex = Replacements[I].Get<0>();
-			const int32 ReplaceStartIndex = Replacements[I].Get<1>();
-			const int32 ReplaceEndIndex = Replacements[I].Get<2>();
+		Replacements.Emplace(AttributeIndex, ReplaceStartIndex, ReplaceEndIndex);
+	}
 
-			InOutSource = InOutSource.Left(ReplaceStartIndex) + FString::FromInt(AttributeIndex) + InOutSource.RightChop(ReplaceEndIndex);
-		}
-	};
+	for (int I = Replacements.Num() - 1; I >= 0; --I)
+	{
+		const int32 AttributeIndex = Replacements[I].Get<0>();
+		const int32 ReplaceStartIndex = Replacements[I].Get<1>();
+		const int32 ReplaceEndIndex = Replacements[I].Get<2>();
 
-	// Implementation note: Attribute name replacement relies on precomputed indices into the source strings, therefore the attribute name replacement must take place
-	// before any other modifications. Otherwise, the indices will be incorrect and the source will become gibberish.
-	ReplaceAttributeNamesWithIds(OutShaderSource, ParsedSource.AttributeFunctions);
-	ReplaceAttributeNamesWithIds(OutShaderFunctions, ParsedFunctions.AttributeFunctions);
+		OutShaderSource = OutShaderSource.Left(ReplaceStartIndex) + FString::FromInt(AttributeIndex) + OutShaderSource.RightChop(ReplaceEndIndex);
+	}
 
 #if PLATFORM_WINDOWS
 	// Remove old-school stuff.
 	OutShaderSource.ReplaceInline(TEXT("\r"), TEXT(""));
-	OutShaderFunctions.ReplaceInline(TEXT("\r"), TEXT(""));
 #endif
 
 	// @todo_pcg: Replace using token ranges instead of find/replace, similar to what we do with the parsed attribute functions.
-	auto ReplaceCopyElementFrom = [](FString& OutSource, const TArray<FPCGParsedCopyElementFunction>& ParsedFunctions)
+	// Replace function calls like Out_CopyElementFrom_In(...) with macro PCG_COPY_ALL_ATTRIBUTES_TO_OUTPUT(Out, In, ...).
+	for (const FPCGParsedCopyElementFunction& ParsedFunction : InParsedSource.CopyElementFunctions)
 	{
-		// Replace function calls like Out_CopyElementFrom_In(...) with macro PCG_COPY_ALL_ATTRIBUTES_TO_OUTPUT(Out, In, ...).
-		for (const FPCGParsedCopyElementFunction& ParsedFunction : ParsedFunctions)
-		{
-			OutSource.ReplaceInline(
-				*FString::Format(TEXT("{2}_{0}_{1}("), { PCGCustomHLSLKernel::CopyElementFunctionKeyword, ParsedFunction.SourcePin, ParsedFunction.TargetPin }),
-				*FString::Format(TEXT("PCG_COPY_ALL_ATTRIBUTES_TO_OUTPUT({1}, {0}, "), { ParsedFunction.SourcePin, ParsedFunction.TargetPin })
-			);
-		}
-	};
-
-	ReplaceCopyElementFrom(OutShaderSource, ParsedSource.CopyElementFunctions);
-	ReplaceCopyElementFrom(OutShaderSource, ParsedFunctions.CopyElementFunctions);
+		OutShaderSource.ReplaceInline(
+			*FString::Format(TEXT("{2}_{0}_{1}("), { PCGCustomHLSLKernel::CopyElementFunctionKeyword, ParsedFunction.SourcePin, ParsedFunction.TargetPin }),
+			*FString::Format(TEXT("PCG_COPY_ALL_ATTRIBUTES_TO_OUTPUT({1}, {0}, "), { ParsedFunction.SourcePin, ParsedFunction.TargetPin })
+		);
+	}
 
 	// We inject data IDs directly into the source. They will get remapped to data indices using a label resolver data interface.
 	for (const TPair<FName, FPCGDataLabels>& Pair : PinDataLabels.PinToDataLabels)
@@ -1376,23 +1422,33 @@ void UPCGCustomHLSLKernel::ProcessShaderSource(const UPCGComputeGraph* InCompute
 			// Matches against {PinName}_AnyFunction('{DataLabel}'...
 			// First capture group is the data label, so that we can find & replace it by index.
 			const FString Pattern = FString::Format(TEXT("{0}_.*?[\\s]*?\\([\\s]*?('{1}')"), { *PinLabel.ToString(), DataLabel });
+			FRegexMatcher ModuleMatcher(FRegexPattern(Pattern), OutShaderSource);
 
-			auto ReplaceDataLabels = [&Pattern, &ReplacementStr](FString& InOutSource)
+			while (ModuleMatcher.FindNext())
 			{
-				FRegexMatcher ModuleMatcher(FRegexPattern(Pattern), InOutSource);
-				while (ModuleMatcher.FindNext())
-				{
-					const int32 DataLabelCharIndexStart = ModuleMatcher.GetCaptureGroupBeginning(1);
-					const int32 DataLabelCharIndexEnd = ModuleMatcher.GetCaptureGroupEnding(1);
-
-					InOutSource = InOutSource.Left(DataLabelCharIndexStart) + ReplacementStr + InOutSource.RightChop(DataLabelCharIndexEnd);
-				}
-			};
+				const int32 DataLabelCharIndexStart = ModuleMatcher.GetCaptureGroupBeginning(1);
+				const int32 DataLabelCharIndexEnd = ModuleMatcher.GetCaptureGroupEnding(1);
 
-			ReplaceDataLabels(OutShaderSource);
-			ReplaceDataLabels(OutShaderFunctions);
+				OutShaderSource = OutShaderSource.Left(DataLabelCharIndexStart) + ReplacementStr + OutShaderSource.RightChop(DataLabelCharIndexEnd);
+			}
 		}
 	}
+
+	return OutShaderSource;
+}
+
+FString UPCGCustomHLSLKernel::ProcessAdditionalShaderSources(FPCGGPUCompilationContext& InOutContext) const
+{
+	// @todo_pcg: We should pivot to a stringbuilder here for perf.
+	FString OutShaderSource;
+
+	// The first parsed source is reserved for the kernel source.
+	for (int SourceIndex = 1; SourceIndex < ParsedSources.Num(); ++SourceIndex)
+	{
+		OutShaderSource += ProcessShaderSource(InOutContext, ParsedSources[SourceIndex]) + "\n\n";
+	}
+
+	return OutShaderSource;
 }
 
 bool UPCGCustomHLSLKernel::AreAllOutputPinsWritten(TArray<FPCGKernelLogEntry>& OutLogEntries) const
@@ -1417,8 +1473,18 @@ bool UPCGCustomHLSLKernel::AreAllOutputPinsWritten(TArray<FPCGKernelLogEntry>& O
 		}
 		
 		const FString& PinStr = CustomHLSLSettings->OutputPins[I].Label.ToString();
+		bool bInitializedByAnySource = false;
 
-		if (!IsPinInitializedBySource(ParsedSource, PinStr) && !IsPinInitializedBySource(ParsedFunctions, PinStr))
+		for (const FPCGCustomHLSLParsedSource& ParsedSource : ParsedSources)
+		{
+			if (IsPinInitializedBySource(ParsedSource, PinStr))
+			{
+				bInitializedByAnySource = true;
+				break;
+			}
+		}
+
+		if (!bInitializedByAnySource)
 		{
 #if PCG_KERNEL_LOGGING_ENABLED
 			OutLogEntries.Emplace(FText::Format(
@@ -1618,14 +1684,12 @@ bool UPCGCustomHLSLKernel::AreAttributesValid(FPCGContext* InContext, FText* Out
 			return true;
 		};
 
-		if (!ValidateParsedAttributeFunctions(ParsedSource.AttributeFunctions))
+		for (const FPCGCustomHLSLParsedSource& ParsedSource : ParsedSources)
 		{
-			return false;
-		}
-
-		if (!ValidateParsedAttributeFunctions(ParsedFunctions.AttributeFunctions))
-		{
-			return false;
+			if (!ValidateParsedAttributeFunctions(ParsedSource.AttributeFunctions))
+			{
+				return false;
+			}
 		}
 	}
 
diff --git a/Engine/Plugins/PCG/Source/PCG/Private/Compute/Elements/PCGCustomHLSLKernel.h b/Engine/Plugins/PCG/Source/PCG/Private/Compute/Elements/PCGCustomHLSLKernel.h
--- a/Engine/Plugins/PCG/Source/PCG/Private/Compute/Elements/PCGCustomHLSLKernel.h
+++ b/Engine/Plugins/PCG/Source/PCG/Private/Compute/Elements/PCGCustomHLSLKernel.h
@@ -92,12 +92,21 @@ struct FPCGCustomHLSLAttributeInfos
 	TArray<FPCGCustomHLSLAttributeInfo> Infos;
 };
 
+// @todo_pcg: It would be ideal someday for parsed source to be editor only. Unfortunately not possible right now because we need to validate the parsed attribute functions during execution.
 /** Holds the results of parsing for some HLSL source. */
 USTRUCT()
 struct FPCGCustomHLSLParsedSource
 {
 	GENERATED_BODY()
 
+	FPCGCustomHLSLParsedSource() = default;
+
+#if WITH_EDITOR
+	FPCGCustomHLSLParsedSource(const FString& InSource)
+		: Source(InSource)
+	{}
+#endif
+
 	UPROPERTY()
 	TArray<FPCGParsedAttributeFunction> AttributeFunctions;
 
@@ -141,7 +150,7 @@ public:
 	virtual bool IsKernelDataValid(FPCGContext* InContext) const override;
 	virtual bool ComputeOutputBindingDataDesc(const UPCGComputeGraph* InGraph, FName InOutputPinLabel, UPCGDataBinding* InBinding, FPCGDataCollectionDesc& OutDataDesc) const override;
 #if WITH_EDITOR
-	virtual FString GetCookedSource(const UPCGComputeGraph* InComputeGraph) const override;
+	virtual FString GetCookedSource(FPCGGPUCompilationContext& InOutContext) const override;
 	virtual FString GetEntryPoint() const override { return EntryPoint; }
 #endif
 	virtual int ComputeThreadCount(const UPCGDataBinding* Binding) const override;
@@ -160,13 +169,14 @@ protected:
 
 	void InitEntryPoint();
 	void ParseShaderSource();
+	void CreateParsedSources();
 	void CollectDataLabels(const FPCGCustomHLSLParsedSource& InParsedSource);
-
 	bool ValidateShaderSource(TArray<FPCGKernelLogEntry>& OutLogEntries) const;
 	bool AreAllOutputPinsWritten(TArray<FPCGKernelLogEntry>& OutLogEntries) const;
 
 	/** Apply any code gen/transformations/expansions to the shader source. */
-	virtual void ProcessShaderSource(const UPCGComputeGraph* InComputeGraph, FString& OutShaderSource, FString& OutShaderFunctions) const;
+	virtual FString ProcessShaderSource(FPCGGPUCompilationContext& InOutContext, const FPCGCustomHLSLParsedSource& InParsedSource) const;
+	virtual FString ProcessAdditionalShaderSources(FPCGGPUCompilationContext& InOutContext) const;
 #endif
 
 	/** Will the ThreadCountMultiplier value be applied when calculating the dispatch thread count. */
@@ -186,12 +196,8 @@ public:
 #endif
 
 protected:
-	// @todo_pcg: To support additional external sources, we should probably have an array/map of parsed sources instead.
-	UPROPERTY()
-	FPCGCustomHLSLParsedSource ParsedSource;
-
 	UPROPERTY()
-	FPCGCustomHLSLParsedSource ParsedFunctions;
+	TArray<FPCGCustomHLSLParsedSource> ParsedSources;
 
 	UPROPERTY()
 	TArray<FPCGKernelAttributeKey> KernelAttributeKeys;
diff --git a/Engine/Plugins/PCG/Source/PCG/Private/Compute/Elements/PCGNormalToDensityKernel.cpp b/Engine/Plugins/PCG/Source/PCG/Private/Compute/Elements/PCGNormalToDensityKernel.cpp
--- a/Engine/Plugins/PCG/Source/PCG/Private/Compute/Elements/PCGNormalToDensityKernel.cpp
+++ b/Engine/Plugins/PCG/Source/PCG/Private/Compute/Elements/PCGNormalToDensityKernel.cpp
@@ -52,7 +52,7 @@ int UPCGNormalToDensityKernel::ComputeThreadCount(const UPCGDataBinding* InBindi
 }
 
 #if WITH_EDITOR
-FString UPCGNormalToDensityKernel::GetCookedSource(const UPCGComputeGraph* InComputeGraph) const
+FString UPCGNormalToDensityKernel::GetCookedSource(FPCGGPUCompilationContext& InOutContext) const
 {
 	FString TemplateFile;
 	ensure(LoadShaderSourceFile(TEXT("/Plugin/PCG/Private/Elements/PCGNormalToDensity.usf"), EShaderPlatform::SP_PCD3D_SM5, &TemplateFile, nullptr));
diff --git a/Engine/Plugins/PCG/Source/PCG/Private/Compute/Elements/PCGNormalToDensityKernel.h b/Engine/Plugins/PCG/Source/PCG/Private/Compute/Elements/PCGNormalToDensityKernel.h
--- a/Engine/Plugins/PCG/Source/PCG/Private/Compute/Elements/PCGNormalToDensityKernel.h
+++ b/Engine/Plugins/PCG/Source/PCG/Private/Compute/Elements/PCGNormalToDensityKernel.h
@@ -20,7 +20,7 @@ public:
 	virtual int ComputeThreadCount(const UPCGDataBinding* Binding) const override;
 
 #if WITH_EDITOR
-	virtual FString GetCookedSource(const UPCGComputeGraph* InComputeGraph) const override;
+	virtual FString GetCookedSource(FPCGGPUCompilationContext& InOutContext) const override;
 	virtual FString GetEntryPoint() const override { return TEXT("PCGNormalToDensityCS"); }
 	void CreateAdditionalInputDataInterfaces(FPCGGPUCompilationContext& InOutContext, UObject* InObjectOuter, TArray<TObjectPtr<UComputeDataInterface>>& OutDataInterfaces) const override;
 #endif
diff --git a/Engine/Plugins/PCG/Source/PCG/Private/Compute/PCGComputeKernel.h b/Engine/Plugins/PCG/Source/PCG/Private/Compute/PCGComputeKernel.h
--- a/Engine/Plugins/PCG/Source/PCG/Private/Compute/PCGComputeKernel.h
+++ b/Engine/Plugins/PCG/Source/PCG/Private/Compute/PCGComputeKernel.h
@@ -5,7 +5,6 @@
 #include "ComputeFramework/ComputeKernel.h"
 #include "Compute/PCGComputeCommon.h"
 
-#include "ComputeFramework/ComputeGraph.h"
 #include "UObject/ObjectKey.h"
 
 #include "PCGComputeKernel.generated.h"
@@ -20,6 +19,8 @@ struct FPCGKernelAttributeKey;
 struct FPCGPinProperties;
 struct FPCGPinPropertiesGPU;
 
+class UComputeDataInterface;
+
 UENUM()
 enum class EPCGKernelLogVerbosity
 {
@@ -81,7 +82,7 @@ public:
 
 #if WITH_EDITOR
 	/** Produces the node specific portion of kernel shader source text, including the main entry point. */
-	virtual FString GetCookedSource(const UPCGComputeGraph* InComputeGraph) const PURE_VIRTUAL(UPCGComputeKernel::GetCookedSource, return {};);
+	virtual FString GetCookedSource(FPCGGPUCompilationContext& InOutContext) const PURE_VIRTUAL(UPCGComputeKernel::GetCookedSource, return {};);
 
 	/** Get the name of the main kernel function in the source. This name is also displayed in GPU profile scopes. */
 	virtual FString GetEntryPoint() const PURE_VIRTUAL(UPCGComputeKernel::GetEntryPoint, return TEXT("Main"););
@@ -102,7 +103,7 @@ public:
 	virtual bool ComputeOutputBindingDataDesc(const UPCGComputeGraph* InGraph, FName InOutputPinLabel, UPCGDataBinding* InBinding, FPCGDataCollectionDesc& OutDataDesc) const PURE_VIRTUAL(UPCGComputeKernel::ComputeOutputBindingDataDesc, return {};);
 
 	/** Compute how many threads should be dispatched to execute this node on the GPU. */
-	virtual int ComputeThreadCount(const UPCGDataBinding* Binding) const PURE_VIRTUAL(UPCGComputeKernel::GetCookedSource, return 0;);
+	virtual int ComputeThreadCount(const UPCGDataBinding* Binding) const PURE_VIRTUAL(UPCGComputeKernel::ComputeThreadCount, return 0;);
 
 	/** Get any data labels that are used statically by this node. */
 	virtual void GetDataLabels(FName InPinLabel, TArray<FString>& OutDataLabels) const {}
diff --git a/Engine/Plugins/PCG/Source/PCG/Private/Compute/PCGComputeSource.cpp b/Engine/Plugins/PCG/Source/PCG/Private/Compute/PCGComputeSource.cpp
--- a/Engine/Plugins/PCG/Source/PCG/Private/Compute/PCGComputeSource.cpp
+++ b/Engine/Plugins/PCG/Source/PCG/Private/Compute/PCGComputeSource.cpp
@@ -27,3 +27,14 @@ FString UPCGComputeSource::GetVirtualPath() const
 	return "";
 #endif
 }
+
+#if WITH_EDITOR
+void UPCGComputeSource::SetSource(const FString& InSource)
+{
+	if (Source != InSource)
+	{
+		Modify();
+		Source = InSource;
+	}
+}
+#endif
diff --git a/Engine/Plugins/PCG/Source/PCG/Private/Compute/PCGComputeSource.h b/Engine/Plugins/PCG/Source/PCG/Private/Compute/PCGComputeSource.h
deleted file mode 100644
diff --git a/Engine/Plugins/PCG/Source/PCG/Private/Elements/Metadata/PCGMetadataPartitionKernel.cpp b/Engine/Plugins/PCG/Source/PCG/Private/Elements/Metadata/PCGMetadataPartitionKernel.cpp
--- a/Engine/Plugins/PCG/Source/PCG/Private/Elements/Metadata/PCGMetadataPartitionKernel.cpp
+++ b/Engine/Plugins/PCG/Source/PCG/Private/Elements/Metadata/PCGMetadataPartitionKernel.cpp
@@ -146,7 +146,7 @@ int UPCGMetadataPartitionKernel::ComputeThreadCount(const UPCGDataBinding* InBin
 }
 
 #if WITH_EDITOR
-FString UPCGMetadataPartitionKernel::GetCookedSource(const UPCGComputeGraph* InComputeGraph) const
+FString UPCGMetadataPartitionKernel::GetCookedSource(FPCGGPUCompilationContext& InOutContext) const
 {
 	FString TemplateFile;
 	ensure(LoadShaderSourceFile(TEXT("/Plugin/PCG/Private/Elements/PCGMetadataPartition.usf"), EShaderPlatform::SP_PCD3D_SM5, &TemplateFile, nullptr));
diff --git a/Engine/Plugins/PCG/Source/PCG/Private/Elements/PCGCopyPointsAnalysisKernel.cpp b/Engine/Plugins/PCG/Source/PCG/Private/Elements/PCGCopyPointsAnalysisKernel.cpp
--- a/Engine/Plugins/PCG/Source/PCG/Private/Elements/PCGCopyPointsAnalysisKernel.cpp
+++ b/Engine/Plugins/PCG/Source/PCG/Private/Elements/PCGCopyPointsAnalysisKernel.cpp
@@ -109,7 +109,7 @@ int UPCGCopyPointsAnalysisKernel::ComputeThreadCount(const UPCGDataBinding* InBi
 }
 
 #if WITH_EDITOR
-FString UPCGCopyPointsAnalysisKernel::GetCookedSource(const UPCGComputeGraph* InComputeGraph) const
+FString UPCGCopyPointsAnalysisKernel::GetCookedSource(FPCGGPUCompilationContext& InOutContext) const
 {
 	FString TemplateFile;
 	ensure(LoadShaderSourceFile(TEXT("/Plugin/PCG/Private/Elements/PCGCopyPointsAnalysis.usf"), EShaderPlatform::SP_PCD3D_SM5, &TemplateFile, nullptr));
diff --git a/Engine/Plugins/PCG/Source/PCG/Private/Elements/PCGCopyPointsKernel.cpp b/Engine/Plugins/PCG/Source/PCG/Private/Elements/PCGCopyPointsKernel.cpp
--- a/Engine/Plugins/PCG/Source/PCG/Private/Elements/PCGCopyPointsKernel.cpp
+++ b/Engine/Plugins/PCG/Source/PCG/Private/Elements/PCGCopyPointsKernel.cpp
@@ -199,7 +199,7 @@ int UPCGCopyPointsKernel::ComputeThreadCount(const UPCGDataBinding* InBinding) c
 }
 
 #if WITH_EDITOR
-FString UPCGCopyPointsKernel::GetCookedSource(const UPCGComputeGraph* InComputeGraph) const
+FString UPCGCopyPointsKernel::GetCookedSource(FPCGGPUCompilationContext& InOutContext) const
 {
 	FString TemplateFile;
 	ensure(LoadShaderSourceFile(TEXT("/Plugin/PCG/Private/Elements/PCGCopyPoints.usf"), EShaderPlatform::SP_PCD3D_SM5, &TemplateFile, nullptr));
diff --git a/Engine/Plugins/PCG/Source/PCG/Private/Elements/PCGSkinnedMeshSpawnerKernel.cpp b/Engine/Plugins/PCG/Source/PCG/Private/Elements/PCGSkinnedMeshSpawnerKernel.cpp
--- a/Engine/Plugins/PCG/Source/PCG/Private/Elements/PCGSkinnedMeshSpawnerKernel.cpp
+++ b/Engine/Plugins/PCG/Source/PCG/Private/Elements/PCGSkinnedMeshSpawnerKernel.cpp
@@ -63,7 +63,7 @@ int UPCGSkinnedMeshSpawnerKernel::ComputeThreadCount(const UPCGDataBinding* InBi
 }
 
 #if WITH_EDITOR
-FString UPCGSkinnedMeshSpawnerKernel::GetCookedSource(const UPCGComputeGraph* InComputeGraph) const
+FString UPCGSkinnedMeshSpawnerKernel::GetCookedSource(FPCGGPUCompilationContext& InOutContext) const
 {
 	FString TemplateFile;
 	ensure(LoadShaderSourceFile(TEXT("/Plugin/PCG/Private/Elements/PCGStaticMeshSpawner.usf"), EShaderPlatform::SP_PCD3D_SM5, &TemplateFile, nullptr));
diff --git a/Engine/Plugins/PCG/Source/PCG/Private/Elements/PCGStaticMeshSpawnerKernel.cpp b/Engine/Plugins/PCG/Source/PCG/Private/Elements/PCGStaticMeshSpawnerKernel.cpp
--- a/Engine/Plugins/PCG/Source/PCG/Private/Elements/PCGStaticMeshSpawnerKernel.cpp
+++ b/Engine/Plugins/PCG/Source/PCG/Private/Elements/PCGStaticMeshSpawnerKernel.cpp
@@ -116,7 +116,7 @@ int UPCGStaticMeshSpawnerKernel::ComputeThreadCount(const UPCGDataBinding* InBin
 }
 
 #if WITH_EDITOR
-FString UPCGStaticMeshSpawnerKernel::GetCookedSource(const UPCGComputeGraph* InComputeGraph) const
+FString UPCGStaticMeshSpawnerKernel::GetCookedSource(FPCGGPUCompilationContext& InOutContext) const
 {
 	FString TemplateFile;
 	ensure(LoadShaderSourceFile(TEXT("/Plugin/PCG/Private/Elements/PCGStaticMeshSpawner.usf"), EShaderPlatform::SP_PCD3D_SM5, &TemplateFile, nullptr));
diff --git a/Engine/Plugins/PCG/Source/PCG/Private/Graph/PCGGPUGraphCompilationContext.h b/Engine/Plugins/PCG/Source/PCG/Private/Graph/PCGGPUGraphCompilationContext.h
--- a/Engine/Plugins/PCG/Source/PCG/Private/Graph/PCGGPUGraphCompilationContext.h
+++ b/Engine/Plugins/PCG/Source/PCG/Private/Graph/PCGGPUGraphCompilationContext.h
@@ -8,6 +8,7 @@
 
 class FPCGGraphCompiler;
 class UPCGComputeGraph;
+struct FPCGKernelAttributeTable;
 
 /** Current context/state used during compilation of GPU graphs. */
 struct FPCGGPUCompilationContext
@@ -64,6 +65,9 @@ struct FPCGGPUCompilationContext
 	FPCGGraphCompiler& GetGraphCompiler() { return GraphCompiler; }
 	TArray<TObjectPtr<UPCGComputeGraph>>& GetCompiledComputeGraphs() { return CompiledComputeGraphs; }
 
+	const FPCGKernelAttributeTable* GetStaticAttributeTable() { return StaticAttributeTable; }
+	void SetStaticAttributeTable(const FPCGKernelAttributeTable* InStaticAttributeTable) { StaticAttributeTable = InStaticAttributeTable; }
+
 	/** Returns the index of the newly added compute graph. */
 	int32 AddCompiledComputeGraph(TObjectPtr<UPCGComputeGraph> NewComputeGraph) { return CompiledComputeGraphs.Add(NewComputeGraph); }
 
@@ -74,6 +78,9 @@ private:
 
 	FPCGGraphCompiler& GraphCompiler;
 	TArray<TObjectPtr<UPCGComputeGraph>> CompiledComputeGraphs;
+
+	/** Attribute table of the currently compiling compute graph. */
+	const FPCGKernelAttributeTable* StaticAttributeTable = nullptr;
 };
 
 #endif // WITH_EDITOR
diff --git a/Engine/Plugins/PCG/Source/PCG/Private/Graph/PCGGraphCompilerGPU.cpp b/Engine/Plugins/PCG/Source/PCG/Private/Graph/PCGGraphCompilerGPU.cpp
--- a/Engine/Plugins/PCG/Source/PCG/Private/Graph/PCGGraphCompilerGPU.cpp
+++ b/Engine/Plugins/PCG/Source/PCG/Private/Graph/PCGGraphCompilerGPU.cpp
@@ -757,6 +757,7 @@ void FPCGGraphCompilerGPU::BuildComputeGraphTask(
 		UObject* ComputeGraphOuter = InOutContext.GetGraphCompiler().IsCooking() ? InGraph : Cast<UObject>(GetTransientPackage());
 		const FName GraphName = MakeUniqueObjectName(InGraph, UPCGComputeGraph::StaticClass(), InGraph->GetFName());
 		ComputeGraph = InOutContext.NewObject_AnyThread<UPCGComputeGraph>(ComputeGraphOuter, GraphName);
+		InOutContext.SetStaticAttributeTable(&ComputeGraph->GetStaticAttributeTable());
 
 		// Not incredibly useful for us - DG adds GetComponentSource()->GetComponentClass() object which allows it to bind at execution time by class.
 		// But execution code requires it currently.
@@ -779,6 +780,9 @@ void FPCGGraphCompilerGPU::BuildComputeGraphTask(
 			// Re-outer kernels to compute graph.
 			InTaskIdToKernel[TaskId]->Rename(nullptr, ComputeGraph);
 		}
+
+		// Clear out the current static attribute table.
+		InOutContext.SetStaticAttributeTable(nullptr);
 	}
 
 	const int32 ComputeGraphIndex = InOutContext.AddCompiledComputeGraph(ComputeGraph);
@@ -1432,8 +1436,7 @@ void FPCGGraphCompilerGPU::CompileComputeGraph(
 			KernelSource->EntryPoint = KernelWithBindings.Kernel->GetEntryPoint();
 			KernelSource->GroupSize = FIntVector(PCGComputeConstants::THREAD_GROUP_SIZE, 1, 1);
 
-			// TODO: Instead of passing a half-baked compute graph into this virtual call, we should pass in a compilation context per compute graph.
-			FString Source = KernelWithBindings.Kernel->GetCookedSource(InOutComputeGraph);
+			FString Source = KernelWithBindings.Kernel->GetCookedSource(InOutContext);
 			KernelSource->SetSource(MoveTemp(Source));
 			KernelWithBindings.Kernel->CreateAdditionalSources(InOutContext, KernelWithBindings.Kernel, KernelSource->AdditionalSources);
 
diff --git a/Engine/Plugins/PCG/Source/PCG/Public/Compute/PCGComputeSource.h b/Engine/Plugins/PCG/Source/PCG/Public/Compute/PCGComputeSource.h
new file mode 100644
--- /dev/null
+++ b/Engine/Plugins/PCG/Source/PCG/Public/Compute/PCGComputeSource.h
@@ -0,0 +1,40 @@
+// Copyright Epic Games, Inc. All Rights Reserved.
+
+#pragma once
+
+#include "Compute/IPCGNodeSourceTextProvider.h"
+
+#include "ComputeFramework/ComputeSource.h"
+
+#include "PCGComputeSource.generated.h"
+
+UCLASS()
+class PCG_API UPCGComputeSource : public UComputeSource, public IPCGNodeSourceTextProvider
+{
+	GENERATED_BODY()
+
+public:
+	//~ Begin UComputeSource Interface.
+	FString GetSource() const override;
+	FString GetVirtualPath() const override;
+	//~ End UComputeSource Interface.
+
+#if WITH_EDITOR
+	//~Begin IPCGNodeSourceTextProvider interface
+	FString GetShaderText() const override { return GetSource(); }
+	FString GetDeclarationsText() const override { return {}; }
+	FString GetShaderFunctionsText() const override { return {}; }
+	void SetShaderFunctionsText(const FString& NewFunctionsText) override {}
+	void SetShaderText(const FString& NewText) override { SetSource(NewText); }
+	bool IsShaderTextReadOnly() const override { return false; }
+	//~End IPCGNodeSourceTextProvider interface
+
+	void SetSource(const FString& InSource);
+#endif
+
+#if WITH_EDITORONLY_DATA
+protected:
+	UPROPERTY(EditAnywhere, Category = "Source", meta = (DisplayAfter = "AdditionalSources"))
+	FString Source;
+#endif
+};
diff --git a/Engine/Plugins/PCG/Source/PCG/Public/Elements/Metadata/PCGMetadataPartitionKernel.h b/Engine/Plugins/PCG/Source/PCG/Public/Elements/Metadata/PCGMetadataPartitionKernel.h
--- a/Engine/Plugins/PCG/Source/PCG/Public/Elements/Metadata/PCGMetadataPartitionKernel.h
+++ b/Engine/Plugins/PCG/Source/PCG/Public/Elements/Metadata/PCGMetadataPartitionKernel.h
@@ -19,7 +19,7 @@ public:
 	virtual bool ComputeOutputBindingDataDesc(const UPCGComputeGraph* InGraph, FName InOutputPinLabel, UPCGDataBinding* InBinding, FPCGDataCollectionDesc& OutDataDesc) const override;
 	virtual int ComputeThreadCount(const UPCGDataBinding* Binding) const override;
 #if WITH_EDITOR
-	virtual FString GetCookedSource(const UPCGComputeGraph* InComputeGraph) const override;
+	virtual FString GetCookedSource(FPCGGPUCompilationContext& InOutContext) const override;
 	virtual FString GetEntryPoint() const override { return TEXT("PCGMetadataPartitionCS"); }
 	void CreateAdditionalInputDataInterfaces(FPCGGPUCompilationContext& InOutContext, UObject* InObjectOuter, TArray<TObjectPtr<UComputeDataInterface>>& OutDataInterfaces) const override;
 	// Split graph to update element counts on CPU before proceeding. Compacts data for downstream kernels.
diff --git a/Engine/Plugins/PCG/Source/PCG/Public/Elements/PCGCopyPointsAnalysisKernel.h b/Engine/Plugins/PCG/Source/PCG/Public/Elements/PCGCopyPointsAnalysisKernel.h
--- a/Engine/Plugins/PCG/Source/PCG/Public/Elements/PCGCopyPointsAnalysisKernel.h
+++ b/Engine/Plugins/PCG/Source/PCG/Public/Elements/PCGCopyPointsAnalysisKernel.h
@@ -19,7 +19,7 @@ public:
 	virtual bool ComputeOutputBindingDataDesc(const UPCGComputeGraph* InGraph, FName InOutputPinLabel, UPCGDataBinding* InBinding, FPCGDataCollectionDesc& OutDataDesc) const override;
 	virtual int ComputeThreadCount(const UPCGDataBinding* Binding) const override;
 #if WITH_EDITOR
-	virtual FString GetCookedSource(const UPCGComputeGraph* InComputeGraph) const override;
+	virtual FString GetCookedSource(FPCGGPUCompilationContext& InOutContext) const override;
 	virtual FString GetEntryPoint() const override { return TEXT("PCGCopyPointsAnalysisCS"); }
 	void CreateAdditionalInputDataInterfaces(FPCGGPUCompilationContext& InOutContext, UObject* InObjectOuter, TArray<TObjectPtr<UComputeDataInterface>>& OutDataInterfaces) const override;
 	// Split graph to read back analysis results.
diff --git a/Engine/Plugins/PCG/Source/PCG/Public/Elements/PCGCopyPointsKernel.h b/Engine/Plugins/PCG/Source/PCG/Public/Elements/PCGCopyPointsKernel.h
--- a/Engine/Plugins/PCG/Source/PCG/Public/Elements/PCGCopyPointsKernel.h
+++ b/Engine/Plugins/PCG/Source/PCG/Public/Elements/PCGCopyPointsKernel.h
@@ -21,7 +21,7 @@ public:
 	virtual bool ComputeOutputBindingDataDesc(const UPCGComputeGraph* InGraph, FName InOutputPinLabel, UPCGDataBinding* InBinding, FPCGDataCollectionDesc& OutDataDesc) const override;
 	virtual int ComputeThreadCount(const UPCGDataBinding* Binding) const override;
 #if WITH_EDITOR
-	virtual FString GetCookedSource(const UPCGComputeGraph* InComputeGraph) const override;
+	virtual FString GetCookedSource(FPCGGPUCompilationContext& InOutContext) const override;
 	virtual FString GetEntryPoint() const override { return TEXT("PCGCopyPointsCS"); }
 	void CreateAdditionalInputDataInterfaces(FPCGGPUCompilationContext& InOutContext, UObject* InObjectOuter, TArray<TObjectPtr<UComputeDataInterface>>& OutDataInterfaces) const override;
 #endif
diff --git a/Engine/Plugins/PCG/Source/PCG/Public/Elements/PCGSkinnedMeshSpawnerKernel.h b/Engine/Plugins/PCG/Source/PCG/Public/Elements/PCGSkinnedMeshSpawnerKernel.h
--- a/Engine/Plugins/PCG/Source/PCG/Public/Elements/PCGSkinnedMeshSpawnerKernel.h
+++ b/Engine/Plugins/PCG/Source/PCG/Public/Elements/PCGSkinnedMeshSpawnerKernel.h
@@ -23,7 +23,7 @@ public:
 	virtual bool ComputeOutputBindingDataDesc(const UPCGComputeGraph* InGraph, FName InOutputPinLabel, UPCGDataBinding* InBinding, FPCGDataCollectionDesc& OutDataDesc) const override;
 	virtual int ComputeThreadCount(const UPCGDataBinding* Binding) const override;
 #if WITH_EDITOR
-	virtual FString GetCookedSource(const UPCGComputeGraph* InComputeGraph) const override;
+	virtual FString GetCookedSource(FPCGGPUCompilationContext& InOutContext) const override;
 	virtual FString GetEntryPoint() const override { return TEXT("PCGStaticMeshSpawnerCS"); }
 	void CreateAdditionalInputDataInterfaces(FPCGGPUCompilationContext& InOutContext, UObject* InObjectOuter, TArray<TObjectPtr<UComputeDataInterface>>& OutDataInterfaces) const override;
 	void CreateAdditionalOutputDataInterfaces(FPCGGPUCompilationContext& InOutContext, UObject* InObjectOuter, TArray<TObjectPtr<UComputeDataInterface>>& OutDataInterfaces) const override;
diff --git a/Engine/Plugins/PCG/Source/PCG/Public/Elements/PCGStaticMeshSpawnerKernel.h b/Engine/Plugins/PCG/Source/PCG/Public/Elements/PCGStaticMeshSpawnerKernel.h
--- a/Engine/Plugins/PCG/Source/PCG/Public/Elements/PCGStaticMeshSpawnerKernel.h
+++ b/Engine/Plugins/PCG/Source/PCG/Public/Elements/PCGStaticMeshSpawnerKernel.h
@@ -23,7 +23,7 @@ public:
 	virtual bool ComputeOutputBindingDataDesc(const UPCGComputeGraph* InGraph, FName InOutputPinLabel, UPCGDataBinding* InBinding, FPCGDataCollectionDesc& OutDataDesc) const override;
 	virtual int ComputeThreadCount(const UPCGDataBinding* Binding) const override;
 #if WITH_EDITOR
-	virtual FString GetCookedSource(const UPCGComputeGraph* InComputeGraph) const override;
+	virtual FString GetCookedSource(FPCGGPUCompilationContext& InOutContext) const override;
 	virtual FString GetEntryPoint() const override { return TEXT("PCGStaticMeshSpawnerCS"); }
 	void CreateAdditionalInputDataInterfaces(FPCGGPUCompilationContext& InOutContext, UObject* InObjectOuter, TArray<TObjectPtr<UComputeDataInterface>>& OutDataInterfaces) const override;
 	void CreateAdditionalOutputDataInterfaces(FPCGGPUCompilationContext& InOutContext, UObject* InObjectOuter, TArray<TObjectPtr<UComputeDataInterface>>& OutDataInterfaces) const override;
diff --git a/Engine/Plugins/PCG/Source/PCGEditor/Private/AssetDefinitions/AssetDefinition_PCGComputeSource.cpp b/Engine/Plugins/PCG/Source/PCGEditor/Private/AssetDefinitions/AssetDefinition_PCGComputeSource.cpp
new file mode 100644
--- /dev/null
+++ b/Engine/Plugins/PCG/Source/PCGEditor/Private/AssetDefinitions/AssetDefinition_PCGComputeSource.cpp
@@ -0,0 +1,32 @@
+// Copyright Epic Games, Inc. All Rights Reserved.
+
+#include "AssetDefinition_PCGComputeSource.h"
+
+#include "Compute/PCGComputeSource.h"
+
+#include "PCGEditorCommon.h"
+
+#define LOCTEXT_NAMESPACE "AssetDefinition_PCGComputeSource"
+
+FText UAssetDefinition_PCGComputeSource::GetAssetDisplayName() const
+{
+	return LOCTEXT("DisplayName", "PCG Compute Source");
+}
+
+FLinearColor UAssetDefinition_PCGComputeSource::GetAssetColor() const
+{
+	return FColor::Orange;
+}
+
+TSoftClassPtr<UObject> UAssetDefinition_PCGComputeSource::GetAssetClass() const
+{
+	return UPCGComputeSource::StaticClass();
+}
+
+TConstArrayView<FAssetCategoryPath> UAssetDefinition_PCGComputeSource::GetAssetCategories() const
+{
+	static const FAssetCategoryPath Categories[] = { FPCGEditorCommon::PCGAdvancedAssetCategoryPath };
+	return Categories;
+}
+
+#undef LOCTEXT_NAMESPACE
diff --git a/Engine/Plugins/PCG/Source/PCGEditor/Private/AssetDefinitions/AssetDefinition_PCGComputeSource.h b/Engine/Plugins/PCG/Source/PCGEditor/Private/AssetDefinitions/AssetDefinition_PCGComputeSource.h
new file mode 100644
--- /dev/null
+++ b/Engine/Plugins/PCG/Source/PCGEditor/Private/AssetDefinitions/AssetDefinition_PCGComputeSource.h
@@ -0,0 +1,19 @@
+// Copyright Epic Games, Inc. All Rights Reserved.
+
+#pragma once
+
+#include "AssetDefinitionDefault.h"
+
+#include "AssetDefinition_PCGComputeSource.generated.h"
+
+UCLASS()
+class UAssetDefinition_PCGComputeSource : public UAssetDefinitionDefault
+{
+	GENERATED_BODY()
+
+public:
+	virtual FText GetAssetDisplayName() const override;
+	virtual FLinearColor GetAssetColor() const override;
+	virtual TSoftClassPtr<UObject> GetAssetClass() const override;
+	virtual TConstArrayView<FAssetCategoryPath> GetAssetCategories() const override;
+};
diff --git a/Engine/Plugins/PCG/Source/PCGEditor/Private/Details/PCGComputeSourceDetails.cpp b/Engine/Plugins/PCG/Source/PCGEditor/Private/Details/PCGComputeSourceDetails.cpp
new file mode 100644
--- /dev/null
+++ b/Engine/Plugins/PCG/Source/PCGEditor/Private/Details/PCGComputeSourceDetails.cpp
@@ -0,0 +1,130 @@
+// Copyright Epic Games, Inc. All Rights Reserved.
+
+#include "Details/PCGComputeSourceDetails.h"
+
+#include "Compute/PCGComputeSource.h"
+
+#include "PCGEditor.h"
+#include "PCGEditorGraph.h"
+#include "PCGHLSLSyntaxHighlighter.h"
+#include "Widgets/SPCGEditorNodeSource.h"
+#include "Widgets/SPCGNodeSourceTextBox.h"
+
+#include "DetailWidgetRow.h"
+#include "DetailLayoutBuilder.h"
+#include "PropertyHandle.h"
+
+#define LOCTEXT_NAMESPACE "PCGComputeSourceDetails"
+
+FPCGComputeSourceDetails::FPCGComputeSourceDetails()
+	: SyntaxHighlighter(FPCGHLSLSyntaxHighlighter::Create())
+{}
+
+TSharedRef<IDetailCustomization> FPCGComputeSourceDetails::MakeInstance()
+{
+	return MakeShared<FPCGComputeSourceDetails>();
+}
+
+void FPCGComputeSourceDetails::CustomizeDetails(IDetailLayoutBuilder& DetailBuilder)
+{
+	TArray<TWeakObjectPtr<UObject>> ObjectsBeingCustomized;
+	DetailBuilder.GetObjectsBeingCustomized(ObjectsBeingCustomized);
+
+	if (ObjectsBeingCustomized.Num() > 1)
+	{
+		return;
+	}
+
+	ComputeSourceWeakPtr = Cast<UPCGComputeSource>(ObjectsBeingCustomized[0].Get());
+
+	if (!ComputeSourceWeakPtr.Get())
+	{
+		return;
+	}
+
+	if (TSharedPtr<IDetailsView> DetailsView = DetailBuilder.GetDetailsViewSharedPtr())
+	{
+		DetailsView->HideFilterArea(true);
+	}
+
+	SExpandableArea::FArguments ExpandableAreaArgs;
+	ExpandableAreaArgs.AreaTitleFont(FAppStyle::GetFontStyle("DetailsView.CategoryFontStyle"));
+	ExpandableAreaArgs.AreaTitle(LOCTEXT("PCGNodeSource_ShaderText_Title", "Shader Source"));
+
+	TSharedRef<IPropertyHandle> SourcePropertyHandle = DetailBuilder.GetProperty(TEXT("Source"));
+	DetailBuilder.EditDefaultProperty(SourcePropertyHandle)->CustomWidget()
+	.OverrideResetToDefault(FResetToDefaultOverride::Hide())
+	[
+		SNew(SVerticalBox)
+		+SVerticalBox::Slot()
+		.AutoHeight()
+		[
+			// @todo_pcg: Share functionality with the NodeSourceEditor
+			//ConstructNonExpandableHeaderWidget(ExpandableAreaArgs.AreaTitle(LOCTEXT("PCGNodeSource_ShaderText_Title", "Shader Source")))
+			SNew(SBorder)
+				.BorderImage(FStyleDefaults::GetNoBrush())
+				.BorderBackgroundColor(FLinearColor::Transparent)
+				.Padding(0.0f)
+				[
+					SNew(SButton)
+					.ButtonStyle(FCoreStyle::Get(), "NoBorder")
+					.ContentPadding(ExpandableAreaArgs._HeaderPadding)
+					.ForegroundColor(FSlateColor::UseForeground())
+					[
+						SNew(SHorizontalBox)
+						+SHorizontalBox::Slot()
+						.AutoWidth()
+						.Padding(ExpandableAreaArgs._AreaTitlePadding)
+						.VAlign(VAlign_Center)
+						[
+							SNew(SSpacer)
+							.Size(ExpandableAreaArgs._Style->CollapsedImage.GetImageSize())
+						]
+						+SHorizontalBox::Slot()
+						.FillWidth(1.0f)
+						.VAlign(VAlign_Center)
+						[
+							SNew(STextBlock)
+							.Text(ExpandableAreaArgs._AreaTitle)
+							.Font(ExpandableAreaArgs._AreaTitleFont)
+						]
+					]
+				]
+		]
+		+SVerticalBox::Slot()
+		.FillHeight(1.0f)
+		[
+			SNew(SBox)
+			.HeightOverride(500)
+			[
+				SAssignNew(SourceTextBox, SPCGNodeSourceTextBox)
+				.Text(this, &FPCGComputeSourceDetails::GetSourceText)
+				.OnTextChanged(this, &FPCGComputeSourceDetails::OnSourceTextChanged)
+				.IsReadOnly(this, &FPCGComputeSourceDetails::IsSourceTextReadOnly)
+				.Marshaller(SyntaxHighlighter)
+			]
+		]
+	];
+}
+
+FText FPCGComputeSourceDetails::GetSourceText() const
+{
+	UPCGComputeSource* ComputeSource = ComputeSourceWeakPtr.Get();
+	return ComputeSource ? FText::FromString(ComputeSource->GetShaderText()) : FText::GetEmpty();
+}
+
+void FPCGComputeSourceDetails::OnSourceTextChanged(const FText& InText)
+{
+	if (UPCGComputeSource* ComputeSource = ComputeSourceWeakPtr.Get())
+	{
+		ComputeSource->SetShaderText(InText.ToString());
+	}
+}
+
+bool FPCGComputeSourceDetails::IsSourceTextReadOnly() const
+{
+	UPCGComputeSource* ComputeSource = ComputeSourceWeakPtr.Get();
+	return ComputeSource ? ComputeSource->IsShaderTextReadOnly() : false;
+}
+
+#undef LOCTEXT_NAMESPACE
diff --git a/Engine/Plugins/PCG/Source/PCGEditor/Private/Details/PCGComputeSourceDetails.h b/Engine/Plugins/PCG/Source/PCGEditor/Private/Details/PCGComputeSourceDetails.h
new file mode 100644
--- /dev/null
+++ b/Engine/Plugins/PCG/Source/PCGEditor/Private/Details/PCGComputeSourceDetails.h
@@ -0,0 +1,32 @@
+// Copyright Epic Games, Inc. All Rights Reserved.
+
+#pragma once
+
+#include "CoreMinimal.h"
+#include "IDetailCustomization.h"
+#include "UObject/WeakObjectPtr.h"
+
+class FPCGHLSLSyntaxHighlighter;
+class SPCGNodeSourceTextBox;
+class UPCGComputeSource;
+
+/**
+ * Implements a details view customization for PCG Compute Sources.
+ */
+class FPCGComputeSourceDetails : public IDetailCustomization
+{
+public:
+	FPCGComputeSourceDetails();
+
+	static TSharedRef<IDetailCustomization> MakeInstance();
+	virtual void CustomizeDetails(IDetailLayoutBuilder& DetailBuilder) override;
+
+protected:
+	FText GetSourceText() const;
+	void OnSourceTextChanged(const FText& InText);
+	bool IsSourceTextReadOnly() const;
+
+	TWeakObjectPtr<UPCGComputeSource> ComputeSourceWeakPtr;
+	TSharedRef<FPCGHLSLSyntaxHighlighter> SyntaxHighlighter;
+	TSharedPtr<SPCGNodeSourceTextBox> SourceTextBox;
+};
diff --git a/Engine/Plugins/PCG/Source/PCGEditor/Private/Factories/PCGComputeSourceFactory.cpp b/Engine/Plugins/PCG/Source/PCGEditor/Private/Factories/PCGComputeSourceFactory.cpp
new file mode 100644
--- /dev/null
+++ b/Engine/Plugins/PCG/Source/PCGEditor/Private/Factories/PCGComputeSourceFactory.cpp
@@ -0,0 +1,17 @@
+// Copyright Epic Games, Inc. All Rights Reserved.
+
+#include "PCGComputeSourceFactory.h"
+
+#include "Compute/PCGComputeSource.h"
+
+UPCGComputeSourceFactory::UPCGComputeSourceFactory()
+{
+	SupportedClass = UPCGComputeSource::StaticClass();
+	bCreateNew = true;
+	bEditAfterNew = true;
+}
+
+UObject* UPCGComputeSourceFactory::FactoryCreateNew(UClass* InClass, UObject* InParent, FName InName, EObjectFlags Flags, UObject* Context, FFeedbackContext* Warn)
+{
+	return NewObject<UPCGComputeSource>(InParent, InClass, InName, Flags);
+}
diff --git a/Engine/Plugins/PCG/Source/PCGEditor/Private/Factories/PCGComputeSourceFactory.h b/Engine/Plugins/PCG/Source/PCGEditor/Private/Factories/PCGComputeSourceFactory.h
new file mode 100644
--- /dev/null
+++ b/Engine/Plugins/PCG/Source/PCGEditor/Private/Factories/PCGComputeSourceFactory.h
@@ -0,0 +1,21 @@
+// Copyright Epic Games, Inc. All Rights Reserved.
+
+#pragma once
+
+#include "Factories/Factory.h"
+
+#include "PCGComputeSourceFactory.generated.h"
+
+UCLASS()
+class UPCGComputeSourceFactory : public UFactory
+{
+	GENERATED_BODY()
+
+public:
+	UPCGComputeSourceFactory();
+
+	//~Begin UFactory interface
+	virtual UObject* FactoryCreateNew(UClass* InClass, UObject* InParent, FName InName, EObjectFlags Flags, UObject* Context, FFeedbackContext* Warn) override;
+	virtual bool ShouldShowInNewMenu() const override { return true; }
+	//~End UFactory interface
+};
diff --git a/Engine/Plugins/PCG/Source/PCGEditor/Private/PCGEditor.cpp b/Engine/Plugins/PCG/Source/PCGEditor/Private/PCGEditor.cpp
--- a/Engine/Plugins/PCG/Source/PCGEditor/Private/PCGEditor.cpp
+++ b/Engine/Plugins/PCG/Source/PCGEditor/Private/PCGEditor.cpp
@@ -2931,7 +2931,7 @@ TSharedRef<SPCGEditorGraphLogView> FPCGEditor::CreateLogWidget()
 
 TSharedRef<SPCGEditorNodeSource> FPCGEditor::CreateNodeSourceWidget()
 {
-	return SNew(SPCGEditorNodeSource, SharedThis(this));
+	return SNew(SPCGEditorNodeSource);
 }
 
 TSharedRef<SPCGEditorGraphUserParametersView> FPCGEditor::CreateGraphParamsWidget()
diff --git a/Engine/Plugins/PCG/Source/PCGEditor/Private/PCGEditorModule.cpp b/Engine/Plugins/PCG/Source/PCGEditor/Private/PCGEditorModule.cpp
--- a/Engine/Plugins/PCG/Source/PCGEditor/Private/PCGEditorModule.cpp
+++ b/Engine/Plugins/PCG/Source/PCGEditor/Private/PCGEditorModule.cpp
@@ -37,6 +37,7 @@
 #include "Details/PCGAttributePropertySelectorDetails.h"
 #include "Details/PCGBlueprintSettingsDetails.h"
 #include "Details/PCGComponentDetails.h"
+#include "Details/PCGComputeSourceDetails.h"
 #include "Details/PCGCustomHLSLSettingsDetails.h"
 #include "Details/PCGEditableUserParameterDetails.h"
 #include "Details/PCGGraphDetails.h"
@@ -284,6 +285,7 @@ void FPCGEditorModule::RegisterDetailsCustomizations()
 	PropertyEditor.RegisterCustomClassLayout("PCGVolume", FOnGetDetailCustomizationInstance::CreateStatic(&FPCGVolumeDetails::MakeInstance));
 	PropertyEditor.RegisterCustomClassLayout("PCGUserParameterGetSettings", FOnGetDetailCustomizationInstance::CreateStatic(&FPCGEditableUserParameterDetails::MakeInstance));
 	PropertyEditor.RegisterCustomClassLayout("PCGCustomHLSLSettings", FOnGetDetailCustomizationInstance::CreateStatic(&FPCGCustomHLSLSettingsDetails::MakeInstance));
+	PropertyEditor.RegisterCustomClassLayout("PCGComputeSource", FOnGetDetailCustomizationInstance::CreateStatic(&FPCGComputeSourceDetails::MakeInstance));
 
 	PropertyEditor.RegisterCustomPropertyTypeLayout("PCGAttributePropertySelector", FOnGetPropertyTypeCustomizationInstance::CreateStatic(&FPCGAttributePropertySelectorDetails::MakeInstance));
 	PropertyEditor.RegisterCustomPropertyTypeLayout("PCGAttributePropertyInputSelector", FOnGetPropertyTypeCustomizationInstance::CreateStatic(&FPCGAttributePropertySelectorDetails::MakeInstance));
diff --git a/Engine/Plugins/PCG/Source/PCGEditor/Private/Widgets/SPCGEditorNodeSource.cpp b/Engine/Plugins/PCG/Source/PCGEditor/Private/Widgets/SPCGEditorNodeSource.cpp
--- a/Engine/Plugins/PCG/Source/PCGEditor/Private/Widgets/SPCGEditorNodeSource.cpp
+++ b/Engine/Plugins/PCG/Source/PCGEditor/Private/Widgets/SPCGEditorNodeSource.cpp
@@ -21,7 +21,7 @@ SPCGEditorNodeSource::SPCGEditorNodeSource()
 {
 }
 
-void SPCGEditorNodeSource::Construct(const FArguments& InArgs, TSharedPtr<FPCGEditor> InPCGEditor)
+void SPCGEditorNodeSource::Construct(const FArguments& InArgs)
 {
 	SExpandableArea::FArguments ExpandableAreaArgs;
 	ExpandableAreaArgs.AreaTitleFont(FAppStyle::GetFontStyle("DetailsView.CategoryFontStyle"));
diff --git a/Engine/Plugins/PCG/Source/PCGEditor/Private/Widgets/SPCGEditorNodeSource.h b/Engine/Plugins/PCG/Source/PCGEditor/Private/Widgets/SPCGEditorNodeSource.h
--- a/Engine/Plugins/PCG/Source/PCGEditor/Private/Widgets/SPCGEditorNodeSource.h
+++ b/Engine/Plugins/PCG/Source/PCGEditor/Private/Widgets/SPCGEditorNodeSource.h
@@ -24,7 +24,7 @@ public:
 
 	SPCGEditorNodeSource();
 
-	void Construct(const FArguments& InArgs, TSharedPtr<FPCGEditor> InPCGEditor);
+	void Construct(const FArguments& InArgs);
 
 	void SetTextProviderObject(UObject* InProviderObject);
 	const UObject* GetTextProviderObject() const { return ShaderTextProviderObject.IsValid() ? ShaderTextProviderObject.Get() : nullptr; }
