Changed the default behavior of the surface sampler to make it driven primarily by the points per squared meters.

diff --git a/Engine/Plugins/PCG/Source/PCG/Private/Elements/PCGSurfaceSampler.cpp b/Engine/Plugins/PCG/Source/PCG/Private/Elements/PCGSurfaceSampler.cpp
--- a/Engine/Plugins/PCG/Source/PCG/Private/Elements/PCGSurfaceSampler.cpp
+++ b/Engine/Plugins/PCG/Source/PCG/Private/Elements/PCGSurfaceSampler.cpp
@@ -17,6 +17,7 @@
 
 #include "HAL/UnrealMemory.h"
 #include "Math/RandomStream.h"
+#include "UObject/FortniteMainBranchObjectVersion.h"
 
 #include UE_INLINE_GENERATED_CPP_BY_NAME(PCGSurfaceSampler)
 
@@ -28,6 +29,7 @@
 	{
 		check(Settings);
 		// Compute used values
+		bUseLegacyGridCreationMethod = Settings->bUseLegacyGridCreationMethod;
 		PointsPerSquaredMeter = Settings->PointsPerSquaredMeter;
 		PointExtents = Settings->PointExtents;
 		Looseness = Settings->Looseness;
@@ -53,10 +55,57 @@
 
 		Seed = Context ? Context->GetSeed() : PCGValueConstants::DefaultSeed;
 
-		// Conceptually, we will break down the surface bounds in a N x M grid
-		InterstitialDistance = Params.PointExtents * 2;
-		InnerCellSize = InterstitialDistance * Params.Looseness;
-		CellSize = InterstitialDistance + InnerCellSize;
+		if (Params.PointExtents.X <= 0 || Params.PointExtents.Y <= 0)
+		{
+			// PointExtents and Looseness are user overridable, if any of those values are 0 or negative, it's invalid, so we early out.
+			PCGLog::LogWarningOnGraph(LOCTEXT("InvalidParametersExtents", "Skipped - Extents are negative or zero."), Context);
+			return false;
+		}
+
+		if (!Params.bUseLegacyGridCreationMethod)
+		{
+			// Points per squared meter -> inverse is meters squared by point.
+			const FVector::FReal SquaredUnitsPerPoint = Params.PointsPerSquaredMeter > 0 ? (100.0 * 100.0) / Params.PointsPerSquaredMeter : UE_DOUBLE_BIG_NUMBER;
+
+			// Compute approximate cell size based on the point extents relative size.
+			// Knowing that X * Y = SQM
+			// if X = aY, then aY * Y = SQM -> aY^2 = SQM -> Y = sqrt(SQM / a)
+			const FVector::FReal MinCellSize = 2.0 * FMath::Min(Params.PointExtents.X, Params.PointExtents.Y);
+			const FVector::FReal MaxCellSize = 2.0 * FMath::Max(Params.PointExtents.X, Params.PointExtents.Y);
+
+			const FVector::FReal BaseCellSize = FMath::Sqrt(SquaredUnitsPerPoint / (MaxCellSize / MinCellSize));
+
+			CellSize = FVector(
+				BaseCellSize * 2.0 * (Params.PointExtents.X / MinCellSize),
+				BaseCellSize * 2.0 * (Params.PointExtents.Y / MinCellSize),
+				2.0 * Params.PointExtents.Z);
+
+			InterstitialDistance = Params.PointExtents * 2;
+			
+			FVector CellRemainder = CellSize - InterstitialDistance;
+
+			// Enforce that the grid is at least the provided extents.
+			if (CellRemainder.X < 0 || CellRemainder.Y < 0)
+			{
+				CellSize.X = FMath::Max(CellSize.X, InterstitialDistance.X);
+				CellSize.Y = FMath::Max(CellSize.Y, InterstitialDistance.Y);
+				CellRemainder = CellSize - InterstitialDistance;
+				check(CellRemainder.X >= -UE_DOUBLE_SMALL_NUMBER && CellRemainder.Y >= -UE_DOUBLE_SMALL_NUMBER);
+			}
+
+			const FVector::FReal ClampedLooseness = FMath::Clamp(Params.Looseness, 0.0, 1.0);
+			InnerCellSize = CellRemainder * ClampedLooseness;
+			InnerCellOffset = CellRemainder * 0.5 * (1.0 - ClampedLooseness);
+		}
+		else
+		{
+			// Legacy grid creation - Conceptually, we will break down the surface bounds in a N x M grid, where the cells are extents * (1 + steepness).
+			InterstitialDistance = Params.PointExtents * 2;
+			InnerCellSize = InterstitialDistance * Params.Looseness;
+			InnerCellOffset = FVector::ZeroVector;
+			CellSize = InterstitialDistance + InnerCellSize;
+		}
+
 		if (CellSize.X <= 0 || CellSize.Y <= 0)
 		{
 			// PointExtents and Looseness are user overridable, if any of those values are 0 or negative, it's invalid, so we early out.
@@ -64,27 +113,5 @@
 			return false;
 		}
 
-		// Local transformation is only needed if we're rotating.
-		bNeedsLocalTransformation = !InSurfaceTransform.Rotator().IsNearlyZero();
-		if (bNeedsLocalTransformation)
-		{
-			// Build the semi-local transform matrix for transforming points to the pre-projection plane
-			const FTransform TranslationTransform(InSurfaceTransform.GetTranslation());
-			PreProjectionTransform = TranslationTransform.Inverse().ToMatrixNoScale();
-			// Find the rotation between the world normal and the surface and apply to the matrix
-			PreProjectionTransform *= FQuat::FindBetweenNormals(FVector::UpVector, InSurfaceTransform.GetRotation().GetUpVector().GetSafeNormal()).ToMatrix();
-			PreProjectionTransform *= TranslationTransform.ToMatrixNoScale();
-		}
-
-		// Drop points slightly by an epsilon otherwise point can be culled. If the sampler has a volume connected as the Bounding Shape,
-		// the volume will call through to PCGHelpers::IsInsideBounds() which is a one sided test and points at the top of the volume
-		// will fail it. TODO perhaps the one-sided check can be isolated to component-bounds
-		constexpr FVector::FReal DefaultHeightModifier = 1.0 - UE_DOUBLE_SMALL_NUMBER;
-		// Try to use a multiplier instead of a simply offset to combat loss of precision in floats. However if MaxZ is very small,
-		// then multiplier will not work, so just use an offset.
-		PreProjectionDisplacement = (FMath::Abs(InEffectiveGridBounds.Max.Z) > UE_DOUBLE_SMALL_NUMBER) ? InEffectiveGridBounds.Max.Z * DefaultHeightModifier : -UE_DOUBLE_SMALL_NUMBER;
-		// Make sure we're still in bounds though!
-		PreProjectionDisplacement = FMath::Max(PreProjectionDisplacement, InEffectiveGridBounds.Min.Z);
-
 		// By using scaled indices in the world, we can easily make this process deterministic
 		CellMinX = FMath::CeilToInt((InEffectiveGridBounds.Min.X) / CellSize.X);
@@ -117,9 +144,19 @@
 
 		check(CellCount > 0);
 
-		constexpr FVector::FReal InvSquaredMeterUnits = 1.0 / (100.0 * 100.0);
-		const FVector::FReal TargetPointCount = (InEffectiveGridBounds.Max.X - InEffectiveGridBounds.Min.X) * (InEffectiveGridBounds.Max.Y - InEffectiveGridBounds.Min.Y) * Params.PointsPerSquaredMeter * InvSquaredMeterUnits;
-		Ratio = static_cast<float>(FMath::Clamp(TargetPointCount / (FVector::FReal)CellCount, 0.0, 1.0));
+		FVector::FReal TargetPointCount = 0;
 
+		if (!Params.bUseLegacyGridCreationMethod)
+		{
+			TargetPointCount = (InEffectiveGridBounds.Max.X - InEffectiveGridBounds.Min.X) * (InEffectiveGridBounds.Max.Y - InEffectiveGridBounds.Min.Y);
+			Ratio = 1.0;
+		}
+		else
+		{
+			constexpr FVector::FReal InvSquaredMeterUnits = 1.0 / (100.0 * 100.0);
+			TargetPointCount = (InEffectiveGridBounds.Max.X - InEffectiveGridBounds.Min.X) * (InEffectiveGridBounds.Max.Y - InEffectiveGridBounds.Min.Y) * Params.PointsPerSquaredMeter * InvSquaredMeterUnits;
+			Ratio = static_cast<float>(FMath::Clamp(TargetPointCount / (FVector::FReal)CellCount, 0.0, 1.0));
+		}
+
 		if (Ratio < UE_SMALL_NUMBER)
 		{
@@ -138,6 +175,28 @@
 			return false;
 		}
 
+		// Local transformation is only needed if we're rotating.
+		bNeedsLocalTransformation = !InSurfaceTransform.Rotator().IsNearlyZero();
+		if (bNeedsLocalTransformation)
+		{
+			// Build the semi-local transform matrix for transforming points to the pre-projection plane
+			const FTransform TranslationTransform(InSurfaceTransform.GetTranslation());
+			PreProjectionTransform = TranslationTransform.Inverse().ToMatrixNoScale();
+			// Find the rotation between the world normal and the surface and apply to the matrix
+			PreProjectionTransform *= FQuat::FindBetweenNormals(FVector::UpVector, InSurfaceTransform.GetRotation().GetUpVector().GetSafeNormal()).ToMatrix();
+			PreProjectionTransform *= TranslationTransform.ToMatrixNoScale();
+		}
+
+		// Drop points slightly by an epsilon otherwise point can be culled. If the sampler has a volume connected as the Bounding Shape,
+		// the volume will call through to PCGHelpers::IsInsideBounds() which is a one sided test and points at the top of the volume
+		// will fail it. TODO perhaps the one-sided check can be isolated to component-bounds
+		constexpr FVector::FReal DefaultHeightModifier = 1.0 - UE_DOUBLE_SMALL_NUMBER;
+		// Try to use a multiplier instead of a simply offset to combat loss of precision in floats. However if MaxZ is very small,
+		// then multiplier will not work, so just use an offset.
+		PreProjectionDisplacement = (FMath::Abs(InEffectiveGridBounds.Max.Z) > UE_DOUBLE_SMALL_NUMBER) ? InEffectiveGridBounds.Max.Z * DefaultHeightModifier : -UE_DOUBLE_SMALL_NUMBER;
+		// Make sure we're still in bounds though!
+		PreProjectionDisplacement = FMath::Max(PreProjectionDisplacement, InEffectiveGridBounds.Min.Z);
+
 		return true;
 	}
 
@@ -208,7 +267,7 @@
 			int32 NumWritten = 0;
 
 			for(int32 ReadIndex = StartReadIndex; ReadIndex < StartReadIndex+Count; ++ReadIndex)
-			{				
+			{
 				const int32 WriteIndex = StartWriteIndex + NumWritten;
 
 				const FSurfaceSamplerParams& SamplerParams = SamplerData.Params;
@@ -216,6 +275,7 @@
 
 				const FVector::FReal CurrentX = Indices.X * SamplerData.CellSize.X;
 				const FVector::FReal CurrentY = Indices.Y * SamplerData.CellSize.Y;
+				const FVector InnerCellOffset = SamplerData.InnerCellOffset;
 				const FVector InnerCellSize = SamplerData.InnerCellSize;
 
 				FRandomStream RandomSource(PCGHelpers::ComputeSeed(SamplerData.Seed, Indices.X, Indices.Y));
@@ -231,7 +291,10 @@
 				const float RandX = RandomSource.FRand();
 				const float RandY = RandomSource.FRand();
 
-				FVector TentativeLocation = FVector(CurrentX + RandX * InnerCellSize.X, CurrentY + RandY * InnerCellSize.Y, SamplerData.PreProjectionDisplacement);
+				FVector TentativeLocation = FVector(
+					CurrentX + InnerCellOffset.X + RandX * InnerCellSize.X, 
+					CurrentY + InnerCellOffset.Y + RandY * InnerCellSize.Y, 
+					SamplerData.PreProjectionDisplacement);
 
 				// If pre-projected points need a local transformation (ex. World Ray Hit Query) and not default to -Z
 				if (SamplerData.bNeedsLocalTransformation)
@@ -374,6 +437,22 @@
 	return PinProperties;
 }
 
+void UPCGSurfaceSamplerSettings::Serialize(FArchive& Ar)
+{
+	Ar.UsingCustomVersion(FFortniteMainBranchObjectVersion::GUID);
+
+	if (Ar.IsLoading() && Ar.CustomVer(FFortniteMainBranchObjectVersion::GUID) < FFortniteMainBranchObjectVersion::PCGChangedSurfaceSamplerDefaultGridCreationMode)
+	{
+		// Default value has changed for the point extents from FVector(100.0f) to FVector(50.0f)
+		PointExtents = FVector(100.0f);
+
+		// Prior to this version all settings were using what we now consider the "legacy" grid creation scheme
+		bUseLegacyGridCreationMethod = true;
+	}
+
+	Super::Serialize(Ar);
+}
+
 void UPCGSurfaceSamplerSettings::PostLoad()
 {
 	Super::PostLoad();
diff --git a/Engine/Plugins/PCG/Source/PCG/Public/Elements/PCGSurfaceSampler.h b/Engine/Plugins/PCG/Source/PCG/Public/Elements/PCGSurfaceSampler.h
--- a/Engine/Plugins/PCG/Source/PCG/Public/Elements/PCGSurfaceSampler.h
+++ b/Engine/Plugins/PCG/Source/PCG/Public/Elements/PCGSurfaceSampler.h
@@ -27,6 +27,7 @@ namespace PCGSurfaceSampler
 	{
 		void InitializeFromSettings(const UPCGSurfaceSamplerSettings* Settings);
 
+		bool bUseLegacyGridCreationMethod = false;
 		float PointsPerSquaredMeter = 1.0f;
 		FVector PointExtents = FVector::One() * 0.5f;
 		float Looseness = 0.0f;
@@ -52,6 +53,7 @@ namespace PCGSurfaceSampler
 		/** Computed values **/
 		FVector InterstitialDistance = FVector::ZeroVector;
 		FVector InnerCellSize = FVector::ZeroVector;
+		FVector InnerCellOffset = FVector::ZeroVector;
 		FVector CellSize = FVector::ZeroVector;
 
 		int32 CellMinX = 0;
@@ -105,6 +107,7 @@ public:
 	UPCGSurfaceSamplerSettings();
 
 	//~Begin UObject interface
+	virtual void Serialize(FArchive& Ar) override;
 	virtual void PostLoad() override;
 	//~End UObject interface
 
@@ -127,12 +130,15 @@ protected:
 	// ~End UPCGSettings interface
 
 public:
+	/** Controls the grid cell size, down to a minimum size defined by the Point Extents. */
 	UPROPERTY(BlueprintReadWrite, EditAnywhere, Category = Settings, meta=(ClampMin="0", PCG_Overridable))
 	float PointsPerSquaredMeter = 0.1f;
 
+	/** Extents of the points to create (e.g. half their total size) */
 	UPROPERTY(BlueprintReadWrite, EditAnywhere, Category = Settings, meta = (PCG_Overridable))
-	FVector PointExtents = FVector(100.0f);
+	FVector PointExtents = FVector(50.0f);
 
+	/** Controls how points are placed in their cell with 0 being in the center and 1 being anywhere inside the full cell size. */
 	UPROPERTY(BlueprintReadWrite, EditAnywhere, Category = Settings, meta=(ClampMin="0", PCG_Overridable))
 	float Looseness = 1.0f;
 
@@ -161,6 +167,10 @@ public:
 	UPROPERTY()
 	float PointRadius_DEPRECATED = 0.0f;
 #endif
+
+	/** Controls whether the prior grid creation mechanism (with cells being extents * (1 + steepness), then filtered by points by squared meters) is used. */
+	UPROPERTY(BlueprintReadWrite, EditAnywhere, Category = Settings, AdvancedDisplay)
+	bool bUseLegacyGridCreationMethod = false;
 };
 
 class FPCGSurfaceSamplerElement : public TPCGTimeSlicedElementBase<PCGSurfaceSampler::FSurfaceSamplerExecutionState, PCGSurfaceSampler::FSurfaceSamplerIterationState>
diff --git a/Engine/Source/Runtime/Core/Public/UObject/FortniteMainBranchObjectVersions.inl b/Engine/Source/Runtime/Core/Public/UObject/FortniteMainBranchObjectVersions.inl
--- a/Engine/Source/Runtime/Core/Public/UObject/FortniteMainBranchObjectVersions.inl
+++ b/Engine/Source/Runtime/Core/Public/UObject/FortniteMainBranchObjectVersions.inl
@@ -631,5 +631,8 @@ FORTNITE_MAIN_VERSION(MigrateLandscapeEditLayerProperties, 205)
 // Added more context data to CVD's traced shapes so we can play it back at the solver stage level (not just game thread frames) 
 FORTNITE_MAIN_VERSION(ThreadContextDataInChaosVisualDebuggerDebugDrawData, 206)
 
+// Changed default grid mode in surface sampler to a version that's more intuitive and less error-prone
+FORTNITE_MAIN_VERSION(PCGChangedSurfaceSamplerDefaultGridCreationMode, 207)
+
 // -----<new versions can be added above this line>-------------------------------------------------
 #undef FORTNITE_MAIN_VERSION
\ No newline at end of file
