[PCG] Enable new StructUtils details view support

Dependent on: 41783637

#rnx

diff --git a/Engine/Plugins/PCG/Source/PCG/Public/PCGGraph.h b/Engine/Plugins/PCG/Source/PCG/Public/PCGGraph.h
--- a/Engine/Plugins/PCG/Source/PCG/Public/PCGGraph.h
+++ b/Engine/Plugins/PCG/Source/PCG/Public/PCGGraph.h
@@ -477,7 +477,7 @@ protected:
 #endif // WITH_EDITORONLY_DATA
 
 	// Parameters
-	UPROPERTY(EditAnywhere, Category = "Graph Parameters", DisplayName = "Parameters", meta = (NoResetToDefault, DefaultType = "EPropertyBagPropertyType::Double", IsPinTypeAccepted = "UserParametersIsPinTypeAccepted", CanRemoveProperty = "UserParametersCanRemoveProperty"))
+	UPROPERTY(EditAnywhere, Category = "Graph Parameters", DisplayName = "Parameters", meta = (NoResetToDefault, DefaultType = "EPropertyBagPropertyType::Double", IsPinTypeAccepted = "UserParametersIsPinTypeAccepted", CanRemoveProperty = "UserParametersCanRemoveProperty", ChildRowFeatures = "Extended"))
 	FInstancedPropertyBag UserParameters;
 
 #if WITH_EDITOR
