PCG: Fix short class paths

diff --git a/Engine/Plugins/PCG/Source/PCG/Private/Compute/Elements/PCGCustomHLSL.h b/Engine/Plugins/PCG/Source/PCG/Private/Compute/Elements/PCGCustomHLSL.h
--- a/Engine/Plugins/PCG/Source/PCG/Private/Compute/Elements/PCGCustomHLSL.h
+++ b/Engine/Plugins/PCG/Source/PCG/Private/Compute/Elements/PCGCustomHLSL.h
@@ -161,11 +161,11 @@ public:
 protected:
 #if WITH_EDITORONLY_DATA
 	/** Override your kernel with a PCG compute source asset. */
-	UPROPERTY(EditAnywhere, Category = "Settings", meta = (AllowedClasses = "PCGComputeSource"))
+	UPROPERTY(EditAnywhere, Category = "Settings", meta = (AllowedClasses = "/Script/PCG.PCGComputeSource"))
 	TObjectPtr<UComputeSource> KernelSourceOverride;
 
 	/** Additional source files to use in your kernel. */
-	UPROPERTY(EditAnywhere, Category = "Settings", meta = (AllowedClasses = "PCGComputeSource"))
+	UPROPERTY(EditAnywhere, Category = "Settings", meta = (AllowedClasses = "/Script/PCG.PCGComputeSource"))
 	TArray<TObjectPtr<UComputeSource>> AdditionalSources;
 #endif
 
