[PCG] Runtime generation radius multiplier

#rnx

diff --git a/Engine/Plugins/PCG/Source/PCG/Private/PCGComponent.cpp b/Engine/Plugins/PCG/Source/PCG/Private/PCGComponent.cpp
--- a/Engine/Plugins/PCG/Source/PCG/Private/PCGComponent.cpp
+++ b/Engine/Plugins/PCG/Source/PCG/Private/PCGComponent.cpp
@@ -83,6 +83,11 @@ namespace PCGComponent
 		false,
 		TEXT("At the end of execution, generated output data is flattened in place if it can be stolen, but it can be disabled if it has unwanted side effects."));
 
+	static TAutoConsoleVariable<float> CVarRuntimeGenerationRadiusMultiplier(
+		TEXT("pcg.RuntimeGeneration.GlobalRadiusMultiplier"),
+		1.0f,
+		TEXT("Global multiplier for generation radius of all runtime gen components."));
+
 	struct FPCGManagedActorLoadingScope
 	{
 		FPCGManagedActorLoadingScope()
@@ -1278,15 +1283,17 @@ void UPCGComponent::SetSchedulingPolicyClass(TSubclassOf<UPCGSchedulingPolicyBas
 
 double UPCGComponent::GetGenerationRadiusFromGrid(EPCGHiGenGrid Grid) const
 {
+	const double Multiplier = FMath::Max(PCGComponent::CVarRuntimeGenerationRadiusMultiplier.GetValueOnAnyThread(), 0.0f);
+
 	if (bOverrideGenerationRadii)
 	{
-		return GenerationRadii.GetGenerationRadiusFromGrid(Grid);
+		return Multiplier * GenerationRadii.GetGenerationRadiusFromGrid(Grid);
 	}
 
 	const UPCGGraph* Graph = GetGraph();
 	if (ensure(Graph))
 	{
-		return Graph->GetGridGenerationRadiusFromGrid(Grid);
+		return Multiplier * Graph->GetGridGenerationRadiusFromGrid(Grid);
 	}
 
 	return 0;
@@ -1294,15 +1301,17 @@ double UPCGComponent::GetGenerationRadiusFromGrid(EPCGHiGenGrid Grid) const
 
 double UPCGComponent::GetCleanupRadiusFromGrid(EPCGHiGenGrid Grid) const
 {
+	const double Multiplier = FMath::Max(PCGComponent::CVarRuntimeGenerationRadiusMultiplier.GetValueOnAnyThread(), 0.0f);
+
 	if (bOverrideGenerationRadii)
 	{
-		return GenerationRadii.GetCleanupRadiusFromGrid(Grid);
+		return Multiplier * GenerationRadii.GetCleanupRadiusFromGrid(Grid);
 	}
 
 	const UPCGGraph* Graph = GetGraph();
 	if (ensure(Graph))
 	{
-		return Graph->GetGridCleanupRadiusFromGrid(Grid);
+		return Multiplier * Graph->GetGridCleanupRadiusFromGrid(Grid);
 	}
 
 	return 0;
