[PCG] Move initialization of VT resources to RT rather than GT proxy creation.

When running Debug a checkSlow() is hit because we're calling into RT code from the GT.

#rnx

diff --git a/Engine/Plugins/PCG/Source/PCG/Private/Compute/DataInterfaces/PCGLandscapeDataInterface.cpp b/Engine/Plugins/PCG/Source/PCG/Private/Compute/DataInterfaces/PCGLandscapeDataInterface.cpp
--- a/Engine/Plugins/PCG/Source/PCG/Private/Compute/DataInterfaces/PCGLandscapeDataInterface.cpp
+++ b/Engine/Plugins/PCG/Source/PCG/Private/Compute/DataInterfaces/PCGLandscapeDataInterface.cpp
@@ -573,42 +573,14 @@ FPCGLandscapeDataProviderProxy::FPCGLandscapeDataProviderProxy(
 	ERuntimeVirtualTextureMaterialType InNormalVirtualTextureMode,
 	FVector2D InLandscapeGridSize)
 	: Resource(InResource)
+	, BaseColorVirtualTexture(InBaseColorVirtualTexture)
+	, HeightVirtualTexture(InHeightVirtualTexture)
+	, NormalVirtualTexture(InNormalVirtualTexture)
+	, bBaseColorSRGB(bInBaseColorSRGB)
+	, NormalVirtualTextureMode(InNormalVirtualTextureMode)
 	, BaseColorVirtualTextureUnpackType(InBaseColorVirtualTextureUnpackType)
 	, LandscapeGridSize(InLandscapeGridSize)
 {
-	BaseColorVirtualLayer.Initialize(InBaseColorVirtualTexture, /*PageTableIndex=*/0, bInBaseColorSRGB);
-	BaseColorVirtualPage.Initialize(InBaseColorVirtualTexture, /*PageTableIndex=*/0, /*bIncludeWorldToUV=*/true, /*bIncludeHeightUnpack=*/false);
-
-	HeightVirtualLayer.Initialize(InHeightVirtualTexture, /*PageTableIndex=*/0, /*bSRGB=*/false);
-	HeightVirtualPage.Initialize(InHeightVirtualTexture, /*PageTableIndex=*/0, /*bIncludeWorldToUV=*/true, /*bIncludeHeightUnpack=*/true);
-
-	NormalVirtualPage.Initialize(InNormalVirtualTexture, /*PageTableIndex=*/0, /*bIncludeWorldToUV=*/true, /*bIncludeHeightUnpack=*/false);
-
-	switch (InNormalVirtualTextureMode)
-	{
-	case ERuntimeVirtualTextureMaterialType::BaseColor_Normal_Roughness:
-		NormalVirtualLayer0.Initialize(InNormalVirtualTexture, /*PageTableIndex=*/0, /*bSRGB=*/false);
-		NormalVirtualLayer1.Initialize(InNormalVirtualTexture, /*PageTableIndex=*/1, /*bSRGB=*/false);
-		NormalUnpackMode = PCGVirtualTextureCommon::ENormalUnpackType::B5G6R5;
-		break;
-	case ERuntimeVirtualTextureMaterialType::BaseColor_Normal_Specular:
-		NormalVirtualLayer0.Initialize(InNormalVirtualTexture, /*PageTableIndex=*/0, /*bSRGB=*/false);
-		NormalVirtualLayer1.Initialize(InNormalVirtualTexture, /*PageTableIndex=*/1, /*bSRGB=*/false);
-		NormalUnpackMode = PCGVirtualTextureCommon::ENormalUnpackType::BC3BC3;
-		break;
-	case ERuntimeVirtualTextureMaterialType::BaseColor_Normal_Specular_YCoCg:
-	case ERuntimeVirtualTextureMaterialType::BaseColor_Normal_Specular_Mask_YCoCg:
-		NormalVirtualLayer0.Initialize(InNormalVirtualTexture, /*PageTableIndex=*/1, /*bSRGB=*/false);
-		NormalVirtualLayer1.Initialize(InNormalVirtualTexture, /*PageTableIndex=*/2, /*bSRGB=*/false);
-		NormalUnpackMode = PCGVirtualTextureCommon::ENormalUnpackType::BC5BC1;
-		break;
-	default:
-		NormalVirtualPage.Reset();
-		NormalVirtualLayer0.Reset();
-		NormalVirtualLayer1.Reset();
-		NormalUnpackMode = PCGVirtualTextureCommon::ENormalUnpackType::None;
-		break;
-	}
 }
 
 bool FPCGLandscapeDataProviderProxy::IsValid(FValidationData const& InValidationData) const
@@ -627,6 +599,51 @@ void FPCGLandscapeDataProviderProxy::GatherDispatchData(FDispatchData const& InD
 
 void FPCGLandscapeDataProviderProxy::AllocateResources(FRDGBuilder& GraphBuilder, FAllocationData const& InAllocationData)
 {
+	if (TStrongObjectPtr<const URuntimeVirtualTexture> BaseColorVT = BaseColorVirtualTexture.Pin())
+	{
+		BaseColorVirtualLayer.Initialize(BaseColorVT.Get(), /*PageTableIndex=*/0, bBaseColorSRGB);
+		BaseColorVirtualPage.Initialize(BaseColorVT.Get(), /*PageTableIndex=*/0, /*bIncludeWorldToUV=*/true, /*bIncludeHeightUnpack=*/false);
+	}
+
+	if (TStrongObjectPtr<const URuntimeVirtualTexture> HeightVT = HeightVirtualTexture.Pin())
+	{
+		HeightVirtualLayer.Initialize(HeightVT.Get(), /*PageTableIndex=*/0, /*bSRGB=*/false);
+		HeightVirtualPage.Initialize(HeightVT.Get(), /*PageTableIndex=*/0, /*bIncludeWorldToUV=*/true, /*bIncludeHeightUnpack=*/true);
+	}
+
+	if (TStrongObjectPtr<const URuntimeVirtualTexture> NormalVT = NormalVirtualTexture.Pin())
+	{
+		const URuntimeVirtualTexture* NormalVTPtr = NormalVT.Get();
+
+		NormalVirtualPage.Initialize(NormalVTPtr, /*PageTableIndex=*/0, /*bIncludeWorldToUV=*/true, /*bIncludeHeightUnpack=*/false);
+
+		switch (NormalVirtualTextureMode)
+		{
+		case ERuntimeVirtualTextureMaterialType::BaseColor_Normal_Roughness:
+			NormalVirtualLayer0.Initialize(NormalVTPtr, /*PageTableIndex=*/0, /*bSRGB=*/false);
+			NormalVirtualLayer1.Initialize(NormalVTPtr, /*PageTableIndex=*/1, /*bSRGB=*/false);
+			NormalUnpackMode = PCGVirtualTextureCommon::ENormalUnpackType::B5G6R5;
+			break;
+		case ERuntimeVirtualTextureMaterialType::BaseColor_Normal_Specular:
+			NormalVirtualLayer0.Initialize(NormalVTPtr, /*PageTableIndex=*/0, /*bSRGB=*/false);
+			NormalVirtualLayer1.Initialize(NormalVTPtr, /*PageTableIndex=*/1, /*bSRGB=*/false);
+			NormalUnpackMode = PCGVirtualTextureCommon::ENormalUnpackType::BC3BC3;
+			break;
+		case ERuntimeVirtualTextureMaterialType::BaseColor_Normal_Specular_YCoCg:
+		case ERuntimeVirtualTextureMaterialType::BaseColor_Normal_Specular_Mask_YCoCg:
+			NormalVirtualLayer0.Initialize(NormalVTPtr, /*PageTableIndex=*/1, /*bSRGB=*/false);
+			NormalVirtualLayer1.Initialize(NormalVTPtr, /*PageTableIndex=*/2, /*bSRGB=*/false);
+			NormalUnpackMode = PCGVirtualTextureCommon::ENormalUnpackType::BC5BC1;
+			break;
+		default:
+			NormalVirtualPage.Reset();
+			NormalVirtualLayer0.Reset();
+			NormalVirtualLayer1.Reset();
+			NormalUnpackMode = PCGVirtualTextureCommon::ENormalUnpackType::None;
+			break;
+		}
+	}
+
 	if (Resource && Resource->LandscapeTexture)
 	{
 		Resource->LandscapeTexture->InitResource(GraphBuilder.RHICmdList);
diff --git a/Engine/Plugins/PCG/Source/PCG/Private/Compute/DataInterfaces/PCGLandscapeDataInterface.h b/Engine/Plugins/PCG/Source/PCG/Private/Compute/DataInterfaces/PCGLandscapeDataInterface.h
--- a/Engine/Plugins/PCG/Source/PCG/Private/Compute/DataInterfaces/PCGLandscapeDataInterface.h
+++ b/Engine/Plugins/PCG/Source/PCG/Private/Compute/DataInterfaces/PCGLandscapeDataInterface.h
@@ -136,6 +136,13 @@ protected:
 protected:
 	TSharedPtr<const FPCGLandscapeResource> Resource;
 
+	TWeakObjectPtr<const URuntimeVirtualTexture> BaseColorVirtualTexture;
+	TWeakObjectPtr<const URuntimeVirtualTexture> HeightVirtualTexture;
+	TWeakObjectPtr<const URuntimeVirtualTexture> NormalVirtualTexture;
+
+	bool bBaseColorSRGB = false;
+	ERuntimeVirtualTextureMaterialType NormalVirtualTextureMode = ERuntimeVirtualTextureMaterialType::Count;
+
 	PCGVirtualTextureCommon::FVirtualTexturePageTable BaseColorVirtualPage;
 	PCGVirtualTextureCommon::FVirtualTextureLayer BaseColorVirtualLayer;
 	PCGVirtualTextureCommon::EBaseColorUnpackType BaseColorVirtualTextureUnpackType = PCGVirtualTextureCommon::EBaseColorUnpackType::None;
