[PCG] Iterate on world loaded queries in runtime gen.

* Cache results which saves a good chunk of queries esp when grid parent queries are enabled.
* Disable testing world loaded when adding generation sources, not needed anymore.

#rnx

diff --git a/Engine/Plugins/PCG/Source/PCG/Private/RuntimeGen/PCGGenSourceManager.cpp b/Engine/Plugins/PCG/Source/PCG/Private/RuntimeGen/PCGGenSourceManager.cpp
--- a/Engine/Plugins/PCG/Source/PCG/Private/RuntimeGen/PCGGenSourceManager.cpp
+++ b/Engine/Plugins/PCG/Source/PCG/Private/RuntimeGen/PCGGenSourceManager.cpp
@@ -14,8 +14,6 @@
 #include "Engine/World.h"
 #include "GameFramework/GameModeBase.h"
 #include "WorldPartition/WorldPartition.h"
-#include "WorldPartition/WorldPartitionRuntimeCell.h"
-#include "WorldPartition/WorldPartitionSubsystem.h"
 
 #if WITH_EDITOR
 #include "EditorViewportClient.h"
@@ -51,56 +49,17 @@ TSet<IPCGGenSourceBase*> FPCGGenSourceManager::GetAllGenSources(const APCGWorldA
 		bDirty = false;
 	}
 
-	UWorldPartitionSubsystem* WorldPartitionSubsystem = InPCGWorldActor ? UWorld::GetSubsystem<UWorldPartitionSubsystem>(InPCGWorldActor->GetWorld()) : nullptr;
-	
-	// Prepare query upfront.
-	TArray<FWorldPartitionStreamingQuerySource> QuerySources;
-	FWorldPartitionStreamingQuerySource& QuerySource = QuerySources.Emplace_GetRef();
-	QuerySource.bSpatialQuery = true;
-	QuerySource.bUseGridLoadingRange = false;
-	// Providing 1cm is standard when checking loading state.
-	QuerySource.Radius = 1.0f;
-	QuerySource.bDataLayersOnly = false;
-
-	auto IsWorldLoaded = [WorldPartitionSubsystem, &QuerySources](IPCGGenSourceBase* InGenSource)
-	{
-		TRACE_CPUPROFILER_EVENT_SCOPE(QueryIsWorldLoaded);
-
-		if (TOptional<FVector> Position = InGenSource->GetPosition())
-		{
-			QuerySources[0].Location = *Position;
-
-			const bool bIsLoaded = WorldPartitionSubsystem->IsStreamingCompleted(EWorldPartitionRuntimeCellState::Activated, QuerySources, /*bExactState*/ false);
-			if (!bIsLoaded)
-			{
-				UE_LOG(LogPCG, Verbose, TEXT("Discarded gen source due to world not loaded."));
-			}
-
-			return bIsLoaded;
-		}
-		else
-		{
-			return false;
-		}
-	};
-
 	TSet<IPCGGenSourceBase*> AllGenSources;
 	AllGenSources.Reserve(RegisteredGenSources.Num() + RegisteredNamedGenSources.Num());
 
 	for (TScriptInterface<IPCGGenSourceBase> GenSource : RegisteredGenSources)
 	{
-		if (IsWorldLoaded(GenSource.GetInterface()))
-		{
-			AllGenSources.Add(GenSource.GetInterface());
-		}
+		AllGenSources.Add(GenSource.GetInterface());
 	}
 
 	for (const TPair<FName, TScriptInterface<IPCGGenSourceBase>>& NamedGenSource : RegisteredNamedGenSources)
 	{
-		if (IsWorldLoaded(NamedGenSource.Value.GetInterface()))
-		{
-			AllGenSources.Add(NamedGenSource.Value.GetInterface());
-		}
+		AllGenSources.Add(NamedGenSource.Value.GetInterface());
 	}
 
 	// Acquire a generation source for each active streaming source that doesn't have a generation source already (based on name collision test).
@@ -119,7 +78,7 @@ TSet<IPCGGenSourceBase*> FPCGGenSourceManager::GetAllGenSources(const APCGWorldA
 			}
 		}
 
-		if (!bAlreadyTracked && IsWorldLoaded(WPGenSource))
+		if (!bAlreadyTracked)
 		{
 			AllGenSources.Add(WPGenSource);
 		}
@@ -127,7 +86,7 @@ TSet<IPCGGenSourceBase*> FPCGGenSourceManager::GetAllGenSources(const APCGWorldA
 
 #if WITH_EDITOR
 	// Acquire a generation source for the active editor viewport if one exists.
-	if (EditorCameraGenSource && EditorCameraGenSource->EditorViewportClient && IsWorldLoaded(EditorCameraGenSource))
+	if (EditorCameraGenSource && EditorCameraGenSource->EditorViewportClient)
 	{
 		AllGenSources.Add(EditorCameraGenSource);
 	}
diff --git a/Engine/Plugins/PCG/Source/PCG/Private/RuntimeGen/PCGRuntimeGenScheduler.cpp b/Engine/Plugins/PCG/Source/PCG/Private/RuntimeGen/PCGRuntimeGenScheduler.cpp
--- a/Engine/Plugins/PCG/Source/PCG/Private/RuntimeGen/PCGRuntimeGenScheduler.cpp
+++ b/Engine/Plugins/PCG/Source/PCG/Private/RuntimeGen/PCGRuntimeGenScheduler.cpp
@@ -325,14 +325,25 @@ void FPCGRuntimeGenScheduler::TickQueueComponentsForGeneration(
 		QuerySource.bDataLayersOnly = false;
 	}
 
-	auto IsWorldStreamingComplete = [WorldPartitionSubsystem, &StreamingQuerySources](const FVector& InLocation, float InGridSize)
+	CachedStreamingQueryResults.Empty(CachedStreamingQueryResults.Num());
+
+	auto IsWorldStreamingComplete = [WorldPartitionSubsystem, &StreamingQuerySources, &CachedStreamingQueryResults=CachedStreamingQueryResults](const FVector& InLocation, float InGridSize)
 	{
 		TRACE_CPUPROFILER_EVENT_SCOPE(IsWorldStreamingComplete);
 
+		FStreamingCompleteQueryKey Key{ InLocation, InGridSize };
+		if (bool* FoundResult = CachedStreamingQueryResults.Find(Key))
+		{
+			return *FoundResult;
+		}
+
 		StreamingQuerySources[0].Radius = InGridSize / 2.0f;
 		StreamingQuerySources[0].Location = InLocation;
 
 		const bool bIsLoaded = WorldPartitionSubsystem->IsStreamingCompleted(EWorldPartitionRuntimeCellState::Activated, StreamingQuerySources, /*bExactState*/ false);
+
+		CachedStreamingQueryResults.Add(Key, bIsLoaded);
+
 		if (!bIsLoaded)
 		{
 			UE_LOG(LogPCG, Verbose, TEXT("Holding back generation of cell at (%.2f, %.2f, %.2f), grid size %f, due to world not loaded."), InLocation.X, InLocation.Y, InLocation.Z, InGridSize);
diff --git a/Engine/Plugins/PCG/Source/PCG/Public/RuntimeGen/PCGRuntimeGenScheduler.h b/Engine/Plugins/PCG/Source/PCG/Public/RuntimeGen/PCGRuntimeGenScheduler.h
--- a/Engine/Plugins/PCG/Source/PCG/Public/RuntimeGen/PCGRuntimeGenScheduler.h
+++ b/Engine/Plugins/PCG/Source/PCG/Public/RuntimeGen/PCGRuntimeGenScheduler.h
@@ -178,4 +178,8 @@ private:
 	const UPCGComponent* OriginalComponentBeingGenerated = nullptr;
 
 	int32 FramesUntilGeneration = 0;
+
+	// Local to TickQueueComponentsForGeneration, cached here for efficiency.
+	using FStreamingCompleteQueryKey = TPair<FVector, float>;
+	TMap<FStreamingCompleteQueryKey, bool> CachedStreamingQueryResults;
 };
