[PCG] Added node to gather contextual information in the graph.

diff --git a/Engine/Plugins/PCG/Source/PCG/Private/Elements/PCGGetExecutionContext.cpp b/Engine/Plugins/PCG/Source/PCG/Private/Elements/PCGGetExecutionContext.cpp
new file mode 100644
--- /dev/null
+++ b/Engine/Plugins/PCG/Source/PCG/Private/Elements/PCGGetExecutionContext.cpp
@@ -0,0 +1,106 @@
+// Copyright Epic Games, Inc. All Rights Reserved.
+
+#include "Elements/PCGGetExecutionContext.h"
+
+#include "PCGComponent.h"
+#include "PCGParamData.h"
+#include "PCGPin.h"
+#include "Helpers/PCGHelpers.h"
+#include "Metadata/PCGMetadata.h"
+
+#include "Engine/World.h"
+
+#include UE_INLINE_GENERATED_CPP_BY_NAME(PCGGetExecutionContext)
+
+#define LOCTEXT_NAMESPACE "PCGGetExecutionContextElement"
+
+namespace PCGGetExecutionContextConstants
+{
+	static const FName AttributeName = TEXT("Info");
+}
+
+#if WITH_EDITOR
+FText UPCGGetExecutionContextSettings::GetNodeTooltipText() const
+{
+	return LOCTEXT("GetExecutionContextTooltip", "Returns some context-specific common information.");
+}
+
+EPCGChangeType UPCGGetExecutionContextSettings::GetChangeTypeForProperty(const FName& InPropertyName) const
+{
+	return Super::GetChangeTypeForProperty(InPropertyName) | EPCGChangeType::Cosmetic;
+}
+#endif
+
+TArray<FPCGPinProperties> UPCGGetExecutionContextSettings::OutputPinProperties() const
+{
+	TArray<FPCGPinProperties> PinProperties;
+	PinProperties.Emplace(PCGPinConstants::DefaultOutputLabel, EPCGDataType::Param);
+
+	return PinProperties;
+}
+
+FString UPCGGetExecutionContextSettings::GetAdditionalTitleInformation() const
+{
+	if (const UEnum* EnumPtr = StaticEnum<EPCGGetExecutionContextMode>())
+	{
+		return FText::Format(LOCTEXT("AdditionalTitle", "Get {0}"), EnumPtr->GetDisplayNameTextByValue(static_cast<int64>(Mode))).ToString();
+	}
+	else
+	{
+		return FString();
+	}
+}
+
+FPCGElementPtr UPCGGetExecutionContextSettings::CreateElement() const
+{
+	return MakeShared<FPCGGetExecutionContextElement>();
+}
+
+bool FPCGGetExecutionContextElement::ExecuteInternal(FPCGContext* Context) const
+{
+	TRACE_CPUPROFILER_EVENT_SCOPE(FPCGGetExecutionContextElement::Execute);
+
+	check(Context);
+
+	const UPCGGetExecutionContextSettings* Settings = Context->GetInputSettings<UPCGGetExecutionContextSettings>();
+	check(Settings);
+
+	const EPCGGetExecutionContextMode Mode = Settings->Mode;
+
+	UPCGParamData* ParamData = FPCGContext::NewObject_AnyThread<UPCGParamData>(Context);
+	check(ParamData && ParamData->Metadata);
+
+	FPCGTaggedData& OutputData = Context->OutputData.TaggedData.Emplace_GetRef();
+	OutputData.Data = ParamData;
+
+	IPCGGraphExecutionSource* ExecutionSource = Context->ExecutionSource.Get();
+	UPCGComponent* SourceComponent = Cast<UPCGComponent>(ExecutionSource);
+
+	bool Value = false;
+
+	if (Mode == EPCGGetExecutionContextMode::IsEditor || Mode == EPCGGetExecutionContextMode::IsRuntime)
+	{
+		UWorld* SupportingWorld = ExecutionSource ? ExecutionSource->GetExecutionState().GetWorld() : nullptr;
+		const bool bContextIsRuntimeOrOnRuntimeGenerationMode = (SupportingWorld && SupportingWorld->IsGameWorld()) || PCGHelpers::IsRuntimeGeneration(ExecutionSource);
+		Value = (bContextIsRuntimeOrOnRuntimeGenerationMode == (Mode == EPCGGetExecutionContextMode::IsRuntime));
+	}
+	else if (Mode == EPCGGetExecutionContextMode::IsOriginal || Mode == EPCGGetExecutionContextMode::IsLocal)
+	{
+		Value = SourceComponent && (SourceComponent->IsLocalComponent() == (Mode == EPCGGetExecutionContextMode::IsLocal));
+	}
+	else if (Mode == EPCGGetExecutionContextMode::IsPartitioned)
+	{
+		Value = SourceComponent && SourceComponent->IsPartitioned();
+	}
+	else if (Mode == EPCGGetExecutionContextMode::IsRuntimeGeneration)
+	{
+		Value = PCGHelpers::IsRuntimeGeneration(ExecutionSource);
+	}
+
+	FPCGMetadataAttribute<bool>* Attribute = ParamData->Metadata->CreateAttribute<bool>(PCGGetExecutionContextConstants::AttributeName, Value, /*bAllowInterpolation=*/false, /*bOverrideParent=*/false);
+	ParamData->Metadata->AddEntry();
+
+	return true;
+}
+
+#undef LOCTEXT_NAMESPACE
\ No newline at end of file
diff --git a/Engine/Plugins/PCG/Source/PCG/Public/Elements/PCGGetExecutionContext.h b/Engine/Plugins/PCG/Source/PCG/Public/Elements/PCGGetExecutionContext.h
new file mode 100644
--- /dev/null
+++ b/Engine/Plugins/PCG/Source/PCG/Public/Elements/PCGGetExecutionContext.h
@@ -0,0 +1,59 @@
+// Copyright Epic Games, Inc. All Rights Reserved.
+
+#pragma once
+
+#include "PCGContext.h"
+#include "PCGSettings.h"
+
+#include "PCGGetExecutionContext.generated.h"
+
+UENUM(BlueprintType)
+enum class EPCGGetExecutionContextMode : uint8
+{
+	IsEditor UMETA(Tooltip="Returns whether the execution happens in the editor (not in PIE)."),
+	IsRuntime UMETA(Tooltip="Returns whether the execution happens in a runtime environment (PIE or cooked build) or not."),
+	IsOriginal UMETA(Tooltip="Returns whether the executing context is the original context (e.g. not partitioned)."),
+	IsLocal UMETA(Tooltip="Returns whether the executing context is partitioned (in opposition to original)."),
+	IsPartitioned UMETA(Tooltip = "Returns whether the executing context is the original context and that it is partitioned."),
+	IsRuntimeGeneration UMETA(Tooltip="Returns whether the executing context is doing runtime gneration.")
+};
+
+/** Returns some context-specific information in the form of an attribute set. */
+UCLASS(MinimalAPI, BlueprintType, ClassGroup=(Procedural))
+class UPCGGetExecutionContextSettings : public UPCGSettings
+{
+	GENERATED_BODY()
+
+public:
+	//~Begin UPCGSettings interface implementation
+#if WITH_EDITOR
+	virtual FName GetDefaultNodeName() const override { return FName(TEXT("GetExecutionContext")); }
+	virtual FText GetDefaultNodeTitle() const override { return NSLOCTEXT("PCGGetExecutionContext", "NodeTitle", "Get Execution Context Info"); }
+	virtual FText GetNodeTooltipText() const override;
+	virtual EPCGSettingsType GetType() const override { return EPCGSettingsType::Generic; }
+	virtual EPCGChangeType GetChangeTypeForProperty(const FName& InPropertyName) const override;
+#endif // WITH_EDITOR
+
+	virtual bool HasFlippedTitleLines() const override { return true; }
+	virtual FString GetAdditionalTitleInformation() const override;
+
+protected:
+	virtual TArray<FPCGPinProperties> InputPinProperties() const override { return TArray<FPCGPinProperties>(); }
+	virtual TArray<FPCGPinProperties> OutputPinProperties() const override;
+
+	virtual FPCGElementPtr CreateElement() const override;
+	// ~End UPCGSettings interface implementation
+
+public:
+	UPROPERTY(BlueprintReadWrite, EditAnywhere, Category = Settings)
+	EPCGGetExecutionContextMode Mode = EPCGGetExecutionContextMode::IsRuntime;
+};
+
+class FPCGGetExecutionContextElement : public IPCGElement
+{
+public:
+	virtual bool IsCacheable(const UPCGSettings* InSettings) const { return false; }
+
+protected:
+	virtual bool ExecuteInternal(FPCGContext* Context) const override;
+};
\ No newline at end of file
