[PCG] Fixed interpolation error introduced recently.

#rnx

diff --git a/Engine/Plugins/PCG/Source/PCG/Private/Elements/PCGAttractElement.cpp b/Engine/Plugins/PCG/Source/PCG/Private/Elements/PCGAttractElement.cpp
--- a/Engine/Plugins/PCG/Source/PCG/Private/Elements/PCGAttractElement.cpp
+++ b/Engine/Plugins/PCG/Source/PCG/Private/Elements/PCGAttractElement.cpp
@@ -82,7 +82,8 @@ namespace PCGAttractElement
 						T OutputValue{};
 						if constexpr (PCG::Private::MetadataTraits<T>::CanInterpolate)
 						{
-							OutputValue = PCG::Private::MetadataTraits<T>::WeightedSum(SourceValue, TargetValue, Alpha);
+							OutputValue = PCG::Private::MetadataTraits<T>::WeightedSum(PCG::Private::MetadataTraits<T>::ZeroValueForWeightedSum(), SourceValue, 1.0 - Alpha);
+							OutputValue = PCG::Private::MetadataTraits<T>::WeightedSum(OutputValue, TargetValue, Alpha);
 
 							if constexpr (PCG::Private::MetadataTraits<T>::InterpolationNeedsNormalization)
 							{
