[PCG - Build Health] Fix unreachable code in PCGExportSelectedAttributes

#rnx

diff --git a/Engine/Plugins/PCG/Source/PCG/Private/Elements/IO/PCGExportSelectedAttributes.cpp b/Engine/Plugins/PCG/Source/PCG/Private/Elements/IO/PCGExportSelectedAttributes.cpp
--- a/Engine/Plugins/PCG/Source/PCG/Private/Elements/IO/PCGExportSelectedAttributes.cpp
+++ b/Engine/Plugins/PCG/Source/PCG/Private/Elements/IO/PCGExportSelectedAttributes.cpp
@@ -166,25 +166,26 @@ namespace PCGExportSelectedAttributes
 			template <typename T>
 			TSharedPtr<FJsonValue> ConvertFloatingPointTypeToJson(const T& Value)
 			{
-				static_assert(PCG::Private::MetadataTraits<T>::IsFloatingPoint, "Helper is only designed to break floating point compound components.");
+				using U = std::remove_cv_t<T>;
+				static_assert(PCG::Private::MetadataTraits<U>::IsFloatingPoint, "Helper is only designed to break floating point compound components.");
 
-				if constexpr (std::is_same_v<T, double>)
+				if constexpr (std::is_same_v<U, float> || std::is_same_v<U, double>)
 				{
 					return MakeShared<FJsonValueNumber>(Value);
 				}
-				if constexpr (std::is_same_v<T, FVector2D>)
+				else if constexpr (std::is_same_v<U, FVector2D>)
 				{
 					return ToJsonArray({Value.X, Value.Y});
 				}
-				else if constexpr (std::is_same_v<T, FVector>)
+				else if constexpr (std::is_same_v<U, FVector>)
 				{
 					return ToJsonArray({Value.X, Value.Y, Value.Z});
 				}
-				else if constexpr (std::is_same_v<T, FVector4> || std::is_same_v<T, FQuat>)
+				else if constexpr (std::is_same_v<U, FVector4> || std::is_same_v<U, FQuat>)
 				{
 					return ToJsonArray({Value.X, Value.Y, Value.Z, Value.W});
 				}
-				else if constexpr (std::is_same_v<T, FTransform>)
+				else if constexpr (std::is_same_v<U, FTransform>)
 				{
 					TSharedPtr<FJsonObject> JsonTransformObject = MakeShared<FJsonObject>();
 					JsonTransformObject->SetArrayField(
@@ -199,12 +200,11 @@ namespace PCGExportSelectedAttributes
 
 					return MakeShared<FJsonValueObject>(JsonTransformObject);
 				}
-				else if constexpr (std::is_same_v<T, FRotator>)
+				else //if constexpr (std::is_same_v<U, FRotator>)
 				{
+					static_assert(std::is_same_v<U, FRotator>, "Unsupported type");
 					return ToJsonArray({Value.Pitch, Value.Yaw, Value.Roll});
 				}
-
-				return nullptr;
 			}
 		}
 	}
