[PCG] Add new node for control points extraction

* Add new function to get metadata entries in a polyline
* Implement those functions in the spline and landscape spline
* Fix bug with UPCGLandscapeSplineData::GetTransform/LocationAtAlpha

#rnx

diff --git a/Engine/Plugins/PCG/Source/PCG/Private/Data/PCGLandscapeSplineData.cpp b/Engine/Plugins/PCG/Source/PCG/Private/Data/PCGLandscapeSplineData.cpp
--- a/Engine/Plugins/PCG/Source/PCG/Private/Data/PCGLandscapeSplineData.cpp
+++ b/Engine/Plugins/PCG/Source/PCG/Private/Data/PCGLandscapeSplineData.cpp
@@ -251,14 +251,14 @@ FVector UPCGLandscapeSplineData::GetLocationAtAlpha(float Alpha) const
 {
 	const float InputKey = GetInputKeyAtAlpha(Alpha);
 	const int SegmentIndex = FMath::TruncToInt(InputKey);
-	return GetLocationAtDistance(SegmentIndex, FMath::Frac(Alpha) * GetSegmentLength(SegmentIndex));
+	return GetLocationAtDistance(SegmentIndex, FMath::Frac(InputKey) * GetSegmentLength(SegmentIndex));
 }
 
 FTransform UPCGLandscapeSplineData::GetTransformAtAlpha(float Alpha) const
 {
 	const float InputKey = GetInputKeyAtAlpha(Alpha);
 	const int SegmentIndex = FMath::TruncToInt(InputKey);
-	return GetTransformAtDistance(SegmentIndex, FMath::Frac(Alpha) * GetSegmentLength(SegmentIndex));
+	return GetTransformAtDistance(SegmentIndex, FMath::Frac(InputKey) * GetSegmentLength(SegmentIndex));
 }
 
 const UPCGPointData* UPCGLandscapeSplineData::CreatePointData(FPCGContext* Context) const
diff --git a/Engine/Plugins/PCG/Source/PCG/Private/Elements/PCGCreateSpline.cpp b/Engine/Plugins/PCG/Source/PCG/Private/Elements/PCGCreateSpline.cpp
--- a/Engine/Plugins/PCG/Source/PCG/Private/Elements/PCGCreateSpline.cpp
+++ b/Engine/Plugins/PCG/Source/PCG/Private/Elements/PCGCreateSpline.cpp
@@ -27,8 +27,8 @@ UPCGCreateSplineSettings::UPCGCreateSplineSettings()
 	// Change the default for Arrive and Leave tangent
 	if (PCGHelpers::IsNewObjectAndNotDefault(this))
 	{
-		ArriveTangentAttribute = PCGSplineSamplerConstants::ArriveTangentName;
-		LeaveTangentAttribute = PCGSplineSamplerConstants::LeaveTangentName;
+		ArriveTangentAttribute = PCGSplineSamplerConstants::ArriveTangentAttributeName;
+		LeaveTangentAttribute = PCGSplineSamplerConstants::LeaveTangentAttributeName;
 	}
 }
 
diff --git a/Engine/Plugins/PCG/Source/PCG/Private/Elements/PCGGetSplineControlPoints.cpp b/Engine/Plugins/PCG/Source/PCG/Private/Elements/PCGGetSplineControlPoints.cpp
new file mode 100644
--- /dev/null
+++ b/Engine/Plugins/PCG/Source/PCG/Private/Elements/PCGGetSplineControlPoints.cpp
@@ -0,0 +1,132 @@
+// Copyright Epic Games, Inc. All Rights Reserved.
+
+#include "Elements/PCGGetSplineControlPoints.h"
+
+#include "Data/PCGPolyLineData.h"
+#include "Data/PCGBasePointData.h"
+
+#include "PCGContext.h"
+#include "Helpers/PCGHelpers.h"
+
+#define LOCTEXT_NAMESPACE "PCGGetSplineControlPointsElement"
+
+#if WITH_EDITOR
+FName UPCGGetSplineControlPointsSettings::GetDefaultNodeName() const
+{
+	return FName(TEXT("GetSplineControlPoints"));
+}
+
+FText UPCGGetSplineControlPointsSettings::GetDefaultNodeTitle() const
+{
+	return LOCTEXT("NodeTitle", "Get Spline Control Points");
+}
+
+FText UPCGGetSplineControlPointsSettings::GetNodeTooltipText() const
+{
+	return LOCTEXT("NodeTooltip", "Extracts the control points from the spline(s) as point data.");
+}
+#endif // WITH_EDITOR
+
+FPCGElementPtr UPCGGetSplineControlPointsSettings::CreateElement() const
+{
+	return MakeShared<FPCGGetSplineControlPointsElement>();
+}
+
+TArray<FPCGPinProperties> UPCGGetSplineControlPointsSettings::InputPinProperties() const
+{
+	TArray<FPCGPinProperties> Properties;
+	Properties.Emplace_GetRef(PCGPinConstants::DefaultInputLabel, EPCGDataType::PolyLine).SetRequiredPin();	
+	return Properties;
+}
+
+TArray<FPCGPinProperties> UPCGGetSplineControlPointsSettings::OutputPinProperties() const
+{ 
+	TArray<FPCGPinProperties> Properties;
+	Properties.Emplace(PCGPinConstants::DefaultOutputLabel, EPCGDataType::Point);	
+	return Properties;
+}
+
+bool FPCGGetSplineControlPointsElement::ExecuteInternal(FPCGContext* InContext) const
+{
+	TRACE_CPUPROFILER_EVENT_SCOPE(FPCGGetSplineControlPointsElement::Execute);
+
+	check(InContext);
+
+	const UPCGGetSplineControlPointsSettings* Settings = InContext->GetInputSettings<UPCGGetSplineControlPointsSettings>();
+	check(Settings);
+
+	for (const FPCGTaggedData& Input : InContext->InputData.GetInputsByPin(PCGPinConstants::DefaultInputLabel))
+	{
+		const UPCGPolyLineData* PolylineData = Cast<const UPCGPolyLineData>(Input.Data);
+
+		if (!PolylineData)
+		{
+			continue;
+		}
+
+		const bool bIsClosed = PolylineData->IsClosed();
+		const int32 NumPoints = bIsClosed ? PolylineData->GetNumSegments() : (PolylineData->GetNumSegments() + 1);
+		const TConstArrayView<PCGMetadataEntryKey> MetadataEntries = PolylineData->GetConstVerticesEntryKeys();
+
+		UPCGBasePointData* PointData = FPCGContext::NewPointData_AnyThread(InContext);	
+		FPCGInitializeFromDataParams Params(PolylineData);
+		Params.DomainMapping = {{PCGMetadataDomainID::Default, PCGMetadataDomainID::Elements}};
+		PointData->InitializeFromDataWithParams(Params);
+		PointData->SetNumPoints(NumPoints);
+		PointData->AllocateProperties(EPCGPointNativeProperties::Transform | EPCGPointNativeProperties::Seed | EPCGPointNativeProperties::MetadataEntry);
+
+		PointData->SetDensity(1.0f);
+		PointData->SetSteepness(1.0f);
+
+		FPCGMetadataDomain* ElementsMetadataDomain = PointData->MutableMetadata()->GetMetadataDomain(PCGMetadataDomainID::Elements);
+		check(ElementsMetadataDomain);
+
+		auto CreateAndValidateAttribute = [InContext, ElementsMetadataDomain]<typename T>(const FName AttributeName, T DefaultValue) -> FPCGMetadataAttribute<T>*
+		{
+			FPCGMetadataAttribute<T>* Attribute = ElementsMetadataDomain->FindOrCreateAttribute<T>(AttributeName, DefaultValue, /*bAllowInterpolation=*/true, /*bOverrideParent=*/true);
+			if (!Attribute)
+			{
+				PCGLog::Metadata::LogFailToCreateAttributeError<T>(AttributeName, InContext);
+			}
+
+			return Attribute;
+		};
+
+		FPCGMetadataAttribute<FVector>* LeaveTangentAttribute = CreateAndValidateAttribute(Settings->LeaveTangentAttributeName, FVector::ZeroVector);
+		FPCGMetadataAttribute<FVector>* ArriveTangentAttribute = CreateAndValidateAttribute(Settings->ArriveTangentAttributeName, FVector::ZeroVector);
+
+		TPCGValueRange<FTransform> PointTransforms = PointData->GetTransformValueRange();
+		TPCGValueRange<int> PointSeeds = PointData->GetSeedValueRange();
+		TPCGValueRange<PCGMetadataEntryKey> PointMetadataEntries = PointData->GetMetadataEntryValueRange();
+
+		for (int32 i = 0; i < NumPoints; ++i)
+		{
+			PointTransforms[i] = PolylineData->GetTransformAtDistance(i, 0);
+			PointSeeds[i] = PCGHelpers::ComputeSeedFromPosition(PointTransforms[i].GetLocation());
+			PointMetadataEntries[i] = MetadataEntries.IsValidIndex(i) ? MetadataEntries[i] : PCGInvalidEntryKey;
+			ElementsMetadataDomain->InitializeOnSet(PointMetadataEntries[i]);
+
+			if (LeaveTangentAttribute || ArriveTangentAttribute)
+			{
+				FVector LeaveTangent, ArriveTangent;
+				PolylineData->GetTangentsAtSegmentStart(i, ArriveTangent, LeaveTangent);
+
+				if (LeaveTangentAttribute)
+				{
+					LeaveTangentAttribute->SetValue(PointMetadataEntries[i], LeaveTangent);
+				}
+
+				if (ArriveTangentAttribute)
+				{
+					ArriveTangentAttribute->SetValue(PointMetadataEntries[i], ArriveTangent);
+				}
+			}
+		}
+
+		InContext->OutputData.TaggedData.Emplace_GetRef(Input).Data = PointData;
+	}
+
+	return true;
+}
+
+#undef LOCTEXT_NAMESPACE
diff --git a/Engine/Plugins/PCG/Source/PCG/Public/Data/PCGPolyLineData.h b/Engine/Plugins/PCG/Source/PCG/Public/Data/PCGPolyLineData.h
--- a/Engine/Plugins/PCG/Source/PCG/Public/Data/PCGPolyLineData.h
+++ b/Engine/Plugins/PCG/Source/PCG/Public/Data/PCGPolyLineData.h
@@ -94,6 +94,8 @@ public:
 	 */
 	UFUNCTION(BlueprintCallable, Category = "PolylineData")
 	virtual void WriteMetadataToPoint(float InputKey, FPCGPoint& OutPoint, UPCGMetadata* OutMetadata) const {}
+
+	virtual TConstArrayView<PCGMetadataEntryKey> GetConstVerticesEntryKeys() const { return {}; }
 };
 
 #undef UE_API
diff --git a/Engine/Plugins/PCG/Source/PCG/Public/Data/PCGSplineData.h b/Engine/Plugins/PCG/Source/PCG/Public/Data/PCGSplineData.h
--- a/Engine/Plugins/PCG/Source/PCG/Public/Data/PCGSplineData.h
+++ b/Engine/Plugins/PCG/Source/PCG/Public/Data/PCGSplineData.h
@@ -83,6 +83,7 @@ public:
 	UE_API virtual FTransform GetTransformAtAlpha(float Alpha) const override;
 	virtual bool IsClosed() const override { return SplineStruct.bClosedLoop; }
 	UE_API virtual void WriteMetadataToPoint(float InputKey, FPCGPoint& OutPoint, UPCGMetadata* OutMetadata) const;
+	virtual TConstArrayView<PCGMetadataEntryKey> GetConstVerticesEntryKeys() const override { return SplineStruct.GetConstControlPointsEntryKeys(); }
 	//~End UPCGPolyLineData interface
 	
 	/** Static helper to create an accessor on a data that doesn't yet exist, as accessors for spline data don't rely on existing data. */
diff --git a/Engine/Plugins/PCG/Source/PCG/Public/Elements/PCGGetSplineControlPoints.h b/Engine/Plugins/PCG/Source/PCG/Public/Elements/PCGGetSplineControlPoints.h
new file mode 100644
--- /dev/null
+++ b/Engine/Plugins/PCG/Source/PCG/Public/Elements/PCGGetSplineControlPoints.h
@@ -0,0 +1,42 @@
+// Copyright Epic Games, Inc. All Rights Reserved.
+
+#pragma once
+
+#include "PCGSettings.h"
+#include "PCGSplineSampler.h"
+
+#include "PCGGetSplineControlPoints.generated.h"
+
+UCLASS(MinimalAPI, BlueprintType, ClassGroup = (Procedural))
+class UPCGGetSplineControlPointsSettings : public UPCGSettings
+{
+	GENERATED_BODY()
+
+public:
+	//~Begin UPCGSettings interface
+#if WITH_EDITOR
+	virtual FName GetDefaultNodeName() const override;
+	virtual FText GetDefaultNodeTitle() const override;
+    virtual FText GetNodeTooltipText() const override;
+	virtual EPCGSettingsType GetType() const override { return EPCGSettingsType::Spatial; }
+#endif
+
+protected:
+	virtual TArray<FPCGPinProperties> InputPinProperties() const override;
+	virtual TArray<FPCGPinProperties> OutputPinProperties() const override;
+	virtual FPCGElementPtr CreateElement() const override;
+	//~End UPCGSettings interface
+
+public:
+	UPROPERTY(EditAnywhere, BlueprintReadWrite, Category="Settings")
+	FName ArriveTangentAttributeName = PCGSplineSamplerConstants::ArriveTangentAttributeName;
+	
+	UPROPERTY(EditAnywhere, BlueprintReadWrite, Category="Settings")
+	FName LeaveTangentAttributeName = PCGSplineSamplerConstants::LeaveTangentAttributeName;
+};
+
+class FPCGGetSplineControlPointsElement : public IPCGElement
+{
+protected:
+	virtual bool ExecuteInternal(FPCGContext* InContext) const override;
+};
diff --git a/Engine/Plugins/PCG/Source/PCG/Public/Elements/PCGSplineSampler.h b/Engine/Plugins/PCG/Source/PCG/Public/Elements/PCGSplineSampler.h
--- a/Engine/Plugins/PCG/Source/PCG/Public/Elements/PCGSplineSampler.h
+++ b/Engine/Plugins/PCG/Source/PCG/Public/Elements/PCGSplineSampler.h
@@ -18,8 +18,15 @@ namespace PCGSplineSamplerConstants
 {
 	const FName SplineLabel = TEXT("Spline");
 	const FName BoundingShapeLabel = TEXT("Bounding Shape");
-	const FName ArriveTangentName = TEXT("ArriveTangent");
-	const FName LeaveTangentName = TEXT("LeaveTangent");
+	const FName NextDirectionDeltaAttributeName = TEXT("NextDirectionDelta");
+	const FName CurvatureAttributeName = TEXT("Curvature");
+	const FName SegmentIndexAttributeName = TEXT("SegmentIndex");
+	const FName SubsegmentIndexAttributeName = TEXT("SubsegmentIndex");
+	const FName ArriveTangentAttributeName = TEXT("ArriveTangent");
+	const FName LeaveTangentAttributeName = TEXT("LeaveTangent");
+	const FName AlphaAttributeName = TEXT("Alpha");
+	const FName DistanceAttributeName = TEXT("Distance");
+	const FName InputKeyAttributeName = TEXT("InputKey");
 }
 
 UENUM()
@@ -140,7 +147,7 @@ struct FPCGSplineSamplerParams
 
 	/** Attribute that will contain the delta angle to the next point on the spline w.r.t to the current's point Up vector. */
 	UPROPERTY(BlueprintReadWrite, EditAnywhere, Category = "Settings|Attributes", meta = (EditCondition = "bComputeDirectionDelta&&Dimension!=EPCGSplineSamplingDimension::OnInterior", EditConditionHides))
-	FName NextDirectionDeltaAttribute = TEXT("NextDirectionDelta");
+	FName NextDirectionDeltaAttribute = PCGSplineSamplerConstants::NextDirectionDeltaAttributeName;
 
 	/** Compute curvature along the spline and write it to an attribute. */
 	UPROPERTY(BlueprintReadWrite, EditAnywhere, Category = "Settings|Attributes", meta = (EditCondition = "Dimension!=EPCGSplineSamplingDimension::OnInterior", EditConditionHides))
@@ -148,7 +155,7 @@ struct FPCGSplineSamplerParams
 
 	/** Attribute that will contain the curvature. Note that the radius of curvature is defined as 1/Curvature, and might need you to scale to world units. */
 	UPROPERTY(BlueprintReadWrite, EditAnywhere, Category = "Settings|Attributes", meta = (EditCondition = "bComputeCurvature&&Dimension!=EPCGSplineSamplingDimension::OnInterior", EditConditionHides))
-	FName CurvatureAttribute = TEXT("Curvature");
+	FName CurvatureAttribute = PCGSplineSamplerConstants::CurvatureAttributeName;
 
 	/** Compute the spline segment index and write it to an attribute. */
 	UPROPERTY(BlueprintReadWrite, EditAnywhere, Category = "Settings|Attributes", meta = (EditCondition = "Dimension!=EPCGSplineSamplingDimension::OnInterior", EditConditionHides))
@@ -156,7 +163,7 @@ struct FPCGSplineSamplerParams
 
 	/** Attribute that will contain the spline segment index. */
 	UPROPERTY(BlueprintReadWrite, EditAnywhere, Category = "Settings|Attributes", meta = (EditCondition = "bComputeSegmentIndex&&Dimension!=EPCGSplineSamplingDimension::OnInterior", EditConditionHides))
-	FName SegmentIndexAttribute = TEXT("SegmentIndex");
+	FName SegmentIndexAttribute = PCGSplineSamplerConstants::SegmentIndexAttributeName;
 
 	/** Compute the sub-segment index of a point on the spline and write it to an attribute. */
 	UPROPERTY(BlueprintReadWrite, EditAnywhere, Category = "Settings|Attributes", meta = (EditCondition = "Mode==EPCGSplineSamplingMode::Subdivision&&Dimension!=EPCGSplineSamplingDimension::OnInterior", EditConditionHides))
@@ -164,7 +171,7 @@ struct FPCGSplineSamplerParams
 
 	/** Attribute that will contain the sub-segment index of a point on the spline. When the sub-segment index is 0, the point is a control point on the actual spline. Only applies to Subdivision mode. */
 	UPROPERTY(BlueprintReadWrite, EditAnywhere, Category = "Settings|Attributes", meta = (EditCondition = "bComputeSubsegmentIndex&&Mode==EPCGSplineSamplingMode::Subdivision&&Dimension!=EPCGSplineSamplingDimension::OnInterior", EditConditionHides))
-	FName SubsegmentIndexAttribute = TEXT("SubsegmentIndex");
+	FName SubsegmentIndexAttribute = PCGSplineSamplerConstants::SubsegmentIndexAttributeName;
 
 	/** Compute arrive and leave tangents along the spline and write them to attributes. */
 	UPROPERTY(BlueprintReadWrite, EditAnywhere, Category = "Settings|Attributes", meta = (EditCondition = "Dimension!=EPCGSplineSamplingDimension::OnInterior", EditConditionHides))
@@ -172,11 +179,11 @@ struct FPCGSplineSamplerParams
 
 	/** Attribute that will contain the arrive tangent vector. For control points, this will be the actual arrive tangent. For non-control points, this will only be the normalized tangent at this point. */
 	UPROPERTY(BlueprintReadWrite, EditAnywhere, Category = "Settings|Attributes", meta = (EditCondition = "bComputeTangents&&Dimension!=EPCGSplineSamplingDimension::OnInterior", EditConditionHides))
-	FName ArriveTangentAttribute = PCGSplineSamplerConstants::ArriveTangentName;
+	FName ArriveTangentAttribute = PCGSplineSamplerConstants::ArriveTangentAttributeName;
 
 	/** Attribute that will contain the leave tangent vector. For control points, this will be the actual leave tangent. For non-control points, this will only be the normalized tangent at this point. */
 	UPROPERTY(BlueprintReadWrite, EditAnywhere, Category = "Settings|Attributes", meta = (EditCondition = "bComputeTangents&&Dimension!=EPCGSplineSamplingDimension::OnInterior", EditConditionHides))
-	FName LeaveTangentAttribute = PCGSplineSamplerConstants::LeaveTangentName;
+	FName LeaveTangentAttribute = PCGSplineSamplerConstants::LeaveTangentAttributeName;
 
 	/** Compute an alpha value along the spline and write it to an attribute. */
 	UPROPERTY(BlueprintReadWrite, EditAnywhere, Category = "Settings|Attributes", meta = (EditCondition = "Dimension!=EPCGSplineSamplingDimension::OnInterior", EditConditionHides))
@@ -187,7 +194,7 @@ struct FPCGSplineSamplerParams
 	 * For example, if there are three segments, each segment will take up 0.333... of the interval.
 	 */
 	UPROPERTY(BlueprintReadWrite, EditAnywhere, Category = "Settings|Attributes", meta = (EditCondition = "bComputeAlpha&&Dimension!=EPCGSplineSamplingDimension::OnInterior", EditConditionHides))
-	FName AlphaAttribute = TEXT("Alpha");
+	FName AlphaAttribute = PCGSplineSamplerConstants::AlphaAttributeName;
 
 	/** Compute distance along the spline and write it to an attribute. */
 	UPROPERTY(BlueprintReadWrite, EditAnywhere, Category = "Settings|Attributes", meta = (EditCondition = "Dimension!=EPCGSplineSamplingDimension::OnInterior", EditConditionHides))
@@ -195,7 +202,7 @@ struct FPCGSplineSamplerParams
 
 	/** Attribute that will contain the distance along the spline at the sample point. */
 	UPROPERTY(BlueprintReadWrite, EditAnywhere, Category = "Settings|Attributes", meta = (EditCondition = "bComputeDistance&&Dimension!=EPCGSplineSamplingDimension::OnInterior", EditConditionHides))
-	FName DistanceAttribute = TEXT("Distance");
+	FName DistanceAttribute = PCGSplineSamplerConstants::DistanceAttributeName;
 
 	/** Compute an input key for each point along the spline and write the key to an attribute. */
 	UPROPERTY(BlueprintReadWrite, EditAnywhere, Category = "Settings|Attributes", meta = (EditCondition = "Dimension!=EPCGSplineSamplingDimension::OnInterior", EditConditionHides))
@@ -203,7 +210,7 @@ struct FPCGSplineSamplerParams
 
 	/** Attribute that will contain the spline input key, a float value between [0, N], where N is the number of control points. Each range [i, i+1] represents an interpolation from 0 to 1 across spline segment i. */
 	UPROPERTY(BlueprintReadWrite, EditAnywhere, Category = "Settings|Attributes", meta = (EditCondition = "bComputeInputKey&&Dimension!=EPCGSplineSamplingDimension::OnInterior", EditConditionHides))
-	FName InputKeyAttribute = TEXT("InputKey");
+	FName InputKeyAttribute = PCGSplineSamplerConstants::InputKeyAttributeName;
 
 	/** If no Bounding Shape input is provided, the actor bounds are used to limit the sample generation domain.
 	* This option allows ignoring the actor bounds and generating over the entire spline. Use with caution as this
