[PCG] Retain world streaming query result cache across frames, invalidate on streaming events.

Massively reduces number of executed queries. Also tried finer-grained invalidation approach which intersects WP cell bounds with cached queries but failed in some cases, will revisit in future with better debugging vis/tools.

#rnx

diff --git a/Engine/Plugins/PCG/Source/PCG/Private/RuntimeGen/PCGRuntimeGenScheduler.cpp b/Engine/Plugins/PCG/Source/PCG/Private/RuntimeGen/PCGRuntimeGenScheduler.cpp
--- a/Engine/Plugins/PCG/Source/PCG/Private/RuntimeGen/PCGRuntimeGenScheduler.cpp
+++ b/Engine/Plugins/PCG/Source/PCG/Private/RuntimeGen/PCGRuntimeGenScheduler.cpp
@@ -20,7 +20,9 @@
 
 #include "DrawDebugHelpers.h"
 #include "Async/ParallelFor.h"
+#include "Engine/LevelStreaming.h"
 #include "Components/RuntimeVirtualTextureComponent.h"
+#include "Streaming/LevelStreamingDelegates.h"
 #include "UObject/UObjectGlobals.h"
 #include "VT/RuntimeVirtualTexture.h"
 #include "WorldPartition/WorldPartition.h"
@@ -124,12 +126,16 @@ FPCGRuntimeGenScheduler::FPCGRuntimeGenScheduler(UWorld* InWorld, FPCGActorAndCo
 	bPoolingWasEnabledLastFrame = PCGRuntimeGenSchedulerHelpers::CVarRuntimeGenerationEnablePooling.GetValueOnAnyThread();
 	BasePoolSizeLastFrame = PCGRuntimeGenSchedulerHelpers::CVarRuntimeGenerationBasePoolSize.GetValueOnAnyThread();
 	FramesUntilGeneration = PCGRuntimeGenSchedulerHelpers::CVarFramesBeforeFirstGenerate.GetValueOnGameThread();
+
+	FLevelStreamingDelegates::OnLevelStreamingStateChanged.AddRaw(this, &FPCGRuntimeGenScheduler::OnLevelStreamingStateChanged);
 }
 
 FPCGRuntimeGenScheduler::~FPCGRuntimeGenScheduler()
 {
 	delete GenSourceManager;
 	GenSourceManager = nullptr;
+
+	FLevelStreamingDelegates::OnLevelStreamingStateChanged.RemoveAll(this);
 }
 
 void FPCGRuntimeGenScheduler::Tick(APCGWorldActor* InPCGWorldActor)
@@ -314,18 +320,22 @@ void FPCGRuntimeGenScheduler::TickQueueComponentsForGeneration(
 
 	// Prepare streaming queries up front.
 	TArray<FWorldPartitionStreamingQuerySource> StreamingQuerySources;
-	UWorldPartitionSubsystem* WorldPartitionSubsystem = InPCGWorldActor ? UWorld::GetSubsystem<UWorldPartitionSubsystem>(InPCGWorldActor->GetWorld()) : nullptr;
+	StreamingQuerySources.Reserve(1);
+	
+	bool bCheckStreaming = PCGRuntimeGenSchedulerHelpers::CVarCheckWorldStreamedIn.GetValueOnGameThread();
+	const UWorldPartitionSubsystem* WorldPartitionSubsystem = (bCheckStreaming && InPCGWorldActor) ? UWorld::GetSubsystem<UWorldPartitionSubsystem>(InPCGWorldActor->GetWorld()) : nullptr;
 
-	const bool bCheckStreaming = PCGRuntimeGenSchedulerHelpers::CVarCheckWorldStreamedIn.GetValueOnGameThread();
-	if (bCheckStreaming)
+	if (WorldPartitionSubsystem)
 	{
 		FWorldPartitionStreamingQuerySource& QuerySource = StreamingQuerySources.Emplace_GetRef();
 		QuerySource.bSpatialQuery = true;
 		QuerySource.bUseGridLoadingRange = false;
 		QuerySource.bDataLayersOnly = false;
 	}
-
-	CachedStreamingQueryResults.Empty(CachedStreamingQueryResults.Num());
+	else
+	{
+		bCheckStreaming = false;
+	}
 
 	auto IsWorldStreamingComplete = [WorldPartitionSubsystem, &StreamingQuerySources, &CachedStreamingQueryResults=CachedStreamingQueryResults](const FVector& InLocation, float InGridSize)
 	{
@@ -1242,6 +1252,15 @@ void FPCGRuntimeGenScheduler::CleanupLocalComponents(const APCGWorldActor* InPCG
 	}
 }
 
+void FPCGRuntimeGenScheduler::OnLevelStreamingStateChanged(UWorld* InWorld, const ULevelStreaming* InLevelStreaming, ULevel* InLevelIfLoaded, ELevelStreamingState InPreviousState, ELevelStreamingState InNewState)
+{
+	if (World == InWorld && (InPreviousState == ELevelStreamingState::LoadedVisible || InNewState == ELevelStreamingState::LoadedVisible))
+	{
+		// todo_pcg: Fine-grained invalidation based on bounds overlap tests did not trivially work on first attempt, improve debug tools/vis and retry.
+		CachedStreamingQueryResults.Empty(CachedStreamingQueryResults.Num());
+	}
+}
+
 void FPCGRuntimeGenScheduler::CleanupLocalComponent(APCGPartitionActor* PartitionActor, UPCGComponent* LocalComponent)
 {
 	if (!PartitionActor)
diff --git a/Engine/Plugins/PCG/Source/PCG/Public/RuntimeGen/PCGRuntimeGenScheduler.h b/Engine/Plugins/PCG/Source/PCG/Public/RuntimeGen/PCGRuntimeGenScheduler.h
--- a/Engine/Plugins/PCG/Source/PCG/Public/RuntimeGen/PCGRuntimeGenScheduler.h
+++ b/Engine/Plugins/PCG/Source/PCG/Public/RuntimeGen/PCGRuntimeGenScheduler.h
@@ -14,9 +14,13 @@ class APCGWorldActor;
 class FPCGActorAndComponentMapping;
 class FPCGGenSourceManager;
 class IPCGGenSourceBase;
+class ULevel;
+class ULevelStreaming;
 class UPCGComponent;
 class UPCGSubsystem;
 class UWorld;
+enum class ELevelStreamingState : uint8;
+enum class ELevelStreamingTargetState : uint8;
 
 /** Used to inform what virtual textures to prime and on what grids they need to be present. */
 USTRUCT(BlueprintType)
@@ -133,6 +137,9 @@ protected:
 	void AddReferencedObjects(FReferenceCollector& Collector);
 
 private:
+	/** Called on world streaming events. */
+	void OnLevelStreamingStateChanged(UWorld* InWorld, const ULevelStreaming* InLevelStreaming, ULevel* InLevelIfLoaded, ELevelStreamingState InPreviousState, ELevelStreamingState InNewState);
+
 	/** Cleanup a local compnent from the Partition actor. */
 	void CleanupLocalComponent(APCGPartitionActor* PartitionActor, UPCGComponent* LocalComponent);
 
