[PCG] Hoist time budget up to subsystem level and use the end time in the RGS.

Currently the RGS budget is only applied to cleanup because time slice latency in cleanup is hopefully fairly absorbable/doesn't cause noticeable issues. At least one component is always cleaned up to ensure the system always makes progress.

With these, the subsystem cost stays relatively bounded in my tests (not precisely within the budget as there are still excluded costs, but mostly within 20% of budget, and this was without throttling the number of generating components).

#rnx

diff --git a/Engine/Plugins/PCG/Source/PCG/Private/Graph/PCGGraphExecutor.cpp b/Engine/Plugins/PCG/Source/PCG/Private/Graph/PCGGraphExecutor.cpp
--- a/Engine/Plugins/PCG/Source/PCG/Private/Graph/PCGGraphExecutor.cpp
+++ b/Engine/Plugins/PCG/Source/PCG/Private/Graph/PCGGraphExecutor.cpp
@@ -1460,7 +1460,7 @@ void FPCGGraphExecutor::DebugDisplayTask(TWeakPtr<FPCGGraphExecutor::FGameThread
 }
 #endif // WITH_EDITOR
 
-double FPCGGraphExecutor::GetTickBudgetInSeconds() const
+double FPCGGraphExecutor::GetTickBudgetInSeconds()
 {
 #if WITH_EDITOR
 	if (GEditor && !GEditor->IsPlaySessionInProgress())
@@ -1736,6 +1736,30 @@ bool FPCGGraphExecutor::ExecuteScheduling(double EndTime, TSharedPtr<FPCGGraphAc
 	return bStateChanged;
 }
 
+void FPCGGraphExecutor::Execute()
+{
+	double EndTime = FPlatformTime::Seconds() + FPCGGraphExecutor::GetTickBudgetInSeconds();
+	Execute(EndTime);
+}
+
+struct FPCGExcludeScopeFromTickBudget
+{
+	FPCGExcludeScopeFromTickBudget(double& InEndTime)
+		: EndTime(InEndTime)
+	{
+		StartTime = FPlatformTime::Seconds();
+	}
+
+	~FPCGExcludeScopeFromTickBudget()
+	{
+		EndTime += FPlatformTime::Seconds() - StartTime;
+	}
+
+private:
+	double& EndTime;
+	double StartTime;
+};
+
 /**
  * Graph Execution Pipeline:
  * 
@@ -1779,47 +1803,50 @@ bool FPCGGraphExecutor::ExecuteScheduling(double EndTime, TSharedPtr<FPCGGraphAc
  *   pcg.Async.OutOfTickBudgetInMilliseconds : Timeslicing budget for UE::Tasks launched by [FPCGGraphExecutor::ExecuteScheduling]
  *   pcg.DisableAsyncTimeSlicingOnGameThread : Disables timeslicing on GameThread execution
  */
-void FPCGGraphExecutor::Execute()
+void FPCGGraphExecutor::Execute(double& InOutEndTime)
 {
 	TRACE_CPUPROFILER_EVENT_SCOPE(FPCGGraphExecutor::Execute);
 
 #if WITH_EDITOR
-	const bool bIsEditorWorld = !World || !World->IsGameWorld();
 
-	if (PCGSystemSwitches::CVarPausePCGExecution.GetValueOnAnyThread() || (bIsEditorWorld && GEditor && GEditor->IsTransactionActive() && PCGGraphExecutor::CVarPausePCGExecutionWhileTransactionActive.GetValueOnAnyThread()))
 	{
-		UpdateGenerationNotification();
-		return;
+		FPCGExcludeScopeFromTickBudget Exclusion(InOutEndTime);
+
+		const bool bIsEditorWorld = !World || !World->IsGameWorld();
+
+		if (PCGSystemSwitches::CVarPausePCGExecution.GetValueOnAnyThread() || (bIsEditorWorld && GEditor && GEditor->IsTransactionActive() && PCGGraphExecutor::CVarPausePCGExecutionWhileTransactionActive.GetValueOnAnyThread()))
+		{
+			UpdateGenerationNotification();
+			return;
+		}
 	}
 #endif // WITH_EDITOR
 
 	const double StartTime = FPlatformTime::Seconds();
+	const double Budget = FMath::Max(0.0, InOutEndTime - StartTime);
 
 	const double ProcessScheduledTasksBudgetPercentage = 0.3;
-	const double ProcessScheduledTasksEndTime = StartTime + GetTickBudgetInSeconds() * ProcessScheduledTasksBudgetPercentage;
+	const double ProcessScheduledTasksEndTime = StartTime + Budget * ProcessScheduledTasksBudgetPercentage;
 
 	ProcessScheduledTasks(ProcessScheduledTasksEndTime);
 
 	const double ProcessScheduledTasksTime = FPlatformTime::Seconds() - StartTime;
 	
-	double NotificationTime = 0.0;
 #if WITH_EDITOR
-	// Update Notifications before capturing StartTime so that this call doesn't eat up task budget
-	UpdateGenerationNotification();
+	{
+		FPCGExcludeScopeFromTickBudget Exclusion(InOutEndTime);
 
-	// Remove notifications from budget
-	NotificationTime = FPlatformTime::Seconds() - StartTime - ProcessScheduledTasksTime;
+		UpdateGenerationNotification();
+	}
 #endif // WITH_EDITOR
 				
 	bool bContinueExecute = GetNonScheduledRemainingTaskCount() > 0;
 	bool bFirstLoop = true;
 
-	const double EndTime = StartTime + NotificationTime + GetTickBudgetInSeconds() - ProcessScheduledTasksTime;
-
-	while ((bFirstLoop || (FPlatformTime::Seconds() < EndTime)) && bContinueExecute)
+	while ((bFirstLoop || (FPlatformTime::Seconds() < InOutEndTime)) && bContinueExecute)
 	{
 		TSharedPtr<FPCGGraphActiveTask> MainThreadTask;
-		bContinueExecute = ExecuteScheduling(EndTime, &MainThreadTask, /*bForceCheckPausedTasks=*/bFirstLoop);
+		bContinueExecute = ExecuteScheduling(InOutEndTime, &MainThreadTask, /*bForceCheckPausedTasks=*/bFirstLoop);
 		bFirstLoop = false;
 
 		if (MainThreadTask)
diff --git a/Engine/Plugins/PCG/Source/PCG/Private/Graph/PCGGraphExecutor.h b/Engine/Plugins/PCG/Source/PCG/Private/Graph/PCGGraphExecutor.h
--- a/Engine/Plugins/PCG/Source/PCG/Private/Graph/PCGGraphExecutor.h
+++ b/Engine/Plugins/PCG/Source/PCG/Private/Graph/PCGGraphExecutor.h
@@ -134,6 +134,7 @@ public:
 
 	/** "Tick" of the graph executor. This call is NOT THREADSAFE */
 	void Execute();
+	void Execute(double& InOutEndTime);
 
 	/** Expose cache so it can be dirtied */
 	FPCGGraphCache& GetCache() { return GraphCache; }
@@ -183,7 +184,7 @@ private:
 	void SetLandscapeHeightPCGData(FPCGTaskId InGraphExecutionTaskId, UPCGData* InData) { SetExecutionCacheData(InGraphExecutionTaskId, EExecutionCacheDataType::LandscapeHeightData, InData); }
 	void SetOriginalActorPCGData(FPCGTaskId InGraphExecutionTaskId, UPCGData* InData) { SetExecutionCacheData(InGraphExecutionTaskId, EExecutionCacheDataType::OriginalActorData, InData); }
 
-	double GetTickBudgetInSeconds() const;
+	static double GetTickBudgetInSeconds();
 		
 	struct FCachedResult
 	{
diff --git a/Engine/Plugins/PCG/Source/PCG/Private/PCGSubsystem.cpp b/Engine/Plugins/PCG/Source/PCG/Private/PCGSubsystem.cpp
--- a/Engine/Plugins/PCG/Source/PCG/Private/PCGSubsystem.cpp
+++ b/Engine/Plugins/PCG/Source/PCG/Private/PCGSubsystem.cpp
@@ -300,6 +300,7 @@ bool UPCGSubsystem::RemoveAndCopyConstructionScriptSourceComponent(AActor* InCom
 
 void UPCGSubsystem::Tick(float DeltaSeconds)
 {
+	TRACE_CPUPROFILER_EVENT_SCOPE(UPCGSubsystem::Tick);
 	Super::Tick(DeltaSeconds);
 
 	ExecuteBeginTickActions();
@@ -327,14 +328,16 @@ void UPCGSubsystem::Tick(float DeltaSeconds)
 		GetLandscapeCache()->Tick(DeltaSeconds);
 	}
 	
+	double EndTime = FPlatformTime::Seconds() + FPCGGraphExecutor::GetTickBudgetInSeconds();
+
 	// If we have any tasks to execute, schedule some
-	GraphExecutor->Execute();
+	GraphExecutor->Execute(EndTime);
 
 	ActorAndComponentMapping.Tick();
 
 	if (PCGWorldActor)
 	{
-		RuntimeGenScheduler->Tick(PCGWorldActor);
+		RuntimeGenScheduler->Tick(PCGWorldActor, EndTime);
 	}
 }
 
diff --git a/Engine/Plugins/PCG/Source/PCG/Private/RuntimeGen/PCGRuntimeGenScheduler.cpp b/Engine/Plugins/PCG/Source/PCG/Private/RuntimeGen/PCGRuntimeGenScheduler.cpp
--- a/Engine/Plugins/PCG/Source/PCG/Private/RuntimeGen/PCGRuntimeGenScheduler.cpp
+++ b/Engine/Plugins/PCG/Source/PCG/Private/RuntimeGen/PCGRuntimeGenScheduler.cpp
@@ -138,7 +138,7 @@ FPCGRuntimeGenScheduler::~FPCGRuntimeGenScheduler()
 	FLevelStreamingDelegates::OnLevelStreamingStateChanged.RemoveAll(this);
 }
 
-void FPCGRuntimeGenScheduler::Tick(APCGWorldActor* InPCGWorldActor)
+void FPCGRuntimeGenScheduler::Tick(APCGWorldActor* InPCGWorldActor, double InEndTime)
 {
 	TRACE_CPUPROFILER_EVENT_SCOPE(FPCGRuntimeGenScheduler::Tick);
 	check(InPCGWorldActor && GenSourceManager);
@@ -207,7 +207,7 @@ void FPCGRuntimeGenScheduler::Tick(APCGWorldActor* InPCGWorldActor)
 
 	if (!GeneratedComponents.IsEmpty())
 	{
-		TickCleanup(GenSources, InPCGWorldActor);
+		TickCleanup(GenSources, InPCGWorldActor, InEndTime);
 	}
 
 	// 3. Schedule generation on components in priority order.
@@ -216,6 +216,8 @@ void FPCGRuntimeGenScheduler::Tick(APCGWorldActor* InPCGWorldActor)
 		// Sort components by priority (will be generated in descending order).
 		ComponentsToGenerate.ValueSort([](double PrioA, double PrioB)->bool { return PrioA > PrioB; });
 
+		// Only apply time budget to cleanup currently. Currently too easy to introduce latency issues so don't hold back generation new components
+		// (and we already have CVarNumGeneratingComponentsAtSameTime to throttle new generations).
 		TickScheduleGeneration(ComponentsToGenerate);
 	}
 }
@@ -584,7 +586,7 @@ void FPCGRuntimeGenScheduler::TickQueueComponentsForGeneration(
 	}
 }
 
-void FPCGRuntimeGenScheduler::TickCleanup(const TSet<IPCGGenSourceBase*>& InGenSources, const APCGWorldActor* InPCGWorldActor)
+void FPCGRuntimeGenScheduler::TickCleanup(const TSet<IPCGGenSourceBase*>& InGenSources, const APCGWorldActor* InPCGWorldActor, double InEndTime)
 {
 	TRACE_CPUPROFILER_EVENT_SCOPE(FPCGRuntimeGenScheduler::TickCleanup);
 
@@ -726,6 +728,12 @@ void FPCGRuntimeGenScheduler::TickCleanup(const TSet<IPCGGenSourceBase*>& InGenS
 	{
 		const PCGComponentToClean& ComponentToClean = ComponentsToClean[I];
 		CleanupComponent(ComponentToClean.Get<0>(), ComponentToClean.Get<1>());
+
+		if (FPlatformTime::Seconds() >= InEndTime)
+		{
+			UE_LOG(LogPCG, Verbose, TEXT("FPCGRuntimeGenScheduler: Time budget exceeded, aborted after cleaning up %d / %d components"), (I + 1), NumComponentsToClean);
+			break;
+		}
 	}
 }
 
diff --git a/Engine/Plugins/PCG/Source/PCG/Public/RuntimeGen/PCGRuntimeGenScheduler.h b/Engine/Plugins/PCG/Source/PCG/Public/RuntimeGen/PCGRuntimeGenScheduler.h
--- a/Engine/Plugins/PCG/Source/PCG/Public/RuntimeGen/PCGRuntimeGenScheduler.h
+++ b/Engine/Plugins/PCG/Source/PCG/Public/RuntimeGen/PCGRuntimeGenScheduler.h
@@ -66,7 +66,7 @@ public:
 	FPCGRuntimeGenScheduler& operator=(const FPCGRuntimeGenScheduler& other) = delete;
 	FPCGRuntimeGenScheduler& operator=(FPCGRuntimeGenScheduler&& other) = delete; 
 
-	void Tick(APCGWorldActor* InPCGWorldActor);
+	void Tick(APCGWorldActor* InPCGWorldActor, double InEndTime);
 
 	void OnOriginalComponentRegistered(UPCGComponent* InOriginalComponent);
 	void OnOriginalComponentUnregistered(UPCGComponent* InOriginalComponent);
@@ -99,7 +99,7 @@ protected:
 		TMap<FGridGenerationKey, double>& OutComponentsToGenerate);
 
 	/** Perform immediate cleanup on components that become out of range. */
-	void TickCleanup(const TSet<IPCGGenSourceBase*>& GenSources, const APCGWorldActor* InPCGWorldActor);
+	void TickCleanup(const TSet<IPCGGenSourceBase*>& GenSources, const APCGWorldActor* InPCGWorldActor, double InEndTime);
 
 	/** Schedule generation on components in priority order. */
 	void TickScheduleGeneration(TMap<FGridGenerationKey, double>& ComponentsToGenerate);
