[PCG] Misc RGS optimizations from last round of deep profiling.

#rnx

diff --git a/Engine/Plugins/PCG/Source/PCG/Private/Grid/PCGPartitionActor.cpp b/Engine/Plugins/PCG/Source/PCG/Private/Grid/PCGPartitionActor.cpp
--- a/Engine/Plugins/PCG/Source/PCG/Private/Grid/PCGPartitionActor.cpp
+++ b/Engine/Plugins/PCG/Source/PCG/Private/Grid/PCGPartitionActor.cpp
@@ -695,6 +695,11 @@ bool APCGPartitionActor::IsSafeForDeletion() const
 	return true;
 }
 
+bool APCGPartitionActor::HasLocalPCGComponents() const
+{
+	return !LocalToOriginal.IsEmpty();
+}
+
 TSet<TObjectPtr<UPCGComponent>> APCGPartitionActor::GetAllLocalPCGComponents() const
 {
 	TSet<TObjectPtr<UPCGComponent>> ResultComponents;
diff --git a/Engine/Plugins/PCG/Source/PCG/Private/PCGGraph.cpp b/Engine/Plugins/PCG/Source/PCG/Private/PCGGraph.cpp
--- a/Engine/Plugins/PCG/Source/PCG/Private/PCGGraph.cpp
+++ b/Engine/Plugins/PCG/Source/PCG/Private/PCGGraph.cpp
@@ -793,8 +793,14 @@ void UPCGGraph::AddReferencedObjects(UObject* InThis, FReferenceCollector& Colle
 
 uint32 UPCGGraph::GetDefaultGridSize() const
 {
-	ensure(IsHierarchicalGenerationEnabled());
-	return PCGHiGenGrid::IsValidGrid(HiGenGridSize) ? (PCGHiGenGrid::GridToGridSize(HiGenGridSize) * (1 << HiGenExponential)) : PCGHiGenGrid::UnboundedGridSize();
+	if (IsHierarchicalGenerationEnabled() && PCGHiGenGrid::IsValidGrid(HiGenGridSize))
+	{
+		return PCGHiGenGrid::GridToGridSize(HiGenGridSize) * (1 << HiGenExponential);
+	}
+	else
+	{
+		return PCGHiGenGrid::UnboundedGridSize();
+	}
 }
 
 TSharedPtr<FComputeGraphInstance> UPCGGraph::RetrieveComputeGraphInstanceFromPool(const FComputeGraphInstanceKey& InKey, bool& bOutNewInstance) const
@@ -1358,50 +1364,66 @@ void UPCGGraph::PostNodeUndo(UPCGNode* InPCGNode)
 
 void UPCGGraph::GetGridSizes(PCGHiGenGrid::FSizeArray& OutGridSizes, bool& bOutHasUnbounded) const
 {
-	bOutHasUnbounded = HiGenGridSize == EPCGHiGenGrid::Unbounded;
+	FScopeLock Lock(&CachedGridInfoLock);
+
+	if (!CachedGridInfo)
+	{
+		CacheGridSizesInternalNoLock();
+	}
+
+	OutGridSizes = CachedGridInfo->GridSizes;
+	bOutHasUnbounded = CachedGridInfo->bHasUnbounded;
+}
+
+void UPCGGraph::CacheGridSizesInternalNoLock() const
+{
+	FGridInfo GridInfo;
+
+	GridInfo.bHasUnbounded = HiGenGridSize == EPCGHiGenGrid::Unbounded;
 
 	const uint32 GraphDefaultGridSize = GetDefaultGridSize();
 	if (!IsHierarchicalGenerationEnabled())
 	{
 		if (PCGHiGenGrid::IsValidGridSize(GetDefaultGridSize()))
 		{
-			OutGridSizes.Add(GraphDefaultGridSize);
+			GridInfo.GridSizes.Add(GraphDefaultGridSize);
 		}
-		return;
 	}
-
-	bool bHasUninitialized = false;
-	for (const UPCGNode* Node : Nodes)
+	else
 	{
-		const uint32 GridSize = GetNodeGenerationGridSize(Node, GraphDefaultGridSize);
-		if (PCGHiGenGrid::IsValidGridSize(GridSize))
+		bool bHasUninitialized = false;
+		for (const UPCGNode* Node : Nodes)
 		{
-			if (!OutGridSizes.Contains(GridSize))
+			const uint32 GridSize = GetNodeGenerationGridSize(Node, GraphDefaultGridSize);
+			if (PCGHiGenGrid::IsValidGridSize(GridSize))
 			{
-				OutGridSizes.Add(GridSize);
+				if (!GridInfo.GridSizes.Contains(GridSize))
+				{
+					GridInfo.GridSizes.Add(GridSize);
+				}
+			}
+			else if (GridSize == PCGHiGenGrid::UnboundedGridSize())
+			{
+				GridInfo.bHasUnbounded = true;
+			}
+			else if (GridSize == PCGHiGenGrid::UninitializedGridSize())
+			{
+				// Outside nodes will not have a concrete grid set
+				bHasUninitialized = true;
 			}
 		}
-		else if (GridSize == PCGHiGenGrid::UnboundedGridSize())
-		{
-			bOutHasUnbounded = true;
-		}
-		else if (GridSize == PCGHiGenGrid::UninitializedGridSize())
+
+		if (bHasUninitialized)
 		{
-			// Outside nodes will not have a concrete grid set
-			bHasUninitialized = true;
+			// Nodes outside grid ranges will execute at graph default
+			GridInfo.GridSizes.Add(GraphDefaultGridSize);
 		}
-	}
 
-	if (bHasUninitialized)
-	{
-		// Nodes outside grid ranges will execute at graph default
-		OutGridSizes.Add(GraphDefaultGridSize);
+		// Descending
+		GridInfo.GridSizes.Sort([](const uint32& A, const uint32& B) { return A > B; });
 	}
 
-	// Descending
-	OutGridSizes.Sort([](const uint32& A, const uint32& B) { return A > B; });
-
-	return;
+	CachedGridInfo = MoveTemp(GridInfo);
 }
 
 double UPCGGraph::GetGridGenerationRadiusFromGrid(EPCGHiGenGrid Grid) const
@@ -1630,6 +1652,12 @@ void UPCGGraph::NotifyGraphChanged(EPCGChangeType ChangeType)
 {
 	TRACE_CPUPROFILER_EVENT_SCOPE(UPCGGraph::NotifyGraphChanged);
 
+	if (!!(ChangeType & EPCGChangeType::Structural))
+	{
+		FScopeLock Lock(&CachedGridInfoLock);
+		CachedGridInfo.Reset();
+	}
+
 	if (GraphChangeNotificationsDisableCounter > 0)
 	{
 		bDelayedChangeNotification = true;
diff --git a/Engine/Plugins/PCG/Source/PCG/Private/RuntimeGen/PCGRuntimeGenScheduler.cpp b/Engine/Plugins/PCG/Source/PCG/Private/RuntimeGen/PCGRuntimeGenScheduler.cpp
--- a/Engine/Plugins/PCG/Source/PCG/Private/RuntimeGen/PCGRuntimeGenScheduler.cpp
+++ b/Engine/Plugins/PCG/Source/PCG/Private/RuntimeGen/PCGRuntimeGenScheduler.cpp
@@ -152,7 +152,7 @@ void FPCGRuntimeGenScheduler::Tick(APCGWorldActor* InPCGWorldActor, double InEnd
 
 	TickCVars(InPCGWorldActor);
 
-	TSet<IPCGGenSourceBase*> GenSources;
+	GenSources.Empty(GenSources.Num());
 
 	if (bAnyRuntimeGenComponentsExist)
 	{
@@ -197,10 +197,18 @@ void FPCGRuntimeGenScheduler::Tick(APCGWorldActor* InPCGWorldActor, double InEnd
 	// 1. Queue nearby components for generation.
 	
 	// Mapping of component + coordinates to priorities - needed to compute max priority over all gen sources.
-	TMap<FGridGenerationKey, double> ComponentsToGenerate;
+	ComponentsToGenerate.Empty(128);
+
 	if (!GenSources.IsEmpty())
 	{
-		TickQueueComponentsForGeneration(GenSources, InPCGWorldActor, ComponentsToGenerate);
+		FTickQueueComponentsForGenerationInputs Inputs;
+		Inputs.GenSources = &GenSources;
+		Inputs.PCGWorldActor = InPCGWorldActor;
+		Inputs.AllPartitionedComponents = ActorAndComponentMapping->GetAllRegisteredPartitionedComponents();
+		Inputs.AllNonPartitionedComponents = ActorAndComponentMapping->GetAllRegisteredNonPartitionedComponents();
+		Inputs.GeneratedComponents = GeneratedComponents;
+
+		TickQueueComponentsForGeneration(Inputs, ComponentsToGenerate);
 	}
 
 	// 2. Schedule cleanup on components that become out of range.
@@ -284,13 +292,12 @@ bool FPCGRuntimeGenScheduler::ShouldTick()
 }
 
 void FPCGRuntimeGenScheduler::TickQueueComponentsForGeneration(
-	const TSet<IPCGGenSourceBase*>& InGenSources,
-	APCGWorldActor* InPCGWorldActor,
+	const FTickQueueComponentsForGenerationInputs& Inputs,
 	TMap<FGridGenerationKey, double>& OutComponentsToGenerate)
 {
 	TRACE_CPUPROFILER_EVENT_SCOPE(FPCGRuntimeGenScheduler::TickQueueComponentsForGeneration);
 
-	check(ActorAndComponentMapping && InPCGWorldActor);
+	check(Inputs.PCGWorldActor && Inputs.GenSources);
 
 	// TODO: Thought - it would be possible to maintain a global maximum generation distance across all components
 	// perhaps in the actor&comp mapping system, and then do a spatial query to get the components here.
@@ -325,7 +332,7 @@ void FPCGRuntimeGenScheduler::TickQueueComponentsForGeneration(
 	StreamingQuerySources.Reserve(1);
 	
 	bool bCheckStreaming = PCGRuntimeGenSchedulerHelpers::CVarCheckWorldStreamedIn.GetValueOnGameThread();
-	const UWorldPartitionSubsystem* WorldPartitionSubsystem = (bCheckStreaming && InPCGWorldActor) ? UWorld::GetSubsystem<UWorldPartitionSubsystem>(InPCGWorldActor->GetWorld()) : nullptr;
+	const UWorldPartitionSubsystem* WorldPartitionSubsystem = bCheckStreaming ? UWorld::GetSubsystem<UWorldPartitionSubsystem>(Inputs.PCGWorldActor->GetWorld()) : nullptr;
 
 	if (WorldPartitionSubsystem)
 	{
@@ -365,7 +372,7 @@ void FPCGRuntimeGenScheduler::TickQueueComponentsForGeneration(
 	};
 
 	// Collect local components from all partitioned components.
-	for (UPCGComponent* OriginalComponent : ActorAndComponentMapping->PartitionedOctree.GetAllComponents())
+	for (UPCGComponent* OriginalComponent : Inputs.AllPartitionedComponents)
 	{
 		TRACE_CPUPROFILER_EVENT_SCOPE(FPCGRuntimeGenScheduler::CollectLocalComponents);
 
@@ -382,7 +389,7 @@ void FPCGRuntimeGenScheduler::TickQueueComponentsForGeneration(
 			ensure(Policy);
 			bool bHasUnbounded = false;
 			PCGHiGenGrid::FSizeArray GridSizes;
-			ensure(PCGHelpers::GetGenerationGridSizes(OriginalComponent->GetGraph(), InPCGWorldActor, GridSizes, bHasUnbounded));
+			ensure(PCGHelpers::GetGenerationGridSizes(OriginalComponent->GetGraph(), Inputs.PCGWorldActor, GridSizes, bHasUnbounded));
 
 			if (GridSizes.IsEmpty() && !bHasUnbounded)
 			{
@@ -392,7 +399,7 @@ void FPCGRuntimeGenScheduler::TickQueueComponentsForGeneration(
 			const EPCGHiGenGrid MaxGrid = bHasUnbounded ? EPCGHiGenGrid::Unbounded : PCGHiGenGrid::GridSizeToGrid(GridSizes[0]);
 			const double MaxGenerationRadius = OriginalComponent->GetGenerationRadiusFromGrid(MaxGrid);
 
-			for (const IPCGGenSourceBase* GenSource : InGenSources)
+			for (const IPCGGenSourceBase* GenSource : *Inputs.GenSources)
 			{
 				const TOptional<FVector> GenSourcePositionOptional = GenSource->GetPosition();
 				if (!GenSourcePositionOptional.IsSet())
@@ -424,7 +431,7 @@ void FPCGRuntimeGenScheduler::TickQueueComponentsForGeneration(
 					
 					// Ignore components that have already been generated or marked for generation. Unbounded grid size means not-partitioned.
 					FGridGenerationKey Key(PCGHiGenGrid::UnboundedGridSize(), FIntVector(0), OriginalComponent);
-					if (!GeneratedComponents.Contains(Key) && 
+					if (!Inputs.GeneratedComponents.Contains(Key) &&
 						(!Policy || Policy->ShouldGenerate(GenSource, OriginalComponentBounds, bIs2DGrid)))
 					{
 						check(Key.GetGridDescriptor().Is2DGrid() == bIs2DGrid);
@@ -465,7 +472,7 @@ void FPCGRuntimeGenScheduler::TickQueueComponentsForGeneration(
 								FGridGenerationKey Key(GridSize, GridCoords, OriginalComponent);
 
 								// Ignore components that have already been generated or marked for generation.
-								if (GeneratedComponents.Find({ GridSize, GridCoords, OriginalComponent }))
+								if (Inputs.GeneratedComponents.Find({ GridSize, GridCoords, OriginalComponent }))
 								{
 									continue;
 								}
@@ -530,7 +537,7 @@ void FPCGRuntimeGenScheduler::TickQueueComponentsForGeneration(
 	}
 
 	// Collect all non-partitioned components.
-	for (UPCGComponent* OriginalComponent : ActorAndComponentMapping->NonPartitionedOctree.GetAllComponents())
+	for (UPCGComponent* OriginalComponent : Inputs.AllNonPartitionedComponents)
 	{
 		TRACE_CPUPROFILER_EVENT_SCOPE(FPCGRuntimeGenScheduler::CollectNonPartitionedComponents);
 
@@ -540,7 +547,7 @@ void FPCGRuntimeGenScheduler::TickQueueComponentsForGeneration(
 		}
 
 		// The generation key for a non-partitioned component should always have unbounded grid size and 0,0,0 cell coord.
-		if (GeneratedComponents.Contains({ PCGHiGenGrid::UnboundedGridSize(), FIntVector(0), OriginalComponent }))
+		if (Inputs.GeneratedComponents.Contains({ PCGHiGenGrid::UnboundedGridSize(), FIntVector(0), OriginalComponent }))
 		{
 			continue;
 		}
@@ -554,7 +561,7 @@ void FPCGRuntimeGenScheduler::TickQueueComponentsForGeneration(
 			// Unbounded will grab the base GenerationRadius used for non-partitioned and unbounded.
 			const double MaxGenerationRadius = OriginalComponent->GetGenerationRadiusFromGrid(EPCGHiGenGrid::Unbounded);
 
-			for (const IPCGGenSourceBase* GenSource : InGenSources)
+			for (const IPCGGenSourceBase* GenSource : *Inputs.GenSources)
 			{
 				const TOptional<FVector> GenSourcePositionOptional = GenSource->GetPosition();
 				if (!GenSourcePositionOptional.IsSet())
@@ -637,7 +644,7 @@ void FPCGRuntimeGenScheduler::TickCleanup(const TSet<IPCGGenSourceBase*>& InGenS
 
 	// Generated Entry Key, Local/Generated Component
 	using PCGComponentToClean = TTuple<FGridGenerationKey, UPCGComponent*>;
-	TArray<PCGComponentToClean, TInlineAllocator<128>> ComponentsToClean;
+	TArray<PCGComponentToClean, TInlineAllocator<256>> ComponentsToClean;
 	TArray<FGridGenerationKey, TInlineAllocator<16>> InvalidKeys;
 
 	ComponentsToClean.SetNumUninitialized(NumGeneratedComponents);
@@ -694,7 +701,12 @@ void FPCGRuntimeGenScheduler::TickCleanup(const TSet<IPCGGenSourceBase*>& InGenS
 		// Otherwise, we have a local component.
 		else
 		{
-			UPCGComponent* LocalComponent = ActorAndComponentMapping->GetLocalComponent(GridDescriptor, GridCoords, OriginalComponent);
+			UPCGComponent* LocalComponent = GenerationKey.GetCachedLocalComponent();
+			if (!LocalComponent)
+			{
+				LocalComponent = ActorAndComponentMapping->GetLocalComponent(GridDescriptor, GridCoords, OriginalComponent);
+			}
+
 			APCGPartitionActor* PartitionActor = LocalComponent ? Cast<APCGPartitionActor>(LocalComponent->GetOwner()) : nullptr;
 			if (!PartitionActor || !OriginalComponent->bActivated)
 			{
@@ -737,7 +749,7 @@ void FPCGRuntimeGenScheduler::TickCleanup(const TSet<IPCGGenSourceBase*>& InGenS
 	}
 }
 
-void FPCGRuntimeGenScheduler::TickScheduleGeneration(TMap<FGridGenerationKey, double>& ComponentsToGenerate)
+void FPCGRuntimeGenScheduler::TickScheduleGeneration(TMap<FGridGenerationKey, double>& InComponentsToGenerate)
 {
 	TRACE_CPUPROFILER_EVENT_SCOPE(FPCGRuntimeGenScheduler::TickScheduleGeneration);
 
@@ -755,11 +767,7 @@ void FPCGRuntimeGenScheduler::TickScheduleGeneration(TMap<FGridGenerationKey, do
 				continue;
 			}
 
-			UPCGComponent* Component = Key.GetOriginalComponent();
-			if (Key.GetGridDescriptor().GetHiGenGrid() != EPCGHiGenGrid::Unbounded)
-			{
-				Component = ActorAndComponentMapping->GetLocalComponent(Key.GetGridDescriptor(), Key.GetGridCoords(), Component);
-			}
+			UPCGComponent* Component = (Key.GetGridDescriptor().GetHiGenGrid() == EPCGHiGenGrid::Unbounded) ? Key.GetOriginalComponent() : Key.GetCachedLocalComponent();
 
 			if (Component && Component->IsGenerating())
 			{
@@ -770,14 +778,16 @@ void FPCGRuntimeGenScheduler::TickScheduleGeneration(TMap<FGridGenerationKey, do
 
 	const int MaxNumGenerating = FMath::Max(1, PCGRuntimeGenSchedulerHelpers::CVarNumGeneratingComponentsAtSameTime.GetValueOnAnyThread());
 
-	for (const TPair<FGridGenerationKey, double>& GenerateEntry : ComponentsToGenerate)
+	GeneratedComponents.Reserve(GeneratedComponents.Num() + InComponentsToGenerate.Num());
+
+	for (TPair<FGridGenerationKey, double>& GenerateEntry : InComponentsToGenerate)
 	{
 		if (NumGenerating >= MaxNumGenerating)
 		{
 			break;
 		}
 
-		const FGridGenerationKey& Key = GenerateEntry.Key;
+		FGridGenerationKey& Key = GenerateEntry.Key;
 		const double Priority = GenerateEntry.Value;
 
 		const FPCGGridDescriptor GridDescriptor = Key.GetGridDescriptor();
@@ -868,7 +878,7 @@ void FPCGRuntimeGenScheduler::TickScheduleGeneration(TMap<FGridGenerationKey, do
 					PartitionActorsPtr->Add(PartitionActor);
 				}
 
-				// Generate local component.
+				// Create local component.
 				LocalComponent = PartitionActor->GetLocalComponent(OriginalComponent);
 			}
 
@@ -894,6 +904,8 @@ void FPCGRuntimeGenScheduler::TickScheduleGeneration(TMap<FGridGenerationKey, do
 
 				// Force to refresh if the component is already generated.
 				LocalComponent->GenerateLocal(EPCGComponentGenerationTrigger::GenerateAtRuntime, /*bForce=*/true, LocalComponent->GetGenerationGrid(), Dependencies);
+
+				Key.SetCachedLocalComponent(LocalComponent);
 			}
 		}
 
@@ -913,8 +925,8 @@ void FPCGRuntimeGenScheduler::TickRequestVirtualTexturePriming(const TSet<IPCGGe
 	{
 		TRACE_CPUPROFILER_EVENT_SCOPE(GatherOriginalComponents);
 
-		TSet<UPCGComponent*> PartitionedComponents = ActorAndComponentMapping->PartitionedOctree.GetAllComponents();
-		TSet<UPCGComponent*> NonPartitionedComponents = ActorAndComponentMapping->NonPartitionedOctree.GetAllComponents();
+		TSet<UPCGComponent*> PartitionedComponents = ActorAndComponentMapping->GetAllRegisteredPartitionedComponents();
+		TSet<UPCGComponent*> NonPartitionedComponents = ActorAndComponentMapping->GetAllRegisteredNonPartitionedComponents();
 
 		OriginalComponents.Reserve(PartitionedComponents.Num() + NonPartitionedComponents.Num());
 
@@ -1304,7 +1316,7 @@ void FPCGRuntimeGenScheduler::CleanupLocalComponent(APCGPartitionActor* Partitio
 	}
 
 	// Cleanup the PA if it no longer has any components (return to pool or destroy).
-	if (PartitionActor->GetAllLocalPCGComponents().IsEmpty())
+	if (PartitionActor->HasLocalPCGComponents())
 	{
 		PartitionActor->UnregisterPCG();
 
@@ -1471,7 +1483,7 @@ void FPCGRuntimeGenScheduler::RefreshComponent(UPCGComponent* InComponent, bool
 				UE_LOG(LogPCG, Warning, TEXT("[RUNTIMEGEN] SHALLOW REFRESH LOCAL COMPONENT: '%s'"), *PartitionActor->GetActorNameOrLabel());
 			}
 
-			GeneratedComponentsToRemove.Emplace({ PartitionActor->GetPCGGridSize(), PartitionActor->GetGridCoord(), OriginalComponent });
+			GeneratedComponentsToRemove.Emplace({ PartitionActor->GetPCGGridSize(), PartitionActor->GetGridCoord(), OriginalComponent, bIsLocalComponent ? InComponent : nullptr });
 			InComponent->CleanupLocalImmediate(/*bRemoveComponents=*/false);
 		}
 		else
@@ -1529,7 +1541,7 @@ void FPCGRuntimeGenScheduler::RefreshComponent(UPCGComponent* InComponent, bool
 			APCGPartitionActor* PartitionActor = CastChecked<APCGPartitionActor>(LocalComponent->GetOwner());
 
 			// Find the specific generation key for this component, if it exists, cleanup and generate.
-			FGridGenerationKey LocalComponentKey(PartitionActor->GetPCGGridSize(), PartitionActor->GetGridCoord(), OriginalComponent);
+			FGridGenerationKey LocalComponentKey(PartitionActor->GetPCGGridSize(), PartitionActor->GetGridCoord(), OriginalComponent, LocalComponent);
 
 			if (GeneratedComponents.Find(LocalComponentKey))
 			{
diff --git a/Engine/Plugins/PCG/Source/PCG/Public/Grid/PCGPartitionActor.h b/Engine/Plugins/PCG/Source/PCG/Public/Grid/PCGPartitionActor.h
--- a/Engine/Plugins/PCG/Source/PCG/Public/Grid/PCGPartitionActor.h
+++ b/Engine/Plugins/PCG/Source/PCG/Public/Grid/PCGPartitionActor.h
@@ -105,6 +105,9 @@ public:
 	/** [Game thread only] Return if the actor is safe for deletion, meaning no generation is currently running on all original components. */
 	UE_API bool IsSafeForDeletion() const;
 
+	/** Whether there are any local PCGCompnents on this actor. */
+	UE_API bool HasLocalPCGComponents() const;
+
 	/** Return an array of all the PCGComponents on this actor */
 	UE_API TSet<TObjectPtr<UPCGComponent>> GetAllLocalPCGComponents() const;
 
diff --git a/Engine/Plugins/PCG/Source/PCG/Public/PCGGraph.h b/Engine/Plugins/PCG/Source/PCG/Public/PCGGraph.h
--- a/Engine/Plugins/PCG/Source/PCG/Public/PCGGraph.h
+++ b/Engine/Plugins/PCG/Source/PCG/Public/PCGGraph.h
@@ -454,6 +454,8 @@ protected:
 
 	UE_API bool ForEachNodeRecursively_Internal(TFunctionRef<bool(UPCGNode*)> Action, TSet<const UPCGGraph*>& VisitedGraphs) const;
 
+	void CacheGridSizesInternalNoLock() const;
+
 	/** Calculates node grid size. Not thread safe, must be called within write lock. */
 	UE_API uint32 CalculateNodeGridSizeRecursive_Unsafe(const UPCGNode* InNode, uint32 InDefaultGridSize) const;
 
@@ -507,6 +509,16 @@ protected:
 	mutable TMap<const UPCGNode*, uint32> NodeToGridSize;
 	mutable FRWLock NodeToGridSizeLock;
 
+	struct FGridInfo
+	{
+		PCGHiGenGrid::FSizeArray GridSizes;
+		bool bHasUnbounded = false;
+	};
+
+	/** Cached information about higen grids present in graph. */
+	mutable TOptional<FGridInfo> CachedGridInfo;
+	mutable FCriticalSection CachedGridInfoLock;
+
 	/** Sets whether this graph is marked as editor-only; note that the IsEditorOnly call depends on the local graph value and the value in all subgraphs, recursively. */
 	UPROPERTY(EditAnywhere, BlueprintReadWrite, Category = Cooking)
 	bool bIsEditorOnly = false;
diff --git a/Engine/Plugins/PCG/Source/PCG/Public/RuntimeGen/PCGRuntimeGenScheduler.h b/Engine/Plugins/PCG/Source/PCG/Public/RuntimeGen/PCGRuntimeGenScheduler.h
--- a/Engine/Plugins/PCG/Source/PCG/Public/RuntimeGen/PCGRuntimeGenScheduler.h
+++ b/Engine/Plugins/PCG/Source/PCG/Public/RuntimeGen/PCGRuntimeGenScheduler.h
@@ -75,27 +75,70 @@ public:
 	void FlushAllGeneratedActors() { bActorFlushRequested = true; }
 
 protected:
-	struct FGridGenerationKey : TTuple<uint32, FIntVector, TWeakObjectPtr<UPCGComponent>>
+	// Grid size, grid coords, original component, local componnent
+	struct FGridGenerationKey
 	{
 		FGridGenerationKey(uint32 InGridSize, const FIntVector& InGridCoords, UPCGComponent* InOriginalComponent)
-			: TTuple<uint32, FIntVector, TWeakObjectPtr<UPCGComponent>>(InGridSize, InGridCoords, InOriginalComponent) {}
+			: GridSize(InGridSize)
+			, GridCoords(InGridCoords)
+			, OriginalComponent(InOriginalComponent)
+			, CachedLocalComponent(nullptr)
+		{}
+
+		FGridGenerationKey(uint32 InGridSize, const FIntVector& InGridCoords, UPCGComponent* InOriginalComponent, UPCGComponent* InLocalComponent)
+			: GridSize(InGridSize)
+			, GridCoords(InGridCoords)
+			, OriginalComponent(InOriginalComponent)
+			, CachedLocalComponent(InLocalComponent)
+		{}
+
+		bool operator==(const FGridGenerationKey& Other) const
+		{
+			return GridSize == Other.GridSize
+				&& GridCoords == Other.GridCoords
+				&& OriginalComponent == Other.OriginalComponent;
+		}
+
+		bool IsValid() const { return OriginalComponent.IsValid(); }
+
+		uint32 GetGridSize() const { return GridSize; }
+		FIntVector GetGridCoords() const { return GridCoords; }
+		UPCGComponent* GetOriginalComponent() const { return OriginalComponent.Get(); }
+
+		UPCGComponent* GetCachedLocalComponent() const { return CachedLocalComponent.Get(); }
+		void SetCachedLocalComponent(UPCGComponent* InComponent) { CachedLocalComponent = InComponent; }
 
-		bool IsValid() const { return Get<2>().IsValid(); }
+		FPCGGridDescriptor GetGridDescriptor() const;
 
-		uint32 GetGridSize() const { return Get<0>(); }
-		FIntVector GetGridCoords() const { return Get<1>(); }
-		UPCGComponent* GetOriginalComponent() const { return Get<2>().Get(); }
+		uint32 GridSize = 0;
+		FIntVector GridCoords = FIntVector(0);
+		TWeakObjectPtr<UPCGComponent> OriginalComponent;
 
-		FPCGGridDescriptor GetGridDescriptor() const;
+		// Optional/opportunistic cached local component if one has been created.
+		TWeakObjectPtr<UPCGComponent> CachedLocalComponent;
+
+		friend uint32 GetTypeHash(const FGridGenerationKey& In)
+		{
+			return HashCombine(GetTypeHash(In.GridCoords), GetTypeHash(In.GridSize), GetTypeHash(In.OriginalComponent));
+		}
 	};
 
 	/** Returns true if the scheduler should tick this frame. */
 	bool ShouldTick();
 
+	// Add constructor, get moves for the sets
+	struct FTickQueueComponentsForGenerationInputs
+	{
+		const TSet<IPCGGenSourceBase*>* GenSources;
+		const APCGWorldActor* PCGWorldActor;
+		TSet<UPCGComponent*> AllPartitionedComponents;
+		TSet<UPCGComponent*> AllNonPartitionedComponents;
+		TSet<FGridGenerationKey> GeneratedComponents;
+	};
+
 	/** Queue nearby components for generation. */
 	void TickQueueComponentsForGeneration(
-		const TSet<IPCGGenSourceBase*>& GenSources,
-		APCGWorldActor* InPCGWorldActor,
+		const FTickQueueComponentsForGenerationInputs& Inputs,
 		TMap<FGridGenerationKey, double>& OutComponentsToGenerate);
 
 	/** Perform immediate cleanup on components that become out of range. */
@@ -157,6 +200,10 @@ private:
 	*/
 	TSet<FGridGenerationKey> GeneratedComponentsToRemove;
 
+	// Local to member functions but hoisted for efficiency.
+	TSet<IPCGGenSourceBase*> GenSources;
+	TMap<FGridGenerationKey, double> ComponentsToGenerate;
+
 	/** Pool of RuntimeGen PartitionActors used for hierarchical generation. */
 	TArray<TObjectPtr<APCGPartitionActor>> PartitionActorPool;
 
