[PCG] Fix partition actors not getting returned to pool.

Minor: Move scopes slightly so we don't get events for ticks that are inactive in current world, and move world streaming query event to the cache miss path which is much more useful.

#rnx

diff --git a/Engine/Plugins/PCG/Source/PCG/Private/RuntimeGen/PCGRuntimeGenScheduler.cpp b/Engine/Plugins/PCG/Source/PCG/Private/RuntimeGen/PCGRuntimeGenScheduler.cpp
--- a/Engine/Plugins/PCG/Source/PCG/Private/RuntimeGen/PCGRuntimeGenScheduler.cpp
+++ b/Engine/Plugins/PCG/Source/PCG/Private/RuntimeGen/PCGRuntimeGenScheduler.cpp
@@ -140,7 +140,6 @@ FPCGRuntimeGenScheduler::~FPCGRuntimeGenScheduler()
 
 void FPCGRuntimeGenScheduler::Tick(APCGWorldActor* InPCGWorldActor, double InEndTime)
 {
-	TRACE_CPUPROFILER_EVENT_SCOPE(FPCGRuntimeGenScheduler::Tick);
 	check(InPCGWorldActor && GenSourceManager);
 
 	// 0. Preamble - check if we should be active in this world and do lazy initialization.
@@ -150,6 +149,8 @@ void FPCGRuntimeGenScheduler::Tick(APCGWorldActor* InPCGWorldActor, double InEnd
 		return;
 	}
 
+	TRACE_CPUPROFILER_EVENT_SCOPE(FPCGRuntimeGenScheduler::Tick);
+
 	TickCVars(InPCGWorldActor);
 
 	GenSources.Empty(GenSources.Num());
@@ -348,14 +349,14 @@ void FPCGRuntimeGenScheduler::TickQueueComponentsForGeneration(
 
 	auto IsWorldStreamingComplete = [WorldPartitionSubsystem, &StreamingQuerySources, &CachedStreamingQueryResults=CachedStreamingQueryResults](const FVector& InLocation, float InGridSize)
 	{
-		TRACE_CPUPROFILER_EVENT_SCOPE(IsWorldStreamingComplete);
-
 		FStreamingCompleteQueryKey Key{ InLocation, InGridSize };
 		if (bool* FoundResult = CachedStreamingQueryResults.Find(Key))
 		{
 			return *FoundResult;
 		}
 
+		TRACE_CPUPROFILER_EVENT_SCOPE(IsWorldStreamingComplete);
+
 		StreamingQuerySources[0].Radius = InGridSize / 2.0f;
 		StreamingQuerySources[0].Location = InLocation;
 
@@ -1316,7 +1317,7 @@ void FPCGRuntimeGenScheduler::CleanupLocalComponent(APCGPartitionActor* Partitio
 	}
 
 	// Cleanup the PA if it no longer has any components (return to pool or destroy).
-	if (PartitionActor->HasLocalPCGComponents())
+	if (!PartitionActor->HasLocalPCGComponents())
 	{
 		PartitionActor->UnregisterPCG();
 
