[PCG] Fixed indexing mismatch in per-data crc mechanism

diff --git a/Engine/Plugins/PCG/Source/PCG/Private/PCGElement.cpp b/Engine/Plugins/PCG/Source/PCG/Private/PCGElement.cpp
--- a/Engine/Plugins/PCG/Source/PCG/Private/PCGElement.cpp
+++ b/Engine/Plugins/PCG/Source/PCG/Private/PCGElement.cpp
@@ -682,10 +682,18 @@ void IPCGElement::PostExecutePrimaryLoopElement(FPCGContext* Context, const UPCG
 				Context->InputData.TaggedData.Append(CachedInputData.TaggedData);
 			}
 
-			const FPCGDataCollection& CachedOutputData = Context->CachedInputToOutputInternalResults[CachedCollectionIndex].Value;
+			// Note: we're modifying the context data, but it will not be reused after
+			FPCGDataCollection& CachedOutputData = Context->CachedInputToOutputInternalResults[CachedCollectionIndex].Value;
 
 			// Since this data comes from the cache, we can't rely on its original index; however, it should be used as-if it was keyed on the input original index.
 			const int32 CacheOutputOriginalIndex = CacheInputOriginalIndex;
+
+			// Apply matching original index so that further iterations place the data at the right indices.
+			for (FPCGTaggedData& TaggedData : CachedOutputData.TaggedData)
+			{
+				TaggedData.OriginalIndex = CacheOutputOriginalIndex;
+			}
+
 			const int32 InsertOutputIndex = CacheOutputOriginalIndex != INDEX_NONE ? Context->OutputData.TaggedData.IndexOfByPredicate([CacheOutputOriginalIndex](const FPCGTaggedData& TaggedData) { return TaggedData.OriginalIndex > CacheOutputOriginalIndex; }) : INDEX_NONE;
 
 			if (InsertOutputIndex != INDEX_NONE)
