[PCG] Remove empty data doesn't automatically type pins

diff --git a/Engine/Plugins/PCG/Source/PCG/Public/Elements/PCGRemoveEmptyData.h b/Engine/Plugins/PCG/Source/PCG/Public/Elements/PCGRemoveEmptyData.h
--- a/Engine/Plugins/PCG/Source/PCG/Public/Elements/PCGRemoveEmptyData.h
+++ b/Engine/Plugins/PCG/Source/PCG/Public/Elements/PCGRemoveEmptyData.h
@@ -20,6 +20,8 @@ public:
 	virtual EPCGSettingsType GetType() const override { return EPCGSettingsType::Filter; }
 #endif
 
+	virtual bool HasDynamicPins() const override { return true; }
+
 protected:
 	virtual TArray<FPCGPinProperties> InputPinProperties() const override;
 	virtual TArray<FPCGPinProperties> OutputPinProperties() const override;
