[PCG] TPCGValueRange: Fix ViewNum() when NumElements == 0

#rnx

diff --git a/Engine/Plugins/PCG/Source/PCG/Public/Utils/PCGValueRange.h b/Engine/Plugins/PCG/Source/PCG/Public/Utils/PCGValueRange.h
--- a/Engine/Plugins/PCG/Source/PCG/Public/Utils/PCGValueRange.h
+++ b/Engine/Plugins/PCG/Source/PCG/Public/Utils/PCGValueRange.h
@@ -85,7 +85,7 @@ public:
 
 	[[nodiscard]] FORCEINLINE SizeType ViewNum() const
 	{
-		return ElementView.Num();
+		return FMath::Min(NumElements, ElementView.Num());
 	}
 
 	[[nodiscard]] FORCEINLINE ElementType& operator[](SizeType Index) const
