[PCG] Fixed an issue with the way bounds were obtained for actor tracking and changes.

diff --git a/Engine/Plugins/PCG/Source/PCG/Private/PCGActorAndComponentMapping.cpp b/Engine/Plugins/PCG/Source/PCG/Private/PCGActorAndComponentMapping.cpp
--- a/Engine/Plugins/PCG/Source/PCG/Private/PCGActorAndComponentMapping.cpp
+++ b/Engine/Plugins/PCG/Source/PCG/Private/PCGActorAndComponentMapping.cpp
@@ -42,23 +42,6 @@
 
 namespace PCGActorAndComponentMapping
 {
-	FBox GetActorBounds(const AActor* InActor)
-	{
-		if (!InActor)
-		{
-			return FBox(EForceInit::ForceInit);
-		}
-
-		FBox ActorBounds = InActor->GetComponentsBoundingBox();
-		if (!ActorBounds.IsValid && InActor->GetRootComponent() != nullptr)
-		{
-			// Try on the RootComponent
-			ActorBounds = InActor->GetRootComponent()->Bounds.GetBox();
-		}
-
-		return ActorBounds;
-	}
-
 #if WITH_EDITOR
 	static TAutoConsoleVariable<bool> CVarDisableObjectDependenciesTracking(
 		TEXT("pcg.DisableObjectDependenciesTracking"),
@@ -1666,7 +1649,7 @@ void FPCGActorAndComponentMapping::OnObjectModified(UObject* InObject)
 			if (IsActorTracked(InActor))
 			{
 				FActorPreviousData& PreviousData = ActorToPreviousDataMap.Add(InActor);
-				PreviousData.Get<0>() = PCGActorAndComponentMapping::GetActorBounds(InActor);
+				PreviousData.Get<0>() = PCGHelpers::GetActorBounds(InActor, /*bIgnorePCGCreatedComponents=*/false);
 				PreviousData.Get<1>() = TSet<FName>(InActor->Tags);
 				PreviousData.Get<2>() = FApp::GetCurrentTime();
 			}
@@ -1884,7 +1867,7 @@ void FPCGActorAndComponentMapping::OnObjectChanged(UObject* InObject, const ECha
 	}
 
 	// We need the actor bounds to know if we are intersecting, for the culled settings to be dirtied.
-	const FBox ActorBounds = PCGActorAndComponentMapping::GetActorBounds(Actor);
+	const FBox ActorBounds = PCGHelpers::GetActorBounds(Actor, /*bIgnorePCGCreatedComponents=*/false);
 
 	{
 		TRACE_CPUPROFILER_EVENT_SCOPE(FPCGActorAndComponentMapping::OnObjectChanged::AlwaysTrackedUpdate);
@@ -2184,7 +2167,7 @@ void FPCGActorAndComponentMapping::OnPCGGraphStartsGenerating(UPCGComponent* InC
 		// Removing all components that aren't intersecting with it, since it won't contribute to refresh
 		for (UPCGComponent* TrackedComponent : TempTrackedComponents)
 		{
-			if (ensure(TrackedComponent) && TrackedComponent->GetOwner() && PCGActorAndComponentMapping::GetActorBounds(TrackedComponent->GetOwner()).Intersect(ComponentBounds))
+			if (ensure(TrackedComponent) && TrackedComponent->GetOwner() && PCGHelpers::GetActorBounds(TrackedComponent->GetOwner(), /*bIgnorePCGCreatedComponents=*/false).Intersect(ComponentBounds))
 			{
 				TrackedComponents.Add(TrackedComponent);
 			}
