[PCG] Use player viewpoint rather than player pawn actor position for generation source position.

This will support cinematics etc and can remove need for WP streaming sources, and allow a single generation source.

#rnx

diff --git a/Engine/Plugins/PCG/Source/PCG/Private/RuntimeGen/GenSources/PCGGenSourcePlayer.cpp b/Engine/Plugins/PCG/Source/PCG/Private/RuntimeGen/GenSources/PCGGenSourcePlayer.cpp
--- a/Engine/Plugins/PCG/Source/PCG/Private/RuntimeGen/GenSources/PCGGenSourcePlayer.cpp
+++ b/Engine/Plugins/PCG/Source/PCG/Private/RuntimeGen/GenSources/PCGGenSourcePlayer.cpp
@@ -34,10 +34,10 @@ TOptional<FVector> UPCGGenSourcePlayer::GetPosition() const
 {
 	if (PlayerController.IsValid())
 	{
-		if (APawn* PlayerPawn = PlayerController->GetPawn())
-		{
-			return PlayerPawn->GetActorLocation();
-		}
+		FVector CameraLocation;
+		FRotator CameraRotation;
+		PlayerController->GetPlayerViewPoint(CameraLocation, CameraRotation);
+		return CameraLocation;
 	}
 
 	return TOptional<FVector>();
@@ -47,10 +47,10 @@ TOptional<FVector> UPCGGenSourcePlayer::GetDirection() const
 {
 	if (PlayerController.IsValid())
 	{
-		if (APawn* PlayerPawn = PlayerController->GetPawn())
-		{
-			return PlayerPawn->GetActorForwardVector();
-		}
+		FVector CameraLocation;
+		FRotator CameraRotation;
+		PlayerController->GetPlayerViewPoint(CameraLocation, CameraRotation);
+		return CameraRotation.Vector();
 	}
 
 	return TOptional<FVector>();
