[PCG] Fix some metadata elements calls that were not respecting time slicing

* Also add some guards in async to make sure if this happens again elsewhere it is caught with a good message and not crash.

#rnx

diff --git a/Engine/Plugins/PCG/Source/PCG/Private/Elements/Metadata/PCGHashAttribute.cpp b/Engine/Plugins/PCG/Source/PCG/Private/Elements/Metadata/PCGHashAttribute.cpp
--- a/Engine/Plugins/PCG/Source/PCG/Private/Elements/Metadata/PCGHashAttribute.cpp
+++ b/Engine/Plugins/PCG/Source/PCG/Private/Elements/Metadata/PCGHashAttribute.cpp
@@ -15,15 +15,13 @@ bool FPCGHashAttributeElement::DoOperation(PCGMetadataOps::FOperationData& InOpe
 {
 	TRACE_CPUPROFILER_EVENT_SCOPE(FPCGHashAttributeElement::Execute);
 
-	PCGMetadataAttribute::CallbackWithRightType(InOperationData.MostComplexInputType, [this, &InOperationData]<typename AttributeType>(AttributeType) -> void
+	return PCGMetadataAttribute::CallbackWithRightType(InOperationData.MostComplexInputType, [this, &InOperationData]<typename AttributeType>(AttributeType) -> bool
 	{
-		DoUnaryOp<AttributeType>(InOperationData, [](const AttributeType& Value) -> int32
+		return DoUnaryOp<AttributeType>(InOperationData, [](const AttributeType& Value) -> int32
 		{
 			return static_cast<int32>(PCG::Private::MetadataTraits<AttributeType>::Hash(Value));
 		});
 	});
-
-	return true;
 }
 
 #undef LOCTEXT_NAMESPACE
diff --git a/Engine/Plugins/PCG/Source/PCG/Private/Elements/Metadata/PCGMetadataMathsOpElement.cpp b/Engine/Plugins/PCG/Source/PCG/Private/Elements/Metadata/PCGMetadataMathsOpElement.cpp
--- a/Engine/Plugins/PCG/Source/PCG/Private/Elements/Metadata/PCGMetadataMathsOpElement.cpp
+++ b/Engine/Plugins/PCG/Source/PCG/Private/Elements/Metadata/PCGMetadataMathsOpElement.cpp
@@ -376,8 +376,7 @@ bool FPCGMetadataMathsElement::DoOperation(PCGMetadataOps::FOperationData& Opera
 			if (Operation == EPCGMetadataMathsOperation::Set)
 			{
 				TRACE_CPUPROFILER_EVENT_SCOPE(FPCGMetadataMathsElement::ExecuteInternal::BinaryOp);
-				DoBinaryOp<AttributeType, AttributeType>(OperationData, [](const AttributeType&, const AttributeType& Value2) -> AttributeType { return Value2; });
-				return true;
+				return DoBinaryOp<AttributeType, AttributeType>(OperationData, [](const AttributeType&, const AttributeType& Value2) -> AttributeType { return Value2; });
 			}
 		}
 
diff --git a/Engine/Plugins/PCG/Source/PCG/Private/Elements/Metadata/PCGMetadataOpElementBase.cpp b/Engine/Plugins/PCG/Source/PCG/Private/Elements/Metadata/PCGMetadataOpElementBase.cpp
--- a/Engine/Plugins/PCG/Source/PCG/Private/Elements/Metadata/PCGMetadataOpElementBase.cpp
+++ b/Engine/Plugins/PCG/Source/PCG/Private/Elements/Metadata/PCGMetadataOpElementBase.cpp
@@ -982,6 +982,14 @@ bool FPCGMetadataElementBase::ExecuteInternal(FPCGContext* Context) const
 			return true;
 		}
 
+		// Make sure the async state is reset, otherwise it means the metadata op is not taking into account time-slicing correctly
+		if (!ensureMsgf(!Context->AsyncState.bStarted,
+			TEXT("Metadata operation has not finish processing the previous data and is starting a new processing.\n"
+			"Make sure that the DoOperation is returning true only when the async processing is done.")))
+		{
+			return true;
+		}
+
 		return DoOperation(IterState);
 	});
 }
diff --git a/Engine/Plugins/PCG/Source/PCG/Private/Helpers/PCGAsyncState.cpp b/Engine/Plugins/PCG/Source/PCG/Private/Helpers/PCGAsyncState.cpp
--- a/Engine/Plugins/PCG/Source/PCG/Private/Helpers/PCGAsyncState.cpp
+++ b/Engine/Plugins/PCG/Source/PCG/Private/Helpers/PCGAsyncState.cpp
@@ -15,6 +15,7 @@ void FPCGAsyncState::Reset()
 	TasksChunksStartEnd.Reset();
 	CurrentChunkToCollapse = 0;
 	ChunkToNumElementsWrittenMap.Reset();
+	InitialNumIterations = 0;
 
 	bStarted = false;
 }
diff --git a/Engine/Plugins/PCG/Source/PCG/Public/Helpers/PCGAsync.h b/Engine/Plugins/PCG/Source/PCG/Public/Helpers/PCGAsync.h
--- a/Engine/Plugins/PCG/Source/PCG/Public/Helpers/PCGAsync.h
+++ b/Engine/Plugins/PCG/Source/PCG/Public/Helpers/PCGAsync.h
@@ -97,6 +97,7 @@ namespace FPCGAsync
 			Initialize();
 
 			AsyncState.bStarted = true;
+			AsyncState.InitialNumIterations = NumIterations;
 
 			// At the beginning, dispatch the chunks that each task needs to process.
 			// It will be stored on the AsyncState so that it can be reused.
@@ -130,6 +131,19 @@ namespace FPCGAsync
 				check(Count == ChunksNumber);
 			}
 		}
+		else
+		{
+			if (!ensureMsgf(NumIterations == AsyncState.InitialNumIterations,
+				TEXT("NumIterations (%d) mismatch with the initial number of iterations (%d).\n"
+				"It might mean that a new processing is started while the previous one is not done.\n"
+				"Validate on the calling site that if time slicing is enabled, it is waiting for this\n"
+				"function to return true (done) before starting a new processing."),
+				NumIterations, AsyncState.InitialNumIterations))
+			{
+				AsyncState.Reset();
+				return true;
+			}
+		}
 
 		const int32 NumFutures = FMath::Max(0, AsyncState.TasksChunksStartEnd.Num() - 1);
 
diff --git a/Engine/Plugins/PCG/Source/PCG/Public/Helpers/PCGAsyncState.h b/Engine/Plugins/PCG/Source/PCG/Public/Helpers/PCGAsyncState.h
--- a/Engine/Plugins/PCG/Source/PCG/Public/Helpers/PCGAsyncState.h
+++ b/Engine/Plugins/PCG/Source/PCG/Public/Helpers/PCGAsyncState.h
@@ -41,6 +41,9 @@ struct FPCGAsyncState
 	/** Keep information if the async has started or not. */
 	bool bStarted = false;
 
+	/** Keep information on the initial number of iteration, so that if we have timeslicing we make sure that the number of iterations matches. */
+	int32 InitialNumIterations = 0;
+
 	/** Keep track of the current chunk we need to collapse. Important to collapse in order. */
 	int32 CurrentChunkToCollapse = 0;
 
