[PCG] Move the recently added ensure at the right place

* We should check that the async state is reset when we are done with the process, not before

#rnx

diff --git a/Engine/Plugins/PCG/Source/PCG/Private/Elements/Metadata/PCGMetadataOpElementBase.cpp b/Engine/Plugins/PCG/Source/PCG/Private/Elements/Metadata/PCGMetadataOpElementBase.cpp
--- a/Engine/Plugins/PCG/Source/PCG/Private/Elements/Metadata/PCGMetadataOpElementBase.cpp
+++ b/Engine/Plugins/PCG/Source/PCG/Private/Elements/Metadata/PCGMetadataOpElementBase.cpp
@@ -982,15 +982,17 @@ bool FPCGMetadataElementBase::ExecuteInternal(FPCGContext* Context) const
 			return true;
 		}
 
-		// Make sure the async state is reset, otherwise it means the metadata op is not taking into account time-slicing correctly
-		if (!ensureMsgf(!Context->AsyncState.bStarted,
-			TEXT("Metadata operation has not finish processing the previous data and is starting a new processing.\n"
-			"Make sure that the DoOperation is returning true only when the async processing is done.")))
+		const bool bIsDone = DoOperation(IterState);
+		
+		if (bIsDone)
 		{
-			return true;
+			// Make sure the async state is reset, otherwise it means the metadata op is not taking into account time-slicing correctly
+			ensureMsgf(!Context->AsyncState.bStarted,
+				TEXT("Metadata operation has not finish processing the previous data and is starting a new processing.\n"
+				"Make sure that the DoOperation is returning true only when the async processing is done."));
 		}
 
-		return DoOperation(IterState);
+		return bIsDone;
 	});
 }
 
