[PCG] Preconfigurations for Create Graph Parameter based on Type

diff --git a/Engine/Plugins/PCG/Source/PCG/Private/Elements/Metadata/PCGAttributeCast.cpp b/Engine/Plugins/PCG/Source/PCG/Private/Elements/Metadata/PCGAttributeCast.cpp
--- a/Engine/Plugins/PCG/Source/PCG/Private/Elements/Metadata/PCGAttributeCast.cpp
+++ b/Engine/Plugins/PCG/Source/PCG/Private/Elements/Metadata/PCGAttributeCast.cpp
@@ -28,7 +28,7 @@ FPCGElementPtr UPCGAttributeCastSettings::CreateElement() const
 #if WITH_EDITOR
 TArray<FPCGPreConfiguredSettingsInfo> UPCGAttributeCastSettings::GetPreconfiguredInfo() const
 {
-	return FPCGPreConfiguredSettingsInfo::PopulateFromEnum<EPCGMetadataTypes>({ EPCGMetadataTypes::Count, EPCGMetadataTypes::Unknown }, LOCTEXT("PreconfigureInfo", "Cast to "));
+	return FPCGPreConfiguredSettingsInfo::PopulateFromEnum<EPCGMetadataTypes>({ EPCGMetadataTypes::Count, EPCGMetadataTypes::Unknown }, FTextFormat(LOCTEXT("PreconfigureInfo", "Cast to {0}")));
 }
 #endif
 
diff --git a/Engine/Plugins/PCG/Source/PCG/Private/Elements/PCGCreateAttribute.cpp b/Engine/Plugins/PCG/Source/PCG/Private/Elements/PCGCreateAttribute.cpp
--- a/Engine/Plugins/PCG/Source/PCG/Private/Elements/PCGCreateAttribute.cpp
+++ b/Engine/Plugins/PCG/Source/PCG/Private/Elements/PCGCreateAttribute.cpp
@@ -11,6 +11,7 @@
 #include "Elements/PCGUserParameterGet.h"
 #include "Elements/Metadata/PCGMetadataElementCommon.h"
 #include "Helpers/PCGConversion.h"
+#include "Helpers/PCGGraphParametersHelpers.h"
 #include "Helpers/PCGPropertyHelpers.h"
 #include "Utils/PCGLogErrors.h"
 
@@ -35,7 +36,6 @@ namespace PCGCreateAttributeConstants
 		constexpr int32 ToGetGraphParameterIndex = 0;
 
 		const FText ToGetGraphParameterTooltip = LOCTEXT("ToGetGraphParameterTooltip", "Convert constant to graph parameter.");
-		const FName ConversionDefaultName(LOCTEXT("ConversionDefaultName", "NewProperty").ToString());
 	}
 }
 
@@ -265,6 +265,13 @@ FText UPCGCreateAttributeSetSettings::GetNodeTooltipText() const
 {
 	return FText::Format(PCGCreateAttributeConstants::NodeTooltipFormatCreateConstant, FText::FromString(AttributeTypes.TypeToString()), FText::FromString(AttributeTypes.ToString()));
 }
+
+TArray<FPCGPreConfiguredSettingsInfo> UPCGCreateAttributeSetSettings::GetPreconfiguredInfo() const
+{
+	return FPCGPreConfiguredSettingsInfo::PopulateFromEnum<EPCGMetadataTypes>(
+		/*InValuesToSkip=*/{EPCGMetadataTypes::Count, EPCGMetadataTypes::Float, EPCGMetadataTypes::Unknown},
+		/*InOptionalFormat=*/FTextFormat(LOCTEXT("PreconfigureFormat", "New {0} Constant")));
+}
 #endif // WITH_EDITOR
 
 TArray<FPCGPinProperties> UPCGCreateAttributeSetSettings::InputPinProperties() const
@@ -336,21 +343,13 @@ bool UPCGCreateAttributeSetSettings::ConvertNode(const FPCGPreconfiguredInfo& Co
 			if (UPCGGraph* Graph = NodeConverter.GetGraph())
 			{
 				const FInstancedPropertyBag* UserParametersStruct = Graph->GetUserParametersStruct();
-				FName PropertyName = (OutputTarget.GetName() != NAME_None)
-					? OutputTarget.GetName()
-					: PCGCreateAttributeConstants::Conversion::ConversionDefaultName;
-
-				// Validate that the property name doesn't exist... if it does, increment until it doesn't
-				uint64 NameIteration = 1;
-				while (UserParametersStruct->FindPropertyDescByName(PropertyName))
+				FName PropertyName = OutputTarget.GetName();
+				if (!PCGGraphParameter::Helpers::GenerateUniqueName(Graph, PropertyName))
 				{
-					if (++NameIteration > 100)
-					{
-						return false;
-					}
-
-					PropertyName.SetNumber(NameIteration);
+					PCGLog::LogErrorOnGraph(FText::Format(LOCTEXT("CouldNotCreateUniqueGraphParameterName", "Could not create a unique graph parameter with name '{0}'. Check if the graph is valid and the name has not reached its maximum."), FText::FromName(PropertyName)));
+					return false;
 				}
+				check(UserParametersStruct->IsPropertyNameValid(PropertyName));
 
 				FPropertyBagPropertyDesc PropertyDesc = PCGPropertyHelpers::CreatePropertyBagDescWithMetadataType(PropertyName, AttributeTypes.Type);
 				Graph->AddUserParameters({PropertyDesc});
diff --git a/Engine/Plugins/PCG/Source/PCG/Private/Elements/PCGUserParameterGet.cpp b/Engine/Plugins/PCG/Source/PCG/Private/Elements/PCGUserParameterGet.cpp
--- a/Engine/Plugins/PCG/Source/PCG/Private/Elements/PCGUserParameterGet.cpp
+++ b/Engine/Plugins/PCG/Source/PCG/Private/Elements/PCGUserParameterGet.cpp
@@ -12,7 +12,7 @@
 #include "Elements/PCGCreateAttribute.h"
 #include "Helpers/PCGConversion.h"
 #include "Helpers/PCGDynamicTrackingHelpers.h"
-#include "Helpers/PCGHelpers.h"
+#include "Helpers/PCGGraphParametersHelpers.h"
 #include "Helpers/PCGPropertyHelpers.h"
 #include "Metadata/PCGAttributePropertySelector.h"
 
@@ -170,13 +170,6 @@ namespace PCGUserParameterGet
 	}
 } // namespace PCGUserParameterGet
 
-UPCGUserParameterGetSettings::UPCGUserParameterGetSettings(const FObjectInitializer& ObjectInitializer)
-{
-#if WITH_EDITORONLY_DATA
-	bExposeToLibrary = false;
-#endif
-}
-
 TArray<FPCGPinProperties> UPCGUserParameterGetSettings::OutputPinProperties() const
 {
 	TArray<FPCGPinProperties> PinProperties;
@@ -230,8 +223,57 @@ TArray<FPCGPreconfiguredInfo> UPCGUserParameterGetSettings::GetConversionInfo()
 
 	return ConversionInfo;
 }
+
+TArray<FPCGPreConfiguredSettingsInfo> UPCGUserParameterGetSettings::GetPreconfiguredInfo() const
+{
+	return FPCGPreConfiguredSettingsInfo::PopulateFromEnum<EPCGMetadataTypes>(
+		/*InValuesToSkip=*/{EPCGMetadataTypes::Count, EPCGMetadataTypes::Unknown},
+		/*InOptionalFormat=*/FTextFormat(LOCTEXT("PreconfigureFormat", "New {0} Parameter")));
+}
 #endif // WITH_EDITOR
 
+void UPCGUserParameterGetSettings::ApplyPreconfiguredSettings(const FPCGPreConfiguredSettingsInfo& PreconfigureInfo)
+{
+	if (PreconfigureInfo.PreconfiguredIndex < 0 || PreconfigureInfo.PreconfiguredIndex >= static_cast<uint8>(EPCGMetadataTypes::Count))
+	{
+		return;
+	}
+
+	EPCGMetadataTypes NewType = EPCGMetadataTypes::Unknown;
+	if (const UEnum* EnumPtr = StaticEnum<EPCGMetadataTypes>())
+	{
+		if (EnumPtr->IsValidEnumValue(PreconfigureInfo.PreconfiguredIndex))
+		{
+			NewType = static_cast<EPCGMetadataTypes>(EnumPtr->GetValueByIndex(PreconfigureInfo.PreconfiguredIndex));
+		}
+		else
+		{
+			return;
+		}
+	}
+
+	const UPCGNode* Node = CastChecked<UPCGNode>(GetOuter());
+	UPCGGraph* Graph = Node ? Node->GetGraph() : nullptr;
+	if (const FInstancedPropertyBag* PropertyBag = Graph ? Graph->GetUserParametersStruct() : nullptr)
+	{
+		FName NewPropertyName;
+		if (!PCGGraphParameter::Helpers::GenerateUniqueName(Graph, PropertyName))
+		{
+			PCGLog::LogErrorOnGraph(FText::Format(LOCTEXT("CouldNotCreateUniqueGraphParameterName", "Could not create a unique graph parameter with name '{0}'. Check if the graph is valid and the name has not reached its maximum."), FText::FromName(PropertyName)));
+			return;
+		}
+		check(PropertyBag->IsPropertyNameValid(NewPropertyName));
+
+		FPropertyBagPropertyDesc PropertyDesc = PCGPropertyHelpers::CreatePropertyBagDescWithMetadataType(NewPropertyName, NewType);
+		Graph->AddUserParameters({PropertyDesc});
+		const FPropertyBagPropertyDesc* Desc = PropertyBag->FindPropertyDescByName(NewPropertyName);
+		check(Desc);
+
+		PropertyGuid = Desc->ID;
+		PropertyName = Desc->Name;
+	}
+}
+
 bool UPCGUserParameterGetSettings::ConvertNode(const FPCGPreconfiguredInfo& ConversionInfo)
 {
 	using namespace PCGUserParameterGet::Constants::Conversion;
diff --git a/Engine/Plugins/PCG/Source/PCG/Private/Helpers/PCGGraphParametersHelpers.cpp b/Engine/Plugins/PCG/Source/PCG/Private/Helpers/PCGGraphParametersHelpers.cpp
--- a/Engine/Plugins/PCG/Source/PCG/Private/Helpers/PCGGraphParametersHelpers.cpp
+++ b/Engine/Plugins/PCG/Source/PCG/Private/Helpers/PCGGraphParametersHelpers.cpp
@@ -17,6 +17,8 @@
 namespace PCGGraphParametersHelpersPrivate
 {
 	static const FText InvalidGraphInterfaceError = LOCTEXT("InvalidGraphInterface", "Invalid graph interface.");
+	static const FName NewGraphParameterDefaultName(LOCTEXT("NewGraphParameterDefaultName", "NewProperty").ToString());
+	static constexpr uint64 MaxIncrementalNameIterations = 100;
 
 	void ThrowBlueprintException(const FText& ErrorMessage)
 	{
@@ -115,6 +117,39 @@ namespace PCGGraphParametersHelpersPrivate
 	}
 }
 
+namespace PCGGraphParameter::Helpers
+{
+	bool GenerateUniqueName(const UPCGGraphInterface* InGraph, FName& InOutName)
+	{
+		if (!InGraph)
+		{
+			return false;
+		}
+
+		if (InOutName == NAME_None)
+		{
+			InOutName = PCGGraphParametersHelpersPrivate::NewGraphParameterDefaultName;
+		}
+
+		const FInstancedPropertyBag* UserParametersStruct = InGraph->GetUserParametersStruct();
+		check(UserParametersStruct);
+
+		// Validate that the property name doesn't exist... if it does, increment until it doesn't
+		uint64 NameIteration = 1;
+		while (UserParametersStruct->FindPropertyDescByName(InOutName))
+		{
+			if (++NameIteration > PCGGraphParametersHelpersPrivate::MaxIncrementalNameIterations)
+			{
+				return false;
+			}
+
+			InOutName.SetNumber(NameIteration);
+		}
+
+		return true;
+	}
+} // PCGGraphParameter::Helpers
+
 //////////////////////////
 // Begin Blueprint Library
 //////////////////////////
diff --git a/Engine/Plugins/PCG/Source/PCG/Public/Elements/Metadata/PCGMetadataElementCommon.h b/Engine/Plugins/PCG/Source/PCG/Public/Elements/Metadata/PCGMetadataElementCommon.h
--- a/Engine/Plugins/PCG/Source/PCG/Public/Elements/Metadata/PCGMetadataElementCommon.h
+++ b/Engine/Plugins/PCG/Source/PCG/Public/Elements/Metadata/PCGMetadataElementCommon.h
@@ -300,6 +300,7 @@ namespace PCGMetadataElementCommon
 	UE_DEPRECATED(5.6, "Please use FPCGPreConfiguredSettingsInfo::PopulateFromEnum instead.")
 	TArray<FPCGPreConfiguredSettingsInfo> FillPreconfiguredSettingsInfoFromEnum(const TSet<EnumOperation>& InValuesToSkip = {}, const FText& InOptionalPrefix = FText())
 	{
-		return FPCGPreConfiguredSettingsInfo::PopulateFromEnum<EnumOperation>(InValuesToSkip, InOptionalPrefix);
+		FTextFormat Format = FTextFormat(FText::Format(INVTEXT("{0}{1}"), InOptionalPrefix, INVTEXT("{0}")));
+		return FPCGPreConfiguredSettingsInfo::PopulateFromEnum<EnumOperation>(InValuesToSkip, std::move(Format));
 	}
 }
diff --git a/Engine/Plugins/PCG/Source/PCG/Public/Elements/PCGCreateAttribute.h b/Engine/Plugins/PCG/Source/PCG/Public/Elements/PCGCreateAttribute.h
--- a/Engine/Plugins/PCG/Source/PCG/Public/Elements/PCGCreateAttribute.h
+++ b/Engine/Plugins/PCG/Source/PCG/Public/Elements/PCGCreateAttribute.h
@@ -102,6 +102,9 @@ public:
 	virtual EPCGSettingsType GetType() const override { return EPCGSettingsType::Param; }
 	virtual TArray<FText> GetNodeTitleAliases() const override;
 	virtual FText GetNodeTooltipText() const override;
+	virtual TArray<FPCGPreConfiguredSettingsInfo> GetPreconfiguredInfo() const override;
+	virtual bool GroupPreconfiguredSettings() const override { return true; }
+	virtual bool OnlyExposePreconfiguredSettings() const override { return true; }
 #endif
 	virtual FString GetAdditionalTitleInformation() const override;
 	virtual void ApplyPreconfiguredSettings(const FPCGPreConfiguredSettingsInfo& PreconfigureInfo) override;
diff --git a/Engine/Plugins/PCG/Source/PCG/Public/Elements/PCGUserParameterGet.h b/Engine/Plugins/PCG/Source/PCG/Public/Elements/PCGUserParameterGet.h
--- a/Engine/Plugins/PCG/Source/PCG/Public/Elements/PCGUserParameterGet.h
+++ b/Engine/Plugins/PCG/Source/PCG/Public/Elements/PCGUserParameterGet.h
@@ -27,8 +27,6 @@ class UPCGUserParameterGetSettings : public UPCGSettings
 	GENERATED_BODY()
 
 public:
-	UPCGUserParameterGetSettings(const FObjectInitializer& ObjectInitializer);
-
 	UPROPERTY()
 	FGuid PropertyGuid;
 
@@ -65,7 +63,11 @@ protected:
 	virtual bool CanUserEditTitle() const override { return true; }
 	virtual EPCGSettingsType GetType() const override { return EPCGSettingsType::GraphParameters; }
 	virtual TArray<FPCGPreconfiguredInfo> GetConversionInfo() const override;
+	virtual TArray<FPCGPreConfiguredSettingsInfo> GetPreconfiguredInfo() const override;
+	virtual bool OnlyExposePreconfiguredSettings() const override { return true; }
+	virtual bool GroupPreconfiguredSettings() const override { return false; }
 #endif
+	virtual void ApplyPreconfiguredSettings(const FPCGPreConfiguredSettingsInfo& PreconfigureInfo) override;
 
 	virtual bool ConvertNode(const FPCGPreconfiguredInfo& ConversionInfo) override;
 	//~End UPCGSettings interface
diff --git a/Engine/Plugins/PCG/Source/PCG/Public/Helpers/PCGGraphParametersHelpers.h b/Engine/Plugins/PCG/Source/PCG/Public/Helpers/PCGGraphParametersHelpers.h
--- a/Engine/Plugins/PCG/Source/PCG/Public/Helpers/PCGGraphParametersHelpers.h
+++ b/Engine/Plugins/PCG/Source/PCG/Public/Helpers/PCGGraphParametersHelpers.h
@@ -10,6 +10,11 @@
 
 class UPCGGraphInterface;
 
+namespace PCGGraphParameter::Helpers
+{
+	PCG_API bool GenerateUniqueName(const UPCGGraphInterface* InGraph, FName& InOutName);
+}
+
 /**
 * Blueprint Library to get or set graph parameters on graphs and graph instances
 */
diff --git a/Engine/Plugins/PCG/Source/PCG/Public/PCGCommon.h b/Engine/Plugins/PCG/Source/PCG/Public/PCGCommon.h
--- a/Engine/Plugins/PCG/Source/PCG/Public/PCGCommon.h
+++ b/Engine/Plugins/PCG/Source/PCG/Public/PCGCommon.h
@@ -185,6 +185,16 @@ enum class EPCGExclusiveDataType : uint8
 	Resource UMETA(Hidden)
 };
 
+UENUM(meta=(DisplayName="PCG Container Type"))
+enum class EPCGContainerType : uint8
+{
+	Element = 0,
+	None = Element,
+	Array,
+	Map,
+	Set
+};
+
 namespace PCGValueConstants
 {
 	constexpr int32 DefaultSeed = 42;
diff --git a/Engine/Plugins/PCG/Source/PCG/Public/Utils/PCGPreconfiguration.h b/Engine/Plugins/PCG/Source/PCG/Public/Utils/PCGPreconfiguration.h
--- a/Engine/Plugins/PCG/Source/PCG/Public/Utils/PCGPreconfiguration.h
+++ b/Engine/Plugins/PCG/Source/PCG/Public/Utils/PCGPreconfiguration.h
@@ -2,8 +2,15 @@
 
 #pragma once
 
+#include "PCGCommon.h"
+#include "Metadata/PCGMetadataAttributeTraits.h"
+
 #include "PCGPreconfiguration.generated.h"
 
+#if WITH_EDITOR
+struct FSlateBrush;
+#endif // WITH_EDITOR
+
 /**
 * Pre-configured info
 * Utility struct for use with any type of PCG preconfiguration.
@@ -34,15 +41,17 @@ struct FPCGPreconfiguredInfo
 	* Can also specify explicitly values that should not be included, i.e. cases that may not be available in non-editor builds.
 	* Can pass in an optional function or lambda for filtering or finer control over the results, for things like metadata.
 	* @param InValuesToSkip A set of enum values that should be skipped during population, like counts or hidden values.
-	* @param InOptionalPrefix An optional string prefix to the name of the action as it appears in the contextual search.
+	* @param InOptionalFormat An optional string format to fit the name of the action into, as it appears in the contextual search.
 	* @param ProcessFunc A post-process or filtering callback for each enum value. Returns false if the enum should be filtered.
 	* @returns An array of preconfigured info for schema actions or conversion, etc.
 	*/
 	template <typename EnumType, typename SubclassType = FPCGPreconfiguredInfo>
-	static TArray<SubclassType> PopulateFromEnum(const TSet<EnumType>& InValuesToSkip = {}, const FText& InOptionalPrefix = FText(), TFunction<bool(SubclassType& InOutInfo, const UEnum* EnumPtr, int32 ValueIndex)> ProcessFunc = nullptr)
+	static TArray<SubclassType> PopulateFromEnum(const TSet<EnumType>& InValuesToSkip = {}, const FTextFormat& InOptionalFormat = FTextFormat(), const TFunction<bool(SubclassType& InOutInfo, const UEnum* EnumPtr, int32 ValueIndex)>& ProcessFunc = nullptr)
 	{
 		static_assert(std::is_base_of_v<FPCGPreconfiguredInfo, SubclassType>, "SubclassType must be a subclass of FPCGPreconfiguredInfo");
 
+		static const FTextFormat EmptyFormat = INVTEXT("{0}");
+		const FTextFormat Format = InOptionalFormat.GetSourceText().IsEmpty() ? EmptyFormat : InOptionalFormat;
 		TArray<SubclassType> PreconfiguredInfo;
 
 		if (const UEnum* EnumPtr = StaticEnum<EnumType>())
@@ -57,7 +66,7 @@ struct FPCGPreconfiguredInfo
 					FText DisplayName = EnumPtr->GetDisplayNameTextByValue(Value);
 					if (!DisplayName.IsEmpty())
 					{
-						SubclassType Info(Value, FText::Format(FText::FromString(TEXT("{0}{1}")), InOptionalPrefix, std::move(DisplayName)));
+						SubclassType Info(Value, FText::Format(Format, std::move(DisplayName)));
 						if (!ProcessFunc || ProcessFunc(Info, EnumPtr, i))
 						{
 							PreconfiguredInfo.Emplace(std::move(Info));
@@ -84,6 +93,65 @@ struct FPCGPreconfiguredInfo
 #endif // WITH_EDITORONLY_DATA
 };
 
+// The type of descriptor used to generate an icon and color for an action menu item.
+enum class EPCGActionIconDescriptorType : uint8
+{
+	Label = 0,
+	Metadata
+};
+
+// Simple interface for overriding an action icon
+class IPCGActionIconDescriptorBase
+{
+public:
+	virtual ~IPCGActionIconDescriptorBase() = default;
+	virtual EPCGActionIconDescriptorType GetActionIconDescriptorType() const = 0;
+};
+
+// Descriptor for setting the icon and color directly.
+struct FPCGActionIconByLabelDescriptor : IPCGActionIconDescriptorBase
+{
+	virtual EPCGActionIconDescriptorType GetActionIconDescriptorType() const override { return EPCGActionIconDescriptorType::Label; }
+
+	FName BrushLabel = NAME_None;
+	FLinearColor Tint = FLinearColor::White;
+};
+
+// Descriptor to convert an EPCGMetadataTypes value directly.
+struct FPCGActionIconByMetadataDescriptor : IPCGActionIconDescriptorBase
+{
+	explicit FPCGActionIconByMetadataDescriptor(const EPCGMetadataTypes InType, const EPCGContainerType InContainerType = EPCGContainerType::Element)
+		: MetadataType(InType)
+		, ContainerType(InContainerType) {}
+
+	virtual EPCGActionIconDescriptorType GetActionIconDescriptorType() const override { return EPCGActionIconDescriptorType::Metadata; }
+	EPCGMetadataTypes GetMetadataType() const { return MetadataType; }
+	EPCGContainerType GetContainerType() const { return ContainerType; }
+
+private:
+	const EPCGMetadataTypes MetadataType;
+	const EPCGContainerType ContainerType;
+};
+
+/**
+ * @todo_pcg:
+ * Look into converting these descriptors as a type trait for extensibility. Right now, MetadataDescriptor requires
+ * a hardcoded lookup in the Editor, which could also potentially be made better extensible with a Factory pattern.
+ *
+ * Ex.
+ * template <typename T>
+ * struct IPCGActorIconTrait
+ * {
+ *	using ActionIconType = FPCGActionIconByLabelDescriptor;
+ * };
+ * 
+ * template <>
+ * struct IPCGActorIconTrait<EPCGMetadataTypes>
+ * {
+ * 	using ActionIconType = FPCGActionIconByMetadataDescriptor;
+ * };
+ */
+
 /**
 * Pre-configured settings info
 * Will be passed to the settings to pre-configure the settings on creation. Also used for pre-configured node elements.
@@ -95,6 +163,8 @@ struct FPCGPreConfiguredSettingsInfo : public FPCGPreconfiguredInfo
 	GENERATED_BODY()
 
 	static constexpr TCHAR SearchHintMetadataKey[] = TEXT("SearchHints");
+	static constexpr TCHAR ActionIconMetadataKey[] = TEXT("ActionIcon");
+	static constexpr TCHAR ActionIconTintMetadataKey[] = TEXT("ActionIconTint");
 
 	FPCGPreConfiguredSettingsInfo() = default;
 
@@ -103,32 +173,68 @@ struct FPCGPreConfiguredSettingsInfo : public FPCGPreconfiguredInfo
 	{}
 
 #if WITH_EDITOR
-	FPCGPreConfiguredSettingsInfo(int32 InIndex, FText InLabel, FText InTooltip, FText InSearchHints = {})
+	FPCGPreConfiguredSettingsInfo(const int32 InIndex, const FText& InLabel, const FText& InTooltip, const FText& InSearchHints = {}, TSharedPtr<IPCGActionIconDescriptorBase> InActionIconDescriptor = nullptr)
 		: FPCGPreconfiguredInfo(InIndex, InLabel, InTooltip)
 		, SearchHints(InSearchHints)
-	{}
+		, ActionIconDescriptor(std::move(InActionIconDescriptor)) {}
 #endif // WITH_EDITOR
 
 	template <typename EnumType>
-	static TArray<FPCGPreConfiguredSettingsInfo> PopulateFromEnum(const TSet<EnumType>& InValuesToSkip = {}, const FText& InOptionalPrefix = FText())
+	static TArray<FPCGPreConfiguredSettingsInfo> PopulateFromEnum(const TSet<EnumType>& InValuesToSkip = {}, const FTextFormat& InOptionalFormat = FTextFormat())
 	{
+		static_assert(std::is_enum_v<EnumType>, "Must use an enum to populate.");
+
 #if WITH_EDITOR
-		return FPCGPreconfiguredInfo::PopulateFromEnum<EnumType, FPCGPreConfiguredSettingsInfo>(InValuesToSkip, InOptionalPrefix, [](FPCGPreConfiguredSettingsInfo& InOutInfo, const UEnum* EnumPtr, const int32 ValueIndex) -> bool
+		return FPCGPreconfiguredInfo::PopulateFromEnum<EnumType, FPCGPreConfiguredSettingsInfo>(std::move(InValuesToSkip), std::move(InOptionalFormat), [](FPCGPreConfiguredSettingsInfo& InOutInfo, const UEnum* EnumPtr, const int32 ValueIndex) -> bool
 		{
-			if (EnumPtr && EnumPtr->HasMetaData(SearchHintMetadataKey, ValueIndex))
+			if (EnumPtr)
 			{
-				InOutInfo.SearchHints = FText::FromString(EnumPtr->GetMetaData(SearchHintMetadataKey, ValueIndex));
+				if (EnumPtr->HasMetaData(SearchHintMetadataKey, ValueIndex))
+				{
+					InOutInfo.SearchHints = FText::FromString(EnumPtr->GetMetaData(SearchHintMetadataKey, ValueIndex));
+					InOutInfo.EnumPtr = EnumPtr;
+				}
+
+				// For anything we populate from metadata type enum, add the type icon as the palette icon
+				if constexpr (std::is_same_v<EnumType, EPCGMetadataTypes>)
+				{
+					if (!InOutInfo.ActionIconDescriptor)
+					{
+						InOutInfo.ActionIconDescriptor = MakeShared<FPCGActionIconByMetadataDescriptor>(static_cast<EPCGMetadataTypes>(EnumPtr->GetValueByIndex(ValueIndex)));
+					}
+				}
+				else // Get the brush icon and color directly from the metadata
+				{
+					FPCGActionIconByLabelDescriptor IconDescriptor;
+					if (EnumPtr->HasMetaData(ActionIconMetadataKey, ValueIndex))
+					{
+						IconDescriptor.BrushLabel = FName(EnumPtr->GetMetaData(ActionIconMetadataKey, ValueIndex));
+					}
+
+					if (EnumPtr->HasMetaData(ActionIconTintMetadataKey, ValueIndex))
+					{
+						// @todo_pcg: Parse the metadata into RGBA channels. For now just make it white.
+					}
+					IconDescriptor.Tint = FLinearColor::White;
+
+					InOutInfo.ActionIconDescriptor = MakeShared<FPCGActionIconByLabelDescriptor>(std::move(IconDescriptor));
+				}
 			}
 
 			return true;
 		});
 #else
-		return FPCGPreconfiguredInfo::PopulateFromEnum<EnumType, FPCGPreConfiguredSettingsInfo>(InValuesToSkip, InOptionalPrefix);
+		return FPCGPreconfiguredInfo::PopulateFromEnum<EnumType, FPCGPreConfiguredSettingsInfo>(std::move(InValuesToSkip), std::move(InOptionalFormat));
 #endif // WITH_EDITOR
 	}
 
 #if WITH_EDITORONLY_DATA
 	UPROPERTY(EditAnywhere, Category = "")
 	FText SearchHints;
+
+	TSharedPtr<IPCGActionIconDescriptorBase> ActionIconDescriptor = nullptr;
+
+private:
+	TObjectPtr<const UEnum> EnumPtr = nullptr;
 #endif // WITH_EDITORONLY_DATA
 };
diff --git a/Engine/Plugins/PCG/Source/PCGEditor/Private/PCGEditor.cpp b/Engine/Plugins/PCG/Source/PCGEditor/Private/PCGEditor.cpp
--- a/Engine/Plugins/PCG/Source/PCGEditor/Private/PCGEditor.cpp
+++ b/Engine/Plugins/PCG/Source/PCGEditor/Private/PCGEditor.cpp
@@ -30,6 +30,7 @@
 #include "Nodes/PCGEditorGraphNodeOutput.h"
 #include "Nodes/PCGEditorGraphNodeReroute.h"
 
+#include "Widgets/SPCGEditorGraphActionMenu.h"
 #include "Widgets/SPCGEditorGraphAttributeListView.h"
 #include "Widgets/SPCGEditorGraphDebugObjectTree.h"
 #include "Widgets/SPCGEditorGraphDetailsView.h"
@@ -49,6 +50,7 @@
 #include "IAssetTools.h"
 #include "LevelEditor.h"
 #include "ScopedTransaction.h"
+#include "SGraphEditorActionMenu.h"
 #include "ShaderCore.h"
 #include "SNodePanel.h"
 #include "SourceCodeNavigation.h"
@@ -69,6 +71,7 @@
 #include "Misc/TransactionObjectEvent.h"
 #include "Preferences/UnrealEdOptions.h"
 #include "Widgets/Docking/SDockTab.h"
+#include "Widgets/Input/SEditableTextBox.h"
 #include "Widgets/Notifications/SNotificationList.h"
 
 #define LOCTEXT_NAMESPACE "PCGGraphEditor"
@@ -1616,6 +1619,17 @@ FReply FPCGEditor::OnSpawnNodeByShortcut(FInputChord InChord, const FVector2f& I
 	}
 }
 
+FActionMenuContent FPCGEditor::OnCreateActionMenuContent(UEdGraph* InGraph, const FVector2f& Location, const TArray<UEdGraphPin*>& InDraggedPins, bool bAutoExpand, SGraphEditor::FActionMenuClosed OnMenuClosed)
+{
+	SAssignNew(ActionMenu, SPCGEditorGraphActionMenu)
+	.GraphObj(InGraph)
+	.NewNodePosition(Location)
+	.DraggedFromPins(InDraggedPins)
+	.AutoExpandActionMenu(bAutoExpand);
+
+	return FActionMenuContent(ActionMenu.ToSharedRef(), ActionMenu->GetFilterTextBox());
+}
+
 void FPCGEditor::OnCollapseNodesInSubgraph()
 {
 	if (!InternalValidationOnAction())
@@ -2781,6 +2795,7 @@ TSharedRef<SGraphEditor> FPCGEditor::CreateGraphEditorWidget()
 	InEvents.OnTextCommitted = FOnNodeTextCommitted::CreateSP(this, &FPCGEditor::OnNodeTitleCommitted);
 	InEvents.OnNodeDoubleClicked = FSingleNodeEvent::CreateSP(this, &FPCGEditor::OnNodeDoubleClicked);
 	InEvents.OnSpawnNodeByShortcutAtLocation = SGraphEditor::FOnSpawnNodeByShortcutAtLocation::CreateSP(this, &FPCGEditor::OnSpawnNodeByShortcut, PCGEditorGraph);
+	InEvents.OnCreateActionMenuAtLocation = SGraphEditor::FOnCreateActionMenuAtLocation::CreateSP(this, &FPCGEditor::OnCreateActionMenuContent);
 
 	return SNew(SGraphEditor)
 		.AdditionalCommands(GraphEditorCommands)
diff --git a/Engine/Plugins/PCG/Source/PCGEditor/Private/PCGEditor.h b/Engine/Plugins/PCG/Source/PCGEditor/Private/PCGEditor.h
--- a/Engine/Plugins/PCG/Source/PCGEditor/Private/PCGEditor.h
+++ b/Engine/Plugins/PCG/Source/PCGEditor/Private/PCGEditor.h
@@ -6,10 +6,11 @@
 #include "Graph/PCGStackContext.h"
 
 #include "EditorUndoClient.h"
+#include "GraphEditor.h"
 #include "Toolkits/AssetEditorToolkit.h"
 
-class FSpawnTabArgs;
 enum class ECheckBoxState : uint8;
+class SPCGEditorGraphActionMenu;
 namespace ETextCommit { enum Type : int; }
 
 class FUICommandList;
@@ -149,6 +150,8 @@ public:
 	FReply OnSpawnNodeByShortcut(FInputChord InChord, const FVector2D& InPosition, UPCGEditorGraph* InGraph);
 	FReply OnSpawnNodeByShortcut(FInputChord InChord, const FVector2f& InPosition, UPCGEditorGraph* InGraph);
 
+	FActionMenuContent OnCreateActionMenuContent(UEdGraph* InGraph, const FVector2f& Location, const TArray<UEdGraphPin*>& InDraggedPins, bool bAutoExpand, SGraphEditor::FActionMenuClosed OnMenuClosed);
+
 	FOnInspectedStackChanged OnInspectedStackChangedDelegate;
 
 protected:
@@ -451,6 +454,8 @@ private:
 	TObjectPtr<UPCGGraph> PCGGraphBeingEdited = nullptr;
 	UPCGEditorGraph* PCGEditorGraph = nullptr;
 
+	TSharedPtr<SPCGEditorGraphActionMenu> ActionMenu;
+
 	TSoftObjectPtr<UPCGComponent> PCGComponentBeingInspected;
 	// Implementation note: we'll keep the last valid component inspected so we don't un-inspect on spurious selection changes
 	TSoftObjectPtr<UPCGComponent> LastValidPCGComponentBeingInspected;
diff --git a/Engine/Plugins/PCG/Source/PCGEditor/Private/PCGEditorGraphSchema.cpp b/Engine/Plugins/PCG/Source/PCGEditor/Private/PCGEditorGraphSchema.cpp
--- a/Engine/Plugins/PCG/Source/PCGEditor/Private/PCGEditorGraphSchema.cpp
+++ b/Engine/Plugins/PCG/Source/PCGEditor/Private/PCGEditorGraphSchema.cpp
@@ -21,6 +21,7 @@
 #include "PCGEditorGraphSchemaActions.h"
 #include "PCGEditorSettings.h"
 #include "PCGEditorUtils.h"
+#include "Helpers/PCGPropertyHelpers.h"
 #include "Nodes/PCGEditorGraphNodeBase.h"
 #include "Nodes/PCGEditorGraphNodeReroute.h"
 
@@ -879,6 +880,65 @@ void UPCGEditorGraphSchema::OnPinConnectionDoubleCicked(UEdGraphPin* PinA, UEdGr
 	}
 }
 
+const FSlateBrush* UPCGEditorGraphSchema::GetMetadataTypeSlateBrush(const EPCGContainerType ContainerType) const
+{
+	switch (ContainerType)
+	{
+		case EPCGContainerType::Array:
+			return FAppStyle::GetBrush(TEXT("Kismet.VariableList.ArrayTypeIcon"));
+		case EPCGContainerType::Set:
+			return FAppStyle::GetBrush(TEXT("Kismet.VariableList.SetTypeIcon"));
+		// @todo_pcg: Enable after Map support finalizes
+		// case EPCGContainerType::Map:
+		//	return FAppStyle::GetBrush(TEXT("Kismet.VariableList.MapKeyTypeIcon"));
+		default: // fall-through
+		case EPCGContainerType::Element:
+			return FAppStyle::GetBrush(TEXT("Kismet.VariableList.TypeIcon"));
+	}
+}
+
+FLinearColor UPCGEditorGraphSchema::GetMetadataTypeColor(const EPCGMetadataTypes Type) const
+{
+	const UGraphEditorSettings* GraphSettings = GetDefault<UGraphEditorSettings>();
+	check(GraphSettings)
+
+	switch (Type)
+	{
+		case EPCGMetadataTypes::Float:
+			return GraphSettings->FloatPinTypeColor;
+		case EPCGMetadataTypes::Double:
+			return GraphSettings->DoublePinTypeColor;
+		case EPCGMetadataTypes::Integer32:
+			return GraphSettings->IntPinTypeColor;
+		case EPCGMetadataTypes::Integer64:
+			return GraphSettings->Int64PinTypeColor;
+		case EPCGMetadataTypes::Vector2: // fall-through
+		case EPCGMetadataTypes::Vector:  // fall-through
+		case EPCGMetadataTypes::Vector4: // fall-through
+			return GraphSettings->VectorPinTypeColor;
+		case EPCGMetadataTypes::Quaternion:
+			return GraphSettings->RotatorPinTypeColor;
+		case EPCGMetadataTypes::Transform:
+			return GraphSettings->TransformPinTypeColor;
+		case EPCGMetadataTypes::String:
+			return GraphSettings->StringPinTypeColor;
+		case EPCGMetadataTypes::Boolean:
+			return GraphSettings->BooleanPinTypeColor;
+		case EPCGMetadataTypes::Rotator:
+			return GraphSettings->RotatorPinTypeColor;
+		case EPCGMetadataTypes::Name:
+			return GraphSettings->NamePinTypeColor;
+		case EPCGMetadataTypes::SoftObjectPath:
+			return GraphSettings->SoftObjectPinTypeColor;
+		case EPCGMetadataTypes::SoftClassPath:
+			return GraphSettings->SoftClassPinTypeColor;
+		default: // fall-through
+		case EPCGMetadataTypes::Unknown:
+			ensure(false);
+			return FLinearColor::White;
+	}
+}
+
 FPCGEditorConnectionDrawingPolicy::FPCGEditorConnectionDrawingPolicy(int32 InBackLayerID, int32 InFrontLayerID, float InZoomFactor, const FSlateRect& InClippingRect, FSlateWindowElementList& InDrawElements, UEdGraph* InGraph)
 	: FConnectionDrawingPolicy(InBackLayerID, InFrontLayerID, InZoomFactor, InClippingRect, InDrawElements)
 	, Graph(CastChecked<UPCGEditorGraph>(InGraph))
diff --git a/Engine/Plugins/PCG/Source/PCGEditor/Private/PCGEditorGraphSchema.h b/Engine/Plugins/PCG/Source/PCGEditor/Private/PCGEditorGraphSchema.h
--- a/Engine/Plugins/PCG/Source/PCGEditor/Private/PCGEditorGraphSchema.h
+++ b/Engine/Plugins/PCG/Source/PCGEditor/Private/PCGEditorGraphSchema.h
@@ -47,6 +47,9 @@ public:
 	virtual void OnPinConnectionDoubleCicked(UEdGraphPin* PinA, UEdGraphPin* PinB, const FVector2f& GraphPosition) const override;
 	//~ End EdGraphSchema Interface
 
+	virtual const FSlateBrush* GetMetadataTypeSlateBrush(EPCGContainerType ContainerType) const;
+	virtual FLinearColor GetMetadataTypeColor(EPCGMetadataTypes Type) const;
+
 private:
 	void GetGraphActions(FGraphActionMenuBuilder& ActionMenuBuilder, const FPCGActionsFilter& ActionFilter, bool bIsContextual) const;
 	void GetNativeElementActions(FGraphActionMenuBuilder& ActionMenuBuilder, const FPCGActionsFilter& ActionFilter) const;
diff --git a/Engine/Plugins/PCG/Source/PCGEditor/Private/PCGEditorGraphSchemaActions.cpp b/Engine/Plugins/PCG/Source/PCGEditor/Private/PCGEditorGraphSchemaActions.cpp
--- a/Engine/Plugins/PCG/Source/PCGEditor/Private/PCGEditorGraphSchemaActions.cpp
+++ b/Engine/Plugins/PCG/Source/PCGEditor/Private/PCGEditorGraphSchemaActions.cpp
@@ -8,7 +8,7 @@
 #include "Nodes/PCGEditorGraphNode.h"
 #include "Nodes/PCGEditorGraphNodeReroute.h"
 
-#include "PCGDataAsset.h"
+#include "PCGEditorGraphSchema.h"
 #include "PCGGraph.h"
 #include "PCGNode.h"
 #include "PCGSubgraph.h"
@@ -88,6 +88,48 @@ UEdGraphNode* FPCGEditorGraphSchemaAction_NewNativeElement::PerformAction(UEdGra
 	return NewNode;
 }
 
+FSlateBrush const* FPCGEditorGraphSchemaAction_NewNativeElement::GetPaletteIcon() const
+{
+	// Get the icon from the descriptor if it exists. Should only need to do this once.
+	if (!PaletteIcon && PreconfiguredInfo.ActionIconDescriptor)
+	{
+		const IPCGActionIconDescriptorBase* IconDescriptor = PreconfiguredInfo.ActionIconDescriptor.Get();
+		const FSlateBrush* BaseBrush = nullptr;
+		FLinearColor TintColor = FLinearColor::White;
+		switch (IconDescriptor->GetActionIconDescriptorType())
+		{
+			case EPCGActionIconDescriptorType::Label:
+			{
+				const FPCGActionIconByLabelDescriptor* LabelDescriptor = static_cast<const FPCGActionIconByLabelDescriptor*>(IconDescriptor);
+				BaseBrush = FAppStyle::GetBrush(LabelDescriptor->BrushLabel);
+				TintColor = LabelDescriptor->Tint;
+			}
+			break;
+			case EPCGActionIconDescriptorType::Metadata:
+			{
+				const UPCGEditorGraphSchema* Schema = GetDefault<UPCGEditorGraphSchema>();
+				check(Schema);
+
+				const FPCGActionIconByMetadataDescriptor* MetadataDescriptor = static_cast<const FPCGActionIconByMetadataDescriptor*>(IconDescriptor);
+				BaseBrush = Schema->GetMetadataTypeSlateBrush(MetadataDescriptor->GetContainerType());
+				TintColor = Schema->GetMetadataTypeColor(MetadataDescriptor->GetMetadataType());
+			}
+			break;
+			default:
+				// Invalid descriptor type
+				checkNoEntry();
+		}
+
+		if (BaseBrush)
+		{
+			PaletteIcon = MakeShared<FSlateBrush>(*BaseBrush);
+			PaletteIcon->TintColor = TintColor;
+		}
+	}
+
+	return PaletteIcon ? PaletteIcon.Get() : nullptr;
+}
+
 void FPCGEditorGraphSchemaAction_NewGetParameterElement::PostCreation(UPCGNode* NewNode)
 {
 	check(NewNode);
diff --git a/Engine/Plugins/PCG/Source/PCGEditor/Private/PCGEditorGraphSchemaActions.h b/Engine/Plugins/PCG/Source/PCGEditor/Private/PCGEditorGraphSchemaActions.h
--- a/Engine/Plugins/PCG/Source/PCGEditor/Private/PCGEditorGraphSchemaActions.h
+++ b/Engine/Plugins/PCG/Source/PCGEditor/Private/PCGEditorGraphSchemaActions.h
@@ -15,6 +15,7 @@ class UEdGraphPin;
 class UPCGBlueprintElement;
 class UPCGNode;
 class UPCGEditorGraph;
+struct FSlateBrush;
 
 UENUM()
 enum class EPCGEditorNewSettingsBehavior : uint8
@@ -59,8 +60,13 @@ struct FPCGEditorGraphSchemaAction_NewNativeElement : public FEdGraphSchemaActio
 	virtual UEdGraphNode* PerformAction(UEdGraph* ParentGraph, UEdGraphPin* FromPin, const FVector2f& Location, bool bSelectNewNode = true) override;
 	//~ End FEdGraphSchemaAction Interface
 
+	virtual FSlateBrush const* GetPaletteIcon() const override;
+
 protected:
 	virtual void PostCreation(UPCGNode* NewNode) {}
+
+private:
+	mutable TSharedPtr<FSlateBrush> PaletteIcon = nullptr;
 };
 
 USTRUCT()
diff --git a/Engine/Plugins/PCG/Source/PCGEditor/Private/Widgets/SPCGEditorGraphActionMenu.h b/Engine/Plugins/PCG/Source/PCGEditor/Private/Widgets/SPCGEditorGraphActionMenu.h
new file mode 100644
--- /dev/null
+++ b/Engine/Plugins/PCG/Source/PCGEditor/Private/Widgets/SPCGEditorGraphActionMenu.h
@@ -0,0 +1,56 @@
+// Copyright Epic Games, Inc. All Rights Reserved.
+
+#pragma once
+
+#include "SGraphActionMenu.h"
+#include "SGraphEditorActionMenu.h"
+#include "SPCGEditorGraphActionWidget.h"
+
+class SPCGEditorGraphActionMenu : public SGraphEditorActionMenu
+{
+	SLATE_BEGIN_ARGS(SPCGEditorGraphActionMenu)
+			: _GraphObj(nullptr)
+			, _NewNodePosition(FVector2f::ZeroVector)
+			, _AutoExpandActionMenu(false) {}
+
+		SLATE_ARGUMENT(UEdGraph*, GraphObj)
+		SLATE_ARGUMENT(FDeprecateSlateVector2D, NewNodePosition)
+		SLATE_ARGUMENT(TArray<UEdGraphPin*>, DraggedFromPins)
+		SLATE_ARGUMENT(SGraphEditor::FActionMenuClosed, OnClosedCallback)
+		SLATE_ARGUMENT(bool, AutoExpandActionMenu)
+	SLATE_END_ARGS()
+
+	void Construct(const FArguments& InArgs)
+	{
+		GraphObj = InArgs._GraphObj;
+		DraggedFromPins = InArgs._DraggedFromPins;
+		NewNodePosition = InArgs._NewNodePosition;
+		OnClosedCallback = InArgs._OnClosedCallback;
+		AutoExpandActionMenu = InArgs._AutoExpandActionMenu;
+
+		// This mirrors the default SGraphActionMenu
+		SBorder::Construct(SBorder::FArguments()
+			.BorderImage(FAppStyle::GetBrush("Menu.Background"))
+			.Padding(5.0f)
+			[
+				SNew(SBox)
+				.WidthOverride(400.0f)
+				.HeightOverride(400.0f)
+				[
+					SAssignNew(GraphActionMenu, SGraphActionMenu)
+					.GraphObj(GraphObj)
+					.OnActionSelected(this, &SPCGEditorGraphActionMenu::OnActionSelected)
+					.OnCollectAllActions(this, &SPCGEditorGraphActionMenu::CollectAllActions)
+					.AutoExpandActionMenu(AutoExpandActionMenu)
+					.DraggedFromPins(DraggedFromPins)
+					.OnCreateWidgetForAction(this, &SPCGEditorGraphActionMenu::OnCreateWidgetForAction)
+				]
+			]
+		);
+	}
+
+	TSharedRef<SWidget> OnCreateWidgetForAction(FCreateWidgetForActionData* const Data) const
+	{
+		return SNew(SPCGGraphActionWidget, Data);
+	}
+};
diff --git a/Engine/Plugins/PCG/Source/PCGEditor/Private/Widgets/SPCGEditorGraphActionWidget.h b/Engine/Plugins/PCG/Source/PCGEditor/Private/Widgets/SPCGEditorGraphActionWidget.h
new file mode 100644
--- /dev/null
+++ b/Engine/Plugins/PCG/Source/PCGEditor/Private/Widgets/SPCGEditorGraphActionWidget.h
@@ -0,0 +1,55 @@
+// Copyright Epic Games, Inc. All Rights Reserved.
+
+#pragma once
+
+#include "SGraphActionMenu.h"
+#include "EdGraph/EdGraphSchema.h"
+#include "Widgets/Images/SLayeredImage.h"
+
+class SPCGGraphActionWidget : public SCompoundWidget
+{
+	SLATE_BEGIN_ARGS(SPCGGraphActionWidget) {}
+		SLATE_ARGUMENT(TSharedPtr<SWidget>, NameWidget)
+	SLATE_END_ARGS()
+
+	void Construct(const FArguments& InArgs, const FCreateWidgetForActionData* const InCreateData)
+	{
+		check(InCreateData->Action.IsValid());
+
+		TSharedRef<SWidget> IconWidget = SNullWidget::NullWidget;
+		const TSharedPtr<FEdGraphSchemaAction>& Action = InCreateData->Action;
+
+		if (const FSlateBrush* const PrimaryIcon = Action->GetPaletteIcon())
+		{
+			// @todo_pcg: Support secondary icon/colors for more complex type icons (like Maps)
+			const FSlateBrush* const SecondaryIcon = nullptr;
+			const FLinearColor SecondaryColor = FLinearColor::White;
+
+			SAssignNew(IconWidget, SLayeredImage, SecondaryIcon, SecondaryColor)
+			.DesiredSizeOverride(FVector2D(16))
+			.Image(PrimaryIcon);
+		}
+
+		TSharedRef<SWidget> NameSlotWidget =
+			InArgs._NameWidget
+				? InArgs._NameWidget.ToSharedRef()
+				: SNew(STextBlock)
+				.Text(Action->GetMenuDescription());
+
+		this->ChildSlot
+		[
+			SNew(SHorizontalBox)
+			+ SHorizontalBox::Slot()
+			.AutoWidth()
+			[
+				IconWidget
+			]
+			+ SHorizontalBox::Slot()
+			.AutoWidth()
+			.Padding(4, 0, 0, 0)
+			[
+				std::move(NameSlotWidget)
+			]
+		];
+	}
+};
diff --git a/Engine/Plugins/PCG/Source/PCGEditor/Private/Widgets/SPCGEditorGraphNodePalette.cpp b/Engine/Plugins/PCG/Source/PCGEditor/Private/Widgets/SPCGEditorGraphNodePalette.cpp
--- a/Engine/Plugins/PCG/Source/PCGEditor/Private/Widgets/SPCGEditorGraphNodePalette.cpp
+++ b/Engine/Plugins/PCG/Source/PCGEditor/Private/Widgets/SPCGEditorGraphNodePalette.cpp
@@ -2,7 +2,6 @@
 
 #include "SPCGEditorGraphNodePalette.h"
 
-#include "AssetRegistry/AssetData.h"
 #include "PCGGraph.h"
 
 #include "PCGEditor.h"
@@ -11,6 +10,7 @@
 #include "PCGEditorGraphSchemaActions.h"
 #include "PCGEditorUtils.h"
 #include "PCGSettingsDragDropAction.h"
+#include "Widgets/SPCGEditorGraphActionWidget.h"
 
 #include "SEnumCombo.h"
 #include "AssetRegistry/AssetRegistryModule.h"
@@ -25,12 +25,10 @@ void SPCGEditorGraphNodePaletteItem::Construct(const FArguments& InArgs, FCreate
 	TSharedPtr<FEdGraphSchemaAction> GraphAction = InCreateData->Action;
 	ActionPtr = InCreateData->Action;
 
-	bool bIsReadOnly = false;
-	TSharedRef<SWidget> NameSlotWidget = CreateTextSlotWidget(InCreateData, bIsReadOnly);
-
 	this->ChildSlot
 	[
-		NameSlotWidget
+		SNew(SPCGGraphActionWidget, InCreateData)
+		.NameWidget(CreateTextSlotWidget(InCreateData, InCreateData->bIsReadOnly).ToSharedPtr())
 	];
 }
 
