[PCG] Fix crash encountered when using VT data interface.

The defines used in the shader did not agree with the ones in C++, which seems to indicate a collision. Add PCG prefixes.

#rnx

diff --git a/Engine/Plugins/PCG/Shaders/Private/PCGVirtualTextureDataInterface.ush b/Engine/Plugins/PCG/Shaders/Private/PCGVirtualTextureDataInterface.ush
--- a/Engine/Plugins/PCG/Shaders/Private/PCGVirtualTextureDataInterface.ush
+++ b/Engine/Plugins/PCG/Shaders/Private/PCGVirtualTextureDataInterface.ush
@@ -21,13 +21,13 @@ PCG_VIRTUAL_TEXTURE_DECLARATION(3)
 #undef PCG_VIRTUAL_TEXTURE_DECLARATION
 
 // Layer infos
-uint4 {DataInterfaceName}_MaterialTypes[MAX_NUM_BINDINGS];
-uint4 {DataInterfaceName}_ValidLayerMasks[MAX_NUM_BINDINGS];
-float4 {DataInterfaceName}_LWCTiles[MAX_NUM_BINDINGS];
-float4x4 {DataInterfaceName}_WorldToUVTransforms[MAX_NUM_BINDINGS];
-uint4 {DataInterfaceName}_Adaptive[MAX_NUM_BINDINGS * MAX_NUM_LAYERS];
-uint4 {DataInterfaceName}_Uniforms[MAX_NUM_BINDINGS * MAX_NUM_LAYERS];
-uint4 {DataInterfaceName}_PageTableUniforms[MAX_NUM_BINDINGS * 2];
+uint4 {DataInterfaceName}_MaterialTypes[PCG_MAX_NUM_BINDINGS];
+uint4 {DataInterfaceName}_ValidLayerMasks[PCG_MAX_NUM_BINDINGS];
+float4 {DataInterfaceName}_LWCTiles[PCG_MAX_NUM_BINDINGS];
+float4x4 {DataInterfaceName}_WorldToUVTransforms[PCG_MAX_NUM_BINDINGS];
+uint4 {DataInterfaceName}_Adaptive[PCG_MAX_NUM_BINDINGS * PCG_MAX_NUM_LAYERS];
+uint4 {DataInterfaceName}_Uniforms[PCG_MAX_NUM_BINDINGS * PCG_MAX_NUM_LAYERS];
+uint4 {DataInterfaceName}_PageTableUniforms[PCG_MAX_NUM_BINDINGS * 2];
 
 // Shared
 uint {DataInterfaceName}_NumVirtualTextures;
@@ -155,7 +155,7 @@ void SampleVirtualTexture_{DataInterfaceName}(
 	const uint ValidLayerMask = {DataInterfaceName}_ValidLayerMasks[InDataIndex].x;
 
 	{
-		const uint PackedLayerIndex = InDataIndex * MAX_NUM_LAYERS;
+		const uint PackedLayerIndex = InDataIndex * PCG_MAX_NUM_LAYERS;
 
 		uint Adaptive0 = {DataInterfaceName}_Adaptive[PackedLayerIndex + 0].x;
 		uint Adaptive1 = {DataInterfaceName}_Adaptive[PackedLayerIndex + 1].x;
diff --git a/Engine/Plugins/PCG/Source/PCG/Private/Compute/DataInterfaces/PCGVirtualTextureDataInterface.cpp b/Engine/Plugins/PCG/Source/PCG/Private/Compute/DataInterfaces/PCGVirtualTextureDataInterface.cpp
--- a/Engine/Plugins/PCG/Source/PCG/Private/Compute/DataInterfaces/PCGVirtualTextureDataInterface.cpp
+++ b/Engine/Plugins/PCG/Source/PCG/Private/Compute/DataInterfaces/PCGVirtualTextureDataInterface.cpp
@@ -22,7 +22,6 @@
 
 #define LOCTEXT_NAMESPACE "PCGVirtualTextureDataInterface"
 
-using namespace PCGVirtualTextureDataInterfaceConstants;
 static const TCHAR* EnableMultipleVirtualTexturePermutationName = TEXT("ENABLE_MULTIPLE_VIRTUAL_TEXTURES");
 
 void UPCGVirtualTextureDataInterface::GetSupportedInputs(TArray<FShaderFunctionDefinition>& OutFunctions) const
@@ -44,28 +43,28 @@ void UPCGVirtualTextureDataInterface::GetSupportedInputs(TArray<FShaderFunctionD
 
 BEGIN_SHADER_PARAMETER_STRUCT(FPCGVirtualTextureDataInterfaceParameters, )
 	// Layer 0
-	SHADER_PARAMETER_SRV_ARRAY(Texture2D, VirtualTexture0, [MAX_NUM_BINDINGS])
-	SHADER_PARAMETER_TEXTURE_ARRAY(Texture2D<uint4>, PageTable0, [MAX_NUM_BINDINGS])
-	SHADER_PARAMETER_TEXTURE_ARRAY(Texture2D<uint>, PageTableIndirection0, [MAX_NUM_BINDINGS])
+	SHADER_PARAMETER_SRV_ARRAY(Texture2D, VirtualTexture0, [PCGVirtualTextureDIConstants::MAX_NUM_BINDINGS])
+	SHADER_PARAMETER_TEXTURE_ARRAY(Texture2D<uint4>, PageTable0, [PCGVirtualTextureDIConstants::MAX_NUM_BINDINGS])
+	SHADER_PARAMETER_TEXTURE_ARRAY(Texture2D<uint>, PageTableIndirection0, [PCGVirtualTextureDIConstants::MAX_NUM_BINDINGS])
 
 	// Layer 1
-	SHADER_PARAMETER_SRV_ARRAY(Texture2D, VirtualTexture1, [MAX_NUM_BINDINGS])
-	SHADER_PARAMETER_TEXTURE_ARRAY(Texture2D<uint4>, PageTable1, [MAX_NUM_BINDINGS])
-	SHADER_PARAMETER_TEXTURE_ARRAY(Texture2D<uint>, PageTableIndirection1, [MAX_NUM_BINDINGS])
+	SHADER_PARAMETER_SRV_ARRAY(Texture2D, VirtualTexture1, [PCGVirtualTextureDIConstants::MAX_NUM_BINDINGS])
+	SHADER_PARAMETER_TEXTURE_ARRAY(Texture2D<uint4>, PageTable1, [PCGVirtualTextureDIConstants::MAX_NUM_BINDINGS])
+	SHADER_PARAMETER_TEXTURE_ARRAY(Texture2D<uint>, PageTableIndirection1, [PCGVirtualTextureDIConstants::MAX_NUM_BINDINGS])
 
 	// Layer 2
-	SHADER_PARAMETER_SRV_ARRAY(Texture2D, VirtualTexture2, [MAX_NUM_BINDINGS])
-	SHADER_PARAMETER_TEXTURE_ARRAY(Texture2D<uint4>, PageTable2, [MAX_NUM_BINDINGS])
-	SHADER_PARAMETER_TEXTURE_ARRAY(Texture2D<uint>, PageTableIndirection2, [MAX_NUM_BINDINGS])
+	SHADER_PARAMETER_SRV_ARRAY(Texture2D, VirtualTexture2, [PCGVirtualTextureDIConstants::MAX_NUM_BINDINGS])
+	SHADER_PARAMETER_TEXTURE_ARRAY(Texture2D<uint4>, PageTable2, [PCGVirtualTextureDIConstants::MAX_NUM_BINDINGS])
+	SHADER_PARAMETER_TEXTURE_ARRAY(Texture2D<uint>, PageTableIndirection2, [PCGVirtualTextureDIConstants::MAX_NUM_BINDINGS])
 
 	// Layer info
-	SHADER_PARAMETER_ARRAY(FUintVector4, MaterialTypes, [MAX_NUM_BINDINGS])
-	SHADER_PARAMETER_ARRAY(FUintVector4, ValidLayerMasks, [MAX_NUM_BINDINGS])
-	SHADER_PARAMETER_ARRAY(FVector4f, LWCTiles, [MAX_NUM_BINDINGS])
-	SHADER_PARAMETER_ARRAY(FMatrix44f, WorldToUVTransforms, [MAX_NUM_BINDINGS])
-	SHADER_PARAMETER_ARRAY(FUintVector4, Adaptive, [MAX_NUM_BINDINGS * MAX_NUM_LAYERS])
-	SHADER_PARAMETER_ARRAY(FUintVector4, Uniforms, [MAX_NUM_BINDINGS * MAX_NUM_LAYERS])
-	SHADER_PARAMETER_ARRAY(FUintVector4, PageTableUniforms, [MAX_NUM_BINDINGS * 2])
+	SHADER_PARAMETER_ARRAY(FUintVector4, MaterialTypes, [PCGVirtualTextureDIConstants::MAX_NUM_BINDINGS])
+	SHADER_PARAMETER_ARRAY(FUintVector4, ValidLayerMasks, [PCGVirtualTextureDIConstants::MAX_NUM_BINDINGS])
+	SHADER_PARAMETER_ARRAY(FVector4f, LWCTiles, [PCGVirtualTextureDIConstants::MAX_NUM_BINDINGS])
+	SHADER_PARAMETER_ARRAY(FMatrix44f, WorldToUVTransforms, [PCGVirtualTextureDIConstants::MAX_NUM_BINDINGS])
+	SHADER_PARAMETER_ARRAY(FUintVector4, Adaptive, [PCGVirtualTextureDIConstants::MAX_NUM_BINDINGS * PCGVirtualTextureDIConstants::MAX_NUM_LAYERS])
+	SHADER_PARAMETER_ARRAY(FUintVector4, Uniforms, [PCGVirtualTextureDIConstants::MAX_NUM_BINDINGS * PCGVirtualTextureDIConstants::MAX_NUM_LAYERS])
+	SHADER_PARAMETER_ARRAY(FUintVector4, PageTableUniforms, [PCGVirtualTextureDIConstants::MAX_NUM_BINDINGS * 2])
 
 	// Shared
 	SHADER_PARAMETER(uint32, NumVirtualTextures)
@@ -107,8 +106,8 @@ void UPCGVirtualTextureDataInterface::GetHLSL(FString& OutHLSL, FString const& I
 
 void UPCGVirtualTextureDataInterface::GetDefines(FComputeKernelDefinitionSet& OutDefinitionSet) const
 {
-	OutDefinitionSet.Defines.Add(FComputeKernelDefinition(TEXT("MAX_NUM_BINDINGS"), FString::FromInt(MAX_NUM_BINDINGS)));
-	OutDefinitionSet.Defines.Add(FComputeKernelDefinition(TEXT("MAX_NUM_LAYERS"), FString::FromInt(MAX_NUM_LAYERS)));
+	OutDefinitionSet.Defines.Add(FComputeKernelDefinition(TEXT("PCG_MAX_NUM_BINDINGS"), FString::FromInt(PCGVirtualTextureDIConstants::MAX_NUM_BINDINGS)));
+	OutDefinitionSet.Defines.Add(FComputeKernelDefinition(TEXT("PCG_MAX_NUM_LAYERS"), FString::FromInt(PCGVirtualTextureDIConstants::MAX_NUM_LAYERS)));
 }
 
 void UPCGVirtualTextureDataInterface::GetPermutations(FComputeKernelPermutationVector& OutPermutationVector) const
@@ -155,7 +154,7 @@ void UPCGVirtualTextureDataProvider::Initialize(const UComputeDataInterface* InD
 			continue;
 		}
 
-		if (VirtualTextures.Num() < MAX_NUM_BINDINGS)
+		if (VirtualTextures.Num() < PCGVirtualTextureDIConstants::MAX_NUM_BINDINGS)
 		{
 			VirtualTextures.Add(VirtualTextureData->GetRuntimeVirtualTexture());
 		}
@@ -163,7 +162,7 @@ void UPCGVirtualTextureDataProvider::Initialize(const UComputeDataInterface* InD
 		{
 			UE_LOG(LogPCG, Warning, TEXT("Texture data interface on pin '%s' received too many textures to bind. Only the first %d textures will be bound."),
 				*DataInterface->DownstreamInputPinLabelAliases[0].ToString(),
-				MAX_NUM_BINDINGS);
+				PCGVirtualTextureDIConstants::MAX_NUM_BINDINGS);
 			break;
 		}
 	}
@@ -191,8 +190,8 @@ FPCGVirtualTextureDataProviderProxy::FPCGVirtualTextureDataProviderProxy(const T
 		if (const URuntimeVirtualTexture* VirtualTexture = InVirtualTextures[VirtualTextureIndex])
 		{
 			MaterialType = VirtualTexture->GetMaterialType();
-			const int32 NumVirtualTextureLayers = FMath::Min(URuntimeVirtualTexture::GetLayerCount(MaterialType), MAX_NUM_LAYERS);
-			const int32 PackedLayerBaseIndex = VirtualTextureIndex * MAX_NUM_LAYERS;
+			const int32 NumVirtualTextureLayers = FMath::Min(URuntimeVirtualTexture::GetLayerCount(MaterialType), PCGVirtualTextureDIConstants::MAX_NUM_LAYERS);
+			const int32 PackedLayerBaseIndex = VirtualTextureIndex * PCGVirtualTextureDIConstants::MAX_NUM_LAYERS;
 
 			for (int32 LayerIndex = 0; LayerIndex < NumVirtualTextureLayers; ++LayerIndex)
 			{
@@ -245,9 +244,9 @@ void FPCGVirtualTextureDataProviderProxy::GatherDispatchData(FDispatchData const
 		Parameters.NumVirtualTextures = NumVirtualTextures;
 		Parameters.Sampler = TStaticSamplerState<SF_Bilinear, AM_Clamp, AM_Clamp, AM_Clamp>::GetRHI();
 
-		for (int32 BindingIndex = 0; BindingIndex < MAX_NUM_BINDINGS; ++BindingIndex)
+		for (int32 BindingIndex = 0; BindingIndex < PCGVirtualTextureDIConstants::MAX_NUM_BINDINGS; ++BindingIndex)
 		{
-			const int32 PackedLayerBaseIndex = BindingIndex * MAX_NUM_LAYERS;
+			const int32 PackedLayerBaseIndex = BindingIndex * PCGVirtualTextureDIConstants::MAX_NUM_LAYERS;
 			const PCGVirtualTextureCommon::FVirtualTextureLayer& Layer0 = Layers[PackedLayerBaseIndex + 0];
 			const PCGVirtualTextureCommon::FVirtualTextureLayer& Layer1 = Layers[PackedLayerBaseIndex + 1];
 			const PCGVirtualTextureCommon::FVirtualTextureLayer& Layer2 = Layers[PackedLayerBaseIndex + 2];
diff --git a/Engine/Plugins/PCG/Source/PCG/Private/Compute/DataInterfaces/PCGVirtualTextureDataInterface.h b/Engine/Plugins/PCG/Source/PCG/Private/Compute/DataInterfaces/PCGVirtualTextureDataInterface.h
--- a/Engine/Plugins/PCG/Source/PCG/Private/Compute/DataInterfaces/PCGVirtualTextureDataInterface.h
+++ b/Engine/Plugins/PCG/Source/PCG/Private/Compute/DataInterfaces/PCGVirtualTextureDataInterface.h
@@ -10,7 +10,7 @@
 class FPCGVirtualTextureDataInterfaceParameters;
 class URuntimeVirtualTexture;
 
-namespace PCGVirtualTextureDataInterfaceConstants
+namespace PCGVirtualTextureDIConstants
 {
 	/**
 	 * The maximum number of virtual textures which can be bound to a PCGVirtualTextureDataInterface.
@@ -81,8 +81,8 @@ public:
 protected:
 	using FParameters = FPCGVirtualTextureDataInterfaceParameters;
 
-	ERuntimeVirtualTextureMaterialType MaterialTypes[PCGVirtualTextureDataInterfaceConstants::MAX_NUM_BINDINGS];
-	PCGVirtualTextureCommon::FVirtualTextureLayer Layers[PCGVirtualTextureDataInterfaceConstants::MAX_NUM_LAYERS * PCGVirtualTextureDataInterfaceConstants::MAX_NUM_BINDINGS];
-	PCGVirtualTextureCommon::FVirtualTexturePageTable PageTables[PCGVirtualTextureDataInterfaceConstants::MAX_NUM_LAYERS * PCGVirtualTextureDataInterfaceConstants::MAX_NUM_BINDINGS];
+	ERuntimeVirtualTextureMaterialType MaterialTypes[PCGVirtualTextureDIConstants::MAX_NUM_BINDINGS];
+	PCGVirtualTextureCommon::FVirtualTextureLayer Layers[PCGVirtualTextureDIConstants::MAX_NUM_LAYERS * PCGVirtualTextureDIConstants::MAX_NUM_BINDINGS];
+	PCGVirtualTextureCommon::FVirtualTexturePageTable PageTables[PCGVirtualTextureDIConstants::MAX_NUM_LAYERS * PCGVirtualTextureDIConstants::MAX_NUM_BINDINGS];
 	int32 NumVirtualTextures = 0;
 };
