[PCG] Fixed distance behavior change. Will be revisited to promote new behavior in next version.

diff --git a/Engine/Plugins/PCG/Source/PCG/Private/Elements/PCGDistance.cpp b/Engine/Plugins/PCG/Source/PCG/Private/Elements/PCGDistance.cpp
--- a/Engine/Plugins/PCG/Source/PCG/Private/Elements/PCGDistance.cpp
+++ b/Engine/Plugins/PCG/Source/PCG/Private/Elements/PCGDistance.cpp
@@ -158,7 +158,9 @@ bool FPCGDistanceElement::ExecuteInternal(FPCGContext* Context) const
 		TargetPointDatas.Add(TargetPointData);
 	}
 
-	if (TargetPointDatas.IsEmpty())
+	// @todo_pcg: normally, we should do nothing if there is no target data, but this would break older data.
+	// We should introduce an advanced option to "perform" the operation even if there are no targets, make that true for old data, and false for newer data.
+	if (TargetPointDatas.IsEmpty() && !Settings->bOutputToAttribute)
 	{
 		// If Target pin has no valid data then we no-op and pass through all data from Source pin.
 		Context->OutputData = PCGGather::GatherDataForPin(Context->InputData, PCGDistance::SourceLabel);
@@ -189,7 +191,7 @@ bool FPCGDistanceElement::ExecuteInternal(FPCGContext* Context) const
 		if (!SourcePointData)
 		{
 			PCGE_LOG(Error, GraphAndLog, LOCTEXT("CannotConvertToPointData", "Cannot convert input Spatial data to Point data"));
-			continue;			
+			continue;
 		}
 
 		UPCGBasePointData* OutputData = FPCGContext::NewPointData_AnyThread(Context);
@@ -220,6 +222,12 @@ bool FPCGDistanceElement::ExecuteInternal(FPCGContext* Context) const
 			}
 		}
 
+		// There is nothing to do as we will search against nothing - leave the default value in the attribute.
+		if (TargetPointDatas.IsEmpty())
+		{
+			continue;
+		}
+
 		TUniquePtr<IPCGAttributeAccessor> Accessor;
 		TUniquePtr<IPCGAttributeAccessorKeys> Keys;
 
@@ -270,7 +278,7 @@ bool FPCGDistanceElement::ExecuteInternal(FPCGContext* Context) const
 				ResultCache.Distances.SetNumUninitialized(SourcePointData->GetNumPoints());
 			}
 		}
-				
+
 		auto ProcessDistanceFunc = [SourceShape, TargetShape, &TargetPointDatas, MaximumDistance, MaximumDistanceRecip, bSetDensity, bOutputDistanceVector, OutputData, SourcePointData, &ResultCache, bWriteToAttribute = Accessor.IsValid(), Settings, SourceIndex](int32 StartReadIndex, int32 StartWriteIndex, int32 Count)
 		{
 			check(StartReadIndex == StartWriteIndex);
