[PCG] Added visualization for primitives, volumes, collision shapes, collision wrapper, spline.

diff --git a/Engine/Plugins/PCG/Source/PCG/Public/Data/PCGCollisionShapeData.h b/Engine/Plugins/PCG/Source/PCG/Public/Data/PCGCollisionShapeData.h
--- a/Engine/Plugins/PCG/Source/PCG/Public/Data/PCGCollisionShapeData.h
+++ b/Engine/Plugins/PCG/Source/PCG/Public/Data/PCGCollisionShapeData.h
@@ -17,6 +17,8 @@ class UPCGCollisionShapeData : public UPCGSpatialDataWithPointCache
 {
 	GENERATED_BODY()
 
+	friend class FPCGCollisionShapeDataVisualization;
+
 public:
 	PCG_API void Initialize(UShapeComponent* InShape);
 	PCG_API static bool IsSupported(UShapeComponent* InShape);
diff --git a/Engine/Plugins/PCG/Source/PCG/Public/Data/PCGCollisionWrapperData.h b/Engine/Plugins/PCG/Source/PCG/Public/Data/PCGCollisionWrapperData.h
--- a/Engine/Plugins/PCG/Source/PCG/Public/Data/PCGCollisionWrapperData.h
+++ b/Engine/Plugins/PCG/Source/PCG/Public/Data/PCGCollisionWrapperData.h
@@ -60,6 +60,8 @@ class UPCGCollisionWrapperData : public UPCGSpatialData
 {
 	GENERATED_BODY()
 
+	friend class FPCGCollisionWrapperDataVisualization;
+
 public:
 	/** Inititializes the collision wrapper on a point data based on the provided attribute selector */
 	PCG_API bool Initialize(const UPCGBasePointData* InPointData, const FPCGAttributePropertyInputSelector& InCollisionSelector, EPCGCollisionQueryFlag InCollisionQueryFlag);
diff --git a/Engine/Plugins/PCG/Source/PCG/Public/Data/PCGVolumeData.h b/Engine/Plugins/PCG/Source/PCG/Public/Data/PCGVolumeData.h
--- a/Engine/Plugins/PCG/Source/PCG/Public/Data/PCGVolumeData.h
+++ b/Engine/Plugins/PCG/Source/PCG/Public/Data/PCGVolumeData.h
@@ -16,6 +16,8 @@ class UPCGVolumeData : public UPCGSpatialDataWithPointCache
 {
 	GENERATED_BODY()
 
+	friend class FPCGVolumeDataVisualization;
+
 public:
 	UE_API ~UPCGVolumeData();
 	UE_API void Initialize(AVolume* InVolume);
diff --git a/Engine/Plugins/PCG/Source/PCGEditor/Private/DataVisualizations/PCGCollisionVisComponent.cpp b/Engine/Plugins/PCG/Source/PCGEditor/Private/DataVisualizations/PCGCollisionVisComponent.cpp
new file mode 100644
--- /dev/null
+++ b/Engine/Plugins/PCG/Source/PCGEditor/Private/DataVisualizations/PCGCollisionVisComponent.cpp
@@ -0,0 +1,144 @@
+// Copyright Epic Games, Inc. All Rights Reserved.
+
+#include "DataVisualizations/PCGCollisionVisComponent.h"
+
+#include "MeshElementCollector.h"
+#include "PrimitiveSceneProxy.h"
+#include "Engine/Engine.h"
+#include "MaterialEditor/MaterialEditorMeshComponent.h"
+#include "Materials/Material.h"
+#include "Materials/MaterialRenderProxy.h"
+#include "Physics/PhysicsInterfaceDeclares.h"
+#include "PhysicsEngine/BodySetup.h"
+
+class FPCGCollisionVisProxy final : public FPrimitiveSceneProxy
+{
+public:
+	SIZE_T GetTypeHash() const override
+	{
+		static size_t UniquePointer;
+		return reinterpret_cast<size_t>(&UniquePointer);
+	}
+
+	FPCGCollisionVisProxy(UPCGCollisionVisComponent* Component)
+		: FPrimitiveSceneProxy(Component)
+	{
+		bWillEverBeLit = false;
+		BodySetups = Component->BodySetups;
+		BodyTransforms = Component->BodyTransforms;
+	}
+
+	virtual void GetDynamicMeshElements(const TArray<const FSceneView*>& Views, const FSceneViewFamily& ViewFamily, uint32 VisibilityMap, FMeshElementCollector& Collector) const override
+	{
+		if (AllowDebugViewmodes())
+		{
+			for (int32 ViewIndex = 0; ViewIndex < Views.Num(); ViewIndex++)
+			{
+				if (VisibilityMap & (1 << ViewIndex))
+				{
+					const FLinearColor DrawColor(GEngine->C_BrushWire);
+					const FColor DrawColorInFColor = DrawColor.ToFColor(true);
+					auto SolidMaterialInstance = new FColoredMaterialRenderProxy(
+						GEngine->ShadedLevelColorationUnlitMaterial->GetRenderProxy(),
+						DrawColor
+					);
+
+					Collector.RegisterOneFrameMaterialProxy(SolidMaterialInstance);
+
+					const FSceneView* View = Views[ViewIndex];
+
+					for (int Instance = 0; Instance < BodySetups.Num(); ++Instance)
+					{
+						UBodySetup* BodySetup = BodySetups[Instance];
+
+						if (BodySetup)
+						{
+							FTransform GeomTransform = BodyTransforms[Instance] * FTransform(GetLocalToWorld());
+							BodySetup->AggGeom.GetAggGeom(GeomTransform, DrawColorInFColor, /*Material=*/SolidMaterialInstance, false, /*bSolid=*/ true, AlwaysHasVelocity(), ViewIndex, Collector);
+						}
+					}
+				}
+			}
+		}
+	}
+
+	virtual FPrimitiveViewRelevance GetViewRelevance(const FSceneView* View) const override
+	{
+		bool bVisible = false;
+
+		if (IsShown(View))
+		{
+			bVisible = true;
+		}
+
+		FPrimitiveViewRelevance Result;
+		Result.bDrawRelevance = bVisible;
+		Result.bDynamicRelevance = true;
+		Result.bShadowRelevance = IsShadowCast(View);
+
+		return Result;
+	}
+
+	virtual uint32 GetMemoryFootprint(void) const override { return(sizeof(*this) + GetAllocatedSize()); }
+	uint32 GetAllocatedSize(void) const { return(FPrimitiveSceneProxy::GetAllocatedSize()); }
+
+	TArray<UBodySetup*> BodySetups;
+	TArray<FTransform> BodyTransforms;
+};
+
+FPrimitiveSceneProxy* UPCGCollisionVisComponent::CreateSceneProxy()
+{
+	return BodySetups.IsEmpty() ? nullptr : new FPCGCollisionVisProxy(this);
+}
+
+FBoxSphereBounds UPCGCollisionVisComponent::CalcBounds(const FTransform& LocalToWorld) const
+{
+	if (BodySetups.IsEmpty())
+	{
+		return FBoxSphereBounds(LocalToWorld.GetLocation(), FVector::ZeroVector, 0.f);
+	}
+	else
+	{
+		FBoxSphereBounds NewBounds(EForceInit::ForceInit);
+		bool bHasBounds = false;
+
+		for (int Instance = 0; Instance < BodySetups.Num(); ++Instance)
+		{
+			FBoxSphereBounds InstanceBounds;
+			BodySetups[Instance]->AggGeom.CalcBoxSphereBounds(InstanceBounds, BodyTransforms[Instance] * LocalToWorld);
+
+			if (bHasBounds)
+			{
+				NewBounds = NewBounds + InstanceBounds;
+			}
+			else
+			{
+				bHasBounds = true;
+				NewBounds = InstanceBounds;
+			}
+		}
+
+		return NewBounds;
+	}
+}
+
+void UPCGCollisionVisComponent::SetBodySetup(UBodySetup* InBodySetup)
+{
+	BodySetups.Reset();
+	BodyTransforms.Reset();
+
+	if (InBodySetup)
+	{
+		BodySetups.Add(InBodySetup);
+		BodyTransforms.Add(FTransform::Identity);
+	}
+}
+
+void UPCGCollisionVisComponent::AddBodySetup(UBodySetup* InBodySetup, const FTransform& InBodyTransform)
+{
+	if (InBodySetup)
+	{
+		BodySetups.Add(InBodySetup);
+		BodyTransforms.Add(InBodyTransform);
+	}
+}
\ No newline at end of file
diff --git a/Engine/Plugins/PCG/Source/PCGEditor/Private/DataVisualizations/PCGSplineDataVisualization.cpp b/Engine/Plugins/PCG/Source/PCGEditor/Private/DataVisualizations/PCGSplineDataVisualization.cpp
--- a/Engine/Plugins/PCG/Source/PCGEditor/Private/DataVisualizations/PCGSplineDataVisualization.cpp
+++ b/Engine/Plugins/PCG/Source/PCGEditor/Private/DataVisualizations/PCGSplineDataVisualization.cpp
@@ -8,6 +8,9 @@
 #include "Data/PCGSplineStruct.h"
 #include "DataVisualizations/PCGDataVisualizationHelpers.h"
 
+#include "AdvancedPreviewScene.h"
+#include "Components/SplineComponent.h"
+
 #define LOCTEXT_NAMESPACE "PCGSplineDataVisualization"
 
 namespace PCGSplineDataVisualizationConstants
@@ -137,4 +140,40 @@ const UPCGBasePointData* IPCGSplineDataVisualization::CollapseToDebugBasePointDa
 	return nullptr;
 }
 
+FPCGSetupSceneFunc IPCGSplineDataVisualization::GetViewportSetupFunc(const UPCGData* Data) const
+{
+	return[this, Data](FAdvancedPreviewScene* InScene, const TArray<UObject*>& InResources, FVector& OutViewLocation, float& OutOrbitDistance, TArray<TObjectPtr<UObject>>& OutManagedResources)
+	{
+		check(InScene);
+
+		const UPCGSplineData* SplineData = Cast<UPCGSplineData>(Data);
+
+		if (!SplineData)
+		{
+			return;
+		}
+
+		TObjectPtr<USplineComponent> SplineComponent = NewObject<USplineComponent>(GetTransientPackage(), NAME_None, RF_Transient);
+		SplineData->ApplyTo(SplineComponent);
+
+		if (GEditor->PreviewPlatform.GetEffectivePreviewFeatureLevel() <= ERHIFeatureLevel::ES3_1)
+		{
+			SplineComponent->SetMobility(EComponentMobility::Static);
+		}
+
+		OutManagedResources.Add(SplineComponent);
+		InScene->AddComponent(SplineComponent, SplineComponent->GetComponentToWorld());
+
+		FBoxSphereBounds Bounds = SplineComponent->CalcBounds(SplineComponent->GetComponentToWorld());
+
+		InScene->SetFloorOffset(-Bounds.Origin.Z + Bounds.BoxExtent.Z);
+		OutViewLocation = Bounds.Origin;
+
+		if (Bounds.SphereRadius > 0.0f)
+		{
+			OutOrbitDistance = Bounds.SphereRadius * 2.0f;
+		}
+	};
+}
+
 #undef LOCTEXT_NAMESPACE
diff --git a/Engine/Plugins/PCG/Source/PCGEditor/Private/DataVisualizations/PCGVolumeDataVisualization.cpp b/Engine/Plugins/PCG/Source/PCGEditor/Private/DataVisualizations/PCGVolumeDataVisualization.cpp
--- a/Engine/Plugins/PCG/Source/PCGEditor/Private/DataVisualizations/PCGVolumeDataVisualization.cpp
+++ b/Engine/Plugins/PCG/Source/PCGEditor/Private/DataVisualizations/PCGVolumeDataVisualization.cpp
@@ -11,8 +11,16 @@
 #include "Elements/PCGVolumeSampler.h"
 
 #include "PCGEditorSettings.h"
+#include "DataVisualizations/PCGCollisionVisComponent.h"
 
-const UPCGBasePointData* IPCGPrimitiveDataVisualization::CollapseToDebugBasePointData(FPCGContext* Context, const UPCGData* Data) const
+#include "AdvancedPreviewScene.h"
+#include "Chaos/ChaosEngineInterface.h"
+#include "Components/BoxComponent.h"
+#include "Components/CapsuleComponent.h"
+#include "Components/PrimitiveComponent.h"
+#include "Components/SphereComponent.h"
+
+const UPCGBasePointData* FPCGPrimitiveDataVisualization::CollapseToDebugBasePointData(FPCGContext* Context, const UPCGData* Data) const
 {
 	const UPCGSpatialData* SpatialData = Cast<UPCGSpatialData>(Data);
 	
@@ -43,3 +51,197 @@ const UPCGBasePointData* IPCGPrimitiveDataVisualization::CollapseToDebugBasePoin
 		return nullptr;
 	}
 }
+
+FPCGSetupSceneFunc FPCGPrimitiveDataVisualization::GetViewportSetupFunc(const UPCGData* Data) const
+{
+	return [this, Data](FAdvancedPreviewScene* InScene, const TArray<UObject*>& InResources, FVector& OutViewLocation, float& OutOrbitDistance, TArray<TObjectPtr<UObject>>& OutManagedResources)
+	{
+		check(InScene);
+
+		TArray<TObjectPtr<UPrimitiveComponent>> Components;
+		TArray<FTransform> ComponentTransforms;
+
+		GetComponentsAndTransforms(Data, Components, ComponentTransforms);
+
+		if (Components.IsEmpty() || Components.Num() != ComponentTransforms.Num())
+		{
+			return;
+		}
+
+		FBoxSphereBounds Bounds(EForceInit::ForceInit);
+		bool bHasBounds = false;
+
+		for(int Index = 0; Index < Components.Num(); ++Index)
+		{
+			TObjectPtr<UPrimitiveComponent> Component = Components[Index];
+			const FTransform& Transform = ComponentTransforms[Index];
+
+			if (GEditor->PreviewPlatform.GetEffectivePreviewFeatureLevel() <= ERHIFeatureLevel::ES3_1)
+			{
+				Component->SetMobility(EComponentMobility::Static);
+			}	
+
+			OutManagedResources.Add(Component);
+			InScene->AddComponent(Component, Transform);
+
+			if (bHasBounds)
+			{
+				Bounds = Bounds + Component->CalcBounds(Transform);
+			}
+			else
+			{
+				bHasBounds = true;
+				Bounds = Component->CalcBounds(Transform);
+			}
+		}
+
+		if (bHasBounds)
+		{
+			InScene->SetFloorOffset(-Bounds.Origin.Z + Bounds.BoxExtent.Z);
+			OutViewLocation = Bounds.Origin;
+
+			if (Bounds.SphereRadius > 0.0f)
+			{
+				OutOrbitDistance = Bounds.SphereRadius * 2.0f;
+			}
+		}
+	};
+}
+
+void FPCGPrimitiveDataVisualization::GetComponentsAndTransforms(const UPCGData* InData, TArray<TObjectPtr<UPrimitiveComponent>>& OutComponents, TArray<FTransform>& OutComponentTransforms) const
+{
+	const UPCGPrimitiveData* PrimitiveData = Cast<UPCGPrimitiveData>(InData);
+
+	if (!PrimitiveData)
+	{
+		return;
+	}
+
+	if (TStrongObjectPtr<UPrimitiveComponent> PrimitiveComponent = PrimitiveData->GetComponent().Pin())
+	{
+		TObjectPtr<UPCGCollisionVisComponent> VisComponent = NewObject<UPCGCollisionVisComponent>(GetTransientPackage(), NAME_None, RF_Transient);
+		VisComponent->SetBodySetup(PrimitiveComponent->GetBodySetup());
+
+		OutComponents.Add(VisComponent);
+		OutComponentTransforms.Add(PrimitiveComponent->GetComponentToWorld());
+	}
+}
+
+void FPCGVolumeDataVisualization::GetComponentsAndTransforms(const UPCGData* InData, TArray<TObjectPtr<UPrimitiveComponent>>& OutComponents, TArray<FTransform>& OutComponentTransforms) const
+{
+	const UPCGVolumeData* VolumeData = Cast<UPCGVolumeData>(InData);
+
+	if (!VolumeData)
+	{
+		return;
+	}
+
+	if (FBodyInstance* Instance = VolumeData->VolumeBodyInstance)
+	{
+		TObjectPtr<UPCGCollisionVisComponent> VisComponent = NewObject<UPCGCollisionVisComponent>(GetTransientPackage(), NAME_None, RF_Transient);
+		VisComponent->SetBodySetup(Instance->GetBodySetup());
+
+		OutComponents.Add(VisComponent);
+		OutComponentTransforms.Add(Instance->GetUnrealWorldTransform());
+	}
+	else
+	{
+		// Create a box component
+		UBoxComponent* BoxComponent = NewObject<UBoxComponent>(GetTransientPackage(), NAME_None, RF_Transient);
+		BoxComponent->SetBoxExtent(VolumeData->Bounds.GetExtent());
+		BoxComponent->SetSimulatePhysics(false);
+
+		OutComponents.Add(BoxComponent);
+		OutComponentTransforms.Add(FTransform(VolumeData->Bounds.GetCenter()));
+	}
+}
+
+void FPCGCollisionShapeDataVisualization::GetComponentsAndTransforms(const UPCGData* InData, TArray<TObjectPtr<UPrimitiveComponent>>& OutComponents, TArray<FTransform>& OutComponentTransforms) const
+{
+	const UPCGCollisionShapeData* ShapeData = Cast<UPCGCollisionShapeData>(InData);
+
+	if (!ShapeData)
+	{
+		return;
+	}
+
+	const FCollisionShape& Shape = ShapeData->Shape;
+
+	if (Shape.IsLine()) // not supported in PCG
+	{
+		return;
+	}
+
+	if (Shape.IsSphere())
+	{
+		USphereComponent* SphereComponent = NewObject<USphereComponent>(GetTransientPackage(), NAME_None, RF_Transient);
+		SphereComponent->SetSphereRadius(Shape.Sphere.Radius);
+
+		OutComponents.Add(SphereComponent);
+	}
+	else if (Shape.IsCapsule())
+	{
+		UCapsuleComponent* CapsuleComponent = NewObject<UCapsuleComponent>(GetTransientPackage(), NAME_None, RF_Transient);
+		CapsuleComponent->SetCapsuleRadius(Shape.Capsule.Radius);
+		CapsuleComponent->SetCapsuleHalfHeight(Shape.Capsule.HalfHeight);
+
+		OutComponents.Add(CapsuleComponent);
+	}
+	else if (Shape.IsBox())
+	{
+		UBoxComponent* BoxComponent = NewObject<UBoxComponent>(GetTransientPackage(), NAME_None, RF_Transient);
+		BoxComponent->SetBoxExtent(FVector(Shape.Box.HalfExtentX, Shape.Box.HalfExtentY, Shape.Box.HalfExtentZ));
+
+		OutComponents.Add(BoxComponent);
+	}
+
+	OutComponentTransforms.Add(ShapeData->Transform);
+}
+
+void FPCGCollisionWrapperDataVisualization::GetComponentsAndTransforms(const UPCGData* InData, TArray<TObjectPtr<UPrimitiveComponent>>& OutComponents, TArray<FTransform>& OutComponentTransforms) const
+{
+	const UPCGCollisionWrapperData* WrapperData = Cast<UPCGCollisionWrapperData>(InData);
+
+	if (!WrapperData || !WrapperData->CollisionWrapper.bInitialized)
+	{
+		return;
+	}
+
+	const FPCGCollisionWrapper& CollisionWrapper = WrapperData->CollisionWrapper;
+	TConstPCGValueRange<FTransform> Transforms = WrapperData->GetPointData()->GetConstTransformValueRange();
+	TConstPCGValueRange<FVector> BoundsMin = WrapperData->GetPointData()->GetConstBoundsMinValueRange();
+	TConstPCGValueRange<FVector> BoundsMax = WrapperData->GetPointData()->GetConstBoundsMaxValueRange();
+
+	// Prepare components
+	UPCGCollisionVisComponent* VisComponent = NewObject<UPCGCollisionVisComponent>(GetTransientPackage(), NAME_None, RF_Transient);
+
+	for (int Index = 0; Index < Transforms.Num(); ++Index)
+	{
+		const FTransform& Transform = Transforms[Index];
+		FBodyInstance* BodyInstance = CollisionWrapper.IndexToBodyInstance[Index] == INDEX_NONE ? nullptr : CollisionWrapper.BodyInstances[CollisionWrapper.IndexToBodyInstance[Index]];
+		if (BodyInstance && BodyInstance->GetBodySetup())
+		{
+			VisComponent->AddBodySetup(BodyInstance->GetBodySetup(), Transform);
+		}
+		else
+		{
+			// Default to creating a box component
+			UBoxComponent* BoxComponent = NewObject<UBoxComponent>(GetTransientPackage(), NAME_None, RF_Transient);
+
+			const FVector Center = (BoundsMax[Index] + BoundsMin[Index]) * 0.5;
+			const FVector Extent = (BoundsMax[Index] - BoundsMin[Index]) * 0.5;
+
+			BoxComponent->SetBoxExtent(Extent);
+			BoxComponent->SetSimulatePhysics(false);
+
+			OutComponents.Add(BoxComponent);
+			OutComponentTransforms.Add(FTransform(Center) * Transform);
+		}
+	}
+
+	if (!VisComponent->BodyTransforms.IsEmpty())
+	{
+		OutComponents.Add(VisComponent);
+		OutComponentTransforms.Add(FTransform::Identity);
+	}
+}
\ No newline at end of file
diff --git a/Engine/Plugins/PCG/Source/PCGEditor/Private/PCGEditorModule.cpp b/Engine/Plugins/PCG/Source/PCGEditor/Private/PCGEditorModule.cpp
--- a/Engine/Plugins/PCG/Source/PCGEditor/Private/PCGEditorModule.cpp
+++ b/Engine/Plugins/PCG/Source/PCGEditor/Private/PCGEditorModule.cpp
@@ -798,10 +798,10 @@ void FPCGEditorModule::RegisterPCGDataVisualizations()
 	DataVisRegistry.InternalRegistry.Add(UPCGParamData::StaticClass(), MakeUnique<const IPCGParamDataVisualization>());
 	DataVisRegistry.InternalRegistry.Add(UPCGSpatialData::StaticClass(), MakeUnique<const IPCGSpatialDataVisualization>());
 	DataVisRegistry.InternalRegistry.Add(UPCGSplineData::StaticClass(), MakeUnique<const IPCGSplineDataVisualization>());
-	DataVisRegistry.InternalRegistry.Add(UPCGVolumeData::StaticClass(), MakeUnique<const IPCGPrimitiveDataVisualization>());
-	DataVisRegistry.InternalRegistry.Add(UPCGPrimitiveData::StaticClass(), MakeUnique<const IPCGPrimitiveDataVisualization>());
-	DataVisRegistry.InternalRegistry.Add(UPCGCollisionShapeData::StaticClass(), MakeUnique<const IPCGPrimitiveDataVisualization>());
-	DataVisRegistry.InternalRegistry.Add(UPCGCollisionWrapperData::StaticClass(), MakeUnique<const IPCGPrimitiveDataVisualization>());
+	DataVisRegistry.InternalRegistry.Add(UPCGVolumeData::StaticClass(), MakeUnique<const FPCGVolumeDataVisualization>());
+	DataVisRegistry.InternalRegistry.Add(UPCGPrimitiveData::StaticClass(), MakeUnique<const FPCGPrimitiveDataVisualization>());
+	DataVisRegistry.InternalRegistry.Add(UPCGCollisionShapeData::StaticClass(), MakeUnique<const FPCGCollisionShapeDataVisualization>());
+	DataVisRegistry.InternalRegistry.Add(UPCGCollisionWrapperData::StaticClass(), MakeUnique<const FPCGCollisionWrapperDataVisualization>());
 	DataVisRegistry.InternalRegistry.Add(UPCGLandscapeData::StaticClass(), MakeUnique<const IPCGLandscapeDataVisualization>());
 	DataVisRegistry.InternalRegistry.Add(UPCGStaticMeshResourceData::StaticClass(), MakeUnique<const IPCGStaticMeshDataVisualization>());
 	DataVisRegistry.InternalRegistry.Add(UPCGBaseTextureData::StaticClass(), MakeUnique<const FPCGBaseTextureDataVisualization>());
diff --git a/Engine/Plugins/PCG/Source/PCGEditor/Public/DataVisualizations/PCGCollisionVisComponent.h b/Engine/Plugins/PCG/Source/PCGEditor/Public/DataVisualizations/PCGCollisionVisComponent.h
new file mode 100644
--- /dev/null
+++ b/Engine/Plugins/PCG/Source/PCGEditor/Public/DataVisualizations/PCGCollisionVisComponent.h
@@ -0,0 +1,33 @@
+// Copyright Epic Games, Inc. All Rights Reserved.
+
+#pragma once
+
+#include "Components/PrimitiveComponent.h"
+
+#include "PCGCollisionVisComponent.generated.h"
+
+UCLASS(MinimalAPI, NotPlaceable)
+class UPCGCollisionVisComponent : public UPrimitiveComponent
+{
+	GENERATED_BODY()
+
+public:
+	UPROPERTY()
+	TArray<TObjectPtr<UBodySetup>> BodySetups;
+
+	UPROPERTY()
+	TArray<FTransform> BodyTransforms;
+
+	//~Begin USceneComponent interface
+	virtual FBoxSphereBounds CalcBounds(const FTransform& LocalToWorld) const override;
+	//~End USceneComponent interface
+
+	//~Begin UPrimitiveComponent interface
+	virtual FPrimitiveSceneProxy* CreateSceneProxy() override;
+	virtual UBodySetup* GetBodySetup() override { return BodySetups.IsEmpty() ? nullptr : BodySetups[0]; }
+	virtual bool IsEditorOnly() const override { return true; }
+	//~End UPrimitiveComponent interface
+
+	void SetBodySetup(UBodySetup* InBodySetup);
+	void AddBodySetup(UBodySetup* InBodySetup, const FTransform& InTransform);
+};
\ No newline at end of file
diff --git a/Engine/Plugins/PCG/Source/PCGEditor/Public/DataVisualizations/PCGSplineDataVisualization.h b/Engine/Plugins/PCG/Source/PCGEditor/Public/DataVisualizations/PCGSplineDataVisualization.h
--- a/Engine/Plugins/PCG/Source/PCGEditor/Public/DataVisualizations/PCGSplineDataVisualization.h
+++ b/Engine/Plugins/PCG/Source/PCGEditor/Public/DataVisualizations/PCGSplineDataVisualization.h
@@ -13,6 +13,7 @@ class PCGEDITOR_API IPCGSplineDataVisualization : public IPCGSpatialDataVisualiz
 public:
 	// ~Begin IPCGDataVisualization interface
 	virtual FPCGTableVisualizerInfo GetTableVisualizerInfoWithDomain(const UPCGData* Data, const FPCGMetadataDomainID& DomainID) const override;
+	virtual FPCGSetupSceneFunc GetViewportSetupFunc(const UPCGData* Data) const override;
 	// ~End IPCGDataVisualization interface
 
 	// ~Begin IPCGSpatialDataVisualization interface
diff --git a/Engine/Plugins/PCG/Source/PCGEditor/Public/DataVisualizations/PCGVolumeDataVisualization.h b/Engine/Plugins/PCG/Source/PCGEditor/Public/DataVisualizations/PCGVolumeDataVisualization.h
--- a/Engine/Plugins/PCG/Source/PCGEditor/Public/DataVisualizations/PCGVolumeDataVisualization.h
+++ b/Engine/Plugins/PCG/Source/PCGEditor/Public/DataVisualizations/PCGVolumeDataVisualization.h
@@ -8,10 +8,35 @@ class UPCGData;
 struct FPCGContext;
 
 /** Overrides collapse behavior to target a given number of points generated, to avoid too coarse debug (points way bigger than data) or enormous amount of points (points way smaller than data). */
-class PCGEDITOR_API IPCGPrimitiveDataVisualization : public IPCGSpatialDataVisualization
+class PCGEDITOR_API FPCGPrimitiveDataVisualization : public IPCGSpatialDataVisualization
 {
 public:
 	// ~Begin IPCGSpatialDataVisualization interface
 	virtual const UPCGBasePointData* CollapseToDebugBasePointData(FPCGContext* Context, const UPCGData* Data) const override;
+	FPCGSetupSceneFunc GetViewportSetupFunc(const UPCGData* Data) const override;
 	// ~End IPCGSpatialDataVisualization interface
+
+protected:
+	virtual void GetComponentsAndTransforms(const UPCGData* InData, TArray<TObjectPtr<UPrimitiveComponent>>& OutComponents, TArray<FTransform>& OutComponentTransforms) const;
 };
+
+/** Specialization for volumes that don't always rely on their primitive component */
+class FPCGVolumeDataVisualization : public FPCGPrimitiveDataVisualization
+{
+protected:
+	virtual void GetComponentsAndTransforms(const UPCGData* InData, TArray<TObjectPtr<UPrimitiveComponent>>& OutComponents, TArray<FTransform>& OutComponentTransforms) const override;
+};
+
+/** Specialization for simple collision shapes */
+class FPCGCollisionShapeDataVisualization : public FPCGPrimitiveDataVisualization
+{
+protected:
+	virtual void GetComponentsAndTransforms(const UPCGData* InData, TArray<TObjectPtr<UPrimitiveComponent>>& OutComponents, TArray<FTransform>& OutComponentTransforms) const override;
+};
+
+/** Specialization for collision wrappers */
+class FPCGCollisionWrapperDataVisualization : public FPCGPrimitiveDataVisualization
+{
+protected:
+	virtual void GetComponentsAndTransforms(const UPCGData* InData, TArray<TObjectPtr<UPrimitiveComponent>>& OutComponents, TArray<FTransform>& OutComponentTransforms) const override;
+};
\ No newline at end of file
