[PCG] Add virtual texture material type as automatic data label

#rnx

diff --git a/Engine/Plugins/PCG/Source/PCG/Private/Data/PCGGetDataFunctionRegistry.cpp b/Engine/Plugins/PCG/Source/PCG/Private/Data/PCGGetDataFunctionRegistry.cpp
--- a/Engine/Plugins/PCG/Source/PCG/Private/Data/PCGGetDataFunctionRegistry.cpp
+++ b/Engine/Plugins/PCG/Source/PCG/Private/Data/PCGGetDataFunctionRegistry.cpp
@@ -293,14 +293,20 @@ int FPCGGetDataFunctionRegistry::DefaultDataFromComponent(FPCGContext* InContext
 			VirtualTextureData->Initialize(VirtualTextureComponent);
 			TaggedData.Data = VirtualTextureData;
 
-			FString VirtualTextureName;
-
 			if (const URuntimeVirtualTexture* VirtualTexture = VirtualTextureData->GetRuntimeVirtualTexture())
 			{
+				FString VirtualTextureName;
 				VirtualTexture->GetName(VirtualTextureName);
+
+				const UEnum* MaterialTypeEnum = StaticEnum<ERuntimeVirtualTextureMaterialType>();
+				check(MaterialTypeEnum);
+
+				const FString MaterialTypeString = MaterialTypeEnum->GetNameStringByValue(static_cast<int64>(VirtualTexture->GetMaterialType()));
+
+				TaggedData.Tags.Add(PCGComputeHelpers::GetPrefixedDataLabel(VirtualTextureName));
+				TaggedData.Tags.Add(PCGComputeHelpers::GetPrefixedDataLabel(MaterialTypeString));
 			}
 
-			TaggedData.Tags.Add(PCGComputeHelpers::GetPrefixedDataLabel(VirtualTextureName));
 		}
 	}
 
