[PCG] Add filter by value and value range on data on filter data by attrbiute

* The existing behavior is now named "Filter by Existence" and allow to target a metadata domain
* The other modes are simetrical to the Attribute Filter node
* Factorized some code/constants from the Attribute Filter for re-use.

#rnx

diff --git a/Engine/Plugins/PCG/Source/PCG/Private/Elements/PCGAttributeFilter.cpp b/Engine/Plugins/PCG/Source/PCG/Private/Elements/PCGAttributeFilter.cpp
--- a/Engine/Plugins/PCG/Source/PCG/Private/Elements/PCGAttributeFilter.cpp
+++ b/Engine/Plugins/PCG/Source/PCG/Private/Elements/PCGAttributeFilter.cpp
@@ -17,11 +17,6 @@
 
 namespace PCGAttributeFilterConstants
 {
-	const FName DataToFilterLabel = TEXT("In");
-	const FName FilterLabel = TEXT("Filter");
-	const FName FilterMinLabel = TEXT("FilterMin");
-	const FName FilterMaxLabel = TEXT("FilterMax");
-
 	constexpr int32 ChunkSize = 256;
 
 #if WITH_EDITOR
@@ -30,69 +25,14 @@ namespace PCGAttributeFilterConstants
 	const FText FilterPinTooltip = LOCTEXT("FilterPinTooltip", "This pin accepts Statial data and Attribute Sets. If the data is Spatial, it will automatically sample input points in it. "
 		"If it is points, it will sample if \"Spatial Query\" is enabled, otherwise points number need to match with input.");
 #endif // WITH_EDITOR
+
+	constexpr int32 DefaultAliasIndex = 0;
+	constexpr int32 PointFilterAliasIndex = 1;
+	constexpr int32 AttributeFilterAliasIndex = 2;
 }
 
 namespace PCGAttributeFilterHelpers
 {
-	template <typename T>
-	bool ApplyCompare(const T& Input1, const T& Input2, EPCGAttributeFilterOperator Operation)
-	{
-		if (Operation == EPCGAttributeFilterOperator::Equal)
-		{
-			return PCG::Private::MetadataTraits<T>::Equal(Input1, Input2);
-		}
-		else if (Operation == EPCGAttributeFilterOperator::NotEqual)
-		{
-			return !PCG::Private::MetadataTraits<T>::Equal(Input1, Input2);
-		}
-
-		if constexpr (PCG::Private::MetadataTraits<T>::CanCompare)
-		{
-			switch (Operation)
-			{
-			case EPCGAttributeFilterOperator::Greater:
-				return PCG::Private::MetadataTraits<T>::Greater(Input1, Input2);
-			case EPCGAttributeFilterOperator::GreaterOrEqual:
-				return PCG::Private::MetadataTraits<T>::GreaterOrEqual(Input1, Input2);
-			case EPCGAttributeFilterOperator::Lesser:
-				return PCG::Private::MetadataTraits<T>::Less(Input1, Input2);
-			case EPCGAttributeFilterOperator::LesserOrEqual:
-				return PCG::Private::MetadataTraits<T>::LessOrEqual(Input1, Input2);
-			default:
-				break;
-			}
-		}
-
-		if constexpr (PCG::Private::MetadataTraits<T>::CanSearchString)
-		{
-			switch (Operation)
-			{
-			case EPCGAttributeFilterOperator::Substring:
-				return PCG::Private::MetadataTraits<T>::Substring(Input1, Input2);
-			case EPCGAttributeFilterOperator::Matches:
-				return PCG::Private::MetadataTraits<T>::Matches(Input1, Input2);
-			default:
-				break;
-			}
-		}
-
-		return false;
-	}
-
-	template <typename T>
-	bool ApplyRange(const T& Input, const T& InMin, const T& InMax, bool bMinIncluded, bool bMaxIncluded)
-	{
-		if constexpr (PCG::Private::MetadataTraits<T>::CanCompare)
-		{
-			return (bMinIncluded ? PCG::Private::MetadataTraits<T>::GreaterOrEqual(Input, InMin) : PCG::Private::MetadataTraits<T>::Greater(Input, InMin)) &&
-				(bMaxIncluded ? PCG::Private::MetadataTraits<T>::LessOrEqual(Input, InMax) : PCG::Private::MetadataTraits<T>::Less(Input, InMax));
-		}
-		else
-		{
-			return false;
-		}
-	}
-
 	struct ThresholdInfo
 	{
 		TUniquePtr<const IPCGAttributeAccessor> ThresholdAccessor;
@@ -288,7 +228,7 @@ EPCGDataType UPCGAttributeFilteringSettings::GetCurrentPinTypes(const UPCGPin* I
 TArray<FPCGPinProperties> UPCGAttributeFilteringSettings::InputPinProperties() const
 {
 	TArray<FPCGPinProperties> PinProperties;
-	FPCGPinProperties& DataToFilterPinProperty = PinProperties.Emplace_GetRef(PCGAttributeFilterConstants::DataToFilterLabel, EPCGDataType::PointOrParam);
+	FPCGPinProperties& DataToFilterPinProperty = PinProperties.Emplace_GetRef(PCGPinConstants::DefaultInputLabel, EPCGDataType::PointOrParam);
 	DataToFilterPinProperty.SetRequiredPin();
 
 #if WITH_EDITOR
@@ -354,8 +294,9 @@ UPCGAttributeFilteringSettings::UPCGAttributeFilteringSettings()
 TArray<FPCGPreConfiguredSettingsInfo> UPCGAttributeFilteringSettings::GetPreconfiguredInfo() const
 {
 	TArray<FPCGPreConfiguredSettingsInfo> PreconfiguredInfo;
-	PreconfiguredInfo.Emplace(0, GetDefaultNodeTitle());
-	PreconfiguredInfo.Emplace(1, LOCTEXT("PointNodeTitle", "Point Filter"));
+	PreconfiguredInfo.Emplace(PCGAttributeFilterConstants::DefaultAliasIndex, GetDefaultNodeTitle());
+	PreconfiguredInfo.Emplace(PCGAttributeFilterConstants::PointFilterAliasIndex, LOCTEXT("PointNodeTitle", "Point Filter"));
+	PreconfiguredInfo.Emplace(PCGAttributeFilterConstants::AttributeFilterAliasIndex, LOCTEXT("AttributeNodeTitle", "Attribute Filter"));
 
 	return PreconfiguredInfo;
 }
@@ -364,7 +305,7 @@ TArray<FPCGPreConfiguredSettingsInfo> UPCGAttributeFilteringSettings::GetPreconf
 void UPCGAttributeFilteringSettings::ApplyPreconfiguredSettings(const FPCGPreConfiguredSettingsInfo& PreconfigureInfo)
 {
 	// If index is 1, it is the default ($Density)
-	if (PreconfigureInfo.PreconfiguredIndex == 0)
+	if (PreconfigureInfo.PreconfiguredIndex != PCGAttributeFilterConstants::PointFilterAliasIndex)
 	{
 		TargetAttribute.SetAttributeName(PCGMetadataAttributeConstants::LastAttributeName);
 		ThresholdAttribute.SetAttributeName(PCGMetadataAttributeConstants::LastAttributeName);
@@ -410,8 +351,9 @@ void UPCGAttributeFilteringRangeSettings::PostLoad()
 TArray<FPCGPreConfiguredSettingsInfo> UPCGAttributeFilteringRangeSettings::GetPreconfiguredInfo() const
 {
 	TArray<FPCGPreConfiguredSettingsInfo> PreconfiguredInfo;
-	PreconfiguredInfo.Emplace(0, GetDefaultNodeTitle());
-	PreconfiguredInfo.Emplace(1, LOCTEXT("PointRangeNodeTitle", "Point Filter Range"));
+	PreconfiguredInfo.Emplace(PCGAttributeFilterConstants::DefaultAliasIndex, GetDefaultNodeTitle());
+	PreconfiguredInfo.Emplace(PCGAttributeFilterConstants::PointFilterAliasIndex, LOCTEXT("PointRangeNodeTitle", "Point Filter Range"));
+	PreconfiguredInfo.Emplace(PCGAttributeFilterConstants::AttributeFilterAliasIndex, LOCTEXT("AttributeRangeNodeTitle", "Attribute Filter Range"));
 
 	return PreconfiguredInfo;
 }
@@ -420,7 +362,7 @@ TArray<FPCGPreConfiguredSettingsInfo> UPCGAttributeFilteringRangeSettings::GetPr
 void UPCGAttributeFilteringRangeSettings::ApplyPreconfiguredSettings(const FPCGPreConfiguredSettingsInfo& PreconfigureInfo)
 {
 	// If index is 1, it is the default ($Density)
-	if (PreconfigureInfo.PreconfiguredIndex == 0)
+	if (PreconfigureInfo.PreconfiguredIndex != PCGAttributeFilterConstants::PointFilterAliasIndex)
 	{
 		TargetAttribute.SetAttributeName(PCGMetadataAttributeConstants::LastAttributeName);
 		MinThreshold.ThresholdAttribute.SetAttributeName(PCGMetadataAttributeConstants::LastAttributeName);
@@ -490,7 +432,7 @@ EPCGDataType UPCGAttributeFilteringRangeSettings::GetCurrentPinTypes(const UPCGP
 TArray<FPCGPinProperties> UPCGAttributeFilteringRangeSettings::InputPinProperties() const
 {
 	TArray<FPCGPinProperties> PinProperties;
-	FPCGPinProperties& DataToFilterPinProperty = PinProperties.Emplace_GetRef(PCGAttributeFilterConstants::DataToFilterLabel, EPCGDataType::PointOrParam);
+	FPCGPinProperties& DataToFilterPinProperty = PinProperties.Emplace_GetRef(PCGPinConstants::DefaultInputLabel, EPCGDataType::PointOrParam);
 	DataToFilterPinProperty.SetRequiredPin();
 
 #if WITH_EDITOR
@@ -550,7 +492,7 @@ bool FPCGAttributeFilterElementBase::DoFiltering(FPCGContext* Context, EPCGAttri
 		bool bIsInputPointData = false;
 	} OperationData;
 
-	TArray<FPCGTaggedData> DataToFilter = Context->InputData.GetInputsByPin(PCGAttributeFilterConstants::DataToFilterLabel);
+	TArray<FPCGTaggedData> DataToFilter = Context->InputData.GetInputsByPin(PCGPinConstants::DefaultInputLabel);
 	TArray<FPCGTaggedData>& Outputs = Context->OutputData.TaggedData;
 	TArray<FPCGTaggedData> FirstFilterData = SecondThreshold ? Context->InputData.GetInputsByPin(PCGAttributeFilterConstants::FilterMinLabel) : Context->InputData.GetInputsByPin(PCGAttributeFilterConstants::FilterLabel);
 	TArray<FPCGTaggedData> SecondFilterData = SecondThreshold ? Context->InputData.GetInputsByPin(PCGAttributeFilterConstants::FilterMaxLabel) : TArray<FPCGTaggedData>{};
diff --git a/Engine/Plugins/PCG/Source/PCG/Private/Elements/PCGFilterByAttribute.cpp b/Engine/Plugins/PCG/Source/PCG/Private/Elements/PCGFilterByAttribute.cpp
--- a/Engine/Plugins/PCG/Source/PCG/Private/Elements/PCGFilterByAttribute.cpp
+++ b/Engine/Plugins/PCG/Source/PCG/Private/Elements/PCGFilterByAttribute.cpp
@@ -11,77 +11,42 @@
 #include "Metadata/Accessors/PCGAttributeAccessorHelpers.h"
 
 #include "Algo/AnyOf.h"
+#include "Elements/Metadata/PCGMetadataElementCommon.h"
+#include "Metadata/Accessors/PCGCustomAccessor.h"
 
 #define LOCTEXT_NAMESPACE "PCGFilterByAttributeElement"
 
-#if WITH_EDITOR
-FText UPCGFilterByAttributeSettings::GetDefaultNodeTitle() const
+namespace PCGFilterByAttributeElement
 {
-	return LOCTEXT("NodeTitle", "Filter Data By Attribute");
-}
-#endif // WITH_EDITOR
-
-FString UPCGFilterByAttributeSettings::GetAdditionalTitleInformation() const
-{
-#if WITH_EDITOR
-	if (IsPropertyOverriddenByPin(GET_MEMBER_NAME_CHECKED(UPCGFilterByAttributeSettings, Attribute)))
-	{
-		return FString();
-	}
-	else
-#endif
+	const FText ThresholdText = LOCTEXT("Threshold", "Threshold");
+	
+	bool FilterByExistence(const FPCGTaggedData& Input, const UPCGFilterByAttributeSettings& Settings, const TArray<FString>& Attributes)
 	{
-		return Attribute.ToString();
-	}
-}
-
-FPCGElementPtr UPCGFilterByAttributeSettings::CreateElement() const
-{
-	return MakeShared<FPCGFilterByAttributeElement>();
-}
+		TArray<FString, TInlineAllocator<64>> DataAttributeStrings;
 
-bool FPCGFilterByAttributeElement::ExecuteInternal(FPCGContext* Context) const
-{
-	TRACE_CPUPROFILER_EVENT_SCOPE(FPCGFilterByAttributeElement::Execute);
-	check(Context);
-
-	const UPCGFilterByAttributeSettings* Settings = Context->GetInputSettings<UPCGFilterByAttributeSettings>();
-	check(Settings);
-
-	TArray<FString> Attributes = PCGHelpers::GetStringArrayFromCommaSeparatedList(Settings->Attribute.ToString());
-
-	TArray<FPCGTaggedData> Inputs = Context->InputData.GetInputsByPin(PCGPinConstants::DefaultInputLabel);
-	TArray<FPCGTaggedData>& Outputs = Context->OutputData.TaggedData;
-
-	TArray<FName> DataAttributes;
-	TArray<EPCGMetadataTypes> DataAttributeTypes;
-	TArray<FString> DataAttributeStrings;
-
-	for (const FPCGTaggedData& Input : Inputs)
-	{
-		FPCGTaggedData& Output = Outputs.Add_GetRef(Input);
-		Output.Pin = PCGPinConstants::DefaultOutFilterLabel;
-
-		if (!Input.Data || !Input.Data->ConstMetadata())
+		const UPCGMetadata* Metadata = Input.Data ? Input.Data->ConstMetadata() : nullptr;
+		const FPCGMetadataDomainID MetadataDomainID = Input.Data && Metadata ? Input.Data->GetMetadataDomainIDFromSelector(FPCGAttributePropertySelector::CreateAttributeSelector(NAME_None, Settings.MetadataDomain)) : PCGMetadataDomainID::Invalid;
+		const FPCGMetadataDomain* MetadataDomain = MetadataDomainID.IsValid() && Metadata ? Metadata->GetConstMetadataDomain(MetadataDomainID) : nullptr;
+		
+		if (!MetadataDomain)
 		{
-			continue;
+			return false;
 		}
 
 		// All attributes from the list must have a match in order to put the data in the In Filter pin.
-		const UPCGMetadata* Metadata = Input.Data->ConstMetadata();
-
-		DataAttributes.Reset();
-		DataAttributeTypes.Reset();
-		DataAttributeStrings.Reset();
-
-		if (Settings->Operator != EPCGStringMatchingOperator::Equal)
+		if (Settings.Operator != EPCGStringMatchingOperator::Equal)
 		{
+			TArray<FName> DataAttributes;
+			TArray<EPCGMetadataTypes> DataAttributeTypes;
+			
 			// Gather all attributes & properties based on data type
-			Metadata->GetAttributes(DataAttributes, DataAttributeTypes);
-			Algo::Transform(DataAttributes, DataAttributeStrings, [](const FName& InAttribute) { return InAttribute.ToString(); });
+			MetadataDomain->GetAttributes(DataAttributes, DataAttributeTypes);
+			Algo::Transform(DataAttributes, DataAttributeStrings, [](const FPCGAttributeIdentifier& InAttribute) { return InAttribute.ToString(); });
 
-			if (!Settings->bIgnoreProperties)
+			if (!Settings.bIgnoreProperties)
 			{
+				// @pcg_todo We should have a way of querying all properties for a given data.
+				
 				if (Cast<UPCGBasePointData>(Input.Data))
 				{
 					if (const UEnum* PointProperties = StaticEnum<EPCGPointProperties>())
@@ -113,7 +78,7 @@ bool FPCGFilterByAttributeElement::ExecuteInternal(FPCGContext* Context) const
 			Selector.Update(Attribute);
 
 			// In the case of the equal test, we can test directly if the selector would yield something valid
-			if (Settings->Operator == EPCGStringMatchingOperator::Equal)
+			if (Settings.Operator == EPCGStringMatchingOperator::Equal)
 			{
 				TUniquePtr<const IPCGAttributeAccessor> Accessor = PCGAttributeAccessorHelpers::CreateConstAccessor(Input.Data, Selector, /*bQuiet=*/true);
 				if (!Accessor || !Accessor.IsValid())
@@ -133,8 +98,8 @@ bool FPCGFilterByAttributeElement::ExecuteInternal(FPCGContext* Context) const
 				bool bFoundValidMatch = false;
 				for (const FString& DataAttribute : DataAttributeStrings)
 				{
-					if ((Settings->Operator == EPCGStringMatchingOperator::Substring && !DataAttribute.Contains(AttributeWithNoAccessor)) ||
-						(Settings->Operator == EPCGStringMatchingOperator::Matches && !DataAttribute.MatchesWildcard(AttributeWithNoAccessor)))
+					if ((Settings.Operator == EPCGStringMatchingOperator::Substring && !DataAttribute.Contains(AttributeWithNoAccessor)) ||
+						(Settings.Operator == EPCGStringMatchingOperator::Matches && !DataAttribute.MatchesWildcard(AttributeWithNoAccessor)))
 					{
 						continue;
 					}
@@ -159,6 +124,325 @@ bool FPCGFilterByAttributeElement::ExecuteInternal(FPCGContext* Context) const
 			}
 		}
 
+		return bInFilter;
+	}
+
+	bool CreateAndValidateAccessorThreshold(
+		TUniquePtr<const IPCGAttributeAccessor>& ThresholdAccessor,
+		TUniquePtr<const IPCGAttributeAccessorKeys>& ThresholdAccessorKeys,
+		const FPCGFilterByAttributeThresholdSettings& ThresholdSettings,
+		const FPCGTaggedData* ThresholdData,
+		const uint16 InputType,
+		const FPCGContext* Context)
+	{
+		if (ThresholdData)
+		{
+			const FPCGAttributePropertyInputSelector ThresholdAttribute = ThresholdSettings.ThresholdAttribute.CopyAndFixLast(ThresholdData->Data);
+			ThresholdAccessor = PCGAttributeAccessorHelpers::CreateConstAccessor(ThresholdData->Data, ThresholdAttribute);
+			ThresholdAccessorKeys = PCGAttributeAccessorHelpers::CreateConstKeys(ThresholdData->Data, ThresholdAttribute);
+
+			if (!ThresholdAccessor || !ThresholdAccessorKeys)
+			{
+				PCGLog::Metadata::LogFailToCreateAccessorError(ThresholdAttribute, Context);
+				return false;
+			}
+		}
+		else
+		{
+			auto ConstantThreshold = [&ThresholdAccessor, &ThresholdAccessorKeys](auto&& Value) -> bool
+			{
+				using ConstantType = std::decay_t<decltype(Value)>;
+
+				ThresholdAccessor = MakeUnique<FPCGConstantValueAccessor<ConstantType>>(std::forward<decltype(Value)>(Value));
+				// Dummy keys
+				ThresholdAccessorKeys = MakeUnique<FPCGAttributeAccessorKeysSingleObjectPtr<void>>();
+
+				return true;
+			};
+
+			if (!ThresholdSettings.AttributeTypes.Dispatcher(ConstantThreshold))
+			{
+				return false;
+			}
+		}
+
+		check(ThresholdAccessor && ThresholdAccessorKeys);
+
+		if (!PCG::Private::IsBroadcastableOrConstructible(ThresholdAccessor->GetUnderlyingType(), InputType))
+		{
+			const FText InputTypeName = PCG::Private::GetTypeNameText(InputType);
+			const FText ThresholdTypeName = PCG::Private::GetTypeNameText(ThresholdAccessor->GetUnderlyingType());
+			PCGLog::LogWarningOnGraph(FText::Format(LOCTEXT("TypeConversionFailed", "Cannot convert threshold type '{0}' to input target type '{1}'"),
+				ThresholdTypeName,
+				InputTypeName), Context);
+			return false;
+		}
+
+		return true;
+	}
+
+	bool FilterByValue(const FPCGTaggedData& Input, const UPCGFilterByAttributeSettings& Settings, const FPCGTaggedData* ThresholdData, const FPCGContext* Context)
+	{
+		const FPCGAttributePropertyInputSelector TargetAttribute = Settings.TargetAttribute.CopyAndFixLast(Input.Data);
+		TUniquePtr<const IPCGAttributeAccessor> InputAccessor = PCGAttributeAccessorHelpers::CreateConstAccessor(Input.Data, TargetAttribute);
+		TUniquePtr<const IPCGAttributeAccessorKeys> InputAccessorKeys = PCGAttributeAccessorHelpers::CreateConstKeys(Input.Data, TargetAttribute);
+
+		if (!InputAccessor || !InputAccessorKeys)
+		{
+			PCGLog::Metadata::LogFailToCreateAccessorError(TargetAttribute, Context);
+			return false;
+		}
+		
+		TUniquePtr<const IPCGAttributeAccessor> ThresholdAccessor;
+		TUniquePtr<const IPCGAttributeAccessorKeys> ThresholdAccessorKeys;
+		if (!CreateAndValidateAccessorThreshold(ThresholdAccessor, ThresholdAccessorKeys, Settings.Threshold, ThresholdData, InputAccessor->GetUnderlyingType(), Context))
+		{
+			return false;
+		}
+
+		check(InputAccessor && InputAccessorKeys && ThresholdAccessor && ThresholdAccessorKeys);
+
+		auto DoFilter = [&Settings, &InputAccessor, &InputAccessorKeys, &ThresholdAccessor, &ThresholdAccessorKeys]<typename T>(T) -> bool
+		{
+			bool bResult = false;
+			// if bResult != bShouldContinue, we can stop the processing
+			const bool bShouldContinue = Settings.FilterByValueMode == EPCGFilterByAttributeValueMode::AllOf;
+
+			PCGMetadataElementCommon::ApplyOnMultiAccessors<T, T>(
+				{InputAccessorKeys.Get(), ThresholdAccessorKeys.Get()},
+				{InputAccessor.Get(), ThresholdAccessor.Get()},
+				[&bResult, FilterOperator = Settings.FilterOperator, bShouldContinue](const T& Value, const T& Threshold, int) -> bool
+			{
+				bResult = PCGAttributeFilterHelpers::ApplyCompare(Value, Threshold, FilterOperator);
+				return bResult == bShouldContinue;
+			}, EPCGAttributeAccessorFlags::AllowBroadcastAndConstructible);
+
+			return bResult;
+		};
+
+		return PCGMetadataAttribute::CallbackWithRightType(InputAccessor->GetUnderlyingType(), DoFilter);
+	}
+
+	bool FilterByValueRange(const FPCGTaggedData& Input, const UPCGFilterByAttributeSettings& Settings, const FPCGTaggedData* MinThresholdData, const FPCGTaggedData* MaxThresholdData, FPCGContext* Context)
+	{
+		const FPCGAttributePropertyInputSelector TargetAttribute = Settings.TargetAttribute.CopyAndFixLast(Input.Data);
+		TUniquePtr<const IPCGAttributeAccessor> InputAccessor = PCGAttributeAccessorHelpers::CreateConstAccessor(Input.Data, TargetAttribute);
+		TUniquePtr<const IPCGAttributeAccessorKeys> InputAccessorKeys = PCGAttributeAccessorHelpers::CreateConstKeys(Input.Data, TargetAttribute);
+
+		if (!InputAccessor || !InputAccessorKeys)
+		{
+			PCGLog::Metadata::LogFailToCreateAccessorError(TargetAttribute, Context);
+			return false;
+		}
+		
+		TUniquePtr<const IPCGAttributeAccessor> MinThresholdAccessor, MaxThresholdAccessor;
+		TUniquePtr<const IPCGAttributeAccessorKeys> MinThresholdAccessorKeys, MaxThresholdAccessorKeys;
+		if (!CreateAndValidateAccessorThreshold(MinThresholdAccessor, MinThresholdAccessorKeys, Settings.MinThreshold, MinThresholdData, InputAccessor->GetUnderlyingType(), Context)
+			|| !CreateAndValidateAccessorThreshold(MaxThresholdAccessor, MaxThresholdAccessorKeys, Settings.MaxThreshold, MaxThresholdData, InputAccessor->GetUnderlyingType(), Context))
+		{
+			return false;
+		}
+
+		check(InputAccessor && InputAccessorKeys && MinThresholdAccessor && MinThresholdAccessorKeys && MaxThresholdAccessor && MaxThresholdAccessorKeys);
+		
+		auto DoFilter = [&Settings, &InputAccessor, &InputAccessorKeys, &MinThresholdAccessor, &MinThresholdAccessorKeys, &MaxThresholdAccessor, &MaxThresholdAccessorKeys]<typename T>(T) -> bool
+		{
+			bool bResult = false;
+			
+			// Remove some types to avoid compiling too many version of ApplyOnMultiAccessor
+			if constexpr (PCG::Private::MetadataTraits<T>::CanCompare)
+			{
+				// if bResult != bShouldContinue, we can stop the processing
+				const bool bShouldContinue = Settings.FilterByValueMode == EPCGFilterByAttributeValueMode::AllOf;
+
+				PCGMetadataElementCommon::ApplyOnMultiAccessors<T, T, T>(
+					{InputAccessorKeys.Get(), MinThresholdAccessorKeys.Get(), MaxThresholdAccessorKeys.Get()},
+					{InputAccessor.Get(), MinThresholdAccessor.Get(), MaxThresholdAccessor.Get()},
+					[&bResult, MinInclusive = Settings.MinThreshold.bInclusive, MaxInclusive = Settings.MaxThreshold.bInclusive, bShouldContinue](const T& Value, const T& MinThreshold, const T& MaxThreshold, int) -> bool
+				{
+					bResult = PCGAttributeFilterHelpers::ApplyRange(Value, MinThreshold, MaxThreshold, MinInclusive, MaxInclusive);
+					return bResult == bShouldContinue;
+				}, EPCGAttributeAccessorFlags::AllowBroadcastAndConstructible);
+			}
+
+			return bResult;
+		};
+
+		return PCGMetadataAttribute::CallbackWithRightType(InputAccessor->GetUnderlyingType(), DoFilter);
+	}
+}
+
+#if WITH_EDITOR
+FText UPCGFilterByAttributeSettings::GetDefaultNodeTitle() const
+{
+	return LOCTEXT("NodeTitle", "Filter Data By Attribute");
+}
+#endif // WITH_EDITOR
+
+FString UPCGFilterByAttributeSettings::GetAdditionalTitleInformation() const
+{
+#if WITH_EDITOR
+	if (IsPropertyOverriddenByPin(GET_MEMBER_NAME_CHECKED(UPCGFilterByAttributeSettings, FilterMode)))
+	{
+		return {};
+	}
+#endif // WITH_EDITOR
+	
+	if (FilterMode == EPCGFilterByAttributeMode::FilterByExistence)
+	{
+		static const FText Subtitle = LOCTEXT("FilterByExistence", "By existence"); 
+#if WITH_EDITOR
+		if (IsPropertyOverriddenByPin(GET_MEMBER_NAME_CHECKED(UPCGFilterByAttributeSettings, Attribute)))
+		{
+			return Subtitle.ToString();
+		}
+		else
+#endif
+		{
+			return FText::Format(INVTEXT("{0}: {1}"), Subtitle, FText::FromName(Attribute)).ToString();
+		}
+	}
+	else if (FilterMode == EPCGFilterByAttributeMode::FilterByValue || FilterMode == EPCGFilterByAttributeMode::FilterByValueRange)
+	{
+		static const FText SubtitleValue = LOCTEXT("FilterByValue", "By value");
+		static const FText SubtitleValueRange = LOCTEXT("FilterByValueRange", "By value range");
+		const FText& Subtitle = FilterMode == EPCGFilterByAttributeMode::FilterByValue ? SubtitleValue : SubtitleValueRange;
+#if WITH_EDITOR
+		if (IsPropertyOverriddenByPin(GET_MEMBER_NAME_CHECKED(UPCGFilterByAttributeSettings, TargetAttribute)))
+		{
+			return Subtitle.ToString();
+		}
+		else
+#endif
+		{
+			return FText::Format(INVTEXT("{0}: {1}"), Subtitle, TargetAttribute.GetDisplayText()).ToString();
+		}
+	}
+
+	return {};
+}
+
+#if WITH_EDITOR
+EPCGChangeType UPCGFilterByAttributeSettings::GetChangeTypeForProperty(const FName& InPropertyName) const
+{
+	EPCGChangeType ChangeType = Super::GetChangeTypeForProperty(InPropertyName) | EPCGChangeType::Cosmetic;
+
+	if (InPropertyName == GET_MEMBER_NAME_CHECKED(UPCGFilterByAttributeSettings, FilterMode))
+	{
+		ChangeType |= EPCGChangeType::Structural;
+	}
+
+	return ChangeType;
+}
+#endif // WITH_EDITOR
+
+TArray<FPCGPinProperties> UPCGFilterByAttributeSettings::InputPinProperties() const
+{
+	TArray<FPCGPinProperties> Properties;
+
+	Properties.Emplace_GetRef(PCGPinConstants::DefaultInputLabel).SetRequiredPin();
+
+	if (FilterMode == EPCGFilterByAttributeMode::FilterByValue && !Threshold.bUseConstantThreshold)
+	{
+		Properties.Emplace_GetRef(PCGAttributeFilterConstants::FilterLabel).SetRequiredPin();
+	}
+	else if (FilterMode == EPCGFilterByAttributeMode::FilterByValueRange)
+	{
+		if (!MinThreshold.bUseConstantThreshold)
+		{
+			Properties.Emplace_GetRef(PCGAttributeFilterConstants::FilterMinLabel).SetRequiredPin();
+		}
+
+		if (!MaxThreshold.bUseConstantThreshold)
+		{
+			Properties.Emplace_GetRef(PCGAttributeFilterConstants::FilterMaxLabel).SetRequiredPin();
+		}
+	}
+
+	return Properties;
+}
+
+FPCGElementPtr UPCGFilterByAttributeSettings::CreateElement() const
+{
+	return MakeShared<FPCGFilterByAttributeElement>();
+}
+
+bool FPCGFilterByAttributeElement::ExecuteInternal(FPCGContext* Context) const
+{
+	TRACE_CPUPROFILER_EVENT_SCOPE(FPCGFilterByAttributeElement::Execute);
+	check(Context);
+
+	const UPCGFilterByAttributeSettings* Settings = Context->GetInputSettings<UPCGFilterByAttributeSettings>();
+	check(Settings);
+
+	TArray<FString> Attributes = PCGHelpers::GetStringArrayFromCommaSeparatedList(Settings->Attribute.ToString());
+
+	TArray<FPCGTaggedData> Inputs = Context->InputData.GetInputsByPin(PCGPinConstants::DefaultInputLabel);
+	TArray<FPCGTaggedData> ThresholdData = Context->InputData.GetInputsByPin(PCGAttributeFilterConstants::FilterLabel);
+	TArray<FPCGTaggedData> MinThresholdData = Context->InputData.GetInputsByPin(PCGAttributeFilterConstants::FilterMinLabel);
+	TArray<FPCGTaggedData> MaxThresholdData = Context->InputData.GetInputsByPin(PCGAttributeFilterConstants::FilterMaxLabel);
+	TArray<FPCGTaggedData>& Outputs = Context->OutputData.TaggedData;
+
+	if (Inputs.IsEmpty())
+	{
+		return true;
+	}
+
+	auto ValidateCardinality = [InputNum = Inputs.Num(), Context](const FPCGFilterByAttributeThresholdSettings& ThresholdSettings, const FName ThresholdLabel, const int ThresholdNum)
+	{
+		if (!ThresholdSettings.bUseConstantThreshold && InputNum != ThresholdNum && ThresholdNum != 1)
+		{
+			PCGLog::InputOutput::LogInvalidCardinalityError(PCGPinConstants::DefaultInputLabel, ThresholdLabel, Context);
+			return false;
+		}
+
+		return true;
+	};
+
+	// Validation
+	if (Settings->FilterMode == EPCGFilterByAttributeMode::FilterByValue
+		&& !ValidateCardinality(Settings->Threshold, PCGAttributeFilterConstants::FilterLabel, ThresholdData.Num()))
+	{
+		Context->OutputData = Context->InputData;
+		return true;
+	}
+	else if (Settings->FilterMode == EPCGFilterByAttributeMode::FilterByValueRange
+		&& (!ValidateCardinality(Settings->MinThreshold, PCGAttributeFilterConstants::FilterMinLabel, MinThresholdData.Num())
+			|| !ValidateCardinality(Settings->MaxThreshold, PCGAttributeFilterConstants::FilterMaxLabel, MaxThresholdData.Num())))
+	{
+		Context->OutputData = Context->InputData;
+		return true;
+	}
+
+	for (int32 i = 0; i < Inputs.Num(); ++i)
+	{
+		const FPCGTaggedData& Input = Inputs[i];
+		FPCGTaggedData& Output = Outputs.Add_GetRef(Input);
+		Output.Pin = PCGPinConstants::DefaultOutFilterLabel;
+
+		bool bInFilter = false;
+		switch (Settings->FilterMode)
+		{
+		case EPCGFilterByAttributeMode::FilterByExistence:
+			bInFilter = PCGFilterByAttributeElement::FilterByExistence(Input, *Settings, Attributes);
+			break;
+		case EPCGFilterByAttributeMode::FilterByValue:
+		{
+			const FPCGTaggedData* ThresholdTaggedData = Settings->Threshold.bUseConstantThreshold ? nullptr : &ThresholdData[i % ThresholdData.Num()];
+			bInFilter = PCGFilterByAttributeElement::FilterByValue(Input, *Settings, ThresholdTaggedData, Context);
+			break;
+		}
+		case EPCGFilterByAttributeMode::FilterByValueRange:
+		{
+			const FPCGTaggedData* MinThresholdTaggedData = Settings->MinThreshold.bUseConstantThreshold ? nullptr : &MinThresholdData[i % MinThresholdData.Num()];
+			const FPCGTaggedData* MaxThresholdTaggedData = Settings->MaxThreshold.bUseConstantThreshold ? nullptr : &MaxThresholdData[i % MaxThresholdData.Num()];
+			bInFilter = PCGFilterByAttributeElement::FilterByValueRange(Input, *Settings, MinThresholdTaggedData, MaxThresholdTaggedData, Context);
+			break;
+		}
+		default:
+			break;
+		}
+
 		if (bInFilter)
 		{
 			Output.Pin = PCGPinConstants::DefaultInFilterLabel;
diff --git a/Engine/Plugins/PCG/Source/PCG/Public/Elements/Metadata/PCGMetadataElementCommon.h b/Engine/Plugins/PCG/Source/PCG/Public/Elements/Metadata/PCGMetadataElementCommon.h
--- a/Engine/Plugins/PCG/Source/PCG/Public/Elements/Metadata/PCGMetadataElementCommon.h
+++ b/Engine/Plugins/PCG/Source/PCG/Public/Elements/Metadata/PCGMetadataElementCommon.h
@@ -3,6 +3,7 @@
 #pragma once
 
 #include "PCGSettings.h"
+#include "Algo/AnyOf.h"
 #include "Metadata/PCGMetadata.h"
 #include "Metadata/PCGMetadataAttributeTraits.h"
 #include "Metadata/Accessors/IPCGAttributeAccessor.h" // IWYU pragma: keep
@@ -85,7 +86,7 @@ namespace PCGMetadataElementCommon
 	/**
 	* Iterate over the full range of the keys (if Count is negative, otherwise, as many times as Count), calling the callback with values range get from the accessor.
 	* ChunkSize influence the max number of values to get in one go with GetRange.
-	* Callback should have a signature: void(const TArrayView<T>& View, int32 Start, int32 Range)
+	* Callback should have a signature: void(const TArrayView<T>& View, int32 Start, int32 Range) or bool(const TArrayView<T>& View, int32 Start, int32 Range) which returns true if we continue, false otherwise.
 	* Return false if it process nothing.
 	*/
 	template <typename T, typename Func>
@@ -108,6 +109,8 @@ namespace PCGMetadataElementCommon
 			TempValues.SetNum(ChunkSize);
 		}
 
+		using ReturnType = decltype(Callback(TArrayView<T>{}, 0, 0));
+
 		const int32 NumberOfIterations = (NumberOfEntries + ChunkSize - 1) / ChunkSize;
 
 		for (int32 i = 0; i < NumberOfIterations; ++i)
@@ -121,7 +124,17 @@ namespace PCGMetadataElementCommon
 				return false;
 			}
 
-			Callback(View, StartIndex, Range);
+			if constexpr (std::is_same_v<ReturnType, bool>)
+			{
+				if (!Callback(View, StartIndex, Range))
+				{
+					break;
+				}
+			}
+			else
+			{
+				Callback(View, StartIndex, Range);
+			}
 		}
 
 		return true;
@@ -130,7 +143,7 @@ namespace PCGMetadataElementCommon
 	/**
 	* Iterate over the full range of the keys (if Count is negative, otherwise, as many times as Count), calling the callback with values get from the accessor.
 	* ChunkSize influence the max number of values to get in one go with GetRange.
-	* Callback should have a signature: void(const T& Value, int32 Index) / void(T Value, int32 Index)
+	* Callback should have a signature: void(const T& Value, int32 Index) / void(T Value, int32 Index). Can also return a bool to know if we need to continue or not.
 	* Return false if it process nothing.
 	*/
 	template <typename T, typename Func>
@@ -138,10 +151,27 @@ namespace PCGMetadataElementCommon
 	{
 		auto RangeCallback = [Callback = std::forward<Func>(InCallback)](const TArrayView<T>& View, int32 Start, int32 Range)
 		{
+			using ReturnType = decltype(Callback(T{}, 0));
+			
 			for (int32 j = 0; j < Range; ++j)
 			{
 				// We can move the value since we don't use it afterward. If the callback is taking a const ref, it's fine.
-				Callback(std::move(View[j]), Start + j);
+				if constexpr (std::is_same_v<ReturnType, bool>)
+				{
+					if (!Callback(std::move(View[j]), Start + j))
+					{
+						return false;
+					}
+				}
+				else
+				{
+					Callback(std::move(View[j]), Start + j);
+				}
+			}
+
+			if constexpr (std::is_same_v<ReturnType, bool>)
+			{
+				return true;
 			}
 		};
 
@@ -150,18 +180,30 @@ namespace PCGMetadataElementCommon
 
 	/**
 	* Iterate over the full range of the keys (if Count is negative, otherwise, as many times as Count), calling the callback with values ranges get from the multiple accessors.
-	* Keys needs to be the same for all accessors.
+	* There can be 1 Keys for all accessors or 1 per accessor, but in that case if Count is set to -1, it will use the first key for the number of elements to process.
 	* ChunkSize influence the max number of values to get in one go with GetRange.
-	* Callback should have a signature: void(const TArrayView<T>&... Args, int32 Start, int32 Range)
+	* Callback should have a signature: void(const TArrayView<T>&... Args, int32 Start, int32 Range). Can also return a bool to know if we need to continue.
 	* Return false if it processes nothing.
 	*/
 	template <typename... T, typename Func = TFunction<void(const TArrayView<T>&..., int32 Start, int32 Range)>>
-	bool ApplyOnMultiAccessorsRange(const IPCGAttributeAccessorKeys& Keys, const TArray<const IPCGAttributeAccessor*>& Accessors, Func&& InCallback, EPCGAttributeAccessorFlags Flags = EPCGAttributeAccessorFlags::StrictType, const int32 ChunkSize = DefaultChunkSize, const int32 Count = -1)
+	bool ApplyOnMultiAccessorsRange(const TConstArrayView<IPCGAttributeAccessorKeys const*> MultiKeys, const TConstArrayView<IPCGAttributeAccessor const*> Accessors, Func&& InCallback, EPCGAttributeAccessorFlags Flags = EPCGAttributeAccessorFlags::StrictType, const int32 ChunkSize = DefaultChunkSize, const int32 Count = -1)
 	{
 		// We use C++20 functionality here, so make sure we are at this version at least
 		static_assert(__cplusplus >= 202002L);
 
-		const int32 NumberOfEntries = Count < 0 ? Keys.GetNum() : Count;
+		// We support 1 key for all accessors or 1 key per accessor.
+		if (MultiKeys.IsEmpty() || Accessors.IsEmpty() || (MultiKeys.Num() != Accessors.Num() && MultiKeys.Num() != 1) || !ensure(Algo::AllOf(MultiKeys))|| !ensure(Algo::AllOf(Accessors)))
+		{
+			return false;
+		}
+
+		// Also make sure that the keys are matching in num (or just have a single entry)
+		if (!ensure(Algo::AllOf(MultiKeys, [FirstKeys=MultiKeys[0]](const IPCGAttributeAccessorKeys* Keys){ return Keys->GetNum() == FirstKeys->GetNum() || Keys->GetNum() == 1; })))
+		{
+			return false;
+		}
+
+		const int32 NumberOfEntries = Count < 0 ? MultiKeys[0]->GetNum() : Count;
 
 		if (NumberOfEntries == 0)
 		{
@@ -225,7 +267,7 @@ namespace PCGMetadataElementCommon
 			[&]
 			{
 				TArrayView<T> View(reinterpret_cast<T*>(RawBuffers[i].GetData()), Range);
-				bSuccess &= Accessors[i]->GetRange<T>(View, StartIndex, Keys, Flags);
+				bSuccess &= Accessors[i]->GetRange<T>(View, StartIndex, *MultiKeys[i % MultiKeys.Num()], Flags);
 				++i;
 			}(), ...);
 
@@ -244,32 +286,88 @@ namespace PCGMetadataElementCommon
 				return TupleType(TArrayView<T>(reinterpret_cast<T*>(RawBuffers[I].GetData()), Range)...);
 			}(std::make_index_sequence<NumArgs>{});
 
-			TempTuple.ApplyBefore(InCallback, StartIndex, Range);
+			using ReturnType = decltype(TempTuple.ApplyBefore(InCallback, StartIndex, Range));
+
+			if constexpr (std::is_same_v<ReturnType, bool>)
+			{
+				if (!TempTuple.ApplyBefore(InCallback, StartIndex, Range))
+				{
+					break;
+				}
+			}
+			else
+			{
+				TempTuple.ApplyBefore(InCallback, StartIndex, Range);
+			}
 		}
 
 		return true;
 	}
 
 	/**
-	* Iterate over the full range of the keys (if Count is negative, otherwise, as many times as Count), calling the callback with values get from multiple accessor.
+	* Iterate over the full range of the keys (if Count is negative, otherwise, as many times as Count), calling the callback with values ranges get from the multiple accessors.
+	* ChunkSize influence the max number of values to get in one go with GetRange.
+	* Callback should have a signature: void(const TArrayView<T>&... Args, int32 Start, int32 Range)
+	* Return false if it processes nothing.
+	*/
+	template <typename... T, typename Func = TFunction<void(const TArrayView<T>&..., int32 Start, int32 Range)>>
+	bool ApplyOnMultiAccessorsRange(const IPCGAttributeAccessorKeys& Keys, const TConstArrayView<IPCGAttributeAccessor const*> Accessors, Func&& InCallback, EPCGAttributeAccessorFlags Flags = EPCGAttributeAccessorFlags::StrictType, const int32 ChunkSize = DefaultChunkSize, const int32 Count = -1)
+	{
+		const IPCGAttributeAccessorKeys* KeysPtr = &Keys;
+		return ApplyOnMultiAccessorsRange<T...>(TConstArrayView<IPCGAttributeAccessorKeys const*>(&KeysPtr, 1), Accessors, std::forward<Func>(InCallback), Flags, ChunkSize, Count);
+	}
+
+	/**
+	* Iterate over the full range of the keys (if Count is negative, otherwise, as many times as Count), calling the callback with values get from the multiple accessors.
+	* There can be 1 keys for all accessors or 1 per accessor, but in that case if Count is set to -1, it will use the first key for the number of elements to process.
 	* ChunkSize influence the max number of values to get in one go with GetRange.
-	* Callback should have a signature: void(const T&... Values, int32 Index) / void(T... Values, int32 Index)
+	* Callback should have a signature: void(const T&... Values, int32 Index) / void(T... Values, int32 Index). Can return a bool to know if we need to continue.
 	* Return false if it process nothing.
 	*/
 	template <typename ...T, typename Func = TFunction<void(const T&..., int32 Index)>>
-	bool ApplyOnMultiAccessors(const IPCGAttributeAccessorKeys& Keys, const TArray<const IPCGAttributeAccessor*>& Accessors, Func&& InCallback, EPCGAttributeAccessorFlags Flags = EPCGAttributeAccessorFlags::StrictType, const int32 ChunkSize = DefaultChunkSize, const int32 Count = -1)
+	bool ApplyOnMultiAccessors(const TConstArrayView<IPCGAttributeAccessorKeys const*> MultiKeys, const TConstArrayView<IPCGAttributeAccessor const*> Accessors, Func&& InCallback, EPCGAttributeAccessorFlags Flags = EPCGAttributeAccessorFlags::StrictType, const int32 ChunkSize = DefaultChunkSize, const int32 Count = -1)
 	{
 		auto RangeCallback = [Callback = std::forward<Func>(InCallback)](const TArrayView<T>& ...Views, int32 Start, int32 Range)
 		{
+			using ReturnType = decltype(Callback(Views[0]..., 0));
+
 			for (int32 j = 0; j < Range; ++j)
 			{
 				// We can move the value since we don't use it afterward. If the callback is taking a const ref, it's fine.
-				Callback(std::move(Views[j])..., Start + j);
+				if constexpr (std::is_same_v<ReturnType, bool>)
+				{
+					if (!Callback(std::move(Views[j])..., Start + j))
+					{
+						return false;
+					}
+				}
+				else
+				{
+					Callback(std::move(Views[j])..., Start + j);
+				}
+			}
+
+			if constexpr (std::is_same_v<ReturnType, bool>)
+			{
+				return true;
 			}
 		};
 
-		return ApplyOnMultiAccessorsRange<T...>(Keys, Accessors, std::move(RangeCallback), Flags, ChunkSize, Count);
-	}	
+		return ApplyOnMultiAccessorsRange<T...>(MultiKeys, Accessors, std::move(RangeCallback), Flags, ChunkSize, Count);
+	}
+
+	/**
+	* Iterate over the full range of the keys (if Count is negative, otherwise, as many times as Count), calling the callback with values get from the multiple accessors.
+	* ChunkSize influence the max number of values to get in one go with GetRange.
+	* Callback should have a signature: void(const T&... Values, int32 Index) / void(T... Values, int32 Index). Can return a bool to know if we need to continue.
+	* Return false if it process nothing.
+	*/
+	template <typename ...T, typename Func = TFunction<void(const T&..., int32 Index)>>
+	bool ApplyOnMultiAccessors(const IPCGAttributeAccessorKeys& Keys, const TConstArrayView<IPCGAttributeAccessor const*> Accessors, Func&& InCallback, EPCGAttributeAccessorFlags Flags = EPCGAttributeAccessorFlags::StrictType, const int32 ChunkSize = DefaultChunkSize, const int32 Count = -1)
+	{
+		const IPCGAttributeAccessorKeys* KeysPtr = &Keys;
+		return ApplyOnMultiAccessors<T...>(TConstArrayView<IPCGAttributeAccessorKeys const*>(&KeysPtr, 1), Accessors, std::forward<Func>(InCallback), Flags, ChunkSize, Count);
+	}
 
 	/**
 	* Read all the values from in accessor and write to out accessor. Iterate as many times than the out accessor keys
diff --git a/Engine/Plugins/PCG/Source/PCG/Public/Elements/PCGAttributeFilter.h b/Engine/Plugins/PCG/Source/PCG/Public/Elements/PCGAttributeFilter.h
--- a/Engine/Plugins/PCG/Source/PCG/Public/Elements/PCGAttributeFilter.h
+++ b/Engine/Plugins/PCG/Source/PCG/Public/Elements/PCGAttributeFilter.h
@@ -51,7 +51,7 @@ struct FPCGAttributeFilterThresholdSettings
 };
 
 /**
-* Attribute filter that allows to do "A op B" type filtering, where A is the input spatial data or Attribute set,
+* Filter elements by attribute that allows to do "A op B" type filtering, where A is the input spatial data or Attribute set,
 * and B is either a constant, another spatial data (if input is a spatial data), an Attribute set (in filter) or the input itself.
 * The filtering can be done either on properties or attributes.
 * Some examples:
@@ -77,7 +77,7 @@ public:
 	//~Begin UPCGSettings interface
 #if WITH_EDITOR
 	virtual FName GetDefaultNodeName() const override { return FName(TEXT("AttributeFilter")); }
-	virtual FText GetDefaultNodeTitle() const override { return NSLOCTEXT("PCGAttributeFilteringElement", "NodeTitle", "Attribute Filter"); }
+	virtual FText GetDefaultNodeTitle() const override { return NSLOCTEXT("PCGAttributeFilteringElement", "NodeTitle", "Filter Attribute Elements"); }
 	virtual EPCGSettingsType GetType() const override { return EPCGSettingsType::Filter; }
 
 	// Expose 2 nodes: Attribute filter and Point Filter that will not have the same defaults.
@@ -156,7 +156,7 @@ public:
 	//~Begin UPCGSettings interface
 #if WITH_EDITOR
 	virtual FName GetDefaultNodeName() const override { return FName(TEXT("AttributeFilterRange")); }
-	virtual FText GetDefaultNodeTitle() const override { return NSLOCTEXT("PCGAttributeFilteringElement", "NodeTitleRange", "Attribute Filter Range"); }
+	virtual FText GetDefaultNodeTitle() const override { return NSLOCTEXT("PCGAttributeFilteringElement", "NodeTitleRange", "Filter Attribute Elements by Range"); }
 	virtual EPCGSettingsType GetType() const override { return EPCGSettingsType::Filter; }
 
 	// Expose 2 nodes: "Attribute Filter Range" and "Point Filter Range" that will not have the same defaults.
@@ -217,3 +217,72 @@ class FPCGAttributeFilterRangeElement : public FPCGAttributeFilterElementBase
 protected:
 	virtual bool ExecuteInternal(FPCGContext* Context) const override;
 };
+
+namespace PCGAttributeFilterConstants
+{
+	const FName FilterLabel = TEXT("Filter");
+	const FName FilterMinLabel = TEXT("FilterMin");
+	const FName FilterMaxLabel = TEXT("FilterMax");
+}
+
+namespace PCGAttributeFilterHelpers
+{
+	template <typename T>
+	bool ApplyCompare(const T& Input1, const T& Input2, EPCGAttributeFilterOperator Operation)
+	{
+		if (Operation == EPCGAttributeFilterOperator::Equal)
+		{
+			return PCG::Private::MetadataTraits<T>::Equal(Input1, Input2);
+		}
+		else if (Operation == EPCGAttributeFilterOperator::NotEqual)
+		{
+			return !PCG::Private::MetadataTraits<T>::Equal(Input1, Input2);
+		}
+
+		if constexpr (PCG::Private::MetadataTraits<T>::CanCompare)
+		{
+			switch (Operation)
+			{
+			case EPCGAttributeFilterOperator::Greater:
+				return PCG::Private::MetadataTraits<T>::Greater(Input1, Input2);
+			case EPCGAttributeFilterOperator::GreaterOrEqual:
+				return PCG::Private::MetadataTraits<T>::GreaterOrEqual(Input1, Input2);
+			case EPCGAttributeFilterOperator::Lesser:
+				return PCG::Private::MetadataTraits<T>::Less(Input1, Input2);
+			case EPCGAttributeFilterOperator::LesserOrEqual:
+				return PCG::Private::MetadataTraits<T>::LessOrEqual(Input1, Input2);
+			default:
+				break;
+			}
+		}
+
+		if constexpr (PCG::Private::MetadataTraits<T>::CanSearchString)
+		{
+			switch (Operation)
+			{
+			case EPCGAttributeFilterOperator::Substring:
+				return PCG::Private::MetadataTraits<T>::Substring(Input1, Input2);
+			case EPCGAttributeFilterOperator::Matches:
+				return PCG::Private::MetadataTraits<T>::Matches(Input1, Input2);
+			default:
+				break;
+			}
+		}
+
+		return false;
+	}
+
+	template <typename T>
+	bool ApplyRange(const T& Input, const T& InMin, const T& InMax, bool bMinIncluded, bool bMaxIncluded)
+	{
+		if constexpr (PCG::Private::MetadataTraits<T>::CanCompare)
+		{
+			return (bMinIncluded ? PCG::Private::MetadataTraits<T>::GreaterOrEqual(Input, InMin) : PCG::Private::MetadataTraits<T>::Greater(Input, InMin)) &&
+				(bMaxIncluded ? PCG::Private::MetadataTraits<T>::LessOrEqual(Input, InMax) : PCG::Private::MetadataTraits<T>::Less(Input, InMax));
+		}
+		else
+		{
+			return false;
+		}
+	}
+}
diff --git a/Engine/Plugins/PCG/Source/PCG/Public/Elements/PCGFilterByAttribute.h b/Engine/Plugins/PCG/Source/PCG/Public/Elements/PCGFilterByAttribute.h
--- a/Engine/Plugins/PCG/Source/PCG/Public/Elements/PCGFilterByAttribute.h
+++ b/Engine/Plugins/PCG/Source/PCG/Public/Elements/PCGFilterByAttribute.h
@@ -2,11 +2,52 @@
 
 #pragma once
 
+#include "PCGAttributeFilter.h"
 #include "Elements/PCGFilterDataBase.h"
+#include "Metadata/PCGMetadataTypesConstantStruct.h"
 
 #include "PCGFilterByAttribute.generated.h"
 
-/** Separates data on whether they have a specific metadata attribute (not by value) */
+UENUM()
+enum class EPCGFilterByAttributeMode
+{
+	FilterByExistence UMETA(Tooltip="Filter the data if they have a specific attribute."),
+	FilterByValue UMETA(Tooltip="Filter the data by comparing the values of an attribute."),
+	FilterByValueRange UMETA(Tooltip="Filter the data by comparing the values of an attribute within a range.")
+};
+
+UENUM()
+enum class EPCGFilterByAttributeValueMode
+{
+	AnyOf UMETA(Tooltip="Any of the values satisfy the filter"),
+	AllOf UMETA(Tooltip="All of the values satisfy the filter"),
+};
+
+USTRUCT(BlueprintType)
+struct FPCGFilterByAttributeThresholdSettings
+{
+	GENERATED_BODY()
+	
+	UPROPERTY(BlueprintReadWrite, EditAnywhere, Category = Settings, meta = (PCG_Overridable))
+	bool bUseConstantThreshold = true;
+
+	UPROPERTY(BlueprintReadWrite, EditAnywhere, Category = Settings, meta = (EditCondition = "!bUseConstantThreshold", EditConditionHides, PCG_Overridable))
+	FPCGAttributePropertyInputSelector ThresholdAttribute;
+
+	UPROPERTY(BlueprintReadWrite, EditAnywhere, Category = Settings, meta = (EditCondition = "bUseConstantThreshold", EditConditionHides, ShowOnlyInnerProperties, DisplayAfter = "bUseConstantThreshold", PCG_NotOverridable))
+	FPCGMetadataTypesConstantStruct AttributeTypes;
+};
+
+USTRUCT(BlueprintType)
+struct FPCGFilterByAttributeThresholdSettingsRange : public FPCGFilterByAttributeThresholdSettings
+{
+	GENERATED_BODY()
+	
+	UPROPERTY(BlueprintReadWrite, EditAnywhere, Category = Settings, meta = (PCG_Overridable, EditCondition = "ShowInclusive()", EditConditionHides))
+	bool bInclusive = false;
+};
+
+/** Separates data on whether they have a specific metadata attribute, depending on some criteria on attribute values. */
 UCLASS(MinimalAPI, BlueprintType, ClassGroup = (Procedural))
 class UPCGFilterByAttributeSettings : public UPCGFilterDataBaseSettings
 {
@@ -17,29 +58,57 @@ public:
 #if WITH_EDITOR
 	virtual FName GetDefaultNodeName() const override { return FName(TEXT("FilterDataByAttribute")); }
 	virtual FText GetDefaultNodeTitle() const override;
-	virtual FText GetNodeTooltipText() const override { return NSLOCTEXT("PCGFilterByAttributeElement", "NodeTooltip", "Separates input data by whether they have the specified attribute or not."); }
+	virtual FText GetNodeTooltipText() const override { return NSLOCTEXT("PCGFilterByAttributeElement", "NodeTooltip", "Separates input data by whether they have the specified attribute or not, or on the data attribute value."); }
 #endif
 	
 	virtual FString GetAdditionalTitleInformation() const override;
 
 protected:
 #if WITH_EDITOR
-	virtual EPCGChangeType GetChangeTypeForProperty(const FName& InPropertyName) const override { return Super::GetChangeTypeForProperty(InPropertyName) | EPCGChangeType::Cosmetic; }
+	virtual EPCGChangeType GetChangeTypeForProperty(const FName& InPropertyName) const override;
 #endif
+
+	virtual TArray<FPCGPinProperties> InputPinProperties() const override;
 	virtual FPCGElementPtr CreateElement() const override;
 	//~End UPCGSettings interface
 
 public:
+	UPROPERTY(BlueprintReadWrite, EditAnywhere, Category = Settings, meta = (PCG_Overridable))
+	EPCGFilterByAttributeMode FilterMode = EPCGFilterByAttributeMode::FilterByExistence;
+
 	/** Comma-separated list of attributes to look for */
-	UPROPERTY(BlueprintReadWrite, EditAnywhere, Category = Settings, meta = (DisplayName = "Attributes", PCG_Overridable))
+	UPROPERTY(BlueprintReadWrite, EditAnywhere, Category = Settings, meta = (DisplayName = "Attributes", PCG_Overridable, EditCondition = "FilterMode == EPCGFilterByAttributeMode::FilterByExistence", EditConditionHides))
 	FName Attribute;
+	
+	/** Domain to target for filtering existence */
+	UPROPERTY(BlueprintReadWrite, EditAnywhere, Category = Settings, meta = (PCG_Overridable, EditCondition = "FilterMode == EPCGFilterByAttributeMode::FilterByExistence", EditConditionHides))
+	FName MetadataDomain = PCGDataConstants::DefaultDomainName;
 
-	UPROPERTY(BlueprintReadWrite, EditAnywhere, Category = Settings, meta = (PCG_Overridable))
+	UPROPERTY(BlueprintReadWrite, EditAnywhere, Category = Settings, meta = (PCG_Overridable, EditCondition = "FilterMode == EPCGFilterByAttributeMode::FilterByExistence", EditConditionHides))
 	EPCGStringMatchingOperator Operator = EPCGStringMatchingOperator::Equal;
 
 	/** Controls whether properties (denoted by $) will be considered in the filter or not. */
-	UPROPERTY(BlueprintReadWrite, EditAnywhere, Category = Settings, meta = (PCG_Overridable))
+	UPROPERTY(BlueprintReadWrite, EditAnywhere, Category = Settings, meta = (PCG_Overridable, EditCondition = "FilterMode == EPCGFilterByAttributeMode::FilterByExistence", EditConditionHides))
 	bool bIgnoreProperties = false;
+
+	UPROPERTY(BlueprintReadWrite, EditAnywhere, Category = Settings, meta = (PCG_Overridable, EditCondition = "FilterMode != EPCGFilterByAttributeMode::FilterByExistence", EditConditionHides))
+	EPCGFilterByAttributeValueMode FilterByValueMode = EPCGFilterByAttributeValueMode::AnyOf;
+	
+	UPROPERTY(BlueprintReadWrite, EditAnywhere, Category = Settings, meta = (PCG_Overridable, EditCondition = "FilterMode != EPCGFilterByAttributeMode::FilterByExistence", EditConditionHides))
+	FPCGAttributePropertyInputSelector TargetAttribute;
+	
+	UPROPERTY(BlueprintReadWrite, EditAnywhere, Category = Settings, meta = (PCG_Overridable, EditCondition = "FilterMode == EPCGFilterByAttributeMode::FilterByValue", EditConditionHides))
+	EPCGAttributeFilterOperator FilterOperator = EPCGAttributeFilterOperator::Greater;
+	
+	UPROPERTY(BlueprintReadWrite, EditAnywhere, Category = Settings, meta = (PCG_Overridable, EditCondition = "FilterMode == EPCGFilterByAttributeMode::FilterByValue", EditConditionHides))
+	FPCGFilterByAttributeThresholdSettings Threshold;
+
+	/** Threshold property/attribute/constant related properties */
+	UPROPERTY(BlueprintReadWrite, EditAnywhere, Category = Settings, meta = (PCG_Overridable, EditCondition = "FilterMode == EPCGFilterByAttributeMode::FilterByValueRange", EditConditionHides))
+	FPCGFilterByAttributeThresholdSettingsRange MinThreshold;
+
+	UPROPERTY(BlueprintReadWrite, EditAnywhere, Category = Settings, meta = (PCG_Overridable, EditCondition = "FilterMode == EPCGFilterByAttributeMode::FilterByValueRange", EditConditionHides))
+	FPCGFilterByAttributeThresholdSettingsRange MaxThreshold;
 };
 
 class FPCGFilterByAttributeElement : public IPCGElement
