[PCG] Add Export/Import Text options to reset to default value on graph parameter overrides

* Do the same thing as the override IDs
* Needed to propagate the reset to instanced BP

#rnx

diff --git a/Engine/Plugins/PCG/Source/PCG/Private/PCGGraph.cpp b/Engine/Plugins/PCG/Source/PCG/Private/PCGGraph.cpp
--- a/Engine/Plugins/PCG/Source/PCG/Private/PCGGraph.cpp
+++ b/Engine/Plugins/PCG/Source/PCG/Private/PCGGraph.cpp
@@ -28,6 +28,7 @@
 #include "Dialogs/Dialogs.h"
 #include "EdGraph/EdGraphPin.h"
 #include "GameFramework/Actor.h"
+#include "UObject/ObjectSaveContext.h"
 #else
 #include "UObject/Package.h"
 #endif
@@ -96,6 +97,20 @@ namespace PCGGraphUtils
 		InSourcePropertyDesc->CachedProperty->CopyCompleteValue(TargetValueAddress, SourceValueAddress);
 	}
 
+#if WITH_EDITOR
+	/** Get the value for a source property as string to be used with Export and Import text. */
+	bool GetDefaultPropertyValueForEditor(const FPropertyBagPropertyDesc* InSourcePropertyDesc, const FInstancedPropertyBag& InSourceInstance, FString& OutValueString)
+	{
+		if (!InSourceInstance.IsValid() || !InSourcePropertyDesc || !InSourcePropertyDesc->CachedProperty)
+		{
+			return false;
+		}
+
+		InSourcePropertyDesc->CachedProperty->ExportText_Direct(OutValueString, InSourceInstance.GetValue().GetMemory(), InSourceInstance.GetValue().GetMemory(), nullptr, PPF_None);
+		return true;
+	}
+#endif // WITH_EDITOR
+
 	EPCGChangeType NotifyTouchedNodes(const TSet<UPCGNode*>& InTouchedNodes, EPCGChangeType InChangeType)
 	{
 		EPCGChangeType FinalChangeType = EPCGChangeType::None;
@@ -2493,11 +2508,24 @@ bool UPCGGraphInstance::IsPropertyOverriddenAndNotDefault(const FProperty* InPro
 	return Graph ? ParametersOverrides.IsPropertyOverriddenAndNotDefault(InProperty, Graph->GetUserParametersStruct()) : false;
 }
 
+
 #if WITH_EDITOR
-FString UPCGGraphInstance::ExportOverriddenPropertyIdsChangeForEditor(const FProperty* InProperty, bool bMarkAsOverridden, bool& bHasChanged)
+FString UPCGGraphInstance::GetDefaultPropertyValueForEditor(const FProperty* InProperty, bool& bIsDifferent) const
+{
+	bIsDifferent = false;
+	
+	if (!IsPropertyOverridden(InProperty))
+	{
+		return {};
+	}
+	
+	return ParametersOverrides.GetDefaultPropertyValueForEditor(InProperty, Graph->GetUserParametersStruct(), bIsDifferent);
+}
+
+FString UPCGGraphInstance::ExportOverriddenPropertyIdsChangeForEditor(const FProperty* InProperty, bool bMarkAsOverridden, bool& bIsDifferent) const
 {
 	TSet<FGuid> CopyOfPropertyIDsOverridden;
-	bHasChanged = false;
+	bIsDifferent = false;
 	
 	if (const FPropertyBagPropertyDesc* PropertyDesc = ParametersOverrides.Parameters.FindPropertyDescByName(InProperty->GetFName()))
 	{
@@ -2505,17 +2533,17 @@ FString UPCGGraphInstance::ExportOverriddenPropertyIdsChangeForEditor(const FPro
 		{
 			CopyOfPropertyIDsOverridden = ParametersOverrides.PropertiesIDsOverridden;
 			CopyOfPropertyIDsOverridden.Add(PropertyDesc->ID);
-			bHasChanged = true;
+			bIsDifferent = true;
 		}
 		else if (!bMarkAsOverridden && ParametersOverrides.PropertiesIDsOverridden.Contains(PropertyDesc->ID))
 		{
 			CopyOfPropertyIDsOverridden = ParametersOverrides.PropertiesIDsOverridden;
 			CopyOfPropertyIDsOverridden.Remove(PropertyDesc->ID);
-			bHasChanged = true;
+			bIsDifferent = true;
 		}
 	}
 
-	if (bHasChanged)
+	if (bIsDifferent)
 	{
 		const FProperty* PropertyIDsOverriddenProperty = FPCGOverrideInstancedPropertyBag::StaticStruct()->FindPropertyByName(GET_MEMBER_NAME_CHECKED(FPCGOverrideInstancedPropertyBag, PropertiesIDsOverridden));
 		check(PropertyIDsOverriddenProperty);
@@ -2732,6 +2760,32 @@ bool FPCGOverrideInstancedPropertyBag::ResetPropertyToDefault(const FProperty* I
 	return false;
 }
 
+#if WITH_EDITOR
+FString FPCGOverrideInstancedPropertyBag::GetDefaultPropertyValueForEditor(const FProperty* InProperty, const FInstancedPropertyBag* ParentUserParameters, bool& bIsDifferent) const
+{
+	bIsDifferent = false;
+	check(ParentUserParameters);
+
+	const FPropertyBagPropertyDesc* OriginalPropertyDesc = ParentUserParameters->FindPropertyDescByName(InProperty->GetFName());
+	const FPropertyBagPropertyDesc* ThisPropertyDesc = Parameters.FindPropertyDescByName(InProperty->GetFName());
+
+	if (OriginalPropertyDesc && ThisPropertyDesc)
+	{
+		if (!PCGGraphUtils::ArePropertiesIdentical(OriginalPropertyDesc, *ParentUserParameters, ThisPropertyDesc, Parameters))
+		{
+			FString OutValueString;
+			if (PCGGraphUtils::GetDefaultPropertyValueForEditor(OriginalPropertyDesc, *ParentUserParameters, OutValueString))
+			{
+				bIsDifferent = true;
+				return OutValueString;
+			}
+		}
+	}
+
+	return {};
+}
+#endif // WITH_EDITOR
+
 bool FPCGOverrideInstancedPropertyBag::IsPropertyOverridden(const FProperty* InProperty) const
 {
 	if (!InProperty)
diff --git a/Engine/Plugins/PCG/Source/PCG/Public/PCGGraph.h b/Engine/Plugins/PCG/Source/PCG/Public/PCGGraph.h
--- a/Engine/Plugins/PCG/Source/PCG/Public/PCGGraph.h
+++ b/Engine/Plugins/PCG/Source/PCG/Public/PCGGraph.h
@@ -67,6 +67,11 @@ public:
 	/** Reset overridden property to its parent value. Return true if the value was different. */
 	UE_API bool ResetPropertyToDefault(const FProperty* InProperty, const FInstancedPropertyBag* ParentUserParameters);
 
+#if WITH_EDITOR
+	/** Get parent value of overridden property for re-importing. Will set bIsDifferent true if the value for the property in this struct is different from the one in the parent struct. */
+	UE_API FString GetDefaultPropertyValueForEditor(const FProperty* InProperty, const FInstancedPropertyBag* ParentUserParameters, bool& bIsDifferent) const;
+#endif // WITH_EDITOR
+
 	/** Return if the property is currently marked overridden. */
 	UE_API bool IsPropertyOverridden(const FProperty* InProperty) const;
 
@@ -637,13 +642,19 @@ public:
 	UE_API void ResetPropertyToDefault(const FProperty* InProperty);
 	bool IsPropertyOverridden(const FProperty* InProperty) const { return ParametersOverrides.IsPropertyOverridden(InProperty); }
 	UE_API bool IsPropertyOverriddenAndNotDefault(const FProperty* InProperty) const;
-	
+
 #if WITH_EDITOR
 	/** For properties to be propagated correctly, changes need to happen with ImportText.
 	 * This function will make a copy of the currently overridden property ids, add/remove the InProperty (depending on bMarkAsOverridden)
 	 * and return the ExportText of the set of overridden property ids, for it to be re-imported by the Editor to effectively make the change.
 	 */
-	UE_API FString ExportOverriddenPropertyIdsChangeForEditor(const FProperty* InProperty, bool bMarkAsOverridden, bool& bHasChanged);
+	UE_API FString ExportOverriddenPropertyIdsChangeForEditor(const FProperty* InProperty, bool bMarkAsOverridden, bool& bIsDifferent) const;
+
+	 /** Same thing for the default value reset, for properties to be propagated correctly, changes need to happen with ImportText.
+	  * This function will return the default value as string, to be re-imported with ImportText. bHasChanged will be set to true if
+	  * the value for this instance is different than the value from the parent.
+	  */
+	UE_API FString GetDefaultPropertyValueForEditor(const FProperty* InProperty, bool& bIsDifferent) const;
 #endif // WITH_EDITOR
 
 	/** Because subgraph nodes hold graph instances, delegation of the selection can't be done elsewhere, so this function gets called when the dialog is opened and will forward to the graph customization if filtering is applied. */
diff --git a/Engine/Plugins/PCG/Source/PCGEditor/Private/Details/PCGInstancedPropertyBagOverrideDetails.cpp b/Engine/Plugins/PCG/Source/PCGEditor/Private/Details/PCGInstancedPropertyBagOverrideDetails.cpp
--- a/Engine/Plugins/PCG/Source/PCGEditor/Private/Details/PCGInstancedPropertyBagOverrideDetails.cpp
+++ b/Engine/Plugins/PCG/Source/PCGEditor/Private/Details/PCGInstancedPropertyBagOverrideDetails.cpp
@@ -120,12 +120,16 @@ FReply FPCGOverrideInstancedPropertyBagDataDetails::OnResetToDefaultValue(TShare
 	{
 		return FReply::Handled();
 	}
-
-	FScopedTransaction Transaction(FText::Format(LOCTEXT("OnResetToDefaultValue", "Reset Override for {0}"), FText::FromName(InPropertyHandle->GetProperty()->GetFName())));
-
+	
 	if (Owner.IsValid())
 	{
-		Owner->ResetPropertyToDefault(InPropertyHandle->GetProperty());
+		bool bHasChanged = false;
+		FString DefaultValue = Owner->GetDefaultPropertyValueForEditor(InPropertyHandle->GetProperty(), bHasChanged);
+		if (bHasChanged)
+		{
+			FScopedTransaction Transaction(FText::Format(LOCTEXT("OnResetToDefaultValue", "Reset Override for {0}"), FText::FromName(InPropertyHandle->GetProperty()->GetFName())));
+			InPropertyHandle->SetPerObjectValues({std::move(DefaultValue)});
+		}
 	}
 
 	return FReply::Handled();
