[PCG] Add type pin and color to Get Graph Parameters actions in the palette.

#rnx

diff --git a/Engine/Plugins/PCG/Source/PCGEditor/Private/PCGEditorGraph.cpp b/Engine/Plugins/PCG/Source/PCGEditor/Private/PCGEditorGraph.cpp
--- a/Engine/Plugins/PCG/Source/PCGEditor/Private/PCGEditorGraph.cpp
+++ b/Engine/Plugins/PCG/Source/PCGEditor/Private/PCGEditorGraph.cpp
@@ -403,8 +403,7 @@ FReply UPCGEditorGraph::OnPropertyBagDetailsDropOnGraphPin(const FPropertyBagPro
 	{
 		FPCGEditorGraphSchemaAction_NewGetParameterElement Action;
 		Action.SettingsClass = UPCGUserParameterGetSettings::StaticClass();
-		Action.PropertyGuid = PropertyDesc.ID;
-		Action.PropertyName = PropertyDesc.Name;
+		Action.PropertyDesc = PropertyDesc;
 		Action.PerformAction(Node->GetGraph(), Pin, GraphPosition, /*bSelectNewNode=*/true);
 	}
 
@@ -417,8 +416,7 @@ FReply UPCGEditorGraph::OnPropertyBagDetailsDropOnGraph(const FPropertyBagProper
 	{
 		FPCGEditorGraphSchemaAction_NewGetParameterElement Action;
 		Action.SettingsClass = UPCGUserParameterGetSettings::StaticClass();
-		Action.PropertyGuid = PropertyDesc.ID;
-		Action.PropertyName = PropertyDesc.Name;
+		Action.PropertyDesc = PropertyDesc;
 		Action.PerformAction(Graph, nullptr, GraphPosition, /*bSelectNewNode=*/true);
 	}
 
diff --git a/Engine/Plugins/PCG/Source/PCGEditor/Private/PCGEditorGraphSchema.cpp b/Engine/Plugins/PCG/Source/PCGEditor/Private/PCGEditorGraphSchema.cpp
--- a/Engine/Plugins/PCG/Source/PCGEditor/Private/PCGEditorGraphSchema.cpp
+++ b/Engine/Plugins/PCG/Source/PCGEditor/Private/PCGEditorGraphSchema.cpp
@@ -489,8 +489,7 @@ void UPCGEditorGraphSchema::GetNativeElementActions(FGraphActionMenuBuilder& Act
 
 				TSharedPtr<FPCGEditorGraphSchemaAction_NewGetParameterElement> NewAction(new FPCGEditorGraphSchemaAction_NewGetParameterElement(UserParameterCategory, MenuDesc, Description, 0));
 				NewAction->SettingsClass = UPCGUserParameterGetSettings::StaticClass();
-				NewAction->PropertyName = PropertyDesc.Name;
-				NewAction->PropertyGuid = PropertyDesc.ID;
+				NewAction->PropertyDesc = PropertyDesc;
 				ActionMenuBuilder.AddAction(NewAction);
 			}
 		}
@@ -939,6 +938,86 @@ FLinearColor UPCGEditorGraphSchema::GetMetadataTypeColor(const EPCGMetadataTypes
 	}
 }
 
+const FSlateBrush* UPCGEditorGraphSchema::GetPropertyBagTypeSlateBrush(EPropertyBagContainerType ContainerType) const
+{
+	switch (ContainerType)
+	{
+	case EPropertyBagContainerType::Array:
+		return FAppStyle::GetBrush(TEXT("Kismet.VariableList.ArrayTypeIcon"));
+	case EPropertyBagContainerType::Set:
+		return FAppStyle::GetBrush(TEXT("Kismet.VariableList.SetTypeIcon"));
+		// @todo_pcg: Enable after Map support finalizes
+		// case EPropertyBagContainerType::Map:
+		//	return FAppStyle::GetBrush(TEXT("Kismet.VariableList.MapKeyTypeIcon"));
+	default: // fall-through
+		return FAppStyle::GetBrush(TEXT("Kismet.VariableList.TypeIcon"));
+	}
+}
+
+FLinearColor UPCGEditorGraphSchema::GetPropertyBagTypeColor(const FPropertyBagPropertyDesc& Desc) const
+{
+	const UGraphEditorSettings* GraphSettings = GetDefault<UGraphEditorSettings>();
+	check(GraphSettings)
+
+	// For struct, extract the common structs
+	if (Desc.ValueType == EPropertyBagPropertyType::Struct && Desc.ValueTypeObject)
+	{
+		if (Desc.ValueTypeObject == TBaseStructure<FVector>::Get()
+			|| Desc.ValueTypeObject == TBaseStructure<FVector2D>::Get()
+			|| Desc.ValueTypeObject == TBaseStructure<FVector4>::Get())
+		{
+			return GraphSettings->VectorPinTypeColor;
+		}
+		else if (Desc.ValueTypeObject == TBaseStructure<FRotator>::Get()
+			|| Desc.ValueTypeObject == TBaseStructure<FQuat>::Get())
+		{
+			return GraphSettings->RotatorPinTypeColor;
+		}
+		else if (Desc.ValueTypeObject == TBaseStructure<FTransform>::Get())
+		{
+			return GraphSettings->TransformPinTypeColor;
+		}
+	}
+
+	switch (Desc.ValueType)
+	{
+	case EPropertyBagPropertyType::Bool:
+		return GraphSettings->BooleanPinTypeColor;
+	case EPropertyBagPropertyType::Byte:
+		return GraphSettings->BytePinTypeColor;
+	case EPropertyBagPropertyType::Int32:
+		return GraphSettings->IntPinTypeColor;
+	case EPropertyBagPropertyType::Int64:
+		return GraphSettings->Int64PinTypeColor;
+	case EPropertyBagPropertyType::Float:
+		return GraphSettings->FloatPinTypeColor;
+	case EPropertyBagPropertyType::Double:
+		return GraphSettings->DoublePinTypeColor;
+	case EPropertyBagPropertyType::Name:
+		return GraphSettings->NamePinTypeColor;
+	case EPropertyBagPropertyType::String:
+		return GraphSettings->StringPinTypeColor;
+	case EPropertyBagPropertyType::Text:
+		return GraphSettings->TextPinTypeColor;
+	case EPropertyBagPropertyType::Enum:
+		return GraphSettings->ObjectPinTypeColor;
+	case EPropertyBagPropertyType::Struct:
+		return GraphSettings->StructPinTypeColor;
+	case EPropertyBagPropertyType::Object:
+		return GraphSettings->ObjectPinTypeColor;
+	case EPropertyBagPropertyType::SoftObject:
+		return GraphSettings->SoftObjectPinTypeColor;
+	case EPropertyBagPropertyType::Class:
+		return GraphSettings->ClassPinTypeColor;
+	case EPropertyBagPropertyType::SoftClass:
+		return GraphSettings->SoftClassPinTypeColor;
+	default:
+		break;
+	}
+
+	return FLinearColor::White;
+}
+
 FPCGEditorConnectionDrawingPolicy::FPCGEditorConnectionDrawingPolicy(int32 InBackLayerID, int32 InFrontLayerID, float InZoomFactor, const FSlateRect& InClippingRect, FSlateWindowElementList& InDrawElements, UEdGraph* InGraph)
 	: FConnectionDrawingPolicy(InBackLayerID, InFrontLayerID, InZoomFactor, InClippingRect, InDrawElements)
 	, Graph(CastChecked<UPCGEditorGraph>(InGraph))
diff --git a/Engine/Plugins/PCG/Source/PCGEditor/Private/PCGEditorGraphSchema.h b/Engine/Plugins/PCG/Source/PCGEditor/Private/PCGEditorGraphSchema.h
--- a/Engine/Plugins/PCG/Source/PCGEditor/Private/PCGEditorGraphSchema.h
+++ b/Engine/Plugins/PCG/Source/PCGEditor/Private/PCGEditorGraphSchema.h
@@ -8,6 +8,7 @@
 #include "Editor/PCGGraphCustomization.h"
 
 #include "ConnectionDrawingPolicy.h"
+#include "StructUtils/PropertyBag.h"
 
 #include "PCGEditorGraphSchema.generated.h"
 
@@ -50,6 +51,9 @@ public:
 	virtual const FSlateBrush* GetMetadataTypeSlateBrush(EPCGContainerType ContainerType) const;
 	virtual FLinearColor GetMetadataTypeColor(EPCGMetadataTypes Type) const;
 
+	virtual const FSlateBrush* GetPropertyBagTypeSlateBrush(EPropertyBagContainerType ContainerType) const;
+	virtual FLinearColor GetPropertyBagTypeColor(const FPropertyBagPropertyDesc& Desc) const;
+
 private:
 	void GetGraphActions(FGraphActionMenuBuilder& ActionMenuBuilder, const FPCGActionsFilter& ActionFilter, bool bIsContextual) const;
 	void GetNativeElementActions(FGraphActionMenuBuilder& ActionMenuBuilder, const FPCGActionsFilter& ActionFilter) const;
diff --git a/Engine/Plugins/PCG/Source/PCGEditor/Private/PCGEditorGraphSchemaActions.cpp b/Engine/Plugins/PCG/Source/PCGEditor/Private/PCGEditorGraphSchemaActions.cpp
--- a/Engine/Plugins/PCG/Source/PCGEditor/Private/PCGEditorGraphSchemaActions.cpp
+++ b/Engine/Plugins/PCG/Source/PCGEditor/Private/PCGEditorGraphSchemaActions.cpp
@@ -130,13 +130,30 @@ FSlateBrush const* FPCGEditorGraphSchemaAction_NewNativeElement::GetPaletteIcon(
 	return PaletteIcon ? PaletteIcon.Get() : nullptr;
 }
 
+FSlateBrush const* FPCGEditorGraphSchemaAction_NewGetParameterElement::GetPaletteIcon() const
+{
+	if (!PaletteIcon)
+	{
+		const UPCGEditorGraphSchema* Schema = GetDefault<UPCGEditorGraphSchema>();
+		check(Schema);
+		
+		if (const FSlateBrush* BaseBrush = Schema->GetPropertyBagTypeSlateBrush(!PropertyDesc.ContainerTypes.IsEmpty() ? PropertyDesc.ContainerTypes[0] : EPropertyBagContainerType::None))
+		{
+			PaletteIcon = MakeShared<FSlateBrush>(*BaseBrush);
+			PaletteIcon->TintColor = Schema->GetPropertyBagTypeColor(PropertyDesc);
+		}
+	}
+
+	return PaletteIcon ? PaletteIcon.Get() : nullptr;
+}
+
 void FPCGEditorGraphSchemaAction_NewGetParameterElement::PostCreation(UPCGNode* NewNode)
 {
 	check(NewNode);
 	UPCGUserParameterGetSettings* Settings = CastChecked<UPCGUserParameterGetSettings>(NewNode->GetSettings());
 
-	Settings->PropertyGuid = PropertyGuid;
-	Settings->PropertyName = PropertyName;
+	Settings->PropertyGuid = PropertyDesc.ID;
+	Settings->PropertyName = PropertyDesc.Name;
 
 	NewNode->NodeTitle = Settings->PropertyName;
 
diff --git a/Engine/Plugins/PCG/Source/PCGEditor/Private/PCGEditorGraphSchemaActions.h b/Engine/Plugins/PCG/Source/PCGEditor/Private/PCGEditorGraphSchemaActions.h
--- a/Engine/Plugins/PCG/Source/PCGEditor/Private/PCGEditorGraphSchemaActions.h
+++ b/Engine/Plugins/PCG/Source/PCGEditor/Private/PCGEditorGraphSchemaActions.h
@@ -78,13 +78,15 @@ struct FPCGEditorGraphSchemaAction_NewGetParameterElement : public FPCGEditorGra
 	using FPCGEditorGraphSchemaAction_NewNativeElement::FPCGEditorGraphSchemaAction_NewNativeElement;
 
 	UPROPERTY()
-	FGuid PropertyGuid;
+	FPropertyBagPropertyDesc PropertyDesc;
 
-	UPROPERTY()
-	FName PropertyName;
+	virtual FSlateBrush const* GetPaletteIcon() const override;
 
 protected:
 	virtual void PostCreation(UPCGNode* NewNode) override;
+
+private:
+	mutable TSharedPtr<FSlateBrush> PaletteIcon = nullptr;
 };
 
 USTRUCT()
