[PCG] Remove inlined functions to not exposed functions in the subsystem

* Also move the actor and component mapping header in private
* Make it a UniquePtr in the subsystem
* Fix some includes

#rnx

diff --git a/Engine/Plugins/PCG/Source/PCG/Private/PCGActorAndComponentMapping.h b/Engine/Plugins/PCG/Source/PCG/Private/PCGActorAndComponentMapping.h
new file mode 100644
--- /dev/null
+++ b/Engine/Plugins/PCG/Source/PCG/Private/PCGActorAndComponentMapping.h
@@ -0,0 +1,295 @@
+// Copyright Epic Games, Inc. All Rights Reserved.
+
+#pragma once
+
+#include "PCGCommon.h"
+#include "Elements/PCGActorSelector.h"
+#include "Grid/PCGGridDescriptor.h"
+#include "Grid/PCGComponentOctree.h"
+#include "RuntimeGen/PCGRuntimeGenScheduler.h"
+
+#include "Containers/Array.h"
+#include "Containers/Map.h"
+#include "Containers/Set.h"
+#include "Templates/Function.h"
+#include "UObject/ObjectKey.h"
+#include "UObject/ObjectPtr.h"
+#include "UObject/ObjectSaveContext.h"
+#include "UObject/SoftObjectPtr.h"
+#include "UObject/WeakObjectPtr.h"
+
+class AActor;
+class ALandscapeProxy;
+class APCGPartitionActor;
+class FLandscapeProxyComponentDataChangedParams;
+class ILevelInstanceInterface;
+class UObject;
+class UPCGComponent;
+class UPCGGraph;
+class UPCGSubsystem;
+class UWorld;
+class UWorldPartition;
+class FWorldPartitionActorDescInstance;
+class FPackageReloadedEvent;
+
+enum class EPackageReloadPhase : uint8;
+
+/**
+* This class handle any necessary mapping between actors and pcg components.
+* Its meant to be part of the PCG Subsystem and owned by it. We offload some logic to this class to avoid to clutter the subsystem.
+* For now it is used for:
+* - Mapping between PCG Components and Partition actors: Dispatch tasks from original components to local components
+* - Tracking non partitioned PCG Components that has tracking actors needs.
+* - Tracking actors: Be able to react to actors changes, and potentially dirty/refresh affected components
+*/
+class FPCGActorAndComponentMapping
+{
+public:
+	friend UPCGSubsystem;
+	friend FPCGRuntimeGenScheduler;
+
+	~FPCGActorAndComponentMapping() = default;
+
+	/** Initializes callbacks, etc, tied to the PCG subsystem */
+	void Initialize();
+
+	/** Deinitializes callbacks, etc, tied to the PCG subsystem */
+	void Deinitialize();
+	
+	/** Should be called by the subsystem to handle delayed operations. */
+	void Tick();
+
+	/** Returns the PCGSubystem's World */
+	UWorld* GetWorld() const;
+
+#if WITH_EDITOR
+	/** Notify that we exited the Landscape edit mode. */
+	void NotifyLandscapeEditModeExited();
+
+	static void BuildPartitionActorRecords(APCGWorldActor* PCGWorldActor, UWorldPartition* WorldPartition, TMap<FPCGGridCellDescriptor, FGuid>& OutPartitionActorRecords, TSet<FGuid>& OutInvalidPartitionActors);
+#endif // WITH_EDITOR
+
+	/** Register a new PCG Component or update it. Returns true if it was added/updated. Thread safe */
+	bool RegisterOrUpdatePCGComponent(UPCGComponent* InComponent, bool bDoActorMapping = true);
+
+	/** In case of BP Actors, we need to remap the old component destroyed by the construction script to the new one. Returns true if a re-mapping was done. Thread safe */
+	bool RemapPCGComponent(const UPCGComponent* OldComponent, UPCGComponent* NewComponent, bool bDoActorMapping);
+
+	/** Unregister a PCG Component. Can force it, if we have a delayed unregister. Thread safe */
+	void UnregisterPCGComponent(UPCGComponent* InComponent, bool bForce = false);
+
+	/** Register a new Partition actor, will be added to a map and will query all intersecting volume to bind to them if asked. Thread safe */
+	void RegisterPartitionActor(APCGPartitionActor* InActor);
+
+	/** Unregister a Partition actor, will be removed from the map and remove itself to all intersecting volumes. Thread safe */
+	void UnregisterPartitionActor(APCGPartitionActor* InActor);
+
+	/** Return a copy of all the registered partitioned components. Thread safe */
+	TSet<UPCGComponent*> GetAllRegisteredPartitionedComponents() const;
+
+	/** Return a copy of all the registered non-partitioned components. Thread safe */
+	TSet<UPCGComponent*> GetAllRegisteredNonPartitionedComponents() const;
+
+	/** Return a copy of all the registered components. Thread safe */
+	TSet<UPCGComponent*> GetAllRegisteredComponents() const;
+
+	/** Retrieves a local component using grid descriptor and grid coordinates, returns nullptr if no such component is found. */
+	UPCGComponent* GetLocalComponent(const FPCGGridDescriptor& GridDescriptor, const FIntVector& CellCoords, const UPCGComponent* InOriginalComponent) const;
+
+	APCGPartitionActor* GetPartitionActor(const FPCGGridDescriptor& GridDescriptor, const FIntVector& CellCoords) const;
+
+#if WITH_EDITOR
+	/** Returns true if there is record of a partition actor living in a certain grid cell, regardless of whether or not it is loaded. */
+	bool DoesPartitionActorRecordExist(const FPCGGridDescriptor& GridDescriptor, const FIntVector& GridCoords) const;
+#endif
+private:
+
+	/** Used to track where a change comes from. */
+	enum class EChangedObjectReason
+	{
+		ActorAdded,
+		ActorChanged,
+		ActorDeleted,
+		ObjectChanged,
+		ObjectSaved,
+		PCGComponentGenerated,
+		Undefined
+	};
+
+#if WITH_EDITOR
+	void RegisterDelegates();
+	void UnregisterDelegates();
+
+	/** If the partition grid size change, call this to empty the Partition actors map */
+	void ResetPartitionActorsMap();
+
+	void RegisterTracking(UPCGComponent* InComponent);
+	void UpdateTracking(UPCGComponent* InComponent, bool bInShouldDirtyActors, const TArray<FPCGSelectionKey>* OptionalChangedKeys = nullptr);
+
+	void ProcessDelayedEvents();
+	bool ShouldDelayActor(AActor* InActor) const;
+#endif
+	bool RegisterOrUpdatePCGComponent(UPCGComponent* InComponent, bool bDoActorMapping, bool bForce);
+
+	// This class is only meant to be used as part of the PCG Subsytem and owned by it.
+	// So we put constructors private.
+	// We also need this class to be default constructible, since the PCGSubsytem needs to be default constructible.
+	FPCGActorAndComponentMapping() = default;
+	explicit FPCGActorAndComponentMapping(UPCGSubsystem* PCGSubsystem);
+
+	bool RegisterOrUpdatePartitionedPCGComponent(UPCGComponent* InComponent, bool bDoActorMapping = true);
+	bool RegisterOrUpdateNonPartitionedPCGComponent(UPCGComponent* InComponent);
+
+	void UnregisterPartitionedPCGComponent(UPCGComponent* InComponent);
+	void UnregisterNonPartitionedPCGComponent(UPCGComponent* InComponent);
+
+	/* Call the InFunc function to all local component registered to the original component. Return the list of all the tasks scheduled. Thread safe*/
+	TArray<FPCGTaskId> DispatchToRegisteredLocalComponents(UPCGComponent* OriginalComponent, const TFunctionRef<FPCGTaskId(UPCGComponent*)>& InFunc) const;
+
+	/* Call the InFunc function to all local component from the set of partition actors. Return the list of all the tasks scheduled. */
+	TArray<FPCGTaskId> DispatchToLocalComponents(UPCGComponent* OriginalComponent, const TSet<TObjectPtr<APCGPartitionActor>>& PartitionActors, const TFunctionRef<FPCGTaskId(UPCGComponent*)>& InFunc) const;
+
+	/** Call the InFunc function to all partitioned components which bounds intersect 'InBounds'. */
+	void ForAllIntersectingPartitionedComponents(const FBoxCenterAndExtent& InBounds, TFunctionRef<void(UPCGComponent*)> InFunc) const;
+
+	/** Call the InFunc function for all original components (regardless of partitioned or not). */
+	void ForAllOriginalComponents(TFunctionRef<void(UPCGComponent*)> InFunc);
+
+	/** Gather all the PCG components within some bounds. */
+	TArray<UPCGComponent*> GetAllIntersectingComponents(const FBoxCenterAndExtent& InBounds) const;
+
+	/** Iterate other all the int coordinates given a box and call a callback. Thread safe */
+	void ForAllIntersectingPartitionActors(const FBox& InBounds, TFunctionRef<void(APCGPartitionActor*)> InFunc) const;
+
+	/** Update the current mapping between a PCG component and its PCG Partition actors */
+	void UpdateMappingPCGComponentPartitionActor(UPCGComponent* InComponent, bool bInChangeGraphInstances = true);
+
+	/** Returns the current mapping between a PCG component and its PCG Partition actors */
+	TSet<TObjectPtr<APCGPartitionActor>> GetPCGComponentPartitionActorMappings(UPCGComponent* InComponent) const;
+	
+	/** Delete the current mapping between a PCG component and its PCG Partition actors */
+	void DeleteMappingPCGComponentPartitionActor(UPCGComponent* InComponent);
+
+	/** Returns true if a component is registered */
+	bool IsComponentRegistered(const UPCGComponent* InComponent) const;
+
+	/** Return true if there are any Original or Non-Partitioned components set to GenerateAtRuntime. */
+	bool AnyRuntimeGenComponentsExist() const;
+
+	// Typedef to store the previous position and previous tags of changed actors.
+	using FActorPreviousData = TTuple<FBox, TSet<FName>, double>;
+
+#if WITH_EDITOR
+	/** Return true if the key is tracked.*/
+	bool IsKeyTracked(const FPCGSelectionKey& InKey) const;
+
+	/** Return true if the actor is tracked.*/
+	bool IsActorTracked(const AActor* InActor) const;
+
+	void OnPackageReloaded(const EPackageReloadPhase InPackageReloadPhase, FPackageReloadedEvent* InPackageReloadedEvent);
+	void OnActorDescInstanceAdded(FWorldPartitionActorDescInstance* InActorDescInstance);
+	void OnActorDescInstanceRemoved(FWorldPartitionActorDescInstance* InActorDescInstance);
+	void OnActorAdded(AActor* InActor);
+	void OnActorLoaded(AActor& InActor);
+	void OnActorAdded_Internal(AActor* InActor, bool bShouldDirty);
+	void OnActorChanged_Internal(AActor* InActor, UObject* InOriginatingChangeObject, const EChangedObjectReason InChangeReason);
+	void OnActorChanged_Recursive(AActor* InActor, UObject* InOriginatingChangeObject, const EChangedObjectReason InChangeReason);
+	void OnActorDeleted(AActor* InActor);
+	void OnActorUnloaded(AActor& InActor);
+	void OnActorDeleted_Internal(AActor* InActor, bool bShouldDirty);
+	void OnLandscapeChanged(ALandscapeProxy* InLandscape, const FLandscapeProxyComponentDataChangedParams& InChangeParams);
+	void OnLevelInstancesUpdated(const TArray<ILevelInstanceInterface*>& InLevelInstances);
+	void OnLevelInstanceEditCancelled(ILevelInstanceInterface* InLevelInstance, bool bInHasDiscardedChanges);
+	void ApplyLandscapeChanges(ALandscapeProxy* InLandscape);
+	void OnObjectModified(UObject* InObject);
+	void OnObjectPropertyChanged(UObject* InObject, FPropertyChangedEvent& InEvent);
+	void OnObjectSaved(UObject* InObject, FObjectPreSaveContext InObjectSaveContext);
+	void OnPCGGraphStartsGenerating(UPCGComponent* InComponent);
+	void OnPCGGraphGeneratedOrCleaned(UPCGComponent* InComponent);
+	void OnPCGGraphCancelled(UPCGComponent* InComponent);
+
+	/** Remap the tracking in case of BP components. */
+	void RemapTracking(const UPCGComponent* InOldComponent, UPCGComponent* InNewComponent);
+
+	/** Unregister tracking when a component is removed */
+	void UnregisterTracking(UPCGComponent* InComponent);
+
+	/** Unregister tracking when a component is removed or keys are unregistered */
+	void UnregisterTracking(UPCGComponent* InComponent, const TSet<FPCGSelectionKey>* OptionalKeysToUntrack);
+
+	/** Trigger an update when an object/actor changed.
+	* Can specify previous actor data if it has moved or tags changed to also update components that were at its previous position or tracked previous tags.
+	* Can also specify an optional object, originating the change, to avoid re-dirtying a component if it was the origin.
+	* Another option when an actor is deleted/unload, don't refresh their components.
+	*/
+	void OnObjectChanged(UObject* InObject, const EChangedObjectReason InChangeReason, const FActorPreviousData* InPreviousData = nullptr, const UObject* InOriginatingChangeObject = nullptr);
+
+	/** Gather all settings from a given component that track the key, and clear the cache for them. Returns true if we should dirty afterwards (aka at least one settings was cleared and/or landscape changed). */
+	bool ClearCacheForKeys(const TArray<FPCGSelectionKey>& InKeys, const UPCGComponent* InComponent, const bool bIntersect, const UObject* InOriginatingChange) const;
+
+	/** Build initial records for existing partition actors */
+	void BuildPartitionActorRecords();
+#endif // WITH_EDITOR
+
+private:
+	// Cached subsystem
+	UPCGSubsystem* PCGSubsystem = nullptr;
+
+	// Mapping Component <-> PartitionActor
+	/** Octree tracking all partitioned pcg components */
+	FPCGComponentOctreeAndMap PartitionedOctree;
+
+	/** Mapping from grid size and grid coords to partition actor. We can only have 1 partition actor per grid cell. */
+	TMap<FPCGGridDescriptor, TMap<FIntVector, TObjectPtr<APCGPartitionActor>>> PartitionActorsMap;
+	mutable FRWLock PartitionActorsMapLock;
+
+	/** Mapping between original components and its overlapping partition actors. */
+	TMap<const UPCGComponent*, TSet<TObjectPtr<APCGPartitionActor>>> ComponentToPartitionActorsMap;
+	mutable FRWLock ComponentToPartitionActorsMapLock;
+
+	/** Will hold all the components that are not partitioned (and not local) and are tracking something. Will be use to dispatch actor tracking updates. */
+	FPCGComponentOctreeAndMap NonPartitionedOctree;
+
+#if WITH_EDITOR
+	/** Components to be registered/unregistered on the next frame. see RegisterOrUpdatePCGComponent/UnregisterPCGComponent for a better understanding on why it is needed. */
+	mutable FCriticalSection DelayedComponentsRegistrationLock;
+	TSet<UPCGComponent*> DelayedComponentsToUnregister;
+	TSet<UPCGComponent*> DelayedComponentsToRegister;
+	
+	// Tracking actors
+	/** Keep a mapping between tracked keys and the components that track them, and the tracking needs to be culled.*/
+	TMap<FPCGSelectionKey, TSet<UPCGComponent*>> CulledTrackedKeysToComponentsMap;
+
+	/** Same mapping but for always tracked keys */
+	TMap<FPCGSelectionKey, TSet<UPCGComponent*>> AlwaysTrackedKeysToComponentsMap;
+
+	mutable FRWLock TrackedComponentsLock;
+
+	// Keep track of actors that aren't yet ready (or if the subsystem is not yet ready), whether we should dirty them in next tick.
+	TMap<TObjectKey<AActor>, TTuple<TObjectKey<UObject>, EChangedObjectReason>> DelayedChangedActors;
+	
+	/** Transient map of actors and their previous data, it's set in the pre object change to be able to track changes (such as tags or positions) */
+	TMap<TObjectKey<AActor>, FActorPreviousData> ActorToPreviousDataMap;
+
+	/** Transient map that keep track of all components that depends on another component currently generating, to trigger the refresh only once, when all are done. */
+	TMap<TObjectKey<UPCGComponent>, TArray<TObjectKey<UPCGComponent>>> ComponentsToDependencyMap;
+
+	// Part for the delayed landscape change update
+	double LastLandscapeDirtyTime = -1.0;
+	TArray<TObjectKey<ALandscapeProxy>> DelayedModifiedLandscapes;
+
+	// Keep track of all dirtied landscapes and force a refresh on them when exiting edit mode.
+	TArray<TObjectKey<ALandscapeProxy>> DirtiedLandscapes;
+
+	// Temp boolean to indicate we are currently exiting the landscape edit mode
+	bool bIsCurrentlyExitingLandscapeEditMode = false;
+
+	// Time keeper for cleaning up cached previous actor data
+	double LastPreviousActorDataCleanup = -1.0;
+
+	// Set of existing PCG Partition Actors (for World Partition worlds)
+	TMap<FPCGGridCellDescriptor, FGuid> PartitionActorRecords;
+	// Previously generated PCG Partition Actors to ignore (will prevent them from getting registered and mark them for deletion)
+	TSet<FGuid> InvalidPartitionActors;
+#endif // WITH_EDITOR
+};
\ No newline at end of file
diff --git a/Engine/Plugins/PCG/Source/PCG/Private/PCGComponent.cpp b/Engine/Plugins/PCG/Source/PCG/Private/PCGComponent.cpp
--- a/Engine/Plugins/PCG/Source/PCG/Private/PCGComponent.cpp
+++ b/Engine/Plugins/PCG/Source/PCG/Private/PCGComponent.cpp
@@ -46,6 +46,7 @@
 #include "ScopedTransaction.h"
 #include "Editor/EditorEngine.h"
 #include "Editor/Transactor.h"
+#include "UObject/ObjectSaveContext.h"
 #include "WorldPartition/WorldPartition.h"
 #include "WorldPartition/WorldPartitionActorDescInstance.h"
 #include "WorldPartition/WorldPartitionHandle.h"
diff --git a/Engine/Plugins/PCG/Source/PCG/Private/PCGSubsystem.cpp b/Engine/Plugins/PCG/Source/PCG/Private/PCGSubsystem.cpp
--- a/Engine/Plugins/PCG/Source/PCG/Private/PCGSubsystem.cpp
+++ b/Engine/Plugins/PCG/Source/PCG/Private/PCGSubsystem.cpp
@@ -2,6 +2,7 @@
 
 #include "PCGSubsystem.h"
 
+#include "PCGActorAndComponentMapping.h"
 #include "PCGComponent.h"
 #include "PCGGraph.h"
 #include "PCGModule.h"
@@ -151,7 +152,7 @@ public:
 
 UPCGSubsystem::UPCGSubsystem()
 	: Super()
-	, ActorAndComponentMapping(this)
+	, ActorAndComponentMapping(new FPCGActorAndComponentMapping(this)) // Can't use MakeUnique because the ctor of FPCGActorAndComponentMapping is private
 {
 }
 
@@ -205,7 +206,7 @@ void UPCGSubsystem::Deinitialize()
 	PCGWorldActor = nullptr;
 	bHasTickedOnce = false;
 
-	ActorAndComponentMapping.Deinitialize();
+	ActorAndComponentMapping->Deinitialize();
 
 	Super::Deinitialize();
 }
@@ -227,7 +228,7 @@ void UPCGSubsystem::PostInitialize()
 		}
 	}
 
-	ActorAndComponentMapping.Initialize();
+	ActorAndComponentMapping->Initialize();
 
 	// Initialize graph executor
 	check(!GraphExecutor);
@@ -235,7 +236,7 @@ void UPCGSubsystem::PostInitialize()
 
 	// Initialize runtime generation scheduler
 	check(!RuntimeGenScheduler);
-	RuntimeGenScheduler = new FPCGRuntimeGenScheduler(GetWorld(), &ActorAndComponentMapping);
+	RuntimeGenScheduler = new FPCGRuntimeGenScheduler(GetWorld(), ActorAndComponentMapping.Get());
 }
 
 UPCGSubsystem* UPCGSubsystem::GetInstance(UWorld* World)
@@ -334,7 +335,7 @@ void UPCGSubsystem::Tick(float DeltaSeconds)
 	// If we have any tasks to execute, schedule some
 	GraphExecutor->Execute(EndTime);
 
-	ActorAndComponentMapping.Tick();
+	ActorAndComponentMapping->Tick();
 
 	if (PCGWorldActor)
 	{
@@ -399,7 +400,7 @@ int32 UPCGSubsystem::GetPCGQualityLevel()
 void UPCGSubsystem::OnPCGQualityLevelChanged()
 {
 	// Trigger deep refresh of all runtime generated components.
-	ActorAndComponentMapping.ForAllOriginalComponents([this](UPCGComponent* InComponent)
+	ActorAndComponentMapping->ForAllOriginalComponents([this](UPCGComponent* InComponent)
 	{
 		if (InComponent && InComponent->IsManagedByRuntimeGenSystem())
 		{
@@ -621,7 +622,7 @@ FPCGTaskId UPCGSubsystem::ScheduleComponent(UPCGComponent* PCGComponent, EPCGHiG
 		}
 #endif // WITH_EDITOR
 
-		ActorAndComponentMapping.UpdateMappingPCGComponentPartitionActor(PCGComponent);
+		ActorAndComponentMapping->UpdateMappingPCGComponentPartitionActor(PCGComponent);
 	}
 
 	// Execution dependencies require a task to finish executing before the dependent task.
@@ -710,7 +711,7 @@ FPCGTaskId UPCGSubsystem::ScheduleComponent(UPCGComponent* PCGComponent, EPCGHiG
 			return LocalComponent->GenerateInternal(bForce, LocalComponent->GetGenerationGrid(), EPCGComponentGenerationTrigger::GenerateOnDemand, *AllDependencies);
 		};
 
-		ExecutionDependencyTasks.Append(ActorAndComponentMapping.DispatchToRegisteredLocalComponents(PCGComponent, LocalGenerateTask));
+		ExecutionDependencyTasks.Append(ActorAndComponentMapping->DispatchToRegisteredLocalComponents(PCGComponent, LocalGenerateTask));
 	}
 
 	if (!ExecutionDependencyTasks.IsEmpty() || !DataDependencyTasks.IsEmpty())
@@ -784,7 +785,7 @@ FPCGTaskId UPCGSubsystem::ScheduleCleanup(UPCGComponent* PCGComponent, bool bRem
 			return LocalComponent->CleanupLocal(bRemoveComponents, Dependencies);
 		};
 
-		AllTasks.Append(ActorAndComponentMapping.DispatchToRegisteredLocalComponents(PCGComponent, LocalCleanupTask));
+		AllTasks.Append(ActorAndComponentMapping->DispatchToRegisteredLocalComponents(PCGComponent, LocalCleanupTask));
 	}
 
 	TWeakObjectPtr<UPCGComponent> ComponentPtr(PCGComponent);
@@ -804,7 +805,7 @@ FPCGTaskId UPCGSubsystem::ScheduleCleanup(UPCGComponent* PCGComponent, bool bRem
 			// the component mappings is counterproductive.
 			if (bRemoveComponents && Component->IsPartitioned())
 			{
-				ActorAndComponentMapping.DeleteMappingPCGComponentPartitionActor(Component);
+				ActorAndComponentMapping->DeleteMappingPCGComponentPartitionActor(Component);
 			}
 		}
 
@@ -928,7 +929,7 @@ void UPCGSubsystem::CancelGeneration(UPCGComponent* Component, bool bCleanupUnus
 			return InvalidPCGTaskId;
 		};
 
-		ActorAndComponentMapping.DispatchToRegisteredLocalComponents(Component, LocalCancel);
+		ActorAndComponentMapping->DispatchToRegisteredLocalComponents(Component, LocalCancel);
 	}
 
 	TArray<IPCGGraphExecutionSource*> CancelledExecutionSources = GraphExecutor->Cancel(Component);
@@ -1038,7 +1039,7 @@ void UPCGSubsystem::ForAllRegisteredLocalComponents(UPCGComponent* InOriginalCom
 		return InvalidPCGTaskId;
 	};
 
-	ActorAndComponentMapping.DispatchToRegisteredLocalComponents(InOriginalComponent, WrapperFunc);
+	ActorAndComponentMapping->DispatchToRegisteredLocalComponents(InOriginalComponent, WrapperFunc);
 }
 
 void UPCGSubsystem::ForAllRegisteredIntersectingLocalComponents(UPCGComponent* InOriginalComponent, const FBoxCenterAndExtent& InBounds, const TFunctionRef<void(UPCGComponent*)>& InFunc) const
@@ -1052,7 +1053,7 @@ void UPCGSubsystem::ForAllRegisteredIntersectingLocalComponents(UPCGComponent* I
 		return;
 	}
 
-	ActorAndComponentMapping.ForAllIntersectingPartitionActors(Overlap, [InOriginalComponent, &InFunc](APCGPartitionActor* Actor)
+	ActorAndComponentMapping->ForAllIntersectingPartitionActors(Overlap, [InOriginalComponent, &InFunc](APCGPartitionActor* Actor)
 	{
 		if (UPCGComponent* LocalComponent = Actor->GetLocalComponent(InOriginalComponent))
 		{
@@ -1123,7 +1124,7 @@ FPCGTaskId UPCGSubsystem::ForAllOverlappingCells(UPCGComponent* InComponent, con
 			{
 				APCGPartitionActor* Actor = FindOrCreatePCGPartitionActor(Descriptor, CellCoord, bCanCreateActor);
 #if WITH_EDITOR
-				if (!Actor && !GetWorld()->IsGameWorld() && ActorAndComponentMapping.DoesPartitionActorRecordExist(Descriptor, CellCoord))
+				if (!Actor && !GetWorld()->IsGameWorld() && ActorAndComponentMapping->DoesPartitionActorRecordExist(Descriptor, CellCoord))
 				{
 					FPCGTaskId ExecuteTaskId = InUnloadedFunc(Descriptor, CellCoord, IntersectedBounds);
 					
@@ -1198,13 +1199,13 @@ void UPCGSubsystem::CleanupLocalComponentsImmediate(UPCGComponent* InOriginalCom
 		return InvalidPCGTaskId;
 	};
 
-	ActorAndComponentMapping.DispatchToRegisteredLocalComponents(InOriginalComponent, LocalCleanupTask);
+	ActorAndComponentMapping->DispatchToRegisteredLocalComponents(InOriginalComponent, LocalCleanupTask);
 
 	// Remove the local component mappings if requested and the component is partitioned. If 'bRemoveComponents' is false, that indicates we are doing a refresh, so destroying
 	// the component mappings is counterproductive.
 	if (bRemoveComponents && InOriginalComponent->IsPartitioned())
 	{
-		ActorAndComponentMapping.DeleteMappingPCGComponentPartitionActor(InOriginalComponent);
+		ActorAndComponentMapping->DeleteMappingPCGComponentPartitionActor(InOriginalComponent);
 	}
 }
 
@@ -1244,12 +1245,12 @@ APCGPartitionActor* UPCGSubsystem::FindOrCreatePCGPartitionActor(const FGuid& Gu
 UPCGComponent* UPCGSubsystem::GetLocalComponent(const FPCGGridDescriptor& GridDescriptor, const FIntVector& CellCoords, const UPCGComponent* InOriginalComponent) const
 {
 	check(InOriginalComponent);
-	return ActorAndComponentMapping.GetLocalComponent(GridDescriptor, CellCoords, InOriginalComponent);
+	return ActorAndComponentMapping->GetLocalComponent(GridDescriptor, CellCoords, InOriginalComponent);
 }
 
 APCGPartitionActor* UPCGSubsystem::GetRegisteredPCGPartitionActor(const FPCGGridDescriptor& GridDescriptor, const FIntVector& GridCoords) const
 {
-	return ActorAndComponentMapping.GetPartitionActor(GridDescriptor, GridCoords);
+	return ActorAndComponentMapping->GetPartitionActor(GridDescriptor, GridCoords);
 }
 
 APCGPartitionActor* UPCGSubsystem::FindOrCreatePCGPartitionActor(const FPCGGridDescriptor& GridDescriptor, const FIntVector& GridCoords, bool bCanCreateActor, bool bHideFromOutliner) const
@@ -1286,7 +1287,7 @@ APCGPartitionActor* UPCGSubsystem::FindOrCreatePCGPartitionActor(const FPCGGridD
 
 #if WITH_EDITOR
 		// Check if there is already an unloaded actor for this cell. RuntimeGenerated PAs are never unloaded, so we ignore them.
-		if (ActorAndComponentMapping.DoesPartitionActorRecordExist(GridDescriptor, GridCoords))
+		if (ActorAndComponentMapping->DoesPartitionActorRecordExist(GridDescriptor, GridCoords))
 		{
 			return nullptr;
 		}
@@ -1394,7 +1395,7 @@ APCGPartitionActor* UPCGSubsystem::FindOrCreatePCGPartitionActor(const FPCGGridD
 
 TSet<TObjectPtr<APCGPartitionActor>> UPCGSubsystem::GetPCGComponentPartitionActorMappings(UPCGComponent* InComponent) const
 {
-	return ActorAndComponentMapping.GetPCGComponentPartitionActorMappings(InComponent);
+	return ActorAndComponentMapping->GetPCGComponentPartitionActorMappings(InComponent);
 }
 
 bool UPCGSubsystem::IsGraphCacheDebuggingEnabled() const
@@ -1464,7 +1465,7 @@ void UPCGSubsystem::DirtyGraph(UPCGComponent* Component, const FBox& InBounds, E
 		return InvalidPCGTaskId;
 	};
 
-	ActorAndComponentMapping.DispatchToRegisteredLocalComponents(Component, DirtyTask);
+	ActorAndComponentMapping->DispatchToRegisteredLocalComponents(Component, DirtyTask);
 }
 
 void UPCGSubsystem::ClearPCGLink(UPCGComponent* InComponent, const FBox& InBounds, AActor* InNewActor)
@@ -1769,13 +1770,13 @@ void UPCGSubsystem::PropagateEditingModeToLocalComponents(UPCGComponent* InCompo
 {
 	if (ensure(InComponent && InComponent->IsPartitioned()))
 	{
-		FBox Bounds = ActorAndComponentMapping.PartitionedOctree.GetBounds(InComponent);
+		FBox Bounds = ActorAndComponentMapping->PartitionedOctree.GetBounds(InComponent);
 		if (!Bounds.IsValid)
 		{
 			return;
 		}
 
-		ActorAndComponentMapping.ForAllIntersectingPartitionActors(Bounds, [InComponent, EditingMode](APCGPartitionActor* Actor)
+		ActorAndComponentMapping->ForAllIntersectingPartitionActors(Bounds, [InComponent, EditingMode](APCGPartitionActor* Actor)
 		{
 			Actor->ChangeTransientState(InComponent, EditingMode);
 		});
@@ -1865,7 +1866,7 @@ uint32 UPCGSubsystem::GetGraphCacheEntryCount(IPCGElement* InElement) const
 
 void UPCGSubsystem::GenerateAllPCGComponents(bool bForce) const
 {
-	for (UPCGComponent* Component : ActorAndComponentMapping.GetAllRegisteredComponents())
+	for (UPCGComponent* Component : ActorAndComponentMapping->GetAllRegisteredComponents())
 	{
 		Component->Generate(bForce);
 	}
@@ -1873,7 +1874,7 @@ void UPCGSubsystem::GenerateAllPCGComponents(bool bForce) const
 
 void UPCGSubsystem::CleanupAllPCGComponents(bool bPurge) const
 {
-	for (UPCGComponent* Component : ActorAndComponentMapping.GetAllRegisteredComponents())
+	for (UPCGComponent* Component : ActorAndComponentMapping->GetAllRegisteredComponents())
 	{
 		if (bPurge)
 		{
@@ -1916,7 +1917,7 @@ FPCGNodeVisualLogs& UPCGSubsystem::GetNodeVisualLogsMutable()
 
 void UPCGSubsystem::NotifyLandscapeEditModeExited()
 {
-	ActorAndComponentMapping.NotifyLandscapeEditModeExited();
+	ActorAndComponentMapping->NotifyLandscapeEditModeExited();
 }
 
 /* UE_DEPRECATED(5.6) */
@@ -2009,7 +2010,7 @@ void UPCGSubsystem::CreateMissingPartitionActors()
 {
 	if (!PCGHelpers::IsRuntimeOrPIE())
 	{
-		ActorAndComponentMapping.ForAllOriginalComponents([this](UPCGComponent* PCGComponent)
+		ActorAndComponentMapping->ForAllOriginalComponents([this](UPCGComponent* PCGComponent)
 		{
 			if (PCGComponent->IsPartitioned() && !PCGComponent->IsManagedByRuntimeGenSystem())
 			{
@@ -2020,7 +2021,7 @@ void UPCGSubsystem::CreateMissingPartitionActors()
 				{
 					CreatePartitionActorsWithinBounds(PCGComponent, PCGComponent->GetGridBounds(), GridSizes);
 				}
-				ActorAndComponentMapping.UpdateMappingPCGComponentPartitionActor(PCGComponent);
+				ActorAndComponentMapping->UpdateMappingPCGComponentPartitionActor(PCGComponent);
 			}
 		});
 	}
@@ -2042,27 +2043,27 @@ void UPCGSubsystem::CreatePartitionActorsWithinBounds(UPCGComponent* InComponent
 
 void UPCGSubsystem::UpdateMappingPCGComponentPartitionActor(UPCGComponent* InComponent) 
 { 
-	ActorAndComponentMapping.UpdateMappingPCGComponentPartitionActor(InComponent); 
+	ActorAndComponentMapping->UpdateMappingPCGComponentPartitionActor(InComponent); 
 }
 
 void UPCGSubsystem::OnPCGGraphCancelled(UPCGComponent* InComponent)
 {
-	ActorAndComponentMapping.OnPCGGraphCancelled(InComponent);
+	ActorAndComponentMapping->OnPCGGraphCancelled(InComponent);
 }
 
 void UPCGSubsystem::OnPCGGraphStartGenerating(UPCGComponent* InComponent)
 {
-	ActorAndComponentMapping.OnPCGGraphStartsGenerating(InComponent);
+	ActorAndComponentMapping->OnPCGGraphStartsGenerating(InComponent);
 }
 
 void UPCGSubsystem::OnPCGGraphGenerated(UPCGComponent* InComponent)
 {
-	ActorAndComponentMapping.OnPCGGraphGeneratedOrCleaned(InComponent);
+	ActorAndComponentMapping->OnPCGGraphGeneratedOrCleaned(InComponent);
 }
 
 void UPCGSubsystem::OnPCGGraphCleaned(UPCGComponent* InComponent)
 {
-	ActorAndComponentMapping.OnPCGGraphGeneratedOrCleaned(InComponent);
+	ActorAndComponentMapping->OnPCGGraphGeneratedOrCleaned(InComponent);
 }
 
 void UPCGSubsystem::SetDisableClearResults(bool bInDisableClearResults)
@@ -2185,3 +2186,60 @@ void UPCGSubsystem::FlushCache()
 	}
 #endif
 }
+
+bool UPCGSubsystem::RegisterOrUpdatePCGComponent(UPCGComponent* InComponent, bool bDoActorMapping)
+{
+	return ActorAndComponentMapping->RegisterOrUpdatePCGComponent(InComponent, bDoActorMapping);
+}
+
+void UPCGSubsystem::UnregisterPartitionActor(APCGPartitionActor* InActor)
+{
+	ActorAndComponentMapping->UnregisterPartitionActor(InActor);
+}
+
+TSet<UPCGComponent*> UPCGSubsystem::GetAllRegisteredPartitionedComponents() const
+{
+	return ActorAndComponentMapping->GetAllRegisteredPartitionedComponents();
+}
+
+TSet<UPCGComponent*> UPCGSubsystem::GetAllRegisteredComponents() const
+{
+	return ActorAndComponentMapping->GetAllRegisteredComponents();
+}
+
+void UPCGSubsystem::ForAllIntersectingPartitionedComponents(const FBoxCenterAndExtent& InBounds, TFunctionRef<void(UPCGComponent*)> InFunc) const
+{
+	return ActorAndComponentMapping->ForAllIntersectingPartitionedComponents(InBounds, std::move(InFunc));
+}
+
+TArray<UPCGComponent*> UPCGSubsystem::GetAllIntersectingComponents(const FBoxCenterAndExtent& InBounds) const
+{
+	return ActorAndComponentMapping->GetAllIntersectingComponents(InBounds);
+}
+
+bool UPCGSubsystem::RemapPCGComponent(const UPCGComponent* OldComponent, UPCGComponent* NewComponent, bool bDoActorMapping)
+{
+	return ActorAndComponentMapping->RemapPCGComponent(OldComponent, NewComponent, bDoActorMapping);
+}
+
+void UPCGSubsystem::UnregisterPCGComponent(UPCGComponent* InComponent, bool bForce)
+{
+	ActorAndComponentMapping->UnregisterPCGComponent(InComponent, bForce);
+}
+
+void UPCGSubsystem::RegisterPartitionActor(APCGPartitionActor* InActor)
+{
+	ActorAndComponentMapping->RegisterPartitionActor(InActor);
+}
+
+#if WITH_EDITOR
+void UPCGSubsystem::UpdateComponentTracking(UPCGComponent* InComponent, bool bShouldDirtyActors, const TArray<FPCGSelectionKey>* OptionalChangedKeys)
+{
+	ActorAndComponentMapping->UpdateTracking(InComponent, bShouldDirtyActors, OptionalChangedKeys);
+}
+
+void UPCGSubsystem::ResetPartitionActorsMap()
+{
+	ActorAndComponentMapping->ResetPartitionActorsMap();
+}
+#endif // WITH_EDITOR
diff --git a/Engine/Plugins/PCG/Source/PCG/Public/PCGActorAndComponentMapping.h b/Engine/Plugins/PCG/Source/PCG/Public/PCGActorAndComponentMapping.h
deleted file mode 100644
diff --git a/Engine/Plugins/PCG/Source/PCG/Public/PCGSubsystem.h b/Engine/Plugins/PCG/Source/PCG/Public/PCGSubsystem.h
--- a/Engine/Plugins/PCG/Source/PCG/Public/PCGSubsystem.h
+++ b/Engine/Plugins/PCG/Source/PCG/Public/PCGSubsystem.h
@@ -5,20 +5,24 @@
 #include "Subsystems/WorldSubsystem.h"
 
 #include "PCGCommon.h"
-#include "PCGActorAndComponentMapping.h"
+#include "Elements/PCGActorSelector.h"
 #include "Grid/PCGGridDescriptor.h"
 #include "Grid/PCGComponentOctree.h"
-#include "UObject/ObjectKey.h"
 #include "Utils/PCGNodeVisualLogs.h"
 
+#include "UObject/ObjectKey.h"
+
 #include "PCGSubsystem.generated.h"
 
 #define UE_API PCG_API
 
 class APCGPartitionActor;
 class APCGWorldActor;
+class FPCGActorAndComponentMapping;
+class FPCGGenSourceManager;
 class FPCGRuntimeGenScheduler;
 class UPCGComputeGraph;
+class UPCGData;
 class UPCGGraph;
 class UPCGLandscapeCache;
 
@@ -220,22 +224,22 @@ public:
 	UE_API void ClearOutputData(FPCGTaskId InTaskId);
 
 	/** Register a new PCG Component or update it, will be added to the octree if it doesn't exists yet. Returns true if it was added/updated. Thread safe */
-	bool RegisterOrUpdatePCGComponent(UPCGComponent* InComponent, bool bDoActorMapping = true) { return ActorAndComponentMapping.RegisterOrUpdatePCGComponent(InComponent, bDoActorMapping); }
+	UE_API bool RegisterOrUpdatePCGComponent(UPCGComponent* InComponent, bool bDoActorMapping = true);
 
 	/** In case of BP Actors, we need to remap the old component destroyed by the construction script to the new one. Returns true if re-mapping succeeded. */
-	bool RemapPCGComponent(const UPCGComponent* OldComponent, UPCGComponent* NewComponent, bool bDoActorMapping) { return ActorAndComponentMapping.RemapPCGComponent(OldComponent, NewComponent, bDoActorMapping); }
+	UE_API bool RemapPCGComponent(const UPCGComponent* OldComponent, UPCGComponent* NewComponent, bool bDoActorMapping);
 
 	/** Unregister a PCG Component, will be removed from the octree. Can force it, if we have a delayed unregister. Thread safe */
-	void UnregisterPCGComponent(UPCGComponent* InComponent, bool bForce = false) { ActorAndComponentMapping.UnregisterPCGComponent(InComponent, bForce); }
+	UE_API void UnregisterPCGComponent(UPCGComponent* InComponent, bool bForce = false);
 
 	/** Register a new Partition actor, will be added to a map and will query all intersecting volume to bind to them if asked. Thread safe */
-	void RegisterPartitionActor(APCGPartitionActor* InActor) { ActorAndComponentMapping.RegisterPartitionActor(InActor); }
+	UE_API void RegisterPartitionActor(APCGPartitionActor* InActor);
 
 	/** Unregister a Partition actor, will be removed from the map and remove itself to all intersecting volumes. Thread safe */
-	void UnregisterPartitionActor(APCGPartitionActor* InActor) { ActorAndComponentMapping.UnregisterPartitionActor(InActor); }
+	UE_API void UnregisterPartitionActor(APCGPartitionActor* InActor);
 
-	TSet<UPCGComponent*> GetAllRegisteredPartitionedComponents() const { return ActorAndComponentMapping.GetAllRegisteredPartitionedComponents(); }
-	TSet<UPCGComponent*> GetAllRegisteredComponents() const { return ActorAndComponentMapping.GetAllRegisteredComponents(); }
+	UE_API TSet<UPCGComponent*> GetAllRegisteredPartitionedComponents() const;
+	UE_API TSet<UPCGComponent*> GetAllRegisteredComponents() const;
 
 	/** Returns the interface to the cache, required for element per-data caching */
 	UE_API IPCGGraphCache* GetCache();
@@ -250,10 +254,10 @@ public:
 	UE_API void ForAllRegisteredIntersectingLocalComponents(UPCGComponent* InOriginalComponent, const FBoxCenterAndExtent& InBounds, const TFunctionRef<void(UPCGComponent*)>& InFunc) const;
 
 	/** Get all components in specified bounds. */
-	void ForAllIntersectingPartitionedComponents(const FBoxCenterAndExtent& InBounds, TFunctionRef<void(UPCGComponent*)> InFunc) const { return ActorAndComponentMapping.ForAllIntersectingPartitionedComponents(InBounds, std::move(InFunc)); }
+	UE_API void ForAllIntersectingPartitionedComponents(const FBoxCenterAndExtent& InBounds, TFunctionRef<void(UPCGComponent*)> InFunc) const;
 
 	/** Gather all the PCG components within some bounds. */
-	TArray<UPCGComponent*> GetAllIntersectingComponents(const FBoxCenterAndExtent& InBounds) const { return ActorAndComponentMapping.GetAllIntersectingComponents(InBounds); }
+	UE_API TArray<UPCGComponent*> GetAllIntersectingComponents(const FBoxCenterAndExtent& InBounds) const;
 
 	/** Traverses the hierarchy associated with the given component and calls InFunc for each overlapping component. */
 	UE_API void ForAllOverlappingComponentsInHierarchy(UPCGComponent* InComponent, const TFunctionRef<void(UPCGComponent*)>& InFunc) const;
@@ -320,7 +324,7 @@ public:
 	UE_API void NotifyGraphChanged(UPCGGraph* InGraph, EPCGChangeType ChangeType);
 
 	/** Update the tracking on a given component. */
-	void UpdateComponentTracking(UPCGComponent* InComponent, bool bShouldDirtyActors, const TArray<FPCGSelectionKey>* OptionalChangedKeys = nullptr) { ActorAndComponentMapping.UpdateTracking(InComponent, bShouldDirtyActors, OptionalChangedKeys); }
+	UE_API void UpdateComponentTracking(UPCGComponent* InComponent, bool bShouldDirtyActors, const TArray<FPCGSelectionKey>* OptionalChangedKeys = nullptr);
 
 	/** Propagates transient state change from an original component to the relevant partition actors */
 	UE_API void PropagateEditingModeToLocalComponents(UPCGComponent* InOriginalComponent, EPCGEditorDirtyMode EditingMode);
@@ -332,7 +336,7 @@ public:
 	UE_API void ClearPCGLink(UPCGComponent* InComponent, const FBox& InBounds, AActor* InNewActor);
 
 	/** If the partition grid size change, call this to empty the Partition actors map */
-	void ResetPartitionActorsMap() { ActorAndComponentMapping.ResetPartitionActorsMap(); }
+	UE_API void ResetPartitionActorsMap();
 
 	/** Builds the landscape data cache. If force build is true, then it will build the cache even if it is never serialized. */
 	UE_API void BuildLandscapeCache(bool bQuiet = false, bool bForceBuild = true);
@@ -439,7 +443,7 @@ private:
 	TSharedPtr<FPCGGraphExecutor> GraphExecutor;
 	FPCGRuntimeGenScheduler* RuntimeGenScheduler = nullptr;
 	bool bHasTickedOnce = false;
-	FPCGActorAndComponentMapping ActorAndComponentMapping;
+	TUniquePtr<FPCGActorAndComponentMapping> ActorAndComponentMapping;
 
 	/** Functions will be executed at the beginning of the tick and then removed from this array. */
 	TArray<FTickAction> BeginTickActions;
