[PCG] Fix incremental build issue

#rnx

diff --git a/Engine/Plugins/PCG/Source/PCGEditor/Private/PCGEditorGraphSchemaActions.cpp b/Engine/Plugins/PCG/Source/PCGEditor/Private/PCGEditorGraphSchemaActions.cpp
--- a/Engine/Plugins/PCG/Source/PCGEditor/Private/PCGEditorGraphSchemaActions.cpp
+++ b/Engine/Plugins/PCG/Source/PCGEditor/Private/PCGEditorGraphSchemaActions.cpp
@@ -4,11 +4,11 @@
 
 #include "PCGEditorCommon.h"
 #include "PCGEditorGraph.h"
+#include "PCGEditorGraphSchema.h"
 #include "PCGEditorModule.h"
 #include "Nodes/PCGEditorGraphNode.h"
 #include "Nodes/PCGEditorGraphNodeReroute.h"
 
-#include "PCGEditorGraphSchema.h"
 #include "PCGGraph.h"
 #include "PCGNode.h"
 #include "PCGSubgraph.h"
@@ -18,6 +18,8 @@
 #include "Elements/PCGReroute.h"
 #include "Elements/PCGUserParameterGet.h"
 
+#include "UObject/UObjectGlobals.h"
+
 #include "EdGraphNode_Comment.h"
 #include "GraphEditor.h"
 #include "ScopedTransaction.h"
