[PCG] Change default value for Data Attribute to tags

* bDiscardAttributeValue was true by default, is now false
* Was just submitted no deprecation needed

#rnx

diff --git a/Engine/Plugins/PCG/Source/PCG/Public/Elements/Metadata/PCGDataAttributesAndTags.h b/Engine/Plugins/PCG/Source/PCG/Public/Elements/Metadata/PCGDataAttributesAndTags.h
--- a/Engine/Plugins/PCG/Source/PCG/Public/Elements/Metadata/PCGDataAttributesAndTags.h
+++ b/Engine/Plugins/PCG/Source/PCG/Public/Elements/Metadata/PCGDataAttributesAndTags.h
@@ -55,7 +55,7 @@ public:
 
 	// Do not output the value of the attribute in the tag (e.g. MyAttr:1), but only the attribute name (e.g. MyAttr)
 	UPROPERTY(EditAnywhere, BlueprintReadWrite, AdvancedDisplay, Category = "Settings", meta = (PCG_Overridable))
-	bool bDiscardAttributeValue = true;
+	bool bDiscardAttributeValue = false;
 };
 
 UCLASS(MinimalAPI, BlueprintType, ClassGroup = (Procedural))
