[PCG] Fix parsing issue with data labels

#rnx

diff --git a/Engine/Plugins/PCG/Source/PCG/Private/Compute/Elements/PCGCustomHLSLKernel.cpp b/Engine/Plugins/PCG/Source/PCG/Private/Compute/Elements/PCGCustomHLSLKernel.cpp
--- a/Engine/Plugins/PCG/Source/PCG/Private/Compute/Elements/PCGCustomHLSLKernel.cpp
+++ b/Engine/Plugins/PCG/Source/PCG/Private/Compute/Elements/PCGCustomHLSLKernel.cpp
@@ -1350,17 +1350,17 @@ FString UPCGCustomHLSLKernel::ProcessShaderSource(FPCGGPUCompilationContext& InO
 		return OutShaderSource;
 	}
 
-	// We inject attribute IDs directly into the source. This is most efficient and saves us passing them into the kernel. However
-	// the trade off is that a shader with a shared source will generate multiple variants if used in different compute graphs with
-	// different attribute IDs.
-	//
-	// Attribute name replacement relies on precomputed indices into the source strings, therefore the attribute name replacement must take place
+	// Replacement relies on precomputed indices into the source strings, therefore the replacement must take place
 	// before any other modifications. Otherwise, the indices will be incorrect and the source will become gibberish.
-	TArray<TTuple</*AttributeIndex=*/int32, /*ReplaceStartIndex=*/int32, /*ReplaceEndIndex=*/int32>> Replacements;
+	using FReplacement = TTuple</*ReplacementString=*/FString, /*ReplaceStartIndex=*/int32, /*ReplaceEndIndex=*/int32>;
+	TArray<FReplacement> Replacements;
 
 	const UEnum* AttributeTypeEnum = StaticEnum<EPCGKernelAttributeType>();
 	check(AttributeTypeEnum);
 
+	// We inject attribute IDs directly into the source. This is most efficient and saves us passing them into the kernel. However
+	// the trade off is that a shader with a shared source will generate multiple variants if used in different compute graphs with
+	// different attribute IDs.
 	for (const FPCGParsedAttributeFunction& ParsedFunction : InParsedSource.AttributeFunctions)
 	{
 		const int64 AttrType = ParsedFunction.AttributeType;
@@ -1377,31 +1377,7 @@ FString UPCGCustomHLSLKernel::ProcessShaderSource(FPCGGPUCompilationContext& InO
 		const int32 ReplaceStartIndex = OutShaderSource.Find(SourceDefinition, ESearchCase::CaseSensitive, ESearchDir::FromStart, ParsedFunction.MatchBeginning);
 		const int32 ReplaceEndIndex = ReplaceStartIndex + SourceDefinition.Len();
 
-		Replacements.Emplace(AttributeIndex, ReplaceStartIndex, ReplaceEndIndex);
-	}
-
-	for (int I = Replacements.Num() - 1; I >= 0; --I)
-	{
-		const int32 AttributeIndex = Replacements[I].Get<0>();
-		const int32 ReplaceStartIndex = Replacements[I].Get<1>();
-		const int32 ReplaceEndIndex = Replacements[I].Get<2>();
-
-		OutShaderSource = OutShaderSource.Left(ReplaceStartIndex) + FString::FromInt(AttributeIndex) + OutShaderSource.RightChop(ReplaceEndIndex);
-	}
-
-#if PLATFORM_WINDOWS
-	// Remove old-school stuff.
-	OutShaderSource.ReplaceInline(TEXT("\r"), TEXT(""));
-#endif
-
-	// @todo_pcg: Replace using token ranges instead of find/replace, similar to what we do with the parsed attribute functions.
-	// Replace function calls like Out_CopyElementFrom_In(...) with macro PCG_COPY_ALL_ATTRIBUTES_TO_OUTPUT(Out, In, ...).
-	for (const FPCGParsedCopyElementFunction& ParsedFunction : InParsedSource.CopyElementFunctions)
-	{
-		OutShaderSource.ReplaceInline(
-			*FString::Format(TEXT("{2}_{0}_{1}("), { PCGCustomHLSLKernel::CopyElementFunctionKeyword, ParsedFunction.SourcePin, ParsedFunction.TargetPin }),
-			*FString::Format(TEXT("PCG_COPY_ALL_ATTRIBUTES_TO_OUTPUT({1}, {0}, "), { ParsedFunction.SourcePin, ParsedFunction.TargetPin })
-		);
+		Replacements.Emplace(FString::FromInt(AttributeIndex), ReplaceStartIndex, ReplaceEndIndex);
 	}
 
 	// We inject data IDs directly into the source. They will get remapped to data indices using a label resolver data interface.
@@ -1429,11 +1405,37 @@ FString UPCGCustomHLSLKernel::ProcessShaderSource(FPCGGPUCompilationContext& InO
 				const int32 DataLabelCharIndexStart = ModuleMatcher.GetCaptureGroupBeginning(1);
 				const int32 DataLabelCharIndexEnd = ModuleMatcher.GetCaptureGroupEnding(1);
 
-				OutShaderSource = OutShaderSource.Left(DataLabelCharIndexStart) + ReplacementStr + OutShaderSource.RightChop(DataLabelCharIndexEnd);
+				Replacements.Emplace(ReplacementStr, DataLabelCharIndexStart, DataLabelCharIndexEnd);
 			}
 		}
 	}
 
+	// Sort the replacements by replacement index and apply them to the source in reverse order.
+	// Note: Assumes that two replacements do not overlap.
+	Algo::Sort(Replacements, [](const FReplacement& A, const FReplacement& B) { return A.Get<1>() < B.Get<1>(); });
+
+	for (int I = Replacements.Num() - 1; I >= 0; --I)
+	{
+		const FString& ReplacementString = Replacements[I].Get<0>();
+		const int32 ReplaceStartIndex = Replacements[I].Get<1>();
+		const int32 ReplaceEndIndex = Replacements[I].Get<2>();
+
+		OutShaderSource = OutShaderSource.Left(ReplaceStartIndex) + ReplacementString + OutShaderSource.RightChop(ReplaceEndIndex);
+	}
+
+	// Remove old-school stuff.
+	OutShaderSource.ReplaceInline(TEXT("\r"), TEXT(""));
+
+	// @todo_pcg: Replace using token ranges instead of find/replace, similar to what we do with the parsed attribute functions.
+	// Replace function calls like Out_CopyElementFrom_In(...) with macro PCG_COPY_ALL_ATTRIBUTES_TO_OUTPUT(Out, In, ...).
+	for (const FPCGParsedCopyElementFunction& ParsedFunction : InParsedSource.CopyElementFunctions)
+	{
+		OutShaderSource.ReplaceInline(
+			*FString::Format(TEXT("{2}_{0}_{1}("), { PCGCustomHLSLKernel::CopyElementFunctionKeyword, ParsedFunction.SourcePin, ParsedFunction.TargetPin }),
+			*FString::Format(TEXT("PCG_COPY_ALL_ATTRIBUTES_TO_OUTPUT({1}, {0}, "), { ParsedFunction.SourcePin, ParsedFunction.TargetPin })
+		);
+	}
+
 	return OutShaderSource;
 }
 
