[PCG] Fix the issue preventing Action Menu actions from being invoked when highlighted

Also fixed a minor issue with Get Param name selection when drag'n'dropped

#rnx

diff --git a/Engine/Plugins/PCG/Source/PCG/Private/Elements/PCGUserParameterGet.cpp b/Engine/Plugins/PCG/Source/PCG/Private/Elements/PCGUserParameterGet.cpp
--- a/Engine/Plugins/PCG/Source/PCG/Private/Elements/PCGUserParameterGet.cpp
+++ b/Engine/Plugins/PCG/Source/PCG/Private/Elements/PCGUserParameterGet.cpp
@@ -256,17 +256,16 @@ void UPCGUserParameterGetSettings::ApplyPreconfiguredSettings(const FPCGPreConfi
 	UPCGGraph* Graph = Node ? Node->GetGraph() : nullptr;
 	if (const FInstancedPropertyBag* PropertyBag = Graph ? Graph->GetUserParametersStruct() : nullptr)
 	{
-		FName NewPropertyName;
 		if (!PCGGraphParameter::Helpers::GenerateUniqueName(Graph, PropertyName))
 		{
 			PCGLog::LogErrorOnGraph(FText::Format(LOCTEXT("CouldNotCreateUniqueGraphParameterName", "Could not create a unique graph parameter with name '{0}'. Check if the graph is valid and the name has not reached its maximum."), FText::FromName(PropertyName)));
 			return;
 		}
-		check(PropertyBag->IsPropertyNameValid(NewPropertyName));
+		check(PropertyBag->IsPropertyNameValid(PropertyName));
 
-		FPropertyBagPropertyDesc PropertyDesc = PCGPropertyHelpers::CreatePropertyBagDescWithMetadataType(NewPropertyName, NewType);
+		FPropertyBagPropertyDesc PropertyDesc = PCGPropertyHelpers::CreatePropertyBagDescWithMetadataType(PropertyName, NewType);
 		Graph->AddUserParameters({PropertyDesc});
-		const FPropertyBagPropertyDesc* Desc = PropertyBag->FindPropertyDescByName(NewPropertyName);
+		const FPropertyBagPropertyDesc* Desc = PropertyBag->FindPropertyDescByName(PropertyName);
 		check(Desc);
 
 		PropertyGuid = Desc->ID;
diff --git a/Engine/Plugins/PCG/Source/PCGEditor/Private/PCGEditor.cpp b/Engine/Plugins/PCG/Source/PCGEditor/Private/PCGEditor.cpp
--- a/Engine/Plugins/PCG/Source/PCGEditor/Private/PCGEditor.cpp
+++ b/Engine/Plugins/PCG/Source/PCGEditor/Private/PCGEditor.cpp
@@ -30,7 +30,7 @@
 #include "Nodes/PCGEditorGraphNodeOutput.h"
 #include "Nodes/PCGEditorGraphNodeReroute.h"
 
-#include "Widgets/SPCGEditorGraphActionMenu.h"
+#include "Widgets/SPCGEditorGraphActionWidget.h"
 #include "Widgets/SPCGEditorGraphAttributeListView.h"
 #include "Widgets/SPCGEditorGraphDebugObjectTree.h"
 #include "Widgets/SPCGEditorGraphDetailsView.h"
@@ -1621,11 +1621,16 @@ FReply FPCGEditor::OnSpawnNodeByShortcut(FInputChord InChord, const FVector2f& I
 
 FActionMenuContent FPCGEditor::OnCreateActionMenuContent(UEdGraph* InGraph, const FVector2f& Location, const TArray<UEdGraphPin*>& InDraggedPins, bool bAutoExpand, SGraphEditor::FActionMenuClosed OnMenuClosed)
 {
-	SAssignNew(ActionMenu, SPCGEditorGraphActionMenu)
+	SAssignNew(ActionMenu, SGraphEditorActionMenu)
 	.GraphObj(InGraph)
 	.NewNodePosition(Location)
 	.DraggedFromPins(InDraggedPins)
-	.AutoExpandActionMenu(bAutoExpand);
+	.AutoExpandActionMenu(bAutoExpand)
+	.OnCreateWidgetForAction(SGraphActionMenu::FOnCreateWidgetForAction::CreateLambda(
+		[](const FCreateWidgetForActionData* const Data)
+		{
+			return SNew(SPCGGraphActionWidget, Data);
+		}));
 
 	return FActionMenuContent(ActionMenu.ToSharedRef(), ActionMenu->GetFilterTextBox());
 }
diff --git a/Engine/Plugins/PCG/Source/PCGEditor/Private/PCGEditor.h b/Engine/Plugins/PCG/Source/PCGEditor/Private/PCGEditor.h
--- a/Engine/Plugins/PCG/Source/PCGEditor/Private/PCGEditor.h
+++ b/Engine/Plugins/PCG/Source/PCGEditor/Private/PCGEditor.h
@@ -10,11 +10,11 @@
 #include "Toolkits/AssetEditorToolkit.h"
 
 enum class ECheckBoxState : uint8;
-class SPCGEditorGraphActionMenu;
 namespace ETextCommit { enum Type : int; }
 
 class FUICommandList;
 class SGraphEditor;
+class SGraphEditorActionMenu;
 class SPCGEditorGraphAttributeListView;
 class SPCGEditorGraphDebugObjectTree;
 class SPCGEditorGraphDetailsView;
@@ -454,7 +454,7 @@ private:
 	TObjectPtr<UPCGGraph> PCGGraphBeingEdited = nullptr;
 	UPCGEditorGraph* PCGEditorGraph = nullptr;
 
-	TSharedPtr<SPCGEditorGraphActionMenu> ActionMenu;
+	TSharedPtr<SGraphEditorActionMenu> ActionMenu;
 
 	TSoftObjectPtr<UPCGComponent> PCGComponentBeingInspected;
 	// Implementation note: we'll keep the last valid component inspected so we don't un-inspect on spurious selection changes
diff --git a/Engine/Plugins/PCG/Source/PCGEditor/Private/Widgets/SPCGEditorGraphActionMenu.h b/Engine/Plugins/PCG/Source/PCGEditor/Private/Widgets/SPCGEditorGraphActionMenu.h
deleted file mode 100644
diff --git a/Engine/Plugins/PCG/Source/PCGEditor/Private/Widgets/SPCGEditorGraphActionWidget.h b/Engine/Plugins/PCG/Source/PCGEditor/Private/Widgets/SPCGEditorGraphActionWidget.h
--- a/Engine/Plugins/PCG/Source/PCGEditor/Private/Widgets/SPCGEditorGraphActionWidget.h
+++ b/Engine/Plugins/PCG/Source/PCGEditor/Private/Widgets/SPCGEditorGraphActionWidget.h
@@ -17,9 +17,10 @@ class SPCGGraphActionWidget : public SCompoundWidget
 		check(InCreateData->Action.IsValid());
 
 		TSharedRef<SWidget> IconWidget = SNullWidget::NullWidget;
-		const TSharedPtr<FEdGraphSchemaAction>& Action = InCreateData->Action;
+		ActionPtr = InCreateData->Action;
+		MouseButtonDownDelegate = InCreateData->MouseButtonDownDelegate;
 
-		if (const FSlateBrush* const PrimaryIcon = Action->GetPaletteIcon())
+		if (const FSlateBrush* const PrimaryIcon = InCreateData->Action->GetPaletteIcon())
 		{
 			// @todo_pcg: Support secondary icon/colors for more complex type icons (like Maps)
 			const FSlateBrush* const SecondaryIcon = nullptr;
@@ -34,11 +35,12 @@ class SPCGGraphActionWidget : public SCompoundWidget
 			InArgs._NameWidget
 				? InArgs._NameWidget.ToSharedRef()
 				: SNew(STextBlock)
-				.Text(Action->GetMenuDescription());
+				.Text(InCreateData->Action->GetMenuDescription());
 
 		this->ChildSlot
 		[
 			SNew(SHorizontalBox)
+			.ToolTipText(InCreateData->Action->GetMenuDescription())
 			+ SHorizontalBox::Slot()
 			.AutoWidth()
 			[
@@ -52,4 +54,19 @@ class SPCGGraphActionWidget : public SCompoundWidget
 			]
 		];
 	}
+
+	virtual FReply OnMouseButtonDown(const FGeometry& MyGeometry, const FPointerEvent& MouseEvent) override
+	{
+		if (MouseButtonDownDelegate.Execute(ActionPtr))
+		{
+			return FReply::Handled();
+		}
+
+		return FReply::Unhandled();
+	}
+
+	/** The action that we want to display with this widget */
+	TWeakPtr<FEdGraphSchemaAction> ActionPtr;
+	/** Delegate executed when mouse button goes down */
+	FCreateWidgetMouseButtonDown MouseButtonDownDelegate;
 };
diff --git a/Engine/Source/Editor/GraphEditor/Private/SGraphEditorActionMenu.cpp b/Engine/Source/Editor/GraphEditor/Private/SGraphEditorActionMenu.cpp
--- a/Engine/Source/Editor/GraphEditor/Private/SGraphEditorActionMenu.cpp
+++ b/Engine/Source/Editor/GraphEditor/Private/SGraphEditorActionMenu.cpp
@@ -8,7 +8,6 @@
 #include "HAL/PlatformMath.h"
 #include "Layout/Margin.h"
 #include "Misc/Attribute.h"
-#include "SGraphActionMenu.h"
 #include "Styling/AppStyle.h"
 #include "Types/SlateStructs.h"
 #include "Widgets/Layout/SBox.h"
@@ -27,6 +26,7 @@ void SGraphEditorActionMenu::Construct( const FArguments& InArgs )
 	this->NewNodePosition = InArgs._NewNodePosition;
 	this->OnClosedCallback = InArgs._OnClosedCallback;
 	this->AutoExpandActionMenu = InArgs._AutoExpandActionMenu;
+	this->OnCreateWidgetForAction = InArgs._OnCreateWidgetForAction;
 
 	// Build the widget layout
 	SBorder::Construct( SBorder::FArguments()
@@ -44,6 +44,7 @@ void SGraphEditorActionMenu::Construct( const FArguments& InArgs )
 				.AutoExpandActionMenu(AutoExpandActionMenu)
 				.DraggedFromPins(DraggedFromPins)
 				.GraphObj(GraphObj)
+				.OnCreateWidgetForAction(OnCreateWidgetForAction)
 			]
 		]
 	);
diff --git a/Engine/Source/Editor/GraphEditor/Public/SGraphEditorActionMenu.h b/Engine/Source/Editor/GraphEditor/Public/SGraphEditorActionMenu.h
--- a/Engine/Source/Editor/GraphEditor/Public/SGraphEditorActionMenu.h
+++ b/Engine/Source/Editor/GraphEditor/Public/SGraphEditorActionMenu.h
@@ -14,6 +14,7 @@
 #if UE_ENABLE_INCLUDE_ORDER_DEPRECATED_IN_5_6
 #include "Widgets/Input/SEditableTextBox.h"
 #endif
+#include "SGraphActionMenu.h"
 #include "Widgets/Layout/SBorder.h"
 
 class SEditableTextBox;
@@ -39,6 +40,7 @@ public:
 		SLATE_ARGUMENT( TArray<UEdGraphPin*>, DraggedFromPins )
 		SLATE_ARGUMENT( SGraphEditor::FActionMenuClosed, OnClosedCallback )
 		SLATE_ARGUMENT( bool, AutoExpandActionMenu )
+		SLATE_EVENT( SGraphActionMenu::FOnCreateWidgetForAction, OnCreateWidgetForAction )
 	SLATE_END_ARGS()
 
 	void Construct( const FArguments& InArgs );
@@ -56,6 +58,7 @@ protected:
 	bool AutoExpandActionMenu;
 
 	SGraphEditor::FActionMenuClosed OnClosedCallback;
+	SGraphActionMenu::FOnCreateWidgetForAction OnCreateWidgetForAction;
 
 	TSharedPtr<SGraphActionMenu> GraphActionMenu;
 
