[PCG] Add reset default value option to Inline Constants

#rnx

diff --git a/Engine/Plugins/PCG/Source/PCG/Private/Elements/Metadata/PCGMetadataOpElementBase.cpp b/Engine/Plugins/PCG/Source/PCG/Private/Elements/Metadata/PCGMetadataOpElementBase.cpp
--- a/Engine/Plugins/PCG/Source/PCG/Private/Elements/Metadata/PCGMetadataOpElementBase.cpp
+++ b/Engine/Plugins/PCG/Source/PCG/Private/Elements/Metadata/PCGMetadataOpElementBase.cpp
@@ -381,6 +381,18 @@ FString UPCGMetadataSettingsBase::GetPinDefaultValueAsString(const FName PinLabe
 
 	return FString();
 }
+
+void UPCGMetadataSettingsBase::ResetDefaultValue(const FName PinLabel)
+{
+	const FName PropertyName = PCGMetadataBase::Helpers::GetDefaultValuePropertyName(GetInputPinIndex(PinLabel));
+	if (PropertyName != NAME_None && DefaultValues.FindProperty(PropertyName))
+	{
+		Modify();
+		const EPCGMetadataTypes CurrentType = DefaultValues.GetCurrentPropertyType(PropertyName);
+		DefaultValues.RemoveProperty(PropertyName);
+		DefaultValues.CreateNewProperty(PropertyName, CurrentType);
+	}
+}
 #endif // WITH_EDITOR
 
 EPCGMetadataTypes UPCGMetadataSettingsBase::GetPinInitialDefaultValueType(const FName PinLabel) const
diff --git a/Engine/Plugins/PCG/Source/PCG/Private/Helpers/PCGDefaultValueContainer.cpp b/Engine/Plugins/PCG/Source/PCG/Private/Helpers/PCGDefaultValueContainer.cpp
--- a/Engine/Plugins/PCG/Source/PCG/Private/Helpers/PCGDefaultValueContainer.cpp
+++ b/Engine/Plugins/PCG/Source/PCG/Private/Helpers/PCGDefaultValueContainer.cpp
@@ -38,6 +38,11 @@ const FProperty* FPCGDefaultValueContainer::FindProperty(const FName PropertyNam
 	return PropertyDesc ? PropertyDesc->CachedProperty : nullptr;
 }
 
+void FPCGDefaultValueContainer::RemoveProperty(const FName PropertyName)
+{
+	PropertyBag.RemovePropertyByName(PropertyName);
+}
+
 EPCGMetadataTypes FPCGDefaultValueContainer::GetCurrentPropertyType(const FName PropertyName) const
 {
 	const FProperty* Property = FindProperty(PropertyName);
diff --git a/Engine/Plugins/PCG/Source/PCG/Public/Elements/Metadata/PCGMetadataOpElementBase.h b/Engine/Plugins/PCG/Source/PCG/Public/Elements/Metadata/PCGMetadataOpElementBase.h
--- a/Engine/Plugins/PCG/Source/PCG/Public/Elements/Metadata/PCGMetadataOpElementBase.h
+++ b/Engine/Plugins/PCG/Source/PCG/Public/Elements/Metadata/PCGMetadataOpElementBase.h
@@ -112,6 +112,7 @@ public:
 	UE_API virtual void ResetDefaultValues() override;
 	virtual FString GetPinInitialDefaultValueString(FName PinLabel) const override { return PCG::Private::MetadataTraits<double>::ZeroValueString(); }
 	UE_API virtual FString GetPinDefaultValueAsString(FName PinLabel) const override;
+	UE_API virtual void ResetDefaultValue(FName PinLabel) override;
 #endif // WITH_EDITOR
 
 protected:
diff --git a/Engine/Plugins/PCG/Source/PCG/Public/Helpers/PCGDefaultValueContainer.h b/Engine/Plugins/PCG/Source/PCG/Public/Helpers/PCGDefaultValueContainer.h
--- a/Engine/Plugins/PCG/Source/PCG/Public/Helpers/PCGDefaultValueContainer.h
+++ b/Engine/Plugins/PCG/Source/PCG/Public/Helpers/PCGDefaultValueContainer.h
@@ -23,6 +23,9 @@ struct FPCGDefaultValueContainer
 	/** Finds an FProperty in the property bag. Returns nullptr if it does not exist. */
 	const FProperty* FindProperty(FName PropertyName) const;
 
+	/** Removes the property in the property bag, if it exists. */
+	void RemoveProperty(FName PropertyName);
+
 	/** Find the current metadata subtype of the specified FProperty. */
 	EPCGMetadataTypes GetCurrentPropertyType(FName PropertyName) const;
 
diff --git a/Engine/Plugins/PCG/Source/PCG/Public/Metadata/PCGDefaultValueInterface.h b/Engine/Plugins/PCG/Source/PCG/Public/Metadata/PCGDefaultValueInterface.h
--- a/Engine/Plugins/PCG/Source/PCG/Public/Metadata/PCGDefaultValueInterface.h
+++ b/Engine/Plugins/PCG/Source/PCG/Public/Metadata/PCGDefaultValueInterface.h
@@ -50,6 +50,9 @@ public:
 	/** Resets all default values to their 'reset' value and deactivates them. */
 	virtual void ResetDefaultValues() {}
 
+	/** Resets the pin's default value to the initial value. */
+	virtual void ResetDefaultValue(FName PinLabel) {}
+
 	/** Sets the pin's default value string directly. */
 	virtual void SetPinDefaultValue(FName PinLabel, const FString& DefaultValue, bool bCreateIfNeeded = false) {}
 
diff --git a/Engine/Plugins/PCG/Source/PCGEditor/Private/Nodes/PCGEditorGraphNodeBase.cpp b/Engine/Plugins/PCG/Source/PCGEditor/Private/Nodes/PCGEditorGraphNodeBase.cpp
--- a/Engine/Plugins/PCG/Source/PCGEditor/Private/Nodes/PCGEditorGraphNodeBase.cpp
+++ b/Engine/Plugins/PCG/Source/PCGEditor/Private/Nodes/PCGEditorGraphNodeBase.cpp
@@ -282,6 +282,18 @@ void UPCGEditorGraphNodeBase::GetNodeContextMenuActions(UToolMenu* Menu, class U
 									FCanExecuteAction::CreateLambda([this, Pin = Context->Pin] { return IsPinDefaultValueEnabled(Pin->PinName) && !Pin->HasAnyConnections(); })));
 					}
 				}
+
+				Section.AddMenuEntry(
+					TEXT("ResetInlineConstantPin"),
+					LOCTEXT("ResetInlineConstantPinLabel", "Reset Inline Constant Value"),
+					LOCTEXT("ResetInlineConstantPinTooltip", "Reset the inline constant to its default value."),
+					FSlateIcon(),
+					FUIAction(
+						FExecuteAction::CreateLambda([this, Pin = Context->Pin]
+						{
+							OnUserResetPinDefaultValue(Pin->PinName, const_cast<UEdGraphPin*>(Pin));
+						}),
+						FCanExecuteAction::CreateLambda([this, Pin = Context->Pin] { return CanResetPinDefaultValue(Pin->PinName) && !Pin->HasAnyConnections(); })));
 			}
 		}
 	}
@@ -852,6 +864,29 @@ bool UPCGEditorGraphNodeBase::CanConvertToDefaultValueMetadataType(const FName P
 	return Interface && Interface->IsPinDefaultValueMetadataTypeValid(PinLabel, DataType);
 }
 
+void UPCGEditorGraphNodeBase::OnUserResetPinDefaultValue(const FName PinLabel, UEdGraphPin* OutPin) const
+{
+	if (IPCGSettingsDefaultValueProvider* Interface = GetDefaultValueInterface())
+	{
+		FScopedTransaction Transaction(
+			*FPCGEditorCommon::ContextIdentifier,
+			LOCTEXT("PCGEditorResetInlineConstantDefaultValue", "Reset Pin Inline Constant to Default Value"),
+			OutPin ? OutPin->GetOwningNode() : nullptr);
+		Interface->ResetDefaultValue(PinLabel);
+		if (OutPin)
+		{
+			OutPin->Modify();
+			OutPin->DefaultValue = Interface->GetPinDefaultValueAsString(PinLabel);
+		}
+	}
+}
+
+bool UPCGEditorGraphNodeBase::CanResetPinDefaultValue(const FName PinLabel) const
+{
+	const IPCGSettingsDefaultValueProvider* Interface = GetDefaultValueInterface();
+	return Interface && Interface->DefaultValuesAreEnabled() && Interface->IsPinDefaultValueActivated(PinLabel);
+}
+
 EPCGMetadataTypes UPCGEditorGraphNodeBase::GetPinDefaultValueType(const FName PinLabel) const
 {
 	const IPCGSettingsDefaultValueProvider* Interface = GetDefaultValueInterface();
diff --git a/Engine/Plugins/PCG/Source/PCGEditor/Private/Nodes/PCGEditorGraphNodeBase.h b/Engine/Plugins/PCG/Source/PCGEditor/Private/Nodes/PCGEditorGraphNodeBase.h
--- a/Engine/Plugins/PCG/Source/PCGEditor/Private/Nodes/PCGEditorGraphNodeBase.h
+++ b/Engine/Plugins/PCG/Source/PCGEditor/Private/Nodes/PCGEditorGraphNodeBase.h
@@ -107,6 +107,12 @@ public:
 	/** The target type is valid for this pin's default value and this pin is not already this type. */
 	bool CanConvertToDefaultValueMetadataType(FName PinLabel, EPCGMetadataTypes DataType) const;
 
+	/** User is resetting the inline constant value back to the default. */
+	void OnUserResetPinDefaultValue(FName PinLabel, UEdGraphPin* OutPin) const;
+
+	/** The pin can be reset back to the default value. */
+	bool CanResetPinDefaultValue(FName PinLabel) const;
+
 	/** Get the default value type on the pin. */
 	EPCGMetadataTypes GetPinDefaultValueType(FName PinLabel) const;
 
