[PCG] Replace usage of UBlueprint::OnChanged by FCoreUObjectDelegates::OnObjectsReplaced
- This properly updates the object ptrs when class is recompiled, no need to do any reinstancing code because it is already handled by the blueprint compiler
- PCGSpawnActor, PCGCreateTargetActor, PCGAddComponent & PCGExecuteBlueprint

diff --git a/Engine/Plugins/PCG/Source/PCG/Private/Elements/PCGAddComponent.cpp b/Engine/Plugins/PCG/Source/PCG/Private/Elements/PCGAddComponent.cpp
--- a/Engine/Plugins/PCG/Source/PCG/Private/Elements/PCGAddComponent.cpp
+++ b/Engine/Plugins/PCG/Source/PCG/Private/Elements/PCGAddComponent.cpp
@@ -116,21 +116,26 @@ FPCGElementPtr UPCGAddComponentSettings::CreateElement() const
 }
 
 #if WITH_EDITOR
-void UPCGAddComponentSettings::OnBlueprintChanged(UBlueprint* InBlueprint)
+void UPCGAddComponentSettings::OnObjectsReplaced(const TMap<UObject*, UObject*>& InOldToNewInstances)
 {
-	RefreshTemplateComponent();
-	DirtyCache();
-	OnSettingsChangedDelegate.Broadcast(this, EPCGChangeType::Settings);
+	if (!TemplateComponent)
+	{
+		return;
+	}
+
+	if (UObject* NewObject = InOldToNewInstances.FindRef(TemplateComponent))
+	{
+		TemplateComponent = Cast<UActorComponent>(NewObject);
+		DirtyCache();
+		OnSettingsChangedDelegate.Broadcast(this, EPCGChangeType::Settings);
+	}
 }
 
 void UPCGAddComponentSettings::SetupOnBlueprintChangedEvent()
 {
 	if (UBlueprintGeneratedClass* BlueprintClass = Cast<UBlueprintGeneratedClass>(TemplateComponentClass))
 	{
-		if (UBlueprint* Blueprint = Cast<UBlueprint>(BlueprintClass->ClassGeneratedBy))
-		{
-			Blueprint->OnChanged().AddUObject(this, &UPCGAddComponentSettings::OnBlueprintChanged);
-		}
+		FCoreUObjectDelegates::OnObjectsReplaced.AddUObject(this, &UPCGAddComponentSettings::OnObjectsReplaced);
 	}
 }
 
@@ -138,10 +143,7 @@ void UPCGAddComponentSettings::TeardownOnBlueprintChangedEvent()
 {
 	if (UBlueprintGeneratedClass* BlueprintClass = Cast<UBlueprintGeneratedClass>(TemplateComponentClass))
 	{
-		if (UBlueprint* Blueprint = Cast<UBlueprint>(BlueprintClass->ClassGeneratedBy))
-		{
-			Blueprint->OnChanged().RemoveAll(this);
-		}
+		FCoreUObjectDelegates::OnObjectsReplaced.RemoveAll(this);
 	}
 }
 
@@ -459,19 +461,19 @@ void FPCGAddComponentElement::AddComponents(FPCGAddComponentContext* Context,
 			ComponentReferencesView[j] = FSoftObjectPath();
 
 			const FTransform& Transform = TransformsView[j];
-			AActor* TargetActor = TSoftObjectPtr<AActor>(TargetActorsView[j]).LoadSynchronous();
+			AActor* TargetActor = TSoftObjectPtr<AActor>(TargetActorsView[j]).Get();
 			const FSoftObjectPath& ClassOrTemplatePath = ComponentsClassOrTemplateView[j];
 
 			TSubclassOf<UActorComponent> ComponentClass;
 			UActorComponent* ComponentTemplate = nullptr;
 
 			TSoftClassPtr<UActorComponent> ComponentClassPtr(ClassOrTemplatePath);
-			ComponentClass = ComponentClassPtr.LoadSynchronous();
+			ComponentClass = ComponentClassPtr.Get();
 
 			if (!ComponentClass)
 			{
 				TSoftObjectPtr<UActorComponent> ComponentObjectPtr(ClassOrTemplatePath);
-				ComponentTemplate = ComponentObjectPtr.LoadSynchronous();
+				ComponentTemplate = ComponentObjectPtr.Get();
 				ComponentClass = ComponentTemplate ? ComponentTemplate->GetClass() : nullptr;
 			}
 			
diff --git a/Engine/Plugins/PCG/Source/PCG/Private/Elements/PCGCreateTargetActor.cpp b/Engine/Plugins/PCG/Source/PCG/Private/Elements/PCGCreateTargetActor.cpp
--- a/Engine/Plugins/PCG/Source/PCG/Private/Elements/PCGCreateTargetActor.cpp
+++ b/Engine/Plugins/PCG/Source/PCG/Private/Elements/PCGCreateTargetActor.cpp
@@ -161,10 +161,7 @@ void UPCGCreateTargetActor::SetupBlueprintEvent()
 {
 	if (UBlueprintGeneratedClass* BlueprintClass = Cast<UBlueprintGeneratedClass>(TemplateActorClass))
 	{
-		if (UBlueprint* Blueprint = Cast<UBlueprint>(BlueprintClass->ClassGeneratedBy))
-		{
-			Blueprint->OnChanged().AddUObject(this, &UPCGCreateTargetActor::OnBlueprintChanged);
-		}
+		FCoreUObjectDelegates::OnObjectsReplaced.AddUObject(this, &UPCGCreateTargetActor::OnObjectsReplaced);
 	}
 }
 
@@ -172,10 +169,7 @@ void UPCGCreateTargetActor::TeardownBlueprintEvent()
 {
 	if (UBlueprintGeneratedClass* BlueprintClass = Cast<UBlueprintGeneratedClass>(TemplateActorClass))
 	{
-		if (UBlueprint* Blueprint = Cast<UBlueprint>(BlueprintClass->ClassGeneratedBy))
-		{
-			Blueprint->OnChanged().RemoveAll(this);
-		}
+		FCoreUObjectDelegates::OnObjectsReplaced.RemoveAll(this);
 	}
 }
 
@@ -224,10 +218,18 @@ void UPCGCreateTargetActor::PostEditUndo()
 	RefreshTemplateActor();
 }
 
-void UPCGCreateTargetActor::OnBlueprintChanged(UBlueprint* InBlueprint)
+void UPCGCreateTargetActor::OnObjectsReplaced(const TMap<UObject*, UObject*>& InOldToNewInstances)
 {
-	RefreshTemplateActor();
-	OnSettingsChangedDelegate.Broadcast(this, EPCGChangeType::Settings);
+	if (!TemplateActor)
+	{
+		return;
+	}
+	
+	if (UObject* NewObject = InOldToNewInstances.FindRef(TemplateActor))
+	{
+		TemplateActor = Cast<AActor>(NewObject);
+		OnSettingsChangedDelegate.Broadcast(this, EPCGChangeType::Settings);
+	}
 }
 
 void UPCGCreateTargetActor::RefreshTemplateActor()
diff --git a/Engine/Plugins/PCG/Source/PCG/Private/Elements/PCGExecuteBlueprint.cpp b/Engine/Plugins/PCG/Source/PCG/Private/Elements/PCGExecuteBlueprint.cpp
--- a/Engine/Plugins/PCG/Source/PCG/Private/Elements/PCGExecuteBlueprint.cpp
+++ b/Engine/Plugins/PCG/Source/PCG/Private/Elements/PCGExecuteBlueprint.cpp
@@ -374,10 +374,7 @@ void UPCGBlueprintSettings::SetupBlueprintEvent()
 #if WITH_EDITOR
 	if (BlueprintElementType)
 	{
-		if (UBlueprint* Blueprint = Cast<UBlueprint>(BlueprintElementType->ClassGeneratedBy))
-		{
-			Blueprint->OnChanged().AddUObject(this, &UPCGBlueprintSettings::OnBlueprintChanged);
-		}
+		FCoreUObjectDelegates::OnObjectsReplaced.AddUObject(this, &UPCGBlueprintSettings::OnObjectsReplaced);
 	}
 #endif
 }
@@ -387,10 +384,7 @@ void UPCGBlueprintSettings::TeardownBlueprintEvent()
 #if WITH_EDITOR
 	if (BlueprintElementType)
 	{
-		if (UBlueprint* Blueprint = Cast<UBlueprint>(BlueprintElementType->ClassGeneratedBy))
-		{
-			Blueprint->OnChanged().RemoveAll(this);
-		}
+		FCoreUObjectDelegates::OnObjectsReplaced.RemoveAll(this);
 	}
 #endif
 }
@@ -495,23 +489,32 @@ void UPCGBlueprintSettings::PostEditChangeProperty(FPropertyChangedEvent& Proper
 	Super::PostEditChangeProperty(PropertyChangedEvent);
 }
 
-void UPCGBlueprintSettings::OnBlueprintChanged(UBlueprint* InBlueprint)
+void UPCGBlueprintSettings::OnObjectsReplaced(const TMap<UObject*, UObject*>& InOldToNewInstances)
 {
-	// When the blueprint changes, the element gets recreated, so we must rewire it here.
-	DirtyCache();
-	TeardownBlueprintElementEvent();
-	SetupBlueprintElementEvent();
-
-	// We are responsible for calling Initialize on the newly recreated element
-	if (BlueprintElementInstance)
+	if (!BlueprintElementInstance)
 	{
-		BlueprintElementInstance->Initialize();
+		return;
 	}
 
-	// Also, reconstruct overrides
-	InitializeCachedOverridableParams(/*bReset=*/true);
+	if (UObject* NewObject = InOldToNewInstances.FindRef(BlueprintElementInstance))
+	{
+		// When the blueprint changes, the element gets recreated, so we must rewire it here.
+		TeardownBlueprintElementEvent();
+		BlueprintElementInstance = Cast<UPCGBlueprintElement>(NewObject);
+		SetupBlueprintElementEvent();
+		
+		DirtyCache();
 
-	OnSettingsChangedDelegate.Broadcast(this, EPCGChangeType::Settings);
+		if (BlueprintElementInstance)
+		{
+			BlueprintElementInstance->Initialize();
+		}
+
+		// Also, reconstruct overrides
+		InitializeCachedOverridableParams(/*bReset=*/true);
+
+		OnSettingsChangedDelegate.Broadcast(this, EPCGChangeType::Settings);
+	}
 }
 
 void UPCGBlueprintSettings::OnBlueprintElementChanged(UPCGBlueprintElement* InElement)
diff --git a/Engine/Plugins/PCG/Source/PCG/Private/Elements/PCGSpawnActor.cpp b/Engine/Plugins/PCG/Source/PCG/Private/Elements/PCGSpawnActor.cpp
--- a/Engine/Plugins/PCG/Source/PCG/Private/Elements/PCGSpawnActor.cpp
+++ b/Engine/Plugins/PCG/Source/PCG/Private/Elements/PCGSpawnActor.cpp
@@ -240,10 +240,7 @@ void UPCGSpawnActorSettings::SetupBlueprintEvent()
 {
 	if (UBlueprintGeneratedClass* BlueprintClass = Cast<UBlueprintGeneratedClass>(TemplateActorClass))
 	{
-		if (UBlueprint* Blueprint = Cast<UBlueprint>(BlueprintClass->ClassGeneratedBy))
-		{
-			Blueprint->OnChanged().AddUObject(this, &UPCGSpawnActorSettings::OnBlueprintChanged);
-		}
+		FCoreUObjectDelegates::OnObjectsReplaced.AddUObject(this, &UPCGSpawnActorSettings::OnObjectsReplaced);
 	}
 }
 
@@ -251,10 +248,7 @@ void UPCGSpawnActorSettings::TeardownBlueprintEvent()
 {
 	if (UBlueprintGeneratedClass* BlueprintClass = Cast<UBlueprintGeneratedClass>(TemplateActorClass))
 	{
-		if (UBlueprint* Blueprint = Cast<UBlueprint>(BlueprintClass->ClassGeneratedBy))
-		{
-			Blueprint->OnChanged().RemoveAll(this);
-		}
+		FCoreUObjectDelegates::OnObjectsReplaced.RemoveAll(this);
 	}
 }
 #endif
@@ -378,11 +372,19 @@ void UPCGSpawnActorSettings::PostEditUndo()
 	RefreshTemplateActor();
 }
 
-void UPCGSpawnActorSettings::OnBlueprintChanged(UBlueprint* InBlueprint)
+void UPCGSpawnActorSettings::OnObjectsReplaced(const TMap<UObject*, UObject*>& InOldToNewInstances)
 {
-	RefreshTemplateActor();
-	DirtyCache();
-	OnSettingsChangedDelegate.Broadcast(this, EPCGChangeType::Settings);
+	if (!TemplateActor)
+	{
+		return;
+	}
+
+	if (UObject* NewObject = InOldToNewInstances.FindRef(TemplateActor))
+	{
+		TemplateActor = Cast<AActor>(NewObject);
+		DirtyCache();
+		OnSettingsChangedDelegate.Broadcast(this, EPCGChangeType::Settings);
+	}
 }
 
 void UPCGSpawnActorSettings::RefreshTemplateActor()
diff --git a/Engine/Plugins/PCG/Source/PCG/Public/Elements/PCGAddComponent.h b/Engine/Plugins/PCG/Source/PCG/Public/Elements/PCGAddComponent.h
--- a/Engine/Plugins/PCG/Source/PCG/Public/Elements/PCGAddComponent.h
+++ b/Engine/Plugins/PCG/Source/PCG/Public/Elements/PCGAddComponent.h
@@ -64,7 +64,7 @@ public:
 
 private:
 #if WITH_EDITOR
-	void OnBlueprintChanged(UBlueprint* Blueprint);
+	void OnObjectsReplaced(const TMap<UObject*, UObject*>& InOldToNewInstances);
 	void SetupOnBlueprintChangedEvent();
 	void TeardownOnBlueprintChangedEvent();
 	void RefreshTemplateComponent();
diff --git a/Engine/Plugins/PCG/Source/PCG/Public/Elements/PCGCreateTargetActor.h b/Engine/Plugins/PCG/Source/PCG/Public/Elements/PCGCreateTargetActor.h
--- a/Engine/Plugins/PCG/Source/PCG/Public/Elements/PCGCreateTargetActor.h
+++ b/Engine/Plugins/PCG/Source/PCG/Public/Elements/PCGCreateTargetActor.h
@@ -48,7 +48,7 @@ protected:
 	virtual EPCGChangeType GetChangeTypeForProperty(const FName& InPropertyName) const override;
 	//~End UPCGSettings interface
 
-	void OnBlueprintChanged(UBlueprint* Blueprint);
+	void OnObjectsReplaced(const TMap<UObject*, UObject*>& InOldToNewInstances);
 	void SetupBlueprintEvent();
 	void TeardownBlueprintEvent();
 
diff --git a/Engine/Plugins/PCG/Source/PCG/Public/Elements/PCGExecuteBlueprint.h b/Engine/Plugins/PCG/Source/PCG/Public/Elements/PCGExecuteBlueprint.h
--- a/Engine/Plugins/PCG/Source/PCG/Public/Elements/PCGExecuteBlueprint.h
+++ b/Engine/Plugins/PCG/Source/PCG/Public/Elements/PCGExecuteBlueprint.h
@@ -390,8 +390,10 @@ protected:
 
 protected:
 #if WITH_EDITOR
-	UE_API void OnBlueprintChanged(UBlueprint* InBlueprint);
+	UE_DEPRECATED(5.6, "No longer supported")
+	void OnBlueprintChanged(UBlueprint* InBlueprint) {}
 	UE_API void OnBlueprintElementChanged(UPCGBlueprintElement* InElement);
+	UE_API void OnObjectsReplaced(const TMap<UObject*, UObject*>& InOldToNewInstances);
 #endif
 
 	UE_API void RefreshBlueprintElement();
diff --git a/Engine/Plugins/PCG/Source/PCG/Public/Elements/PCGSpawnActor.h b/Engine/Plugins/PCG/Source/PCG/Public/Elements/PCGSpawnActor.h
--- a/Engine/Plugins/PCG/Source/PCG/Public/Elements/PCGSpawnActor.h
+++ b/Engine/Plugins/PCG/Source/PCG/Public/Elements/PCGSpawnActor.h
@@ -119,7 +119,6 @@ public:
 	virtual void BeginDestroy() override;
 
 #if WITH_EDITOR
-	void OnBlueprintChanged(UBlueprint* Blueprint);
 	virtual void PreEditChange(FProperty* PropertyAboutToChange) override;
 	virtual void PostEditChangeProperty(FPropertyChangedEvent& PropertyChangedEvent) override;
 	virtual void PreEditUndo() override;
@@ -163,7 +162,7 @@ protected:
 
 	void SetupBlueprintEvent();
 	void TeardownBlueprintEvent();
-
+	void OnObjectsReplaced(const TMap<UObject*, UObject*>& InOldToNewInstances);
 private:
 	void RefreshTemplateActor();
 #endif
