[PCG] Fix PISM component bounds issues.

Using LocalBounds in descriptor was a flawed approach. PISMC can be set up before final component transform / final actor transform are known, which means its difficult or impossible to know correct local bounds. Switch to WorldBounds instead.

#rnx

diff --git a/Engine/Plugins/PCG/Source/PCG/Private/Components/PCGProceduralISMComponent.cpp b/Engine/Plugins/PCG/Source/PCG/Private/Components/PCGProceduralISMComponent.cpp
--- a/Engine/Plugins/PCG/Source/PCG/Private/Components/PCGProceduralISMComponent.cpp
+++ b/Engine/Plugins/PCG/Source/PCG/Private/Components/PCGProceduralISMComponent.cpp
@@ -176,7 +176,7 @@ void FPCGProceduralISMComponentDescriptor::InitComponent(UPCGProceduralISMCompon
 	InProceduralISMComponent->RuntimeVirtualTextures = RuntimeVirtualTextures;
 	InProceduralISMComponent->SetNumInstances(NumInstances);
 	InProceduralISMComponent->SetNumCustomDataFloats(NumCustomFloats);
-	InProceduralISMComponent->SetLocalBounds(LocalBounds);
+	InProceduralISMComponent->SetBounds(WorldBounds);
 	InProceduralISMComponent->SetMinDrawDistance(InstanceMinDrawDistance);
 	InProceduralISMComponent->SetCullDistances(InstanceStartCullDistance, InstanceEndCullDistance);
 	InProceduralISMComponent->ComponentTags = ComponentTags;
@@ -215,7 +215,7 @@ void FPCGProceduralISMComponentDescriptor::InitFrom(UPCGProceduralISMComponent*
 	RuntimeVirtualTextures = InProceduralISMComponent->RuntimeVirtualTextures;
 	NumInstances = InProceduralISMComponent->GetNumInstances();
 	NumCustomFloats = InProceduralISMComponent->GetNumCustomDataFloats();
-	InProceduralISMComponent->GetLocalBounds(LocalBounds);
+	InProceduralISMComponent->GetBounds(WorldBounds);
 	InstanceMinDrawDistance = InProceduralISMComponent->MinDrawDistance;
 	InProceduralISMComponent->GetCullDistances(InstanceStartCullDistance, InstanceEndCullDistance);
 	ComponentTags = InProceduralISMComponent->ComponentTags;
@@ -270,7 +270,7 @@ bool FPCGProceduralISMComponentDescriptor::operator==(const FPCGProceduralISMCom
 		&& RuntimeVirtualTextures == Other.RuntimeVirtualTextures
 		&& NumInstances == Other.NumInstances
 		&& NumCustomFloats == Other.NumCustomFloats
-		&& LocalBounds == Other.LocalBounds
+		&& WorldBounds == Other.WorldBounds
 		&& InstanceMinDrawDistance == Other.InstanceMinDrawDistance
 		&& InstanceStartCullDistance == Other.InstanceStartCullDistance
 		&& InstanceEndCullDistance == Other.InstanceEndCullDistance
@@ -333,6 +333,16 @@ UPCGProceduralISMComponent::UPCGProceduralISMComponent(const FObjectInitializer&
 #endif
 }
 
+void UPCGProceduralISMComponent::PostLoad()
+{
+	PRAGMA_DISABLE_DEPRECATION_WARNINGS
+	if (LocalBounds.IsValid)
+	{
+		WorldBounds = LocalBounds.TransformBy(GetComponentTransform());
+	}
+	PRAGMA_ENABLE_DEPRECATION_WARNINGS
+}
+
 #if WITH_EDITOR
 void UPCGProceduralISMComponent::OnRegister()
 {
@@ -404,7 +414,7 @@ FMatrix UPCGProceduralISMComponent::GetRenderMatrix() const
 
 FBoxSphereBounds UPCGProceduralISMComponent::CalcBounds(const FTransform& BoundTransform) const
 {
-	return LocalBounds.TransformBy(BoundTransform);
+	return WorldBounds.InverseTransformBy(GetComponentTransform()).TransformBy(BoundTransform);
 }
 
 #if WITH_EDITOR
@@ -510,8 +520,8 @@ void UPCGProceduralISMComponent::BuildSceneDesc(FPrimitiveSceneProxyDesc* InScen
 	OutPrimitiveSceneDesc.PrimitiveSceneData = &GetSceneData();
 	OutPrimitiveSceneDesc.RenderMatrix = GetRenderMatrix();
 	OutPrimitiveSceneDesc.AttachmentRootPosition = GetComponentLocation();
-	OutPrimitiveSceneDesc.LocalBounds = LocalBounds;
-	OutPrimitiveSceneDesc.Bounds = CalcBounds(GetComponentToWorld());
+	OutPrimitiveSceneDesc.Bounds = WorldBounds;
+	OutPrimitiveSceneDesc.LocalBounds = GetLocalBounds();
 	OutPrimitiveSceneDesc.Mobility = InSceneProxyDesc->Mobility;
 }
 
@@ -619,11 +629,21 @@ void UPCGProceduralISMComponent::SetNumCustomDataFloats(int32 InNumCustomDataFlo
 	}
 }
 
-void UPCGProceduralISMComponent::SetLocalBounds(const FBox& InLocalBounds)
+FBox UPCGProceduralISMComponent::GetLocalBounds()
+{
+	return WorldBounds.InverseTransformBy(GetComponentTransform());
+}
+
+void UPCGProceduralISMComponent::GetBounds(FBox& OutWorldBounds)
+{
+	OutWorldBounds = WorldBounds;
+}
+
+void UPCGProceduralISMComponent::SetBounds(const FBox& InWorldBounds)
 {
-	if (InLocalBounds != LocalBounds)
+	if (InWorldBounds != WorldBounds)
 	{
-		LocalBounds = InLocalBounds;
+		WorldBounds = InWorldBounds;
 		MarkRenderStateDirty();
 	}
 }
diff --git a/Engine/Plugins/PCG/Source/PCG/Private/Components/PCGProceduralISMComponent.h b/Engine/Plugins/PCG/Source/PCG/Private/Components/PCGProceduralISMComponent.h
--- a/Engine/Plugins/PCG/Source/PCG/Private/Components/PCGProceduralISMComponent.h
+++ b/Engine/Plugins/PCG/Source/PCG/Private/Components/PCGProceduralISMComponent.h
@@ -34,6 +34,10 @@ class UPCGProceduralISMComponent : public UStaticMeshComponent
 	GENERATED_UCLASS_BODY()
 
 public:
+	//~Begin UObject interface
+	UE_API virtual void PostLoad() override;
+	//~End UObject interface
+
 #if WITH_EDITOR
 	//~Begin UActorComponent Interface
 	virtual void OnRegister() override;
@@ -55,11 +59,15 @@ public:
 	void SetNumCustomDataFloats(int32 InNumCustomDataFloats);
 
 	/** Get local bounds that were previously set for all instances of this component. */
-	void GetLocalBounds(FBox& OutLocalBounds) { OutLocalBounds = LocalBounds; }
+	FBox GetLocalBounds();
+
+	/** Get world bounds for all instances of this component. */
+	UFUNCTION(BlueprintCallable, Category = "Components|PCGProceduralISMC")
+	void GetBounds(FBox& OutWorldBounds);
 
-	/** Set local bounds for all instances of this component (bounds are not computed automatically from instances). */
+	/** Set world bounds for all instances of this component (bounds are not computed automatically from instances and must be set using this function). */
 	UFUNCTION(BlueprintCallable, Category = "Components|PCGProceduralISMC")
-	void SetLocalBounds(const FBox& InLocalBounds);
+	void SetBounds(const FBox& InWorldBounds);
 
 	/** Clear all instances being rendered by this component. */
 	UFUNCTION(BlueprintCallable, Category="Components|PCGProceduralISMC")
@@ -143,7 +151,7 @@ protected:
 	static inline const float DefaultBoundsExtent = 500.0f;
 
 	UPROPERTY(EditAnywhere, Category = Instances)
-	FBox LocalBounds = FBox(-FVector::OneVector * DefaultBoundsExtent, FVector::OneVector * DefaultBoundsExtent);
+	FBox WorldBounds = FBox(-FVector::OneVector * DefaultBoundsExtent, FVector::OneVector * DefaultBoundsExtent);
 
 	/** Distance from camera at which each instance begins to draw. */
 	UPROPERTY(EditAnywhere, BlueprintReadOnly, Category=Culling)
@@ -161,6 +169,14 @@ protected:
 	/** Used for dynamic stats */
 	TStatId StatId;
 #endif
+
+	UE_DEPRECATED(5.6, "Use WorldBounds instead.")
+	UPROPERTY(meta = (DeprecatedProperty, DeprecationMessage = "Use WorldBounds instead."))
+	FBox LocalBounds;
+
+	UE_DEPRECATED(5.6, "Use SetBounds to set world-space bounds instead.")
+	UFUNCTION(BlueprintCallable, Category = "Components|PCGProceduralISMC", meta = (DeprecatedProperty, DeprecationMessage = "Use SetBounds to set world-space bounds instead."))
+	void SetLocalBounds(const FBox& InWorldBounds) {}
 };
 
 /** Struct that holds properties that can be used to initialize Procedural ISM Components. */
@@ -221,7 +237,7 @@ public:
 	int32 NumCustomFloats;
 
 	UPROPERTY()
-	FBox LocalBounds;
+	FBox WorldBounds;
 
 	UPROPERTY()
 	int32 InstanceMinDrawDistance;
@@ -309,6 +325,10 @@ public:
 
 	UPROPERTY()
 	TEnumAsByte<enum EDetailMode> DetailMode;
+
+	UE_DEPRECATED(5.6, "Use WorldBounds instead.")
+	UPROPERTY(meta = (DeprecatedProperty, DeprecationMessage = "Use WorldBounds instead."))
+	FBox LocalBounds;
 };
 
 struct FPCGProceduralISMCBuilderParameters
diff --git a/Engine/Plugins/PCG/Source/PCG/Private/Compute/DataInterfaces/Elements/PCGStaticMeshSpawnerDataInterface.cpp b/Engine/Plugins/PCG/Source/PCG/Private/Compute/DataInterfaces/Elements/PCGStaticMeshSpawnerDataInterface.cpp
--- a/Engine/Plugins/PCG/Source/PCG/Private/Compute/DataInterfaces/Elements/PCGStaticMeshSpawnerDataInterface.cpp
+++ b/Engine/Plugins/PCG/Source/PCG/Private/Compute/DataInterfaces/Elements/PCGStaticMeshSpawnerDataInterface.cpp
@@ -502,8 +502,6 @@ void UPCGStaticMeshSpawnerDataProvider::CreatePrimitiveDescriptors(FPCGContext*
 		return;
 	}
 
-	const FBox LocalBounds = SourceComponent->GetGridBounds().InverseTransformBy(SourceComponent->GetOwner()->GetActorTransform());
-
 	if (Settings->InstanceDataPackerParameters && Settings->InstanceDataPackerParameters->GetAttributeNames(/*OutNames=*/nullptr))
 	{
 		TArray<FName> AttributeNames;
@@ -605,7 +603,7 @@ void UPCGStaticMeshSpawnerDataProvider::CreatePrimitiveDescriptors(FPCGContext*
 			FPCGProceduralISMComponentDescriptor Descriptor;
 			Descriptor = SelectorByAttribute->TemplateDescriptor;
 			Descriptor.NumInstances = StringKeyAndInstanceCount.Value;
-			Descriptor.LocalBounds = LocalBounds;
+			Descriptor.WorldBounds = SourceComponent->GetGridBounds();
 			Descriptor.NumCustomFloats = CustomFloatCount;
 			Descriptor.StaticMesh = FSoftObjectPath(MeshPathString);
 
@@ -670,7 +668,7 @@ void UPCGStaticMeshSpawnerDataProvider::CreatePrimitiveDescriptors(FPCGContext*
 
 				FPCGProceduralISMComponentDescriptor Descriptor;
 				Descriptor = Entry.Descriptor;
-				Descriptor.LocalBounds = LocalBounds;
+				Descriptor.WorldBounds = SourceComponent->GetGridBounds();
 				Descriptor.NumCustomFloats = CustomFloatCount;
 				Descriptor.StaticMesh = Entry.Descriptor.StaticMesh;
 
