[PCG] Fix CIS issues.

#rnx

diff --git a/Engine/Plugins/PCG/Source/PCG/Private/Components/PCGProceduralISMComponent.h b/Engine/Plugins/PCG/Source/PCG/Private/Components/PCGProceduralISMComponent.h
--- a/Engine/Plugins/PCG/Source/PCG/Private/Components/PCGProceduralISMComponent.h
+++ b/Engine/Plugins/PCG/Source/PCG/Private/Components/PCGProceduralISMComponent.h
@@ -187,7 +187,15 @@ struct FPCGProceduralISMComponentDescriptor
 
 public:
 	FPCGProceduralISMComponentDescriptor();
+
+	// Note: We need to explicitly disable warnings on these constructors/operators for clang to be happy with deprecated variables
+	PRAGMA_DISABLE_DEPRECATION_WARNINGS
 	virtual ~FPCGProceduralISMComponentDescriptor() = default;
+	FPCGProceduralISMComponentDescriptor(const FPCGProceduralISMComponentDescriptor&) = default;
+	FPCGProceduralISMComponentDescriptor(FPCGProceduralISMComponentDescriptor&&) = default;
+	FPCGProceduralISMComponentDescriptor& operator=(const FPCGProceduralISMComponentDescriptor&) = default;
+	FPCGProceduralISMComponentDescriptor& operator=(FPCGProceduralISMComponentDescriptor&&) = default;
+	PRAGMA_ENABLE_DEPRECATION_WARNINGS
 
 	FPCGProceduralISMComponentDescriptor& operator=(const FPCGSoftISMComponentDescriptor& Other);
 
