[PCG] Support output textures in GPU kernels

diff --git a/Engine/Plugins/PCG/Shaders/Private/PCGTextureDataInterface.ush b/Engine/Plugins/PCG/Shaders/Private/PCGTextureDataInterface.ush
--- a/Engine/Plugins/PCG/Shaders/Private/PCGTextureDataInterface.ush
+++ b/Engine/Plugins/PCG/Shaders/Private/PCGTextureDataInterface.ush
@@ -1,50 +1,122 @@
 // Copyright Epic Games, Inc. All Rights Reserved.
 
-Texture2D {DataInterfaceName}_Texture_0;
-Texture2D {DataInterfaceName}_Texture_1;
-Texture2D {DataInterfaceName}_Texture_2;
-Texture2D {DataInterfaceName}_Texture_3;
-Texture2D {DataInterfaceName}_Texture_4;
-Texture2D {DataInterfaceName}_Texture_5;
-Texture2D {DataInterfaceName}_Texture_6;
-Texture2D {DataInterfaceName}_Texture_7;
-
-Texture2DArray {DataInterfaceName}_TextureArray_0;
-Texture2DArray {DataInterfaceName}_TextureArray_1;
-Texture2DArray {DataInterfaceName}_TextureArray_2;
-Texture2DArray {DataInterfaceName}_TextureArray_3;
-Texture2DArray {DataInterfaceName}_TextureArray_4;
-Texture2DArray {DataInterfaceName}_TextureArray_5;
-Texture2DArray {DataInterfaceName}_TextureArray_6;
-Texture2DArray {DataInterfaceName}_TextureArray_7;
+// Mirror of ETextureDimension
+#define PCG_TextureDimension_Texture2D          0
+#define PCG_TextureDimension_Texture2DArray     1
+//#define PCG_TextureDimension_Texture3D        2
+//#define PCG_TextureDimension_TextureCube      3
+//#define PCG_TextureDimension_TextureCubeArray 4
+
+// @todo_pcg: The Texture2DArray binding can go away in the future if we support a UPCGTexture2DArrayData.
+#define PCG_TEXTURE_SRV_DECLARATION(Index) \
+	Texture2D {DataInterfaceName}_Texture_##Index; \
+	Texture2DArray {DataInterfaceName}_TextureArray_##Index;
+
+PCG_TEXTURE_SRV_DECLARATION(0);
+PCG_TEXTURE_SRV_DECLARATION(1);
+PCG_TEXTURE_SRV_DECLARATION(2);
+PCG_TEXTURE_SRV_DECLARATION(3);
+PCG_TEXTURE_SRV_DECLARATION(4);
+PCG_TEXTURE_SRV_DECLARATION(5);
+PCG_TEXTURE_SRV_DECLARATION(6);
+PCG_TEXTURE_SRV_DECLARATION(7);
+
+#undef PCG_TEXTURE_SRV_DECLARATION
+
+#define PCG_TEXTURE_UAV_DECLARATION(Index) \
+	RWTexture2D<float4> {DataInterfaceName}_TextureUAV_##Index; \
+
+PCG_TEXTURE_UAV_DECLARATION(0);
+
+#undef PCG_TEXTURE_UAV_DECLARATION
 
 SamplerState {DataInterfaceName}_Sampler_Linear;
 SamplerState {DataInterfaceName}_Sampler_Point;
 
-float4 {DataInterfaceName}_TextureBounds[MAX_NUM_BINDINGS];
+float4 {DataInterfaceName}_TextureBounds[PCG_MAX_NUM_SRV_BINDINGS];
 
-uint {DataInterfaceName}_NumTextureInfos;
+/** Layout of TextureInfos buffer:
+ * - TextureInfo0
+ *   - uint32 - BindingIndex
+ *   - uint32 - SliceIndex (aka TextureIndex if bound to a TextureArray, 0 otherwise.)
+ *   - uint32 - Dimension
+ *   - uint32 - Sampler
+ *   - uint32 - SizeX
+ *   - uint32 - SizeY
+ * - TextureInfo1
+ * - TextureInfo2
+ * - ...
+ */
 StructuredBuffer<uint4> {DataInterfaceName}_TextureInfos;
+uint {DataInterfaceName}_NumTextureInfos;
+
+// #################### INTERNAL HELPERS ##########################
 
-void GetTextureInfoInternal_{DataInterfaceName}(uint InDataIndex, out uint OutBindingIndex, out uint OutSliceIndex, out uint OutIsTextureArray, out uint OutPointSample)
+void GetTextureInfoInternal_{DataInterfaceName}(uint InDataIndex, out uint OutBindingIndex, out uint OutSliceIndex, out uint OutTextureDimension, out uint OutPointSample)
 {
 	if (InDataIndex >= 0u && InDataIndex < {DataInterfaceName}_NumTextureInfos)
 	{
-		const uint4 TextureInfo = {DataInterfaceName}_TextureInfos[InDataIndex];
+		const uint4 TextureInfo = {DataInterfaceName}_TextureInfos[InDataIndex * 2 + 0];
 		OutBindingIndex = TextureInfo.x;
 		OutSliceIndex = TextureInfo.y;
-		OutIsTextureArray = TextureInfo.z;
+		OutTextureDimension = TextureInfo.z;
 		OutPointSample = TextureInfo.w;
 	}
 	else
 	{
 		OutBindingIndex = 0u;
 		OutSliceIndex = 0u;
-		OutIsTextureArray = 0u;
+		OutTextureDimension = 0u;
 		OutPointSample = 0u;
 	}
 }
 
+// #################### GENERAL HELPERS ##########################
+
+uint GetNumData_{DataInterfaceName}()
+{
+	return {DataInterfaceName}_NumTextureInfos;
+}
+
+uint2 GetNumElements_{DataInterfaceName}(uint InDataIndex)
+{
+	if (InDataIndex >= 0u && InDataIndex < {DataInterfaceName}_NumTextureInfos)
+	{
+		return {DataInterfaceName}_TextureInfos[InDataIndex * 2 + 1].xy;
+	}
+	else
+	{
+		return uint2(0u, 0u);
+	}
+}
+
+bool GetThreadData_{DataInterfaceName}(uint InThreadIndex, out uint OutDataIndex, out uint2 OutElementIndex)
+{
+	uint LinearElementIndex = InThreadIndex;
+
+	for (uint DataIndex = 0; DataIndex < {DataInterfaceName}_NumTextureInfos; ++DataIndex)
+	{
+		const uint2 NumElements = GetNumElements_{DataInterfaceName}(DataIndex);
+		const uint ElementCount = NumElements.x * NumElements.y;
+
+		if (LinearElementIndex < ElementCount)
+		{
+			OutDataIndex = DataIndex;
+			OutElementIndex.x = LinearElementIndex % NumElements.x;
+			OutElementIndex.y = LinearElementIndex / NumElements.x;
+			return true;
+		}
+
+		LinearElementIndex -= ElementCount;
+	}
+
+	OutDataIndex = (uint)-1;
+	OutElementIndex = (uint2)-1;
+	return false;
+}
+
+// #################### INPUT FUNCTIONS ##########################
+
 float2 GetTexCoords_{DataInterfaceName}(float2 WorldPos, float2 Min, float2 Max)
 {
 	// Inverse lerp
@@ -62,7 +134,7 @@ float2 GetTexCoords_{DataInterfaceName}(float2 WorldPos, float2 Min, float2 Max)
 #define PCG_CONDITIONAL_TEXTURE_SAMPLE(Index) \
 	if (BindingIndex == Index) \
 	{ \
-		if (IsTextureArray == 0) \
+		if (TextureDimension == PCG_TextureDimension_Texture2D) \
 		{ \
 			if (PointSample == 0) \
 			{ \
@@ -73,7 +145,7 @@ float2 GetTexCoords_{DataInterfaceName}(float2 WorldPos, float2 Min, float2 Max)
 				return {DataInterfaceName}_Texture_##Index.SampleLevel({DataInterfaceName}_Sampler_Point, InTextureUV, 0); \
 			} \
 		} \
-		else \
+		else if (TextureDimension == PCG_TextureDimension_Texture2DArray) \
 		{ \
 			if (PointSample == 0) \
 			{ \
@@ -90,9 +162,9 @@ float4 Sample_{DataInterfaceName}(uint InDataIndex, float2 InTextureUV)
 {
 	uint BindingIndex = 0;
 	uint SliceIndex = 0;
-	uint IsTextureArray = 0;
+	uint TextureDimension = 0;
 	uint PointSample = 0;
-	GetTextureInfoInternal_{DataInterfaceName}(InDataIndex, BindingIndex, SliceIndex, IsTextureArray, PointSample);
+	GetTextureInfoInternal_{DataInterfaceName}(InDataIndex, BindingIndex, SliceIndex, TextureDimension, PointSample);
 
 #if ENABLE_MULTIPLE_TEXTURE_OBJECTS
 	PCG_CONDITIONAL_TEXTURE_SAMPLE(0)
@@ -103,11 +175,8 @@ float4 Sample_{DataInterfaceName}(uint InDataIndex, float2 InTextureUV)
 	PCG_CONDITIONAL_TEXTURE_SAMPLE(5)
 	PCG_CONDITIONAL_TEXTURE_SAMPLE(6)
 	PCG_CONDITIONAL_TEXTURE_SAMPLE(7)
-
-	// TODO: Runtime diagnostic error.
-	return (float4)0;
 #else
-	if (IsTextureArray == 0)
+	if (TextureDimension == PCG_TextureDimension_Texture2D)
 	{
 		if (PointSample == 0)
 		{
@@ -118,7 +187,7 @@ float4 Sample_{DataInterfaceName}(uint InDataIndex, float2 InTextureUV)
 			return {DataInterfaceName}_Texture_0.SampleLevel({DataInterfaceName}_Sampler_Point, InTextureUV, 0);
 		}
 	}
-	else
+	else if (TextureDimension == PCG_TextureDimension_Texture2DArray)
 	{
 		if (PointSample == 0)
 		{
@@ -128,9 +197,11 @@ float4 Sample_{DataInterfaceName}(uint InDataIndex, float2 InTextureUV)
 		{
 			return {DataInterfaceName}_TextureArray_0.SampleLevel({DataInterfaceName}_Sampler_Point, float3(InTextureUV, SliceIndex), 0);
 		}
-
 	}
 #endif
+
+	// TODO: Runtime diagnostic error.
+	return (float4)0;
 }
 
 #undef PCG_CONDITIONAL_TEXTURE_SAMPLE
@@ -139,7 +210,7 @@ float4 Sample_{DataInterfaceName}(uint InDataIndex, float2 InTextureUV)
 	if (BindingIndex == Index) \
 	{ \
 		const float2 TextureUV = GetTexCoords_{DataInterfaceName}(WorldPos, {DataInterfaceName}_TextureBounds[Index].xy, {DataInterfaceName}_TextureBounds[Index].zw); \
-		if (IsTextureArray == 0) \
+		if (TextureDimension == PCG_TextureDimension_Texture2D) \
 		{ \
 			if (PointSample == 0) \
 			{ \
@@ -150,7 +221,7 @@ float4 Sample_{DataInterfaceName}(uint InDataIndex, float2 InTextureUV)
 				return {DataInterfaceName}_Texture_##Index.SampleLevel({DataInterfaceName}_Sampler_Point, TextureUV, 0); \
 			} \
 		} \
-		else \
+		else if (TextureDimension == PCG_TextureDimension_Texture2DArray) \
 		{ \
 			if (PointSample == 0) \
 			{ \
@@ -167,9 +238,9 @@ float4 SampleWorldPos_{DataInterfaceName}(uint InDataIndex, float2 WorldPos)
 {
 	uint BindingIndex = 0;
 	uint SliceIndex = 0;
-	uint IsTextureArray = 0;
+	uint TextureDimension = 0;
 	uint PointSample = 0;
-	GetTextureInfoInternal_{DataInterfaceName}(InDataIndex, BindingIndex, SliceIndex, IsTextureArray, PointSample);
+	GetTextureInfoInternal_{DataInterfaceName}(InDataIndex, BindingIndex, SliceIndex, TextureDimension, PointSample);
 
 #if ENABLE_MULTIPLE_TEXTURE_OBJECTS
 	PCG_CONDITIONAL_TEXTURE_SAMPLE_WORLD_POS(0)
@@ -180,13 +251,10 @@ float4 SampleWorldPos_{DataInterfaceName}(uint InDataIndex, float2 WorldPos)
 	PCG_CONDITIONAL_TEXTURE_SAMPLE_WORLD_POS(5)
 	PCG_CONDITIONAL_TEXTURE_SAMPLE_WORLD_POS(6)
 	PCG_CONDITIONAL_TEXTURE_SAMPLE_WORLD_POS(7)
-
-	// TODO: Runtime diagnostic error.
-	return (float4)0;
 #else
 	const float2 TextureUV = GetTexCoords_{DataInterfaceName}(WorldPos, {DataInterfaceName}_TextureBounds[0].xy, {DataInterfaceName}_TextureBounds[0].zw);
 
-	if (IsTextureArray == 0)
+	if (TextureDimension == PCG_TextureDimension_Texture2D)
 	{
 		if (PointSample == 0)
 		{
@@ -197,7 +265,7 @@ float4 SampleWorldPos_{DataInterfaceName}(uint InDataIndex, float2 WorldPos)
 			return {DataInterfaceName}_Texture_0.SampleLevel({DataInterfaceName}_Sampler_Point, TextureUV, 0);
 		}
 	}
-	else
+	else if (TextureDimension == PCG_TextureDimension_Texture2DArray)
 	{
 		if (PointSample == 0)
 		{
@@ -209,10 +277,96 @@ float4 SampleWorldPos_{DataInterfaceName}(uint InDataIndex, float2 WorldPos)
 		}
 	}
 #endif
+
+	// TODO: Runtime diagnostic error.
+	return (float4)0;
 }
 
 #undef PCG_CONDITIONAL_TEXTURE_SAMPLE_WORLD_POS
 
+#define PCG_CONDITIONAL_TEXTURE_LOAD(Index) \
+	if (BindingIndex == Index) \
+	{ \
+		if (TextureDimension == PCG_TextureDimension_Texture2D) \
+		{ \
+			return {DataInterfaceName}_Texture_##Index[InElementIndex]; \
+		} \
+		else if (TextureDimension == PCG_TextureDimension_Texture2DArray) \
+		{ \
+			return {DataInterfaceName}_TextureArray_##Index[uint3(InElementIndex, SliceIndex)]; \
+		} \
+	}
+
+float4 Load_{DataInterfaceName}(uint InDataIndex, uint2 InElementIndex)
+{
+	uint BindingIndex = 0;
+	uint SliceIndex = 0;
+	uint TextureDimension = 0;
+	uint PointSample = 0;
+	GetTextureInfoInternal_{DataInterfaceName}(InDataIndex, BindingIndex, SliceIndex, TextureDimension, PointSample);
+
+#if ENABLE_MULTIPLE_TEXTURE_OBJECTS
+	PCG_CONDITIONAL_TEXTURE_LOAD(0)
+	PCG_CONDITIONAL_TEXTURE_LOAD(1)
+	PCG_CONDITIONAL_TEXTURE_LOAD(2)
+	PCG_CONDITIONAL_TEXTURE_LOAD(3)
+	PCG_CONDITIONAL_TEXTURE_LOAD(4)
+	PCG_CONDITIONAL_TEXTURE_LOAD(5)
+	PCG_CONDITIONAL_TEXTURE_LOAD(6)
+	PCG_CONDITIONAL_TEXTURE_LOAD(7)
+#else
+	if (TextureDimension == PCG_TextureDimension_Texture2D)
+	{
+		return {DataInterfaceName}_Texture_0[InElementIndex];
+	}
+	else if (TextureDimension == PCG_TextureDimension_Texture2DArray)
+	{
+		return {DataInterfaceName}_TextureArray_0[uint3(InElementIndex, SliceIndex)];
+	}
+#endif
+
+	// TODO: Runtime diagnostic error.
+	return (float4)0;
+}
+
+#undef PCG_CONDITIONAL_TEXTURE_LOAD
+
+#define PCG_CONDITIONAL_TEXTURE_STORE(Index) \
+	if (BindingIndex == Index) \
+	{ \
+		{DataInterfaceName}_TextureUAV_##Index[InElementIndex] = InValue; \
+		return; \
+	}
+
+// #################### OUTPUT FUNCTIONS ##########################
+
+void Store_{DataInterfaceName}(uint InDataIndex, uint2 InElementIndex, float4 InValue)
+{
+	uint BindingIndex = 0;
+	uint SliceIndex = 0;
+	uint TextureDimension = 0;
+	uint PointSample = 0;
+	GetTextureInfoInternal_{DataInterfaceName}(InDataIndex, BindingIndex, SliceIndex, TextureDimension, PointSample);
+
+	if (TextureDimension != PCG_TextureDimension_Texture2D)
+	{
+		// TODO: Runtime diagnostic error.
+		//return;
+	}
+
+#if ENABLE_MULTIPLE_TEXTURE_OBJECTS
+	PCG_CONDITIONAL_TEXTURE_STORE(0)
+
+	// TODO: Runtime diagnostic error.
+#else
+	{DataInterfaceName}_TextureUAV_0[InElementIndex] = InValue;
+#endif
+}
+
+#undef PCG_CONDITIONAL_TEXTURE_STORE
+
+// #################### DEPRECATED FUNCTIONS ##########################
+
 // Deprecated (5.6). Fallback to using first data.
 float4 Sample_{DataInterfaceName}(float2 TexCoords)
 {
@@ -224,3 +378,9 @@ float4 SampleWorldPos_{DataInterfaceName}(float2 WorldPos)
 {
 	return SampleWorldPos_{DataInterfaceName}(/*InDataIndex=*/0, WorldPos);
 }
+
+#undef PCG_TextureDimension_Texture2D
+#undef PCG_TextureDimension_Texture2DArray
+//#undef PCG_TextureDimension_Texture3D
+//#undef PCG_TextureDimension_TextureCube
+//#undef PCG_TextureDimension_TextureCubeArray
diff --git a/Engine/Plugins/PCG/Source/PCG/Internal/Compute/PCGComputeCommon.h b/Engine/Plugins/PCG/Source/PCG/Internal/Compute/PCGComputeCommon.h
--- a/Engine/Plugins/PCG/Source/PCG/Internal/Compute/PCGComputeCommon.h
+++ b/Engine/Plugins/PCG/Source/PCG/Internal/Compute/PCGComputeCommon.h
@@ -68,22 +68,6 @@ namespace PCGComputeConstants
 	constexpr int POINT_SEED_ATTRIBUTE_ID = 7;
 	constexpr int POINT_STEEPNESS_ATTRIBUTE_ID = 8;
 
-	/** PCG data types supported in GPU node inputs. */
-	constexpr EPCGDataType AllowedInputTypes =
-		EPCGDataType::Point
-		| EPCGDataType::Param
-		| EPCGDataType::Landscape
-		| EPCGDataType::BaseTexture
-		| EPCGDataType::VirtualTexture
-		| EPCGDataType::Resource
-		| EPCGDataType::ProxyForGPU;
-
-	/** PCG data types supported in GPU node outputs. */
-	constexpr EPCGDataType AllowedOutputTypes = EPCGDataType::Point | EPCGDataType::Param | EPCGDataType::ProxyForGPU;
-
-	/** PCG data types supported in GPU data collections. */
-	constexpr EPCGDataType AllowedDataCollectionTypes = EPCGDataType::Point | EPCGDataType::Param | EPCGDataType::ProxyForGPU;
-
 	constexpr uint32 KernelExecutedFlag = 1 << 31;
 
 	constexpr uint32 MAX_PRIMITIVE_COMPONENTS_PER_SPAWNER = 256;
@@ -102,6 +86,9 @@ namespace PCGComputeHelpers
 	/** Gets the element count for a given data. E.g. number of points in a PointData, number of metadata entries in a ParamData, etc. */
 	int GetElementCount(const UPCGData* InData);
 
+	/** Gets the 2D element count for a given data. E.g. texture size for UPCGTextureData. */
+	FIntPoint GetElementCount2D(const UPCGData* InData);
+
 	/** True if 'Type' is valid on a GPU input pin. */
 	bool IsTypeAllowedAsInput(EPCGDataType Type);
 
diff --git a/Engine/Plugins/PCG/Source/PCG/Internal/Compute/PCGDataForGPU.h b/Engine/Plugins/PCG/Source/PCG/Internal/Compute/PCGDataForGPU.h
--- a/Engine/Plugins/PCG/Source/PCG/Internal/Compute/PCGDataForGPU.h
+++ b/Engine/Plugins/PCG/Source/PCG/Internal/Compute/PCGDataForGPU.h
@@ -155,7 +155,9 @@ struct FPCGDataDesc
 {
 public:
 	FPCGDataDesc() = default;
+	explicit FPCGDataDesc(EPCGDataType InType);
 	FPCGDataDesc(EPCGDataType InType, int32 InElementCount);
+	FPCGDataDesc(EPCGDataType InType, int32 InElementCount, FIntPoint InElementCount2D);
 	FPCGDataDesc(const FPCGTaggedData& InTaggedData, const UPCGDataBinding* InBinding);
 
 	uint64 ComputePackedSize() const;
@@ -170,10 +172,14 @@ public:
 
 	int32 GetElementCountForAttribute(const FPCGKernelAttributeDesc& AttributeDesc) const;
 
+	bool IsDomain2D() const { return !!(Type & EPCGDataType::BaseTexture); }
+
 public:
 	EPCGDataType Type = EPCGDataType::Point;
 	TArray<FPCGKernelAttributeDesc> AttributeDescs;
+	// @todo_pcg: These counts should be merged into one FIntVector
 	int32 ElementCount = 0;
+	FIntPoint ElementCount2D = FIntPoint::ZeroValue;
 	TArray<int32, TInlineAllocator<4>> TagStringKeys;
 
 private:
diff --git a/Engine/Plugins/PCG/Source/PCG/Private/Compute/DataInterfaces/BuiltInKernels/PCGDataNumDataInterface.cpp b/Engine/Plugins/PCG/Source/PCG/Private/Compute/DataInterfaces/BuiltInKernels/PCGDataNumDataInterface.cpp
--- a/Engine/Plugins/PCG/Source/PCG/Private/Compute/DataInterfaces/BuiltInKernels/PCGDataNumDataInterface.cpp
+++ b/Engine/Plugins/PCG/Source/PCG/Private/Compute/DataInterfaces/BuiltInKernels/PCGDataNumDataInterface.cpp
@@ -56,9 +56,6 @@ void UPCGDataNumDataProvider::Initialize(const UComputeDataInterface* InDataInte
 
 	Super::Initialize(InDataInterface, InBinding, InInputMask, InOutputMask);
 
-	const UPCGDataNumDataInterface* DataInterface = CastChecked<UPCGDataNumDataInterface>(InDataInterface);
-	ProducerKernel = DataInterface->ProducerKernel;
-
 	const UPCGDataNumSettings* Settings = CastChecked<UPCGDataNumSettings>(ProducerKernel->GetSettings());
 	OutputAttributeName = Settings->OutputAttributeName;
 }
@@ -88,11 +85,10 @@ bool UPCGDataNumDataProvider::PrepareForExecute_GameThread(UPCGDataBinding* InBi
 
 void UPCGDataNumDataProvider::Reset()
 {
-	ProducerKernel = nullptr;
+	Super::Reset();
+
 	OutputAttributeName = NAME_None;
 	OutputAttributeId = INDEX_NONE;
-
-	Super::Reset();
 }
 
 bool FPCGDataNumDataProviderProxy::IsValid(FValidationData const& InValidationData) const
diff --git a/Engine/Plugins/PCG/Source/PCG/Private/Compute/DataInterfaces/BuiltInKernels/PCGDataNumDataInterface.h b/Engine/Plugins/PCG/Source/PCG/Private/Compute/DataInterfaces/BuiltInKernels/PCGDataNumDataInterface.h
--- a/Engine/Plugins/PCG/Source/PCG/Private/Compute/DataInterfaces/BuiltInKernels/PCGDataNumDataInterface.h
+++ b/Engine/Plugins/PCG/Source/PCG/Private/Compute/DataInterfaces/BuiltInKernels/PCGDataNumDataInterface.h
@@ -38,9 +38,6 @@ public:
 	virtual void Reset() override;
 	//~ End UComputeDataProvider Interface
 
-	UPROPERTY()
-	TObjectPtr<const UPCGComputeKernel> ProducerKernel;
-
 	FName OutputAttributeName = NAME_None;
 	int32 OutputAttributeId = INDEX_NONE;
 };
diff --git a/Engine/Plugins/PCG/Source/PCG/Private/Compute/DataInterfaces/Elements/PCGCopyPointsDataInterface.cpp b/Engine/Plugins/PCG/Source/PCG/Private/Compute/DataInterfaces/Elements/PCGCopyPointsDataInterface.cpp
--- a/Engine/Plugins/PCG/Source/PCG/Private/Compute/DataInterfaces/Elements/PCGCopyPointsDataInterface.cpp
+++ b/Engine/Plugins/PCG/Source/PCG/Private/Compute/DataInterfaces/Elements/PCGCopyPointsDataInterface.cpp
@@ -142,16 +142,6 @@ UComputeDataProvider* UPCGCopyPointsDataInterface::CreateDataProvider() const
 	return NewObject<UPCGCopyPointsDataProvider>();
 }
 
-void UPCGCopyPointsDataProvider::Initialize(const UComputeDataInterface* InDataInterface, UObject* InBinding, uint64 InInputMask, uint64 InOutputMask)
-{
-	TRACE_CPUPROFILER_EVENT_SCOPE(UPCGCopyPointsDataProvider::Initialize);
-
-	Super::Initialize(InDataInterface, InBinding, InInputMask, InOutputMask);
-
-	const UPCGCopyPointsDataInterface* DataInterface = CastChecked<UPCGCopyPointsDataInterface>(InDataInterface);
-	ProducerKernel = DataInterface->ProducerKernel;
-}
-
 bool UPCGCopyPointsDataProvider::PerformPreExecuteReadbacks_GameThread(UPCGDataBinding* InBinding)
 {
 	TRACE_CPUPROFILER_EVENT_SCOPE(UPCGCopyPointsDataProvider::PerformPreExecuteReadbacks_GameThread);
@@ -271,10 +261,9 @@ FComputeDataProviderRenderProxy* UPCGCopyPointsDataProvider::GetRenderProxy()
 
 void UPCGCopyPointsDataProvider::Reset()
 {
-	ProducerKernel = nullptr;
-	SourceAndTargetDataIndices.Empty();
-
 	Super::Reset();
+
+	SourceAndTargetDataIndices.Empty();
 }
 
 bool FPCGCopyPointsDataProviderProxy::IsValid(FValidationData const& InValidationData) const
diff --git a/Engine/Plugins/PCG/Source/PCG/Private/Compute/DataInterfaces/Elements/PCGCopyPointsDataInterface.h b/Engine/Plugins/PCG/Source/PCG/Private/Compute/DataInterfaces/Elements/PCGCopyPointsDataInterface.h
--- a/Engine/Plugins/PCG/Source/PCG/Private/Compute/DataInterfaces/Elements/PCGCopyPointsDataInterface.h
+++ b/Engine/Plugins/PCG/Source/PCG/Private/Compute/DataInterfaces/Elements/PCGCopyPointsDataInterface.h
@@ -39,15 +39,11 @@ public:
 	//~ End UPCGComputeDataProvider Interface
 
 	//~ Begin UComputeDataProvider Interface
-	virtual void Initialize(const UComputeDataInterface* InDataInterface, UObject* InBinding, uint64 InInputMask, uint64 InOutputMask) override;
 	FComputeDataProviderRenderProxy* GetRenderProxy() override;
 	virtual void Reset() override;
 	//~ End UComputeDataProvider Interface
 
 public:
-	UPROPERTY()
-	TObjectPtr<const UPCGComputeKernel> ProducerKernel;
-
 	UPROPERTY()
 	TArray<FUintVector2> SourceAndTargetDataIndices;
 };
diff --git a/Engine/Plugins/PCG/Source/PCG/Private/Compute/DataInterfaces/Elements/PCGNormalToDensityDataInterface.cpp b/Engine/Plugins/PCG/Source/PCG/Private/Compute/DataInterfaces/Elements/PCGNormalToDensityDataInterface.cpp
--- a/Engine/Plugins/PCG/Source/PCG/Private/Compute/DataInterfaces/Elements/PCGNormalToDensityDataInterface.cpp
+++ b/Engine/Plugins/PCG/Source/PCG/Private/Compute/DataInterfaces/Elements/PCGNormalToDensityDataInterface.cpp
@@ -70,16 +70,6 @@ UComputeDataProvider* UPCGNormalToDensityDataInterface::CreateDataProvider() con
 	return NewObject<UPCGNormalToDensityProvider>();
 }
 
-void UPCGNormalToDensityProvider::Initialize(const UComputeDataInterface* InDataInterface, UObject* InBinding, uint64 InInputMask, uint64 InOutputMask)
-{
-	TRACE_CPUPROFILER_EVENT_SCOPE(UPCGNormalToDensityProvider::Initialize);
-
-	Super::Initialize(InDataInterface, InBinding, InInputMask, InOutputMask);
-
-	const UPCGNormalToDensityDataInterface* DataInterface = CastChecked<UPCGNormalToDensityDataInterface>(InDataInterface);
-	ProducerKernel = DataInterface->ProducerKernel;
-}
-
 FComputeDataProviderRenderProxy* UPCGNormalToDensityProvider::GetRenderProxy()
 {
 	const UPCGNormalToDensitySettings* Settings = CastChecked<UPCGNormalToDensitySettings>(ProducerKernel->GetSettings());
@@ -93,13 +83,6 @@ FComputeDataProviderRenderProxy* UPCGNormalToDensityProvider::GetRenderProxy()
 	return new FPCGNormalToDensityDataProviderProxy(MoveTemp(ProxyData));
 }
 
-void UPCGNormalToDensityProvider::Reset()
-{
-	ProducerKernel = nullptr;
-
-	Super::Reset();
-}
-
 bool FPCGNormalToDensityDataProviderProxy::IsValid(FValidationData const& InValidationData) const
 {
 	return InValidationData.ParameterStructSize == sizeof(FParameters);
diff --git a/Engine/Plugins/PCG/Source/PCG/Private/Compute/DataInterfaces/Elements/PCGNormalToDensityDataInterface.h b/Engine/Plugins/PCG/Source/PCG/Private/Compute/DataInterfaces/Elements/PCGNormalToDensityDataInterface.h
--- a/Engine/Plugins/PCG/Source/PCG/Private/Compute/DataInterfaces/Elements/PCGNormalToDensityDataInterface.h
+++ b/Engine/Plugins/PCG/Source/PCG/Private/Compute/DataInterfaces/Elements/PCGNormalToDensityDataInterface.h
@@ -33,14 +33,8 @@ class UPCGNormalToDensityProvider : public UPCGComputeDataProvider
 
 public:
 	//~ Begin UComputeDataProvider Interface
-	virtual void Initialize(const UComputeDataInterface* InDataInterface, UObject* InBinding, uint64 InInputMask, uint64 InOutputMask) override;
 	FComputeDataProviderRenderProxy* GetRenderProxy() override;
-	virtual void Reset() override;
 	//~ End UComputeDataProvider Interface
-
-public:
-	UPROPERTY()
-	TObjectPtr<const UPCGComputeKernel> ProducerKernel;
 };
 
 class FPCGNormalToDensityDataProviderProxy : public FComputeDataProviderRenderProxy
diff --git a/Engine/Plugins/PCG/Source/PCG/Private/Compute/DataInterfaces/Elements/PCGSkinnedMeshSpawnerDataInterface.cpp b/Engine/Plugins/PCG/Source/PCG/Private/Compute/DataInterfaces/Elements/PCGSkinnedMeshSpawnerDataInterface.cpp
--- a/Engine/Plugins/PCG/Source/PCG/Private/Compute/DataInterfaces/Elements/PCGSkinnedMeshSpawnerDataInterface.cpp
+++ b/Engine/Plugins/PCG/Source/PCG/Private/Compute/DataInterfaces/Elements/PCGSkinnedMeshSpawnerDataInterface.cpp
@@ -200,16 +200,6 @@ UComputeDataProvider* UPCGSkinnedMeshSpawnerDataInterface::CreateDataProvider()
 	return NewObject<UPCGSkinnedMeshSpawnerDataProvider>();
 }
 
-void UPCGSkinnedMeshSpawnerDataProvider::Initialize(const UComputeDataInterface* InDataInterface, UObject* InBinding, uint64 InInputMask, uint64 InOutputMask)
-{
-	TRACE_CPUPROFILER_EVENT_SCOPE(UPCGSkinnedMeshSpawnerDataProvider::Initialize);
-
-	Super::Initialize(InDataInterface, InBinding, InInputMask, InOutputMask);
-
-	const UPCGSkinnedMeshSpawnerDataInterface* DataInterface = CastChecked<UPCGSkinnedMeshSpawnerDataInterface>(InDataInterface);
-	ProducerKernel = CastChecked<UPCGSkinnedMeshSpawnerKernel>(DataInterface->ProducerKernel);
-}
-
 bool UPCGSkinnedMeshSpawnerDataProvider::PerformPreExecuteReadbacks_GameThread(UPCGDataBinding* InBinding)
 {
 	TRACE_CPUPROFILER_EVENT_SCOPE(UPCGSkinnedMeshSpawnerDataProvider::PerformPreExecuteReadbacks_GameThread);
@@ -441,7 +431,8 @@ FComputeDataProviderRenderProxy* UPCGSkinnedMeshSpawnerDataProvider::GetRenderPr
 
 void UPCGSkinnedMeshSpawnerDataProvider::Reset()
 {
-	ProducerKernel = nullptr;
+	Super::Reset();
+
 	AttributeIdOffsetStrides.Empty();
 	PrimitiveStringKeys.Empty();
 	PrimitiveMeshBounds.Empty();
@@ -458,8 +449,6 @@ void UPCGSkinnedMeshSpawnerDataProvider::Reset()
 	NumPrimitivesSetup = 0;
 	bRegisteredPrimitives = false;
 	bStaticMeshesLoaded = false;
-
-	Super::Reset();
 }
 
 void UPCGSkinnedMeshSpawnerDataProvider::CreatePrimitiveDescriptors(FPCGContext* InContext, UPCGDataBinding* InBinding)
diff --git a/Engine/Plugins/PCG/Source/PCG/Private/Compute/DataInterfaces/Elements/PCGSkinnedMeshSpawnerDataInterface.h b/Engine/Plugins/PCG/Source/PCG/Private/Compute/DataInterfaces/Elements/PCGSkinnedMeshSpawnerDataInterface.h
--- a/Engine/Plugins/PCG/Source/PCG/Private/Compute/DataInterfaces/Elements/PCGSkinnedMeshSpawnerDataInterface.h
+++ b/Engine/Plugins/PCG/Source/PCG/Private/Compute/DataInterfaces/Elements/PCGSkinnedMeshSpawnerDataInterface.h
@@ -39,7 +39,6 @@ class UPCGSkinnedMeshSpawnerDataProvider : public UPCGComputeDataProvider
 
 public:
 	//~ Begin UComputeDataProvider Interface
-	virtual void Initialize(const UComputeDataInterface* InDataInterface, UObject* InBinding, uint64 InInputMask, uint64 InOutputMask) override;
 	virtual bool PerformPreExecuteReadbacks_GameThread(UPCGDataBinding* InBinding) override;
 	virtual bool PrepareForExecute_GameThread(UPCGDataBinding* InBinding) override;
 	FComputeDataProviderRenderProxy* GetRenderProxy() override;
@@ -52,9 +51,6 @@ public:
 	virtual bool SetupPrimitives(FPCGContext* InContext, UPCGDataBinding* InBinding);
 
 public:
-	UPROPERTY()
-	TObjectPtr<const UPCGSkinnedMeshSpawnerKernel> ProducerKernel;
-
 	/** Attributes to use for writing per-instance custom floats. */
 	UPROPERTY()
 	TArray<FUintVector4> AttributeIdOffsetStrides;
diff --git a/Engine/Plugins/PCG/Source/PCG/Private/Compute/DataInterfaces/Elements/PCGStaticMeshSpawnerDataInterface.cpp b/Engine/Plugins/PCG/Source/PCG/Private/Compute/DataInterfaces/Elements/PCGStaticMeshSpawnerDataInterface.cpp
--- a/Engine/Plugins/PCG/Source/PCG/Private/Compute/DataInterfaces/Elements/PCGStaticMeshSpawnerDataInterface.cpp
+++ b/Engine/Plugins/PCG/Source/PCG/Private/Compute/DataInterfaces/Elements/PCGStaticMeshSpawnerDataInterface.cpp
@@ -200,16 +200,6 @@ UComputeDataProvider* UPCGStaticMeshSpawnerDataInterface::CreateDataProvider() c
 	return NewObject<UPCGStaticMeshSpawnerDataProvider>();
 }
 
-void UPCGStaticMeshSpawnerDataProvider::Initialize(const UComputeDataInterface* InDataInterface, UObject* InBinding, uint64 InInputMask, uint64 InOutputMask)
-{
-	TRACE_CPUPROFILER_EVENT_SCOPE(UPCGStaticMeshSpawnerDataProvider::Initialize);
-
-	Super::Initialize(InDataInterface, InBinding, InInputMask, InOutputMask);
-
-	const UPCGStaticMeshSpawnerDataInterface* DataInterface = CastChecked<UPCGStaticMeshSpawnerDataInterface>(InDataInterface);
-	ProducerKernel = CastChecked<UPCGStaticMeshSpawnerKernel>(DataInterface->ProducerKernel);
-}
-
 bool UPCGStaticMeshSpawnerDataProvider::PerformPreExecuteReadbacks_GameThread(UPCGDataBinding* InBinding)
 {
 	TRACE_CPUPROFILER_EVENT_SCOPE(UPCGStaticMeshSpawnerDataProvider::PerformPreExecuteReadbacks_GameThread);
@@ -449,7 +439,8 @@ FComputeDataProviderRenderProxy* UPCGStaticMeshSpawnerDataProvider::GetRenderPro
 
 void UPCGStaticMeshSpawnerDataProvider::Reset()
 {
-	ProducerKernel = nullptr;
+	Super::Reset();
+
 	AttributeIdOffsetStrides.Empty();
 	PrimitiveStringKeys.Empty();
 	PrimitiveMeshBounds.Empty();
@@ -467,8 +458,6 @@ void UPCGStaticMeshSpawnerDataProvider::Reset()
 	NumPrimitivesSetup = 0;
 	bRegisteredPrimitives = false;
 	bStaticMeshesLoaded = false;
-
-	Super::Reset();
 }
 
 void UPCGStaticMeshSpawnerDataProvider::CreatePrimitiveDescriptors(FPCGContext* InContext, UPCGDataBinding* InBinding)
diff --git a/Engine/Plugins/PCG/Source/PCG/Private/Compute/DataInterfaces/Elements/PCGStaticMeshSpawnerDataInterface.h b/Engine/Plugins/PCG/Source/PCG/Private/Compute/DataInterfaces/Elements/PCGStaticMeshSpawnerDataInterface.h
--- a/Engine/Plugins/PCG/Source/PCG/Private/Compute/DataInterfaces/Elements/PCGStaticMeshSpawnerDataInterface.h
+++ b/Engine/Plugins/PCG/Source/PCG/Private/Compute/DataInterfaces/Elements/PCGStaticMeshSpawnerDataInterface.h
@@ -40,7 +40,6 @@ class UPCGStaticMeshSpawnerDataProvider : public UPCGComputeDataProvider
 
 public:
 	//~ Begin UComputeDataProvider Interface
-	virtual void Initialize(const UComputeDataInterface* InDataInterface, UObject* InBinding, uint64 InInputMask, uint64 InOutputMask) override;
 	virtual bool PerformPreExecuteReadbacks_GameThread(UPCGDataBinding* InBinding) override;
 	virtual bool PrepareForExecute_GameThread(UPCGDataBinding* InBinding) override;
 	FComputeDataProviderRenderProxy* GetRenderProxy() override;
@@ -53,9 +52,6 @@ public:
 	virtual bool SetupPrimitives(FPCGContext* InContext, UPCGDataBinding* InBinding);
 
 public:
-	UPROPERTY()
-	TObjectPtr<const UPCGStaticMeshSpawnerKernel> ProducerKernel;
-
 	/** Attributes to use for writing per-instance custom floats. */
 	UPROPERTY()
 	TArray<FUintVector4> AttributeIdOffsetStrides;
diff --git a/Engine/Plugins/PCG/Source/PCG/Private/Compute/DataInterfaces/PCGComputeDataInterface.cpp b/Engine/Plugins/PCG/Source/PCG/Private/Compute/DataInterfaces/PCGComputeDataInterface.cpp
new file mode 100644
--- /dev/null
+++ b/Engine/Plugins/PCG/Source/PCG/Private/Compute/DataInterfaces/PCGComputeDataInterface.cpp
@@ -0,0 +1,106 @@
+// Copyright Epic Games, Inc. All Rights Reserved.
+
+#include "Compute/DataInterfaces/PCGComputeDataInterface.h"
+
+#include "PCGGraphExecutionInspection.h"
+#include "PCGGraphExecutionStateInterface.h"
+#include "PCGSettings.h"
+#include "Compute/PCGComputeKernel.h"
+#include "Compute/PCGDataBinding.h"
+
+#include UE_INLINE_GENERATED_CPP_BY_NAME(PCGComputeDataInterface)
+
+void UPCGComputeDataInterface::AddDownstreamInputPin(FName InInputPinLabel, const FName* InOptionalInputPinLabelAlias)
+{
+	DownstreamInputPinLabelAliases.AddUnique(InOptionalInputPinLabelAlias ? *InOptionalInputPinLabelAlias : InInputPinLabel);
+}
+
+void UPCGComputeDataInterface::SetOutputPin(FName InOutputPinLabel, const FName* InOptionalOutputPinLabelAlias)
+{
+	OutputPinLabel = InOutputPinLabel;
+	OutputPinLabelAlias = InOptionalOutputPinLabelAlias ? *InOptionalOutputPinLabelAlias : InOutputPinLabel;
+}
+
+void UPCGComputeDataProvider::Initialize(const UComputeDataInterface* InDataInterface, UObject* InBinding, uint64 InInputMask, uint64 InOutputMask)
+{
+	// Bump counter so any async callbacks from previous usages of this provider will be invalidated.
+	++GenerationCounter;
+
+	Super::Initialize(InDataInterface, InBinding, InInputMask, InOutputMask);
+
+	const UPCGComputeDataInterface* DataInterface = CastChecked<UPCGComputeDataInterface>(InDataInterface);
+	ProducerKernel = DataInterface->ProducerKernel;
+	ProducerSettings = DataInterface->ProducerSettings;
+	GraphBindingIndex = DataInterface->GraphBindingIndex;
+
+	// The original label is needed to store inspection data.
+	OutputPinLabel = DataInterface->OutputPinLabel;
+
+	// Use the aliased label for normal data output as this is the output from the compute graph.
+	OutputPinLabelAlias = DataInterface->OutputPinLabelAlias;
+
+	DownstreamInputPinLabelAliases = DataInterface->DownstreamInputPinLabelAliases;
+}
+
+void UPCGComputeDataProvider::Reset()
+{
+	Super::Reset();
+
+	// Bump counter so any async callbacks from usages of this provider will be invalidated.
+	++GenerationCounter;
+
+	ProducerKernel = nullptr;
+	ProducerSettings = nullptr;
+	GraphBindingIndex = INDEX_NONE;
+	OutputPinLabel = NAME_None;
+	OutputPinLabelAlias = NAME_None;
+	DownstreamInputPinLabelAliases.Empty();
+}
+
+void UPCGExportableDataProvider::Initialize(const UComputeDataInterface* InDataInterface, UObject* InBinding, uint64 InInputMask, uint64 InOutputMask)
+{
+	Super::Initialize(InDataInterface, InBinding, InInputMask, InOutputMask);
+	Binding = CastChecked<UPCGDataBinding>(InBinding);
+
+	const UPCGExportableDataInterface* DataInterface = CastChecked<UPCGExportableDataInterface>(InDataInterface);
+	ExportMode = DataInterface->GetRequiresExport() ? EPCGExportMode::ComputeGraphOutput : EPCGExportMode::NoExport;
+
+#if WITH_EDITOR
+	check(ProducerSettings);
+
+	const IPCGGraphExecutionSource* ExecutionSource = Binding->GetExecutionSource();
+	if (ProducerSettings->bIsInspecting && ExecutionSource && ExecutionSource->GetExecutionState().GetInspection().IsInspecting())
+	{
+		ExportMode |= EPCGExportMode::Inspection;
+	}
+
+	if (ProducerSettings->bDebug)
+	{
+		ExportMode |= EPCGExportMode::DebugVisualization;
+	}
+#endif
+}
+
+void UPCGExportableDataProvider::Reset()
+{
+	Super::Reset();
+
+	ExportMode = EPCGExportMode::NoExport;
+	OnDataExported = {};
+	Binding.Reset();
+	PinDesc = nullptr;
+}
+
+bool UPCGExportableDataProvider::PrepareForExecute_GameThread(UPCGDataBinding* InBinding)
+{
+	check(InBinding);
+
+	if (!Super::PrepareForExecute_GameThread(InBinding))
+	{
+		return false;
+	}
+
+	PinDesc = InBinding->GetCachedKernelPinDataDesc(GraphBindingIndex);
+
+	return true;
+}
diff --git a/Engine/Plugins/PCG/Source/PCG/Private/Compute/DataInterfaces/PCGComputeDataInterface.h b/Engine/Plugins/PCG/Source/PCG/Private/Compute/DataInterfaces/PCGComputeDataInterface.h
--- a/Engine/Plugins/PCG/Source/PCG/Private/Compute/DataInterfaces/PCGComputeDataInterface.h
+++ b/Engine/Plugins/PCG/Source/PCG/Private/Compute/DataInterfaces/PCGComputeDataInterface.h
@@ -2,6 +2,8 @@
 
 #pragma once
 
+#include "Compute/PCGComputeCommon.h"
+
 #include "ComputeFramework/ComputeDataInterface.h"
 #include "ComputeFramework/ComputeDataProvider.h"
 
@@ -17,17 +19,13 @@ class UPCGComputeDataInterface : public UComputeDataInterface
 
 public:
 	/** Register a downstream pin (and optionally compute graph element pin alias). */
-	virtual void AddDownstreamInputPin(FName InInputPinLabel, const FName* InOptionalInputPinLabelAlias = nullptr)
-	{
-		DownstreamInputPinLabelAliases.AddUnique(InOptionalInputPinLabelAlias ? *InOptionalInputPinLabelAlias : InInputPinLabel);
-	}
+	virtual void AddDownstreamInputPin(FName InInputPinLabel, const FName* InOptionalInputPinLabelAlias = nullptr);
 
 	/** Set the output pin label and label alias that this data interface is associated with. */
-	virtual void SetOutputPin(FName InOutputPinLabel, const FName* InOptionalOutputPinLabelAlias = nullptr)
-	{
-		OutputPinLabel = InOutputPinLabel;
-		OutputPinLabelAlias = InOptionalOutputPinLabelAlias ? *InOptionalOutputPinLabelAlias : InOutputPinLabel;
-	}
+	virtual void SetOutputPin(FName InOutputPinLabel, const FName* InOptionalOutputPinLabelAlias = nullptr);
+
+	void SetProducerSettings(TObjectPtr<const UPCGSettings> InProducerSettings) { ProducerSettings = InProducerSettings; }
+	TObjectPtr<const UPCGSettings> GetProducerSettings() const { return ProducerSettings; }
 
 public:
 	UPROPERTY()
@@ -36,6 +34,9 @@ public:
 	UPROPERTY()
 	TObjectPtr<const UPCGComputeKernel> ProducerKernel;
 
+	UPROPERTY()
+	TObjectPtr<const UPCGSettings> ProducerSettings;
+
 	/** Label of output pin that this data interface is associated with. */
 	UPROPERTY()
 	FName OutputPinLabel;
@@ -54,6 +55,9 @@ class UPCGComputeDataProvider : public UComputeDataProvider
 	GENERATED_BODY()
 
 public:
+	virtual void Initialize(const UComputeDataInterface* InDataInterface, UObject* InBinding, uint64 InInputMask, uint64 InOutputMask) override;
+	virtual void Reset() override;
+
 	/** Called before PrepareForExecute to allow data providers to do any readbacks and finalize data descriptions. */
 	virtual bool PerformPreExecuteReadbacks_GameThread(UPCGDataBinding* InBinding) { return true; }
 
@@ -65,4 +69,71 @@ public:
 
 	/** Release any handles to transient resources like GPU buffers. */
 	virtual void ReleaseTransientResources(const TCHAR* InReason = nullptr) {}
+
+	TObjectPtr<const UPCGSettings> GetProducerSettings() const { return ProducerSettings; }
+
+public:
+	int32 GraphBindingIndex = INDEX_NONE;
+	FName OutputPinLabel;
+	FName OutputPinLabelAlias;
+	TArray<FName> DownstreamInputPinLabelAliases;
+
+	/** Bumped each time the data provider is initialized or reset, so that async callbacks can detect if they originated from a previous usage of the data provider and no-op. */
+	std::atomic<uint64> GenerationCounter = 0;
+
+protected:
+	UPROPERTY()
+	TObjectPtr<const UPCGComputeKernel> ProducerKernel;
+
+	/** Settings of node that produces this data, normally the upstream node. */
+	UPROPERTY()
+	TObjectPtr<const UPCGSettings> ProducerSettings;
+};
+
+UCLASS(Abstract, ClassGroup = (Procedural))
+class UPCGExportableDataInterface : public UPCGComputeDataInterface
+{
+	GENERATED_BODY()
+
+public:
+	void SetRequiresExport(bool bInRequiresExport) { bRequiresExport = bInRequiresExport; }
+	bool GetRequiresExport() const { return bRequiresExport; }
+
+protected:
+	/** Whether this data is passed to downstream tasks outside of this compute graph. */
+	UPROPERTY()
+	bool bRequiresExport = false;
+};
+
+UCLASS(Abstract, ClassGroup (Procedural))
+class UPCGExportableDataProvider : public UPCGComputeDataProvider
+{
+	GENERATED_BODY()
+
+public:
+	//~ Begin UComputeDataInterface Interface
+	virtual void Initialize(const UComputeDataInterface* InDataInterface, UObject* InBinding, uint64 InInputMask, uint64 InOutputMask) override;
+	virtual void Reset() override;
+	//~ End UComputeDataInterface Interface
+
+	//~ Begin UPCGComputeDataProvider Interface
+	virtual bool PrepareForExecute_GameThread(UPCGDataBinding* InBinding) override;
+	//~ End UPCGComputeDataProvider Interface
+
+	/** Returns true if the generated buffer should be exported to a persistent buffer to be passed to downstream nodes. */
+	bool IsExportRequired() const { return ExportMode != EPCGExportMode::NoExport; }
+
+	EPCGExportMode GetExportMode() const { return ExportMode; }
+	UPCGDataBinding* GetDataBinding() const { return Binding.Get(); }
+	const FPCGDataCollectionDesc* GetPinDescription() const { return PinDesc; }
+
+	DECLARE_EVENT(UPCGExportableDataProvider, FOnDataExported);
+	FOnDataExported& OnDataExported_GameThread() { return OnDataExported; }
+
+protected:
+	EPCGExportMode ExportMode = EPCGExportMode::NoExport;
+	FOnDataExported OnDataExported;
+
+	TWeakObjectPtr<UPCGDataBinding> Binding;
+	const FPCGDataCollectionDesc* PinDesc = nullptr;
 };
diff --git a/Engine/Plugins/PCG/Source/PCG/Private/Compute/DataInterfaces/PCGCustomKernelDataInterface.h b/Engine/Plugins/PCG/Source/PCG/Private/Compute/DataInterfaces/PCGCustomKernelDataInterface.h
--- a/Engine/Plugins/PCG/Source/PCG/Private/Compute/DataInterfaces/PCGCustomKernelDataInterface.h
+++ b/Engine/Plugins/PCG/Source/PCG/Private/Compute/DataInterfaces/PCGCustomKernelDataInterface.h
@@ -15,7 +15,7 @@ class UPCGComputeKernel;
 
 /** Interface for any meta data provided to the compute kernel, such as num threads. */
 UCLASS(MinimalAPI, ClassGroup = (Procedural))
-class UPCGCustomKernelDataInterface : public UComputeDataInterface
+class UPCGCustomKernelDataInterface : public UPCGComputeDataInterface
 {
 	GENERATED_BODY()
 
diff --git a/Engine/Plugins/PCG/Source/PCG/Private/Compute/DataInterfaces/PCGDataCollectionDataInterface.cpp b/Engine/Plugins/PCG/Source/PCG/Private/Compute/DataInterfaces/PCGDataCollectionDataInterface.cpp
--- a/Engine/Plugins/PCG/Source/PCG/Private/Compute/DataInterfaces/PCGDataCollectionDataInterface.cpp
+++ b/Engine/Plugins/PCG/Source/PCG/Private/Compute/DataInterfaces/PCGDataCollectionDataInterface.cpp
@@ -457,46 +457,15 @@ void UPCGDataCollectionDataProvider::Initialize(const UComputeDataInterface* InD
 {
 	TRACE_CPUPROFILER_EVENT_SCOPE(UPCGDataCollectionDataProvider::Initialize);
 
-	// Bump counter so any async callbacks from previous usages of this provider will be invalidated.
-	++GenerationCounter;
-
 	Super::Initialize(InDataInterface, InBinding, InInputMask, InOutputMask);
 
 	const UPCGDataCollectionDataInterface* DataInterface = CastChecked<UPCGDataCollectionDataInterface>(InDataInterface);
 
-	Binding = CastChecked<UPCGDataBinding>(InBinding);
-
-	check(DataInterface->ProducerSettings);
-
-	ProducerSettings = DataInterface->ProducerSettings;
-	GraphBindingIndex = DataInterface->GraphBindingIndex;
-
-	ExportMode = DataInterface->bRequiresExport ? EPCGExportMode::ComputeGraphOutput : EPCGExportMode::NoExport;
-
 	bRequiresZeroInitialization = DataInterface->bRequiresZeroInitialization;
 
 	ElementCountMultiplier = DataInterface->ElementCountMultiplier;
 
 	bExposeElementCounters = DataInterface->bExposeElementCounters;
-
-#if WITH_EDITOR
-	const IPCGGraphExecutionSource* ExecutionSource = Binding->GetExecutionSource();
-	if (ProducerSettings->bIsInspecting && ExecutionSource && ExecutionSource->GetExecutionState().GetInspection().IsInspecting())
-	{
-		ExportMode |= EPCGExportMode::Inspection;
-	}
-
-	if (ProducerSettings->bDebug)
-	{
-		ExportMode |= EPCGExportMode::DebugVisualization;
-	}
-#endif
-
-	// Use the aliased label for normal data output as this is the output from the compute graph.
-	OutputPinLabelAlias = DataInterface->OutputPinLabelAlias;
-
-	// The original label is needed to store inspection data.
-	OutputPinLabel = DataInterface->OutputPinLabel;
 }
 
 bool UPCGDataCollectionDataProvider::PrepareForExecute_GameThread(UPCGDataBinding* InBinding)
@@ -509,9 +478,9 @@ bool UPCGDataCollectionDataProvider::PrepareForExecute_GameThread(UPCGDataBindin
 		return false;
 	}
 
-	PinDesc = Binding->GetCachedKernelPinDataDesc(GraphBindingIndex);
+	check(PinDesc);
 
-	if (ensure(PinDesc) && PCGComputeHelpers::IsBufferSizeTooLarge(PinDesc->ComputePackedSizeBytes()))
+	if (PCGComputeHelpers::IsBufferSizeTooLarge(PinDesc->ComputePackedSizeBytes()))
 	{
 		PinDesc = nullptr;
 	}
@@ -576,25 +545,14 @@ FComputeDataProviderRenderProxy* UPCGDataCollectionDataProvider::GetRenderProxy(
 
 void UPCGDataCollectionDataProvider::Reset()
 {
-	// Bump counter so any async callbacks from usages of this provider will be invalidated.
-	++GenerationCounter;
+	Super::Reset();
 
-	OnDataExported = {};
 	OnCountersReadBack = {};
-	ProducerSettings = nullptr;
-	Binding.Reset();
-	GraphBindingIndex = INDEX_NONE;
-	ExportMode = EPCGExportMode::NoExport;
 	ElementCountMultiplier = 0;
 	bExposeElementCounters = false;
-	OutputPinLabel = NAME_None;
-	OutputPinLabelAlias = NAME_None;
-	PinDesc = nullptr;
 	bReadbackComplete = false;
 	ReadbackCounters.Empty();
 	bRequiresZeroInitialization = false;
-
-	Super::Reset();
 }
 
 FPCGDataCollectionDataProviderProxy::FPCGDataCollectionDataProviderProxy(
@@ -739,7 +697,7 @@ void FPCGDataCollectionDataProviderProxy::AllocateResources(FRDGBuilder& GraphBu
 				return;
 			}
 
-			UPCGDataBinding* Binding = DataProvider ? DataProvider->Binding.Get() : nullptr;
+			UPCGDataBinding* Binding = DataProvider ? DataProvider->GetDataBinding() : nullptr;
 			UPCGSubsystem* Subsystem = UPCGSubsystem::GetSubsystemForCurrentWorld();
 
 			if (Binding && Subsystem)
@@ -752,7 +710,7 @@ void FPCGDataCollectionDataProviderProxy::AllocateResources(FRDGBuilder& GraphBu
 					Proxy->Initialize(DataCollectionOnGPU, DataIndex);
 
 					// TODO - binding is doing a lot of work. Could store a context handle in the data provider instead?
-					Binding->ReceiveDataFromGPU_GameThread(Proxy, DataProvider->ProducerSettings, ExportMode, OutputPinLabel, OutputPinLabelAlias);
+					Binding->ReceiveDataFromGPU_GameThread(Proxy, DataProvider->GetProducerSettings(), ExportMode, OutputPinLabel, OutputPinLabelAlias);
 				}
 
 				DataProvider->OnDataExported_GameThread().Broadcast();
diff --git a/Engine/Plugins/PCG/Source/PCG/Private/Compute/DataInterfaces/PCGDataCollectionDataInterface.h b/Engine/Plugins/PCG/Source/PCG/Private/Compute/DataInterfaces/PCGDataCollectionDataInterface.h
--- a/Engine/Plugins/PCG/Source/PCG/Private/Compute/DataInterfaces/PCGDataCollectionDataInterface.h
+++ b/Engine/Plugins/PCG/Source/PCG/Private/Compute/DataInterfaces/PCGDataCollectionDataInterface.h
@@ -16,7 +16,7 @@ class UPCGDataBinding;
 
 /** Compute Framework Data Interface for reading PCG data. */
 UCLASS(ClassGroup = (Procedural))
-class UPCGDataCollectionDataInterface : public UPCGComputeDataInterface
+class UPCGDataCollectionDataInterface : public UPCGExportableDataInterface
 {
 	GENERATED_BODY()
 
@@ -38,10 +38,6 @@ public:
 	UComputeDataProvider* CreateDataProvider() const override;
 	//~ End UComputeDataInterface Interface
 
-	void SetProducerSettings(const UPCGSettings* InProducerSettings) { ProducerSettings = InProducerSettings; }
-
-	void SetRequiresExport(bool bInRequiresExport) { bRequiresExport = bInRequiresExport; }
-
 	void SetRequiresZeroInitialization(bool bInZeroInit) { bRequiresZeroInitialization = bInZeroInit; }
 
 	void SetElementCountMultiplier(uint32 InElementCountMultiplier) { ElementCountMultiplier = InElementCountMultiplier; }
@@ -49,18 +45,10 @@ public:
 	void SetExposeElementCounters(bool bInExposeElementCounters) { bExposeElementCounters = bInExposeElementCounters; }
 
 public:
-	/** Settings of node that produces this data, normally the upstream node. */
-	UPROPERTY()
-	TObjectPtr<const UPCGSettings> ProducerSettings;
-
 	/** Whether to perform full 0-initialization of the buffer. */
 	UPROPERTY()
 	bool bRequiresZeroInitialization = false;
 
-	/** Whether this data is passed to downstream tasks outside of this compute graph. */
-	UPROPERTY()
-	bool bRequiresExport = false;
-
 	UPROPERTY()
 	uint32 ElementCountMultiplier = 0;
 
@@ -73,7 +61,7 @@ protected:
 
 /** Compute Framework Data Provider for reading a PCG Data Collection. */
 UCLASS()
-class UPCGDataCollectionDataProvider : public UPCGComputeDataProvider
+class UPCGDataCollectionDataProvider : public UPCGExportableDataProvider
 {
 	GENERATED_BODY()
 
@@ -88,28 +76,10 @@ public:
 	virtual void Reset() override;
 	//~ End UComputeDataProvider Interface
 
-	/** Returns true if the generated buffer should be exported to a persistent buffer to be passed to downstream nodes. */
-	bool IsExportRequired() const { return ExportMode != EPCGExportMode::NoExport; }
-
-	DECLARE_EVENT(UPCGDataCollectionDataProvider, FOnDataExported);
-	FOnDataExported& OnDataExported_GameThread() { return OnDataExported; }
-
 	DECLARE_EVENT(UPCGDataCollectionDataProvider, FOnCountersReadBack);
 	FOnCountersReadBack& OnCountersReadBack_GameThread() { return OnCountersReadBack; }
 
 public:
-	/** Settings of node that produces this data, normally the upstream node. */
-	UPROPERTY()
-	TObjectPtr<const UPCGSettings> ProducerSettings;
-
-	TWeakObjectPtr<UPCGDataBinding> Binding;
-
-	UPROPERTY()
-	int32 GraphBindingIndex = INDEX_NONE;
-
-	UPROPERTY()
-	EPCGExportMode ExportMode = EPCGExportMode::NoExport;
-
 	UPROPERTY()
 	uint32 ElementCountMultiplier = 0;
 
@@ -119,21 +89,7 @@ public:
 	UPROPERTY()
 	bool bExposeElementCounters = false;
 
-	UPROPERTY()
-	FName OutputPinLabel;
-
-	UPROPERTY()
-	FName OutputPinLabelAlias;
-
-	/** Bumped each time the data provider is initialized or reset, so that async callbacks can detect if they originated from a previous usage of the data provider and no-op. */
-	std::atomic<uint64> GenerationCounter = 0;
-
-protected:
-	const FPCGDataCollectionDesc* PinDesc = nullptr;
-
 private:
-	FOnDataExported OnDataExported;
-
 	FOnCountersReadBack OnCountersReadBack;
 
 	std::atomic<bool> bReadbackComplete = false;
diff --git a/Engine/Plugins/PCG/Source/PCG/Private/Compute/DataInterfaces/PCGDataLabelResolverDataInterface.h b/Engine/Plugins/PCG/Source/PCG/Private/Compute/DataInterfaces/PCGDataLabelResolverDataInterface.h
--- a/Engine/Plugins/PCG/Source/PCG/Private/Compute/DataInterfaces/PCGDataLabelResolverDataInterface.h
+++ b/Engine/Plugins/PCG/Source/PCG/Private/Compute/DataInterfaces/PCGDataLabelResolverDataInterface.h
@@ -15,7 +15,7 @@ class FPCGDataLabelResolverDataInterfaceParameters;
 
 /** Data interface for mapping from data label to data index. */
 UCLASS(MinimalAPI, ClassGroup = (Procedural))
-class UPCGDataLabelResolverDataInterface : public UComputeDataInterface
+class UPCGDataLabelResolverDataInterface : public UPCGComputeDataInterface
 {
 	GENERATED_BODY()
 
diff --git a/Engine/Plugins/PCG/Source/PCG/Private/Compute/DataInterfaces/PCGDebugDataInterface.h b/Engine/Plugins/PCG/Source/PCG/Private/Compute/DataInterfaces/PCGDebugDataInterface.h
--- a/Engine/Plugins/PCG/Source/PCG/Private/Compute/DataInterfaces/PCGDebugDataInterface.h
+++ b/Engine/Plugins/PCG/Source/PCG/Private/Compute/DataInterfaces/PCGDebugDataInterface.h
@@ -2,8 +2,7 @@
 
 #pragma once
 
-#include "ComputeFramework/ComputeDataInterface.h"
-#include "ComputeFramework/ComputeDataProvider.h"
+#include "Compute/DataInterfaces/PCGComputeDataInterface.h"
 
 #include "PCGDebugDataInterface.generated.h"
 
@@ -15,7 +14,7 @@ class FRDGBufferUAV;
 
 /** Interface for data about the kernel debug, such as debug value buffer. */
 UCLASS(MinimalAPI, ClassGroup = (Procedural))
-class UPCGDebugDataInterface : public UComputeDataInterface
+class UPCGDebugDataInterface : public UPCGComputeDataInterface
 {
 	GENERATED_BODY()
 
@@ -39,7 +38,7 @@ public:
 
 /** Compute Framework Data Provider for each custom compute kernel. */
 UCLASS()
-class UPCGDebugDataProvider : public UComputeDataProvider
+class UPCGDebugDataProvider : public UPCGComputeDataProvider
 {
 	GENERATED_BODY()
 
diff --git a/Engine/Plugins/PCG/Source/PCG/Private/Compute/DataInterfaces/PCGInstanceDataInterface.cpp b/Engine/Plugins/PCG/Source/PCG/Private/Compute/DataInterfaces/PCGInstanceDataInterface.cpp
--- a/Engine/Plugins/PCG/Source/PCG/Private/Compute/DataInterfaces/PCGInstanceDataInterface.cpp
+++ b/Engine/Plugins/PCG/Source/PCG/Private/Compute/DataInterfaces/PCGInstanceDataInterface.cpp
@@ -132,19 +132,6 @@ UComputeDataProvider* UPCGInstanceDataInterface::CreateDataProvider() const
 	return NewObject<UPCGInstanceDataProvider>();
 }
 
-void UPCGInstanceDataProvider::Initialize(const UComputeDataInterface* InDataInterface, UObject* InBinding, uint64 InInputMask, uint64 InOutputMask)
-{
-	TRACE_CPUPROFILER_EVENT_SCOPE(UPCGInstanceDataProvider::Initialize);
-
-	Super::Initialize(InDataInterface, InBinding, InInputMask, InOutputMask);
-
-	const UPCGInstanceDataInterface* DataInterface = CastChecked<UPCGInstanceDataInterface>(InDataInterface);
-	UPCGDataBinding* Binding = CastChecked<UPCGDataBinding>(InBinding);
-
-	ProducerKernel = CastChecked<UComputeKernel>(DataInterface->ProducerKernel);
-	check(ProducerKernel);
-}
-
 bool UPCGInstanceDataProvider::PrepareForExecute_GameThread(UPCGDataBinding* InBinding)
 {
 	TRACE_CPUPROFILER_EVENT_SCOPE(UPCGInstanceDataProvider::PrepareForExecute_GameThread);
@@ -202,15 +189,12 @@ FComputeDataProviderRenderProxy* UPCGInstanceDataProvider::GetRenderProxy()
 
 void UPCGInstanceDataProvider::Reset()
 {
-	ProducerKernel = nullptr;
+	Super::Reset();
+
 	Primitives.Empty(Primitives.Num());
 	NumInstancesAllPrimitives = 0;
 	NumCustomFloatsPerInstance = 0;
 	bWroteInstances = false;
-
-	++GenerationCounter;
-
-	Super::Reset();
 }
 
 FPCGInstanceDataProviderProxy::FPCGInstanceDataProviderProxy(TArray<FPCGPrimitiveComponentAndNumInstances> InPrimitives, uint32 InNumInstancesAllPrimitives, uint32 InNumCustomFloatsPerInstance, TWeakObjectPtr<UPCGInstanceDataProvider> InDataProvider)
diff --git a/Engine/Plugins/PCG/Source/PCG/Private/Compute/DataInterfaces/PCGInstanceDataInterface.h b/Engine/Plugins/PCG/Source/PCG/Private/Compute/DataInterfaces/PCGInstanceDataInterface.h
--- a/Engine/Plugins/PCG/Source/PCG/Private/Compute/DataInterfaces/PCGInstanceDataInterface.h
+++ b/Engine/Plugins/PCG/Source/PCG/Private/Compute/DataInterfaces/PCGInstanceDataInterface.h
@@ -49,16 +49,12 @@ class UPCGInstanceDataProvider : public UPCGComputeDataProvider
 
 public:
 	//~ Begin UComputeDataProvider Interface
-	virtual void Initialize(const UComputeDataInterface* InDataInterface, UObject* InBinding, uint64 InInputMask, uint64 InOutputMask) override;
 	virtual bool PrepareForExecute_GameThread(UPCGDataBinding* InBinding) override;
 	virtual bool PostExecute(UPCGDataBinding* InBinding) override;
 	FComputeDataProviderRenderProxy* GetRenderProxy() override;
 	virtual void Reset() override;
 	//~ End UComputeDataProvider Interface
 
-	UPROPERTY()
-	TObjectPtr<const UComputeKernel> ProducerKernel;
-
 	TArray<FPCGPrimitiveComponentAndNumInstances> Primitives;
 
 	uint32 NumInstancesAllPrimitives = 0;
diff --git a/Engine/Plugins/PCG/Source/PCG/Private/Compute/DataInterfaces/PCGTextureDataInterface.cpp b/Engine/Plugins/PCG/Source/PCG/Private/Compute/DataInterfaces/PCGTextureDataInterface.cpp
--- a/Engine/Plugins/PCG/Source/PCG/Private/Compute/DataInterfaces/PCGTextureDataInterface.cpp
+++ b/Engine/Plugins/PCG/Source/PCG/Private/Compute/DataInterfaces/PCGTextureDataInterface.cpp
@@ -2,9 +2,14 @@
 
 #include "PCGTextureDataInterface.h"
 
+#include "PCGComponent.h"
+#include "PCGContext.h"
 #include "PCGData.h"
+#include "PCGSettings.h"
 #include "Compute/PCGDataBinding.h"
 #include "Data/PCGTextureData.h"
+#include "Helpers/PCGBlueprintHelpers.h"
+#include "Helpers/PCGHelpers.h"
 
 #include "GlobalRenderResources.h"
 #include "RHIResources.h"
@@ -19,6 +24,7 @@
 #include "ComputeFramework/ComputeKernelPermutationVector.h"
 #include "ComputeFramework/ShaderParamTypeDefinition.h"
 #include "Engine/Texture.h"
+#include "GameFramework/Actor.h"
 
 #include UE_INLINE_GENERATED_CPP_BY_NAME(PCGTextureDataInterface)
 
@@ -29,24 +35,53 @@ namespace PCGTextureDataInterfaceConstants
 	const TCHAR* EnableMultipleTextureObjectsPermutationName = TEXT("ENABLE_MULTIPLE_TEXTURE_OBJECTS");
 }
 
+namespace PCGTextureDataInterfaceHelpers
+{
+	void GetSupportedShaderFunctions(TArray<FShaderFunctionDefinition>& OutFunctions)
+	{
+		OutFunctions.AddDefaulted_GetRef()
+			.SetName(TEXT("GetNumData"))
+			.AddReturnType(EShaderFundamentalType::Uint);
+
+		OutFunctions.AddDefaulted_GetRef()
+			.SetName(TEXT("GetNumElements"))
+			.AddReturnType(EShaderFundamentalType::Uint, 2)
+			.AddParam(EShaderFundamentalType::Uint); // InDataIndex
+
+		OutFunctions.AddDefaulted_GetRef()
+			.SetName(TEXT("GetThreadData"))
+			.AddReturnType(EShaderFundamentalType::Bool)
+			.AddParam(EShaderFundamentalType::Uint) // InThreadIndex
+			.AddParam(EShaderFundamentalType::Uint, 0, 0, EShaderParamModifier::Out) // OutDataIndex
+			.AddParam(EShaderFundamentalType::Uint, 2, 0, EShaderParamModifier::Out); // OutElementIndex
+	}
+}
+
 FPCGTextureBindingInfo::FPCGTextureBindingInfo(const UPCGBaseTextureData* InTextureData)
 {
 	check(InTextureData);
 
 	ResourceType = InTextureData->GetTextureResourceType();
 	Texture = InTextureData->GetTextureRHI();
+	ExportedTexture = InTextureData->GetRefCountedTexture();
+	Transform = InTextureData->GetTransform();
 	TextureBounds = InTextureData->GetBounds();
 	bPointSample = InTextureData->Filter == EPCGTextureFilter::Point;
 
-	if (ResourceType == EPCGTextureResourceType::ExportedTexture)
+	if (Texture.IsValid())
 	{
-		ExportedTexture = InTextureData->GetRefCountedTexture();
+		const FRHITextureDesc& Desc = Texture->GetDesc();
+		Size = Desc.Extent;
+		Dimension = Desc.Dimension;
 	}
+}
 
-	if (Texture)
-	{
-		bIsTextureArray = Texture->GetDesc().Dimension == ETextureDimension::Texture2DArray;
-	}
+FPCGTextureBindingInfo::FPCGTextureBindingInfo(const FPCGDataDesc& InDataDesc, const FTransform& InTransform)
+{
+	// @todo_pcg: Need to support other information here too, like format, slice index, resource type, etc.
+	Size = FIntPoint(FMath::Max(InDataDesc.ElementCount2D.X, 1), FMath::Max(InDataDesc.ElementCount2D.Y, 1));
+	Transform = InTransform;
+	TextureBounds = FBox(FVector(-1.0f, -1.0f, 0.0f), FVector(1.0f, 1.0f, 0.0f)).TransformBy(Transform);
 }
 
 bool FPCGTextureBindingInfo::operator==(const FPCGTextureBindingInfo& Other) const
@@ -54,8 +89,10 @@ bool FPCGTextureBindingInfo::operator==(const FPCGTextureBindingInfo& Other) con
 	return ResourceType == Other.ResourceType
 		&& Texture == Other.Texture
 		&& ExportedTexture == Other.ExportedTexture
-		&& bIsTextureArray == Other.bIsTextureArray
+		&& Size == Other.Size
+		&& Dimension == Other.Dimension
 		&& bPointSample == Other.bPointSample
+		&& Transform.Equals(Transform)
 		&& TextureBounds == Other.TextureBounds;
 }
 
@@ -63,6 +100,8 @@ void UPCGTextureDataInterface::GetSupportedInputs(TArray<FShaderFunctionDefiniti
 {
 	// TODO: Support texture3d, etc. UNiagaraDataInterfaceRenderTarget2D and related are useful references.
 
+	PCGTextureDataInterfaceHelpers::GetSupportedShaderFunctions(OutFunctions);
+
 	OutFunctions.AddDefaulted_GetRef()
 		.SetName(TEXT("GetTexCoords"))
 		.AddReturnType(EShaderFundamentalType::Float, 2)
@@ -76,11 +115,6 @@ void UPCGTextureDataInterface::GetSupportedInputs(TArray<FShaderFunctionDefiniti
 		.AddParam(EShaderFundamentalType::Uint) // InDataIndex
 		.AddParam(EShaderFundamentalType::Float, 2); // InTextureUV
 
-	OutFunctions.AddDefaulted_GetRef()
-		.SetName(TEXT("Sample"))
-		.AddReturnType(EShaderFundamentalType::Float, 4) // InDataIndex
-		.AddParam(EShaderFundamentalType::Float, 2); // InTextureUV
-
 	OutFunctions.AddDefaulted_GetRef()
 		.SetName(TEXT("SampleWorldPos"))
 		.AddReturnType(EShaderFundamentalType::Float, 4)
@@ -88,17 +122,45 @@ void UPCGTextureDataInterface::GetSupportedInputs(TArray<FShaderFunctionDefiniti
 		.AddParam(EShaderFundamentalType::Float, 2); // WorldPos
 
 	OutFunctions.AddDefaulted_GetRef()
-		.SetName(TEXT("SampleWorldPos"))
+		.SetName(TEXT("Load"))
 		.AddReturnType(EShaderFundamentalType::Float, 4)
-		.AddParam(EShaderFundamentalType::Float, 2); // WorldPos
+		.AddParam(EShaderFundamentalType::Uint) // InDataIndex
+		.AddParam(EShaderFundamentalType::Uint, 2); // InElementIndex
+
+	{
+		// Pseudo-deprecated. Assumes data index is 0.
+		OutFunctions.AddDefaulted_GetRef()
+			.SetName(TEXT("Sample"))
+			.AddReturnType(EShaderFundamentalType::Float, 4)
+			.AddParam(EShaderFundamentalType::Float, 2); // InTextureUV
+
+		OutFunctions.AddDefaulted_GetRef()
+			.SetName(TEXT("SampleWorldPos"))
+			.AddReturnType(EShaderFundamentalType::Float, 4)
+			.AddParam(EShaderFundamentalType::Float, 2); // WorldPos
+	}
+}
+
+void UPCGTextureDataInterface::GetSupportedOutputs(TArray<FShaderFunctionDefinition>& OutFunctions) const
+{
+	PCGTextureDataInterfaceHelpers::GetSupportedShaderFunctions(OutFunctions);
+
+	OutFunctions.AddDefaulted_GetRef()
+		.SetName(TEXT("Store"))
+		.AddParam(EShaderFundamentalType::Uint) // InDataIndex
+		.AddParam(EShaderFundamentalType::Uint, 2) // InElementIndex
+		.AddParam(EShaderFundamentalType::Float, 4); // InValue
+
+	// @todo_pcg: We could also add a LoadUAV function. Note that we can't use the existing input Load function because it loads from the SRV, not the UAV.
 }
 
 BEGIN_SHADER_PARAMETER_STRUCT(FPCGTextureDataInterfaceParameters,)
-	SHADER_PARAMETER_RDG_TEXTURE_SRV_ARRAY(Texture2D, Texture, [PCGTextureDataInterfaceConstants::MAX_NUM_BINDINGS])
-	SHADER_PARAMETER_RDG_TEXTURE_SRV_ARRAY(Texture2DArray, TextureArray, [PCGTextureDataInterfaceConstants::MAX_NUM_BINDINGS])
+	SHADER_PARAMETER_RDG_TEXTURE_SRV_ARRAY(Texture2D, Texture, [PCGTextureDataInterfaceConstants::MAX_NUM_SRV_BINDINGS])
+	SHADER_PARAMETER_RDG_TEXTURE_SRV_ARRAY(Texture2DArray, TextureArray, [PCGTextureDataInterfaceConstants::MAX_NUM_SRV_BINDINGS])
+	SHADER_PARAMETER_RDG_TEXTURE_UAV_ARRAY(RWTexture2D<float4>, TextureUAV, [PCGTextureDataInterfaceConstants::MAX_NUM_UAV_BINDINGS])
 	SHADER_PARAMETER_SAMPLER(SamplerState, Sampler_Linear)
 	SHADER_PARAMETER_SAMPLER(SamplerState, Sampler_Point)
-	SHADER_PARAMETER_ARRAY(FVector4f, TextureBounds, [PCGTextureDataInterfaceConstants::MAX_NUM_BINDINGS])
+	SHADER_PARAMETER_ARRAY(FVector4f, TextureBounds, [PCGTextureDataInterfaceConstants::MAX_NUM_SRV_BINDINGS])
 	SHADER_PARAMETER(uint32, NumTextureInfos)
 	SHADER_PARAMETER_RDG_BUFFER_SRV(StructuredBuffer<FUintVector4>, TextureInfos)
 END_SHADER_PARAMETER_STRUCT()
@@ -136,7 +198,7 @@ void UPCGTextureDataInterface::GetHLSL(FString& OutHLSL, FString const& InDataIn
 
 void UPCGTextureDataInterface::GetDefines(FComputeKernelDefinitionSet& OutDefinitionSet) const
 {
-	OutDefinitionSet.Defines.Add(FComputeKernelDefinition(TEXT("MAX_NUM_BINDINGS"), FString::FromInt(PCGTextureDataInterfaceConstants::MAX_NUM_BINDINGS)));
+	OutDefinitionSet.Defines.Add(FComputeKernelDefinition(TEXT("PCG_MAX_NUM_SRV_BINDINGS"), FString::FromInt(PCGTextureDataInterfaceConstants::MAX_NUM_SRV_BINDINGS)));
 }
 
 void UPCGTextureDataInterface::GetPermutations(FComputeKernelPermutationVector& OutPermutationVector) const
@@ -151,16 +213,62 @@ UComputeDataProvider* UPCGTextureDataInterface::CreateDataProvider() const
 
 void UPCGTextureDataProvider::Initialize(const UComputeDataInterface* InDataInterface, UObject* InBinding, uint64 InInputMask, uint64 InOutputMask)
 {
-	TRACE_CPUPROFILER_EVENT_SCOPE(UPCGTextureDataProvider::Initialize);
-
 	Super::Initialize(InDataInterface, InBinding, InInputMask, InOutputMask);
 
 	const UPCGTextureDataInterface* DataInterface = CastChecked<UPCGTextureDataInterface>(InDataInterface);
-	UPCGDataBinding* Binding = CastChecked<UPCGDataBinding>(InBinding);
+	bInitializeFromDataCollection = DataInterface->GetInitializeFromDataCollection();
+}
+
+FComputeDataProviderRenderProxy* UPCGTextureDataProvider::GetRenderProxy()
+{
+	return new FPCGTextureDataProviderProxy(this);
+}
+
+void UPCGTextureDataProvider::Reset()
+{
+	Super::Reset();
+
+	BindingInfos.Empty();
+	TextureInfos.Empty();
+	bInitializeFromDataCollection = false;
+}
+
+bool UPCGTextureDataProvider::PrepareForExecute_GameThread(UPCGDataBinding* InBinding)
+{
+	TRACE_CPUPROFILER_EVENT_SCOPE(UPCGTextureDataProvider::PrepareForExecute_GameThread);
+
+	check(InBinding);
+
+	if (!Super::PrepareForExecute_GameThread(InBinding))
+	{
+		return false;
+	}
+
+	// @todo_pcg: Check if the texture size is too large
+	//if (PCGComputeHelpers::IsBufferSizeTooLarge(PinDesc->ComputePackedSizeBytes()))
+	//{
+	//	PinDesc = nullptr;
+	//}
+
+	if (bInitializeFromDataCollection)
+	{
+		BuildInfosFromDataCollection(InBinding);
+	}
+	else
+	{
+		BuildInfosFromDataDescription(InBinding);
+	}
+
+	return true;
+}
+
+void UPCGTextureDataProvider::BuildInfosFromDataCollection(UPCGDataBinding* InBinding)
+{
+	check(InBinding);
 
 	// Take any input pin label alias to obtain the data from the input data collection.
-	check(!DataInterface->DownstreamInputPinLabelAliases.IsEmpty());
-	const TArray<FPCGTaggedData> InputTaggedData = Binding->InputDataCollection.GetInputsByPin(DataInterface->DownstreamInputPinLabelAliases[0]);
+	check(!DownstreamInputPinLabelAliases.IsEmpty());
+	const TArray<FPCGTaggedData> InputTaggedData = InBinding->InputDataCollection.GetInputsByPin(DownstreamInputPinLabelAliases[0]);
 
 	for (int I = 0; I < InputTaggedData.Num(); ++I)
 	{
@@ -178,15 +286,15 @@ void UPCGTextureDataProvider::Initialize(const UComputeDataInterface* InDataInte
 
 		if (BindingIndex == INDEX_NONE)
 		{
-			if (BindingInfos.Num() < PCGTextureDataInterfaceConstants::MAX_NUM_BINDINGS)
+			if (BindingInfos.Num() < PCGTextureDataInterfaceConstants::MAX_NUM_SRV_BINDINGS)
 			{
 				BindingIndex = BindingInfos.Add(MoveTemp(BindingInfo));
 			}
 			else
 			{
 				UE_LOG(LogPCG, Warning, TEXT("Texture data interface on pin '%s' received too many textures to bind. Only the first %d textures will be bound."),
-					*DataInterface->DownstreamInputPinLabelAliases[0].ToString(),
-					PCGTextureDataInterfaceConstants::MAX_NUM_BINDINGS);
+					*DownstreamInputPinLabelAliases[0].ToString(),
+					PCGTextureDataInterfaceConstants::MAX_NUM_SRV_BINDINGS);
 				continue;
 			}
 		}
@@ -197,17 +305,78 @@ void UPCGTextureDataProvider::Initialize(const UComputeDataInterface* InDataInte
 	}
 }
 
-FComputeDataProviderRenderProxy* UPCGTextureDataProvider::GetRenderProxy()
+void UPCGTextureDataProvider::BuildInfosFromDataDescription(UPCGDataBinding* InBinding)
 {
-	return new FPCGTextureDataProviderProxy(BindingInfos, TextureInfos);
+	check(InBinding);
+	check(PinDesc);
+
+	TSharedPtr<FPCGContextHandle> ContextHandle = InBinding->ContextHandle.Pin();
+	FPCGContext* Context = ContextHandle ? ContextHandle->GetContext() : nullptr;
+
+	if (!ensure(Context))
+	{
+		return;
+	}
+
+	// @todo_pcg: Instead of using this helper and relying on an actor (which may not exist for some cases), we should use the context execution source/state.
+	AActor* OriginalActor = UPCGBlueprintHelpers::GetOriginalComponent(*Context)->GetOwner();
+	if (!ensure(OriginalActor))
+	{
+		return;
+	}
+
+	const FBox OriginalActorLocalBounds = PCGHelpers::GetActorLocalBounds(OriginalActor);
+	FTransform Transform = OriginalActor->GetTransform();
+	Transform.SetScale3D(Transform.GetScale3D() * 0.5 * (OriginalActorLocalBounds.Max - OriginalActorLocalBounds.Min));
+
+	for (int DataIndex = 0; DataIndex < PinDesc->DataDescs.Num(); ++DataIndex)
+	{
+		const FPCGDataDesc& DataDesc = PinDesc->DataDescs[DataIndex];
+
+		FPCGTextureBindingInfo BindingInfo(DataDesc, Transform);
+		int32 BindingIndex = BindingInfos.Find(BindingInfo);
+
+		if (BindingIndex == INDEX_NONE)
+		{
+			if (BindingInfos.Num() < PCGTextureDataInterfaceConstants::MAX_NUM_UAV_BINDINGS)
+			{
+				BindingIndex = BindingInfos.Add(MoveTemp(BindingInfo));
+			}
+			else
+			{
+				UE_LOG(LogPCG, Warning, TEXT("Texture data interface on pin '%s' received too many texture UAVs to bind. Only the first %d texture UAVs will be bound."),
+					*OutputPinLabel.ToString(),
+					PCGTextureDataInterfaceConstants::MAX_NUM_UAV_BINDINGS);
+				continue;
+			}
+		}
+
+		FPCGTextureInfo& TextureInfo = TextureInfos.Emplace_GetRef();
+		TextureInfo.BindingIndex = BindingIndex;
+		TextureInfo.SliceIndex = 0;
+	}
 }
 
-void UPCGTextureDataProvider::Reset()
+FPCGTextureDataProviderProxy::FPCGTextureDataProviderProxy(TWeakObjectPtr<UPCGTextureDataProvider> InDataProvider)
 {
-	BindingInfos.Empty();
-	TextureInfos.Empty();
+	check(IsInGameThread());
+	check(InDataProvider.IsValid());
+
+	BindingInfos = InDataProvider->GetBindingInfos();
+	TextureInfos = InDataProvider->GetTextureInfos();
+	ExportMode = InDataProvider->GetExportMode();
+	OutputPinLabel = InDataProvider->OutputPinLabel;
+	OutputPinLabelAlias = InDataProvider->OutputPinLabelAlias;
+	OriginatingGenerationCount = InDataProvider->GenerationCounter;
+	DataProviderWeakPtr_GT = InDataProvider;
+
+	if (InDataProvider->GetPinDescription())
+	{
+		TRACE_CPUPROFILER_EVENT_SCOPE(CopyDataDescription);
 
-	Super::Reset();
+		// @todo_pcg: Avoid costly copy of data description.
+		PinDesc = *InDataProvider->GetPinDescription();
+	}
 }
 
 bool FPCGTextureDataProviderProxy::IsValid(FValidationData const& InValidationData) const
@@ -263,8 +432,10 @@ void FPCGTextureDataProviderProxy::GatherDispatchData(FDispatchData const& InDis
 
 		Parameters.Sampler_Linear = TStaticSamplerState<SF_Bilinear, AM_Clamp, AM_Clamp, AM_Clamp>::GetRHI();
 		Parameters.Sampler_Point = TStaticSamplerState<SF_Point, AM_Clamp, AM_Clamp, AM_Clamp>::GetRHI();
+		Parameters.NumTextureInfos = TextureInfos.Num();
+		Parameters.TextureInfos = TextureInfosBufferSRV;
 
-		for (int BindingIndex = 0; BindingIndex < PCGTextureDataInterfaceConstants::MAX_NUM_BINDINGS; ++BindingIndex)
+		for (int BindingIndex = 0; BindingIndex < PCGTextureDataInterfaceConstants::MAX_NUM_SRV_BINDINGS; ++BindingIndex)
 		{
 			Parameters.Texture[BindingIndex] = TextureSRV[BindingIndex];
 			Parameters.TextureArray[BindingIndex] = TextureArraySRV[BindingIndex];
@@ -281,60 +452,186 @@ void FPCGTextureDataProviderProxy::GatherDispatchData(FDispatchData const& InDis
 			{
 				Parameters.TextureBounds[BindingIndex] = FVector4f::Zero();
 			}
+		}
 
-			Parameters.NumTextureInfos = TextureInfos.Num();
-			Parameters.TextureInfos = TextureInfosBufferSRV;
+		for (int BindingIndex = 0; BindingIndex < PCGTextureDataInterfaceConstants::MAX_NUM_UAV_BINDINGS; ++BindingIndex)
+		{
+			Parameters.TextureUAV[BindingIndex] = TextureUAV[BindingIndex];
 		}
 	}
 }
 
 void FPCGTextureDataProviderProxy::AllocateResources(FRDGBuilder& GraphBuilder, FAllocationData const& InAllocationData)
 {
-	CreateTextureSRVs(GraphBuilder);
+	CreateDefaultTextures(GraphBuilder);
+	CreateTextures(GraphBuilder);
 	PackTextureInfos(GraphBuilder);
 }
 
-void FPCGTextureDataProviderProxy::CreateTextureSRVs(FRDGBuilder& GraphBuilder)
+void FPCGTextureDataProviderProxy::CreateDefaultTextures(FRDGBuilder& GraphBuilder)
 {
-	FRDGTextureRef DummyTextureRDG = GSystemTextures.GetDefaultTexture2D(GraphBuilder, EPixelFormat::PF_G8, 0.0f);
-	FRDGTextureRef DummyTextureArrayRDG = GSystemTextures.GetDefaultTexture(GraphBuilder, ETextureDimension::Texture2DArray, EPixelFormat::PF_G8, 0.0f);
+	FRDGTextureRef DummyTextureForSRV = GSystemTextures.GetDefaultTexture2D(GraphBuilder, EPixelFormat::PF_G8, 0.0f);
+	FRDGTextureRef DummyTextureArrayForSRV = GSystemTextures.GetDefaultTexture(GraphBuilder, ETextureDimension::Texture2DArray, EPixelFormat::PF_G8, 0.0f);
 
-	for (int BindingIndex = 0; BindingIndex < PCGTextureDataInterfaceConstants::MAX_NUM_BINDINGS; ++BindingIndex)
+	FRDGTextureRef DummyTextureForUAV = GraphBuilder.CreateTexture(
+			FRDGTextureDesc::Create2D(FIntPoint(1, 1), PF_G8, FClearValueBinding::Black, ETextureCreateFlags::ShaderResource | ETextureCreateFlags::UAV),
+			TEXT("PCGTextureDI_DummyTextureForUAV"));
+
+	FRDGTextureRef DummyTextureArrayForUAV = GraphBuilder.CreateTexture(
+			FRDGTextureDesc::Create2DArray(FIntPoint(1, 1), PF_G8, FClearValueBinding::Black, ETextureCreateFlags::ShaderResource | ETextureCreateFlags::UAV, /*ArraySize=*/1),
+			TEXT("PCGTextureDI_DummyTextureArrayForUAV"));
+
+	for (int BindingIndex = 0; BindingIndex < PCGTextureDataInterfaceConstants::MAX_NUM_SRV_BINDINGS; ++BindingIndex)
+	{
+		TextureSRV[BindingIndex] = GraphBuilder.CreateSRV(DummyTextureForSRV);
+		TextureArraySRV[BindingIndex] = GraphBuilder.CreateSRV(DummyTextureArrayForSRV);
+	}
+
+	for (int BindingIndex = 0; BindingIndex < PCGTextureDataInterfaceConstants::MAX_NUM_UAV_BINDINGS; ++BindingIndex)
 	{
-		// Default to blank dummy texture if no texture was provided.
-		FRDGTextureRef TextureRDG = DummyTextureRDG;
-		FRDGTextureRef TextureArrayRDG = DummyTextureArrayRDG;
+		TextureUAV[BindingIndex] = GraphBuilder.CreateUAV(DummyTextureForUAV);
+	}
+}
 
-		if (BindingIndex < BindingInfos.Num())
+void FPCGTextureDataProviderProxy::CreateTextures(FRDGBuilder& GraphBuilder)
+{
+	FRDGTextureRef ExportableTextures[PCGTextureDataInterfaceConstants::MAX_NUM_UAV_BINDINGS] = {};
+
+	for (int BindingIndex = 0; BindingIndex < BindingInfos.Num(); ++BindingIndex)
+	{
+		FRDGTextureRef Texture = nullptr;
+		bool bCanCreateUAV = false;
+
+		if (BindingInfos[BindingIndex].ResourceType == EPCGTextureResourceType::TextureObject && BindingInfos[BindingIndex].Texture)
+		{
+			Texture = GraphBuilder.RegisterExternalTexture(CreateRenderTarget(BindingInfos[BindingIndex].Texture, TEXT("PCGTextureDI_RenderTarget")));
+		}
+		else if (BindingInfos[BindingIndex].ResourceType == EPCGTextureResourceType::ExportedTexture && BindingInfos[BindingIndex].ExportedTexture)
+		{
+			Texture = GraphBuilder.RegisterExternalTexture(BindingInfos[BindingIndex].ExportedTexture);
+		}
+		else if (BindingIndex < PCGTextureDataInterfaceConstants::MAX_NUM_UAV_BINDINGS)
 		{
-			FRDGTextureRef Texture = nullptr;
+			// @todo_pcg: Expose options like size, format, dimensions (e.g. Tex2D vs TexArray), etc.
+			Texture = GraphBuilder.CreateTexture(
+				FRDGTextureDesc::Create2D(
+					BindingInfos[BindingIndex].Size,
+					EPixelFormat::PF_FloatRGBA,
+					FClearValueBinding::Black,
+					ETextureCreateFlags::ShaderResource | ETextureCreateFlags::UAV),
+				TEXT("PCGTextureDI_UAV"));
+
+			bCanCreateUAV = true;
+		}
 
-			if (BindingInfos[BindingIndex].ResourceType == EPCGTextureResourceType::TextureObject && BindingInfos[BindingIndex].Texture)
-			{
-				Texture = GraphBuilder.RegisterExternalTexture(
-					CreateRenderTarget(BindingInfos[BindingIndex].Texture, *FText::Format(LOCTEXT("RenderTargetDebugName", "PCGDataInterfaceRenderTarget_{0}"), BindingIndex).ToString()));
-			}
-			else if (BindingInfos[BindingIndex].ResourceType == EPCGTextureResourceType::ExportedTexture && BindingInfos[BindingIndex].ExportedTexture)
+		if (!Texture)
+		{
+			continue;
+		}
+
+		ExportableTextures[BindingIndex] = Texture;
+
+		if (BindingInfos[BindingIndex].Dimension == ETextureDimension::Texture2D)
+		{
+			TextureSRV[BindingIndex] = GraphBuilder.CreateSRV(Texture);
+
+			if (bCanCreateUAV)
 			{
-				Texture = GraphBuilder.RegisterExternalTexture(BindingInfos[BindingIndex].ExportedTexture);
+				TextureUAV[BindingIndex] = GraphBuilder.CreateUAV(Texture);
 			}
+		}
+		else if (BindingInfos[BindingIndex].Dimension == ETextureDimension::Texture2DArray)
+		{
+			TextureArraySRV[BindingIndex] = GraphBuilder.CreateSRV(Texture);
+		}
+		else
+		{
+			checkNoEntry();
+		}
+	}
 
-			if (Texture)
+	if (ExportMode != EPCGExportMode::NoExport)
+	{
+		TArray<TRefCountPtr<IPooledRenderTarget>> ExportedTextures;
+
+		for (const FPCGTextureInfo& TextureInfo : TextureInfos)
+		{
+			if (ensure(TextureInfo.BindingIndex >= 0 && TextureInfo.BindingIndex < PCGTextureDataInterfaceConstants::MAX_NUM_UAV_BINDINGS))
 			{
-				if (BindingInfos[BindingIndex].bIsTextureArray)
+				if (FRDGTextureRef ExportableTexture = ExportableTextures[TextureInfo.BindingIndex])
 				{
-					TextureArrayRDG = Texture;
+					ExportedTextures.Add(GraphBuilder.ConvertToExternalTexture(ExportableTexture));
 				}
-				else
+			}
+		}
+
+		ExportTextureUAVs(ExportedTextures);
+	}
+}
+
+void FPCGTextureDataProviderProxy::ExportTextureUAVs(const TArray<TRefCountPtr<IPooledRenderTarget>>& ExportedTextures)
+{
+	// Export textures and pass a reference back to the main thread where it will be picked up by the compute graph element.
+	ExecuteOnGameThread(UE_SOURCE_LOCATION, [ExportedTextures=ExportedTextures, DataProviderWeakPtr=DataProviderWeakPtr_GT, PinDesc=PinDesc, OutputPinLabel=OutputPinLabel, ExportMode=ExportMode, OutputPinLabelAlias=OutputPinLabelAlias, GenerationCount=OriginatingGenerationCount]()
+	{
+		// Obtain objects. No ensures added because a graph cancellation could feasibly destroy some or all of these.
+		UPCGTextureDataProvider* DataProvider = DataProviderWeakPtr.Get();
+		if (!DataProvider)
+		{
+			UE_LOG(LogPCG, Error, TEXT("Could not resolve UPCGTextureDataProvider object to pass back buffer handle."));
+			return;
+		}
+
+		if (DataProvider->GenerationCounter != GenerationCount)
+		{
+			return;
+		}
+
+		if (DataProvider->bInitializeFromDataCollection)
+		{
+			DataProvider->OnDataExported_GameThread().Broadcast();
+			return;
+		}
+
+		if (UPCGDataBinding* Binding = DataProvider->GetDataBinding())
+		{
+			const TArray<FPCGTextureInfo>& TextureInfos = DataProvider->GetTextureInfos();
+			const TArray<FPCGTextureBindingInfo>& BindingInfos = DataProvider->GetBindingInfos();
+			const TArray<FString>& StringTable = Binding->GetStringTable();
+
+			for (int TextureDataIndex = 0; TextureDataIndex < TextureInfos.Num(); ++TextureDataIndex)
+			{
+				const FPCGTextureInfo& TextureInfo = TextureInfos[TextureDataIndex];
+				check(ExportedTextures.IsValidIndex(TextureDataIndex));
+				check(PinDesc.DataDescs.IsValidIndex(TextureDataIndex))
+				check(BindingInfos.IsValidIndex(TextureInfo.BindingIndex));
+
+				const FPCGTextureBindingInfo& BindingInfo = BindingInfos[TextureInfo.BindingIndex];
+				const FTransform TextureTransform = BindingInfos[TextureInfo.BindingIndex].Transform;
+				const FVector::FReal XSize = 2.0 * TextureTransform.GetScale3D().X;
+				const FVector::FReal YSize = 2.0 * TextureTransform.GetScale3D().Y;
+				const FVector2D TexelSize = FVector2D(XSize, YSize) / BindingInfo.Size;
+
+				UPCGTextureData* ExportedData = NewObject<UPCGTextureData>();
+				ExportedData->TexelSize = TexelSize.GetMin();
+				ExportedData->Initialize(ExportedTextures[TextureDataIndex], TextureInfo.SliceIndex, TextureTransform, /*bInSkipReadbackToCPU=*/true);
+
+				TSet<FString> Tags;
+				for (const int32 TagStringKey : PinDesc.DataDescs[TextureDataIndex].TagStringKeys)
 				{
-					TextureRDG = Texture;
+					if (ensure(StringTable.IsValidIndex(TagStringKey)))
+					{
+						Tags.Add(StringTable[TagStringKey]);
+					}
 				}
+
+				// @todo_pcg: Binding is doing a lot of work. Could store a context handle in the data provider instead?
+				Binding->ReceiveDataFromGPU_GameThread(ExportedData, DataProvider->GetProducerSettings(), ExportMode, OutputPinLabel, OutputPinLabelAlias, Tags);
 			}
 		}
 
-		TextureSRV[BindingIndex] = GraphBuilder.CreateSRV(TextureRDG);
-		TextureArraySRV[BindingIndex] = GraphBuilder.CreateSRV(TextureArrayRDG);
-	}
+		DataProvider->OnDataExported_GameThread().Broadcast();
+	});
 }
 
 void FPCGTextureDataProviderProxy::PackTextureInfos(FRDGBuilder& GraphBuilder)
@@ -343,26 +640,30 @@ void FPCGTextureDataProviderProxy::PackTextureInfos(FRDGBuilder& GraphBuilder)
 	// - TextureInfo0
 	//   - uint32 - BindingIndex
 	//   - uint32 - SliceIndex (aka TextureIndex if bound to a TextureArray, 0 otherwise.)
-	//   - uint32 - IsTextureArray
+	//   - uint32 - Dimension
+	//   - uint32 - Sampler
+	//   - uint32 - SizeX
+	//   - uint32 - SizeY
 	// - TextureInfo1
 	// - TextureInfo2
 	// - ...
 
 	if (!TextureInfos.IsEmpty())
 	{
-		const FRDGBufferDesc TextureInfosDesc = FRDGBufferDesc::CreateStructuredDesc(sizeof(FUintVector4), TextureInfos.Num());
+		const FRDGBufferDesc TextureInfosDesc = FRDGBufferDesc::CreateStructuredDesc(sizeof(FUintVector4), TextureInfos.Num() * 2);
 
-		FRDGBufferRef TextureInfosBuffer = GraphBuilder.CreateBuffer(TextureInfosDesc, TEXT("PCGTextureInfosBuffer"));
+		FRDGBufferRef TextureInfosBuffer = GraphBuilder.CreateBuffer(TextureInfosDesc, TEXT("PCGTextureDI_InfosBuffer"));
 		TextureInfosBufferSRV = GraphBuilder.CreateSRV(TextureInfosBuffer);
 
 		TArray<FUintVector4> TextureInfosPackedData;
-		TextureInfosPackedData.SetNumUninitialized(TextureInfos.Num());
+		TextureInfosPackedData.SetNumUninitialized(TextureInfos.Num() * 2);
 
 		for (int TextureInfoIndex = 0; TextureInfoIndex < TextureInfos.Num(); ++TextureInfoIndex)
 		{
 			const FPCGTextureInfo& TextureInfo = TextureInfos[TextureInfoIndex];
 			const FPCGTextureBindingInfo& BindingInfo = BindingInfos[TextureInfo.BindingIndex];
-			TextureInfosPackedData[TextureInfoIndex] = FUintVector4(TextureInfo.BindingIndex, TextureInfo.SliceIndex, BindingInfo.bIsTextureArray, BindingInfo.bPointSample);
+			TextureInfosPackedData[TextureInfoIndex * 2 + 0] = FUintVector4(TextureInfo.BindingIndex, TextureInfo.SliceIndex, static_cast<uint32>(BindingInfo.Dimension), BindingInfo.bPointSample);
+			TextureInfosPackedData[TextureInfoIndex * 2 + 1] = FUintVector4(BindingInfo.Size.X, BindingInfo.Size.Y, 0u, 0u);
 		}
 
 		GraphBuilder.QueueBufferUpload(TextureInfosBuffer, MakeArrayView(TextureInfosPackedData));
diff --git a/Engine/Plugins/PCG/Source/PCG/Private/Compute/DataInterfaces/PCGTextureDataInterface.h b/Engine/Plugins/PCG/Source/PCG/Private/Compute/DataInterfaces/PCGTextureDataInterface.h
--- a/Engine/Plugins/PCG/Source/PCG/Private/Compute/DataInterfaces/PCGTextureDataInterface.h
+++ b/Engine/Plugins/PCG/Source/PCG/Private/Compute/DataInterfaces/PCGTextureDataInterface.h
@@ -2,12 +2,12 @@
 
 #pragma once
 
-#include "ComputeFramework/ComputeDataProvider.h"
+#include "Compute/PCGDataForGPU.h"
 #include "Compute/DataInterfaces/PCGComputeDataInterface.h"
-
 #include "Data/PCGTextureData.h"
 
 #include "RHIResources.h"
+#include "ComputeFramework/ComputeDataProvider.h"
 
 #include "PCGTextureDataInterface.generated.h"
 
@@ -17,25 +17,29 @@ class UPCGBaseTextureData;
 namespace PCGTextureDataInterfaceConstants
 {
 	/**
-	 * The maximum number of textures which can be bound to a PCGTextureDataInterface.
+	 * The maximum number of texture SRVs and UAVs which can be bound to a PCGTextureDataInterface.
 	 * We can't use full bindless on all platforms, so fallback to emulating.
 	 * If you change this number, make sure update PCGTextureDataInterface.ush as well.
 	 */
-	constexpr uint32 MAX_NUM_BINDINGS = 8;
+	constexpr uint32 MAX_NUM_SRV_BINDINGS = 8;
+	constexpr uint32 MAX_NUM_UAV_BINDINGS = 1;
 }
 
 struct FPCGTextureBindingInfo
 {
-	FPCGTextureBindingInfo(const UPCGBaseTextureData* InTextureData);
+	explicit FPCGTextureBindingInfo(const UPCGBaseTextureData* InTextureData);
+	FPCGTextureBindingInfo(const FPCGDataDesc& InDataDesc, const FTransform& InTransform);
 	bool operator==(const FPCGTextureBindingInfo& Other) const;
 
 	EPCGTextureResourceType ResourceType = EPCGTextureResourceType::TextureObject;
 	FTextureRHIRef Texture = nullptr;
 	TRefCountPtr<IPooledRenderTarget> ExportedTexture = nullptr;
-	bool bIsTextureArray = false;
+	FIntPoint Size = FIntPoint::ZeroValue;
+	ETextureDimension Dimension = ETextureDimension::Texture2D;
 	bool bPointSample = false;
 
-	// @todo_pcg: TextureBounds could potentially be independent of the texture binding, and queried via lookup from the texture info buffer instead.
+	// @todo_pcg: Eventually replace TextureBounds with only the Transform
+	FTransform Transform = FTransform::Identity;
 	FBox TextureBounds = FBox(EForceInit::ForceInit);
 };
 
@@ -47,7 +51,7 @@ struct FPCGTextureInfo
 
 /** Data Interface allowing sampling of a texture. */
 UCLASS(ClassGroup = (Procedural))
-class UPCGTextureDataInterface : public UPCGComputeDataInterface
+class UPCGTextureDataInterface : public UPCGExportableDataInterface
 {
 	GENERATED_BODY()
 
@@ -55,6 +59,7 @@ public:
 	//~ Begin UComputeDataInterface Interface
 	TCHAR const* GetClassName() const override { return TEXT("PCGTexture"); }
 	void GetSupportedInputs(TArray<FShaderFunctionDefinition>& OutFunctions) const override;
+	void GetSupportedOutputs(TArray<FShaderFunctionDefinition>& OutFunctions) const override;
 	void GetShaderParameters(TCHAR const* UID, FShaderParametersMetadataBuilder& InOutBuilder, FShaderParametersMetadataAllocations& InOutAllocations) const override;
 	TCHAR const* GetShaderVirtualPath() const override;
 	void GetShaderHash(FString& InOutKey) const override;
@@ -64,13 +69,20 @@ public:
 	UComputeDataProvider* CreateDataProvider() const override;
 	//~ End UComputeDataInterface Interface
 
+	void SetInitializeFromDataCollection(bool bInInitializeFromDataCollection) { bInitializeFromDataCollection = bInInitializeFromDataCollection; }
+	bool GetInitializeFromDataCollection() const { return bInitializeFromDataCollection; }
+
+protected:
+	UPROPERTY()
+	bool bInitializeFromDataCollection = false;
+
 private:
 	static TCHAR const* TemplateFilePath;
 };
 
 /** Compute Framework Data Provider for reading a texture. */
 UCLASS()
-class UPCGTextureDataProvider : public UComputeDataProvider
+class UPCGTextureDataProvider : public UPCGExportableDataProvider
 {
 	GENERATED_BODY()
 
@@ -81,17 +93,27 @@ public:
 	virtual void Reset() override;
 	//~ End UComputeDataProvider Interface
 
+	//~ Begin UPCGComputeDataProvider Interface
+	virtual bool PrepareForExecute_GameThread(UPCGDataBinding* InBinding) override;
+	//~ End UPCGComputeDataProvider Interface
+
+	const TArray<FPCGTextureBindingInfo>& GetBindingInfos() const { return BindingInfos; }
+	const TArray<FPCGTextureInfo>& GetTextureInfos() const { return TextureInfos; }
+
+protected:
+	void BuildInfosFromDataCollection(UPCGDataBinding* InBinding);
+	void BuildInfosFromDataDescription(UPCGDataBinding* InBinding);
+
+public:
 	TArray<FPCGTextureBindingInfo> BindingInfos;
 	TArray<FPCGTextureInfo> TextureInfos;
+	bool bInitializeFromDataCollection = false;
 };
 
 class FPCGTextureDataProviderProxy : public FComputeDataProviderRenderProxy
 {
 public:
-	FPCGTextureDataProviderProxy(const TArray<FPCGTextureBindingInfo>& InBindingInfos, const TArray<FPCGTextureInfo>& InTextureInfos)
-		: BindingInfos(InBindingInfos)
-		, TextureInfos(InTextureInfos)
-	{}
+	explicit FPCGTextureDataProviderProxy(TWeakObjectPtr<UPCGTextureDataProvider> InDataProvider);
 
 	//~ Begin FComputeDataProviderRenderProxy Interface
 	bool IsValid(FValidationData const& InValidationData) const override;
@@ -100,7 +122,9 @@ public:
 	void AllocateResources(FRDGBuilder& GraphBuilder, FAllocationData const& InAllocationData) override;
 	//~ End FComputeDataProviderRenderProxy Interface
 
-	void CreateTextureSRVs(FRDGBuilder& GraphBuilder);
+	void CreateDefaultTextures(FRDGBuilder& GraphBuilder);
+	void CreateTextures(FRDGBuilder& GraphBuilder);
+	void ExportTextureUAVs(const TArray<TRefCountPtr<IPooledRenderTarget>>& ExportedTextures);
 	void PackTextureInfos(FRDGBuilder& GraphBuilder);
 
 protected:
@@ -109,7 +133,22 @@ protected:
 	TArray<FPCGTextureBindingInfo> BindingInfos;
 	TArray<FPCGTextureInfo> TextureInfos;
 
-	FRDGTextureSRVRef TextureSRV[PCGTextureDataInterfaceConstants::MAX_NUM_BINDINGS];
-	FRDGTextureSRVRef TextureArraySRV[PCGTextureDataInterfaceConstants::MAX_NUM_BINDINGS];
+	// For export
+	EPCGExportMode ExportMode = EPCGExportMode::NoExport;
+	FPCGDataCollectionDesc PinDesc;
+	FName OutputPinLabel;
+	FName OutputPinLabelAlias;
+
+	/** Generation count of the data provider when the proxy was created. */
+	uint64 OriginatingGenerationCount = 0;
+
+	// Weak pointer useful for passing back texture handles. Do not access this directly from the render thread.
+	TWeakObjectPtr<UPCGTextureDataProvider> DataProviderWeakPtr_GT;
+
+	// Allocated resources
+	FRDGTextureSRVRef TextureSRV[PCGTextureDataInterfaceConstants::MAX_NUM_SRV_BINDINGS];
+	FRDGTextureSRVRef TextureArraySRV[PCGTextureDataInterfaceConstants::MAX_NUM_SRV_BINDINGS];
+	FRDGTextureUAVRef TextureUAV[PCGTextureDataInterfaceConstants::MAX_NUM_UAV_BINDINGS];
+	FRDGTextureUAVRef TextureArrayUAV[PCGTextureDataInterfaceConstants::MAX_NUM_UAV_BINDINGS];
 	FRDGBufferSRVRef TextureInfosBufferSRV = nullptr;
 };
diff --git a/Engine/Plugins/PCG/Source/PCG/Private/Compute/Elements/PCGComputeGraphElement.cpp b/Engine/Plugins/PCG/Source/PCG/Private/Compute/Elements/PCGComputeGraphElement.cpp
--- a/Engine/Plugins/PCG/Source/PCG/Private/Compute/Elements/PCGComputeGraphElement.cpp
+++ b/Engine/Plugins/PCG/Source/PCG/Private/Compute/Elements/PCGComputeGraphElement.cpp
@@ -8,6 +8,7 @@
 #include "PCGSubsystem.h"
 #include "Components/PCGProceduralISMComponent.h"
 #include "Compute/PCGComputeCommon.h"
+#include "Data/PCGTextureData.h"
 #include "Compute/Data/PCGProxyForGPUData.h"
 #include "Compute/DataInterfaces/PCGDataCollectionDataInterface.h"
 #include "Elements/PCGStaticMeshSpawnerKernel.h"
@@ -223,19 +224,19 @@ bool FPCGComputeGraphElement::ExecuteInternal(FPCGContext* InContext) const
 		// Register all providers running async operations.
 		for (UComputeDataProvider* ComputeDataProvider : Context->ComputeGraphInstance->GetDataProviders())
 		{
-			UPCGDataCollectionDataProvider* DataProvider = Cast<UPCGDataCollectionDataProvider>(ComputeDataProvider);
+			UPCGExportableDataProvider* DataProvider = Cast<UPCGExportableDataProvider>(ComputeDataProvider);
 			if (!DataProvider)
 			{
 				continue;
 			}
 
-			const bool bDataProviderDebugged = DataProvider->ProducerSettings && DataProvider->ProducerSettings->bDebug;
+			const bool bDataProviderDebugged = DataProvider->GetProducerSettings() && DataProvider->GetProducerSettings()->bDebug;
 
 			if (DataProvider->IsExportRequired() || bDataProviderDebugged)
 			{
 				Context->ProvidersWithBufferExports.Add(DataProvider);
 
-				TWeakObjectPtr<UPCGDataCollectionDataProvider> DataProviderWeak = DataProvider;
+				TWeakObjectPtr<UPCGExportableDataProvider> DataProviderWeak = DataProvider;
 
 				const uint64 OriginatingGenerationCount = DataProvider->GenerationCounter;
 
@@ -245,7 +246,7 @@ bool FPCGComputeGraphElement::ExecuteInternal(FPCGContext* InContext) const
 
 					TSharedPtr<FPCGContextHandle> SharedHandle = ContextHandle.Pin();
 					FPCGComputeGraphContext* ContextPtr = static_cast<FPCGComputeGraphContext*>(SharedHandle ? SharedHandle->GetContext() : nullptr);
-					UPCGDataCollectionDataProvider* DataProvider = DataProviderWeak.Get();
+					UPCGExportableDataProvider* DataProvider = DataProviderWeak.Get();
 
 					if (!DataProvider || !ContextPtr || GenerationCount != DataProvider->GenerationCounter)
 					{
@@ -262,13 +263,15 @@ bool FPCGComputeGraphElement::ExecuteInternal(FPCGContext* InContext) const
 				});
 			}
 
-			if (DataProvider->bExposeElementCounters)
+			UPCGDataCollectionDataProvider* DataCollectionDataProvider = Cast<UPCGDataCollectionDataProvider>(ComputeDataProvider);
+
+			if (DataCollectionDataProvider && DataCollectionDataProvider->bExposeElementCounters)
 			{
-				Context->ProvidersWithReadbacks.Add(DataProvider);
+				Context->ProvidersWithReadbacks.Add(DataCollectionDataProvider);
 
-				TWeakObjectPtr<UPCGDataCollectionDataProvider> DataProviderWeak = DataProvider;
+				TWeakObjectPtr<UPCGDataCollectionDataProvider> DataProviderWeak = DataCollectionDataProvider;
 
-				DataProvider->OnCountersReadBack_GameThread().AddLambda([ContextHandle = Context->GetOrCreateHandle(), DataProviderWeak]()
+				DataCollectionDataProvider->OnCountersReadBack_GameThread().AddLambda([ContextHandle = Context->GetOrCreateHandle(), DataProviderWeak]()
 				{
 					check(IsInGameThread());
 
@@ -585,15 +588,31 @@ bool FPCGComputeGraphElement::ExecuteInternal(FPCGContext* InContext) const
 			return false;
 		}
 
-		// Call both to start all required readbacks.
-		const bool bDebugDrawDone = ExecuteDebugDraw(Context);
-		const bool bInspectionDone = StoreDataForInspection(Context);
+		CollectDebugDataPrepareActions(Context);
 
-		if (!bDebugDrawDone || !bInspectionDone)
+		bool bAllReady = true;
+
+		for (int32 Index = Context->DebugDataPrepareActions.Num() - 1; Index >= 0; --Index)
+		{
+			FPCGComputeGraphContext::FDebugDataPrepareAction Action = Context->DebugDataPrepareActions[Index];
+			const bool bActionIsDone = !Action.IsSet() || Action(Context);
+
+			bAllReady &= bActionIsDone;
+
+			if (bActionIsDone)
+			{
+				Context->DebugDataPrepareActions.RemoveAtSwap(Index);
+			}
+		}
+
+		if (!bAllReady)
 		{
 			SleepUntilNextFrame();
 			return false;
 		}
+
+		ExecuteDebugDraw(Context);
+		StoreDataForInspection(Context);
 #endif
 	}
 	}
@@ -645,107 +664,125 @@ void FPCGComputeGraphElement::ResetAsyncOperations(FPCGContext* InContext) const
 }
 
 #if WITH_EDITOR
-bool FPCGComputeGraphElement::ExecuteDebugDraw(FPCGComputeGraphContext* InContext) const
+void FPCGComputeGraphElement::CollectDebugDataPrepareActions(FPCGComputeGraphContext* InContext) const
 {
-	if (!InContext || !InContext->DataBinding || InContext->DataBinding->DataToDebug.IsEmpty())
+	if (!InContext || !InContext->DataBinding)
 	{
-		return true;
+		return;
 	}
 
-	const FPCGDataVisualizationRegistry& DataVisRegistry = FPCGModule::GetConstPCGDataVisualizationRegistry();
-
-	bool bAllDone = true;
-
-	for (int Index = InContext->DataBinding->DataToDebug.Num() - 1; Index >= 0; --Index)
+	auto CollectDebugDataPrepareAction = [InContext](FPCGDataToDebug& DataToDebug)
 	{
-		FPCGDataToDebug& DataToDebug = InContext->DataBinding->DataToDebug[Index];
-		const UPCGProxyForGPUData* Proxy = DataToDebug.Proxy.Get();
-		const UPCGSettings* ProducerSettings = DataToDebug.ProducerSettings.Get();
-
-		if (Proxy && ProducerSettings)
+		if (const UPCGProxyForGPUData* ProxyData = Cast<UPCGProxyForGPUData>(DataToDebug.Data.Get()))
 		{
-			UPCGProxyForGPUData::FReadbackResult Readback = Proxy->GetCPUData(InContext);
-			if (Readback.bComplete)
+			InContext->DebugDataPrepareActions.Add([&DataToDebug](FPCGComputeGraphContext* InContext)
+			{
+				const UPCGProxyForGPUData* ProxyData = CastChecked<UPCGProxyForGPUData>(DataToDebug.Data.Get());
+				UPCGProxyForGPUData::FReadbackResult Readback = ProxyData->GetCPUData(InContext);
+				return Readback.bComplete;
+			});
+		}
+		else if (const UPCGTextureData* TextureData = Cast<UPCGTextureData>(DataToDebug.Data.Get()))
+		{
+			InContext->DebugDataPrepareActions.Add([&DataToDebug](FPCGComputeGraphContext* InContext)
 			{
-				if (Readback.TaggedData.Data)
+				// Duplicate and CPU-initialize the debug texture data.
+				const UPCGTextureData* TextureData = CastChecked<UPCGTextureData>(DataToDebug.Data.Get());
+				TObjectPtr<UPCGTextureData> DuplicateTextureData = Cast<UPCGTextureData>(DataToDebug.DataPendingInit);
+
+				if (!DuplicateTextureData)
 				{
-					if (const IPCGDataVisualization* DataVis = DataVisRegistry.GetDataVisualization(Readback.TaggedData.Data->GetClass()))
-					{
-						DataVis->ExecuteDebugDisplay(InContext, ProducerSettings, Readback.TaggedData.Data, InContext->GetTargetActor(nullptr));
-					}
+					DuplicateTextureData = InContext->NewObject_AnyThread<UPCGTextureData>(InContext, GetTransientPackage());
+					DuplicateTextureData->TexelSize = TextureData->TexelSize;
+					DataToDebug.DataPendingInit = DuplicateTextureData;
 				}
 
-				InContext->DataBinding->DataToDebug.RemoveAtSwap(Index);
-			}
-			else
-			{
-				// Readback needs another frame to complete.
-				bAllDone = false;
-			}
-		}
-		else
-		{
-			// Could not obtain proxy, remove this data.
-			InContext->DataBinding->DataToDebug.RemoveAtSwap(Index);
+				return DuplicateTextureData->Initialize(TextureData->GetRefCountedTexture(), TextureData->TextureIndex, TextureData->GetTransform(), /*bInSkipReadbackToCPU=*/false);
+			});
 		}
-	}
+	};
 
-	if (bAllDone)
+	for (FPCGDataToDebug& DataToDebug : InContext->DataBinding->DataToDebug)
 	{
-		InContext->DataBinding->DataToDebug.Empty();
+		CollectDebugDataPrepareAction(DataToDebug);
 	}
 
-	return bAllDone;
+	for (FPCGDataToDebug& DataToDebug : InContext->DataBinding->DataToInspect)
+	{
+		CollectDebugDataPrepareAction(DataToDebug);
+	}
 }
 
-bool FPCGComputeGraphElement::StoreDataForInspection(FPCGComputeGraphContext* InContext) const
+void FPCGComputeGraphElement::ExecuteDebugDraw(FPCGComputeGraphContext* InContext) const
 {
-	if (!InContext || !InContext->DataBinding || InContext->DataBinding->DataToInspect.IsEmpty() || !InContext->ExecutionSource.IsValid())
+	if (!InContext || !InContext->DataBinding)
 	{
-		return true;
+		return;
 	}
 
 	const FPCGDataVisualizationRegistry& DataVisRegistry = FPCGModule::GetConstPCGDataVisualizationRegistry();
 
-	bool bAllReady = true;
-
-	// First ensure all are read back.
-	for (FPCGDataToDebug& DataToInspect : InContext->DataBinding->DataToInspect)
+	for (int Index = InContext->DataBinding->DataToDebug.Num() - 1; Index >= 0; --Index)
 	{
-		const UPCGProxyForGPUData* Proxy = DataToInspect.Proxy.Get();
-		const UPCGSettings* ProducerSettings = DataToInspect.ProducerSettings.Get();
+		FPCGDataToDebug& DataToDebug = InContext->DataBinding->DataToDebug[Index];
+		const UPCGData* DataToDisplay = nullptr;
 
-		if (Proxy && ProducerSettings)
+		if (const UPCGProxyForGPUData* Proxy = Cast<UPCGProxyForGPUData>(DataToDebug.Data.Get()))
 		{
 			UPCGProxyForGPUData::FReadbackResult Readback = Proxy->GetCPUData(InContext);
-			if (!Readback.bComplete)
+			ensure(Readback.bComplete);
+
+			DataToDisplay = Readback.TaggedData.Data.Get();
+		}
+		else if (const UPCGTextureData* Data = Cast<UPCGTextureData>(DataToDebug.Data.Get()))
+		{
+			DataToDisplay = DataToDebug.DataPendingInit;
+		}
+
+		InContext->DataBinding->DataToDebug.RemoveAtSwap(Index);
+
+		const UPCGSettings* ProducerSettings = DataToDebug.ProducerSettings.Get();
+
+		if (DataToDisplay && ProducerSettings)
+		{
+			if (const IPCGDataVisualization* DataVis = DataVisRegistry.GetDataVisualization(DataToDisplay->GetClass()))
 			{
-				bAllReady = false;
+				DataVis->ExecuteDebugDisplay(InContext, ProducerSettings, DataToDisplay, InContext->GetTargetActor(nullptr));
 			}
 		}
 	}
+}
 
-	if (!bAllReady)
+void FPCGComputeGraphElement::StoreDataForInspection(FPCGComputeGraphContext* InContext) const
+{
+	if (!InContext || !InContext->DataBinding || !InContext->ExecutionSource.IsValid())
 	{
-		return false;
+		return;
 	}
 
-	// Now collect all data into a collection and store it.
+	// Collect all data into a collection and store it.
 	TMap<const UPCGSettings*, FPCGDataCollection> SettingsToDataCollection;
 
 	for (FPCGDataToDebug& DataToInspect : InContext->DataBinding->DataToInspect)
 	{
-		const UPCGProxyForGPUData* Proxy = DataToInspect.Proxy.Get();
-		const UPCGSettings* ProducerSettings = DataToInspect.ProducerSettings.Get();
-
-		if (Proxy && ProducerSettings)
+		if (const UPCGSettings* ProducerSettings = DataToInspect.ProducerSettings.Get())
 		{
-			UPCGProxyForGPUData::FReadbackResult Readback = Proxy->GetCPUData(InContext);
-			if (ensure(Readback.bComplete && Readback.TaggedData.Data))
+			FPCGDataCollection& DataCollection = SettingsToDataCollection.FindOrAdd(ProducerSettings);
+
+			if (const UPCGProxyForGPUData* ProxyData = Cast<UPCGProxyForGPUData>(DataToInspect.Data.Get()))
 			{
-				FPCGDataCollection& DataCollection = SettingsToDataCollection.FindOrAdd(ProducerSettings);
+				UPCGProxyForGPUData::FReadbackResult Readback = ProxyData->GetCPUData(InContext);
 				FPCGTaggedData TaggedData = Readback.TaggedData;
 				TaggedData.Pin = DataToInspect.PinLabel;
+				TaggedData.Tags.Append(DataToInspect.AdditionalTags);
+				DataCollection.TaggedData.Add(TaggedData);
+			}
+			else if (const UPCGTextureData* TextureData = Cast<UPCGTextureData>(DataToInspect.Data.Get()))
+			{
+				FPCGTaggedData TaggedData;
+				TaggedData.Data = DataToInspect.DataPendingInit;
+				TaggedData.Pin = DataToInspect.PinLabel;
+				TaggedData.Tags = DataToInspect.AdditionalTags;
 				DataCollection.TaggedData.Add(TaggedData);
 			}
 		}
@@ -763,8 +800,6 @@ bool FPCGComputeGraphElement::StoreDataForInspection(FPCGComputeGraphContext* In
 	}
 
 	InContext->DataBinding->DataToInspect.Empty();
-
-	return bAllReady;
 }
 
 void FPCGComputeGraphElement::LogCompilationMessages(FPCGComputeGraphContext* InContext) const
diff --git a/Engine/Plugins/PCG/Source/PCG/Private/Compute/Elements/PCGComputeGraphElement.h b/Engine/Plugins/PCG/Source/PCG/Private/Compute/Elements/PCGComputeGraphElement.h
--- a/Engine/Plugins/PCG/Source/PCG/Private/Compute/Elements/PCGComputeGraphElement.h
+++ b/Engine/Plugins/PCG/Source/PCG/Private/Compute/Elements/PCGComputeGraphElement.h
@@ -82,6 +82,9 @@ public:
 	EPCGComputeGraphExecutionPhase ExecutionSubPhase = EPCGComputeGraphExecutionPhase::GetComputeGraph;
 
 	std::atomic<bool> bGraphValid = false;
+
+	using FDebugDataPrepareAction = TFunction<bool(FPCGComputeGraphContext*)>;
+	TArray<FDebugDataPrepareAction> DebugDataPrepareActions;
 };
 
 /** Executes a CF graph. Created by the compiler when collapsing GPU nodes rather than by a settings/node. */
@@ -122,11 +125,13 @@ protected:
 	void ResetAsyncOperations(FPCGContext* InContext) const;
 
 #if WITH_EDITOR
+	void CollectDebugDataPrepareActions(FPCGComputeGraphContext* InContext) const;
+
 	// TODO: Debug draw has a special path here because the compute graph element represents a set of original elements which do not themselves
 	// execute. Review if this can be reconciled or unified with the normal debug draw path.
-	bool ExecuteDebugDraw(FPCGComputeGraphContext* InContext) const;
+	void ExecuteDebugDraw(FPCGComputeGraphContext* InContext) const;
 
-	bool StoreDataForInspection(FPCGComputeGraphContext* InContext) const;
+	void StoreDataForInspection(FPCGComputeGraphContext* InContext) const;
 
 	void LogCompilationMessages(FPCGComputeGraphContext* InContext) const;
 #endif
diff --git a/Engine/Plugins/PCG/Source/PCG/Private/Compute/Elements/PCGCustomHLSL.cpp b/Engine/Plugins/PCG/Source/PCG/Private/Compute/Elements/PCGCustomHLSL.cpp
--- a/Engine/Plugins/PCG/Source/PCG/Private/Compute/Elements/PCGCustomHLSL.cpp
+++ b/Engine/Plugins/PCG/Source/PCG/Private/Compute/Elements/PCGCustomHLSL.cpp
@@ -172,8 +172,8 @@ void UPCGCustomHLSLSettings::ApplyPreconfiguredSettings(const FPCGPreConfiguredS
 		{
 			KernelType = EPCGKernelType(PreconfiguredInfo.PreconfiguredIndex);
 
-			// Point Generator doesn't utilize the default input pin, so let's not add it by default.
-			if (KernelType == EPCGKernelType::PointGenerator)
+			// Generators don't utilize the default input pin, so let's not add it by default.
+			if (IsGeneratorKernel())
 			{
 				InputPins.Empty();
 			}
@@ -231,20 +231,6 @@ const UPCGPin* UPCGCustomHLSLSettings::GetFirstInputPin() const
 	return nullptr;
 }
 
-const UPCGPin* UPCGCustomHLSLSettings::GetPointProcessingInputPin() const
-{
-	if (const UPCGNode* Node = Cast<UPCGNode>(GetOuter()))
-	{
-		const UPCGPin* FirstPin = !Node->GetInputPins().IsEmpty() ? Node->GetInputPins()[0] : nullptr;
-		if (FirstPin && FirstPin->Properties.AllowedTypes == EPCGDataType::Point)
-		{
-			return FirstPin;
-		}
-	}
-
-	return nullptr;
-}
-
 const UPCGPin* UPCGCustomHLSLSettings::GetFirstOutputPin() const
 {
 	if (const UPCGNode* Node = Cast<UPCGNode>(GetOuter()))
@@ -255,20 +241,6 @@ const UPCGPin* UPCGCustomHLSLSettings::GetFirstOutputPin() const
 	return nullptr;
 }
 
-const UPCGPin* UPCGCustomHLSLSettings::GetFirstPointOutputPin() const
-{
-	if (const UPCGNode* Node = Cast<UPCGNode>(GetOuter()))
-	{
-		const UPCGPin* FirstPin = !Node->GetOutputPins().IsEmpty() ? Node->GetOutputPins()[0] : nullptr;
-		if (FirstPin && FirstPin->Properties.AllowedTypes == EPCGDataType::Point)
-		{
-			return FirstPin;
-		}
-	}
-
-	return nullptr;
-}
-
 #if WITH_EDITOR
 FString UPCGCustomHLSLSettings::GetDeclarationsText() const
 {
@@ -374,31 +346,40 @@ void UPCGCustomHLSLSettings::UpdateInputDeclarations()
 
 	// Constants category
 	{
-		if (KernelType == EPCGKernelType::PointGenerator)
+		if (IsGeneratorKernel())
 		{
 			InputDeclarations += TEXT("/*** INPUT CONSTANTS ***/\n\n");
+		}
+
+		if (KernelType == EPCGKernelType::PointGenerator)
+		{
+			// @todo_pcg: NumPoints should also be ElementCount.
 			InputDeclarations += FString::Format(TEXT("const uint NumPoints = {0};\n\n"), { PointCount });
 		}
+		else if (KernelType == EPCGKernelType::TextureGenerator)
+		{
+			InputDeclarations += FString::Format(TEXT("const uint2 ElementCount = uint2({0}, {1});\n\n"), { TextureSize.X, TextureSize.Y });
+		}
 
 		InputDeclarations += TEXT("/*** INPUT PER-THREAD CONSTANTS ***/\n\n");
 		InputDeclarations += TEXT("const uint ThreadIndex;\n");
 
-		if (KernelType == EPCGKernelType::PointProcessor)
+		if (IsProcessorKernel())
 		{
-			const UPCGPin* PointProcessingInputPin = GetPointProcessingInputPin();
-			const UPCGPin* PointProcessingOutputPin = GetFirstPointOutputPin();
+			const UPCGPin* InputPin = GetFirstInputPin();
+			const UPCGPin* OutputPin = GetFirstOutputPin();
 
-			if (PointProcessingInputPin && PointProcessingOutputPin)
+			if (InputPin && OutputPin)
 			{
 				InputDeclarations += FString::Format(TEXT(
 					"const uint {0}_DataIndex;\n"
 					"const uint {1}_DataIndex;\n"),
-					{ PointProcessingInputPin->Properties.Label.ToString(),  PointProcessingOutputPin->Properties.Label.ToString() });
+					{ InputPin->Properties.Label.ToString(),  OutputPin->Properties.Label.ToString() });
 			}
 		}
-		else if (KernelType == EPCGKernelType::PointGenerator)
+		else if (IsGeneratorKernel())
 		{
-			if (const UPCGPin* PointProcessingOutputPin = GetFirstPointOutputPin())
+			if (const UPCGPin* PointProcessingOutputPin = GetFirstOutputPin())
 			{
 				InputDeclarations += FString::Format(
 					TEXT("const uint {0}_DataIndex;\n"),
@@ -406,12 +387,20 @@ void UPCGCustomHLSLSettings::UpdateInputDeclarations()
 			}
 		}
 
-		InputDeclarations += TEXT("const uint ElementIndex;\n\n");
+		if (IsPointKernel())
+		{
+			InputDeclarations += TEXT("const uint ElementIndex;\n");
+		}
+		else if (IsTextureKernel())
+		{
+			InputDeclarations += TEXT("const uint2 ElementIndex;\n");
+		}
+
+		InputDeclarations += TEXT("\n");
 	}
 
 	TArray<FString> DataCollectionDataPins;
 	TArray<FString> PointDataPins;
-	TArray<FString> SplineDataPins;
 	TArray<FString> LandscapeDataPins;
 	TArray<FString> TextureDataPins;
 	TArray<FString> VirtualTextureDataPins;
@@ -430,17 +419,12 @@ void UPCGCustomHLSLSettings::UpdateInputDeclarations()
 			PointDataPins.Add(Pin.Label.ToString());
 		}
 
-		if (!!(Pin.AllowedTypes & EPCGDataType::Spline))
-		{
-			SplineDataPins.Add(Pin.Label.ToString());
-		}
-
 		if (!!(Pin.AllowedTypes & EPCGDataType::Landscape))
 		{
 			LandscapeDataPins.Add(Pin.Label.ToString());
 		}
 
-		if (!!(Pin.AllowedTypes & EPCGDataType::Texture))
+		if (!!(Pin.AllowedTypes & EPCGDataType::BaseTexture))
 		{
 			TextureDataPins.Add(Pin.Label.ToString());
 		}
@@ -543,10 +527,14 @@ void UPCGCustomHLSLSettings::UpdateInputDeclarations()
 		}
 
 		InputDeclarations += FString::Format(TEXT(
-			"float2 {0}_GetTexCoords(float2 WorldPos, float2 Min, float2 Max);\n"
+			"uint {0}_GetNumData();\n"
+			"uint2 {0}_GetNumElements(uint DataIndex);\n"
+			"// Computes a texture coordinate as if WorldPos lies in a UV plane defined by WorldMin and WorldMax.\n"
+			"float2 {0}_GetTexCoords(float2 WorldPos, float2 WorldMin, float2 WorldMax);\n"
 			"float4 {0}_Sample(uint DataIndex, float2 TexCoords);\n"
 			"// Computes sample coordinates of the WorldPos relative to the texture data's bounds.\n"
-			"float4 {0}_SampleWorldPos(uint DataIndex, float2 WorldPos);\n"),
+			"float4 {0}_SampleWorldPos(uint DataIndex, float2 WorldPos);\n"
+			"float4 {0}_Load(uint DataIndex, uint2 ElementIndex);\n"),
 			{ bMultiPin ? PCGHLSLElement::PinDeclTemplateStr : TextureDataPins[0] });
 
 		InputDeclarations += TEXT("\n");
@@ -642,7 +630,7 @@ void UPCGCustomHLSLSettings::UpdateOutputDeclarations()
 
 	TArray<FString> DataCollectionDataPins;
 	TArray<FString> PointDataPins;
-	TArray<FString> SplineDataPins;
+	TArray<FString> TextureDataPins;
 	TArray<FString> RawBufferDataPins;
 
 	for (const FPCGPinProperties& Pin : OutputPinProperties())
@@ -657,9 +645,9 @@ void UPCGCustomHLSLSettings::UpdateOutputDeclarations()
 			PointDataPins.Add(Pin.Label.ToString());
 		}
 
-		if (!!(Pin.AllowedTypes & EPCGDataType::Spline))
+		if (!!(Pin.AllowedTypes & EPCGDataType::BaseTexture))
 		{
-			SplineDataPins.Add(Pin.Label.ToString());
+			TextureDataPins.Add(Pin.Label.ToString());
 		}
 	}
 
@@ -720,6 +708,26 @@ void UPCGCustomHLSLSettings::UpdateOutputDeclarations()
 		OutputDeclarations += TEXT("\n");
 	}
 
+	if (!TextureDataPins.IsEmpty())
+	{
+		OutputDeclarations += TEXT("/*** OUTPUT TEXTURE DATA FUNCTIONS ***/\n\n");
+
+		const bool bMultiPin = TextureDataPins.Num() > 1;
+
+		if (bMultiPin)
+		{
+			OutputDeclarations += TEXT("// Valid pins: ") + FString::Join(TextureDataPins, TEXT(", ")) + TEXT("\n\n");
+		}
+
+		OutputDeclarations += FString::Format(TEXT(
+			"uint {0}_GetNumData();\n"
+			"uint2 {0}_GetNumElements(uint DataIndex);\n"
+			"void {0}_Store(uint DataIndex, uint2 ElementIndex, float4 Value);\n"),
+			{ bMultiPin ? PCGHLSLElement::PinDeclTemplateStr : TextureDataPins[0] });
+
+		OutputDeclarations += TEXT("\n");
+	}
+
 	if (!RawBufferDataPins.IsEmpty())
 	{
 		OutputDeclarations += TEXT("/*** OUTPUT BYTE ADDRESS BUFFER DATA FUNCTIONS ***/\n\n");
@@ -755,32 +763,67 @@ void UPCGCustomHLSLSettings::UpdateHelperDeclarations()
 
 		// Get thread data - useful in all kernel types for secondary pins.
 		{
-			TArray<FString> PinNames;
+			TArray<FString> DataCollectionPinNames;
+			TArray<FString> TextureDataPinNames;
 
 			for (const FPCGPinProperties& Properties : InputPinProperties())
 			{
-				PinNames.Add(Properties.Label.ToString());
+				if (PCGComputeHelpers::IsTypeAllowedInDataCollection(Properties.AllowedTypes))
+				{
+					DataCollectionPinNames.Add(Properties.Label.ToString());
+				}
+
+				if (!!(Properties.AllowedTypes & EPCGDataType::BaseTexture))
+				{
+					TextureDataPinNames.Add(Properties.Label.ToString());
+				}
 			}
 
 			for (const FPCGPinProperties& Properties : OutputPinProperties())
 			{
-				PinNames.Add(Properties.Label.ToString());
+				if (PCGComputeHelpers::IsTypeAllowedInDataCollection(Properties.AllowedTypes))
+				{
+					DataCollectionPinNames.Add(Properties.Label.ToString());
+				}
+
+				if (!!(Properties.AllowedTypes & EPCGDataType::BaseTexture))
+				{
+					TextureDataPinNames.Add(Properties.Label.ToString());
+				}
 			}
 
-			if (!PinNames.IsEmpty())
+			if (!DataCollectionPinNames.IsEmpty())
 			{
 				HelperDeclarations += TEXT("\n// Returns false if thread has no data to operate on.\n");
 
-				const bool bMultiPin = PinNames.Num() > 1;
+				const bool bMultiPin = DataCollectionPinNames.Num() > 1;
 
 				if (bMultiPin)
 				{
-					HelperDeclarations += TEXT("// Valid pins: ") + FString::Join(PinNames, TEXT(", ")) + TEXT("\n");
+					HelperDeclarations += TEXT("// Valid pins: ") + FString::Join(DataCollectionPinNames, TEXT(", ")) + TEXT("\n");
 				}
 
 				HelperDeclarations += FString::Format(
 					TEXT("bool {0}_GetThreadData(uint ThreadIndex, out uint OutDataIndex, out uint OutElementIndex);\n"),
-					{ bMultiPin ? PCGHLSLElement::PinDeclTemplateStr : PinNames[0] });
+					{ bMultiPin ? PCGHLSLElement::PinDeclTemplateStr : DataCollectionPinNames[0] });
+
+				HelperDeclarations += TEXT("\n");
+			}
+
+			if (!TextureDataPinNames.IsEmpty())
+			{
+				HelperDeclarations += TEXT("\n// Returns false if thread has no data to operate on.\n");
+
+				const bool bMultiPin = TextureDataPinNames.Num() > 1;
+
+				if (bMultiPin)
+				{
+					HelperDeclarations += TEXT("// Valid pins: ") + FString::Join(TextureDataPinNames, TEXT(", ")) + TEXT("\n");
+				}
+
+				HelperDeclarations += FString::Format(
+					TEXT("bool {0}_GetThreadData(uint ThreadIndex, out uint OutDataIndex, out uint2 OutElementIndex);\n"),
+					{ bMultiPin ? PCGHLSLElement::PinDeclTemplateStr : TextureDataPinNames[0] });
 
 				HelperDeclarations += TEXT("\n");
 			}
@@ -806,7 +849,7 @@ void UPCGCustomHLSLSettings::UpdateHelperDeclarations()
 		{
 			HelperDeclarations += FString::Format(TEXT(
 				"\n"
-				"// Writes floats to the debug buffer array, which will be readback and logged in the console for inspection."
+				"// Writes floats to the debug buffer array, which will be readback and logged in the console for inspection.\n"
 				"void WriteDebugValue(uint Index, float Value); // Index in [0, {0}] (set from 'Debug Buffer Size' property)\n"),
 				{ DebugBufferSize - 1 });
 		}
@@ -828,6 +871,23 @@ void UPCGCustomHLSLSettings::UpdatePinSettings()
 			Properties.AllowedTypes = EPCGDataType::Point;
 		}
 
+		// Allow kernel type to drive the first pin type.
+		if (PinIndex == 0)
+		{
+			if (KernelType == EPCGKernelType::PointProcessor)
+			{
+				Properties.AllowedTypes = EPCGDataType::Point;
+			}
+			else if (KernelType == EPCGKernelType::TextureProcessor)
+			{
+				Properties.AllowedTypes = EPCGDataType::BaseTexture;
+
+				// Texture kernels are single-data for now.
+				Properties.bAllowMultipleData = false;
+				Properties.bAllowEditMultipleData = false;
+			}
+		}
+
 		if (!!(Properties.AllowedTypes & EPCGDataType::Landscape))
 		{
 			// Don't allow multiple data on this pin because we do not support a dynamic number of landscapes bound to a
@@ -860,7 +920,28 @@ void UPCGCustomHLSLSettings::UpdatePinSettings()
 		}
 
 		// Only allow editing the initialization mode if it's not driven by the kernel type.
-		const bool bInitModeDrivenByKernel = PinIndex == 0 && (KernelType == EPCGKernelType::PointProcessor || KernelType == EPCGKernelType::PointGenerator);
+		bool bInitModeDrivenByKernel = false;
+		if (PinIndex == 0)
+		{
+			if (KernelType == EPCGKernelType::PointProcessor || KernelType == EPCGKernelType::PointGenerator)
+			{
+				bInitModeDrivenByKernel = true;
+				Properties.AllowedTypes = EPCGDataType::Point;
+			}
+			else if (KernelType == EPCGKernelType::TextureProcessor || KernelType == EPCGKernelType::TextureGenerator)
+			{
+				bInitModeDrivenByKernel = true;
+				Properties.AllowedTypes = EPCGDataType::BaseTexture;
+			}
+		}
+
+		if (!!(Properties.AllowedTypes & EPCGDataType::BaseTexture))
+		{
+			// Texture outputs are single-data for now.
+			Properties.bAllowMultipleData = false;
+			Properties.bAllowEditMultipleData = false;
+		}
+
 		Properties.PropertiesGPU.bAllowEditInitMode = !bInitModeDrivenByKernel;
 		Properties.PropertiesGPU.bMultipleInitPins = Properties.PropertiesGPU.PinsToInititalizeFrom.Num() > 1;
 
diff --git a/Engine/Plugins/PCG/Source/PCG/Private/Compute/Elements/PCGCustomHLSL.h b/Engine/Plugins/PCG/Source/PCG/Private/Compute/Elements/PCGCustomHLSL.h
--- a/Engine/Plugins/PCG/Source/PCG/Private/Compute/Elements/PCGCustomHLSL.h
+++ b/Engine/Plugins/PCG/Source/PCG/Private/Compute/Elements/PCGCustomHLSL.h
@@ -20,6 +20,8 @@ enum class EPCGKernelType : uint8
 {
 	PointProcessor UMETA(Tooltip = "Kernel executes on each point in first input pin."),
 	PointGenerator UMETA(Tooltip = "Kernel executes for fixed number of points, configurable on node."),
+	TextureProcessor UMETA(Tooltip = "Kernel executes on each texel in the first input pin."),
+	TextureGenerator UMETA(Tooltip = "Kernel executes for each texel in a fixed size texture, configurable on node."),
 	Custom UMETA(Tooltip = "Execution thread counts and output buffer sizes configurable on node. All data read/write indices must be manually bounds checked."),
 };
 
@@ -120,13 +122,16 @@ protected:
 	const UPCGPin* GetInputPin(FName Label) const;
 	const UPCGPin* GetOutputPin(FName Label) const;
 	const UPCGPin* GetFirstInputPin() const;
-	const UPCGPin* GetPointProcessingInputPin() const;
 	const UPCGPin* GetFirstOutputPin() const;
-	const UPCGPin* GetFirstPointOutputPin() const;
 
 	/** Will the ThreadCountMultiplier value be applied when calculating the dispatch thread count. */
 	bool IsThreadCountMultiplierInUse() const { return KernelType == EPCGKernelType::Custom && DispatchThreadCount != EPCGDispatchThreadCount::Fixed; }
 
+	bool IsProcessorKernel() const { return KernelType == EPCGKernelType::PointProcessor || KernelType == EPCGKernelType::TextureProcessor; }
+	bool IsGeneratorKernel() const { return KernelType == EPCGKernelType::PointGenerator || KernelType == EPCGKernelType::TextureGenerator; }
+	bool IsTextureKernel() const { return KernelType == EPCGKernelType::TextureProcessor || KernelType == EPCGKernelType::TextureGenerator; }
+	bool IsPointKernel() const { return KernelType == EPCGKernelType::PointProcessor || KernelType == EPCGKernelType::PointGenerator; }
+
 protected:
 	UPROPERTY(EditAnywhere, Category = "Settings")
 	EPCGKernelType KernelType = EPCGKernelType::PointProcessor;
@@ -134,6 +139,9 @@ protected:
 	UPROPERTY(EditAnywhere, DisplayName = "Number of Points", Category = "Settings", meta = (EditCondition = "KernelType == EPCGKernelType::PointGenerator", EditConditionHides))
 	int PointCount = 256;
 
+	UPROPERTY(EditAnywhere, DisplayName = "Element Count", Category = "Settings", meta = (EditCondition = "KernelType == EPCGKernelType::TextureGenerator", EditConditionHides))
+	FIntPoint TextureSize = FIntPoint(64, 64);
+
 	UPROPERTY(EditAnywhere, Category = "Settings|Thread Count", meta = (EditCondition = "KernelType == EPCGKernelType::Custom", EditConditionHides))
 	EPCGDispatchThreadCount DispatchThreadCount = EPCGDispatchThreadCount::FromFirstOutputPin;
 
@@ -169,7 +177,7 @@ protected:
 	TArray<TObjectPtr<UComputeSource>> AdditionalSources;
 #endif
 
-	/** Mute uninitialized data warnings. */
+	/** Mute uninitialized data errors. */
 	UPROPERTY(EditAnywhere, Category = "Settings")
 	bool bMuteUnwrittenPinDataErrors = false;
 
diff --git a/Engine/Plugins/PCG/Source/PCG/Private/Compute/Elements/PCGCustomHLSLKernel.cpp b/Engine/Plugins/PCG/Source/PCG/Private/Compute/Elements/PCGCustomHLSLKernel.cpp
--- a/Engine/Plugins/PCG/Source/PCG/Private/Compute/Elements/PCGCustomHLSLKernel.cpp
+++ b/Engine/Plugins/PCG/Source/PCG/Private/Compute/Elements/PCGCustomHLSLKernel.cpp
@@ -29,6 +29,7 @@ namespace PCGCustomHLSLKernel
 	constexpr TCHAR AttributeFunctionSetKeyword[] = { TEXT("Set") };
 	constexpr TCHAR CopyElementFunctionKeyword[] = { TEXT("CopyElementFrom") };
 	constexpr TCHAR InitializeFunctionKeyword[] = { TEXT("Initialize") };
+	constexpr TCHAR StoreFunctionKeyword[] = { TEXT("Store") };
 
 #if WITH_EDITOR
 	enum class EParseState : uint8
@@ -246,11 +247,11 @@ bool UPCGCustomHLSLKernel::ComputeOutputBindingDataDesc(const UPCGComputeGraph*
 	const FName FirstOutputPinLabel = CustomHLSLSettings->OutputPins[0].Label;
 
 	// The primary output pin follows any rules prescribed by kernel type.
-	if (InOutputPinLabel == FirstOutputPinLabel && CustomHLSLSettings->KernelType == EPCGKernelType::PointProcessor)
+	if (InOutputPinLabel == FirstOutputPinLabel && CustomHLSLSettings->IsProcessorKernel())
 	{
 		FPCGDataCollectionDesc InputDesc;
 
-		if (const FPCGPinProperties* FirstInputPinProps = GetFirstPointInputPin())
+		if (const FPCGPinProperties* FirstInputPinProps = GetFirstInputPin())
 		{
 			const FPCGKernelPin FirstKernelPin(KernelIndex, FirstInputPinProps->Label, /*bIsInput=*/true);
 
@@ -268,6 +269,13 @@ bool UPCGCustomHLSLKernel::ComputeOutputBindingDataDesc(const UPCGComputeGraph*
 		// Generators always produce a single point data with known point count.
 		OutputDataDesc.DataDescs.Emplace(EPCGDataType::Point, CustomHLSLSettings->PointCount);
 	}
+	else if (InOutputPinLabel == FirstOutputPinLabel && CustomHLSLSettings->KernelType == EPCGKernelType::TextureGenerator)
+	{
+		// Generators always produce a single texture data with known size.
+		FPCGDataDesc& NewDataDesc = OutputDataDesc.DataDescs.Emplace_GetRef();
+		NewDataDesc.Type = EPCGDataType::BaseTexture;
+		NewDataDesc.ElementCount2D = CustomHLSLSettings->TextureSize;
+	}
 	else if (ensure(OutputPinProperties))
 	{
 		ComputeDataDescFromPinProperties(*OutputPinProperties, MakeArrayView(CustomHLSLSettings->InputPins), InBinding, OutputDataDesc);
@@ -435,31 +443,36 @@ FString UPCGCustomHLSLKernel::GetCookedSource(FPCGGPUCompilationContext& InOutCo
 
 	for (const FPCGPinPropertiesGPU& PinProps : CustomHLSLSettings->OutputPins)
 	{
-		SetAsExecuted += FString::Format(TEXT("    if (all(GroupId == 0u) && GroupIndex == 0) {0}_SetAsExecutedInternal();\n"), { PinProps.Label.ToString() });
+		if (PCGComputeHelpers::IsTypeAllowedInDataCollection(PinProps.AllowedTypes))
+		{
+			SetAsExecuted += FString::Format(TEXT("    if (all(GroupId == 0u) && GroupIndex == 0) {0}_SetAsExecutedInternal();\n"), { PinProps.Label.ToString() });
+		}
 	}
 
 	// Per-kernel-type preamble. Set up shader inputs and initialize output data.
 	FString KernelSpecificPreamble = TEXT("    // Kernel preamble\n");
 
-	auto AddThreadInfoForPin = [&KernelSpecificPreamble](FName PinLabel)
-		{
-			KernelSpecificPreamble += FString::Format(TEXT(
-				"    uint {0}_DataIndex;\n"
-				"    if (!{0}_GetThreadData(ThreadIndex, {0}_DataIndex, ElementIndex)) return;\n"),
-				{ PinLabel.ToString() });
-		};
+	auto AddThreadInfoForPin = [&KernelSpecificPreamble](const FPCGPinProperties* InPin)
+	{
+		check(InPin);
+
+		KernelSpecificPreamble += FString::Format(TEXT(
+			"    uint {0}_DataIndex;\n"
+			"    if (!{0}_GetThreadData(ThreadIndex, {0}_DataIndex, ElementIndex)) return;\n"),
+			{ InPin->Label.ToString() });
+	};
 
 	if (CustomHLSLSettings->KernelType == EPCGKernelType::PointProcessor)
 	{
-		const FPCGPinProperties* InputPin = GetFirstPointInputPin();
-		const FPCGPinPropertiesGPU* OutputPin = GetFirstPointOutputPin();
+		const FPCGPinProperties* InputPin = GetFirstInputPin();
+		const FPCGPinPropertiesGPU* OutputPin = GetFirstOutputPin();
 
 		if (InputPin && OutputPin)
 		{
 			KernelSpecificPreamble += TEXT("    uint ElementIndex; // Assumption - element index identical in input and output data.\n");
 
-			AddThreadInfoForPin(InputPin->Label);
-			AddThreadInfoForPin(OutputPin->Label);
+			AddThreadInfoForPin(InputPin);
+			AddThreadInfoForPin(OutputPin);
 
 			// If input point is invalid, mark output point as invalid and abort.
 			KernelSpecificPreamble += FString::Format(TEXT(
@@ -474,8 +487,8 @@ FString UPCGCustomHLSLKernel::GetCookedSource(FPCGGPUCompilationContext& InOutCo
 			// TODO pass in IDs of attributes that are actually present.
 			KernelSpecificPreamble += FString::Format(TEXT(
 				"\n"
-				"// Point processor always initializes outputs by copying input data elements.\n"
-				"PCG_COPY_ALL_ATTRIBUTES_TO_OUTPUT({1}, {0}, {1}_DataIndex, ElementIndex, {0}_DataIndex, ElementIndex);\n"),
+				"    // Point processor always initializes outputs by copying input data elements.\n"
+				"    PCG_COPY_ALL_ATTRIBUTES_TO_OUTPUT({1}, {0}, {1}_DataIndex, ElementIndex, {0}_DataIndex, ElementIndex);\n"),
 				{ InputPin->Label.ToString(), OutputPin->Label.ToString() });
 		}
 	}
@@ -483,19 +496,56 @@ FString UPCGCustomHLSLKernel::GetCookedSource(FPCGGPUCompilationContext& InOutCo
 	{
 		KernelSpecificPreamble += FString::Format(TEXT("    const uint NumPoints = {0};\n"), { CustomHLSLSettings->PointCount });
 
-		if (const FPCGPinPropertiesGPU* OutputPin = GetFirstPointOutputPin())
+		if (const FPCGPinPropertiesGPU* OutputPin = GetFirstOutputPin())
 		{
 			KernelSpecificPreamble += TEXT("uint ElementIndex; // Assumption - element index identical in input and output data.\n");
 
-			AddThreadInfoForPin(OutputPin->Label);
+			AddThreadInfoForPin(OutputPin);
 
 			KernelSpecificPreamble += FString::Format(TEXT(
-				"    \n"
+				"\n"
 				"    // Initialize all values to defaults for output pin {0}\n"
 				"    {0}_InitializePoint({0}_DataIndex, ElementIndex);\n"),
 				{ OutputPin->Label.ToString() });
 		}
 	}
+	else if (CustomHLSLSettings->KernelType == EPCGKernelType::TextureProcessor)
+	{
+		const FPCGPinProperties* InputPin = GetFirstInputPin();
+		const FPCGPinPropertiesGPU* OutputPin = GetFirstOutputPin();
+
+		if (InputPin && OutputPin)
+		{
+			KernelSpecificPreamble += TEXT("    uint2 ElementIndex; // Assumption - texel index identical in input and output data.\n");
+
+			AddThreadInfoForPin(InputPin);
+			AddThreadInfoForPin(OutputPin);
+
+			// Automatically copy the input texture across.
+			KernelSpecificPreamble += FString::Format(TEXT(
+				"\n"
+				"    // Texture processor always initializes outputs by copying the input texture.\n"
+				"    {1}_Store({1}_DataIndex, ElementIndex, {0}_Load({0}_DataIndex, ElementIndex));\n"),
+				{ InputPin->Label.ToString(), OutputPin->Label.ToString() });
+		}
+	}
+	else if (CustomHLSLSettings->KernelType == EPCGKernelType::TextureGenerator)
+	{
+		KernelSpecificPreamble += FString::Format(TEXT("    const uint2 ElementCount = uint2({0}, {1});\n"), { CustomHLSLSettings->TextureSize.X, CustomHLSLSettings->TextureSize.Y });
+
+		if (const FPCGPinPropertiesGPU* OutputPin = GetFirstOutputPin())
+		{
+			KernelSpecificPreamble += TEXT("    uint2 ElementIndex; // Assumption - texel index identical in input and output data.\n");
+
+			AddThreadInfoForPin(OutputPin);
+
+			KernelSpecificPreamble += FString::Format(TEXT(
+				"\n"
+				"    // Zero-initialize for output pin {0}\n"
+				"    {0}_Store({0}_DataIndex, ElementIndex, (float4)0.0f);\n"),
+				{ OutputPin->Label.ToString() });
+		}
+	}
 
 	FString Result;
 
@@ -545,10 +595,15 @@ int UPCGCustomHLSLKernel::ComputeThreadCount(const UPCGDataBinding* InBinding) c
 		// Point generator has fixed thread count.
 		ThreadCount = CustomHLSLSettings->PointCount;
 	}
-	else if (CustomHLSLSettings->KernelType == EPCGKernelType::PointProcessor)
+	else if (CustomHLSLSettings->KernelType == EPCGKernelType::TextureGenerator)
+	{
+		// Texture generator has fixed thread count.
+		ThreadCount = CustomHLSLSettings->TextureSize.X * CustomHLSLSettings->TextureSize.Y;
+	}
+	else if (CustomHLSLSettings->IsProcessorKernel())
 	{
 		// Processing volume depends on data arriving on primary pin.
-		if (const FPCGPinProperties* InputPin = GetFirstPointInputPin())
+		if (const FPCGPinProperties* InputPin = GetFirstInputPin())
 		{
 			ThreadCount = GetElementCountForInputPin(*InputPin, InBinding);
 		}
@@ -560,7 +615,7 @@ int UPCGCustomHLSLKernel::ComputeThreadCount(const UPCGDataBinding* InBinding) c
 			if (const FPCGPinPropertiesGPU* OutputPin = CustomHLSLSettings->OutputPins.IsEmpty() ? nullptr : &CustomHLSLSettings->OutputPins[0])
 			{
 				const FPCGDataCollectionDesc* Desc = InBinding->GetCachedKernelPinDataDesc(this, OutputPin->Label, /*bIsInput=*/false);
-				ThreadCount = Desc ? OutputPin->GetElementCountMultiplier() * Desc->ComputeDataElementCount(EPCGDataType::Any) : 0;
+				ThreadCount = Desc ? OutputPin->GetElementCountMultiplier() * Desc->ComputeDataElementCount(OutputPin->AllowedTypes) : 0;
 			}
 		}
 		else if (CustomHLSLSettings->DispatchThreadCount == EPCGDispatchThreadCount::FromProductOfInputPins)
@@ -714,20 +769,27 @@ void UPCGCustomHLSLKernel::ParseShaderSource()
 	{
 		const FString PinStr = *PinProps.Label.ToString();
 
-		for (const FString& AttributeTypeString : AttributeTypeStrings)
+		if (!!(PinProps.AllowedTypes & (EPCGDataType::Point | EPCGDataType::Param)))
 		{
-			AttributeKeywords.Add(PinStr + "_" + PCGCustomHLSLKernel::AttributeFunctionSetKeyword + AttributeTypeString);
-		}
+			for (const FString& AttributeTypeString : AttributeTypeStrings)
+			{
+				AttributeKeywords.Add(PinStr + "_" + PCGCustomHLSLKernel::AttributeFunctionSetKeyword + AttributeTypeString);
+			}
 
-		for (const FPCGPinProperties& InputPinProps : CustomHLSLSettings->InputPins)
-		{
-			if (InputPinProps.AllowedTypes == PinProps.AllowedTypes)
+			for (const FPCGPinProperties& InputPinProps : CustomHLSLSettings->InputPins)
 			{
-				CopyElementKeywords.Add(PinStr + "_" + PCGCustomHLSLKernel::CopyElementFunctionKeyword + "_" + InputPinProps.Label.ToString());
+				if (InputPinProps.AllowedTypes == PinProps.AllowedTypes)
+				{
+					CopyElementKeywords.Add(PinStr + "_" + PCGCustomHLSLKernel::CopyElementFunctionKeyword + "_" + InputPinProps.Label.ToString());
+				}
 			}
+			
+			InitializeKeywords.Add(PinStr + "_" + PCGCustomHLSLKernel::InitializeFunctionKeyword);
+		}
+		else if (!!(PinProps.AllowedTypes & EPCGDataType::BaseTexture))
+		{
+			InitializeKeywords.Add(PinStr + "_" + PCGCustomHLSLKernel::StoreFunctionKeyword);
 		}
-
-		InitializeKeywords.Add(PinStr + "_" + PCGCustomHLSLKernel::InitializeFunctionKeyword);
 	}
 
 	AdditionalKeywords.Append(AttributeKeywords);
@@ -1109,6 +1171,19 @@ bool UPCGCustomHLSLKernel::PerformStaticValidation()
 				return false;
 			}
 		}
+		else if (bIsFirstInputPin && CustomHLSLSettings->KernelType == EPCGKernelType::TextureProcessor)
+		{
+			if (!(Properties.AllowedTypes & EPCGDataType::BaseTexture))
+			{
+#if PCG_KERNEL_LOGGING_ENABLED
+				StaticLogEntries.Emplace(FText::Format(
+					LOCTEXT("InvalidNonTexturePrimaryInput", "'Texture Processor' nodes require primary input pin to be of type 'Base Texture', but found '{0}'."),
+					FText::FromString(PCGCustomHLSLKernel::GetDataTypeString(Properties.AllowedTypes))),
+					EPCGKernelLogVerbosity::Error);
+#endif
+				return false;
+			}
+		}
 
 		if (!PCGComputeHelpers::IsTypeAllowedAsInput(Properties.AllowedTypes))
 		{
@@ -1134,9 +1209,9 @@ bool UPCGCustomHLSLKernel::PerformStaticValidation()
 			return false;
 		}
 
-		const bool bPinIsDefinedByKernel = bIsFirstOutputPin && (CustomHLSLSettings->KernelType == EPCGKernelType::PointGenerator || CustomHLSLSettings->KernelType == EPCGKernelType::PointProcessor);
+		const bool bPinIsDefinedByKernel = bIsFirstOutputPin && (CustomHLSLSettings->IsPointKernel() || CustomHLSLSettings->IsTextureKernel());
 
-		if (bPinIsDefinedByKernel)
+		if (bIsFirstOutputPin && CustomHLSLSettings->IsPointKernel())
 		{
 			if (Properties.AllowedTypes != EPCGDataType::Point)
 			{
@@ -1149,6 +1224,19 @@ bool UPCGCustomHLSLKernel::PerformStaticValidation()
 				return false;
 			}
 		}
+		else if (bIsFirstInputPin && CustomHLSLSettings->IsTextureKernel())
+		{
+			if (!(Properties.AllowedTypes & EPCGDataType::BaseTexture))
+			{
+#if PCG_KERNEL_LOGGING_ENABLED
+				StaticLogEntries.Emplace(FText::Format(
+					LOCTEXT("InvalidNonTexturePrimaryOutput", "'Texture Processor' and 'Texture Generator' nodes require primary output pin to be of type 'Base Texture', but found '{0}'."),
+					FText::FromString(PCGCustomHLSLKernel::GetDataTypeString(Properties.AllowedTypes))),
+					EPCGKernelLogVerbosity::Error);
+#endif
+				return false;
+			}
+		}
 
 		if (!PCGComputeHelpers::IsTypeAllowedAsOutput(Properties.AllowedTypes))
 		{
@@ -1188,7 +1276,7 @@ bool UPCGCustomHLSLKernel::PerformStaticValidation()
 
 					if (InitPinProps)
 					{
-						if (!PCGComputeHelpers::IsTypeAllowedInDataCollection(InitPinProps->AllowedTypes))
+						if (!PCGComputeHelpers::IsTypeAllowedAsOutput(InitPinProps->AllowedTypes))
 						{
 #if PCG_KERNEL_LOGGING_ENABLED
 							StaticLogEntries.Emplace(FText::Format(
@@ -1245,6 +1333,19 @@ bool UPCGCustomHLSLKernel::PerformStaticValidation()
 						FText::AsNumber(Props.ElementCount),
 						FText::FromName(Properties.Label)),
 						EPCGKernelLogVerbosity::Error);
+#endif
+					return false;
+				}
+
+				if (Props.ElementCount2D.GetMin() < 1)
+				{
+#if PCG_KERNEL_LOGGING_ENABLED
+					StaticLogEntries.Emplace(FText::Format(
+						LOCTEXT("InvalidElementCount", "Invalid fixed element count 2D ({0}, {1}) on output pin '{2}'. Must be greater than 0."),
+						FText::AsNumber(Props.ElementCount2D.X),
+						FText::AsNumber(Props.ElementCount2D.Y),
+						FText::FromName(Properties.Label)),
+						EPCGKernelLogVerbosity::Error);
 #endif
 					return false;
 				}
@@ -1465,7 +1566,7 @@ bool UPCGCustomHLSLKernel::AreAllOutputPinsWritten(TArray<FPCGKernelLogEntry>& O
 	const UPCGCustomHLSLSettings* CustomHLSLSettings = CastChecked<UPCGCustomHLSLSettings>(Settings);
 
 	// Processor/Generator kernels initialize the first output pin data automatically.
-	const bool bSkipFirstPin = CustomHLSLSettings->KernelType == EPCGKernelType::PointProcessor || CustomHLSLSettings->KernelType == EPCGKernelType::PointGenerator;
+	const bool bSkipFirstPin = CustomHLSLSettings->IsProcessorKernel() || CustomHLSLSettings->IsGeneratorKernel();
 
 	for (int32 I = 0; I < CustomHLSLSettings->OutputPins.Num(); ++I)
 	{
@@ -1473,8 +1574,9 @@ bool UPCGCustomHLSLKernel::AreAllOutputPinsWritten(TArray<FPCGKernelLogEntry>& O
 		{
 			continue;
 		}
-		
-		const FString& PinStr = CustomHLSLSettings->OutputPins[I].Label.ToString();
+
+		const FPCGPinPropertiesGPU& PinProps = CustomHLSLSettings->OutputPins[I];
+		const FString& PinStr = PinProps.Label.ToString();
 		bool bInitializedByAnySource = false;
 
 		for (const FPCGCustomHLSLParsedSource& ParsedSource : ParsedSources)
@@ -1698,18 +1800,16 @@ bool UPCGCustomHLSLKernel::AreAttributesValid(FPCGContext* InContext, FText* Out
 	return true;
 }
 
-const FPCGPinProperties* UPCGCustomHLSLKernel::GetFirstPointInputPin() const
+const FPCGPinProperties* UPCGCustomHLSLKernel::GetFirstInputPin() const
 {
-	// Point processing input pin is first pin of type Point.
 	const UPCGCustomHLSLSettings* CustomHLSLSettings = CastChecked<UPCGCustomHLSLSettings>(Settings);
-	return CustomHLSLSettings->InputPins.FindByPredicate([](const FPCGPinProperties& InProps) { return InProps.AllowedTypes == EPCGDataType::Point; });
+	return !CustomHLSLSettings->InputPins.IsEmpty() ? &CustomHLSLSettings->InputPins[0] : nullptr;
 }
 
-const FPCGPinPropertiesGPU* UPCGCustomHLSLKernel::GetFirstPointOutputPin() const
+const FPCGPinPropertiesGPU* UPCGCustomHLSLKernel::GetFirstOutputPin() const
 {
-	// Point processing output pin is first pin of type Point.
 	const UPCGCustomHLSLSettings* CustomHLSLSettings = CastChecked<UPCGCustomHLSLSettings>(Settings);
-	return CustomHLSLSettings->OutputPins.FindByPredicate([](const FPCGPinPropertiesGPU& InProps) { return InProps.AllowedTypes == EPCGDataType::Point; });
+	return !CustomHLSLSettings->OutputPins.IsEmpty() ? &CustomHLSLSettings->OutputPins[0] : nullptr;
 }
 
 int UPCGCustomHLSLKernel::GetElementCountForInputPin(const FPCGPinProperties& InInputPinProps, const UPCGDataBinding* InBinding) const
diff --git a/Engine/Plugins/PCG/Source/PCG/Private/Compute/Elements/PCGCustomHLSLKernel.h b/Engine/Plugins/PCG/Source/PCG/Private/Compute/Elements/PCGCustomHLSLKernel.h
--- a/Engine/Plugins/PCG/Source/PCG/Private/Compute/Elements/PCGCustomHLSLKernel.h
+++ b/Engine/Plugins/PCG/Source/PCG/Private/Compute/Elements/PCGCustomHLSLKernel.h
@@ -183,8 +183,8 @@ protected:
 	bool IsThreadCountMultiplierInUse() const;
 	bool AreAttributesValid(FPCGContext* InContext, FText* OutErrorText) const;
 
-	const FPCGPinProperties* GetFirstPointInputPin() const;
-	const FPCGPinPropertiesGPU* GetFirstPointOutputPin() const;
+	const FPCGPinProperties* GetFirstInputPin() const;
+	const FPCGPinPropertiesGPU* GetFirstOutputPin() const;
 
 	int GetElementCountForInputPin(const FPCGPinProperties& InInputPinProps, const UPCGDataBinding* InBinding) const;
 
diff --git a/Engine/Plugins/PCG/Source/PCG/Private/Compute/PCGComputeCommon.cpp b/Engine/Plugins/PCG/Source/PCG/Private/Compute/PCGComputeCommon.cpp
--- a/Engine/Plugins/PCG/Source/PCG/Private/Compute/PCGComputeCommon.cpp
+++ b/Engine/Plugins/PCG/Source/PCG/Private/Compute/PCGComputeCommon.cpp
@@ -10,12 +10,32 @@
 #include "Compute/PCGPinPropertiesGPU.h"
 #include "Compute/Data/PCGProxyForGPUData.h"
 #include "Data/PCGPointData.h"
+#include "Data/PCGTextureData.h"
 #include "Editor/IPCGEditorModule.h"
 
 #include "DynamicRHI.h"
 #include "RHIStats.h"
 #include "RenderResource.h"
 
+namespace PCGComputeConstants
+{
+	/** PCG data types supported in GPU node inputs. */
+	constexpr EPCGDataType AllowedInputTypes =
+		EPCGDataType::Point
+		| EPCGDataType::Param
+		| EPCGDataType::Landscape
+		| EPCGDataType::BaseTexture
+		| EPCGDataType::VirtualTexture
+		| EPCGDataType::Resource
+		| EPCGDataType::ProxyForGPU;
+
+	/** PCG data types supported in GPU node outputs. */
+	constexpr EPCGDataType AllowedOutputTypes = EPCGDataType::Point | EPCGDataType::Param | EPCGDataType::ProxyForGPU | EPCGDataType::BaseTexture;
+
+	/** PCG data types supported in GPU data collections. */
+	constexpr EPCGDataType AllowedDataCollectionTypes = EPCGDataType::Point | EPCGDataType::Param | EPCGDataType::ProxyForGPU;
+}
+
 namespace PCGComputeHelpers
 {
 	static TAutoConsoleVariable<float> CVarMaxGPUBufferSizeProportion(
@@ -44,6 +64,16 @@ namespace PCGComputeHelpers
 		return 0;
 	}
 
+	FIntPoint GetElementCount2D(const UPCGData* InData)
+	{
+		if (const UPCGBaseTextureData* TextureData = Cast<UPCGBaseTextureData>(InData))
+		{
+			return TextureData->GetTextureSize();
+		}
+
+		return FIntPoint::ZeroValue;
+	}
+
 	bool IsTypeAllowedAsInput(EPCGDataType Type)
 	{
 		return (Type | PCGComputeConstants::AllowedInputTypes) == PCGComputeConstants::AllowedInputTypes;
diff --git a/Engine/Plugins/PCG/Source/PCG/Private/Compute/PCGComputeGraph.cpp b/Engine/Plugins/PCG/Source/PCG/Private/Compute/PCGComputeGraph.cpp
--- a/Engine/Plugins/PCG/Source/PCG/Private/Compute/PCGComputeGraph.cpp
+++ b/Engine/Plugins/PCG/Source/PCG/Private/Compute/PCGComputeGraph.cpp
@@ -209,6 +209,7 @@ void UPCGComputeGraph::DebugLogDataDescriptions(const UPCGDataBinding* InBinding
 			UE_LOG(LogPCG, Display, TEXT("\t\tData Index %d"), DataIndex);
 			UE_LOG(LogPCG, Display, TEXT("\t\t\tType: %s"), *PCGDataTypeEnum->GetDisplayNameTextByValue(static_cast<int64>(DataDesc.Type)).ToString());
 			UE_LOG(LogPCG, Display, TEXT("\t\t\tElementCount: %d"), DataDesc.ElementCount);
+			UE_LOG(LogPCG, Display, TEXT("\t\t\tElementCount2D: (%d, %d)"), DataDesc.ElementCount2D.X, DataDesc.ElementCount2D.Y);
 			UE_LOG(LogPCG, Display, TEXT("\t\t\tAttributes (%d)"), DataDesc.AttributeDescs.Num());
 
 			for (int32 AttributeIndex = 0; AttributeIndex < DataDesc.AttributeDescs.Num(); ++AttributeIndex)
diff --git a/Engine/Plugins/PCG/Source/PCG/Private/Compute/PCGComputeKernel.cpp b/Engine/Plugins/PCG/Source/PCG/Private/Compute/PCGComputeKernel.cpp
--- a/Engine/Plugins/PCG/Source/PCG/Private/Compute/PCGComputeKernel.cpp
+++ b/Engine/Plugins/PCG/Source/PCG/Private/Compute/PCGComputeKernel.cpp
@@ -140,7 +140,10 @@ void UPCGComputeKernel::ComputeDataDescFromPinProperties(const FPCGPinProperties
 			{
 				// Set element count to 0 for now, but we will overwrite it.
 				FPCGDataDesc& DataDesc = OutPinDesc.DataDescs.Emplace_GetRef(OutputPinProps.AllowedTypes, /*SetNumElements=*/0);
+				const bool bDomainIs2D = DataDesc.IsDomain2D();
+
 				int TotalNumElements = 0; // Total number of elements computed for this data index.
+				FIntPoint TotalNumElements2D = FIntPoint::ZeroValue;
 
 				// For each data index, loop over all the pins and create the uber-data.
 				for (int InputPinIndex = 0; InputPinIndex < NumInitPins; ++InputPinIndex)
@@ -156,16 +159,35 @@ void UPCGComputeKernel::ComputeDataDescFromPinProperties(const FPCGPinProperties
 					}
 
 					const FPCGDataDesc& InputDataDesc = InputDesc.DataDescs.IsValidIndex(ClampedDataIndex) ? InputDesc.DataDescs[ClampedDataIndex] : FPCGDataDesc(EPCGDataType::Any, 0);
+					const bool bSourceDomainIs2D = InputDataDesc.IsDomain2D();
 
 					if (Props.ElementCountMode == EPCGElementCountMode::FromInputData)
 					{
 						if (Props.ElementMultiplicity == EPCGElementMultiplicity::Product)
 						{
-							TotalNumElements = FMath::Max(TotalNumElements, 1) * InputDataDesc.ElementCount;
+							if (bDomainIs2D)
+							{
+								TotalNumElements2D = TotalNumElements2D.ComponentMax(FIntPoint(1, 1)) * (bSourceDomainIs2D
+									? InputDataDesc.ElementCount2D.X * InputDataDesc.ElementCount2D.Y
+									: InputDataDesc.ElementCount);
+							}
+							else
+							{
+								TotalNumElements = FMath::Max(TotalNumElements, 1) * (bSourceDomainIs2D
+									? InputDataDesc.ElementCount2D.X * InputDataDesc.ElementCount2D.Y
+									: InputDataDesc.ElementCount);
+							}
 						}
 						else if (Props.ElementMultiplicity == EPCGElementMultiplicity::Sum)
 						{
-							TotalNumElements += InputDataDesc.ElementCount;
+							if (bDomainIs2D)
+							{
+								TotalNumElements2D += bSourceDomainIs2D ? InputDataDesc.ElementCount2D : InputDataDesc.ElementCount;
+							}
+							else
+							{
+								TotalNumElements += bSourceDomainIs2D ? InputDataDesc.ElementCount2D.X * InputDataDesc.ElementCount2D.Y : InputDataDesc.ElementCount;
+							}
 						}
 						else
 						{
@@ -175,6 +197,7 @@ void UPCGComputeKernel::ComputeDataDescFromPinProperties(const FPCGPinProperties
 					else if (Props.ElementCountMode == EPCGElementCountMode::Fixed)
 					{
 						TotalNumElements += Props.ElementCount;
+						TotalNumElements2D += Props.ElementCount2D;
 					}
 					else
 					{
@@ -191,9 +214,10 @@ void UPCGComputeKernel::ComputeDataDescFromPinProperties(const FPCGPinProperties
 					{
 						DataDesc.TagStringKeys.AddUnique(TagStringKey);
 					}
-
-					DataDesc.ElementCount = TotalNumElements;
 				}
+
+				DataDesc.ElementCount = TotalNumElements;
+				DataDesc.ElementCount2D = TotalNumElements2D;
 			}
 		};
 
@@ -204,7 +228,18 @@ void UPCGComputeKernel::ComputeDataDescFromPinProperties(const FPCGPinProperties
 			{
 				for (const FPCGDataDesc& InputDataDesc : InputDescs[0].DataDescs)
 				{
-					FPCGDataDesc& DataDesc = OutPinDesc.DataDescs.Emplace_GetRef(OutputPinProps.AllowedTypes, InputDataDesc.ElementCount);
+					FPCGDataDesc& DataDesc = OutPinDesc.DataDescs.Emplace_GetRef(OutputPinProps.AllowedTypes);
+					const bool bDomainIs2D = DataDesc.IsDomain2D();
+					const bool bSourceDomainIs2D = InputDataDesc.IsDomain2D();
+
+					if (bDomainIs2D)
+					{
+						DataDesc.ElementCount2D = bSourceDomainIs2D ? InputDataDesc.ElementCount2D : InputDataDesc.ElementCount;
+					}
+					else
+					{
+						DataDesc.ElementCount = bSourceDomainIs2D ? InputDataDesc.ElementCount2D.X * InputDataDesc.ElementCount2D.Y : InputDataDesc.ElementCount;
+					}
 
 					if (Props.AttributeInheritanceMode == EPCGAttributeInheritanceMode::CopyAttributeSetup)
 					{
@@ -229,17 +264,62 @@ void UPCGComputeKernel::ComputeDataDescFromPinProperties(const FPCGPinProperties
 						{
 							for (const FPCGDataDesc& OtherInputDataDesc : OtherInputDesc.DataDescs)
 							{
-								FPCGDataDesc* DataDesc = nullptr;
+								FPCGDataDesc& DataDesc = OutPinDesc.DataDescs.Emplace_GetRef(OutputPinProps.AllowedTypes);
+								const bool bDomainIs2D = DataDesc.IsDomain2D();
+								const bool bSourceDomainIs2D = InputDataDesc.IsDomain2D();
+								const bool bOtherSourceDomainIs2D = OtherInputDataDesc.IsDomain2D();
 
 								if (Props.ElementCountMode == EPCGElementCountMode::FromInputData)
 								{
 									if (Props.ElementMultiplicity == EPCGElementMultiplicity::Product)
 									{
-										DataDesc = &OutPinDesc.DataDescs.Emplace_GetRef(OutputPinProps.AllowedTypes, InputDataDesc.ElementCount * OtherInputDataDesc.ElementCount);
+										if (bDomainIs2D)
+										{
+											if (bSourceDomainIs2D)
+											{
+												DataDesc.ElementCount2D = bOtherSourceDomainIs2D ? InputDataDesc.ElementCount2D * OtherInputDataDesc.ElementCount2D : InputDataDesc.ElementCount2D * OtherInputDataDesc.ElementCount;
+											}
+											else
+											{
+												DataDesc.ElementCount2D = bOtherSourceDomainIs2D ?  OtherInputDataDesc.ElementCount2D * InputDataDesc.ElementCount : InputDataDesc.ElementCount * OtherInputDataDesc.ElementCount;
+											}
+										}
+										else
+										{
+											if (bSourceDomainIs2D)
+											{
+												DataDesc.ElementCount = (bOtherSourceDomainIs2D ? InputDataDesc.ElementCount2D * OtherInputDataDesc.ElementCount2D : InputDataDesc.ElementCount2D * OtherInputDataDesc.ElementCount).X;
+											}
+											else
+											{
+												DataDesc.ElementCount = (bOtherSourceDomainIs2D ?  OtherInputDataDesc.ElementCount2D * InputDataDesc.ElementCount : InputDataDesc.ElementCount * OtherInputDataDesc.ElementCount).X;
+											}
+										}
 									}
 									else if (Props.ElementMultiplicity == EPCGElementMultiplicity::Sum)
 									{
-										DataDesc = &OutPinDesc.DataDescs.Emplace_GetRef(OutputPinProps.AllowedTypes, InputDataDesc.ElementCount + OtherInputDataDesc.ElementCount);
+										if (bDomainIs2D)
+										{
+											if (bSourceDomainIs2D)
+											{
+												DataDesc.ElementCount2D = bOtherSourceDomainIs2D ? InputDataDesc.ElementCount2D + OtherInputDataDesc.ElementCount2D : InputDataDesc.ElementCount2D + OtherInputDataDesc.ElementCount;
+											}
+											else
+											{
+												DataDesc.ElementCount2D = bOtherSourceDomainIs2D ?  OtherInputDataDesc.ElementCount2D + InputDataDesc.ElementCount : InputDataDesc.ElementCount + OtherInputDataDesc.ElementCount;
+											}
+										}
+										else
+										{
+											if (bSourceDomainIs2D)
+											{
+												DataDesc.ElementCount = (bOtherSourceDomainIs2D ? InputDataDesc.ElementCount2D + OtherInputDataDesc.ElementCount2D : InputDataDesc.ElementCount2D + OtherInputDataDesc.ElementCount).X;
+											}
+											else
+											{
+												DataDesc.ElementCount = (bOtherSourceDomainIs2D ?  OtherInputDataDesc.ElementCount2D + InputDataDesc.ElementCount : InputDataDesc.ElementCount + OtherInputDataDesc.ElementCount).X;
+											}
+										}
 									}
 									else
 									{
@@ -248,31 +328,29 @@ void UPCGComputeKernel::ComputeDataDescFromPinProperties(const FPCGPinProperties
 								}
 								else if (Props.ElementCountMode == EPCGElementCountMode::Fixed)
 								{
-									DataDesc = &OutPinDesc.DataDescs.Emplace_GetRef(OutputPinProps.AllowedTypes, Props.ElementCount);
+									DataDesc.ElementCount = Props.ElementCount;
+									DataDesc.ElementCount2D = Props.ElementCount2D;
 								}
 								else
 								{
 									checkNoEntry();
 								}
 
-								if (ensure(DataDesc))
+								if (Props.AttributeInheritanceMode == EPCGAttributeInheritanceMode::CopyAttributeSetup)
 								{
-									if (Props.AttributeInheritanceMode == EPCGAttributeInheritanceMode::CopyAttributeSetup)
-									{
-										AddAttributesFromData(InputDataDesc, *DataDesc);
-										AddAttributesFromData(OtherInputDataDesc, *DataDesc);
-									}
+									AddAttributesFromData(InputDataDesc, DataDesc);
+									AddAttributesFromData(OtherInputDataDesc, DataDesc);
+								}
 
-									// Add unique tags from both input data.
-									for (const int32 TagStringKey : InputDataDesc.TagStringKeys)
-									{
-										DataDesc->TagStringKeys.AddUnique(TagStringKey);
-									}
+								// Add unique tags from both input data.
+								for (const int32 TagStringKey : InputDataDesc.TagStringKeys)
+								{
+									DataDesc.TagStringKeys.AddUnique(TagStringKey);
+								}
 
-									for (const int32 TagStringKey : OtherInputDataDesc.TagStringKeys)
-									{
-										DataDesc->TagStringKeys.AddUnique(TagStringKey);
-									}
+								for (const int32 TagStringKey : OtherInputDataDesc.TagStringKeys)
+								{
+									DataDesc.TagStringKeys.AddUnique(TagStringKey);
 								}
 							}
 						}
@@ -311,13 +389,14 @@ void UPCGComputeKernel::ComputeDataDescFromPinProperties(const FPCGPinProperties
 		for (FPCGDataDesc& Desc : OutPinDesc.DataDescs)
 		{
 			Desc.ElementCount *= Multiplier;
+			Desc.ElementCount2D *= Multiplier;
 		}
 	}
 	else if (Props.InitializationMode == EPCGPinInitMode::Custom)
 	{
 		for (int I = 0; I < Props.DataCount; ++I)
 		{
-			OutPinDesc.DataDescs.Emplace(OutputPinProps.AllowedTypes, Props.ElementCount);
+			OutPinDesc.DataDescs.Emplace(OutputPinProps.AllowedTypes, Props.ElementCount, Props.ElementCount2D);
 		}
 	}
 	else
diff --git a/Engine/Plugins/PCG/Source/PCG/Private/Compute/PCGDataBinding.cpp b/Engine/Plugins/PCG/Source/PCG/Private/Compute/PCGDataBinding.cpp
--- a/Engine/Plugins/PCG/Source/PCG/Private/Compute/PCGDataBinding.cpp
+++ b/Engine/Plugins/PCG/Source/PCG/Private/Compute/PCGDataBinding.cpp
@@ -167,11 +167,22 @@ const FPCGDataCollectionDesc* UPCGDataBinding::GetCachedKernelPinDataDesc(int32
 }
 
 void UPCGDataBinding::ReceiveDataFromGPU_GameThread(
-	UPCGProxyForGPUData* InData,
+	UPCGData* InData,
 	const UPCGSettings* InProducerSettings,
 	EPCGExportMode InExportMode,
 	FName InPinLabel,
 	FName InPinLabelAlias)
+{
+	ReceiveDataFromGPU_GameThread(InData, InProducerSettings, InExportMode, InPinLabel, InPinLabelAlias, /*InAdditionalTags=*/{});
+}
+
+void UPCGDataBinding::ReceiveDataFromGPU_GameThread(
+	UPCGData* InData,
+	const UPCGSettings* InProducerSettings,
+	EPCGExportMode InExportMode,
+	FName InPinLabel,
+	FName InPinLabelAlias,
+	const TSet<FString>& InAdditionalTags)
 {
 	check(IsInGameThread());
 
@@ -180,16 +191,17 @@ void UPCGDataBinding::ReceiveDataFromGPU_GameThread(
 		FPCGTaggedData& TaggedData = OutputDataCollection.TaggedData.Emplace_GetRef();
 		TaggedData.Pin = InPinLabelAlias;
 		TaggedData.Data = InData;
+		TaggedData.Tags = InAdditionalTags;
 	}
 
 	if (!!(InExportMode & EPCGExportMode::DebugVisualization))
 	{
-		DataToDebug.Add(FPCGDataToDebug{ .Proxy = InData, .ProducerSettings = InProducerSettings, .PinLabel = InPinLabel, .PinLabelAlias = InPinLabelAlias });
+		DataToDebug.Add(FPCGDataToDebug{ .Data = InData, .ProducerSettings = InProducerSettings, .PinLabel = InPinLabel, .PinLabelAlias = InPinLabelAlias, .AdditionalTags = InAdditionalTags });
 	}
 
 	if (!!(InExportMode & EPCGExportMode::Inspection))
 	{
-		DataToInspect.Add(FPCGDataToDebug{ .Proxy = InData, .ProducerSettings = InProducerSettings, .PinLabel = InPinLabel, .PinLabelAlias = InPinLabelAlias });
+		DataToInspect.Add(FPCGDataToDebug{ .Data = InData, .ProducerSettings = InProducerSettings, .PinLabel = InPinLabel, .PinLabelAlias = InPinLabelAlias, .AdditionalTags = InAdditionalTags });
 	}
 }
 
@@ -215,10 +227,13 @@ void UPCGDataBinding::ApplyElementCounts_GameThread()
 	{
 		for (FPCGDataToDebug& Data : DebugData)
 		{
-			const TArray<uint32>* FoundElementCounts = This->PinToElementCounts.Find(Data.PinLabelAlias);
-			if (FoundElementCounts && Data.Proxy)
+			if (UPCGProxyForGPUData* Proxy = Cast<UPCGProxyForGPUData>(Data.Data))
 			{
-				Data.Proxy->UpdateElementCountsFromReadback(*FoundElementCounts);
+				const TArray<uint32>* FoundElementCounts = This->PinToElementCounts.Find(Data.PinLabelAlias);
+				if (FoundElementCounts)
+				{
+					Proxy->UpdateElementCountsFromReadback(*FoundElementCounts);
+				}
 			}
 		}
 	};
diff --git a/Engine/Plugins/PCG/Source/PCG/Private/Compute/PCGDataBinding.h b/Engine/Plugins/PCG/Source/PCG/Private/Compute/PCGDataBinding.h
--- a/Engine/Plugins/PCG/Source/PCG/Private/Compute/PCGDataBinding.h
+++ b/Engine/Plugins/PCG/Source/PCG/Private/Compute/PCGDataBinding.h
@@ -58,7 +58,10 @@ struct FPCGDataToDebug
 	GENERATED_BODY()
 
 	UPROPERTY()
-	TObjectPtr<UPCGProxyForGPUData> Proxy;
+	TObjectPtr<UPCGData> Data;
+
+	UPROPERTY()
+	TObjectPtr<UPCGData> DataPendingInit;
 	
 	UPROPERTY()
 	TWeakObjectPtr<const UPCGSettings> ProducerSettings;
@@ -68,6 +71,11 @@ struct FPCGDataToDebug
 	
 	UPROPERTY()
 	FName PinLabelAlias;
+
+	// @todo_pcg: This is a big hack to support tags on texture data proxies. Should be replaced with a proper abstraction.
+	/** Allow appending additional tags to the CPU tagged data. */
+	UPROPERTY()
+	TSet<FString> AdditionalTags;
 };
 
 UCLASS(Transient, Category = PCG)
@@ -102,7 +110,8 @@ public:
 	const FPCGDataCollectionDesc* GetCachedKernelPinDataDesc(const UPCGComputeKernel* InKernel, FName InPinLabel, bool bIsInput) const;
 	const FPCGDataCollectionDesc* GetCachedKernelPinDataDesc(int32 InGraphBindingIndex) const;
 
-	void ReceiveDataFromGPU_GameThread(UPCGProxyForGPUData* InData, const UPCGSettings* InProducerSettings, EPCGExportMode InExportMode, FName InPinLabel, FName InPinLabelAlias);
+	void ReceiveDataFromGPU_GameThread(UPCGData* InData, const UPCGSettings* InProducerSettings, EPCGExportMode InExportMode, FName InPinLabel, FName InPinLabelAlias);
+	void ReceiveDataFromGPU_GameThread(UPCGData* InData, const UPCGSettings* InProducerSettings, EPCGExportMode InExportMode, FName InPinLabel, FName InPinLabelAlias, const TSet<FString>& AdditionalTags);
 
 	void ReceiveElementCountsFromGPU_GameThread(FName InPinLabel, const TArray<uint32>& InCounters);
 
diff --git a/Engine/Plugins/PCG/Source/PCG/Private/Compute/PCGDataForGPU.cpp b/Engine/Plugins/PCG/Source/PCG/Private/Compute/PCGDataForGPU.cpp
--- a/Engine/Plugins/PCG/Source/PCG/Private/Compute/PCGDataForGPU.cpp
+++ b/Engine/Plugins/PCG/Source/PCG/Private/Compute/PCGDataForGPU.cpp
@@ -725,9 +725,20 @@ bool FPCGKernelAttributeDesc::operator==(const FPCGKernelAttributeDesc& Other) c
 	return AttributeId == Other.AttributeId && AttributeKey == Other.AttributeKey;
 }
 
+FPCGDataDesc::FPCGDataDesc(EPCGDataType InType)
+	: FPCGDataDesc(InType, 0, FIntPoint::ZeroValue)
+{
+}
+
 FPCGDataDesc::FPCGDataDesc(EPCGDataType InType, int InElementCount)
+	: FPCGDataDesc(InType, InElementCount, FIntPoint::ZeroValue)
+{
+}
+
+FPCGDataDesc::FPCGDataDesc(EPCGDataType InType, int InElementCount, FIntPoint InElementCount2D)
 	: Type(InType)
 	, ElementCount(InElementCount)
+	, ElementCount2D(InElementCount2D)
 {
 	InitializeAttributeDescs(/*InData=*/nullptr, /*InBinding*/nullptr);
 }
@@ -738,6 +749,7 @@ FPCGDataDesc::FPCGDataDesc(const FPCGTaggedData& InTaggedData, const UPCGDataBin
 
 	Type = InTaggedData.Data->GetDataType();
 	ElementCount = PCGComputeHelpers::GetElementCount(InTaggedData.Data);
+	ElementCount2D = PCGComputeHelpers::GetElementCount2D(InTaggedData.Data);
 
 	TagStringKeys.Reserve(InTaggedData.Tags.Num());
 
@@ -932,7 +944,7 @@ FPCGDataCollectionDesc FPCGDataCollectionDesc::BuildFromDataCollection(
 
 	for (const FPCGTaggedData& Data : InDataCollection.TaggedData)
 	{
-		if (!Data.Data || !PCGComputeHelpers::IsTypeAllowedInDataCollection(Data.Data->GetDataType()))
+		if (!Data.Data || !PCGComputeHelpers::IsTypeAllowedAsOutput(Data.Data->GetDataType()))
 		{
 			continue;
 		}
@@ -1619,7 +1631,14 @@ uint32 FPCGDataCollectionDesc::ComputeDataElementCount(EPCGDataType InDataType)
 	{
 		if (!!(DataDesc.Type & InDataType))
 		{
-			ElementCount += DataDesc.ElementCount;
+			if (!!(DataDesc.Type & EPCGDataType::BaseTexture))
+			{
+				ElementCount += DataDesc.ElementCount2D.X * DataDesc.ElementCount2D.Y;
+			}
+			else
+			{
+				ElementCount += DataDesc.ElementCount;
+			}
 		}
 	}
 
diff --git a/Engine/Plugins/PCG/Source/PCG/Private/Compute/PCGPinPropertiesGPU.h b/Engine/Plugins/PCG/Source/PCG/Private/Compute/PCGPinPropertiesGPU.h
--- a/Engine/Plugins/PCG/Source/PCG/Private/Compute/PCGPinPropertiesGPU.h
+++ b/Engine/Plugins/PCG/Source/PCG/Private/Compute/PCGPinPropertiesGPU.h
@@ -97,6 +97,12 @@ struct FPCGPinPropertiesGPUStruct
 	UPROPERTY(BlueprintReadWrite, EditAnywhere, Category = Settings, meta = (EditCondition = "bAllowEditInitMode && (InitializationMode == EPCGPinInitMode::Custom || ElementCountMode == EPCGElementCountMode::Fixed)", EditConditionHides, ClampMin = 1))
 	int ElementCount = 1;
 
+	// @todo_pcg: Show texture data type info, e.g. Dimension, Format, etc.
+
+	/** Fixed number of elements to create in each output data. */
+	UPROPERTY(BlueprintReadWrite, EditAnywhere, Category = Settings, meta = (EditCondition = "bAllowEditInitMode && (InitializationMode == EPCGPinInitMode::Custom || ElementCountMode == EPCGElementCountMode::Fixed)", EditConditionHides, ClampMin = 1))
+	FIntPoint ElementCount2D = FIntPoint(64, 64);
+
 	/** Fixed number of elements to create in each output data. */
 	UPROPERTY(BlueprintReadWrite, EditAnywhere, Category = Settings, meta = (EditCondition = "bAllowEditInitMode && (InitializationMode == EPCGPinInitMode::Custom || ElementCountMode == EPCGElementCountMode::FromInputData)", EditConditionHides, ClampMin = 1))
 	int ElementCountMultiplier = 1;
diff --git a/Engine/Plugins/PCG/Source/PCG/Private/Data/PCGRenderTargetData.cpp b/Engine/Plugins/PCG/Source/PCG/Private/Data/PCGRenderTargetData.cpp
--- a/Engine/Plugins/PCG/Source/PCG/Private/Data/PCGRenderTargetData.cpp
+++ b/Engine/Plugins/PCG/Source/PCG/Private/Data/PCGRenderTargetData.cpp
@@ -17,30 +17,33 @@ void UPCGRenderTargetData::Initialize(UTextureRenderTarget2D* InRenderTarget, co
 
 	ColorData.Reset();
 
-	if (RenderTarget && !bSkipReadbackToCPU)
+	if (RenderTarget)
 	{
-		if (FTextureRenderTargetResource* RTResource = RenderTarget->GameThread_GetRenderTargetResource())
+		Width = RenderTarget->SizeX;
+		Height = RenderTarget->SizeY;
+
+		if (!bSkipReadbackToCPU)
 		{
-			TRACE_CPUPROFILER_EVENT_SCOPE(UPCGRenderTargetData::Initialize::ReadData);
-			ETextureRenderTargetFormat Format = RenderTarget->RenderTargetFormat;
-			if ((Format == RTF_RGBA16f) || (Format == RTF_RGBA32f) || (Format == RTF_RGBA8) || (Format == RTF_R8) || (Format == RTF_RGB10A2))
+			if (FTextureRenderTargetResource* RTResource = RenderTarget->GameThread_GetRenderTargetResource())
 			{
-				const FIntRect Rect = FIntRect(0, 0, RenderTarget->SizeX, RenderTarget->SizeY);
-				const FReadSurfaceDataFlags ReadPixelFlags(RCM_MinMax);
-				RTResource->ReadLinearColorPixels(ColorData, ReadPixelFlags, Rect);
+				TRACE_CPUPROFILER_EVENT_SCOPE(UPCGRenderTargetData::Initialize::ReadData);
+				ETextureRenderTargetFormat Format = RenderTarget->RenderTargetFormat;
+				if ((Format == RTF_RGBA16f) || (Format == RTF_RGBA32f) || (Format == RTF_RGBA8) || (Format == RTF_R8) || (Format == RTF_RGB10A2))
+				{
+					const FIntRect Rect = FIntRect(0, 0, RenderTarget->SizeX, RenderTarget->SizeY);
+					const FReadSurfaceDataFlags ReadPixelFlags(RCM_MinMax);
+					RTResource->ReadLinearColorPixels(ColorData, ReadPixelFlags, Rect);
+				}
 			}
 		}
-	}
 
-	// Never take resource ownership on assets.
-	if (RenderTarget && RenderTarget->IsAsset())
-	{
-		bOwnsRenderTarget = false;
+		// Never take resource ownership on assets.
+		if (RenderTarget->IsAsset())
+		{
+			bOwnsRenderTarget = false;
+		}
 	}
 
-	Width = !ColorData.IsEmpty() ? RenderTarget->SizeX : 0;
-	Height = !ColorData.IsEmpty() ? RenderTarget->SizeY : 0;
-
 	Bounds = FBox(EForceInit::ForceInit);
 	Bounds += FVector(-1.0f, -1.0f, 0.0f);
 	Bounds += FVector(1.0f, 1.0f, 0.0f);
diff --git a/Engine/Plugins/PCG/Source/PCG/Private/Data/PCGTextureData.cpp b/Engine/Plugins/PCG/Source/PCG/Private/Data/PCGTextureData.cpp
--- a/Engine/Plugins/PCG/Source/PCG/Private/Data/PCGTextureData.cpp
+++ b/Engine/Plugins/PCG/Source/PCG/Private/Data/PCGTextureData.cpp
@@ -398,7 +398,7 @@ const UPCGBasePointData* UPCGBaseTextureData::CreateBasePointData(FPCGContext* C
 
 bool UPCGBaseTextureData::IsValid() const
 {
-	return (Height > 0 && Width > 0 && !ColorData.IsEmpty()) || bSkipReadbackToCPU;
+	return Height > 0 && Width > 0 && (!ColorData.IsEmpty() || bSkipReadbackToCPU);
 }
 
 bool UPCGBaseTextureData::SamplePointLocal(const FVector2D& LocalPosition, FVector4& OutColor, float& OutDensity) const
@@ -496,6 +496,11 @@ void UPCGTextureData::InitializeInternal(UTexture* InTexture, uint32 InTextureIn
 
 	if (bInSkipReadbackToCPU)
 	{
+		FTextureRHIRef ResourceRHI = GetTextureRHI();
+		const FIntPoint Extent = ResourceRHI ? ResourceRHI->GetDesc().Extent : FIntPoint::ZeroValue;
+		Width = Extent.X;
+		Height = Extent.Y;
+
 		bSuccessfullyInitialized = true;
 		SetInitCompleted(true);
 		return;
@@ -643,6 +648,9 @@ bool UPCGTextureData::Initialize(TRefCountPtr<IPooledRenderTarget> InTextureHand
 
 	if (bInSkipReadbackToCPU)
 	{
+		const FIntPoint Extent = InTextureHandle.IsValid() ? InTextureHandle->GetDesc().Extent : FIntPoint::ZeroValue;
+		Width = Extent.X;
+		Height = Extent.Y;
 		bSuccessfullyInitialized = true;
 	}
 	else
diff --git a/Engine/Plugins/PCG/Source/PCG/Private/Graph/PCGGraphCompilerGPU.cpp b/Engine/Plugins/PCG/Source/PCG/Private/Graph/PCGGraphCompilerGPU.cpp
--- a/Engine/Plugins/PCG/Source/PCG/Private/Graph/PCGGraphCompilerGPU.cpp
+++ b/Engine/Plugins/PCG/Source/PCG/Private/Graph/PCGGraphCompilerGPU.cpp
@@ -895,7 +895,6 @@ void FPCGGraphCompilerGPU::CreateDataInterfaces(
 			}
 
 			check(DataInterfacePCGData);
-			DataInterfacePCGData->SetProducerSettings(ProducerSettings);
 			DataInterfacePCGData->SetRequiresExport(bRequiresExport);
 
 			if (ProducerPin)
@@ -917,7 +916,12 @@ void FPCGGraphCompilerGPU::CreateDataInterfaces(
 		case EPCGDataType::Texture:
 		case EPCGDataType::RenderTarget:
 		{
-			DataInterface = InOutContext.NewObject_AnyThread<UPCGTextureDataInterface>(InOutComputeGraph);
+			UPCGTextureDataInterface* TextureDataInterface = InOutContext.NewObject_AnyThread<UPCGTextureDataInterface>(InOutComputeGraph);
+			TextureDataInterface->SetRequiresExport(bRequiresExport);
+			TextureDataInterface->SetInitializeFromDataCollection(!bUpstreamIsGPUTask);
+
+			DataInterface = TextureDataInterface;
+
 			break;
 		}
 		case EPCGDataType::Landscape:
@@ -944,6 +948,7 @@ void FPCGGraphCompilerGPU::CreateDataInterfaces(
 		if (DataInterface)
 		{
 			DataInterface->ProducerKernel = ProducerKernel;
+			DataInterface->SetProducerSettings(ProducerSettings);
 			DataInterface->SetOutputPin(InOutputPinProperties.Label);
 		}
 
diff --git a/Engine/Plugins/PCG/Source/PCG/Public/Data/PCGTextureData.h b/Engine/Plugins/PCG/Source/PCG/Public/Data/PCGTextureData.h
--- a/Engine/Plugins/PCG/Source/PCG/Public/Data/PCGTextureData.h
+++ b/Engine/Plugins/PCG/Source/PCG/Public/Data/PCGTextureData.h
@@ -98,6 +98,7 @@ public:
 	UE_API virtual EPCGTextureResourceType GetTextureResourceType() const PURE_VIRTUAL(UPCGBaseTextureData::GetTextureResourceType, return EPCGTextureResourceType::Invalid;);
 	virtual TRefCountPtr<IPooledRenderTarget> GetRefCountedTexture() const { return nullptr; }
 	virtual int GetTextureSlice() const { return 0; }
+	virtual FIntPoint GetTextureSize() const { return FIntPoint(Width, Height); }
 
 protected:
 	UE_API const UPCGBasePointData* CreateBasePointData(FPCGContext* Context, TSubclassOf<UPCGBasePointData> PointDataClass) const;
