[PCG] Update Graph Parameters on Category change

#rnx

diff --git a/Engine/Plugins/PCG/Source/PCG/Private/PCGGraph.cpp b/Engine/Plugins/PCG/Source/PCG/Private/PCGGraph.cpp
--- a/Engine/Plugins/PCG/Source/PCG/Private/PCGGraph.cpp
+++ b/Engine/Plugins/PCG/Source/PCG/Private/PCGGraph.cpp
@@ -19,6 +19,7 @@
 #include "Graph/PCGGraphCompilationData.h"
 #include "Graph/PCGGraphCompiler.h"
 #include "Graph/PCGGraphExecutor.h"
+#include "Helpers/PCGPropertyHelpers.h"
 
 #include "AssetRegistry/AssetData.h"
 
@@ -337,7 +338,10 @@ bool UPCGGraphInterface::VerifyIfGraphCustomizationChanged(const FPropertyChange
 EPCGChangeType UPCGGraphInterface::GetChangeTypeForGraphParameterChange(EPCGGraphParameterEvent InChangeType, FName InChangedPropertyName)
 {
 	// If the parameter had its order changed in the struct, was just added or was removed but was not used in the graph, it is not a change that requires a refresh, so we go with Cosmetic change type.
-	if (InChangeType == EPCGGraphParameterEvent::PropertyMoved || InChangeType == EPCGGraphParameterEvent::Added || InChangeType == EPCGGraphParameterEvent::RemovedUnused)
+	if (InChangeType == EPCGGraphParameterEvent::PropertyMoved
+		|| InChangeType == EPCGGraphParameterEvent::Added
+		|| InChangeType == EPCGGraphParameterEvent::RemovedUnused
+		|| InChangeType == EPCGGraphParameterEvent::CategoryChanged)
 	{
 		return EPCGChangeType::Cosmetic;
 	}
@@ -1849,6 +1853,23 @@ void UPCGGraph::PostEditChangeChainProperty(FPropertyChangedChainEvent& Property
 					ChangedPropertyName = PostDesc.Name;
 					break;
 				}
+				// Category changed for the property
+				else if (PreDesc.CachedProperty && PostDesc.CachedProperty)
+				{
+					using namespace PCGPropertyHelpers::Constants;
+					const bool bPreCategoriesEnabled = PreDesc.CachedProperty->HasMetaData(EnableCategoriesMetadataName);
+					const bool bPostCategoriesEnabled = PostDesc.CachedProperty->HasMetaData(EnableCategoriesMetadataName);
+					// Check first if categories or enabled and then compare if they have the same category.
+					if (bPreCategoriesEnabled != bPostCategoriesEnabled
+						|| (bPreCategoriesEnabled
+							&& bPostCategoriesEnabled
+							&& !PreDesc.CachedProperty->GetMetaData(CategoryMetadataName).Equals(PostDesc.CachedProperty->GetMetaData(CategoryMetadataName))))
+					{
+						ChangeType = EPCGGraphParameterEvent::CategoryChanged;
+						ChangedPropertyName = PostDesc.Name;
+						break;
+					}
+				}
 			}
 		}
 
@@ -1977,9 +1998,6 @@ bool UPCGGraph::UserParametersIsPinTypeAccepted(FEdGraphPinType InPinType, bool
 
 void UPCGGraph::OnGraphParametersChanged(EPCGGraphParameterEvent InChangeType, FName InChangedPropertyName)
 {
-	bool bWasModified = false;
-	EPCGGraphParameterEvent ChangeType = InChangeType;
-
 	if (InChangeType == EPCGGraphParameterEvent::RemovedUsed || InChangeType == EPCGGraphParameterEvent::RemovedUnused || InChangeType == EPCGGraphParameterEvent::PropertyRenamed)
 	{
 		// Look for all the Get Parameter nodes and make sure to delete all nodes that doesn't exist anymore
@@ -1997,7 +2015,6 @@ void UPCGGraph::OnGraphParametersChanged(EPCGGraphParameterEvent InChangeType, F
 				const FPropertyBagPropertyDesc* PropertyDesc = UserParameters.FindPropertyDescByID(Settings->PropertyGuid);
 				if (!PropertyDesc)
 				{
-					ChangeType = EPCGGraphParameterEvent::RemovedUsed;
 					NodesToRemove.Add(Node);
 				}
 				else if (Settings->PropertyName != PropertyDesc->Name)
@@ -2013,7 +2030,6 @@ void UPCGGraph::OnGraphParametersChanged(EPCGGraphParameterEvent InChangeType, F
 
 		if (!NodesToRemove.IsEmpty())
 		{
-			bWasModified = true;
 			Modify();
 
 			for (UPCGNode* Node : NodesToRemove)
@@ -2612,6 +2628,7 @@ bool FPCGOverrideInstancedPropertyBag::RefreshParameters(const FInstancedPropert
 	case EPCGGraphParameterEvent::RemovedUsed: // fall-through
 	case EPCGGraphParameterEvent::PropertyRenamed: // fall-through
 	case EPCGGraphParameterEvent::PropertyMoved: // fall-through
+	case EPCGGraphParameterEvent::CategoryChanged: // fall-through
 	case EPCGGraphParameterEvent::PropertyTypeModified:
 	{
 		bWasModified = true;
diff --git a/Engine/Plugins/PCG/Source/PCG/Public/Helpers/PCGPropertyHelpers.h b/Engine/Plugins/PCG/Source/PCG/Public/Helpers/PCGPropertyHelpers.h
--- a/Engine/Plugins/PCG/Source/PCG/Public/Helpers/PCGPropertyHelpers.h
+++ b/Engine/Plugins/PCG/Source/PCG/Public/Helpers/PCGPropertyHelpers.h
@@ -121,6 +121,12 @@ namespace PCGPropertyHelpers
 	* Since the names in editor and at runtime can be different for the properties defined in a UserDefinedStruct, the search will also look for sanitized version of the name if it is not a valid name.
 	*/
 	const FProperty* FindPropertyInUserDefinedStruct(const UUserDefinedStruct* InStruct, const FName InName);
+
+	namespace Constants
+	{
+		const FName EnableCategoriesMetadataName("EnableCategories");
+		const FName CategoryMetadataName("Category");
+	}
 }
 
 //////
diff --git a/Engine/Plugins/PCG/Source/PCG/Public/PCGGraph.h b/Engine/Plugins/PCG/Source/PCG/Public/PCGGraph.h
--- a/Engine/Plugins/PCG/Source/PCG/Public/PCGGraph.h
+++ b/Engine/Plugins/PCG/Source/PCG/Public/PCGGraph.h
@@ -40,6 +40,7 @@ enum class EPCGGraphParameterEvent
 	ValueModifiedByParent,
 	MultiplePropertiesAdded,
 	UndoRedo,
+	CategoryChanged,
 	None
 };
 
