[PCG] Retrieve the ray hit material with GetMaterialFromCollisionFaceIndex, instead of needing to specify a material index.

diff --git a/Engine/Plugins/PCG/Source/PCG/Private/Helpers/PCGWorldQueryHelpers.cpp b/Engine/Plugins/PCG/Source/PCG/Private/Helpers/PCGWorldQueryHelpers.cpp
--- a/Engine/Plugins/PCG/Source/PCG/Private/Helpers/PCGWorldQueryHelpers.cpp
+++ b/Engine/Plugins/PCG/Source/PCG/Private/Helpers/PCGWorldQueryHelpers.cpp
@@ -10,6 +10,7 @@
 #include "Helpers/PCGHelpers.h"
 
 #include "LandscapeProxy.h"
+#include "StaticMeshResources.h"
 #include "Algo/AnyOf.h"
 #include "Components/BrushComponent.h"
 #include "Components/PrimitiveComponent.h"
@@ -218,10 +219,12 @@ namespace PCGWorldQueryHelpers
 		bResult &= CreateAttribute(PCGPointDataConstants::ActorReferenceAttribute, QueryParams.bGetReferenceToActorHit, FSoftObjectPath());
 		bResult &= CreateAttribute(PCGWorldQueryConstants::PhysicalMaterialReferenceAttribute, QueryParams.bGetReferenceToPhysicalMaterial, FSoftObjectPath());
 		bResult &= CreateAttribute(PCGWorldQueryConstants::RenderMaterialReferenceAttribute, QueryParams.bGetReferenceToRenderMaterial, FSoftObjectPath());
+		bResult &= CreateAttribute(PCGWorldQueryConstants::RenderMaterialIndexAttribute, QueryParams.bTraceComplex && QueryParams.bGetRenderMaterialIndex, int32{0});
 		bResult &= CreateAttribute(PCGWorldQueryConstants::StaticMeshReferenceAttribute, QueryParams.bGetReferenceToStaticMesh, FSoftObjectPath());
 		bResult &= CreateAttribute(PCGWorldQueryConstants::ElementIndexAttribute, QueryParams.bGetElementIndex, int32{0});
 		bResult &= CreateAttribute(PCGWorldQueryConstants::UVCoordAttribute, QueryParams.bTraceComplex && QueryParams.bGetUVCoords, FVector2D::ZeroVector);
 		bResult &= CreateAttribute(PCGWorldQueryConstants::FaceIndexAttribute, QueryParams.bTraceComplex && QueryParams.bGetFaceIndex, int32{0});
+		bResult &= CreateAttribute(PCGWorldQueryConstants::SectionIndexAttribute, QueryParams.bTraceComplex && QueryParams.bGetSectionIndex, int32{0});
 
 		return bResult;
 	}
@@ -299,6 +302,7 @@ namespace PCGWorldQueryHelpers
 		bResult &= ApplyAttribute(PCGWorldQueryConstants::PhysicalMaterialReferenceAttribute, FSoftObjectPath(Hit.PhysMaterial.Get()), QueryParams.bGetReferenceToPhysicalMaterial);
 		bResult &= ApplyAttribute(PCGWorldQueryConstants::ElementIndexAttribute, static_cast<int32>(Hit.ElementIndex), QueryParams.bGetElementIndex);
 
+		// Handle landscape related hit result
 		if (ALandscapeProxy* Landscape = Cast<ALandscapeProxy>(Hit.GetActor()))
 		{
 			if (const UMaterialInterface* RenderMaterial = Landscape->GetLandscapeMaterial())
@@ -322,20 +326,16 @@ namespace PCGWorldQueryHelpers
 				}
 			}
 		}
+		// Handle primitive component related hit result
 		else if (UPrimitiveComponent* HitComponent = Hit.GetComponent())
 		{
-			if (QueryParams.bGetReferenceToRenderMaterial)
-			{
-				const UMaterialInterface* RenderMaterial = HitComponent->GetMaterial(QueryParams.RenderMaterialIndex);
-				bResult &= RenderMaterial && ApplyAttribute(PCGWorldQueryConstants::RenderMaterialReferenceAttribute, FSoftObjectPath(RenderMaterial));
-			}
-
 			if (QueryParams.bGetLocalImpactPoint)
 			{
 				const FVector LocalHitLocation = HitComponent->GetComponentToWorld().InverseTransformPosition(Hit.ImpactPoint);
 				bResult &= ApplyAttribute(PCGWorldQueryConstants::LocalImpactPointAttribute, LocalHitLocation);
 			}
 
+			// Handle static mesh component related hit result
 			if (const UStaticMeshComponent* StaticMeshComponent = Cast<UStaticMeshComponent>(HitComponent))
 			{
 				if (UStaticMesh* StaticMesh = StaticMeshComponent->GetStaticMesh())
@@ -346,12 +346,72 @@ namespace PCGWorldQueryHelpers
 				// Implementation note: FaceIndex will return -1 if complex queries are disabled
 				bResult &= ApplyAttribute(PCGWorldQueryConstants::FaceIndexAttribute, Hit.FaceIndex, QueryParams.bTraceComplex && QueryParams.bGetFaceIndex);
 
-				if (QueryParams.bTraceComplex && QueryParams.bGetUVCoords && UPhysicsSettings::Get()->bSupportUVFromHitResults)
+				if (QueryParams.bTraceComplex)
+				{
+					int32 SectionIndex = INDEX_NONE;
+					int32 MaterialIndex = INDEX_NONE;
+					UMaterialInterface* RenderMaterial = nullptr;
+
+					// Attempt to get these values from LOD0, only if needed.
+					if (QueryParams.bGetSectionIndex
+						|| QueryParams.bGetRenderMaterialIndex
+						|| (QueryParams.bGetReferenceToRenderMaterial && !QueryParams.bUseRenderMaterialIndex))
+					{
+						if (Hit.FaceIndex != INDEX_NONE
+							&& StaticMeshComponent->GetStaticMesh()
+							&& StaticMeshComponent->GetStaticMesh()->GetRenderData()
+							&& !StaticMeshComponent->GetStaticMesh()->GetRenderData()->LODResources.IsEmpty())
+						{
+							const FStaticMeshLODResources& LOD = StaticMeshComponent->GetStaticMesh()->GetRenderData()->LODResources[0];
+							for (int32 CurrentIndex = 0; CurrentIndex < LOD.Sections.Num(); ++CurrentIndex)
+							{
+								const FStaticMeshSection& Section = LOD.Sections[CurrentIndex];
+								const uint32 HitFaceIndex = static_cast<uint32>(Hit.FaceIndex);
+
+								// Find the section by iterating through the indices 3 at a time.
+								const uint32 FirstFaceIndex = Section.FirstIndex / 3u;
+								const uint32 LastFaceIndex = FirstFaceIndex + Section.NumTriangles;
+
+								// Find the section index
+								if (HitFaceIndex >= FirstFaceIndex && HitFaceIndex <= LastFaceIndex)
+								{
+									SectionIndex = CurrentIndex;
+									MaterialIndex = Section.MaterialIndex;
+									break;
+								}
+							}
+						}
+					}
+
+					if (QueryParams.bUseRenderMaterialIndex)  // Use explicit Render Material Index
+					{
+						RenderMaterial = HitComponent->GetMaterial(QueryParams.RenderMaterialIndex);
+						bResult &= ApplyAttribute(PCGWorldQueryConstants::RenderMaterialIndexAttribute, QueryParams.RenderMaterialIndex, QueryParams.bGetRenderMaterialIndex);
+					}
+					else
+					{
+						RenderMaterial = HitComponent->GetMaterialFromCollisionFaceIndex(Hit.FaceIndex, SectionIndex);
+						bResult &= ApplyAttribute(PCGWorldQueryConstants::RenderMaterialIndexAttribute, MaterialIndex, QueryParams.bGetRenderMaterialIndex);
+					}
+
+					bResult &= ApplyAttribute(PCGWorldQueryConstants::SectionIndexAttribute, SectionIndex, QueryParams.bGetSectionIndex);
+					bResult &= ApplyAttribute(PCGWorldQueryConstants::RenderMaterialReferenceAttribute, FSoftObjectPath(RenderMaterial), QueryParams.bGetReferenceToRenderMaterial);
+
+					if (QueryParams.bGetUVCoords && UPhysicsSettings::Get()->bSupportUVFromHitResults)
+					{
+						FVector2D UVCoords;
+						if (UGameplayStatics::FindCollisionUV(Hit, QueryParams.UVChannel, UVCoords))
+						{
+							bResult &= ApplyAttribute(PCGWorldQueryConstants::UVCoordAttribute, UVCoords);
+						}
+					}
+				}
+				else // @todo_pcg: Address other primitive types here
 				{
-					FVector2D UVCoords;
-					if (UGameplayStatics::FindCollisionUV(Hit, QueryParams.UVChannel, UVCoords))
+					if (QueryParams.bGetReferenceToRenderMaterial)
 					{
-						bResult &= ApplyAttribute(PCGWorldQueryConstants::UVCoordAttribute, UVCoords);
+						const UMaterialInterface* RenderMaterial = HitComponent->GetMaterial(QueryParams.RenderMaterialIndex);
+						bResult &= RenderMaterial && ApplyAttribute(PCGWorldQueryConstants::RenderMaterialReferenceAttribute, FSoftObjectPath(RenderMaterial));
 					}
 				}
 			}
diff --git a/Engine/Plugins/PCG/Source/PCG/Public/Data/PCGWorldData.h b/Engine/Plugins/PCG/Source/PCG/Public/Data/PCGWorldData.h
--- a/Engine/Plugins/PCG/Source/PCG/Public/Data/PCGWorldData.h
+++ b/Engine/Plugins/PCG/Source/PCG/Public/Data/PCGWorldData.h
@@ -167,7 +167,7 @@ struct FPCGWorldRaycastQueryParams : public FPCGWorldCommonQueryParams
 	UPROPERTY(BlueprintReadWrite, EditAnywhere, Category = "Data|Attributes", meta = (PCG_Overridable))
 	uint8 bGetLocalImpactPoint : 1 = false;
 
-	/** Create an attribute for the render material. */
+	/** Create an attribute for the render material. Requires 'bTraceComplex' for use with Primitive Components. */
 	UPROPERTY(BlueprintReadWrite, EditAnywhere, Category = "Data|Attributes", meta = (PCG_Overridable))
 	uint8 bGetReferenceToRenderMaterial : 1 = false;
 
@@ -187,14 +187,26 @@ struct FPCGWorldRaycastQueryParams : public FPCGWorldCommonQueryParams
 	UPROPERTY(BlueprintReadWrite, EditAnywhere, Category = "Data|Attributes", meta = (PCG_Overridable))
 	uint8 bGetElementIndex : 1 = false;
 
+	/** Create an attribute for the index of the section hit. Currently only works for Static Meshes. */
+	UPROPERTY(BlueprintReadWrite, EditAnywhere, Category = "Data|Attributes", meta = (PCG_Overridable, EditCondition = "bTraceComplex"))
+	uint8 bGetSectionIndex : 1 = false;
+
 	/** Will apply landscape layers and their values at the impact point. */
 	UPROPERTY(BlueprintReadWrite, EditAnywhere, Category = "Data|Attributes", meta = (PCG_Overridable, EditCondition = "SelectLandscapeHits != EPCGWorldQuerySelectLandscapeHits::Exclude"))
 	uint8 bApplyMetadataFromLandscape : 1 = false;
 
-	/** The index of the render material when it is queried from the hit. */
-	UPROPERTY(BlueprintReadWrite, EditAnywhere, Category = "Data|Attributes", meta = (PCG_Overridable, EditCondition = "bGetReferenceToRenderMaterial", EditConditionHides, DisplayAfter = "bGetReferenceToRenderMaterial"))
+	/** Retrieve the material index explicitly. If false, the render material will be assumed from the primitive. Currently only works for Static Meshes. */
+	UPROPERTY(BlueprintReadWrite, EditAnywhere, Category = "Data|Attributes", DisplayName = "Set Render Material Index", meta = (PCG_Overridable, EditCondition = "bTraceComplex && bGetReferenceToRenderMaterial", EditConditionHides, DisplayAfter = "bGetReferenceToRenderMaterial"))
+	uint8 bUseRenderMaterialIndex : 1 = false;
+
+	/** Define the index of the render material to retrieve when a primitive is hit. */
+	UPROPERTY(BlueprintReadWrite, EditAnywhere, Category = "Data|Attributes", meta = (PCG_Overridable, EditCondition = "bTraceComplex && bUseRenderMaterialIndex", EditConditionHides, DisplayAfter = "bUseRenderMaterialIndex"))
 	int32 RenderMaterialIndex = 0;
 
+	/** The index of the render material to query when a primitive is hit. Currently only works for Static Meshes. */
+	UPROPERTY(BlueprintReadWrite, EditAnywhere, Category = "Data|Attributes", meta = (PCG_Overridable, EditCondition = "bTraceComplex && !bUseRenderMaterialIndex", EditConditionHides, DisplayAfter = "bUseRenderMaterialIndex"))
+	uint8 bGetRenderMaterialIndex : 1 = false;
+
 	/** This UV Channel will be selected when retrieving UV Coordinates from a raycast query. */
 	UPROPERTY(BlueprintReadWrite, EditAnywhere, Category = "Data|Attributes", meta = (PCG_Overridable, EditCondition = "bTraceComplex && bGetUVCoords", EditConditionHides, DisplayAfter = "bGetUVCoords"))
 	int32 UVChannel = 0;
diff --git a/Engine/Plugins/PCG/Source/PCG/Public/Helpers/PCGWorldQueryHelpers.h b/Engine/Plugins/PCG/Source/PCG/Public/Helpers/PCGWorldQueryHelpers.h
--- a/Engine/Plugins/PCG/Source/PCG/Public/Helpers/PCGWorldQueryHelpers.h
+++ b/Engine/Plugins/PCG/Source/PCG/Public/Helpers/PCGWorldQueryHelpers.h
@@ -43,6 +43,8 @@ namespace PCGWorldQueryConstants
 	const FName ElementIndexAttribute = TEXT("ImpactElementIndex");
 	const FName UVCoordAttribute = TEXT("ImpactUVCoords");
 	const FName FaceIndexAttribute = TEXT("ImpactFaceIndex");
+	const FName SectionIndexAttribute = TEXT("ImpactSectionIndex");
+	const FName RenderMaterialIndexAttribute = TEXT("ImpactRenderMaterialIndex");
 }
 
 namespace PCGWorldQueryHelpers
