[PCG] Fix build health error in PCGTextureDataInterface

#rnx

diff --git a/Engine/Plugins/PCG/Source/PCG/Private/Compute/DataInterfaces/PCGTextureDataInterface.cpp b/Engine/Plugins/PCG/Source/PCG/Private/Compute/DataInterfaces/PCGTextureDataInterface.cpp
--- a/Engine/Plugins/PCG/Source/PCG/Private/Compute/DataInterfaces/PCGTextureDataInterface.cpp
+++ b/Engine/Plugins/PCG/Source/PCG/Private/Compute/DataInterfaces/PCGTextureDataInterface.cpp
@@ -595,19 +595,19 @@ void FPCGTextureDataProviderProxy::ExportTextureUAVs(const TArray<TRefCountPtr<I
 
 		if (UPCGDataBinding* Binding = DataProvider->GetDataBinding())
 		{
-			const TArray<FPCGTextureInfo>& TextureInfos = DataProvider->GetTextureInfos();
-			const TArray<FPCGTextureBindingInfo>& BindingInfos = DataProvider->GetBindingInfos();
+			const TArray<FPCGTextureInfo>& TextureInfos_GT = DataProvider->GetTextureInfos();
+			const TArray<FPCGTextureBindingInfo>& BindingInfos_GT = DataProvider->GetBindingInfos();
 			const TArray<FString>& StringTable = Binding->GetStringTable();
 
-			for (int TextureDataIndex = 0; TextureDataIndex < TextureInfos.Num(); ++TextureDataIndex)
+			for (int TextureDataIndex = 0; TextureDataIndex < TextureInfos_GT.Num(); ++TextureDataIndex)
 			{
-				const FPCGTextureInfo& TextureInfo = TextureInfos[TextureDataIndex];
+				const FPCGTextureInfo& TextureInfo = TextureInfos_GT[TextureDataIndex];
 				check(ExportedTextures.IsValidIndex(TextureDataIndex));
 				check(PinDesc.DataDescs.IsValidIndex(TextureDataIndex))
-				check(BindingInfos.IsValidIndex(TextureInfo.BindingIndex));
+				check(BindingInfos_GT.IsValidIndex(TextureInfo.BindingIndex));
 
-				const FPCGTextureBindingInfo& BindingInfo = BindingInfos[TextureInfo.BindingIndex];
-				const FTransform TextureTransform = BindingInfos[TextureInfo.BindingIndex].Transform;
+				const FPCGTextureBindingInfo& BindingInfo = BindingInfos_GT[TextureInfo.BindingIndex];
+				const FTransform TextureTransform = BindingInfos_GT[TextureInfo.BindingIndex].Transform;
 				const FVector::FReal XSize = 2.0 * TextureTransform.GetScale3D().X;
 				const FVector::FReal YSize = 2.0 * TextureTransform.GetScale3D().Y;
 				const FVector2D TexelSize = FVector2D(XSize, YSize) / BindingInfo.Size;
