[PCG] Fixed subgraph filtering when nodes are filtered too.

#rnx

diff --git a/Engine/Plugins/PCG/Source/PCG/Private/Editor/PCGGraphCustomization.cpp b/Engine/Plugins/PCG/Source/PCG/Private/Editor/PCGGraphCustomization.cpp
--- a/Engine/Plugins/PCG/Source/PCG/Private/Editor/PCGGraphCustomization.cpp
+++ b/Engine/Plugins/PCG/Source/PCG/Private/Editor/PCGGraphCustomization.cpp
@@ -71,6 +71,6 @@ bool FPCGGraphEditorCustomization::FilterSubgraph(const FSoftObjectPath& InSubgr
 	}
 	else
 	{
-		return bFilterSubgraphs && NodeFilterType == EPCGGraphEditorFiltering::Allow;
+		return bFilterSubgraphs && SubgraphFilterType == EPCGGraphEditorFiltering::Allow;
 	} 
 }
\ No newline at end of file
