[PCG] Changing graph customization does not trigger graph refreshes anymore.

#rnx

diff --git a/Engine/Plugins/PCG/Source/PCG/Private/PCGComponent.cpp b/Engine/Plugins/PCG/Source/PCG/Private/PCGComponent.cpp
--- a/Engine/Plugins/PCG/Source/PCG/Private/PCGComponent.cpp
+++ b/Engine/Plugins/PCG/Source/PCG/Private/PCGComponent.cpp
@@ -2189,7 +2189,9 @@ void UPCGComponent::RefreshAfterGraphChanged(UPCGGraphInterface* InGraph, EPCGCh
 		return;
 	}
 
-	if (ChangeType == EPCGChangeType::Cosmetic || ChangeType == EPCGChangeType::None)
+	if (ChangeType == EPCGChangeType::Cosmetic ||
+		ChangeType == EPCGChangeType::GraphCustomization ||
+		ChangeType == EPCGChangeType::None)
 	{
 		// If it is a cosmetic change (or no change), nothing to do
 		return;
diff --git a/Engine/Plugins/PCG/Source/PCG/Private/PCGGraph.cpp b/Engine/Plugins/PCG/Source/PCG/Private/PCGGraph.cpp
--- a/Engine/Plugins/PCG/Source/PCG/Private/PCGGraph.cpp
+++ b/Engine/Plugins/PCG/Source/PCG/Private/PCGGraph.cpp
@@ -1715,7 +1715,7 @@ void UPCGGraph::NotifyGraphChanged(EPCGChangeType ChangeType)
 		}
 	}
 
-	if (ChangeType != EPCGChangeType::Cosmetic)
+	if ((ChangeType & ~(EPCGChangeType::Cosmetic | EPCGChangeType::GraphCustomization)) != EPCGChangeType::None)
 	{
 		{
 			// Graph settings, nodes, graph structure can all change the higen grid sizes.
