[PCG] Removed too stringent check in landscape cache. Fixed issue with BP method that didn't infer world context automatically.

diff --git a/Engine/Plugins/PCG/Source/PCG/Private/Grid/PCGLandscapeCache.cpp b/Engine/Plugins/PCG/Source/PCG/Private/Grid/PCGLandscapeCache.cpp
--- a/Engine/Plugins/PCG/Source/PCG/Private/Grid/PCGLandscapeCache.cpp
+++ b/Engine/Plugins/PCG/Source/PCG/Private/Grid/PCGLandscapeCache.cpp
@@ -273,10 +273,7 @@ void FPCGLandscapeCacheEntry::GetInterpolatedPointInternal(const PCGLandscapeCac
 
 void FPCGLandscapeCacheEntry::GetInterpolatedPointMetadataInternal(const PCGLandscapeCache::FSafeIndices& Indices, int64& OutMetadataEntry, UPCGMetadata* OutMetadata) const
 {
-	check(bDataLoaded);
-	check(2 * Indices.X1Y1 < PositionsAndNormals.Num());
-
-	if (OutMetadata && !LayerData.IsEmpty())
+	if (bDataLoaded && OutMetadata && !LayerData.IsEmpty())
 	{
 		OutMetadata->InitializeOnSet(OutMetadataEntry);
 
@@ -286,6 +283,7 @@ void FPCGLandscapeCacheEntry::GetInterpolatedPointMetadataInternal(const PCGLand
 		{
 			const FName& CurrentLayerName = LayerDataNames[LayerIndex];
 			const TArray<uint8>& CurrentLayerData = LayerData[LayerIndex];
+			check(Indices.X1Y1 < CurrentLayerData.Num());
 
 			if (FPCGMetadataAttributeBase* Attribute = OutMetadata->GetMutableAttribute(CurrentLayerName))
 			{
diff --git a/Engine/Plugins/PCG/Source/PCG/Public/Helpers/PCGBlueprintHelpers.h b/Engine/Plugins/PCG/Source/PCG/Public/Helpers/PCGBlueprintHelpers.h
--- a/Engine/Plugins/PCG/Source/PCG/Public/Helpers/PCGBlueprintHelpers.h
+++ b/Engine/Plugins/PCG/Source/PCG/Public/Helpers/PCGBlueprintHelpers.h
@@ -81,7 +81,7 @@ public:
 	UFUNCTION(BlueprintCallable, Category = "PCG|Helpers", meta = (ScriptMethod))
 	static UE_API UPCGData* CreatePCGDataFromActor(AActor* InActor, bool bParseActor = true);
 
-	UFUNCTION(BlueprintCallable, Category = "PCG|Helpers", meta = (ScriptMethod))
+	UFUNCTION(BlueprintCallable, Category = "PCG|Helpers", meta = (WorldContext="WorldContextObject"))
 	static UE_API TArray<FPCGLandscapeLayerWeight> GetInterpolatedPCGLandscapeLayerWeights(UObject* WorldContextObject, const FVector& Location);
 
 	UFUNCTION(BLueprintCallable, Category = "PCG|Helpers", meta = (ScriptMethod))
