[PCG] Fix warning in StaticMeshDataInterface

#rnx

diff --git a/Engine/Plugins/PCG/Source/PCG/Private/Compute/DataInterfaces/PCGStaticMeshDataInterface.cpp b/Engine/Plugins/PCG/Source/PCG/Private/Compute/DataInterfaces/PCGStaticMeshDataInterface.cpp
--- a/Engine/Plugins/PCG/Source/PCG/Private/Compute/DataInterfaces/PCGStaticMeshDataInterface.cpp
+++ b/Engine/Plugins/PCG/Source/PCG/Private/Compute/DataInterfaces/PCGStaticMeshDataInterface.cpp
@@ -63,9 +63,9 @@ void UPCGStaticMeshDataInterface::GetSupportedInputs(TArray<FShaderFunctionDefin
 			.SetName(TEXT("GetTriangleIndices"))
 			.AddParam(EShaderFundamentalType::Int) // DataIndex
 			.AddParam(EShaderFundamentalType::Int) // TriangleIndex
-			.AddParam(EShaderFundamentalType::Int, 1, 0, EShaderParamModifier::Out) // OutIndex0
-			.AddParam(EShaderFundamentalType::Int, 1, 0, EShaderParamModifier::Out) // OutIndex1
-			.AddParam(EShaderFundamentalType::Int, 1, 0, EShaderParamModifier::Out); // OutIndex2
+			.AddParam(EShaderFundamentalType::Int, 0, 0, EShaderParamModifier::Out) // OutIndex0
+			.AddParam(EShaderFundamentalType::Int, 0, 0, EShaderParamModifier::Out) // OutIndex1
+			.AddParam(EShaderFundamentalType::Int, 0, 0, EShaderParamModifier::Out); // OutIndex2
 
 		OutFunctions.AddDefaulted_GetRef()
 			.SetName(TEXT("SampleTriangle"))
