[PCG] Fix issue where some characters in node titles can cause kernels to stop compiling

#rnx

diff --git a/Engine/Plugins/PCG/Source/PCG/Private/Compute/Elements/PCGCustomHLSLKernel.cpp b/Engine/Plugins/PCG/Source/PCG/Private/Compute/Elements/PCGCustomHLSLKernel.cpp
--- a/Engine/Plugins/PCG/Source/PCG/Private/Compute/Elements/PCGCustomHLSLKernel.cpp
+++ b/Engine/Plugins/PCG/Source/PCG/Private/Compute/Elements/PCGCustomHLSLKernel.cpp
@@ -715,10 +715,13 @@ void UPCGCustomHLSLKernel::InitEntryPoint()
 	// the full fname with number as this will change across executions and cause DDC misses.
 	EntryPoint += TEXT("_") + GetFName().GetPlainNameString();
 
-	const TCHAR InvalidCharacters[] = {
+	const TCHAR InvalidCharacters[] =
+	{
 		TCHAR('\"'), TCHAR('\\'), TCHAR('\''), TCHAR('\"'), TCHAR(' '), TCHAR(','), TCHAR('.'), TCHAR('|'), TCHAR('&'), TCHAR('!'),
 		TCHAR('~'), TCHAR('\n'), TCHAR('\r'), TCHAR('\t'), TCHAR('@'), TCHAR('#'), TCHAR('/'), TCHAR('('), TCHAR(')'), TCHAR('{'),
-		TCHAR('}'), TCHAR('['), TCHAR(']'), TCHAR('='), TCHAR(';'), TCHAR(':'), TCHAR('^'), TCHAR('%'), TCHAR('$'), TCHAR('`') };
+		TCHAR('}'), TCHAR('['), TCHAR(']'), TCHAR('='), TCHAR(';'), TCHAR(':'), TCHAR('^'), TCHAR('%'), TCHAR('$'), TCHAR('`'),
+		TCHAR('-'), TCHAR('+'), TCHAR('*'), TCHAR('<'), TCHAR('>'), TCHAR('?')
+	};
 
 	for (TCHAR Char : InvalidCharacters)
 	{
