[PCG] Second pass of world streaming queries emitted from runtime generation.

Before a cell of grid X is generated, we check the largest dependent parent grid Y for which streaming queries have been enabled.

#rnx

diff --git a/Engine/Plugins/PCG/Source/PCG/Private/PCGComponent.cpp b/Engine/Plugins/PCG/Source/PCG/Private/PCGComponent.cpp
--- a/Engine/Plugins/PCG/Source/PCG/Private/PCGComponent.cpp
+++ b/Engine/Plugins/PCG/Source/PCG/Private/PCGComponent.cpp
@@ -3121,6 +3121,11 @@ const UPCGComponent* UPCGComponent::GetConstOriginalComponent() const
 	return ensure(PartitionActor) ? PartitionActor->GetOriginalComponent(this) : this;
 }
 
+bool UPCGComponent::DoesGridDependOnWorldStreaming(uint32 InGridSize) const
+{
+	return SchedulingPolicy && SchedulingPolicy->DoesGridDependOnWorldStreaming(InGridSize);
+}
+
 UPCGData* UPCGComponent::CreateActorPCGData() const
 {
 	return CreateActorPCGData(GetOwner(), bParseActorComponents);
diff --git a/Engine/Plugins/PCG/Source/PCG/Private/PCGGraph.cpp b/Engine/Plugins/PCG/Source/PCG/Private/PCGGraph.cpp
--- a/Engine/Plugins/PCG/Source/PCG/Private/PCGGraph.cpp
+++ b/Engine/Plugins/PCG/Source/PCG/Private/PCGGraph.cpp
@@ -1445,6 +1445,53 @@ void UPCGGraph::CacheGridSizesInternalNoLock() const
 	CachedGridInfo = MoveTemp(GridInfo);
 }
 
+void UPCGGraph::GetParentGridSizes(const uint32 InChildGridSize, PCGHiGenGrid::FSizeArray& OutParentGridSizes) const
+{
+	if (!PCGHiGenGrid::IsValidGridSize(InChildGridSize))
+	{
+		// Grid size is 0 or Unbounded or some other invalid value, and will not have any parent grids.
+		return;
+	}
+
+	FScopeLock Lock(&CachedGridInfoLock);
+	
+	if (const PCGHiGenGrid::FSizeArray* FoundGridSizeInfo = ChildGridSizeToParentGridSizes.Find(InChildGridSize))
+	{
+		OutParentGridSizes = *FoundGridSizeInfo;
+		return;
+	}
+
+	// No higen means no parent dependencies.
+	if (IsHierarchicalGenerationEnabled())
+	{
+		// Use in parent grid calculations to cache results. Note this map is specific to InChildGridSize 
+		const uint32 DefaultGridSize = GetDefaultGridSize();
+
+		for (const UPCGNode* Node : Nodes)
+		{
+			if (Node)
+			{
+				PCGHiGenGrid::FSizeArray NodeAllGridSizes = CalculateNodeGridSizesRecursiveNoLock(Node, DefaultGridSize);
+
+				if (!NodeAllGridSizes.IsEmpty() && NodeAllGridSizes.Last() == InChildGridSize)
+				{
+					for (uint32 GridSize : NodeAllGridSizes)
+					{
+						OutParentGridSizes.AddUnique(GridSize);
+					}
+				}
+			}
+		}
+
+		OutParentGridSizes.Remove(InChildGridSize);
+
+		// Always output in descending order.
+		OutParentGridSizes.Sort([](const uint32& A, const uint32& B) { return A > B; });
+	}
+
+	ChildGridSizeToParentGridSizes.Add(InChildGridSize, OutParentGridSizes);
+}
+
 double UPCGGraph::GetGridGenerationRadiusFromGrid(EPCGHiGenGrid Grid) const
 {
 	if (Grid == EPCGHiGenGrid::Unbounded || Grid == EPCGHiGenGrid::Uninitialized)
@@ -1671,10 +1718,27 @@ void UPCGGraph::NotifyGraphChanged(EPCGChangeType ChangeType)
 {
 	TRACE_CPUPROFILER_EVENT_SCOPE(UPCGGraph::NotifyGraphChanged);
 
-	if (!!(ChangeType & EPCGChangeType::Structural))
+	const bool bNonCosmeticChange = (ChangeType & ~(EPCGChangeType::Cosmetic | EPCGChangeType::GraphCustomization)) != EPCGChangeType::None;
+
+	if (bNonCosmeticChange)
 	{
-		FScopeLock Lock(&CachedGridInfoLock);
-		CachedGridInfo.Reset();
+		// Graph settings, nodes, graph structure can all change the higen grid sizes.
+		{
+			FWriteScopeLock GridSizeLock(NodeToGridSizeLock);
+			NodeToGridSize.Reset();
+			NodeToAllGridSizes.Reset();
+		}
+
+		{
+			FScopeLock Lock(&CachedGridInfoLock);
+			CachedGridInfo.Reset();
+			ChildGridSizeToParentGridSizes.Reset();
+		}
+
+		// Any non-trivial change to graph may change compute graphs.
+		check(IsInGameThread());
+		AvailableComputeGraphInstances.Empty(AvailableComputeGraphInstances.Num());
+		AllComputeGraphInstances.Empty(AllComputeGraphInstances.Num());
 	}
 
 	if (GraphChangeNotificationsDisableCounter > 0)
@@ -1715,14 +1779,8 @@ void UPCGGraph::NotifyGraphChanged(EPCGChangeType ChangeType)
 		}
 	}
 
-	if ((ChangeType & ~(EPCGChangeType::Cosmetic | EPCGChangeType::GraphCustomization)) != EPCGChangeType::None)
+	if (bNonCosmeticChange)
 	{
-		{
-			// Graph settings, nodes, graph structure can all change the higen grid sizes.
-			FWriteScopeLock GridSizeLock(NodeToGridSizeLock);
-			NodeToGridSize.Reset();
-		}
-
 		// Also notify other systems that this graph changed, only if the owner is not a PCG Component nor PCG Subgraph.
 		// They already have their own system to trigger a refresh.
 		const UObject* Outer = GetOuter();
@@ -1731,11 +1789,6 @@ void UPCGGraph::NotifyGraphChanged(EPCGChangeType ChangeType)
 			FPropertyChangedEvent EmptyEvent{ nullptr };
 			FCoreUObjectDelegates::OnObjectPropertyChanged.Broadcast(this, EmptyEvent);
 		}
-
-		// Any non-trivial change to graph may change compute graphs.
-		check(IsInGameThread());
-		AvailableComputeGraphInstances.Empty(AvailableComputeGraphInstances.Num());
-		AllComputeGraphInstances.Empty(AllComputeGraphInstances.Num());
 	}
 
 	OnGraphChangedDelegate.Broadcast(this, ChangeType);
@@ -1770,6 +1823,8 @@ void UPCGGraph::OnNodeChanged(UPCGNode* InNode, EPCGChangeType ChangeType)
 
 			FWriteScopeLock Lock(NodeToGridSizeLock);
 			NodeToGridSize.Reset();
+			NodeToAllGridSizes.Reset();
+			ChildGridSizeToParentGridSizes.Reset();
 		}
 	}
 
@@ -2117,6 +2172,51 @@ uint32 UPCGGraph::CalculateNodeGridSizeRecursive_Unsafe(const UPCGNode* InNode,
 	return GridSize;
 }
 
+PCGHiGenGrid::FSizeArray UPCGGraph::CalculateNodeGridSizesRecursiveNoLock(const UPCGNode* InNode, uint32 InDefaultGridSize) const
+{
+	if (!InNode)
+	{
+		return {};
+	}
+
+	if (const PCGHiGenGrid::FSizeArray* CachedGridSizes = NodeToAllGridSizes.Find(InNode))
+	{
+		return *CachedGridSizes;
+	}
+
+	PCGHiGenGrid::FSizeArray AllGridSizes;
+
+	const uint32 NodeGridSize = CalculateNodeGridSizeRecursive_Unsafe(InNode, InDefaultGridSize);
+	AllGridSizes.AddUnique(NodeGridSize);
+
+	for (const UPCGPin* Pin : InNode->GetInputPins())
+	{
+		if (Pin)
+		{
+			for (const UPCGEdge* Edge : Pin->Edges)
+			{
+				const UPCGPin* OtherPin = Edge ? Edge->InputPin : nullptr;
+				if (OtherPin && OtherPin->Node.Get())
+				{
+					PCGHiGenGrid::FSizeArray InputGridSizes = CalculateNodeGridSizesRecursiveNoLock(OtherPin->Node.Get(), InDefaultGridSize);
+
+					for (uint32 GridSize : InputGridSizes)
+					{
+						AllGridSizes.AddUnique(GridSize);
+					}
+				}
+			}
+		}
+	}
+
+	// Descending order.
+	AllGridSizes.Sort([](const uint32& A, const uint32& B) { return A > B; });
+
+	NodeToAllGridSizes.Add(InNode, AllGridSizes);
+
+	return AllGridSizes;
+}
+
 void UPCGGraph::AddUserParameters(const TArray<FPropertyBagPropertyDesc>& InDescs, const UPCGGraph* InOptionalOriginalGraph)
 {
 	UserParameters.AddProperties(InDescs);
diff --git a/Engine/Plugins/PCG/Source/PCG/Private/RuntimeGen/PCGRuntimeGenScheduler.cpp b/Engine/Plugins/PCG/Source/PCG/Private/RuntimeGen/PCGRuntimeGenScheduler.cpp
--- a/Engine/Plugins/PCG/Source/PCG/Private/RuntimeGen/PCGRuntimeGenScheduler.cpp
+++ b/Engine/Plugins/PCG/Source/PCG/Private/RuntimeGen/PCGRuntimeGenScheduler.cpp
@@ -89,9 +89,9 @@ namespace PCGRuntimeGenSchedulerHelpers
 			}
 		}));
 
-	static TAutoConsoleVariable<bool> CVarCheckWorldStreamedIn(
-		TEXT("pcg.RuntimeGeneration.CheckWorldStreamingStatusForLocalComponents"),
-		false,
+	static TAutoConsoleVariable<bool> CVarEnableWorldStreamingQueries(
+		TEXT("pcg.RuntimeGeneration.EnableWorldStreamingQueries"),
+		true,
 		TEXT("Checks that the world is streamed in before triggering generation of local (partitioned) components."));
 
 	static TAutoConsoleVariable<int32> CVarFramesBeforeFirstGenerate(
@@ -332,7 +332,7 @@ void FPCGRuntimeGenScheduler::TickQueueComponentsForGeneration(
 	TArray<FWorldPartitionStreamingQuerySource> StreamingQuerySources;
 	StreamingQuerySources.Reserve(1);
 	
-	bool bCheckStreaming = PCGRuntimeGenSchedulerHelpers::CVarCheckWorldStreamedIn.GetValueOnGameThread();
+	bool bCheckStreaming = PCGRuntimeGenSchedulerHelpers::CVarEnableWorldStreamingQueries.GetValueOnGameThread();
 	const UWorldPartitionSubsystem* WorldPartitionSubsystem = bCheckStreaming ? UWorld::GetSubsystem<UWorldPartitionSubsystem>(Inputs.PCGWorldActor->GetWorld()) : nullptr;
 
 	if (WorldPartitionSubsystem)
@@ -397,6 +397,39 @@ void FPCGRuntimeGenScheduler::TickQueueComponentsForGeneration(
 				continue;
 			}
 
+			// For each relevant grid index, the largest grid size that has been marked in the runtime gen policy as depending on world streaming.
+			PCGHiGenGrid::FSizeArray WorldStreamingQueryGridSizes;
+
+			if (bCheckStreaming)
+			{
+				WorldStreamingQueryGridSizes.SetNumUninitialized(GridSizes.Num());
+
+				for (int GridIndex = 0; GridIndex < GridSizes.Num(); ++GridIndex)
+				{
+					WorldStreamingQueryGridSizes[GridIndex] = PCGHiGenGrid::UninitializedGridSize();
+
+					PCGHiGenGrid::FSizeArray ParentGridsDescending;
+					OriginalComponent->GetGraph()->GetParentGridSizes(GridSizes[GridIndex], ParentGridsDescending);
+					for (uint32 ParentGridSize : ParentGridsDescending)
+					{
+						if (OriginalComponent->DoesGridDependOnWorldStreaming(ParentGridSize) && ensure(ParentGridSize > GridSizes[GridIndex]))
+						{
+							WorldStreamingQueryGridSizes[GridIndex] = ParentGridSize;
+							break;
+						}
+					}
+
+					if (WorldStreamingQueryGridSizes[GridIndex] == PCGHiGenGrid::UninitializedGridSize() && OriginalComponent->DoesGridDependOnWorldStreaming(GridSizes[GridIndex]))
+					{
+						WorldStreamingQueryGridSizes[GridIndex] = GridSizes[GridIndex];
+					}
+				}
+			}
+			else
+			{
+				WorldStreamingQueryGridSizes.SetNumZeroed(GridSizes.Num());
+			}
+
 			const EPCGHiGenGrid MaxGrid = bHasUnbounded ? EPCGHiGenGrid::Unbounded : PCGHiGenGrid::GridSizeToGrid(GridSizes[0]);
 			const double MaxGenerationRadius = OriginalComponent->GetGenerationRadiusFromGrid(MaxGrid);
 
@@ -501,26 +534,28 @@ void FPCGRuntimeGenScheduler::TickQueueComponentsForGeneration(
 								{
 									bool bStreamingComplete = true;
 
-									if (bCheckStreaming)
+									if (WorldStreamingQueryGridSizes[GridIndex] == GridSize)
 									{
-										if (GridIndex > 0)
+										bStreamingComplete = IsWorldStreamingComplete(Center, GridSize);
+
+										if (!bStreamingComplete)
 										{
-											// We use the topmost parent because we don't currently know which ancestor(s) this grid actually relies on.
-											// todo_pcg: Precompute the topmost ancestor that relies on world streaming.
-											const uint32 ParentGridSize = GridSizes[0];
-											const FVector ParentCenter = UPCGActorHelpers::GetCellCenter(Center, ParentGridSize, Key.GetGridDescriptor().Is2DGrid());
-
-											bStreamingComplete = IsWorldStreamingComplete(ParentCenter, ParentGridSize);
-
-											if (!bStreamingComplete)
-											{
-												UE_LOG(LogPCG, VeryVerbose, TEXT("Cell %u (%d, %d, %d) rejected as parent on grid size %u was not fully loaded."),
-													GridSize, GridCoords.X, GridCoords.Y, GridCoords.Z, ParentGridSize);
-											}
+											UE_LOG(LogPCG, VeryVerbose, TEXT("Cell %u (%d, %d, %d) rejected as world is not fully loaded."),
+												GridSize, GridCoords.X, GridCoords.Y, GridCoords.Z);
 										}
-										else
+									}
+									else if (WorldStreamingQueryGridSizes[GridIndex] != PCGHiGenGrid::UninitializedGridSize())
+									{
+										// Check world loaded status using the pre-calculated largest parent grid that depends on world streaming.
+										const uint32 ParentGridSize = WorldStreamingQueryGridSizes[GridIndex];
+										const FVector ParentCenter = UPCGActorHelpers::GetCellCenter(Center, ParentGridSize, Key.GetGridDescriptor().Is2DGrid());
+
+										bStreamingComplete = IsWorldStreamingComplete(ParentCenter, ParentGridSize);
+
+										if (!bStreamingComplete)
 										{
-											bStreamingComplete = IsWorldStreamingComplete(Center, GridSize);
+											UE_LOG(LogPCG, VeryVerbose, TEXT("Cell %u (%d, %d, %d) rejected as parent on grid size %u is not fully loaded."),
+												GridSize, GridCoords.X, GridCoords.Y, GridCoords.Z, ParentGridSize);
 										}
 									}
 
diff --git a/Engine/Plugins/PCG/Source/PCG/Public/PCGComponent.h b/Engine/Plugins/PCG/Source/PCG/Public/PCGComponent.h
--- a/Engine/Plugins/PCG/Source/PCG/Public/PCGComponent.h
+++ b/Engine/Plugins/PCG/Source/PCG/Public/PCGComponent.h
@@ -186,6 +186,8 @@ public:
 
 	UPCGSchedulingPolicyBase* GetRuntimeGenSchedulingPolicy() const { return SchedulingPolicy; }
 
+	UE_API bool DoesGridDependOnWorldStreaming(uint32 InGridSize) const;
+
 	UE_API bool CanPartition() const;
 
 	UE_API UPCGGraph* GetGraph() const;
diff --git a/Engine/Plugins/PCG/Source/PCG/Public/PCGGraph.h b/Engine/Plugins/PCG/Source/PCG/Public/PCGGraph.h
--- a/Engine/Plugins/PCG/Source/PCG/Public/PCGGraph.h
+++ b/Engine/Plugins/PCG/Source/PCG/Public/PCGGraph.h
@@ -392,6 +392,9 @@ public:
 	/** Determine the relevant grid sizes by inspecting all HiGenGridSize nodes. */
 	UE_API void GetGridSizes(PCGHiGenGrid::FSizeArray& OutGridSizes, bool& bOutHasUnbounded) const;
 
+	/** Returns all parent grid sizes for the given child grid size, calculated by inspecting nodes. */
+	UE_API void GetParentGridSizes(const uint32 InChildGridSize, PCGHiGenGrid::FSizeArray& OutParentGridSizes) const;
+
 	/** Returns exponential on grid size, which represents a shift in the grid */
 	uint32 GetGridExponential() const { return HiGenExponential; }
 
@@ -465,6 +468,9 @@ protected:
 	/** Calculates node grid size. Not thread safe, must be called within write lock. */
 	UE_API uint32 CalculateNodeGridSizeRecursive_Unsafe(const UPCGNode* InNode, uint32 InDefaultGridSize) const;
 
+	/** Calculates all parent grid sizes that a given node depends on. Not thread safe, must be called within write lock. */
+	UE_API PCGHiGenGrid::FSizeArray CalculateNodeGridSizesRecursiveNoLock(const UPCGNode* InNode, uint32 InDefaultGridSize) const;
+
 	UPROPERTY(BlueprintReadOnly, VisibleAnywhere, Category = Graph, meta = (NoResetToDefault))
 	TArray<TObjectPtr<UPCGNode>> Nodes;
 
@@ -513,6 +519,10 @@ protected:
 
 	/** Execution grid size for nodes. */
 	mutable TMap<const UPCGNode*, uint32> NodeToGridSize;
+
+	/** All execution grid sizes for nodes (node grid size and all parent grid sizes that it depends on). */
+	mutable TMap<const UPCGNode*, PCGHiGenGrid::FSizeArray> NodeToAllGridSizes;
+
 	mutable FRWLock NodeToGridSizeLock;
 
 	struct FGridInfo
@@ -523,6 +533,7 @@ protected:
 
 	/** Cached information about higen grids present in graph. */
 	mutable TOptional<FGridInfo> CachedGridInfo;
+	mutable TMap<uint32, PCGHiGenGrid::FSizeArray> ChildGridSizeToParentGridSizes;
 	mutable FCriticalSection CachedGridInfoLock;
 
 	/** Sets whether this graph is marked as editor-only; note that the IsEditorOnly call depends on the local graph value and the value in all subgraphs, recursively. */
diff --git a/Engine/Plugins/PCG/Source/PCG/Public/RuntimeGen/SchedulingPolicies/PCGSchedulingPolicyBase.h b/Engine/Plugins/PCG/Source/PCG/Public/RuntimeGen/SchedulingPolicies/PCGSchedulingPolicyBase.h
--- a/Engine/Plugins/PCG/Source/PCG/Public/RuntimeGen/SchedulingPolicies/PCGSchedulingPolicyBase.h
+++ b/Engine/Plugins/PCG/Source/PCG/Public/RuntimeGen/SchedulingPolicies/PCGSchedulingPolicyBase.h
@@ -2,6 +2,8 @@
 
 #pragma once
 
+#include "PCGCommon.h"
+
 #include "PCGSchedulingPolicyBase.generated.h"
 
 #define UE_API PCG_API
@@ -39,6 +41,20 @@ public:
 	void SetShouldDisplayProperties(bool bInShouldDisplayProperties) { bShouldDisplayProperties = bInShouldDisplayProperties; }
 #endif // WITH_EDITOR
 
+	bool DoesGridDependOnWorldStreaming(uint32 InGridSize) const
+	{
+		const EPCGHiGenGrid Grid = PCGHiGenGrid::GridSizeToGrid(InGridSize);
+		return Grid != EPCGHiGenGrid::Uninitialized && GridsDependentOnWorldStreaming.Contains(Grid);
+	}
+
+public:
+	/**
+	* Grids that depend on world streaming. Cells on these grids will only generate if the world within the generation volume is fully streamed. Enable this
+	* if the graph depends on actors in the world such as a landscape or a set of streamable actors.
+	*/
+	UPROPERTY(EditAnywhere, Category = "RuntimeGeneration|Scheduling Policy Parameters", meta = (EditCondition = "bShouldDisplayProperties", EditConditionHides, HideEditConditionToggle))
+	TArray<EPCGHiGenGrid> GridsDependentOnWorldStreaming;
+
 #if WITH_EDITORONLY_DATA
 private:
 	/** Hidden property to control display of SchedulingPolicy properties. */
