[PCG] Added debug draw info for currently executing components to help understand what is executing.

#rnx

diff --git a/Engine/Plugins/PCG/Source/PCG/Private/Elements/PCGGenerateGrassMaps.cpp b/Engine/Plugins/PCG/Source/PCG/Private/Elements/PCGGenerateGrassMaps.cpp
--- a/Engine/Plugins/PCG/Source/PCG/Private/Elements/PCGGenerateGrassMaps.cpp
+++ b/Engine/Plugins/PCG/Source/PCG/Private/Elements/PCGGenerateGrassMaps.cpp
@@ -426,6 +426,14 @@ bool FPCGGenerateGrassMapsElement::ExecuteInternal(FPCGContext* InContext) const
 		
 		if (!bAllReady)
 		{
+#if UE_ENABLE_DEBUG_DRAWING
+			if (PCGSystemSwitches::CVarPCGDebugDrawGeneratedCells.GetValueOnGameThread())
+			{
+				FColor DebugColor = FColor::Yellow;
+				PCGHelpers::DebugDrawGenerationVolume(InContext, &DebugColor);
+			}
+#endif
+
 			// Sleep until next frame, no use spinning on this.
 			Context->bIsPaused = true;
 			FPCGModule::GetPCGModuleChecked().ExecuteNextTick([ContextHandle = InContext->GetOrCreateHandle()]()
diff --git a/Engine/Plugins/PCG/Source/PCG/Private/Helpers/PCGHelpers.cpp b/Engine/Plugins/PCG/Source/PCG/Private/Helpers/PCGHelpers.cpp
--- a/Engine/Plugins/PCG/Source/PCG/Private/Helpers/PCGHelpers.cpp
+++ b/Engine/Plugins/PCG/Source/PCG/Private/Helpers/PCGHelpers.cpp
@@ -10,6 +10,7 @@
 #include "PCGSubsystem.h"
 #include "PCGWorldActor.h"
 #include "Grid/PCGPartitionActor.h"
+#include "Helpers/PCGActorHelpers.h"
 
 #include "Components/BillboardComponent.h"
 #include "Landscape.h"
@@ -742,6 +743,57 @@ namespace PCGHelpers
 
 		return RandomIndices;
 	}
+
+#if UE_ENABLE_DEBUG_DRAWING
+	void DebugDrawGenerationVolume(FPCGContext* InContext, const FColor* InOverrideColor)
+	{
+		if (!InContext || !InContext->ExecutionSource.IsValid())
+		{
+			return;
+		}
+
+		const UPCGComponent* Component = Cast<UPCGComponent>(InContext->ExecutionSource.Get());
+		const FColor Color = InOverrideColor ? *InOverrideColor : FColor::MakeRandomSeededColor(Component ? (2 + Component->GetGenerationGridSize()) : 0);
+
+		PCGHelpers::ExecuteOnGameThread(TEXT("Execution debug"), [ExecutionSource = InContext->ExecutionSource, Color]()
+		{
+			UWorld* World = ExecutionSource.IsValid() ? ExecutionSource->GetExecutionState().GetWorld() : nullptr;
+			if (!World)
+			{
+				return;
+			}
+
+			const FBox Bounds = ExecutionSource->GetExecutionState().GetBounds();
+			const UPCGComponent* ComponentInner = Cast<UPCGComponent>(ExecutionSource.Get());
+
+			// Unfortunately DrawDebugString does not work in editor.
+			if (ComponentInner && World->IsGameWorld())
+			{
+				FString Text;
+				if (ComponentInner->GetGenerationGrid() == EPCGHiGenGrid::Unbounded)
+				{
+					if (ComponentInner->GetOwner())
+					{
+						Text = ComponentInner->GetOwner()->GetName();
+					}
+				}
+				else
+				{
+					const FIntVector CellCoord = UPCGActorHelpers::GetCellCoord(Bounds.GetCenter(), ComponentInner->GetGenerationGridSize(), ComponentInner->Use2DGrid());
+					Text = FString::Format(TEXT("{0} ({1}, {2}, {3})"), { ComponentInner->GetGenerationGridSize(), CellCoord.X, CellCoord.Y, CellCoord.Z });
+				}
+
+				DrawDebugString(World, Bounds.GetCenter() + FVector(0.0, 0.0, 100.0), Text, /*TestBaseActor=*/nullptr, Color, /*Duration=*/0.0f);
+			}
+
+			DrawDebugBox(World, Bounds.GetCenter(), Bounds.GetExtent(), Color, /*bPersistentLines=*/false, /*LifeTime=*/0.0f);
+
+			// Add additional boxes to visually distinguish this vis and make the wireframe look "thicker".
+			DrawDebugBox(World, Bounds.GetCenter(), Bounds.GetExtent() * FVector(0.95, 1.0, 1.0), Color, /*bPersistentLines=*/false, /*LifeTime=*/0.0f);
+			DrawDebugBox(World, Bounds.GetCenter(), Bounds.GetExtent() * FVector(1.0, 0.95, 1.0), Color, /*bPersistentLines=*/false, /*LifeTime=*/0.0f);
+		});
+	}
+#endif // UE_ENABLE_DEBUG_DRAWING
 }
 
 #undef LOCTEXT_NAMESPACE
diff --git a/Engine/Plugins/PCG/Source/PCG/Private/PCGCommon.cpp b/Engine/Plugins/PCG/Source/PCG/Private/PCGCommon.cpp
--- a/Engine/Plugins/PCG/Source/PCG/Private/PCGCommon.cpp
+++ b/Engine/Plugins/PCG/Source/PCG/Private/PCGCommon.cpp
@@ -107,6 +107,13 @@ namespace PCGSystemSwitches
 		TEXT("pcg.Experimental.ReleaseTransientResourcesEarly"),
 		false,
 		TEXT("Aggressively release transient (esp. GPU) data when possible."));
+
+#if UE_ENABLE_DEBUG_DRAWING 
+	TAutoConsoleVariable<bool> CVarPCGDebugDrawGeneratedCells(
+		TEXT("pcg.GraphExecution.DebugDrawGeneratedCells"),
+		false,
+		TEXT("Draws debug boxes around any cell that is executing, colored by grid level, and in PIE/standalone labels grid size and coords."));
+#endif // UE_ENABLE_DEBUG_DRAWING 
 }
 
 namespace PCGHiGenGrid
diff --git a/Engine/Plugins/PCG/Source/PCG/Private/PCGElement.cpp b/Engine/Plugins/PCG/Source/PCG/Private/PCGElement.cpp
--- a/Engine/Plugins/PCG/Source/PCG/Private/PCGElement.cpp
+++ b/Engine/Plugins/PCG/Source/PCG/Private/PCGElement.cpp
@@ -11,12 +11,14 @@
 #include "Data/PCGPointArrayData.h"
 #include "Data/PCGPointData.h"
 #include "Graph/PCGGraphCache.h"
+#include "Helpers/PCGActorHelpers.h"
 #include "Helpers/PCGHelpers.h"
 
 #if WITH_EDITOR
 #include "PCGDataVisualization.h"
 #endif
 
+#include "DrawDebugHelpers.h"
 #include "HAL/IConsoleManager.h"
 #include "Utils/PCGExtraCapture.h"
 
@@ -204,6 +206,13 @@ bool IPCGElement::Execute(FPCGContext* Context) const
 				TRACE_CPUPROFILER_EVENT_SCOPE(EPCGExecutionPhase::Execute);
 				PCG_ELEMENT_EXECUTION_BREAKPOINT();
 
+#if UE_ENABLE_DEBUG_DRAWING 
+				if (PCGSystemSwitches::CVarPCGDebugDrawGeneratedCells.GetValueOnAnyThread())
+				{
+					PCGHelpers::DebugDrawGenerationVolume(Context);
+				}
+#endif // UE_ENABLE_DEBUG_DRAWING
+
 				if (ExecuteInternal(Context))
 				{
 					Context->CurrentPhase = EPCGExecutionPhase::PostExecute;
diff --git a/Engine/Plugins/PCG/Source/PCG/Private/RuntimeGen/PCGRuntimeGenScheduler.cpp b/Engine/Plugins/PCG/Source/PCG/Private/RuntimeGen/PCGRuntimeGenScheduler.cpp
--- a/Engine/Plugins/PCG/Source/PCG/Private/RuntimeGen/PCGRuntimeGenScheduler.cpp
+++ b/Engine/Plugins/PCG/Source/PCG/Private/RuntimeGen/PCGRuntimeGenScheduler.cpp
@@ -372,6 +372,10 @@ void FPCGRuntimeGenScheduler::TickQueueComponentsForGeneration(
 		return bIsLoaded;
 	};
 
+#if UE_ENABLE_DEBUG_DRAWING
+	const bool bDebugDrawGenerationSources = PCGSystemSwitches::CVarPCGDebugDrawGeneratedCells.GetValueOnGameThread();
+#endif
+
 	// Collect local components from all partitioned components.
 	for (UPCGComponent* OriginalComponent : Inputs.AllPartitionedComponents)
 	{
@@ -441,6 +445,21 @@ void FPCGRuntimeGenScheduler::TickQueueComponentsForGeneration(
 					continue;
 				}
 
+#if UE_ENABLE_DEBUG_DRAWING
+				if (bDebugDrawGenerationSources)
+				{
+					if (TOptional<FVector> Pos = GenSource->GetPosition())
+					{
+						if (const UObject* GenSourceObject = Cast<UObject>(GenSource))
+						{
+							DrawDebugString(World, *GenSourcePositionOptional, GenSourceObject->GetName(), /*TestBaseActor=*/nullptr, FColor::Red, /*Duration=*/0.0f);
+						}
+
+						DrawDebugSphere(World, *GenSourcePositionOptional, MaxGenerationRadius, 64, FColor::Red, /*bPersistentLines=*/false, /*LifeTime=*/0.0f);
+					}
+				}
+#endif // UE_ENABLE_DEBUG_DRAWING
+
 				const FVector GenSourcePosition = GenSourcePositionOptional.GetValue();
 				const FBox OriginalComponentBounds = OriginalComponent->GetGridBounds();
 				const bool bIs2DGrid = OriginalComponent->Use2DGrid();
diff --git a/Engine/Plugins/PCG/Source/PCG/Public/Helpers/PCGHelpers.h b/Engine/Plugins/PCG/Source/PCG/Public/Helpers/PCGHelpers.h
--- a/Engine/Plugins/PCG/Source/PCG/Public/Helpers/PCGHelpers.h
+++ b/Engine/Plugins/PCG/Source/PCG/Public/Helpers/PCGHelpers.h
@@ -173,6 +173,10 @@ namespace PCGHelpers
 			::ExecuteOnGameThread(DebugName, std::forward<FunctorType>(Functor));
 		}
 	}
+
+#if UE_ENABLE_DEBUG_DRAWING 
+	PCG_API void DebugDrawGenerationVolume(FPCGContext* InContext, const FColor* InOverrideColor = nullptr);
+#endif
 }
 
 /** Holds function prototypes used to match against actor function signatures. */
diff --git a/Engine/Plugins/PCG/Source/PCG/Public/PCGCommon.h b/Engine/Plugins/PCG/Source/PCG/Public/PCGCommon.h
--- a/Engine/Plugins/PCG/Source/PCG/Public/PCGCommon.h
+++ b/Engine/Plugins/PCG/Source/PCG/Public/PCGCommon.h
@@ -299,6 +299,7 @@ namespace PCGSystemSwitches
 #endif
 
 	extern PCG_API TAutoConsoleVariable<bool> CVarReleaseTransientResourcesEarly;
+	extern PCG_API TAutoConsoleVariable<bool> CVarPCGDebugDrawGeneratedCells;
 }
 
 /** Describes space referential for operations that create data */
