[PCG] Fix loctext warning in PCGCustomHLSLKernel

#rnx

diff --git a/Engine/Plugins/PCG/Source/PCG/Private/Compute/Elements/PCGCustomHLSLKernel.cpp b/Engine/Plugins/PCG/Source/PCG/Private/Compute/Elements/PCGCustomHLSLKernel.cpp
--- a/Engine/Plugins/PCG/Source/PCG/Private/Compute/Elements/PCGCustomHLSLKernel.cpp
+++ b/Engine/Plugins/PCG/Source/PCG/Private/Compute/Elements/PCGCustomHLSLKernel.cpp
@@ -1344,7 +1344,7 @@ bool UPCGCustomHLSLKernel::PerformStaticValidation()
 				{
 #if PCG_KERNEL_LOGGING_ENABLED
 					StaticLogEntries.Emplace(FText::Format(
-						LOCTEXT("InvalidElementCount", "Invalid fixed element count 2D ({0}, {1}) on output pin '{2}'. Must be greater than 0."),
+						LOCTEXT("InvalidElementCount2D", "Invalid fixed element count 2D ({0}, {1}) on output pin '{2}'. Must be greater than 0."),
 						FText::AsNumber(Props.ElementCount2D.X),
 						FText::AsNumber(Props.ElementCount2D.Y),
 						FText::FromName(Properties.Label)),
