[PCG] Clean up Custom HLSL node UI and minor tweaks to API

#rnx

diff --git a/Engine/Plugins/PCG/Source/PCG/Private/Compute/Elements/PCGCustomHLSL.cpp b/Engine/Plugins/PCG/Source/PCG/Private/Compute/Elements/PCGCustomHLSL.cpp
--- a/Engine/Plugins/PCG/Source/PCG/Private/Compute/Elements/PCGCustomHLSL.cpp
+++ b/Engine/Plugins/PCG/Source/PCG/Private/Compute/Elements/PCGCustomHLSL.cpp
@@ -353,12 +353,11 @@ void UPCGCustomHLSLSettings::UpdateInputDeclarations()
 
 		if (KernelType == EPCGKernelType::PointGenerator)
 		{
-			// @todo_pcg: NumPoints should also be ElementCount.
-			InputDeclarations += FString::Format(TEXT("const uint NumPoints = {0};\n\n"), { PointCount });
+			InputDeclarations += FString::Format(TEXT("const uint NumElements = {0};\n\n"), { PointCount });
 		}
 		else if (KernelType == EPCGKernelType::TextureGenerator)
 		{
-			InputDeclarations += FString::Format(TEXT("const uint2 ElementCount = uint2({0}, {1});\n\n"), { TextureSize.X, TextureSize.Y });
+			InputDeclarations += FString::Format(TEXT("const uint2 NumElements = uint2({0}, {1});\n\n"), { NumElements2D.X, NumElements2D.Y });
 		}
 
 		InputDeclarations += TEXT("/*** INPUT PER-THREAD CONSTANTS ***/\n\n");
@@ -940,6 +939,11 @@ void UPCGCustomHLSLSettings::UpdatePinSettings()
 			// Texture outputs are single-data for now.
 			Properties.bAllowMultipleData = false;
 			Properties.bAllowEditMultipleData = false;
+			Properties.PropertiesGPU.bShowTexturePinSettings = true;
+		}
+		else
+		{
+			Properties.PropertiesGPU.bShowTexturePinSettings = false;
 		}
 
 		Properties.PropertiesGPU.bAllowEditInitMode = !bInitModeDrivenByKernel;
diff --git a/Engine/Plugins/PCG/Source/PCG/Private/Compute/Elements/PCGCustomHLSL.h b/Engine/Plugins/PCG/Source/PCG/Private/Compute/Elements/PCGCustomHLSL.h
--- a/Engine/Plugins/PCG/Source/PCG/Private/Compute/Elements/PCGCustomHLSL.h
+++ b/Engine/Plugins/PCG/Source/PCG/Private/Compute/Elements/PCGCustomHLSL.h
@@ -136,11 +136,11 @@ protected:
 	UPROPERTY(EditAnywhere, Category = "Settings")
 	EPCGKernelType KernelType = EPCGKernelType::PointProcessor;
 
-	UPROPERTY(EditAnywhere, DisplayName = "Number of Points", Category = "Settings", meta = (EditCondition = "KernelType == EPCGKernelType::PointGenerator", EditConditionHides))
+	UPROPERTY(EditAnywhere, DisplayName = "Num Elements", Category = "Settings", meta = (EditCondition = "KernelType == EPCGKernelType::PointGenerator", EditConditionHides))
 	int PointCount = 256;
 
-	UPROPERTY(EditAnywhere, DisplayName = "Element Count", Category = "Settings", meta = (EditCondition = "KernelType == EPCGKernelType::TextureGenerator", EditConditionHides))
-	FIntPoint TextureSize = FIntPoint(64, 64);
+	UPROPERTY(EditAnywhere, DisplayName = "Num Elements", Category = "Settings", meta = (EditCondition = "KernelType == EPCGKernelType::TextureGenerator", EditConditionHides))
+	FIntPoint NumElements2D = FIntPoint(64, 64);
 
 	UPROPERTY(EditAnywhere, Category = "Settings|Thread Count", meta = (EditCondition = "KernelType == EPCGKernelType::Custom", EditConditionHides))
 	EPCGDispatchThreadCount DispatchThreadCount = EPCGDispatchThreadCount::FromFirstOutputPin;
diff --git a/Engine/Plugins/PCG/Source/PCG/Private/Compute/Elements/PCGCustomHLSLKernel.cpp b/Engine/Plugins/PCG/Source/PCG/Private/Compute/Elements/PCGCustomHLSLKernel.cpp
--- a/Engine/Plugins/PCG/Source/PCG/Private/Compute/Elements/PCGCustomHLSLKernel.cpp
+++ b/Engine/Plugins/PCG/Source/PCG/Private/Compute/Elements/PCGCustomHLSLKernel.cpp
@@ -274,7 +274,7 @@ bool UPCGCustomHLSLKernel::ComputeOutputBindingDataDesc(const UPCGComputeGraph*
 		// Generators always produce a single texture data with known size.
 		FPCGDataDesc& NewDataDesc = OutputDataDesc.DataDescs.Emplace_GetRef();
 		NewDataDesc.Type = EPCGDataType::BaseTexture;
-		NewDataDesc.ElementCount2D = CustomHLSLSettings->TextureSize;
+		NewDataDesc.ElementCount2D = CustomHLSLSettings->NumElements2D;
 	}
 	else if (ensure(OutputPinProperties))
 	{
@@ -494,7 +494,11 @@ FString UPCGCustomHLSLKernel::GetCookedSource(FPCGGPUCompilationContext& InOutCo
 	}
 	else if (CustomHLSLSettings->KernelType == EPCGKernelType::PointGenerator)
 	{
-		KernelSpecificPreamble += FString::Format(TEXT("    const uint NumPoints = {0};\n"), { CustomHLSLSettings->PointCount });
+		KernelSpecificPreamble += FString::Format(TEXT(
+			"    const uint NumElements = {0};\n"
+			"    // NumPoints is deprecated.\n"
+			"    const uint NumPoints = NumElements;\n"
+		), { CustomHLSLSettings->PointCount });
 
 		if (const FPCGPinPropertiesGPU* OutputPin = GetFirstOutputPin())
 		{
@@ -531,7 +535,7 @@ FString UPCGCustomHLSLKernel::GetCookedSource(FPCGGPUCompilationContext& InOutCo
 	}
 	else if (CustomHLSLSettings->KernelType == EPCGKernelType::TextureGenerator)
 	{
-		KernelSpecificPreamble += FString::Format(TEXT("    const uint2 ElementCount = uint2({0}, {1});\n"), { CustomHLSLSettings->TextureSize.X, CustomHLSLSettings->TextureSize.Y });
+		KernelSpecificPreamble += FString::Format(TEXT("    const uint2 NumElements = uint2({0}, {1});\n"), { CustomHLSLSettings->NumElements2D.X, CustomHLSLSettings->NumElements2D.Y });
 
 		if (const FPCGPinPropertiesGPU* OutputPin = GetFirstOutputPin())
 		{
@@ -598,7 +602,7 @@ int UPCGCustomHLSLKernel::ComputeThreadCount(const UPCGDataBinding* InBinding) c
 	else if (CustomHLSLSettings->KernelType == EPCGKernelType::TextureGenerator)
 	{
 		// Texture generator has fixed thread count.
-		ThreadCount = CustomHLSLSettings->TextureSize.X * CustomHLSLSettings->TextureSize.Y;
+		ThreadCount = CustomHLSLSettings->NumElements2D.X * CustomHLSLSettings->NumElements2D.Y;
 	}
 	else if (CustomHLSLSettings->IsProcessorKernel())
 	{
@@ -1332,7 +1336,7 @@ bool UPCGCustomHLSLKernel::PerformStaticValidation()
 				{
 #if PCG_KERNEL_LOGGING_ENABLED
 					StaticLogEntries.Emplace(FText::Format(
-						LOCTEXT("InvalidElementCount", "Invalid fixed element count {0} on output pin '{1}'. Must be greater than 0."),
+						LOCTEXT("InvalidElementCount", "Invalid fixed num elements {0} on output pin '{1}'. Must be greater than 0."),
 						FText::AsNumber(Props.ElementCount),
 						FText::FromName(Properties.Label)),
 						EPCGKernelLogVerbosity::Error);
@@ -1340,13 +1344,13 @@ bool UPCGCustomHLSLKernel::PerformStaticValidation()
 					return false;
 				}
 
-				if (Props.ElementCount2D.GetMin() < 1)
+				if (Props.NumElements2D.GetMin() < 1)
 				{
 #if PCG_KERNEL_LOGGING_ENABLED
 					StaticLogEntries.Emplace(FText::Format(
-						LOCTEXT("InvalidElementCount2D", "Invalid fixed element count 2D ({0}, {1}) on output pin '{2}'. Must be greater than 0."),
-						FText::AsNumber(Props.ElementCount2D.X),
-						FText::AsNumber(Props.ElementCount2D.Y),
+						LOCTEXT("InvalidElementCount2D", "Invalid fixed num elements ({0}, {1}) on output pin '{2}'. Must be greater than 0."),
+						FText::AsNumber(Props.NumElements2D.X),
+						FText::AsNumber(Props.NumElements2D.Y),
 						FText::FromName(Properties.Label)),
 						EPCGKernelLogVerbosity::Error);
 #endif
diff --git a/Engine/Plugins/PCG/Source/PCG/Private/Compute/PCGComputeGraph.cpp b/Engine/Plugins/PCG/Source/PCG/Private/Compute/PCGComputeGraph.cpp
--- a/Engine/Plugins/PCG/Source/PCG/Private/Compute/PCGComputeGraph.cpp
+++ b/Engine/Plugins/PCG/Source/PCG/Private/Compute/PCGComputeGraph.cpp
@@ -208,8 +208,16 @@ void UPCGComputeGraph::DebugLogDataDescriptions(const UPCGDataBinding* InBinding
 
 			UE_LOG(LogPCG, Display, TEXT("\t\tData Index %d"), DataIndex);
 			UE_LOG(LogPCG, Display, TEXT("\t\t\tType: %s"), *PCGDataTypeEnum->GetDisplayNameTextByValue(static_cast<int64>(DataDesc.Type)).ToString());
-			UE_LOG(LogPCG, Display, TEXT("\t\t\tElementCount: %d"), DataDesc.ElementCount);
-			UE_LOG(LogPCG, Display, TEXT("\t\t\tElementCount2D: (%d, %d)"), DataDesc.ElementCount2D.X, DataDesc.ElementCount2D.Y);
+			
+			if (DataDesc.IsDomain2D())
+			{
+				UE_LOG(LogPCG, Display, TEXT("\t\t\tNum Elements: (%d, %d)"), DataDesc.ElementCount2D.X, DataDesc.ElementCount2D.Y);
+			}
+			else
+			{
+				UE_LOG(LogPCG, Display, TEXT("\t\t\tNum Elements: %d"), DataDesc.ElementCount);
+			}
+
 			UE_LOG(LogPCG, Display, TEXT("\t\t\tAttributes (%d)"), DataDesc.AttributeDescs.Num());
 
 			for (int32 AttributeIndex = 0; AttributeIndex < DataDesc.AttributeDescs.Num(); ++AttributeIndex)
diff --git a/Engine/Plugins/PCG/Source/PCG/Private/Compute/PCGComputeKernel.cpp b/Engine/Plugins/PCG/Source/PCG/Private/Compute/PCGComputeKernel.cpp
--- a/Engine/Plugins/PCG/Source/PCG/Private/Compute/PCGComputeKernel.cpp
+++ b/Engine/Plugins/PCG/Source/PCG/Private/Compute/PCGComputeKernel.cpp
@@ -197,7 +197,7 @@ void UPCGComputeKernel::ComputeDataDescFromPinProperties(const FPCGPinProperties
 					else if (Props.ElementCountMode == EPCGElementCountMode::Fixed)
 					{
 						TotalNumElements += Props.ElementCount;
-						TotalNumElements2D += Props.ElementCount2D;
+						TotalNumElements2D += Props.NumElements2D;
 					}
 					else
 					{
@@ -329,7 +329,7 @@ void UPCGComputeKernel::ComputeDataDescFromPinProperties(const FPCGPinProperties
 								else if (Props.ElementCountMode == EPCGElementCountMode::Fixed)
 								{
 									DataDesc.ElementCount = Props.ElementCount;
-									DataDesc.ElementCount2D = Props.ElementCount2D;
+									DataDesc.ElementCount2D = Props.NumElements2D;
 								}
 								else
 								{
@@ -396,7 +396,7 @@ void UPCGComputeKernel::ComputeDataDescFromPinProperties(const FPCGPinProperties
 	{
 		for (int I = 0; I < Props.DataCount; ++I)
 		{
-			OutPinDesc.DataDescs.Emplace(OutputPinProps.AllowedTypes, Props.ElementCount, Props.ElementCount2D);
+			OutPinDesc.DataDescs.Emplace(OutputPinProps.AllowedTypes, Props.ElementCount, Props.NumElements2D);
 		}
 	}
 	else
diff --git a/Engine/Plugins/PCG/Source/PCG/Private/Compute/PCGPinPropertiesGPU.h b/Engine/Plugins/PCG/Source/PCG/Private/Compute/PCGPinPropertiesGPU.h
--- a/Engine/Plugins/PCG/Source/PCG/Private/Compute/PCGPinPropertiesGPU.h
+++ b/Engine/Plugins/PCG/Source/PCG/Private/Compute/PCGPinPropertiesGPU.h
@@ -94,25 +94,25 @@ struct FPCGPinPropertiesGPUStruct
 	EPCGElementMultiplicity ElementMultiplicity = EPCGElementMultiplicity::Product;
 
 	/** Fixed number of elements to create in each output data. */
-	UPROPERTY(BlueprintReadWrite, EditAnywhere, Category = Settings, meta = (EditCondition = "bAllowEditInitMode && (InitializationMode == EPCGPinInitMode::Custom || ElementCountMode == EPCGElementCountMode::Fixed)", EditConditionHides, ClampMin = 1))
+	UPROPERTY(BlueprintReadWrite, EditAnywhere, Category = Settings, DisplayName = "Num Elements", meta = (EditCondition = "bAllowEditInitMode && (InitializationMode == EPCGPinInitMode::Custom || ElementCountMode == EPCGElementCountMode::Fixed) && !bShowTexturePinSettings", EditConditionHides, ClampMin = 1))
 	int ElementCount = 1;
 
 	// @todo_pcg: Show texture data type info, e.g. Dimension, Format, etc.
 
 	/** Fixed number of elements to create in each output data. */
-	UPROPERTY(BlueprintReadWrite, EditAnywhere, Category = Settings, meta = (EditCondition = "bAllowEditInitMode && (InitializationMode == EPCGPinInitMode::Custom || ElementCountMode == EPCGElementCountMode::Fixed)", EditConditionHides, ClampMin = 1))
-	FIntPoint ElementCount2D = FIntPoint(64, 64);
+	UPROPERTY(BlueprintReadWrite, EditAnywhere, Category = Settings, DisplayName = "Num Elements", meta = (EditCondition = "bAllowEditInitMode && (InitializationMode == EPCGPinInitMode::Custom || ElementCountMode == EPCGElementCountMode::Fixed) && bShowTexturePinSettings", EditConditionHides, ClampMin = 1))
+	FIntPoint NumElements2D = FIntPoint(64, 64);
 
-	/** Fixed number of elements to create in each output data. */
+	/** Scalar on the fixed number of elements to create in each output data. */
 	UPROPERTY(BlueprintReadWrite, EditAnywhere, Category = Settings, meta = (EditCondition = "bAllowEditInitMode && (InitializationMode == EPCGPinInitMode::Custom || ElementCountMode == EPCGElementCountMode::FromInputData)", EditConditionHides, ClampMin = 1))
 	int ElementCountMultiplier = 1;
 
 	/** How to inherit attribute names, types, and values. */
-	UPROPERTY(BlueprintReadWrite, EditAnywhere, Category = Settings, meta = (EditCondition = "bAllowEditInitMode && InitializationMode == EPCGPinInitMode::FromInputPins", EditConditionHides))
+	UPROPERTY(BlueprintReadWrite, EditAnywhere, Category = Settings, meta = (EditCondition = "bAllowEditInitMode && InitializationMode == EPCGPinInitMode::FromInputPins && !bShowTexturePinSettings", EditConditionHides))
 	EPCGAttributeInheritanceMode AttributeInheritanceMode = EPCGAttributeInheritanceMode::CopyAttributeSetup;
 
 	/** Add entries to create new attributes on data emitted by this pin. */
-	UPROPERTY(EditAnywhere, DisplayName = "Attributes to Create", Category = Settings)
+	UPROPERTY(EditAnywhere, DisplayName = "Attributes to Create", Category = Settings, meta = (EditCondition = "!bShowTexturePinSettings", EditConditionHides, HideEditConditionToggle))
 	TArray<FPCGKernelAttributeKey> CreatedKernelAttributeKeys;
 
 #if WITH_EDITORONLY_DATA
@@ -125,6 +125,9 @@ struct FPCGPinPropertiesGPUStruct
 
 	UPROPERTY(Transient)
 	bool bMultipleInitPins = false;
+
+	UPROPERTY(Transient)
+	bool bShowTexturePinSettings = false;
 #endif // WITH_EDITORONLY_DATA
 };
 
