[PCG] Prefix all shader defines with PCG_ to avoid collisions with engine shaders

#rnx

diff --git a/Engine/Plugins/PCG/Shaders/Private/PCGDataCollectionAdaptor.usf b/Engine/Plugins/PCG/Shaders/Private/PCGDataCollectionAdaptor.usf
--- a/Engine/Plugins/PCG/Shaders/Private/PCGDataCollectionAdaptor.usf
+++ b/Engine/Plugins/PCG/Shaders/Private/PCGDataCollectionAdaptor.usf
@@ -23,9 +23,9 @@
 //   For each present attribute:
 //     <Packed element data> (element count * stride)
 
-#define DATA_COLLECTION_HEADER_SIZE_BYTES 4
-#define DATA_HEADER_SIZE_BYTES 1036
-#define ATTRIBUTE_HEADER_SIZE_BYTES 8
+#define PCG_DATA_COLLECTION_HEADER_SIZE_BYTES 4
+#define PCG_DATA_HEADER_SIZE_BYTES 1036
+#define PCG_ATTRIBUTE_HEADER_SIZE_BYTES 8
 
 StructuredBuffer<int2> InSourceToTargetAttributeId;
 StructuredBuffer<uint> InDataElementCounts;
@@ -37,8 +37,8 @@ uint InNumRemappedAttributes;
 void Out_SetNumElements(uint InDataIndex, uint InNumElements)
 {
 	const uint Address =
-		DATA_COLLECTION_HEADER_SIZE_BYTES +
-		InDataIndex * DATA_HEADER_SIZE_BYTES +
+		PCG_DATA_COLLECTION_HEADER_SIZE_BYTES +
+		InDataIndex * PCG_DATA_HEADER_SIZE_BYTES +
 		/*TypeId*/4 + /*Attribute Count*/4;
 
 	return OutDataCollection.Store(Address, InNumElements);
@@ -46,10 +46,10 @@ void Out_SetNumElements(uint InDataIndex, uint InNumElements)
 
 uint GetAttributeHeaderAddress(uint InDataIndex, int InAttributeId)
 {
-	return DATA_COLLECTION_HEADER_SIZE_BYTES +
-		InDataIndex * DATA_HEADER_SIZE_BYTES +
+	return PCG_DATA_COLLECTION_HEADER_SIZE_BYTES +
+		InDataIndex * PCG_DATA_HEADER_SIZE_BYTES +
 		/*This TypeId*/4 + /*This Attribute Count*/4 + /*This Element Count*/4 +
-		InAttributeId * ATTRIBUTE_HEADER_SIZE_BYTES;
+		InAttributeId * PCG_ATTRIBUTE_HEADER_SIZE_BYTES;
 }
 
 [numthreads(NUM_THREADS_PER_GROUP, 1, 1)]
diff --git a/Engine/Plugins/PCG/Shaders/Private/PCGDataCollectionDataInterface.ush b/Engine/Plugins/PCG/Shaders/Private/PCGDataCollectionDataInterface.ush
--- a/Engine/Plugins/PCG/Shaders/Private/PCGDataCollectionDataInterface.ush
+++ b/Engine/Plugins/PCG/Shaders/Private/PCGDataCollectionDataInterface.ush
@@ -16,27 +16,6 @@
 //   For each present attribute:
 //     <Packed element data> (element count * stride)
 
-// TODO: These defines could be code-gen'd to stay in sync with our PCGComputeConstants, see UComputeDataInterface::GetDefines()
-#define POINT_POSITION_ID 0
-#define POINT_ROTATION_ID 1
-#define POINT_SCALE_ID 2
-#define POINT_BOUNDS_MIN_ID 3
-#define POINT_BOUNDS_MAX_ID 4
-#define POINT_COLOR_ID 5
-#define POINT_DENSITY_ID 6
-#define POINT_SEED_ID 7
-#define POINT_STEEPNESS_ID 8
-#define POINT_METADATA_ENTRY_ID 9
-
-// Used to represent invalid/removed points. We use a value slightly less than max float,
-// as not all platforms support float infinity in shaders.
-#define INVALID_DENSITY 3.402823e+38f
-
-#define DATA_COLLECTION_HEADER_SIZE_BYTES 4
-#define DATA_HEADER_SIZE_BYTES 1036
-#define ATTRIBUTE_HEADER_SIZE_BYTES 8
-#define KERNEL_EXECUTED_FLAG (1u << 31)
-
 RWByteAddressBuffer {DataInterfaceName}_DataCollectionBuffer;
 
 // Multiplier that can be configured on output pins.
@@ -50,14 +29,14 @@ uint {DataInterfaceName}_ElementCountersPresent;
 uint GetNumData_{DataInterfaceName}()
 {
 	// We use the most significant bit of NumData to store the KernelExecuted flag.
-	return {DataInterfaceName}_DataCollectionBuffer.Load(0) & ~KERNEL_EXECUTED_FLAG;
+	return {DataInterfaceName}_DataCollectionBuffer.Load(0) & ~PCG_KERNEL_EXECUTED_FLAG;
 }
 
 uint GetNumElements_{DataInterfaceName}(uint InDataIndex)
 {
 	const uint ReadAddress =
-		DATA_COLLECTION_HEADER_SIZE_BYTES +
-		InDataIndex * DATA_HEADER_SIZE_BYTES +
+		PCG_DATA_COLLECTION_HEADER_SIZE_BYTES +
+		InDataIndex * PCG_DATA_HEADER_SIZE_BYTES +
 		/*TypeId*/4 + /*Attribute Count*/4;
 
 	return {DataInterfaceName}_DataCollectionBuffer.Load(ReadAddress);
@@ -128,10 +107,10 @@ void StoreBufferInternal_{DataInterfaceName}(uint Address, uint Value)
 uint GetFirstElementAddressInternal_{DataInterfaceName}(uint InDataIndex, int InAttributeId)
 {
 	const uint ReadAddress =
-		DATA_COLLECTION_HEADER_SIZE_BYTES +
-		InDataIndex * DATA_HEADER_SIZE_BYTES +
+		PCG_DATA_COLLECTION_HEADER_SIZE_BYTES +
+		InDataIndex * PCG_DATA_HEADER_SIZE_BYTES +
 		/*TypeId*/4 + /*Attribute Count*/4 + /*Element Count*/4 +
-		InAttributeId * ATTRIBUTE_HEADER_SIZE_BYTES +
+		InAttributeId * PCG_ATTRIBUTE_HEADER_SIZE_BYTES +
 		/*PackedAttributeIdAndStride*/4;
 
 	return {DataInterfaceName}_DataCollectionBuffer.Load(ReadAddress);
@@ -140,10 +119,10 @@ uint GetFirstElementAddressInternal_{DataInterfaceName}(uint InDataIndex, int In
 uint GetAttributeStrideInternal_{DataInterfaceName}(uint InDataIndex, int InAttributeId)
 {
 	const uint ReadAddress =
-		DATA_COLLECTION_HEADER_SIZE_BYTES +
-		InDataIndex * DATA_HEADER_SIZE_BYTES +
+		PCG_DATA_COLLECTION_HEADER_SIZE_BYTES +
+		InDataIndex * PCG_DATA_HEADER_SIZE_BYTES +
 		/*TypeId*/4 + /*Attribute Count*/4 + /*Element Count*/4 +
-		InAttributeId * ATTRIBUTE_HEADER_SIZE_BYTES;
+		InAttributeId * PCG_ATTRIBUTE_HEADER_SIZE_BYTES;
 
 	const uint PackedAttributeIdAndStride = {DataInterfaceName}_DataCollectionBuffer.Load(ReadAddress);
 	return PackedAttributeIdAndStride & (0xFFFFFFFF >> 24);
@@ -159,8 +138,8 @@ uint GetElementAddressInternal_{DataInterfaceName}(uint InDataIndex, uint InElem
 uint GetDataNumAttributesInternal_{DataInterfaceName}(uint InDataIndex)
 {
 	const uint ReadAddress =
-		DATA_COLLECTION_HEADER_SIZE_BYTES +
-		InDataIndex * DATA_HEADER_SIZE_BYTES +
+		PCG_DATA_COLLECTION_HEADER_SIZE_BYTES +
+		InDataIndex * PCG_DATA_HEADER_SIZE_BYTES +
 		/*TypeId*/4;
 
 	return {DataInterfaceName}_DataCollectionBuffer.Load(ReadAddress);
@@ -170,7 +149,7 @@ void SetAsExecutedInternal_{DataInterfaceName}()
 {
 	// We use the most significant bit of NumData to store the KernelExecuted flag
 	const uint NumData = GetNumData_{DataInterfaceName}();
-	{DataInterfaceName}_DataCollectionBuffer.Store(0, NumData | KERNEL_EXECUTED_FLAG);
+	{DataInterfaceName}_DataCollectionBuffer.Store(0, NumData | PCG_KERNEL_EXECUTED_FLAG);
 }
 
 // #################### ATTRIBUTE GETTERS ##########################
@@ -396,47 +375,47 @@ int AtomicAddInt_{DataInterfaceName}(uint DataIndex, uint ElementIndex, int Attr
 
 float3 GetPosition_{DataInterfaceName}(uint DataIndex, uint ElementIndex)
 {
-	return GetFloat3_{DataInterfaceName}(DataIndex, ElementIndex, POINT_POSITION_ID);
+	return GetFloat3_{DataInterfaceName}(DataIndex, ElementIndex, PCG_POINT_POSITION_ID);
 }
 
 float4 GetRotation_{DataInterfaceName}(uint DataIndex, uint ElementIndex)
 {
-	return GetFloat4_{DataInterfaceName}(DataIndex, ElementIndex, POINT_ROTATION_ID);
+	return GetFloat4_{DataInterfaceName}(DataIndex, ElementIndex, PCG_POINT_ROTATION_ID);
 }
 
 float3 GetScale_{DataInterfaceName}(uint DataIndex, uint ElementIndex)
 {
-	return GetFloat3_{DataInterfaceName}(DataIndex, ElementIndex, POINT_SCALE_ID);
+	return GetFloat3_{DataInterfaceName}(DataIndex, ElementIndex, PCG_POINT_SCALE_ID);
 }
 
 float3 GetBoundsMin_{DataInterfaceName}(uint DataIndex, uint ElementIndex)
 {
-	return GetFloat3_{DataInterfaceName}(DataIndex, ElementIndex, POINT_BOUNDS_MIN_ID);
+	return GetFloat3_{DataInterfaceName}(DataIndex, ElementIndex, PCG_POINT_BOUNDS_MIN_ID);
 }
 
 float3 GetBoundsMax_{DataInterfaceName}(uint DataIndex, uint ElementIndex)
 {
-	return GetFloat3_{DataInterfaceName}(DataIndex, ElementIndex, POINT_BOUNDS_MAX_ID);
+	return GetFloat3_{DataInterfaceName}(DataIndex, ElementIndex, PCG_POINT_BOUNDS_MAX_ID);
 }
 
 float4 GetColor_{DataInterfaceName}(uint DataIndex, uint ElementIndex)
 {
-	return GetFloat4_{DataInterfaceName}(DataIndex, ElementIndex, POINT_COLOR_ID);
+	return GetFloat4_{DataInterfaceName}(DataIndex, ElementIndex, PCG_POINT_COLOR_ID);
 }
 
 float GetDensity_{DataInterfaceName}(uint DataIndex, uint ElementIndex)
 {
-	return GetFloat_{DataInterfaceName}(DataIndex, ElementIndex, POINT_DENSITY_ID);
+	return GetFloat_{DataInterfaceName}(DataIndex, ElementIndex, PCG_POINT_DENSITY_ID);
 }
 
 float GetSteepness_{DataInterfaceName}(uint DataIndex, uint ElementIndex)
 {
-	return GetFloat_{DataInterfaceName}(DataIndex, ElementIndex, POINT_STEEPNESS_ID);
+	return GetFloat_{DataInterfaceName}(DataIndex, ElementIndex, PCG_POINT_STEEPNESS_ID);
 }
 
 uint GetSeed_{DataInterfaceName}(uint DataIndex, uint ElementIndex)
 {
-	return GetUint_{DataInterfaceName}(DataIndex, ElementIndex, POINT_SEED_ID);
+	return GetUint_{DataInterfaceName}(DataIndex, ElementIndex, PCG_POINT_SEED_ID);
 }
 
 float4x4 GetPointTransform_{DataInterfaceName}(uint InDataIndex, uint InElementIndex)
@@ -459,54 +438,54 @@ float4x4 GetPointTransform_{DataInterfaceName}(uint InDataIndex, uint InElementI
 
 bool IsPointRemoved_{DataInterfaceName}(uint InDataIndex, uint InElementIndex)
 {
-	return GetDensity_{DataInterfaceName}(InDataIndex, InElementIndex) == INVALID_DENSITY;
+	return GetDensity_{DataInterfaceName}(InDataIndex, InElementIndex) == PCG_INVALID_DENSITY;
 }
 
 // #################### POINT ATTRIBUTE SETTERS ##########################
 
 void SetPosition_{DataInterfaceName}(uint DataIndex, uint ElementIndex, float3 Position)
 {
-	SetFloat3_{DataInterfaceName}(DataIndex, ElementIndex, POINT_POSITION_ID, Position);
+	SetFloat3_{DataInterfaceName}(DataIndex, ElementIndex, PCG_POINT_POSITION_ID, Position);
 }
 
 void SetRotation_{DataInterfaceName}(uint DataIndex, uint ElementIndex, float4 Rotation)
 {
-	SetFloat4_{DataInterfaceName}(DataIndex, ElementIndex, POINT_ROTATION_ID, Rotation);
+	SetFloat4_{DataInterfaceName}(DataIndex, ElementIndex, PCG_POINT_ROTATION_ID, Rotation);
 }
 
 void SetScale_{DataInterfaceName}(uint DataIndex, uint ElementIndex, float3 Scale)
 {
-	SetFloat3_{DataInterfaceName}(DataIndex, ElementIndex, POINT_SCALE_ID, Scale);
+	SetFloat3_{DataInterfaceName}(DataIndex, ElementIndex, PCG_POINT_SCALE_ID, Scale);
 }
 
 void SetBoundsMin_{DataInterfaceName}(uint DataIndex, uint ElementIndex, float3 BoundsMin)
 {
-	SetFloat3_{DataInterfaceName}(DataIndex, ElementIndex, POINT_BOUNDS_MIN_ID, BoundsMin);
+	SetFloat3_{DataInterfaceName}(DataIndex, ElementIndex, PCG_POINT_BOUNDS_MIN_ID, BoundsMin);
 }
 
 void SetBoundsMax_{DataInterfaceName}(uint DataIndex, uint ElementIndex, float3 BoundsMax)
 {
-	SetFloat3_{DataInterfaceName}(DataIndex, ElementIndex, POINT_BOUNDS_MAX_ID, BoundsMax);
+	SetFloat3_{DataInterfaceName}(DataIndex, ElementIndex, PCG_POINT_BOUNDS_MAX_ID, BoundsMax);
 }
 
 void SetColor_{DataInterfaceName}(uint DataIndex, uint ElementIndex, float4 Color)
 {
-	SetFloat4_{DataInterfaceName}(DataIndex, ElementIndex, POINT_COLOR_ID, Color);
+	SetFloat4_{DataInterfaceName}(DataIndex, ElementIndex, PCG_POINT_COLOR_ID, Color);
 }
 
 void SetDensity_{DataInterfaceName}(uint DataIndex, uint ElementIndex, float Density)
 {
-	SetFloat_{DataInterfaceName}(DataIndex, ElementIndex, POINT_DENSITY_ID, Density);
+	SetFloat_{DataInterfaceName}(DataIndex, ElementIndex, PCG_POINT_DENSITY_ID, Density);
 }
 
 void SetSteepness_{DataInterfaceName}(uint DataIndex, uint ElementIndex, float Steepness)
 {
-	SetFloat_{DataInterfaceName}(DataIndex, ElementIndex, POINT_STEEPNESS_ID, Steepness);
+	SetFloat_{DataInterfaceName}(DataIndex, ElementIndex, PCG_POINT_STEEPNESS_ID, Steepness);
 }
 
 void SetSeed_{DataInterfaceName}(uint DataIndex, uint ElementIndex, uint Seed)
 {
-	SetUint_{DataInterfaceName}(DataIndex, ElementIndex, POINT_SEED_ID, Seed);
+	SetUint_{DataInterfaceName}(DataIndex, ElementIndex, PCG_POINT_SEED_ID, Seed);
 }
 
 void SetPointTransform_{DataInterfaceName}(uint InDataIndex, uint InElementIndex, float4x4 Transform)
@@ -529,7 +508,7 @@ void SetPointTransform_{DataInterfaceName}(uint InDataIndex, uint InElementIndex
 
 void RemovePoint_{DataInterfaceName}(uint InDataIndex, uint InElementIndex)
 {
-	SetDensity_{DataInterfaceName}(InDataIndex, InElementIndex, INVALID_DENSITY);
+	SetDensity_{DataInterfaceName}(InDataIndex, InElementIndex, PCG_INVALID_DENSITY);
 }
 
 // Initialize a single point with default values.
@@ -547,7 +526,7 @@ void InitializePoint_{DataInterfaceName}(uint InDataIndex, uint InElementIndex)
 	
 	int NumAttributesRemaining = (int)GetDataNumAttributesInternal_{DataInterfaceName}(InDataIndex);
 
-	for (int AttributeId = NUM_RESERVED_ATTRS; AttributeId < MAX_NUM_ATTRS; ++AttributeId)
+	for (int AttributeId = PCG_NUM_RESERVED_ATTRS; AttributeId < PCG_MAX_NUM_ATTRS; ++AttributeId)
 	{
 		const uint Stride = GetAttributeStrideInternal_{DataInterfaceName}(InDataIndex, AttributeId);
 
diff --git a/Engine/Plugins/PCG/Shaders/Private/PCGGrassMapUnpackerCS.usf b/Engine/Plugins/PCG/Shaders/Private/PCGGrassMapUnpackerCS.usf
--- a/Engine/Plugins/PCG/Shaders/Private/PCGGrassMapUnpackerCS.usf
+++ b/Engine/Plugins/PCG/Shaders/Private/PCGGrassMapUnpackerCS.usf
@@ -3,9 +3,6 @@
 #include "/Engine/Private/Common.ush"
 #include "/Engine/Public/Platform.ush"
 
-// Must match FPCGGrassMapUnpackerCS::MaxNumLandscapeComponents
-#define MAX_NUM_LANDSCAPE_COMPONENTS 64
-
 // Input:
 // First 2 channels of first texture always reserved for heightmap (see constructor of FLandscapeGrassWeightExporter_RenderThread::FComponentInfo).
 // So first two grass maps will be in z & w channels of first set of pass, otherwise packed in remaining channels.
@@ -35,7 +32,7 @@ Texture2D<float4> InPackedGrassMaps;
 //                 * ResX
 RWTexture2DArray<float> OutUnpackedGrassMaps;
 
-int4 InLinearTileIndexToComponentIndex[MAX_NUM_LANDSCAPE_COMPONENTS]; // 16 byte alignment on arrays
+int4 InLinearTileIndexToComponentIndex[PCG_MAX_NUM_LANDSCAPE_COMPONENTS]; // 16 byte alignment on arrays
 uint InNumTilesX;
 uint InLandscapeComponentResolution;
 uint InNumGrassMapPasses;
diff --git a/Engine/Plugins/PCG/Shaders/Private/PCGInstanceDataInterface.ush b/Engine/Plugins/PCG/Shaders/Private/PCGInstanceDataInterface.ush
--- a/Engine/Plugins/PCG/Shaders/Private/PCGInstanceDataInterface.ush
+++ b/Engine/Plugins/PCG/Shaders/Private/PCGInstanceDataInterface.ush
@@ -1,8 +1,6 @@
 // Copyright Epic Games, Inc. All Rights Reserved.
 
-#define MAX_PRIMITIVES 256 // Should be the same as MAX_PRIMITIVE_COMPONENTS_PER_SPAWNER
-
-uint4 {DataInterfaceName}_PrimitiveParameters[MAX_PRIMITIVES]; // (NumInstancesAllocated, InstanceOffset, Unused, Unused)
+uint4 {DataInterfaceName}_PrimitiveParameters[PCG_MAX_PRIMITIVES]; // (NumInstancesAllocated, InstanceOffset, Unused, Unused)
 
 uint {DataInterfaceName}_NumPrimitives;
 uint {DataInterfaceName}_NumInstancesAllPrimitives;
diff --git a/Engine/Plugins/PCG/Shaders/Private/PCGShaderUtilsInternal.ush b/Engine/Plugins/PCG/Shaders/Private/PCGShaderUtilsInternal.ush
--- a/Engine/Plugins/PCG/Shaders/Private/PCGShaderUtilsInternal.ush
+++ b/Engine/Plugins/PCG/Shaders/Private/PCGShaderUtilsInternal.ush
@@ -3,11 +3,11 @@
 #pragma once
 
 // Must match C++ constants in PCGComputeCommon.h
-#ifndef NUM_RESERVED_ATTRS
-#define NUM_RESERVED_ATTRS 32
+#ifndef PCG_NUM_RESERVED_ATTRS
+#define PCG_NUM_RESERVED_ATTRS 32
 #endif
-#ifndef MAX_NUM_ATTRS
-#define MAX_NUM_ATTRS 128
+#ifndef PCG_MAX_NUM_ATTRS
+#define PCG_MAX_NUM_ATTRS 128
 #endif
 
 // TODO pass in attribute IDs to copy?
@@ -16,9 +16,9 @@
 		const uint NumAttributes = TargetPin##_GetDataNumAttributesInternal(Target_DataIndex); \
 		uint NumAttributesProcessed = 0; \
 		\
-		const int FirstAttributeId = (bMetadataOnly) ? NUM_RESERVED_ATTRS : 0; \
+		const int FirstAttributeId = (bMetadataOnly) ? PCG_NUM_RESERVED_ATTRS : 0; \
 		\
-		for (int AttributeId = FirstAttributeId; AttributeId < MAX_NUM_ATTRS; ++AttributeId) \
+		for (int AttributeId = FirstAttributeId; AttributeId < PCG_MAX_NUM_ATTRS; ++AttributeId) \
 		{ \
 			const uint Target_Stride = TargetPin##_GetAttributeStrideInternal(Target_DataIndex, AttributeId); \
 			\
diff --git a/Engine/Plugins/PCG/Source/PCG/Private/Compute/DataInterfaces/PCGDataCollectionDataInterface.cpp b/Engine/Plugins/PCG/Source/PCG/Private/Compute/DataInterfaces/PCGDataCollectionDataInterface.cpp
--- a/Engine/Plugins/PCG/Source/PCG/Private/Compute/DataInterfaces/PCGDataCollectionDataInterface.cpp
+++ b/Engine/Plugins/PCG/Source/PCG/Private/Compute/DataInterfaces/PCGDataCollectionDataInterface.cpp
@@ -13,11 +13,12 @@
 #include "Compute/Elements/PCGComputeGraphElement.h"
 #include "Helpers/PCGHelpers.h"
 
-#include "ComputeFramework/ShaderParamTypeDefinition.h"
 #include "RenderGraphBuilder.h"
 #include "RenderGraphResources.h"
 #include "ShaderCompilerCore.h"
 #include "ShaderParameterMetadataBuilder.h"
+#include "ComputeFramework/ComputeKernelPermutationSet.h"
+#include "ComputeFramework/ShaderParamTypeDefinition.h"
 
 #define LOCTEXT_NAMESPACE "PCGDataCollectionDataInterface"
 
@@ -448,6 +449,31 @@ void UPCGDataCollectionDataInterface::GetHLSL(FString& OutHLSL, FString const& I
 	}
 }
 
+void UPCGDataCollectionDataInterface::GetDefines(FComputeKernelDefinitionSet& OutDefinitionSet) const
+{
+	Super::GetDefines(OutDefinitionSet);
+
+	// Point property IDs
+	OutDefinitionSet.Defines.Add(FComputeKernelDefinition(TEXT("PCG_POINT_POSITION_ID"),   FString::FromInt(PCGComputeConstants::POINT_POSITION_ATTRIBUTE_ID)));
+	OutDefinitionSet.Defines.Add(FComputeKernelDefinition(TEXT("PCG_POINT_ROTATION_ID"),   FString::FromInt(PCGComputeConstants::POINT_ROTATION_ATTRIBUTE_ID)));
+	OutDefinitionSet.Defines.Add(FComputeKernelDefinition(TEXT("PCG_POINT_SCALE_ID"),      FString::FromInt(PCGComputeConstants::POINT_SCALE_ATTRIBUTE_ID)));
+	OutDefinitionSet.Defines.Add(FComputeKernelDefinition(TEXT("PCG_POINT_BOUNDS_MIN_ID"), FString::FromInt(PCGComputeConstants::POINT_BOUNDS_MIN_ATTRIBUTE_ID)));
+	OutDefinitionSet.Defines.Add(FComputeKernelDefinition(TEXT("PCG_POINT_BOUNDS_MAX_ID"), FString::FromInt(PCGComputeConstants::POINT_BOUNDS_MAX_ATTRIBUTE_ID)));
+	OutDefinitionSet.Defines.Add(FComputeKernelDefinition(TEXT("PCG_POINT_COLOR_ID"),      FString::FromInt(PCGComputeConstants::POINT_COLOR_ATTRIBUTE_ID)));
+	OutDefinitionSet.Defines.Add(FComputeKernelDefinition(TEXT("PCG_POINT_DENSITY_ID"),    FString::FromInt(PCGComputeConstants::POINT_DENSITY_ATTRIBUTE_ID)));
+	OutDefinitionSet.Defines.Add(FComputeKernelDefinition(TEXT("PCG_POINT_SEED_ID"),       FString::FromInt(PCGComputeConstants::POINT_SEED_ATTRIBUTE_ID)));
+	OutDefinitionSet.Defines.Add(FComputeKernelDefinition(TEXT("PCG_POINT_STEEPNESS_ID"),  FString::FromInt(PCGComputeConstants::POINT_STEEPNESS_ATTRIBUTE_ID)));
+
+	// Header sizes
+	OutDefinitionSet.Defines.Add(FComputeKernelDefinition(TEXT("PCG_DATA_COLLECTION_HEADER_SIZE_BYTES"), FString::FromInt(PCGComputeConstants::DATA_COLLECTION_HEADER_SIZE_BYTES)));
+	OutDefinitionSet.Defines.Add(FComputeKernelDefinition(TEXT("PCG_DATA_HEADER_SIZE_BYTES"),            FString::FromInt(PCGComputeConstants::DATA_HEADER_SIZE_BYTES)));
+	OutDefinitionSet.Defines.Add(FComputeKernelDefinition(TEXT("PCG_ATTRIBUTE_HEADER_SIZE_BYTES"),       FString::FromInt(PCGComputeConstants::ATTRIBUTE_HEADER_SIZE_BYTES)));
+
+	// Misc
+	OutDefinitionSet.Defines.Add(FComputeKernelDefinition(TEXT("PCG_KERNEL_EXECUTED_FLAG"), FString::FromInt(PCGComputeConstants::KernelExecutedFlag)));
+	OutDefinitionSet.Defines.Add(FComputeKernelDefinition(TEXT("PCG_INVALID_DENSITY"),      FString::Printf(TEXT("%f"), PCGComputeConstants::INVALID_DENSITY)));
+}
+
 UComputeDataProvider* UPCGDataCollectionDataInterface::CreateDataProvider() const
 {
 	return NewObject<UPCGDataCollectionDataProvider>();
diff --git a/Engine/Plugins/PCG/Source/PCG/Private/Compute/DataInterfaces/PCGDataCollectionDataInterface.h b/Engine/Plugins/PCG/Source/PCG/Private/Compute/DataInterfaces/PCGDataCollectionDataInterface.h
--- a/Engine/Plugins/PCG/Source/PCG/Private/Compute/DataInterfaces/PCGDataCollectionDataInterface.h
+++ b/Engine/Plugins/PCG/Source/PCG/Private/Compute/DataInterfaces/PCGDataCollectionDataInterface.h
@@ -33,6 +33,7 @@ public:
 	TCHAR const* GetShaderVirtualPath() const override; // TODO probably easier to just inline rather than external source?
 	void GetShaderHash(FString& InOutKey) const override;
 	void GetHLSL(FString& OutHLSL, FString const& InDataInterfaceName) const override;
+	void GetDefines(FComputeKernelDefinitionSet& OutDefinitionSet) const override;
 	// This data interface only uses the CF readback functionality for the element counters.
 	bool GetRequiresReadback() const override { return bExposeElementCounters; }
 	UComputeDataProvider* CreateDataProvider() const override;
diff --git a/Engine/Plugins/PCG/Source/PCG/Private/Compute/DataInterfaces/PCGInstanceDataInterface.cpp b/Engine/Plugins/PCG/Source/PCG/Private/Compute/DataInterfaces/PCGInstanceDataInterface.cpp
--- a/Engine/Plugins/PCG/Source/PCG/Private/Compute/DataInterfaces/PCGInstanceDataInterface.cpp
+++ b/Engine/Plugins/PCG/Source/PCG/Private/Compute/DataInterfaces/PCGInstanceDataInterface.cpp
@@ -127,6 +127,13 @@ void UPCGInstanceDataInterface::GetHLSL(FString& OutHLSL, FString const& InDataI
 	}
 }
 
+void UPCGInstanceDataInterface::GetDefines(FComputeKernelDefinitionSet& OutDefinitionSet) const
+{
+	Super::GetDefines(OutDefinitionSet);
+
+	OutDefinitionSet.Defines.Add(FComputeKernelDefinition(TEXT("PCG_MAX_PRIMITIVES"), FString::FromInt(PCGComputeConstants::MAX_PRIMITIVE_COMPONENTS_PER_SPAWNER)));
+}
+
 UComputeDataProvider* UPCGInstanceDataInterface::CreateDataProvider() const
 {
 	return NewObject<UPCGInstanceDataProvider>();
diff --git a/Engine/Plugins/PCG/Source/PCG/Private/Compute/DataInterfaces/PCGInstanceDataInterface.h b/Engine/Plugins/PCG/Source/PCG/Private/Compute/DataInterfaces/PCGInstanceDataInterface.h
--- a/Engine/Plugins/PCG/Source/PCG/Private/Compute/DataInterfaces/PCGInstanceDataInterface.h
+++ b/Engine/Plugins/PCG/Source/PCG/Private/Compute/DataInterfaces/PCGInstanceDataInterface.h
@@ -29,6 +29,7 @@ public:
 	TCHAR const* GetShaderVirtualPath() const override;
 	void GetShaderHash(FString& InOutKey) const override;
 	void GetHLSL(FString& OutHLSL, FString const& InDataInterfaceName) const override;
+	void GetDefines(FComputeKernelDefinitionSet& OutDefinitionSet) const override;
 	bool GetRequiresPostSubmitCall() const override { return true; }
 	UComputeDataProvider* CreateDataProvider() const override;
 	//~ End UComputeDataInterface Interface
diff --git a/Engine/Plugins/PCG/Source/PCG/Private/Compute/DataInterfaces/PCGTextureDataInterface.cpp b/Engine/Plugins/PCG/Source/PCG/Private/Compute/DataInterfaces/PCGTextureDataInterface.cpp
--- a/Engine/Plugins/PCG/Source/PCG/Private/Compute/DataInterfaces/PCGTextureDataInterface.cpp
+++ b/Engine/Plugins/PCG/Source/PCG/Private/Compute/DataInterfaces/PCGTextureDataInterface.cpp
@@ -198,6 +198,8 @@ void UPCGTextureDataInterface::GetHLSL(FString& OutHLSL, FString const& InDataIn
 
 void UPCGTextureDataInterface::GetDefines(FComputeKernelDefinitionSet& OutDefinitionSet) const
 {
+	Super::GetDefines(OutDefinitionSet);
+
 	OutDefinitionSet.Defines.Add(FComputeKernelDefinition(TEXT("PCG_MAX_NUM_SRV_BINDINGS"), FString::FromInt(PCGTextureDataInterfaceConstants::MAX_NUM_SRV_BINDINGS)));
 }
 
diff --git a/Engine/Plugins/PCG/Source/PCG/Private/Compute/DataInterfaces/PCGVirtualTextureDataInterface.cpp b/Engine/Plugins/PCG/Source/PCG/Private/Compute/DataInterfaces/PCGVirtualTextureDataInterface.cpp
--- a/Engine/Plugins/PCG/Source/PCG/Private/Compute/DataInterfaces/PCGVirtualTextureDataInterface.cpp
+++ b/Engine/Plugins/PCG/Source/PCG/Private/Compute/DataInterfaces/PCGVirtualTextureDataInterface.cpp
@@ -106,6 +106,8 @@ void UPCGVirtualTextureDataInterface::GetHLSL(FString& OutHLSL, FString const& I
 
 void UPCGVirtualTextureDataInterface::GetDefines(FComputeKernelDefinitionSet& OutDefinitionSet) const
 {
+	Super::GetDefines(OutDefinitionSet);
+
 	OutDefinitionSet.Defines.Add(FComputeKernelDefinition(TEXT("PCG_MAX_NUM_BINDINGS"), FString::FromInt(PCGVirtualTextureDIConstants::MAX_NUM_BINDINGS)));
 	OutDefinitionSet.Defines.Add(FComputeKernelDefinition(TEXT("PCG_MAX_NUM_LAYERS"), FString::FromInt(PCGVirtualTextureDIConstants::MAX_NUM_LAYERS)));
 }
diff --git a/Engine/Plugins/PCG/Source/PCGCompute/Private/PCGDataCollectionAdaptorCS.cpp b/Engine/Plugins/PCG/Source/PCGCompute/Private/PCGDataCollectionAdaptorCS.cpp
--- a/Engine/Plugins/PCG/Source/PCGCompute/Private/PCGDataCollectionAdaptorCS.cpp
+++ b/Engine/Plugins/PCG/Source/PCGCompute/Private/PCGDataCollectionAdaptorCS.cpp
@@ -2,11 +2,24 @@
 
 #include "PCGDataCollectionAdaptorCS.h"
 
+namespace PCGComputeConstants
+{
+	// @todo_pcg: These must be kept in sync with the constants in PCGComputeCommon.h, but currently that is an internal file that we cannot reference from this module.
+	constexpr int DATA_COLLECTION_HEADER_SIZE_BYTES = 4;
+	constexpr int DATA_HEADER_PREAMBLE_SIZE_BYTES = 12;
+	constexpr int MAX_NUM_ATTRS = 128;
+	constexpr int ATTRIBUTE_HEADER_SIZE_BYTES = 8;
+	constexpr int DATA_HEADER_SIZE_BYTES = DATA_HEADER_PREAMBLE_SIZE_BYTES + MAX_NUM_ATTRS * ATTRIBUTE_HEADER_SIZE_BYTES;
+}
+
 void FPCGDataCollectionAdaptorCS::ModifyCompilationEnvironment(const FGlobalShaderPermutationParameters& Parameters, FShaderCompilerEnvironment& OutEnvironment)
 {
 	FGlobalShader::ModifyCompilationEnvironment(Parameters, OutEnvironment);
 
 	OutEnvironment.SetDefine(TEXT("NUM_THREADS_PER_GROUP"), NumThreadsPerGroup);
+	OutEnvironment.SetDefine(TEXT("PCG_DATA_COLLECTION_HEADER_SIZE_BYTES"), PCGComputeConstants::DATA_COLLECTION_HEADER_SIZE_BYTES);
+	OutEnvironment.SetDefine(TEXT("PCG_DATA_HEADER_SIZE_BYTES"), PCGComputeConstants::DATA_HEADER_SIZE_BYTES);
+	OutEnvironment.SetDefine(TEXT("PCG_ATTRIBUTE_HEADER_SIZE_BYTES"), PCGComputeConstants::ATTRIBUTE_HEADER_SIZE_BYTES);
 }
 
 IMPLEMENT_GLOBAL_SHADER(FPCGDataCollectionAdaptorCS, "/PCGComputeShaders/PCGDataCollectionAdaptor.usf", "MainCS", SF_Compute);
diff --git a/Engine/Plugins/PCG/Source/PCGCompute/Private/PCGGrassMapUnpackerCS.cpp b/Engine/Plugins/PCG/Source/PCGCompute/Private/PCGGrassMapUnpackerCS.cpp
--- a/Engine/Plugins/PCG/Source/PCGCompute/Private/PCGGrassMapUnpackerCS.cpp
+++ b/Engine/Plugins/PCG/Source/PCGCompute/Private/PCGGrassMapUnpackerCS.cpp
@@ -9,6 +9,7 @@ void FPCGGrassMapUnpackerCS::ModifyCompilationEnvironment(const FGlobalShaderPer
 	OutEnvironment.SetDefine(TEXT("THREADGROUPSIZE_X"), ThreadGroupDim);
 	OutEnvironment.SetDefine(TEXT("THREADGROUPSIZE_Y"), ThreadGroupDim);
 	OutEnvironment.SetDefine(TEXT("THREADGROUPSIZE_Z"), 1);
+	OutEnvironment.SetDefine(TEXT("PCG_MAX_NUM_LANDSCAPE_COMPONENTS"), MaxNumLandscapeComponents);
 }
 
 IMPLEMENT_GLOBAL_SHADER(FPCGGrassMapUnpackerCS, "/PCGComputeShaders/PCGGrassMapUnpackerCS.usf", "PCGGrassMapUnpacker_CS", SF_Compute);
