[PCG] Visualize transfers of data to/from GPU in graph editor.

#rnx

diff --git a/Engine/Plugins/PCG/Source/PCG/Private/Compute/DataInterfaces/PCGComputeDataInterface.cpp b/Engine/Plugins/PCG/Source/PCG/Private/Compute/DataInterfaces/PCGComputeDataInterface.cpp
--- a/Engine/Plugins/PCG/Source/PCG/Private/Compute/DataInterfaces/PCGComputeDataInterface.cpp
+++ b/Engine/Plugins/PCG/Source/PCG/Private/Compute/DataInterfaces/PCGComputeDataInterface.cpp
@@ -2,6 +2,8 @@
 
 #include "Compute/DataInterfaces/PCGComputeDataInterface.h"
 
+#include "PCGContext.h"
+#include "PCGGraph.h"
 #include "PCGGraphExecutionInspection.h"
 #include "PCGGraphExecutionStateInterface.h"
 #include "PCGSettings.h"
@@ -57,6 +59,35 @@ void UPCGComputeDataProvider::Reset()
 	DownstreamInputPinLabelAliases.Empty();
 }
 
+#if WITH_EDITOR
+void UPCGComputeDataProvider::NotifyProducerUploadedData(UPCGDataBinding* InBinding)
+{
+	if (ProducerSettings && !ProducerSettings->ShouldExecuteOnGPU())
+	{
+		const UPCGNode* ProducerNode = Cast<UPCGNode>(ProducerSettings->GetOuter());
+
+		// Works around current issue where input output settings are outer'd to the graph rather than their node.
+		if (!ProducerNode)
+		{
+			if (const UPCGGraph* Graph = Cast<UPCGGraph>(ProducerSettings->GetOuter()))
+			{
+				if (Graph->GetInputNode() && Graph->GetInputNode()->GetSettings() == ProducerSettings)
+				{
+					ProducerNode = Graph->GetInputNode();
+				}
+			}
+		}
+
+		TSharedPtr<FPCGContextHandle> ContextHandle = InBinding->ContextHandle.Pin();
+		FPCGContext* Context = ContextHandle ? ContextHandle->GetContext() : nullptr;
+		if (ProducerNode && Context && Context->GetStack() && Context->ExecutionSource.IsValid())
+		{
+			Context->ExecutionSource->GetExecutionState().GetInspection().NotifyCPUToGPUUpload(ProducerNode, Context->GetStack());
+		}
+	}
+}
+#endif // WITH_EDITOR
+
 void UPCGExportableDataProvider::Initialize(const UComputeDataInterface* InDataInterface, UObject* InBinding, uint64 InInputMask, uint64 InOutputMask)
 {
 	Super::Initialize(InDataInterface, InBinding, InInputMask, InOutputMask);
diff --git a/Engine/Plugins/PCG/Source/PCG/Private/Compute/DataInterfaces/PCGComputeDataInterface.h b/Engine/Plugins/PCG/Source/PCG/Private/Compute/DataInterfaces/PCGComputeDataInterface.h
--- a/Engine/Plugins/PCG/Source/PCG/Private/Compute/DataInterfaces/PCGComputeDataInterface.h
+++ b/Engine/Plugins/PCG/Source/PCG/Private/Compute/DataInterfaces/PCGComputeDataInterface.h
@@ -72,6 +72,11 @@ public:
 
 	TObjectPtr<const UPCGSettings> GetProducerSettings() const { return ProducerSettings; }
 
+protected:
+#if WITH_EDITOR
+	void NotifyProducerUploadedData(UPCGDataBinding* InBinding);
+#endif
+
 public:
 	int32 GraphBindingIndex = INDEX_NONE;
 	FName OutputPinLabel;
diff --git a/Engine/Plugins/PCG/Source/PCG/Private/Compute/DataInterfaces/PCGDataCollectionUploadDataInterface.cpp b/Engine/Plugins/PCG/Source/PCG/Private/Compute/DataInterfaces/PCGDataCollectionUploadDataInterface.cpp
--- a/Engine/Plugins/PCG/Source/PCG/Private/Compute/DataInterfaces/PCGDataCollectionUploadDataInterface.cpp
+++ b/Engine/Plugins/PCG/Source/PCG/Private/Compute/DataInterfaces/PCGDataCollectionUploadDataInterface.cpp
@@ -3,6 +3,7 @@
 #include "Compute/DataInterfaces/PCGDataCollectionUploadDataInterface.h"
 
 #include "PCGEdge.h"
+#include "PCGGraphExecutionInspection.h"
 #include "PCGModule.h"
 #include "PCGNode.h"
 #include "PCGSettings.h"
@@ -31,6 +32,7 @@ void UPCGDataProviderDataCollectionUpload::Initialize(const UComputeDataInterfac
 	const UPCGDataCollectionUploadDataInterface* DataInterface = CastChecked<UPCGDataCollectionUploadDataInterface>(InDataInterface);
 
 	Binding = CastChecked<UPCGDataBinding>(InBinding);
+	ProducerSettings = DataInterface->ProducerSettings;
 	GraphBindingIndex = DataInterface->GraphBindingIndex;
 
 	// Pick the data items from input data collection using any of the compute graph element virtual input pin labels.
@@ -54,7 +56,16 @@ bool UPCGDataProviderDataCollectionUpload::PrepareForExecute_GameThread(UPCGData
 		}
 	}
 
-	return !DataAdaptor || DataAdaptor->PrepareData_GameThread();
+	const bool bPreparationDone = !DataAdaptor || DataAdaptor->PrepareData_GameThread();
+
+#if WITH_EDITOR
+	if (bPreparationDone && DataAdaptor && DataAdaptor->IsUploadingFromCPU())
+	{
+		NotifyProducerUploadedData(InBinding);
+	}
+#endif
+
+	return bPreparationDone;
 }
 
 void UPCGDataProviderDataCollectionUpload::ReleaseTransientResources(const TCHAR* InReason)
diff --git a/Engine/Plugins/PCG/Source/PCG/Private/Compute/DataInterfaces/PCGLandscapeDataInterface.cpp b/Engine/Plugins/PCG/Source/PCG/Private/Compute/DataInterfaces/PCGLandscapeDataInterface.cpp
--- a/Engine/Plugins/PCG/Source/PCG/Private/Compute/DataInterfaces/PCGLandscapeDataInterface.cpp
+++ b/Engine/Plugins/PCG/Source/PCG/Private/Compute/DataInterfaces/PCGLandscapeDataInterface.cpp
@@ -277,7 +277,10 @@ void UPCGLandscapeDataProvider::Initialize(const UComputeDataInterface* InDataIn
 	Super::Initialize(InDataInterface, InBinding, InInputMask, InOutputMask);
 
 	const UPCGLandscapeDataInterface* DataInterface = CastChecked<UPCGLandscapeDataInterface>(InDataInterface);
+	ProducerSettings = DataInterface->ProducerSettings;
+
 	UPCGDataBinding* Binding = CastChecked<UPCGDataBinding>(InBinding);
+
 	const IPCGGraphExecutionSource* ExecutionSource = Binding->GetExecutionSource();
 	check(ExecutionSource);
 
@@ -326,7 +329,7 @@ void UPCGLandscapeDataProvider::Initialize(const UComputeDataInterface* InDataIn
 					{
 						if (TestLandscape(Landscape))
 						{
-							InitFromLandscape(Landscape, ComponentBounds, bCanSampleVirtualTextures && LandscapeData->CanSampleVirtualTextures(), LandscapeData->CanSampleVirtualTextureNormals());
+							InitFromLandscape(Binding, Landscape, ComponentBounds, bCanSampleVirtualTextures && LandscapeData->CanSampleVirtualTextures(), LandscapeData->CanSampleVirtualTextureNormals());
 							break;
 						}
 					}
@@ -336,7 +339,7 @@ void UPCGLandscapeDataProvider::Initialize(const UComputeDataInterface* InDataIn
 	}
 }
 
-void UPCGLandscapeDataProvider::InitFromLandscape(ALandscape* InLandscape, const FBox& Bounds, bool bAllowSampleVirtualTexture, bool bAllowSampleVirtualTextureNormal)
+void UPCGLandscapeDataProvider::InitFromLandscape(UPCGDataBinding* InBinding, ALandscape* InLandscape, const FBox& Bounds, bool bAllowSampleVirtualTexture, bool bAllowSampleVirtualTextureNormal)
 {
 	TRACE_CPUPROFILER_EVENT_SCOPE(UPCGLandscapeDataProvider::InitFromLandscape);
 	check(IsInGameThread());
@@ -354,7 +357,7 @@ void UPCGLandscapeDataProvider::InitFromLandscape(ALandscape* InLandscape, const
 	// If the height can't be accessed via RVT, build a physics based height texture instead.
 	if (!HeightVirtualTexture)
 	{
-		InitializeFromLandscapeCollision(InLandscape, Bounds);
+		InitializeFromLandscapeCollision(InBinding, InLandscape, Bounds);
 	}
 
 	LandscapeGridSize = FVector2D(InLandscape->GetTransform().GetScale3D());
@@ -438,7 +441,7 @@ void UPCGLandscapeDataProvider::InitializeRuntimeVirtualTextures(ALandscape* InL
 	}
 }
 
-void UPCGLandscapeDataProvider::InitializeFromLandscapeCollision(ALandscape* InLandscape, const FBox& Bounds)
+void UPCGLandscapeDataProvider::InitializeFromLandscapeCollision(UPCGDataBinding* InBinding, ALandscape* InLandscape, const FBox& Bounds)
 {
 	TRACE_CPUPROFILER_EVENT_SCOPE(UPCGLandscapeDataProvider::InitializeFromLandscapeCollision);
 	check(InLandscape);
@@ -533,6 +536,10 @@ void UPCGLandscapeDataProvider::InitializeFromLandscapeCollision(ALandscape* InL
 	}
 
 	Resource = MakeShared<FPCGLandscapeResource>(Key);
+
+#if WITH_EDITOR
+	NotifyProducerUploadedData(InBinding);
+#endif // WITH_EDITOR
 }
 
 FComputeDataProviderRenderProxy* UPCGLandscapeDataProvider::GetRenderProxy()
diff --git a/Engine/Plugins/PCG/Source/PCG/Private/Compute/DataInterfaces/PCGLandscapeDataInterface.h b/Engine/Plugins/PCG/Source/PCG/Private/Compute/DataInterfaces/PCGLandscapeDataInterface.h
--- a/Engine/Plugins/PCG/Source/PCG/Private/Compute/DataInterfaces/PCGLandscapeDataInterface.h
+++ b/Engine/Plugins/PCG/Source/PCG/Private/Compute/DataInterfaces/PCGLandscapeDataInterface.h
@@ -64,7 +64,7 @@ private:
 
 /** Compute Framework Data Provider for reading a landscape. */
 UCLASS()
-class UPCGLandscapeDataProvider : public UComputeDataProvider
+class UPCGLandscapeDataProvider : public UPCGComputeDataProvider
 {
 	GENERATED_BODY()
 
@@ -76,9 +76,9 @@ public:
 	//~ End UComputeDataProvider Interface
 
 private:
-	void InitFromLandscape(ALandscape* InLandscape, const FBox& Bounds, bool bAllowSampleVirtualTexture, bool bAllowSampleVirtualTextureNormal);
+	void InitFromLandscape(UPCGDataBinding* InBinding, ALandscape* InLandscape, const FBox& Bounds, bool bAllowSampleVirtualTexture, bool bAllowSampleVirtualTextureNormal);
 	void InitializeRuntimeVirtualTextures(ALandscape* InLandscape, bool bAllowSampleVirtualTextureNormal);
-	void InitializeFromLandscapeCollision(ALandscape* InLandscape, const FBox& Bounds);
+	void InitializeFromLandscapeCollision(UPCGDataBinding* InBinding, ALandscape* InLandscape, const FBox& Bounds);
 
 private:
 	UPROPERTY()
diff --git a/Engine/Plugins/PCG/Source/PCG/Private/Compute/PCGDataUploadAdaptor.h b/Engine/Plugins/PCG/Source/PCG/Private/Compute/PCGDataUploadAdaptor.h
--- a/Engine/Plugins/PCG/Source/PCG/Private/Compute/PCGDataUploadAdaptor.h
+++ b/Engine/Plugins/PCG/Source/PCG/Private/Compute/PCGDataUploadAdaptor.h
@@ -28,6 +28,8 @@ public:
 	/** Gets the buffer that can then be used to read the data in kernels. */
 	FRDGBufferRef GetBuffer_RenderThread(FRDGBuilder& InGraphBuilder, EPCGExportMode InExportMode);
 
+	bool IsUploadingFromCPU() const { return !ExternalBufferForReuse; }
+
 private:
 	TRefCountPtr<FRDGPooledBuffer> ExternalBufferForReuse;
 
diff --git a/Engine/Plugins/PCG/Source/PCG/Private/Elements/PCGGenerateGrassMaps.cpp b/Engine/Plugins/PCG/Source/PCG/Private/Elements/PCGGenerateGrassMaps.cpp
--- a/Engine/Plugins/PCG/Source/PCG/Private/Elements/PCGGenerateGrassMaps.cpp
+++ b/Engine/Plugins/PCG/Source/PCG/Private/Elements/PCGGenerateGrassMaps.cpp
@@ -649,6 +649,16 @@ bool FPCGGenerateGrassMapsElement::ExecuteInternal(FPCGContext* InContext) const
 		return false;
 	}
 
+#if WITH_EDITOR
+	if (!Context->TextureDatas.IsEmpty() && !Settings->bSkipReadbackToCPU)
+	{
+		if (!Settings->bSkipReadbackToCPU && Context->Node && Context->GetStack() && Context->ExecutionSource.IsValid())
+		{
+			Context->ExecutionSource->GetExecutionState().GetInspection().NotifyGPUToCPUReadback(Context->Node, Context->GetStack());
+		}
+	}
+#endif // WITH_EDITOR
+
 	// 5. Emit texture data objects.
 	for (int DataIndex = 0; DataIndex < Context->SelectedGrassTypes.Num(); ++DataIndex)
 	{
diff --git a/Engine/Plugins/PCG/Source/PCG/Private/Elements/PCGTextureSampler.cpp b/Engine/Plugins/PCG/Source/PCG/Private/Elements/PCGTextureSampler.cpp
--- a/Engine/Plugins/PCG/Source/PCG/Private/Elements/PCGTextureSampler.cpp
+++ b/Engine/Plugins/PCG/Source/PCG/Private/Elements/PCGTextureSampler.cpp
@@ -227,6 +227,13 @@ bool FPCGTextureSamplerElement::ExecuteInternal(FPCGContext* InContext) const
 		UPCGRenderTargetData* RenderTargetData = FPCGContext::NewObject_AnyThread<UPCGRenderTargetData>(Context);
 		RenderTargetData->Initialize(RenderTarget, Context->Transform, Settings->bSkipReadbackToCPU);
 		BaseTextureData = RenderTargetData;
+
+#if WITH_EDITOR
+		if (!Settings->bSkipReadbackToCPU && Context->Node && Context->GetStack() && Context->ExecutionSource.IsValid())
+		{
+			Context->ExecutionSource->GetExecutionState().GetInspection().NotifyGPUToCPUReadback(Context->Node, Context->GetStack());
+		}
+#endif // WITH_EDITOR
 	}
 	else
 	{
@@ -257,6 +264,13 @@ bool FPCGTextureSamplerElement::ExecuteInternal(FPCGContext* InContext) const
 			const bool bForceEditorOnlyCPUSampling = false;
 #endif
 
+#if WITH_EDITOR
+			if (!Settings->bSkipReadbackToCPU && Context->Node && Context->GetStack() && Context->ExecutionSource.IsValid())
+			{
+				Context->ExecutionSource->GetExecutionState().GetInspection().NotifyGPUToCPUReadback(Context->Node, Context->GetStack());
+			}
+#endif
+
 			if (!TextureData->Initialize(Texture, TextureArrayIndex, Context->Transform, bForceEditorOnlyCPUSampling, Settings->bSkipReadbackToCPU))
 			{
 				// Initialization not complete. Could be waiting on async texture processing or for GPU readback. Sleep until next frame.
diff --git a/Engine/Plugins/PCG/Source/PCG/Private/PCGElement.cpp b/Engine/Plugins/PCG/Source/PCG/Private/PCGElement.cpp
--- a/Engine/Plugins/PCG/Source/PCG/Private/PCGElement.cpp
+++ b/Engine/Plugins/PCG/Source/PCG/Private/PCGElement.cpp
@@ -291,8 +291,7 @@ void IPCGElement::PreExecute(FPCGContext* Context) const
 
 	const UPCGSettings* Settings = Context->GetInputSettings<UPCGSettings>();
 
-	// If we were supposed to execute on GPU and end up here, then GPU compilation failed. Run validation in noisy mode
-	// to generate runtime errors/graph errors, then pass through.
+	// If we were supposed to execute on GPU and end up here, then GPU compilation failed. Pass through.
 	if (Settings && Settings->ShouldExecuteOnGPU())
 	{
 		DisabledPassThroughData(Context);
@@ -406,6 +405,13 @@ bool IPCGElement::ConvertInputsIfNeeded(FPCGContext* Context) const
 		{
 			if (const UPCGProxyForGPUData* DataWithGPUSupport = Cast<UPCGProxyForGPUData>(TaggedData.Data))
 			{
+#if WITH_EDITOR
+				if (Context->Node && Context->GetStack() && Context->ExecutionSource.IsValid())
+				{
+					Context->ExecutionSource->GetExecutionState().GetInspection().NotifyGPUToCPUReadback(Context->Node, Context->GetStack());
+				}
+#endif
+
 				// Poll until readback is done.
 				UPCGProxyForGPUData::FReadbackResult Result = DataWithGPUSupport->GetCPUData(Context);
 
diff --git a/Engine/Plugins/PCG/Source/PCG/Private/PCGGraphExecutionInspection.cpp b/Engine/Plugins/PCG/Source/PCG/Private/PCGGraphExecutionInspection.cpp
--- a/Engine/Plugins/PCG/Source/PCG/Private/PCGGraphExecutionInspection.cpp
+++ b/Engine/Plugins/PCG/Source/PCG/Private/PCGGraphExecutionInspection.cpp
@@ -228,6 +228,12 @@ void FPCGGraphExecutionInspection::ClearInspectionData(bool bClearPerNodeExecuti
 			FWriteScopeLock Lock(NodeToStackToInactivePinMaskLock);
 			NodeToStackToInactivePinMask.Reset();
 		}
+
+		{
+			FWriteScopeLock Lock(NodeToStacksTriggeringGPUTransfersLock);
+			NodeToStacksTriggeringGPUUploads.Reset();
+			NodeToStacksTriggeringGPUReadbacks.Reset();
+		}
 	}
 }
 
@@ -240,6 +246,42 @@ bool FPCGGraphExecutionInspection::HasNodeProducedData(const UPCGNode* InNode, c
 	return StacksThatProducedData && StacksThatProducedData->Contains(Stack);
 }
 
+void FPCGGraphExecutionInspection::NotifyGPUToCPUReadback(const UPCGNode* InNode, const FPCGStack* InStack) const
+{
+	if (!ensure(InNode) || !ensure(InStack))
+	{
+		return;
+	}
+
+	FWriteScopeLock Lock(NodeToStacksTriggeringGPUTransfersLock);
+	NodeToStacksTriggeringGPUReadbacks.FindOrAdd(InNode).Add(*InStack);
+}
+
+void FPCGGraphExecutionInspection::NotifyCPUToGPUUpload(const UPCGNode* InNode, const FPCGStack* InStack) const
+{
+	if (!ensure(InNode) || !ensure(InStack))
+	{
+		return;
+	}
+
+	FWriteScopeLock Lock(NodeToStacksTriggeringGPUTransfersLock);
+	NodeToStacksTriggeringGPUUploads.FindOrAdd(InNode).Add(*InStack);
+}
+
+bool FPCGGraphExecutionInspection::DidNodeTriggerGPUToCPUReadback(const UPCGNode* InNode, const FPCGStack& Stack) const
+{
+	FReadScopeLock Lock(NodeToStacksTriggeringGPUTransfersLock);
+	const TSet<FPCGStack>* FoundStacks = NodeToStacksTriggeringGPUReadbacks.Find(InNode);
+	return FoundStacks && FoundStacks->Contains(Stack);
+}
+
+bool FPCGGraphExecutionInspection::DidNodeTriggerCPUToGPUUpload(const UPCGNode* InNode, const FPCGStack& Stack) const
+{
+	FReadScopeLock Lock(NodeToStacksTriggeringGPUTransfersLock);
+	const TSet<FPCGStack>* FoundStacks = NodeToStacksTriggeringGPUUploads.Find(InNode);
+	return FoundStacks && FoundStacks->Contains(Stack);
+}
+
 void FPCGGraphExecutionInspection::AddReferencedObjects(FReferenceCollector& Collector)
 {
 	for (auto& Pair : InspectionCache)
diff --git a/Engine/Plugins/PCG/Source/PCG/Public/PCGGraphExecutionInspection.h b/Engine/Plugins/PCG/Source/PCG/Public/PCGGraphExecutionInspection.h
--- a/Engine/Plugins/PCG/Source/PCG/Public/PCGGraphExecutionInspection.h
+++ b/Engine/Plugins/PCG/Source/PCG/Public/PCGGraphExecutionInspection.h
@@ -55,6 +55,12 @@ public:
 	/** Did the given node produce one or more data items in the given stack during the last execution. */
 	UE_API bool HasNodeProducedData(const UPCGNode* InNode, const FPCGStack& Stack) const;
 
+	UE_API void NotifyGPUToCPUReadback(const UPCGNode* InNode, const FPCGStack* InStack) const;
+	UE_API void NotifyCPUToGPUUpload(const UPCGNode* InNode, const FPCGStack* InStack) const;
+
+	UE_API bool DidNodeTriggerGPUToCPUReadback(const UPCGNode* InNode, const FPCGStack& Stack) const;
+	UE_API bool DidNodeTriggerCPUToGPUUpload(const UPCGNode* InNode, const FPCGStack& Stack) const;
+
 private:
 	UE_API void AddReferencedObjects(FReferenceCollector& Collector);
 
@@ -75,6 +81,11 @@ private:
 	/** Map from nodes to stacks to mask of output pins that were deactivated during execution. */
 	mutable TMap<TObjectKey<const UPCGNode>, TMap<const FPCGStack, uint64>> NodeToStackToInactivePinMask;
 	mutable FRWLock NodeToStackToInactivePinMaskLock;
+
+	/** Map from nodes to all stacks for which GPU data transfers occurred. */
+	mutable TMap<TObjectKey<const UPCGNode>, TSet<FPCGStack>> NodeToStacksTriggeringGPUUploads;
+	mutable TMap<TObjectKey<const UPCGNode>, TSet<FPCGStack>> NodeToStacksTriggeringGPUReadbacks;
+	mutable FRWLock NodeToStacksTriggeringGPUTransfersLock;
 };
 #endif
 
diff --git a/Engine/Plugins/PCG/Source/PCGEditor/Private/Nodes/PCGEditorGraphNodeBase.cpp b/Engine/Plugins/PCG/Source/PCGEditor/Private/Nodes/PCGEditorGraphNodeBase.cpp
--- a/Engine/Plugins/PCG/Source/PCGEditor/Private/Nodes/PCGEditorGraphNodeBase.cpp
+++ b/Engine/Plugins/PCG/Source/PCGEditor/Private/Nodes/PCGEditorGraphNodeBase.cpp
@@ -169,6 +169,7 @@ void UPCGEditorGraphNodeBase::Construct(UPCGNode* InPCGNode)
 	const UPCGEditorGraph* Graph = Cast<UPCGEditorGraph>(GetOuter());
 	const FPCGEditor* Editor = Graph ? Graph->GetEditor().Pin().Get() : nullptr;
 	UpdateStructuralVisualization(Editor ? Editor->GetPCGComponentBeingInspected() : nullptr, Editor ? Editor->GetStackBeingInspected() : nullptr, /*bNewlyPlaced=*/true);
+	UpdateGPUVisualization(Editor ? Editor->GetPCGComponentBeingInspected() : nullptr, Editor ? Editor->GetStackBeingInspected() : nullptr);
 }
 
 void UPCGEditorGraphNodeBase::BeginDestroy()
@@ -805,6 +806,33 @@ EPCGChangeType UPCGEditorGraphNodeBase::UpdateStructuralVisualization(UPCGCompon
 	return ChangeType;
 }
 
+EPCGChangeType UPCGEditorGraphNodeBase::UpdateGPUVisualization(class UPCGComponent* InComponentBeingDebugged, const FPCGStack* InStackBeingInspected)
+{
+	EPCGChangeType ChangeType = EPCGChangeType::None;
+
+	const UPCGNode* Node = GetPCGNode();
+
+	const bool bTriggeredUpload = InComponentBeingDebugged && InStackBeingInspected && Node
+		&& InComponentBeingDebugged->GetExecutionState().GetInspection().DidNodeTriggerCPUToGPUUpload(Node, *InStackBeingInspected);
+
+	if (bTriggeredUpload != GetTriggeredGPUUpload())
+	{
+		SetTriggeredGPUUpload(bTriggeredUpload);
+		ChangeType |= EPCGChangeType::Cosmetic;
+	}
+
+	const bool bTriggeredReadback = InComponentBeingDebugged && InStackBeingInspected && Node
+		&& InComponentBeingDebugged->GetExecutionState().GetInspection().DidNodeTriggerGPUToCPUReadback(Node, *InStackBeingInspected);
+
+	if (bTriggeredReadback != GetTriggeredGPUReadback())
+	{
+		SetTriggeredGPUReadback(bTriggeredReadback);
+		ChangeType |= EPCGChangeType::Cosmetic;
+	}
+	 
+	return ChangeType;
+}
+
 bool UPCGEditorGraphNodeBase::CanUserAddRemoveDynamicInputPins() const
 {
 	const UPCGSettings* Settings = PCGNode ? PCGNode->GetSettings() : nullptr;
diff --git a/Engine/Plugins/PCG/Source/PCGEditor/Private/Nodes/PCGEditorGraphNodeBase.h b/Engine/Plugins/PCG/Source/PCGEditor/Private/Nodes/PCGEditorGraphNodeBase.h
--- a/Engine/Plugins/PCG/Source/PCGEditor/Private/Nodes/PCGEditorGraphNodeBase.h
+++ b/Engine/Plugins/PCG/Source/PCGEditor/Private/Nodes/PCGEditorGraphNodeBase.h
@@ -76,6 +76,12 @@ public:
 	void SetIsCulledFromExecution(bool bInIsCulledFromExecution) { bIsCulledFromExecution = bInIsCulledFromExecution; }
 	bool IsCulledFromExecution() const { return bIsCulledFromExecution; }
 
+	void SetTriggeredGPUUpload(bool bInTriggeredGPUUpload) { bTriggeredGPUUpload = bInTriggeredGPUUpload; }
+	bool GetTriggeredGPUUpload() const { return bTriggeredGPUUpload; }
+
+	void SetTriggeredGPUReadback(bool bInTriggeredGPUReadback) { bTriggeredGPUReadback = bInTriggeredGPUReadback; }
+	bool GetTriggeredGPUReadback() const { return bTriggeredGPUReadback; }
+
 	/** Increase deferred reconstruct counter, calls to ReconstructNode will flag reconstruct to happen when count hits zero */
 	void EnableDeferredReconstruct();
 	/** Decrease deferred reconstruct counter, ReconstructNode will be called if counter hits zero and the node is flagged for reconstruction  */
@@ -87,6 +93,8 @@ public:
 	/** If the currently inspected grid size is smaller than the grid size of this node, display transparent. */
 	EPCGChangeType UpdateStructuralVisualization(class UPCGComponent* InComponentBeingDebugged, const FPCGStack* InStackBeingInspected, bool bNewlyPlaced = false);
 
+	EPCGChangeType UpdateGPUVisualization(class UPCGComponent* InComponentBeingDebugged, const FPCGStack* InStackBeingInspected);
+
 	bool CanUserAddRemoveDynamicInputPins() const;
 
 	/** The settings support default value inline constants. */
@@ -186,6 +194,9 @@ protected:
 	/** Whether this node was culled in the last execution. */
 	bool bIsCulledFromExecution = false;
 
+	bool bTriggeredGPUUpload = false;
+	bool bTriggeredGPUReadback = false;
+
 	/** Bitmask of inactive output pins. Bit N will be set if output pin index N is inactive. */
 	uint64 InactiveOutputPinMask = 0;
 
diff --git a/Engine/Plugins/PCG/Source/PCGEditor/Private/Nodes/SPCGEditorGraphNode.cpp b/Engine/Plugins/PCG/Source/PCGEditor/Private/Nodes/SPCGEditorGraphNode.cpp
--- a/Engine/Plugins/PCG/Source/PCGEditor/Private/Nodes/SPCGEditorGraphNode.cpp
+++ b/Engine/Plugins/PCG/Source/PCGEditor/Private/Nodes/SPCGEditorGraphNode.cpp
@@ -297,23 +297,43 @@ void SPCGEditorGraphNode::GetOverlayBrushes(bool bSelected, const FVector2f& Wid
 {
 	check(PCGEditorGraphNode);
 	
-	FVector2f OverlayOffset(0.0f, 0.0f);
+	float YOffsetLeft = 0.0f;
+	// Start lower down to be clear of grid size label.
+	float YOffsetRight = UsesHiGenOverlay() ? 18.0f : 0.0f;
 
-	auto AddOverlayBrush = [&OverlayOffset, &Brushes](const FName& BrushName)
+	auto AddOverlayBrush = [&YOffsetLeft, &YOffsetRight, &Brushes, this](const FName& BrushName, bool bRightSide = false)
 	{
 		const FSlateBrush* Brush = FPCGEditorStyle::Get().GetBrush(BrushName);
 
 		if (Brush)
 		{
+			float& YOffset = bRightSide ? YOffsetRight : YOffsetLeft;
+
 			FOverlayBrushInfo BrushInfo;
 			BrushInfo.Brush = Brush;
-			BrushInfo.OverlayOffset = OverlayOffset - Brush->GetImageSize() / 2.0f;
+			BrushInfo.OverlayOffset = FVector2f(0.0f, YOffset) - Brush->GetImageSize() / 2.0f;
+
+			if (bRightSide)
+			{
+				BrushInfo.OverlayOffset.X += GetDesiredSize().X;
+			}
+
 			Brushes.Add(BrushInfo);
 
-			OverlayOffset.Y += Brush->GetImageSize().Y;
+			YOffset += Brush->GetImageSize().Y;
 		}
 	};
 
+	if (PCGEditorGraphNode->GetTriggeredGPUUpload())
+	{
+		AddOverlayBrush(TEXT("PCG.NodeOverlay.GPUUpload"), /*bRightSide=*/true);
+	}
+
+	if (PCGEditorGraphNode->GetTriggeredGPUReadback())
+	{
+		AddOverlayBrush(TEXT("PCG.NodeOverlay.GPUReadback"));
+	}
+
 	if (PCGEditorGraphNode->IsCulledFromExecution())
 	{
 		AddOverlayBrush(PCGEditorStyleConstants::Node_Overlay_Inactive);
diff --git a/Engine/Plugins/PCG/Source/PCGEditor/Private/Nodes/SPCGEditorGraphNode.h b/Engine/Plugins/PCG/Source/PCGEditor/Private/Nodes/SPCGEditorGraphNode.h
--- a/Engine/Plugins/PCG/Source/PCGEditor/Private/Nodes/SPCGEditorGraphNode.h
+++ b/Engine/Plugins/PCG/Source/PCGEditor/Private/Nodes/SPCGEditorGraphNode.h
@@ -49,9 +49,15 @@ private:
 	/** Adds the Hierarchical Generation overlay to the array, displaying the HiGen grid size on the node. */
 	void AddHiGenOverlayWidget(TArray<FOverlayWidgetInfo>& OverlayWidgets) const;
 
-	/** Adds the "GPU" tag to the node, for . */
+	/** Adds the "GPU" tag to the node to indicate the node will execute on the GPU. */
 	void AddGPUOverlayWidget(TArray<FOverlayWidgetInfo>& OverlayWidgets) const;
 
+	/** Indicates an upload of data to the GPU occurred as a result of this node. */
+	void AddGPUUploadWidget(TArray<FOverlayWidgetInfo>& OverlayWidgets) const;
+
+	/** Indicates a readback of data from the GPU occurred as a result of this node. */
+	void AddGPUReadbackWidget(TArray<FOverlayWidgetInfo>& OverlayWidgets) const;
+
 	/**
 	* Get the border brush for the given combination of grid sizes and enabled state. All a big workaround for FSlateRoundedBoxBrush not respecting
 	* the tint colour.
diff --git a/Engine/Plugins/PCG/Source/PCGEditor/Private/PCGEditor.cpp b/Engine/Plugins/PCG/Source/PCGEditor/Private/PCGEditor.cpp
--- a/Engine/Plugins/PCG/Source/PCGEditor/Private/PCGEditor.cpp
+++ b/Engine/Plugins/PCG/Source/PCGEditor/Private/PCGEditor.cpp
@@ -372,6 +372,7 @@ void FPCGEditor::UpdateAfterInspectedStackChanged()
 			// Update now that component has changed. Will fire OnNodeChanged if necessary.
 			EPCGChangeType ChangeType = PCGNode->UpdateErrorsAndWarnings();
 			ChangeType |= PCGNode->UpdateStructuralVisualization(Component, &StackBeingInspected);
+			ChangeType |= PCGNode->UpdateGPUVisualization(Component, &StackBeingInspected);
 
 			if (ChangeType != EPCGChangeType::None)
 			{
@@ -3154,6 +3155,7 @@ void FPCGEditor::OnComponentGenerationDone(UPCGSubsystem* Subsystem, UPCGCompone
 
 			ChangeType |= PCGEditorGraphNode->UpdateErrorsAndWarnings();
 			ChangeType |= PCGEditorGraphNode->UpdateStructuralVisualization(PCGComponentBeingInspected.Get(), &StackBeingInspected);
+			ChangeType |= PCGEditorGraphNode->UpdateGPUVisualization(PCGComponentBeingInspected.Get(), &StackBeingInspected);
 
 			if (ChangeType != EPCGChangeType::None)
 			{
diff --git a/Engine/Plugins/PCG/Source/PCGEditor/Private/PCGEditorGraph.cpp b/Engine/Plugins/PCG/Source/PCGEditor/Private/PCGEditorGraph.cpp
--- a/Engine/Plugins/PCG/Source/PCGEditor/Private/PCGEditorGraph.cpp
+++ b/Engine/Plugins/PCG/Source/PCGEditor/Private/PCGEditorGraph.cpp
@@ -142,7 +142,7 @@ void UPCGEditorGraph::ReconstructGraph()
 	}
 
 	// Ensure graph structure visualization is nice and fresh upon opening.
-	UpdateStructuralVisualization(nullptr, nullptr);
+	UpdateVisualizations(nullptr, nullptr);
 }
 
 void UPCGEditorGraph::BeginDestroy()
@@ -178,14 +178,20 @@ void UPCGEditorGraph::CreateLinks(UPCGEditorGraphNodeBase* GraphNode, bool bCrea
 	CreateLinks(GraphNode, bCreateInbound, bCreateOutbound, PCGNodeToPCGEditorNodeMap);
 }
 
-void UPCGEditorGraph::UpdateStructuralVisualization(UPCGComponent* PCGComponentBeingInspected, const FPCGStack* PCGStackBeingInspected)
+void UPCGEditorGraph::UpdateVisualizations(UPCGComponent* PCGComponentBeingInspected, const FPCGStack* PCGStackBeingInspected)
 {
 	for (UEdGraphNode* EditorNode : Nodes)
 	{
-		UPCGEditorGraphNodeBase* PCGEditorNode = Cast<UPCGEditorGraphNodeBase>(EditorNode);
-		if (PCGEditorNode && PCGEditorNode->UpdateStructuralVisualization(PCGComponentBeingInspected, PCGStackBeingInspected) != EPCGChangeType::None)
+		if (UPCGEditorGraphNodeBase* PCGEditorNode = Cast<UPCGEditorGraphNodeBase>(EditorNode))
 		{
-			PCGEditorNode->ReconstructNode();
+			EPCGChangeType ChangeType = EPCGChangeType::None;
+			ChangeType |= PCGEditorNode->UpdateStructuralVisualization(PCGComponentBeingInspected, PCGStackBeingInspected);
+			ChangeType |= PCGEditorNode->UpdateGPUVisualization(PCGComponentBeingInspected, PCGStackBeingInspected);
+
+			if (ChangeType != EPCGChangeType::None)
+			{
+				PCGEditorNode->ReconstructNode();
+			}
 		}
 	}
 }
diff --git a/Engine/Plugins/PCG/Source/PCGEditor/Private/PCGEditorGraph.h b/Engine/Plugins/PCG/Source/PCGEditor/Private/PCGEditorGraph.h
--- a/Engine/Plugins/PCG/Source/PCGEditor/Private/PCGEditorGraph.h
+++ b/Engine/Plugins/PCG/Source/PCGEditor/Private/PCGEditorGraph.h
@@ -48,7 +48,7 @@ public:
 	TWeakPtr<FPCGEditor> GetEditor() const { return PCGEditor; }
 
 	/** Updates the grid size visualization in the editor. */
-	void UpdateStructuralVisualization(UPCGComponent* PCGComponentBeingInspected, const FPCGStack* PCGStackBeingInspected);
+	void UpdateVisualizations(UPCGComponent* PCGComponentBeingInspected, const FPCGStack* PCGStackBeingInspected);
 
 	/** Returns the PCG editor graph node corresponding to the given PCG node. */
 	const UPCGEditorGraphNodeBase* GetEditorNodeFromPCGNode(const UPCGNode* InPCGNode) const;
diff --git a/Engine/Plugins/PCG/Source/PCGEditor/Private/PCGEditorStyle.cpp b/Engine/Plugins/PCG/Source/PCGEditor/Private/PCGEditorStyle.cpp
--- a/Engine/Plugins/PCG/Source/PCGEditor/Private/PCGEditorStyle.cpp
+++ b/Engine/Plugins/PCG/Source/PCGEditor/Private/PCGEditorStyle.cpp
@@ -36,6 +36,8 @@ FPCGEditorStyle::FPCGEditorStyle() : FSlateStyleSet("PCGEditorStyle")
 	Set("PCG.NodeOverlay.Debug", new CORE_IMAGE_BRUSH_SVG("Starship/Blueprints/Breakpoint_Valid", Icon20x20, FSlateColor(FColor::Cyan)));
 	Set("PCG.NodeOverlay.Inspect", new CORE_IMAGE_BRUSH_SVG("Starship/Blueprints/Breakpoint_Valid", Icon20x20, FSlateColor(FColor::Orange)));
 	Set("PCG.NodeOverlay.OnInactiveBranch", new CORE_IMAGE_BRUSH_SVG("Starship/Blueprints/CompileStatus_Fail_Badge", Icon20x20, FSlateColor(FColor::White)));
+	Set("PCG.NodeOverlay.GPUUpload", new CORE_IMAGE_BRUSH("Common/UpArrow", Icon16x16, FSlateColor(FColor::Yellow)));
+	Set("PCG.NodeOverlay.GPUReadback", new CORE_IMAGE_BRUSH("Common/DownArrow", Icon16x16, FSlateColor(FColor::Yellow)));
 
 	static const FVector2D PinSize(22.0f, 22.0f);
 	Set(PCGEditorStyleConstants::Pin_SD_SC_IN_C, new IMAGE_BRUSH_SVG("Style/PCG_Graph_SD_SC_IN", PinSize));
