[PCG] Fix Attribute select returning the default point instead of the selected point

* Regression from Point Array data, point value range was considered empty
* Add a condition that if we are an orphan point data with a single point, we create the range on the default value
* Also change it in the Get Attribute From Point Index

#rnx

diff --git a/Engine/Plugins/PCG/Source/PCG/Private/Elements/PCGAttributeGetFromPointIndexElement.cpp b/Engine/Plugins/PCG/Source/PCG/Private/Elements/PCGAttributeGetFromPointIndexElement.cpp
--- a/Engine/Plugins/PCG/Source/PCG/Private/Elements/PCGAttributeGetFromPointIndexElement.cpp
+++ b/Engine/Plugins/PCG/Source/PCG/Private/Elements/PCGAttributeGetFromPointIndexElement.cpp
@@ -172,11 +172,15 @@ bool FPCGAttributeGetFromPointIndexElement::ExecuteInternal(FPCGContext* Context
 
 					FPCGInitializeFromDataParams InitializeFromDataParams(PointData);
 					InitializeFromDataParams.bInheritSpatialData = false;
-					
 					OutputPointData->InitializeFromDataWithParams(InitializeFromDataParams);
+
 					OutputPointData->SetNumPoints(1);
+					OutputPointData->CopyUnallocatedPropertiesFrom(PointData);
+
+					const FConstPCGPointValueRanges InRanges(PointData);
+					FPCGPointValueRanges OutRanges(OutputPointData, /*bAllocate=*/false);
 
-					PointData->CopyPointsTo(OutputPointData, Index, 0, 1);
+					OutRanges.SetFromValueRanges(0, InRanges, Index);
 
 					FPCGTaggedData& OutputPoint = Outputs.Add_GetRef(Input);
 					OutputPoint.Data = OutputPointData;
diff --git a/Engine/Plugins/PCG/Source/PCG/Public/Data/PCGBasePointData.h b/Engine/Plugins/PCG/Source/PCG/Public/Data/PCGBasePointData.h
--- a/Engine/Plugins/PCG/Source/PCG/Public/Data/PCGBasePointData.h
+++ b/Engine/Plugins/PCG/Source/PCG/Public/Data/PCGBasePointData.h
@@ -352,43 +352,48 @@ struct FPCGPointValueRanges
 	FPCGPointValueRanges(UPCGBasePointData* InBasePointData, bool bAllocate = true)
 	{
 		EPCGPointNativeProperties AllocatedProperties = InBasePointData->GetAllocatedProperties(/*bWithInheritance=*/false);
+		
+		// In the very specific case where we have just a single point, that we don't inherit from a spatial parent (Orphan point data with a single point)
+		// we can create a range that will allow to write into the default value of the point data.
+		const bool bIsSingleOrphanPointData = !bAllocate && InBasePointData->GetNumPoints() == 1 && !InBasePointData->HasSpatialDataParent();
+		const bool bMustCreateRange = bAllocate || bIsSingleOrphanPointData;
 
-		if (bAllocate || EnumHasAnyFlags(AllocatedProperties, EPCGPointNativeProperties::Transform))
+		if (bMustCreateRange || EnumHasAnyFlags(AllocatedProperties, EPCGPointNativeProperties::Transform))
 		{
 			TransformRange = InBasePointData->GetTransformValueRange(bAllocate);
 		}
 
-		if (bAllocate || EnumHasAnyFlags(AllocatedProperties, EPCGPointNativeProperties::Density))
+		if (bMustCreateRange || EnumHasAnyFlags(AllocatedProperties, EPCGPointNativeProperties::Density))
 		{
 			DensityRange = InBasePointData->GetDensityValueRange(bAllocate);
 		}
 
-		if (bAllocate || EnumHasAnyFlags(AllocatedProperties, EPCGPointNativeProperties::Steepness))
+		if (bMustCreateRange || EnumHasAnyFlags(AllocatedProperties, EPCGPointNativeProperties::Steepness))
 		{
 			SteepnessRange = InBasePointData->GetSteepnessValueRange(bAllocate);
 		}
 
-		if (bAllocate || EnumHasAnyFlags(AllocatedProperties, EPCGPointNativeProperties::BoundsMin))
+		if (bMustCreateRange || EnumHasAnyFlags(AllocatedProperties, EPCGPointNativeProperties::BoundsMin))
 		{
 			BoundsMinRange = InBasePointData->GetBoundsMinValueRange(bAllocate);
 		}
 
-		if (bAllocate || EnumHasAnyFlags(AllocatedProperties, EPCGPointNativeProperties::BoundsMax))
+		if (bMustCreateRange || EnumHasAnyFlags(AllocatedProperties, EPCGPointNativeProperties::BoundsMax))
 		{
 			BoundsMaxRange = InBasePointData->GetBoundsMaxValueRange(bAllocate);
 		}
 
-		if (bAllocate || EnumHasAnyFlags(AllocatedProperties, EPCGPointNativeProperties::Color))
+		if (bMustCreateRange || EnumHasAnyFlags(AllocatedProperties, EPCGPointNativeProperties::Color))
 		{
 			ColorRange = InBasePointData->GetColorValueRange(bAllocate);
 		}
 
-		if (bAllocate || EnumHasAnyFlags(AllocatedProperties, EPCGPointNativeProperties::Seed))
+		if (bMustCreateRange || EnumHasAnyFlags(AllocatedProperties, EPCGPointNativeProperties::Seed))
 		{
 			SeedRange = InBasePointData->GetSeedValueRange(bAllocate);
 		}
 
-		if (bAllocate || EnumHasAnyFlags(AllocatedProperties, EPCGPointNativeProperties::MetadataEntry))
+		if (bMustCreateRange || EnumHasAnyFlags(AllocatedProperties, EPCGPointNativeProperties::MetadataEntry))
 		{
 			MetadataEntryRange = InBasePointData->GetMetadataEntryValueRange(bAllocate);
 		}
