[PCG] Preview actors now have an external package (in WorldPartition worlds) so they can support data layers
- Changing the transient state now doesn't need to set an external package it just needs to change the transient flag and dirtyness of the actor package
#rnx

diff --git a/Engine/Plugins/PCG/Source/PCG/Private/Elements/PCGCreateTargetActor.cpp b/Engine/Plugins/PCG/Source/PCG/Private/Elements/PCGCreateTargetActor.cpp
--- a/Engine/Plugins/PCG/Source/PCG/Private/Elements/PCGCreateTargetActor.cpp
+++ b/Engine/Plugins/PCG/Source/PCG/Private/Elements/PCGCreateTargetActor.cpp
@@ -380,11 +380,6 @@ bool FPCGCreateTargetActorElement::ExecuteInternal(FPCGContext* Context) const
 	SpawnParams.SpawnCollisionHandlingOverride = ESpawnActorCollisionHandlingMethod::AlwaysSpawn;
 	SpawnParams.OverrideLevel = TargetActor->GetLevel();
 
-	if (PCGHelpers::IsRuntimeOrPIE() || (SourceComponent && SourceComponent->IsInPreviewMode()))
-	{
-		SpawnParams.ObjectFlags |= RF_Transient | RF_NonPIEDuplicateTransient;
-	}
-
 	FTransform Transform = TargetActor->GetTransform();
 	if(Context->IsValueOverriden(GET_MEMBER_NAME_CHECKED(UPCGCreateTargetActor, ActorPivot)))
 	{
@@ -394,6 +389,8 @@ bool FPCGCreateTargetActorElement::ExecuteInternal(FPCGContext* Context) const
 	UPCGActorHelpers::FSpawnDefaultActorParams SpawnDefaultActorParams(TargetActor->GetWorld(), Settings->TemplateActorClass, Transform, SpawnParams);
 	
 #if WITH_EDITOR
+	SpawnDefaultActorParams.bIsPreviewActor = (SourceComponent && SourceComponent->IsInPreviewMode());
+
 	int32 DataLayerCrc;
 	SpawnDefaultActorParams.DataLayerInstances = PCGDataLayerHelpers::GetDataLayerInstancesAndCrc(Context, Settings->DataLayerSettings, TargetActor, DataLayerCrc);
 
@@ -443,13 +440,20 @@ bool FPCGCreateTargetActorElement::ExecuteInternal(FPCGContext* Context) const
 
 			if (SourceComponent)
 			{
-				SourceComponent->ForEachManagedResource([&ReusableResource, ResourceCrc, &Context](UPCGManagedResource* InResource)
+				SourceComponent->ForEachManagedResource([&ReusableResource, ResourceCrc, &Context, bIsPreview = SpawnDefaultActorParams.bIsPreviewActor](UPCGManagedResource* InResource)
 				{
 					if (ReusableResource)
 					{
 						return;
 					}
 
+#if WITH_EDITOR
+					if (InResource->IsPreview() != bIsPreview)
+					{
+						return;
+					}
+#endif
+
 					if (UPCGManagedActors* Resource = Cast<UPCGManagedActors>(InResource))
 					{
 						if (Resource->GetCrc().IsValid() && Resource->GetCrc() == ResourceCrc && Resource->GetConstGeneratedActors().Num() == 1 && Resource->GetConstGeneratedActors()[0].IsValid())
diff --git a/Engine/Plugins/PCG/Source/PCG/Private/Elements/PCGSpawnActor.cpp b/Engine/Plugins/PCG/Source/PCG/Private/Elements/PCGSpawnActor.cpp
--- a/Engine/Plugins/PCG/Source/PCG/Private/Elements/PCGSpawnActor.cpp
+++ b/Engine/Plugins/PCG/Source/PCG/Private/Elements/PCGSpawnActor.cpp
@@ -849,15 +849,13 @@ void FPCGSpawnActorElement::SpawnActors(FPCGSubgraphContext* Context, AActor* Ta
 	SpawnParams.OverrideLevel = TargetActor->GetLevel();
 
 	UPCGComponent* SourceComponent = Cast<UPCGComponent>(Context->ExecutionSource.Get());
-	if (PCGHelpers::IsRuntimeOrPIE() || (SourceComponent && SourceComponent->IsInPreviewMode()))
-	{
-		SpawnParams.ObjectFlags |= RF_Transient | RF_NonPIEDuplicateTransient;
-	}
 
 	UPCGActorHelpers::FSpawnDefaultActorParams SpawnDefaultActorParams(TargetActor->GetWorld(), InTemplateActorClass, FTransform::Identity, SpawnParams);
 	SpawnDefaultActorParams.bForceStaticMobility = false; // Always respect the actor's mobility
 
 #if WITH_EDITOR
+	SpawnDefaultActorParams.bIsPreviewActor = (SourceComponent && SourceComponent->IsInPreviewMode());
+
 	int32 DataLayerCrc;
 	SpawnDefaultActorParams.DataLayerInstances = PCGDataLayerHelpers::GetDataLayerInstancesAndCrc(Context, Settings->DataLayerSettings, TargetActor, DataLayerCrc);
 
@@ -901,7 +899,7 @@ void FPCGSpawnActorElement::SpawnActors(FPCGSubgraphContext* Context, AActor* Ta
 
 		if (InputDependenciesCrc.IsValid())
 		{
-			SourceComponent->ForEachManagedResource([&ReusedManagedActorsResource, &InputDependenciesCrc, &Context, NumPoints = PointData->GetNumPoints()](UPCGManagedResource* InResource)
+			SourceComponent->ForEachManagedResource([&ReusedManagedActorsResource, &InputDependenciesCrc, &Context, NumPoints = PointData->GetNumPoints(), bIsPreview = SpawnDefaultActorParams.bIsPreviewActor](UPCGManagedResource* InResource)
 			{
 				if (ReusedManagedActorsResource)
 				{
@@ -910,6 +908,13 @@ void FPCGSpawnActorElement::SpawnActors(FPCGSubgraphContext* Context, AActor* Ta
 
 				if (UPCGManagedActors* Resource = Cast<UPCGManagedActors>(InResource))
 				{
+#if WITH_EDITOR
+					if (Resource->IsPreview() != bIsPreview)
+					{
+						return;
+					}
+#endif
+
 					// We can only re-use the resource if it matches the number of points (if actor failed to spawned for whatever reason, we won't know which point is associated with the fail)
 					if (Resource->GetCrc().IsValid() && Resource->GetCrc() == InputDependenciesCrc && Resource->GetConstGeneratedActors().Num() == NumPoints)
 					{
diff --git a/Engine/Plugins/PCG/Source/PCG/Private/Helpers/PCGActorHelpers.cpp b/Engine/Plugins/PCG/Source/PCG/Private/Helpers/PCGActorHelpers.cpp
--- a/Engine/Plugins/PCG/Source/PCG/Private/Helpers/PCGActorHelpers.cpp
+++ b/Engine/Plugins/PCG/Source/PCG/Private/Helpers/PCGActorHelpers.cpp
@@ -23,6 +23,7 @@
 #include "Editor.h"
 #include "Engine/Level.h"
 #include "Editor/IPCGEditorModule.h"
+#include "Engine/Blueprint.h"
 #include "Modules/ModuleManager.h"
 #include "Subsystems/ActorEditorContextSubsystem.h"
 #include "WorldPartition/DataLayer/DataLayerInstance.h"
@@ -689,12 +690,43 @@ AActor* UPCGActorHelpers::SpawnDefaultActor(const FSpawnDefaultActorParams& Para
 		SpawnParams.OverrideLevel = Params.World->PersistentLevel;
 	}
 
-	if (PCGHelpers::IsRuntimeOrPIE())
+	if (PCGHelpers::IsRuntimeOrPIE() || Params.bIsPreviewActor)
 	{
 		SpawnParams.ObjectFlags |= RF_Transient | RF_NonPIEDuplicateTransient;
 	}
 
 #if WITH_EDITOR
+	// Capture label before potentially changing it for previewing
+	FString Label = SpawnParams.Name.IsNone() ? FString() : SpawnParams.Name.ToString();
+
+	// If we are spawning a preview actor in a world partition world, we need to assign it an external package as Data Layers are not supported on non OFPA actors.
+	// Note: this could at somepoint become an option in the FActorSpawnParameters to allow RF_Transient actors to create their own package in UWorld::SpawnActor.
+	if (Params.bIsPreviewActor && SpawnParams.OverrideLevel->GetWorldPartition())
+	{
+		// Remove trailing _C from BP generated class
+		FString ActorClassName = Params.ActorClass->GetName();
+		if (Cast<UBlueprint>(Params.ActorClass->ClassGeneratedBy))
+		{
+			ActorClassName.RemoveFromEnd(TEXT("_C"), ESearchCase::CaseSensitive);
+		}
+
+		SpawnParams.Name = FActorSpawnUtils::MakeUniqueActorName(SpawnParams.OverrideLevel, Params.ActorClass, *ActorClassName, /*bGloballyUniqueName*/true);
+		
+		// If Label wasn't empty add suffix to it, if not use basename as label
+		if (Label.IsEmpty())
+		{
+			Label = ActorClassName;
+		}
+		
+		SpawnParams.OverridePackage = UPCGActorHelpers::CreatePreviewPackage(SpawnParams.OverrideLevel, SpawnParams.Name.ToString());
+		if (SpawnParams.OverridePackage)
+		{
+			SpawnParams.bCreateActorPackage = false;
+		}
+	}
+
+	SpawnParams.InitialActorLabel = Label;
+
 	// Find External Data Layer if it exists so we can create scope that will allow actor to be properly created
 	const UExternalDataLayerInstance* ExternalDataLayerInstance = nullptr;
 	bool bPushedContext = false;
@@ -730,7 +762,7 @@ AActor* UPCGActorHelpers::SpawnDefaultActor(const FSpawnDefaultActorParams& Para
 	{
 		EditorModule->SetOutlinerUIRefreshDelay(UE::PCGActorHelpers::Local::OutlinerUIRefreshDelay);
 	}
-#endif
+#endif // WITH_EDITOR
 
 	AActor* NewActor = Params.World->SpawnActor(*Params.ActorClass, &Params.Transform, SpawnParams);
 	
@@ -756,11 +788,6 @@ AActor* UPCGActorHelpers::SpawnDefaultActor(const FSpawnDefaultActorParams& Para
 			}
 		}
 	}
-	
-	if (SpawnParams.Name != NAME_None)
-	{
-		NewActor->SetActorLabel(SpawnParams.Name.ToString());
-	}
 #endif // WITH_EDITOR
 
 	USceneComponent* RootComponent = NewActor->GetRootComponent();
@@ -792,6 +819,26 @@ AActor* UPCGActorHelpers::SpawnDefaultActor(const FSpawnDefaultActorParams& Para
 	return NewActor;
 }
 
+#if WITH_EDITOR
+UPackage* UPCGActorHelpers::CreatePreviewPackage(ULevel* InLevel, const FString& InActorName)
+{
+	check(InLevel);
+	TStringBuilderWithBuffer<TCHAR, NAME_SIZE> ActorPath;
+	ActorPath += InLevel->GetPathName();
+	ActorPath += TEXT(".Preview_");
+	ActorPath += InActorName;
+
+	UPackage* ActorPackage = ULevel::CreateActorPackage(InLevel->GetPackage(), InLevel->GetActorPackagingScheme(), ActorPath.ToString());
+	if (ActorPackage)
+	{
+		// Set dirty flag to false before adding transient flag so that UX updates properly
+		ActorPackage->SetDirtyFlag(false);
+		ActorPackage->SetFlags(RF_Transient);
+	}
+	return ActorPackage;
+}
+#endif // WITH_EDITOR
+
 FIntVector UPCGActorHelpers::GetCellCoord(FVector InPosition, int InGridSize, bool bUse2DGrid)
 {
 	check(InGridSize > 0);
diff --git a/Engine/Plugins/PCG/Source/PCG/Private/PCGManagedResource.cpp b/Engine/Plugins/PCG/Source/PCG/Private/PCGManagedResource.cpp
--- a/Engine/Plugins/PCG/Source/PCG/Private/PCGManagedResource.cpp
+++ b/Engine/Plugins/PCG/Source/PCG/Private/PCGManagedResource.cpp
@@ -4,6 +4,7 @@
 
 #include "PCGComponent.h"
 #include "PCGModule.h"
+#include "Helpers/PCGActorHelpers.h"
 #include "Helpers/PCGHelpers.h"
 
 #include "Algo/AnyOf.h"
@@ -61,10 +62,12 @@ bool UPCGManagedResource::DebugForcePurgeAllResourcesOnGenerate()
 }
 
 #if WITH_EDITOR
-void UPCGManagedResource::ChangeTransientState(EPCGEditorDirtyMode /*NewEditingMode*/)
+void UPCGManagedResource::ChangeTransientState(EPCGEditorDirtyMode NewEditingMode)
 {
 	// Any change in the transient state resets the transient state that was set on load, regardless of the bNowTransient flag
 	bMarkedTransientOnLoad = false;
+
+	bIsPreview = (NewEditingMode == EPCGEditorDirtyMode::Preview);
 }
 #endif // WITH_EDITOR
 
@@ -119,8 +122,8 @@ bool UPCGManagedActors::Release(bool bHardRelease, TSet<TSoftObjectPtr<AActor>>&
 		// Here, instead of adding the actors to be deleted (which has the side effect of potentially emptying the package, which leads to its deletion, we will hide the actors instead
 		for (TSoftObjectPtr<AActor> GeneratedActor : GeneratedActorsArray)
 		{
-			// In rare cases where the actor wouldn't be currently loaded (in WP) we must make sure it is done here
-			if (AActor* Actor = GeneratedActor.LoadSynchronous())
+			// Hide actor if it is loaded
+			if (AActor* Actor = GeneratedActor.Get())
 			{
 				Actor->SetIsTemporarilyHiddenInEditor(true);
 				Actor->SetHidden(true);
@@ -237,22 +240,46 @@ void UPCGManagedActors::ChangeTransientState(EPCGEditorDirtyMode NewEditingMode)
 {
 	const bool bNowTransient = (NewEditingMode == EPCGEditorDirtyMode::Preview);
 
-	for (TSoftObjectPtr<AActor> GeneratedActor : GeneratedActorsArray)
+	for (TSoftObjectPtr<AActor> GeneratedActorPtr : GeneratedActorsArray)
 	{
 		// Make sure to load if needed because we need to affect the actors regardless of the current WP state
-		if (GeneratedActor.LoadSynchronous() != nullptr)
+		if (AActor* GeneratedActor = GeneratedActorPtr.Get())
 		{
 			const bool bWasTransient = GeneratedActor->HasAnyFlags(RF_Transient);
 
 			if (bNowTransient != bWasTransient)
 			{
+				UPackage* CurrentPackage = GeneratedActor->GetExternalPackage();
 				if (bNowTransient)
 				{
 					GeneratedActor->SetFlags(RF_Transient);
+					if (CurrentPackage)
+					{
+						CurrentPackage->SetDirtyFlag(true);
+
+						// Disable external packaging first because AActor::SetPackageExternal will early out if package is already external.
+						// In this case we want to change the external package so we remove the previous one first.
+						GeneratedActor->SetPackageExternal(/*bExternal=*/false, /*bShouldDirty=*/false);
+
+						UPackage* PreviewPackage = UPCGActorHelpers::CreatePreviewPackage(GeneratedActor->GetLevel(), GeneratedActor->GetName());
+						ensure(PreviewPackage);
+
+						// Use the preview package.
+						GeneratedActor->SetPackageExternal(/*bExternal=*/true, /*bShouldDirty=*/false, PreviewPackage);
+					}
 				}
 				else
 				{
 					GeneratedActor->ClearFlags(RF_Transient);
+					if (CurrentPackage)
+					{
+						// Disable external packaging first because AActor::SetPackageExternal will early out if package is already external.
+						// In this case we want to change the external package so we remove the previous one first.
+						GeneratedActor->SetPackageExternal(/*bExternal=*/false, /*bShouldDirty=*/false);
+
+						// Use the default external package for this actor.
+						GeneratedActor->SetPackageExternal(/*bExternal=*/true, /*bShouldDirty=*/false);
+					}
 				}
 			}
 
@@ -270,26 +297,9 @@ void UPCGManagedActors::ChangeTransientState(EPCGEditorDirtyMode NewEditingMode)
 
 			if (bNowTransient != bWasTransient)
 			{
-				if (GeneratedActor->GetLevel() && GeneratedActor->GetLevel()->IsUsingExternalActors())
-				{
-					UPackage* PreviousExternalPackage = bNowTransient ? GeneratedActor->GetExternalPackage() : nullptr;
-
-					GeneratedActor->SetPackageExternal(/*bExternal=*/!bNowTransient, /*bShouldDirty=*/false);
-
-					if (PreviousExternalPackage)
-					{
-						PreviousExternalPackage->MarkPackageDirty();
-					}
-				
-					if (!bNowTransient && GeneratedActor->GetExternalPackage())
-					{
-						MarkPackageDirty();
-					}
-				}
-
-				if (!bNowTransient || GeneratedActor->IsPackageExternal())
+				if (!bNowTransient)
 				{
-					ForEachObjectWithOuter(GeneratedActor.Get(), [bNowTransient](UObject* Object)
+					ForEachObjectWithOuter(GeneratedActor, [bNowTransient](UObject* Object)
 					{
 						if (bNowTransient)
 						{
diff --git a/Engine/Plugins/PCG/Source/PCG/Public/Helpers/PCGActorHelpers.h b/Engine/Plugins/PCG/Source/PCG/Public/Helpers/PCGActorHelpers.h
--- a/Engine/Plugins/PCG/Source/PCG/Public/Helpers/PCGActorHelpers.h
+++ b/Engine/Plugins/PCG/Source/PCG/Public/Helpers/PCGActorHelpers.h
@@ -215,6 +215,7 @@ PRAGMA_ENABLE_DEPRECATION_WARNINGS
 		FActorSpawnParameters SpawnParams;
 		AActor* Parent = nullptr;
 		bool bForceStaticMobility = true;
+		bool bIsPreviewActor = false;
 #if WITH_EDITOR
 		UHLODLayer* HLODLayer = nullptr;
 		TArray<const UDataLayerInstance*> DataLayerInstances;
@@ -239,6 +240,10 @@ PRAGMA_ENABLE_DEPRECATION_WARNINGS
 
 	/** Extract the tags and the actor reference of the given actor and hash it. Useful for CRC dependencies that depends on the tags or the instance of the actor.*/
 	static UE_API int ComputeHashFromActorTagsAndReference(const AActor* InActor, const bool bIncludeTags, const bool bIncludeActorReference);
+
+#if WITH_EDITOR
+	static UPackage* CreatePreviewPackage(ULevel* InLevel, const FString& InActorName);
+#endif
 };
 
 #undef UE_API
diff --git a/Engine/Plugins/PCG/Source/PCG/Public/PCGManagedResource.h b/Engine/Plugins/PCG/Source/PCG/Public/PCGManagedResource.h
--- a/Engine/Plugins/PCG/Source/PCG/Public/PCGManagedResource.h
+++ b/Engine/Plugins/PCG/Source/PCG/Public/PCGManagedResource.h
@@ -65,6 +65,9 @@ public:
 	virtual void MarkTransientOnLoad() { bMarkedTransientOnLoad = true; }
 
 	bool IsMarkedTransientOnLoad() const { return bMarkedTransientOnLoad; }
+
+	bool IsPreview() const { return bIsPreview; }
+	void SetIsPreview(bool bInIsPreview) { bIsPreview = bInIsPreview; }
 #endif // WITH_EDITOR
 
 protected:
@@ -81,6 +84,9 @@ protected:
 	// Note that this flag will be reset if there is a transient state change originating from the component, which might trigger resource deletion, flags change, etc.
 	UPROPERTY(VisibleAnywhere, Category = GeneratedData, meta = (NoResetToDefault))
 	bool bMarkedTransientOnLoad = false;
+
+	UPROPERTY(Transient)
+	bool bIsPreview = false;
 #endif
 };
 
