[PCG] Adjusted budget cvars and enabled early transient resources cleanup by default

diff --git a/Engine/Plugins/PCG/Source/PCG/Private/Graph/PCGGraphExecutor.cpp b/Engine/Plugins/PCG/Source/PCG/Private/Graph/PCGGraphExecutor.cpp
--- a/Engine/Plugins/PCG/Source/PCG/Private/Graph/PCGGraphExecutor.cpp
+++ b/Engine/Plugins/PCG/Source/PCG/Private/Graph/PCGGraphExecutor.cpp
@@ -61,7 +61,7 @@ namespace PCGGraphExecutor
 {
 	TAutoConsoleVariable<float> CVarTimePerFrame(
 		TEXT("pcg.FrameTime"),
-		1000.0f / 60.0f,
+		5.0f,
 		TEXT("Allocated time in ms per frame"));
 
 	TAutoConsoleVariable<bool> CVarGraphMultithreading(
@@ -72,7 +72,7 @@ namespace PCGGraphExecutor
 #if WITH_EDITOR
 	TAutoConsoleVariable<float> CVarEditorTimePerFrame(
 		TEXT("pcg.EditorFrameTime"),
-		1000.0f / 20.0f,
+		15.0f,
 		TEXT("Allocated time in ms per frame when running in editor (non pie)"));
 
 	TAutoConsoleVariable<float> CVarEditorNotificationDelayInSeconds(
diff --git a/Engine/Plugins/PCG/Source/PCG/Private/PCGCommon.cpp b/Engine/Plugins/PCG/Source/PCG/Private/PCGCommon.cpp
--- a/Engine/Plugins/PCG/Source/PCG/Private/PCGCommon.cpp
+++ b/Engine/Plugins/PCG/Source/PCG/Private/PCGCommon.cpp
@@ -104,8 +104,8 @@ namespace PCGSystemSwitches
 #endif
 
 	TAutoConsoleVariable<bool> CVarReleaseTransientResourcesEarly(
-		TEXT("pcg.Experimental.ReleaseTransientResourcesEarly"),
-		false,
+		TEXT("pcg.ReleaseTransientResourcesEarly"),
+		true,
 		TEXT("Aggressively release transient (esp. GPU) data when possible."));
 
 #if UE_ENABLE_DEBUG_DRAWING 
