[PCG] remove ensure that could fire with older data (5.2 allowed ActorSelectionClass to be edited by user)

#rnx

diff --git a/Engine/Plugins/PCG/Source/PCG/Private/Elements/PCGTypedGetter.cpp b/Engine/Plugins/PCG/Source/PCG/Private/Elements/PCGTypedGetter.cpp
--- a/Engine/Plugins/PCG/Source/PCG/Private/Elements/PCGTypedGetter.cpp
+++ b/Engine/Plugins/PCG/Source/PCG/Private/Elements/PCGTypedGetter.cpp
@@ -102,13 +102,13 @@ void FPCGGetLandscapeDataElement::ProcessActors(FPCGContext* Context, const UPCG
 
 	for (AActor* FoundActor : FoundActors)
 	{
-		if (!FoundActor || !IsValid(FoundActor))
+		if (!IsValid(FoundActor))
 		{
 			continue;
 		}
 
 		ALandscapeProxy* Landscape = Cast<ALandscapeProxy>(FoundActor);
-		if (ensure(Landscape))
+		if (Landscape)
 		{
 			Landscapes.Add(Landscape);
 			LandscapeBounds += PCGHelpers::GetGridBounds(Landscape, nullptr);
