[PCG] FGridGenerationKey fix
- cache bUse2DGrid instead of resolving OriginalComponent which might be stale
- use FObjectKey for equality comparison / hash

#rnx

diff --git a/Engine/Plugins/PCG/Source/PCG/Private/RuntimeGen/PCGRuntimeGenScheduler.cpp b/Engine/Plugins/PCG/Source/PCG/Private/RuntimeGen/PCGRuntimeGenScheduler.cpp
--- a/Engine/Plugins/PCG/Source/PCG/Private/RuntimeGen/PCGRuntimeGenScheduler.cpp
+++ b/Engine/Plugins/PCG/Source/PCG/Private/RuntimeGen/PCGRuntimeGenScheduler.cpp
@@ -112,7 +112,7 @@ namespace PCGRuntimeGenSchedulerHelpers
 
 FPCGGridDescriptor FPCGRuntimeGenScheduler::FGridGenerationKey::GetGridDescriptor() const
 {
-	return FPCGGridDescriptor().SetGridSize(GetGridSize()).SetIsRuntime(true).SetIs2DGrid(GetOriginalComponent()->Use2DGrid());
+	return FPCGGridDescriptor().SetGridSize(GetGridSize()).SetIsRuntime(true).SetIs2DGrid(Use2DGrid());
 }
 
 FPCGRuntimeGenScheduler::FPCGRuntimeGenScheduler(UWorld* InWorld, FPCGActorAndComponentMapping* InActorAndComponentMapping)
diff --git a/Engine/Plugins/PCG/Source/PCG/Public/RuntimeGen/PCGRuntimeGenScheduler.h b/Engine/Plugins/PCG/Source/PCG/Public/RuntimeGen/PCGRuntimeGenScheduler.h
--- a/Engine/Plugins/PCG/Source/PCG/Public/RuntimeGen/PCGRuntimeGenScheduler.h
+++ b/Engine/Plugins/PCG/Source/PCG/Public/RuntimeGen/PCGRuntimeGenScheduler.h
@@ -79,47 +79,50 @@ protected:
 	struct FGridGenerationKey
 	{
 		FGridGenerationKey(uint32 InGridSize, const FIntVector& InGridCoords, UPCGComponent* InOriginalComponent)
-			: GridSize(InGridSize)
-			, GridCoords(InGridCoords)
-			, OriginalComponent(InOriginalComponent)
-			, CachedLocalComponent(nullptr)
-		{}
+			: FGridGenerationKey(InGridSize, InGridCoords, InOriginalComponent, nullptr)
+		{
+		}
 
 		FGridGenerationKey(uint32 InGridSize, const FIntVector& InGridCoords, UPCGComponent* InOriginalComponent, UPCGComponent* InLocalComponent)
 			: GridSize(InGridSize)
 			, GridCoords(InGridCoords)
 			, OriginalComponent(InOriginalComponent)
 			, CachedLocalComponent(InLocalComponent)
-		{}
+		{
+			bUse2DGrid = InOriginalComponent ? InOriginalComponent->Use2DGrid() : true;
+		}
 
 		bool operator==(const FGridGenerationKey& Other) const
 		{
 			return GridSize == Other.GridSize
 				&& GridCoords == Other.GridCoords
+				&& bUse2DGrid == Other.bUse2DGrid
 				&& OriginalComponent == Other.OriginalComponent;
 		}
 
-		bool IsValid() const { return OriginalComponent.IsValid(); }
+		bool IsValid() const { return !!OriginalComponent.ResolveObjectPtr(); }
 
+		bool Use2DGrid() const { return bUse2DGrid; }
 		uint32 GetGridSize() const { return GridSize; }
 		FIntVector GetGridCoords() const { return GridCoords; }
-		UPCGComponent* GetOriginalComponent() const { return OriginalComponent.Get(); }
+		UPCGComponent* GetOriginalComponent() const { return OriginalComponent.ResolveObjectPtr(); }
 
 		UPCGComponent* GetCachedLocalComponent() const { return CachedLocalComponent.Get(); }
 		void SetCachedLocalComponent(UPCGComponent* InComponent) { CachedLocalComponent = InComponent; }
 
 		FPCGGridDescriptor GetGridDescriptor() const;
 
+		bool bUse2DGrid = false;
 		uint32 GridSize = 0;
 		FIntVector GridCoords = FIntVector(0);
-		TWeakObjectPtr<UPCGComponent> OriginalComponent;
+		TObjectKey<UPCGComponent> OriginalComponent;
 
 		// Optional/opportunistic cached local component if one has been created.
 		TWeakObjectPtr<UPCGComponent> CachedLocalComponent;
 
 		friend uint32 GetTypeHash(const FGridGenerationKey& In)
 		{
-			return HashCombine(GetTypeHash(In.GridCoords), GetTypeHash(In.GridSize), GetTypeHash(In.OriginalComponent));
+			return HashCombine(GetTypeHash(In.GridCoords), GetTypeHash(In.GridSize), GetTypeHash(In.bUse2DGrid), GetTypeHash(In.OriginalComponent));
 		}
 	};
 
