[PCG] Support GPU texture outputs in data viewport

#virtualized

diff --git a/Engine/Plugins/PCG/Source/PCGEditor/PCGEditor.Build.cs b/Engine/Plugins/PCG/Source/PCGEditor/PCGEditor.Build.cs
--- a/Engine/Plugins/PCG/Source/PCGEditor/PCGEditor.Build.cs
+++ b/Engine/Plugins/PCG/Source/PCGEditor/PCGEditor.Build.cs
@@ -54,6 +54,7 @@ namespace UnrealBuildTool.Rules
 					"KismetWidgets",
 					"PCG",
 					"PropertyEditor",
+					"RHI",
 					"RenderCore",
 					"Slate",
 					"SlateCore",
diff --git a/Engine/Plugins/PCG/Source/PCGEditor/Private/DataVisualizations/PCGBaseTextureDataVisualization.cpp b/Engine/Plugins/PCG/Source/PCGEditor/Private/DataVisualizations/PCGBaseTextureDataVisualization.cpp
--- a/Engine/Plugins/PCG/Source/PCGEditor/Private/DataVisualizations/PCGBaseTextureDataVisualization.cpp
+++ b/Engine/Plugins/PCG/Source/PCGEditor/Private/DataVisualizations/PCGBaseTextureDataVisualization.cpp
@@ -2,6 +2,8 @@
 
 #include "DataVisualizations/PCGBaseTextureDataVisualization.h"
 
+#include "DataVisualizations/PCGVisualizationTexture2D.h"
+
 #include "Data/PCGTextureData.h"
 
 #include "AdvancedPreviewScene.h"
@@ -42,15 +44,39 @@ FPCGSetupSceneFunc FPCGBaseTextureDataVisualization::GetViewportSetupFunc(const
 			return;
 		}
 
-		if (WeakData->GetTextureResourceType() != EPCGTextureResourceType::TextureObject)
+		TObjectPtr<UTexture> Texture = nullptr;
+		bool bIsTextureArray = false;
+
+		if (WeakData->GetTextureResourceType() == EPCGTextureResourceType::TextureObject)
+		{
+			Texture = WeakData->GetTexture();
+			bIsTextureArray = WeakData->GetTextureRHI() && WeakData->GetTextureRHI()->GetDesc().Dimension == ETextureDimension::Texture2DArray;
+		}
+		else if (WeakData->GetTextureResourceType() == EPCGTextureResourceType::ExportedTexture)
 		{
-			// @todo_pcg: Need to create a dynamic UTexture asset from an exported buffer handle in order to support EPCGTextureResourceType::ExportedTexture
-			UE_LOG(LogPCG, Verbose, TEXT("Texture data uses an unsupported resource type for data viewport visualization."));
+			Texture = UPCGVisualizationTexture2D::Create(WeakData);
+			OutManagedResources.Add(Texture);
+		}
+		else
+		{
+			UE_LOG(LogPCG, Error, TEXT("Texture data uses an unsupported resource type for data viewport visualization."));
+			return;
 		}
 
 		UMaterialInterface* DebugMaterial = Cast<UMaterialInterface>(InResources[1]);
 		UMaterialInstanceDynamic* MaterialInstance = UMaterialInstanceDynamic::Create(DebugMaterial, GetTransientPackage());
-		MaterialInstance->SetTextureParameterValue(FName(TEXT("DebugTexture")), WeakData->GetTexture());
+
+		if (bIsTextureArray)
+		{
+			MaterialInstance->SetTextureParameterValue(FName(TEXT("DebugTextureArray")), Texture);
+			MaterialInstance->SetScalarParameterValue(FName(TEXT("SliceIndex")), WeakData->GetTextureSlice());
+			MaterialInstance->SetScalarParameterValue(FName(TEXT("UseTextureArray")), 1.0f);
+		}
+		else
+		{
+			MaterialInstance->SetTextureParameterValue(FName(TEXT("DebugTexture")), Texture);
+		}
+
 		OutManagedResources.Add(MaterialInstance);
 
 		TObjectPtr<UStaticMeshComponent> MeshComponent = NewObject<UStaticMeshComponent>(GetTransientPackage(), NAME_None, RF_Transient);
diff --git a/Engine/Plugins/PCG/Source/PCGEditor/Private/DataVisualizations/PCGVisualizationTexture2D.cpp b/Engine/Plugins/PCG/Source/PCGEditor/Private/DataVisualizations/PCGVisualizationTexture2D.cpp
new file mode 100644
--- /dev/null
+++ b/Engine/Plugins/PCG/Source/PCGEditor/Private/DataVisualizations/PCGVisualizationTexture2D.cpp
@@ -0,0 +1,156 @@
+// Copyright Epic Games, Inc. All Rights Reserved.
+
+#include "DataVisualizations/PCGVisualizationTexture2D.h"
+
+#include "Data/PCGTextureData.h"
+
+#include "RHIStaticStates.h"
+#include "RenderCaptureInterface.h"
+#include "TextureResource.h"
+
+#include UE_INLINE_GENERATED_CPP_BY_NAME(PCGVisualizationTexture2D)
+
+namespace PCGVisTexture2DHelpers
+{
+	static int32 GTriggerGPUCaptureDispatches = 0;
+	static FAutoConsoleVariableRef CVarTriggerGPUCapture(
+		TEXT("pcg.RenderCaptures.TriggerVisualizationTextureResourceInitCaptures"),
+		GTriggerGPUCaptureDispatches,
+		TEXT("Trigger GPU captures for this many of the subsequent visualization texture resource initializations."));
+
+	/** Creates a black 1x1 dummy texture. */
+	FTextureRHIRef CreateDummyTexture(FRHICommandListImmediate& RHICmdList)
+	{
+		const FRHITextureCreateDesc DummyTextureDesc =
+			FRHITextureCreateDesc::Create2D(TEXT("PCGDummyTexture"), 1, 1, PF_G8)
+			.SetFlags(ETextureCreateFlags::ShaderResource);
+
+		FTextureRHIRef DummyTexture = RHICreateTexture(DummyTextureDesc);
+
+		{
+			uint32 DestStride;
+			uint8* DestBuffer = (uint8*)RHICmdList.LockTexture2D(DummyTexture, 0, RLM_WriteOnly, DestStride, false);
+			*DestBuffer = 0;
+			RHICmdList.UnlockTexture2D(DummyTexture, 0, false);
+		}
+
+		return DummyTexture;
+	}
+}
+
+/**
+ * RT Resource for UPCGVisualizationTexture2D. Wraps and manages the texture/sampler RHI for the UTexture.
+ * Created from UPCGVisualizationTexture2D::CreateResource(), but otherwise managed in the UTexture base class.
+ */
+class FPCGVisualizationTexture2DResource : public FTextureResource
+{
+public:
+	FPCGVisualizationTexture2DResource(UPCGVisualizationTexture2D* InOwner)
+		: Owner(InOwner)
+	{}
+
+	virtual uint32 GetSizeX() const override { return ensure(Owner) ? Owner->GetSurfaceWidth() : 1u; }
+	virtual uint32 GetSizeY() const override { return ensure(Owner) ? Owner->GetSurfaceHeight() : 1u; }
+
+	/** Called when the resource is initialized. This is only called by the rendering thread. */
+	virtual void InitRHI(FRHICommandListBase& RHICmdList) override;
+
+	/** Called when the resource is released. This is only called by the rendering thread. */
+	virtual void ReleaseRHI() override;
+
+protected:
+	/** The texture owning this resource. */
+	UPCGVisualizationTexture2D* Owner;
+};
+
+void FPCGVisualizationTexture2DResource::InitRHI(FRHICommandListBase& RHICmdListBase)
+{
+	if (!Owner)
+	{
+		return;
+	}
+
+	FRHICommandListImmediate& RHICmdList = RHICmdListBase.GetAsImmediate();
+
+	const UPCGBaseTextureData* TextureData = Owner->GetTextureData();
+	FTextureRHIRef SourceTextureRHI = TextureData ? TextureData->GetTextureRHI() : nullptr;
+
+	if (TextureData && SourceTextureRHI.IsValid())
+	{
+		RenderCaptureInterface::FScopedCapture RenderCapture(PCGVisTexture2DHelpers::GTriggerGPUCaptureDispatches > 0, &RHICmdList, TEXT("FPCGVisualizationTextureResource::InitRHI::CopySource"));
+		PCGVisTexture2DHelpers::GTriggerGPUCaptureDispatches = FMath::Max(PCGVisTexture2DHelpers::GTriggerGPUCaptureDispatches - 1, 0);
+
+		// Create TextureRHI to match our SourceTextureRHI
+		const FRHITextureDesc& SourceDesc = SourceTextureRHI->GetDesc();
+
+		const FRHITextureCreateDesc CreateDesc =
+			FRHITextureCreateDesc::Create2D(TEXT("PCGVisualizationTexture2D"))
+			.SetExtent(SourceDesc.Extent)
+			.SetFormat(SourceDesc.Format)
+			.SetFlags(SourceDesc.Flags)
+			.SetInitialState(ERHIAccess::CopyDest);
+
+		TextureRHI = RHICreateTexture(CreateDesc);
+
+		// Prepare to copy the source texture.
+		RHICmdList.Transition(FRHITransitionInfo(SourceTextureRHI, ERHIAccess::Unknown, ERHIAccess::CopySrc));
+
+		// Copy the source texture.
+		FRHICopyTextureInfo CopyInfo;
+		CopyInfo.SourceSliceIndex = TextureData->GetTextureSlice();
+		RHICmdList.CopyTexture(SourceTextureRHI, TextureRHI, CopyInfo);
+
+		RHICmdList.Transition(FRHITransitionInfo(TextureRHI, ERHIAccess::CopyDest, ERHIAccess::SRVMask));
+	}
+	else
+	{
+		// If there is no source texture RHI, just use a black dummy texture.
+		TextureRHI = PCGVisTexture2DHelpers::CreateDummyTexture(RHICmdList);
+	}
+
+	SamplerStateRHI = TStaticSamplerState<SF_Bilinear, AM_Clamp, AM_Clamp, AM_Clamp>::GetRHI();
+
+	RHIUpdateTextureReference(Owner->TextureReference.TextureReferenceRHI, TextureRHI);
+}
+
+void FPCGVisualizationTexture2DResource::ReleaseRHI()
+{
+	if (Owner)
+	{
+		RHIUpdateTextureReference(Owner->TextureReference.TextureReferenceRHI, nullptr);
+	}
+
+	FTextureResource::ReleaseRHI();
+	TextureRHI.SafeRelease();
+}
+
+FTextureResource* UPCGVisualizationTexture2D::CreateResource()
+{
+	return new FPCGVisualizationTexture2DResource(this);
+}
+
+float UPCGVisualizationTexture2D::GetSurfaceWidth() const
+{
+	return TextureData.Get() ? TextureData.Get()->GetTextureSize().X : 1u;
+}
+
+float UPCGVisualizationTexture2D::GetSurfaceHeight() const
+{
+	return TextureData.Get() ? TextureData.Get()->GetTextureSize().Y : 1u;
+}
+
+void UPCGVisualizationTexture2D::Init(TWeakObjectPtr<const UPCGBaseTextureData> InTextureData)
+{
+	TextureData = InTextureData;
+
+	UpdateResource();
+}
+
+UPCGVisualizationTexture2D* UPCGVisualizationTexture2D::Create(TWeakObjectPtr<const UPCGBaseTextureData> InTextureData)
+{
+	UPCGVisualizationTexture2D* NewTexture = NewObject<UPCGVisualizationTexture2D>(GetTransientPackage(), NAME_None, RF_Transient);
+
+	NewTexture->Init(InTextureData);
+
+	return NewTexture;
+}
diff --git a/Engine/Plugins/PCG/Source/PCGEditor/Private/DataVisualizations/PCGVisualizationTexture2D.h b/Engine/Plugins/PCG/Source/PCGEditor/Private/DataVisualizations/PCGVisualizationTexture2D.h
new file mode 100644
--- /dev/null
+++ b/Engine/Plugins/PCG/Source/PCGEditor/Private/DataVisualizations/PCGVisualizationTexture2D.h
@@ -0,0 +1,40 @@
+// Copyright Epic Games, Inc. All Rights Reserved.
+
+#pragma once
+
+#include "Engine/Texture.h"
+
+#include "PCGVisualizationTexture2D.generated.h"
+
+class UPCGBaseTextureData;
+
+/**
+ * UTexture implementation for visualizing UPCGBaseTextureData in the data viewport.
+ * This is necessary because some texture data do not have their own UTexture, so one needs
+ * to be created in order to override the texture parameter for the debug vis material.
+ */
+UCLASS(MinimalAPI)
+class UPCGVisualizationTexture2D : public UTexture
+{
+	GENERATED_BODY()
+
+public:
+	//~ Begin UTexture Interface.
+	virtual ETextureClass GetTextureClass() const { return ETextureClass::TwoDDynamic; }
+	virtual FTextureResource* CreateResource() override;
+	virtual EMaterialValueType GetMaterialType() const override { return MCT_Texture2D; }
+	virtual float GetSurfaceWidth() const override;
+	virtual float GetSurfaceHeight() const override;
+	virtual float GetSurfaceDepth() const override { return 0; }
+	virtual uint32 GetSurfaceArraySize() const override { return 0; }
+	//~ End UTexture Interface.
+	
+	void Init(TWeakObjectPtr<const UPCGBaseTextureData> InTextureData);
+	const UPCGBaseTextureData* GetTextureData() const { return TextureData.Get(); }
+
+	/** Creates and initializes a new UPCGVisualizationTexture2D from the texture data. */
+	static UPCGVisualizationTexture2D* Create(TWeakObjectPtr<const UPCGBaseTextureData> InTextureData);
+
+protected:
+	TWeakObjectPtr<const UPCGBaseTextureData> TextureData;
+};
