[PCG] Set group for User Params panel to avoid it showing up multiple times in the file menu

#rnx

diff --git a/Engine/Plugins/PCG/Source/PCGEditor/Private/PCGEditor.cpp b/Engine/Plugins/PCG/Source/PCGEditor/Private/PCGEditor.cpp
--- a/Engine/Plugins/PCG/Source/PCGEditor/Private/PCGEditor.cpp
+++ b/Engine/Plugins/PCG/Source/PCGEditor/Private/PCGEditor.cpp
@@ -652,7 +652,8 @@ void FPCGEditor::RegisterTabSpawners(const TSharedRef<FTabManager>& InTabManager
 		.SetGroup(WorkspaceMenuCategoryRef);
 
 	InTabManager->RegisterTabSpawner(FPCGEditor_private::UserParamsID, FOnSpawnTab::CreateSP(this, &FPCGEditor::SpawnTab_UserParams))
-		.SetDisplayName(FPCGEditor_private::UserParamsTabName);
+		.SetDisplayName(FPCGEditor_private::UserParamsTabName)
+		.SetGroup(WorkspaceMenuCategoryRef);
 
 	InTabManager->RegisterTabSpawner(FPCGEditor_private::ViewportID[0], FOnSpawnTab::CreateSP(this, &FPCGEditor::SpawnTab_Viewport, 0))
 		.SetDisplayName(LOCTEXT("ViewportTab1", "Data Viewport 1"))
