[PCG] Fix default value reset broken with latest changes for graph instances

* Was only working for the first parameter

#rnx

diff --git a/Engine/Plugins/PCG/Source/PCG/Private/PCGGraph.cpp b/Engine/Plugins/PCG/Source/PCG/Private/PCGGraph.cpp
--- a/Engine/Plugins/PCG/Source/PCG/Private/PCGGraph.cpp
+++ b/Engine/Plugins/PCG/Source/PCG/Private/PCGGraph.cpp
@@ -107,7 +107,7 @@ namespace PCGGraphUtils
 			return false;
 		}
 
-		InSourcePropertyDesc->CachedProperty->ExportText_Direct(OutValueString, InSourceInstance.GetValue().GetMemory(), InSourceInstance.GetValue().GetMemory(), nullptr, PPF_None);
+		InSourcePropertyDesc->CachedProperty->ExportText_InContainer(0, OutValueString, InSourceInstance.GetValue().GetMemory(), InSourceInstance.GetValue().GetMemory(), nullptr, PPF_None);
 		return true;
 	}
 #endif // WITH_EDITOR
