[PCG] Add option to disable readback of GPU textures in compute graphs when inspected

#rnx

diff --git a/Engine/Plugins/PCG/Source/PCG/Private/Compute/Elements/PCGComputeGraphElement.cpp b/Engine/Plugins/PCG/Source/PCG/Private/Compute/Elements/PCGComputeGraphElement.cpp
--- a/Engine/Plugins/PCG/Source/PCG/Private/Compute/Elements/PCGComputeGraphElement.cpp
+++ b/Engine/Plugins/PCG/Source/PCG/Private/Compute/Elements/PCGComputeGraphElement.cpp
@@ -26,6 +26,13 @@
 
 #define LOCTEXT_NAMESPACE "PCGComputeGraphElement"
 
+#if WITH_EDITOR
+static TAutoConsoleVariable<bool> CVarReadbackTextureDataOnInspect(
+	TEXT("pcg.GPU.ReadbackTextureDataOnInspect"),
+	true,
+	TEXT("Reads texture data back to the CPU when inspected."));
+#endif
+
 FPCGComputeGraphContext::~FPCGComputeGraphContext()
 {
 	const FPCGStack* StackPtr = GetStack();
@@ -671,7 +678,9 @@ void FPCGComputeGraphElement::CollectDebugDataPrepareActions(FPCGComputeGraphCon
 		return;
 	}
 
-	auto CollectDebugDataPrepareAction = [InContext](FPCGDataToDebug& DataToDebug)
+	const bool bReadbackTextureDataOnInspect = CVarReadbackTextureDataOnInspect.GetValueOnGameThread();
+
+	auto CollectDebugDataPrepareAction = [InContext, bReadbackTextureDataOnInspect](FPCGDataToDebug& DataToDebug, bool bInIsInspecting)
 	{
 		if (const UPCGProxyForGPUData* ProxyData = Cast<UPCGProxyForGPUData>(DataToDebug.Data.Get()))
 		{
@@ -684,8 +693,13 @@ void FPCGComputeGraphElement::CollectDebugDataPrepareActions(FPCGComputeGraphCon
 		}
 		else if (const UPCGTextureData* TextureData = Cast<UPCGTextureData>(DataToDebug.Data.Get()))
 		{
-			InContext->DebugDataPrepareActions.Add([&DataToDebug](FPCGComputeGraphContext* InContext)
+			InContext->DebugDataPrepareActions.Add([&DataToDebug, bInIsInspecting, bReadbackTextureDataOnInspect](FPCGComputeGraphContext* InContext)
 			{
+				if (bInIsInspecting && !bReadbackTextureDataOnInspect)
+				{
+					return true;
+				}
+
 				// Duplicate and CPU-initialize the debug texture data.
 				const UPCGTextureData* TextureData = CastChecked<UPCGTextureData>(DataToDebug.Data.Get());
 				TObjectPtr<UPCGTextureData> DuplicateTextureData = Cast<UPCGTextureData>(DataToDebug.DataPendingInit);
@@ -704,12 +718,12 @@ void FPCGComputeGraphElement::CollectDebugDataPrepareActions(FPCGComputeGraphCon
 
 	for (FPCGDataToDebug& DataToDebug : InContext->DataBinding->DataToDebug)
 	{
-		CollectDebugDataPrepareAction(DataToDebug);
+		CollectDebugDataPrepareAction(DataToDebug, /*bIsInspecting=*/false);
 	}
 
 	for (FPCGDataToDebug& DataToDebug : InContext->DataBinding->DataToInspect)
 	{
-		CollectDebugDataPrepareAction(DataToDebug);
+		CollectDebugDataPrepareAction(DataToDebug, /*bIsInspecting=*/true);
 	}
 }
 
@@ -779,8 +793,10 @@ void FPCGComputeGraphElement::StoreDataForInspection(FPCGComputeGraphContext* In
 			}
 			else if (const UPCGTextureData* TextureData = Cast<UPCGTextureData>(DataToInspect.Data.Get()))
 			{
+				const bool bReadbackTextureDataOnInspect = CVarReadbackTextureDataOnInspect.GetValueOnGameThread();
+
 				FPCGTaggedData TaggedData;
-				TaggedData.Data = DataToInspect.DataPendingInit;
+				TaggedData.Data = bReadbackTextureDataOnInspect ? DataToInspect.DataPendingInit : DataToInspect.Data;
 				TaggedData.Pin = DataToInspect.PinLabel;
 				TaggedData.Tags = DataToInspect.AdditionalTags;
 				DataCollection.TaggedData.Add(TaggedData);
