[PCG] Fix Crash caused by PCGAssetExporterUtils bad reuse mechanism

diff --git a/Engine/Plugins/PCG/Source/PCG/Private/PCGAssetExporterUtils.cpp b/Engine/Plugins/PCG/Source/PCG/Private/PCGAssetExporterUtils.cpp
--- a/Engine/Plugins/PCG/Source/PCG/Private/PCGAssetExporterUtils.cpp
+++ b/Engine/Plugins/PCG/Source/PCG/Private/PCGAssetExporterUtils.cpp
@@ -10,6 +10,7 @@
 #include "FileHelpers.h"
 #include "IContentBrowserSingleton.h"
 #include "AssetRegistry/AssetRegistryModule.h"
+#include "Engine/Engine.h"
 #endif // WITH_EDITOR
 
 #define LOCTEXT_NAMESPACE "PCGAssetExporterUtils"
@@ -105,24 +106,21 @@ UPackage* UPCGAssetExporterUtils::CreateAsset(const UClass* AssetClass, const FP
 		}
 	}
 
-	UPackage* Package = FPackageName::DoesPackageExist(PackageName) ? LoadPackage(nullptr, *PackageName, LOAD_None) : nullptr;
+	UPackage* Package = FPackageName::DoesPackageExist(PackageName) ? LoadPackage(nullptr, *PackageName, LOAD_None) : FindPackage(nullptr, *PackageName);
 
 	UObject* Asset = nullptr;
+	UObject* PreviousAsset = nullptr;
 	bool NewAssetCreated = false;
 
 	if (Package)
 	{
-		UObject* Object = FindObjectFast<UObject>(Package, *AssetName);
-		if (Object && Object->GetClass() != AssetClass)
+		PreviousAsset = FindObjectFast<UObject>(Package, *AssetName);
+		if (PreviousAsset)
 		{
-			Object->SetFlags(RF_Transient);
-			Object->Rename(nullptr, GetTransientPackage(), REN_DontCreateRedirectors | REN_NonTransactional);
+			PreviousAsset->SetFlags(RF_Transient);
+			PreviousAsset->Rename(nullptr, GetTransientPackage(), REN_DontCreateRedirectors | REN_NonTransactional);
 			NewAssetCreated = true;
 		}
-		else
-		{
-			Asset = Object;
-		}
 	}
 	else
 	{
@@ -153,15 +151,16 @@ UPackage* UPCGAssetExporterUtils::CreateAsset(const UClass* AssetClass, const FP
 			FAssetRegistryModule::AssetCreated(Asset);
 		}
 
-		if (ExportFunc(PackageName, Asset))
+		const bool bSuccess = ExportFunc(PackageName, Asset);
+
+		if (PreviousAsset)
 		{
-			// Make sure everybody knows we changed the asset
-			if (!NewAssetCreated)
-			{
-				FCoreUObjectDelegates::BroadcastOnObjectModified(Asset);
-			}
+			TMap<UObject*, UObject*> OldToNew;
+			OldToNew.Emplace(PreviousAsset, Asset);
+			GEngine->NotifyToolsOfObjectReplacement(OldToNew);
 		}
-		else
+
+		if (!bSuccess)
 		{
 			return nullptr;
 		}
