[PCG] Add a backward compatibility CVar: pcg.Advanced.ConvertToPointDataInPostProcessGraph (disabled by default) that allows converting PCGComponent output from UPCGPointArrayData to UPCGPointData for licensees that rely on UPCGPointData
- The user will get a warning in the output log telling them to upgrade their code to support UPCGBasePointData

diff --git a/Engine/Plugins/PCG/Source/PCG/Private/PCGComponent.cpp b/Engine/Plugins/PCG/Source/PCG/Private/PCGComponent.cpp
--- a/Engine/Plugins/PCG/Source/PCG/Private/PCGComponent.cpp
+++ b/Engine/Plugins/PCG/Source/PCG/Private/PCGComponent.cpp
@@ -13,6 +13,8 @@
 #include "PCGSubsystem.h"
 #include "Data/PCGIntersectionData.h"
 #include "Data/PCGLandscapeData.h"
+#include "Data/PCGPointArrayData.h"
+#include "Data/PCGPointData.h"
 #include "Data/PCGSpatialData.h"
 #include "Data/PCGUnionData.h"
 #include "Graph/PCGStackContext.h"
@@ -89,6 +91,11 @@ namespace PCGComponent
 		1.0f,
 		TEXT("Global multiplier for generation radius of all runtime gen components."));
 
+	static TAutoConsoleVariable<bool> CVarConvertToPointDataInPostProcessGraph(
+		TEXT("pcg.Advanced.ConvertToPointDataInPostProcessGraph"),
+		false,
+		TEXT("At the end of execution, generated output data will be converted from UPCGPointArrayData to UPCGPointData."));
+
 	struct FPCGManagedActorLoadingScope
 	{
 		FPCGManagedActorLoadingScope()
@@ -761,6 +768,19 @@ void UPCGComponent::PostProcessGraph(const FBox& InNewBounds, bool bInGenerated,
 						}
 					});
 
+					if (PCGComponent::CVarConvertToPointDataInPostProcessGraph.GetValueOnAnyThread())
+					{
+						// To be backward compatible with licensee code that might rely on this output being strongly typed to UPCGPointData
+						// we convert UPCGPointArrayData back to UPCGPointData. This should be a temporary solution and licensees can convert their code to use UPCGBasePointData so
+						// that it is compatible with both point data classes and disable that CVar.
+						if (const UPCGPointArrayData* PointArrayData = Cast<UPCGPointArrayData>(OutputTaggedData.Data.Get()))
+						{
+							OutputTaggedData.Data = PointArrayData->ToPointData(Context);
+
+							UE_LOG(LogPCG, Warning, TEXT("UPCGPointArrayData was converted to UPCGPointData. Please update your code to support UPCGBasePointData if needed and set pcg.Advanced.ConvertToPointDataInPostProcessGraph to 0"));
+						}
+					}
+
 					// Finally add to the generated output collection
 					GeneratedGraphOutput.TaggedData.Add(OutputTaggedData);
 				}
