[PCG] Default view textures as top-down orthographic and unlit in data viewport
* Also future-proof the viewport callback function signature, so that it's easy to add new parameters in the future

#rnx

diff --git a/Engine/Plugins/PCG/Source/PCG/Public/PCGDataVisualization.h b/Engine/Plugins/PCG/Source/PCG/Public/PCGDataVisualization.h
--- a/Engine/Plugins/PCG/Source/PCG/Public/PCGDataVisualization.h
+++ b/Engine/Plugins/PCG/Source/PCG/Public/PCGDataVisualization.h
@@ -14,14 +14,22 @@ class UPCGSettingsInterface;
 struct FPCGContext;
 
 class FAdvancedPreviewScene;
+class FEditorViewportClient;
 struct FStreamableHandle;
 
-using FPCGSetupSceneFunc = TFunction<void(
-	FAdvancedPreviewScene* InScene,
-	const TArray<UObject*>& InResources,
-	FVector& OutViewLocation,
-	float& OutOrbitDistance,
-	TArray<TObjectPtr<UObject>>& OutManagedResources)>;
+struct FPCGSceneSetupParams
+{
+	FAdvancedPreviewScene* Scene = nullptr;
+	FEditorViewportClient* EditorViewportClient = nullptr;
+
+	/** Resources which are kept loaded for scene setup. */
+	TConstArrayView<UObject*> Resources;
+
+	/** Any UObjects created during scene setup should be tracked here to be kept visible for GC. */
+	TArray<TObjectPtr<UObject>> ManagedResources;
+};
+
+using FPCGSetupSceneFunc = TFunction<void(FPCGSceneSetupParams& InOutParams)>;
 
 enum class EPCGTableVisualizerColumnSortingMode
 {
diff --git a/Engine/Plugins/PCG/Source/PCGEditor/Private/DataVisualizations/PCGBaseTextureDataVisualization.cpp b/Engine/Plugins/PCG/Source/PCGEditor/Private/DataVisualizations/PCGBaseTextureDataVisualization.cpp
--- a/Engine/Plugins/PCG/Source/PCGEditor/Private/DataVisualizations/PCGBaseTextureDataVisualization.cpp
+++ b/Engine/Plugins/PCG/Source/PCGEditor/Private/DataVisualizations/PCGBaseTextureDataVisualization.cpp
@@ -7,6 +7,7 @@
 #include "Data/PCGTextureData.h"
 
 #include "AdvancedPreviewScene.h"
+#include "EditorViewportClient.h"
 #include "Components/StaticMeshComponent.h"
 #include "Engine/AssetManager.h"
 #include "Engine/Texture.h"
@@ -29,11 +30,13 @@ TArray<TSharedPtr<FStreamableHandle>> FPCGBaseTextureDataVisualization::LoadRequ
 
 FPCGSetupSceneFunc FPCGBaseTextureDataVisualization::GetViewportSetupFunc(const UPCGData* Data) const
 {
-	return [this, WeakData=TWeakObjectPtr<const UPCGBaseTextureData>(Cast<UPCGBaseTextureData>(Data))](FAdvancedPreviewScene* InScene, const TArray<UObject*>& InResources, FVector& OutViewLocation, float& OutOrbitDistance, TArray<TObjectPtr<UObject>>& OutManagedResources)
+	return [this, WeakData=TWeakObjectPtr<const UPCGBaseTextureData>(Cast<UPCGBaseTextureData>(Data))](FPCGSceneSetupParams& InOutParams)
 	{
-		check(InScene);
+		check(InOutParams.Scene);
+		check(InOutParams.EditorViewportClient);
+		check(InOutParams.EditorViewportClient->Viewport);
 
-		if (InResources.Num() != 2)
+		if (InOutParams.Resources.Num() != 2)
 		{
 			return;
 		}
@@ -55,7 +58,7 @@ FPCGSetupSceneFunc FPCGBaseTextureDataVisualization::GetViewportSetupFunc(const
 		else if (WeakData->GetTextureResourceType() == EPCGTextureResourceType::ExportedTexture)
 		{
 			Texture = UPCGVisualizationTexture2D::Create(WeakData);
-			OutManagedResources.Add(Texture);
+			InOutParams.ManagedResources.Add(Texture);
 		}
 		else
 		{
@@ -63,7 +66,7 @@ FPCGSetupSceneFunc FPCGBaseTextureDataVisualization::GetViewportSetupFunc(const
 			return;
 		}
 
-		UMaterialInterface* DebugMaterial = Cast<UMaterialInterface>(InResources[1]);
+		UMaterialInterface* DebugMaterial = Cast<UMaterialInterface>(InOutParams.Resources[1]);
 		UMaterialInstanceDynamic* MaterialInstance = UMaterialInstanceDynamic::Create(DebugMaterial, GetTransientPackage());
 
 		if (bIsTextureArray)
@@ -77,29 +80,38 @@ FPCGSetupSceneFunc FPCGBaseTextureDataVisualization::GetViewportSetupFunc(const
 			MaterialInstance->SetTextureParameterValue(FName(TEXT("DebugTexture")), Texture);
 		}
 
-		OutManagedResources.Add(MaterialInstance);
+		InOutParams.ManagedResources.Add(MaterialInstance);
 
 		TObjectPtr<UStaticMeshComponent> MeshComponent = NewObject<UStaticMeshComponent>(GetTransientPackage(), NAME_None, RF_Transient);
-		MeshComponent->SetStaticMesh(Cast<UStaticMesh>(InResources[0]));
+		MeshComponent->SetStaticMesh(Cast<UStaticMesh>(InOutParams.Resources[0]));
 		MeshComponent->OverrideMaterials.Add(MaterialInstance);
-		OutManagedResources.Add(MeshComponent);
+		InOutParams.ManagedResources.Add(MeshComponent);
 
 		if (GEditor->PreviewPlatform.GetEffectivePreviewFeatureLevel() <= ERHIFeatureLevel::ES3_1)
 		{
 			MeshComponent->SetMobility(EComponentMobility::Static);
 		}
 
-		static const FTransform MeshTransform(FRotator(0, 0, 90), FVector::ZeroVector, FVector::OneVector);
-		InScene->AddComponent(MeshComponent, MeshTransform);
+		static const FTransform MeshTransform(FRotator(0.0, -90.0, 0.0), FVector::ZeroVector, FVector::OneVector);
+		InOutParams.Scene->AddComponent(MeshComponent, MeshTransform);
+		InOutParams.Scene->SetFloorVisibility(false);
+		InOutParams.Scene->SetEnvironmentVisibility(false);
+		InOutParams.EditorViewportClient->SetViewportType(ELevelViewportType::LVT_OrthoTop);
+		InOutParams.EditorViewportClient->SetViewMode(EViewModeIndex::VMI_Unlit);
+		InOutParams.EditorViewportClient->SetViewLocation(FVector::UpVector);
 
-		// Bounds will be updated already by SetStaticMesh() call.
-		const FBoxSphereBounds Bounds = MeshComponent->Bounds;
-		InScene->SetFloorOffset(-Bounds.Origin.Z + Bounds.BoxExtent.Z);
-		OutViewLocation = Bounds.Origin;
+		const FIntPoint ViewportSize = InOutParams.EditorViewportClient->Viewport->GetSizeXY();
 
-		if (Bounds.SphereRadius > 0.0f)
+		if (ViewportSize.X > 0 && ViewportSize.Y > 0)
 		{
-			OutOrbitDistance = Bounds.SphereRadius * 2.0f;
+			// Bounds will be updated already by SetStaticMesh() call.
+			const FBoxSphereBounds Bounds = MeshComponent->Bounds;
+
+			// Fit the orthographic zoom to the mesh. 0.8f is chosen arbitrarily to add some padding around the mesh.
+			const float MeshUnitsPerPixel = FMath::Max(Bounds.BoxExtent.X / ViewportSize.X, Bounds.BoxExtent.Y / ViewportSize.Y) * 2.0f;
+			const float OrthoZoom = FMath::Clamp(MeshUnitsPerPixel * DEFAULT_ORTHOZOOM * 0.8f, MIN_ORTHOZOOM, MAX_ORTHOZOOM);
+
+			InOutParams.EditorViewportClient->SetOrthoZoom(OrthoZoom);
 		}
 	};
 }
diff --git a/Engine/Plugins/PCG/Source/PCGEditor/Private/DataVisualizations/PCGSpatialDataVisualization.cpp b/Engine/Plugins/PCG/Source/PCGEditor/Private/DataVisualizations/PCGSpatialDataVisualization.cpp
--- a/Engine/Plugins/PCG/Source/PCGEditor/Private/DataVisualizations/PCGSpatialDataVisualization.cpp
+++ b/Engine/Plugins/PCG/Source/PCGEditor/Private/DataVisualizations/PCGSpatialDataVisualization.cpp
@@ -19,6 +19,7 @@
 #include "Metadata/Accessors/PCGCustomAccessor.h"
 
 #include "AdvancedPreviewScene.h"
+#include "EditorViewportClient.h"
 #include "Components/InstancedStaticMeshComponent.h"
 #include "Engine/AssetManager.h"
 
@@ -406,9 +407,10 @@ TArray<FPCGMetadataDomainID> IPCGSpatialDataVisualization::GetAllSupportedDomain
 
 FPCGSetupSceneFunc IPCGSpatialDataVisualization::GetViewportSetupFunc(const UPCGData* Data) const
 {
-	return [this, WeakData=TWeakObjectPtr<const UPCGData>(Data)](FAdvancedPreviewScene* InScene, const TArray<UObject*>& InResources, FVector& OutViewLocation, float& OutOrbitDistance, TArray<TObjectPtr<UObject>>& OutManagedResources)
+	return [this, WeakData=TWeakObjectPtr<const UPCGData>(Data)](FPCGSceneSetupParams& InOutParams)
 	{
-		check(InScene);
+		check(InOutParams.Scene);
+		check(InOutParams.EditorViewportClient);
 
 		if (!WeakData.IsValid())
 		{
@@ -425,12 +427,12 @@ FPCGSetupSceneFunc IPCGSpatialDataVisualization::GetViewportSetupFunc(const UPCG
 			nullptr,
 			nullptr,
 			FPCGCrc(),
-			[&OutManagedResources, InScene, &bInitializedBounds, &BoundsMin, &BoundsMax](UInstancedStaticMeshComponent* ISMC)
+			[&InOutParams, &bInitializedBounds, &BoundsMin, &BoundsMax](UInstancedStaticMeshComponent* ISMC)
 			{
 				check(ISMC);
 
-				OutManagedResources.Add(ISMC);
-				InScene->AddComponent(ISMC, FTransform::Identity);
+				InOutParams.ManagedResources.Add(ISMC);
+				InOutParams.Scene->AddComponent(ISMC, FTransform::Identity);
 
 				const FVector CurrentBoundsMin = ISMC->Bounds.Origin - ISMC->Bounds.BoxExtent;
 				const FVector CurrentBoundsMax = ISMC->Bounds.Origin + ISMC->Bounds.BoxExtent;
@@ -450,13 +452,13 @@ FPCGSetupSceneFunc IPCGSpatialDataVisualization::GetViewportSetupFunc(const UPCG
 		);
 
 		const FBoxSphereBounds Bounds = FBoxSphereBounds(FBox(BoundsMin, BoundsMax));
-		InScene->SetFloorOffset(-Bounds.Origin.Z + Bounds.BoxExtent.Z);
+		InOutParams.Scene->SetFloorOffset(-Bounds.Origin.Z + Bounds.BoxExtent.Z);
 
-		OutViewLocation = Bounds.Origin;
+		InOutParams.EditorViewportClient->SetViewLocation(Bounds.Origin);
 
 		if (Bounds.SphereRadius > 0.0f)
 		{
-			OutOrbitDistance = Bounds.SphereRadius * 2.0f;
+			InOutParams.EditorViewportClient->SetViewLocationForOrbiting(Bounds.Origin, Bounds.SphereRadius * 2.0f);
 		}
 	};
 }
diff --git a/Engine/Plugins/PCG/Source/PCGEditor/Private/DataVisualizations/PCGSplineDataVisualization.cpp b/Engine/Plugins/PCG/Source/PCGEditor/Private/DataVisualizations/PCGSplineDataVisualization.cpp
--- a/Engine/Plugins/PCG/Source/PCGEditor/Private/DataVisualizations/PCGSplineDataVisualization.cpp
+++ b/Engine/Plugins/PCG/Source/PCGEditor/Private/DataVisualizations/PCGSplineDataVisualization.cpp
@@ -9,6 +9,7 @@
 #include "DataVisualizations/PCGDataVisualizationHelpers.h"
 
 #include "AdvancedPreviewScene.h"
+#include "EditorViewportClient.h"
 #include "Components/SplineComponent.h"
 
 #define LOCTEXT_NAMESPACE "PCGSplineDataVisualization"
@@ -142,9 +143,10 @@ const UPCGBasePointData* IPCGSplineDataVisualization::CollapseToDebugBasePointDa
 
 FPCGSetupSceneFunc IPCGSplineDataVisualization::GetViewportSetupFunc(const UPCGData* Data) const
 {
-	return[this, Data](FAdvancedPreviewScene* InScene, const TArray<UObject*>& InResources, FVector& OutViewLocation, float& OutOrbitDistance, TArray<TObjectPtr<UObject>>& OutManagedResources)
+	return[this, Data](FPCGSceneSetupParams& InOutParams)
 	{
-		check(InScene);
+		check(InOutParams.Scene);
+		check(InOutParams.EditorViewportClient);
 
 		const UPCGSplineData* SplineData = Cast<UPCGSplineData>(Data);
 
@@ -161,17 +163,17 @@ FPCGSetupSceneFunc IPCGSplineDataVisualization::GetViewportSetupFunc(const UPCGD
 			SplineComponent->SetMobility(EComponentMobility::Static);
 		}
 
-		OutManagedResources.Add(SplineComponent);
-		InScene->AddComponent(SplineComponent, SplineComponent->GetComponentToWorld());
+		InOutParams.ManagedResources.Add(SplineComponent);
+		InOutParams.Scene->AddComponent(SplineComponent, SplineComponent->GetComponentToWorld());
 
 		FBoxSphereBounds Bounds = SplineComponent->CalcBounds(SplineComponent->GetComponentToWorld());
 
-		InScene->SetFloorOffset(-Bounds.Origin.Z + Bounds.BoxExtent.Z);
-		OutViewLocation = Bounds.Origin;
+		InOutParams.Scene->SetFloorOffset(-Bounds.Origin.Z + Bounds.BoxExtent.Z);
+		InOutParams.EditorViewportClient->SetViewLocation(Bounds.Origin);
 
 		if (Bounds.SphereRadius > 0.0f)
 		{
-			OutOrbitDistance = Bounds.SphereRadius * 2.0f;
+			InOutParams.EditorViewportClient->SetViewLocationForOrbiting(Bounds.Origin, Bounds.SphereRadius * 2.0f);
 		}
 	};
 }
diff --git a/Engine/Plugins/PCG/Source/PCGEditor/Private/DataVisualizations/PCGStaticMeshDataVisualization.cpp b/Engine/Plugins/PCG/Source/PCGEditor/Private/DataVisualizations/PCGStaticMeshDataVisualization.cpp
--- a/Engine/Plugins/PCG/Source/PCGEditor/Private/DataVisualizations/PCGStaticMeshDataVisualization.cpp
+++ b/Engine/Plugins/PCG/Source/PCGEditor/Private/DataVisualizations/PCGStaticMeshDataVisualization.cpp
@@ -6,6 +6,7 @@
 #include "DataVisualizations/PCGDataVisualizationHelpers.h"
 
 #include "AdvancedPreviewScene.h"
+#include "EditorViewportClient.h"
 #include "MaterialEditor/MaterialEditorMeshComponent.h"
 
 FPCGTableVisualizerInfo IPCGStaticMeshDataVisualization::GetTableVisualizerInfoWithDomain(const UPCGData* Data, const FPCGMetadataDomainID& DomainID) const
@@ -31,38 +32,39 @@ TArray<TSharedPtr<FStreamableHandle>> IPCGStaticMeshDataVisualization::LoadRequi
 
 FPCGSetupSceneFunc IPCGStaticMeshDataVisualization::GetViewportSetupFunc(const UPCGData* Data) const
 {
-	return [this](FAdvancedPreviewScene* InScene, const TArray<UObject*>& InResources, FVector& OutViewLocation, float& OutOrbitDistance, TArray<TObjectPtr<UObject>>& OutManagedResources)
+	return [this](FPCGSceneSetupParams& InOutParams)
 	{
-		check(InScene);
+		check(InOutParams.Scene);
+		check(InOutParams.EditorViewportClient);
 
-		if (InResources.IsEmpty())
+		if (InOutParams.Resources.IsEmpty())
 		{
 			return;
 		}
 
 		// Note: Using UMaterialEditorMeshComponent subclass for more accurate mesh bounds.
 		TObjectPtr<UStaticMeshComponent> MeshComponent = NewObject<UMaterialEditorMeshComponent>(GetTransientPackage(), NAME_None, RF_Transient);
-		OutManagedResources.Add(MeshComponent);
+		InOutParams.ManagedResources.Add(MeshComponent);
 
 		if (GEditor->PreviewPlatform.GetEffectivePreviewFeatureLevel() <= ERHIFeatureLevel::ES3_1)
 		{
 			MeshComponent->SetMobility(EComponentMobility::Static);
 		}
 
-		InScene->AddComponent(MeshComponent, FTransform::Identity);
+		InOutParams.Scene->AddComponent(MeshComponent, FTransform::Identity);
 
-		UStaticMesh* StaticMesh = Cast<UStaticMesh>(InResources[0]);
+		UStaticMesh* StaticMesh = Cast<UStaticMesh>(InOutParams.Resources[0]);
 		MeshComponent->SetStaticMesh(StaticMesh);
 
 		// Bounds will be updated already by SetStaticMesh() call.
 		const FBoxSphereBounds Bounds = MeshComponent->Bounds;
-		InScene->SetFloorOffset(-Bounds.Origin.Z + Bounds.BoxExtent.Z);
+		InOutParams.Scene->SetFloorOffset(-Bounds.Origin.Z + Bounds.BoxExtent.Z);
 
-		OutViewLocation = Bounds.Origin;
+		InOutParams.EditorViewportClient->SetViewLocation(Bounds.Origin);
 
 		if (Bounds.SphereRadius > 0.0f)
 		{
-			OutOrbitDistance = Bounds.SphereRadius * 2.0f;
+			InOutParams.EditorViewportClient->SetViewLocationForOrbiting(Bounds.Origin, Bounds.SphereRadius * 2.0f);
 		}
 	};
 }
diff --git a/Engine/Plugins/PCG/Source/PCGEditor/Private/DataVisualizations/PCGVolumeDataVisualization.cpp b/Engine/Plugins/PCG/Source/PCGEditor/Private/DataVisualizations/PCGVolumeDataVisualization.cpp
--- a/Engine/Plugins/PCG/Source/PCGEditor/Private/DataVisualizations/PCGVolumeDataVisualization.cpp
+++ b/Engine/Plugins/PCG/Source/PCGEditor/Private/DataVisualizations/PCGVolumeDataVisualization.cpp
@@ -14,6 +14,7 @@
 #include "DataVisualizations/PCGCollisionVisComponent.h"
 
 #include "AdvancedPreviewScene.h"
+#include "EditorViewportClient.h"
 #include "Chaos/ChaosEngineInterface.h"
 #include "Components/BoxComponent.h"
 #include "Components/CapsuleComponent.h"
@@ -54,9 +55,10 @@ const UPCGBasePointData* FPCGPrimitiveDataVisualization::CollapseToDebugBasePoin
 
 FPCGSetupSceneFunc FPCGPrimitiveDataVisualization::GetViewportSetupFunc(const UPCGData* Data) const
 {
-	return [this, Data](FAdvancedPreviewScene* InScene, const TArray<UObject*>& InResources, FVector& OutViewLocation, float& OutOrbitDistance, TArray<TObjectPtr<UObject>>& OutManagedResources)
+	return [this, Data](FPCGSceneSetupParams& InOutParams)
 	{
-		check(InScene);
+		check(InOutParams.Scene);
+		check(InOutParams.EditorViewportClient);
 
 		TArray<TObjectPtr<UPrimitiveComponent>> Components;
 		TArray<FTransform> ComponentTransforms;
@@ -81,8 +83,8 @@ FPCGSetupSceneFunc FPCGPrimitiveDataVisualization::GetViewportSetupFunc(const UP
 				Component->SetMobility(EComponentMobility::Static);
 			}	
 
-			OutManagedResources.Add(Component);
-			InScene->AddComponent(Component, Transform);
+			InOutParams.ManagedResources.Add(Component);
+			InOutParams.Scene->AddComponent(Component, Transform);
 
 			if (bHasBounds)
 			{
@@ -97,12 +99,12 @@ FPCGSetupSceneFunc FPCGPrimitiveDataVisualization::GetViewportSetupFunc(const UP
 
 		if (bHasBounds)
 		{
-			InScene->SetFloorOffset(-Bounds.Origin.Z + Bounds.BoxExtent.Z);
-			OutViewLocation = Bounds.Origin;
+			InOutParams.Scene->SetFloorOffset(-Bounds.Origin.Z + Bounds.BoxExtent.Z);
+			InOutParams.EditorViewportClient->SetViewLocation(Bounds.Origin);
 
 			if (Bounds.SphereRadius > 0.0f)
 			{
-				OutOrbitDistance = Bounds.SphereRadius * 2.0f;
+				InOutParams.EditorViewportClient->SetViewLocationForOrbiting(Bounds.Origin, Bounds.SphereRadius * 2.0f);
 			}
 		}
 	};
diff --git a/Engine/Plugins/PCG/Source/PCGEditor/Private/Widgets/AssetEditorViewport/SPCGEditorViewport.cpp b/Engine/Plugins/PCG/Source/PCGEditor/Private/Widgets/AssetEditorViewport/SPCGEditorViewport.cpp
--- a/Engine/Plugins/PCG/Source/PCGEditor/Private/Widgets/AssetEditorViewport/SPCGEditorViewport.cpp
+++ b/Engine/Plugins/PCG/Source/PCGEditor/Private/Widgets/AssetEditorViewport/SPCGEditorViewport.cpp
@@ -12,6 +12,7 @@
 namespace PCGEditorViewportConstants
 {
 	static const FVector DefaultViewLocation = FVector::ZeroVector;
+	static const FRotator DefaultViewRotation = FRotator(-25.0f, -135.0f, 0.0f);
 	static const float DefaultOrbitDistance = 500.0f;
 }
 
@@ -116,6 +117,8 @@ void FPCGEditorViewportClient::Tick(float DeltaSeconds)
 
 FAdvancedPreviewScene* FPCGEditorViewportClient::ResetScene()
 {
+	TRACE_CPUPROFILER_EVENT_SCOPE(FPCGEditorViewportClient::ResetScene);
+
 	Scene = MakeShareable(new FAdvancedPreviewScene(FPreviewScene::ConstructionValues()));
 	PreviewScene = Scene.Get();
 
@@ -125,13 +128,17 @@ FAdvancedPreviewScene* FPCGEditorViewportClient::ResetScene()
 		World->ChangeFeatureLevel(GWorld->GetFeatureLevel());
 	}
 
-	UAssetViewerSettings* Settings = UAssetViewerSettings::Get();
-	const int32 ProfileIndex = Scene->GetCurrentProfileIndex();
+	SetViewportType(ELevelViewportType::LVT_Perspective);
+	SetViewMode(EViewModeIndex::VMI_Lit);
+	SetViewLocation(PCGEditorViewportConstants::DefaultViewLocation);
+	SetViewRotation(PCGEditorViewportConstants::DefaultViewRotation);
+	SetViewLocationForOrbiting(PCGEditorViewportConstants::DefaultViewLocation, PCGEditorViewportConstants::DefaultOrbitDistance);
 
-	if (Settings && Settings->Profiles.IsValidIndex(ProfileIndex))
-	{
-		Scene->SetEnvironmentVisibility(Settings->Profiles[ProfileIndex].bShowEnvironment, true);
-	}
+	Scene->SetFloorOffset(0.0f);
+	Scene->SetFloorVisibility(true);
+	Scene->SetEnvironmentVisibility(true);
+
+	Invalidate();
 
 	return Scene.Get();
 }
@@ -151,7 +158,6 @@ void SPCGEditorViewport::Construct(const FArguments& InArgs)
 	SEditorViewport::Construct(SEditorViewport::FArguments());
 }
 
-
 TSharedRef<class SEditorViewport> SPCGEditorViewport::GetViewportWidget()
 {
 	return SharedThis(this);
@@ -164,6 +170,8 @@ TSharedPtr<FExtender> SPCGEditorViewport::GetExtenders() const
 
 void SPCGEditorViewport::SetupScene(const TArray<UObject*>& InResources, const FPCGSetupSceneFunc& SetupFunc)
 {
+	TRACE_CPUPROFILER_EVENT_SCOPE(SPCGEditorViewport::SetupScene);
+
 	ResetScene();
 
 	if (!ensure(EditorViewportClient.IsValid()))
@@ -172,17 +180,22 @@ void SPCGEditorViewport::SetupScene(const TArray<UObject*>& InResources, const F
 	}
 
 	FAdvancedPreviewScene* Scene = EditorViewportClient->GetPreviewScene();
-	FVector ViewLocation = PCGEditorViewportConstants::DefaultViewLocation;
-	float OrbitDistance = PCGEditorViewportConstants::DefaultOrbitDistance;
 
 	if (Scene && SetupFunc.IsSet())
 	{
-		SetupFunc(Scene, InResources, ViewLocation, OrbitDistance, ManagedResources);
-	}
+		TRACE_CPUPROFILER_EVENT_SCOPE(SPCGEditorViewport::SetupSceneCallback);
 
-	EditorViewportClient->SetViewLocation(ViewLocation);
-	EditorViewportClient->SetViewLocationForOrbiting(ViewLocation, OrbitDistance);
-	EditorViewportClient->Invalidate();
+		FPCGSceneSetupParams SceneSetupParams;
+		SceneSetupParams.Scene = Scene;
+		SceneSetupParams.EditorViewportClient = EditorViewportClient.Get();
+		SceneSetupParams.Resources = InResources;
+
+		SetupFunc(SceneSetupParams);
+
+		ManagedResources = MoveTemp(SceneSetupParams.ManagedResources);
+
+		EditorViewportClient->Invalidate();
+	}
 }
 
 void SPCGEditorViewport::ResetScene()
@@ -195,16 +208,13 @@ void SPCGEditorViewport::ResetScene()
 	}
 
 	EditorViewportClient->ResetScene();
-	EditorViewportClient->SetViewLocation(PCGEditorViewportConstants::DefaultViewLocation);
-	EditorViewportClient->SetViewLocationForOrbiting(PCGEditorViewportConstants::DefaultViewLocation, PCGEditorViewportConstants::DefaultOrbitDistance);
-	EditorViewportClient->Invalidate();
 }
 
 TSharedRef<FEditorViewportClient> SPCGEditorViewport::MakeEditorViewportClient()
 {
 	EditorViewportClient = MakeShareable(new FPCGEditorViewportClient(SharedThis(this)));
 	EditorViewportClient->SetViewLocation(PCGEditorViewportConstants::DefaultViewLocation);
-	EditorViewportClient->SetViewRotation(FRotator(-25.0f, -135.0f, 0.0f));
+	EditorViewportClient->SetViewRotation(PCGEditorViewportConstants::DefaultViewRotation);
 	EditorViewportClient->SetViewLocationForOrbiting(PCGEditorViewportConstants::DefaultViewLocation, PCGEditorViewportConstants::DefaultOrbitDistance);
 	EditorViewportClient->bSetListenerPosition = false;
 	EditorViewportClient->EngineShowFlags.EnableAdvancedFeatures();
@@ -219,7 +229,7 @@ TSharedRef<FEditorViewportClient> SPCGEditorViewport::MakeEditorViewportClient()
 
 TSharedPtr<SWidget> SPCGEditorViewport::BuildViewportToolbar()
 {
-	const FName ToolbarName = "PCG.DataViewportToolbar";
+	const FName ToolbarName = "PCG.ViewportToolbar";
 	
 	if (!UToolMenus::Get()->IsMenuRegistered(ToolbarName))
 	{
diff --git a/Engine/Plugins/PCGInterops/PCGGeometryScriptInterop/Source/PCGGeometryScriptInterop/Private/Editor/PCGDynamicMeshDataVisualization.cpp b/Engine/Plugins/PCGInterops/PCGGeometryScriptInterop/Source/PCGGeometryScriptInterop/Private/Editor/PCGDynamicMeshDataVisualization.cpp
--- a/Engine/Plugins/PCGInterops/PCGGeometryScriptInterop/Source/PCGGeometryScriptInterop/Private/Editor/PCGDynamicMeshDataVisualization.cpp
+++ b/Engine/Plugins/PCGInterops/PCGGeometryScriptInterop/Source/PCGGeometryScriptInterop/Private/Editor/PCGDynamicMeshDataVisualization.cpp
@@ -11,6 +11,7 @@
 #include "Resources/PCGDynamicMeshManagedComponent.h"
 
 #include "AdvancedPreviewScene.h"
+#include "EditorViewportClient.h"
 #include "UDynamicMesh.h"
 #include "Components/DynamicMeshComponent.h"
 
@@ -38,9 +39,10 @@ void FPCGDynamicMeshDataVisualization::ExecuteDebugDisplay(FPCGContext* Context,
 
 FPCGSetupSceneFunc FPCGDynamicMeshDataVisualization::GetViewportSetupFunc(const UPCGData* Data) const
 {
-	return [this, WeakData=TWeakObjectPtr<const UPCGDynamicMeshData>(Cast<UPCGDynamicMeshData>(Data))](FAdvancedPreviewScene* InScene, const TArray<UObject*>& InResources, FVector& OutViewLocation, float& OutOrbitDistance, TArray<TObjectPtr<UObject>>& OutManagedResources)
+	return [this, WeakData=TWeakObjectPtr<const UPCGDynamicMeshData>(Cast<UPCGDynamicMeshData>(Data))](FPCGSceneSetupParams& InOutParams)
 	{
-		check(InScene);
+		check(InOutParams.Scene);
+		check(InOutParams.EditorViewportClient);
 
 		if (!WeakData.IsValid())
 		{
@@ -51,23 +53,23 @@ FPCGSetupSceneFunc FPCGDynamicMeshDataVisualization::GetViewportSetupFunc(const
 		TObjectPtr<UDynamicMeshComponent> DynamicMeshComponent = NewObject<UDynamicMeshComponent>(GetTransientPackage(), NAME_None, RF_Transient);
 		WeakData->InitializeDynamicMeshComponentFromData(DynamicMeshComponent);
 
-		OutManagedResources.Add(DynamicMeshComponent);
+		InOutParams.ManagedResources.Add(DynamicMeshComponent);
 
 		if (GEditor->PreviewPlatform.GetEffectivePreviewFeatureLevel() <= ERHIFeatureLevel::ES3_1)
 		{
 			DynamicMeshComponent->SetMobility(EComponentMobility::Static);
 		}
 
-		InScene->AddComponent(DynamicMeshComponent, FTransform::Identity);
+		InOutParams.Scene->AddComponent(DynamicMeshComponent, FTransform::Identity);
 
 		const FBoxSphereBounds Bounds = DynamicMeshComponent->CalcLocalBounds();
-		InScene->SetFloorOffset(-Bounds.Origin.Z + Bounds.BoxExtent.Z);
+		InOutParams.Scene->SetFloorOffset(-Bounds.Origin.Z + Bounds.BoxExtent.Z);
 
-		OutViewLocation = Bounds.Origin;
+		InOutParams.EditorViewportClient->SetViewLocation(Bounds.Origin);
 
 		if (Bounds.SphereRadius > 0.0f)
 		{
-			OutOrbitDistance = Bounds.SphereRadius * 2.0f;
+			InOutParams.EditorViewportClient->SetViewLocationForOrbiting(Bounds.Origin, Bounds.SphereRadius * 2.0f);
 		}
 	};
 }
