[PCG] Finalize GPU API for 5.6
* Mark NumPoints as deprecated in the declarations
* Remove Get/Set attribute override that takes an index from the declarations
* Fix empty category dropdown in kernel-driven texture output pin settings

#rnx

diff --git a/Engine/Plugins/PCG/Source/PCG/Private/Compute/Elements/PCGCustomHLSL.cpp b/Engine/Plugins/PCG/Source/PCG/Private/Compute/Elements/PCGCustomHLSL.cpp
--- a/Engine/Plugins/PCG/Source/PCG/Private/Compute/Elements/PCGCustomHLSL.cpp
+++ b/Engine/Plugins/PCG/Source/PCG/Private/Compute/Elements/PCGCustomHLSL.cpp
@@ -353,7 +353,11 @@ void UPCGCustomHLSLSettings::UpdateInputDeclarations()
 
 		if (KernelType == EPCGKernelType::PointGenerator)
 		{
-			InputDeclarations += FString::Format(TEXT("const uint NumElements = {0};\n\n"), { PointCount });
+			InputDeclarations += FString::Format(TEXT(
+				"const uint NumElements = {0};\n"
+				"const uint NumPoints = {0}; // Deprecated, use NumElements instead.\n" // Deprecated 5.6
+				"\n"
+			), { PointCount });
 		}
 		else if (KernelType == EPCGKernelType::TextureGenerator)
 		{
@@ -457,7 +461,6 @@ void UPCGCustomHLSLSettings::UpdateInputDeclarations()
 			"\n"
 			"// Valid types: bool, int, float, float2, float3, float4, Rotator (float3), Quat (float4), Transform (float4x4), StringKey (int), Name (uint2)\n"
 			"\n"
-			"{type} {0}_Get{type}(uint DataIndex, uint ElementIndex, int AttributeId);\n"
 			"{type} {0}_Get{type}(uint DataIndex, uint ElementIndex, 'AttributeName');\n"
 			"\n"
 			"// Example: {0}_GetFloat({0}_DataIndex, ElementIndex, 'MyFloatAttr');\n"),
@@ -667,7 +670,6 @@ void UPCGCustomHLSLSettings::UpdateOutputDeclarations()
 		OutputDeclarations += FString::Format(TEXT(
 			"// Valid types: bool, int, float, float2, float3, float4, Rotator (float3), Quat (float4), Transform (float4x4), StringKey (int), Name (uint2)\n"
 			"\n"
-			"void {0}_Set{type}(uint DataIndex, uint ElementIndex, int AttributeId, {type} Value);\n"
 			"void {0}_Set{type}(uint DataIndex, uint ElementIndex, 'AttributeName', {type} Value);\n"
 			"\n"
 			"// Example: {0}_SetFloat({0}_DataIndex, ElementIndex, 'MyFloatAttr', MyValue);\n"),
@@ -932,6 +934,8 @@ void UPCGCustomHLSLSettings::UpdatePinSettings()
 				bInitModeDrivenByKernel = true;
 				Properties.AllowedTypes = EPCGDataType::BaseTexture;
 			}
+
+			Properties.bShowPropertiesGPU = !IsTextureKernel();
 		}
 
 		if (!!(Properties.AllowedTypes & EPCGDataType::BaseTexture))
diff --git a/Engine/Plugins/PCG/Source/PCG/Private/Compute/PCGPinPropertiesGPU.h b/Engine/Plugins/PCG/Source/PCG/Private/Compute/PCGPinPropertiesGPU.h
--- a/Engine/Plugins/PCG/Source/PCG/Private/Compute/PCGPinPropertiesGPU.h
+++ b/Engine/Plugins/PCG/Source/PCG/Private/Compute/PCGPinPropertiesGPU.h
@@ -155,8 +155,13 @@ public:
 #endif
 
 public:
-	UPROPERTY(BlueprintReadWrite, EditAnywhere, Category = Settings, meta = (DisplayName = "GPU Properties", DisplayAfter = "Tooltip"))
+	UPROPERTY(BlueprintReadWrite, EditAnywhere, Category = Settings, meta = (DisplayName = "GPU Properties", DisplayAfter = "Tooltip", EditCondition = "bShowPropertiesGPU", EditConditionHides, HideEditConditionToggle))
 	FPCGPinPropertiesGPUStruct PropertiesGPU;
+
+#if WITH_EDITORONLY_DATA
+	UPROPERTY(Transient)
+	bool bShowPropertiesGPU = true;
+#endif // WITH_EDITORONLY_DATA
 };
 
 template<>
