[PCG] Fix remapping of attribute IDs when reusing GPU buffers.

Only apply mapping that relate to the buffer.

#rnx

diff --git a/Engine/Plugins/PCG/Source/PCG/Private/Compute/PCGDataUploadAdaptor.cpp b/Engine/Plugins/PCG/Source/PCG/Private/Compute/PCGDataUploadAdaptor.cpp
--- a/Engine/Plugins/PCG/Source/PCG/Private/Compute/PCGDataUploadAdaptor.cpp
+++ b/Engine/Plugins/PCG/Source/PCG/Private/Compute/PCGDataUploadAdaptor.cpp
@@ -119,16 +119,22 @@ bool FPCGDataUploadAdaptor::PrepareData_GameThread()
 
 		SourceToTargetAttributeId.Reserve(Binding->GetAttributeIdRemap().Num());
 
-		Algo::Transform(Binding->GetAttributeIdRemap(), SourceToTargetAttributeId, [](const TPair<UPCGDataBinding::FSourceBufferAndAttributeIndex, int32>& Remapping)
+		for (const TPair<UPCGDataBinding::FSourceBufferAndAttributeIndex, int32>& Entry : Binding->GetAttributeIdRemap())
 		{
-			const int32 TargetAttributeId = PCGComputeHelpers::GetAttributeIdFromMetadataAttributeIndex(Remapping.Value);
-			const int32 SourceAttributeId = PCGComputeHelpers::GetAttributeIdFromMetadataAttributeIndex(Remapping.Key.Value);
+			const TSharedPtr<const FPCGProxyForGPUDataCollection>& SourceBuffer = Entry.Key.Key;
 
-			ensure(TargetAttributeId >= PCGComputeConstants::NUM_RESERVED_ATTRS);
-			ensure(SourceAttributeId >= PCGComputeConstants::NUM_RESERVED_ATTRS);
+			if (SourceBuffer && SourceBuffer->GetBuffer() == ExternalBufferForReuse)
+			{
+				const int32 GraphAttributeId = PCGComputeHelpers::GetAttributeIdFromMetadataAttributeIndex(Entry.Value);
+
+				const int32 SourceBufferAttributeId = PCGComputeHelpers::GetAttributeIdFromMetadataAttributeIndex(Entry.Key.Value);
+
+				ensure(GraphAttributeId >= PCGComputeConstants::NUM_RESERVED_ATTRS);
+				ensure(SourceBufferAttributeId >= PCGComputeConstants::NUM_RESERVED_ATTRS);
 
-			return FIntVector2(SourceAttributeId, TargetAttributeId);
-		});
+				SourceToTargetAttributeId.Add(FIntVector2(SourceBufferAttributeId, GraphAttributeId));
+			}
+		}
 	}
 	else
 	{
