[PCG] PCGLandscapeSplineData: remove checks in favor of default behaviors with logging to prevent crashes from bad input

#rnx

diff --git a/Engine/Plugins/PCG/Source/PCG/Private/Data/PCGLandscapeSplineData.cpp b/Engine/Plugins/PCG/Source/PCG/Private/Data/PCGLandscapeSplineData.cpp
--- a/Engine/Plugins/PCG/Source/PCG/Private/Data/PCGLandscapeSplineData.cpp
+++ b/Engine/Plugins/PCG/Source/PCG/Private/Data/PCGLandscapeSplineData.cpp
@@ -40,6 +40,24 @@ namespace PCGLandscapeDataHelpers
 	}
 }
 
+bool UPCGLandscapeSplineData::CheckSpline() const
+{
+	if (!Spline.IsValid())
+	{
+		if (!bLoggedInvalidSpline)
+		{
+			UE_LOG(LogPCG, Warning, TEXT("Invalid landscape spline component"));
+			bLoggedInvalidSpline = true;
+		}
+		
+		return false;
+	}
+	else
+	{
+		return true;
+	}
+}
+
 void UPCGLandscapeSplineData::Initialize(ULandscapeSplinesComponent* InSplineComponent)
 {
 	check(InSplineComponent);
@@ -64,13 +82,12 @@ void UPCGLandscapeSplineData::AddToCrc(FArchiveCrc32& Ar, bool bFullDataCrc) con
 
 FTransform UPCGLandscapeSplineData::GetTransform() const
 {
-	return Spline.IsValid() ? Spline->GetComponentTransform() : FTransform::Identity;
+	return CheckSpline() ? Spline->GetComponentTransform() : FTransform::Identity;
 }
 
 int UPCGLandscapeSplineData::GetNumSegments() const
 {
-	check(Spline.IsValid());
-	return Spline->GetSegments().Num();
+	return CheckSpline() ? Spline->GetSegments().Num() : 0;
 }
 
 FVector::FReal UPCGLandscapeSplineData::GetSegmentLength(int SegmentIndex) const
@@ -80,23 +97,30 @@ FVector::FReal UPCGLandscapeSplineData::GetSegmentLength(int SegmentIndex) const
 
 FTransform UPCGLandscapeSplineData::GetTransformAtDistance(int SegmentIndex, FVector::FReal Distance, bool bWorldSpace, FBox* OutBounds) const
 {
-	check(Spline.IsValid());
-
-	if (GetNumSegments() == 0)
+	if (!CheckSpline() || GetNumSegments() == 0)
 	{
 		return GetTransform();
 	}
 
-	check(Spline->GetSegments().IsValidIndex(SegmentIndex));
+	const int ClampedSegmentIndex = FMath::Clamp(SegmentIndex, 0, GetNumSegments()-1);
 
-	const ULandscapeSplineSegment* Segment = Spline->GetSegments()[SegmentIndex];
+	const ULandscapeSplineSegment* Segment = Spline->GetSegments()[ClampedSegmentIndex];
 	check(Segment);
 
 	const TArray<FLandscapeSplineInterpPoint>& InterpPoints = Segment->GetPoints();
 
-	int32 PointIndex;
-	FVector::FReal Alpha;
-	GetInterpPointAtDistance(SegmentIndex, Distance, PointIndex, /*bComputeAlpha=*/true, Alpha);
+	int32 PointIndex = 0;
+	FVector::FReal Alpha = 0.0f;
+
+	if (ClampedSegmentIndex == SegmentIndex)
+	{
+		GetInterpPointAtDistance(ClampedSegmentIndex, Distance, PointIndex, /*bComputeAlpha=*/true, Alpha);
+	}
+	else if(ClampedSegmentIndex < SegmentIndex)
+	{
+		Alpha = 1.0f;
+		PointIndex = FMath::Max(0, InterpPoints.Num() - 2);
+	}
 
 	check(InterpPoints.IsValidIndex(PointIndex));
 	const FLandscapeSplineInterpPoint& PreviousPoint = InterpPoints[PointIndex];
@@ -131,14 +155,12 @@ FTransform UPCGLandscapeSplineData::GetTransformAtDistance(int SegmentIndex, FVe
 
 FVector::FReal UPCGLandscapeSplineData::GetCurvatureAtDistance(int SegmentIndex, FVector::FReal Distance) const
 {
-	check(Spline.IsValid());
-
-	if (GetNumSegments() == 0)
+	if (!CheckSpline() || GetNumSegments() == 0)
 	{
 		return 0.0;
 	}
 
-	check(Spline->GetSegments().IsValidIndex(SegmentIndex));
+	const int ClampedSegmentIndex = FMath::Clamp(SegmentIndex, 0, GetNumSegments()-1);
 
 	const ULandscapeSplineSegment* Segment = Spline->GetSegments()[SegmentIndex];
 	check(Segment);
@@ -151,9 +173,17 @@ FVector::FReal UPCGLandscapeSplineData::GetCurvatureAtDistance(int SegmentIndex,
 		return 0;
 	}
 
-	int32 PointIndex;
-	FVector::FReal Alpha;
-	GetInterpPointAtDistance(SegmentIndex, Distance, PointIndex, /*bComputeAlpha=*/false, Alpha);
+	int32 PointIndex = 0;
+
+	if (ClampedSegmentIndex == SegmentIndex)
+	{
+		FVector::FReal Alpha;
+		GetInterpPointAtDistance(ClampedSegmentIndex, Distance, PointIndex, /*bComputeAlpha=*/false, Alpha);
+	}
+	else if (ClampedSegmentIndex < SegmentIndex)
+	{
+		PointIndex = FMath::Max(0, InterpPoints.Num() - 2);
+	}
 
 	// If our sample overshoots the segment, clamp it back to the last point.
 	if (PointIndex == InterpPoints.Num() - 1)
@@ -208,12 +238,15 @@ float UPCGLandscapeSplineData::GetInputKeyAtDistance(int SegmentIndex, FVector::
 
 void UPCGLandscapeSplineData::GetTangentsAtSegmentStart(int SegmentIndex, FVector& OutArriveTangent, FVector& OutLeaveTangent) const
 {
-	check(Spline.IsValid());
+	if (!CheckSpline())
+	{
+		return;
+	}
 
 	const TArray<TObjectPtr<ULandscapeSplineSegment>>& Segments = Spline->GetSegments();
 
-	const TObjectPtr<ULandscapeSplineSegment> PreviousSegment = SegmentIndex > 0 ? Segments[SegmentIndex - 1] : nullptr;
-	const TObjectPtr<ULandscapeSplineSegment> CurrentSegment = SegmentIndex < Segments.Num() ? Segments[SegmentIndex] : nullptr;
+	const TObjectPtr<ULandscapeSplineSegment> PreviousSegment = SegmentIndex > 0 && Segments.IsValidIndex(SegmentIndex - 1) ? Segments[SegmentIndex - 1] : nullptr;
+	const TObjectPtr<ULandscapeSplineSegment> CurrentSegment = Segments.IsValidIndex(SegmentIndex) ? Segments[SegmentIndex] : nullptr;
 
 	// Arrive tangent lives in the end-point of the previous segment
 	const TObjectPtr<ULandscapeSplineControlPoint> ArrivePoint = PreviousSegment ? PreviousSegment->Connections[1].ControlPoint : nullptr;
@@ -277,8 +310,6 @@ const UPCGPointArrayData* UPCGLandscapeSplineData::CreatePointArrayData(FPCGCont
 
 const UPCGBasePointData* UPCGLandscapeSplineData::CreateBasePointData(FPCGContext* Context, TSubclassOf<UPCGBasePointData> PointDataClass) const
 {
-	check(Spline.IsValid());
-
 	UPCGBasePointData* Data = FPCGContext::NewObject_AnyThread<UPCGBasePointData>(Context, GetTransientPackage(), PointDataClass);
 
 	FPCGInitializeFromDataParams InitializeFromDataParams(this);
@@ -291,16 +322,23 @@ const UPCGBasePointData* UPCGLandscapeSplineData::CreateBasePointData(FPCGContex
 
 	PCGSplineSamplerHelpers::SampleLineData(Context, /*LineData=*/this, /*InBoundingShapeData=*/nullptr, /*InProjectionTarget=*/nullptr, /*InProjectionParams=*/{}, SamplerParams, Data);
 
-	UE_LOG(LogPCG, Verbose, TEXT("Landscape spline %s generated %d points"), *Spline->GetFName().ToString(), Data->GetNumPoints());
+	if (CheckSpline())
+	{
+		UE_LOG(LogPCG, Verbose, TEXT("Landscape spline %s generated %d points"), *Spline->GetFName().ToString(), Data->GetNumPoints());
+	}
 
 	return Data;
 }
 
 FBox UPCGLandscapeSplineData::GetBounds() const
 {
-	check(Spline.IsValid());
-
 	FBox Bounds(EForceInit::ForceInit);
+
+	if (!CheckSpline())
+	{		
+		return Bounds;
+	}
+
 	for (const TObjectPtr<ULandscapeSplineSegment>& Segment : Spline->GetSegments())
 	{
 		Bounds += Segment->GetBounds();
@@ -321,9 +359,12 @@ bool UPCGLandscapeSplineData::SamplePoint(const FTransform& InTransform, const F
 
 	// This does not move the query point, but does not take into account the Z axis at all - so this is inherently a projection. There's some
 	// things that need double checking but we should look at incorporating the Z axis into the calculation.
+	if (!CheckSpline())
+	{
+		return false;
+	}
 
 	const ULandscapeSplinesComponent* SplinePtr = Spline.Get();
-	check(SplinePtr);
 
 	OutPoint.Transform = InTransform;
 	OutPoint.SetLocalBounds(InBounds); // TODO: should maybe do Min.Z = Max.Z = 0 ?
@@ -390,7 +431,10 @@ UPCGSpatialData* UPCGLandscapeSplineData::CopyInternal(FPCGContext* Context) con
 
 void UPCGLandscapeSplineData::UpdateReparamTable()
 {
-	check(Spline.IsValid());
+	if (!CheckSpline())
+	{
+		return;
+	}
 
 	const TArray<TObjectPtr<ULandscapeSplineSegment>>& Segments = Spline->GetSegments();
 	const int32 NumSegments = Segments.Num();
diff --git a/Engine/Plugins/PCG/Source/PCG/Public/Data/PCGLandscapeSplineData.h b/Engine/Plugins/PCG/Source/PCG/Public/Data/PCGLandscapeSplineData.h
--- a/Engine/Plugins/PCG/Source/PCG/Public/Data/PCGLandscapeSplineData.h
+++ b/Engine/Plugins/PCG/Source/PCG/Public/Data/PCGLandscapeSplineData.h
@@ -67,6 +67,9 @@ protected:
 
 	/** Reparameterization of the spline by distance. Useful to query the InputKey at arbitrary distance. */
 	FInterpCurveFloat ReparamTable;
+
+	bool CheckSpline() const;
+	mutable bool bLoggedInvalidSpline = false;
 };
 
 #undef UE_API
