[PCG] Removed cog from the graph parameters dialog as it will not do what we are expecting.

#rnx

diff --git a/Engine/Plugins/PCG/Source/PCGEditor/Private/Widgets/SPCGEditorGraphParamsView.cpp b/Engine/Plugins/PCG/Source/PCGEditor/Private/Widgets/SPCGEditorGraphParamsView.cpp
--- a/Engine/Plugins/PCG/Source/PCGEditor/Private/Widgets/SPCGEditorGraphParamsView.cpp
+++ b/Engine/Plugins/PCG/Source/PCGEditor/Private/Widgets/SPCGEditorGraphParamsView.cpp
@@ -34,7 +34,7 @@ void SPCGEditorGraphUserParametersView::Construct(const FArguments& InArgs, cons
 		DetailsViewArgs.DefaultsOnlyVisibility = EEditDefaultsOnlyNodeVisibility::Automatic;
 		DetailsViewArgs.bAllowSearch = true;
 		DetailsViewArgs.bShowScrollBar = true;
-		DetailsViewArgs.bShowOptions = true;
+		DetailsViewArgs.bShowOptions = false;
 		DetailsViewArgs.bShowObjectLabel = true;
 
 		DetailsViewArgs.bShowModifiedPropertiesOption = false;
