[PCG] Refactor how we initialize metadata when collapsing.

* Add a new InitializeAnotherMetadata that by default it just initialize/add attributes to the metadata.
* For composite data, it use the inner data for initialisation
* For surface data, it will add the attribute that were previously created on the surface data itself (sampling).
* Changed how metadata domain are traversed when queried from the Metadata, default is now skipped
* This forces a change on the GPU attribute binding, as the identifier was now Elements instead of Default. So force it back to default for now when adding the attribute to GPU data binding.

diff --git a/Engine/Plugins/PCG/Source/PCG/Private/Compute/PCGDataBinding.cpp b/Engine/Plugins/PCG/Source/PCG/Private/Compute/PCGDataBinding.cpp
--- a/Engine/Plugins/PCG/Source/PCG/Private/Compute/PCGDataBinding.cpp
+++ b/Engine/Plugins/PCG/Source/PCG/Private/Compute/PCGDataBinding.cpp
@@ -334,7 +334,7 @@ void UPCGDataBinding::AddInputDataAttributesToTable()
 		}
 		else if (const UPCGMetadata* Metadata = Data.Data ? Data.Data->ConstMetadata() : nullptr)
 		{
-			const FPCGMetadataDomainID DefaultDomain = Metadata->GetConstDefaultMetadataDomain()->GetDomainID();
+			const FPCGMetadataDomainID MetadataDefaultDomainID = Metadata->GetConstDefaultMetadataDomain()->GetDomainID();
 			
 			TArray<FPCGAttributeIdentifier> AttributeIdentifiers;
 			TArray<EPCGMetadataTypes> AttributeTypes;
@@ -346,11 +346,17 @@ void UPCGDataBinding::AddInputDataAttributesToTable()
 
 			for (int Index = 0; Index < AttributeIdentifiers.Num(); ++Index)
 			{
-				const FPCGAttributeIdentifier AttributeIdentifier = AttributeIdentifiers[Index];
-				if (!AttributeIdentifier.MetadataDomain.IsDefault() && AttributeIdentifier.MetadataDomain != DefaultDomain)
+				FPCGAttributeIdentifier AttributeIdentifier = AttributeIdentifiers[Index];
+				if (AttributeIdentifier.MetadataDomain != PCGMetadataDomainID::Default && AttributeIdentifier.MetadataDomain != MetadataDefaultDomainID)
 				{
 					continue;
 				}
+
+				// If the domain is the default domain, force it to the default identifier.
+				if (AttributeIdentifier.MetadataDomain == MetadataDefaultDomainID)
+				{
+					AttributeIdentifier.MetadataDomain = PCGMetadataDomainID::Default;
+				}
 					
 				const EPCGKernelAttributeType AttributeType = PCGDataForGPUHelpers::GetAttributeTypeFromMetadataType(AttributeTypes[Index]);
 
@@ -411,7 +417,7 @@ void UPCGDataBinding::AddInputDataStringsToTable()
 				continue;
 			}
 			
-			const FPCGMetadataDomainID DefaultDomain = Metadata->GetConstDefaultMetadataDomain()->GetDomainID();
+			const FPCGMetadataDomainID MetadataDefaultDomainID = Metadata->GetConstDefaultMetadataDomain()->GetDomainID();
 
 			TArray<FPCGAttributeIdentifier> AttributeIdentifiers;
 			TArray<EPCGMetadataTypes> AttributeTypes;
@@ -425,11 +431,17 @@ void UPCGDataBinding::AddInputDataStringsToTable()
 				// @todo_pcg: Attributes on other domains than the default are ignored at the moment, until we have a better way of representing
 				// different domains in the GPU header.
 				// It means those are lost.
-				const FPCGAttributeIdentifier AttributeIdentifier = AttributeIdentifiers[AttributeIndex];
-				if (!AttributeIdentifier.MetadataDomain.IsDefault() && AttributeIdentifier.MetadataDomain != DefaultDomain)
+				FPCGAttributeIdentifier AttributeIdentifier = AttributeIdentifiers[AttributeIndex];
+				if (AttributeIdentifier.MetadataDomain != PCGMetadataDomainID::Default && AttributeIdentifier.MetadataDomain != MetadataDefaultDomainID)
 				{
 					continue;
 				}
+
+				// If the domain is the default domain, force it to the default identifier.
+				if (AttributeIdentifier.MetadataDomain == MetadataDefaultDomainID)
+				{
+					AttributeIdentifier.MetadataDomain = PCGMetadataDomainID::Default;
+				}
 				
 				const EPCGKernelAttributeType AttributeType = PCGDataForGPUHelpers::GetAttributeTypeFromMetadataType(AttributeTypes[AttributeIndex]);
 
diff --git a/Engine/Plugins/PCG/Source/PCG/Private/Compute/PCGDataForGPU.cpp b/Engine/Plugins/PCG/Source/PCG/Private/Compute/PCGDataForGPU.cpp
--- a/Engine/Plugins/PCG/Source/PCG/Private/Compute/PCGDataForGPU.cpp
+++ b/Engine/Plugins/PCG/Source/PCG/Private/Compute/PCGDataForGPU.cpp
@@ -851,7 +851,7 @@ void FPCGDataDesc::InitializeAttributeDescs(const UPCGData* InData, const UPCGDa
 
 	if (InBinding && Metadata)
 	{
-		const FPCGMetadataDomainID DefaultDomain = Metadata->GetConstDefaultMetadataDomain()->GetDomainID();
+		const FPCGMetadataDomainID MetadataDefaultDomainID = Metadata->GetConstDefaultMetadataDomain()->GetDomainID();
 		
 		TArray<FPCGAttributeIdentifier> AttributeIdentifiers;
 		TArray<EPCGMetadataTypes> AttributeTypes;
@@ -865,11 +865,17 @@ void FPCGDataDesc::InitializeAttributeDescs(const UPCGData* InData, const UPCGDa
 			// @todo_pcg: Attributes on other domains than the default are ignored at the moment, until we have a better way of representing
 			// different domains in the GPU header.
 			// It means those are lost.
-			const FPCGAttributeIdentifier AttributeIdentifier = AttributeIdentifiers[CustomAttributeIndex];
-			if (!AttributeIdentifier.MetadataDomain.IsDefault() && AttributeIdentifier.MetadataDomain != DefaultDomain)
+			FPCGAttributeIdentifier AttributeIdentifier = AttributeIdentifiers[CustomAttributeIndex];
+			if (AttributeIdentifier.MetadataDomain != PCGMetadataDomainID::Default && AttributeIdentifier.MetadataDomain != MetadataDefaultDomainID)
 			{
 				continue;
 			}
+
+			// If the domain is the default domain, force it to the default identifier.
+			if (AttributeIdentifier.MetadataDomain == MetadataDefaultDomainID)
+			{
+				AttributeIdentifier.MetadataDomain = PCGMetadataDomainID::Default;
+			}
 			
 			const EPCGKernelAttributeType AttributeType = PCGDataForGPUHelpers::GetAttributeTypeFromMetadataType(AttributeTypes[CustomAttributeIndex]);
 
diff --git a/Engine/Plugins/PCG/Source/PCG/Private/Data/PCGDifferenceData.cpp b/Engine/Plugins/PCG/Source/PCG/Private/Data/PCGDifferenceData.cpp
--- a/Engine/Plugins/PCG/Source/PCG/Private/Data/PCGDifferenceData.cpp
+++ b/Engine/Plugins/PCG/Source/PCG/Private/Data/PCGDifferenceData.cpp
@@ -156,6 +156,7 @@ void UPCGDifferenceData::AddToCrc(FArchiveCrc32& Ar, bool bFullDataCrc) const
 	Super::AddToCrc(Ar, bFullDataCrc);
 
 	// Implementation note: no metadata in composite data at this point.
+	// @todo_pcg: need metadata serialization now
 
 	FString ClassName = StaticClass()->GetPathName();
 	Ar << ClassName;
@@ -261,8 +262,10 @@ const UPCGBasePointData* UPCGDifferenceData::CreateBasePointData(FPCGContext* Co
 	UPCGBasePointData* TargetPointData = FPCGContext::NewObject_AnyThread<UPCGBasePointData>(Context, GetTransientPackage(), PointDataClass);
 	const UPCGMetadata* SourceMetadata = SourcePointData->ConstMetadata();
 
-	FPCGInitializeFromDataParams InitializeFromDataParams(SourcePointData);
+	FPCGInitializeFromDataParams InitializeFromDataParams(this);
 	InitializeFromDataParams.bInheritSpatialData = false;
+	// Since we have collapsed the source data, we need to inherit from this one.
+	InitializeFromDataParams.SourceOverride = SourcePointData;
 	TargetPointData->InitializeFromDataWithParams(InitializeFromDataParams);
 
 	UPCGMetadata* OutMetadata = TargetPointData->Metadata;
@@ -271,7 +274,7 @@ const UPCGBasePointData* UPCGDifferenceData::CreateBasePointData(FPCGContext* Co
 	if (bDiffMetadata && OutMetadata)
 	{
 		TempDiffMetadata = FPCGContext::NewObject_AnyThread<UPCGMetadata>(Context);
-		DifferenceData->InitializeOtherMetadata(FPCGInitializeFromDataParams{DifferenceData}, TempDiffMetadata);
+		DifferenceData->InitializeTargetMetadata(FPCGInitializeFromDataParams{DifferenceData}, TempDiffMetadata);
 	}
 
 	constexpr int ChunkSize = FPCGSpatialDataProcessing::DefaultSamplePointsChunkSize;
@@ -390,7 +393,7 @@ UPCGSpatialData* UPCGDifferenceData::CopyInternal(FPCGContext* Context) const
 	return NewDifferenceData;
 }
 
-void UPCGDifferenceData::InitializeOtherMetadataInternal(const FPCGInitializeFromDataParams& InParams, UPCGMetadata* MetadataToInitialize) const
+void UPCGDifferenceData::InitializeTargetMetadata(const FPCGInitializeFromDataParams& InParams, UPCGMetadata* MetadataToInitialize) const
 {
 	check(InParams.bInheritMetadata);
 	check(MetadataToInitialize);
@@ -398,7 +401,7 @@ void UPCGDifferenceData::InitializeOtherMetadataInternal(const FPCGInitializeFro
 	// Duplicate data case, call the spatial base method
 	if (InParams.bIsDuplicatingData)
 	{
-		UPCGSpatialData::InitializeOtherMetadataInternal(InParams, MetadataToInitialize);
+		UPCGSpatialData::InitializeTargetMetadata(InParams, MetadataToInitialize);
 		return;
 	}
 
@@ -407,8 +410,10 @@ void UPCGDifferenceData::InitializeOtherMetadataInternal(const FPCGInitializeFro
 	if (Source)
 	{
 		FPCGInitializeFromDataParams CopyParams = InParams;
-		CopyParams.Source = Source;
-		Source->InitializeOtherMetadata(InParams, MetadataToInitialize);
+		CopyParams.SourceOverride = nullptr;
+
+		CopyParams.Source = InParams.SourceOverride ? InParams.SourceOverride : Source;
+		CopyParams.Source->InitializeTargetMetadata(CopyParams, MetadataToInitialize);
 	}
 
 	MetadataToInitialize->AddAttributes(Metadata);
diff --git a/Engine/Plugins/PCG/Source/PCG/Private/Data/PCGIntersectionData.cpp b/Engine/Plugins/PCG/Source/PCG/Private/Data/PCGIntersectionData.cpp
--- a/Engine/Plugins/PCG/Source/PCG/Private/Data/PCGIntersectionData.cpp
+++ b/Engine/Plugins/PCG/Source/PCG/Private/Data/PCGIntersectionData.cpp
@@ -83,6 +83,7 @@ void UPCGIntersectionData::AddToCrc(FArchiveCrc32& Ar, bool bFullDataCrc) const
 	Super::AddToCrc(Ar, bFullDataCrc);
 
 	// Implementation note: no metadata in composite data at this point.
+	// @todo_pcg: need metadata serialization now
 
 	FString ClassName = StaticClass()->GetPathName();
 	Ar << ClassName;
@@ -220,20 +221,19 @@ UPCGBasePointData* UPCGIntersectionData::CreateAndFilterPointData(FPCGContext* C
 	UPCGBasePointData* OutputData = FPCGContext::NewObject_AnyThread<UPCGBasePointData>(Context, GetTransientPackage(), PointDataClass);
 	UPCGMetadata* OutputMetadata = OutputData->MutableMetadata();
 
-	FPCGInitializeFromDataParams InitializeFromDataParams(SourcePointData);
+	FPCGInitializeFromDataParams InitializeFromDataParams(this);
 	InitializeFromDataParams.bInheritSpatialData = false;
+	// Since we have collapsed the source data, we need to inherit from this one.
+	InitializeFromDataParams.SourceOverride = SourcePointData;
 	OutputData->InitializeFromDataWithParams(InitializeFromDataParams);
 
 	UPCGMetadata* TempYMetadata = nullptr;
 	if (Y && Y->ConstMetadata())
 	{
 		TempYMetadata = FPCGContext::NewObject_AnyThread<UPCGMetadata>(Context);
-		Y->InitializeOtherMetadata(FPCGInitializeFromDataParams{Y}, TempYMetadata);
-		OutputMetadata->AddAttributes(TempYMetadata);
+		Y->InitializeTargetMetadata(FPCGInitializeFromDataParams{Y}, TempYMetadata);
 	}
 
-	OutputMetadata->AddAttributes(ThisMetadata);
-
 	const bool bPointDataHasCommonAttributes = (SourceMetadata && TempYMetadata && SourceMetadata->HasCommonAttributes(TempYMetadata));
 
 	constexpr int ChunkSize = FPCGSpatialDataProcessing::DefaultSamplePointsChunkSize;
@@ -343,35 +343,34 @@ UPCGSpatialData* UPCGIntersectionData::CopyInternal(FPCGContext* Context) const
 	return NewIntersectionData;
 }
 
-void UPCGIntersectionData::InitializeOtherMetadataInternal(const FPCGInitializeFromDataParams& InParams, UPCGMetadata* MetadataToInitialize) const
+void UPCGIntersectionData::InitializeTargetMetadata(const FPCGInitializeFromDataParams& InParams, UPCGMetadata* MetadataToInitialize) const
 {
 	check(InParams.bInheritMetadata);
 	check(MetadataToInitialize);
-	check(GetA());
-	check(GetB());
 
 	// Duplicate data case, call the spatial base method
 	if (InParams.bIsDuplicatingData)
 	{
-		UPCGSpatialData::InitializeOtherMetadataInternal(InParams, MetadataToInitialize);
+		UPCGSpatialData::InitializeTargetMetadata(InParams, MetadataToInitialize);
 		return;
 	}
 
+	check(GetA());
+	check(GetB());
+
 	// In the case of the intersection, we initialize the data from either A or B, depending on the dimension, not the intersection itself. The intersection will add its own
 	// attributes afterward.
 	FPCGInitializeFromDataParams ParamsCopy = InParams;
-	if (GetA()->GetDimension() <= GetB()->GetDimension())
-	{
-		ParamsCopy.Source = GetA();
-		GetA()->InitializeOtherMetadata(ParamsCopy, MetadataToInitialize);
-		MetadataToInitialize->AddAttributes(GetB()->Metadata);
-	}
-	else
-	{
-		ParamsCopy.Source = GetB();
-		GetB()->InitializeOtherMetadata(ParamsCopy, MetadataToInitialize);
-		MetadataToInitialize->AddAttributes(GetA()->Metadata);
-	}
+	ParamsCopy.SourceOverride = nullptr;
+
+	const UPCGSpatialData* SourceData = (GetA()->GetDimension() <= GetB()->GetDimension()) ? GetA() : GetB();
+	const UPCGSpatialData* OtherData = SourceData == GetA() ? GetB() : GetA();
+
+	ParamsCopy.Source = InParams.SourceOverride ? InParams.SourceOverride.Get() : SourceData;
+	ParamsCopy.Source->InitializeTargetMetadata(ParamsCopy, MetadataToInitialize);
+
+	ParamsCopy.Source = OtherData;
+	ParamsCopy.Source->InitializeTargetMetadata(ParamsCopy, MetadataToInitialize);
 
 	MetadataToInitialize->AddAttributes(Metadata);
 }
diff --git a/Engine/Plugins/PCG/Source/PCG/Private/Data/PCGLandscapeData.cpp b/Engine/Plugins/PCG/Source/PCG/Private/Data/PCGLandscapeData.cpp
--- a/Engine/Plugins/PCG/Source/PCG/Private/Data/PCGLandscapeData.cpp
+++ b/Engine/Plugins/PCG/Source/PCG/Private/Data/PCGLandscapeData.cpp
@@ -77,47 +77,6 @@ void UPCGLandscapeData::Initialize(const TArray<TWeakObjectPtr<ALandscapeProxy>>
 	// Store cache pointer for easier access
 	UPCGSubsystem* PCGSubsystem = UPCGSubsystem::GetInstance(FirstLandscape->GetWorld());
 	LandscapeCache = PCGSubsystem ? PCGSubsystem->GetLandscapeCache() : nullptr;
-
-	// TODO: find a better way to do this - maybe there should be a prototype metadata in the landscape cache
-	if (LandscapeCache)
-	{
-		if (DataProps.bGetLayerWeights)
-		{
-			for (TSoftObjectPtr<ALandscapeProxy> Landscape : Landscapes)
-			{
-				const TArray<FName> Layers = LandscapeCache->GetLayerNames(Landscape.Get());
-
-				for (const FName& Layer : Layers)
-				{
-					if (!Metadata->HasAttribute(Layer))
-					{
-						Metadata->CreateFloatAttribute(Layer, 0.0f, /*bAllowInterpolation=*/true);
-					}
-				}
-			}
-		}
-	}
-	else
-	{
-		UE_LOG(LogPCG, Error, TEXT("Landscape data is unable to access the landscape cache (Note: Will happen if there is no PCG world actor."));
-	}
-
-	// Create secondary attributes as we need them
-	if (DataProps.bGetActorReference)
-	{
-		Metadata->CreateAttribute<FSoftObjectPath>(PCGPointDataConstants::ActorReferenceAttribute, FSoftObjectPath(), /*bAllowInterpolation=*/false, /*bOverrideParent=*/false);
-	}
-
-	if (DataProps.bGetPhysicalMaterial)
-	{
-		Metadata->CreateAttribute<FSoftObjectPath>(PCGWorldQueryConstants::PhysicalMaterialReferenceAttribute, FSoftObjectPath(), /*bAllowInterpolation=*/false, /*bOverrideParent*/false);
-	}
-
-	if (DataProps.bGetComponentCoordinates)
-	{
-		Metadata->CreateInteger32Attribute(PCGLandscapeDataConstants::ComponentXAttribute, 0, /*bAllowsInterpolation=*/false);
-		Metadata->CreateInteger32Attribute(PCGLandscapeDataConstants::ComponentYAttribute, 0, /*bAllowsInterpolation=*/false);
-	}
 }
 
 void UPCGLandscapeData::PostLoad()
@@ -404,7 +363,7 @@ bool UPCGLandscapeData::ProjectPoint(const FTransform& InTransform, const FBox&
 
 	if (DataProps.bGetActorReference && OutMetadata && LandscapeCollisionComponent)
 	{
-		if (FPCGMetadataAttribute<FSoftObjectPath>* ActorReferenceAttribute = OutMetadata->GetMutableTypedAttribute<FSoftObjectPath>(PCGPointDataConstants::ActorReferenceAttribute))
+		if (FPCGMetadataAttribute<FSoftObjectPath>* ActorReferenceAttribute = OutMetadata->FindOrCreateAttribute<FSoftObjectPath>(PCGPointDataConstants::ActorReferenceAttribute))
 		{
 			// Landscape code seems to indicate the XYtoComponentMap can be sometimes invalid, so rely on the collision map instead
 			OutMetadata->InitializeOnSet(OutPoint.MetadataEntry);
@@ -414,7 +373,7 @@ bool UPCGLandscapeData::ProjectPoint(const FTransform& InTransform, const FBox&
 
 	if (DataProps.bGetPhysicalMaterial && OutMetadata && LandscapeCollisionComponent)
 	{
-		if (FPCGMetadataAttribute<FSoftObjectPath>* PhysicalMaterialAttribute = OutMetadata->GetMutableTypedAttribute<FSoftObjectPath>(PCGWorldQueryConstants::PhysicalMaterialReferenceAttribute))
+		if (FPCGMetadataAttribute<FSoftObjectPath>* PhysicalMaterialAttribute = OutMetadata->FindOrCreateAttribute<FSoftObjectPath>(PCGWorldQueryConstants::PhysicalMaterialReferenceAttribute))
 		{
 			if(UPhysicalMaterial* PhysicalMaterial = LandscapeCollisionComponent->GetPhysicalMaterial(static_cast<float>(ComponentLocalPoint.X), static_cast<float>(ComponentLocalPoint.Y), EHeightfieldSource::Complex))
 			{
@@ -426,13 +385,13 @@ bool UPCGLandscapeData::ProjectPoint(const FTransform& InTransform, const FBox&
 
 	if (DataProps.bGetComponentCoordinates && OutMetadata)
 	{
-		if(FPCGMetadataAttribute<int32>* ComponentXAttribute = OutMetadata->GetMutableTypedAttribute<int32>(PCGLandscapeDataConstants::ComponentXAttribute))
+		if(FPCGMetadataAttribute<int32>* ComponentXAttribute = OutMetadata->FindOrCreateAttribute<int32>(PCGLandscapeDataConstants::ComponentXAttribute))
 		{
 			OutMetadata->InitializeOnSet(OutPoint.MetadataEntry);
 			ComponentXAttribute->SetValue(OutPoint.MetadataEntry, ComponentMapKey.X);
 		}
 
-		if (FPCGMetadataAttribute<int32>* ComponentYAttribute = OutMetadata->GetMutableTypedAttribute<int32>(PCGLandscapeDataConstants::ComponentYAttribute))
+		if (FPCGMetadataAttribute<int32>* ComponentYAttribute = OutMetadata->FindOrCreateAttribute<int32>(PCGLandscapeDataConstants::ComponentYAttribute))
 		{
 			OutMetadata->InitializeOnSet(OutPoint.MetadataEntry);
 			ComponentYAttribute->SetValue(OutPoint.MetadataEntry, ComponentMapKey.Y);
@@ -508,7 +467,7 @@ const UPCGBasePointData* UPCGLandscapeData::CreateBasePointData(FPCGContext* Con
 		return Data;
 	}
 
-	UPCGMetadata* OutMetadata = DataProps.bGetLayerWeights ? Data->Metadata : nullptr;
+	UPCGMetadata* OutMetadata = UseMetadata() ? Data->Metadata : nullptr;
 	if (OutMetadata)
 	{
 		Data->AllocateProperties(EPCGPointNativeProperties::MetadataEntry);
@@ -616,6 +575,11 @@ const UPCGBasePointData* UPCGLandscapeData::CreateBasePointData(FPCGContext* Con
 	return Data;
 }
 
+bool UPCGLandscapeData::UseMetadata() const
+{
+	return DataProps.bGetActorReference || DataProps.bGetComponentCoordinates || DataProps.bGetLayerWeights || DataProps.bGetPhysicalMaterial;
+}
+
 const ULandscapeInfo* UPCGLandscapeData::GetLandscapeInfo(const FVector& InPosition) const
 {
 	check(!LandscapeInfos.IsEmpty());
@@ -714,4 +678,63 @@ TArray<FPCGTaskId> UPCGLandscapeData::PrepareForSpatialQuery(FPCGContext* InCont
 #endif
 
 	return { InContext->ScheduleGeneric(Params) };
-}
\ No newline at end of file
+}
+
+void UPCGLandscapeData::InitializeTargetMetadata(const FPCGInitializeFromDataParams& InParams, UPCGMetadata* MetadataToInitialize) const
+{
+	check(MetadataToInitialize);
+
+	// Initialize the new metadata normally
+	Super::InitializeTargetMetadata(InParams, MetadataToInitialize);
+
+	if (!UseMetadata())
+	{
+		// If we have no extra metadata, early out.
+		return;
+	}
+
+	// Only add those special attributes if the metadata to initialize support elements domain
+	FPCGMetadataDomain* MetadataDomain = MetadataToInitialize->GetMetadataDomain(PCGMetadataDomainID::Elements);
+	if (!MetadataDomain)
+	{
+		return;
+	}
+
+	// TODO: find a better way to do this - maybe there should be a prototype metadata in the landscape cache
+	if (LandscapeCache)
+	{
+		if (DataProps.bGetLayerWeights)
+		{
+			for (TSoftObjectPtr<ALandscapeProxy> Landscape : Landscapes)
+			{
+				const TArray<FName> Layers = LandscapeCache->GetLayerNames(Landscape.Get());
+
+				for (const FName& Layer : Layers)
+				{
+					MetadataDomain->FindOrCreateAttribute<float>(Layer, 0.0f, /*bAllowInterpolation=*/true);
+				}
+			}
+		}
+	}
+	else
+	{
+		UE_LOG(LogPCG, Error, TEXT("Landscape data is unable to access the landscape cache (Note: Will happen if there is no PCG world actor."));
+	}
+
+	// Create secondary attributes as we need them
+	if (DataProps.bGetActorReference)
+	{
+		MetadataDomain->FindOrCreateAttribute<FSoftObjectPath>(PCGPointDataConstants::ActorReferenceAttribute, FSoftObjectPath(), /*bAllowInterpolation=*/false, /*bOverrideParent=*/false);
+	}
+
+	if (DataProps.bGetPhysicalMaterial)
+	{
+		MetadataDomain->FindOrCreateAttribute<FSoftObjectPath>(PCGWorldQueryConstants::PhysicalMaterialReferenceAttribute, FSoftObjectPath(), /*bAllowInterpolation=*/false, /*bOverrideParent*/false);
+	}
+
+	if (DataProps.bGetComponentCoordinates)
+	{
+		MetadataDomain->FindOrCreateAttribute<int32>(PCGLandscapeDataConstants::ComponentXAttribute, 0, /*bAllowsInterpolation=*/false);
+		MetadataDomain->FindOrCreateAttribute<int32>(PCGLandscapeDataConstants::ComponentYAttribute, 0, /*bAllowsInterpolation=*/false);
+	}
+}
diff --git a/Engine/Plugins/PCG/Source/PCG/Private/Data/PCGProjectionData.cpp b/Engine/Plugins/PCG/Source/PCG/Private/Data/PCGProjectionData.cpp
--- a/Engine/Plugins/PCG/Source/PCG/Private/Data/PCGProjectionData.cpp
+++ b/Engine/Plugins/PCG/Source/PCG/Private/Data/PCGProjectionData.cpp
@@ -93,6 +93,7 @@ void UPCGProjectionData::AddToCrc(FArchiveCrc32& Ar, bool bFullDataCrc) const
 	Super::AddToCrc(Ar, bFullDataCrc);
 
 	// Implementation note: no metadata in composite data at this point.
+	// @todo_pcg: need metadata serialization now
 
 	FString ClassName = StaticClass()->GetPathName();
 	Ar << ClassName;
@@ -277,6 +278,8 @@ const UPCGBasePointData* UPCGProjectionData::CreateBasePointData(FPCGContext* Co
 	// Copy metadata attributes from source point including values
 	FPCGInitializeFromDataParams InitializeFromDataParams(this);
 	InitializeFromDataParams.bInheritSpatialData = false;
+	// Since we have collapsed the source data, we need to inherit from this one.
+	InitializeFromDataParams.SourceOverride = SourcePointData;
 	PointData->InitializeFromDataWithParams(InitializeFromDataParams);
 	UPCGMetadata* OutMetadata = PointData->Metadata;
 	check(OutMetadata);
@@ -286,6 +289,7 @@ const UPCGBasePointData* UPCGProjectionData::CreateBasePointData(FPCGContext* Co
 	if (Target->Metadata)
 	{
 		TempTargetMetadata = FPCGContext::NewObject_AnyThread<UPCGMetadata>(Context);
+		TempTargetMetadata->SetupDomainsFromPCGDataType<UPCGBasePointData>();
 
 		// We achieve filtering of metadata attributes by manipulating this temporary metadata, which works because
 		// the projection operation operates on the attributes in this metadata.
@@ -295,16 +299,9 @@ const UPCGBasePointData* UPCGProjectionData::CreateBasePointData(FPCGContext* Co
 		// * An excluded attribute that does not exist on source data will not be kept in result
 		// * Included attributes are the only attributes that can be changed during projection
 		// * Included attributes are the only attributes that will be added from target data
-
-		TSet<FName> AttributesFilter;
-		GetIncludeExcludeAttributeNames(AttributesFilter);
-
-		TempTargetMetadata->InitializeWithAttributeFilter(Target->Metadata, AttributesFilter, ProjectionParams.AttributeMode);
+		SetupTargetMetadata(TempTargetMetadata);
 	}
 
-	// Add any attributes from filtered metadata to produce final list
-	OutMetadata->AddAttributes(TempTargetMetadata);
-
 	auto InitializeFunc = [PointData, SourcePointData]()
 	{
 		PointData->SetNumPoints(SourcePointData->GetNumPoints(), /*bInitializeValues=*/false);
@@ -452,7 +449,26 @@ UPCGSpatialData* UPCGProjectionData::CopyInternal(FPCGContext* Context) const
 	return NewProjectionData;
 }
 
-void UPCGProjectionData::InitializeOtherMetadataInternal(const FPCGInitializeFromDataParams& InParams, UPCGMetadata* MetadataToInitialize) const
+void UPCGProjectionData::SetupTargetMetadata(UPCGMetadata* MetadataToInitialize) const
+{
+	if (!Target || !Target->ConstMetadata())
+	{
+		return;
+	}
+	
+	FPCGInitializeFromDataParams TargetInitializeFromDataParams(Target);
+	
+	TSet<FName> AttributesFilter;
+	GetIncludeExcludeAttributeNames(AttributesFilter);
+
+	FPCGMetadataDomainInitializeParams MetadataParams(nullptr, &AttributesFilter);
+	MetadataParams.FilterMode = ProjectionParams.AttributeMode;
+	TargetInitializeFromDataParams.MetadataInitializeParams.DomainInitializeParams.Emplace(PCGMetadataDomainID::Elements, std::move(MetadataParams));
+
+	Target->InitializeTargetMetadata(TargetInitializeFromDataParams, MetadataToInitialize);
+}
+
+void UPCGProjectionData::InitializeTargetMetadata(const FPCGInitializeFromDataParams& InParams, UPCGMetadata* MetadataToInitialize) const
 {
 	check(InParams.bInheritMetadata);
 	check(MetadataToInitialize);
@@ -460,16 +476,18 @@ void UPCGProjectionData::InitializeOtherMetadataInternal(const FPCGInitializeFro
 	// Duplicate data case, call the spatial base method
 	if (InParams.bIsDuplicatingData)
 	{
-		UPCGSpatialData::InitializeOtherMetadataInternal(InParams, MetadataToInitialize);
+		UPCGSpatialData::InitializeTargetMetadata(InParams, MetadataToInitialize);
 		return;
 	}
-	
+
 	if (Source)
 	{
 		FPCGInitializeFromDataParams ParamsCopy = InParams;
-		ParamsCopy.Source = Source;
-		Source->InitializeOtherMetadata(ParamsCopy, MetadataToInitialize);
+		ParamsCopy.Source = InParams.SourceOverride ? InParams.SourceOverride : Source;
+		ParamsCopy.Source->InitializeTargetMetadata(ParamsCopy, MetadataToInitialize);
 	}
 
+	SetupTargetMetadata(MetadataToInitialize);
+
 	MetadataToInitialize->AddAttributes(Metadata);
 }
\ No newline at end of file
diff --git a/Engine/Plugins/PCG/Source/PCG/Private/Data/PCGSpatialData.cpp b/Engine/Plugins/PCG/Source/PCG/Private/Data/PCGSpatialData.cpp
--- a/Engine/Plugins/PCG/Source/PCG/Private/Data/PCGSpatialData.cpp
+++ b/Engine/Plugins/PCG/Source/PCG/Private/Data/PCGSpatialData.cpp
@@ -23,6 +23,17 @@ namespace PCGSpatialData
 		TEXT("Enable UPCGSpatialData subclass PrepareForSpatialQuery task scheduling"));
 }
 
+FPCGInitializeFromDataParams::FPCGInitializeFromDataParams()
+	: FPCGInitializeFromDataParams(nullptr)
+{
+}
+
+FPCGInitializeFromDataParams::FPCGInitializeFromDataParams(const UPCGSpatialData* InSource)
+	: Source(InSource)
+	, MetadataInitializeParams(InSource ? InSource->ConstMetadata() : nullptr)
+{
+}
+
 UPCGSpatialData::UPCGSpatialData(const FObjectInitializer& ObjectInitializer)
 	: Super(ObjectInitializer)
 {
@@ -326,24 +337,44 @@ UPCGMetadata* UPCGSpatialData::CreateEmptyMetadata()
 	return Metadata;
 }
 
-void UPCGSpatialData::InitializeOtherMetadata(const FPCGInitializeFromDataParams& InParams, UPCGMetadata* MetadataToInitialize) const
+void UPCGSpatialData::InitializeMetadata(const FPCGInitializeFromDataParams& InParams)
 {
 	if (InParams.bInheritMetadata)
 	{
-		InitializeOtherMetadataInternal(InParams, MetadataToInitialize);
+		InitializeMetadataInternal(InParams);
 	}
 }
 
-void UPCGSpatialData::InitializeOtherMetadataInternal(const FPCGInitializeFromDataParams& InParams, UPCGMetadata* MetadataToInitialize) const
+void UPCGSpatialData::InitializeMetadataInternal(const FPCGInitializeFromDataParams& InParams)
+{
+	check(InParams.Source);
+	InParams.Source->InitializeTargetMetadata(InParams, Metadata);
+}
+
+void UPCGSpatialData::InitializeTargetMetadata(const FPCGInitializeFromDataParams& InParams, UPCGMetadata* MetadataToInitialize) const
 {
 	check(InParams.bInheritMetadata);
 	check(MetadataToInitialize);
 
-	FPCGMetadataInitializeParams InitializeParams(Metadata, TSet<FName>(), InParams.bInheritAttributes ? EPCGMetadataFilterMode::ExcludeAttributes : EPCGMetadataFilterMode::IncludeAttributes);
-	InitializeParams.DomainMapping.Reserve(InParams.DomainMapping.Num());
-	Algo::Transform(InParams.DomainMapping, InitializeParams.DomainMapping, [](const auto& It) { return It; });
+	// Making sure the metadata initialized params are setup correctly.
+	const FPCGMetadataInitializeParams* ParamsToUse = &InParams.MetadataInitializeParams;
+	FPCGMetadataInitializeParams TemporaryParams{nullptr};
+	if (InParams.MetadataInitializeParams.Parent != InParams.Source->ConstMetadata())
+	{
+		TemporaryParams = InParams.MetadataInitializeParams;
+		TemporaryParams.Parent = InParams.Source->ConstMetadata();
+		ParamsToUse = &TemporaryParams;
+	}
 
-	MetadataToInitialize->Initialize(InitializeParams);
+	// If the metadata to initialize already have a parent, we can't initialize it again, so just add the attributes.
+	if (MetadataToInitialize->GetParent())
+	{
+		MetadataToInitialize->AddAttributes(*ParamsToUse);
+	}
+	else
+	{
+		MetadataToInitialize->Initialize(*ParamsToUse);
+	}
 }
 
 void UPCGSpatialData::InitializeSpatialDataInternal(const FPCGInitializeFromDataParams& InParams)
@@ -373,7 +404,7 @@ void UPCGSpatialData::InitializeFromData(const UPCGSpatialData* InSource, const
 		}
 		else
 		{
-			InSource->InitializeOtherMetadata(Params, Metadata);
+			InitializeMetadata(Params);
 		}
 	}
 }
@@ -386,7 +417,7 @@ void UPCGSpatialData::InitializeFromDataWithParams(const FPCGInitializeFromDataP
 	InitializeSpatialDataInternal(InParams);
 	if (InParams.Source)
 	{
-		InParams.Source->InitializeOtherMetadata(InParams, Metadata);
+		InitializeMetadata(InParams);
 	}
 }
 
@@ -396,14 +427,10 @@ UPCGSpatialData* UPCGSpatialData::DuplicateData(FPCGContext* Context, bool bInit
 	check(NewSpatialData);
 	check(NewSpatialData->Metadata);
 	
-	NewSpatialData->InitializeSpatialDataInternal(this);
-
-	if (bInitializeMetadata)
-	{
-		FPCGInitializeFromDataParams Params(this);
-		Params.bIsDuplicatingData = true;
-		InitializeOtherMetadata(Params, NewSpatialData->Metadata);
-	}
+	FPCGInitializeFromDataParams Params(this);
+	Params.bInheritMetadata = bInitializeMetadata;
+	Params.bIsDuplicatingData = true;
+	NewSpatialData->InitializeFromDataWithParams(Params);
 
 	if (bHasCachedLastSelector)
 	{
diff --git a/Engine/Plugins/PCG/Source/PCG/Private/Data/PCGUnionData.cpp b/Engine/Plugins/PCG/Source/PCG/Private/Data/PCGUnionData.cpp
--- a/Engine/Plugins/PCG/Source/PCG/Private/Data/PCGUnionData.cpp
+++ b/Engine/Plugins/PCG/Source/PCG/Private/Data/PCGUnionData.cpp
@@ -190,6 +190,7 @@ void UPCGUnionData::AddToCrc(FArchiveCrc32& Ar, bool bFullDataCrc) const
 	Super::AddToCrc(Ar, bFullDataCrc);
 
 	// Implementation note: no metadata in composite data at this point.
+	// @todo_pcg: need metadata serialization now
 
 	FString ClassName = StaticClass()->GetPathName();
 	Ar << ClassName;
@@ -347,7 +348,7 @@ const UPCGBasePointData* UPCGUnionData::CreateBasePointData(FPCGContext* Context
 	
 	const UPCGBasePointData* FirstSource = ToPointDataFunc(Context, DataRawPtr[0]);
 	InputMetadatas[0] = FirstSource->Metadata;
-	
+
 	for (int32 i = 1; i < DataRawPtr.Num(); i++)
 	{
 		InputMetadatas[i] = ToPointDataFunc(Context, DataRawPtr[i])->Metadata;
@@ -356,8 +357,10 @@ const UPCGBasePointData* UPCGUnionData::CreateBasePointData(FPCGContext* Context
 	UPCGBasePointData* PointData = FPCGContext::NewObject_AnyThread<UPCGBasePointData>(Context, GetTransientPackage(), PointDataClass);
 	UPCGMetadata* OutMetadata = PointData->MutableMetadata();
 
-	FPCGInitializeFromDataParams InitializeFromDataParams(FirstSource);
+	FPCGInitializeFromDataParams InitializeFromDataParams(this);
 	InitializeFromDataParams.bInheritSpatialData = false;
+	// Since we have collapsed the source data, we need to inherit from this one.
+	InitializeFromDataParams.SourceOverride = FirstSource;
 	PointData->InitializeFromDataWithParams(InitializeFromDataParams);
 
 	switch (UnionType)
@@ -660,7 +663,7 @@ UPCGSpatialData* UPCGUnionData::CopyInternal(FPCGContext* Context) const
 	return NewUnionData;
 }
 
-void UPCGUnionData::InitializeOtherMetadataInternal(const FPCGInitializeFromDataParams& InParams, UPCGMetadata* MetadataToInitialize) const
+void UPCGUnionData::InitializeTargetMetadata(const FPCGInitializeFromDataParams& InParams, UPCGMetadata* MetadataToInitialize) const
 {
 	check(InParams.bInheritMetadata);
 	check(MetadataToInitialize);
@@ -668,17 +671,25 @@ void UPCGUnionData::InitializeOtherMetadataInternal(const FPCGInitializeFromData
 	// Duplicate data case, call the spatial base method
 	if (InParams.bIsDuplicatingData)
 	{
-		UPCGSpatialData::InitializeOtherMetadataInternal(InParams, MetadataToInitialize);
+		UPCGSpatialData::InitializeTargetMetadata(InParams, MetadataToInitialize);
 		return;
 	}
 
-	// In the case of the Union, we initialize the data from all the sources. The union will add its own
-	// attributes afterward.
-	FPCGInitializeFromDataParams CopyParams = InParams;
-	for (const UPCGSpatialData* UnionedData : Data)
+	if (!Data.IsEmpty())
 	{
-		CopyParams.Source = UnionedData;
-		UnionedData->InitializeOtherMetadata(InParams, MetadataToInitialize);
+		// In the case of the Union, we initialize the data from all the sources, special case for the first element, use the source override if specified.
+		// The union will add its own attributes afterward.
+		FPCGInitializeFromDataParams CopyParams = InParams;
+		CopyParams.SourceOverride = nullptr;
+		
+		for (int32 i = 0; i < Data.Num(); ++i)
+		{
+			if (const UPCGSpatialData* UnionedData = Data[i])
+			{
+				CopyParams.Source = (InParams.SourceOverride && i == 0) ? InParams.SourceOverride.Get() : UnionedData;
+				CopyParams.Source->InitializeTargetMetadata(CopyParams, MetadataToInitialize);
+			}
+		}
 	}
 
 	MetadataToInitialize->AddAttributes(Metadata);
diff --git a/Engine/Plugins/PCG/Source/PCG/Private/Data/PCGWorldData.cpp b/Engine/Plugins/PCG/Source/PCG/Private/Data/PCGWorldData.cpp
--- a/Engine/Plugins/PCG/Source/PCG/Private/Data/PCGWorldData.cpp
+++ b/Engine/Plugins/PCG/Source/PCG/Private/Data/PCGWorldData.cpp
@@ -173,7 +173,7 @@ bool UPCGWorldVolumetricData::SamplePoint(const FTransform& InTransform, const F
 
 	check(World.IsValid());
 
-	FPCGMetadataAttribute<FSoftObjectPath>* ActorOverlappedAttribute = ((OutMetadata && QueryParams.bGetReferenceToActorHit) ? OutMetadata->GetMutableTypedAttribute<FSoftObjectPath>(PCGPointDataConstants::ActorReferenceAttribute) : nullptr);
+	FPCGMetadataAttribute<FSoftObjectPath>* ActorOverlappedAttribute = ((OutMetadata && QueryParams.bGetReferenceToActorHit) ? OutMetadata->FindOrCreateAttribute<FSoftObjectPath>(PCGPointDataConstants::ActorReferenceAttribute) : nullptr);
 
 	FCollisionObjectQueryParams ObjectQueryParams(QueryParams.CollisionChannel);
 	FCollisionShape CollisionShape = FCollisionShape::MakeBox(InBounds.GetExtent() * InTransform.GetScale3D());
@@ -281,7 +281,6 @@ void UPCGWorldRayHitData::Initialize(UWorld* InWorld, const FTransform& InTransf
 	Transform = InTransform;
 	Bounds = InBounds;
 	LocalBounds = InLocalBounds;
-	PCGWorldQueryHelpers::CreateRayHitAttributes(QueryParams, Metadata);
 }
 
 void UPCGWorldRayHitData::AddToCrc(FArchiveCrc32& Ar, bool bFullDataCrc) const
@@ -343,8 +342,6 @@ const UPCGPointArrayData* UPCGWorldRayHitData::CreatePointArrayData(FPCGContext*
 
 const UPCGBasePointData* UPCGWorldRayHitData::CreateBasePointData(FPCGContext* Context, const FBox& InBounds, TSubclassOf<UPCGBasePointData> PointDataClass) const
 {
-	PCGWorldQueryHelpers::CreateRayHitAttributes(QueryParams, Metadata);
-
 	FBox EffectiveBounds = Bounds;
 	if (InBounds.IsValid)
 	{
@@ -376,6 +373,30 @@ const UPCGBasePointData* UPCGWorldRayHitData::CreateBasePointData(FPCGContext* C
 	return PCGSurfaceSampler::SampleSurface(Context, Params, /*InSurface=*/this, /*InBoundingShape=*/nullptr, EffectiveBounds, PointDataClass);
 }
 
+void UPCGWorldRayHitData::InitializeTargetMetadata(const FPCGInitializeFromDataParams& InParams, UPCGMetadata* MetadataToInitialize) const
+{
+	check(MetadataToInitialize);
+
+	// Initialize the new metadata normally
+	Super::InitializeTargetMetadata(InParams, MetadataToInitialize);
+	
+	// Only add those special attributes if the metadata to initialize support elements domain
+	FPCGMetadataDomain* MetadataDomain = MetadataToInitialize->GetMetadataDomain(PCGMetadataDomainID::Elements);
+	if (!MetadataDomain)
+	{
+		return;
+	}
+	
+	// Add all landscape layers
+	for (const FName LayerName : CachedLandscapeLayerNames)
+	{
+		MetadataDomain->FindOrCreateAttribute<float>(LayerName, 0.0f, /*bAllowInterpolation=*/true);
+	}
+
+	// Then all the other attributes
+	PCGWorldQueryHelpers::CreateRayHitAttributes(QueryParams, MetadataDomain);
+}
+
 UPCGSpatialData* UPCGWorldRayHitData::CopyInternal(FPCGContext* Context) const
 {
 	UPCGWorldRayHitData* NewData = FPCGContext::NewObject_AnyThread<UPCGWorldRayHitData>(Context);
@@ -388,6 +409,7 @@ UPCGSpatialData* UPCGWorldRayHitData::CopyInternal(FPCGContext* Context) const
 	NewData->QueryParams = QueryParams;
 	NewData->QueryParams.Initialize();
 	NewData->ActorFilter.GetFilterActorsMutable() = ActorFilter.GetFilterActorsConst();
+	NewData->CachedLandscapeLayerNames = CachedLandscapeLayerNames;
 
 	return NewData;
 }
diff --git a/Engine/Plugins/PCG/Source/PCG/Private/Elements/PCGGetSplineControlPoints.cpp b/Engine/Plugins/PCG/Source/PCG/Private/Elements/PCGGetSplineControlPoints.cpp
--- a/Engine/Plugins/PCG/Source/PCG/Private/Elements/PCGGetSplineControlPoints.cpp
+++ b/Engine/Plugins/PCG/Source/PCG/Private/Elements/PCGGetSplineControlPoints.cpp
@@ -70,7 +70,6 @@ bool FPCGGetSplineControlPointsElement::ExecuteInternal(FPCGContext* InContext)
 
 		UPCGBasePointData* PointData = FPCGContext::NewPointData_AnyThread(InContext);	
 		FPCGInitializeFromDataParams Params(PolylineData);
-		Params.DomainMapping = {{PCGMetadataDomainID::Default, PCGMetadataDomainID::Elements}};
 		PointData->InitializeFromDataWithParams(Params);
 		PointData->SetNumPoints(NumPoints);
 		PointData->AllocateProperties(EPCGPointNativeProperties::Transform | EPCGPointNativeProperties::Seed | EPCGPointNativeProperties::MetadataEntry);
diff --git a/Engine/Plugins/PCG/Source/PCG/Private/Elements/PCGWorldQuery.cpp b/Engine/Plugins/PCG/Source/PCG/Private/Elements/PCGWorldQuery.cpp
--- a/Engine/Plugins/PCG/Source/PCG/Private/Elements/PCGWorldQuery.cpp
+++ b/Engine/Plugins/PCG/Source/PCG/Private/Elements/PCGWorldQuery.cpp
@@ -219,7 +219,6 @@ bool FPCGWorldRayHitQueryElement::ExecuteInternal(FPCGContext* Context) const
 	Data->CollisionShape = Settings->CollisionShape;
 	Data->Initialize(World, Transform, /*InBounds=*/FBox(EForceInit::ForceInit), LocalBounds);
 
-	bool bHasLandscapeMetadata = false;
 	if (QueryParams.bApplyMetadataFromLandscape && Data->Metadata)
 	{
 		UPCGSubsystem* PCGSubsystem = UPCGSubsystem::GetInstance(World);
@@ -250,21 +249,13 @@ bool FPCGWorldRayHitQueryElement::ExecuteInternal(FPCGContext* Context) const
 			{
 				if (ALandscapeProxy* LandscapeProxy = Cast<ALandscapeProxy>(Landscape))
 				{
-					const TArray<FName> Layers = LandscapeCache->GetLayerNames(LandscapeProxy);
-					for (FName Layer : Layers)
-					{
-						if (!Data->Metadata->HasAttribute(Layer))
-						{
-							Data->Metadata->CreateFloatAttribute(Layer, 0.0f, /*bAllowInterpolation=*/true);
-							bHasLandscapeMetadata = true;
-						}
-					}
+					Data->CachedLandscapeLayerNames.Append(LandscapeCache->GetLayerNames(LandscapeProxy));
 				}
 			}
 		}
 	}
 
-	if (!bHasLandscapeMetadata)
+	if (Data->CachedLandscapeLayerNames.IsEmpty())
 	{
 		Data->QueryParams.bApplyMetadataFromLandscape = false;
 	}
diff --git a/Engine/Plugins/PCG/Source/PCG/Private/Grid/PCGLandscapeCache.cpp b/Engine/Plugins/PCG/Source/PCG/Private/Grid/PCGLandscapeCache.cpp
--- a/Engine/Plugins/PCG/Source/PCG/Private/Grid/PCGLandscapeCache.cpp
+++ b/Engine/Plugins/PCG/Source/PCG/Private/Grid/PCGLandscapeCache.cpp
@@ -285,14 +285,13 @@ void FPCGLandscapeCacheEntry::GetInterpolatedPointMetadataInternal(const PCGLand
 			const TArray<uint8>& CurrentLayerData = LayerData[LayerIndex];
 			check(Indices.X1Y1 < CurrentLayerData.Num());
 
-			if (FPCGMetadataAttributeBase* Attribute = OutMetadata->GetMutableAttribute(CurrentLayerName))
+			if (FPCGMetadataAttribute<float>* Attribute = OutMetadata->FindOrCreateAttribute<float>(CurrentLayerName))
 			{
-				check(Attribute->GetTypeId() == PCG::Private::MetadataTypes<float>::Id);
 				const float Y0Data = FMath::Lerp((float)CurrentLayerData[Indices.X0Y0] / 255.0f, (float)CurrentLayerData[Indices.X1Y0] / 255.0f, Indices.XFraction);
 				const float Y1Data = FMath::Lerp((float)CurrentLayerData[Indices.X0Y1] / 255.0f, (float)CurrentLayerData[Indices.X1Y1] / 255.0f, Indices.XFraction);
 				const float Data = FMath::Lerp(Y0Data, Y1Data, Indices.YFraction);
 
-				static_cast<FPCGMetadataAttribute<float>*>(Attribute)->SetValue(OutMetadataEntry, Data);
+				Attribute->SetValue(OutMetadataEntry, Data);
 			}
 		}
 	}
diff --git a/Engine/Plugins/PCG/Source/PCG/Private/Helpers/PCGWorldQueryHelpers.cpp b/Engine/Plugins/PCG/Source/PCG/Private/Helpers/PCGWorldQueryHelpers.cpp
--- a/Engine/Plugins/PCG/Source/PCG/Private/Helpers/PCGWorldQueryHelpers.cpp
+++ b/Engine/Plugins/PCG/Source/PCG/Private/Helpers/PCGWorldQueryHelpers.cpp
@@ -8,6 +8,7 @@
 #include "Data/PCGWorldData.h"
 #include "Grid/PCGLandscapeCache.h"
 #include "Helpers/PCGHelpers.h"
+#include "Metadata/PCGMetadataDomain.h"
 
 #include "LandscapeProxy.h"
 #include "StaticMeshResources.h"
@@ -203,6 +204,16 @@ namespace PCGWorldQueryHelpers
 
 	bool CreateRayHitAttributes(const FPCGWorldRaycastQueryParams& QueryParams, UPCGMetadata* OutMetadata)
 	{
+		return CreateRayHitAttributes(QueryParams, OutMetadata ? OutMetadata->GetDefaultMetadataDomain() : nullptr);
+	}
+
+	bool CreateRayHitAttributes(const FPCGWorldRaycastQueryParams& QueryParams, FPCGMetadataDomain* OutMetadata)
+	{
+		if (!OutMetadata)
+		{
+			return false;
+		}
+
 		auto CreateAttribute = [OutMetadata]<typename Type>(FName AttributeName, bool bShouldCreate, const Type& DefaultValue) -> bool
 		{
 			return !bShouldCreate || OutMetadata->FindOrCreateAttribute<Type>(AttributeName, DefaultValue, /*bAllowsInterpolation=*/true, /*bOverrideParent=*/false, /*bOverwriteIfTypeMismatch=*/true);
@@ -273,7 +284,7 @@ namespace PCGWorldQueryHelpers
 				return true;
 			}
 
-			if (FPCGMetadataAttribute<Type>* Attribute = OutMetadata->GetMutableTypedAttribute<Type>(AttributeName))
+			if (FPCGMetadataAttribute<Type>* Attribute = OutMetadata->FindOrCreateAttribute<Type>(AttributeName))
 			{
 				OutMetadata->InitializeOnSet(OutMetadataEntry);
 				Attribute->SetValue(OutMetadataEntry, Value);
@@ -314,12 +325,10 @@ namespace PCGWorldQueryHelpers
 			{
 				if (UPCGLandscapeCache* LandscapeCache = World->GetSubsystem<UPCGSubsystem>()->GetLandscapeCache())
 				{
-					// TODO: This is not ideal, but we won't have the Landscape until execution and the LandscapeCache intends on requiring it
 					const TArray<FName> Layers = LandscapeCache->GetLayerNames(Landscape);
 					for (const FName& Layer : Layers)
 					{
-						// Use Threadsafe version for on the fly attribute creation.
-						OutMetadata->FindOrCreateAttribute<float>(Layer, float());
+						OutMetadata->FindOrCreateAttribute<float>(Layer);
 					}
 
 					LandscapeCache->SampleMetadataOnPoint(Landscape, InTransform, OutMetadataEntry, OutMetadata);
diff --git a/Engine/Plugins/PCG/Source/PCG/Private/Metadata/PCGMetadata.cpp b/Engine/Plugins/PCG/Source/PCG/Private/Metadata/PCGMetadata.cpp
--- a/Engine/Plugins/PCG/Source/PCG/Private/Metadata/PCGMetadata.cpp
+++ b/Engine/Plugins/PCG/Source/PCG/Private/Metadata/PCGMetadata.cpp
@@ -73,8 +73,11 @@ void UPCGMetadata::FindFixOrCreateDomainInitializeParams(const FPCGMetadataIniti
 		OtherDomain = *It;
 	}
 
-	check(Params);
-	InFunc(FindOrCreateMetadataDomain(OtherDomain), *Params);
+	if (FPCGMetadataDomain* MetadataDomain = FindOrCreateMetadataDomain(OtherDomain))
+	{
+		check(Params);
+		InFunc(*MetadataDomain, *Params);
+	}
 }
 
 template <typename Func, typename ...Args>
@@ -137,66 +140,42 @@ decltype(auto) UPCGMetadata::WithConstMetadataDomain(const FPCGMetadataDomainID&
 template <typename Func>
 void UPCGMetadata::ForEachValidUniqueConstMetadataDomain(Func InFunc) const
 {
-	auto Apply = [this, &InFunc](const FPCGMetadataDomainID& DomainID)
+	for (const auto& It : MetadataDomains)
 	{
-		if (DomainID == DefaultDomain && !DomainID.IsDefault())
-		{
-			return;
-		}
-		
-		const FPCGMetadataDomain* CurrentMetadataDomain = MetadataDomains[DomainID].Get();
-		if (!CurrentMetadataDomain)
+		const FPCGMetadataDomainID DomainID = It.Key;
+		if (DomainID.IsDefault())
 		{
-			return;
+			continue;
 		}
 		
-		InFunc(DomainID, *CurrentMetadataDomain);
-	};
-
-	// Always check the default first
-	Apply(PCGMetadataDomainID::Default);
-	
-	for (const auto& It : MetadataDomains)
-	{
-		if (It.Key.IsDefault())
+		const FPCGMetadataDomain* MetadataDomain = It.Value.Get();
+		if (!MetadataDomain)
 		{
 			continue;
 		}
 
-		Apply(It.Key);
+		InFunc(It.Key, *MetadataDomain);
 	}
 }
 
 template <typename Func>
 void UPCGMetadata::ForEachValidUniqueMetadataDomain(Func InFunc)
-{	
-	auto Apply = [this, &InFunc](const FPCGMetadataDomainID& DomainID)
+{
+	for (const auto& It : MetadataDomains)
 	{
-		if (DomainID == DefaultDomain && !DomainID.IsDefault())
+		const FPCGMetadataDomainID DomainID = It.Key;
+		if (DomainID.IsDefault())
 		{
-			return;
+			continue;
 		}
 		
-		FPCGMetadataDomain* CurrentMetadataDomain = MetadataDomains[DomainID].Get();
-		if (!CurrentMetadataDomain)
-		{
-			return;
-		}
-		
-		InFunc(DomainID, *CurrentMetadataDomain);
-	};
-	
-	// Always check the default first
-	Apply(PCGMetadataDomainID::Default);
-	
-	for (const auto& It : MetadataDomains)
-	{
-		if (It.Key.IsDefault())
+		FPCGMetadataDomain* MetadataDomain = It.Value.Get();
+		if (!MetadataDomain)
 		{
 			continue;
 		}
 
-		Apply(It.Key);
+		InFunc(It.Key, *MetadataDomain);
 	}
 }
 
@@ -226,7 +205,7 @@ const FPCGMetadataDomain* UPCGMetadata::GetConstMetadataDomainFromSelector(const
 FPCGMetadataDomain* UPCGMetadata::FindOrCreateMetadataDomain(const FPCGMetadataDomainID& InMetadataDomainID)
 {
 	TSharedPtr<FPCGMetadataDomain>* FoundMetadataDomain = MetadataDomains.Find(InMetadataDomainID);
-	if(!ensure(FoundMetadataDomain))
+	if (!FoundMetadataDomain)
 	{
 		return nullptr;
 	}
@@ -505,9 +484,9 @@ void UPCGMetadata::Initialize(const FPCGMetadataInitializeParams& InParams)
 	
 	Parent->ForEachValidUniqueConstMetadataDomain([this, &InParams](const FPCGMetadataDomainID& DomainID, const FPCGMetadataDomain& OtherMetadataDomain)
 	{
-		auto InitializeFunc = [](FPCGMetadataDomain* CurrentMetadataDomain, const FPCGMetadataDomainInitializeParams& DomainParams)
+		auto InitializeFunc = [](FPCGMetadataDomain& CurrentMetadataDomain, const FPCGMetadataDomainInitializeParams& DomainParams)
 		{
-			CurrentMetadataDomain->Initialize(DomainParams);
+			CurrentMetadataDomain.Initialize(DomainParams);
 		};
 		
 		FindFixOrCreateDomainInitializeParams(InParams, DomainID, OtherMetadataDomain, std::move(InitializeFunc));
@@ -545,12 +524,12 @@ void UPCGMetadata::InitializeAsCopy(const FPCGMetadataInitializeParams& InParams
 		
 	InParams.Parent->ForEachValidUniqueConstMetadataDomain([this, &bShouldParent, &InParams](const FPCGMetadataDomainID& DomainID, const FPCGMetadataDomain& OtherMetadataDomain)
 	{
-		auto InitializeFunc = [&bShouldParent, &InParams](FPCGMetadataDomain* CurrentMetadataDomain, const FPCGMetadataDomainInitializeParams& DomainParams)
+		auto InitializeFunc = [&bShouldParent, &InParams](FPCGMetadataDomain& CurrentMetadataDomain, const FPCGMetadataDomainInitializeParams& DomainParams)
 		{
-			CurrentMetadataDomain->InitializeAsCopy(DomainParams);
+			CurrentMetadataDomain.InitializeAsCopy(DomainParams);
 
 			// We still validate that the parent was set correctly and it matches the metadata to copy.
-			if (CurrentMetadataDomain->Parent && CurrentMetadataDomain->Parent->TopMetadata == InParams.Parent)
+			if (CurrentMetadataDomain.Parent && CurrentMetadataDomain.Parent->TopMetadata == InParams.Parent)
 			{
 				bShouldParent = true;
 			}
@@ -587,9 +566,9 @@ void UPCGMetadata::AddAttributes(const FPCGMetadataInitializeParams& InParams)
 
 	InParams.Parent->ForEachValidUniqueConstMetadataDomain([this, &InParams, &bAddSucceeded](const FPCGMetadataDomainID& DomainID, const FPCGMetadataDomain& OtherMetadataDomain)
 	{
-		auto InitializeFunc = [&bAddSucceeded](FPCGMetadataDomain* CurrentMetadataDomain, const FPCGMetadataDomainInitializeParams& DomainParams)
+		auto InitializeFunc = [&bAddSucceeded](FPCGMetadataDomain& CurrentMetadataDomain, const FPCGMetadataDomainInitializeParams& DomainParams)
 		{
-			bAddSucceeded |= CurrentMetadataDomain->AddAttributes(DomainParams);
+			bAddSucceeded |= CurrentMetadataDomain.AddAttributes(DomainParams);
 		};
 
 		FindFixOrCreateDomainInitializeParams(InParams, DomainID, OtherMetadataDomain, std::move(InitializeFunc));
@@ -636,8 +615,10 @@ void UPCGMetadata::CopyAttributes(const UPCGMetadata* InOther)
 
 	InOther->ForEachValidUniqueConstMetadataDomain([this](const FPCGMetadataDomainID& DomainID, const FPCGMetadataDomain& OtherMetadataDomain)
 	{
-		FPCGMetadataDomain* CurrentMetadataDomain = FindOrCreateMetadataDomain(DomainID);
-		CurrentMetadataDomain->CopyAttributes(&OtherMetadataDomain);
+		if (FPCGMetadataDomain* CurrentMetadataDomain = FindOrCreateMetadataDomain(DomainID))
+		{
+			CurrentMetadataDomain->CopyAttributes(&OtherMetadataDomain);
+		}
 	});
 }
 
diff --git a/Engine/Plugins/PCG/Source/PCG/Public/Data/PCGDifferenceData.h b/Engine/Plugins/PCG/Source/PCG/Public/Data/PCGDifferenceData.h
--- a/Engine/Plugins/PCG/Source/PCG/Public/Data/PCGDifferenceData.h
+++ b/Engine/Plugins/PCG/Source/PCG/Public/Data/PCGDifferenceData.h
@@ -77,9 +77,8 @@ public:
 	UE_API virtual bool SamplePoint(const FTransform& Transform, const FBox& Bounds, FPCGPoint& OutPoint, UPCGMetadata* OutMetadata) const override;
 	UE_API virtual bool HasNonTrivialTransform() const override;
 	virtual const UPCGSpatialData* FindFirstConcreteShapeFromNetwork() const override { return GetSource() ? GetSource()->FindFirstConcreteShapeFromNetwork() : nullptr; }
-
+	UE_API virtual void InitializeTargetMetadata(const FPCGInitializeFromDataParams& InParams, UPCGMetadata* MetadataToInitialize) const override;
 protected:
-	UE_API virtual void InitializeOtherMetadataInternal(const FPCGInitializeFromDataParams& InParams, UPCGMetadata* MetadataToInitialize) const override;
 	UE_API virtual UPCGSpatialData* CopyInternal(FPCGContext* Context) const override;
 	//~End UPCGSpatialData interface
 
diff --git a/Engine/Plugins/PCG/Source/PCG/Public/Data/PCGIntersectionData.h b/Engine/Plugins/PCG/Source/PCG/Public/Data/PCGIntersectionData.h
--- a/Engine/Plugins/PCG/Source/PCG/Public/Data/PCGIntersectionData.h
+++ b/Engine/Plugins/PCG/Source/PCG/Public/Data/PCGIntersectionData.h
@@ -49,10 +49,9 @@ public:
 	UE_API virtual bool SamplePoint(const FTransform& Transform, const FBox& Bounds, FPCGPoint& OutPoint, UPCGMetadata* OutMetadata) const override;
 	UE_API virtual bool HasNonTrivialTransform() const override;
 	UE_API virtual const UPCGSpatialData* FindFirstConcreteShapeFromNetwork() const override;
-
+	UE_API virtual void InitializeTargetMetadata(const FPCGInitializeFromDataParams& InParams, UPCGMetadata* MetadataToInitialize) const override;
 protected:
 	UE_API virtual UPCGSpatialData* CopyInternal(FPCGContext* Context) const override;
-	UE_API virtual void InitializeOtherMetadataInternal(const FPCGInitializeFromDataParams& InParams, UPCGMetadata* MetadataToInitialize) const override;
 	//~End UPCGSpatialData interface
 
 public:
diff --git a/Engine/Plugins/PCG/Source/PCG/Public/Data/PCGLandscapeData.h b/Engine/Plugins/PCG/Source/PCG/Public/Data/PCGLandscapeData.h
--- a/Engine/Plugins/PCG/Source/PCG/Public/Data/PCGLandscapeData.h
+++ b/Engine/Plugins/PCG/Source/PCG/Public/Data/PCGLandscapeData.h
@@ -80,6 +80,7 @@ public:
 	UE_API virtual bool ProjectPoint(const FTransform& InTransform, const FBox& InBounds, const FPCGProjectionParams& InParams, FPCGPoint& OutPoint, UPCGMetadata* OutMetadata) const override;
 	virtual bool HasNonTrivialTransform() const override { return true; }
 	UE_API virtual TArray<FPCGTaskId> PrepareForSpatialQuery(FPCGContext* InContext, const FBox& InBounds) const override;
+	UE_API virtual void InitializeTargetMetadata(const FPCGInitializeFromDataParams& InParams, UPCGMetadata* MetadataToInitialize) const override;
 protected:
 	UE_API virtual UPCGSpatialData* CopyInternal(FPCGContext* Context) const override;
 	//~End UPCGSpatialData interface
@@ -123,6 +124,8 @@ protected:
 #endif // WITH_EDITORONLY_DATA
 
 private:
+	bool UseMetadata() const;
+
 	// Transient data
 	TArray<TPair<FBox, ULandscapeInfo*>> BoundsToLandscapeInfos;
 	TArray<ULandscapeInfo*> LandscapeInfos;
diff --git a/Engine/Plugins/PCG/Source/PCG/Public/Data/PCGProjectionData.h b/Engine/Plugins/PCG/Source/PCG/Public/Data/PCGProjectionData.h
--- a/Engine/Plugins/PCG/Source/PCG/Public/Data/PCGProjectionData.h
+++ b/Engine/Plugins/PCG/Source/PCG/Public/Data/PCGProjectionData.h
@@ -37,8 +37,8 @@ public:
 	UE_API virtual bool SamplePoint(const FTransform& Transform, const FBox& Bounds, FPCGPoint& OutPoint, UPCGMetadata* OutMetadata) const override;
 	UE_API virtual bool HasNonTrivialTransform() const override;
 	UE_API virtual bool RequiresCollapseToSample() const override;
+	UE_API virtual void InitializeTargetMetadata(const FPCGInitializeFromDataParams& InParams, UPCGMetadata* MetadataToInitialize) const override;
 protected:
-	UE_API virtual void InitializeOtherMetadataInternal(const FPCGInitializeFromDataParams& InParams, UPCGMetadata* MetadataToInitialize) const override;
 	UE_API virtual UPCGSpatialData* CopyInternal(FPCGContext* Context) const override;
 	//~End UPCGSpatialData interface
 
@@ -60,6 +60,8 @@ protected:
 
 	UE_API void GetIncludeExcludeAttributeNames(TSet<FName>& OutAttributeNames) const;
 
+	UE_API void SetupTargetMetadata(UPCGMetadata* MetadataToInitialize) const;
+
 	UPROPERTY(BlueprintReadOnly, VisibleAnywhere, Category = SpatialData)
 	TObjectPtr<const UPCGSpatialData> Source = nullptr;
 
diff --git a/Engine/Plugins/PCG/Source/PCG/Public/Data/PCGSpatialData.h b/Engine/Plugins/PCG/Source/PCG/Public/Data/PCGSpatialData.h
--- a/Engine/Plugins/PCG/Source/PCG/Public/Data/PCGSpatialData.h
+++ b/Engine/Plugins/PCG/Source/PCG/Public/Data/PCGSpatialData.h
@@ -29,21 +29,22 @@ namespace PCGSpatialData
 	extern TAutoConsoleVariable<bool> CVarEnablePrepareForSpatialQuery;
 }
 
-USTRUCT(BlueprintType)
+USTRUCT(MinimalAPI, BlueprintType)
 struct FPCGInitializeFromDataParams
 {
 	GENERATED_BODY()
 
-	FPCGInitializeFromDataParams() = default;
+	UE_API FPCGInitializeFromDataParams();
 	
-	FPCGInitializeFromDataParams(const UPCGSpatialData* InSource)
-		: Source(InSource)
-	{
-	}
+	UE_API FPCGInitializeFromDataParams(const UPCGSpatialData* InSource);
 	
 	UPROPERTY(BlueprintReadWrite, EditAnywhere, Category = Default)
 	TObjectPtr<const UPCGSpatialData> Source = nullptr;
-		
+
+	/** In the case of collapse of composite data, we need to inherit metadata from another source. */
+	UPROPERTY(BlueprintReadWrite, EditAnywhere, Category = Default)
+	TObjectPtr<const UPCGSpatialData> SourceOverride = nullptr;
+
 	UPROPERTY(BlueprintReadWrite, EditAnywhere, Category = Default)
 	bool bInheritMetadata = true;
 
@@ -56,9 +57,8 @@ struct FPCGInitializeFromDataParams
 	// Special boolean to be set to true when we duplicate data.
 	bool bIsDuplicatingData = false;
 
-	/** When initializing metadata from different types, can provide an extra mapping to know how to initialize metadata across domains. If empty, will match with the same domains. */
-	UPROPERTY(BlueprintReadWrite, EditAnywhere, Category = Default)
-	TMap<FPCGMetadataDomainID, FPCGMetadataDomainID> DomainMapping;
+	/** When initializing metadata, can provide an extra set of params to initialize (for filtering attributes for example) */
+	FPCGMetadataInitializeParams MetadataInitializeParams;
 };
 
 /**
@@ -233,10 +233,12 @@ public:
 	UPROPERTY(BlueprintReadWrite, EditAnywhere, Category = Data)
 	bool bKeepZeroDensityPoints = false;
 
-	/** Internal function meant to be used by composite data to initialize their metadata using other sources, or ourselves. Need to be public to be called by child classes.
-	 * InitializeFromData or InitializeFromDataWithParams are the one to use normally.
+	/**
+	 * Initialize another metadata from this data. By default, it will just initialize/add attributes to the metadata to initialize.
+	 * But it can also be overridden if there is specific logic when there is a cross domain boundary (like Surface -> Points).
+	 * If MetadataToInitialize is not yet initialized, will initialize it. If it is already initialized, will add the attributes.
 	 */
-	UE_API void InitializeOtherMetadata(const FPCGInitializeFromDataParams& InParams, UPCGMetadata* MetadataToInitialize) const;
+	UE_API virtual void InitializeTargetMetadata(const FPCGInitializeFromDataParams& InParams, UPCGMetadata* MetadataToInitialize) const;
 
 protected:
 	UE_DEPRECATED(5.5, "Call/Implement version with FPCGContext parameter")
@@ -244,8 +246,8 @@ protected:
 
 	UE_API virtual UPCGSpatialData* CopyInternal(FPCGContext* Context) const PURE_VIRTUAL(UPCGSpatialData::CopyInternal, return nullptr;);
 	UE_API virtual void InitializeSpatialDataInternal(const FPCGInitializeFromDataParams& InParams);
-
-	UE_API virtual void InitializeOtherMetadataInternal(const FPCGInitializeFromDataParams& InParams, UPCGMetadata* MetadataToInitialize) const;
+	UE_API void InitializeMetadata(const FPCGInitializeFromDataParams& InParams);
+	UE_API virtual void InitializeMetadataInternal(const FPCGInitializeFromDataParams& InParams);
 
 private:
 	/** Cache to keep track of the latest attribute manipulated on this data. */
diff --git a/Engine/Plugins/PCG/Source/PCG/Public/Data/PCGUnionData.h b/Engine/Plugins/PCG/Source/PCG/Public/Data/PCGUnionData.h
--- a/Engine/Plugins/PCG/Source/PCG/Public/Data/PCGUnionData.h
+++ b/Engine/Plugins/PCG/Source/PCG/Public/Data/PCGUnionData.h
@@ -55,9 +55,8 @@ public:
 	UE_API virtual bool SamplePoint(const FTransform& Transform, const FBox& Bounds, FPCGPoint& OutPoint, UPCGMetadata* OutMetadata) const override;
 	UE_API virtual bool HasNonTrivialTransform() const override;
 	UE_API virtual const UPCGSpatialData* FindFirstConcreteShapeFromNetwork() const override;
-
+	UE_API virtual void InitializeTargetMetadata(const FPCGInitializeFromDataParams& InParams, UPCGMetadata* MetadataToInitialize) const override;
 protected:
-	UE_API virtual void InitializeOtherMetadataInternal(const FPCGInitializeFromDataParams& InParams, UPCGMetadata* MetadataToInitialize) const override;
 	UE_API virtual UPCGSpatialData* CopyInternal(FPCGContext* Context) const override;
 	//~End UPCGSpatialData interface
 
diff --git a/Engine/Plugins/PCG/Source/PCG/Public/Data/PCGWorldData.h b/Engine/Plugins/PCG/Source/PCG/Public/Data/PCGWorldData.h
--- a/Engine/Plugins/PCG/Source/PCG/Public/Data/PCGWorldData.h
+++ b/Engine/Plugins/PCG/Source/PCG/Public/Data/PCGWorldData.h
@@ -366,9 +366,10 @@ public:
 	virtual FBox GetBounds() const override { return Bounds; }
 	virtual FBox GetStrictBounds() const override { return Bounds; }
 	virtual bool IsBounded() const override { return !!Bounds.IsValid; }
-	virtual bool SamplePoint(const FTransform& Transform, const FBox& Bounds, FPCGPoint& OutPoint, UPCGMetadata* OutMetadata) const override;
+	PCG_API virtual bool SamplePoint(const FTransform& Transform, const FBox& Bounds, FPCGPoint& OutPoint, UPCGMetadata* OutMetadata) const override;
 	virtual bool HasNonTrivialTransform() const override { return true; }
 	virtual FVector GetNormal() const override { return Transform.GetRotation().GetUpVector(); }
+	PCG_API virtual void InitializeTargetMetadata(const FPCGInitializeFromDataParams& InParams, UPCGMetadata* MetadataToInitialize) const override;
 protected:
 	virtual UPCGSpatialData* CopyInternal(FPCGContext* Context) const override;
 	//~End UPCGSpatialData interface
@@ -400,4 +401,8 @@ public:
 	
 	UPROPERTY()
 	FPCGWorldQueryActorFilterCache ActorFilter;
+	
+	/** Attributes related to landscape layers are added when we initialize the target metadata, so we need to cache them when this data gets initialized. */
+	UPROPERTY()
+	TSet<FName> CachedLandscapeLayerNames;
 };
diff --git a/Engine/Plugins/PCG/Source/PCG/Public/Helpers/PCGWorldQueryHelpers.h b/Engine/Plugins/PCG/Source/PCG/Public/Helpers/PCGWorldQueryHelpers.h
--- a/Engine/Plugins/PCG/Source/PCG/Public/Helpers/PCGWorldQueryHelpers.h
+++ b/Engine/Plugins/PCG/Source/PCG/Public/Helpers/PCGWorldQueryHelpers.h
@@ -11,6 +11,7 @@
 #include "PCGWorldQueryHelpers.generated.h"
 
 class AActor;
+class FPCGMetadataDomain;
 class UPCGComponent;
 class UPCGMetadata;
 struct FHitResult;
@@ -74,6 +75,9 @@ namespace PCGWorldQueryHelpers
 
 	/** Creates hit result attributes based off query params. Can be called before ApplyRayHitMetadata. */
 	bool CreateRayHitAttributes(const FPCGWorldRaycastQueryParams& QueryParams, UPCGMetadata* OutMetadata);
+	
+	/** Creates hit result attributes based off query params. Can be called before ApplyRayHitMetadata. */
+	bool CreateRayHitAttributes(const FPCGWorldRaycastQueryParams& QueryParams, FPCGMetadataDomain* OutMetadata);
 
 	/** Applies a 'miss' hit result to the metadata. To be called if the ray misses the target and the point should be kept. */
 	bool ApplyRayMissMetadata(const FPCGWorldRaycastQueryParams& QueryParams, int64& OutMetadataEntry, UPCGMetadata* OutMetadata);
diff --git a/Engine/Plugins/PCG/Source/PCG/Public/Metadata/PCGMetadata.h b/Engine/Plugins/PCG/Source/PCG/Public/Metadata/PCGMetadata.h
--- a/Engine/Plugins/PCG/Source/PCG/Public/Metadata/PCGMetadata.h
+++ b/Engine/Plugins/PCG/Source/PCG/Public/Metadata/PCGMetadata.h
@@ -69,7 +69,7 @@ public:
 
 	/** If we have floating metadata , setup it to match a given data type. */
 	template <typename DataType, std::enable_if_t<std::is_base_of_v<UPCGData, DataType>, int> = 0>
-	void SetupDomainsFromPCGDataType() { SetupFromPCGDataType(DataType::StaticClass()); }
+	void SetupDomainsFromPCGDataType() { SetupDomainsFromPCGDataType(DataType::StaticClass()); }
 
 	UE_API void SetupDomainsFromPCGDataType(const TSubclassOf<UPCGData>& PCGDataType);
 
