[PCG] Update PCG's default layout v1.0 and

* Rename Node Source Editor frontend to "HLSL Source"
- Remove Determinism button

TODO: Add Determinism Tests to the Tools file menu

#rnx

diff --git a/Engine/Plugins/PCG/Source/PCG/Private/Compute/Elements/PCGCustomHLSL.h b/Engine/Plugins/PCG/Source/PCG/Private/Compute/Elements/PCGCustomHLSL.h
--- a/Engine/Plugins/PCG/Source/PCG/Private/Compute/Elements/PCGCustomHLSL.h
+++ b/Engine/Plugins/PCG/Source/PCG/Private/Compute/Elements/PCGCustomHLSL.h
@@ -184,23 +184,23 @@ protected:
 #if WITH_EDITORONLY_DATA
 	// Shader source and declarations are entirely editor-only, and should never be serialized outside of the editor.
 
-	/** Optional functions that can be called from the source. Intended to be edited using the Node Source Editor window. */
+	/** Optional functions that can be called from the source. Intended to be edited using the HLSL Source Editor window. */
 	UPROPERTY()
 	FString ShaderFunctions = "/** CUSTOM SHADER FUNCTIONS **/\n";
 
-	/** Shader code that forms the body of the kernel. Intended to be edited using the Node Source Editor window. */
+	/** Shader code that forms the body of the kernel. Intended to be edited using the HLSL Source Editor window. */
 	UPROPERTY()
 	FString ShaderSource;
 
-	/** Inputs data accessors that can be used from the shader code. Intended to be viewed using the Node Source Editor window. */
+	/** Inputs data accessors that can be used from the shader code. Intended to be viewed using the HLSL Source Editor window. */
 	UPROPERTY(Transient)
 	FString InputDeclarations;
 
-	/** Output data accessors that can be used from the shader code. Intended to be viewed using the Node Source Editor window. */
+	/** Output data accessors that can be used from the shader code. Intended to be viewed using the HLSL Source Editor window. */
 	UPROPERTY(Transient)
 	FString OutputDeclarations;
 
-	/** Helper data and functions that can be used from the shader code. Intended to be viewed using the Node Source Editor window. */
+	/** Helper data and functions that can be used from the shader code. Intended to be viewed using the HLSL Source Editor window. */
 	UPROPERTY(Transient)
 	FString HelperDeclarations;
 #endif
diff --git a/Engine/Plugins/PCG/Source/PCGEditor/Private/Details/PCGCustomHLSLSettingsDetails.cpp b/Engine/Plugins/PCG/Source/PCGEditor/Private/Details/PCGCustomHLSLSettingsDetails.cpp
--- a/Engine/Plugins/PCG/Source/PCGEditor/Private/Details/PCGCustomHLSLSettingsDetails.cpp
+++ b/Engine/Plugins/PCG/Source/PCGEditor/Private/Details/PCGCustomHLSLSettingsDetails.cpp
@@ -58,11 +58,11 @@ void FPCGCustomHLSLSettingsDetails::CustomizeDetails(IDetailLayoutBuilder& Detai
 
 							return FReply::Handled();
 						})
-						.ToolTipText(FText::FromString("Opens Node Source Editor Panel."))
+						.ToolTipText(FText::FromString("Opens HLSL Source Editor Panel."))
 						[
 							SNew(STextBlock)
 							.Font(IDetailLayoutBuilder::GetDetailFont())
-							.Text(LOCTEXT("ButtonOpenSourceEditor", "Open Source Editor"))
+							.Text(LOCTEXT("ButtonOpenSourceEditor", "Open HLSL Editor"))
 						]
 					]
 				];
diff --git a/Engine/Plugins/PCG/Source/PCGEditor/Private/PCGEditor.cpp b/Engine/Plugins/PCG/Source/PCGEditor/Private/PCGEditor.cpp
--- a/Engine/Plugins/PCG/Source/PCGEditor/Private/PCGEditor.cpp
+++ b/Engine/Plugins/PCG/Source/PCGEditor/Private/PCGEditor.cpp
@@ -95,7 +95,7 @@
 	const FName DeterminismID = FName(TEXT("Determinism"));
 	const FName ProfilingID = FName(TEXT("Profiling"));
 	const FName LogID = FName(TEXT("Log"));
-	const FName NodeSourceID = FName(TEXT("NodeSource"));
+	const FName HLSLSourceID = FName(TEXT("HLSLSource"));
 	const FName UserParamsID = FName(TEXT("UserParams"));
 	const FName ViewportID[] = {
 		FName(TEXT("Viewport")),
@@ -175,58 +175,7 @@
 	BindCommands();
 	RegisterToolbar();
 
-	const TSharedRef<FTabManager::FLayout> StandaloneDefaultLayout = FTabManager::NewLayout("Standalone_PCGGraphEditor_Layout_v0.6")
-		->AddArea
-		(
-			FTabManager::NewPrimaryArea()->SetOrientation(Orient_Vertical)
-			->Split
-			(
-				FTabManager::NewSplitter()->SetOrientation(Orient_Horizontal)
-				->SetSizeCoefficient(0.72)
-				->Split
-				(
-					FTabManager::NewStack()
-					->SetSizeCoefficient(0.10f)
-					->SetHideTabWell(true)
-					->AddTab(FPCGEditor_private::PaletteID, ETabState::OpenedTab)
-				)
-				->Split
-				(
-					FTabManager::NewStack()
-					->SetSizeCoefficient(0.70)
-					->SetHideTabWell(true)
-					->AddTab(FPCGEditor_private::GraphEditorID, ETabState::OpenedTab)
-				)
-				->Split
-				(
-					FTabManager::NewStack()
-					->SetSizeCoefficient(0.20f)
-					->SetHideTabWell(true)
-					->AddTab(FPCGEditor_private::PropertyDetailsID[0], ETabState::OpenedTab)
-				)
-			)
-			->Split
-			(
-				FTabManager::NewSplitter()->SetOrientation(Orient_Horizontal)
-				->SetSizeCoefficient(0.28)
-				->Split
-				(
-					FTabManager::NewStack()
-					->SetSizeCoefficient(0.2)
-					->SetHideTabWell(true)
-					->AddTab(FPCGEditor_private::DebugObjectID, ETabState::OpenedTab)
-				)
-				->Split
-				(
-					FTabManager::NewStack()
-					->SetSizeCoefficient(0.8)
-					->SetHideTabWell(true)
-					->AddTab(FPCGEditor_private::AttributesID[0], ETabState::OpenedTab)
-					->AddTab(FPCGEditor_private::DeterminismID, ETabState::ClosedTab)
-					->AddTab(FPCGEditor_private::FindID, ETabState::ClosedTab)
-				)
-			)
-		);
+	const TSharedRef<FTabManager::FLayout> StandaloneDefaultLayout = GetDefaultLayout();
 
 	const FName PCGGraphEditorAppName = FName(TEXT("PCGEditorApp"));
 
@@ -535,7 +484,7 @@
 		case EPCGEditorPanel::NodePalette:
 			return FPCGEditor_private::PaletteID;
 		case EPCGEditorPanel::NodeSource:
-			return FPCGEditor_private::NodeSourceID;
+			return FPCGEditor_private::HLSLSourceID;
 		case EPCGEditorPanel::Profiling:
 			return FPCGEditor_private::ProfilingID;
 		case EPCGEditorPanel::PropertyDetails1:
@@ -571,11 +520,27 @@
 	}
 }
 
+void FPCGEditor::CloseGraphPanel(const EPCGEditorPanel Panel) const
+{
+	const FName PanelID = GetPanelID(Panel);
+	if (PanelID != NAME_None)
+	{
+		const TSharedPtr<SDockTab> Tab = TabManager->FindExistingLiveTab(PanelID);
+		Tab->RequestCloseTab();
+	}
+}
+
 bool FPCGEditor::IsPanelCurrentlyOpen(const EPCGEditorPanel Panel) const
 {
 	return TabManager.IsValid() && TabManager->FindExistingLiveTab(GetPanelID(Panel));
 }
 
+bool FPCGEditor::IsPanelCurrentlyForeground(const EPCGEditorPanel Panel) const
+{
+		const TSharedPtr<SDockTab> DockTab = TabManager.IsValid() ? TabManager->FindExistingLiveTab(GetPanelID(Panel)) : nullptr;
+		return DockTab.IsValid() && DockTab->IsForeground();
+}
+
 void FPCGEditor::RegisterTabSpawners(const TSharedRef<FTabManager>& InTabManager)
 {
 	WorkspaceMenuCategory = InTabManager->AddLocalWorkspaceMenuCategory(LOCTEXT("WorkspaceMenu_PCGEditor", "PCG Editor"));
@@ -647,8 +612,8 @@
 		.SetDisplayName(LOCTEXT("LogCaptureTab", "Log Capture"))
 		.SetGroup(WorkspaceMenuCategoryRef);
 
-	InTabManager->RegisterTabSpawner(FPCGEditor_private::NodeSourceID, FOnSpawnTab::CreateSP(this, &FPCGEditor::SpawnTab_NodeSource))
-		.SetDisplayName(LOCTEXT("NodeSourceTab", "Node Source Editor"))
+	InTabManager->RegisterTabSpawner(FPCGEditor_private::HLSLSourceID, FOnSpawnTab::CreateSP(this, &FPCGEditor::SpawnTab_NodeSource))
+		.SetDisplayName(LOCTEXT("NodeSourceTab", "HLSL Source"))
 		.SetGroup(WorkspaceMenuCategoryRef);
 
 	InTabManager->RegisterTabSpawner(FPCGEditor_private::UserParamsID, FOnSpawnTab::CreateSP(this, &FPCGEditor::SpawnTab_UserParams))
@@ -689,7 +654,7 @@
 	InTabManager->UnregisterTabSpawner(FPCGEditor_private::DeterminismID);
 	InTabManager->UnregisterTabSpawner(FPCGEditor_private::ProfilingID);
 	InTabManager->UnregisterTabSpawner(FPCGEditor_private::LogID);
-	InTabManager->UnregisterTabSpawner(FPCGEditor_private::NodeSourceID);
+	InTabManager->UnregisterTabSpawner(FPCGEditor_private::HLSLSourceID);
 	InTabManager->UnregisterTabSpawner(FPCGEditor_private::UserParamsID);
 	InTabManager->UnregisterTabSpawner(FPCGEditor_private::ViewportID[0]);
 	InTabManager->UnregisterTabSpawner(FPCGEditor_private::ViewportID[1]);
@@ -780,14 +745,13 @@
 {
 	UToolMenus* ToolMenus = UToolMenus::Get();
 	FName ParentName;
-	const FName MenuName = GetToolMenuToolbarName(ParentName);
-	if (!ToolMenus->IsMenuRegistered(MenuName))
+	const FName ToolbarName = GetToolMenuToolbarName(ParentName);
+	if (!ToolMenus->IsMenuRegistered(ToolbarName))
 	{
-		UToolMenu* ToolBar = ToolMenus->RegisterMenu(MenuName, ParentName, EMultiBoxType::ToolBar);
+		UToolMenu* ToolBar = ToolMenus->RegisterMenu(ToolbarName, ParentName, EMultiBoxType::ToolBar);
 
 		const FPCGEditorCommands& PCGEditorCommands = FPCGEditorCommands::Get();
-		const FToolMenuInsert InsertAfterAssetSection("Asset", EToolMenuInsertType::After);
-		FToolMenuSection& Section = ToolBar->AddSection("PCGToolbar", TAttribute<FText>(), InsertAfterAssetSection);
+		FToolMenuSection& Section = ToolBar->AddSection("PCGToolbar", TAttribute<FText>());
 
 		Section.AddEntry(FToolMenuEntry::InitToolBarButton(
 			PCGEditorCommands.Find,
@@ -829,12 +793,10 @@
 		Section.AddSeparator(NAME_None);
 
 		Section.AddEntry(FToolMenuEntry::InitToolBarButton(
-			 PCGEditorCommands.RunDeterminismGraphTest,
+			 PCGEditorCommands.ToggleGraphParams,
 			 TAttribute<FText>(),
 			 TAttribute<FText>(),
-			 FSlateIcon(FPCGEditorStyle::Get().GetStyleSetName(), "PCG.Command.RunDeterminismTest")));
-
-		Section.AddSeparator(NAME_None);
+			 FSlateIcon(FPCGEditorStyle::Get().GetStyleSetName(), "PCG.Command.OpenGraphParams")));
 
 		Section.AddEntry(FToolMenuEntry::InitToolBarButton(
 			 PCGEditorCommands.EditGraphSettings,
@@ -884,5 +846,11 @@
 		FCanExecuteAction::CreateSP(this, &FPCGEditor::CanRunDeterminismGraphTest));
 
 	ToolkitCommands->MapAction(
+		PCGEditorCommands.ToggleGraphParams,
+		FExecuteAction::CreateSP(this, &FPCGEditor::OnToggleGraphParamsPanel),
+		FCanExecuteAction(),
+		FIsActionChecked::CreateSP(this, &FPCGEditor::IsToggleGraphParamsToggled));
+
+	ToolkitCommands->MapAction(
 		PCGEditorCommands.EditGraphSettings,
 		FExecuteAction::CreateSP(this, &FPCGEditor::OnEditGraphSettings),
@@ -1341,6 +1309,23 @@
 	return false;
 }
 
+void FPCGEditor::OnToggleGraphParamsPanel() const
+{
+	if (IsPanelCurrentlyForeground(EPCGEditorPanel::UserParams))
+	{
+		CloseGraphPanel(EPCGEditorPanel::UserParams);
+	}
+	else
+	{
+		BringFocusToPanel(EPCGEditorPanel::UserParams);
+	}
+}
+
+bool FPCGEditor::IsToggleGraphParamsToggled() const
+{
+	return IsPanelCurrentlyOpen(EPCGEditorPanel::UserParams);
+}
+
 bool FPCGEditor::CanCollapseNodesInSubgraph() const
 {
 	bool HasPCGNode = false;
@@ -3283,6 +3268,70 @@
 	}
 }
 
+TSharedRef<FTabManager::FLayout> FPCGEditor::GetDefaultLayout()
+{
+	return FTabManager::NewLayout("Standalone_PCGGraphEditor_DefaultLayout_v1.0")
+		->AddArea // Main PCG Graph Editor Area
+		(
+			FTabManager::NewPrimaryArea()->SetOrientation(Orient_Vertical)
+			->Split // Top Section - Graph, Data Viewport, HLSL Source Editor, and Details View
+			(
+				FTabManager::NewSplitter()->SetOrientation(Orient_Horizontal)
+				->SetSizeCoefficient(0.65f)
+				->Split // Graph Palette
+				(
+					FTabManager::NewStack()
+					->AddTab(FPCGEditor_private::PaletteID, ETabState::SidebarTab, ESidebarLocation::Left, /*SidebarSizeCoefficient=*/0.13f)
+				)
+				->Split // Data Viewport/HLSL Source Editor
+				(
+					FTabManager::NewStack()
+					->SetSizeCoefficient(0.2f)
+					->AddTab(FPCGEditor_private::ViewportID[0], ETabState::OpenedTab)
+					->AddTab(FPCGEditor_private::HLSLSourceID, ETabState::OpenedTab)
+					->SetForegroundTab(FPCGEditor_private::ViewportID[0])
+				)
+				->Split // Node Graph
+				(
+					FTabManager::NewStack()
+					->SetSizeCoefficient(0.6f)
+					->AddTab(FPCGEditor_private::GraphEditorID, ETabState::OpenedTab)
+					->SetHideTabWell(true)
+				)
+				->Split // Details View
+				(
+					FTabManager::NewStack()
+					->SetSizeCoefficient(0.2f)
+					->AddTab(FPCGEditor_private::PropertyDetailsID[0], ETabState::OpenedTab)
+				)
+			)
+			->Split // Bottom Section - Debug/Params and ALV
+			(
+				FTabManager::NewSplitter()->SetOrientation(Orient_Horizontal)
+				->SetSizeCoefficient(0.35f)
+				->Split
+				(
+					FTabManager::NewStack()
+					->SetSizeCoefficient(0.2f)
+					->AddTab(FPCGEditor_private::DebugObjectID, ETabState::OpenedTab)
+					->AddTab(FPCGEditor_private::UserParamsID, ETabState::OpenedTab)
+					->SetForegroundTab(FPCGEditor_private::DebugObjectID)
+				)
+				->Split // ALV, Profiling, Find, Determinism
+				(
+					FTabManager::NewStack()
+					->SetSizeCoefficient(0.8f)
+					->SetHideTabWell(false)
+					->AddTab(FPCGEditor_private::AttributesID[0], ETabState::OpenedTab)
+					->AddTab(FPCGEditor_private::ProfilingID, ETabState::OpenedTab)
+					->AddTab(FPCGEditor_private::FindID, ETabState::OpenedTab)
+					->AddTab(FPCGEditor_private::DeterminismID, ETabState::ClosedTab)
+					->SetForegroundTab(FPCGEditor_private::AttributesID[0])
+				)
+			)
+		);
+}
+
 TSharedRef<SDockTab> FPCGEditor::SpawnTab_GraphEditor(const FSpawnTabArgs& Args)
 {
 	return SNew(SDockTab)
@@ -3383,7 +3432,7 @@
 TSharedRef<SDockTab> FPCGEditor::SpawnTab_NodeSource(const FSpawnTabArgs& Args)
 {
 	return SNew(SDockTab)
-		.Label(LOCTEXT("PCGNodeSourceTitle", "Node Source Editor"))
+		.Label(LOCTEXT("PCGHLSLSourceTitle", "HLSL Source"))
 		.TabColorScale(GetTabColorScale())
 		[
 			NodeSourceWidget.ToSharedRef()
diff --git a/Engine/Plugins/PCG/Source/PCGEditor/Private/PCGEditor.h b/Engine/Plugins/PCG/Source/PCGEditor/Private/PCGEditor.h
--- a/Engine/Plugins/PCG/Source/PCGEditor/Private/PCGEditor.h
+++ b/Engine/Plugins/PCG/Source/PCGEditor/Private/PCGEditor.h
@@ -101,10 +101,14 @@ public:
 
 	/** Get the TabID of the editor panel. */
 	FName GetPanelID(EPCGEditorPanel Panel) const;
-	/** Focuses the user on a specific Details View panel and flashes the tab. */
+	/** Focuses the user on a specific panel and flashes the tab. */
 	void BringFocusToPanel(EPCGEditorPanel Panel) const;
+	/** Attempts to close the specific panel if it's open. */
+	void CloseGraphPanel(EPCGEditorPanel Panel) const;
 	/** Returns true if the selected tab is currently open. */
 	bool IsPanelCurrentlyOpen(EPCGEditorPanel Panel) const;
+	/** Returns true if the selected tab is currently open and focused. */
+	bool IsPanelCurrentlyForeground(EPCGEditorPanel Panel) const;
 
 	/** Helper to get to the subsystem. */
 	static class UPCGSubsystem* GetSubsystem();
@@ -246,6 +250,11 @@ private:
 	/** Whether the PCG object being edited is opened in details view or not */
 	bool IsEditGraphSettingsToggled() const;
 
+	/** Open panel to view and edit the graph parameters. */
+	void OnToggleGraphParamsPanel() const;
+	/** Is the graph params panel is open. */
+	bool IsToggleGraphParamsToggled() const;
+
 	/** Select every node in the graph */
 	void SelectAllNodes();
 	/** Whether we can select every node */
@@ -437,6 +446,8 @@ private:
 	FText GetAttributesTabLabel(int AttributesIndex);
 	FText GetViewportTabLabel(int ViewportIndex);
 
+	static TSharedRef<FTabManager::FLayout> GetDefaultLayout();
+
 	TSharedPtr<SGraphEditor> GraphEditorWidget;
 	TArray<TSharedPtr<SPCGEditorGraphDetailsView>> PropertyDetailsWidgets;
 	TSharedPtr<SPCGEditorGraphNodePalette> PaletteWidget;
diff --git a/Engine/Plugins/PCG/Source/PCGEditor/Private/PCGEditorCommands.cpp b/Engine/Plugins/PCG/Source/PCGEditor/Private/PCGEditorCommands.cpp
--- a/Engine/Plugins/PCG/Source/PCGEditor/Private/PCGEditorCommands.cpp
+++ b/Engine/Plugins/PCG/Source/PCGEditor/Private/PCGEditorCommands.cpp
@@ -33,6 +33,7 @@ void FPCGEditorCommands::RegisterCommands()
 	UI_COMMAND(RunDeterminismNodeTest, "Run Determinism Test on Node", "Evaluate the current node for determinism.", EUserInterfaceActionType::Button, FInputChord(EModifierKey::Alt, EKeys::T));
 	UI_COMMAND(RunDeterminismGraphTest, "Graph Determinism Test", "Evaluate the current graph for determinism.", EUserInterfaceActionType::Button, FInputChord());
 	UI_COMMAND(EditGraphSettings, "Graph Settings", "Edit the graph settings.", EUserInterfaceActionType::ToggleButton, FInputChord());
+	UI_COMMAND(ToggleGraphParams, "Graph Parameters", "Open the graph settings panel.", EUserInterfaceActionType::ToggleButton, FInputChord());
 	UI_COMMAND(CancelExecution, "Cancel Execution", "Cancels the execution of the current graph", EUserInterfaceActionType::Button, FInputChord(EModifierKey::Control, EKeys::Escape));
 	UI_COMMAND(ToggleEnabled, "Toggle Enabled", "Toggle node enabled state for selected nodes.", EUserInterfaceActionType::ToggleButton, FInputChord(EKeys::E));
 	UI_COMMAND(ToggleDebug, "Toggle Debug", "Toggle node debug state for selected nodes", EUserInterfaceActionType::ToggleButton, FInputChord(EKeys::D));
diff --git a/Engine/Plugins/PCG/Source/PCGEditor/Private/PCGEditorCommands.h b/Engine/Plugins/PCG/Source/PCGEditor/Private/PCGEditorCommands.h
--- a/Engine/Plugins/PCG/Source/PCGEditor/Private/PCGEditorCommands.h
+++ b/Engine/Plugins/PCG/Source/PCGEditor/Private/PCGEditorCommands.h
@@ -30,6 +30,7 @@ public:
 	TSharedPtr<FUICommandInfo> RunDeterminismNodeTest;
 	TSharedPtr<FUICommandInfo> RunDeterminismGraphTest;
 	TSharedPtr<FUICommandInfo> EditGraphSettings;
+	TSharedPtr<FUICommandInfo> ToggleGraphParams;
 	TSharedPtr<FUICommandInfo> CancelExecution;
 	TSharedPtr<FUICommandInfo> ToggleEnabled;
 	TSharedPtr<FUICommandInfo> ToggleDebug;
diff --git a/Engine/Plugins/PCG/Source/PCGEditor/Private/PCGEditorStyle.cpp b/Engine/Plugins/PCG/Source/PCGEditor/Private/PCGEditorStyle.cpp
--- a/Engine/Plugins/PCG/Source/PCGEditor/Private/PCGEditorStyle.cpp
+++ b/Engine/Plugins/PCG/Source/PCGEditor/Private/PCGEditorStyle.cpp
@@ -138,6 +138,7 @@ FPCGEditorStyle::FPCGEditorStyle() : FSlateStyleSet("PCGEditorStyle")
 	Set("PCG.Command.ForceRegenClearCache", new IMAGE_BRUSH_SVG("Style/PCG_Command_ForceRegenClearCache", Icon20x20));
 	Set("PCG.Command.PauseRegen", new IMAGE_BRUSH_SVG("Style/PCG_Command_PauseRegen", Icon20x20));
 	Set("PCG.Command.StopRegen", new IMAGE_BRUSH_SVG("Style/PCG_Command_StopRegen", Icon20x20));
+	Set("PCG.Command.OpenGraphParams", new IMAGE_BRUSH_SVG("Style/PCG_Graph_OpenGraphParams", Icon20x20));
 	Set("PCG.Command.GraphSettings", new IMAGE_BRUSH_SVG("Style/PCG_Command_GraphSettings", Icon20x20));
 	Set("PCG.Command.OpenDebugTreeTab", new IMAGE_BRUSH_SVG("Style/PCG_Command_OpenDebugTreeTab", Icon20x20));
 	Set("PCG.Command.RunDeterminismTest", new IMAGE_BRUSH_SVG("Style/PCG_Command_RunDeterminismTest", Icon20x20));
diff --git a/Engine/Plugins/PCG/Source/PCGEditor/Private/Widgets/SPCGNodeSourceTextBox.cpp b/Engine/Plugins/PCG/Source/PCGEditor/Private/Widgets/SPCGNodeSourceTextBox.cpp
--- a/Engine/Plugins/PCG/Source/PCGEditor/Private/Widgets/SPCGNodeSourceTextBox.cpp
+++ b/Engine/Plugins/PCG/Source/PCGEditor/Private/Widgets/SPCGNodeSourceTextBox.cpp
@@ -55,7 +55,7 @@ static bool IsWhiteSpace(const TCHAR& InCharacter)
 FPCGNodeSourceEditorTextBoxCommands::FPCGNodeSourceEditorTextBoxCommands() 
 	: TCommands<FPCGNodeSourceEditorTextBoxCommands>(
 		"PCGShaderTextEditorDocumentTextBox", // Context name for fast lookup
-		NSLOCTEXT("Contexts", "PCGShaderTextEditorDocumentTextBox", "PCG Node Source Editor Document TextBox"), // Localized context name for display
+		NSLOCTEXT("Contexts", "PCGShaderTextEditorDocumentTextBox", "PCG HLSL Source Editor Document TextBox"), // Localized context name for display
 		NAME_None,
 		FAppStyle::GetAppStyleSetName()
 	)
