[PCG] Fixed minor issue where bGetImpactFace was required to retrieve other mesh related information

* Abstracted conversion of FPCGRaycastQueryParams into FCollisionQueryParams for consistency.

#rnx

diff --git a/Engine/Plugins/PCG/Source/PCG/Private/Data/PCGWorldData.cpp b/Engine/Plugins/PCG/Source/PCG/Private/Data/PCGWorldData.cpp
--- a/Engine/Plugins/PCG/Source/PCG/Private/Data/PCGWorldData.cpp
+++ b/Engine/Plugins/PCG/Source/PCG/Private/Data/PCGWorldData.cpp
@@ -298,10 +298,7 @@ bool UPCGWorldRayHitData::SamplePoint(const FTransform& InTransform, const FBox&
 
 	// Todo: consider prebuilding this
 	const FCollisionObjectQueryParams ObjectQueryParams(QueryParams.CollisionChannel);
-	FCollisionQueryParams CollisionQueryParams; // TODO: apply properties from the settings when/if they exist
-	CollisionQueryParams.bTraceComplex = QueryParams.bTraceComplex;
-	CollisionQueryParams.bReturnPhysicalMaterial = QueryParams.bGetReferenceToPhysicalMaterial;
-	CollisionQueryParams.bReturnFaceIndex = QueryParams.bGetFaceIndex || QueryParams.bGetUVCoords;
+	const FCollisionQueryParams CollisionQueryParams = FPCGWorldRaycastQueryParams::ConvertPCGRaycastQueryToCollisionQuery(QueryParams); // TODO: apply properties from the settings when/if they exist
 
 	FVector RayOrigin = InTransform.GetLocation() - ((InTransform.GetLocation() - QueryParams.RayOrigin) | QueryParams.RayDirection) * QueryParams.RayDirection;
 	FVector RayEnd = RayOrigin + QueryParams.RayDirection * QueryParams.RayLength;
diff --git a/Engine/Plugins/PCG/Source/PCG/Private/Elements/PCGWorldRaycast.cpp b/Engine/Plugins/PCG/Source/PCG/Private/Elements/PCGWorldRaycast.cpp
--- a/Engine/Plugins/PCG/Source/PCG/Private/Elements/PCGWorldRaycast.cpp
+++ b/Engine/Plugins/PCG/Source/PCG/Private/Elements/PCGWorldRaycast.cpp
@@ -259,9 +259,7 @@ bool FPCGWorldRaycastElement::PrepareDataInternal(FPCGContext* InContext) const
 			return EPCGTimeSliceInitResult::NoOperation;
 		}
 
-		OutState.CollisionQueryParams.bTraceComplex = Settings->WorldQueryParams.bTraceComplex;
-		OutState.CollisionQueryParams.bReturnPhysicalMaterial = Settings->WorldQueryParams.bGetReferenceToPhysicalMaterial;
-		OutState.CollisionQueryParams.bReturnFaceIndex = Settings->WorldQueryParams.bGetFaceIndex || Settings->WorldQueryParams.bGetUVCoords;
+		OutState.CollisionQueryParams = FPCGWorldRaycastQueryParams::ConvertPCGRaycastQueryToCollisionQuery(Settings->WorldQueryParams);
 		OutState.CollisionObjectQueryParams = FCollisionObjectQueryParams(Settings->WorldQueryParams.CollisionChannel);
 
 		return EPCGTimeSliceInitResult::Success;
diff --git a/Engine/Plugins/PCG/Source/PCG/Public/Data/PCGWorldData.h b/Engine/Plugins/PCG/Source/PCG/Public/Data/PCGWorldData.h
--- a/Engine/Plugins/PCG/Source/PCG/Public/Data/PCGWorldData.h
+++ b/Engine/Plugins/PCG/Source/PCG/Public/Data/PCGWorldData.h
@@ -7,6 +7,7 @@
 #include "PCGSurfaceData.h"
 #include "Helpers/PCGWorldQueryHelpers.h"
 
+#include "CollisionQueryParams.h"
 #include "Engine/EngineTypes.h"
 #include "UObject/ObjectKey.h"
 
@@ -210,6 +211,23 @@ struct FPCGWorldRaycastQueryParams : public FPCGWorldCommonQueryParams
 	/** This UV Channel will be selected when retrieving UV Coordinates from a raycast query. */
 	UPROPERTY(BlueprintReadWrite, EditAnywhere, Category = "Data|Attributes", meta = (PCG_Overridable, EditCondition = "bTraceComplex && bGetUVCoords", EditConditionHides, DisplayAfter = "bGetUVCoords"))
 	int32 UVChannel = 0;
+
+	// Helper function to ensure cohesive conversion of FPCGWorldRaycastQueryParams->FCollisionQueryParams.
+	static FCollisionQueryParams ConvertPCGRaycastQueryToCollisionQuery(const FPCGWorldRaycastQueryParams& InParams)
+	{
+		FCollisionQueryParams Params;
+		Params.bTraceComplex = InParams.bTraceComplex;
+		Params.bReturnPhysicalMaterial = InParams.bGetReferenceToPhysicalMaterial;
+		// We need the face index to calculate UVs, Render Material, and Mesh information.
+		Params.bReturnFaceIndex =
+			InParams.bGetFaceIndex
+			|| InParams.bGetUVCoords
+			|| InParams.bGetSectionIndex
+			|| InParams.bGetRenderMaterialIndex
+			|| InParams.bGetReferenceToRenderMaterial;
+
+		return Params;
+	}
 };
 
 template<>
