[PCG] Fix error message when a tracked non-actor is changed

#rnx

diff --git a/Engine/Plugins/PCG/Source/PCG/Private/PCGActorAndComponentMapping.cpp b/Engine/Plugins/PCG/Source/PCG/Private/PCGActorAndComponentMapping.cpp
--- a/Engine/Plugins/PCG/Source/PCG/Private/PCGActorAndComponentMapping.cpp
+++ b/Engine/Plugins/PCG/Source/PCG/Private/PCGActorAndComponentMapping.cpp
@@ -1867,7 +1867,7 @@ void FPCGActorAndComponentMapping::OnObjectChanged(UObject* InObject, const ECha
 	}
 
 	// We need the actor bounds to know if we are intersecting, for the culled settings to be dirtied.
-	const FBox ActorBounds = PCGHelpers::GetActorBounds(Actor, /*bIgnorePCGCreatedComponents=*/false);
+	const FBox ActorBounds = Actor ? PCGHelpers::GetActorBounds(Actor, /*bIgnorePCGCreatedComponents=*/false) : FBox{EForceInit::ForceInit};
 
 	{
 		TRACE_CPUPROFILER_EVENT_SCOPE(FPCGActorAndComponentMapping::OnObjectChanged::AlwaysTrackedUpdate);
@@ -1887,9 +1887,9 @@ void FPCGActorAndComponentMapping::OnObjectChanged(UObject* InObject, const ECha
 			// Since when we clear the cache for a settings, we clear it all, it's only necessary to do it once on the original component. And we will dirty all the local components (with dirty dispatch)
 			if (!DirtyComponents.Contains(PCGComponent))
 			{
-				const FBox ComponentBounds = PCGComponent->GetGridBounds();
-				const bool bIntersect = bOwnerChanged || ActorBounds.Intersect(ComponentBounds);// || (OldActorBoundsPtr && OldActorBoundsPtr->Intersect(ComponentBounds));
-
+				// Technically, a key that is always tracked can also be culled tracked, so maybe having an intersection test here to not
+				// dirty components that culled tracked and don't intersect with the component would be more efficient.
+				// In practice, it's mostly relevent for partitioned graphs and we will dirty all the locals anyway.
 				bShouldDirty |= ClearCacheForKeys(MatchedKeys, PCGComponent, /*bIntersect=*/true, InOriginatingChangeObject);
 				if (bShouldDirty)
 				{
