[PCG] Unify render doc capture cvars

#rnx

diff --git a/Engine/Plugins/PCG/Source/PCG/Private/Compute/DataInterfaces/PCGInstanceDataInterface.cpp b/Engine/Plugins/PCG/Source/PCG/Private/Compute/DataInterfaces/PCGInstanceDataInterface.cpp
--- a/Engine/Plugins/PCG/Source/PCG/Private/Compute/DataInterfaces/PCGInstanceDataInterface.cpp
+++ b/Engine/Plugins/PCG/Source/PCG/Private/Compute/DataInterfaces/PCGInstanceDataInterface.cpp
@@ -45,7 +45,7 @@ namespace PCGInstanceDataInterface
 #if !UE_BUILD_SHIPPING
 	static int32 TriggerGPUCaptureDispatchIndex = 0;
 	static FAutoConsoleVariableRef CVarTriggerGPUCaptureDispatchIndex(
-		TEXT("pcg.GPU.Instancing.TriggerSceneWriterGPUCaptureDispatchIndex"),
+		TEXT("pcg.GPU.TriggerRenderCaptures.InstanceSceneWriter"),
 		TriggerGPUCaptureDispatchIndex,
 		TEXT("Index of the next dispatch to capture. I.e. if set to 2, will ignore one dispatch and then trigger a capture on the next one."),
 		ECVF_RenderThreadSafe
@@ -372,7 +372,7 @@ void FPCGInstanceDataProviderProxy::PostSubmit(FRDGBuilder& InGraphBuilder) cons
 #if UE_BUILD_SHIPPING
 				const bool bTriggerCapture = false;
 #else
-				const bool bTriggerCapture = (PCGInstanceDataInterface::TriggerGPUCaptureDispatchIndex == 1);
+				const bool bTriggerCapture = (PCGInstanceDataInterface::TriggerGPUCaptureDispatchIndex > 0);
 				PCGInstanceDataInterface::TriggerGPUCaptureDispatchIndex = FMath::Max(PCGInstanceDataInterface::TriggerGPUCaptureDispatchIndex - 1, 0);
 #endif // UE_BUILD_SHIPPING
 
diff --git a/Engine/Plugins/PCG/Source/PCG/Private/Compute/PCGDataUploadAdaptor.cpp b/Engine/Plugins/PCG/Source/PCG/Private/Compute/PCGDataUploadAdaptor.cpp
--- a/Engine/Plugins/PCG/Source/PCG/Private/Compute/PCGDataUploadAdaptor.cpp
+++ b/Engine/Plugins/PCG/Source/PCG/Private/Compute/PCGDataUploadAdaptor.cpp
@@ -19,7 +19,7 @@ namespace PCGDataUploadAdaptor
 #if !UE_BUILD_SHIPPING
 	static int32 TriggerGPUCaptureDispatches = 0;
 	static FAutoConsoleVariableRef CVarTriggerGPUCapture(
-		TEXT("pcg.GPU.DataAdaptor.TriggerGPUCaptures"),
+		TEXT("pcg.GPU.TriggerRenderCaptures.DataUploadAdaptor"),
 		TriggerGPUCaptureDispatches,
 		TEXT("Trigger GPU captures for this many of the subsequent data collection update dispatches."));
 #endif // !UE_BUILD_SHIPPING
diff --git a/Engine/Plugins/PCG/Source/PCG/Private/Data/PCGTextureData.cpp b/Engine/Plugins/PCG/Source/PCG/Private/Data/PCGTextureData.cpp
--- a/Engine/Plugins/PCG/Source/PCG/Private/Data/PCGTextureData.cpp
+++ b/Engine/Plugins/PCG/Source/PCG/Private/Data/PCGTextureData.cpp
@@ -21,7 +21,7 @@ namespace PCGTextureSamplingHelpers
 {
 	static int32 GTriggerReadbackCaptureDispatches = 0;
 	static FAutoConsoleVariableRef CVarTriggerReadbackCapture(
-		TEXT("pcg.Texture.TriggerReadbackCaptures"),
+		TEXT("pcg.GPU.TriggerRenderCaptures.TextureReadback"),
 		GTriggerReadbackCaptureDispatches,
 		TEXT("Trigger GPU readback captures for this many of the subsequent texture data initializations."));
 
diff --git a/Engine/Plugins/PCG/Source/PCG/Private/Elements/PCGGenerateGrassMaps.cpp b/Engine/Plugins/PCG/Source/PCG/Private/Elements/PCGGenerateGrassMaps.cpp
--- a/Engine/Plugins/PCG/Source/PCG/Private/Elements/PCGGenerateGrassMaps.cpp
+++ b/Engine/Plugins/PCG/Source/PCG/Private/Elements/PCGGenerateGrassMaps.cpp
@@ -36,7 +36,7 @@ namespace PCGGenerateGrassMaps
 
 	static int32 GTriggerGPUCaptureDispatches = 0;
 	static FAutoConsoleVariableRef CVarTriggerGPUCapture(
-		TEXT("pcg.Grass.TriggerGPUCaptures"),
+		TEXT("pcg.GPU.TriggerRenderCaptures.GrassMapGeneration"),
 		GTriggerGPUCaptureDispatches,
 		TEXT("Trigger GPU captures for this many of the subsequent grass generations."));
 
diff --git a/Engine/Plugins/PCG/Source/PCGEditor/Private/DataVisualizations/PCGVisualizationTexture2D.cpp b/Engine/Plugins/PCG/Source/PCGEditor/Private/DataVisualizations/PCGVisualizationTexture2D.cpp
--- a/Engine/Plugins/PCG/Source/PCGEditor/Private/DataVisualizations/PCGVisualizationTexture2D.cpp
+++ b/Engine/Plugins/PCG/Source/PCGEditor/Private/DataVisualizations/PCGVisualizationTexture2D.cpp
@@ -14,7 +14,7 @@ namespace PCGVisTexture2DHelpers
 {
 	static int32 GTriggerGPUCaptureDispatches = 0;
 	static FAutoConsoleVariableRef CVarTriggerGPUCapture(
-		TEXT("pcg.RenderCaptures.TriggerVisualizationTextureResourceInitCaptures"),
+		TEXT("pcg.GPU.TriggerRenderCaptures.VisTexture2DInit"),
 		GTriggerGPUCaptureDispatches,
 		TEXT("Trigger GPU captures for this many of the subsequent visualization texture resource initializations."));
 
