[PCG] Clean up Get Virtual Texture Data node UI

#rnx

diff --git a/Engine/Plugins/PCG/Source/PCG/Private/Elements/PCGTypedGetter.cpp b/Engine/Plugins/PCG/Source/PCG/Private/Elements/PCGTypedGetter.cpp
--- a/Engine/Plugins/PCG/Source/PCG/Private/Elements/PCGTypedGetter.cpp
+++ b/Engine/Plugins/PCG/Source/PCG/Private/Elements/PCGTypedGetter.cpp
@@ -243,6 +243,7 @@ UPCGGetVirtualTextureSettings::UPCGGetVirtualTextureSettings()
 	ActorSelector.bShowActorFilter = false;
 	ActorSelector.ActorSelection = EPCGActorSelection::ByClass;
 	ActorSelector.ActorSelectionClass = ARuntimeVirtualTextureVolume::StaticClass();
+	ActorSelector.bShowActorSelection = false;
 	ActorSelector.bShowActorSelectionClass = false;
 	ActorSelector.bIncludeChildren = false;
 	ActorSelector.bSelectMultiple = true;
