[PCG] Release preview resources when unloading PCG Component

diff --git a/Engine/Plugins/PCG/Source/PCG/Private/Helpers/PCGActorHelpers.cpp b/Engine/Plugins/PCG/Source/PCG/Private/Helpers/PCGActorHelpers.cpp
--- a/Engine/Plugins/PCG/Source/PCG/Private/Helpers/PCGActorHelpers.cpp
+++ b/Engine/Plugins/PCG/Source/PCG/Private/Helpers/PCGActorHelpers.cpp
@@ -21,8 +21,11 @@
 
 #if WITH_EDITOR
 #include "Editor.h"
+#include "Selection.h"
+#include "UnrealEdGlobals.h"
 #include "Engine/Level.h"
 #include "Editor/IPCGEditorModule.h"
+#include "Editor/UnrealEdEngine.h"
 #include "Engine/Blueprint.h"
 #include "Modules/ModuleManager.h"
 #include "Subsystems/ActorEditorContextSubsystem.h"
@@ -519,95 +522,40 @@ bool UPCGActorHelpers::DeleteActors(UWorld* World, const TArray<TSoftObjectPtr<A
 	}
 
 #if WITH_EDITOR
-	// Remove potential references to to-be deleted objects from the global selection sets.
-	/*if (GIsEditor)
-	{
-		GEditor->ResetAllSelectionSets();
-	}*/
-
-	/** Note: the following code block is commented out because this is currently getting hit when partition actors & partitioned components are unloaded.
-	* Destroying the actors is fine, but trying to delete them from SCC is a major issue and it can hang UE.
-	* There are multiple code paths calling this method that would need to be revisited if we want to control this a bit better
-	* or we can centralize these kind of requirements in the subsystem - this could then be taken care of in a single place, knowing the actual setting.
-	*/
-	//UWorldPartition* WorldPartition = World->GetWorldPartition();
-	//if (!PCGHelpers::IsRuntimeOrPIE() && WorldPartition)
-	//{
-	//	TArray<FString> PackagesToDeleteFromSCC;
-	//	TSet<UPackage*> PackagesToCleanup;
-
-	//	for (const TSoftObjectPtr<AActor>& ManagedActor : ActorsToDelete)
-	//	{
-	//		// If actor is loaded, just remove from world and keep track of package to cleanup
-	//		if (AActor* Actor = ManagedActor.Get())
-	//		{
-	//			if (UPackage* ActorPackage = Actor->GetExternalPackage())
-	//			{
-	//				PackagesToCleanup.Emplace(ActorPackage);
-	//			}
-	//			
-	//			if (Actor->GetWorld() == World)
-	//			{
-	//				World->DestroyActor(Actor);
-	//			}
-	//			else
-	//			{
-	//				// If we're here and the world is null, then the actor has either been destroyed already or it will be picked up by GC by design.
-	//				// Otherwise, we have bigger issues, something is very wrong.
-	//				check(Actor->GetWorld() == nullptr);
-	//			}
-	//		}
-	//		// Otherwise, get from World Partition.
-	//		// Note that it is possible that some actors don't exist anymore, so a null here is not a critical condition
-	//		else if (const FWorldPartitionActorDesc* ActorDesc = WorldPartition->GetActorDesc(ManagedActor.ToSoftObjectPath()))
-	//		{
-	//			PackagesToDeleteFromSCC.Emplace(ActorDesc->GetActorPackage().ToString());
-	//			WorldPartition->RemoveActor(ActorDesc->GetGuid());
-	//		}
-	//	}
-
-	//	// Save currently loaded packages so they get deleted
-	//	if (PackagesToCleanup.Num() > 0)
-	//	{
-	//		ObjectTools::CleanupAfterSuccessfulDelete(PackagesToCleanup.Array(), /*bPerformReferenceCheck=*/true);
-	//	}
-
-	//	// Delete outstanding unloaded packages
-	//	if (PackagesToDeleteFromSCC.Num() > 0)
-	//	{
-	//		FPackageSourceControlHelper PackageHelper;
-	//		if (!PackageHelper.Delete(PackagesToDeleteFromSCC))
-	//		{
-	//			return false;
-	//		}
-	//	}
-	//}
-	//else
-#endif
+	if (IPCGEditorModule* EditorModule = IPCGEditorModule::Get())
 	{
-#if WITH_EDITOR
-
-		if (IPCGEditorModule* EditorModule = IPCGEditorModule::Get())
-		{
-			EditorModule->SetOutlinerUIRefreshDelay(UE::PCGActorHelpers::Local::OutlinerUIRefreshDelay);
-		}
+		EditorModule->SetOutlinerUIRefreshDelay(UE::PCGActorHelpers::Local::OutlinerUIRefreshDelay);
+	}
 
+	TArray<AActor*> DeleteActors;
 #endif
 
-		// Not in editor, really unlikely to happen but might be slow
-		for (const TSoftObjectPtr<AActor>& ManagedActor : ActorsToDelete)
+	for (const TSoftObjectPtr<AActor>& ManagedActor : ActorsToDelete)
+	{
+		// @todo_pcg: Revisit this GetWorld() check when fixing UE-215065
+		if (AActor* Actor = ManagedActor.Get(); Actor && Actor->GetWorld())
 		{
-			// @todo_pcg: Revisit this GetWorld() check when fixing UE-215065
-			if (AActor* Actor = ManagedActor.Get(); Actor && Actor->GetWorld())
+#if WITH_EDITOR
+			if (!Actor->GetWorld()->IsGameWorld())
 			{
-				if (!ensure(World->DestroyActor(Actor)))
-				{
-					UE_LOG(LogPCG, Warning, TEXT("Actor %s failed to be destroyed."), *Actor->GetPathName());
-				}
+				DeleteActors.Add(Actor);
+			}
+			else
+#endif
+			if (!ensure(World->DestroyActor(Actor)))
+			{
+				UE_LOG(LogPCG, Warning, TEXT("Actor %s failed to be destroyed."), *Actor->GetPathName());
 			}
 		}
 	}
 
+#if WITH_EDITOR
+	if (!DeleteActors.IsEmpty())
+	{
+		GUnrealEd->DeleteActors(DeleteActors, World, GUnrealEd->GetSelectedActors()->GetElementSelectionSet(), /*bVerifyDeletionCanHappen=*/false, /*bWarnAboutReferences=*/false, /*bWarnAboutSoftReferences=*/false);
+	}
+#endif
+
 	return true;
 }
 
diff --git a/Engine/Plugins/PCG/Source/PCG/Private/PCGComponent.cpp b/Engine/Plugins/PCG/Source/PCG/Private/PCGComponent.cpp
--- a/Engine/Plugins/PCG/Source/PCG/Private/PCGComponent.cpp
+++ b/Engine/Plugins/PCG/Source/PCG/Private/PCGComponent.cpp
@@ -96,6 +96,11 @@ namespace PCGComponent
 		false,
 		TEXT("At the end of execution, generated output data will be converted from UPCGPointArrayData to UPCGPointData."));
 
+	static TAutoConsoleVariable<bool> CVarDeletePreviewResourcesWhenUnloading(
+		TEXT("pcg.Advanced.DeletePreviewResourcesWhenUnloading"),
+		true,
+		TEXT("Delete generated preview resources (actors) when PCG Component gets unloaded."));
+
 	struct FPCGManagedActorLoadingScope
 	{
 		FPCGManagedActorLoadingScope()
@@ -1835,6 +1840,9 @@ void UPCGComponent::BeginPlay()
 
 void UPCGComponent::EndPlay(const EEndPlayReason::Type EndPlayReason)
 {
+	// General comment: we shouldn't usually be cleaning up resources in the EndPlay call.
+	// There might be an exception currently with unregistering of Local components (UE-215065) which needs fixing.
+	
 	// Always try to unregister itself, if it doesn't exist, it will early out. 
 	// Just making sure that we don't left some resources registered while dead.
 	if (UPCGSubsystem* Subsystem = GetSubsystem())
@@ -1858,6 +1866,40 @@ void UPCGComponent::OnUnregister()
 
 		Subsystem->OnPCGComponentUnregistered.Broadcast(IsValidChecked(this) ? this : nullptr);
 	}
+
+	// We shouldn't cleanup resources in OnUnregister in most cases.
+	// This specific case is to handle World Partition unloading of actors where our actor might have a PCG Component which generated some Preview Resources.
+	// In this case we do want to release those resources and the actors they might hold on to. If we don't those actors become orphaned and we might end up with duplicates when reloading the component.
+	if (PCGComponent::CVarDeletePreviewResourcesWhenUnloading.GetValueOnAnyThread() &&
+		CurrentEditingMode == EPCGEditorDirtyMode::Preview && // We have preview resources
+		UE::GetIsEditorLoadingPackage() && // We are loading (we don't want to release resources if we are getting unregistered outside of a loading scope)
+		(!IsCreatedByConstructionScript() || !GIsReconstructingBlueprintInstances) && // But not inside a construction script
+		(GetWorld() && GetWorld()->IsPartitionedWorld())) // and in a Partitioned World because if not the preview resources will get unloaded with us (no partial unloading of levels)
+	{
+		TSet<TSoftObjectPtr<AActor>> ActorsToDelete;
+		for (UPCGManagedResource* ManagedResource : GeneratedResources)
+		{
+			// Make sure resource is indeed a transient/preview resource
+			if (ManagedResource &&
+				!ManagedResource->IsMarkedTransientOnLoad() &&
+				!LoadedPreviewResources.Contains(ManagedResource))
+			{
+				ManagedResource->Release(/*bHardRelease=*/true, ActorsToDelete);
+			}
+		}
+
+		if (!ActorsToDelete.IsEmpty())
+		{
+			// Dispatch delete of actors to next tick
+			::ExecuteOnGameThread(UE_SOURCE_LOCATION, [WeakWorld = TWeakObjectPtr<UWorld>(GetWorld()), InActorsToDelete = MoveTemp(ActorsToDelete)]()
+				{
+					if (UWorld* World = WeakWorld.Get())
+					{
+						UPCGActorHelpers::DeleteActors(World, InActorsToDelete.Array());
+					}
+				});
+		}
+	}
 #endif // WITH_EDITOR
 
 	Super::OnUnregister();
