[PCG] Fix GPU nodes with internal data emitting information to the ALV

#rnx

diff --git a/Engine/Plugins/PCG/Source/PCG/Internal/Compute/PCGComputeCommon.h b/Engine/Plugins/PCG/Source/PCG/Internal/Compute/PCGComputeCommon.h
--- a/Engine/Plugins/PCG/Source/PCG/Internal/Compute/PCGComputeCommon.h
+++ b/Engine/Plugins/PCG/Source/PCG/Internal/Compute/PCGComputeCommon.h
@@ -141,7 +141,7 @@ struct FPCGPinReference
 	}
 
 	/** Reference a pin by kernel and label. */
-	explicit FPCGPinReference(const UPCGComputeKernel* InKernel, FName InLabel)
+	explicit FPCGPinReference(UPCGComputeKernel* InKernel, FName InLabel)
 		: Kernel(InKernel)
 		, Label(InLabel)
 	{
@@ -154,7 +154,7 @@ struct FPCGPinReference
 	}
 
 	/** Associated kernel. If null then compiler will look for pin on owning node. */
-	const UPCGComputeKernel* Kernel = nullptr;
+	UPCGComputeKernel* Kernel = nullptr;
 
 	/** Pin label. */
 	FName Label;
@@ -174,6 +174,9 @@ struct FPCGKernelEdge
 	bool IsConnectedToNodeInput() const { return UpstreamPin.Kernel == nullptr; }
 	bool IsConnectedToNodeOutput() const { return DownstreamPin.Kernel == nullptr; }
 
+	UPCGComputeKernel* GetUpstreamKernel() const { return UpstreamPin.Kernel; }
+	UPCGComputeKernel* GetDownstreamKernel() const { return DownstreamPin.Kernel; }
+
 	FPCGPinReference UpstreamPin;
 	FPCGPinReference DownstreamPin;
 };
diff --git a/Engine/Plugins/PCG/Source/PCG/Private/Compute/DataInterfaces/PCGComputeDataInterface.cpp b/Engine/Plugins/PCG/Source/PCG/Private/Compute/DataInterfaces/PCGComputeDataInterface.cpp
--- a/Engine/Plugins/PCG/Source/PCG/Private/Compute/DataInterfaces/PCGComputeDataInterface.cpp
+++ b/Engine/Plugins/PCG/Source/PCG/Private/Compute/DataInterfaces/PCGComputeDataInterface.cpp
@@ -12,6 +12,16 @@
 
 #include UE_INLINE_GENERATED_CPP_BY_NAME(PCGComputeDataInterface)
 
+#if WITH_EDITOR
+namespace PCGComputeDataInterfaceHelpers
+{
+	static TAutoConsoleVariable<bool> CVarEnableInternalKernelPinInspection(
+		TEXT("pcg.GPU.EnableInternalKernelPinInspection"),
+		false,
+		TEXT("Allows internal kernel pins to be inspected and debugged. Note, this only works if the internal kernel pin name matches a pin name on the node."));
+}
+#endif
+
 void UPCGComputeDataInterface::AddDownstreamInputPin(FName InInputPinLabel, const FName* InOptionalInputPinLabelAlias)
 {
 	DownstreamInputPinLabelAliases.AddUnique(InOptionalInputPinLabelAlias ? *InOptionalInputPinLabelAlias : InInputPinLabel);
@@ -99,15 +109,21 @@ void UPCGExportableDataProvider::Initialize(const UComputeDataInterface* InDataI
 #if WITH_EDITOR
 	check(ProducerSettings);
 
-	const IPCGGraphExecutionSource* ExecutionSource = Binding->GetExecutionSource();
-	if (ProducerSettings->bIsInspecting && ExecutionSource && ExecutionSource->GetExecutionState().GetInspection().IsInspecting())
-	{
-		ExportMode |= EPCGExportMode::Inspection;
-	}
+	const bool bCanInspectOrDebug = !ProducerKernel || !ProducerKernel->IsPinInternal(OutputPinLabel) || PCGComputeDataInterfaceHelpers::CVarEnableInternalKernelPinInspection.GetValueOnAnyThread();
 
-	if (ProducerSettings->bDebug)
+	if (bCanInspectOrDebug)
 	{
-		ExportMode |= EPCGExportMode::DebugVisualization;
+		const IPCGGraphExecutionSource* ExecutionSource = Binding->GetExecutionSource();
+
+		if (ProducerSettings->bIsInspecting && ExecutionSource && ExecutionSource->GetExecutionState().GetInspection().IsInspecting())
+		{
+			ExportMode |= EPCGExportMode::Inspection;
+		}
+
+		if (ProducerSettings->bDebug)
+		{
+			ExportMode |= EPCGExportMode::DebugVisualization;
+		}
 	}
 #endif
 }
diff --git a/Engine/Plugins/PCG/Source/PCG/Private/Compute/PCGComputeKernel.h b/Engine/Plugins/PCG/Source/PCG/Private/Compute/PCGComputeKernel.h
--- a/Engine/Plugins/PCG/Source/PCG/Private/Compute/PCGComputeKernel.h
+++ b/Engine/Plugins/PCG/Source/PCG/Private/Compute/PCGComputeKernel.h
@@ -97,6 +97,12 @@ public:
 
 	/** Whether compute graphs should be split at the output of this node. Useful for reading back runtime statistics or diagnostic info before continuing. */
 	virtual bool SplitGraphAtOutput() const { return false; }
+
+	/** Mark a pin as being internal to the kernel graph. This means it is not on the CPU/GPU boundary. */
+	void AddInternalPin(FName InPinLabel) { InternalPinLabels.Add(InPinLabel); }
+
+	/** Whether the pin matching the given label should be eligible for inspect/debug. */
+	bool IsPinInternal(FName InPinLabel) const { return InternalPinLabels.Contains(InPinLabel); }
 #endif
 
 	/** Compute a description of data that will be output from pin InOutputPinLabel. */
@@ -173,4 +179,8 @@ protected:
 	/** Log entries created and cached on initialization. Are logged during AreKernelSettingsValid(), unless PCG_KERNEL_LOGGING_ENABLED is false. */
 	UPROPERTY()
 	TArray<FPCGKernelLogEntry> StaticLogEntries;
+
+#if WITH_EDITOR
+	TSet<FName> InternalPinLabels;
+#endif
 };
diff --git a/Engine/Plugins/PCG/Source/PCG/Private/Graph/PCGGraphCompilerGPU.cpp b/Engine/Plugins/PCG/Source/PCG/Private/Graph/PCGGraphCompilerGPU.cpp
--- a/Engine/Plugins/PCG/Source/PCG/Private/Graph/PCGGraphCompilerGPU.cpp
+++ b/Engine/Plugins/PCG/Source/PCG/Private/Graph/PCGGraphCompilerGPU.cpp
@@ -291,6 +291,22 @@ void FPCGGraphCompilerGPU::ExpandGPUNodeKernelsToTasks(
 			continue;
 		}
 
+		for (FPCGKernelEdge& KernelEdge : NodeKernelEdges)
+		{
+			if (!KernelEdge.IsConnectedToNodeInput() && !KernelEdge.IsConnectedToNodeOutput())
+			{
+				if (UPCGComputeKernel* UpstreamKernel = KernelEdge.GetUpstreamKernel())
+				{
+					UpstreamKernel->AddInternalPin(KernelEdge.UpstreamPin.Label);
+				}
+
+				if (UPCGComputeKernel* DownstreamKernel = KernelEdge.GetDownstreamKernel())
+				{
+					DownstreamKernel->AddInternalPin(KernelEdge.DownstreamPin.Label);
+				}
+			}
+		}
+
 		KernelToTaskId.Reset();
 
 		// Create a new task for each kernel emitted by the node.
