[PCG] Change ConvertPCGRaycastQueryToCollisionQuery from static to member

#rnx

diff --git a/Engine/Plugins/PCG/Source/PCG/Private/Data/PCGWorldData.cpp b/Engine/Plugins/PCG/Source/PCG/Private/Data/PCGWorldData.cpp
--- a/Engine/Plugins/PCG/Source/PCG/Private/Data/PCGWorldData.cpp
+++ b/Engine/Plugins/PCG/Source/PCG/Private/Data/PCGWorldData.cpp
@@ -298,7 +298,7 @@ bool UPCGWorldRayHitData::SamplePoint(const FTransform& InTransform, const FBox&
 
 	// Todo: consider prebuilding this
 	const FCollisionObjectQueryParams ObjectQueryParams(QueryParams.CollisionChannel);
-	const FCollisionQueryParams CollisionQueryParams = FPCGWorldRaycastQueryParams::ConvertPCGRaycastQueryToCollisionQuery(QueryParams); // TODO: apply properties from the settings when/if they exist
+	const FCollisionQueryParams CollisionQueryParams = QueryParams.ToCollisionQuery(); // TODO: apply properties from the settings when/if they exist
 
 	FVector RayOrigin = InTransform.GetLocation() - ((InTransform.GetLocation() - QueryParams.RayOrigin) | QueryParams.RayDirection) * QueryParams.RayDirection;
 	FVector RayEnd = RayOrigin + QueryParams.RayDirection * QueryParams.RayLength;
diff --git a/Engine/Plugins/PCG/Source/PCG/Private/Elements/PCGWorldRaycast.cpp b/Engine/Plugins/PCG/Source/PCG/Private/Elements/PCGWorldRaycast.cpp
--- a/Engine/Plugins/PCG/Source/PCG/Private/Elements/PCGWorldRaycast.cpp
+++ b/Engine/Plugins/PCG/Source/PCG/Private/Elements/PCGWorldRaycast.cpp
@@ -259,7 +259,7 @@ bool FPCGWorldRaycastElement::PrepareDataInternal(FPCGContext* InContext) const
 			return EPCGTimeSliceInitResult::NoOperation;
 		}
 
-		OutState.CollisionQueryParams = FPCGWorldRaycastQueryParams::ConvertPCGRaycastQueryToCollisionQuery(Settings->WorldQueryParams);
+		OutState.CollisionQueryParams = Settings->WorldQueryParams.ToCollisionQuery();
 		OutState.CollisionObjectQueryParams = FCollisionObjectQueryParams(Settings->WorldQueryParams.CollisionChannel);
 
 		return EPCGTimeSliceInitResult::Success;
diff --git a/Engine/Plugins/PCG/Source/PCG/Public/Data/PCGWorldData.h b/Engine/Plugins/PCG/Source/PCG/Public/Data/PCGWorldData.h
--- a/Engine/Plugins/PCG/Source/PCG/Public/Data/PCGWorldData.h
+++ b/Engine/Plugins/PCG/Source/PCG/Public/Data/PCGWorldData.h
@@ -213,18 +213,18 @@ struct FPCGWorldRaycastQueryParams : public FPCGWorldCommonQueryParams
 	int32 UVChannel = 0;
 
 	// Helper function to ensure cohesive conversion of FPCGWorldRaycastQueryParams->FCollisionQueryParams.
-	static FCollisionQueryParams ConvertPCGRaycastQueryToCollisionQuery(const FPCGWorldRaycastQueryParams& InParams)
+	FCollisionQueryParams ToCollisionQuery() const
 	{
 		FCollisionQueryParams Params;
-		Params.bTraceComplex = InParams.bTraceComplex;
-		Params.bReturnPhysicalMaterial = InParams.bGetReferenceToPhysicalMaterial;
+		Params.bTraceComplex = bTraceComplex;
+		Params.bReturnPhysicalMaterial = bGetReferenceToPhysicalMaterial;
 		// We need the face index to calculate UVs, Render Material, and Mesh information.
 		Params.bReturnFaceIndex =
-			InParams.bGetFaceIndex
-			|| InParams.bGetUVCoords
-			|| InParams.bGetSectionIndex
-			|| InParams.bGetRenderMaterialIndex
-			|| InParams.bGetReferenceToRenderMaterial;
+			bGetFaceIndex
+			|| bGetUVCoords
+			|| bGetSectionIndex
+			|| bGetRenderMaterialIndex
+			|| bGetReferenceToRenderMaterial;
 
 		return Params;
 	}
