[PCG] Fixed determinism test triggering ensures - one stemming from the shift that didn't work as intended.

#rnx

diff --git a/Engine/Plugins/PCG/Source/PCG/Private/Tests/Determinism/PCGDeterminismTestsCommon.cpp b/Engine/Plugins/PCG/Source/PCG/Private/Tests/Determinism/PCGDeterminismTestsCommon.cpp
--- a/Engine/Plugins/PCG/Source/PCG/Private/Tests/Determinism/PCGDeterminismTestsCommon.cpp
+++ b/Engine/Plugins/PCG/Source/PCG/Private/Tests/Determinism/PCGDeterminismTestsCommon.cpp
@@ -1607,6 +1607,8 @@ namespace PCGDeterminismTests
 		}
 
 		TestData.OutputData = MoveTemp(Context->OutputData);
+
+		FPCGContext::Release(Context.Release());
 	}
 
 	void ExecuteWithSameTestData(const PCGTestsCommon::FTestData& TestData, const UPCGNode* PCGNode, FPCGDataCollection& OutFirstOutputData, FPCGDataCollection& OutSecondOutputData)
@@ -1629,6 +1631,9 @@ namespace PCGDeterminismTests
 
 		OutFirstOutputData = FirstContext->OutputData;
 		OutSecondOutputData = SecondContext->OutputData;
+
+		FPCGContext::Release(FirstContext.Release());
+		FPCGContext::Release(SecondContext.Release());
 	}
 
 	void ExecuteWithSameTestDataSameElement(const PCGTestsCommon::FTestData& TestData, const UPCGNode* PCGNode, FPCGDataCollection& OutFirstOutputData, FPCGDataCollection& OutSecondOutputData)
@@ -1650,6 +1655,9 @@ namespace PCGDeterminismTests
 
 		OutFirstOutputData = FirstContext->OutputData;
 		OutSecondOutputData = SecondContext->OutputData;
+
+		FPCGContext::Release(FirstContext.Release());
+		FPCGContext::Release(SecondContext.Release());
 	}
 
 	bool ExecutionIsDeterministic(PCGTestsCommon::FTestData& FirstTestData, PCGTestsCommon::FTestData& SecondTestData, const UPCGNode* PCGNode)
diff --git a/Engine/Plugins/PCG/Source/PCG/Public/Helpers/PCGHelpers.h b/Engine/Plugins/PCG/Source/PCG/Public/Helpers/PCGHelpers.h
--- a/Engine/Plugins/PCG/Source/PCG/Public/Helpers/PCGHelpers.h
+++ b/Engine/Plugins/PCG/Source/PCG/Public/Helpers/PCGHelpers.h
@@ -157,7 +157,11 @@ namespace PCGHelpers
 			TempArray[i] = std::move(Array[i - NumShifts]);
 		}
 
-		Array = std::move(TempArray);
+		// @todo: there is a better implementation where we need to allocate at most the shift so we can move things in place.
+		for (int32 i = 0; i < Count; ++i)
+		{
+			Array[i] = std::move(TempArray[i]);
+		}
 	}
 
 	/** Execute given functor on game thread. If called from game thread, executes immediately. */
