[PCG] HLSL Editor UX improvements
* Disable editing in HLSL editor whenever a single, individual HLSL node is not selected
* Disable editing in HLSL editor when overridden by a compute source

#rnx

diff --git a/Engine/Plugins/PCG/Source/PCG/Private/Compute/Elements/PCGCustomHLSL.cpp b/Engine/Plugins/PCG/Source/PCG/Private/Compute/Elements/PCGCustomHLSL.cpp
--- a/Engine/Plugins/PCG/Source/PCG/Private/Compute/Elements/PCGCustomHLSL.cpp
+++ b/Engine/Plugins/PCG/Source/PCG/Private/Compute/Elements/PCGCustomHLSL.cpp
@@ -295,7 +295,7 @@ void UPCGCustomHLSLSettings::SetShaderText(const FString& NewText)
 
 bool UPCGCustomHLSLSettings::IsShaderTextReadOnly() const
 {
-	return false;
+	return KernelSourceOverride != nullptr;
 }
 
 void UPCGCustomHLSLSettings::CreateKernels(FPCGGPUCompilationContext& InOutContext, UObject* InObjectOuter, TArray<UPCGComputeKernel*>& OutKernels, TArray<FPCGKernelEdge>& OutEdges) const
diff --git a/Engine/Plugins/PCG/Source/PCGEditor/Private/Details/PCGComputeSourceDetails.cpp b/Engine/Plugins/PCG/Source/PCGEditor/Private/Details/PCGComputeSourceDetails.cpp
--- a/Engine/Plugins/PCG/Source/PCGEditor/Private/Details/PCGComputeSourceDetails.cpp
+++ b/Engine/Plugins/PCG/Source/PCGEditor/Private/Details/PCGComputeSourceDetails.cpp
@@ -99,12 +99,16 @@ void FPCGComputeSourceDetails::CustomizeDetails(IDetailLayoutBuilder& DetailBuil
 			[
 				SAssignNew(SourceTextBox, SPCGNodeSourceTextBox)
 				.Text(this, &FPCGComputeSourceDetails::GetSourceText)
+				.IsReadOnly(false)
 				.OnTextChanged(this, &FPCGComputeSourceDetails::OnSourceTextChanged)
-				.IsReadOnly(this, &FPCGComputeSourceDetails::IsSourceTextReadOnly)
+				.OnTextCommitted(this, &FPCGComputeSourceDetails::OnSourceTextCommitted)
+				.OnTextChangesApplied(this, &FPCGComputeSourceDetails::OnSourceTextChangesApplied)
 				.Marshaller(SyntaxHighlighter)
 			]
 		]
 	];
+
+	SourceTextBox->SetTextProviderObject(ComputeSourceWeakPtr.Get());
 }
 
 FText FPCGComputeSourceDetails::GetSourceText() const
@@ -115,16 +119,31 @@ FText FPCGComputeSourceDetails::GetSourceText() const
 
 void FPCGComputeSourceDetails::OnSourceTextChanged(const FText& InText)
 {
-	if (UPCGComputeSource* ComputeSource = ComputeSourceWeakPtr.Get())
-	{
-		ComputeSource->SetShaderText(InText.ToString());
-	}
+	SourceText = InText;
 }
 
-bool FPCGComputeSourceDetails::IsSourceTextReadOnly() const
+void FPCGComputeSourceDetails::OnSourceTextCommitted(const FText& InText, ETextCommit::Type InCommitInfo)
 {
-	UPCGComputeSource* ComputeSource = ComputeSourceWeakPtr.Get();
-	return ComputeSource ? ComputeSource->IsShaderTextReadOnly() : false;
+	SourceText = InText;
+	SetComputeSourceText();
+}
+
+void FPCGComputeSourceDetails::OnSourceTextChangesApplied() const
+{
+	SetComputeSourceText();
+}
+
+void FPCGComputeSourceDetails::SetComputeSourceText() const
+{
+	if (UPCGComputeSource* ComputeSource = ComputeSourceWeakPtr.Get())
+	{
+		const FString SourceString = SourceText.ToString();
+
+		if (SourceString != ComputeSource->GetShaderText())
+		{
+			ComputeSource->SetShaderText(SourceString);
+		}
+	}
 }
 
 #undef LOCTEXT_NAMESPACE
diff --git a/Engine/Plugins/PCG/Source/PCGEditor/Private/Details/PCGComputeSourceDetails.h b/Engine/Plugins/PCG/Source/PCGEditor/Private/Details/PCGComputeSourceDetails.h
--- a/Engine/Plugins/PCG/Source/PCGEditor/Private/Details/PCGComputeSourceDetails.h
+++ b/Engine/Plugins/PCG/Source/PCGEditor/Private/Details/PCGComputeSourceDetails.h
@@ -4,6 +4,7 @@
 
 #include "CoreMinimal.h"
 #include "IDetailCustomization.h"
+#include "Types/SlateEnums.h"
 #include "UObject/WeakObjectPtr.h"
 
 class FPCGHLSLSyntaxHighlighter;
@@ -24,9 +25,13 @@ public:
 protected:
 	FText GetSourceText() const;
 	void OnSourceTextChanged(const FText& InText);
-	bool IsSourceTextReadOnly() const;
+	void OnSourceTextCommitted(const FText& InText, ETextCommit::Type InCommitInfo);
+	void OnSourceTextChangesApplied() const;
+	void SetComputeSourceText() const;
 
 	TWeakObjectPtr<UPCGComputeSource> ComputeSourceWeakPtr;
 	TSharedRef<FPCGHLSLSyntaxHighlighter> SyntaxHighlighter;
 	TSharedPtr<SPCGNodeSourceTextBox> SourceTextBox;
+
+	FText SourceText;
 };
diff --git a/Engine/Plugins/PCG/Source/PCGEditor/Private/Widgets/SPCGEditorNodeSource.cpp b/Engine/Plugins/PCG/Source/PCGEditor/Private/Widgets/SPCGEditorNodeSource.cpp
--- a/Engine/Plugins/PCG/Source/PCGEditor/Private/Widgets/SPCGEditorNodeSource.cpp
+++ b/Engine/Plugins/PCG/Source/PCGEditor/Private/Widgets/SPCGEditorNodeSource.cpp
@@ -95,6 +95,9 @@ void SPCGEditorNodeSource::SetTextProviderObject(UObject* InProviderObject)
 {
 	ShaderTextProviderObject = InProviderObject;
 
+	ShaderFunctionsTextBox->SetTextProviderObject(InProviderObject);
+	ShaderTextTextBox->SetTextProviderObject(InProviderObject);
+
 	ShaderFunctionsText = GetShaderFunctionsAsText();
 	ShaderSourceText = GetShaderTextAsText();
 
diff --git a/Engine/Plugins/PCG/Source/PCGEditor/Private/Widgets/SPCGNodeSourceTextBox.cpp b/Engine/Plugins/PCG/Source/PCGEditor/Private/Widgets/SPCGNodeSourceTextBox.cpp
--- a/Engine/Plugins/PCG/Source/PCGEditor/Private/Widgets/SPCGNodeSourceTextBox.cpp
+++ b/Engine/Plugins/PCG/Source/PCGEditor/Private/Widgets/SPCGNodeSourceTextBox.cpp
@@ -2,6 +2,8 @@
 
 #include "SPCGNodeSourceTextBox.h"
 
+#include "Compute/IPCGNodeSourceTextProvider.h"
+
 #include "PCGEditorStyle.h"
 #include "SPCGShaderTextSearchWidget.h"
 
@@ -128,7 +130,7 @@ void SPCGNodeSourceTextBox::Construct(const FArguments& InArgs)
 						}
 						
 					}))
-			.IsReadOnly(bReadOnly)
+			.IsReadOnly(this, &SPCGNodeSourceTextBox::IsReadOnly)
 			.HScrollBar(HScrollBar)
 			.VScrollBar(VScrollBar);
 
@@ -167,6 +169,13 @@ void SPCGNodeSourceTextBox::Construct(const FArguments& InArgs)
 	];
 }
 
+void SPCGNodeSourceTextBox::SetTextProviderObject(UObject* InProviderObject)
+{
+	ShaderTextProviderObject = InProviderObject;
+
+	Refresh();
+}
+
 void SPCGNodeSourceTextBox::RegisterCommands()
 {
 	const FPCGNodeSourceEditorTextBoxCommands& Commands = FPCGNodeSourceEditorTextBoxCommands::Get();
@@ -738,4 +747,14 @@ void SPCGNodeSourceTextBox::HandleAutoIndent() const
 	}
 }
 
+bool SPCGNodeSourceTextBox::IsReadOnly() const
+{
+	if (IPCGNodeSourceTextProvider* Provider = Cast<IPCGNodeSourceTextProvider>(ShaderTextProviderObject))
+	{
+		return Provider->IsShaderTextReadOnly();
+	}
+
+	return true;
+}
+
 #undef LOCTEXT_NAMESPACE
diff --git a/Engine/Plugins/PCG/Source/PCGEditor/Private/Widgets/SPCGNodeSourceTextBox.h b/Engine/Plugins/PCG/Source/PCGEditor/Private/Widgets/SPCGNodeSourceTextBox.h
--- a/Engine/Plugins/PCG/Source/PCGEditor/Private/Widgets/SPCGNodeSourceTextBox.h
+++ b/Engine/Plugins/PCG/Source/PCGEditor/Private/Widgets/SPCGNodeSourceTextBox.h
@@ -67,6 +67,8 @@ public:
 
 	void Refresh() const;
 
+	void SetTextProviderObject(UObject* InProviderObject);
+
 private:
 	//~ Begin SWidget interface
 	// Use Preview Key Down so that the commands related to this tab is given the priority
@@ -99,6 +101,8 @@ private:
 	// Figure out if we need to auto-indent.
 	void HandleAutoIndent() const;
 
+	bool IsReadOnly() const;
+
 private:
 	TSharedPtr<SVerticalBox> TabBody;
 	
@@ -114,4 +118,6 @@ private:
 	TSharedRef<FUICommandList> TextCommandList;
 	
 	FOnTextChangesApplied OnTextChangesAppliedCallback;
+
+	TWeakObjectPtr<UObject> ShaderTextProviderObject;
 };
