[PCG] Fix attribute count in GPU graphs exceeding maximum limit.

We were silently overflowing the limit which led to over-reading from fixed size attribute tables and GPU hangs.

* Actively enforce the limit and emit errors when exceeded.
* Avoid importing attributes from input data which are not supported on GPU.
* Also add missing PostLoad call which caused a crash in a test level.

#rnx

diff --git a/Engine/Plugins/PCG/Source/PCG/Internal/Compute/PCGComputeCommon.h b/Engine/Plugins/PCG/Source/PCG/Internal/Compute/PCGComputeCommon.h
--- a/Engine/Plugins/PCG/Source/PCG/Internal/Compute/PCGComputeCommon.h
+++ b/Engine/Plugins/PCG/Source/PCG/Internal/Compute/PCGComputeCommon.h
@@ -98,6 +98,9 @@ namespace PCGComputeHelpers
 	/** True if 'Type' is valid in a GPU data collection. Some types are only supported as DataInterfaces, and cannot be uploaded in data collections. */
 	bool IsTypeAllowedInDataCollection(EPCGDataType Type);
 
+	/** Whether metadata attributes should be read from the given data and registered for use in GPU graphs. */
+	bool ShouldImportAttributesFromData(const UPCGData* InData);
+
 #if PCG_KERNEL_LOGGING_ENABLED
 	/** Logs a warning on a GPU node in the graph and console. */
 	void LogKernelWarning(const FPCGContext* Context, const UPCGSettings* Settings, const FText& InText);
diff --git a/Engine/Plugins/PCG/Source/PCG/Internal/Compute/PCGDataForGPU.h b/Engine/Plugins/PCG/Source/PCG/Internal/Compute/PCGDataForGPU.h
--- a/Engine/Plugins/PCG/Source/PCG/Internal/Compute/PCGDataForGPU.h
+++ b/Engine/Plugins/PCG/Source/PCG/Internal/Compute/PCGDataForGPU.h
@@ -96,9 +96,12 @@ struct FPCGKernelAttributeTable
 
 	int32 GetAttributeId(FPCGAttributeIdentifier InIdentifier, EPCGKernelAttributeType InType) const;
 
-	/** Adds an attribute of given name and type. Returns added attribute ID. */
+	/** Adds an attribute of given name and type. Returns index or INDEX_NONE if add failed (happens if max table size reached). */
+	int32 AddAttribute(const FPCGKernelAttributeKey& Key);
 	int32 AddAttribute(FPCGAttributeIdentifier InIdentifier, EPCGKernelAttributeType InType);
 
+	int32 Num() const { return AttributeTable.Num(); }
+
 #if PCG_KERNEL_LOGGING_ENABLED
 	void DebugLog() const;
 #endif
@@ -106,8 +109,6 @@ struct FPCGKernelAttributeTable
 protected:
 	UPROPERTY()
 	TArray<FPCGKernelAttributeKey> AttributeTable;
-
-	friend class FPCGGraphCompilerGPU;
 };
 
 /** Data description for a metadata attribute. Stores identifying name and type as well as the unique attribute ID. */
diff --git a/Engine/Plugins/PCG/Source/PCG/Private/Components/PCGProceduralISMComponent.cpp b/Engine/Plugins/PCG/Source/PCG/Private/Components/PCGProceduralISMComponent.cpp
--- a/Engine/Plugins/PCG/Source/PCG/Private/Components/PCGProceduralISMComponent.cpp
+++ b/Engine/Plugins/PCG/Source/PCG/Private/Components/PCGProceduralISMComponent.cpp
@@ -335,6 +335,8 @@ UPCGProceduralISMComponent::UPCGProceduralISMComponent(const FObjectInitializer&
 
 void UPCGProceduralISMComponent::PostLoad()
 {
+	Super::PostLoad();
+
 	PRAGMA_DISABLE_DEPRECATION_WARNINGS
 	if (LocalBounds.IsValid)
 	{
diff --git a/Engine/Plugins/PCG/Source/PCG/Private/Compute/PCGComputeCommon.cpp b/Engine/Plugins/PCG/Source/PCG/Private/Compute/PCGComputeCommon.cpp
--- a/Engine/Plugins/PCG/Source/PCG/Private/Compute/PCGComputeCommon.cpp
+++ b/Engine/Plugins/PCG/Source/PCG/Private/Compute/PCGComputeCommon.cpp
@@ -89,6 +89,13 @@ namespace PCGComputeHelpers
 		return (Type | PCGComputeConstants::AllowedDataCollectionTypes) == PCGComputeConstants::AllowedDataCollectionTypes;
 	}
 
+	bool ShouldImportAttributesFromData(const UPCGData* InData)
+	{
+		// We only read and expose attributes to compute from the following types. Other types are supported but we don't
+		// register/upload their metadata attributes automatically.
+		return InData && (InData->IsA<UPCGParamData>() || InData->IsA<UPCGPointData>());
+	}
+
 #if PCG_KERNEL_LOGGING_ENABLED
 	void LogKernelWarning(const FPCGContext* Context, const UPCGSettings* Settings, const FText& InText)
 	{
diff --git a/Engine/Plugins/PCG/Source/PCG/Private/Compute/PCGDataBinding.cpp b/Engine/Plugins/PCG/Source/PCG/Private/Compute/PCGDataBinding.cpp
--- a/Engine/Plugins/PCG/Source/PCG/Private/Compute/PCGDataBinding.cpp
+++ b/Engine/Plugins/PCG/Source/PCG/Private/Compute/PCGDataBinding.cpp
@@ -332,7 +332,7 @@ void UPCGDataBinding::AddInputDataAttributesToTable()
 				}
 			}
 		}
-		else if (const UPCGMetadata* Metadata = Data.Data ? Data.Data->ConstMetadata() : nullptr)
+		else if (const UPCGMetadata* Metadata = (Data.Data && PCGComputeHelpers::ShouldImportAttributesFromData(Data.Data)) ? Data.Data->ConstMetadata() : nullptr)
 		{
 			const FPCGMetadataDomainID MetadataDefaultDomainID = Metadata->GetConstDefaultMetadataDomain()->GetDomainID();
 			
@@ -360,10 +360,16 @@ void UPCGDataBinding::AddInputDataAttributesToTable()
 					
 				const EPCGKernelAttributeType AttributeType = PCGDataForGPUHelpers::GetAttributeTypeFromMetadataType(AttributeTypes[Index]);
 
-				AttributeTable.AddAttribute(AttributeIdentifier, AttributeType);
+				if (AttributeTable.AddAttribute(AttributeIdentifier, AttributeType) == INDEX_NONE)
+				{
+					UE_LOG(LogPCG, Error, TEXT("Attribute table exceeded maximum size (%d), use the 'Dump Data Descriptions' setting on the GPU node(s) to list attributes that are present."), PCGComputeConstants::MAX_NUM_CUSTOM_ATTRS);
+					break;
+				}
 			}
 		}
 	}
+
+	ensure(AttributeTable.Num() <= PCGComputeConstants::MAX_NUM_CUSTOM_ATTRS);
 }
 
 void UPCGDataBinding::AddInputDataStringsToTable()
diff --git a/Engine/Plugins/PCG/Source/PCG/Private/Compute/PCGDataForGPU.cpp b/Engine/Plugins/PCG/Source/PCG/Private/Compute/PCGDataForGPU.cpp
--- a/Engine/Plugins/PCG/Source/PCG/Private/Compute/PCGDataForGPU.cpp
+++ b/Engine/Plugins/PCG/Source/PCG/Private/Compute/PCGDataForGPU.cpp
@@ -687,9 +687,22 @@ int32 FPCGKernelAttributeTable::GetAttributeId(FPCGAttributeIdentifier InIdentif
 	return GetAttributeId(FPCGKernelAttributeKey(InIdentifier, InType));
 }
 
+int32 FPCGKernelAttributeTable::AddAttribute(const FPCGKernelAttributeKey& Key)
+{
+	int32 Index = AttributeTable.Find(Key);
+
+	if (Index == INDEX_NONE && AttributeTable.Num() < PCGComputeConstants::MAX_NUM_CUSTOM_ATTRS)
+	{
+		Index = AttributeTable.Num();
+		AttributeTable.Add(Key);
+	}
+
+	return Index;
+}
+
 int32 FPCGKernelAttributeTable::AddAttribute(FPCGAttributeIdentifier InIdentifier, EPCGKernelAttributeType InType)
 {
-	return AttributeTable.AddUnique(FPCGKernelAttributeKey(InIdentifier, InType));
+	return AddAttribute(FPCGKernelAttributeKey(InIdentifier, InType));
 }
 
 #if PCG_KERNEL_LOGGING_ENABLED
@@ -847,7 +860,7 @@ void FPCGDataDesc::InitializeAttributeDescs(const UPCGData* InData, const UPCGDa
 	}
 	else { /* TODO: More types! */ }
 
-	const UPCGMetadata* Metadata = InData ? InData->ConstMetadata() : nullptr;
+	const UPCGMetadata* Metadata = (InData && PCGComputeHelpers::ShouldImportAttributesFromData(InData)) ? InData->ConstMetadata() : nullptr;
 
 	if (InBinding && Metadata)
 	{
diff --git a/Engine/Plugins/PCG/Source/PCG/Private/Graph/PCGGraphCompilerGPU.cpp b/Engine/Plugins/PCG/Source/PCG/Private/Graph/PCGGraphCompilerGPU.cpp
--- a/Engine/Plugins/PCG/Source/PCG/Private/Graph/PCGGraphCompilerGPU.cpp
+++ b/Engine/Plugins/PCG/Source/PCG/Private/Graph/PCGGraphCompilerGPU.cpp
@@ -843,13 +843,19 @@ void FPCGGraphCompilerGPU::BuildComputeGraphStaticData(
 
 			for (const FPCGKernelAttributeKey& Key : AttributeKeys)
 			{
-				InOutComputeGraph->StaticAttributeTable.AttributeTable.AddUnique(Key);
+				if (InOutComputeGraph->StaticAttributeTable.AddAttribute(Key) == INDEX_NONE)
+				{
+					UE_LOG(LogPCG, Error, TEXT("FPCGGraphCompilerGPU: Static attribute table exceeded maximum size (%d), use the 'Dump Data Descriptions' setting on the GPU node(s) to list attributes that are present."), PCGComputeConstants::MAX_NUM_CUSTOM_ATTRS);
+					break;
+				}
 			}
 
 			// Record strings.
 			Kernel->AddStaticCreatedStrings(InOutComputeGraph->StringTable);
 		}
 	}
+
+	ensure(InOutComputeGraph->StaticAttributeTable.Num() <= PCGComputeConstants::MAX_NUM_CUSTOM_ATTRS);
 }
 
 void FPCGGraphCompilerGPU::CreateDataInterfaces(
