[PCG] Fix UPCGActorHelpers::DeleteActors regression (42448814) when used from Commandlet

#rnx

diff --git a/Engine/Plugins/PCG/Source/PCG/Private/Helpers/PCGActorHelpers.cpp b/Engine/Plugins/PCG/Source/PCG/Private/Helpers/PCGActorHelpers.cpp
--- a/Engine/Plugins/PCG/Source/PCG/Private/Helpers/PCGActorHelpers.cpp
+++ b/Engine/Plugins/PCG/Source/PCG/Private/Helpers/PCGActorHelpers.cpp
@@ -536,7 +536,8 @@ bool UPCGActorHelpers::DeleteActors(UWorld* World, const TArray<TSoftObjectPtr<A
 		if (AActor* Actor = ManagedActor.Get(); Actor && Actor->GetWorld())
 		{
 #if WITH_EDITOR
-			if (!Actor->GetWorld()->IsGameWorld())
+			// In Editor, prefer UUnrealEdEngine::DeleteActors which handles selection
+			if (GUnrealEd && !Actor->GetWorld()->IsGameWorld())
 			{
 				DeleteActors.Add(Actor);
 			}
@@ -550,7 +551,7 @@ bool UPCGActorHelpers::DeleteActors(UWorld* World, const TArray<TSoftObjectPtr<A
 	}
 
 #if WITH_EDITOR
-	if (!DeleteActors.IsEmpty())
+	if (GUnrealEd && !DeleteActors.IsEmpty())
 	{
 		GUnrealEd->DeleteActors(DeleteActors, World, GUnrealEd->GetSelectedActors()->GetElementSelectionSet(), /*bVerifyDeletionCanHappen=*/false, /*bWarnAboutReferences=*/false, /*bWarnAboutSoftReferences=*/false);
 	}
