[PCG] Fix property bags not showing in PCG details view for PCG nodes.

#rnx

diff --git a/Engine/Plugins/PCG/Source/PCGEditor/Private/Widgets/SPCGEditorGraphDetailsView.cpp b/Engine/Plugins/PCG/Source/PCGEditor/Private/Widgets/SPCGEditorGraphDetailsView.cpp
--- a/Engine/Plugins/PCG/Source/PCGEditor/Private/Widgets/SPCGEditorGraphDetailsView.cpp
+++ b/Engine/Plugins/PCG/Source/PCGEditor/Private/Widgets/SPCGEditorGraphDetailsView.cpp
@@ -152,8 +152,8 @@ void SPCGEditorGraphDetailsView::Construct(const FArguments& InArgs)
 		{
 			if (const FStructProperty* PropertyStruct = CastField<FStructProperty>(Property))
 			{
-				// Explicit because FInstancedPropertyBag holds the Custom DetailsView.
-				return PropertyStruct->Struct == FInstancedPropertyBag::StaticStruct();
+				// Explicit because FInstancedPropertyBag holds the Custom DetailsView for the graphs.
+				return PropertyStruct->Struct == FInstancedPropertyBag::StaticStruct() && PropertyStruct->GetOwnerClass() == UPCGGraph::StaticClass();
 			}
 		}
 
