[PCG] Optimization for trivial cases in the Cull Points Outside Bounds node.

#rnx

diff --git a/Engine/Plugins/PCG/Source/PCG/Private/Elements/PCGCullPointsOutsideActorBounds.cpp b/Engine/Plugins/PCG/Source/PCG/Private/Elements/PCGCullPointsOutsideActorBounds.cpp
--- a/Engine/Plugins/PCG/Source/PCG/Private/Elements/PCGCullPointsOutsideActorBounds.cpp
+++ b/Engine/Plugins/PCG/Source/PCG/Private/Elements/PCGCullPointsOutsideActorBounds.cpp
@@ -58,8 +58,8 @@ bool FPCGCullPointsOutsideActorBoundsElement::ExecuteInternal(FPCGContext* Conte
 	{
 		const UPCGBasePointData* InputPointData = Cast<UPCGBasePointData>(Input.Data);
 
-		// Skip non-points or empty point data
-		if(!InputPointData || InputPointData->GetNumPoints() == 0)
+		// Skip non-points or empty point data, or cases where the point data bounds do not intersect at all with the bounds.
+		if(!InputPointData || InputPointData->GetNumPoints() == 0 || !InputPointData->GetBounds().Intersect(BoundsBox))
 		{
 			continue;
 		}
@@ -68,10 +68,16 @@ bool FPCGCullPointsOutsideActorBoundsElement::ExecuteInternal(FPCGContext* Conte
 
 		TArray<int32, TInlineAllocator<4096>> KeptPointIndices;
 		bool bNeedsSort = false;
+		bool bTrivialOperation = false;
 
+		// First, test if the point bounds are fully inside the bounds box, if so it's a trivial operation
+		if (BoundsBox.IsInsideOrOn(InputPointData->GetBounds()))
+		{
+			bTrivialOperation = true;
+		}
 		// If an octree is available, we'll perform a query on it instead of going through all points, as in most cases it will be significantly less work.
 		// Otherwise, fall-back to testing every point (on multiple threads), which will always yield sorted results.
-		if (InputPointData->IsPointOctreeDirty())
+		else if (InputPointData->IsPointOctreeDirty())
 		{
 			// Implementation note: while this is a slower code path in general, building the octree normally dwarfs filtering (on MT, especially)
 			// @todo_pcg: review this if we improve on the octree creation
@@ -141,7 +147,7 @@ bool FPCGCullPointsOutsideActorBoundsElement::ExecuteInternal(FPCGContext* Conte
 			bNeedsSort = Settings->Mode == EPCGCullPointsMode::Ordered;
 		}
 
-		if (KeptPointIndices.Num() == InputPointData->GetNumPoints())
+		if (bTrivialOperation || KeptPointIndices.Num() == InputPointData->GetNumPoints())
 		{
 			// Trivial operation, we've kept all the points
 			OutputData = InputPointData;
