On behalf of ryan.buehler:
[PCG] Fix phantom edge on drag drop actions in editor

#rnx

diff --git a/Engine/Plugins/PCG/Source/PCGEditor/Private/PCGEditor.cpp b/Engine/Plugins/PCG/Source/PCGEditor/Private/PCGEditor.cpp
--- a/Engine/Plugins/PCG/Source/PCGEditor/Private/PCGEditor.cpp
+++ b/Engine/Plugins/PCG/Source/PCGEditor/Private/PCGEditor.cpp
@@ -1608,18 +1608,19 @@ FReply FPCGEditor::OnSpawnNodeByShortcut(FInputChord InChord, const FVector2f& I
 
 FActionMenuContent FPCGEditor::OnCreateActionMenuContent(UEdGraph* InGraph, const FVector2f& Location, const TArray<UEdGraphPin*>& InDraggedPins, bool bAutoExpand, SGraphEditor::FActionMenuClosed OnMenuClosed)
 {
-	SAssignNew(ActionMenu, SGraphEditorActionMenu)
-	.GraphObj(InGraph)
-	.NewNodePosition(Location)
-	.DraggedFromPins(InDraggedPins)
-	.AutoExpandActionMenu(bAutoExpand)
-	.OnCreateWidgetForAction(SGraphActionMenu::FOnCreateWidgetForAction::CreateLambda(
-		[](const FCreateWidgetForActionData* const Data)
-		{
-			return SNew(SPCGGraphActionWidget, Data);
-		}));
-
-	return FActionMenuContent(ActionMenu.ToSharedRef(), ActionMenu->GetFilterTextBox());
+	TSharedRef<SGraphEditorActionMenu> Menu = SNew(SGraphEditorActionMenu)
+		.GraphObj(InGraph)
+		.NewNodePosition(Location)
+		.DraggedFromPins(InDraggedPins)
+		.AutoExpandActionMenu(bAutoExpand)
+		.OnClosedCallback(OnMenuClosed)
+		.OnCreateWidgetForAction(SGraphActionMenu::FOnCreateWidgetForAction::CreateLambda(
+			[](const FCreateWidgetForActionData* const Data)
+			{
+				return SNew(SPCGGraphActionWidget, Data);
+			}));
+
+	return FActionMenuContent(Menu, Menu->GetFilterTextBox());
 }
 
 void FPCGEditor::OnCollapseNodesInSubgraph()
diff --git a/Engine/Plugins/PCG/Source/PCGEditor/Private/PCGEditor.h b/Engine/Plugins/PCG/Source/PCGEditor/Private/PCGEditor.h
--- a/Engine/Plugins/PCG/Source/PCGEditor/Private/PCGEditor.h
+++ b/Engine/Plugins/PCG/Source/PCGEditor/Private/PCGEditor.h
@@ -465,8 +465,6 @@ private:
 	TObjectPtr<UPCGGraph> PCGGraphBeingEdited = nullptr;
 	UPCGEditorGraph* PCGEditorGraph = nullptr;
 
-	TSharedPtr<SGraphEditorActionMenu> ActionMenu;
-
 	TSoftObjectPtr<UPCGComponent> PCGComponentBeingInspected;
 	// Implementation note: we'll keep the last valid component inspected so we don't un-inspect on spurious selection changes
 	TSoftObjectPtr<UPCGComponent> LastValidPCGComponentBeingInspected;
