[PCG] Fix crash in Generate Grass Maps node when execution source is lost.

#rnx

diff --git a/Engine/Plugins/PCG/Source/PCG/Private/Elements/PCGGenerateGrassMaps.cpp b/Engine/Plugins/PCG/Source/PCG/Private/Elements/PCGGenerateGrassMaps.cpp
--- a/Engine/Plugins/PCG/Source/PCG/Private/Elements/PCGGenerateGrassMaps.cpp
+++ b/Engine/Plugins/PCG/Source/PCG/Private/Elements/PCGGenerateGrassMaps.cpp
@@ -163,9 +163,13 @@ bool FPCGGenerateGrassMapsElement::ExecuteInternal(FPCGContext* InContext) const
 	check(Context);
 
 	const UPCGGenerateGrassMapsSettings* Settings = Context->GetInputSettings<UPCGGenerateGrassMapsSettings>();
-	const UPCGComponent* SourceComponent = Cast<UPCGComponent>(Context->ExecutionSource.Get());
 	check(Settings);
-	check(SourceComponent);
+
+	IPCGGraphExecutionSource* ExecutionSource = Context->ExecutionSource.Get();
+	if (!ExecutionSource)
+	{
+		return true;
+	}
 
 	// 1. Select landscape components that overlap the given bounds.
 	if (!Context->bLandscapeComponentsFiltered)
@@ -192,7 +196,7 @@ bool FPCGGenerateGrassMapsElement::ExecuteInternal(FPCGContext* InContext) const
 			return true;
 		}
 
-		if (!UE::Landscape::IsRuntimeGrassMapGenerationSupported() && (!SourceComponent->GetWorld() || SourceComponent->GetWorld()->IsGameWorld()))
+		if (!UE::Landscape::IsRuntimeGrassMapGenerationSupported() && (!ExecutionSource->GetExecutionState().GetWorld() || ExecutionSource->GetExecutionState().GetWorld()->IsGameWorld()))
 		{
 			UE_LOG(LogPCG, Warning, TEXT("Grass map generation is disabled outside of editor worlds. Try enabling the CVar 'grass.GrassMap.AlwaysBuildRuntimeGenerationResources'."));
 			return true;
@@ -227,7 +231,7 @@ bool FPCGGenerateGrassMapsElement::ExecuteInternal(FPCGContext* InContext) const
 			}
 		}
 
-		const FBox ComponentBounds = SourceComponent->GetGridBounds();
+		const FBox GenerationBounds = ExecutionSource->GetExecutionState().GetBounds();
 
 		for (TSoftObjectPtr<ALandscapeProxy> LandscapeProxyPtr : LandscapeData->Landscapes)
 		{
@@ -258,7 +262,7 @@ bool FPCGGenerateGrassMapsElement::ExecuteInternal(FPCGContext* InContext) const
 
 				// Only generate grass map if there is meaningful overlap with our domain of interest.
 				const FBox LandscapeComponentBounds = LandscapeComponent->Bounds.GetBox();
-				if (LandscapeComponentBounds.Overlap(ComponentBounds).GetVolume() > UE_KINDA_SMALL_NUMBER)
+				if (LandscapeComponentBounds.Overlap(GenerationBounds).GetVolume() > UE_KINDA_SMALL_NUMBER)
 				{
 					if (Context->LandscapeComponents.Num() >= FPCGGrassMapUnpackerCS::MaxNumLandscapeComponents)
 					{
@@ -338,7 +342,7 @@ bool FPCGGenerateGrassMapsElement::ExecuteInternal(FPCGContext* InContext) const
 #if WITH_EDITOR
 					if (PCGGenerateGrassMaps::CVarDebugDrawGeneratedComponents.GetValueOnGameThread())
 					{
-						DrawDebugBox(SourceComponent->GetWorld(),
+						DrawDebugBox(ExecutionSource->GetExecutionState().GetWorld(),
 							LandscapeComponentBounds.GetCenter(),
 							LandscapeComponentBounds.GetExtent(),
 							FColor::MakeRandomSeededColor(Context->TaskId),
@@ -361,7 +365,7 @@ bool FPCGGenerateGrassMapsElement::ExecuteInternal(FPCGContext* InContext) const
 	// 2. Wait for landscape components to be ready for grass map rendering.
 	if (!Context->bTextureStreamingRequested)
 	{
-		if (const UWorld* World = SourceComponent->GetWorld())
+		if (const UWorld* World = ExecutionSource->GetExecutionState().GetWorld())
 		{
 			for (TWeakObjectPtr<ULandscapeComponent> LandscapeComponentWeak : Context->LandscapeComponents)
 			{
