[PCG] Fix loss of transient data across higen grid links.

The previous logic cleared the usage flag of the old data and then set usage on the new data. If old data and new data are the same, then clearing the usage flag freed the data.

#rnx

diff --git a/Engine/Plugins/PCG/Source/PCG/Private/PCGComponent.cpp b/Engine/Plugins/PCG/Source/PCG/Private/PCGComponent.cpp
--- a/Engine/Plugins/PCG/Source/PCG/Private/PCGComponent.cpp
+++ b/Engine/Plugins/PCG/Source/PCG/Private/PCGComponent.cpp
@@ -1271,11 +1271,20 @@ void UPCGComponent::StoreOutputDataForPin(const FString& InResourceKey, const FP
 
 	InData.MarkUsage(EPCGDataUsage::ComponentPerPinOutputData);
 
-	if (FPCGDataCollection* FoundData = PerPinGeneratedOutput.Find(InResourceKey))
+	FPCGDataCollection* FoundExistingData = PerPinGeneratedOutput.Find(InResourceKey);
+
+	// For all existing data items, clear their usage which may release transient resources if the data is not present in the new data collection.
+	if (FoundExistingData)
 	{
-		FoundData->ClearUsage(EPCGDataUsage::ComponentPerPinOutputData);
+		for (const FPCGTaggedData& ExistingData : FoundExistingData->TaggedData)
+		{
+			if (ExistingData.Data && !InData.TaggedData.FindByPredicate([&ExistingData](const FPCGTaggedData& NewData) { return ExistingData.Data == NewData.Data; }))
+			{
+				ExistingData.Data->ClearUsage(EPCGDataUsage::ComponentPerPinOutputData);
+			}
+		}
 
-		*FoundData = InData;
+		*FoundExistingData = InData;
 	}
 	else
 	{
diff --git a/Engine/Plugins/PCG/Source/PCG/Private/PCGData.cpp b/Engine/Plugins/PCG/Source/PCG/Private/PCGData.cpp
--- a/Engine/Plugins/PCG/Source/PCG/Private/PCGData.cpp
+++ b/Engine/Plugins/PCG/Source/PCG/Private/PCGData.cpp
@@ -63,7 +63,7 @@ void UPCGData::ClearUsage(EPCGDataUsage InUsage) const
 {
 #ifdef PCG_DATA_USAGE_LOGGING
 	const UEnum* PCGDataUsageEnum = StaticEnum<EPCGDataUsage>();
-	UE_LOG(LogPCG, Warning, TEXT("%s: Umarked usage flag %s, flags %d -> %d"),
+	UE_LOG(LogPCG, Warning, TEXT("%s: Unmarked usage flag %s, flags %d -> %d"),
 		*GetName(),
 		*PCGDataUsageEnum->GetNameStringByValue(static_cast<__underlying_type(EPCGDataUsage)>(InUsage)),
 		Usage.load(),
