[PCG] Fix crash: UPCGLandscapeData::CreateBasePointData was allocating before setting number of points

#rnx

diff --git a/Engine/Plugins/PCG/Source/PCG/Private/Data/PCGLandscapeData.cpp b/Engine/Plugins/PCG/Source/PCG/Private/Data/PCGLandscapeData.cpp
--- a/Engine/Plugins/PCG/Source/PCG/Private/Data/PCGLandscapeData.cpp
+++ b/Engine/Plugins/PCG/Source/PCG/Private/Data/PCGLandscapeData.cpp
@@ -453,7 +453,7 @@ const UPCGBasePointData* UPCGLandscapeData::CreateBasePointData(FPCGContext* Con
 	Data->InitializeFromDataWithParams(InitializeFromDataParams);
 
 	// @todo_pcg: try to optimize allocation of BoundsMin/BoundsMax as it currently depends on the Landscape components and should always be the same.
-	Data->AllocateProperties(EPCGPointNativeProperties::Transform | EPCGPointNativeProperties::Seed | EPCGPointNativeProperties::BoundsMin | EPCGPointNativeProperties::BoundsMax);
+	EPCGPointNativeProperties PropertiesToAllocate = EPCGPointNativeProperties::Transform | EPCGPointNativeProperties::Seed | EPCGPointNativeProperties::BoundsMin | EPCGPointNativeProperties::BoundsMax;
 	
 	FBox EffectiveBounds = Bounds;
 	if (InBounds.IsValid)
@@ -467,10 +467,11 @@ const UPCGBasePointData* UPCGLandscapeData::CreateBasePointData(FPCGContext* Con
 		return Data;
 	}
 
-	UPCGMetadata* OutMetadata = UseMetadata() ? Data->Metadata : nullptr;
+	// @todo_pcg: for now this method does not write out all the selected attributes (see UseMetaData())
+	UPCGMetadata* OutMetadata = DataProps.bGetLayerWeights ? Data->Metadata : nullptr;
 	if (OutMetadata)
 	{
-		Data->AllocateProperties(EPCGPointNativeProperties::MetadataEntry);
+		PropertiesToAllocate |= EPCGPointNativeProperties::MetadataEntry;
 	}
 		
 	int32 NumPoints = 0;
@@ -509,6 +510,7 @@ const UPCGBasePointData* UPCGLandscapeData::CreateBasePointData(FPCGContext* Con
 		if (PointCountUpperBound > 0)
 		{
 			Data->SetNumPoints(PointsBeforeNum + PointCountUpperBound);
+			Data->AllocateProperties(PropertiesToAllocate);
 		}
 
 		FPCGPointValueRanges OutRanges(Data, /*bAllocate=*/false);
@@ -555,7 +557,11 @@ const UPCGBasePointData* UPCGLandscapeData::CreateBasePointData(FPCGContext* Con
 						else
 						{
 							LandscapeCacheEntry->GetPoint(PointIndex, Point, OutMetadata);
-							OutRanges.MetadataEntryRange[NumPoints] = Point.MetadataEntry;
+
+							if (OutMetadata)
+							{
+								OutRanges.MetadataEntryRange[NumPoints] = Point.MetadataEntry;
+							}
 						}
 
 						OutRanges.TransformRange[NumPoints] = Point.Transform;
