[PCG] Fix HLSL Source Editor context menu Cut/Copy/Paste/Undo/Delete not working
- Missing PopCommandList() call

#rnx

diff --git a/Engine/Plugins/PCG/Source/PCGEditor/Private/Widgets/SPCGNodeSourceTextBox.cpp b/Engine/Plugins/PCG/Source/PCGEditor/Private/Widgets/SPCGNodeSourceTextBox.cpp
--- a/Engine/Plugins/PCG/Source/PCGEditor/Private/Widgets/SPCGNodeSourceTextBox.cpp
+++ b/Engine/Plugins/PCG/Source/PCGEditor/Private/Widgets/SPCGNodeSourceTextBox.cpp
@@ -127,6 +127,7 @@ void SPCGNodeSourceTextBox::Construct(const FArguments& InArgs)
 							InBuilder.PushCommandList(TextCommandList);
 							InBuilder.AddMenuEntry(Commands.ApplyChanges);
 							InBuilder.AddMenuEntry(Commands.ToggleComment);
+							InBuilder.PopCommandList();
 						}
 						
 					}))
