[PCG] PCGGetExecutionContext: Add IsDedicatedServer / HasAuthority

#rnx

diff --git a/Engine/Plugins/PCG/Source/PCG/Private/Elements/PCGGetExecutionContext.cpp b/Engine/Plugins/PCG/Source/PCG/Private/Elements/PCGGetExecutionContext.cpp
--- a/Engine/Plugins/PCG/Source/PCG/Private/Elements/PCGGetExecutionContext.cpp
+++ b/Engine/Plugins/PCG/Source/PCG/Private/Elements/PCGGetExecutionContext.cpp
@@ -73,14 +73,14 @@ bool FPCGGetExecutionContextElement::ExecuteInternal(FPCGContext* Context) const
 	FPCGTaggedData& OutputData = Context->OutputData.TaggedData.Emplace_GetRef();
 	OutputData.Data = ParamData;
 
-	IPCGGraphExecutionSource* ExecutionSource = Context->ExecutionSource.Get();
-	UPCGComponent* SourceComponent = Cast<UPCGComponent>(ExecutionSource);
+	const IPCGGraphExecutionSource* ExecutionSource = Context->ExecutionSource.Get();
+	const UPCGComponent* SourceComponent = Cast<UPCGComponent>(ExecutionSource);
 
 	bool Value = false;
+	const UWorld* SupportingWorld = ExecutionSource ? ExecutionSource->GetExecutionState().GetWorld() : nullptr;
 
 	if (Mode == EPCGGetExecutionContextMode::IsEditor || Mode == EPCGGetExecutionContextMode::IsRuntime)
 	{
-		UWorld* SupportingWorld = ExecutionSource ? ExecutionSource->GetExecutionState().GetWorld() : nullptr;
 		const bool bContextIsRuntimeOrOnRuntimeGenerationMode = (SupportingWorld && SupportingWorld->IsGameWorld()) || PCGHelpers::IsRuntimeGeneration(ExecutionSource);
 		Value = (bContextIsRuntimeOrOnRuntimeGenerationMode == (Mode == EPCGGetExecutionContextMode::IsRuntime));
 	}
@@ -96,6 +96,14 @@ bool FPCGGetExecutionContextElement::ExecuteInternal(FPCGContext* Context) const
 	{
 		Value = PCGHelpers::IsRuntimeGeneration(ExecutionSource);
 	}
+	else if (Mode == EPCGGetExecutionContextMode::IsDedicatedServer)
+	{
+		Value = SupportingWorld && SupportingWorld->IsNetMode(NM_DedicatedServer);
+	}
+	else if (Mode == EPCGGetExecutionContextMode::HasAuthority)
+	{
+		Value = ExecutionSource && ExecutionSource->GetExecutionState().HasAuthority();
+	}
 
 	FPCGMetadataAttribute<bool>* Attribute = ParamData->Metadata->CreateAttribute<bool>(PCGGetExecutionContextConstants::AttributeName, Value, /*bAllowInterpolation=*/false, /*bOverrideParent=*/false);
 	ParamData->Metadata->AddEntry();
diff --git a/Engine/Plugins/PCG/Source/PCG/Private/PCGComponent.cpp b/Engine/Plugins/PCG/Source/PCG/Private/PCGComponent.cpp
--- a/Engine/Plugins/PCG/Source/PCG/Private/PCGComponent.cpp
+++ b/Engine/Plugins/PCG/Source/PCG/Private/PCGComponent.cpp
@@ -4168,6 +4168,12 @@ UWorld* FPCGComponentExecutionState::GetWorld() const
 	return Component->GetWorld();
 }
 
+bool FPCGComponentExecutionState::HasAuthority() const
+{
+	check(Component && Component->GetOwner());
+	return Component->GetOwner()->HasAuthority();
+}
+
 FBox FPCGComponentExecutionState::GetBounds() const
 {
 	check(Component);
diff --git a/Engine/Plugins/PCG/Source/PCG/Public/Elements/PCGGetExecutionContext.h b/Engine/Plugins/PCG/Source/PCG/Public/Elements/PCGGetExecutionContext.h
--- a/Engine/Plugins/PCG/Source/PCG/Public/Elements/PCGGetExecutionContext.h
+++ b/Engine/Plugins/PCG/Source/PCG/Public/Elements/PCGGetExecutionContext.h
@@ -15,7 +15,9 @@ enum class EPCGGetExecutionContextMode : uint8
 	IsOriginal UMETA(Tooltip="Returns whether the executing context is the original context (e.g. not partitioned)."),
 	IsLocal UMETA(Tooltip="Returns whether the executing context is partitioned (in opposition to original)."),
 	IsPartitioned UMETA(Tooltip = "Returns whether the executing context is the original context and that it is partitioned."),
-	IsRuntimeGeneration UMETA(Tooltip="Returns whether the executing context is doing runtime gneration.")
+	IsRuntimeGeneration UMETA(Tooltip="Returns whether the executing context is doing runtime gneration."),
+	IsDedicatedServer UMETA(Tooltip="Returns whether the executing context is running on a dedicated server."),
+	HasAuthority UMETA(Tooltip="Returns whether the execution context is running with network authority.")
 };
 
 /** Returns some context-specific information in the form of an attribute set. */
diff --git a/Engine/Plugins/PCG/Source/PCG/Public/PCGComponent.h b/Engine/Plugins/PCG/Source/PCG/Public/PCGComponent.h
--- a/Engine/Plugins/PCG/Source/PCG/Public/PCGComponent.h
+++ b/Engine/Plugins/PCG/Source/PCG/Public/PCGComponent.h
@@ -101,6 +101,7 @@ public:
 	virtual FString GetDebugName() const override;
 	virtual FTransform GetTransform() const override;
 	virtual UWorld* GetWorld() const override;
+	virtual bool HasAuthority() const override;
 	virtual FBox GetBounds() const override;
 	virtual UPCGGraph* GetGraph() const override;
 	virtual UPCGGraphInstance* GetGraphInstance() const override;
diff --git a/Engine/Plugins/PCG/Source/PCG/Public/PCGGraphExecutionStateInterface.h b/Engine/Plugins/PCG/Source/PCG/Public/PCGGraphExecutionStateInterface.h
--- a/Engine/Plugins/PCG/Source/PCG/Public/PCGGraphExecutionStateInterface.h
+++ b/Engine/Plugins/PCG/Source/PCG/Public/PCGGraphExecutionStateInterface.h
@@ -40,6 +40,9 @@ public:
 	/** Returns a World, can be null */
 	virtual UWorld* GetWorld() const = 0;
 
+	/** Returns true if the ExecutionState has network authority */
+	virtual bool HasAuthority() const = 0;
+
 	/** Returns a Transform if the ExecutionState is a spatial one. */
 	virtual FTransform GetTransform() const = 0;
 
