[PCG] Cleanup of comment nodes in PCG Graphs

* Add a backend representation of the comment node, stored on the PCG Graph
* Add a frontend representation of the comment node, able to initialize itself from the backend, inherited from the base comment node
* Replace the base comment node with our own version
* Add deprecation to convert serialized comment on the PCG Graph
* Also add modify calls and callbacks to track changes to comment nodes so it is correctly serialized on the PCG Graph

#rnx

diff --git a/Engine/Plugins/PCG/Source/PCG/Private/Editor/PCGGraphComment.cpp b/Engine/Plugins/PCG/Source/PCG/Private/Editor/PCGGraphComment.cpp
new file mode 100644
--- /dev/null
+++ b/Engine/Plugins/PCG/Source/PCG/Private/Editor/PCGGraphComment.cpp
@@ -0,0 +1,24 @@
+// Copyright Epic Games, Inc. All Rights Reserved.
+
+#include "Editor/PCGGraphComment.h"
+
+#if WITH_EDITOR
+void FPCGGraphCommentNodeData::InitializeFromCommentNode(const UEdGraphNode_Comment& CommentNode)
+{
+	NodePosX = CommentNode.NodePosX;
+	NodePosY = CommentNode.NodePosY;
+	NodeWidth = CommentNode.NodeWidth;
+	NodeHeight = CommentNode.NodeHeight;
+	NodeComment = CommentNode.NodeComment;
+	CommentColor = CommentNode.CommentColor;
+	FontSize = CommentNode.FontSize;
+	bCommentBubbleVisible_InDetailsPanel = CommentNode.bCommentBubbleVisible_InDetailsPanel;
+	bColorCommentBubble = CommentNode.bColorCommentBubble;
+	MoveMode = static_cast<uint8>(CommentNode.MoveMode);
+	NodeDetails = CommentNode.NodeDetails;
+	CommentDepth = CommentNode.CommentDepth;
+	GUID = CommentNode.NodeGuid;
+	bCommentBubblePinned = CommentNode.bCommentBubblePinned;
+	bCommentBubbleVisible = CommentNode.bCommentBubbleVisible;
+}
+#endif // WITH_EDITOR
\ No newline at end of file
diff --git a/Engine/Plugins/PCG/Source/PCG/Private/Helpers/PCGSubgraphHelpers.cpp b/Engine/Plugins/PCG/Source/PCG/Private/Helpers/PCGSubgraphHelpers.cpp
--- a/Engine/Plugins/PCG/Source/PCG/Private/Helpers/PCGSubgraphHelpers.cpp
+++ b/Engine/Plugins/PCG/Source/PCG/Private/Helpers/PCGSubgraphHelpers.cpp
@@ -230,12 +230,23 @@ namespace PCGSubgraphHelpersExtra
 #if WITH_EDITOR
 		// Also duplicate the extra nodes and assign them to the new graph
 		TArray<TObjectPtr<const UObject>> NewExtraGraphNodes;
+		TArray<FPCGGraphCommentNodeData> NewCommentNodes;
+		
 		for (const UObject* ExtraNode : ExtraNodesToCollapse)
 		{
-			NewExtraGraphNodes.Add(DuplicateObject(ExtraNode, NewGraph));
+			// Special treatment of comment nodes
+			if (const UEdGraphNode_Comment* CommentNode = Cast<const UEdGraphNode_Comment>(ExtraNode))
+			{
+				NewCommentNodes.Emplace_GetRef().InitializeFromCommentNode(*CommentNode);
+			}
+			else
+			{
+				NewExtraGraphNodes.Add(DuplicateObject(ExtraNode, NewGraph));
+			}
 		}
 
 		NewGraph->SetExtraEditorNodes(NewExtraGraphNodes);
+		NewGraph->SetCommentNodes(std::move(NewCommentNodes));
 #endif ///WITH_EDITOR
 	}
 
@@ -249,7 +260,15 @@ namespace PCGSubgraphHelpersExtra
 #if WITH_EDITOR
 		for (const UObject* ExtraNode : ExtraNodesToCollapse)
 		{
-			OriginalGraph->RemoveExtraEditorNode(ExtraNode);
+			// Special treatment of comment nodes
+			if (const UEdGraphNode_Comment* CommentNode = Cast<const UEdGraphNode_Comment>(ExtraNode))
+			{
+				OriginalGraph->RemoveCommentNode(CommentNode->NodeGuid);
+			}
+			else
+			{
+				OriginalGraph->RemoveExtraEditorNode(ExtraNode);
+			}
 		}
 #endif // WITH_EDITOR
 
diff --git a/Engine/Plugins/PCG/Source/PCG/Private/PCGGraph.cpp b/Engine/Plugins/PCG/Source/PCG/Private/PCGGraph.cpp
--- a/Engine/Plugins/PCG/Source/PCG/Private/PCGGraph.cpp
+++ b/Engine/Plugins/PCG/Source/PCG/Private/PCGGraph.cpp
@@ -498,6 +498,13 @@ void UPCGGraph::PostLoad()
 		else
 		{
 			ExtraEditorNodes[i]->ConditionalPostLoad();
+
+			// And convert Comment nodes to comment node data.
+			if (const UEdGraphNode_Comment* CommentNode = Cast<UEdGraphNode_Comment>(ExtraEditorNodes[i]))
+			{
+				CommentNodes.Emplace_GetRef().InitializeFromCommentNode(*CommentNode);
+				ExtraEditorNodes.RemoveAtSwap(i);
+			}
 		}
 	}
 
@@ -1576,6 +1583,11 @@ void UPCGGraph::SetExtraEditorNodes(const TArray<TObjectPtr<const UObject>>& InN
 	}
 }
 
+void UPCGGraph::RemoveCommentNode(const FGuid& InNodeGUID)
+{
+	CommentNodes.RemoveAll([InNodeGUID](const FPCGGraphCommentNodeData& NodeData) { return InNodeGUID == NodeData.GUID; });
+}
+
 void UPCGGraph::RemoveExtraEditorNode(const UObject* InNode)
 {
 	ExtraEditorNodes.Remove(const_cast<UObject*>(InNode));
@@ -1655,6 +1667,8 @@ FPCGSelectionKeyToSettingsMap UPCGGraph::GetTrackedActorKeysToSettings() const
 
 void UPCGGraph::GetTrackedActorKeysToSettings(FPCGSelectionKeyToSettingsMap& OutTagsToSettings, TArray<TObjectPtr<const UPCGGraph>>& OutVisitedGraphs) const
 {
+	TRACE_CPUPROFILER_EVENT_SCOPE(UPCGGraph::GetTrackedActorKeysToSettings);
+	
 	if (OutVisitedGraphs.Contains(this))
 	{
 		return;
diff --git a/Engine/Plugins/PCG/Source/PCG/Public/Editor/PCGGraphComment.h b/Engine/Plugins/PCG/Source/PCG/Public/Editor/PCGGraphComment.h
new file mode 100644
--- /dev/null
+++ b/Engine/Plugins/PCG/Source/PCG/Public/Editor/PCGGraphComment.h
@@ -0,0 +1,84 @@
+// Copyright Epic Games, Inc. All Rights Reserved.
+
+#pragma once
+
+#if WITH_EDITOR
+#include "EdGraphNode_Comment.h"
+#endif // WITH_EDITOR
+
+#include "PCGGraphComment.generated.h"
+
+
+/** Mimic UEdGraphNode_Comment class layout to store the information in the backend (PCG Graph). */
+USTRUCT(meta=(Hidden))
+struct FPCGGraphCommentNodeData
+{
+	GENERATED_BODY();
+
+#if WITH_EDITOR
+	PCG_API void InitializeFromCommentNode(const UEdGraphNode_Comment& CommentNode);
+#endif // WITH_EDITOR
+
+#if WITH_EDITORONLY_DATA
+	/** X position of node in the editor */
+	UPROPERTY()
+	int32 NodePosX;
+
+	/** Y position of node in the editor */
+	UPROPERTY()
+	int32 NodePosY;
+
+	/** Width of node in the editor; only used when the node can be resized */
+	UPROPERTY()
+	int32 NodeWidth;
+
+	/** Height of node in the editor; only used when the node can be resized */
+	UPROPERTY()
+	int32 NodeHeight;
+
+	UPROPERTY()
+	FString NodeComment;
+	
+	/** Color to style comment with */
+	UPROPERTY()
+	FLinearColor CommentColor;
+
+	/** Size of the text in the comment box */
+	UPROPERTY()
+	int32 FontSize;
+
+	/** Whether to show a zoom-invariant comment bubble when zoomed out (making the comment readable at any distance). */
+	UPROPERTY()
+	uint32 bCommentBubbleVisible_InDetailsPanel:1;
+
+	/** Whether to use Comment Color to color the background of the comment bubble shown when zoomed out. */
+	UPROPERTY()
+	uint32 bColorCommentBubble:1;
+
+	/** Whether the comment should move any fully enclosed nodes around when it is moved.
+	 * Underlying enum: ECommentBoxMode::Type
+	 * Implementation node: This needs to be type erased for UHT since ECommentBoxMode::Type is not visible outside of editor. */
+	UPROPERTY()
+	uint8 MoveMode;
+
+	/** Details field if more info is needed to be communicated (will show up in tooltip) */
+	UPROPERTY()
+	FText NodeDetails;
+
+	/** comment Depth */
+	UPROPERTY()
+	int32 CommentDepth;
+
+	/** Extra GUID to be unique. */
+	UPROPERTY()
+	FGuid GUID;
+
+	/** Comment bubble pinned state */
+	UPROPERTY()
+	uint8 bCommentBubblePinned : 1;
+
+	/** Comment bubble visibility */
+	UPROPERTY()
+	uint8 bCommentBubbleVisible : 1;
+#endif // WITH_EDITORONLY_DATA
+};
\ No newline at end of file
diff --git a/Engine/Plugins/PCG/Source/PCG/Public/PCGGraph.h b/Engine/Plugins/PCG/Source/PCG/Public/PCGGraph.h
--- a/Engine/Plugins/PCG/Source/PCG/Public/PCGGraph.h
+++ b/Engine/Plugins/PCG/Source/PCG/Public/PCGGraph.h
@@ -9,6 +9,10 @@
 #include "Graph/PCGStackContext.h"
 #include "Helpers/PCGGraphParameterExtension.h"
 
+#if WITH_EDITOR
+#include "Editor/PCGGraphComment.h"
+#endif // WITH_EDITOR
+
 #include "ComputeFramework/ComputeGraphInstance.h"
 #include "StructUtils/PropertyBag.h"
 #include "UObject/ObjectPtr.h"
@@ -419,6 +423,10 @@ public:
 	const TArray<TObjectPtr<UObject>>& GetExtraEditorNodes() const { return ExtraEditorNodes; }
 	UE_API void SetExtraEditorNodes(const TArray<TObjectPtr<const UObject>>& InNodes);
 
+	const TArray<FPCGGraphCommentNodeData>& GetCommentNodes() const { return CommentNodes; }
+	UE_API void SetCommentNodes(TArray<FPCGGraphCommentNodeData> InNodes) { CommentNodes = std::move(InNodes); }
+	UE_API void RemoveCommentNode(const FGuid& InNodeGUID);
+
 	bool IsInspecting() const { return bIsInspecting; }
 	void EnableInspection(const FPCGStack& InInspectedStack) { bIsInspecting = true; InspectedStack = InInspectedStack; }
 	void DisableInspection() { bIsInspecting = false; InspectedStack = FPCGStack(); }
@@ -482,10 +490,14 @@ protected:
 	TObjectPtr<UPCGNode> OutputNode;
 
 #if WITH_EDITORONLY_DATA
-	// Extra data to hold information that is useful only in editor (like comments)
+	// Extra data to hold information that is useful only in editor
 	UPROPERTY()
 	TArray<TObjectPtr<UObject>> ExtraEditorNodes;
 
+	// Extra data to hold information for comments
+	UPROPERTY()
+	TArray<FPCGGraphCommentNodeData> CommentNodes;
+
 	// Editor graph created from PCG Editor but owned by this, reference is collected using AddReferencedObjects
 	TObjectPtr<UPCGEditorGraph> PCGEditorGraph = nullptr;
 #endif // WITH_EDITORONLY_DATA
diff --git a/Engine/Plugins/PCG/Source/PCGEditor/Private/Nodes/PCGEditorGraphNodeBase.cpp b/Engine/Plugins/PCG/Source/PCGEditor/Private/Nodes/PCGEditorGraphNodeBase.cpp
--- a/Engine/Plugins/PCG/Source/PCGEditor/Private/Nodes/PCGEditorGraphNodeBase.cpp
+++ b/Engine/Plugins/PCG/Source/PCGEditor/Private/Nodes/PCGEditorGraphNodeBase.cpp
@@ -1260,7 +1260,7 @@ FText UPCGEditorGraphNodeBase::GetTooltipText() const
 	return FText::Format(LOCTEXT("NodeTooltip", "{0}\n\n{1} - Node index {2}"),
 		Description,
 		PCGNode ? FText::FromName(PCGNode->GetFName()) : LOCTEXT("InvalidNodeName", "Unbound node"),
-		PCGNode ? FText::AsNumber(PCGNode->GetGraph()->GetNodes().IndexOfByKey(PCGNode)) : LOCTEXT("InvalidNodeIndex", "Invalid index"));
+		PCGNode && PCGNode->GetGraph() ? FText::AsNumber(PCGNode->GetGraph()->GetNodes().IndexOfByKey(PCGNode)) : LOCTEXT("InvalidNodeIndex", "Invalid index"));
 }
 
 void UPCGEditorGraphNodeBase::GetPinHoverText(const UEdGraphPin& Pin, FString& HoverTextOut) const
diff --git a/Engine/Plugins/PCG/Source/PCGEditor/Private/Nodes/PCGEditorGraphNodeComment.cpp b/Engine/Plugins/PCG/Source/PCGEditor/Private/Nodes/PCGEditorGraphNodeComment.cpp
new file mode 100644
--- /dev/null
+++ b/Engine/Plugins/PCG/Source/PCGEditor/Private/Nodes/PCGEditorGraphNodeComment.cpp
@@ -0,0 +1,40 @@
+// Copyright Epic Games, Inc. All Rights Reserved.
+
+#include "Nodes/PCGEditorGraphNodeComment.h"
+
+#include "PCGEditorGraph.h"
+#include "Misc/TransactionObjectEvent.h"
+
+void UPCGEditorGraphNodeComment::PostTransacted(const FTransactionObjectEvent& TransactionEvent)
+{
+	Super::PostTransacted(TransactionEvent);
+
+	// If anything happens, we need to dirty the pcg graph, as it will need to serialize the changes
+	if (TransactionEvent.HasPropertyChanges())
+	{
+		UPCGEditorGraph* PCGEditorGraph = CastChecked<UPCGEditorGraph>(GetGraph());
+		if (UPCGGraph* PCGGraph = PCGEditorGraph->GetPCGGraph(); ensure(PCGGraph))
+		{
+			PCGGraph->Modify();	
+		}
+	}
+}
+
+void UPCGEditorGraphNodeComment::InitializeFromNodeData(const FPCGGraphCommentNodeData& NodeData)
+{
+	NodePosX = NodeData.NodePosX;
+	NodePosY = NodeData.NodePosY;
+	NodeWidth = NodeData.NodeWidth;
+	NodeHeight = NodeData.NodeHeight;
+	NodeComment = NodeData.NodeComment;
+	CommentColor = NodeData.CommentColor;
+	FontSize = NodeData.FontSize;
+	bCommentBubbleVisible_InDetailsPanel = NodeData.bCommentBubbleVisible_InDetailsPanel;
+	bColorCommentBubble = NodeData.bColorCommentBubble;
+	MoveMode = static_cast<ECommentBoxMode::Type>(NodeData.MoveMode);
+	NodeDetails = NodeData.NodeDetails;
+	CommentDepth = NodeData.CommentDepth;
+	NodeGuid = NodeData.GUID;
+	bCommentBubblePinned = NodeData.bCommentBubblePinned;
+	bCommentBubbleVisible = NodeData.bCommentBubbleVisible;
+}
diff --git a/Engine/Plugins/PCG/Source/PCGEditor/Private/Nodes/PCGEditorGraphNodeComment.h b/Engine/Plugins/PCG/Source/PCGEditor/Private/Nodes/PCGEditorGraphNodeComment.h
new file mode 100644
--- /dev/null
+++ b/Engine/Plugins/PCG/Source/PCGEditor/Private/Nodes/PCGEditorGraphNodeComment.h
@@ -0,0 +1,24 @@
+// Copyright Epic Games, Inc. All Rights Reserved.
+
+#pragma once
+
+#include "EdGraphNode_Comment.h"
+
+#include "Editor/PCGGraphComment.h"
+
+#include "PCGEditorGraphNodeComment.generated.h"
+
+
+UCLASS()
+class UPCGEditorGraphNodeComment : public UEdGraphNode_Comment
+{
+	GENERATED_BODY()
+
+public:
+	// ~Begin UObject interface
+	virtual void PostTransacted(const FTransactionObjectEvent& TransactionEvent) override;
+	// ~End UObject interface
+	
+	void InitializeFromNodeData(const FPCGGraphCommentNodeData& NodeData);
+	// Export version is the import version in FPCGGraphCommentNodeData
+};
\ No newline at end of file
diff --git a/Engine/Plugins/PCG/Source/PCGEditor/Private/PCGEditor.cpp b/Engine/Plugins/PCG/Source/PCGEditor/Private/PCGEditor.cpp
--- a/Engine/Plugins/PCG/Source/PCGEditor/Private/PCGEditor.cpp
+++ b/Engine/Plugins/PCG/Source/PCGEditor/Private/PCGEditor.cpp
@@ -2409,10 +2409,10 @@ void FPCGEditor::DeleteSelectedNodes()
 				}
 			}
 
-			PCGEditorGraph->Modify();
-
 			if (bChanged)
 			{
+				// Need to modify the pcg graph so comments are also caught.
+				PCGGraph->Modify();
 				PCGGraph->RemoveNodes(NodesToRemove);
 			}
 		}
@@ -2572,6 +2572,8 @@ void FPCGEditor::PasteNodesHere(const FVector2D& Location)
 		}
 	}
 
+	// Need to modify the pcg graph so comments are also caught.
+	PCGGraphBeingEdited->Modify();
 	PCGGraphBeingEdited->AddNodes(NodesToPaste);
 
 	for (UEdGraphNode* PastedNode : PastedNodes)
@@ -2806,10 +2808,15 @@ void FPCGEditor::ReplicateExtraNodes() const
 		if (UPCGGraph* PCGGraph = PCGEditorGraph->GetPCGGraph())
 		{
 			TArray<TObjectPtr<const UObject>> ExtraNodes;
+			TArray<FPCGGraphCommentNodeData> CommentData;
 			for (const UEdGraphNode* GraphNode : PCGEditorGraph->Nodes)
 			{
 				check(GraphNode);
-				if (!GraphNode->IsA<UPCGEditorGraphNodeBase>())
+				if (const UEdGraphNode_Comment* CommentNode = Cast<const UEdGraphNode_Comment>(GraphNode))
+				{
+					CommentData.Emplace_GetRef().InitializeFromCommentNode(*CommentNode);
+				}
+				else if (!GraphNode->IsA<UPCGEditorGraphNodeBase>())
 				{
 					ExtraNodes.Add(GraphNode);
 				}
@@ -2817,6 +2824,7 @@ void FPCGEditor::ReplicateExtraNodes() const
 			}
 
 			PCGGraph->SetExtraEditorNodes(ExtraNodes);
+			PCGGraph->SetCommentNodes(std::move(CommentData));
 		}
 	}
 }
diff --git a/Engine/Plugins/PCG/Source/PCGEditor/Private/PCGEditorGraph.cpp b/Engine/Plugins/PCG/Source/PCGEditor/Private/PCGEditorGraph.cpp
--- a/Engine/Plugins/PCG/Source/PCGEditor/Private/PCGEditorGraph.cpp
+++ b/Engine/Plugins/PCG/Source/PCGEditor/Private/PCGEditorGraph.cpp
@@ -10,6 +10,7 @@
 #include "Elements/PCGReroute.h"
 #include "Elements/PCGUserParameterGet.h"
 #include "Nodes/PCGEditorGraphNode.h"
+#include "Nodes/PCGEditorGraphNodeComment.h"
 #include "Nodes/PCGEditorGraphNodeGetUserParameter.h"
 #include "Nodes/PCGEditorGraphNodeInput.h"
 #include "Nodes/PCGEditorGraphNodeOutput.h"
@@ -140,6 +141,13 @@ void UPCGEditorGraph::ReconstructGraph()
 			AddNode(NewNode, /*bIsUserAction=*/false, bSelectNewNode);
 		}
 	}
+	
+	for (const FPCGGraphCommentNodeData& CommentData : PCGGraph->GetCommentNodes())
+	{
+		UPCGEditorGraphNodeComment* NewNode = NewObject<UPCGEditorGraphNodeComment>(this, NAME_None, RF_Transactional);
+		NewNode->InitializeFromNodeData(CommentData);
+		AddNode(NewNode, /*bIsUserAction=*/false, bSelectNewNode);
+	}
 
 	// Ensure graph structure visualization is nice and fresh upon opening.
 	UpdateVisualizations(nullptr, nullptr);
diff --git a/Engine/Plugins/PCG/Source/PCGEditor/Private/PCGEditorGraphSchemaActions.cpp b/Engine/Plugins/PCG/Source/PCGEditor/Private/PCGEditorGraphSchemaActions.cpp
--- a/Engine/Plugins/PCG/Source/PCGEditor/Private/PCGEditorGraphSchemaActions.cpp
+++ b/Engine/Plugins/PCG/Source/PCGEditor/Private/PCGEditorGraphSchemaActions.cpp
@@ -7,6 +7,7 @@
 #include "PCGEditorGraphSchema.h"
 #include "PCGEditorModule.h"
 #include "Nodes/PCGEditorGraphNode.h"
+#include "Nodes/PCGEditorGraphNodeComment.h"
 #include "Nodes/PCGEditorGraphNodeReroute.h"
 
 #include "PCGGraph.h"
@@ -20,7 +21,6 @@
 
 #include "UObject/UObjectGlobals.h"
 
-#include "EdGraphNode_Comment.h"
 #include "GraphEditor.h"
 #include "ScopedTransaction.h"
 #include "Framework/Application/SlateApplication.h"
@@ -639,7 +639,14 @@ UEdGraphNode* FPCGEditorGraphSchemaAction_NewComment::PerformAction(class UEdGra
 		return nullptr;
 	}
 
-	UEdGraphNode_Comment* CommentTemplate = NewObject<UEdGraphNode_Comment>();
+	UPCGGraph* PCGGraph = EditorGraph->GetPCGGraph();
+	if (!PCGGraph)
+	{
+		UE_LOG(LogPCGEditor, Error, TEXT("Invalid PCGGraph"));
+		return nullptr;
+	}
+
+	UPCGEditorGraphNodeComment* CommentTemplate = NewObject<UPCGEditorGraphNodeComment>();
 
 	TSharedPtr<SGraphEditor> GraphEditorPtr = SGraphEditor::FindGraphEditorForGraph(ParentGraph);
 
@@ -654,9 +661,10 @@ UEdGraphNode* FPCGEditorGraphSchemaAction_NewComment::PerformAction(class UEdGra
 	
 	const FScopedTransaction Transaction(*FPCGEditorCommon::ContextIdentifier, LOCTEXT("PCGEditorGraphSchemaAction_NewComment", "PCG Editor: New Comment"), nullptr);
 	EditorGraph->Modify();
+	PCGGraph->Modify();
+	
+	UPCGEditorGraphNodeComment* NewNode = FEdGraphSchemaAction_NewNode::SpawnNodeFromTemplate<UPCGEditorGraphNodeComment>(ParentGraph, CommentTemplate, SpawnLocation, bSelectNewNode);
 	
-	UEdGraphNode_Comment* NewNode = FEdGraphSchemaAction_NewNode::SpawnNodeFromTemplate<UEdGraphNode_Comment>(ParentGraph, CommentTemplate, SpawnLocation, bSelectNewNode);
-
 	return NewNode;
 }
 
