[PCG] Fix Attribute Property Selector details view that was incorrectly reporting values to be invalid

#rnx

diff --git a/Engine/Plugins/PCG/Source/PCG/Private/Metadata/PCGAttributePropertySelector.cpp b/Engine/Plugins/PCG/Source/PCG/Private/Metadata/PCGAttributePropertySelector.cpp
--- a/Engine/Plugins/PCG/Source/PCG/Private/Metadata/PCGAttributePropertySelector.cpp
+++ b/Engine/Plugins/PCG/Source/PCG/Private/Metadata/PCGAttributePropertySelector.cpp
@@ -329,7 +329,7 @@ void FPCGAttributePropertySelector::ImportFromOtherSelector(const FPCGAttributeP
 
 bool FPCGAttributePropertySelector::IsValid() const
 {
-	const FName ThisAttributeName = GetAttributeName();
+	const FName ThisAttributeName = GetName();
 	static const FName EmptyName = TEXT("");
 
 	if (!ExtraNames.IsEmpty() && ThisAttributeName == EmptyName)
