[PCG] Register static tracked keys when the graph execution is dynamic

#rnx

diff --git a/Engine/Plugins/PCG/Source/PCG/Private/Elements/PCGLoopElement.cpp b/Engine/Plugins/PCG/Source/PCG/Private/Elements/PCGLoopElement.cpp
--- a/Engine/Plugins/PCG/Source/PCG/Private/Elements/PCGLoopElement.cpp
+++ b/Engine/Plugins/PCG/Source/PCG/Private/Elements/PCGLoopElement.cpp
@@ -192,8 +192,18 @@ bool FPCGLoopElement::ExecuteInternal(FPCGContext* InContext) const
 	{
 		if (Settings->SubgraphOverride)
 		{
+			TRACE_CPUPROFILER_EVENT_SCOPE(FPCGLoopElement::ExecuteInternal::DynamicTrackingAndOverrides);
+
 #if WITH_EDITOR
 			FPCGDynamicTrackingHelper::AddSingleDynamicTrackingKey(Context, FPCGSelectionKey::CreateFromPath(Settings->SubgraphOverride), /*bIsCulled=*/false);
+			// Also need to add statically tracked keys
+			if (InContext && InContext->ExecutionSource.IsValid())
+			{
+				if (const UPCGGraph* Graph = Settings->SubgraphOverride->GetGraph())
+				{
+					InContext->ExecutionSource->GetExecutionState().RegisterDynamicTracking(Graph->GetTrackedActorKeysToSettings());
+				}
+			}
 #endif // WITH_EDITOR
 
 			Context->UpdateOverridesWithOverriddenGraph();
diff --git a/Engine/Plugins/PCG/Source/PCG/Private/PCGComponent.cpp b/Engine/Plugins/PCG/Source/PCG/Private/PCGComponent.cpp
--- a/Engine/Plugins/PCG/Source/PCG/Private/PCGComponent.cpp
+++ b/Engine/Plugins/PCG/Source/PCG/Private/PCGComponent.cpp
@@ -3535,6 +3535,23 @@ void UPCGComponent::RegisterDynamicTracking(const UPCGSettings* InSettings, cons
 	}
 }
 
+void UPCGComponent::RegisterDynamicTracking(const FPCGSelectionKeyToSettingsMap& InKeysToSettings)
+{
+	if (InKeysToSettings.IsEmpty())
+	{
+		return;
+	}
+
+	UE::TScopeLock Lock(CurrentExecutionDynamicTrackingLock);
+
+	CurrentExecutionDynamicTracking.Append(InKeysToSettings);
+
+	for (const auto& It : InKeysToSettings)
+	{
+		Algo::Transform(It.Value, CurrentExecutionDynamicTrackingSettings, [](const FPCGSettingsAndCulling& SettingsAndCulling) { return SettingsAndCulling.Key.Get();});
+	}
+}
+
 void UPCGComponent::UpdateDynamicTracking()
 {
 	TRACE_CPUPROFILER_EVENT_SCOPE(UPCGComponent::UpdateDynamicTracking);
@@ -4235,6 +4252,11 @@ void FPCGComponentExecutionState::RegisterDynamicTracking(const UPCGSettings* In
 	Component->RegisterDynamicTracking(InSettings, InDynamicKeysAndCulling);
 }
 
+void FPCGComponentExecutionState::RegisterDynamicTracking(const FPCGSelectionKeyToSettingsMap& InKeysToSettings)
+{
+	check(Component);
+	Component->RegisterDynamicTracking(InKeysToSettings);
+}
 #endif // WITH_EDITOR
 
 #undef LOCTEXT_NAMESPACE
diff --git a/Engine/Plugins/PCG/Source/PCG/Private/PCGSubgraph.cpp b/Engine/Plugins/PCG/Source/PCG/Private/PCGSubgraph.cpp
--- a/Engine/Plugins/PCG/Source/PCG/Private/PCGSubgraph.cpp
+++ b/Engine/Plugins/PCG/Source/PCG/Private/PCGSubgraph.cpp
@@ -718,8 +718,18 @@ bool FPCGSubgraphElement::ExecuteInternal(FPCGContext* InContext) const
 
 	if (bIsDynamic && !Context->bScheduledSubgraph && Settings->SubgraphOverride)
 	{
+		TRACE_CPUPROFILER_EVENT_SCOPE(FPCGSubgraphElement::ExecuteInternal::DynamicTrackingAndOverrides);
+
 #if WITH_EDITOR
 		FPCGDynamicTrackingHelper::AddSingleDynamicTrackingKey(Context, FPCGSelectionKey::CreateFromPath(Settings->SubgraphOverride), /*bIsCulled=*/false);
+		// Also need to add statically tracked keys
+		if (InContext && InContext->ExecutionSource.IsValid())
+		{
+			if (const UPCGGraph* Graph = Settings->SubgraphOverride->GetGraph())
+			{
+				InContext->ExecutionSource->GetExecutionState().RegisterDynamicTracking(Graph->GetTrackedActorKeysToSettings());
+			}
+		}
 #endif // WITH_EDITOR
 
 		Context->UpdateOverridesWithOverriddenGraph();
diff --git a/Engine/Plugins/PCG/Source/PCG/Public/PCGComponent.h b/Engine/Plugins/PCG/Source/PCG/Public/PCGComponent.h
--- a/Engine/Plugins/PCG/Source/PCG/Public/PCGComponent.h
+++ b/Engine/Plugins/PCG/Source/PCG/Public/PCGComponent.h
@@ -115,7 +115,8 @@ public:
 	virtual const FPCGGraphExecutionInspection& GetInspection() const override;
 	virtual FPCGGraphExecutionInspection& GetInspection() override;
 
-	void RegisterDynamicTracking(const UPCGSettings* InSettings, const TArrayView<TPair<FPCGSelectionKey, bool>>& InDynamicKeysAndCulling);
+	virtual void RegisterDynamicTracking(const UPCGSettings* InSettings, const TArrayView<TPair<FPCGSelectionKey, bool>>& InDynamicKeysAndCulling) override;
+	virtual void RegisterDynamicTracking(const FPCGSelectionKeyToSettingsMap& InKeysToSettings) override;
 #endif // WITH_EDITOR
 
 private:
@@ -448,6 +449,9 @@ public:
 	/** To be called by an element to notify the component that this settings have a dynamic dependency. */
 	UE_API void RegisterDynamicTracking(const UPCGSettings* InSettings, const TArrayView<TPair<FPCGSelectionKey, bool>>& InDynamicKeysAndCulling);
 
+	/** To be called to notify the component that this list of keys have a dynamic dependency. */
+	UE_API void RegisterDynamicTracking(const FPCGSelectionKeyToSettingsMap& InKeysToSettings);
+
 	/** For duration of the current/next generation, any change triggers from this change origin will be discarded. */
 	UE_API void StartIgnoringChangeOriginDuringGeneration(const UObject* InChangeOriginToIgnore);
 	UE_API void StartIgnoringChangeOriginsDuringGeneration(const TArrayView<const UObject*> InChangeOriginsToIgnore);
diff --git a/Engine/Plugins/PCG/Source/PCG/Public/PCGGraphExecutionStateInterface.h b/Engine/Plugins/PCG/Source/PCG/Public/PCGGraphExecutionStateInterface.h
--- a/Engine/Plugins/PCG/Source/PCG/Public/PCGGraphExecutionStateInterface.h
+++ b/Engine/Plugins/PCG/Source/PCG/Public/PCGGraphExecutionStateInterface.h
@@ -76,6 +76,9 @@ public:
 	/** Register tracking dependencies, so ExecutionState can be updated when they change */
 	virtual void RegisterDynamicTracking(const UPCGSettings* InSettings, const TArrayView<TPair<FPCGSelectionKey, bool>>& InDynamicKeysAndCulling) = 0;
 
+	/** Register multiple tracking dependencies, so ExecutionState can be updated when they change */
+	virtual void RegisterDynamicTracking(const FPCGSelectionKeyToSettingsMap& InKeysToSettings) = 0;
+
 #endif
 };
 
