[PCG] Fix crash with Intersection/Difference but sampling of points generates entries (always allocate metadata entries)

diff --git a/Engine/Plugins/PCG/Source/PCG/Private/Data/PCGDifferenceData.cpp b/Engine/Plugins/PCG/Source/PCG/Private/Data/PCGDifferenceData.cpp
--- a/Engine/Plugins/PCG/Source/PCG/Private/Data/PCGDifferenceData.cpp
+++ b/Engine/Plugins/PCG/Source/PCG/Private/Data/PCGDifferenceData.cpp
@@ -361,7 +361,7 @@ const UPCGBasePointData* UPCGDifferenceData::CreateBasePointData(FPCGContext* Co
 		return NumWritten;
 	};
 
-	const EPCGPointNativeProperties PropertiesToAllocate = SourcePointData->GetAllocatedProperties() | EPCGPointNativeProperties::Density;
+	const EPCGPointNativeProperties PropertiesToAllocate = SourcePointData->GetAllocatedProperties() | EPCGPointNativeProperties::Density | EPCGPointNativeProperties::MetadataEntry;
 
 	FPCGSpatialDataProcessing::SampleBasedRangeProcessing<ChunkSize>(Context ? &Context->AsyncState : nullptr, ChunkSamplePoints, SourcePointData, TargetPointData, PropertiesToAllocate);
 
diff --git a/Engine/Plugins/PCG/Source/PCG/Private/Data/PCGIntersectionData.cpp b/Engine/Plugins/PCG/Source/PCG/Private/Data/PCGIntersectionData.cpp
--- a/Engine/Plugins/PCG/Source/PCG/Private/Data/PCGIntersectionData.cpp
+++ b/Engine/Plugins/PCG/Source/PCG/Private/Data/PCGIntersectionData.cpp
@@ -316,7 +316,7 @@ UPCGBasePointData* UPCGIntersectionData::CreateAndFilterPointData(FPCGContext* C
 		return NumWritten;
 	};
 
-	const EPCGPointNativeProperties PropertiesToAllocate = SourcePointData->GetAllocatedProperties() | EPCGPointNativeProperties::Density | EPCGPointNativeProperties::Color;
+	const EPCGPointNativeProperties PropertiesToAllocate = SourcePointData->GetAllocatedProperties() | EPCGPointNativeProperties::Density | EPCGPointNativeProperties::Color | EPCGPointNativeProperties::MetadataEntry;
 
 	FPCGSpatialDataProcessing::SampleBasedRangeProcessing<ChunkSize>(Context ? &Context->AsyncState : nullptr, ChunkSamplePoints, SourcePointData, OutputData, PropertiesToAllocate);
 
