[PCG] Use AccessorKeys to pull data from splines in Subdivide and Duplicate cross section

* Using the FirstKey doesn't make sense if the attribute comes from control points

#rnx

diff --git a/Engine/Plugins/PCG/Source/PCG/Private/Elements/Grammar/PCGDuplicateCrossSections.cpp b/Engine/Plugins/PCG/Source/PCG/Private/Elements/Grammar/PCGDuplicateCrossSections.cpp
--- a/Engine/Plugins/PCG/Source/PCG/Private/Elements/Grammar/PCGDuplicateCrossSections.cpp
+++ b/Engine/Plugins/PCG/Source/PCG/Private/Elements/Grammar/PCGDuplicateCrossSections.cpp
@@ -90,13 +90,14 @@ bool FPCGDuplicateCrossSectionsElement::ExecuteInternal(FPCGContext* InContext)
 		{
 			const FPCGAttributePropertyInputSelector Selector = InSelector.CopyAndFixLast(InputSplineData);
 			const TUniquePtr<const IPCGAttributeAccessor> Accessor = PCGAttributeAccessorHelpers::CreateConstAccessor(InputSplineData, Selector);
-			if (!Accessor)
+			const TUniquePtr<const IPCGAttributeAccessorKeys> Keys = PCGAttributeAccessorHelpers::CreateConstKeys(InputSplineData, Selector);
+			if (!Accessor || !Keys)
 			{
 				PCGLog::Metadata::LogFailToCreateAccessorError(Selector, InContext);
 				return false;
 			}
 
-			if (!Accessor->Get(OutValue, FPCGAttributeAccessorKeysEntries(PCGFirstEntryKey), EPCGAttributeAccessorFlags::AllowBroadcastAndConstructible))
+			if (!Accessor->Get(OutValue, *Keys, EPCGAttributeAccessorFlags::AllowBroadcastAndConstructible))
 			{
 				PCGLog::Metadata::LogFailToGetAttributeError<T>(Selector, Accessor.Get(), InContext);
 				return false;
@@ -122,13 +123,14 @@ bool FPCGDuplicateCrossSectionsElement::ExecuteInternal(FPCGContext* InContext)
 		{
 			const FPCGAttributePropertyInputSelector Selector = Settings->SeedAttribute.CopyAndFixLast(InputSplineData);
 			TUniquePtr<const IPCGAttributeAccessor> SeedAccessor = PCGAttributeAccessorHelpers::CreateConstAccessor(InputSplineData, Selector);
+			TUniquePtr<const IPCGAttributeAccessorKeys> SeedAccessorKeys = PCGAttributeAccessorHelpers::CreateConstKeys(InputSplineData, Selector);
 
-			if (!SeedAccessor)
+			if (!SeedAccessor || !SeedAccessorKeys)
 			{
 				PCGLog::Metadata::LogFailToCreateAccessorError(Selector, InContext);
 			}
 			// Otherwise, get the value, if it fails, the attribute wasn't compatible
-			else if (!SeedAccessor->Get(AdditionalSeed, FPCGAttributeAccessorKeysEntries(PCGFirstEntryKey), EPCGAttributeAccessorFlags::AllowBroadcastAndConstructible))
+			else if (!SeedAccessor->Get(AdditionalSeed, *SeedAccessorKeys, EPCGAttributeAccessorFlags::AllowBroadcastAndConstructible))
 			{
 				PCGLog::Metadata::LogFailToGetAttributeError<int32>(Selector, SeedAccessor.Get(), InContext);
 			}
diff --git a/Engine/Plugins/PCG/Source/PCG/Private/Elements/Grammar/PCGSubdivideSpline.cpp b/Engine/Plugins/PCG/Source/PCG/Private/Elements/Grammar/PCGSubdivideSpline.cpp
--- a/Engine/Plugins/PCG/Source/PCG/Private/Elements/Grammar/PCGSubdivideSpline.cpp
+++ b/Engine/Plugins/PCG/Source/PCG/Private/Elements/Grammar/PCGSubdivideSpline.cpp
@@ -355,14 +355,15 @@ bool FPCGSubdivideSplineElement::ExecuteInternal(FPCGContext* InContext) const
 			check(InputPolyLineData->Metadata);
 			const FPCGAttributePropertyInputSelector Selector = Settings->ModuleHeightAttribute.CopyAndFixLast(InputPolyLineData);
 			TUniquePtr<const IPCGAttributeAccessor> HeightAccessor = PCGAttributeAccessorHelpers::CreateConstAccessor(InputPolyLineData, Selector);
+			TUniquePtr<const IPCGAttributeAccessorKeys> HeightAccessorKeys = PCGAttributeAccessorHelpers::CreateConstKeys(InputPolyLineData, Selector);
 
-			if (!HeightAccessor)
+			if (!HeightAccessor || !HeightAccessorKeys)
 			{
 				PCGLog::Metadata::LogFailToCreateAccessorError(Selector, InContext);
 				continue;
 			}
 
-			if (!HeightAccessor->Get<double>(ModuleHeight, FPCGAttributeAccessorKeysEntries(PCGFirstEntryKey), EPCGAttributeAccessorFlags::AllowBroadcastAndConstructible))
+			if (!HeightAccessor->Get<double>(ModuleHeight, *HeightAccessorKeys, EPCGAttributeAccessorFlags::AllowBroadcastAndConstructible))
 			{
 				PCGLog::Metadata::LogFailToGetAttributeError(Selector, InContext);
 				continue;
@@ -410,14 +411,15 @@ bool FPCGSubdivideSplineElement::ExecuteInternal(FPCGContext* InContext) const
 		if (Settings->bUseSeedAttribute)
 		{
 			const FPCGAttributePropertyInputSelector Selector = Settings->SeedAttribute.CopyAndFixLast(InputPolyLineData);
-			TUniquePtr<const IPCGAttributeAccessor> SeedAccessor = PCGAttributeAccessorHelpers::CreateConstAccessor(InputPolyLineData, Selector);
+			const TUniquePtr<const IPCGAttributeAccessor> SeedAccessor = PCGAttributeAccessorHelpers::CreateConstAccessor(InputPolyLineData, Selector);
+			const TUniquePtr<const IPCGAttributeAccessorKeys> SeedAccessorKeys = PCGAttributeAccessorHelpers::CreateConstKeys(InputPolyLineData, Selector);
 
-			if (!SeedAccessor)
+			if (!SeedAccessor || !SeedAccessorKeys)
 			{
 				PCGLog::Metadata::LogFailToCreateAccessorError(Selector, InContext);
 			}
 			// Otherwise, get the value, if it fails, the attribute wasn't compatible
-			else if (!SeedAccessor->Get(Parameters.AdditionalSeed, FPCGAttributeAccessorKeysEntries(PCGFirstEntryKey), EPCGAttributeAccessorFlags::AllowBroadcastAndConstructible))
+			else if (!SeedAccessor->Get(Parameters.AdditionalSeed, *SeedAccessorKeys, EPCGAttributeAccessorFlags::AllowBroadcastAndConstructible))
 			{
 				PCGLog::Metadata::LogFailToGetAttributeError<int32>(Selector, SeedAccessor.Get(), InContext);
 			}
@@ -427,19 +429,23 @@ bool FPCGSubdivideSplineElement::ExecuteInternal(FPCGContext* InContext) const
 		{
 			const FPCGAttributePropertyInputSelector Selector = Settings->GrammarSelection.GrammarAttribute.CopyAndFixLast(InputPolyLineData);
 			const TUniquePtr<const IPCGAttributeAccessor> GrammarAccessor = PCGAttributeAccessorHelpers::CreateConstAccessor(InputPolyLineData, Selector);
-			const FPCGAttributeAccessorKeysEntries Keys(PCGFirstEntryKey);
-			if (!GrammarAccessor)
+			const TUniquePtr<const IPCGAttributeAccessorKeys> GrammarAccessorKeys = PCGAttributeAccessorHelpers::CreateConstKeys(InputPolyLineData, Selector);
+			if (!GrammarAccessor || !GrammarAccessorKeys)
 			{
 				PCGLog::Metadata::LogFailToCreateAccessorError(Selector, InContext);
 				continue;
 			}
-
-			auto Process = [&Parameters](const FString& InGrammar, int32 Index) -> void
+			
+			// Otherwise, get the value, if it fails, the attribute wasn't compatible
+			FString Grammar;
+			if (GrammarAccessor->Get(Grammar, *GrammarAccessorKeys, EPCGAttributeAccessorFlags::AllowBroadcastAndConstructible))
 			{
-				PCGSubdivideSplineHelpers::Process(Parameters, InGrammar);
-			};
-
-			PCGMetadataElementCommon::ApplyOnAccessor<FString>(Keys, *GrammarAccessor, Process);
+				PCGSubdivideSplineHelpers::Process(Parameters, Grammar);
+			}
+			else
+			{
+				PCGLog::Metadata::LogFailToGetAttributeError<int32>(Selector, GrammarAccessor.Get(), InContext);
+			}
 		}
 		else
 		{
diff --git a/Engine/Plugins/PCG/Source/PCG/Private/Elements/Grammar/PCGSubdivisionBase.cpp b/Engine/Plugins/PCG/Source/PCG/Private/Elements/Grammar/PCGSubdivisionBase.cpp
--- a/Engine/Plugins/PCG/Source/PCG/Private/Elements/Grammar/PCGSubdivisionBase.cpp
+++ b/Engine/Plugins/PCG/Source/PCG/Private/Elements/Grammar/PCGSubdivisionBase.cpp
@@ -126,13 +126,14 @@ PCGGrammar::FTokenizedGrammar FPCGSubdivisionBaseElement::GetTokenizedGrammar(FP
 	{
 		const FPCGAttributePropertyInputSelector Selector = InSettings->GrammarSelection.GrammarAttribute.CopyAndFixLast(InputData);
 		const TUniquePtr<const IPCGAttributeAccessor> Accessor = PCGAttributeAccessorHelpers::CreateConstAccessor(InputData, Selector);
-		if (!Accessor)
+		const TUniquePtr<const IPCGAttributeAccessorKeys> Keys = PCGAttributeAccessorHelpers::CreateConstKeys(InputData, Selector);
+		if (!Accessor || !Keys)
 		{
 			PCGLog::Metadata::LogFailToCreateAccessorError(Selector, InContext);
 			return {};
 		}
 
-		if (!Accessor->Get(Grammar, FPCGAttributeAccessorKeysEntries(PCGFirstEntryKey), EPCGAttributeAccessorFlags::AllowBroadcastAndConstructible))
+		if (!Accessor->Get(Grammar, *Keys, EPCGAttributeAccessorFlags::AllowBroadcastAndConstructible))
 		{
 			PCGLog::Metadata::LogFailToGetAttributeError<FString>(Selector, Accessor.Get(), InContext);
 			return {};
