[PCG] Add missing SetIsPreview calls which were preventing actor reuse (regression from CL 42448814)

#rnx

diff --git a/Engine/Plugins/PCG/Source/PCG/Private/Elements/PCGCreateTargetActor.cpp b/Engine/Plugins/PCG/Source/PCG/Private/Elements/PCGCreateTargetActor.cpp
--- a/Engine/Plugins/PCG/Source/PCG/Private/Elements/PCGCreateTargetActor.cpp
+++ b/Engine/Plugins/PCG/Source/PCG/Private/Elements/PCGCreateTargetActor.cpp
@@ -513,6 +513,9 @@ bool FPCGCreateTargetActorElement::ExecuteInternal(FPCGContext* Context) const
 			{
 				ManagedActors->SetCrc(ResourceCrc);
 			}
+#if WITH_EDITOR
+			ManagedActors->SetIsPreview(SpawnDefaultActorParams.bIsPreviewActor);
+#endif
 			ManagedActors->GetMutableGeneratedActors().AddUnique(GeneratedActor);
 			ManagedActors->bSupportsReset = !Settings->bDeleteActorsBeforeGeneration;
 
diff --git a/Engine/Plugins/PCG/Source/PCG/Private/Elements/PCGSpawnActor.cpp b/Engine/Plugins/PCG/Source/PCG/Private/Elements/PCGSpawnActor.cpp
--- a/Engine/Plugins/PCG/Source/PCG/Private/Elements/PCGSpawnActor.cpp
+++ b/Engine/Plugins/PCG/Source/PCG/Private/Elements/PCGSpawnActor.cpp
@@ -968,6 +968,9 @@ void FPCGSpawnActorElement::SpawnActors(FPCGSubgraphContext* Context, AActor* Ta
 
 		// Create managed resource for actor tracking
 		UPCGManagedActors* ManagedActors = NewObject<UPCGManagedActors>(SourceComponent);
+#if WITH_EDITOR
+		ManagedActors->SetIsPreview(SpawnDefaultActorParams.bIsPreviewActor);
+#endif
 		ManagedActors->SetCrc(InputDependenciesCrc);
 		ManagedActors->bSupportsReset = !Settings->bDeleteActorsBeforeGeneration;
 
