[PCG] Added a button in the graph details view to run the graph determinism test

#rnx

diff --git a/Engine/Plugins/PCG/Source/PCG/Private/Graph/PCGGraphExecutor.cpp b/Engine/Plugins/PCG/Source/PCG/Private/Graph/PCGGraphExecutor.cpp
--- a/Engine/Plugins/PCG/Source/PCG/Private/Graph/PCGGraphExecutor.cpp
+++ b/Engine/Plugins/PCG/Source/PCG/Private/Graph/PCGGraphExecutor.cpp
@@ -798,6 +798,13 @@ FPCGTaskId FPCGGraphExecutor::ScheduleSingleElement(TSharedRef<IPCGElement> Elem
 void FPCGGraphExecutor::MarkInputResults(TArrayView<const FPCGTaskId> InInputResults)
 {
 	TRACE_CPUPROFILER_EVENT_SCOPE(FPCGGraphExecutor::MarkInputResults);
+#if WITH_EDITOR
+	if (bDebugKeepIntermediateResults)
+	{
+		return;
+	}
+#endif
+
 	if (InInputResults.Num() > 0)
 	{
 		PCGGraphExecutor::TScopeLock ScopeLock(TaskOutputsLock);
@@ -2687,6 +2694,8 @@ void FPCGGraphExecutor::LogTaskStateNoLock() const
 
 FPCGTaskId FPCGGraphExecutor::ScheduleDebugWithTaskCallback(UPCGComponent* InComponent, TFunction<void(FPCGTaskId/* TaskId*/, const UPCGNode*/* Node*/, const FPCGDataCollection&/* TaskOutput*/)> TaskCompleteCallback)
 {
+	ensure(bDebugKeepIntermediateResults); // If this triggers, call SetKeepIntermediateResults first, as otherwise there are no guarantees the results will be kept, and this might trigger other ensures.
+
 	check(InComponent);
 	FPCGTaskId FinalTaskID = Schedule(InComponent, {});
 	InComponent->CurrentGenerationTask = FinalTaskID;
diff --git a/Engine/Plugins/PCG/Source/PCG/Private/Graph/PCGGraphExecutor.h b/Engine/Plugins/PCG/Source/PCG/Private/Graph/PCGGraphExecutor.h
--- a/Engine/Plugins/PCG/Source/PCG/Private/Graph/PCGGraphExecutor.h
+++ b/Engine/Plugins/PCG/Source/PCG/Private/Graph/PCGGraphExecutor.h
@@ -123,8 +123,12 @@ public:
 	const FPCGGraphCompiler* GetCompiler() const { return &GraphCompiler; }
 
 #if WITH_EDITOR
+	/** Experimental & Internal - Visits all intermediate results for all tasks. Requires SetKeepIntermediateResults to be called before hand. */
 	FPCGTaskId ScheduleDebugWithTaskCallback(UPCGComponent* InComponent, TFunction<void(FPCGTaskId, const UPCGNode*, const FPCGDataCollection&)> TaskCompleteCallback);
 
+	/** Notifies the executor to keep all intermediate results. Is needed prior to the ScheduleDebugWithTaskCallback call. @todo_pcg: Replace this & SetDisableClearResults with an enum. */
+	void SetKeepIntermediateResults(bool bInKeepIntermediateResults) { bDebugKeepIntermediateResults = bInKeepIntermediateResults; }
+
 	/** Notify compiler that graph has changed so it'll be removed from the cache */
 	void NotifyGraphChanged(UPCGGraph* InGraph, EPCGChangeType ChangeType);
 
@@ -374,6 +378,7 @@ private:
 	/** Set true when a non-runtime-gen component is scheduled and reset to false when all tasks are done executing. */
 	std::atomic<bool> bAnyNonRuntimeGenComponentScheduled = false;
 
+	bool bDebugKeepIntermediateResults = false;
 	bool bDisableClearResults = false;
 #endif
 
diff --git a/Engine/Plugins/PCG/Source/PCG/Private/Tests/Determinism/PCGDeterminismTestsCommon.cpp b/Engine/Plugins/PCG/Source/PCG/Private/Tests/Determinism/PCGDeterminismTestsCommon.cpp
--- a/Engine/Plugins/PCG/Source/PCG/Private/Tests/Determinism/PCGDeterminismTestsCommon.cpp
+++ b/Engine/Plugins/PCG/Source/PCG/Private/Tests/Determinism/PCGDeterminismTestsCommon.cpp
@@ -110,6 +110,7 @@ namespace PCGDeterminismTests
 		}
 
 		TSharedPtr<FPCGGraphExecutor> Executor = MakeShared<FPCGGraphExecutor>();
+		Executor->SetKeepIntermediateResults(true);
 
 		ON_SCOPE_EXIT
 		{
diff --git a/Engine/Plugins/PCG/Source/PCGEditor/Private/Details/PCGGraphDetails.cpp b/Engine/Plugins/PCG/Source/PCGEditor/Private/Details/PCGGraphDetails.cpp
--- a/Engine/Plugins/PCG/Source/PCGEditor/Private/Details/PCGGraphDetails.cpp
+++ b/Engine/Plugins/PCG/Source/PCGEditor/Private/Details/PCGGraphDetails.cpp
@@ -32,8 +32,7 @@ void FPCGGraphDetails::CustomizeDetails(IDetailLayoutBuilder& DetailBuilder)
 		}
 	}
 
-	// UE_DEPRECATED(5.6, "Added for familiarity and convenience to the Graph Parameters panel moving.")
-	IDetailCategoryBuilder& SettingsCategory = DetailBuilder.EditCategory(TEXT("Instance"));
+	TWeakPtr<const FPCGEditor> PCGEditorWeakPtr = nullptr;
 	if (!ObjectsBeingCustomized.IsEmpty() && ObjectsBeingCustomized[0].IsValid())
 	{
 		const UPCGEditorGraph* PCGEditorGraph = nullptr;
@@ -47,48 +46,94 @@ void FPCGGraphDetails::CustomizeDetails(IDetailLayoutBuilder& DetailBuilder)
 			PCGEditorGraph = FPCGEditor::GetPCGEditorGraph(FirstSettings);
 		}
 
-		const TWeakPtr<const FPCGEditor> PCGEditorWeakPtr = PCGEditorGraph ? PCGEditorGraph->GetEditor() : nullptr;
+		PCGEditorWeakPtr = PCGEditorGraph ? PCGEditorGraph->GetEditor() : nullptr;
+	}
 
+	if (PCGEditorWeakPtr.IsValid())
+	{
+		// UE_DEPRECATED(5.6, "Added for familiarity and convenience to the Graph Parameters panel moving.")
 		// Add an "Open Graph Parameters" button if the source editor tab is not already visible.
-		if (PCGEditorWeakPtr.IsValid())
-		{
-			SettingsCategory.AddCustomRow(FText::GetEmpty())
-				.Visibility(
-				TAttribute<EVisibility>::CreateLambda(
-				[PCGEditorWeakPtr]
-				{
-					const TSharedPtr<const FPCGEditor> PCGEditorSharedPtr = PCGEditorWeakPtr.Pin();
-					return (PCGEditorSharedPtr.IsValid() && PCGEditorSharedPtr->IsPanelCurrentlyOpen(EPCGEditorPanel::UserParams)) ? EVisibility::Collapsed : EVisibility::Visible;
-				}))
-				.ValueContent()
-				.MaxDesiredWidth(120.f)
+		IDetailCategoryBuilder& SettingsCategory = DetailBuilder.EditCategory(TEXT("Instance"));
+
+		SettingsCategory.AddCustomRow(FText::GetEmpty())
+			.Visibility(
+			TAttribute<EVisibility>::CreateLambda(
+			[PCGEditorWeakPtr]
+			{
+				const TSharedPtr<const FPCGEditor> PCGEditorSharedPtr = PCGEditorWeakPtr.Pin();
+				return (PCGEditorSharedPtr.IsValid() && PCGEditorSharedPtr->IsPanelCurrentlyOpen(EPCGEditorPanel::UserParams)) ? EVisibility::Collapsed : EVisibility::Visible;
+			}))
+			.ValueContent()
+			.MaxDesiredWidth(120.f)
+			[
+				SNew(SHorizontalBox)
+				+ SHorizontalBox::Slot()
+				.AutoWidth()
+				.Padding(2.0f, 0.0f)
+				.VAlign(VAlign_Fill)
 				[
-					SNew(SHorizontalBox)
-					+ SHorizontalBox::Slot()
-					.AutoWidth()
-					.Padding(2.0f, 0.0f)
-					.VAlign(VAlign_Fill)
+					SNew(SButton)
+					.OnClicked_Lambda(
+					[PCGEditorWeakPtr]() -> FReply
+					{
+						if (const TSharedPtr<const FPCGEditor> SharedPtr = PCGEditorWeakPtr.Pin())
+						{
+							SharedPtr->BringFocusToPanel(EPCGEditorPanel::UserParams);
+						}
+
+						return FReply::Handled();
+					})
+					.ToolTipText(LOCTEXT("OpenGraphParamPanelTooltip", "Opens the Graph Parameters Panel."))
 					[
-						SNew(SButton)
-						.OnClicked_Lambda(
-						[PCGEditorWeakPtr]() -> FReply
+						SNew(STextBlock)
+						.Font(IDetailLayoutBuilder::GetDetailFont())
+						.Text(LOCTEXT("OpenGraphParamPanelButtonText", "Open Graph Parameters"))
+					]
+				]
+			];
+
+		// Add a "Run Graph Determinism Test" button in the debug category
+		IDetailCategoryBuilder& DebugCategory = DetailBuilder.EditCategory(TEXT("Debug"));
+
+		DebugCategory.AddCustomRow(FText::GetEmpty())
+			.ValueContent()
+			.MaxDesiredWidth(120.f)
+			[
+				SNew(SHorizontalBox)
+				+ SHorizontalBox::Slot()
+				.AutoWidth()
+				.Padding(2.0f, 0.0f)
+				.VAlign(VAlign_Fill)
+				[
+					SNew(SButton)
+					.OnClicked_Lambda(
+					[PCGEditorWeakPtr]() -> FReply
+					{
+						if (const TSharedPtr<const FPCGEditor> PCGEditor = PCGEditorWeakPtr.Pin())
+						{
+							PCGEditor->OnDeterminismGraphTest();
+						}
+
+						return FReply::Handled();
+					})
+					.IsEnabled_Lambda(
+					[PCGEditorWeakPtr]() -> bool
+					{
+						if (const TSharedPtr<const FPCGEditor> PCGEditor = PCGEditorWeakPtr.Pin())
 						{
-							if (const TSharedPtr<const FPCGEditor> SharedPtr = PCGEditorWeakPtr.Pin())
-							{
-								SharedPtr->BringFocusToPanel(EPCGEditorPanel::UserParams);
-							}
-
-							return FReply::Handled();
-						})
-						.ToolTipText(LOCTEXT("OpenGraphParamPanelTooltip", "Opens the Graph Parameters Panel."))
-						[
-							SNew(STextBlock)
-							.Font(IDetailLayoutBuilder::GetDetailFont())
-							.Text(LOCTEXT("OpenGraphParamPanelButtonText", "Open Graph Parameters"))
-						]
+							return PCGEditor->CanRunDeterminismGraphTest();
+						}
+
+						return false;
+					})
+					.ToolTipText(LOCTEXT("RunGraphDeterminismTestTooltip", "Runs the graph-level determinism test on the currently selected debug object."))
+					[
+						SNew(STextBlock)
+						.Font(IDetailLayoutBuilder::GetDetailFont())
+						.Text(LOCTEXT("RunGraphDeterminism", "Run Determinism Test"))
 					]
-				];
-		}
+				]
+			];
 	}
 
 	const FName PCGCategoryName("PCG");
diff --git a/Engine/Plugins/PCG/Source/PCGEditor/Private/PCGEditor.cpp b/Engine/Plugins/PCG/Source/PCGEditor/Private/PCGEditor.cpp
--- a/Engine/Plugins/PCG/Source/PCGEditor/Private/PCGEditor.cpp
+++ b/Engine/Plugins/PCG/Source/PCGEditor/Private/PCGEditor.cpp
@@ -1123,7 +1123,7 @@ bool FPCGEditor::CanRunDeterminismNodeTest() const
 	return false;
 }
 
-void FPCGEditor::OnDeterminismNodeTest()
+void FPCGEditor::OnDeterminismNodeTest() const
 {
 	check(GraphEditorWidget.IsValid());
 
@@ -1228,7 +1228,7 @@ bool FPCGEditor::CanRunDeterminismGraphTest() const
 	return PCGEditorGraph && PCGComponentBeingInspected.IsValid();
 }
 
-void FPCGEditor::OnDeterminismGraphTest()
+void FPCGEditor::OnDeterminismGraphTest() const
 {
 	check(GraphEditorWidget.IsValid());
 
diff --git a/Engine/Plugins/PCG/Source/PCGEditor/Private/PCGEditor.h b/Engine/Plugins/PCG/Source/PCGEditor/Private/PCGEditor.h
--- a/Engine/Plugins/PCG/Source/PCGEditor/Private/PCGEditor.h
+++ b/Engine/Plugins/PCG/Source/PCGEditor/Private/PCGEditor.h
@@ -158,6 +158,11 @@ public:
 
 	FOnInspectedStackChanged OnInspectedStackChangedDelegate;
 
+	/** Can determinism be tested on the current graph */
+	bool CanRunDeterminismGraphTest() const;
+	/** Run the determinism test on the current graph */
+	void OnDeterminismGraphTest() const;
+
 protected:
 	// ~Begin FAssetEditorToolkit interface
 	/** Called when "Save" is clicked for this asset */
@@ -238,12 +243,7 @@ private:
 	/** Can determinism be tested on the selected node(s) */
 	bool CanRunDeterminismNodeTest() const;
 	/** Run the determinism test on the selected node(s) */
-	void OnDeterminismNodeTest();
-
-	/** Can determinism be tested on the current graph */
-	bool CanRunDeterminismGraphTest() const;
-	/** Run the determinism test on the current graph */
-	void OnDeterminismGraphTest();
+	void OnDeterminismNodeTest() const;
 
 	/** Open details view for the PCG object being edited */
 	void OnEditGraphSettings();
