[PCG] Fix CopyPoints not outputting proper transforms if Source input doesn't have allocated transforms (SoA only)

#rnx

diff --git a/Engine/Plugins/PCG/Source/PCG/Private/Elements/PCGCopyPoints.cpp b/Engine/Plugins/PCG/Source/PCG/Private/Elements/PCGCopyPoints.cpp
--- a/Engine/Plugins/PCG/Source/PCG/Private/Elements/PCGCopyPoints.cpp
+++ b/Engine/Plugins/PCG/Source/PCG/Private/Elements/PCGCopyPoints.cpp
@@ -318,6 +318,7 @@ bool FPCGCopyPointsElement::ExecuteInternal(FPCGContext* Context) const
 				}
 			};
 
+			PropertiesToAllocate |= EPCGPointNativeProperties::Transform;
 			PropertiesToAllocate |= GetPropertyToAllocate(SeedInheritance, EPCGPointNativeProperties::Seed);
 			PropertiesToAllocate |= GetPropertyToAllocate(ColorInheritance, EPCGPointNativeProperties::Color);
 
