Fixup API macros

#rnx

diff --git a/Engine/Plugins/PCG/Source/PCG/Public/PCGGraph.h b/Engine/Plugins/PCG/Source/PCG/Public/PCGGraph.h
--- a/Engine/Plugins/PCG/Source/PCG/Public/PCGGraph.h
+++ b/Engine/Plugins/PCG/Source/PCG/Public/PCGGraph.h
@@ -424,7 +424,7 @@ public:
 	UE_API void SetExtraEditorNodes(const TArray<TObjectPtr<const UObject>>& InNodes);
 
 	const TArray<FPCGGraphCommentNodeData>& GetCommentNodes() const { return CommentNodes; }
-	UE_API void SetCommentNodes(TArray<FPCGGraphCommentNodeData> InNodes) { CommentNodes = std::move(InNodes); }
+	void SetCommentNodes(TArray<FPCGGraphCommentNodeData> InNodes) { CommentNodes = std::move(InNodes); }
 	UE_API void RemoveCommentNode(const FGuid& InNodeGUID);
 
 	bool IsInspecting() const { return bIsInspecting; }
