[PCG] Memory check in the surface sampler was not on the actual number of element allocated

* Not a silver bullet to fix all OOM but will be way more accurate than the previous count.

#rnx

diff --git a/Engine/Plugins/PCG/Source/PCG/Private/Elements/PCGSurfaceSampler.cpp b/Engine/Plugins/PCG/Source/PCG/Private/Elements/PCGSurfaceSampler.cpp
--- a/Engine/Plugins/PCG/Source/PCG/Private/Elements/PCGSurfaceSampler.cpp
+++ b/Engine/Plugins/PCG/Source/PCG/Private/Elements/PCGSurfaceSampler.cpp
@@ -169,9 +169,9 @@ namespace PCGSurfaceSampler
 		}
 
 		if (PCGFeatureSwitches::CVarCheckSamplerMemory.GetValueOnAnyThread()
-			&& PCGFeatureSwitches::Helpers::GetAvailableMemoryForSamplers() < (sizeof(FPCGPoint) * FMath::Min(TargetPointCount, CellCount)))
+			&& PCGFeatureSwitches::Helpers::GetAvailableMemoryForSamplers() < (sizeof(FPCGPoint) * CellCount))
 		{
-			PCGLog::LogErrorOnGraph(FText::Format((LOCTEXT("TooManyPoints", "Skipped - tried to generate too many points ({0}).\nAdjust 'pcg.SamplerMemoryThreshold' if needed.")), FMath::Min(TargetPointCount, CellCount)), Context);
+			PCGLog::LogErrorOnGraph(FText::Format((LOCTEXT("TooManyPoints", "Skipped - tried to generate too many points ({0}).\nAdjust 'pcg.SamplerMemoryThreshold' if needed.")), CellCount), Context);
 			return false;
 		}
 
