[PCG] Added culling option on execution dependency pins.

diff --git a/Engine/Plugins/PCG/Source/PCG/Private/PCGSettings.cpp b/Engine/Plugins/PCG/Source/PCG/Private/PCGSettings.cpp
--- a/Engine/Plugins/PCG/Source/PCG/Private/PCGSettings.cpp
+++ b/Engine/Plugins/PCG/Source/PCG/Private/PCGSettings.cpp
@@ -431,7 +431,7 @@ void UPCGSettings::AddExecutionDependencyPin(TArray<FPCGPinProperties>& OutPins)
 		DepPin.Tooltip = PCGPinConstants::Tooltips::ExecutionDependencyTooltip;
 #endif // WITH_EDITOR
 		DepPin.Usage = EPCGPinUsage::DependencyOnly;
-		DepPin.PinStatus = EPCGPinStatus::Advanced;
+		DepPin.PinStatus = bExecutionDependencyRequired ? EPCGPinStatus::Required : EPCGPinStatus::Advanced;
 	}
 }
 
@@ -652,7 +652,8 @@ EPCGChangeType UPCGSettings::GetChangeTypeForProperty(const FName& InPropertyNam
 {
 	EPCGChangeType ChangeType = EPCGChangeType::Settings;
 
-	if (InPropertyName == GET_MEMBER_NAME_CHECKED(UPCGSettings, bExecuteOnGPU))
+	if (InPropertyName == GET_MEMBER_NAME_CHECKED(UPCGSettings, bExecuteOnGPU) ||
+		InPropertyName == GET_MEMBER_NAME_CHECKED(UPCGSettings, bExecutionDependencyRequired))
 	{
 		ChangeType |= EPCGChangeType::Structural;
 	}
@@ -733,6 +734,11 @@ void UPCGSettings::PostPaste()
 	InitializeCachedOverridableParams(/*bReset=*/true);
 }
 
+bool UPCGSettings::CanUseExecutionDependencyPinForCulling() const
+{
+	return HasExecutionDependencyPin();
+}
+
 void UPCGSettings::ApplyDeprecationBeforeUpdatePins(UPCGNode* InOutNode, TArray<TObjectPtr<UPCGPin>>& InputPins, TArray<TObjectPtr<UPCGPin>>& OutputPins)
 {
 	check(InOutNode);
diff --git a/Engine/Plugins/PCG/Source/PCG/Public/PCGSettings.h b/Engine/Plugins/PCG/Source/PCG/Public/PCGSettings.h
--- a/Engine/Plugins/PCG/Source/PCG/Public/PCGSettings.h
+++ b/Engine/Plugins/PCG/Source/PCG/Public/PCGSettings.h
@@ -308,6 +308,10 @@ public:
 
 	/** Perform post-operations when an editor node is copied */
 	UE_API virtual void PostPaste();
+
+	/** Returns whether the execution dependency pin can be configured to be used for culling. */
+	UFUNCTION()
+	bool CanUseExecutionDependencyPinForCulling() const;
 #endif // WITH_EDITOR
 
 	/** Derived classes can implement this to expose additional information or context, such as an asset in use by the node. */
@@ -363,6 +367,10 @@ public:
 	UPROPERTY(EditAnywhere, BlueprintReadWrite, Category = Settings, meta=(EditCondition="UseSeed()", EditConditionHides, PCG_Overridable))
 	int Seed = 0xC35A9631; // Default seed is a random prime number, but will be overriden for new settings based on the class type name hash, making each settings class have a different default seed.
 
+	/** If enabled, the execution dependency pin will require to be connected for execution (as any other required pin) and will be used for culling. */
+	UPROPERTY(EditAnywhere, BlueprintReadWrite, Category = Settings, AdvancedDisplay, meta=(EditCondition="CanUseExecutionDependencyPinForCulling()", EditConditionHides))
+	bool bExecutionDependencyRequired = false;
+
 #if WITH_EDITORONLY_DATA
 	UPROPERTY()
 	TSet<FString> TagsAppliedOnOutput_DEPRECATED;
