[PCG] Restore CreateAttribute for discoverability

#rnx

diff --git a/Engine/Plugins/PCG/Source/PCG/Private/Elements/PCGCreateAttribute.cpp b/Engine/Plugins/PCG/Source/PCG/Private/Elements/PCGCreateAttribute.cpp
--- a/Engine/Plugins/PCG/Source/PCG/Private/Elements/PCGCreateAttribute.cpp
+++ b/Engine/Plugins/PCG/Source/PCG/Private/Elements/PCGCreateAttribute.cpp
@@ -269,7 +269,7 @@ FText UPCGCreateAttributeSetSettings::GetNodeTooltipText() const
 TArray<FPCGPreConfiguredSettingsInfo> UPCGCreateAttributeSetSettings::GetPreconfiguredInfo() const
 {
 	return FPCGPreConfiguredSettingsInfo::PopulateFromEnum<EPCGMetadataTypes>(
-		/*InValuesToSkip=*/{EPCGMetadataTypes::Count, EPCGMetadataTypes::Float, EPCGMetadataTypes::Unknown},
+		/*InValuesToSkip=*/{EPCGMetadataTypes::Count, EPCGMetadataTypes::Unknown},
 		/*InOptionalFormat=*/FTextFormat(LOCTEXT("PreconfigureFormat", "New {0} Constant")));
 }
 #endif // WITH_EDITOR
diff --git a/Engine/Plugins/PCG/Source/PCG/Public/Elements/PCGCreateAttribute.h b/Engine/Plugins/PCG/Source/PCG/Public/Elements/PCGCreateAttribute.h
--- a/Engine/Plugins/PCG/Source/PCG/Public/Elements/PCGCreateAttribute.h
+++ b/Engine/Plugins/PCG/Source/PCG/Public/Elements/PCGCreateAttribute.h
@@ -104,7 +104,7 @@ public:
 	virtual FText GetNodeTooltipText() const override;
 	virtual TArray<FPCGPreConfiguredSettingsInfo> GetPreconfiguredInfo() const override;
 	virtual bool GroupPreconfiguredSettings() const override { return true; }
-	virtual bool OnlyExposePreconfiguredSettings() const override { return true; }
+	virtual bool OnlyExposePreconfiguredSettings() const override { return false; }
 #endif
 	virtual FString GetAdditionalTitleInformation() const override;
 	virtual void ApplyPreconfiguredSettings(const FPCGPreConfiguredSettingsInfo& PreconfigureInfo) override;
