[PCG] Fix cook warnings related to data interfaces.

The texture data interface changes were because the compiler produces warnings when control flow can fall through an if statement that has a return in it.

[RN] Fixed warnings emitted from cooking GPU node kernels.

diff --git a/Engine/Plugins/PCG/Shaders/Private/PCGLandscapeDataInterface.ush b/Engine/Plugins/PCG/Shaders/Private/PCGLandscapeDataInterface.ush
--- a/Engine/Plugins/PCG/Shaders/Private/PCGLandscapeDataInterface.ush
+++ b/Engine/Plugins/PCG/Shaders/Private/PCGLandscapeDataInterface.ush
@@ -125,8 +125,26 @@ float GetHeight_{DataInterfaceName}(float3 InWorldPos)
 		const float2 UV = GetCollisionHeightTextureUV_{DataInterfaceName}(InWorldPos);
 		return {DataInterfaceName}_CollisionHeightTexture.SampleLevel({DataInterfaceName}_CollisionHeightTextureSampler, UV, 0).x;
 	}
+	else
+	{
+		return InWorldPos.z;
+	}
+}
+
+float3 CalculateNormalFromHeights_{DataInterfaceName}(float3 InWorldPos)
+{
+	float3 Pos1 = InWorldPos;
+	float3 Pos2 = InWorldPos + float3(0.0f, {DataInterfaceName}_LandscapeGridSize.y, 0.0f);
+	float3 Pos3 = InWorldPos + float3({DataInterfaceName}_LandscapeGridSize.x, 0.0f, 0.0f);
 
-	return InWorldPos.z;
+	Pos1.z = GetHeight_{DataInterfaceName}(Pos1);
+	Pos2.z = GetHeight_{DataInterfaceName}(Pos2);
+	Pos3.z = GetHeight_{DataInterfaceName}(Pos3);
+
+	const float3 DeltaX = Pos2 - Pos1;
+	const float3 DeltaY = Pos3 - Pos1;
+
+	return normalize(cross(DeltaY, DeltaX));
 }
 
 float3 GetNormal_{DataInterfaceName}(float3 InWorldPos)
@@ -192,21 +210,15 @@ float3 GetNormal_{DataInterfaceName}(float3 InWorldPos)
 		{
 			return VirtualTextureUnpackNormalBGR565(PackedValue1);
 		}
+		else
+		{
+			return CalculateNormalFromHeights_{DataInterfaceName}(InWorldPos);
+		}
+	}
+	else
+	{
+		return CalculateNormalFromHeights_{DataInterfaceName}(InWorldPos);
 	}
-
-	float3 Pos1 = InWorldPos;
-	float3 Pos2 = InWorldPos + float3(0.0f, {DataInterfaceName}_LandscapeGridSize.y, 0.0f);
-	float3 Pos3 = InWorldPos + float3({DataInterfaceName}_LandscapeGridSize.x, 0.0f, 0.0f);
-
-	Pos1.z = GetHeight_{DataInterfaceName}(Pos1);
-	Pos2.z = GetHeight_{DataInterfaceName}(Pos2);
-	Pos3.z = GetHeight_{DataInterfaceName}(Pos3);
-
-	const float3 DeltaX = Pos2 - Pos1;
-	const float3 DeltaY = Pos3 - Pos1;
-
-	return normalize(cross(DeltaY, DeltaX));
-
 }
 
 float3 GetBaseColor_{DataInterfaceName}(float3 InWorldPos)
@@ -266,8 +278,10 @@ float3 GetBaseColor_{DataInterfaceName}(float3 InWorldPos)
 			return PackedValue.rgb;
 		}
 	}
-
-	return (float3)0.0f;
+	else
+	{
+		return (float3)0.0f;
+	}
 }
 
 #undef VIRTUAL_TEXTURE_UNPACK_NORMAL_NONE
diff --git a/Engine/Plugins/PCG/Shaders/Private/PCGShaderUtils.ush b/Engine/Plugins/PCG/Shaders/Private/PCGShaderUtils.ush
--- a/Engine/Plugins/PCG/Shaders/Private/PCGShaderUtils.ush
+++ b/Engine/Plugins/PCG/Shaders/Private/PCGShaderUtils.ush
@@ -30,22 +30,22 @@ float FRand(inout uint InOutSeed)
 	return Result - 1.0f;
 }
 
-float3 CreateGrid2D(int ElementIndex, int NumPoints, int NumRows, float3 Min, float3 Max)
+float3 CreateGrid2D(uint ElementIndex, uint NumPoints, uint NumRows, float3 Min, float3 Max)
 {
 	if (NumRows == 0)
 	{
 		return Min;
 	}
 
-	const int NumCols = ceil((0.f + NumPoints) / NumRows);
+	const uint NumCols = ceil((0.f + NumPoints) / NumRows);
 
 	if (NumCols == 0)
 	{
 		return Min;
 	}
 
-	const int IndexX = ElementIndex / NumRows;
-	const int IndexY = ElementIndex - NumRows * IndexX;
+	const uint IndexX = ElementIndex / NumRows;
+	const uint IndexY = ElementIndex - NumRows * IndexX;
 
 	float3 Spacing = (Max - Min);
 
@@ -66,32 +66,32 @@ float3 CreateGrid2D(int ElementIndex, int NumPoints, int NumRows, float3 Min, fl
 	return Position;
 }
 
-float3 CreateGrid2D(int ElementIndex, int NumPoints, float3 Min, float3 Max)
+float3 CreateGrid2D(uint ElementIndex, uint NumPoints, float3 Min, float3 Max)
 {
-	if (NumPoints <= 0)
+	if (NumPoints == 0)
 	{
 		return Min;
 	}
 
-	const int NumRows = ceil(sqrt(NumPoints));
+	const uint NumRows = ceil(sqrt(NumPoints));
 	return CreateGrid2D(ElementIndex, NumPoints, NumRows, Min, Max);
 }
 
-float3 CreateGrid3D(int ElementIndex, int NumPoints, int NumRows, int NumCols, float3 Min, float3 Max)
+float3 CreateGrid3D(uint ElementIndex, uint NumPoints, uint NumRows, uint NumCols, float3 Min, float3 Max)
 {
 	if (NumRows == 0 || NumCols == 0)
 	{
 		return Min;
 	}
 
-	const int NumLayers = ceil((0.f + NumPoints) / (NumRows * NumCols));
+	const uint NumLayers = ceil((0.f + NumPoints) / (NumRows * NumCols));
 	
 	if (NumLayers == 0)
 	{
 		return Min;
 	}
 	
-	int3 Index;
+	uint3 Index;
 	Index.x = ElementIndex % NumCols;
 	Index.y = (ElementIndex / NumCols) % NumRows;
 	Index.z = ElementIndex / (NumRows * NumCols);
@@ -113,15 +113,15 @@ float3 CreateGrid3D(int ElementIndex, int NumPoints, int NumRows, int NumCols, f
 	return Min + Index * Spacing;
 }
 
-float3 CreateGrid3D(int ElementIndex, int NumPoints, float3 Min, float3 Max)
+float3 CreateGrid3D(uint ElementIndex, uint NumPoints, float3 Min, float3 Max)
 {
-	if (NumPoints <= 0)
+	if (NumPoints == 0)
 	{
 		return Min;
 	}
 
-	const int NumRows = round(pow(NumPoints, 1.f / 3.f));
-	const int NumCols = NumRows;
+	const uint NumRows = round(pow(NumPoints, 1.f / 3.f));
+	const uint NumCols = NumRows;
 	
 	return CreateGrid3D(ElementIndex, NumPoints, NumRows, NumCols, Min, Max);
 }
diff --git a/Engine/Plugins/PCG/Shaders/Private/PCGTextureDataInterface.ush b/Engine/Plugins/PCG/Shaders/Private/PCGTextureDataInterface.ush
--- a/Engine/Plugins/PCG/Shaders/Private/PCGTextureDataInterface.ush
+++ b/Engine/Plugins/PCG/Shaders/Private/PCGTextureDataInterface.ush
@@ -1,5 +1,7 @@
 // Copyright Epic Games, Inc. All Rights Reserved.
 
+// todo_pcg: Runtime diagnostic errors would be useful if texture resource types don't match expectations below.
+
 // Mirror of ETextureDimension
 #define PCG_TextureDimension_Texture2D          0
 #define PCG_TextureDimension_Texture2DArray     1
@@ -156,7 +158,12 @@ float2 GetTexCoords_{DataInterfaceName}(float2 WorldPos, float2 Min, float2 Max)
 				return {DataInterfaceName}_TextureArray_##Index.SampleLevel({DataInterfaceName}_Sampler_Point, float3(InTextureUV, SliceIndex), 0); \
 			} \
 		} \
-	}
+		else \
+		{ \
+			return (float4)0; \
+		} \
+	} \
+    else
 
 float4 Sample_{DataInterfaceName}(uint InDataIndex, float2 InTextureUV)
 {
@@ -175,6 +182,10 @@ float4 Sample_{DataInterfaceName}(uint InDataIndex, float2 InTextureUV)
 	PCG_CONDITIONAL_TEXTURE_SAMPLE(5)
 	PCG_CONDITIONAL_TEXTURE_SAMPLE(6)
 	PCG_CONDITIONAL_TEXTURE_SAMPLE(7)
+	// else - from macro
+	{
+		return (float4)0;
+	}
 #else
 	if (TextureDimension == PCG_TextureDimension_Texture2D)
 	{
@@ -198,10 +209,11 @@ float4 Sample_{DataInterfaceName}(uint InDataIndex, float2 InTextureUV)
 			return {DataInterfaceName}_TextureArray_0.SampleLevel({DataInterfaceName}_Sampler_Point, float3(InTextureUV, SliceIndex), 0);
 		}
 	}
+	else
+	{
+		return (float4)0;
+	}
 #endif
-
-	// TODO: Runtime diagnostic error.
-	return (float4)0;
 }
 
 #undef PCG_CONDITIONAL_TEXTURE_SAMPLE
@@ -232,7 +244,12 @@ float4 Sample_{DataInterfaceName}(uint InDataIndex, float2 InTextureUV)
 				return {DataInterfaceName}_TextureArray_##Index.SampleLevel({DataInterfaceName}_Sampler_Point, float3(TextureUV, SliceIndex), 0); \
 			} \
 		} \
-	}
+		else \
+		{ \
+			return (float4)0; \
+		} \
+	} \
+	else
 
 float4 SampleWorldPos_{DataInterfaceName}(uint InDataIndex, float2 WorldPos)
 {
@@ -251,6 +268,10 @@ float4 SampleWorldPos_{DataInterfaceName}(uint InDataIndex, float2 WorldPos)
 	PCG_CONDITIONAL_TEXTURE_SAMPLE_WORLD_POS(5)
 	PCG_CONDITIONAL_TEXTURE_SAMPLE_WORLD_POS(6)
 	PCG_CONDITIONAL_TEXTURE_SAMPLE_WORLD_POS(7)
+	// else - from macro
+	{
+		return (float4)0;
+	}
 #else
 	const float2 TextureUV = GetTexCoords_{DataInterfaceName}(WorldPos, {DataInterfaceName}_TextureBounds[0].xy, {DataInterfaceName}_TextureBounds[0].zw);
 
@@ -276,10 +297,11 @@ float4 SampleWorldPos_{DataInterfaceName}(uint InDataIndex, float2 WorldPos)
 			return {DataInterfaceName}_TextureArray_0.SampleLevel({DataInterfaceName}_Sampler_Point, float3(TextureUV, SliceIndex), 0);
 		}
 	}
+	else
+	{
+		return (float4)0;
+	}
 #endif
-
-	// TODO: Runtime diagnostic error.
-	return (float4)0;
 }
 
 #undef PCG_CONDITIONAL_TEXTURE_SAMPLE_WORLD_POS
@@ -295,7 +317,12 @@ float4 SampleWorldPos_{DataInterfaceName}(uint InDataIndex, float2 WorldPos)
 		{ \
 			return {DataInterfaceName}_TextureArray_##Index[uint3(InElementIndex, SliceIndex)]; \
 		} \
-	}
+		else \
+		{ \
+			return (float4)0; \
+		} \
+	} \
+	else
 
 float4 Load_{DataInterfaceName}(uint InDataIndex, uint2 InElementIndex)
 {
@@ -314,6 +341,10 @@ float4 Load_{DataInterfaceName}(uint InDataIndex, uint2 InElementIndex)
 	PCG_CONDITIONAL_TEXTURE_LOAD(5)
 	PCG_CONDITIONAL_TEXTURE_LOAD(6)
 	PCG_CONDITIONAL_TEXTURE_LOAD(7)
+	// else - from macro
+	{
+		return (float4)0;
+	}
 #else
 	if (TextureDimension == PCG_TextureDimension_Texture2D)
 	{
@@ -323,10 +354,11 @@ float4 Load_{DataInterfaceName}(uint InDataIndex, uint2 InElementIndex)
 	{
 		return {DataInterfaceName}_TextureArray_0[uint3(InElementIndex, SliceIndex)];
 	}
+	else
+	{
+		return (float4)0;
+	}
 #endif
-
-	// TODO: Runtime diagnostic error.
-	return (float4)0;
 }
 
 #undef PCG_CONDITIONAL_TEXTURE_LOAD
@@ -350,14 +382,11 @@ void Store_{DataInterfaceName}(uint InDataIndex, uint2 InElementIndex, float4 In
 
 	if (TextureDimension != PCG_TextureDimension_Texture2D)
 	{
-		// TODO: Runtime diagnostic error.
 		//return;
 	}
 
 #if ENABLE_MULTIPLE_TEXTURE_OBJECTS
 	PCG_CONDITIONAL_TEXTURE_STORE(0)
-
-	// TODO: Runtime diagnostic error.
 #else
 	{DataInterfaceName}_TextureUAV_0[InElementIndex] = InValue;
 #endif
diff --git a/Engine/Plugins/PCG/Source/PCG/Private/Compute/DataInterfaces/PCGDataLabelResolverDataInterface.cpp b/Engine/Plugins/PCG/Source/PCG/Private/Compute/DataInterfaces/PCGDataLabelResolverDataInterface.cpp
--- a/Engine/Plugins/PCG/Source/PCG/Private/Compute/DataInterfaces/PCGDataLabelResolverDataInterface.cpp
+++ b/Engine/Plugins/PCG/Source/PCG/Private/Compute/DataInterfaces/PCGDataLabelResolverDataInterface.cpp
@@ -29,10 +29,8 @@ void UPCGDataLabelResolverDataInterface::GetSupportedInputs(TArray<FShaderFuncti
 
 void UPCGDataLabelResolverDataInterface::GetSupportedOutputs(TArray<FShaderFunctionDefinition>& OutFunctions) const
 {
-	OutFunctions.AddDefaulted_GetRef()
-		.SetName(TEXT("GetDataIndexFromIdInternal"))
-		.AddReturnType(EShaderFundamentalType::Uint) // Index of data in input data collection
-		.AddParam(EShaderFundamentalType::Uint); // InDataId
+	// Usage symmetrical across kernel inputs and kernel outputs.
+	GetSupportedInputs(OutFunctions);
 }
 
 BEGIN_SHADER_PARAMETER_STRUCT(FPCGDataLabelResolverDataInterfaceParameters, )
@@ -56,9 +54,9 @@ void UPCGDataLabelResolverDataInterface::GetHLSL(FString& OutHLSL, FString const
 		"uint {DataInterfaceName}_InNumData;\n"
 		"StructuredBuffer<int> {DataInterfaceName}_InDataIdToDataIndexMap;\n"
 		"\n"
-		"uint GetDataIndexFromIdInternal_{DataInterfaceName}(int InDataId)\n"
+		"uint GetDataIndexFromIdInternal_{DataInterfaceName}(uint InDataId)\n"
 		"{\n"
-		"	if (InDataId < 0 || InDataId >= {DataInterfaceName}_InNumData)\n"
+		"	if (InDataId >= {DataInterfaceName}_InNumData)\n"
 		"	{\n"
 		"		return 0;\n"
 		"	}\n"
diff --git a/Engine/Plugins/PCG/Source/PCG/Private/Compute/Elements/PCGCustomHLSL.cpp b/Engine/Plugins/PCG/Source/PCG/Private/Compute/Elements/PCGCustomHLSL.cpp
--- a/Engine/Plugins/PCG/Source/PCG/Private/Compute/Elements/PCGCustomHLSL.cpp
+++ b/Engine/Plugins/PCG/Source/PCG/Private/Compute/Elements/PCGCustomHLSL.cpp
@@ -841,10 +841,10 @@ void UPCGCustomHLSLSettings::UpdateHelperDeclarations()
 			"uint ComputeSeedFromPosition(float3 Position);\n"
 			"\n"
 			"// Returns the position of the Nth point in a 2D or 3D grid with the given constraints.\n"
-			"float3 CreateGrid2D(int ElementIndex, int NumPoints, float3 Min, float3 Max);\n"
-			"float3 CreateGrid2D(int ElementIndex, int NumPoints, int NumRows, float3 Min, float3 Max);\n"
-			"float3 CreateGrid3D(int ElementIndex, int NumPoints, float3 Min, float3 Max);\n"
-			"float3 CreateGrid3D(int ElementIndex, int NumPoints, int NumRows, int NumCols, float3 Min, float3 Max);\n");
+			"float3 CreateGrid2D(uint ElementIndex, uint NumPoints, float3 Min, float3 Max);\n"
+			"float3 CreateGrid2D(uint ElementIndex, uint NumPoints, uint NumRows, float3 Min, float3 Max);\n"
+			"float3 CreateGrid3D(uint ElementIndex, uint NumPoints, float3 Min, float3 Max);\n"
+			"float3 CreateGrid3D(uint ElementIndex, uint NumPoints, uint NumRows, uint NumCols, float3 Min, float3 Max);\n");
 
 		if (bPrintShaderDebugValues)
 		{
diff --git a/Engine/Plugins/PCG/Source/PCG/Private/Compute/Elements/PCGCustomHLSLKernel.cpp b/Engine/Plugins/PCG/Source/PCG/Private/Compute/Elements/PCGCustomHLSLKernel.cpp
--- a/Engine/Plugins/PCG/Source/PCG/Private/Compute/Elements/PCGCustomHLSLKernel.cpp
+++ b/Engine/Plugins/PCG/Source/PCG/Private/Compute/Elements/PCGCustomHLSLKernel.cpp
@@ -1497,7 +1497,7 @@ FString UPCGCustomHLSLKernel::ProcessShaderSource(FPCGGPUCompilationContext& InO
 		for (int DataId = 0; DataId < DataLabels.Num(); ++DataId)
 		{
 			const FString& DataLabel = DataLabels[DataId];
-			const FString ReplacementStr = FString::Format(TEXT("{0}_GetDataIndexFromIdInternal(/*DataId=*/{1})"),
+			const FString ReplacementStr = FString::Format(TEXT("{0}_GetDataIndexFromIdInternal(/*DataId=*/{1}u)"),
 			{
 				PCGComputeHelpers::GetDataLabelResolverName(PinLabel),
 				FString::FromInt(DataId)
