[PCG] Fixed custom seed accessor in the Attribute noise node

#rnx

diff --git a/Engine/Plugins/PCG/Source/PCG/Private/Elements/PCGAttributeNoise.cpp b/Engine/Plugins/PCG/Source/PCG/Private/Elements/PCGAttributeNoise.cpp
--- a/Engine/Plugins/PCG/Source/PCG/Private/Elements/PCGAttributeNoise.cpp
+++ b/Engine/Plugins/PCG/Source/PCG/Private/Elements/PCGAttributeNoise.cpp
@@ -287,7 +287,7 @@ bool FPCGAttributeNoiseElement::ExecuteInternal(FPCGContext* InContext) const
 			{
 				if (Settings->bHasCustomSeedSource)
 				{
-					Context->CustomSeedSource = Settings->InputSource.CopyAndFixLast(InputData);
+					Context->CustomSeedSource = Settings->CustomSeedSource.CopyAndFixLast(InputData);
 				}
 				else
 				{
