[PCG] Give licensees a FPCGContext blueprint scope (Preview1 ensure fix)

#rnx

diff --git a/Engine/Plugins/PCG/Source/PCG/Private/Elements/PCGExecuteBlueprint.cpp b/Engine/Plugins/PCG/Source/PCG/Private/Elements/PCGExecuteBlueprint.cpp
--- a/Engine/Plugins/PCG/Source/PCG/Private/Elements/PCGExecuteBlueprint.cpp
+++ b/Engine/Plugins/PCG/Source/PCG/Private/Elements/PCGExecuteBlueprint.cpp
@@ -31,6 +31,30 @@ namespace PCGBlueprintConstants
 	constexpr int32 RunawayResetFrequency = 1024;
 }
 
+FPCGContextBlueprintScope::FPCGContextBlueprintScope(FPCGContext* InContext)
+	: Context(InContext)
+{
+	if (Context)
+	{
+		// Make sure handle is not copied when calling ExecuteWithContext which creates a FPCGContext copy on the stack
+		ContextHandle = Context->Handle;
+		Context->Handle = nullptr;
+
+		// Similarly we don't want to register to GC-protection on release.
+		GraphExecutor = Context->GraphExecutor;
+		Context->GraphExecutor = nullptr;
+	}
+}
+
+FPCGContextBlueprintScope::~FPCGContextBlueprintScope()
+{
+	if (Context)
+	{
+		Context->Handle = ContextHandle;
+		Context->GraphExecutor = GraphExecutor;
+	}
+}
+
 #if WITH_EDITOR
 namespace PCGBlueprintHelper
 {
@@ -811,29 +835,22 @@ bool FPCGExecuteBlueprintElement::ExecuteInternal(FPCGContext* InContext) const
 		/** Finally, execute the actual blueprint */
 		Context->BlueprintElementInstance->SetCurrentContext(Context);
 		
-		// Make sure handle is not copied when calling ExecuteWithContext which creates a FPCGContext copy on the stack
-		TSharedPtr<FPCGContextHandle> Handle = Context->Handle;
-		Context->Handle = nullptr;
-
-		// Similarly we don't want to register to GC-protection on release.
-		TWeakPtr<FPCGGraphExecutor> Executor = Context->GraphExecutor;
-		Context->GraphExecutor = nullptr;
-
-		// When running outside of main thread make sure GC can't run (BP Nodes might create objects which can't happen while GC runs)
-		if(!IsInGameThread())
-		{
-			ensure(!Context->AsyncState.bIsRunningOnMainThread);
-			FGCScopeGuard Scope;
-			Context->BlueprintElementInstance->ExecuteWithContext(*Context, Context->InputData, Context->OutputData);
-		}
-		else
 		{
-			Context->BlueprintElementInstance->ExecuteWithContext(*Context, Context->InputData, Context->OutputData);
+			FPCGContextBlueprintScope BlueprintScope(Context);
+
+			// When running outside of main thread make sure GC can't run (BP Nodes might create objects which can't happen while GC runs)
+			if (!IsInGameThread())
+			{
+				ensure(!Context->AsyncState.bIsRunningOnMainThread);
+				FGCScopeGuard Scope;
+				Context->BlueprintElementInstance->ExecuteWithContext(*Context, Context->InputData, Context->OutputData);
+			}
+			else
+			{
+				Context->BlueprintElementInstance->ExecuteWithContext(*Context, Context->InputData, Context->OutputData);
+			}
 		}
 
-		// Restore handle & executor
-		Context->GraphExecutor = Executor;
-		Context->Handle = Handle;
 		Context->BlueprintElementInstance->SetCurrentContext(nullptr);
 	}
 	else
diff --git a/Engine/Plugins/PCG/Source/PCG/Public/Elements/PCGExecuteBlueprint.h b/Engine/Plugins/PCG/Source/PCG/Public/Elements/PCGExecuteBlueprint.h
--- a/Engine/Plugins/PCG/Source/PCG/Public/Elements/PCGExecuteBlueprint.h
+++ b/Engine/Plugins/PCG/Source/PCG/Public/Elements/PCGExecuteBlueprint.h
@@ -28,6 +28,23 @@ namespace PCGBlueprintHelper
 }
 #endif // WITH_EDITOR
 
+/**
+ * Helper class which can be used before calling a Blueprint function with a FPCGContext& parameter.
+ * 
+ * Since BP calls will copy the incoming FPCGContext reference we need to clear out a couple of fields before the BP copy is made to prevent issues when the copy is destroyed.
+ */
+class FPCGContextBlueprintScope
+{
+public:
+	UE_API explicit FPCGContextBlueprintScope(FPCGContext* InContext);
+	UE_API ~FPCGContextBlueprintScope();
+
+private:
+	FPCGContext* Context = nullptr;
+	TSharedPtr<FPCGContextHandle> ContextHandle = nullptr;
+	TWeakPtr<FPCGGraphExecutor> GraphExecutor = nullptr;
+};
+
 UCLASS(MinimalAPI, Abstract, BlueprintType, Blueprintable, hidecategories = (Object))
 class UPCGBlueprintElement : public UObject
 {
diff --git a/Engine/Plugins/PCG/Source/PCG/Public/PCGContext.h b/Engine/Plugins/PCG/Source/PCG/Public/PCGContext.h
--- a/Engine/Plugins/PCG/Source/PCG/Public/PCGContext.h
+++ b/Engine/Plugins/PCG/Source/PCG/Public/PCGContext.h
@@ -79,7 +79,7 @@ struct FPCGContext
 {
 	GENERATED_BODY()
 
-	friend class FPCGExecuteBlueprintElement;
+	friend class FPCGContextBlueprintScope;
 	friend class FPCGGraphExecutor;
 	friend struct FPCGGraphActiveTask;
 	friend struct FPCGGraphTask;
