[PCG] Allow CleanSpline to work on closed loops

* Before it was not taking into account the extra segment in a closed loop.
* Also fix inconsistent logic between tooltips, comment and code.
* Added unit tests.

#rnx

diff --git a/Engine/Plugins/PCG/Source/PCG/Private/Elements/PCGCleanSpline.cpp b/Engine/Plugins/PCG/Source/PCG/Private/Elements/PCGCleanSpline.cpp
--- a/Engine/Plugins/PCG/Source/PCG/Private/Elements/PCGCleanSpline.cpp
+++ b/Engine/Plugins/PCG/Source/PCG/Private/Elements/PCGCleanSpline.cpp
@@ -114,6 +114,7 @@ bool FPCGCleanSplineElement::ExecuteInternal(FPCGContext* InContext) const
 		const TConstArrayView<PCGMetadataEntryKey> ControlPointsMetadataEntry = InputSplineData->SplineStruct.GetConstControlPointsEntryKeys();
 		const int32 NumControlPoints = ControlPointsPosition.Points.Num();
 		bool bControlPointWasRemoved = false;
+		const bool bIsClosed = InputSplineData->IsClosed();
 
 		TArray<FSplinePoint> ControlPoints;
 		ControlPoints.Reserve(NumControlPoints);
@@ -153,49 +154,70 @@ bool FPCGCleanSplineElement::ExecuteInternal(FPCGContext* InContext) const
 			bControlPointWasRemoved = true;
 		};
 
+		auto GetPreviousIndexFromOffset = [&ControlPoints](int32 Index, int32 Offset)
+		{
+			if (!ensure(!ControlPoints.IsEmpty()))
+			{
+				return 0;
+			}
+			
+			int32 Result = Index - Offset;
+			while (Result < 0)
+			{
+				Result += ControlPoints.Num();
+			}
+			
+			return Result;
+		};
+
 		if (Settings->bFuseColocatedControlPoints)
 		{
+			const int32 MinIndex = bIsClosed ? 0 : 1;
+			
 			// Will evaluate by pairs. Reverse order for optimizing RemoveAt.
-			for (int Index = ControlPoints.Num() - 1; Index > 0; --Index)
+			for (int Index = ControlPoints.Num() - 1; Index >= MinIndex; --Index)
 			{
+				const int32 PreviousPointIndex = GetPreviousIndexFromOffset(Index, 1);
+				const int32 CurrentPointIndex = Index;
+				
 				// Evaluate by pairs for colocated.
 				if (PCGCleanSplineHelpers::ControlPointsAreColocated(
-						ControlPoints[Index - 1],
-						ControlPoints[Index],
+						ControlPoints[PreviousPointIndex],
+						ControlPoints[CurrentPointIndex],
 						SplineTransform,
 						Settings->ColocationDistanceThreshold,
 						Settings->bUseSplineLocalSpace))
 				{
 					EPCGControlPointFuseMode FuseMode = Settings->FuseMode;
-					// Generally, keep the first control point, but preserve the final control point to maintain the spline's length
+					// Generally, keep the previous control point (first), but preserve the final control point to maintain the spline's length (second) in non-closed splines.
 					if (FuseMode == EPCGControlPointFuseMode::Auto)
 					{
-						FuseMode = (Index == ControlPoints.Num() - 1) ? EPCGControlPointFuseMode::KeepFirst : FuseMode = EPCGControlPointFuseMode::KeepSecond;
+						FuseMode = (!bIsClosed && CurrentPointIndex == ControlPoints.Num() - 1) ? EPCGControlPointFuseMode::KeepSecond : EPCGControlPointFuseMode::KeepFirst;
 					}
 
 					switch (FuseMode)
 					{
 						// Keep the first control point
 						case EPCGControlPointFuseMode::KeepFirst:
-							ControlPoints[Index - 1].LeaveTangent = ControlPoints[Index].LeaveTangent;
-							RemovePoint(Index);
+							ControlPoints[PreviousPointIndex].LeaveTangent = ControlPoints[CurrentPointIndex].LeaveTangent;
+							RemovePoint(CurrentPointIndex);
 							break;
 						// Keep the second control point
 						case EPCGControlPointFuseMode::KeepSecond:
-							ControlPoints[Index].ArriveTangent = ControlPoints[Index - 1].ArriveTangent;
-							RemovePoint(Index - 1);
+							ControlPoints[CurrentPointIndex].ArriveTangent = ControlPoints[PreviousPointIndex].ArriveTangent;
+							RemovePoint(PreviousPointIndex);
 							break;
 						// Average the two control points' transforms and update the leave tangent
 						case EPCGControlPointFuseMode::Merge:
-							ControlPoints[Index - 1].Position = FMath::Lerp(ControlPoints[Index - 1].Position, ControlPoints[Index].Position, 0.5);
-							ControlPoints[Index - 1].Rotation = FMath::Lerp(ControlPoints[Index - 1].Rotation, ControlPoints[Index].Rotation, 0.5);
-							ControlPoints[Index - 1].Scale = FMath::Lerp(ControlPoints[Index - 1].Scale, ControlPoints[Index].Scale, 0.5);
-							ControlPoints[Index - 1].LeaveTangent = ControlPoints[Index].LeaveTangent;
-							RemovePoint(Index);
+							ControlPoints[PreviousPointIndex].Position = FMath::Lerp(ControlPoints[PreviousPointIndex].Position, ControlPoints[CurrentPointIndex].Position, 0.5);
+							ControlPoints[PreviousPointIndex].Rotation = FMath::Lerp(ControlPoints[PreviousPointIndex].Rotation, ControlPoints[CurrentPointIndex].Rotation, 0.5);
+							ControlPoints[PreviousPointIndex].Scale = FMath::Lerp(ControlPoints[PreviousPointIndex].Scale, ControlPoints[CurrentPointIndex].Scale, 0.5);
+							ControlPoints[PreviousPointIndex].LeaveTangent = ControlPoints[CurrentPointIndex].LeaveTangent;
+							RemovePoint(CurrentPointIndex);
 							break;
 						case EPCGControlPointFuseMode::Auto: // Should've been picked up by now. Fallthrough...
-						default:
 							checkNoEntry();
+						default:
 							break;
 					}
 				}
@@ -204,18 +226,24 @@ bool FPCGCleanSplineElement::ExecuteInternal(FPCGContext* InContext) const
 
 		if (Settings->bRemoveCollinearControlPoints)
 		{
+			const int32 MinIndex = bIsClosed ? 0 : 2;
+			
 			// Will evaluate by triplets. Reverse order for optimizing RemoveAt.
-			for (int Index = ControlPoints.Num() - 1; Index > 1; --Index)
+			for (int Index = ControlPoints.Num() - 1; Index >= MinIndex; --Index)
 			{
+				const int32 SecondPreviousPointIndex = GetPreviousIndexFromOffset(Index, 2);
+				const int32 PreviousPointIndex = GetPreviousIndexFromOffset(Index, 1);
+				const int32 CurrentPointIndex = Index;
+				
 				if (PCGCleanSplineHelpers::ControlPointsAreCollinear(
-						ControlPoints[Index - 2],
-						ControlPoints[Index - 1],
-						ControlPoints[Index],
+						ControlPoints[SecondPreviousPointIndex],
+						ControlPoints[PreviousPointIndex],
+						ControlPoints[CurrentPointIndex],
 						DotProductTolerance))
 				{
-					ControlPoints[Index - 2].LeaveTangent = ControlPoints[Index].Position - ControlPoints[Index - 2].Position;
-					ControlPoints[Index].ArriveTangent = ControlPoints[Index - 2].LeaveTangent;
-					RemovePoint(Index - 1);
+					ControlPoints[SecondPreviousPointIndex].LeaveTangent = ControlPoints[CurrentPointIndex].Position - ControlPoints[SecondPreviousPointIndex].Position;
+					ControlPoints[CurrentPointIndex].ArriveTangent = ControlPoints[SecondPreviousPointIndex].LeaveTangent;
+					RemovePoint(PreviousPointIndex);
 				}
 			}
 		}
diff --git a/Engine/Plugins/PCG/Source/PCG/Private/Tests/Elements/PCGCleanSplineTests.cpp b/Engine/Plugins/PCG/Source/PCG/Private/Tests/Elements/PCGCleanSplineTests.cpp
new file mode 100644
--- /dev/null
+++ b/Engine/Plugins/PCG/Source/PCG/Private/Tests/Elements/PCGCleanSplineTests.cpp
@@ -0,0 +1,484 @@
+// Copyright Epic Games, Inc. All Rights Reserved.
+
+
+#include "PCGContext.h"
+#if WITH_EDITOR
+
+#include "Tests/PCGTestsCommon.h"
+
+#include "Data/PCGSplineData.h"
+#include "Elements/PCGCleanSpline.h"
+
+#include "UObject/Package.h"
+
+class FPCGCleanSplineTest : public FPCGTestBaseClass
+{
+	using FPCGTestBaseClass::FPCGTestBaseClass;
+	
+protected:
+	bool RunTestInternal(UPCGCleanSplineSettings* Settings, const TArray<FVector>& Points, bool bLinear, bool bIsClosed, TArray<int32> ExpectedOutputPoints, const TArray<FVector>* ExpectedPositions = nullptr)
+	{
+		PCGTestsCommon::FTestData TestData;
+		TestData.Reset(Settings);
+
+		// Add dummy metadata entries to keep track of which control point was removed.
+		TArray<PCGMetadataEntryKey> EntryKeys;
+		EntryKeys.Reserve(Points.Num());
+		for (int64 i = 0; i < Points.Num(); ++i)
+		{
+			EntryKeys.Emplace(i);
+		}
+		
+		UPCGSplineData* SplineData = NewObject<UPCGSplineData>(GetTransientPackage(), NAME_None, RF_Transient);
+
+		TArray<FSplinePoint> SplinePoints;
+		SplinePoints.Reserve(Points.Num());
+		for (int32 i = 0; i < Points.Num(); ++i)
+		{
+			SplinePoints.Emplace_GetRef(i, Points[i]).Type = bLinear ? ESplinePointType::Linear : ESplinePointType::Curve;
+		}
+
+		SplineData->Initialize(SplinePoints, bIsClosed, FTransform::Identity, MoveTemp(EntryKeys));
+
+		FPCGTaggedData& Inputs = TestData.InputData.TaggedData.Emplace_GetRef();
+		Inputs.Pin = PCGPinConstants::DefaultInputLabel;
+		Inputs.Data = SplineData;
+	
+		FPCGElementPtr TestElement = TestData.Settings->GetElement();
+
+		TUniquePtr<FPCGContext> Context = TestData.InitializeTestContext();
+
+		while (!TestElement->Execute(Context.Get())) {}
+
+		UTEST_EQUAL("Output has 1 data", Context->OutputData.TaggedData.Num(), 1);
+
+		const UPCGSplineData* OutputSplineData = Cast<const UPCGSplineData>(Context->OutputData.TaggedData[0].Data);
+		UTEST_NOT_NULL("Output is a spline", OutputSplineData);
+
+		UTEST_EQUAL(*FString::Printf(TEXT("Output spline has %d points"), ExpectedOutputPoints.Num()), OutputSplineData->SplineStruct.GetNumberOfPoints(), ExpectedOutputPoints.Num());
+
+		// With Auto, we keep the first point by default, except for the last point. So in this case, the metadata entries should be 0 and 3.
+		UTEST_EQUAL(*FString::Printf(TEXT("Output spline has %d metadata entries"), ExpectedOutputPoints.Num()), OutputSplineData->GetConstVerticesEntryKeys().Num(), ExpectedOutputPoints.Num())
+
+		for (int32 i = 0; i < ExpectedOutputPoints.Num(); ++i)
+		{
+			const int64 ExpectedIndex = ExpectedOutputPoints[i];
+			UTEST_EQUAL(*FString::Printf(TEXT("Point nº%d has was the original point nº%lld"), i, ExpectedIndex), OutputSplineData->GetConstVerticesEntryKeys()[i], ExpectedIndex);
+		}
+
+		if (ExpectedPositions)
+		{
+			for (int32 i = 0; i < ExpectedPositions->Num(); ++i)
+			{
+				UTEST_EQUAL(*FString::Printf(TEXT("Point nº%d is at the right position"), i), OutputSplineData->SplineStruct.GetLocationAtSplineInputKey(i, ESplineCoordinateSpace::World), (*ExpectedPositions)[i]);
+			}
+		}
+		
+		return true;
+	}
+};
+
+IMPLEMENT_CUSTOM_SIMPLE_AUTOMATION_TEST(FPCGCleanSplineTest_CollocatedControlPoints_Auto, FPCGCleanSplineTest, "Plugins.PCG.CleanSpline.CollocatedControlPoints.Auto", PCGTestsCommon::TestFlags)
+IMPLEMENT_CUSTOM_SIMPLE_AUTOMATION_TEST(FPCGCleanSplineTest_CollocatedControlPoints_KeepFirst, FPCGCleanSplineTest, "Plugins.PCG.CleanSpline.CollocatedControlPoints.KeepFirst", PCGTestsCommon::TestFlags)
+IMPLEMENT_CUSTOM_SIMPLE_AUTOMATION_TEST(FPCGCleanSplineTest_CollocatedControlPoints_KeepSecond, FPCGCleanSplineTest, "Plugins.PCG.CleanSpline.CollocatedControlPoints.KeepSecond", PCGTestsCommon::TestFlags)
+IMPLEMENT_CUSTOM_SIMPLE_AUTOMATION_TEST(FPCGCleanSplineTest_CollocatedControlPoints_Merge, FPCGCleanSplineTest, "Plugins.PCG.CleanSpline.CollocatedControlPoints.Merge", PCGTestsCommon::TestFlags)
+IMPLEMENT_CUSTOM_SIMPLE_AUTOMATION_TEST(FPCGCleanSplineTest_CollocatedControlPoints_Closed_Auto, FPCGCleanSplineTest, "Plugins.PCG.CleanSpline.CollocatedControlPoints.Closed_Auto", PCGTestsCommon::TestFlags)
+IMPLEMENT_CUSTOM_SIMPLE_AUTOMATION_TEST(FPCGCleanSplineTest_CollocatedControlPoints_Closed_KeepFirst, FPCGCleanSplineTest, "Plugins.PCG.CleanSpline.CollocatedControlPoints.Closed_KeepFirst", PCGTestsCommon::TestFlags)
+IMPLEMENT_CUSTOM_SIMPLE_AUTOMATION_TEST(FPCGCleanSplineTest_CollocatedControlPoints_Closed_KeepSecond, FPCGCleanSplineTest, "Plugins.PCG.CleanSpline.CollocatedControlPoints.Closed_KeepSecond", PCGTestsCommon::TestFlags)
+IMPLEMENT_CUSTOM_SIMPLE_AUTOMATION_TEST(FPCGCleanSplineTest_CollocatedControlPoints_Closed_Merge, FPCGCleanSplineTest, "Plugins.PCG.CleanSpline.CollocatedControlPoints.Closed_Merge", PCGTestsCommon::TestFlags)
+
+IMPLEMENT_CUSTOM_SIMPLE_AUTOMATION_TEST(FPCGCleanSplineTest_CollinearControlPoints_Linear, FPCGCleanSplineTest, "Plugins.PCG.CleanSpline.CollinearControlPoints.Linear", PCGTestsCommon::TestFlags)
+IMPLEMENT_CUSTOM_SIMPLE_AUTOMATION_TEST(FPCGCleanSplineTest_CollinearControlPoints_Linear_Closed, FPCGCleanSplineTest, "Plugins.PCG.CleanSpline.CollinearControlPoints.Linear_Closed", PCGTestsCommon::TestFlags)
+IMPLEMENT_CUSTOM_SIMPLE_AUTOMATION_TEST(FPCGCleanSplineTest_CollinearControlPoints_Curve, FPCGCleanSplineTest, "Plugins.PCG.CleanSpline.CollinearControlPoints.Curve", PCGTestsCommon::TestFlags)
+IMPLEMENT_CUSTOM_SIMPLE_AUTOMATION_TEST(FPCGCleanSplineTest_CollinearControlPoints_Curve_Closed, FPCGCleanSplineTest, "Plugins.PCG.CleanSpline.CollinearControlPoints.Curve_Closed", PCGTestsCommon::TestFlags)
+
+
+bool FPCGCleanSplineTest_CollocatedControlPoints_Auto::RunTest(const FString& Parameters)
+{
+	TArray<FVector> Points =
+	{
+		{0, 0, 0},
+		{0, 0, 0},
+		{100, 100, 100},
+		{100, 100, 100},
+	};
+
+	bool bIsClosed = false;
+	bool bLinear = true;
+
+	UPCGCleanSplineSettings* Settings = NewObject<UPCGCleanSplineSettings>(GetTransientPackage(), NAME_None, RF_Transient);
+	Settings->bFuseColocatedControlPoints = true;
+	Settings->bRemoveCollinearControlPoints = false;
+	Settings->FuseMode = EPCGControlPointFuseMode::Auto;
+	Settings->ColocationDistanceThreshold = 1.0;
+
+	TArray<int32> ExpectedOutputPoints = {0, 3};
+
+	return RunTestInternal(Settings, Points, bLinear, bIsClosed, MoveTemp(ExpectedOutputPoints));
+}
+
+bool FPCGCleanSplineTest_CollocatedControlPoints_KeepFirst::RunTest(const FString& Parameters)
+{
+	TArray<FVector> Points =
+	{
+		{0, 0, 0},
+		{0, 0, 0},
+		{100, 100, 100},
+		{100, 100, 100},
+	};
+
+	bool bIsClosed = false;
+	bool bLinear = true;
+
+	UPCGCleanSplineSettings* Settings = NewObject<UPCGCleanSplineSettings>(GetTransientPackage(), NAME_None, RF_Transient);
+	Settings->bFuseColocatedControlPoints = true;
+	Settings->bRemoveCollinearControlPoints = false;
+	Settings->FuseMode = EPCGControlPointFuseMode::KeepFirst;
+	Settings->ColocationDistanceThreshold = 1.0;
+
+	TArray<int32> ExpectedOutputPoints = {0, 2};
+
+	return RunTestInternal(Settings, Points, bLinear, bIsClosed, MoveTemp(ExpectedOutputPoints));
+}
+
+bool FPCGCleanSplineTest_CollocatedControlPoints_KeepSecond::RunTest(const FString& Parameters)
+{
+	TArray<FVector> Points =
+	{
+		{0, 0, 0},
+		{0, 0, 0},
+		{100, 100, 100},
+		{100, 100, 100},
+	};
+
+	bool bIsClosed = false;
+	bool bLinear = true;
+
+	UPCGCleanSplineSettings* Settings = NewObject<UPCGCleanSplineSettings>(GetTransientPackage(), NAME_None, RF_Transient);
+	Settings->bFuseColocatedControlPoints = true;
+	Settings->bRemoveCollinearControlPoints = false;
+	Settings->FuseMode = EPCGControlPointFuseMode::KeepSecond;
+	Settings->ColocationDistanceThreshold = 1.0;
+
+	TArray<int32> ExpectedOutputPoints = {1, 3};
+
+	return RunTestInternal(Settings, Points, bLinear, bIsClosed, MoveTemp(ExpectedOutputPoints));
+}
+
+bool FPCGCleanSplineTest_CollocatedControlPoints_Merge::RunTest(const FString& Parameters)
+{
+	TArray<FVector> Points =
+	{
+		{0, 0, 0},
+		{20, 0, 0},
+		{100, 0, 0},
+		{120, 0, 0},
+	};
+
+	bool bIsClosed = false;
+	bool bLinear = true;
+
+	UPCGCleanSplineSettings* Settings = NewObject<UPCGCleanSplineSettings>(GetTransientPackage(), NAME_None, RF_Transient);
+	Settings->bFuseColocatedControlPoints = true;
+	Settings->bRemoveCollinearControlPoints = false;
+	Settings->FuseMode = EPCGControlPointFuseMode::Merge;
+	Settings->ColocationDistanceThreshold = 21.0;
+
+	TArray<int32> ExpectedOutputPoints = {0, 2};
+	TArray<FVector> ExpectedPositions =
+	{
+		{10, 0, 0},
+		{110, 0, 0}
+	};
+
+	return RunTestInternal(Settings, Points, bLinear, bIsClosed, MoveTemp(ExpectedOutputPoints), &ExpectedPositions);
+}
+
+bool FPCGCleanSplineTest_CollocatedControlPoints_Closed_Auto::RunTest(const FString& Parameters)
+{
+	TArray<FVector> Points =
+	{
+		{0, 0, 0},
+		{0, 0, 0},
+		{100, 100, 100},
+		{100, 100, 100},
+		{0, 0, 0},
+	};
+
+	bool bIsClosed = true;
+	bool bLinear = true;
+
+	UPCGCleanSplineSettings* Settings = NewObject<UPCGCleanSplineSettings>(GetTransientPackage(), NAME_None, RF_Transient);
+	Settings->bFuseColocatedControlPoints = true;
+	Settings->bRemoveCollinearControlPoints = false;
+	Settings->FuseMode = EPCGControlPointFuseMode::Auto;
+	Settings->ColocationDistanceThreshold = 1.0;
+
+	TArray<int32> ExpectedOutputPoints = {2, 4};
+
+	return RunTestInternal(Settings, Points, bLinear, bIsClosed, MoveTemp(ExpectedOutputPoints));
+}
+
+bool FPCGCleanSplineTest_CollocatedControlPoints_Closed_KeepFirst::RunTest(const FString& Parameters)
+{
+	TArray<FVector> Points =
+	{
+		{0, 0, 0},
+		{0, 0, 0},
+		{100, 100, 100},
+		{100, 100, 100},
+		{0, 0, 0},
+	};
+
+	bool bIsClosed = true;
+	bool bLinear = true;
+
+	UPCGCleanSplineSettings* Settings = NewObject<UPCGCleanSplineSettings>(GetTransientPackage(), NAME_None, RF_Transient);
+	Settings->bFuseColocatedControlPoints = true;
+	Settings->bRemoveCollinearControlPoints = false;
+	Settings->FuseMode = EPCGControlPointFuseMode::KeepFirst;
+	Settings->ColocationDistanceThreshold = 1.0;
+
+	TArray<int32> ExpectedOutputPoints = {2, 4};
+
+	return RunTestInternal(Settings, Points, bLinear, bIsClosed, MoveTemp(ExpectedOutputPoints));
+}
+
+bool FPCGCleanSplineTest_CollocatedControlPoints_Closed_KeepSecond::RunTest(const FString& Parameters)
+{
+	TArray<FVector> Points =
+	{
+		{0, 0, 0},
+		{0, 0, 0},
+		{100, 100, 100},
+		{100, 100, 100},
+		{0, 0, 0},
+	};
+
+	bool bIsClosed = true;
+	bool bLinear = true;
+
+	UPCGCleanSplineSettings* Settings = NewObject<UPCGCleanSplineSettings>(GetTransientPackage(), NAME_None, RF_Transient);
+	Settings->bFuseColocatedControlPoints = true;
+	Settings->bRemoveCollinearControlPoints = false;
+	Settings->FuseMode = EPCGControlPointFuseMode::KeepSecond;
+	Settings->ColocationDistanceThreshold = 1.0;
+
+	TArray<int32> ExpectedOutputPoints = {1, 3};
+
+	return RunTestInternal(Settings, Points, bLinear, bIsClosed, MoveTemp(ExpectedOutputPoints));
+}
+
+bool FPCGCleanSplineTest_CollocatedControlPoints_Closed_Merge::RunTest(const FString& Parameters)
+{
+	TArray<FVector> Points =
+	{
+		{0, 0, 0},
+		{10, 0, 0},
+		{100, 0, 0},
+		{120, 0, 0},
+		{-10, 0, 0},
+	};
+
+	bool bIsClosed = true;
+	bool bLinear = true;
+
+	UPCGCleanSplineSettings* Settings = NewObject<UPCGCleanSplineSettings>(GetTransientPackage(), NAME_None, RF_Transient);
+	Settings->bFuseColocatedControlPoints = true;
+	Settings->bRemoveCollinearControlPoints = false;
+	Settings->FuseMode = EPCGControlPointFuseMode::Merge;
+	Settings->ColocationDistanceThreshold = 21.0;
+
+	TArray<int32> ExpectedOutputPoints = {2, 4};
+	TArray<FVector> ExpectedPositions =
+	{
+		{110, 0, 0},
+		{-2.5, 0, 0}
+	};
+
+	return RunTestInternal(Settings, Points, bLinear, bIsClosed, MoveTemp(ExpectedOutputPoints), &ExpectedPositions);
+}
+
+bool FPCGCleanSplineTest_CollinearControlPoints_Linear::RunTest(const FString& Parameters)
+{
+	/**
+	 *          4
+	 *          |
+	 *          3
+	 *        2
+	 *      1
+	 *    0
+	 *
+	 * becomes
+	 *          4
+	 *          |
+	 *          3
+	 *        / 
+	 *      / 
+	 *    0
+	 */
+	
+	TArray<FVector> Points =
+	{
+		{0, 0, 0},
+		{100, 100, 100},
+		{200, 200, 200},
+		{300, 300, 300},
+		{300, 500, 500}
+	};
+
+	bool bIsClosed = false;
+	bool bLinear = true;
+
+	UPCGCleanSplineSettings* Settings = NewObject<UPCGCleanSplineSettings>(GetTransientPackage(), NAME_None, RF_Transient);
+	Settings->bFuseColocatedControlPoints = false;
+	Settings->bRemoveCollinearControlPoints = true;
+
+	TArray<int32> ExpectedOutputPoints = {0, 3, 4};
+
+	return RunTestInternal(Settings, Points, bLinear, bIsClosed, MoveTemp(ExpectedOutputPoints));
+}
+
+bool FPCGCleanSplineTest_CollinearControlPoints_Linear_Closed::RunTest(const FString& Parameters)
+{
+	/**
+	 *   5----------- 4
+	 *   |            |
+	 *   |            3
+	 *   |          2
+	 *   |        1
+	 *   |      0
+	 *   |    /
+	 *   |  7
+	 *   6
+	 *
+	 * becomes
+	 *   5----------- 4
+	 *   |            |
+	 *   |            3
+	 *   |          /
+	 *   |        /
+	 *   |      /
+	 *   |    /
+	 *   |  /
+	 *   6
+	 */
+	
+	TArray<FVector> Points =
+	{
+		{0, 0, 0},
+		{100, 100, 0},
+		{200, 200, 0},
+		{300, 300, 0},
+		{300, 500, 0},
+		{-300, 500, 0},
+		{-300, -300, 0},
+		{-200, -200, 0},
+	};
+
+	bool bIsClosed = true;
+	bool bLinear = true;
+
+	UPCGCleanSplineSettings* Settings = NewObject<UPCGCleanSplineSettings>(GetTransientPackage(), NAME_None, RF_Transient);
+	Settings->bFuseColocatedControlPoints = false;
+	Settings->bRemoveCollinearControlPoints = true;
+
+	TArray<int32> ExpectedOutputPoints = {3, 4, 5, 6};
+
+	return RunTestInternal(Settings, Points, bLinear, bIsClosed, MoveTemp(ExpectedOutputPoints));
+}
+
+bool FPCGCleanSplineTest_CollinearControlPoints_Curve::RunTest(const FString& Parameters)
+{
+	/**
+	 *	 With curves, tangent in 3 is different, so 1,2,3 aren't collinear anymore.
+	 *	 
+	 *          4
+	 *          |
+	 *          3
+	 *        2
+	 *      1
+	 *    0
+	 *
+	 * becomes
+	 *          4
+	 *          |
+	 *          3
+	 *        2 
+	 *      / 
+	 *    0
+	 */
+	
+	TArray<FVector> Points =
+	{
+		{0, 0, 0},
+		{100, 100, 0},
+		{200, 200, 0},
+		{300, 300, 0},
+		{300, 500, 0}
+	};
+
+	bool bIsClosed = false;
+	bool bLinear = false;
+
+	UPCGCleanSplineSettings* Settings = NewObject<UPCGCleanSplineSettings>(GetTransientPackage(), NAME_None, RF_Transient);
+	Settings->bFuseColocatedControlPoints = false;
+	Settings->bRemoveCollinearControlPoints = true;
+
+	TArray<int32> ExpectedOutputPoints = {0, 2, 3, 4};
+
+	return RunTestInternal(Settings, Points, bLinear, bIsClosed, MoveTemp(ExpectedOutputPoints));
+}
+
+bool FPCGCleanSplineTest_CollinearControlPoints_Curve_Closed::RunTest(const FString& Parameters)
+{
+	/**
+	 *   With curves, tangents in 3 and 6 are different, so 1,2,3 and 6,7,0 aren't collinear anymore.
+	 *   
+	 *   5----------- 4
+	 *   |            |
+	 *   |            3
+	 *   |          2
+	 *   |        1
+	 *   |      0
+	 *   |    /
+	 *   |  7
+	 *   6
+	 *
+	 * becomes
+	 *   5----------- 4
+	 *   |            |
+	 *   |            3
+	 *   |          2
+	 *   |        /
+	 *   |      /
+	 *   |    /
+	 *   |  7
+	 *   6
+	 */
+	
+	TArray<FVector> Points =
+	{
+		{0, 0, 0},
+		{100, 100, 0},
+		{200, 200, 0},
+		{300, 300, 0},
+		{300, 500, 0},
+		{-300, 500, 0},
+		{-300, -300, 0},
+		{-200, -200, 0},
+	};
+
+	bool bIsClosed = true;
+	bool bLinear = false;
+
+	UPCGCleanSplineSettings* Settings = NewObject<UPCGCleanSplineSettings>(GetTransientPackage(), NAME_None, RF_Transient);
+	Settings->bFuseColocatedControlPoints = false;
+	Settings->bRemoveCollinearControlPoints = true;
+
+	TArray<int32> ExpectedOutputPoints = {2, 3, 4, 5, 6, 7};
+
+	return RunTestInternal(Settings, Points, bLinear, bIsClosed, MoveTemp(ExpectedOutputPoints));
+}
+
+#endif // WITH_EDITOR
\ No newline at end of file
diff --git a/Engine/Plugins/PCG/Source/PCG/Public/Elements/PCGCleanSpline.h b/Engine/Plugins/PCG/Source/PCG/Public/Elements/PCGCleanSpline.h
--- a/Engine/Plugins/PCG/Source/PCG/Public/Elements/PCGCleanSpline.h
+++ b/Engine/Plugins/PCG/Source/PCG/Public/Elements/PCGCleanSpline.h
@@ -9,8 +9,8 @@
 UENUM()
 enum class EPCGControlPointFuseMode
 {
-	KeepFirst UMETA(Tooltip = "Keep the first of two control points."),
-	KeepSecond UMETA(Tooltip = "Keep the second of two control points."),
+	KeepFirst UMETA(Tooltip = "If index 0 and 1 are colocated, keep 0."),
+	KeepSecond UMETA(Tooltip = "If index 0 and 1 are colocated, keep 1."),
 	Merge UMETA(Tooltip = "Fuse the two control points into a new control point, located halfway between the original two."),
 	Auto UMETA(Tooltip = "Will generally keep the second of both control points, except keep the first when it would otherwise alter the length of the spline--ie. final control point.")
 };
