[PCG] Fix compile error due to missing engine define.

#rnx

diff --git a/Engine/Plugins/PCG/Source/PCG/Private/Elements/PCGGenerateGrassMaps.cpp b/Engine/Plugins/PCG/Source/PCG/Private/Elements/PCGGenerateGrassMaps.cpp
--- a/Engine/Plugins/PCG/Source/PCG/Private/Elements/PCGGenerateGrassMaps.cpp
+++ b/Engine/Plugins/PCG/Source/PCG/Private/Elements/PCGGenerateGrassMaps.cpp
@@ -13,6 +13,7 @@
 #include "Metadata/Accessors/PCGAttributeAccessorHelpers.h"
 
 #include "DrawDebugHelpers.h"
+#include "EngineDefines.h" // For UE_ENABLE_DEBUG_DRAWING
 #include "Landscape.h"
 #include "LandscapeComponent.h"
 #include "LandscapeGrassType.h"
diff --git a/Engine/Plugins/PCG/Source/PCG/Private/PCGCommon.cpp b/Engine/Plugins/PCG/Source/PCG/Private/PCGCommon.cpp
--- a/Engine/Plugins/PCG/Source/PCG/Private/PCGCommon.cpp
+++ b/Engine/Plugins/PCG/Source/PCG/Private/PCGCommon.cpp
@@ -5,6 +5,8 @@
 #include "PCGSubsystem.h"
 #include "PCGWorldActor.h"
 
+#include "EngineDefines.h" // For UE_ENABLE_DEBUG_DRAWING
+
 namespace PCGLog
 {
 	FString GetComponentOwnerName(const UPCGComponent* InComponent, bool bUseLabel, FString InDefaultName)
diff --git a/Engine/Plugins/PCG/Source/PCG/Private/PCGElement.cpp b/Engine/Plugins/PCG/Source/PCG/Private/PCGElement.cpp
--- a/Engine/Plugins/PCG/Source/PCG/Private/PCGElement.cpp
+++ b/Engine/Plugins/PCG/Source/PCG/Private/PCGElement.cpp
@@ -19,6 +19,7 @@
 #endif
 
 #include "DrawDebugHelpers.h"
+#include "EngineDefines.h" // For UE_ENABLE_DEBUG_DRAWING
 #include "HAL/IConsoleManager.h"
 #include "Utils/PCGExtraCapture.h"
 
diff --git a/Engine/Plugins/PCG/Source/PCG/Private/RuntimeGen/PCGRuntimeGenScheduler.cpp b/Engine/Plugins/PCG/Source/PCG/Private/RuntimeGen/PCGRuntimeGenScheduler.cpp
--- a/Engine/Plugins/PCG/Source/PCG/Private/RuntimeGen/PCGRuntimeGenScheduler.cpp
+++ b/Engine/Plugins/PCG/Source/PCG/Private/RuntimeGen/PCGRuntimeGenScheduler.cpp
@@ -19,6 +19,7 @@
 #include "RuntimeGen/PCGGenSourceManager.h"
 
 #include "DrawDebugHelpers.h"
+#include "EngineDefines.h" // For UE_ENABLE_DEBUG_DRAWING
 #include "Async/ParallelFor.h"
 #include "Engine/LevelStreaming.h"
 #include "Components/RuntimeVirtualTextureComponent.h"
diff --git a/Engine/Plugins/PCG/Source/PCG/Public/Helpers/PCGHelpers.h b/Engine/Plugins/PCG/Source/PCG/Public/Helpers/PCGHelpers.h
--- a/Engine/Plugins/PCG/Source/PCG/Public/Helpers/PCGHelpers.h
+++ b/Engine/Plugins/PCG/Source/PCG/Public/Helpers/PCGHelpers.h
@@ -5,6 +5,7 @@
 #include "PCGCommon.h"
 #include "PCGPoint.h"
 
+#include "EngineDefines.h" // For UE_ENABLE_DEBUG_DRAWING
 #include "Containers/Ticker.h"
 #include "Kismet/BlueprintFunctionLibrary.h"
 #include "Math/Box.h"
