[PCG] Forced default details view layout in user parameters panel, otherwise buttons (like the determinism test button) can show up.

#rnx

diff --git a/Engine/Plugins/PCG/Source/PCGEditor/Private/Widgets/SPCGEditorGraphParamsView.cpp b/Engine/Plugins/PCG/Source/PCGEditor/Private/Widgets/SPCGEditorGraphParamsView.cpp
--- a/Engine/Plugins/PCG/Source/PCGEditor/Private/Widgets/SPCGEditorGraphParamsView.cpp
+++ b/Engine/Plugins/PCG/Source/PCGEditor/Private/Widgets/SPCGEditorGraphParamsView.cpp
@@ -84,6 +84,11 @@ void SPCGEditorGraphUserParametersView::Construct(const FArguments& InArgs, cons
 
 	// Note: Single Property View Header notes that it doesn't work with arrays or structs, so using filters for now.
 	const TSharedRef<IDetailsView> DetailsView = PropertyEditorModule.CreateDetailView(DetailsViewArgs);
+
+	// Important: since this is using the graph object and that has details customization, we want to ensure we use the default layout instead.
+	// Otherwise we'll get buttons and more.
+	DetailsView->RegisterInstancedCustomPropertyLayout(UPCGGraph::StaticClass(), DetailsView->GetGenericLayoutDetailsDelegate());
+
 	DetailsView->SetObject(PCGGraph);
 
 	ChildSlot
