[PCG] Fix Merge element not properly allocating (SoA)
- Move Union helper into PCGBasePointData and rename GetPropertiesToAllocateFromPointDatas

diff --git a/Engine/Plugins/PCG/Source/PCG/Private/Data/PCGUnionData.cpp b/Engine/Plugins/PCG/Source/PCG/Private/Data/PCGUnionData.cpp
--- a/Engine/Plugins/PCG/Source/PCG/Private/Data/PCGUnionData.cpp
+++ b/Engine/Plugins/PCG/Source/PCG/Private/Data/PCGUnionData.cpp
@@ -39,85 +39,6 @@ namespace PCGUnionDataMaths
 	}
 }
 
-
-namespace PCGUnionHelper
-{
-	template<class T>
-	bool NeedToAllocateSingleValueProperty(EPCGPointNativeProperties InProperty, const TArray<const UPCGBasePointData*>& PointInputDatas)
-	{	
-		// Here we are comparing single values (non-allocated properties) to see if we have single values that differ between inputs.
-		// If they do differ, we need to allocate those properties because the output will have multiple values.
-		TOptional<const T> DefaultValue;
-		for (int32 i = 0; i < PointInputDatas.Num(); ++i)
-		{
-			const TConstPCGValueRange<T> ValueRange = PointInputDatas[i]->GetConstValueRange<T>(InProperty);
-			TOptional<const T> SingleValue = ValueRange.GetSingleValue();
-			if (!DefaultValue.IsSet())
-			{
-				DefaultValue = SingleValue;
-			}
-			else if (DefaultValue.IsSet() && SingleValue.IsSet() && !PCG::Private::MetadataTraits<T>::Equal(DefaultValue.GetValue(), SingleValue.GetValue()))
-			{
-				return true;
-			}
-		}
-
-		return false;
-	}
-
-	EPCGPointNativeProperties GetPropertiesToAllocateFromInputs(const TArray<const UPCGBasePointData*>& PointInputDatas)
-	{
-		// Start by doing a union of all input allocated properties
-		EPCGPointNativeProperties PropertiesToAllocate = EPCGPointNativeProperties::None;
-		for (const UPCGBasePointData* PointInputData : PointInputDatas)
-		{
-			PropertiesToAllocate |= PointInputData->GetAllocatedProperties();
-		}
-
-		if (!EnumHasAnyFlags(PropertiesToAllocate, EPCGPointNativeProperties::Transform) && NeedToAllocateSingleValueProperty<FTransform>(EPCGPointNativeProperties::Transform, PointInputDatas))
-		{
-			PropertiesToAllocate |= EPCGPointNativeProperties::Transform;
-		}
-
-		if (!EnumHasAnyFlags(PropertiesToAllocate, EPCGPointNativeProperties::Density) && NeedToAllocateSingleValueProperty<float>(EPCGPointNativeProperties::Density, PointInputDatas))
-		{
-			PropertiesToAllocate |= EPCGPointNativeProperties::Density;
-		}
-
-		if (!EnumHasAnyFlags(PropertiesToAllocate, EPCGPointNativeProperties::BoundsMin) && NeedToAllocateSingleValueProperty<FVector>(EPCGPointNativeProperties::BoundsMin, PointInputDatas))
-		{
-			PropertiesToAllocate |= EPCGPointNativeProperties::BoundsMin;
-		}
-
-		if (!EnumHasAnyFlags(PropertiesToAllocate, EPCGPointNativeProperties::BoundsMax) && NeedToAllocateSingleValueProperty<FVector>(EPCGPointNativeProperties::BoundsMax, PointInputDatas))
-		{
-			PropertiesToAllocate |= EPCGPointNativeProperties::BoundsMax;
-		}
-
-		if (!EnumHasAnyFlags(PropertiesToAllocate, EPCGPointNativeProperties::Color) && NeedToAllocateSingleValueProperty<FVector4>(EPCGPointNativeProperties::Color, PointInputDatas))
-		{
-			PropertiesToAllocate |= EPCGPointNativeProperties::Color;
-		}
-
-		if (!EnumHasAnyFlags(PropertiesToAllocate, EPCGPointNativeProperties::Steepness) && NeedToAllocateSingleValueProperty<float>(EPCGPointNativeProperties::Steepness, PointInputDatas))
-		{
-			PropertiesToAllocate |= EPCGPointNativeProperties::Steepness;
-		}
-
-		if (!EnumHasAnyFlags(PropertiesToAllocate, EPCGPointNativeProperties::Seed) && NeedToAllocateSingleValueProperty<int32>(EPCGPointNativeProperties::Seed, PointInputDatas))
-		{
-			PropertiesToAllocate |= EPCGPointNativeProperties::Seed;
-		}
-
-		if (!EnumHasAnyFlags(PropertiesToAllocate, EPCGPointNativeProperties::MetadataEntry) && NeedToAllocateSingleValueProperty<int64>(EPCGPointNativeProperties::MetadataEntry, PointInputDatas))
-		{
-			PropertiesToAllocate |= EPCGPointNativeProperties::MetadataEntry;
-		}
-
-		return PropertiesToAllocate;
-	}
-}
-
 void UPCGUnionData::Initialize(const UPCGSpatialData* InA, const UPCGSpatialData* InB)
 {
 	check(InA && InB);
@@ -385,7 +306,7 @@ const UPCGBasePointData* UPCGUnionData::CreateBasePointData(FPCGContext* Context
 				PointInputCount += PointInputDatas[Index]->GetNumPoints();
 			}
 
-			const EPCGPointNativeProperties PropertiesToAllocate = PCGUnionHelper::GetPropertiesToAllocateFromInputs(PointInputDatas);
+			const EPCGPointNativeProperties PropertiesToAllocate = UPCGBasePointData::GetPropertiesToAllocateFromPointData(PointInputDatas);
 
 			if (PointInputCount > 0)
 			{
@@ -491,7 +412,7 @@ void UPCGUnionData::CreateSequentialPointData(FPCGContext* Context, TArray<const
 	
 	OutPointData->SetNumPoints(PointInputCount, /*bInitializeValues=*/false);
 
-	const EPCGPointNativeProperties PropertiesToAllocate = PCGUnionHelper::GetPropertiesToAllocateFromInputs(PointInputDatas);
+	const EPCGPointNativeProperties PropertiesToAllocate = UPCGBasePointData::GetPropertiesToAllocateFromPointData(PointInputDatas);
 	OutPointData->AllocateProperties(PropertiesToAllocate | EPCGPointNativeProperties::MetadataEntry | EPCGPointNativeProperties::Density);
 	const bool bSetColor = EnumHasAnyFlags(PropertiesToAllocate, EPCGPointNativeProperties::Color);
 
diff --git a/Engine/Plugins/PCG/Source/PCG/Private/Elements/PCGMergeElement.cpp b/Engine/Plugins/PCG/Source/PCG/Private/Elements/PCGMergeElement.cpp
--- a/Engine/Plugins/PCG/Source/PCG/Private/Elements/PCGMergeElement.cpp
+++ b/Engine/Plugins/PCG/Source/PCG/Private/Elements/PCGMergeElement.cpp
@@ -74,6 +74,7 @@ bool FPCGMergeElement::ExecuteInternal(FPCGContext* Context) const
 	FPCGTaggedData* TargetTaggedData = nullptr;
 
 	int32 TotalPointCount = 0;
+	TArray<const UPCGBasePointData*> SourcePointDatas;
 
 	// Prepare data & metadata
 	// Done in two passes for futureproofing - expecting changes in the metadata attribute creation vs. usage in points
@@ -92,6 +93,7 @@ bool FPCGMergeElement::ExecuteInternal(FPCGContext* Context) const
 			continue;
 		}
 
+		SourcePointDatas.Add(SourcePointData);
 		TotalPointCount += SourcePointData->GetNumPoints();
 
 		// First data that's valid - will be the original output so we don't do a copy unless we actually need it
@@ -132,17 +134,15 @@ bool FPCGMergeElement::ExecuteInternal(FPCGContext* Context) const
 	}
 
 	TargetPointData->SetNumPoints(TotalPointCount);
+	TargetPointData->AllocateProperties(UPCGBasePointData::GetPropertiesToAllocateFromPointData(SourcePointDatas));
+
 	int32 PointOffset = 0;
 
-	for(int32 SourceDataIndex = 0; SourceDataIndex < Sources.Num(); ++SourceDataIndex)
+	for(int32 SourceDataIndex = 0; SourceDataIndex < SourcePointDatas.Num(); ++SourceDataIndex)
 	{
-		const UPCGBasePointData* SourcePointData = Cast<const UPCGBasePointData>(Sources[SourceDataIndex].Data);
-		const int32 NumSourcePoints = SourcePointData ? SourcePointData->GetNumPoints() : 0;
-
-		if (NumSourcePoints == 0)
-		{
-			continue;
-		}
+		const UPCGBasePointData* SourcePointData = SourcePointDatas[SourceDataIndex];
+		const int32 NumSourcePoints = SourcePointData->GetNumPoints();
+		ensure(NumSourcePoints > 0);
 
 		SourcePointData->CopyPointsTo(TargetPointData, 0, PointOffset, NumSourcePoints);
 
diff --git a/Engine/Plugins/PCG/Source/PCG/Public/Data/PCGBasePointData.h b/Engine/Plugins/PCG/Source/PCG/Public/Data/PCGBasePointData.h
--- a/Engine/Plugins/PCG/Source/PCG/Public/Data/PCGBasePointData.h
+++ b/Engine/Plugins/PCG/Source/PCG/Public/Data/PCGBasePointData.h
@@ -157,6 +157,58 @@ public:
 
 	static UE_API void SetPoints(const UPCGBasePointData* From, UPCGBasePointData* To, const TArrayView<const int32>& InDataIndices, bool bCopyAll);
 
+	static EPCGPointNativeProperties GetPropertiesToAllocateFromPointData(const TArray<const UPCGBasePointData*>& PointDatas)
+	{
+		// Start by doing a union of all input allocated properties
+		EPCGPointNativeProperties PropertiesToAllocate = EPCGPointNativeProperties::None;
+		for (const UPCGBasePointData* PointInputData : PointDatas)
+		{
+			PropertiesToAllocate |= PointInputData->GetAllocatedProperties();
+		}
+
+		if (!EnumHasAnyFlags(PropertiesToAllocate, EPCGPointNativeProperties::Transform) && NeedToAllocateSingleValueProperty<FTransform>(EPCGPointNativeProperties::Transform, PointDatas))
+		{
+			PropertiesToAllocate |= EPCGPointNativeProperties::Transform;
+		}
+
+		if (!EnumHasAnyFlags(PropertiesToAllocate, EPCGPointNativeProperties::Density) && NeedToAllocateSingleValueProperty<float>(EPCGPointNativeProperties::Density, PointDatas))
+		{
+			PropertiesToAllocate |= EPCGPointNativeProperties::Density;
+		}
+
+		if (!EnumHasAnyFlags(PropertiesToAllocate, EPCGPointNativeProperties::BoundsMin) && NeedToAllocateSingleValueProperty<FVector>(EPCGPointNativeProperties::BoundsMin, PointDatas))
+		{
+			PropertiesToAllocate |= EPCGPointNativeProperties::BoundsMin;
+		}
+
+		if (!EnumHasAnyFlags(PropertiesToAllocate, EPCGPointNativeProperties::BoundsMax) && NeedToAllocateSingleValueProperty<FVector>(EPCGPointNativeProperties::BoundsMax, PointDatas))
+		{
+			PropertiesToAllocate |= EPCGPointNativeProperties::BoundsMax;
+		}
+
+		if (!EnumHasAnyFlags(PropertiesToAllocate, EPCGPointNativeProperties::Color) && NeedToAllocateSingleValueProperty<FVector4>(EPCGPointNativeProperties::Color, PointDatas))
+		{
+			PropertiesToAllocate |= EPCGPointNativeProperties::Color;
+		}
+
+		if (!EnumHasAnyFlags(PropertiesToAllocate, EPCGPointNativeProperties::Steepness) && NeedToAllocateSingleValueProperty<float>(EPCGPointNativeProperties::Steepness, PointDatas))
+		{
+			PropertiesToAllocate |= EPCGPointNativeProperties::Steepness;
+		}
+
+		if (!EnumHasAnyFlags(PropertiesToAllocate, EPCGPointNativeProperties::Seed) && NeedToAllocateSingleValueProperty<int32>(EPCGPointNativeProperties::Seed, PointDatas))
+		{
+			PropertiesToAllocate |= EPCGPointNativeProperties::Seed;
+		}
+
+		if (!EnumHasAnyFlags(PropertiesToAllocate, EPCGPointNativeProperties::MetadataEntry) && NeedToAllocateSingleValueProperty<int64>(EPCGPointNativeProperties::MetadataEntry, PointDatas))
+		{
+			PropertiesToAllocate |= EPCGPointNativeProperties::MetadataEntry;
+		}
+
+		return PropertiesToAllocate;
+	}
+
 protected:
 
 	void RebuildOctreeIfNeeded() const
@@ -185,6 +237,33 @@ protected:
 		bBoundsAreDirty = true;
 	}
 
+private:
+
+	template<class T>
+	static bool NeedToAllocateSingleValueProperty(EPCGPointNativeProperties InProperty, const TArray<const UPCGBasePointData*>& PointDatas)
+	{
+		// Here we are comparing single values (non-allocated properties) to see if we have single values that differ between inputs.
+		// If they do differ, we need to allocate those properties because the output will have multiple values.
+		TOptional<const T> DefaultValue;
+		for (int32 i = 0; i < PointDatas.Num(); ++i)
+		{
+			const TConstPCGValueRange<T> ValueRange = PointDatas[i]->GetConstValueRange<T>(InProperty);
+			TOptional<const T> SingleValue = ValueRange.GetSingleValue();
+			if (!DefaultValue.IsSet())
+			{
+				DefaultValue = SingleValue;
+			}
+			else if (DefaultValue.IsSet() && SingleValue.IsSet() && !PCG::Private::MetadataTraits<T>::Equal(DefaultValue.GetValue(), SingleValue.GetValue()))
+			{
+				return true;
+			}
+		}
+
+		return false;
+	}
+
+protected:
+
 	mutable FCriticalSection CachedDataLock;
 	mutable PCGPointOctree::FPointOctree PCGPointOctree;
 	mutable FBox Bounds;
