[PCG] Fix extraction of properties that are array of objects

* We were missing one level of indirection when expending the containers
* Also protect from nullptr dereferencing.

#rnx

diff --git a/Engine/Plugins/PCG/Source/PCG/Private/Helpers/PCGPropertyHelpers.cpp b/Engine/Plugins/PCG/Source/PCG/Private/Helpers/PCGPropertyHelpers.cpp
--- a/Engine/Plugins/PCG/Source/PCG/Private/Helpers/PCGPropertyHelpers.cpp
+++ b/Engine/Plugins/PCG/Source/PCG/Private/Helpers/PCGPropertyHelpers.cpp
@@ -21,18 +21,6 @@ namespace PCGPropertyHelpers
 	static constexpr uint64 ExcludePropertyFlags = CPF_DisableEditOnInstance;
 	static constexpr uint64 IncludePropertyFlags = CPF_BlueprintVisible;
 
-	void LogError(const FText& ErrorMessage, FPCGContext* InOptionalContext)
-	{
-		if (InOptionalContext)
-		{
-			PCGE_LOG_C(Error, GraphAndLog, InOptionalContext, ErrorMessage);
-		}
-		else
-		{
-			UE_LOG(LogPCG, Error, TEXT("%s"), *ErrorMessage.ToString());
-		}
-	}
-
 	/**
 	* Expands container locations to their contents when the property passed in is an array or a set.
 	* This is useful to allow extraction downstream of properties inside of arrays/sets and also to generate the list of addresses/values to look at
@@ -94,7 +82,7 @@ namespace PCGPropertyHelpers
 		{
 			if (!bQuiet)
 			{
-				LogError(FText::Format(LOCTEXT("PropertyDoesNotExist", "Property '{0}' does not exist in {1}."), FText::FromName(CurrentName), FText::FromName(CurrentClass->GetFName())), OptionalContext);
+				PCGLog::LogErrorOnGraph(FText::Format(LOCTEXT("PropertyDoesNotExist", "Property '{0}' does not exist in {1}."), FText::FromName(CurrentName), FText::FromName(CurrentClass->GetFName())), OptionalContext);
 			}
 
 			return nullptr;
@@ -105,7 +93,7 @@ namespace PCGPropertyHelpers
 		{
 			if (!bQuiet)
 			{
-				LogError(FText::Format(LOCTEXT("PropertyExistsButNotVisible", "Property '{0}' does exist in {1}, but is not visible."), FText::FromName(CurrentName), FText::FromName(CurrentClass->GetFName())), OptionalContext);
+				PCGLog::LogErrorOnGraph(FText::Format(LOCTEXT("PropertyExistsButNotVisible", "Property '{0}' does exist in {1}, but is not visible."), FText::FromName(CurrentName), FText::FromName(CurrentClass->GetFName())), OptionalContext);
 			}
 
 			return nullptr;
@@ -114,12 +102,13 @@ namespace PCGPropertyHelpers
 		auto ExtractContainers = [&OutContainers](UStruct*& NextClass, const auto* ContainerProperty, const FProperty* InnerProperty) -> bool
 		{
 			bool bPropertyNotExtractable = false;
+			const FObjectProperty* InnerObjectProperty = nullptr;
 			
 			if (const FStructProperty* InnerStructProperty = CastField<FStructProperty>(InnerProperty))
 			{
 				NextClass = InnerStructProperty->Struct;
 			}
-			else if (const FObjectProperty* InnerObjectProperty = CastField<FObjectProperty>(InnerProperty))
+			else if ((InnerObjectProperty = CastField<FObjectProperty>(InnerProperty)) != nullptr)
 			{
 				NextClass = InnerObjectProperty->PropertyClass;
 			}
@@ -133,6 +122,16 @@ namespace PCGPropertyHelpers
 			{
 				TArray<const void*> Subcontainers;
 				PCGPropertyHelpers::ExpandContainers(ContainerProperty, OutContainers, Subcontainers);
+
+				// If we have an object, we also need to apply the object indirection on the Subcontainers
+				if (InnerObjectProperty)
+				{
+					for (const void*& OutContainer : Subcontainers)
+					{
+						OutContainer = OutContainer ? InnerObjectProperty->GetObjectPropertyValue_InContainer(OutContainer) : nullptr;
+					}
+				}
+
 				OutContainers = MoveTemp(Subcontainers);
 			}
 			
@@ -149,7 +148,7 @@ namespace PCGPropertyHelpers
 				NextClass = StructProperty->Struct;
 				for (const void*& OutContainer : OutContainers)
 				{
-					OutContainer = StructProperty->ContainerPtrToValuePtr<void>(OutContainer);
+					OutContainer = OutContainer ? StructProperty->ContainerPtrToValuePtr<void>(OutContainer) : nullptr;
 				}
 			}
 			else if (const FObjectProperty* ObjectProperty = CastField<FObjectProperty>(Property))
@@ -157,7 +156,7 @@ namespace PCGPropertyHelpers
 				NextClass = ObjectProperty->PropertyClass;
 				for (const void*& OutContainer : OutContainers)
 				{
-					OutContainer = ObjectProperty->GetObjectPropertyValue_InContainer(OutContainer);
+					OutContainer = OutContainer ? ObjectProperty->GetObjectPropertyValue_InContainer(OutContainer) : nullptr;
 				}
 			}
 			else if (const FArrayProperty* ArrayProperty = CastField<FArrayProperty>(Property))
@@ -177,7 +176,7 @@ namespace PCGPropertyHelpers
 			{
 				if (!bQuiet)
 				{
-					LogError(FText::Format(LOCTEXT("PropertyIsNotExtractable", "Property '{0}' does exist in {1}, but is not extractable."), FText::FromName(CurrentName), FText::FromName(CurrentClass->GetFName())), OptionalContext);
+					PCGLog::LogErrorOnGraph(FText::Format(LOCTEXT("PropertyIsNotExtractable", "Property '{0}' does exist in {1}, but is not extractable."), FText::FromName(CurrentName), FText::FromName(CurrentClass->GetFName())), OptionalContext);
 				}
 
 				return nullptr;
@@ -369,12 +368,12 @@ UPCGParamData* PCGPropertyHelpers::ExtractPropertyAsAttributeSet(const PCGProper
 			else if (const FStructProperty* StructProperty = CastField<FStructProperty>(Property))
 			{
 				UnderlyingClass = StructProperty->Struct;
-				AddressFunc = [StructProperty](const void* InAddress) { return StructProperty->ContainerPtrToValuePtr<void>(InAddress); };
+				AddressFunc = [StructProperty](const void* InAddress) { return InAddress ? StructProperty->ContainerPtrToValuePtr<void>(InAddress) : nullptr; };
 			}
 			else if (ObjectProperty)
 			{
 				UnderlyingClass = ObjectProperty->PropertyClass;
-				AddressFunc = [ObjectProperty](const void* InAddress) { return ObjectProperty->GetObjectPropertyValue_InContainer(InAddress); };
+				AddressFunc = [ObjectProperty](const void* InAddress) { return InAddress ? ObjectProperty->GetObjectPropertyValue_InContainer(InAddress) : nullptr; };
 			}
 
 			check(UnderlyingClass);
@@ -442,7 +441,7 @@ UPCGParamData* PCGPropertyHelpers::ExtractPropertyAsAttributeSet(const PCGProper
 		{
 			if (!bQuiet)
 			{
-				LogError(LOCTEXT("NoPropertiesFound", "No properties found to extract"), OptionalContext);
+				PCGLog::LogErrorOnGraph(LOCTEXT("NoPropertiesFound", "No properties found to extract"), OptionalContext);
 			}
 
 			return nullptr;
@@ -485,13 +484,15 @@ UPCGParamData* PCGPropertyHelpers::ExtractPropertyAsAttributeSet(const PCGProper
 			{
 				if (!bQuiet)
 				{
-					LogError(LOCTEXT("InvalidCardinality", "Unable to extract because some properties are of mismatched sizes"), OptionalContext);
+					PCGLog::LogErrorOnGraph(LOCTEXT("InvalidCardinality", "Unable to extract because some properties are of mismatched sizes"), OptionalContext);
 				}
 
 				return nullptr;
 			}
 		}
 
+		bool bHasWarnNullContainerPtr = false;
+
 		for(int32 ElementAddressIndex = 0; ElementAddressIndex < ElementAddresses.Num(); ++ElementAddressIndex)
 		{
 			const void* ElementAddress = ElementAddresses[ElementAddressIndex];
@@ -500,6 +501,17 @@ UPCGParamData* PCGPropertyHelpers::ExtractPropertyAsAttributeSet(const PCGProper
 			// Offset the address if needed
 			const void* ContainerPtr = AddressFunc(ElementAddress);
 
+			if (!ContainerPtr)
+			{
+				if (!bHasWarnNullContainerPtr)
+				{
+					bHasWarnNullContainerPtr = true;
+					PCGLog::LogErrorOnGraph(LOCTEXT("NullPtrContainer", "Some resolved objects were not assigned (null pointers), some results are discarded"), OptionalContext);
+				}
+				
+				continue;
+			}
+
 			for (ExtractablePropertyTuple& ExtractableProperty : ExtractableProperties)
 			{
 				FString AttributeNameStr = ExtractableProperty.Get<0>();
@@ -521,7 +533,7 @@ UPCGParamData* PCGPropertyHelpers::ExtractPropertyAsAttributeSet(const PCGProper
 				{
 					if (!bQuiet)
 					{
-						LogError(FText::Format(LOCTEXT("ErrorCreatingAttribute", "Error while creating an attribute for property '{0}'. Either the property type is not supported by PCG or attribute creation failed."), FText::FromString(FinalProperty->GetName())), OptionalContext);
+						PCGLog::LogErrorOnGraph(FText::Format(LOCTEXT("ErrorCreatingAttribute", "Error while creating an attribute for property '{0}'. Either the property type is not supported by PCG or attribute creation failed."), FText::FromString(FinalProperty->GetName())), OptionalContext);
 					}
 
 					bValidOperation = false;
