[PCG] Hook the comment node save to PreSave instead of SaveAsset

* Will catch all the flows to save a pcg graph, even if the editor graph is not the active one.
* Also initialized FPCGGraphCommentNodeData struct using the default values from UEdGraphNode_Comment

#rnx

diff --git a/Engine/Plugins/PCG/Source/PCG/Private/PCGGraph.cpp b/Engine/Plugins/PCG/Source/PCG/Private/PCGGraph.cpp
--- a/Engine/Plugins/PCG/Source/PCG/Private/PCGGraph.cpp
+++ b/Engine/Plugins/PCG/Source/PCG/Private/PCGGraph.cpp
@@ -666,6 +666,11 @@ void UPCGGraph::PreSave(FObjectPreSaveContext ObjectSaveContext)
 {
 	Super::PreSave(ObjectSaveContext);
 
+	if (IPCGEditorModule* PCGEditorModule = IPCGEditorModule::Get())
+	{
+		PCGEditorModule->OnGraphPreSave(this, ObjectSaveContext);
+	}
+
 	if (ObjectSaveContext.IsCooking())
 	{
 		FPCGGraphCompiler GraphCompiler(/*bIsCooking=*/true);
diff --git a/Engine/Plugins/PCG/Source/PCG/Public/Editor/IPCGEditorModule.h b/Engine/Plugins/PCG/Source/PCG/Public/Editor/IPCGEditorModule.h
--- a/Engine/Plugins/PCG/Source/PCG/Public/Editor/IPCGEditorModule.h
+++ b/Engine/Plugins/PCG/Source/PCG/Public/Editor/IPCGEditorModule.h
@@ -37,6 +37,7 @@ protected:
 	friend class UPCGSubsystem;
 
 	virtual void OnScheduleGraph(const FPCGStackContext& StackContext) = 0;
+	virtual void OnGraphPreSave(UPCGGraph* Graph, FObjectPreSaveContext ObjectSaveContext) = 0;
 
 	/** Clear any data collected during execution, normally called prior to generating the component. */
 	virtual void ClearExecutionMetadata(UPCGComponent* InComponent) = 0;
diff --git a/Engine/Plugins/PCG/Source/PCG/Public/Editor/PCGGraphComment.h b/Engine/Plugins/PCG/Source/PCG/Public/Editor/PCGGraphComment.h
--- a/Engine/Plugins/PCG/Source/PCG/Public/Editor/PCGGraphComment.h
+++ b/Engine/Plugins/PCG/Source/PCG/Public/Editor/PCGGraphComment.h
@@ -9,7 +9,10 @@
 #include "PCGGraphComment.generated.h"
 
 
-/** Mimic UEdGraphNode_Comment class layout to store the information in the backend (PCG Graph). */
+/**
+ * Mimic UEdGraphNode_Comment class layout to store the information in the backend (PCG Graph).
+ * Default values come from the default ctor of UEdGraphNode_Comment.
+ */
 USTRUCT(meta=(Hidden))
 struct FPCGGraphCommentNodeData
 {
@@ -22,44 +25,44 @@ struct FPCGGraphCommentNodeData
 #if WITH_EDITORONLY_DATA
 	/** X position of node in the editor */
 	UPROPERTY()
-	int32 NodePosX;
+	int32 NodePosX = 0;
 
 	/** Y position of node in the editor */
 	UPROPERTY()
-	int32 NodePosY;
+	int32 NodePosY = 0;
 
 	/** Width of node in the editor; only used when the node can be resized */
 	UPROPERTY()
-	int32 NodeWidth;
+	int32 NodeWidth = 400;
 
 	/** Height of node in the editor; only used when the node can be resized */
 	UPROPERTY()
-	int32 NodeHeight;
+	int32 NodeHeight = 100;
 
 	UPROPERTY()
 	FString NodeComment;
 	
 	/** Color to style comment with */
 	UPROPERTY()
-	FLinearColor CommentColor;
+	FLinearColor CommentColor = FLinearColor::White;
 
 	/** Size of the text in the comment box */
 	UPROPERTY()
-	int32 FontSize;
+	int32 FontSize = 18;
 
 	/** Whether to show a zoom-invariant comment bubble when zoomed out (making the comment readable at any distance). */
 	UPROPERTY()
-	uint32 bCommentBubbleVisible_InDetailsPanel:1;
+	uint32 bCommentBubbleVisible_InDetailsPanel:1 = 1;
 
 	/** Whether to use Comment Color to color the background of the comment bubble shown when zoomed out. */
 	UPROPERTY()
-	uint32 bColorCommentBubble:1;
+	uint32 bColorCommentBubble:1 = 0;
 
 	/** Whether the comment should move any fully enclosed nodes around when it is moved.
 	 * Underlying enum: ECommentBoxMode::Type
 	 * Implementation node: This needs to be type erased for UHT since ECommentBoxMode::Type is not visible outside of editor. */
 	UPROPERTY()
-	uint8 MoveMode;
+	uint8 MoveMode = 0;
 
 	/** Details field if more info is needed to be communicated (will show up in tooltip) */
 	UPROPERTY()
@@ -67,7 +70,7 @@ struct FPCGGraphCommentNodeData
 
 	/** comment Depth */
 	UPROPERTY()
-	int32 CommentDepth;
+	int32 CommentDepth = -1;
 
 	/** Extra GUID to be unique. */
 	UPROPERTY()
@@ -75,10 +78,10 @@ struct FPCGGraphCommentNodeData
 
 	/** Comment bubble pinned state */
 	UPROPERTY()
-	uint8 bCommentBubblePinned : 1;
+	uint8 bCommentBubblePinned : 1 = 1;
 
 	/** Comment bubble visibility */
 	UPROPERTY()
-	uint8 bCommentBubbleVisible : 1;
+	uint8 bCommentBubbleVisible : 1 = 1;
 #endif // WITH_EDITORONLY_DATA
 };
\ No newline at end of file
diff --git a/Engine/Plugins/PCG/Source/PCG/Public/PCGGraph.h b/Engine/Plugins/PCG/Source/PCG/Public/PCGGraph.h
--- a/Engine/Plugins/PCG/Source/PCG/Public/PCGGraph.h
+++ b/Engine/Plugins/PCG/Source/PCG/Public/PCGGraph.h
@@ -255,6 +255,7 @@ class UPCGGraph : public UPCGGraphInterface
 {
 #if WITH_EDITOR
 	friend class FPCGEditor;
+	friend class FPCGEditorModule;
 	friend class FPCGSubgraphHelpers;
 #endif // WITH_EDITOR
 	friend class UPCGSubgraphSettings;
diff --git a/Engine/Plugins/PCG/Source/PCGEditor/Private/PCGEditor.cpp b/Engine/Plugins/PCG/Source/PCGEditor/Private/PCGEditor.cpp
--- a/Engine/Plugins/PCG/Source/PCGEditor/Private/PCGEditor.cpp
+++ b/Engine/Plugins/PCG/Source/PCGEditor/Private/PCGEditor.cpp
@@ -2801,42 +2801,6 @@ TSharedRef<SGraphEditor> FPCGEditor::CreateGraphEditorWidget()
 		.ShowGraphStateOverlay(false);
 }
 
-void FPCGEditor::ReplicateExtraNodes() const
-{
-	if (PCGEditorGraph)
-	{
-		if (UPCGGraph* PCGGraph = PCGEditorGraph->GetPCGGraph())
-		{
-			TArray<TObjectPtr<const UObject>> ExtraNodes;
-			TArray<FPCGGraphCommentNodeData> CommentData;
-			for (const UEdGraphNode* GraphNode : PCGEditorGraph->Nodes)
-			{
-				check(GraphNode);
-				if (const UEdGraphNode_Comment* CommentNode = Cast<const UEdGraphNode_Comment>(GraphNode))
-				{
-					CommentData.Emplace_GetRef().InitializeFromCommentNode(*CommentNode);
-				}
-				else if (!GraphNode->IsA<UPCGEditorGraphNodeBase>())
-				{
-					ExtraNodes.Add(GraphNode);
-				}
-
-			}
-
-			PCGGraph->SetExtraEditorNodes(ExtraNodes);
-			PCGGraph->SetCommentNodes(std::move(CommentData));
-		}
-	}
-}
-
-void FPCGEditor::SaveAsset_Execute()
-{
-	// Extra nodes are replicated on asset save, to be saved in the underlying PCGGraph
-	ReplicateExtraNodes();
-
-	FAssetEditorToolkit::SaveAsset_Execute();
-}
-
 void FPCGEditor::OnClose()
 {
 	if (PCGEditorGraph)
@@ -2857,9 +2821,6 @@ void FPCGEditor::OnClose()
 	FEditorDelegates::PostPIEStarted.RemoveAll(this);
 	FEditorDelegates::EndPIE.RemoveAll(this);
 
-	// Extra nodes are replicated on editor close, to be saved in the underlying PCGGraph
-	ReplicateExtraNodes();
-
 	FAssetEditorToolkit::OnClose();
 
 	if (PCGComponentBeingInspected.IsValid())
diff --git a/Engine/Plugins/PCG/Source/PCGEditor/Private/PCGEditor.h b/Engine/Plugins/PCG/Source/PCGEditor/Private/PCGEditor.h
--- a/Engine/Plugins/PCG/Source/PCGEditor/Private/PCGEditor.h
+++ b/Engine/Plugins/PCG/Source/PCGEditor/Private/PCGEditor.h
@@ -163,12 +163,6 @@ public:
 	/** Run the determinism test on the current graph */
 	void OnDeterminismGraphTest() const;
 
-protected:
-	// ~Begin FAssetEditorToolkit interface
-	/** Called when "Save" is clicked for this asset */
-	virtual void SaveAsset_Execute() override;
-	// ~End FAssetEditorToolkit interface
-
 private:
 	/** Register PCG specific toolbar for the editor */
 	void RegisterToolbar() const;
@@ -403,9 +397,6 @@ private:
 	 */
 	void JumpToDefinition(const UClass* Class) const;
 
-	/** To be called everytime we need to replicate our extra nodes to the underlying PCGGraph */
-	void ReplicateExtraNodes() const;
-
 	/** Called when a PCG component unregisters. */
 	void OnComponentUnregistered(UPCGComponent* Component);
 
diff --git a/Engine/Plugins/PCG/Source/PCGEditor/Private/PCGEditorGraph.cpp b/Engine/Plugins/PCG/Source/PCGEditor/Private/PCGEditorGraph.cpp
--- a/Engine/Plugins/PCG/Source/PCGEditor/Private/PCGEditorGraph.cpp
+++ b/Engine/Plugins/PCG/Source/PCGEditor/Private/PCGEditorGraph.cpp
@@ -162,6 +162,8 @@ void UPCGEditorGraph::BeginDestroy()
 
 void UPCGEditorGraph::OnClose()
 {
+	ReplicateExtraNodes();
+
 	if (PCGGraph)
 	{
 		PCGGraph->OnGraphParametersChangedDelegate.RemoveAll(this);
@@ -186,6 +188,31 @@ void UPCGEditorGraph::CreateLinks(UPCGEditorGraphNodeBase* GraphNode, bool bCrea
 	CreateLinks(GraphNode, bCreateInbound, bCreateOutbound, PCGNodeToPCGEditorNodeMap);
 }
 
+void UPCGEditorGraph::ReplicateExtraNodes() const
+{
+	if (PCGGraph)
+	{
+		TArray<TObjectPtr<const UObject>> ExtraNodes;
+		TArray<FPCGGraphCommentNodeData> CommentData;
+		for (const UEdGraphNode* GraphNode : Nodes)
+		{
+			check(GraphNode);
+			if (const UEdGraphNode_Comment* CommentNode = Cast<const UEdGraphNode_Comment>(GraphNode))
+			{
+				CommentData.Emplace_GetRef().InitializeFromCommentNode(*CommentNode);
+			}
+			else if (!GraphNode->IsA<UPCGEditorGraphNodeBase>())
+			{
+				ExtraNodes.Add(GraphNode);
+			}
+
+		}
+
+		PCGGraph->SetExtraEditorNodes(ExtraNodes);
+		PCGGraph->SetCommentNodes(std::move(CommentData));
+	}
+}
+
 void UPCGEditorGraph::UpdateVisualizations(UPCGComponent* PCGComponentBeingInspected, const FPCGStack* PCGStackBeingInspected)
 {
 	for (UEdGraphNode* EditorNode : Nodes)
diff --git a/Engine/Plugins/PCG/Source/PCGEditor/Private/PCGEditorGraph.h b/Engine/Plugins/PCG/Source/PCGEditor/Private/PCGEditorGraph.h
--- a/Engine/Plugins/PCG/Source/PCGEditor/Private/PCGEditorGraph.h
+++ b/Engine/Plugins/PCG/Source/PCGEditor/Private/PCGEditorGraph.h
@@ -42,6 +42,9 @@ public:
 	/** Creates the links for a given node */
 	void CreateLinks(UPCGEditorGraphNodeBase* InGraphNode, bool bCreateInbound, bool bCreateOutbound);
 
+	/** To be called everytime we need to replicate our extra nodes to the underlying PCGGraph */
+	void ReplicateExtraNodes() const;
+
 	UPCGGraph* GetPCGGraph() const { return PCGGraph; }
 
 	void SetEditor(TWeakPtr<FPCGEditor> InEditor) { PCGEditor = InEditor; }
diff --git a/Engine/Plugins/PCG/Source/PCGEditor/Private/PCGEditorModule.cpp b/Engine/Plugins/PCG/Source/PCGEditor/Private/PCGEditorModule.cpp
--- a/Engine/Plugins/PCG/Source/PCGEditor/Private/PCGEditorModule.cpp
+++ b/Engine/Plugins/PCG/Source/PCGEditor/Private/PCGEditorModule.cpp
@@ -21,6 +21,7 @@
 #include "WorldPartitionBuilder/PCGWorldPartitionBuilder.h"
 
 #include "PCGEditorCommands.h"
+#include "PCGEditorGraph.h"
 #include "PCGEditorMenuUtils.h"
 #include "PCGEditorProgressNotification.h"
 #include "PCGEditorSettings.h"
@@ -63,6 +64,7 @@
 #include "Editor/SceneOutliner/Public/ISceneOutliner.h"
 #include "Framework/MultiBox/MultiBoxBuilder.h"
 #include "Framework/Notifications/NotificationManager.h"
+#include "UObject/ObjectSaveContext.h"
 #include "Widgets/SPCGNodeSourceTextBox.h"
 
 #define LOCTEXT_NAMESPACE "FPCGEditorModule"
@@ -859,6 +861,20 @@ void FPCGEditorModule::OnScheduleGraph(const FPCGStackContext& StackContext)
 	}
 }
 
+void FPCGEditorModule::OnGraphPreSave(UPCGGraph* Graph, FObjectPreSaveContext ObjectSaveContext)
+{
+	if (!Graph || !Graph->PCGEditorGraph)
+	{
+		return;
+	}
+
+	// No need to do it on cooking
+	if (!ObjectSaveContext.IsProceduralSave())
+	{
+		Graph->PCGEditorGraph->ReplicateExtraNodes();
+	}
+}
+
 void FPCGEditorModule::ClearExecutionMetadata(UPCGComponent* InComponent)
 {
 	if (InComponent)
diff --git a/Engine/Plugins/PCG/Source/PCGEditor/Public/PCGEditorModule.h b/Engine/Plugins/PCG/Source/PCGEditor/Public/PCGEditorModule.h
--- a/Engine/Plugins/PCG/Source/PCGEditor/Public/PCGEditorModule.h
+++ b/Engine/Plugins/PCG/Source/PCGEditor/Public/PCGEditorModule.h
@@ -31,19 +31,20 @@ public:
 	virtual TWeakPtr<IPCGEditorProgressNotification> CreateProgressNotification(const FTextFormat& TextFormat, bool bCanCancel) override;
 	virtual void ReleaseProgressNotification(TWeakPtr<IPCGEditorProgressNotification> InNotification) override;
 	virtual void SetOutlinerUIRefreshDelay(float InDelay) override;
-	virtual const FPCGNodeVisualLogs& GetNodeVisualLogs() const { return NodeVisualLogs; }
-	virtual FPCGNodeVisualLogs& GetNodeVisualLogsMutable() { return NodeVisualLogs; }
+	virtual const FPCGNodeVisualLogs& GetNodeVisualLogs() const override { return NodeVisualLogs; }
+	virtual FPCGNodeVisualLogs& GetNodeVisualLogsMutable() override { return NodeVisualLogs; }
 	virtual bool CanSelectPartitionActors() const override;
 	virtual TSharedPtr<class ISyntaxTokenizer> CreateHLSLSyntaxTokenizer(const FPCGSyntaxTokenizerParams& InParams) const override;
 
 protected:
-	virtual void OnScheduleGraph(const FPCGStackContext& StackContext);
-	virtual void ClearExecutionMetadata(UPCGComponent* InComponent);
-	virtual void ClearExecutedStacks(const UPCGComponent* InRootComponent);
-	virtual void ClearExecutedStacks(const UPCGGraph* InContainingGraph);
-	virtual	TArray<FPCGStackSharedPtr> GetExecutedStacksPtrs(const FPCGStack& BeginningWithStack);
-	virtual TArray<FPCGStackSharedPtr> GetExecutedStacksPtrs(const UPCGComponent* InComponent, const UPCGGraph* InSubgraph, bool bOnlyWithSubgraphAsCurrentFrame);
-	void NotifyGraphChanged(UPCGGraph* InGraph, EPCGChangeType ChangeType);
+	virtual void OnScheduleGraph(const FPCGStackContext& StackContext) override;
+	virtual void OnGraphPreSave(UPCGGraph* Graph, FObjectPreSaveContext ObjectSaveContext) override;
+	virtual void ClearExecutionMetadata(UPCGComponent* InComponent) override;
+	virtual void ClearExecutedStacks(const UPCGComponent* InRootComponent) override;
+	virtual void ClearExecutedStacks(const UPCGGraph* InContainingGraph) override;
+	virtual	TArray<FPCGStackSharedPtr> GetExecutedStacksPtrs(const FPCGStack& BeginningWithStack) override;
+	virtual TArray<FPCGStackSharedPtr> GetExecutedStacksPtrs(const UPCGComponent* InComponent, const UPCGGraph* InSubgraph, bool bOnlyWithSubgraphAsCurrentFrame) override;
+	virtual void NotifyGraphChanged(UPCGGraph* InGraph, EPCGChangeType ChangeType) override;
 	// ~End IPCGEditorModule implementation
 
 	void RegisterDetailsCustomizations();
