[PCG] Added cvar to control cache cleanup behavior to prevent hitches in game

diff --git a/Engine/Plugins/PCG/Source/PCG/Private/Graph/PCGGraphExecutor.cpp b/Engine/Plugins/PCG/Source/PCG/Private/Graph/PCGGraphExecutor.cpp
--- a/Engine/Plugins/PCG/Source/PCG/Private/Graph/PCGGraphExecutor.cpp
+++ b/Engine/Plugins/PCG/Source/PCG/Private/Graph/PCGGraphExecutor.cpp
@@ -96,6 +96,24 @@ namespace PCGGraphExecutor
 		true,
 		TEXT("Whether proxies for GPU data are cached in per pin output data and passed through grid links. If false data is read back to CPU."));
 
+#if WITH_EDITOR
+	TAutoConsoleVariable<int> CVarCacheEditorCleanupBehavior(
+		TEXT("pcg.Cache.EditorCleanupLevel"),
+		2,
+		TEXT("Level for the cache cleanup, when the memory budget is enforced.\n")
+		TEXT(" 0: Release objects only;\n")
+		TEXT(" 1: Release objects and force garbage collection (iterative);\n")
+		TEXT(" 2: Release objects and force garbage collection (complete)\n"));
+#endif
+
+	TAutoConsoleVariable<int> CVarCacheCleanupBehavior(
+		TEXT("pcg.Cache.CleanupLevel"),
+		0,
+		TEXT("Level for the cache cleanup, when the memory budget is enforced.\n")
+		TEXT(" 0: Release objects only;\n")
+		TEXT(" 1: Release objects and force garbage collection (iterative);\n")
+		TEXT(" 2: Release objects and force garbage collection (complete)\n"));
+
 	void ClearAsyncFlags(TSet<TObjectPtr<UObject>>& AsyncObjects)
 	{
 		// Remove Async Flags from Objects created in Async PCG Tasks so that they can get tracked by the GarbageCollector properly
@@ -1025,7 +1043,23 @@ void FPCGGraphExecutor::ExecuteTasksEnded()
 	// Purge things from cache if memory usage is too high
 	if (GraphCache.EnforceMemoryBudget() && GEngine)
 	{
-		GEngine->ForceGarbageCollection(/*bFullPurge=*/true);
+#if WITH_EDITOR
+		const int CleanupMode = (World && World->IsGameWorld()) ? PCGGraphExecutor::CVarCacheCleanupBehavior.GetValueOnAnyThread() : PCGGraphExecutor::CVarCacheEditorCleanupBehavior.GetValueOnAnyThread();
+#else
+		const int CleanupMode = PCGGraphExecutor::CVarCacheCleanupBehavior.GetValueOnAnyThread();
+#endif
+		if (CleanupMode <= 0)
+		{
+			// Do nothing
+		}
+		else if (CleanupMode == 1)
+		{
+			GEngine->ForceGarbageCollection();
+		}
+		else // CleanupMode >= 2
+		{
+			GEngine->ForceGarbageCollection(/*bFullPurge=*/true);
+		}
 	}
 }
 
