Fixup API macros for merged modules, PURE_VIRTUAL does not need API export

#rnx

diff --git a/Engine/Plugins/Animation/ControlRig/Source/ControlRig/Public/AnimNode_ControlRigBase.h b/Engine/Plugins/Animation/ControlRig/Source/ControlRig/Public/AnimNode_ControlRigBase.h
--- a/Engine/Plugins/Animation/ControlRig/Source/ControlRig/Public/AnimNode_ControlRigBase.h
+++ b/Engine/Plugins/Animation/ControlRig/Source/ControlRig/Public/AnimNode_ControlRigBase.h
@@ -45,8 +45,8 @@ struct FAnimNode_ControlRigBase : public FAnimNode_CustomProperty
 	UE_API FAnimNode_ControlRigBase();
 
 	/* return Control Rig of current object */
-	UE_API virtual UControlRig* GetControlRig() const PURE_VIRTUAL(FAnimNode_ControlRigBase::GetControlRig, return nullptr; );
-	UE_API virtual TSubclassOf<UControlRig> GetControlRigClass() const PURE_VIRTUAL(FAnimNode_ControlRigBase::GetControlRigClass, return nullptr; );
+	virtual UControlRig* GetControlRig() const PURE_VIRTUAL(FAnimNode_ControlRigBase::GetControlRig, return nullptr; );
+	virtual TSubclassOf<UControlRig> GetControlRigClass() const PURE_VIRTUAL(FAnimNode_ControlRigBase::GetControlRigClass, return nullptr; );
 	
 	// FAnimNode_Base interface
 	UE_API virtual void OnInitializeAnimInstance(const FAnimInstanceProxy* InProxy, const UAnimInstance* InAnimInstance) override;
diff --git a/Engine/Plugins/Animation/DeformerGraph/Source/OptimusCore/Public/OptimusComponentSource.h b/Engine/Plugins/Animation/DeformerGraph/Source/OptimusCore/Public/OptimusComponentSource.h
--- a/Engine/Plugins/Animation/DeformerGraph/Source/OptimusCore/Public/OptimusComponentSource.h
+++ b/Engine/Plugins/Animation/DeformerGraph/Source/OptimusCore/Public/OptimusComponentSource.h
@@ -23,20 +23,20 @@ class UOptimusComponentSource :
 
 public:
 	/** Returns the component display name to show in the lister. Should be unique. */
-	UE_API virtual FText GetDisplayName() const
+	virtual FText GetDisplayName() const
 	PURE_VIRTUAL(UOptimusComponentSource::GetDisplayName, return {}; );
 
 	/** Returns a suggested name for the binding. The name may be modified to preserve uniqueness. */
-	UE_API virtual FName GetBindingName() const
+	virtual FName GetBindingName() const
 	PURE_VIRTUAL(UOptimusComponentSource::GetSuggestedBindingName, return {}; );
 
 	/** Returns the actor component that this provider can operate on */
-	UE_API virtual TSubclassOf<UActorComponent> GetComponentClass() const
+	virtual TSubclassOf<UActorComponent> GetComponentClass() const
 	PURE_VIRTUAL(UOptimusComponentSource::GetComponentClass, return {}; );
 
 	/** Returns the available execution contexts. The first context is a preferred context when setting
 	  * initial data domain for a resource. */
-	UE_API virtual TArray<FName> GetExecutionDomains() const
+	virtual TArray<FName> GetExecutionDomains() const
 	PURE_VIRTUAL(UOptimusComponentSource::GetExecutionContexts, return {}; );
 
 	/** Returns the current LOD index associated with a component associated with this source. */
diff --git a/Engine/Plugins/Animation/DeformerGraph/Source/OptimusCore/Public/OptimusComputeDataInterface.h b/Engine/Plugins/Animation/DeformerGraph/Source/OptimusCore/Public/OptimusComputeDataInterface.h
--- a/Engine/Plugins/Animation/DeformerGraph/Source/OptimusCore/Public/OptimusComputeDataInterface.h
+++ b/Engine/Plugins/Animation/DeformerGraph/Source/OptimusCore/Public/OptimusComputeDataInterface.h
@@ -160,13 +160,13 @@ public:
 	};
 
 	/// Returns the name to show on the node that will proxy this interface in the graph view.
-	UE_API virtual FString GetDisplayName() const PURE_VIRTUAL(UOptimusComputeDataInterface::GetDisplayName, return {};)
+	virtual FString GetDisplayName() const PURE_VIRTUAL(UOptimusComputeDataInterface::GetDisplayName, return {};)
 
 	/// Returns the category for the node.
 	virtual FName GetCategory() const { return CategoryName::DataInterfaces; }
 
 	/// Returns the list of pins that will map to the shader functions provided by this data interface.
-	UE_API virtual TArray<FOptimusCDIPinDefinition> GetPinDefinitions() const PURE_VIRTUAL(UOptimusComputeDataInterface::GetPinDefinitions, return {};)
+	virtual TArray<FOptimusCDIPinDefinition> GetPinDefinitions() const PURE_VIRTUAL(UOptimusComputeDataInterface::GetPinDefinitions, return {};)
 
 	/// Returns the list of pins that are related to the properties of the data interface.
 	virtual TArray<FOptimusCDIPropertyPinDefinition> GetPropertyPinDefinitions() const { return {}; }
@@ -186,7 +186,7 @@ public:
 	/**
 	 * @return Returns the component type that this data interface operates on.
 	 */
-	UE_API virtual TSubclassOf<UActorComponent> GetRequiredComponentClass() const PURE_VIRTUAL(UOptimusComputeDataInterface::GetRequiredComponent, return nullptr;)
+	virtual TSubclassOf<UActorComponent> GetRequiredComponentClass() const PURE_VIRTUAL(UOptimusComputeDataInterface::GetRequiredComponent, return nullptr;)
 
 	virtual void OnDataTypeChanged(FName InTypeName) {};
 	
diff --git a/Engine/Plugins/Animation/DeformerGraph/Source/OptimusCore/Public/OptimusNode.h b/Engine/Plugins/Animation/DeformerGraph/Source/OptimusCore/Public/OptimusNode.h
--- a/Engine/Plugins/Animation/DeformerGraph/Source/OptimusCore/Public/OptimusNode.h
+++ b/Engine/Plugins/Animation/DeformerGraph/Source/OptimusCore/Public/OptimusNode.h
@@ -53,7 +53,7 @@ public:
 	 * @return The node class category.
 	 */
 	UFUNCTION(BlueprintCallable, Category = OptimusNodeGraph)
-	UE_API virtual FName GetNodeCategory() const PURE_VIRTUAL(, return NAME_None;);
+	virtual FName GetNodeCategory() const PURE_VIRTUAL(, return NAME_None;);
 
 	/** Returns true if the node can be deleted by the user */
 	virtual bool CanUserDeleteNode() const { return true; }
diff --git a/Engine/Plugins/Animation/PoseSearch/Source/Runtime/Public/PoseSearch/MultiAnimAsset.h b/Engine/Plugins/Animation/PoseSearch/Source/Runtime/Public/PoseSearch/MultiAnimAsset.h
--- a/Engine/Plugins/Animation/PoseSearch/Source/Runtime/Public/PoseSearch/MultiAnimAsset.h
+++ b/Engine/Plugins/Animation/PoseSearch/Source/Runtime/Public/PoseSearch/MultiAnimAsset.h
@@ -19,20 +19,20 @@ class UMultiAnimAsset : public UObject
 	GENERATED_BODY()
 public:
 
-	[[nodiscard]] UE_API virtual bool IsLooping() const PURE_VIRTUAL(UMultiAnimAsset::IsLooping, return false;);
-	[[nodiscard]] UE_API virtual bool HasRootMotion() const PURE_VIRTUAL(UMultiAnimAsset::HasRootMotion, return false;);
-	[[nodiscard]] UE_API virtual float GetPlayLength(const FVector& BlendParameters) const PURE_VIRTUAL(UMultiAnimAsset::GetPlayLength, return 0.f;);
+	[[nodiscard]] virtual bool IsLooping() const PURE_VIRTUAL(UMultiAnimAsset::IsLooping, return false;);
+	[[nodiscard]] virtual bool HasRootMotion() const PURE_VIRTUAL(UMultiAnimAsset::HasRootMotion, return false;);
+	[[nodiscard]] virtual float GetPlayLength(const FVector& BlendParameters) const PURE_VIRTUAL(UMultiAnimAsset::GetPlayLength, return 0.f;);
 
 #if WITH_EDITOR
-	[[nodiscard]] UE_API virtual USkeletalMesh* GetPreviewMesh(const FName& Role) const PURE_VIRTUAL(UMultiAnimAsset::GetPreviewMesh, return nullptr;);
+	[[nodiscard]] virtual USkeletalMesh* GetPreviewMesh(const FName& Role) const PURE_VIRTUAL(UMultiAnimAsset::GetPreviewMesh, return nullptr;);
 #endif // WITH_EDITOR
 
-	[[nodiscard]] UE_API virtual int32 GetNumRoles() const PURE_VIRTUAL(UMultiAnimAsset::GetNumRoles, return 0;);
-	[[nodiscard]] UE_API virtual FName GetRole(int32 RoleIndex) const PURE_VIRTUAL(UMultiAnimAsset::GetRole, return FName(););
-	[[nodiscard]] UE_API virtual UAnimationAsset* GetAnimationAsset(const FName& Role) const PURE_VIRTUAL(UMultiAnimAsset::GetAnimationAsset, return nullptr;);
-	[[nodiscard]] UE_API virtual FTransform GetOrigin(const FName& Role) const PURE_VIRTUAL(UMultiAnimAsset::GetOrigin, return FTransform::Identity;);
+	[[nodiscard]] virtual int32 GetNumRoles() const PURE_VIRTUAL(UMultiAnimAsset::GetNumRoles, return 0;);
+	[[nodiscard]] virtual FName GetRole(int32 RoleIndex) const PURE_VIRTUAL(UMultiAnimAsset::GetRole, return FName(););
+	[[nodiscard]] virtual UAnimationAsset* GetAnimationAsset(const FName& Role) const PURE_VIRTUAL(UMultiAnimAsset::GetAnimationAsset, return nullptr;);
+	[[nodiscard]] virtual FTransform GetOrigin(const FName& Role) const PURE_VIRTUAL(UMultiAnimAsset::GetOrigin, return FTransform::Identity;);
 
-	UE_API virtual void CalculateWarpTransforms(float Time, const TConstArrayView<const FTransform> ActorRootBoneTransforms, TArrayView<FTransform> FullAlignedActorRootBoneTransforms, const TConstArrayView<const UMirrorDataTable*> MirrorDataTables = TConstArrayView<const UMirrorDataTable*>(), const TConstArrayView<bool> RelevantRoleIndexes = TArrayView<bool>()) const PURE_VIRTUAL(UMultiAnimAsset::CalculateWarpTransforms, );
+	virtual void CalculateWarpTransforms(float Time, const TConstArrayView<const FTransform> ActorRootBoneTransforms, TArrayView<FTransform> FullAlignedActorRootBoneTransforms, const TConstArrayView<const UMirrorDataTable*> MirrorDataTables = TConstArrayView<const UMirrorDataTable*>(), const TConstArrayView<bool> RelevantRoleIndexes = TArrayView<bool>()) const PURE_VIRTUAL(UMultiAnimAsset::CalculateWarpTransforms, );
 
 	UFUNCTION(BlueprintPure, Category = "Animation", meta=(BlueprintThreadSafe, DisplayName = "Get Animation Asset"))
 	UAnimationAsset* BP_GetAnimationAsset(const FName& Role) const { return GetAnimationAsset(Role); }
diff --git a/Engine/Plugins/Animation/PoseSearch/Source/Runtime/Public/PoseSearch/PoseSearchFeatureChannel.h b/Engine/Plugins/Animation/PoseSearch/Source/Runtime/Public/PoseSearch/PoseSearchFeatureChannel.h
--- a/Engine/Plugins/Animation/PoseSearch/Source/Runtime/Public/PoseSearch/PoseSearchFeatureChannel.h
+++ b/Engine/Plugins/Animation/PoseSearch/Source/Runtime/Public/PoseSearch/PoseSearchFeatureChannel.h
@@ -123,10 +123,10 @@ public:
 	int32 GetChannelDataOffset() const { checkSlow(ChannelDataOffset >= 0); return ChannelDataOffset; }
 
 	// Called during UPoseSearchSchema::Finalize to prepare the schema for this channel
-	UE_API virtual bool Finalize(UPoseSearchSchema* Schema) PURE_VIRTUAL(UPoseSearchFeatureChannel::Finalize, return false;);
+	virtual bool Finalize(UPoseSearchSchema* Schema) PURE_VIRTUAL(UPoseSearchFeatureChannel::Finalize, return false;);
 	
 	// Called at runtime to add this channel's data to the query pose vector
-	UE_API virtual void BuildQuery(UE::PoseSearch::FSearchContext& SearchContext) const PURE_VIRTUAL(UPoseSearchFeatureChannel::BuildQuery, );
+	virtual void BuildQuery(UE::PoseSearch::FSearchContext& SearchContext) const PURE_VIRTUAL(UPoseSearchFeatureChannel::BuildQuery, );
 
 	// UPoseSearchFeatureChannels can hold sub channels
 	virtual TArrayView<TObjectPtr<UPoseSearchFeatureChannel>> GetSubChannels() { return TArrayView<TObjectPtr<UPoseSearchFeatureChannel>>(); }
@@ -146,10 +146,10 @@ public:
 	// Called at database build time to collect feature weights.
 	// Weights is sized to the cardinality of the schema and the feature channel should write
 	// its weights at the channel's data offset. Channels should provide a weight for each dimension.
-	UE_API virtual void FillWeights(TArrayView<float> Weights) const PURE_VIRTUAL(UPoseSearchFeatureChannel::FillWeights, );
+	virtual void FillWeights(TArrayView<float> Weights) const PURE_VIRTUAL(UPoseSearchFeatureChannel::FillWeights, );
 
 	// Called at database build time to populate pose vectors with this channel's data
-	UE_API virtual bool IndexAsset(UE::PoseSearch::FAssetIndexer& Indexer) const PURE_VIRTUAL(UPoseSearchFeatureChannel::IndexAsset, return false;);
+	virtual bool IndexAsset(UE::PoseSearch::FAssetIndexer& Indexer) const PURE_VIRTUAL(UPoseSearchFeatureChannel::IndexAsset, return false;);
 
 	// returns the TLabelBuilder used editor side to identify this UPoseSearchFeatureChannel (for instance in the pose search debugger)
 	UE_API virtual UE::PoseSearch::TLabelBuilder& GetLabel(UE::PoseSearch::TLabelBuilder& LabelBuilder, UE::PoseSearch::ELabelFormat LabelFormat = UE::PoseSearch::ELabelFormat::Full_Horizontal) const;
diff --git a/Engine/Plugins/Cameras/GameplayCameras/Source/GameplayCameras/Public/Core/CameraVariableAssets.h b/Engine/Plugins/Cameras/GameplayCameras/Source/GameplayCameras/Public/Core/CameraVariableAssets.h
--- a/Engine/Plugins/Cameras/GameplayCameras/Source/GameplayCameras/Public/Core/CameraVariableAssets.h
+++ b/Engine/Plugins/Cameras/GameplayCameras/Source/GameplayCameras/Public/Core/CameraVariableAssets.h
@@ -29,8 +29,8 @@ public:
 
 	const FGuid& GetGuid() const { return Guid; }
 
-	UE_API virtual ECameraVariableType GetVariableType() const PURE_VIRTUAL(UCameraVariableAsset::GetVariableType, return ECameraVariableType::Boolean;);
-	UE_API virtual const uint8* GetDefaultValuePtr() const PURE_VIRTUAL(UCameraVariableAsset::GetDefaultValuePtr, return nullptr;);
+	virtual ECameraVariableType GetVariableType() const PURE_VIRTUAL(UCameraVariableAsset::GetVariableType, return ECameraVariableType::Boolean;);
+	virtual const uint8* GetDefaultValuePtr() const PURE_VIRTUAL(UCameraVariableAsset::GetDefaultValuePtr, return nullptr;);
 
 #if WITH_EDITORONLY_DATA
 	UE_API FString GetDisplayName() const;
@@ -39,7 +39,7 @@ public:
 #if WITH_EDITOR
 	UE_API FText GetDisplayText() const;
 
-	UE_API virtual FString FormatDefaultValue() const PURE_VIRTUAL(UCameraVariableAsset::FormatDefaultValue, return FString(););
+	virtual FString FormatDefaultValue() const PURE_VIRTUAL(UCameraVariableAsset::FormatDefaultValue, return FString(););
 #endif  // WITH_EDITOR
 
 public:
diff --git a/Engine/Plugins/ChaosClothAsset/Source/ChaosClothAssetEngine/Public/ChaosClothAsset/ClothAssetBase.h b/Engine/Plugins/ChaosClothAsset/Source/ChaosClothAssetEngine/Public/ChaosClothAsset/ClothAssetBase.h
--- a/Engine/Plugins/ChaosClothAsset/Source/ChaosClothAssetEngine/Public/ChaosClothAsset/ClothAssetBase.h
+++ b/Engine/Plugins/ChaosClothAsset/Source/ChaosClothAssetEngine/Public/ChaosClothAsset/ClothAssetBase.h
@@ -32,31 +32,31 @@ public:
 
 	//~ Begin UChaosClothAssetBase interface
 	/** Return whether the asset has any valid cloth simulation models and is simulation enabled. */
-	UE_API virtual bool HasValidClothSimulationModels() const
+	virtual bool HasValidClothSimulationModels() const
 	PURE_VIRTUAL(UChaosClothAssetBase::HasValidClothSimulationModels, return false;);
 
 	/** Return the number of cloth simulation models in this asset. */
-	UE_API virtual int32 GetNumClothSimulationModels() const
+	virtual int32 GetNumClothSimulationModels() const
 	PURE_VIRTUAL(UChaosClothAssetBase::GetNumClothSimulationModels, return 0;);
 
 	/** Return the simulation model physical data. */
-	UE_API virtual TSharedPtr<const FChaosClothSimulationModel> GetClothSimulationModel(int32 ModelIndex) const
+	virtual TSharedPtr<const FChaosClothSimulationModel> GetClothSimulationModel(int32 ModelIndex) const
 	PURE_VIRTUAL(UChaosClothAssetBase::GetClothSimulationModel, return nullptr;);
 
 	/** Return the name of the specified cloth simulation models in this asset. */
-	UE_API virtual FName GetClothSimulationModelName(int32 ModelIndex) const
+	virtual FName GetClothSimulationModelName(int32 ModelIndex) const
 	PURE_VIRTUAL(UChaosClothAssetBase::GetClothSimulationModelName, return NAME_None;);
 
 	/** Return the collections for this asset model (one per LOD). */
-	UE_API virtual const TArray<TSharedRef<const FManagedArrayCollection>>& GetCollections(int32 ModelIndex) const
+	virtual const TArray<TSharedRef<const FManagedArrayCollection>>& GetCollections(int32 ModelIndex) const
 	PURE_VIRTUAL(UChaosClothAssetBase::GetCollections, static const TArray<TSharedRef<const FManagedArrayCollection>> EmptyArray; return EmptyArray;);
 
 	/** Return the physics asset used for the simulation by this asset model. */
-	UE_API virtual const UPhysicsAsset* GetPhysicsAssetForModel(int32 ModelIndex) const
+	virtual const UPhysicsAsset* GetPhysicsAssetForModel(int32 ModelIndex) const
 	PURE_VIRTUAL(UChaosClothAssetBase::GetPhysicsAssetForModel, return nullptr;);
 
 	/** Return the asset GUID used to match the render sections for this asset model. */
-	UE_API virtual FGuid GetAssetGuid(int32 ModelIndex) const
+	virtual FGuid GetAssetGuid(int32 ModelIndex) const
 	PURE_VIRTUAL(UChaosClothAssetBase::GetAssetGuid, return FGuid(););
 
 #if WITH_EDITOR
diff --git a/Engine/Plugins/ChaosClothAsset/Source/ChaosClothAssetEngine/Public/ChaosClothAsset/SkeletalMeshConverterClassProvider.h b/Engine/Plugins/ChaosClothAsset/Source/ChaosClothAssetEngine/Public/ChaosClothAsset/SkeletalMeshConverterClassProvider.h
--- a/Engine/Plugins/ChaosClothAsset/Source/ChaosClothAssetEngine/Public/ChaosClothAsset/SkeletalMeshConverterClassProvider.h
+++ b/Engine/Plugins/ChaosClothAsset/Source/ChaosClothAssetEngine/Public/ChaosClothAsset/SkeletalMeshConverterClassProvider.h
@@ -15,7 +15,7 @@ class UClothAssetSkeletalMeshConverter : public UObject
 
 public:
 	/** Build a skeletal mesh from the specified Cloth Asset. */
-	UE_API virtual bool ExportToSkeletalMesh(const class UChaosClothAssetBase& ClothAssetBase, class USkeletalMesh& SkeletalMesh) const
+	virtual bool ExportToSkeletalMesh(const class UChaosClothAssetBase& ClothAssetBase, class USkeletalMesh& SkeletalMesh) const
 		PURE_VIRTUAL(UClothAssetSkeletalMeshConverter::ExportToSkeletalMesh, return false;);
 };
 
diff --git a/Engine/Plugins/Experimental/Water/Source/Runtime/Public/WaterBodyComponent.h b/Engine/Plugins/Experimental/Water/Source/Runtime/Public/WaterBodyComponent.h
--- a/Engine/Plugins/Experimental/Water/Source/Runtime/Public/WaterBodyComponent.h
+++ b/Engine/Plugins/Experimental/Water/Source/Runtime/Public/WaterBodyComponent.h
@@ -166,7 +166,7 @@ public:
 	WATER_API virtual FBox GetCollisionComponentBounds() const;
 
 	/** Returns the type of body */
-	WATER_API virtual EWaterBodyType GetWaterBodyType() const PURE_VIRTUAL(UWaterBodyComponent::GetWaterBodyType, return EWaterBodyType::Transition; )
+	virtual EWaterBodyType GetWaterBodyType() const PURE_VIRTUAL(UWaterBodyComponent::GetWaterBodyType, return EWaterBodyType::Transition; )
 
 	/** Returns collision half-extents */
 	virtual FVector GetCollisionExtents() const { return FVector::ZeroVector; }
diff --git a/Engine/Plugins/Experimental/Water/Source/Runtime/Public/WaterWaves.h b/Engine/Plugins/Experimental/Water/Source/Runtime/Public/WaterWaves.h
--- a/Engine/Plugins/Experimental/Water/Source/Runtime/Public/WaterWaves.h
+++ b/Engine/Plugins/Experimental/Water/Source/Runtime/Public/WaterWaves.h
@@ -16,20 +16,20 @@ class UWaterWavesBase : public UObject
 
 public:
 	/** Returns the maximum wave height that can be reached by those waves */
-	UE_API virtual float GetMaxWaveHeight() const PURE_VIRTUAL(UWaterWavesBase::GetMaxWaveHeight, return 0.0f;)
+	virtual float GetMaxWaveHeight() const PURE_VIRTUAL(UWaterWavesBase::GetMaxWaveHeight, return 0.0f;)
 
 	/** Return the underlying Water class type. This will jump through the potential UWaterWavesAssetReference to get to the actual wave data */
-	UE_API virtual const UWaterWaves* GetWaterWaves() const PURE_VIRTUAL(UWaterWavesBase::GetWaterWaves, return nullptr;)
-	UE_API virtual UWaterWaves* GetWaterWaves() PURE_VIRTUAL(UWaterWavesBase::GetWaterWaves, return nullptr;)
+	virtual const UWaterWaves* GetWaterWaves() const PURE_VIRTUAL(UWaterWavesBase::GetWaterWaves, return nullptr;)
+	virtual UWaterWaves* GetWaterWaves() PURE_VIRTUAL(UWaterWavesBase::GetWaterWaves, return nullptr;)
 
 	/** Computes the raw wave perturbation of the water height/normal */
-	UE_API virtual float GetWaveHeightAtPosition(const FVector& InPosition, float InWaterDepth, float InTime, FVector& OutNormal) const PURE_VIRTUAL(UWaterWavesBase::GetWaveHeightAtPosition, return 0.0f;)
+	virtual float GetWaveHeightAtPosition(const FVector& InPosition, float InWaterDepth, float InTime, FVector& OutNormal) const PURE_VIRTUAL(UWaterWavesBase::GetWaveHeightAtPosition, return 0.0f;)
 
 	/** Computes the raw wave perturbation of the water height only (simple version : faster computation) */
-	UE_API virtual float GetSimpleWaveHeightAtPosition(const FVector& InPosition, float InWaterDepth, float InTime) const PURE_VIRTUAL(UWaterWavesBase::GetSimpleWaveHeightAtPosition, return 0.0f;)
+	virtual float GetSimpleWaveHeightAtPosition(const FVector& InPosition, float InWaterDepth, float InTime) const PURE_VIRTUAL(UWaterWavesBase::GetSimpleWaveHeightAtPosition, return 0.0f;)
 
 	/** Computes the attenuation factor to apply to the raw wave perturbation. Attenuates : normal/wave height/max wave height. */
-	UE_API virtual float GetWaveAttenuationFactor(const FVector& InPosition, float InWaterDepth, float InTargetWaveMaskDepth) const PURE_VIRTUAL(UWaterWavesBase::GetWaveAttenuationFactor, return 0.0f;)
+	virtual float GetWaveAttenuationFactor(const FVector& InPosition, float InWaterDepth, float InTargetWaveMaskDepth) const PURE_VIRTUAL(UWaterWavesBase::GetWaveAttenuationFactor, return 0.0f;)
 
 #if WITH_EDITOR
 	UE_API virtual void PostEditChangeProperty(struct FPropertyChangedEvent& PropertyChangedEvent) override;
diff --git a/Engine/Plugins/FX/Niagara/Source/Niagara/Classes/NiagaraDataInterfaceCurveBase.h b/Engine/Plugins/FX/Niagara/Source/Niagara/Classes/NiagaraDataInterfaceCurveBase.h
--- a/Engine/Plugins/FX/Niagara/Source/Niagara/Classes/NiagaraDataInterfaceCurveBase.h
+++ b/Engine/Plugins/FX/Niagara/Source/Niagara/Classes/NiagaraDataInterfaceCurveBase.h
@@ -152,12 +152,12 @@ public:
 	NIAGARA_API virtual void GetExposedVariables(TArray<FNiagaraVariableBase>& OutVariables) const override;
 	NIAGARA_API virtual bool GetExposedVariableValue(const FNiagaraVariableBase& InVariable, void* InPerInstanceData, FNiagaraSystemInstance* InSystemInstance, void* OutData) const override;
 
-	NIAGARA_API virtual int32 GetCurveNumElems() const PURE_VIRTUAL(UNiagaraDataInterfaceCurveBase::GetCurveNumElems(), return 0;)
+	virtual int32 GetCurveNumElems() const PURE_VIRTUAL(UNiagaraDataInterfaceCurveBase::GetCurveNumElems(), return 0;)
 #if WITH_EDITORONLY_DATA
-	NIAGARA_API virtual FName GetCurveSampleFunctionName() const PURE_VIRTUAL(UNiagaraDataInterfaceCurveBase::GetCurveSampleFunctionName(), return NAME_None;)
+	virtual FName GetCurveSampleFunctionName() const PURE_VIRTUAL(UNiagaraDataInterfaceCurveBase::GetCurveSampleFunctionName(), return NAME_None;)
 #endif
 
-	NIAGARA_API virtual void UpdateTimeRanges() PURE_VIRTUAL(UNiagaraDataInterfaceCurveBase::UpdateTimeRanges(), )
+	virtual void UpdateTimeRanges() PURE_VIRTUAL(UNiagaraDataInterfaceCurveBase::UpdateTimeRanges(), )
 	virtual TArray<float> BuildLUT(int32 NumEntries) const PURE_VIRTUAL(UNiagaraDataInterfaceCurveBase::BuildLUT(), return TArray<float>(); )
 
 	FORCEINLINE float GetMinTime()const { return LUTMinTime; }
diff --git a/Engine/Plugins/FX/Niagara/Source/Niagara/Classes/NiagaraEffectType.h b/Engine/Plugins/FX/Niagara/Source/Niagara/Classes/NiagaraEffectType.h
--- a/Engine/Plugins/FX/Niagara/Source/Niagara/Classes/NiagaraEffectType.h
+++ b/Engine/Plugins/FX/Niagara/Source/Niagara/Classes/NiagaraEffectType.h
@@ -365,7 +365,7 @@ class UNiagaraSignificanceHandler : public UObject
 	GENERATED_BODY()
 
 public:
-	NIAGARA_API virtual void CalculateSignificance(TConstArrayView<UNiagaraComponent*> Components, TArrayView<FNiagaraScalabilityState> OutState, TConstArrayView<FNiagaraScalabilitySystemData> SystemData, TArray<int32>& OutIndices)PURE_VIRTUAL(CalculateSignificance, );
+	virtual void CalculateSignificance(TConstArrayView<UNiagaraComponent*> Components, TArrayView<FNiagaraScalabilityState> OutState, TConstArrayView<FNiagaraScalabilitySystemData> SystemData, TArray<int32>& OutIndices)PURE_VIRTUAL(CalculateSignificance, );
 };
 
 /** Significance is determined by the system's distance to the nearest camera. Closer systems are more significant. */
@@ -560,4 +560,4 @@ public:
 	virtual bool NeedsSystemStats() const { return true; }
 	virtual bool NeedsComponentStats() const { return false; }
 };
-#endif//WITH_PER_FXTYPE_PARTICLE_PERF_STATS
\ No newline at end of file
+#endif//WITH_PER_FXTYPE_PARTICLE_PERF_STATS
diff --git a/Engine/Plugins/FX/Niagara/Source/Niagara/Public/NiagaraDataChannel.h b/Engine/Plugins/FX/Niagara/Source/Niagara/Public/NiagaraDataChannel.h
--- a/Engine/Plugins/FX/Niagara/Source/Niagara/Public/NiagaraDataChannel.h
+++ b/Engine/Plugins/FX/Niagara/Source/Niagara/Public/NiagaraDataChannel.h
@@ -183,7 +183,7 @@ public:
 	bool KeepPreviousFrameData() const { return bKeepPreviousFrameData; }
 
 	/** Create the appropriate handler object for this data channel. */
-	NIAGARA_API virtual UNiagaraDataChannelHandler* CreateHandler(UWorld* OwningWorld) const PURE_VIRTUAL(UNiagaraDataChannel::CreateHandler, {return nullptr;} );
+	virtual UNiagaraDataChannelHandler* CreateHandler(UWorld* OwningWorld) const PURE_VIRTUAL(UNiagaraDataChannel::CreateHandler, {return nullptr;} );
 	
 	NIAGARA_API const FNiagaraDataChannelLayoutInfoPtr GetLayoutInfo()const;
 
diff --git a/Engine/Plugins/FX/Niagara/Source/Niagara/Public/NiagaraDataChannelHandler.h b/Engine/Plugins/FX/Niagara/Source/Niagara/Public/NiagaraDataChannelHandler.h
--- a/Engine/Plugins/FX/Niagara/Source/Niagara/Public/NiagaraDataChannelHandler.h
+++ b/Engine/Plugins/FX/Niagara/Source/Niagara/Public/NiagaraDataChannelHandler.h
@@ -44,7 +44,7 @@ public:
 	Finds the correct internal data for this data channel and the given search parameters. For example in some cases this may require a search of several elements of data that correspond to different spacial regions.
 	This shared ptr provides access to Game level data, CPU simulation data and a render thread proxy that can be given to the RT and provides access to GPUSimulaiton data.
 	*/
-	NIAGARA_API virtual FNiagaraDataChannelDataPtr FindData(FNiagaraDataChannelSearchParameters SearchParams, ENiagaraResourceAccess AccessType)  PURE_VIRTUAL(UNiagaraDataChannelHandler::FindData, return nullptr;);
+	virtual FNiagaraDataChannelDataPtr FindData(FNiagaraDataChannelSearchParameters SearchParams, ENiagaraResourceAccess AccessType)  PURE_VIRTUAL(UNiagaraDataChannelHandler::FindData, return nullptr;);
 
 	const UNiagaraDataChannel* GetDataChannel()const
 	{
@@ -127,4 +127,4 @@ void UNiagaraDataChannelHandler::ForEachNDCData(TAction Func)
 			Func(DataPtr);
 		}
 	}
-}
\ No newline at end of file
+}
diff --git a/Engine/Plugins/FX/Niagara/Source/Niagara/Public/NiagaraRendererProperties.h b/Engine/Plugins/FX/Niagara/Source/Niagara/Public/NiagaraRendererProperties.h
--- a/Engine/Plugins/FX/Niagara/Source/Niagara/Public/NiagaraRendererProperties.h
+++ b/Engine/Plugins/FX/Niagara/Source/Niagara/Public/NiagaraRendererProperties.h
@@ -310,9 +310,9 @@ public:
 
 	//UObject Interface End
 	
-	NIAGARA_API virtual FNiagaraRenderer* CreateEmitterRenderer(ERHIFeatureLevel::Type FeatureLevel, const FNiagaraEmitterInstance* Emitter, const FNiagaraSystemInstanceController& InController) PURE_VIRTUAL ( UNiagaraRendererProperties::CreateEmitterRenderer, return nullptr;);
-	NIAGARA_API virtual class FNiagaraBoundsCalculator* CreateBoundsCalculator() PURE_VIRTUAL(UNiagaraRendererProperties::CreateBoundsCalculator, return nullptr;);
-	NIAGARA_API virtual void GetUsedMaterials(const FNiagaraEmitterInstance* InEmitter, TArray<UMaterialInterface*>& OutMaterials) const PURE_VIRTUAL(UNiagaraRendererProperties::GetUsedMaterials,);
+	virtual FNiagaraRenderer* CreateEmitterRenderer(ERHIFeatureLevel::Type FeatureLevel, const FNiagaraEmitterInstance* Emitter, const FNiagaraSystemInstanceController& InController) PURE_VIRTUAL ( UNiagaraRendererProperties::CreateEmitterRenderer, return nullptr;);
+	virtual class FNiagaraBoundsCalculator* CreateBoundsCalculator() PURE_VIRTUAL(UNiagaraRendererProperties::CreateBoundsCalculator, return nullptr;);
+	virtual void GetUsedMaterials(const FNiagaraEmitterInstance* InEmitter, TArray<UMaterialInterface*>& OutMaterials) const PURE_VIRTUAL(UNiagaraRendererProperties::GetUsedMaterials,);
 	virtual const FVertexFactoryType* GetVertexFactoryType() const { return nullptr; }
 	virtual bool IsBackfaceCullingDisabled() const { return false; }
 
@@ -377,7 +377,7 @@ public:
 	NIAGARA_API void CreateRendererWidgetsForAssets(TConstArrayView<UMaterialInterface*> InMaterials, TSharedPtr<FAssetThumbnailPool> InThumbnailPool, TArray<TSharedPtr<SWidget>>& OutWidgets) const;
 
 	NIAGARA_API virtual void GetRendererWidgets(const FNiagaraEmitterInstance* InEmitter, TArray<TSharedPtr<SWidget>>& OutWidgets, TSharedPtr<FAssetThumbnailPool> InThumbnailPool) const;
-	NIAGARA_API virtual void GetRendererTooltipWidgets(const FNiagaraEmitterInstance* InEmitter, TArray<TSharedPtr<SWidget>>& OutWidgets, TSharedPtr<FAssetThumbnailPool> InThumbnailPool) const PURE_VIRTUAL(UNiagaraRendererProperties::GetRendererTooltipWidgets, );
+	virtual void GetRendererTooltipWidgets(const FNiagaraEmitterInstance* InEmitter, TArray<TSharedPtr<SWidget>>& OutWidgets, TSharedPtr<FAssetThumbnailPool> InThumbnailPool) const PURE_VIRTUAL(UNiagaraRendererProperties::GetRendererTooltipWidgets, );
 	virtual void GetRendererFeedback(const FVersionedNiagaraEmitter& InEmitter, TArray<FText>& OutErrors, TArray<FText>& OutWarnings, TArray<FText>& OutInfo) const {}
 	NIAGARA_API virtual void GetRendererFeedback(const FVersionedNiagaraEmitter& InEmitter, TArray<FNiagaraRendererFeedback>& OutErrors, TArray<FNiagaraRendererFeedback>& OutWarnings, TArray<FNiagaraRendererFeedback>& OutInfo) const;
 
diff --git a/Engine/Plugins/FX/Niagara/Source/Niagara/Public/NiagaraSimulationStageBase.h b/Engine/Plugins/FX/Niagara/Source/Niagara/Public/NiagaraSimulationStageBase.h
--- a/Engine/Plugins/FX/Niagara/Source/Niagara/Public/NiagaraSimulationStageBase.h
+++ b/Engine/Plugins/FX/Niagara/Source/Niagara/Public/NiagaraSimulationStageBase.h
@@ -41,7 +41,7 @@ public:
 #if WITH_EDITORONLY_DATA
 	NIAGARA_API virtual bool AppendCompileHash(FNiagaraCompileHashVisitor* InVisitor) const;
 
-	NIAGARA_API virtual bool FillCompilationData(TArray<FNiagaraSimulationStageCompilationData>& CompilationSimStageData) const PURE_VIRTUAL(UNiagaraSimulationStageBase::FillCompileSimStageData, return false;);
+	virtual bool FillCompilationData(TArray<FNiagaraSimulationStageCompilationData>& CompilationSimStageData) const PURE_VIRTUAL(UNiagaraSimulationStageBase::FillCompileSimStageData, return false;);
 	/** Return the FName to use in place of the default for the location in the stack context. If this would be the default, return NAME_None.*/
 	virtual FName GetStackContextReplacementName() const { return NAME_None; }
 	NIAGARA_API void SetEnabled(bool bEnabled);
diff --git a/Engine/Plugins/Interchange/Runtime/Source/FactoryNodes/Public/InterchangeMeshFactoryNode.h b/Engine/Plugins/Interchange/Runtime/Source/FactoryNodes/Public/InterchangeMeshFactoryNode.h
--- a/Engine/Plugins/Interchange/Runtime/Source/FactoryNodes/Public/InterchangeMeshFactoryNode.h
+++ b/Engine/Plugins/Interchange/Runtime/Source/FactoryNodes/Public/InterchangeMeshFactoryNode.h
@@ -265,8 +265,8 @@ private:
 	UE::Interchange::TMapAttributeHelper<FString, FString> SlotMaterialDependencies;
 
 protected:
-	UE_API virtual void FillAssetClassFromAttribute() PURE_VIRTUAL("FillAssetClassFromAttribute");
-	UE_API virtual bool SetNodeClassFromClassAttribute() PURE_VIRTUAL("SetNodeClassFromClassAttribute", return false;);
+	virtual void FillAssetClassFromAttribute() PURE_VIRTUAL("FillAssetClassFromAttribute");
+	virtual bool SetNodeClassFromClassAttribute() PURE_VIRTUAL("SetNodeClassFromClassAttribute", return false;);
 
 	const UE::Interchange::FAttributeKey ClassNameAttributeKey = UE::Interchange::FBaseNodeStaticData::ClassTypeAttributeKey();
 
diff --git a/Engine/Plugins/Interchange/Runtime/Source/Import/Public/Texture/InterchangeSlicedTexturePayloadInterface.h b/Engine/Plugins/Interchange/Runtime/Source/Import/Public/Texture/InterchangeSlicedTexturePayloadInterface.h
--- a/Engine/Plugins/Interchange/Runtime/Source/Import/Public/Texture/InterchangeSlicedTexturePayloadInterface.h
+++ b/Engine/Plugins/Interchange/Runtime/Source/Import/Public/Texture/InterchangeSlicedTexturePayloadInterface.h
@@ -40,7 +40,7 @@ public:
 	 * @param AlternateTexturePath - When applicable, set to the path of the file actually loaded to create the FImportImage.
 	 * @return a PayloadData containing the import image data. The TOptional will not be set if there is an error.
 	 */
-	UE_API virtual TOptional<UE::Interchange::FImportSlicedImage> GetSlicedTexturePayloadData(const FString& PayloadKey, TOptional<FString>& AlternateTexturePath) const PURE_VIRTUAL(IInterchangeSlicedTexturePayloadInterface::GetSlicedTexturePayloadData, return{};);
+	virtual TOptional<UE::Interchange::FImportSlicedImage> GetSlicedTexturePayloadData(const FString& PayloadKey, TOptional<FString>& AlternateTexturePath) const PURE_VIRTUAL(IInterchangeSlicedTexturePayloadInterface::GetSlicedTexturePayloadData, return{};);
 };
 
 #undef UE_API
diff --git a/Engine/Plugins/Interchange/Runtime/Source/Import/Public/Texture/InterchangeTextureLightProfilePayloadInterface.h b/Engine/Plugins/Interchange/Runtime/Source/Import/Public/Texture/InterchangeTextureLightProfilePayloadInterface.h
--- a/Engine/Plugins/Interchange/Runtime/Source/Import/Public/Texture/InterchangeTextureLightProfilePayloadInterface.h
+++ b/Engine/Plugins/Interchange/Runtime/Source/Import/Public/Texture/InterchangeTextureLightProfilePayloadInterface.h
@@ -40,7 +40,7 @@ public:
 	 * @param AlternateTexturePath - When applicable, set to the path of the file actually loaded to create the FImportImage.
 	 * @return a PayloadData containing the import light profile data. The TOptional will not be set if there is an error.
 	 */
-	UE_API virtual TOptional<UE::Interchange::FImportLightProfile> GetLightProfilePayloadData(const FString& PayloadKey, TOptional<FString>& AlternateTexturePath) const PURE_VIRTUAL(IInterchangeTextureLightProfilePayloadInterface::GetLightProfilePayloadData, return{};);
+	virtual TOptional<UE::Interchange::FImportLightProfile> GetLightProfilePayloadData(const FString& PayloadKey, TOptional<FString>& AlternateTexturePath) const PURE_VIRTUAL(IInterchangeTextureLightProfilePayloadInterface::GetLightProfilePayloadData, return{};);
 
 	/**
 	* IESTranslator uses this, but we also expose it here in case otrher translators want to use directly.
diff --git a/Engine/Plugins/Interchange/Runtime/Source/Import/Public/Texture/InterchangeTexturePayloadInterface.h b/Engine/Plugins/Interchange/Runtime/Source/Import/Public/Texture/InterchangeTexturePayloadInterface.h
--- a/Engine/Plugins/Interchange/Runtime/Source/Import/Public/Texture/InterchangeTexturePayloadInterface.h
+++ b/Engine/Plugins/Interchange/Runtime/Source/Import/Public/Texture/InterchangeTexturePayloadInterface.h
@@ -40,7 +40,7 @@ public:
 	 * @param AlternateTexturePath - When applicable, set to the path of the file actually loaded to create the FImportImage.
 	 * @return a PayloadData containing the import image data. The TOptional will not be set if there is an error.
 	 */
-	UE_API virtual TOptional<UE::Interchange::FImportImage> GetTexturePayloadData(const FString& PayloadKey, TOptional<FString>& AlternateTexturePath) const PURE_VIRTUAL(IInterchangeTexturePayloadInterface::GetTexturePayloadData, return{};);
+	virtual TOptional<UE::Interchange::FImportImage> GetTexturePayloadData(const FString& PayloadKey, TOptional<FString>& AlternateTexturePath) const PURE_VIRTUAL(IInterchangeTexturePayloadInterface::GetTexturePayloadData, return{};);
 
 	/**
 	 * Does this translator support the request for a compressed source data?
diff --git a/Engine/Plugins/Mutable/Source/CustomizableObject/Public/MuCO/CustomizableInstanceLODManagement.h b/Engine/Plugins/Mutable/Source/CustomizableObject/Public/MuCO/CustomizableInstanceLODManagement.h
--- a/Engine/Plugins/Mutable/Source/CustomizableObject/Public/MuCO/CustomizableInstanceLODManagement.h
+++ b/Engine/Plugins/Mutable/Source/CustomizableObject/Public/MuCO/CustomizableInstanceLODManagement.h
@@ -20,13 +20,13 @@ public:
 	virtual ~UCustomizableInstanceLODManagementBase() {};
 
 	// WARNING! The following methods must be overriden in derived classes
-	UE_API virtual void UpdateInstanceDistsAndLODs(FMutableInstanceUpdateMap& InOutRequestedUpdates) PURE_VIRTUAL(UCustomizableInstanceLODManagementBase::UpdateInstanceDistsAndLODs, return;);
-	UE_API virtual int32 GetNumGeneratedInstancesLimitFullLODs() const PURE_VIRTUAL(UCustomizableInstanceLODManagementBase::GetNumGeneratedInstancesLimitFullLODs, return 0;);
-	UE_API virtual int32 GetNumGeneratedInstancesLimitLOD1() const PURE_VIRTUAL(UCustomizableInstanceLODManagementBase::GetNumGeneratedInstancesLimitLOD1, return 0;);
-	UE_API virtual int32 GetNumGeneratedInstancesLimitLOD2() const PURE_VIRTUAL(UCustomizableInstanceLODManagementBase::GetNumGeneratedInstancesLimitLOD2, return 0;);
-	UE_API virtual float GetOnlyUpdateCloseCustomizableObjectsDist() const PURE_VIRTUAL(UCustomizableInstanceLODManagementBase::GetOnlyUpdateCloseCustomizableObjectsDist, return 0.f;);
-	UE_API virtual bool IsOnlyUpdateCloseCustomizableObjectsEnabled() const PURE_VIRTUAL(UCustomizableInstanceLODManagementBase::IsOnlyUpdateCloseCustomizableObjectsEnabled, return false;);
-	UE_API virtual bool IsOnlyGenerateRequestedLODLevelsEnabled() const PURE_VIRTUAL(UCustomizableInstanceLODManagementBase::IsOnlyGenerateRequestedLODLevelsEnabled, return false;);
+	virtual void UpdateInstanceDistsAndLODs(FMutableInstanceUpdateMap& InOutRequestedUpdates) PURE_VIRTUAL(UCustomizableInstanceLODManagementBase::UpdateInstanceDistsAndLODs, return;);
+	virtual int32 GetNumGeneratedInstancesLimitFullLODs() const PURE_VIRTUAL(UCustomizableInstanceLODManagementBase::GetNumGeneratedInstancesLimitFullLODs, return 0;);
+	virtual int32 GetNumGeneratedInstancesLimitLOD1() const PURE_VIRTUAL(UCustomizableInstanceLODManagementBase::GetNumGeneratedInstancesLimitLOD1, return 0;);
+	virtual int32 GetNumGeneratedInstancesLimitLOD2() const PURE_VIRTUAL(UCustomizableInstanceLODManagementBase::GetNumGeneratedInstancesLimitLOD2, return 0;);
+	virtual float GetOnlyUpdateCloseCustomizableObjectsDist() const PURE_VIRTUAL(UCustomizableInstanceLODManagementBase::GetOnlyUpdateCloseCustomizableObjectsDist, return 0.f;);
+	virtual bool IsOnlyUpdateCloseCustomizableObjectsEnabled() const PURE_VIRTUAL(UCustomizableInstanceLODManagementBase::IsOnlyUpdateCloseCustomizableObjectsEnabled, return false;);
+	virtual bool IsOnlyGenerateRequestedLODLevelsEnabled() const PURE_VIRTUAL(UCustomizableInstanceLODManagementBase::IsOnlyGenerateRequestedLODLevelsEnabled, return false;);
 
 #if WITH_EDITOR
 	UE_API void EditorUpdateComponent(UCustomizableObjectInstanceUsage* InstanceUsage);
diff --git a/Engine/Plugins/Online/OnlineFramework/Source/Party/Public/Chat/SocialChatChannel.h b/Engine/Plugins/Online/OnlineFramework/Source/Party/Public/Chat/SocialChatChannel.h
--- a/Engine/Plugins/Online/OnlineFramework/Source/Party/Public/Chat/SocialChatChannel.h
+++ b/Engine/Plugins/Online/OnlineFramework/Source/Party/Public/Chat/SocialChatChannel.h
@@ -45,7 +45,7 @@ public:
 	DECLARE_EVENT_OneParam(USocialChatChannel, FOnChannelDisplayNameChanged, const FText&);
 	virtual FOnChannelDisplayNameChanged& OnChannelDisplayNameChanged() const { return OnChannelDisplayNameChangedEvent; }
 
-	UE_API virtual void Initialize(USocialUser* InSocialUser, const FChatRoomId& InChannelId, ESocialChannelType InSourceChannelType) PURE_VIRTUAL(USocialChatChannel::Initialize, return;);
+	virtual void Initialize(USocialUser* InSocialUser, const FChatRoomId& InChannelId, ESocialChannelType InSourceChannelType) PURE_VIRTUAL(USocialChatChannel::Initialize, return;);
 
 	/**
 	* Manually adds the given message to the channel's log locally. Representations of this channel on other clients will not receive the message.
@@ -62,7 +62,7 @@ public:
 	* Sends a text message to all other users in this channel.
 	* @return True if the message was sent successfully
 	*/
-	UE_API virtual bool SendMessage(const FString& Message) PURE_VIRTUAL(USocialChatChannel::SendMessage, return false;)
+	virtual bool SendMessage(const FString& Message) PURE_VIRTUAL(USocialChatChannel::SendMessage, return false;)
 
 	UE_API void SetChannelDisplayName(const FText& InDisplayName);
 	
diff --git a/Engine/Plugins/Online/OnlineFramework/Source/Party/Public/Chat/SocialChatMessage.h b/Engine/Plugins/Online/OnlineFramework/Source/Party/Public/Chat/SocialChatMessage.h
--- a/Engine/Plugins/Online/OnlineFramework/Source/Party/Public/Chat/SocialChatMessage.h
+++ b/Engine/Plugins/Online/OnlineFramework/Source/Party/Public/Chat/SocialChatMessage.h
@@ -41,7 +41,7 @@ public:
 	void SetPreviousMessage(TSharedPtr<FSocialChatMessage> InPrevious) { PreviousMessage = InPrevious; }
 	TSharedPtr<FSocialChatMessage> GetPreviousMessage() { return PreviousMessage; }
 
-	UE_API virtual FString GetSourceName() const PURE_VIRTUAL(FSocialChatMessage::GetSourceName, return TEXT(""););
+	virtual FString GetSourceName() const PURE_VIRTUAL(FSocialChatMessage::GetSourceName, return TEXT(""););
 	ESocialChannelType GetSourceChannelType() const { return SourceChannelType; }
 
 protected:
diff --git a/Engine/Plugins/Online/OnlineFramework/Source/Rejoin/Public/RejoinCheck.h b/Engine/Plugins/Online/OnlineFramework/Source/Rejoin/Public/RejoinCheck.h
--- a/Engine/Plugins/Online/OnlineFramework/Source/Rejoin/Public/RejoinCheck.h
+++ b/Engine/Plugins/Online/OnlineFramework/Source/Rejoin/Public/RejoinCheck.h
@@ -206,7 +206,7 @@ private:
 	 *
 	 * @return a rejoin status code that indicates whether or not this result should be pursued
 	 */
-	UE_API virtual ERejoinStatus GetRejoinStateFromSearchResult(const FOnlineSessionSearchResult& InSearchResult) const PURE_VIRTUAL(URejoinCheck::GetRejoinStateFromSearchResult, return ERejoinStatus::NeedsRecheck;);
+	virtual ERejoinStatus GetRejoinStateFromSearchResult(const FOnlineSessionSearchResult& InSearchResult) const PURE_VIRTUAL(URejoinCheck::GetRejoinStateFromSearchResult, return ERejoinStatus::NeedsRecheck;);
 	
 	/**
 	 * Delegate fired when a rejoin check has completed
@@ -238,7 +238,7 @@ private:
 	 * Game specific method to rejoin the last session in progress
 	 * Use OnRejoinFailure to communicate with base class the state of the rejoin
 	 */
-	UE_API virtual void RejoinViaSession() PURE_VIRTUAL(URejoinCheck::RejoinViaSession, );
+	virtual void RejoinViaSession() PURE_VIRTUAL(URejoinCheck::RejoinViaSession, );
 
 	/**
 	 * Set the timer for another rejoin search after a given time period
diff --git a/Engine/Plugins/PCG/Source/PCG/Public/Data/PCGBasePointData.h b/Engine/Plugins/PCG/Source/PCG/Public/Data/PCGBasePointData.h
--- a/Engine/Plugins/PCG/Source/PCG/Public/Data/PCGBasePointData.h
+++ b/Engine/Plugins/PCG/Source/PCG/Public/Data/PCGBasePointData.h
@@ -61,20 +61,20 @@ public:
 	/** Adds a single point based on the given actor */
 	UE_API void AddSinglePointFromActor(AActor* InActor, bool* bOutOptionalSanitizedTagAttributeName = nullptr);
 
-	UE_API virtual bool IsValidRef(const PCGPointOctree::FPointRef& InPointRef) const PURE_VIRTUAL(UPCGBasePointData::IsValidRef, return false;);
+	virtual bool IsValidRef(const PCGPointOctree::FPointRef& InPointRef) const PURE_VIRTUAL(UPCGBasePointData::IsValidRef, return false;);
 	
 	virtual EPCGDataType GetDataType() const override { return EPCGDataType::Point; }
 
-	UE_API virtual int32 GetNumPoints() const PURE_VIRTUAL(UPCGBasePointData::GetNumPoints, return 0;);
+	virtual int32 GetNumPoints() const PURE_VIRTUAL(UPCGBasePointData::GetNumPoints, return 0;);
 	bool IsEmpty() const { return GetNumPoints() == 0; }
-	UE_API virtual void SetNumPoints(int32 InNumPoints, bool bInitializeValues = true) PURE_VIRTUAL(UPCGBasePointData::SetNumPoints, return;);
+	virtual void SetNumPoints(int32 InNumPoints, bool bInitializeValues = true) PURE_VIRTUAL(UPCGBasePointData::SetNumPoints, return;);
 
 	virtual void AllocateProperties(EPCGPointNativeProperties Properties) {}
 	virtual void FreeProperties(EPCGPointNativeProperties Properties) {}
-	UE_API virtual void MoveRange(int32 RangeStartIndex, int32 MoveToIndex, int32 NumElements) PURE_VIRTUAL(UPCGBasePointData::MoveRange, return;);
+	virtual void MoveRange(int32 RangeStartIndex, int32 MoveToIndex, int32 NumElements) PURE_VIRTUAL(UPCGBasePointData::MoveRange, return;);
 	virtual void CopyUnallocatedPropertiesFrom(const UPCGBasePointData* InPointData) {}
 
-	UE_API virtual TArray<FTransform> GetTransformsCopy() const PURE_VIRTUAL(UPCGBasePointData::GetTransformsCopy, return TArray<FTransform>(););
+	virtual TArray<FTransform> GetTransformsCopy() const PURE_VIRTUAL(UPCGBasePointData::GetTransformsCopy, return TArray<FTransform>(););
 	
 	/**
 	 * GetValueRange / GetConstValueRange / Get*ValueRange / GetConst*ValueRange
@@ -88,26 +88,26 @@ public:
 	template <typename T>
 	TPCGValueRange<T> GetValueRange(EPCGPointNativeProperties NativeProperty, bool bAllocate = true);
 
-	UE_API virtual TPCGValueRange<FTransform> GetTransformValueRange(bool bAllocate = true) PURE_VIRTUAL(UPCGBasePointData::GetTransformValueRange, return TPCGValueRange<FTransform>(););
-	UE_API virtual TPCGValueRange<float> GetDensityValueRange(bool bAllocate = true) PURE_VIRTUAL(UPCGBasePointData::GetDensityValueRange, return TPCGValueRange<float>(););
-	UE_API virtual TPCGValueRange<FVector> GetBoundsMinValueRange(bool bAllocate = true) PURE_VIRTUAL(UPCGBasePointData::GetBoundsMinValueRange, return TPCGValueRange<FVector>(););
-	UE_API virtual TPCGValueRange<FVector> GetBoundsMaxValueRange(bool bAllocate = true) PURE_VIRTUAL(UPCGBasePointData::GetBoundsMaxValueRange, return TPCGValueRange<FVector>(););
-	UE_API virtual TPCGValueRange<FVector4> GetColorValueRange(bool bAllocate = true) PURE_VIRTUAL(UPCGBasePointData::GetColorValueRange, return TPCGValueRange<FVector4>(););
-	UE_API virtual TPCGValueRange<float> GetSteepnessValueRange(bool bAllocate = true) PURE_VIRTUAL(UPCGBasePointData::GetSteepnessValueRange, return TPCGValueRange<float>(););
-	UE_API virtual TPCGValueRange<int32> GetSeedValueRange(bool bAllocate = true) PURE_VIRTUAL(UPCGBasePointData::GetSeedValueRange, return TPCGValueRange<int32>(););
-	UE_API virtual TPCGValueRange<int64> GetMetadataEntryValueRange(bool bAllocate = true) PURE_VIRTUAL(UPCGBasePointData::GetMetadataEntryValueRange, return TPCGValueRange<int64>(););
+	virtual TPCGValueRange<FTransform> GetTransformValueRange(bool bAllocate = true) PURE_VIRTUAL(UPCGBasePointData::GetTransformValueRange, return TPCGValueRange<FTransform>(););
+	virtual TPCGValueRange<float> GetDensityValueRange(bool bAllocate = true) PURE_VIRTUAL(UPCGBasePointData::GetDensityValueRange, return TPCGValueRange<float>(););
+	virtual TPCGValueRange<FVector> GetBoundsMinValueRange(bool bAllocate = true) PURE_VIRTUAL(UPCGBasePointData::GetBoundsMinValueRange, return TPCGValueRange<FVector>(););
+	virtual TPCGValueRange<FVector> GetBoundsMaxValueRange(bool bAllocate = true) PURE_VIRTUAL(UPCGBasePointData::GetBoundsMaxValueRange, return TPCGValueRange<FVector>(););
+	virtual TPCGValueRange<FVector4> GetColorValueRange(bool bAllocate = true) PURE_VIRTUAL(UPCGBasePointData::GetColorValueRange, return TPCGValueRange<FVector4>(););
+	virtual TPCGValueRange<float> GetSteepnessValueRange(bool bAllocate = true) PURE_VIRTUAL(UPCGBasePointData::GetSteepnessValueRange, return TPCGValueRange<float>(););
+	virtual TPCGValueRange<int32> GetSeedValueRange(bool bAllocate = true) PURE_VIRTUAL(UPCGBasePointData::GetSeedValueRange, return TPCGValueRange<int32>(););
+	virtual TPCGValueRange<int64> GetMetadataEntryValueRange(bool bAllocate = true) PURE_VIRTUAL(UPCGBasePointData::GetMetadataEntryValueRange, return TPCGValueRange<int64>(););
 
 	template <typename T>
 	TConstPCGValueRange<T> GetConstValueRange(EPCGPointNativeProperties NativeProperty) const;
 
-	UE_API virtual TConstPCGValueRange<FTransform> GetConstTransformValueRange() const PURE_VIRTUAL(UPCGBasePointData::GetConstTransformValueRange, return TConstPCGValueRange<FTransform>(););
-	UE_API virtual TConstPCGValueRange<float> GetConstDensityValueRange() const PURE_VIRTUAL(UPCGBasePointData::GetConstDensityValueRange, return TConstPCGValueRange<float>(););
-	UE_API virtual TConstPCGValueRange<FVector> GetConstBoundsMinValueRange() const PURE_VIRTUAL(UPCGBasePointData::GetConstBoundsMinValueRange, return TConstPCGValueRange<FVector>(););
-	UE_API virtual TConstPCGValueRange<FVector> GetConstBoundsMaxValueRange() const PURE_VIRTUAL(UPCGBasePointData::GetConstBoundsMaxValueRange, return TConstPCGValueRange<FVector>(););
-	UE_API virtual TConstPCGValueRange<FVector4> GetConstColorValueRange() const PURE_VIRTUAL(UPCGBasePointData::GetConstColorValueRange, return TConstPCGValueRange<FVector4>(););
-	UE_API virtual TConstPCGValueRange<float> GetConstSteepnessValueRange() const PURE_VIRTUAL(UPCGBasePointData::GetConstSteepnessValueRange, return TConstPCGValueRange<float>(););
-	UE_API virtual TConstPCGValueRange<int32> GetConstSeedValueRange() const PURE_VIRTUAL(UPCGBasePointData::GetConstSeedValueRange, return TConstPCGValueRange<int32>(););
-	UE_API virtual TConstPCGValueRange<int64> GetConstMetadataEntryValueRange() const PURE_VIRTUAL(UPCGBasePointData::GetConstMetadataEntryValueRange, return TConstPCGValueRange<int64>(););
+	virtual TConstPCGValueRange<FTransform> GetConstTransformValueRange() const PURE_VIRTUAL(UPCGBasePointData::GetConstTransformValueRange, return TConstPCGValueRange<FTransform>(););
+	virtual TConstPCGValueRange<float> GetConstDensityValueRange() const PURE_VIRTUAL(UPCGBasePointData::GetConstDensityValueRange, return TConstPCGValueRange<float>(););
+	virtual TConstPCGValueRange<FVector> GetConstBoundsMinValueRange() const PURE_VIRTUAL(UPCGBasePointData::GetConstBoundsMinValueRange, return TConstPCGValueRange<FVector>(););
+	virtual TConstPCGValueRange<FVector> GetConstBoundsMaxValueRange() const PURE_VIRTUAL(UPCGBasePointData::GetConstBoundsMaxValueRange, return TConstPCGValueRange<FVector>(););
+	virtual TConstPCGValueRange<FVector4> GetConstColorValueRange() const PURE_VIRTUAL(UPCGBasePointData::GetConstColorValueRange, return TConstPCGValueRange<FVector4>(););
+	virtual TConstPCGValueRange<float> GetConstSteepnessValueRange() const PURE_VIRTUAL(UPCGBasePointData::GetConstSteepnessValueRange, return TConstPCGValueRange<float>(););
+	virtual TConstPCGValueRange<int32> GetConstSeedValueRange() const PURE_VIRTUAL(UPCGBasePointData::GetConstSeedValueRange, return TConstPCGValueRange<int32>(););
+	virtual TConstPCGValueRange<int64> GetConstMetadataEntryValueRange() const PURE_VIRTUAL(UPCGBasePointData::GetConstMetadataEntryValueRange, return TConstPCGValueRange<int64>(););
 		
 	UE_API virtual void SetTransform(const FTransform& InTransform);
 	UE_API virtual void SetDensity(float InDensity);
diff --git a/Engine/Plugins/PCG/Source/PCG/Public/Data/PCGPolyLineData.h b/Engine/Plugins/PCG/Source/PCG/Public/Data/PCGPolyLineData.h
--- a/Engine/Plugins/PCG/Source/PCG/Public/Data/PCGPolyLineData.h
+++ b/Engine/Plugins/PCG/Source/PCG/Public/Data/PCGPolyLineData.h
@@ -29,26 +29,26 @@ public:
 
 	/** Get the number of segments in this line. If the line is closed, this will be the same as the number of control points in the line. */
 	UFUNCTION(BlueprintCallable, Category = "PolylineData")
-	UE_API virtual int GetNumSegments() const PURE_VIRTUAL(UPCGPolyLineData::GetNumSegments, return 0;);
+	virtual int GetNumSegments() const PURE_VIRTUAL(UPCGPolyLineData::GetNumSegments, return 0;);
 
 	/** Get the length of a specific segment of the line. */
 	UFUNCTION(BlueprintCallable, Category = "PolylineData")
-	UE_API virtual double GetSegmentLength(int SegmentIndex) const PURE_VIRTUAL(UPCGPolyLineData::GetSegmentLength, return 0;);
+	virtual double GetSegmentLength(int SegmentIndex) const PURE_VIRTUAL(UPCGPolyLineData::GetSegmentLength, return 0;);
 
 	/** Get the location of the point at the normalized [0, 1] parameter across the entire the poly line. */
 	UFUNCTION(BlueprintCallable, Category = "PolylineData")
-	UE_API virtual FVector GetLocationAtAlpha(float Alpha) const PURE_VIRTUAL(UPCGPolyLine::GetLocationAtAlpha, return FVector::ZeroVector;)
+	virtual FVector GetLocationAtAlpha(float Alpha) const PURE_VIRTUAL(UPCGPolyLine::GetLocationAtAlpha, return FVector::ZeroVector;)
 
 	/** Get the full transform at the normalized [0, 1] parameter across the entire the poly line. */
 	UFUNCTION(BlueprintCallable, Category = "PolylineData")
-	UE_API virtual FTransform GetTransformAtAlpha(float Alpha) const PURE_VIRTUAL(UPCGPolyLine::GetTransformAtAlpha, return FTransform::Identity;)
+	virtual FTransform GetTransformAtAlpha(float Alpha) const PURE_VIRTUAL(UPCGPolyLine::GetTransformAtAlpha, return FTransform::Identity;)
 
 	/** Get the total length of the line. */
 	UFUNCTION(BlueprintCallable, Category = "PolylineData")
 	UE_API virtual double GetLength() const;
 
 	/** Get the location at a distance along the line. */
-	UE_API virtual FTransform GetTransformAtDistance(int SegmentIndex, double Distance, bool bWorldSpace = true, FBox* OutBounds = nullptr) const PURE_VIRTUAL(UPCGPolyLine::GetTransformAtDistance, return FTransform(););
+	virtual FTransform GetTransformAtDistance(int SegmentIndex, double Distance, bool bWorldSpace = true, FBox* OutBounds = nullptr) const PURE_VIRTUAL(UPCGPolyLine::GetTransformAtDistance, return FTransform(););
 
 	/** Get the location at a distance along the line. */
 	UFUNCTION(BlueprintCallable, Category = "PolylineData", meta = (DisplayName = "Get Transform At Distance", AutoCreateRefTerm = "OutBounds"))
diff --git a/Engine/Plugins/PCG/Source/PCG/Public/Data/PCGResourceData.h b/Engine/Plugins/PCG/Source/PCG/Public/Data/PCGResourceData.h
--- a/Engine/Plugins/PCG/Source/PCG/Public/Data/PCGResourceData.h
+++ b/Engine/Plugins/PCG/Source/PCG/Public/Data/PCGResourceData.h
@@ -21,7 +21,7 @@ public:
 	virtual EPCGDataType GetDataType() const override { return EPCGDataType::Resource; }
 	//~ End UPCGData interface
 
-	UE_API virtual FSoftObjectPath GetResourcePath() const PURE_VIRTUAL(UPCGResourceData::GetResourcePath, return FSoftObjectPath(););
+	virtual FSoftObjectPath GetResourcePath() const PURE_VIRTUAL(UPCGResourceData::GetResourcePath, return FSoftObjectPath(););
 
 	UE_API TSharedPtr<FStreamableHandle> RequestResourceLoad(bool bAsynchronous = true) const;
 };
diff --git a/Engine/Plugins/PCG/Source/PCG/Public/Data/PCGSpatialData.h b/Engine/Plugins/PCG/Source/PCG/Public/Data/PCGSpatialData.h
--- a/Engine/Plugins/PCG/Source/PCG/Public/Data/PCGSpatialData.h
+++ b/Engine/Plugins/PCG/Source/PCG/Public/Data/PCGSpatialData.h
@@ -103,11 +103,11 @@ public:
 
 	/** Returns the dimension of the data type, which has nothing to do with the dimension of its points */
 	UFUNCTION(BlueprintCallable, Category = SpatialData)
-	UE_API virtual int GetDimension() const PURE_VIRTUAL(UPCGSpatialData::GetDimension, return 0;);
+	virtual int GetDimension() const PURE_VIRTUAL(UPCGSpatialData::GetDimension, return 0;);
 
 	/** Returns the full bounds (including density fall-off) of the data */
 	UFUNCTION(BlueprintCallable, Category = SpatialData)
-	UE_API virtual FBox GetBounds() const PURE_VIRTUAL(UPCGSpatialData::GetBounds, return FBox(EForceInit::ForceInit););
+	virtual FBox GetBounds() const PURE_VIRTUAL(UPCGSpatialData::GetBounds, return FBox(EForceInit::ForceInit););
 
 	/** Returns whether a given spatial data is bounded as some data types do not require bounds by themselves */
 	virtual bool IsBounded() const { return true; }
@@ -131,14 +131,14 @@ public:
 	UFUNCTION(BlueprintCallable, Category = SpatialData)
 	const UPCGPointData* ToPointDataWithContext(UPARAM(ref) FPCGContext& Context) const { return ToPointData(&Context); }
 
-	UE_API virtual const UPCGPointData* ToPointData(FPCGContext* Context, const FBox& InBounds = FBox(EForceInit::ForceInit)) const PURE_VIRTUAL(UPCGSpatialData::ToPointData, return nullptr;);
+	virtual const UPCGPointData* ToPointData(FPCGContext* Context, const FBox& InBounds = FBox(EForceInit::ForceInit)) const PURE_VIRTUAL(UPCGSpatialData::ToPointData, return nullptr;);
 
 	UE_API virtual const UPCGPointArrayData* ToPointArrayData(FPCGContext* Context, const FBox& InBounds = FBox(EForceInit::ForceInit)) const;
 
 	UE_API const UPCGBasePointData* ToBasePointData(FPCGContext* Context, const FBox& InBounds = FBox(EForceInit::ForceInit)) const;
 
 	/** Sample rotation, scale and other attributes from this data at the query position. Returns true if Transform location and Bounds overlaps this data. */
-	UE_API virtual bool SamplePoint(const FTransform& Transform, const FBox& Bounds, FPCGPoint& OutPoint, UPCGMetadata* OutMetadata) const PURE_VIRTUAL(UPCGSpatialData::SamplePoint, return false;);
+	virtual bool SamplePoint(const FTransform& Transform, const FBox& Bounds, FPCGPoint& OutPoint, UPCGMetadata* OutMetadata) const PURE_VIRTUAL(UPCGSpatialData::SamplePoint, return false;);
 
 	/** Performs multiple samples at the same time.
 	* Contrary to the single SamplePoint call, this is expected to set the density to 0 for points that were not overlapping - but the other properties can be anything.
@@ -244,7 +244,7 @@ protected:
 	UE_DEPRECATED(5.5, "Call/Implement version with FPCGContext parameter")
 	virtual UPCGSpatialData* CopyInternal() const { return nullptr; }
 
-	UE_API virtual UPCGSpatialData* CopyInternal(FPCGContext* Context) const PURE_VIRTUAL(UPCGSpatialData::CopyInternal, return nullptr;);
+	virtual UPCGSpatialData* CopyInternal(FPCGContext* Context) const PURE_VIRTUAL(UPCGSpatialData::CopyInternal, return nullptr;);
 	UE_API virtual void InitializeSpatialDataInternal(const FPCGInitializeFromDataParams& InParams);
 	UE_API void InitializeMetadata(const FPCGInitializeFromDataParams& InParams);
 	UE_API virtual void InitializeMetadataInternal(const FPCGInitializeFromDataParams& InParams);
@@ -293,7 +293,7 @@ public:
 
 protected:
 	virtual bool SupportsBoundedPointData() const { return false; }
-	UE_API virtual const UPCGPointData* CreatePointData(FPCGContext* Context) const PURE_VIRTUAL(UPCGSpatialData::CreatePointData, return nullptr;);
+	virtual const UPCGPointData* CreatePointData(FPCGContext* Context) const PURE_VIRTUAL(UPCGSpatialData::CreatePointData, return nullptr;);
 	virtual const UPCGPointData* CreatePointData(FPCGContext* Context, const FBox& InBounds) const { return CreatePointData(Context); }
 	UE_API virtual const UPCGPointArrayData* CreatePointArrayData(FPCGContext* Context, const FBox& InBounds) const; 
 
diff --git a/Engine/Plugins/PCG/Source/PCG/Public/Data/PCGTextureData.h b/Engine/Plugins/PCG/Source/PCG/Public/Data/PCGTextureData.h
--- a/Engine/Plugins/PCG/Source/PCG/Public/Data/PCGTextureData.h
+++ b/Engine/Plugins/PCG/Source/PCG/Public/Data/PCGTextureData.h
@@ -93,9 +93,9 @@ public:
 
 	UE_API virtual bool IsValid() const;
 
-	UE_API virtual UTexture* GetTexture() const PURE_VIRTUAL(UPCGBaseTextureData::GetTexture, return nullptr;)
-	UE_API virtual FTextureRHIRef GetTextureRHI() const PURE_VIRTUAL(UPCGBaseTextureData::GetTextureResource, return nullptr;);
-	UE_API virtual EPCGTextureResourceType GetTextureResourceType() const PURE_VIRTUAL(UPCGBaseTextureData::GetTextureResourceType, return EPCGTextureResourceType::Invalid;);
+	virtual UTexture* GetTexture() const PURE_VIRTUAL(UPCGBaseTextureData::GetTexture, return nullptr;)
+	virtual FTextureRHIRef GetTextureRHI() const PURE_VIRTUAL(UPCGBaseTextureData::GetTextureResource, return nullptr;);
+	virtual EPCGTextureResourceType GetTextureResourceType() const PURE_VIRTUAL(UPCGBaseTextureData::GetTextureResourceType, return EPCGTextureResourceType::Invalid;);
 	virtual TRefCountPtr<IPooledRenderTarget> GetRefCountedTexture() const { return nullptr; }
 	virtual int GetTextureSlice() const { return 0; }
 	virtual FIntPoint GetTextureSize() const { return FIntPoint(Width, Height); }
diff --git a/Engine/Plugins/PCG/Source/PCG/Public/PCGGraph.h b/Engine/Plugins/PCG/Source/PCG/Public/PCGGraph.h
--- a/Engine/Plugins/PCG/Source/PCG/Public/PCGGraph.h
+++ b/Engine/Plugins/PCG/Source/PCG/Public/PCGGraph.h
@@ -116,10 +116,10 @@ public:
 	UFUNCTION(BlueprintCallable, Category = Graph)
 	const UPCGGraph* GetConstPCGGraph() const { return GetGraph(); }
 
-	UE_API virtual UPCGGraph* GetGraph() PURE_VIRTUAL(UPCGGraphInterface::GetGraph, return nullptr;)
-	UE_API virtual const UPCGGraph* GetGraph() const PURE_VIRTUAL(UPCGGraphInterface::GetGraph, return nullptr;)
+	virtual UPCGGraph* GetGraph() PURE_VIRTUAL(UPCGGraphInterface::GetGraph, return nullptr;)
+	virtual const UPCGGraph* GetGraph() const PURE_VIRTUAL(UPCGGraphInterface::GetGraph, return nullptr;)
 
-	UE_API virtual const FInstancedPropertyBag* GetUserParametersStruct() const PURE_VIRTUAL(UPCGGraphInterface::GetUserParametersStruct, return nullptr;)
+	virtual const FInstancedPropertyBag* GetUserParametersStruct() const PURE_VIRTUAL(UPCGGraphInterface::GetUserParametersStruct, return nullptr;)
 
 	// Mutable version - should not be used outside of testing. Use UpdateUserParametersStruct on PCGGraph,
 	// or the different get/set/update functions in PCGGraphInterface instead to have proper callbacks.
@@ -241,10 +241,10 @@ public:
 	 */
 	UE_API bool UpdateSetGraphParameter(const FName PropertyName, TFunctionRef<bool(FPropertyBagSetRef& PropertyBagSetRef)> Callback);
 
-	UE_API virtual void OnGraphParametersChanged(EPCGGraphParameterEvent InChangeType, FName InChangedPropertyName) PURE_VIRTUAL(UPCGGraphInterface::OnGraphParametersChanged, )
+	virtual void OnGraphParametersChanged(EPCGGraphParameterEvent InChangeType, FName InChangedPropertyName) PURE_VIRTUAL(UPCGGraphInterface::OnGraphParametersChanged, )
 
 protected:
-	UE_API virtual FInstancedPropertyBag* GetMutableUserParametersStruct() PURE_VIRTUAL(UPCGGraphInterface::GetMutableUserParametersStruct, return nullptr;)
+	virtual FInstancedPropertyBag* GetMutableUserParametersStruct() PURE_VIRTUAL(UPCGGraphInterface::GetMutableUserParametersStruct, return nullptr;)
 
 	/** Detecting if we need to refresh the graph depending on the type of change in the Graph Parameter. */
 	UE_API EPCGChangeType GetChangeTypeForGraphParameterChange(EPCGGraphParameterEvent InChangeType, FName InChangedPropertyName);
diff --git a/Engine/Plugins/PCG/Source/PCG/Public/PCGManagedResource.h b/Engine/Plugins/PCG/Source/PCG/Public/PCGManagedResource.h
--- a/Engine/Plugins/PCG/Source/PCG/Public/PCGManagedResource.h
+++ b/Engine/Plugins/PCG/Source/PCG/Public/PCGManagedResource.h
@@ -174,8 +174,8 @@ public:
 	virtual bool SupportsComponentReset() const { return false; }
 
 protected:
-	UE_API virtual TArrayView<TSoftObjectPtr<UActorComponent>> GetComponentsArray() PURE_VIRTUAL(UPCGManagedComponentBase::GetComponentsArray, return TArrayView<TSoftObjectPtr<UActorComponent>>{};);
-	UE_API virtual int32 GetComponentsCount() const PURE_VIRTUAL(UPCGManagedComponentBase::GetComponentsCount, return 0;);
+	virtual TArrayView<TSoftObjectPtr<UActorComponent>> GetComponentsArray() PURE_VIRTUAL(UPCGManagedComponentBase::GetComponentsArray, return TArrayView<TSoftObjectPtr<UActorComponent>>{};);
+	virtual int32 GetComponentsCount() const PURE_VIRTUAL(UPCGManagedComponentBase::GetComponentsCount, return 0;);
 	UE_API void SetupGeneratedComponentFromBP(TSoftObjectPtr<UActorComponent> InGeneratedComponent);
 };
 
diff --git a/Engine/Plugins/PCG/Source/PCG/Public/PCGSettings.h b/Engine/Plugins/PCG/Source/PCG/Public/PCGSettings.h
--- a/Engine/Plugins/PCG/Source/PCG/Public/PCGSettings.h
+++ b/Engine/Plugins/PCG/Source/PCG/Public/PCGSettings.h
@@ -135,8 +135,8 @@ class UPCGSettingsInterface : public UPCGData
 	GENERATED_BODY()
 
 public:
-	UE_API virtual UPCGSettings* GetSettings() PURE_VIRTUAL(UPCGSettingsInterface::GetSettings, return nullptr;);
-	UE_API virtual const UPCGSettings* GetSettings() const PURE_VIRTUAL(UPCGSettingsInterface::GetSettings, return nullptr;);
+	virtual UPCGSettings* GetSettings() PURE_VIRTUAL(UPCGSettingsInterface::GetSettings, return nullptr;);
+	virtual const UPCGSettings* GetSettings() const PURE_VIRTUAL(UPCGSettingsInterface::GetSettings, return nullptr;);
 
 	UE_API bool IsInstance() const;
 	/** Dedicated method to change enable state because some nodes have more complex behavior on enable/disable (such as subgraphs) */
@@ -407,7 +407,7 @@ protected:
 	UFUNCTION(BlueprintCallable, Category = "Settings|DynamicPins", DisplayName = "Get Type Union Of Incident Edges")
 	int32 BP_GetTypeUnionOfIncidentEdges(const FName& PinLabel) const { return static_cast<int32>(GetTypeUnionOfIncidentEdges(PinLabel)); }
 
-	UE_API virtual FPCGElementPtr CreateElement() const PURE_VIRTUAL(UPCGSettings::CreateElement, return nullptr;);
+	virtual FPCGElementPtr CreateElement() const PURE_VIRTUAL(UPCGSettings::CreateElement, return nullptr;);
 
 	/** An additional custom version number that external system users can use to track versions. This version will be serialized into the asset and will be provided by UserDataVersion after load. */
 	virtual FGuid GetUserCustomVersionGuid() { return FGuid(); }
diff --git a/Engine/Plugins/PCG/Source/PCG/Public/PCGSettingsWithDynamicInputs.h b/Engine/Plugins/PCG/Source/PCG/Public/PCGSettingsWithDynamicInputs.h
--- a/Engine/Plugins/PCG/Source/PCG/Public/PCGSettingsWithDynamicInputs.h
+++ b/Engine/Plugins/PCG/Source/PCG/Public/PCGSettingsWithDynamicInputs.h
@@ -26,7 +26,7 @@ public:
 	/** User driven event to add a dynamic source pin */
 	UE_API virtual void OnUserAddDynamicInputPin();
 	/** Overridden logic to add a default source pin */
-	UE_API virtual void AddDefaultDynamicInputPin() PURE_VIRTUAL(PCGDynamicSettings::AddDefaultSourcePin, );
+	virtual void AddDefaultDynamicInputPin() PURE_VIRTUAL(PCGDynamicSettings::AddDefaultSourcePin, );
 	/** Check if the pin to remove is dynamic */
 	UE_API virtual bool CanUserRemoveDynamicInputPin(int32 PinIndex);
 	/** User driven event to remove a dynamic source pin */
diff --git a/Engine/Plugins/PCG/Source/PCG/Public/RuntimeGen/GenSources/PCGGenSourceBase.h b/Engine/Plugins/PCG/Source/PCG/Public/RuntimeGen/GenSources/PCGGenSourceBase.h
--- a/Engine/Plugins/PCG/Source/PCG/Public/RuntimeGen/GenSources/PCGGenSourceBase.h
+++ b/Engine/Plugins/PCG/Source/PCG/Public/RuntimeGen/GenSources/PCGGenSourceBase.h
@@ -28,10 +28,10 @@ public:
 	virtual void Tick() {}
 
 	/** Returns the world space position of this gen source. */
-	UE_API virtual TOptional<FVector> GetPosition() const PURE_VIRTUAL(UPCGGenSourceBase::GetPosition, return TOptional<FVector>(););
+	virtual TOptional<FVector> GetPosition() const PURE_VIRTUAL(UPCGGenSourceBase::GetPosition, return TOptional<FVector>(););
 
 	/** Returns the normalized forward vector of this gen source. */
-	UE_API virtual TOptional<FVector> GetDirection() const PURE_VIRTUAL(UPCGGenSourceBase::GetDirection, return TOptional<FVector>(););
+	virtual TOptional<FVector> GetDirection() const PURE_VIRTUAL(UPCGGenSourceBase::GetDirection, return TOptional<FVector>(););
 
 	/** Returns the view frustum of this gen source. */
 	virtual TOptional<FConvexVolume> GetViewFrustum(bool bIs2DGrid) const { return TOptional<FConvexVolume>(); }
diff --git a/Engine/Plugins/PCG/Source/PCG/Public/RuntimeGen/SchedulingPolicies/PCGSchedulingPolicyBase.h b/Engine/Plugins/PCG/Source/PCG/Public/RuntimeGen/SchedulingPolicies/PCGSchedulingPolicyBase.h
--- a/Engine/Plugins/PCG/Source/PCG/Public/RuntimeGen/SchedulingPolicies/PCGSchedulingPolicyBase.h
+++ b/Engine/Plugins/PCG/Source/PCG/Public/RuntimeGen/SchedulingPolicies/PCGSchedulingPolicyBase.h
@@ -24,7 +24,7 @@ class UPCGSchedulingPolicyBase : public UObject
 
 public:
 	/** Calculate the runtime scheduling priority with respect to a Generation Source. Should return a value in the range [0, 1], where higher values will be scheduled sooner. */
-	UE_API virtual double CalculatePriority(const IPCGGenSourceBase* InGenSource, const FBox& GenerationBounds, bool bUse2DGrid) const 
+	virtual double CalculatePriority(const IPCGGenSourceBase* InGenSource, const FBox& GenerationBounds, bool bUse2DGrid) const 
 		PURE_VIRTUAL(UPCGSchedulingPolicyBase::CalculatePriority, return 0.0;);
 
 	/** True if the generation source would consider the given bounds for generation. */
@@ -34,7 +34,7 @@ public:
 	virtual bool ShouldCull(const IPCGGenSourceBase* InGenSource, const FBox& GenerationBounds, bool bUse2DGrid) const { return false; }
 
 	/** A SchedulingPolicy is equivalent to another SchedulingPolicy if they are the same (same ptr), or if they have the same type and parameter values. */
-	UE_API virtual bool IsEquivalent(const UPCGSchedulingPolicyBase* OtherSchedulingPolicy) const PURE_VIRTUAL(UPCGSchedulingPolicyBase::IsEquivalent, return false;);
+	virtual bool IsEquivalent(const UPCGSchedulingPolicyBase* OtherSchedulingPolicy) const PURE_VIRTUAL(UPCGSchedulingPolicyBase::IsEquivalent, return false;);
 
 #if WITH_EDITOR
 	/** Sets whether or not properties should be displayed in the editor.Used to hide instanced SchedulingPolicy properties when runtime generation is not enabled. */
diff --git a/Engine/Plugins/Runtime/CommonUI/Source/CommonUI/Public/Groups/CommonWidgetGroupBase.h b/Engine/Plugins/Runtime/CommonUI/Source/CommonUI/Public/Groups/CommonWidgetGroupBase.h
--- a/Engine/Plugins/Runtime/CommonUI/Source/CommonUI/Public/Groups/CommonWidgetGroupBase.h
+++ b/Engine/Plugins/Runtime/CommonUI/Source/CommonUI/Public/Groups/CommonWidgetGroupBase.h
@@ -42,9 +42,9 @@ public:
 	}
 
 protected:
-	UE_API virtual void OnWidgetAdded(UWidget* NewWidget) PURE_VIRTUAL(UCommonWidgetGroupBase::OnWidgetAdded, );
-	UE_API virtual void OnWidgetRemoved(UWidget* OldWidget) PURE_VIRTUAL(UCommonWidgetGroupBase::OnWidgetRemoved, );
-	UE_API virtual void OnRemoveAll() PURE_VIRTUAL(UCommonWidgetGroupBase::OnRemoveAll, );
+	virtual void OnWidgetAdded(UWidget* NewWidget) PURE_VIRTUAL(UCommonWidgetGroupBase::OnWidgetAdded, );
+	virtual void OnWidgetRemoved(UWidget* OldWidget) PURE_VIRTUAL(UCommonWidgetGroupBase::OnWidgetRemoved, );
+	virtual void OnRemoveAll() PURE_VIRTUAL(UCommonWidgetGroupBase::OnRemoveAll, );
 };
 
 #undef UE_API
diff --git a/Engine/Plugins/Runtime/ComputeFramework/Source/ComputeFramework/Public/ComputeFramework/ComputeDataInterface.h b/Engine/Plugins/Runtime/ComputeFramework/Source/ComputeFramework/Public/ComputeFramework/ComputeDataInterface.h
--- a/Engine/Plugins/Runtime/ComputeFramework/Source/ComputeFramework/Public/ComputeFramework/ComputeDataInterface.h
+++ b/Engine/Plugins/Runtime/ComputeFramework/Source/ComputeFramework/Public/ComputeFramework/ComputeDataInterface.h
@@ -26,7 +26,7 @@ class UComputeDataInterface : public UObject
 
 public:
 	/** Get a name for referencing the data interface its shader parameter structures. The InUniqueIndex is optional for use in generating the name. */
-	UE_API virtual TCHAR const* GetClassName() const PURE_VIRTUAL(UComputeDataInterface::GetClassName, return nullptr;)
+	virtual TCHAR const* GetClassName() const PURE_VIRTUAL(UComputeDataInterface::GetClassName, return nullptr;)
 	/** Does the associated UComputeDataProvider provide invocations and thread counts. One and only one data interface per kernel should return true. */
 	virtual bool IsExecutionInterface() const { return false; }
 	/** Return true if the associated UComputeDataProvider holds data that can be combined into a single dispatch invocation. */
diff --git a/Engine/Plugins/Runtime/ComputeFramework/Source/ComputeFramework/Public/ComputeFramework/ComputeKernelSource.h b/Engine/Plugins/Runtime/ComputeFramework/Source/ComputeFramework/Public/ComputeFramework/ComputeKernelSource.h
--- a/Engine/Plugins/Runtime/ComputeFramework/Source/ComputeFramework/Public/ComputeFramework/ComputeKernelSource.h
+++ b/Engine/Plugins/Runtime/ComputeFramework/Source/ComputeFramework/Public/ComputeFramework/ComputeKernelSource.h
@@ -21,7 +21,7 @@ class UComputeKernelSource : public UObject
 
 public:
 	/** Get kernel source code ready for HLSL compilation. */
-	UE_API virtual FString GetSource() const PURE_VIRTUAL(UComputeKernelSource::GetSource, return {};);
+	virtual FString GetSource() const PURE_VIRTUAL(UComputeKernelSource::GetSource, return {};);
 
 public:
 	/** Kernel entry point. */
diff --git a/Engine/Plugins/Runtime/GameplayAbilities/Source/GameplayAbilities/Public/Abilities/GameplayAbilityTargetActor_Trace.h b/Engine/Plugins/Runtime/GameplayAbilities/Source/GameplayAbilities/Public/Abilities/GameplayAbilityTargetActor_Trace.h
--- a/Engine/Plugins/Runtime/GameplayAbilities/Source/GameplayAbilities/Public/Abilities/GameplayAbilityTargetActor_Trace.h
+++ b/Engine/Plugins/Runtime/GameplayAbilities/Source/GameplayAbilities/Public/Abilities/GameplayAbilityTargetActor_Trace.h
@@ -53,7 +53,7 @@ public:
 	bool bTraceAffectsAimPitch;
 
 protected:
-	UE_API virtual FHitResult PerformTrace(AActor* InSourceActor) PURE_VIRTUAL(AGameplayAbilityTargetActor_Trace, return FHitResult(););
+	virtual FHitResult PerformTrace(AActor* InSourceActor) PURE_VIRTUAL(AGameplayAbilityTargetActor_Trace, return FHitResult(););
 
 	UE_API FGameplayAbilityTargetDataHandle MakeTargetData(const FHitResult& HitResult) const;
 
diff --git a/Engine/Plugins/Runtime/HairStrands/Source/HairStrandsCore/Public/GroomAssetCards.h b/Engine/Plugins/Runtime/HairStrands/Source/HairStrandsCore/Public/GroomAssetCards.h
--- a/Engine/Plugins/Runtime/HairStrands/Source/HairStrandsCore/Public/GroomAssetCards.h
+++ b/Engine/Plugins/Runtime/HairStrands/Source/HairStrandsCore/Public/GroomAssetCards.h
@@ -100,7 +100,7 @@ class UHairCardGenerationSettings : public UObject
 {
 	GENERATED_BODY()
 public:
-	UE_API virtual void BuildDDCKey(FArchive& Ar) PURE_VIRTUAL(BuildDDCKey,);
+	virtual void BuildDDCKey(FArchive& Ar) PURE_VIRTUAL(BuildDDCKey,);
 };
 
 USTRUCT(BlueprintType)
diff --git a/Engine/Plugins/Runtime/InstancedActors/Source/InstancedActors/Public/InstancedActorsModifiers.h b/Engine/Plugins/Runtime/InstancedActors/Source/InstancedActors/Public/InstancedActorsModifiers.h
--- a/Engine/Plugins/Runtime/InstancedActors/Source/InstancedActors/Public/InstancedActorsModifiers.h
+++ b/Engine/Plugins/Runtime/InstancedActors/Source/InstancedActors/Public/InstancedActorsModifiers.h
@@ -67,7 +67,7 @@ public:
 	 * 
 	 * @return Return true to continue modification of subsequent instances, false to break iteration.
 	 */
-	UE_API virtual bool ModifyInstance(const FInstancedActorsInstanceHandle& InstanceHandle, const FTransform& InstanceTransform, FInstancedActorsIterationContext& IterationContext) PURE_VIRTUAL(UInstancedActorsModifierBase::ModifyInstance, return false;)
+	virtual bool ModifyInstance(const FInstancedActorsInstanceHandle& InstanceHandle, const FTransform& InstanceTransform, FInstancedActorsIterationContext& IterationContext) PURE_VIRTUAL(UInstancedActorsModifierBase::ModifyInstance, return false;)
 	
 	/** 
 	 * Callback to modify all instances in Manager, whose location falls within Bounds. 
diff --git a/Engine/Plugins/Runtime/MIDIDevice/Source/MIDIDevice/Public/MIDIDeviceControllerBase.h b/Engine/Plugins/Runtime/MIDIDevice/Source/MIDIDevice/Public/MIDIDeviceControllerBase.h
--- a/Engine/Plugins/Runtime/MIDIDevice/Source/MIDIDevice/Public/MIDIDeviceControllerBase.h
+++ b/Engine/Plugins/Runtime/MIDIDevice/Source/MIDIDevice/Public/MIDIDeviceControllerBase.h
@@ -16,19 +16,19 @@ class UMIDIDeviceControllerBase : public UObject
 
 public:
 	/** Called from UMIDIDeviceManager after the controller is created to get it ready to use.  Don't call this directly. */
-	UE_API virtual void StartupDevice(const int32 InitDeviceID, const int32 InitMIDIBufferSize, bool& bOutWasSuccessful) PURE_VIRTUAL(UMIDIDeviceControllerBase::StartupDevice,;)
+	virtual void StartupDevice(const int32 InitDeviceID, const int32 InitMIDIBufferSize, bool& bOutWasSuccessful) PURE_VIRTUAL(UMIDIDeviceControllerBase::StartupDevice,;)
 
 	/** Called every frame by UMIDIDeviceManager to poll for new MIDI events and broadcast them out to subscribers of OnMIDIEvent.  Don't call this directly. */
 	virtual void ProcessIncomingMIDIEvents() { };
 
 	/** Called during destruction to clean up this device.  Don't call this directly. */
-	UE_API virtual void ShutdownDevice() PURE_VIRTUAL(UMIDIDeviceControllerBase::ShutdownDevice,;)
+	virtual void ShutdownDevice() PURE_VIRTUAL(UMIDIDeviceControllerBase::ShutdownDevice,;)
 
 	/** The name of this device.  This name comes from the MIDI hardware, any might not be unique */
-	UE_API virtual FString GetDeviceName() const PURE_VIRTUAL(UMIDIDeviceControllerBase::GetDeviceName, return TEXT("");)
+	virtual FString GetDeviceName() const PURE_VIRTUAL(UMIDIDeviceControllerBase::GetDeviceName, return TEXT("");)
 
 	/** Size of the MIDI buffer in bytes */
-	UE_API virtual int32 GetMIDIBufferSize() const PURE_VIRTUAL(UMIDIDeviceControllerBase::GetMIDIBufferSize, return 0;)
+	virtual int32 GetMIDIBufferSize() const PURE_VIRTUAL(UMIDIDeviceControllerBase::GetMIDIBufferSize, return 0;)
 };
 
 #undef UE_API
diff --git a/Engine/Plugins/Runtime/MassGameplay/Source/MassReplication/Public/MassReplicationProcessor.h b/Engine/Plugins/Runtime/MassGameplay/Source/MassReplication/Public/MassReplicationProcessor.h
--- a/Engine/Plugins/Runtime/MassGameplay/Source/MassReplication/Public/MassReplicationProcessor.h
+++ b/Engine/Plugins/Runtime/MassGameplay/Source/MassReplication/Public/MassReplicationProcessor.h
@@ -74,13 +74,13 @@ public:
 	 * Must override to add specific entity query requirements for replication
 	 * Usually we add replication processor handler requirements
 	 */
-	MASSREPLICATION_API virtual void AddRequirements(FMassEntityQuery& EntityQuery) PURE_VIRTUAL(UMassReplicatorBase::ConfigureQueries, );
+	virtual void AddRequirements(FMassEntityQuery& EntityQuery) PURE_VIRTUAL(UMassReplicatorBase::ConfigureQueries, );
 
 	/**
 	 * Must override to process the client replication
 	 * This methods should call CalculateClientReplication with the appropriate callback implementation
 	 */
-	MASSREPLICATION_API virtual void ProcessClientReplication(FMassExecutionContext& Context, FMassReplicationContext& ReplicationContext) PURE_VIRTUAL(UMassReplicatorBase::ProcessClientReplication, );
+	virtual void ProcessClientReplication(FMassExecutionContext& Context, FMassReplicationContext& ReplicationContext) PURE_VIRTUAL(UMassReplicatorBase::ProcessClientReplication, );
 
 protected:
 	/**
diff --git a/Engine/Plugins/Runtime/MassGameplay/Source/MassSignals/Public/MassSignalProcessorBase.h b/Engine/Plugins/Runtime/MassGameplay/Source/MassSignals/Public/MassSignalProcessorBase.h
--- a/Engine/Plugins/Runtime/MassGameplay/Source/MassSignals/Public/MassSignalProcessorBase.h
+++ b/Engine/Plugins/Runtime/MassGameplay/Source/MassSignals/Public/MassSignalProcessorBase.h
@@ -37,7 +37,7 @@ protected:
 	 * @param Context is the execution context to be passed when executing the lambdas
 	 * @param EntitySignals Look up to retrieve for each entities their raised signal via GetSignalsForEntity
 	 */
-	 UE_API virtual void SignalEntities(FMassEntityManager& EntityManager, FMassExecutionContext& Context, FMassSignalNameLookup& EntitySignals) PURE_VIRTUAL(UMassSignalProcessorBase::SignalEntities, );
+	 virtual void SignalEntities(FMassEntityManager& EntityManager, FMassExecutionContext& Context, FMassSignalNameLookup& EntitySignals) PURE_VIRTUAL(UMassSignalProcessorBase::SignalEntities, );
 
 	/**
 	 * Callback that is being called when new signal is raised
diff --git a/Engine/Plugins/Runtime/MassGameplay/Source/MassSpawner/Public/MassEntitySpawnDataGeneratorBase.h b/Engine/Plugins/Runtime/MassGameplay/Source/MassSpawner/Public/MassEntitySpawnDataGeneratorBase.h
--- a/Engine/Plugins/Runtime/MassGameplay/Source/MassSpawner/Public/MassEntitySpawnDataGeneratorBase.h
+++ b/Engine/Plugins/Runtime/MassGameplay/Source/MassSpawner/Public/MassEntitySpawnDataGeneratorBase.h
@@ -60,7 +60,7 @@ public:
 	 * @param Count of point to generate
 	 * @param FinishedGeneratingSpawnPointsDelegate is the callback to call once the generation is done
 	 */
-	MASSSPAWNER_API virtual void Generate(UObject& QueryOwner, TConstArrayView<FMassSpawnedEntityType> EntityTypes, int32 Count, FFinishedGeneratingSpawnDataSignature& FinishedGeneratingSpawnPointsDelegate) const PURE_VIRTUAL(UMassEntitySpawnDataGeneratorBase::GenerateSpawnPoints, );
+	virtual void Generate(UObject& QueryOwner, TConstArrayView<FMassSpawnedEntityType> EntityTypes, int32 Count, FFinishedGeneratingSpawnDataSignature& FinishedGeneratingSpawnPointsDelegate) const PURE_VIRTUAL(UMassEntitySpawnDataGeneratorBase::GenerateSpawnPoints, );
 
 	/** Populates empty generator results from EntityTypes based on the provided proportions. 
 	 * @param SpawnCount How many entities to distributed among the EntityTypes
diff --git a/Engine/Plugins/Runtime/MassGameplay/Source/MassSpawner/Public/MassEntityTraitBase.h b/Engine/Plugins/Runtime/MassGameplay/Source/MassSpawner/Public/MassEntityTraitBase.h
--- a/Engine/Plugins/Runtime/MassGameplay/Source/MassSpawner/Public/MassEntityTraitBase.h
+++ b/Engine/Plugins/Runtime/MassGameplay/Source/MassSpawner/Public/MassEntityTraitBase.h
@@ -50,7 +50,7 @@ public:
 
 
 	/** Appends items into the entity template required for the trait. */
-	MASSSPAWNER_API virtual void BuildTemplate(FMassEntityTemplateBuildContext& BuildContext, const UWorld& World) const PURE_VIRTUAL(UMassEntityTraitBase::BuildTemplate, return; );
+	virtual void BuildTemplate(FMassEntityTemplateBuildContext& BuildContext, const UWorld& World) const PURE_VIRTUAL(UMassEntityTraitBase::BuildTemplate, return; );
 
 	UE_DEPRECATED(5.5, "This flavor of DestroyTemplate has been deprecated, use the version taking the World parameter")
 	MASSSPAWNER_API virtual void DestroyTemplate() const;
diff --git a/Engine/Plugins/Runtime/MeshModelingToolset/Source/ModelingComponents/Public/BaseTools/MultiSelectionMeshEditingTool.h b/Engine/Plugins/Runtime/MeshModelingToolset/Source/ModelingComponents/Public/BaseTools/MultiSelectionMeshEditingTool.h
--- a/Engine/Plugins/Runtime/MeshModelingToolset/Source/ModelingComponents/Public/BaseTools/MultiSelectionMeshEditingTool.h
+++ b/Engine/Plugins/Runtime/MeshModelingToolset/Source/ModelingComponents/Public/BaseTools/MultiSelectionMeshEditingTool.h
@@ -28,7 +28,7 @@ public:
 	UE_API virtual UInteractiveTool* BuildTool(const FToolBuilderState& SceneState) const override;
 
 	/** @return new Tool instance. Override this in subclasses to build a different Tool class type */
-	UE_API virtual UMultiSelectionMeshEditingTool* CreateNewTool(const FToolBuilderState& SceneState) const PURE_VIRTUAL(UMultiSelectionMeshEditingToolBuilder::CreateNewTool, return nullptr; );
+	virtual UMultiSelectionMeshEditingTool* CreateNewTool(const FToolBuilderState& SceneState) const PURE_VIRTUAL(UMultiSelectionMeshEditingToolBuilder::CreateNewTool, return nullptr; );
 
 	/** Called by BuildTool to configure the Tool with the input mesh source(s) based on the SceneState */
 	UE_API virtual void InitializeNewTool(UMultiSelectionMeshEditingTool* Tool, const FToolBuilderState& SceneState) const;
diff --git a/Engine/Plugins/Runtime/MeshModelingToolset/Source/ModelingComponents/Public/BaseTools/MultiTargetWithSelectionTool.h b/Engine/Plugins/Runtime/MeshModelingToolset/Source/ModelingComponents/Public/BaseTools/MultiTargetWithSelectionTool.h
--- a/Engine/Plugins/Runtime/MeshModelingToolset/Source/ModelingComponents/Public/BaseTools/MultiTargetWithSelectionTool.h
+++ b/Engine/Plugins/Runtime/MeshModelingToolset/Source/ModelingComponents/Public/BaseTools/MultiTargetWithSelectionTool.h
@@ -36,7 +36,7 @@ public:
 	UE_API virtual UInteractiveTool* BuildTool(const FToolBuilderState& SceneState) const override;
 
 	/** @return new Tool instance. Override this in subclasses to build a different Tool class type */
-	UE_API virtual UMultiTargetWithSelectionTool* CreateNewTool(const FToolBuilderState& SceneState) const PURE_VIRTUAL(UMultiTargetWithSelectionToolBuilder::CreateNewTool, return nullptr; );
+	virtual UMultiTargetWithSelectionTool* CreateNewTool(const FToolBuilderState& SceneState) const PURE_VIRTUAL(UMultiTargetWithSelectionToolBuilder::CreateNewTool, return nullptr; );
 
 	/** Called by BuildTool to configure the Tool with the input mesh source(s) based on the SceneState */
 	UE_API virtual void InitializeNewTool(UMultiTargetWithSelectionTool* NewTool, const FToolBuilderState& SceneState) const;
diff --git a/Engine/Plugins/Runtime/MeshModelingToolset/Source/ModelingComponents/Public/BaseTools/SingleSelectionMeshEditingTool.h b/Engine/Plugins/Runtime/MeshModelingToolset/Source/ModelingComponents/Public/BaseTools/SingleSelectionMeshEditingTool.h
--- a/Engine/Plugins/Runtime/MeshModelingToolset/Source/ModelingComponents/Public/BaseTools/SingleSelectionMeshEditingTool.h
+++ b/Engine/Plugins/Runtime/MeshModelingToolset/Source/ModelingComponents/Public/BaseTools/SingleSelectionMeshEditingTool.h
@@ -28,7 +28,7 @@ public:
 	UE_API virtual UInteractiveTool* BuildTool(const FToolBuilderState& SceneState) const override;
 
 	/** @return new Tool instance. Override this in subclasses to build a different Tool class type */
-	UE_API virtual USingleSelectionMeshEditingTool* CreateNewTool(const FToolBuilderState& SceneState) const PURE_VIRTUAL(USingleSelectionMeshEditingToolBuilder::CreateNewTool, return nullptr; );
+	virtual USingleSelectionMeshEditingTool* CreateNewTool(const FToolBuilderState& SceneState) const PURE_VIRTUAL(USingleSelectionMeshEditingToolBuilder::CreateNewTool, return nullptr; );
 
 	/** Called by BuildTool to configure the Tool with the input MeshSource based on the SceneState */
 	UE_API virtual void InitializeNewTool(USingleSelectionMeshEditingTool* Tool, const FToolBuilderState& SceneState) const;
diff --git a/Engine/Plugins/Runtime/MeshModelingToolset/Source/ModelingComponents/Public/BaseTools/SingleTargetWithSelectionTool.h b/Engine/Plugins/Runtime/MeshModelingToolset/Source/ModelingComponents/Public/BaseTools/SingleTargetWithSelectionTool.h
--- a/Engine/Plugins/Runtime/MeshModelingToolset/Source/ModelingComponents/Public/BaseTools/SingleTargetWithSelectionTool.h
+++ b/Engine/Plugins/Runtime/MeshModelingToolset/Source/ModelingComponents/Public/BaseTools/SingleTargetWithSelectionTool.h
@@ -27,7 +27,7 @@ public:
 	UE_API virtual UInteractiveTool* BuildTool(const FToolBuilderState& SceneState) const override;
 
 	/** @return new Tool instance. Override this in subclasses to build a different Tool class type */
-	UE_API virtual USingleTargetWithSelectionTool* CreateNewTool(const FToolBuilderState& SceneState) const PURE_VIRTUAL(USingleTargetWithSelectionToolBuilder::CreateNewTool, return nullptr; );
+	virtual USingleTargetWithSelectionTool* CreateNewTool(const FToolBuilderState& SceneState) const PURE_VIRTUAL(USingleTargetWithSelectionToolBuilder::CreateNewTool, return nullptr; );
 
 	/** Called by BuildTool to configure the Tool with the input MeshSource based on the SceneState */
 	UE_API virtual void InitializeNewTool(USingleTargetWithSelectionTool* Tool, const FToolBuilderState& SceneState) const;
diff --git a/Engine/Plugins/Runtime/MeshModelingToolset/Source/ModelingComponents/Public/Selection/MeshTopologySelectionMechanic.h b/Engine/Plugins/Runtime/MeshModelingToolset/Source/ModelingComponents/Public/Selection/MeshTopologySelectionMechanic.h
--- a/Engine/Plugins/Runtime/MeshModelingToolset/Source/ModelingComponents/Public/Selection/MeshTopologySelectionMechanic.h
+++ b/Engine/Plugins/Runtime/MeshModelingToolset/Source/ModelingComponents/Public/Selection/MeshTopologySelectionMechanic.h
@@ -252,7 +252,7 @@ public:
 	 * Update the hover highlight based on the hit elements at the given World Ray
 	 * @return true if something was hit and is now being hovered
 	 */
-	UE_API virtual bool UpdateHighlight(const FRay& WorldRay) PURE_VIRTUAL(UMeshTopologySelectionMechanic::UpdateHighlight, return false; );
+	virtual bool UpdateHighlight(const FRay& WorldRay) PURE_VIRTUAL(UMeshTopologySelectionMechanic::UpdateHighlight, return false; );
 
 	/**
 	 * Clear current hover-highlight
@@ -268,7 +268,7 @@ public:
 	 * Intersect the ray with the mesh and update the selection based on the hit element, modifier states, etc
 	 * @return true if selection was modified
 	 */
-	UE_API virtual bool UpdateSelection(const FRay& WorldRay, FVector3d& LocalHitPositionOut, FVector3d& LocalHitNormalOut) PURE_VIRTUAL(UMeshTopologySelectionMechanic::UpdateSelection, return false; );
+	virtual bool UpdateSelection(const FRay& WorldRay, FVector3d& LocalHitPositionOut, FVector3d& LocalHitNormalOut) PURE_VIRTUAL(UMeshTopologySelectionMechanic::UpdateSelection, return false; );
 
 	/**
 	 * Replace the current selection with an external selection. 
diff --git a/Engine/Plugins/Runtime/Metasound/Source/MetasoundEngine/Public/Metasound.h b/Engine/Plugins/Runtime/Metasound/Source/MetasoundEngine/Public/Metasound.h
--- a/Engine/Plugins/Runtime/Metasound/Source/MetasoundEngine/Public/Metasound.h
+++ b/Engine/Plugins/Runtime/Metasound/Source/MetasoundEngine/Public/Metasound.h
@@ -41,7 +41,7 @@ public:
 	virtual bool IsEditorOnly() const override { return true; }
 	virtual bool NeedsLoadForEditorGame() const override { return false; }
 
-	UE_API virtual void RegisterGraphWithFrontend() PURE_VIRTUAL(UMetasoundEditorGraphBase::RegisterGraphWithFrontend(), )
+	virtual void RegisterGraphWithFrontend() PURE_VIRTUAL(UMetasoundEditorGraphBase::RegisterGraphWithFrontend(), )
 
 #if WITH_EDITORONLY_DATA
 	UE_DEPRECATED(5.5, "ModifyContext is to be replaced by builder API delegates providing context when items changed and it will be up to the caller to track modification deltas.")
@@ -59,7 +59,7 @@ public:
 	UE_DEPRECATED(5.5, "Editor Graph is now transient, so versioning flag moved to AssetBase.")
 	virtual void SetVersionedOnLoad() {  }
 
-	UE_API virtual void MigrateEditorDocumentData(FMetaSoundFrontendDocumentBuilder & OutBuilder) PURE_VIRTUAL(UMetasoundEditorGraphBase::MigrateEditorDocumentData(), )
+	virtual void MigrateEditorDocumentData(FMetaSoundFrontendDocumentBuilder & OutBuilder) PURE_VIRTUAL(UMetasoundEditorGraphBase::MigrateEditorDocumentData(), )
 #endif // WITH_EDITORONLY_DATA
 
 	UE_API int32 GetHighestMessageSeverity() const;
diff --git a/Engine/Plugins/Runtime/Metasound/Source/MetasoundEngine/Public/MetasoundBuilderBase.h b/Engine/Plugins/Runtime/Metasound/Source/MetasoundEngine/Public/MetasoundBuilderBase.h
--- a/Engine/Plugins/Runtime/Metasound/Source/MetasoundEngine/Public/MetasoundBuilderBase.h
+++ b/Engine/Plugins/Runtime/Metasound/Source/MetasoundEngine/Public/MetasoundBuilderBase.h
@@ -538,7 +538,7 @@ public:
 	UE_API virtual bool ConformObjectToDocument();
 
 	// Returns the base class registered with the MetaSound UObject registry.
-	UE_API virtual const UClass& GetBaseMetaSoundUClass() const PURE_VIRTUAL(UMetaSoundBuilderBase::GetBaseMetaSoundUClass, return *UClass::StaticClass(); );
+	virtual const UClass& GetBaseMetaSoundUClass() const PURE_VIRTUAL(UMetaSoundBuilderBase::GetBaseMetaSoundUClass, return *UClass::StaticClass(); );
 
 	UE_DEPRECATED(5.5, "Renamed to 'GetBaseMetaSoundUClass' for consistency")
 	virtual const UClass& GetBuilderUClass() const { return GetBaseMetaSoundUClass(); }
@@ -578,11 +578,11 @@ public:
 	UE_API void Initialize();
 
 protected:
-	UE_API virtual void BuildAndOverwriteMetaSoundInternal(TScriptInterface<IMetaSoundDocumentInterface> ExistingMetaSound, bool bForceUniqueClassName) const PURE_VIRTUAL(UMetaSoundBuilderBase::BuildAndOverwriteMetaSoundInternal, ;);
+	virtual void BuildAndOverwriteMetaSoundInternal(TScriptInterface<IMetaSoundDocumentInterface> ExistingMetaSound, bool bForceUniqueClassName) const PURE_VIRTUAL(UMetaSoundBuilderBase::BuildAndOverwriteMetaSoundInternal, ;);
 	UE_API virtual void InitDelegates(Metasound::Frontend::FDocumentModifyDelegates& OutDocumentDelegates);
 
-	UE_API virtual void OnAssetReferenceAdded(TScriptInterface<IMetaSoundDocumentInterface> DocInterface) PURE_VIRTUAL(UMetaSoundBuilderBase::OnAssetReferenceAdded, );
-	UE_API virtual void OnRemovingAssetReference(TScriptInterface<IMetaSoundDocumentInterface> DocInterface) PURE_VIRTUAL(UMetaSoundBuilderBase::OnRemovingAssetReference, );
+	virtual void OnAssetReferenceAdded(TScriptInterface<IMetaSoundDocumentInterface> DocInterface) PURE_VIRTUAL(UMetaSoundBuilderBase::OnAssetReferenceAdded, );
+	virtual void OnRemovingAssetReference(TScriptInterface<IMetaSoundDocumentInterface> DocInterface) PURE_VIRTUAL(UMetaSoundBuilderBase::OnRemovingAssetReference, );
 
 	UE_DEPRECATED(5.5, "Moved to internal implementation")
 	virtual void CreateTransientBuilder() { };
diff --git a/Engine/Plugins/Runtime/PropertyBindingUtils/Source/PropertyBindingUtils/Public/PropertyBindingBindingCollection.h b/Engine/Plugins/Runtime/PropertyBindingUtils/Source/PropertyBindingUtils/Public/PropertyBindingBindingCollection.h
--- a/Engine/Plugins/Runtime/PropertyBindingUtils/Source/PropertyBindingUtils/Public/PropertyBindingBindingCollection.h
+++ b/Engine/Plugins/Runtime/PropertyBindingUtils/Source/PropertyBindingUtils/Public/PropertyBindingBindingCollection.h
@@ -155,12 +155,12 @@ public:
 	}
 
 	/** @return Number of bindable struct descriptors available in the collection. */
-	PROPERTYBINDINGUTILS_API virtual int32 GetNumBindableStructDescriptors() const PURE_VIRTUAL(FPropertyBindingBindingCollection::GetNumBindableStructDescriptors, return INDEX_NONE;);
+	virtual int32 GetNumBindableStructDescriptors() const PURE_VIRTUAL(FPropertyBindingBindingCollection::GetNumBindableStructDescriptors, return INDEX_NONE;);
 
-	PROPERTYBINDINGUTILS_API virtual const FPropertyBindingBindableStructDescriptor* GetBindableStructDescriptorFromHandle(FConstStructView InSourceHandleView) const PURE_VIRTUAL(FPropertyBindingBindingCollection::GetBindableStructDescriptorFromHandle, return nullptr; );
+	virtual const FPropertyBindingBindableStructDescriptor* GetBindableStructDescriptorFromHandle(FConstStructView InSourceHandleView) const PURE_VIRTUAL(FPropertyBindingBindingCollection::GetBindableStructDescriptorFromHandle, return nullptr; );
 
 	/** @return The amount of registered bindings. */
-	PROPERTYBINDINGUTILS_API virtual int32 GetNumBindings() const PURE_VIRTUAL(FPropertyBindingBindingCollection::GetNumBindings, return INDEX_NONE;);
+	virtual int32 GetNumBindings() const PURE_VIRTUAL(FPropertyBindingBindingCollection::GetNumBindings, return INDEX_NONE;);
 
 	/**
 	 * Iterates through all bindings and calls the provided function on each non-mutable binding.
@@ -169,7 +169,7 @@ public:
 	 * @note To be able to break iteration, use VisitBindings or VisitMutableBindings
 	 * @see ForEachMutableBinding, VisitBindings, VisitMutableBindings
 	 */
-	PROPERTYBINDINGUTILS_API virtual void ForEachBinding(TFunctionRef<void(const FPropertyBindingBinding& Binding)> InFunction) const PURE_VIRTUAL(FPropertyBindingBindingCollection::ForEachBinding, );
+	virtual void ForEachBinding(TFunctionRef<void(const FPropertyBindingBinding& Binding)> InFunction) const PURE_VIRTUAL(FPropertyBindingBindingCollection::ForEachBinding, );
 
 	/**
 	 * Iterates through all bindings between indices [Begin, End[ and calls the provided function on each non-mutable binding.
@@ -180,7 +180,7 @@ public:
 	 * @note To be able to break iteration, use VisitBindings or VisitMutableBindings
 	 * @see ForEachMutableBinding, VisitBindings, VisitMutableBindings
 	 */
-	PROPERTYBINDINGUTILS_API virtual void ForEachBinding(FPropertyBindingIndex16 InBegin, FPropertyBindingIndex16 InEnd, TFunctionRef<void(const FPropertyBindingBinding& Binding, const int32 BindingIndex)> InFunction) const PURE_VIRTUAL(FPropertyBindingBindingCollection::ForEachBinding, );
+	virtual void ForEachBinding(FPropertyBindingIndex16 InBegin, FPropertyBindingIndex16 InEnd, TFunctionRef<void(const FPropertyBindingBinding& Binding, const int32 BindingIndex)> InFunction) const PURE_VIRTUAL(FPropertyBindingBindingCollection::ForEachBinding, );
 
 	/**
 	 * Iterates through all bindings and calls the provided function on each mutable binding.
@@ -189,7 +189,7 @@ public:
 	 * @note To be able to break iteration, use VisitBindings or VisitMutableBindings
 	 * @see ForEachBinding, VisitBindings, VisitMutableBindings
 	 */
-	PROPERTYBINDINGUTILS_API virtual void ForEachMutableBinding(TFunctionRef<void(FPropertyBindingBinding& Binding)> InFunction) PURE_VIRTUAL(FPropertyBindingBindingCollection::ForEachMutableBinding, );
+	virtual void ForEachMutableBinding(TFunctionRef<void(FPropertyBindingBinding& Binding)> InFunction) PURE_VIRTUAL(FPropertyBindingBindingCollection::ForEachMutableBinding, );
 
 	/**
 	 * Iterates through all bindings and calls the provided function on each non-mutable binding.
@@ -199,7 +199,7 @@ public:
 	 * @note If breaking the iteration is not required, consider using ForEachBinding or ForEachMutableBinding
 	 * @see VisitMutableBindings, ForEachBinding, ForEachMutableBinding
 	 */
-	PROPERTYBINDINGUTILS_API virtual void VisitBindings(TFunctionRef<EVisitResult(const FPropertyBindingBinding& Binding)> InFunction) const PURE_VIRTUAL(FPropertyBindingBindingCollection::VisitBindings, );
+	virtual void VisitBindings(TFunctionRef<EVisitResult(const FPropertyBindingBinding& Binding)> InFunction) const PURE_VIRTUAL(FPropertyBindingBindingCollection::VisitBindings, );
 
 	/**
 	 * Iterates through all bindings and calls the provided function on each mutable binding.
@@ -209,7 +209,7 @@ public:
 	 * @note If breaking the iteration is not required, consider using ForEachBinding or ForEachMutableBinding
 	 * @see VisitBindings, ForEachBinding, ForEachMutableBinding
 	 */
-	PROPERTYBINDINGUTILS_API virtual void VisitMutableBindings(TFunctionRef<EVisitResult(FPropertyBindingBinding& Binding)> InFunction) PURE_VIRTUAL(FPropertyBindingBindingCollection::VisitMutableBindings, );
+	virtual void VisitMutableBindings(TFunctionRef<EVisitResult(FPropertyBindingBinding& Binding)> InFunction) PURE_VIRTUAL(FPropertyBindingBindingCollection::VisitMutableBindings, );
 
 	/**
 	 * Clears all bindings and related data.
