[PCG] Fix bug for where runtime gen components leak partition actors into the world when Partition Grid Size setting changed on PCGWorldActor.

[RN] Fixed runtime generation bug where changing Partition Grid Size setting on PCGWorldActor resulted in generated actors being leaked into the level.

diff --git a/Engine/Plugins/PCG/Source/PCG/Private/PCGWorldActor.cpp b/Engine/Plugins/PCG/Source/PCG/Private/PCGWorldActor.cpp
--- a/Engine/Plugins/PCG/Source/PCG/Private/PCGWorldActor.cpp
+++ b/Engine/Plugins/PCG/Source/PCG/Private/PCGWorldActor.cpp
@@ -296,7 +296,7 @@ void APCGWorldActor::OnPartitionGridSizeChanged()
 	{
 		check(PCGComponent);
 		PCGComponent->DirtyGenerated();
-		PCGComponent->Refresh();
+		PCGComponent->Refresh(EPCGChangeType::GenerationGrid, /*bCancelExistingRefresh=*/true);
 	}
 }
 
