[PCG] Fix issue where data descriptions can have non-unique attributes with different IDs

#rnx

diff --git a/Engine/Plugins/PCG/Source/PCG/Private/Compute/Elements/PCGCustomHLSLKernel.cpp b/Engine/Plugins/PCG/Source/PCG/Private/Compute/Elements/PCGCustomHLSLKernel.cpp
--- a/Engine/Plugins/PCG/Source/PCG/Private/Compute/Elements/PCGCustomHLSLKernel.cpp
+++ b/Engine/Plugins/PCG/Source/PCG/Private/Compute/Elements/PCGCustomHLSLKernel.cpp
@@ -171,30 +171,8 @@ void UPCGCustomHLSLKernel::InitializeInternal()
 	Super::InitializeInternal();
 
 	InitEntryPoint();
+	PopulateAttributeKeysFromPinSettings();
 	ParseShaderSource();
-
-	const UPCGCustomHLSLSettings* CustomHLSLSettings = CastChecked<UPCGCustomHLSLSettings>(GetSettings());
-
-	// Process each output pin for any new attributes they want to create.
-	for (const FPCGPinPropertiesGPU& OutputPin : CustomHLSLSettings->OutputPins)
-	{
-		for (const FPCGKernelAttributeKey& AuthoredKey : OutputPin.PropertiesGPU.CreatedKernelAttributeKeys)
-		{
-			const uint32 AttributeKeyIndex = KernelAttributeKeys.AddUnique(AuthoredKey);
-			TArray<FPCGCustomHLSLAttributeInfo>& Infos = PinToAttributeInfos.FindOrAdd(OutputPin.Label).Infos;
-			FPCGCustomHLSLAttributeInfo* Info = Infos.FindByPredicate([AttributeKeyIndex](const FPCGCustomHLSLAttributeInfo& Info) { return Info.Index == AttributeKeyIndex; });
-
-			// Mark as created on GPU.
-			if (Info)
-			{
-				Info->bCreatedOnGPU = true;
-			}
-			else
-			{
-				Infos.Emplace(AttributeKeyIndex, /*bCreatedOnGPU=*/true);
-			}
-		}
-	}
 }
 #endif // WITH_EDITOR
 
@@ -281,27 +259,10 @@ bool UPCGCustomHLSLKernel::ComputeOutputBindingDataDesc(const UPCGComputeGraph*
 		ComputeDataDescFromPinProperties(*OutputPinProperties, MakeArrayView(CustomHLSLSettings->InputPins), InBinding, OutputDataDesc);
 	}
 
-	// Add attributes that will be created for this pin on the GPU.
-	if (const FPCGCustomHLSLAttributeInfos* Infos = PinToAttributeInfos.Find(InOutputPinLabel))
+	// Add attributes that will be created for this pin on the GPU. This will stomp any existing attributes if they collide!
+	for (const FPCGKernelAttributeKey& CreatedKey : OutputPinProperties->PropertiesGPU.CreatedKernelAttributeKeys)
 	{
-		for (const FPCGCustomHLSLAttributeInfo& Info : Infos->Infos)
-		{
-			if (Info.bCreatedOnGPU && KernelAttributeKeys.IsValidIndex(Info.Index))
-			{
-				const FPCGKernelAttributeKey& AttributeKey = KernelAttributeKeys[Info.Index];
-
-				for (FPCGDataDesc& DataDesc : OutputDataDesc.DataDescs)
-				{
-					const int32 AttributeId = InBinding->GetAttributeId(AttributeKey);
-
-					if (AttributeId != INDEX_NONE)
-					{
-						const FPCGKernelAttributeDesc AttributeDesc(AttributeId, AttributeKey.Type, AttributeKey.Identifier);
-						DataDesc.AttributeDescs.AddUnique(AttributeDesc);
-					}
-				}
-			}
-		}
+		OutputDataDesc.AddAttributeToAllData(CreatedKey, InBinding);
 	}
 
 	// Try to propagate string keys across node. Not trivial because there could be one or more string key attributes on input pins and on output pins,
@@ -733,6 +694,23 @@ void UPCGCustomHLSLKernel::InitEntryPoint()
 	}
 }
 
+void UPCGCustomHLSLKernel::PopulateAttributeKeysFromPinSettings()
+{
+	const UPCGCustomHLSLSettings* CustomHLSLSettings = CastChecked<UPCGCustomHLSLSettings>(GetSettings());
+
+	// Process each output pin for any new attributes they want to create.
+	for (const FPCGPinPropertiesGPU& OutputPin : CustomHLSLSettings->OutputPins)
+	{
+		for (const FPCGKernelAttributeKey& AuthoredKey : OutputPin.PropertiesGPU.CreatedKernelAttributeKeys)
+		{
+			if (AuthoredKey.IsValid())
+			{
+				KernelAttributeKeys.AddUnique(AuthoredKey);
+			}
+		}
+	}
+}
+
 void UPCGCustomHLSLKernel::ParseShaderSource()
 {
 	CreateParsedSources();
@@ -845,7 +823,7 @@ void UPCGCustomHLSLKernel::ParseShaderSource()
 		FString SingleQuoteString = "";
 		bool bLookingForSingleQuoteMatch = false;
 
-		auto AddCompletedAttributeFunction = [&KernelAttributeKeys=KernelAttributeKeys, &PinToAttributeInfos=PinToAttributeInfos, &AttributeFunctions=InOutParsedSource.AttributeFunctions](const FAttributeFunctionMatch& InAttributeFunction)
+		auto AddCompletedAttributeFunction = [&KernelAttributeKeys=KernelAttributeKeys, &AttributeFunctions=InOutParsedSource.AttributeFunctions](const FAttributeFunctionMatch& InAttributeFunction)
 		{
 			if (InAttributeFunction.EncounteredCommas != InAttributeFunction.RequiredCommas)
 			{
@@ -883,12 +861,7 @@ void UPCGCustomHLSLKernel::ParseShaderSource()
 
 			if (Key.IsValid())
 			{
-				const uint32 AttributeKeyIndex = KernelAttributeKeys.AddUnique(Key);
-
-				// Add an entry mapping this pin to the referenced attribute, if the entry doesn't already exist.
-				FPCGCustomHLSLAttributeInfos& Infos = PinToAttributeInfos.FindOrAdd(FName(*PinStr));
-				Infos.Infos.AddUnique({ AttributeKeyIndex, /*bCreatedOnGPU=*/false });
-			
+				KernelAttributeKeys.AddUnique(Key);
 				AttributeFunctions.Emplace(PinStr, FuncStr, AttributeType, NameStr, InAttributeFunction.MatchBeginning);
 			}
 		};
diff --git a/Engine/Plugins/PCG/Source/PCG/Private/Compute/Elements/PCGCustomHLSLKernel.h b/Engine/Plugins/PCG/Source/PCG/Private/Compute/Elements/PCGCustomHLSLKernel.h
--- a/Engine/Plugins/PCG/Source/PCG/Private/Compute/Elements/PCGCustomHLSLKernel.h
+++ b/Engine/Plugins/PCG/Source/PCG/Private/Compute/Elements/PCGCustomHLSLKernel.h
@@ -59,39 +59,6 @@ public:
 	FString TargetPin;
 };
 
-/** Helper struct for serializing attribute info. */
-USTRUCT()
-struct FPCGCustomHLSLAttributeInfo
-{
-	GENERATED_BODY()
-
-	FPCGCustomHLSLAttributeInfo() {}
-	FPCGCustomHLSLAttributeInfo(uint32 InIndex, bool bInCreatedOnGPU)
-		: Index(InIndex)
-		, bCreatedOnGPU(bInCreatedOnGPU)
-	{}
-
-	bool operator==(const FPCGCustomHLSLAttributeInfo& Other) const { return Index == Other.Index && bCreatedOnGPU == Other.bCreatedOnGPU; }
-
-	/** Index into KernelAttributeKeys. */
-	UPROPERTY()
-	uint32 Index = INDEX_NONE;
-
-	/** Whether or not this attribute was created on the GPU. */
-	UPROPERTY()
-	bool bCreatedOnGPU = false;
-};
-
-/** Helper struct for serializing map of pin name to attribute infos. */
-USTRUCT()
-struct FPCGCustomHLSLAttributeInfos
-{
-	GENERATED_BODY()
-
-	UPROPERTY()
-	TArray<FPCGCustomHLSLAttributeInfo> Infos;
-};
-
 // @todo_pcg: It would be ideal someday for parsed source to be editor only. Unfortunately not possible right now because we need to validate the parsed attribute functions during execution.
 /** Holds the results of parsing for some HLSL source. */
 USTRUCT()
@@ -168,6 +135,7 @@ protected:
 	// ~End UPCGComputeKernel interface
 
 	void InitEntryPoint();
+	void PopulateAttributeKeysFromPinSettings();
 	void ParseShaderSource();
 	void CreateParsedSources();
 	void CollectDataLabels(const FPCGCustomHLSLParsedSource& InParsedSource);
@@ -202,9 +170,6 @@ protected:
 	UPROPERTY()
 	TArray<FPCGKernelAttributeKey> KernelAttributeKeys;
 
-	UPROPERTY()
-	TMap<FName, FPCGCustomHLSLAttributeInfos> PinToAttributeInfos;
-
 	UPROPERTY()
 	FPCGPinDataLabels PinDataLabels;
 };
diff --git a/Engine/Plugins/PCG/Source/PCG/Private/Compute/PCGDataForGPU.cpp b/Engine/Plugins/PCG/Source/PCG/Private/Compute/PCGDataForGPU.cpp
--- a/Engine/Plugins/PCG/Source/PCG/Private/Compute/PCGDataForGPU.cpp
+++ b/Engine/Plugins/PCG/Source/PCG/Private/Compute/PCGDataForGPU.cpp
@@ -825,23 +825,20 @@ void FPCGDataDesc::AddAttribute(FPCGKernelAttributeKey InAttribute, const UPCGDa
 		return;
 	}
 
-	bool bAttributeFound = false;
-
-	for (FPCGKernelAttributeDesc& AttributeDesc : AttributeDescs)
+	// Remove existing attributes if they collide with the new attribute (that is, if they have the same Identifier). The new attribute will stomp them,
+	// which is consistent with the general behavior on CPU. The node is authoritative on its attributes types.
+	// @todo_pcg: We should probably encapsulate the AttributeDescs, and they can only be added through AddAttribute() to ensure consistent behavior (w.r.t. stomping existing attributes).
+	for (int AttributeIndex = AttributeDescs.Num() - 1; AttributeIndex >= 0; --AttributeIndex)
 	{
-		if (AttributeDesc.AttributeKey == InAttribute)
-		{
-			ensureMsgf(!bAttributeFound, TEXT("Attribute '%s' encountered multiple times in a single data."), *AttributeDesc.AttributeKey.Identifier.ToString());
-			bAttributeFound = true;
+		FPCGKernelAttributeDesc& ExistingAttributeDesc = AttributeDescs[AttributeIndex];
 
-			AttributeDesc.AttributeKey = InAttribute;
+		if (ExistingAttributeDesc.AttributeKey.Identifier == InAttribute.Identifier)
+		{
+			AttributeDescs.RemoveAtSwap(AttributeIndex);
 		}
 	}
 
-	if (!bAttributeFound)
-	{
-		AttributeDescs.Emplace(AttributeId, InAttribute.Type, InAttribute.Identifier, InOptionalUniqueStringKeys);
-	}
+	AttributeDescs.Emplace(AttributeId, InAttribute.Type, InAttribute.Identifier, InOptionalUniqueStringKeys);
 }
 
 int32 FPCGDataDesc::GetElementCountForAttribute(const FPCGKernelAttributeDesc& AttributeDesc) const
