[PCG] Fix broken determinism for GPU nodes that use additional sources
* We no longer create additional compute sources during compilation because their object names will end up in the cooked source, which causes compilation to be non-deterministic.
* Instead we gather compute existing compute sources. To allow this, internal PCG compute source files have been moved to PCG compute source assets

#rnx

diff --git a/Engine/Plugins/PCG/Shaders/Private/PCGQuaternion.ush b/Engine/Plugins/PCG/Shaders/Private/PCGQuaternion.ush
deleted file mode 100644
diff --git a/Engine/Plugins/PCG/Shaders/Private/PCGShaderUtils.ush b/Engine/Plugins/PCG/Shaders/Private/PCGShaderUtils.ush
deleted file mode 100644
diff --git a/Engine/Plugins/PCG/Shaders/Private/PCGShaderUtilsInternal.ush b/Engine/Plugins/PCG/Shaders/Private/PCGShaderUtilsInternal.ush
deleted file mode 100644
diff --git a/Engine/Plugins/PCG/Source/PCG/Private/Compute/PCGComputeKernel.cpp b/Engine/Plugins/PCG/Source/PCG/Private/Compute/PCGComputeKernel.cpp
--- a/Engine/Plugins/PCG/Source/PCG/Private/Compute/PCGComputeKernel.cpp
+++ b/Engine/Plugins/PCG/Source/PCG/Private/Compute/PCGComputeKernel.cpp
@@ -25,6 +25,10 @@ namespace PCGComputeKernel
 		TEXT("pcg.Graph.GPU.WarnOnOverridePinUsage"),
 		true,
 		TEXT("Enables warnings when parameters are overidden on GPU nodes."));
+
+	static const FSoftObjectPath ShaderUtilsComputeSourcePath(TEXT("PCGComputeSource'/PCG/ComputeSources/PCGCS_ShaderUtils.PCGCS_ShaderUtils'"));
+	static const FSoftObjectPath ShaderUtilsInternalComputeSourcePath(TEXT("PCGComputeSource'/PCG/ComputeSources/PCGCS_ShaderUtilsInternal.PCGCS_ShaderUtilsInternal'"));
+	static const FSoftObjectPath QuaternionComputeSourcePath(TEXT("PCGComputeSource'/PCG/ComputeSources/PCGCS_Quaternion.PCGCS_Quaternion'"));
 }
 
 #if WITH_EDITOR
@@ -66,26 +70,23 @@ bool UPCGComputeKernel::AreKernelSettingsValid(FPCGContext* InContext) const
 }
 
 #if WITH_EDITOR
-void UPCGComputeKernel::CreateAdditionalSources(FPCGGPUCompilationContext& InOutContext, UObject* InObjectOuter, TArray<TObjectPtr<UComputeSource>>& OutAdditionalSources) const
+void UPCGComputeKernel::GatherAdditionalSources(TArray<TObjectPtr<UComputeSource>>& OutAdditionalSources) const
 {
-	auto AddComputeSourceForFile = [&InOutContext, InObjectOuter, &OutAdditionalSources](const FString& InSourceFile)
+	auto GatherComputeSource = [&OutAdditionalSources](const FSoftObjectPath& InComputeSourcePath)
 	{
-		FString Source;
+		TSoftObjectPtr<UPCGComputeSource> ComputeSourcePtr(InComputeSourcePath);
+		UPCGComputeSource* ComputeSource = ComputeSourcePtr.LoadSynchronous();
+		check(ComputeSource);
 
-		if (ensure(LoadShaderSourceFile(*InSourceFile, EShaderPlatform::SP_PCD3D_SM5, &Source, nullptr)))
-		{
-			UPCGComputeSource* AdditionalSource = InOutContext.NewObject_AnyThread<UPCGComputeSource>(InObjectOuter);
-			AdditionalSource->SetSource(MoveTemp(Source));
-			OutAdditionalSources.Add(AdditionalSource);
-		}
+		OutAdditionalSources.Add(ComputeSource);
 	};
 
-	AddComputeSourceForFile(TEXT("/Engine/Private/ComputeShaderUtils.ush"));
-	AddComputeSourceForFile(TEXT("/Plugin/PCG/Private/PCGShaderUtils.ush"));
-	AddComputeSourceForFile(TEXT("/Plugin/PCG/Private/PCGShaderUtilsInternal.ush"));
+	GatherComputeSource(PCGComputeKernel::ShaderUtilsComputeSourcePath);
+	GatherComputeSource(PCGComputeKernel::ShaderUtilsInternalComputeSourcePath);
 
-	// Note: PCGDataCollectionDataInterface.ush depends on PCGQuaternion.ush, therefore all kernels also depend on PCGQuaternion.ush
-	AddComputeSourceForFile(TEXT("/Plugin/PCG/Private/PCGQuaternion.ush"));
+	// Note: PCGDataCollectionDataInterface.ush depends on the quaternion helpers, therefore all kernels also depend on the quaternion helpers.
+	// @todo_pcg: In the future quaternion compute source could be opt-in if the kernel does not manipulate point/attribute data.
+	GatherComputeSource(PCGComputeKernel::QuaternionComputeSourcePath);
 }
 #endif
 
diff --git a/Engine/Plugins/PCG/Source/PCG/Private/Compute/PCGComputeKernel.h b/Engine/Plugins/PCG/Source/PCG/Private/Compute/PCGComputeKernel.h
--- a/Engine/Plugins/PCG/Source/PCG/Private/Compute/PCGComputeKernel.h
+++ b/Engine/Plugins/PCG/Source/PCG/Private/Compute/PCGComputeKernel.h
@@ -87,7 +87,12 @@ public:
 	/** Get the name of the main kernel function in the source. This name is also displayed in GPU profile scopes. */
 	virtual FString GetEntryPoint() const PURE_VIRTUAL(UPCGComputeKernel::GetEntryPoint, return TEXT("Main"););
 
-	virtual void CreateAdditionalSources(FPCGGPUCompilationContext& InOutContext, UObject* InObjectOuter, TArray<TObjectPtr<class UComputeSource>>& OutAdditionalSources) const;
+	/**
+	 * Gathers additional compute sources referenced by this kernel.
+	 * This is preferred to directly using includes, since it allows us to detect hash diffs and recompile when additional sources are changed externally.
+	 * Note, you should not create new compute sources here, as their object names are not deterministic and will result in a new hash every time.
+	 */
+	virtual void GatherAdditionalSources(TArray<TObjectPtr<class UComputeSource>>& OutAdditionalSources) const;
 
 	/** Create additional input data interfaces to marshal any required input data. */
 	virtual void CreateAdditionalInputDataInterfaces(FPCGGPUCompilationContext& InOutContext, UObject* InObjectOuter, TArray<TObjectPtr<UComputeDataInterface>>& OutDataInterfaces) const {}
diff --git a/Engine/Plugins/PCG/Source/PCG/Private/Graph/PCGGraphCompilerGPU.cpp b/Engine/Plugins/PCG/Source/PCG/Private/Graph/PCGGraphCompilerGPU.cpp
--- a/Engine/Plugins/PCG/Source/PCG/Private/Graph/PCGGraphCompilerGPU.cpp
+++ b/Engine/Plugins/PCG/Source/PCG/Private/Graph/PCGGraphCompilerGPU.cpp
@@ -1463,9 +1463,10 @@ void FPCGGraphCompilerGPU::CompileComputeGraph(
 			KernelSource->EntryPoint = KernelWithBindings.Kernel->GetEntryPoint();
 			KernelSource->GroupSize = FIntVector(PCGComputeConstants::THREAD_GROUP_SIZE, 1, 1);
 
-			FString Source = KernelWithBindings.Kernel->GetCookedSource(InOutContext);
+			// All kernels require ComputeShaderUtils.ush, so inject that before anything else.
+			FString Source = TEXT("#include \"/Engine/Private/ComputeShaderUtils.ush\"\n\n") + KernelWithBindings.Kernel->GetCookedSource(InOutContext);
 			KernelSource->SetSource(MoveTemp(Source));
-			KernelWithBindings.Kernel->CreateAdditionalSources(InOutContext, KernelWithBindings.Kernel, KernelSource->AdditionalSources);
+			KernelWithBindings.Kernel->GatherAdditionalSources(KernelSource->AdditionalSources);
 
 			if (Settings->bDumpCookedHLSL)
 			{
