[PCG] Viewport display of zero points isn't behaving properly

#rnx

diff --git a/Engine/Plugins/PCG/Source/PCGEditor/Private/DataVisualizations/PCGSpatialDataVisualization.cpp b/Engine/Plugins/PCG/Source/PCGEditor/Private/DataVisualizations/PCGSpatialDataVisualization.cpp
--- a/Engine/Plugins/PCG/Source/PCGEditor/Private/DataVisualizations/PCGSpatialDataVisualization.cpp
+++ b/Engine/Plugins/PCG/Source/PCGEditor/Private/DataVisualizations/PCGSpatialDataVisualization.cpp
@@ -451,14 +451,17 @@ FPCGSetupSceneFunc IPCGSpatialDataVisualization::GetViewportSetupFunc(const UPCG
 			}
 		);
 
-		const FBoxSphereBounds Bounds = FBoxSphereBounds(FBox(BoundsMin, BoundsMax));
-		InOutParams.Scene->SetFloorOffset(-Bounds.Origin.Z + Bounds.BoxExtent.Z);
+		if (bInitializedBounds)
+		{
+			const FBoxSphereBounds Bounds = FBoxSphereBounds(FBox(BoundsMin, BoundsMax));
+			InOutParams.Scene->SetFloorOffset(-Bounds.Origin.Z + Bounds.BoxExtent.Z);
 
-		InOutParams.EditorViewportClient->SetViewLocation(Bounds.Origin);
+			InOutParams.EditorViewportClient->SetViewLocation(Bounds.Origin);
 
-		if (Bounds.SphereRadius > 0.0f)
-		{
-			InOutParams.EditorViewportClient->SetViewLocationForOrbiting(Bounds.Origin, Bounds.SphereRadius * 2.0f);
+			if (Bounds.SphereRadius > 0.0f)
+			{
+				InOutParams.EditorViewportClient->SetViewLocationForOrbiting(Bounds.Origin, Bounds.SphereRadius * 2.0f);
+			}
 		}
 	};
 }
