[PCG] Filter out object flags RF_WasLoaded & RF_LoadCompleted when duplicating PCG Settings (to avoid Dev/Debug VerifyObjectLoadFlagsWhenFinishedLoading ensure)
- Also remove RF_Transactional which has no use in execution context

#rnx

diff --git a/Engine/Plugins/PCG/Source/PCG/Private/PCGContext.cpp b/Engine/Plugins/PCG/Source/PCG/Private/PCGContext.cpp
--- a/Engine/Plugins/PCG/Source/PCG/Private/PCGContext.cpp
+++ b/Engine/Plugins/PCG/Source/PCG/Private/PCGContext.cpp
@@ -268,6 +268,11 @@ void FPCGContext::InitializeSettings(bool bSkipPostLoad)
 				FObjectDuplicationParameters DuplicateParams(const_cast<UPCGSettings*>(NodeSettings), GetTransientPackage());
 				DuplicateParams.bSkipPostLoad = bSkipPostLoad;
 				DuplicateParams.ApplyFlags = RF_Transient;
+				// Remove RF_WasLoaded and RF_LoadCompleted because those objects aren't loaded they are duplicated.
+				// When this happens outside of the GameThread the duplicated objects would end up with the Async flags + the Loaded flags which
+				// can trigger an ensure in VerifyObjectLoadFlagsWhenFinishedLoading (in !WITH_EDITOR dev/debug builds).
+				// We also remove the RF_Transactional flag as those duplicated objects are execution only objects and can't/should never be transacted.
+				DuplicateParams.FlagMask &= ~(RF_WasLoaded | RF_LoadCompleted | RF_Transactional);
 			
 				TMap<UObject*, UObject*> CreatedObjects;
 				DuplicateParams.CreatedObjects = &CreatedObjects;
