[PCG] GCGuard SoftObjectPtr resolving (because we are resolving outside of Gamethread)

diff --git a/Engine/Plugins/PCG/Source/PCG/Private/Graph/PCGGraphExecutor.cpp b/Engine/Plugins/PCG/Source/PCG/Private/Graph/PCGGraphExecutor.cpp
--- a/Engine/Plugins/PCG/Source/PCG/Private/Graph/PCGGraphExecutor.cpp
+++ b/Engine/Plugins/PCG/Source/PCG/Private/Graph/PCGGraphExecutor.cpp
@@ -3274,14 +3274,22 @@ UPCGGridLinkageSettings::UPCGGridLinkageSettings()
 
 FPCGElementPtr UPCGGridLinkageSettings::CreateElement() const
 {
-	auto GridLinkageOperation = [this](FPCGContext* InContext)
+	const UPCGPin* ResolvedUpstreamPin = nullptr;
+	{
+		FGCScopeGuard Guard;
+		ResolvedUpstreamPin = UpstreamPin.Get();
+	}
+
+	const FName UpstreamPinLabel = ResolvedUpstreamPin ? ResolvedUpstreamPin->Properties.Label : NAME_None;
+
+	auto GridLinkageOperation = [this, UpstreamPinLabel](FPCGContext* InContext)
 	{
 		return PCGGraphExecutor::ExecuteGridLinkage(
 			GenerationGrid,
 			FromGrid,
 			ToGrid,
 			ResourceKey,
-			UpstreamPin.Get() ? UpstreamPin->Properties.Label : NAME_None,
+			UpstreamPinLabel,
 			static_cast<FPCGGridLinkageContext*>(InContext));
 	};
 
@@ -3290,5 +3298,5 @@ FPCGElementPtr UPCGGridLinkageSettings::CreateElement() const
 		return new FPCGGridLinkageContext();
 	};
 
-	return MakeShared<PCGGraphExecutor::FPCGGridLinkageElement>(GridLinkageOperation, ContextAllocator, FromGrid, ToGrid, GenerationGrid, ResourceKey, UpstreamPin.Get());
+	return MakeShared<PCGGraphExecutor::FPCGGridLinkageElement>(GridLinkageOperation, ContextAllocator, FromGrid, ToGrid, GenerationGrid, ResourceKey, ResolvedUpstreamPin);
 }
\ No newline at end of file
