[PCG] Fix GPU Compute cooking errors when referencing static sub graphs with GPU nodes

#rnx

diff --git a/Engine/Plugins/PCG/Source/PCG/Private/Compute/BuiltInKernels/PCGCountUniqueAttributeValuesKernel.cpp b/Engine/Plugins/PCG/Source/PCG/Private/Compute/BuiltInKernels/PCGCountUniqueAttributeValuesKernel.cpp
--- a/Engine/Plugins/PCG/Source/PCG/Private/Compute/BuiltInKernels/PCGCountUniqueAttributeValuesKernel.cpp
+++ b/Engine/Plugins/PCG/Source/PCG/Private/Compute/BuiltInKernels/PCGCountUniqueAttributeValuesKernel.cpp
@@ -52,7 +52,7 @@ bool UPCGCountUniqueAttributeValuesKernel::IsKernelDataValid(FPCGContext* InCont
 		{
 			if (!InputDataDesc->DataDescs.IsEmpty())
 			{
-				PCG_KERNEL_VALIDATION_ERR(InContext, Settings, FText::Format(
+				PCG_KERNEL_VALIDATION_ERR(InContext, GetSettings(), FText::Format(
 					LOCTEXT("AttributeMissing", "Count attribute '{0}' not found, this attribute must be present on all input data, and be of type String Key."),
 					FText::FromName(AttributeName)));
 			}
@@ -62,7 +62,7 @@ bool UPCGCountUniqueAttributeValuesKernel::IsKernelDataValid(FPCGContext* InCont
 
 		if (bConflictingTypesFound)
 		{
-			PCG_KERNEL_VALIDATION_ERR(InContext, Settings, FText::Format(
+			PCG_KERNEL_VALIDATION_ERR(InContext, GetSettings(), FText::Format(
 				LOCTEXT("AttributeHasMultipleTypes", "Count attribute '{0}' found with multiple types in input data, all attributes must be of type String Key."),
 				FText::FromName(AttributeName)));
 
@@ -72,7 +72,7 @@ bool UPCGCountUniqueAttributeValuesKernel::IsKernelDataValid(FPCGContext* InCont
 		if (AttributeDesc.AttributeKey.Type != EPCGKernelAttributeType::StringKey)
 		{
 			// Attribute value counting only currently supported for attributes of type StringKey.
-			PCG_KERNEL_VALIDATION_ERR(Context, Settings, FText::Format(
+			PCG_KERNEL_VALIDATION_ERR(Context, GetSettings(), FText::Format(
 				LOCTEXT("AttributeTypeInvalid", "Cannot count values for attribute '{0}', only attributes of type String Key are currently supported."),
 				FText::FromName(AttributeName)));
 
diff --git a/Engine/Plugins/PCG/Source/PCG/Private/Compute/BuiltInKernels/PCGDataNumKernel.cpp b/Engine/Plugins/PCG/Source/PCG/Private/Compute/BuiltInKernels/PCGDataNumKernel.cpp
--- a/Engine/Plugins/PCG/Source/PCG/Private/Compute/BuiltInKernels/PCGDataNumKernel.cpp
+++ b/Engine/Plugins/PCG/Source/PCG/Private/Compute/BuiltInKernels/PCGDataNumKernel.cpp
@@ -61,7 +61,7 @@ bool UPCGDataNumKernel::IsKernelDataValid(FPCGContext* InContext) const
 			{
 				if (OutCountAttributeKey.Type == EPCGKernelAttributeType::Invalid)
 				{
-					PCG_KERNEL_VALIDATION_ERR(InContext, Settings, FText::Format(
+					PCG_KERNEL_VALIDATION_ERR(InContext, GetSettings(), FText::Format(
 						LOCTEXT("InvalidAttributeKeyType", "Output attribute for the count could not be created with an 'invalid' type."),
 						FText::FromString(OutCountAttributeKey.Identifier.ToString())));
 
@@ -72,7 +72,7 @@ bool UPCGDataNumKernel::IsKernelDataValid(FPCGContext* InContext) const
 			}
 			else
 			{
-				PCG_KERNEL_VALIDATION_ERR(InContext, Settings, FText::Format(
+				PCG_KERNEL_VALIDATION_ERR(InContext, GetSettings(), FText::Format(
 					LOCTEXT("AttributeMissing", "Output attribute '{0}' not found."),
 					FText::FromString(OutCountAttributeKey.Identifier.ToString())));
 
diff --git a/Engine/Plugins/PCG/Source/PCG/Private/Compute/DataInterfaces/PCGComputeDataInterface.cpp b/Engine/Plugins/PCG/Source/PCG/Private/Compute/DataInterfaces/PCGComputeDataInterface.cpp
--- a/Engine/Plugins/PCG/Source/PCG/Private/Compute/DataInterfaces/PCGComputeDataInterface.cpp
+++ b/Engine/Plugins/PCG/Source/PCG/Private/Compute/DataInterfaces/PCGComputeDataInterface.cpp
@@ -33,6 +33,31 @@ void UPCGComputeDataInterface::SetOutputPin(FName InOutputPinLabel, const FName*
 	OutputPinLabelAlias = InOptionalOutputPinLabelAlias ? *InOptionalOutputPinLabelAlias : InOutputPinLabel;
 }
 
+void UPCGComputeDataInterface::SetProducerSettings(const UPCGSettings* InProducerSettings)
+{
+	ResolvedProducerSettings = InProducerSettings;
+	ProducerSettings = ResolvedProducerSettings;
+}
+
+const UPCGSettings* UPCGComputeDataInterface::GetProducerSettings() const
+{
+	if (!ResolvedProducerSettings)
+	{
+		FGCScopeGuard Guard;
+		ResolvedProducerSettings = ProducerSettings.Get();
+	}
+
+	return ResolvedProducerSettings;
+}
+
+void UPCGComputeDataInterface::AddReferencedObjects(UObject* InThis, FReferenceCollector& Collector)
+{
+	Super::AddReferencedObjects(InThis, Collector);
+
+	UPCGComputeDataInterface* This = CastChecked<UPCGComputeDataInterface>(InThis);
+	Collector.AddReferencedObject(This->ResolvedProducerSettings);
+}
+
 void UPCGComputeDataProvider::Initialize(const UComputeDataInterface* InDataInterface, UObject* InBinding, uint64 InInputMask, uint64 InOutputMask)
 {
 	// Bump counter so any async callbacks from previous usages of this provider will be invalidated.
@@ -42,7 +67,9 @@ void UPCGComputeDataProvider::Initialize(const UComputeDataInterface* InDataInte
 
 	const UPCGComputeDataInterface* DataInterface = CastChecked<UPCGComputeDataInterface>(InDataInterface);
 	ProducerKernel = DataInterface->ProducerKernel;
-	ProducerSettings = DataInterface->ProducerSettings;
+	ResolvedProducerSettings = DataInterface->GetProducerSettings();
+	ProducerSettings = ResolvedProducerSettings;
+
 	GraphBindingIndex = DataInterface->GraphBindingIndex;
 
 	// The original label is needed to store inspection data.
@@ -63,12 +90,32 @@ void UPCGComputeDataProvider::Reset()
 
 	ProducerKernel = nullptr;
 	ProducerSettings = nullptr;
+	ResolvedProducerSettings = nullptr;
 	GraphBindingIndex = INDEX_NONE;
 	OutputPinLabel = NAME_None;
 	OutputPinLabelAlias = NAME_None;
 	DownstreamInputPinLabelAliases.Empty();
 }
 
+const UPCGSettings* UPCGComputeDataProvider::GetProducerSettings() const
+{
+	if (!ResolvedProducerSettings)
+	{
+		FGCScopeGuard Guard;
+		ResolvedProducerSettings = ProducerSettings.Get();
+	}
+
+	return ResolvedProducerSettings;
+}
+
+void UPCGComputeDataProvider::AddReferencedObjects(UObject* InThis, FReferenceCollector& Collector)
+{
+	Super::AddReferencedObjects(InThis, Collector);
+
+	UPCGComputeDataProvider* This = CastChecked<UPCGComputeDataProvider>(InThis);
+	Collector.AddReferencedObject(This->ResolvedProducerSettings);
+}
+
 #if WITH_EDITOR
 void UPCGComputeDataProvider::NotifyProducerUploadedData(UPCGDataBinding* InBinding)
 {
diff --git a/Engine/Plugins/PCG/Source/PCG/Private/Compute/DataInterfaces/PCGComputeDataInterface.h b/Engine/Plugins/PCG/Source/PCG/Private/Compute/DataInterfaces/PCGComputeDataInterface.h
--- a/Engine/Plugins/PCG/Source/PCG/Private/Compute/DataInterfaces/PCGComputeDataInterface.h
+++ b/Engine/Plugins/PCG/Source/PCG/Private/Compute/DataInterfaces/PCGComputeDataInterface.h
@@ -24,8 +24,10 @@ public:
 	/** Set the output pin label and label alias that this data interface is associated with. */
 	virtual void SetOutputPin(FName InOutputPinLabel, const FName* InOptionalOutputPinLabelAlias = nullptr);
 
-	void SetProducerSettings(TObjectPtr<const UPCGSettings> InProducerSettings) { ProducerSettings = InProducerSettings; }
-	TObjectPtr<const UPCGSettings> GetProducerSettings() const { return ProducerSettings; }
+	void SetProducerSettings(const UPCGSettings* InProducerSettings);
+	const UPCGSettings* GetProducerSettings() const;
+
+	static void AddReferencedObjects(UObject* InThis, FReferenceCollector& Collector);
 
 public:
 	UPROPERTY()
@@ -34,9 +36,6 @@ public:
 	UPROPERTY()
 	TObjectPtr<const UPCGComputeKernel> ProducerKernel;
 
-	UPROPERTY()
-	TObjectPtr<const UPCGSettings> ProducerSettings;
-
 	/** Label of output pin that this data interface is associated with. */
 	UPROPERTY()
 	FName OutputPinLabel;
@@ -47,6 +46,12 @@ public:
 	/** Generated PCG data will be assigned these labels. */
 	UPROPERTY()
 	TArray<FName> DownstreamInputPinLabelAliases;
+
+protected:
+	UPROPERTY()
+	TSoftObjectPtr<const UPCGSettings> ProducerSettings;
+
+	mutable TObjectPtr<const UPCGSettings> ResolvedProducerSettings = nullptr;
 };
 
 UCLASS(Abstract, ClassGroup (Procedural))
@@ -70,7 +75,9 @@ public:
 	/** Release any handles to transient resources like GPU buffers. */
 	virtual void ReleaseTransientResources(const TCHAR* InReason = nullptr) {}
 
-	TObjectPtr<const UPCGSettings> GetProducerSettings() const { return ProducerSettings; }
+	const UPCGSettings* GetProducerSettings() const;
+
+	static void AddReferencedObjects(UObject* InThis, FReferenceCollector& Collector);
 
 protected:
 #if WITH_EDITOR
@@ -92,7 +99,9 @@ protected:
 
 	/** Settings of node that produces this data, normally the upstream node. */
 	UPROPERTY()
-	TObjectPtr<const UPCGSettings> ProducerSettings;
+	TSoftObjectPtr<const UPCGSettings> ProducerSettings;
+
+	mutable TObjectPtr<const UPCGSettings> ResolvedProducerSettings = nullptr;
 };
 
 UCLASS(Abstract, ClassGroup = (Procedural))
diff --git a/Engine/Plugins/PCG/Source/PCG/Private/Compute/DataInterfaces/PCGCustomKernelDataInterface.cpp b/Engine/Plugins/PCG/Source/PCG/Private/Compute/DataInterfaces/PCGCustomKernelDataInterface.cpp
--- a/Engine/Plugins/PCG/Source/PCG/Private/Compute/DataInterfaces/PCGCustomKernelDataInterface.cpp
+++ b/Engine/Plugins/PCG/Source/PCG/Private/Compute/DataInterfaces/PCGCustomKernelDataInterface.cpp
@@ -16,6 +16,31 @@
 
 #include UE_INLINE_GENERATED_CPP_BY_NAME(PCGCustomKernelDataInterface)
 
+void UPCGCustomKernelDataInterface::SetSettings(const UPCGSettings* InSettings)
+{
+	ResolvedSettings = InSettings;
+	Settings = ResolvedSettings;
+}
+
+const UPCGSettings* UPCGCustomKernelDataInterface::GetSettings() const
+{
+	if (!ResolvedSettings)
+	{
+		FGCScopeGuard Guard;
+		ResolvedSettings = Settings.Get();
+	}
+	
+	return ResolvedSettings;
+}
+
+void UPCGCustomKernelDataInterface::AddReferencedObjects(UObject* InThis, FReferenceCollector& Collector)
+{
+	Super::AddReferencedObjects(InThis, Collector);
+
+	UPCGCustomKernelDataInterface* This = CastChecked<UPCGCustomKernelDataInterface>(InThis);
+	Collector.AddReferencedObject(This->ResolvedSettings);
+}
+
 void UPCGCustomKernelDataInterface::GetSupportedInputs(TArray<FShaderFunctionDefinition>& OutFunctions) const
 {
 	OutFunctions.AddDefaulted_GetRef()
@@ -105,13 +130,15 @@ void UPCGCustomComputeKernelDataProvider::Initialize(const UComputeDataInterface
 	UPCGDataBinding* Binding = CastChecked<UPCGDataBinding>(InBinding);
 	const IPCGGraphExecutionSource* ExecutionSource = Binding->GetExecutionSource();
 	check(ExecutionSource);
-	check(DataInterface->Settings);
+
+	const UPCGSettings* DataInterfaceSettings = DataInterface->GetSettings();
+	check(DataInterfaceSettings);
 	check(DataInterface->Kernel);
 
 	Kernel = DataInterface->Kernel;
 	ThreadCountMultiplier = DataInterface->Kernel->GetThreadCountMultiplier();
-	Seed = static_cast<uint32>(DataInterface->Settings->GetSeed(ExecutionSource));
-	SeedSettings = static_cast<uint32>(DataInterface->Settings->Seed);
+	Seed = static_cast<uint32>(DataInterfaceSettings->GetSeed(ExecutionSource));
+	SeedSettings = static_cast<uint32>(DataInterfaceSettings->Seed);
 	SeedComponent = static_cast<uint32>(ExecutionSource->GetExecutionState().GetSeed());
 
 	SourceComponentBounds = ExecutionSource->GetExecutionState().GetBounds();
diff --git a/Engine/Plugins/PCG/Source/PCG/Private/Compute/DataInterfaces/PCGCustomKernelDataInterface.h b/Engine/Plugins/PCG/Source/PCG/Private/Compute/DataInterfaces/PCGCustomKernelDataInterface.h
--- a/Engine/Plugins/PCG/Source/PCG/Private/Compute/DataInterfaces/PCGCustomKernelDataInterface.h
+++ b/Engine/Plugins/PCG/Source/PCG/Private/Compute/DataInterfaces/PCGCustomKernelDataInterface.h
@@ -32,12 +32,20 @@ public:
 	// This DI will provide execution parameters like dispatch information.
 	bool IsExecutionInterface() const override { return true; }
 	//~ End UComputeDataInterface Interface
+		
+	void SetSettings(const UPCGSettings* InSettings);
+	const UPCGSettings* GetSettings() const;
 
 	UPROPERTY()
-	TObjectPtr<const UPCGSettings> Settings = nullptr;
+	TObjectPtr<const UPCGComputeKernel> Kernel = nullptr;
+
+	static void AddReferencedObjects(UObject* InThis, FReferenceCollector& Collector);
 
+protected:
 	UPROPERTY()
-	TObjectPtr<const UPCGComputeKernel> Kernel = nullptr;
+	TSoftObjectPtr<const UPCGSettings> Settings = nullptr;
+
+	mutable TObjectPtr<const UPCGSettings> ResolvedSettings = nullptr;
 };
 
 /** Compute Framework Data Provider for each custom compute kernel. */
diff --git a/Engine/Plugins/PCG/Source/PCG/Private/Compute/DataInterfaces/PCGDataCollectionUploadDataInterface.cpp b/Engine/Plugins/PCG/Source/PCG/Private/Compute/DataInterfaces/PCGDataCollectionUploadDataInterface.cpp
--- a/Engine/Plugins/PCG/Source/PCG/Private/Compute/DataInterfaces/PCGDataCollectionUploadDataInterface.cpp
+++ b/Engine/Plugins/PCG/Source/PCG/Private/Compute/DataInterfaces/PCGDataCollectionUploadDataInterface.cpp
@@ -32,7 +32,7 @@ void UPCGDataProviderDataCollectionUpload::Initialize(const UComputeDataInterfac
 	const UPCGDataCollectionUploadDataInterface* DataInterface = CastChecked<UPCGDataCollectionUploadDataInterface>(InDataInterface);
 
 	Binding = CastChecked<UPCGDataBinding>(InBinding);
-	ProducerSettings = DataInterface->ProducerSettings;
+	ProducerSettings = DataInterface->GetProducerSettings();
 	GraphBindingIndex = DataInterface->GraphBindingIndex;
 
 	// Pick the data items from input data collection using any of the compute graph element virtual input pin labels.
diff --git a/Engine/Plugins/PCG/Source/PCG/Private/Compute/DataInterfaces/PCGLandscapeDataInterface.cpp b/Engine/Plugins/PCG/Source/PCG/Private/Compute/DataInterfaces/PCGLandscapeDataInterface.cpp
--- a/Engine/Plugins/PCG/Source/PCG/Private/Compute/DataInterfaces/PCGLandscapeDataInterface.cpp
+++ b/Engine/Plugins/PCG/Source/PCG/Private/Compute/DataInterfaces/PCGLandscapeDataInterface.cpp
@@ -277,7 +277,7 @@ void UPCGLandscapeDataProvider::Initialize(const UComputeDataInterface* InDataIn
 	Super::Initialize(InDataInterface, InBinding, InInputMask, InOutputMask);
 
 	const UPCGLandscapeDataInterface* DataInterface = CastChecked<UPCGLandscapeDataInterface>(InDataInterface);
-	ProducerSettings = DataInterface->ProducerSettings;
+	ProducerSettings = DataInterface->GetProducerSettings();
 
 	UPCGDataBinding* Binding = CastChecked<UPCGDataBinding>(InBinding);
 
diff --git a/Engine/Plugins/PCG/Source/PCG/Private/Compute/Elements/PCGCustomHLSLKernel.cpp b/Engine/Plugins/PCG/Source/PCG/Private/Compute/Elements/PCGCustomHLSLKernel.cpp
--- a/Engine/Plugins/PCG/Source/PCG/Private/Compute/Elements/PCGCustomHLSLKernel.cpp
+++ b/Engine/Plugins/PCG/Source/PCG/Private/Compute/Elements/PCGCustomHLSLKernel.cpp
@@ -197,7 +197,7 @@ bool UPCGCustomHLSLKernel::IsKernelDataValid(FPCGContext* InContext) const
 		{
 			if (ErrorTextPtr)
 			{
-				PCG_KERNEL_VALIDATION_ERR(InContext, Settings, *ErrorTextPtr);
+				PCG_KERNEL_VALIDATION_ERR(InContext, GetSettings(), *ErrorTextPtr);
 			}
 			return false;
 		}
@@ -211,7 +211,7 @@ bool UPCGCustomHLSLKernel::ComputeOutputBindingDataDesc(const UPCGComputeGraph*
 	check(InGraph);
 	check(InBinding);
 	
-	const UPCGCustomHLSLSettings* CustomHLSLSettings = CastChecked<UPCGCustomHLSLSettings>(Settings);
+	const UPCGCustomHLSLSettings* CustomHLSLSettings = CastChecked<UPCGCustomHLSLSettings>(GetSettings());
 
 	FPCGDataCollectionDesc OutputDataDesc;
 
@@ -381,7 +381,7 @@ bool UPCGCustomHLSLKernel::ComputeOutputBindingDataDesc(const UPCGComputeGraph*
 #if WITH_EDITOR
 FString UPCGCustomHLSLKernel::GetCookedSource(FPCGGPUCompilationContext& InOutContext) const
 {
-	const UPCGCustomHLSLSettings* CustomHLSLSettings = CastChecked<UPCGCustomHLSLSettings>(Settings);
+	const UPCGCustomHLSLSettings* CustomHLSLSettings = CastChecked<UPCGCustomHLSLSettings>(GetSettings());
 	FString ShaderPathName = CustomHLSLSettings->GetPathName();
 	PCGComputeHelpers::ConvertObjectPathToShaderFilePath(ShaderPathName);
 
@@ -552,7 +552,7 @@ FString UPCGCustomHLSLKernel::GetCookedSource(FPCGGPUCompilationContext& InOutCo
 
 int UPCGCustomHLSLKernel::ComputeThreadCount(const UPCGDataBinding* InBinding) const
 {
-	const UPCGCustomHLSLSettings* CustomHLSLSettings = CastChecked<UPCGCustomHLSLSettings>(Settings);
+	const UPCGCustomHLSLSettings* CustomHLSLSettings = CastChecked<UPCGCustomHLSLSettings>(GetSettings());
 	int ThreadCount = 0;
 
 	if (CustomHLSLSettings->KernelType == EPCGKernelType::PointGenerator)
@@ -626,13 +626,13 @@ void UPCGCustomHLSLKernel::GetKernelAttributeKeys(TArray<FPCGKernelAttributeKey>
 
 uint32 UPCGCustomHLSLKernel::GetThreadCountMultiplier() const
 {
-	const UPCGCustomHLSLSettings* CustomHLSLSettings = CastChecked<UPCGCustomHLSLSettings>(Settings);
+	const UPCGCustomHLSLSettings* CustomHLSLSettings = CastChecked<UPCGCustomHLSLSettings>(GetSettings());
 	return CustomHLSLSettings->ThreadCountMultiplier;
 }
 
 uint32 UPCGCustomHLSLKernel::GetElementCountMultiplier(FName InOutputPinLabel) const
 {
-	const UPCGCustomHLSLSettings* CustomHLSLSettings = CastChecked<UPCGCustomHLSLSettings>(Settings);
+	const UPCGCustomHLSLSettings* CustomHLSLSettings = CastChecked<UPCGCustomHLSLSettings>(GetSettings());
 
 	if (const FPCGPinPropertiesGPU* PinProps = CustomHLSLSettings->OutputPins.FindByPredicate([InOutputPinLabel](const FPCGPinPropertiesGPU& InProps) { return InProps.Label == InOutputPinLabel; }))
 	{
@@ -647,13 +647,13 @@ uint32 UPCGCustomHLSLKernel::GetElementCountMultiplier(FName InOutputPinLabel) c
 
 void UPCGCustomHLSLKernel::GetInputPins(TArray<FPCGPinProperties>& OutPins) const
 {
-	const UPCGCustomHLSLSettings* CustomHLSLSettings = CastChecked<UPCGCustomHLSLSettings>(Settings);
+	const UPCGCustomHLSLSettings* CustomHLSLSettings = CastChecked<UPCGCustomHLSLSettings>(GetSettings());
 	OutPins = CustomHLSLSettings->InputPins;
 }
 
 void UPCGCustomHLSLKernel::GetOutputPins(TArray<FPCGPinPropertiesGPU>& OutPins) const
 {
-	const UPCGCustomHLSLSettings* CustomHLSLSettings = CastChecked<UPCGCustomHLSLSettings>(Settings);
+	const UPCGCustomHLSLSettings* CustomHLSLSettings = CastChecked<UPCGCustomHLSLSettings>(GetSettings());
 	OutPins = CustomHLSLSettings->OutputPins;
 }
 
@@ -737,7 +737,7 @@ void UPCGCustomHLSLKernel::ParseShaderSource()
 	TArray<FString> CopyElementKeywords;
 	TArray<FString> InitializeKeywords;
 
-	const UPCGCustomHLSLSettings* CustomHLSLSettings = CastChecked<UPCGCustomHLSLSettings>(Settings);
+	const UPCGCustomHLSLSettings* CustomHLSLSettings = CastChecked<UPCGCustomHLSLSettings>(GetSettings());
 
 	for (const FPCGPinProperties& PinProps : CustomHLSLSettings->InputPins)
 	{
@@ -1001,7 +1001,7 @@ void UPCGCustomHLSLKernel::CreateParsedSources()
 		AdditionalSources.Add(AdditionalSource);
 	};
 
-	const UPCGCustomHLSLSettings* CustomHLSLSettings = CastChecked<UPCGCustomHLSLSettings>(Settings);
+	const UPCGCustomHLSLSettings* CustomHLSLSettings = CastChecked<UPCGCustomHLSLSettings>(GetSettings());
 	TArray<TObjectPtr<UComputeSource>> AdditionalSourcesToProcess = CustomHLSLSettings->AdditionalSources;
 
 	if (TObjectPtr<UComputeSource> KernelSourceOverride = CustomHLSLSettings->KernelSourceOverride)
@@ -1048,7 +1048,7 @@ void UPCGCustomHLSLKernel::CollectDataLabels(const FPCGCustomHLSLParsedSource& I
 		}
 	};
 
-	const UPCGCustomHLSLSettings* CustomHLSLSettings = CastChecked<UPCGCustomHLSLSettings>(Settings);
+	const UPCGCustomHLSLSettings* CustomHLSLSettings = CastChecked<UPCGCustomHLSLSettings>(GetSettings());
 
 	for (const FPCGPinProperties& PinProps : CustomHLSLSettings->InputPins)
 	{
@@ -1070,7 +1070,7 @@ bool UPCGCustomHLSLKernel::PerformStaticValidation()
 		return false;
 	}
 
-	const UPCGCustomHLSLSettings* CustomHLSLSettings = CastChecked<UPCGCustomHLSLSettings>(Settings);
+	const UPCGCustomHLSLSettings* CustomHLSLSettings = CastChecked<UPCGCustomHLSLSettings>(GetSettings());
 
 	if (CustomHLSLSettings->OutputPins.IsEmpty())
 	{
@@ -1408,7 +1408,7 @@ bool UPCGCustomHLSLKernel::PerformStaticValidation()
 
 bool UPCGCustomHLSLKernel::ValidateShaderSource(TArray<FPCGKernelLogEntry>& OutLogEntries) const
 {
-	const UPCGCustomHLSLSettings* CustomHLSLSettings = CastChecked<UPCGCustomHLSLSettings>(Settings);
+	const UPCGCustomHLSLSettings* CustomHLSLSettings = CastChecked<UPCGCustomHLSLSettings>(GetSettings());
 
 	if (!CustomHLSLSettings->bMuteUnwrittenPinDataErrors && !AreAllOutputPinsWritten(OutLogEntries))
 	{
@@ -1543,7 +1543,7 @@ bool UPCGCustomHLSLKernel::AreAllOutputPinsWritten(TArray<FPCGKernelLogEntry>& O
 		return InParsedSource.InitializedOutputPins.Contains(PinStr);
 	};
 
-	const UPCGCustomHLSLSettings* CustomHLSLSettings = CastChecked<UPCGCustomHLSLSettings>(Settings);
+	const UPCGCustomHLSLSettings* CustomHLSLSettings = CastChecked<UPCGCustomHLSLSettings>(GetSettings());
 
 	// Processor/Generator kernels initialize the first output pin data automatically.
 	const bool bSkipFirstPin = CustomHLSLSettings->IsProcessorKernel() || CustomHLSLSettings->IsGeneratorKernel();
@@ -1586,7 +1586,7 @@ bool UPCGCustomHLSLKernel::AreAllOutputPinsWritten(TArray<FPCGKernelLogEntry>& O
 
 bool UPCGCustomHLSLKernel::IsThreadCountMultiplierInUse() const
 {
-	const UPCGCustomHLSLSettings* CustomHLSLSettings = CastChecked<UPCGCustomHLSLSettings>(Settings);
+	const UPCGCustomHLSLSettings* CustomHLSLSettings = CastChecked<UPCGCustomHLSLSettings>(GetSettings());
 	return CustomHLSLSettings->KernelType == EPCGKernelType::Custom && CustomHLSLSettings->DispatchThreadCount != EPCGDispatchThreadCount::Fixed;
 }
 
@@ -1601,7 +1601,7 @@ bool UPCGCustomHLSLKernel::AreAttributesValid(FPCGContext* InContext, FText* Out
 		return true;
 	}
 
-	const UPCGCustomHLSLSettings* CustomHLSLSettings = CastChecked<UPCGCustomHLSLSettings>(Settings);
+	const UPCGCustomHLSLSettings* CustomHLSLSettings = CastChecked<UPCGCustomHLSLSettings>(GetSettings());
 	const FPCGComputeGraphContext* Context = static_cast<FPCGComputeGraphContext*>(InContext);
 	const UPCGDataBinding* DataBinding = Context->DataBinding.Get();
 
@@ -1782,13 +1782,13 @@ bool UPCGCustomHLSLKernel::AreAttributesValid(FPCGContext* InContext, FText* Out
 
 const FPCGPinProperties* UPCGCustomHLSLKernel::GetFirstInputPin() const
 {
-	const UPCGCustomHLSLSettings* CustomHLSLSettings = CastChecked<UPCGCustomHLSLSettings>(Settings);
+	const UPCGCustomHLSLSettings* CustomHLSLSettings = CastChecked<UPCGCustomHLSLSettings>(GetSettings());
 	return !CustomHLSLSettings->InputPins.IsEmpty() ? &CustomHLSLSettings->InputPins[0] : nullptr;
 }
 
 const FPCGPinPropertiesGPU* UPCGCustomHLSLKernel::GetFirstOutputPin() const
 {
-	const UPCGCustomHLSLSettings* CustomHLSLSettings = CastChecked<UPCGCustomHLSLSettings>(Settings);
+	const UPCGCustomHLSLSettings* CustomHLSLSettings = CastChecked<UPCGCustomHLSLSettings>(GetSettings());
 	return !CustomHLSLSettings->OutputPins.IsEmpty() ? &CustomHLSLSettings->OutputPins[0] : nullptr;
 }
 
diff --git a/Engine/Plugins/PCG/Source/PCG/Private/Compute/Elements/PCGNormalToDensityKernel.cpp b/Engine/Plugins/PCG/Source/PCG/Private/Compute/Elements/PCGNormalToDensityKernel.cpp
--- a/Engine/Plugins/PCG/Source/PCG/Private/Compute/Elements/PCGNormalToDensityKernel.cpp
+++ b/Engine/Plugins/PCG/Source/PCG/Private/Compute/Elements/PCGNormalToDensityKernel.cpp
@@ -11,11 +11,6 @@
 
 #define LOCTEXT_NAMESPACE "PCGTemplateKernel"
 
-const UPCGSettings* UPCGNormalToDensityKernel::GetSettings() const
-{
-	return Settings;
-}
-
 bool UPCGNormalToDensityKernel::ComputeOutputBindingDataDesc(const UPCGComputeGraph* InGraph, FName InOutputPinLabel, UPCGDataBinding* InBinding, FPCGDataCollectionDesc& OutDataDesc) const
 {
 	check(InGraph);
diff --git a/Engine/Plugins/PCG/Source/PCG/Private/Compute/Elements/PCGNormalToDensityKernel.h b/Engine/Plugins/PCG/Source/PCG/Private/Compute/Elements/PCGNormalToDensityKernel.h
--- a/Engine/Plugins/PCG/Source/PCG/Private/Compute/Elements/PCGNormalToDensityKernel.h
+++ b/Engine/Plugins/PCG/Source/PCG/Private/Compute/Elements/PCGNormalToDensityKernel.h
@@ -15,7 +15,6 @@ class UPCGNormalToDensityKernel : public UPCGComputeKernel
 	GENERATED_BODY()
 
 public:
-	virtual const UPCGSettings* GetSettings() const override;
 	virtual bool ComputeOutputBindingDataDesc(const UPCGComputeGraph* InGraph, FName InOutputPinLabel, UPCGDataBinding* InBinding, FPCGDataCollectionDesc& OutDataDesc) const override;
 	virtual int ComputeThreadCount(const UPCGDataBinding* Binding) const override;
 
diff --git a/Engine/Plugins/PCG/Source/PCG/Private/Compute/PCGComputeKernel.cpp b/Engine/Plugins/PCG/Source/PCG/Private/Compute/PCGComputeKernel.cpp
--- a/Engine/Plugins/PCG/Source/PCG/Private/Compute/PCGComputeKernel.cpp
+++ b/Engine/Plugins/PCG/Source/PCG/Private/Compute/PCGComputeKernel.cpp
@@ -34,7 +34,8 @@ namespace PCGComputeKernel
 #if WITH_EDITOR
 void UPCGComputeKernel::Initialize(const FPCGComputeKernelParams& InParams)
 {
-	Settings = InParams.Settings;
+	ResolvedSettings = InParams.Settings;
+	Settings = ResolvedSettings;
 	bLogDataDescriptions = InParams.bLogDescriptions;
 	bInitialized = true;
 
@@ -44,11 +45,30 @@ void UPCGComputeKernel::Initialize(const FPCGComputeKernelParams& InParams)
 }
 #endif
 
+const UPCGSettings* UPCGComputeKernel::GetSettings() const
+{
+	if (!ResolvedSettings)
+	{
+		FGCScopeGuard Guard;
+		ResolvedSettings = Settings.Get();
+	}
+
+	return ResolvedSettings;
+}
+
+void UPCGComputeKernel::AddReferencedObjects(UObject* InThis, FReferenceCollector& Collector)
+{
+	Super::AddReferencedObjects(InThis, Collector);
+
+	UPCGComputeKernel* This = CastChecked<UPCGComputeKernel>(InThis);
+	Collector.AddReferencedObject(This->ResolvedSettings);
+}
+
 bool UPCGComputeKernel::AreKernelSettingsValid(FPCGContext* InContext) const
 {
 	if (!bInitialized)
 	{
-		PCG_KERNEL_VALIDATION_ERR(InContext, Settings, LOCTEXT("UninitializeKernel", "Kernel was not initialized during compilation. Make sure to call Initialize() when creating your kernels."));
+		PCG_KERNEL_VALIDATION_ERR(InContext, GetSettings(), LOCTEXT("UninitializeKernel", "Kernel was not initialized during compilation. Make sure to call Initialize() when creating your kernels."));
 		return false;
 	}
 
@@ -57,11 +77,11 @@ bool UPCGComputeKernel::AreKernelSettingsValid(FPCGContext* InContext) const
 	{
 		if (StaticLogEntry.Verbosity == EPCGKernelLogVerbosity::Error)
 		{
-			PCG_KERNEL_VALIDATION_ERR(InContext, Settings, StaticLogEntry.Message);
+			PCG_KERNEL_VALIDATION_ERR(InContext, GetSettings(), StaticLogEntry.Message);
 		}
 		else
 		{
-			PCG_KERNEL_VALIDATION_WARN(InContext, Settings, StaticLogEntry.Message);
+			PCG_KERNEL_VALIDATION_WARN(InContext, GetSettings(), StaticLogEntry.Message);
 		}
 	}
 #endif
diff --git a/Engine/Plugins/PCG/Source/PCG/Private/Compute/PCGComputeKernel.h b/Engine/Plugins/PCG/Source/PCG/Private/Compute/PCGComputeKernel.h
--- a/Engine/Plugins/PCG/Source/PCG/Private/Compute/PCGComputeKernel.h
+++ b/Engine/Plugins/PCG/Source/PCG/Private/Compute/PCGComputeKernel.h
@@ -72,7 +72,7 @@ public:
 	void SetKernelIndex(int32 InKernelIndex) { KernelIndex = InKernelIndex; }
 
 	/** Gets settings for node associated with this kernel, if any. */
-	virtual const UPCGSettings* GetSettings() const { return Settings; }
+	virtual const UPCGSettings* GetSettings() const;
 
 	/** Performs settings validation and returns true if this node is suitable for execution. */
 	virtual bool AreKernelSettingsValid(FPCGContext* InContext) const;
@@ -148,6 +148,7 @@ public:
 
 	virtual bool GetLogDataDescriptions() const { return bLogDataDescriptions; }
 
+	static void AddReferencedObjects(UObject* InThis, FReferenceCollector& Collector);
 protected:
 #if WITH_EDITOR
 	/** Implement to do any additional initialization specific to your kernel. */
@@ -169,7 +170,7 @@ protected:
 	int32 KernelIndex = INDEX_NONE;
 
 	UPROPERTY()
-	TObjectPtr<const UPCGSettings> Settings;
+	TSoftObjectPtr<const UPCGSettings> Settings;
 
 	UPROPERTY()
 	bool bLogDataDescriptions = false;
@@ -188,4 +189,6 @@ protected:
 #if WITH_EDITOR
 	TSet<FName> InternalPinLabels;
 #endif
+
+	mutable TObjectPtr<const UPCGSettings> ResolvedSettings = nullptr;
 };
diff --git a/Engine/Plugins/PCG/Source/PCG/Private/Elements/Metadata/PCGMetadataPartitionKernel.cpp b/Engine/Plugins/PCG/Source/PCG/Private/Elements/Metadata/PCGMetadataPartitionKernel.cpp
--- a/Engine/Plugins/PCG/Source/PCG/Private/Elements/Metadata/PCGMetadataPartitionKernel.cpp
+++ b/Engine/Plugins/PCG/Source/PCG/Private/Elements/Metadata/PCGMetadataPartitionKernel.cpp
@@ -18,11 +18,6 @@
 
 #define LOCTEXT_NAMESPACE "PCGMetadataPartitionKernel"
 
-const UPCGSettings* UPCGMetadataPartitionKernel::GetSettings() const
-{
-	return Settings;
-}
-
 bool UPCGMetadataPartitionKernel::IsKernelDataValid(FPCGContext* InContext) const
 {
 	TRACE_CPUPROFILER_EVENT_SCOPE(UPCGMetadataPartitionKernel::IsKernelDataValid);
@@ -36,7 +31,7 @@ bool UPCGMetadataPartitionKernel::IsKernelDataValid(FPCGContext* InContext) cons
 
 	if (Context && Context->DataBinding)
 	{
-		const UPCGMetadataPartitionSettings* MPSettings = CastChecked<UPCGMetadataPartitionSettings>(Settings);
+		const UPCGMetadataPartitionSettings* MPSettings = CastChecked<UPCGMetadataPartitionSettings>(GetSettings());
 		const FName AttributeName = MPSettings->PartitionAttributeSelectors[0].GetAttributeName();
 
 		const FPCGDataCollectionDesc* InputDesc = Context->DataBinding->GetCachedKernelPinDataDesc(this, PCGPinConstants::DefaultInputLabel, /*bIsInput=*/true);
@@ -54,7 +49,7 @@ bool UPCGMetadataPartitionKernel::IsKernelDataValid(FPCGContext* InContext) cons
 		{
 			if (!InputDesc->DataDescs.IsEmpty())
 			{
-				PCG_KERNEL_VALIDATION_ERR(InContext, Settings, FText::Format(
+				PCG_KERNEL_VALIDATION_ERR(InContext, GetSettings(), FText::Format(
 					LOCTEXT("PartitionAttributeMissing", "Partition attribute was not present in GPU data."),
 					FText::FromName(AttributeName)));
 			}
@@ -64,7 +59,7 @@ bool UPCGMetadataPartitionKernel::IsKernelDataValid(FPCGContext* InContext) cons
 
 		if (bConflictingTypesInData)
 		{
-			PCG_KERNEL_VALIDATION_ERR(InContext, Settings, FText::Format(
+			PCG_KERNEL_VALIDATION_ERR(InContext, GetSettings(), FText::Format(
 				LOCTEXT("PartitionAttributeTypeConflict", "Attribute '{0}' encountered with multiple different types in input data."),
 				FText::FromName(AttributeName)));
 			return false;
@@ -72,7 +67,7 @@ bool UPCGMetadataPartitionKernel::IsKernelDataValid(FPCGContext* InContext) cons
 
 		if (AttributeDesc.AttributeKey.Type != EPCGKernelAttributeType::StringKey)
 		{
-			PCG_KERNEL_VALIDATION_ERR(Context, Settings, FText::Format(
+			PCG_KERNEL_VALIDATION_ERR(Context, GetSettings(), FText::Format(
 				LOCTEXT("PartitionAttributeTypeInvalid", "Attribute '{0}' not usable for partitioning, only attributes of type String Key are currently supported."),
 				FText::FromName(AttributeName)));
 			return false;
@@ -87,7 +82,7 @@ bool UPCGMetadataPartitionKernel::ComputeOutputBindingDataDesc(const UPCGCompute
 	check(InGraph);
 	check(InBinding);
 
-	const UPCGMetadataPartitionSettings* MPSettings = CastChecked<UPCGMetadataPartitionSettings>(Settings);
+	const UPCGMetadataPartitionSettings* MPSettings = CastChecked<UPCGMetadataPartitionSettings>(GetSettings());
 
 	// Code assumes single output pin.
 	if (!ensure(InOutputPinLabel == PCGPinConstants::DefaultOutputLabel))
@@ -165,7 +160,7 @@ void UPCGMetadataPartitionKernel::CreateAdditionalInputDataInterfaces(FPCGGPUCom
 
 void UPCGMetadataPartitionKernel::GetKernelAttributeKeys(TArray<FPCGKernelAttributeKey>& OutKeys) const
 {
-	const UPCGMetadataPartitionSettings* MPSettings = CastChecked<UPCGMetadataPartitionSettings>(Settings);
+	const UPCGMetadataPartitionSettings* MPSettings = CastChecked<UPCGMetadataPartitionSettings>(GetSettings());
 
 	// Register the attributes this node reads or writes for which we know the attribute type. Currently
 	// only StringKey attributes can be partitioned when executing on GPU, so declare the attribute here
@@ -203,7 +198,7 @@ bool UPCGMetadataPartitionKernel::PerformStaticValidation()
 		return false;
 	}
 	
-	const UPCGMetadataPartitionSettings* MPSettings = CastChecked<UPCGMetadataPartitionSettings>(Settings);
+	const UPCGMetadataPartitionSettings* MPSettings = CastChecked<UPCGMetadataPartitionSettings>(GetSettings());
 
 	if (MPSettings->PartitionAttributeSelectors.Num() != 1)
 	{
diff --git a/Engine/Plugins/PCG/Source/PCG/Private/Elements/PCGCopyPointsAnalysisKernel.cpp b/Engine/Plugins/PCG/Source/PCG/Private/Elements/PCGCopyPointsAnalysisKernel.cpp
--- a/Engine/Plugins/PCG/Source/PCG/Private/Elements/PCGCopyPointsAnalysisKernel.cpp
+++ b/Engine/Plugins/PCG/Source/PCG/Private/Elements/PCGCopyPointsAnalysisKernel.cpp
@@ -21,11 +21,6 @@
 
 #define LOCTEXT_NAMESPACE "PCGCopyPointsAnalysisKernel"
 
-const UPCGSettings* UPCGCopyPointsAnalysisKernel::GetSettings() const
-{
-	return Settings;
-}
-
 bool UPCGCopyPointsAnalysisKernel::IsKernelDataValid(FPCGContext* InContext) const
 {
 	TRACE_CPUPROFILER_EVENT_SCOPE(UPCGCopyPointsAnalysisKernel::IsKernelDataValid);
@@ -40,7 +35,7 @@ bool UPCGCopyPointsAnalysisKernel::IsKernelDataValid(FPCGContext* InContext) con
 		return true;
 	}
 
-	return PCGCopyPointsKernel::IsKernelDataValid(this, CastChecked<UPCGCopyPointsSettings>(Settings), static_cast<FPCGComputeGraphContext*>(InContext));
+	return PCGCopyPointsKernel::IsKernelDataValid(this, CastChecked<UPCGCopyPointsSettings>(GetSettings()), static_cast<FPCGComputeGraphContext*>(InContext));
 }
 
 bool UPCGCopyPointsAnalysisKernel::ComputeOutputBindingDataDesc(const UPCGComputeGraph* InGraph, FName InOutputPinLabel, UPCGDataBinding* InBinding, FPCGDataCollectionDesc& OutDataDesc) const
@@ -71,7 +66,7 @@ bool UPCGCopyPointsAnalysisKernel::ComputeOutputBindingDataDesc(const UPCGComput
 	FPCGDataDesc& DataDesc = OutDesc.DataDescs.Emplace_GetRef();
 	DataDesc.Type = EPCGDataType::Param;
 
-	const UPCGCopyPointsSettings* CopyPointSettings = CastChecked<UPCGCopyPointsSettings>(Settings);
+	const UPCGCopyPointsSettings* CopyPointSettings = CastChecked<UPCGCopyPointsSettings>(GetSettings());
 
 	if (CopyPointSettings->bCopyEachSourceOnEveryTarget)
 	{
diff --git a/Engine/Plugins/PCG/Source/PCG/Private/Elements/PCGCopyPointsKernel.cpp b/Engine/Plugins/PCG/Source/PCG/Private/Elements/PCGCopyPointsKernel.cpp
--- a/Engine/Plugins/PCG/Source/PCG/Private/Elements/PCGCopyPointsKernel.cpp
+++ b/Engine/Plugins/PCG/Source/PCG/Private/Elements/PCGCopyPointsKernel.cpp
@@ -22,11 +22,6 @@
 
 #define LOCTEXT_NAMESPACE "PCGCopyPointsKernel"
 
-const UPCGSettings* UPCGCopyPointsKernel::GetSettings() const
-{
-	return Settings;
-}
-
 bool UPCGCopyPointsKernel::IsKernelDataValid(FPCGContext* InContext) const
 {
 	TRACE_CPUPROFILER_EVENT_SCOPE(UPCGCopyPointsKernel::IsKernelDataValid);
@@ -41,7 +36,7 @@ bool UPCGCopyPointsKernel::IsKernelDataValid(FPCGContext* InContext) const
 		return true;
 	}
 
-	return PCGCopyPointsKernel::IsKernelDataValid(this, CastChecked<UPCGCopyPointsSettings>(Settings), static_cast<FPCGComputeGraphContext*>(InContext));
+	return PCGCopyPointsKernel::IsKernelDataValid(this, CastChecked<UPCGCopyPointsSettings>(GetSettings()), static_cast<FPCGComputeGraphContext*>(InContext));
 }
 
 bool UPCGCopyPointsKernel::ComputeOutputBindingDataDesc(const UPCGComputeGraph* InGraph, FName InOutputPinLabel, UPCGDataBinding* InBinding, FPCGDataCollectionDesc& OutDataDesc) const
@@ -49,7 +44,7 @@ bool UPCGCopyPointsKernel::ComputeOutputBindingDataDesc(const UPCGComputeGraph*
 	check(InGraph);
 	check(InBinding);
 
-	const UPCGCopyPointsSettings* CPSettings = CastChecked<UPCGCopyPointsSettings>(Settings);
+	const UPCGCopyPointsSettings* CPSettings = CastChecked<UPCGCopyPointsSettings>(GetSettings());
 
 	// Code assumes single output pin.
 	if (!ensure(InOutputPinLabel == PCGPinConstants::DefaultOutputLabel))
@@ -219,7 +214,7 @@ void UPCGCopyPointsKernel::CreateAdditionalInputDataInterfaces(FPCGGPUCompilatio
 
 void UPCGCopyPointsKernel::GetInputPins(TArray<FPCGPinProperties>& OutPins) const
 {
-	const UPCGCopyPointsSettings* CPSettings = CastChecked<UPCGCopyPointsSettings>(Settings);
+	const UPCGCopyPointsSettings* CPSettings = CastChecked<UPCGCopyPointsSettings>(GetSettings());
 
 	OutPins.Emplace(PCGCopyPointsConstants::SourcePointsLabel, EPCGDataType::Point);
 	OutPins.Emplace(PCGCopyPointsConstants::TargetPointsLabel, EPCGDataType::Point);
diff --git a/Engine/Plugins/PCG/Source/PCG/Private/Elements/PCGSkinnedMeshSpawnerKernel.cpp b/Engine/Plugins/PCG/Source/PCG/Private/Elements/PCGSkinnedMeshSpawnerKernel.cpp
--- a/Engine/Plugins/PCG/Source/PCG/Private/Elements/PCGSkinnedMeshSpawnerKernel.cpp
+++ b/Engine/Plugins/PCG/Source/PCG/Private/Elements/PCGSkinnedMeshSpawnerKernel.cpp
@@ -31,11 +31,6 @@ namespace PCGSkinnedMeshSpawnerKernel
 	static const FText NoMeshEntriesFormat = LOCTEXT("NoMeshEntries", "No mesh entries provided in weighted mesh selector.");
 }
 
-const UPCGSettings* UPCGSkinnedMeshSpawnerKernel::GetSettings() const
-{
-	return Settings;
-}
-
 bool UPCGSkinnedMeshSpawnerKernel::ComputeOutputBindingDataDesc(const UPCGComputeGraph* InGraph, FName InOutputPinLabel, UPCGDataBinding* InBinding, FPCGDataCollectionDesc& OutDataDesc) const
 {
 	check(InGraph);
@@ -111,7 +106,7 @@ bool UPCGSkinnedMeshSpawnerKernel::PerformStaticValidation()
 		return false;
 	}
 
-	const UPCGSkinnedMeshSpawnerSettings* SMSettings = CastChecked<UPCGSkinnedMeshSpawnerSettings>(Settings);
+	const UPCGSkinnedMeshSpawnerSettings* SMSettings = CastChecked<UPCGSkinnedMeshSpawnerSettings>(GetSettings());
 
 	// Currently instance packers must be able to specify a full list of attribute names upfront, to build the attribute table at compile time.
 	// TODO: We should be able to augment a static attribute table with new attributes at execution time, which will allow other types like regex.
diff --git a/Engine/Plugins/PCG/Source/PCG/Private/Elements/PCGStaticMeshSpawnerKernel.cpp b/Engine/Plugins/PCG/Source/PCG/Private/Elements/PCGStaticMeshSpawnerKernel.cpp
--- a/Engine/Plugins/PCG/Source/PCG/Private/Elements/PCGStaticMeshSpawnerKernel.cpp
+++ b/Engine/Plugins/PCG/Source/PCG/Private/Elements/PCGStaticMeshSpawnerKernel.cpp
@@ -32,11 +32,6 @@ namespace PCGStaticMeshSpawnerKernel
 	static const FText NoMeshEntriesFormat = LOCTEXT("NoMeshEntries", "No mesh entries provided in weighted mesh selector.");
 }
 
-const UPCGSettings* UPCGStaticMeshSpawnerKernel::GetSettings() const
-{
-	return Settings;
-}
-
 bool UPCGStaticMeshSpawnerKernel::ComputeOutputBindingDataDesc(const UPCGComputeGraph* InGraph, FName InOutputPinLabel, UPCGDataBinding* InBinding, FPCGDataCollectionDesc& OutDataDesc) const
 {
 	check(InGraph);
@@ -55,7 +50,7 @@ bool UPCGStaticMeshSpawnerKernel::ComputeOutputBindingDataDesc(const UPCGCompute
 
 	// Add output attribute (selected mesh).
 	{
-		const UPCGStaticMeshSpawnerSettings* SMSettings = CastChecked<UPCGStaticMeshSpawnerSettings>(Settings);
+		const UPCGStaticMeshSpawnerSettings* SMSettings = CastChecked<UPCGStaticMeshSpawnerSettings>(GetSettings());
 
 		// Create unique value keys for the output string.
 		if (UPCGMeshSelectorWeighted* SelectorWeighted = Cast<UPCGMeshSelectorWeighted>(SMSettings->MeshSelectorParameters))
@@ -145,7 +140,7 @@ void UPCGStaticMeshSpawnerKernel::CreateAdditionalOutputDataInterfaces(FPCGGPUCo
 
 void UPCGStaticMeshSpawnerKernel::AddStaticCreatedStrings(TArray<FString>& InOutStringTable) const
 {
-	const UPCGStaticMeshSpawnerSettings* SMSettings = CastChecked<UPCGStaticMeshSpawnerSettings>(Settings);
+	const UPCGStaticMeshSpawnerSettings* SMSettings = CastChecked<UPCGStaticMeshSpawnerSettings>(GetSettings());
 	if (UPCGMeshSelectorWeighted* SelectorWeighted = Cast<UPCGMeshSelectorWeighted>(SMSettings->MeshSelectorParameters))
 	{
 		for (const FPCGMeshSelectorWeightedEntry& Entry : SelectorWeighted->MeshEntries)
@@ -158,7 +153,7 @@ void UPCGStaticMeshSpawnerKernel::AddStaticCreatedStrings(TArray<FString>& InOut
 void UPCGStaticMeshSpawnerKernel::GetKernelAttributeKeys(TArray<FPCGKernelAttributeKey>& OutKeys) const
 {
 	// Register the attribute this node creates.
-	const UPCGStaticMeshSpawnerSettings* SMSettings = CastChecked<UPCGStaticMeshSpawnerSettings>(Settings);
+	const UPCGStaticMeshSpawnerSettings* SMSettings = CastChecked<UPCGStaticMeshSpawnerSettings>(GetSettings());
 	OutKeys.AddUnique(FPCGKernelAttributeKey(SMSettings->OutAttributeName, EPCGKernelAttributeType::StringKey));
 }
 
@@ -183,7 +178,7 @@ bool UPCGStaticMeshSpawnerKernel::PerformStaticValidation()
 		return false;
 	}
 
-	const UPCGStaticMeshSpawnerSettings* SMSettings = CastChecked<UPCGStaticMeshSpawnerSettings>(Settings);
+	const UPCGStaticMeshSpawnerSettings* SMSettings = CastChecked<UPCGStaticMeshSpawnerSettings>(GetSettings());
 
 	if (UPCGMeshSelectorWeighted* SelectorWeighted = Cast<UPCGMeshSelectorWeighted>(SMSettings->MeshSelectorParameters))
 	{
diff --git a/Engine/Plugins/PCG/Source/PCG/Private/Graph/PCGGraphCompilerGPU.cpp b/Engine/Plugins/PCG/Source/PCG/Private/Graph/PCGGraphCompilerGPU.cpp
--- a/Engine/Plugins/PCG/Source/PCG/Private/Graph/PCGGraphCompilerGPU.cpp
+++ b/Engine/Plugins/PCG/Source/PCG/Private/Graph/PCGGraphCompilerGPU.cpp
@@ -1338,7 +1338,7 @@ void FPCGGraphCompilerGPU::CompileComputeGraph(
 		// Kernel data interface.
 		{
 			UPCGCustomKernelDataInterface* KernelDI = InOutContext.NewObject_AnyThread<UPCGCustomKernelDataInterface>(InOutComputeGraph);
-			KernelDI->Settings = Settings;
+			KernelDI->SetSettings(Settings);
 			KernelDI->Kernel = KernelWithBindings.Kernel;
 
 			const int32 KernelDIIndex = InOutComputeGraph->DataInterfaces.Num();
diff --git a/Engine/Plugins/PCG/Source/PCG/Public/Elements/Metadata/PCGMetadataPartitionKernel.h b/Engine/Plugins/PCG/Source/PCG/Public/Elements/Metadata/PCGMetadataPartitionKernel.h
--- a/Engine/Plugins/PCG/Source/PCG/Public/Elements/Metadata/PCGMetadataPartitionKernel.h
+++ b/Engine/Plugins/PCG/Source/PCG/Public/Elements/Metadata/PCGMetadataPartitionKernel.h
@@ -14,7 +14,6 @@ class UPCGMetadataPartitionKernel : public UPCGComputeKernel
 	GENERATED_BODY()
 
 public:
-	virtual const UPCGSettings* GetSettings() const override;
 	virtual bool IsKernelDataValid(FPCGContext* InContext) const override;
 	virtual bool ComputeOutputBindingDataDesc(const UPCGComputeGraph* InGraph, FName InOutputPinLabel, UPCGDataBinding* InBinding, FPCGDataCollectionDesc& OutDataDesc) const override;
 	virtual int ComputeThreadCount(const UPCGDataBinding* Binding) const override;
diff --git a/Engine/Plugins/PCG/Source/PCG/Public/Elements/PCGCopyPointsAnalysisKernel.h b/Engine/Plugins/PCG/Source/PCG/Public/Elements/PCGCopyPointsAnalysisKernel.h
--- a/Engine/Plugins/PCG/Source/PCG/Public/Elements/PCGCopyPointsAnalysisKernel.h
+++ b/Engine/Plugins/PCG/Source/PCG/Public/Elements/PCGCopyPointsAnalysisKernel.h
@@ -14,7 +14,6 @@ class UPCGCopyPointsAnalysisKernel : public UPCGComputeKernel
 	GENERATED_BODY()
 
 public:
-	virtual const UPCGSettings* GetSettings() const override;
 	virtual bool IsKernelDataValid(FPCGContext* InContext) const override;
 	virtual bool ComputeOutputBindingDataDesc(const UPCGComputeGraph* InGraph, FName InOutputPinLabel, UPCGDataBinding* InBinding, FPCGDataCollectionDesc& OutDataDesc) const override;
 	virtual int ComputeThreadCount(const UPCGDataBinding* Binding) const override;
diff --git a/Engine/Plugins/PCG/Source/PCG/Public/Elements/PCGCopyPointsKernel.h b/Engine/Plugins/PCG/Source/PCG/Public/Elements/PCGCopyPointsKernel.h
--- a/Engine/Plugins/PCG/Source/PCG/Public/Elements/PCGCopyPointsKernel.h
+++ b/Engine/Plugins/PCG/Source/PCG/Public/Elements/PCGCopyPointsKernel.h
@@ -16,7 +16,6 @@ class UPCGCopyPointsKernel : public UPCGComputeKernel
 	GENERATED_BODY()
 
 public:
-	virtual const UPCGSettings* GetSettings() const override;
 	virtual bool IsKernelDataValid(FPCGContext* InContext) const override;
 	virtual bool ComputeOutputBindingDataDesc(const UPCGComputeGraph* InGraph, FName InOutputPinLabel, UPCGDataBinding* InBinding, FPCGDataCollectionDesc& OutDataDesc) const override;
 	virtual int ComputeThreadCount(const UPCGDataBinding* Binding) const override;
diff --git a/Engine/Plugins/PCG/Source/PCG/Public/Elements/PCGSkinnedMeshSpawnerKernel.h b/Engine/Plugins/PCG/Source/PCG/Public/Elements/PCGSkinnedMeshSpawnerKernel.h
--- a/Engine/Plugins/PCG/Source/PCG/Public/Elements/PCGSkinnedMeshSpawnerKernel.h
+++ b/Engine/Plugins/PCG/Source/PCG/Public/Elements/PCGSkinnedMeshSpawnerKernel.h
@@ -19,7 +19,6 @@ class UPCGSkinnedMeshSpawnerKernel : public UPCGComputeKernel
 	GENERATED_BODY()
 
 public:
-	virtual const UPCGSettings* GetSettings() const override;
 	virtual bool ComputeOutputBindingDataDesc(const UPCGComputeGraph* InGraph, FName InOutputPinLabel, UPCGDataBinding* InBinding, FPCGDataCollectionDesc& OutDataDesc) const override;
 	virtual int ComputeThreadCount(const UPCGDataBinding* Binding) const override;
 #if WITH_EDITOR
diff --git a/Engine/Plugins/PCG/Source/PCG/Public/Elements/PCGStaticMeshSpawnerKernel.h b/Engine/Plugins/PCG/Source/PCG/Public/Elements/PCGStaticMeshSpawnerKernel.h
--- a/Engine/Plugins/PCG/Source/PCG/Public/Elements/PCGStaticMeshSpawnerKernel.h
+++ b/Engine/Plugins/PCG/Source/PCG/Public/Elements/PCGStaticMeshSpawnerKernel.h
@@ -19,7 +19,6 @@ class UPCGStaticMeshSpawnerKernel : public UPCGComputeKernel
 	GENERATED_BODY()
 
 public:
-	virtual const UPCGSettings* GetSettings() const override;
 	virtual bool ComputeOutputBindingDataDesc(const UPCGComputeGraph* InGraph, FName InOutputPinLabel, UPCGDataBinding* InBinding, FPCGDataCollectionDesc& OutDataDesc) const override;
 	virtual int ComputeThreadCount(const UPCGDataBinding* Binding) const override;
 #if WITH_EDITOR
