[PCG] Soften check() when loading compute sources to ensure() after hits on one of the projects.

#rnx

diff --git a/Engine/Plugins/PCG/Source/PCG/Private/Compute/PCGComputeKernel.cpp b/Engine/Plugins/PCG/Source/PCG/Private/Compute/PCGComputeKernel.cpp
--- a/Engine/Plugins/PCG/Source/PCG/Private/Compute/PCGComputeKernel.cpp
+++ b/Engine/Plugins/PCG/Source/PCG/Private/Compute/PCGComputeKernel.cpp
@@ -95,10 +95,14 @@ void UPCGComputeKernel::GatherAdditionalSources(TArray<TObjectPtr<UComputeSource
 	auto GatherComputeSource = [&OutAdditionalSources](const FSoftObjectPath& InComputeSourcePath)
 	{
 		TSoftObjectPtr<UPCGComputeSource> ComputeSourcePtr(InComputeSourcePath);
-		UPCGComputeSource* ComputeSource = ComputeSourcePtr.LoadSynchronous();
-		check(ComputeSource);
-
-		OutAdditionalSources.Add(ComputeSource);
+		if (UPCGComputeSource* ComputeSource = ComputeSourcePtr.LoadSynchronous(); ensure(ComputeSource))
+		{
+			OutAdditionalSources.Add(ComputeSource);
+		}
+		else
+		{
+			UE_LOG(LogPCG, Error, TEXT("Compute source at '%s' could not be loaded."), *InComputeSourcePath.ToString());
+		}
 	};
 
 	GatherComputeSource(PCGComputeKernel::ShaderUtilsComputeSourcePath);
