[PCG] Make metadata op pins required if plugged

* This was not the case if the inline constant was activated before it begin plugged

#rnx

diff --git a/Engine/Plugins/PCG/Source/PCG/Private/Elements/Metadata/PCGMetadataOpElementBase.cpp b/Engine/Plugins/PCG/Source/PCG/Private/Elements/Metadata/PCGMetadataOpElementBase.cpp
--- a/Engine/Plugins/PCG/Source/PCG/Private/Elements/Metadata/PCGMetadataOpElementBase.cpp
+++ b/Engine/Plugins/PCG/Source/PCG/Private/Elements/Metadata/PCGMetadataOpElementBase.cpp
@@ -408,7 +408,8 @@ EPCGMetadataTypes UPCGMetadataSettingsBase::GetPinInitialDefaultValueType(const
 
 bool UPCGMetadataSettingsBase::IsInputPinRequiredByExecution(const UPCGPin* InPin) const
 {
-	return InPin && !IsPinDefaultValueActivated(InPin->Properties.Label);
+	// If the default value is not activated, or the pin is connected, it becomes required.
+	return InPin && (!IsPinDefaultValueActivated(InPin->Properties.Label) || InPin->IsConnected());
 }
 
 void UPCGMetadataSettingsBase::AddDefaultValuesToCrc(FArchiveCrc32& Crc32) const
