[PCG] Add virtual destructor to IPCGGraphExecutionState

#rnx

diff --git a/Engine/Plugins/PCG/Source/PCG/Public/PCGGraphExecutionStateInterface.h b/Engine/Plugins/PCG/Source/PCG/Public/PCGGraphExecutionStateInterface.h
--- a/Engine/Plugins/PCG/Source/PCG/Public/PCGGraphExecutionStateInterface.h
+++ b/Engine/Plugins/PCG/Source/PCG/Public/PCGGraphExecutionStateInterface.h
@@ -27,6 +27,7 @@ namespace PCGUtils
 class IPCGGraphExecutionState
 {
 public:
+	virtual ~IPCGGraphExecutionState() = default;
 
 	/** Returns a UPCGData representation of the ExecutionState. */
 	virtual UPCGData* GetSelfData() const = 0;
