[PCG] bTreatEditorViewportAsGenerationSource : Fix performance regression when deleting actors
- Go back to original implementation where we are using DestroyActor but add the usage of FDeselectedActorsEvent to handle actor deselection instead of calling UnrealEdEngine::DeleteActors which triggers a full GC

#rnx

diff --git a/Engine/Plugins/PCG/Source/PCG/Private/Helpers/PCGActorHelpers.cpp b/Engine/Plugins/PCG/Source/PCG/Private/Helpers/PCGActorHelpers.cpp
--- a/Engine/Plugins/PCG/Source/PCG/Private/Helpers/PCGActorHelpers.cpp
+++ b/Engine/Plugins/PCG/Source/PCG/Private/Helpers/PCGActorHelpers.cpp
@@ -526,37 +526,38 @@ bool UPCGActorHelpers::DeleteActors(UWorld* World, const TArray<TSoftObjectPtr<A
 	{
 		EditorModule->SetOutlinerUIRefreshDelay(UE::PCGActorHelpers::Local::OutlinerUIRefreshDelay);
 	}
-
-	TArray<AActor*> DeleteActors;
 #endif
 
+	TArray<AActor*> ActorsToDestroy;
+	ActorsToDestroy.Reserve(ActorsToDelete.Num());
+
+	// Gather actors to destroy
 	for (const TSoftObjectPtr<AActor>& ManagedActor : ActorsToDelete)
 	{
 		// @todo_pcg: Revisit this GetWorld() check when fixing UE-215065
 		if (AActor* Actor = ManagedActor.Get(); Actor && Actor->GetWorld())
 		{
-#if WITH_EDITOR
-			// In Editor, prefer UUnrealEdEngine::DeleteActors which handles selection
-			if (GUnrealEd && !Actor->GetWorld()->IsGameWorld())
-			{
-				DeleteActors.Add(Actor);
-			}
-			else
-#endif
-			if (!ensure(World->DestroyActor(Actor)))
-			{
-				UE_LOG(LogPCG, Warning, TEXT("Actor %s failed to be destroyed."), *Actor->GetPathName());
-			}
+			ActorsToDestroy.Add(Actor);
 		}
 	}
 
 #if WITH_EDITOR
-	if (GUnrealEd && !DeleteActors.IsEmpty())
+	if(!ActorsToDestroy.IsEmpty())
 	{
-		GUnrealEd->DeleteActors(DeleteActors, World, GUnrealEd->GetSelectedActors()->GetElementSelectionSet(), /*bVerifyDeletionCanHappen=*/false, /*bWarnAboutReferences=*/false, /*bWarnAboutSoftReferences=*/false);
+		// Deselect them if needed
+		FDeselectedActorsEvent DeselectedActorsEvent(ActorsToDestroy);
 	}
 #endif
 
+	// Destroy actors
+	for (AActor* ActorToDestroy : ActorsToDestroy)
+	{
+		if (!ensure(World->DestroyActor(ActorToDestroy)))
+		{
+			UE_LOG(LogPCG, Warning, TEXT("Actor %s failed to be destroyed."), *ActorToDestroy->GetPathName());
+		}
+	}
+
 	return true;
 }
 
