[PCG] Fixed deprecation bug in Make Rotator that created the default value as vectors when the operation was by angle and called for doubles

#rnx

diff --git a/Engine/Plugins/PCG/Source/PCG/Private/Elements/Metadata/PCGMetadataMakeRotator.cpp b/Engine/Plugins/PCG/Source/PCG/Private/Elements/Metadata/PCGMetadataMakeRotator.cpp
--- a/Engine/Plugins/PCG/Source/PCG/Private/Elements/Metadata/PCGMetadataMakeRotator.cpp
+++ b/Engine/Plugins/PCG/Source/PCG/Private/Elements/Metadata/PCGMetadataMakeRotator.cpp
@@ -5,6 +5,8 @@
 #include "PCGParamData.h"
 #include "Elements/Metadata/PCGMetadataElementCommon.h"
 #include "Metadata/PCGMetadata.h"
+
+#include "Algo/Transform.h"
 #include "UObject/FortniteMainBranchObjectVersion.h"
 
 #define LOCTEXT_NAMESPACE "PCGMetadataMakeRotatorSettings"
@@ -159,11 +161,28 @@ void UPCGMetadataMakeRotatorSettings::ApplyDeprecationBeforeUpdatePins(UPCGNode*
 		&& Operation == EPCGMetadataMakeRotatorOp::MakeRotFromAngles
 		&& GetLinkerCustomVersion(FFortniteMainBranchObjectVersion::GUID) < FFortniteMainBranchObjectVersion::PCGInlineConstantDefaultValues)
 	{
-		for (const UPCGPin* Pin : InputPins)
+		// Gather the labels instead of operating on the pins directly, because converting the pin types will trigger an update.
+		TArray<FName> PinLabels;
+		PinLabels.Reserve(InputPins.Num());
+
+		// Only transform the pins that aren't connected.
+		Algo::TransformIf(InputPins, PinLabels,
+			/*Predicate=*/[](const UPCGPin* Pin)
+			{
+				return !Pin->IsConnected();
+			},
+			/*Trans=*/[](const UPCGPin* Pin)
+			{
+				return Pin->Properties.Label;
+			});
+
+		// Activate the remaining unconnected pins and for MakeRotFromAngles, change the type from vector to double.
+		for (const FName Label : PinLabels)
 		{
-			if (IsPinDefaultValueEnabled(Pin->Properties.Label) && !Pin->IsConnected())
+			if (IsPinDefaultValueEnabled(Label))
 			{
-				SetPinDefaultValueIsActivated(Pin->Properties.Label, /*bIsActivated=*/true, /*bDirtySettings=*/false);
+				SetPinDefaultValueIsActivated(Label, /*bIsActivated=*/true, /*bDirtySettings=*/false);
+				ConvertPinDefaultValueMetadataType(Label, EPCGMetadataTypes::Double);
 			}
 		}
 	}
