[PCG] Fix crash when accessing ProducerSettings in UPCGExportableDataProvider::Initialize

#rnx

diff --git a/Engine/Plugins/PCG/Source/PCG/Private/Compute/DataInterfaces/PCGComputeDataInterface.cpp b/Engine/Plugins/PCG/Source/PCG/Private/Compute/DataInterfaces/PCGComputeDataInterface.cpp
--- a/Engine/Plugins/PCG/Source/PCG/Private/Compute/DataInterfaces/PCGComputeDataInterface.cpp
+++ b/Engine/Plugins/PCG/Source/PCG/Private/Compute/DataInterfaces/PCGComputeDataInterface.cpp
@@ -119,16 +119,18 @@ void UPCGComputeDataProvider::AddReferencedObjects(UObject* InThis, FReferenceCo
 #if WITH_EDITOR
 void UPCGComputeDataProvider::NotifyProducerUploadedData(UPCGDataBinding* InBinding)
 {
-	if (ProducerSettings && !ProducerSettings->ShouldExecuteOnGPU())
+	const UPCGSettings* LocalProducerSettings = GetProducerSettings();
+
+	if (LocalProducerSettings && !LocalProducerSettings->ShouldExecuteOnGPU())
 	{
-		const UPCGNode* ProducerNode = Cast<UPCGNode>(ProducerSettings->GetOuter());
+		const UPCGNode* ProducerNode = Cast<UPCGNode>(LocalProducerSettings->GetOuter());
 
 		// Works around current issue where input output settings are outer'd to the graph rather than their node.
 		if (!ProducerNode)
 		{
-			if (const UPCGGraph* Graph = Cast<UPCGGraph>(ProducerSettings->GetOuter()))
+			if (const UPCGGraph* Graph = Cast<UPCGGraph>(LocalProducerSettings->GetOuter()))
 			{
-				if (Graph->GetInputNode() && Graph->GetInputNode()->GetSettings() == ProducerSettings)
+				if (Graph->GetInputNode() && Graph->GetInputNode()->GetSettings() == LocalProducerSettings)
 				{
 					ProducerNode = Graph->GetInputNode();
 				}
@@ -154,20 +156,20 @@ void UPCGExportableDataProvider::Initialize(const UComputeDataInterface* InDataI
 	ExportMode = DataInterface->GetRequiresExport() ? EPCGExportMode::ComputeGraphOutput : EPCGExportMode::NoExport;
 
 #if WITH_EDITOR
-	check(ProducerSettings);
-
-	const bool bCanInspectOrDebug = !ProducerKernel || !ProducerKernel->IsPinInternal(OutputPinLabel) || PCGComputeDataInterfaceHelpers::CVarEnableInternalKernelPinInspection.GetValueOnAnyThread();
+	// Some exportable data providers don't support inspect/debug because their producer settings are not available (e.g. GridLinkage).
+	const UPCGSettings* LocalProducerSettings = GetProducerSettings();
+	const bool bCanInspectOrDebug = LocalProducerSettings && (!ProducerKernel || !ProducerKernel->IsPinInternal(OutputPinLabel) || PCGComputeDataInterfaceHelpers::CVarEnableInternalKernelPinInspection.GetValueOnAnyThread());
 
 	if (bCanInspectOrDebug)
 	{
 		const IPCGGraphExecutionSource* ExecutionSource = Binding->GetExecutionSource();
 
-		if (ProducerSettings->bIsInspecting && ExecutionSource && ExecutionSource->GetExecutionState().GetInspection().IsInspecting())
+		if (LocalProducerSettings->bIsInspecting && ExecutionSource && ExecutionSource->GetExecutionState().GetInspection().IsInspecting())
 		{
 			ExportMode |= EPCGExportMode::Inspection;
 		}
 
-		if (ProducerSettings->bDebug)
+		if (LocalProducerSettings->bDebug)
 		{
 			ExportMode |= EPCGExportMode::DebugVisualization;
 		}
diff --git a/Engine/Plugins/PCG/Source/PCG/Private/Compute/DataInterfaces/PCGComputeDataInterface.h b/Engine/Plugins/PCG/Source/PCG/Private/Compute/DataInterfaces/PCGComputeDataInterface.h
--- a/Engine/Plugins/PCG/Source/PCG/Private/Compute/DataInterfaces/PCGComputeDataInterface.h
+++ b/Engine/Plugins/PCG/Source/PCG/Private/Compute/DataInterfaces/PCGComputeDataInterface.h
@@ -97,6 +97,7 @@ protected:
 	UPROPERTY()
 	TObjectPtr<const UPCGComputeKernel> ProducerKernel;
 
+private:
 	/** Settings of node that produces this data, normally the upstream node. */
 	UPROPERTY()
 	TSoftObjectPtr<const UPCGSettings> ProducerSettings;
diff --git a/Engine/Plugins/PCG/Source/PCG/Private/Compute/DataInterfaces/PCGDataCollectionUploadDataInterface.cpp b/Engine/Plugins/PCG/Source/PCG/Private/Compute/DataInterfaces/PCGDataCollectionUploadDataInterface.cpp
--- a/Engine/Plugins/PCG/Source/PCG/Private/Compute/DataInterfaces/PCGDataCollectionUploadDataInterface.cpp
+++ b/Engine/Plugins/PCG/Source/PCG/Private/Compute/DataInterfaces/PCGDataCollectionUploadDataInterface.cpp
@@ -29,10 +29,11 @@ void UPCGDataProviderDataCollectionUpload::Initialize(const UComputeDataInterfac
 {
 	TRACE_CPUPROFILER_EVENT_SCOPE(UPCGDataProviderDataCollectionUpload::Initialize);
 
+	Super::Initialize(InDataInterface, InBinding, InInputMask, InOutputMask);
+
 	const UPCGDataCollectionUploadDataInterface* DataInterface = CastChecked<UPCGDataCollectionUploadDataInterface>(InDataInterface);
 
 	Binding = CastChecked<UPCGDataBinding>(InBinding);
-	ProducerSettings = DataInterface->GetProducerSettings();
 	GraphBindingIndex = DataInterface->GraphBindingIndex;
 
 	// Pick the data items from input data collection using any of the compute graph element virtual input pin labels.
diff --git a/Engine/Plugins/PCG/Source/PCG/Private/Compute/DataInterfaces/PCGLandscapeDataInterface.cpp b/Engine/Plugins/PCG/Source/PCG/Private/Compute/DataInterfaces/PCGLandscapeDataInterface.cpp
--- a/Engine/Plugins/PCG/Source/PCG/Private/Compute/DataInterfaces/PCGLandscapeDataInterface.cpp
+++ b/Engine/Plugins/PCG/Source/PCG/Private/Compute/DataInterfaces/PCGLandscapeDataInterface.cpp
@@ -277,7 +277,6 @@ void UPCGLandscapeDataProvider::Initialize(const UComputeDataInterface* InDataIn
 	Super::Initialize(InDataInterface, InBinding, InInputMask, InOutputMask);
 
 	const UPCGLandscapeDataInterface* DataInterface = CastChecked<UPCGLandscapeDataInterface>(InDataInterface);
-	ProducerSettings = DataInterface->GetProducerSettings();
 
 	UPCGDataBinding* Binding = CastChecked<UPCGDataBinding>(InBinding);
 
