[PCG] Fix ensure when applying primitive scene updates before the scene proxies are ready

#rnx

diff --git a/Engine/Plugins/PCG/Source/PCG/Private/Compute/DataInterfaces/PCGInstanceDataInterface.cpp b/Engine/Plugins/PCG/Source/PCG/Private/Compute/DataInterfaces/PCGInstanceDataInterface.cpp
--- a/Engine/Plugins/PCG/Source/PCG/Private/Compute/DataInterfaces/PCGInstanceDataInterface.cpp
+++ b/Engine/Plugins/PCG/Source/PCG/Private/Compute/DataInterfaces/PCGInstanceDataInterface.cpp
@@ -33,6 +33,7 @@
 #include "Components/InstancedSkinnedMeshComponent.h"
 #include "ComputeFramework/ComputeKernelPermutationSet.h"
 #include "ComputeFramework/ShaderParamTypeDefinition.h"
+#include "Containers/Ticker.h"
 #include "Engine/Texture.h"
 #include "Engine/World.h"
 
@@ -301,13 +302,29 @@ void FPCGInstanceDataProviderProxy::PostSubmit(FRDGBuilder& InGraphBuilder) cons
 	const TRefCountPtr<FRDGPooledBuffer> InstanceCustomFloatDataExported = InGraphBuilder.ConvertToExternalBuffer(InstanceCustomFloatData);
 	const TRefCountPtr<FRDGPooledBuffer> WriteCountersExported = InGraphBuilder.ConvertToExternalBuffer(WriteCounters);
 	
-	ExecuteOnGameThread(TEXT("ApplyPrimitiveSceneUpdates"), [DataProviderWeak=DataProvider, DataProviderGeneration=DataProviderGeneration, InstanceDataExported, InstanceCustomFloatDataExported, WriteCountersExported]()
+	FTSTicker::GetCoreTicker().AddTicker(TEXT("ApplyPrimitiveSceneUpdates"), 0.0f, [DataProviderWeak=DataProvider, DataProviderGeneration=DataProviderGeneration, InstanceDataExported, InstanceCustomFloatDataExported, WriteCountersExported](float)
 	{
 		TStrongObjectPtr<UPCGInstanceDataProvider> DataProviderInner = DataProviderWeak.Pin();
 		if (!DataProviderInner || DataProviderInner->GenerationCounter != DataProviderGeneration)
 		{
 			UE_LOG(LogPCG, Verbose, TEXT("Data provider object lost, GPU instancing will fail."));
-			return;
+			return false;
+		}
+
+		bool bAllProxiesAreReady = true;
+
+		for (const FPCGPrimitiveComponentAndNumInstances& PrimitiveAndInstanceCount : DataProviderInner->Primitives)
+		{
+			if (UPrimitiveComponent* PrimitiveComponent = PrimitiveAndInstanceCount.Get<0>().Get())
+			{
+				bAllProxiesAreReady &= PrimitiveComponent->GetSceneProxy() != nullptr;
+			}
+		}
+
+		if (!bAllProxiesAreReady)
+		{
+			UE_LOG(LogPCG, Verbose, TEXT("FPCGInstanceDataProvider: One or more scene proxies were not ready. Will try again on the next tick."));
+			return true;
 		}
 
 		// All instance data is stored in a single buffer, so this is used to give the scene writer an index to the first primitive.
@@ -428,5 +445,7 @@ void FPCGInstanceDataProviderProxy::PostSubmit(FRDGBuilder& InGraphBuilder) cons
 		}
 
 		DataProviderInner->bWroteInstances = true;
+
+		return false;
 	});
 }
