[PCG] Fixed behavior issues in the create graph from template dialog

#rnx

diff --git a/Engine/Plugins/PCG/Source/PCGEditor/Private/PCGEditorUtils.cpp b/Engine/Plugins/PCG/Source/PCGEditor/Private/PCGEditorUtils.cpp
--- a/Engine/Plugins/PCG/Source/PCGEditor/Private/PCGEditorUtils.cpp
+++ b/Engine/Plugins/PCG/Source/PCGEditor/Private/PCGEditorUtils.cpp
@@ -197,6 +197,56 @@ protected:
 	FText DisplayName;
 };
 
+class SPCGSimpleOkCancelWindow : public SWindow
+{
+public:
+	SLATE_BEGIN_ARGS(SPCGSimpleOkCancelWindow)
+		{
+		}
+		SLATE_ARGUMENT(TOptional<FText>, Title)
+	SLATE_END_ARGS()
+
+	void Construct(const FArguments& InArgs)
+	{
+		SWindow::Construct(SWindow::FArguments()
+			.Title(InArgs._Title.Get(LOCTEXT("PCGSimpleOkCancelWindowTitle", "PCG Confirmation Window")))
+			.SizingRule(ESizingRule::UserSized)
+			.ClientSize(FVector2D(400.0f, 500.0f))
+			.SupportsMaximize(false)
+			.SupportsMinimize(false));
+	}
+
+	virtual FReply OnKeyDown(const FGeometry& MyGeometry, const FKeyEvent& InKeyEvent) override
+	{
+		FReply Reply = SWindow::OnKeyDown(MyGeometry, InKeyEvent);
+		if (InKeyEvent.GetKey() == EKeys::Escape)
+		{
+			OnCancel();
+		}
+
+		return Reply;
+	}
+
+	void OnOk()
+	{
+		bProceedWithAction = true;
+		RequestDestroyWindow();
+	}
+	void OnCancel()
+	{
+		bProceedWithAction = false;
+		RequestDestroyWindow();
+	}
+
+	bool ShouldProceedWithAction() const
+	{
+		return bProceedWithAction;
+	}
+
+private:
+	bool bProceedWithAction = false;
+};
+
 bool PCGEditorUtils::PickGraphTemplate(FAssetData& OutAssetData, const FText& TitleOverride)
 {
 	auto HandleAssetSelected = [&OutAssetData](const FAssetData& InSelectedAsset)
@@ -205,17 +255,11 @@ bool PCGEditorUtils::PickGraphTemplate(FAssetData& OutAssetData, const FText& Ti
 	};
 
 	FAssetPickerConfig AssetPickerConfig;
-	AssetPickerConfig.OnAssetEnterPressed = FOnAssetEnterPressed::CreateLambda([HandleAssetSelected](const TArray<FAssetData>& SelectedAssetData)
-	{
-		if (SelectedAssetData.Num() == 1)
-		{
-			HandleAssetSelected(SelectedAssetData[0]);
-		}
-	});
 	AssetPickerConfig.OnAssetSelected = FOnAssetSelected::CreateLambda(HandleAssetSelected);
 	AssetPickerConfig.SelectionMode = ESelectionMode::Single;
 	AssetPickerConfig.bAllowNullSelection = false;
 	AssetPickerConfig.InitialAssetViewType = EAssetViewType::List;
+	AssetPickerConfig.InitialThumbnailSize = EThumbnailSize::Small;
 	AssetPickerConfig.Filter.ClassPaths.Add(UPCGGraph::StaticClass()->GetClassPathName());
 	AssetPickerConfig.OnShouldFilterAsset = FOnShouldFilterAsset::CreateLambda([](const FAssetData& AssetData)
 	{
@@ -261,16 +305,28 @@ bool PCGEditorUtils::PickGraphTemplate(FAssetData& OutAssetData, const FText& Ti
 	FContentBrowserModule& ContentBrowserModule = FModuleManager::Get().LoadModuleChecked<FContentBrowserModule>(TEXT("ContentBrowser"));
 	FText Title = (TitleOverride.IsEmpty() ? LOCTEXT("SelectTemplateDialogTitle", "Initialize from Graph Template...") : TitleOverride);
 
-	bool bClickedOk = false;
-	TSharedPtr<SWindow> Dialog;
-	Dialog = SNew(SWindow)
-		.Title(Title)
-		.SizingRule(ESizingRule::Autosized);
+	TSharedPtr<SPCGSimpleOkCancelWindow> Dialog = SNew(SPCGSimpleOkCancelWindow)
+		.Title(Title);
+
+	// Finally, bind double-click & enter to the asset picker
+	auto HandleAssetActivated = [HandleAssetSelected, WeakDialog = Dialog.ToWeakPtr()](const TArray<FAssetData>& ActivatedAssets, EAssetTypeActivationMethod::Type ActivationMethod)
+	{
+		if ((ActivationMethod == EAssetTypeActivationMethod::DoubleClicked || ActivationMethod == EAssetTypeActivationMethod::Opened) && ActivatedAssets.Num() == 1 && ActivatedAssets[0].IsValid())
+		{
+			HandleAssetSelected(ActivatedAssets[0]);
+
+			if (TSharedPtr<SPCGSimpleOkCancelWindow> PinnedDialog = WeakDialog.Pin())
+			{
+				PinnedDialog->OnOk();
+			}
+		}
+	};
+
+	AssetPickerConfig.OnAssetsActivated = FOnAssetsActivated::CreateLambda(HandleAssetActivated);
 
 	Dialog->SetContent(
 		SNew(SVerticalBox)
 		+ SVerticalBox::Slot()
-		.AutoHeight()
 		.Padding(5.0f)
 		[
 			SNew(SBox)
@@ -290,12 +346,11 @@ bool PCGEditorUtils::PickGraphTemplate(FAssetData& OutAssetData, const FText& Ti
 			.HAlign(HAlign_Left)
 			[
 				SNew(SButton)
-				.OnClicked_Lambda([&bClickedOk, &OutAssetData, WeakDialog = TWeakPtr<SWindow>(Dialog)]()
+				.OnClicked_Lambda([&OutAssetData, WeakDialog = Dialog.ToWeakPtr()]()
 				{
-					bClickedOk = true;
-					if (TSharedPtr<SWindow> PinnedDialog = WeakDialog.Pin())
+					if (TSharedPtr<SPCGSimpleOkCancelWindow> PinnedDialog = WeakDialog.Pin())
 					{
-						PinnedDialog->RequestDestroyWindow();
+						PinnedDialog->OnOk();
 					}
 
 					OutAssetData = FAssetData();
@@ -311,12 +366,11 @@ bool PCGEditorUtils::PickGraphTemplate(FAssetData& OutAssetData, const FText& Ti
 			.AutoWidth()
 			[
 				SNew(SPrimaryButton)
-				.OnClicked_Lambda([&bClickedOk, WeakDialog = TWeakPtr<SWindow>(Dialog)]()
+				.OnClicked_Lambda([WeakDialog = Dialog.ToWeakPtr()]()
 				{
-					bClickedOk = true;
-					if (TSharedPtr<SWindow> PinnedDialog = WeakDialog.Pin())
+					if (TSharedPtr<SPCGSimpleOkCancelWindow> PinnedDialog = WeakDialog.Pin())
 					{
-						PinnedDialog->RequestDestroyWindow();
+						PinnedDialog->OnOk();
 					}
 
 					return FReply::Handled();
@@ -331,11 +385,11 @@ bool PCGEditorUtils::PickGraphTemplate(FAssetData& OutAssetData, const FText& Ti
 			.AutoWidth()
 			[
 				SNew(SButton)
-				.OnClicked_Lambda([WeakDialog = TWeakPtr<SWindow>(Dialog)]()
+				.OnClicked_Lambda([WeakDialog = Dialog.ToWeakPtr()]()
 				{
-					if (TSharedPtr<SWindow> PinnedDialog = WeakDialog.Pin())
+					if (TSharedPtr<SPCGSimpleOkCancelWindow> PinnedDialog = WeakDialog.Pin())
 					{
-						PinnedDialog->RequestDestroyWindow();
+						PinnedDialog->OnCancel();
 					}
 
 					return FReply::Handled();
@@ -345,7 +399,7 @@ bool PCGEditorUtils::PickGraphTemplate(FAssetData& OutAssetData, const FText& Ti
 		]);
 
 	FSlateApplication::Get().AddModalWindow(Dialog.ToSharedRef(), FGlobalTabmanager::Get()->GetRootWindow());
-	return bClickedOk;
+	return Dialog->ShouldProceedWithAction();
 }
 
 void PCGEditorUtils::OpenAssetOrMoveToActorOrComponent(const FSoftObjectPath& InPath)
