[PCG] Removed an ensure firing when the EditorGraph attempts to CreateLinks

* There are redundant reconstructs happening as part of the event handling, during (re)generation.
* The ensure coverage was useful internally, but not for the user, who will not notice the redundancy, but would likely notice the ensure.

#rnx

diff --git a/Engine/Plugins/PCG/Source/PCGEditor/Private/PCGEditorGraph.cpp b/Engine/Plugins/PCG/Source/PCGEditor/Private/PCGEditorGraph.cpp
--- a/Engine/Plugins/PCG/Source/PCGEditor/Private/PCGEditorGraph.cpp
+++ b/Engine/Plugins/PCG/Source/PCGEditor/Private/PCGEditorGraph.cpp
@@ -313,16 +313,35 @@ void UPCGEditorGraph::CreateLinks(UPCGEditorGraphNodeBase* GraphNode, bool bCrea
 					continue;
 				}
 
-				UPCGEditorGraphNodeBase* const* ConnectedGraphNode = InboundNode ? InPCGNodeToPCGEditorNodeMap.Find(InboundNode) : nullptr;
-				UEdGraphPin* OutPin = ConnectedGraphNode ? (*ConnectedGraphNode)->FindPin(InboundEdge->InputPin->Properties.Label, EEdGraphPinDirection::EGPD_Output) : nullptr;
-				if (OutPin)
+				if (UPCGEditorGraphNodeBase* const* ConnectedGraphNode = InboundNode ? InPCGNodeToPCGEditorNodeMap.Find(InboundNode) : nullptr)
 				{
-					OutPin->MakeLinkTo(InPin);
+					if (UEdGraphPin* OutPin = (*ConnectedGraphNode)->FindPin(InboundEdge->InputPin->Properties.Label, EGPD_Output))
+					{
+						OutPin->MakeLinkTo(InPin);
+					}
+					else
+					{
+						UE_LOG(LogPCGEditor, Error, TEXT("Could not find inbound pin '%s' to link to input pin '%s' from node '%s'"),
+							*InboundEdge->InputPin->Properties.Label.ToString(),
+							*InputPin->Properties.Label.ToString(),
+							*InboundNode->GetFName().ToString());
+						ensure(false);
+					}
 				}
 				else
 				{
-					UE_LOG(LogPCGEditor, Error, TEXT("Could not create link to InputPin %s from Node %s"), *InputPin->Properties.Label.ToString(), *InboundNode->GetFName().ToString());
-					ensure(false);
+					/**
+					 * Note: The ConnectedGraphNode may be in a state where it is not yet in the Nodes list. Ex. The
+					 * graph has been reconstructed after a cancelled generation, but before the reconstruction that
+					 * creates the UEdGraphNode.
+					 * 
+					 * All cases seem to be user driven actions that happen during generation, where a reconstruct is
+					 * invoked as part of a cancelled generation during the user action, but before the user action's
+					 * reconstruction--which then correctly links the edges.
+					 * 
+					 * @todo_pcg: As part of an overall effort to clean up the graph construction process, investigate
+					 * this with better Editor tracing.
+					 */
 				}
 			}
 		}
@@ -361,16 +380,24 @@ void UPCGEditorGraph::CreateLinks(UPCGEditorGraphNodeBase* GraphNode, bool bCrea
 					continue;
 				}
 
-				UPCGEditorGraphNodeBase* const* ConnectedGraphNode = OutboundNode ? InPCGNodeToPCGEditorNodeMap.Find(OutboundNode) : nullptr;
-				UEdGraphPin* InPin = ConnectedGraphNode ? (*ConnectedGraphNode)->FindPin(OutboundEdge->OutputPin->Properties.Label, EEdGraphPinDirection::EGPD_Input) : nullptr;
-				if (InPin)
+				if (UPCGEditorGraphNodeBase* const* ConnectedGraphNode = OutboundNode ? InPCGNodeToPCGEditorNodeMap.Find(OutboundNode) : nullptr)
 				{
-					OutPin->MakeLinkTo(InPin);
+					if (UEdGraphPin* InPin = ConnectedGraphNode ? (*ConnectedGraphNode)->FindPin(OutboundEdge->OutputPin->Properties.Label, EGPD_Input) : nullptr)
+					{
+						OutPin->MakeLinkTo(InPin);
+					}
+					else
+					{
+						UE_LOG(LogPCGEditor, Error, TEXT("Could not find outbound pin '%s' to link to input pin '%s' on node '%s'"),
+							*OutboundEdge->OutputPin->Properties.Label.ToString(),
+							*OutputPin->Properties.Label.ToString(),
+							*OutboundNode->GetFName().ToString());
+						ensure(false);
+					}
 				}
 				else
 				{
-					UE_LOG(LogPCGEditor, Error, TEXT("Could not create link from OutputPin %s to Node %s"), *OutputPin->Properties.Label.ToString(), *OutboundNode->GetFName().ToString());
-					ensure(false);
+					// See above comment in Inbound version
 				}
 			}
 		}
