[PCG] Fix crash when binding index lookup fails while computing kernel pin data descriptions

#rnx

diff --git a/Engine/Plugins/PCG/Source/PCG/Private/Compute/PCGDataBinding.cpp b/Engine/Plugins/PCG/Source/PCG/Private/Compute/PCGDataBinding.cpp
--- a/Engine/Plugins/PCG/Source/PCG/Private/Compute/PCGDataBinding.cpp
+++ b/Engine/Plugins/PCG/Source/PCG/Private/Compute/PCGDataBinding.cpp
@@ -111,6 +111,12 @@ bool UPCGDataBinding::ComputeKernelPinDataDesc(const FPCGKernelPin& InKernelPin,
 
 	const int32 GraphBindingIndex = Graph->GetBindingIndex(InKernelPin);
 
+	if (GraphBindingIndex == -1)
+	{
+		UE_LOG(LogPCG, Error, TEXT("Failed to compute data description for kernel pin '%s'."), *InKernelPin.PinLabel.ToString());
+		return false;
+	}
+
 	if (FPCGDataCollectionDesc* FoundDesc = DataDescriptionCache.Find(GraphBindingIndex))
 	{
 		if (OutDesc)
