[PCG] Fixed duplicate advanced categories in settings

#rnx

diff --git a/Engine/Plugins/PCG/Source/PCG/Public/Elements/Grammar/PCGSubdivideSpline.h b/Engine/Plugins/PCG/Source/PCG/Public/Elements/Grammar/PCGSubdivideSpline.h
--- a/Engine/Plugins/PCG/Source/PCG/Public/Elements/Grammar/PCGSubdivideSpline.h
+++ b/Engine/Plugins/PCG/Source/PCG/Public/Elements/Grammar/PCGSubdivideSpline.h
@@ -68,7 +68,7 @@ public:
 	FName IsFinalAttributeName = TEXT("IsFinal");
 
 	/** Maximum acceptable distance tolerance between placed modules (overlap or gap) - a smaller number is more precise. Affects performance if the number is very small. */
-	UPROPERTY(BlueprintReadWrite, EditAnywhere, Category = "Settings|Advanced", meta = (PCG_Overridable, ClampMin="0.000001"))
+	UPROPERTY(BlueprintReadWrite, EditAnywhere, Category = Settings, AdvancedDisplay, meta = (PCG_Overridable, ClampMin="0.000001"))
 	double ModulePlacementTolerance = 0.01; 
 };
 
diff --git a/Engine/Plugins/PCG/Source/PCG/Public/Elements/PCGApplyOnActor.h b/Engine/Plugins/PCG/Source/PCG/Public/Elements/PCGApplyOnActor.h
--- a/Engine/Plugins/PCG/Source/PCG/Public/Elements/PCGApplyOnActor.h
+++ b/Engine/Plugins/PCG/Source/PCG/Public/Elements/PCGApplyOnActor.h
@@ -56,7 +56,7 @@ public:
 	TArray<FName> PostProcessFunctionNames;
 
 	/** Opt-in option to silence errors when the path is Empty or nothing to extract. */
-	UPROPERTY(BlueprintReadWrite, EditAnywhere, Category = "Settings|Advanced")
+	UPROPERTY(BlueprintReadWrite, EditAnywhere, Category = Settings, AdvancedDisplay)
 	bool bSilenceErrorOnEmptyObjectPath = false;
 
 	/** By default, object loading is asynchronous, can force it synchronous if needed. */
diff --git a/Engine/Plugins/PCG/Source/PCG/Public/Elements/PCGClusterElement.h b/Engine/Plugins/PCG/Source/PCG/Public/Elements/PCGClusterElement.h
--- a/Engine/Plugins/PCG/Source/PCG/Public/Elements/PCGClusterElement.h
+++ b/Engine/Plugins/PCG/Source/PCG/Public/Elements/PCGClusterElement.h
@@ -57,25 +57,25 @@ public:
 	FPCGAttributePropertyOutputNoSourceSelector ClusterAttribute;
 
 	/** A limit on the number of iterations to run on each algorithm, if it doesn't otherwise converge. A higher number can offer more accuracy at the cost of processing time. */
-	UPROPERTY(EditAnywhere, BlueprintReadWrite, Category = "Settings|Advanced", meta = (ClampMin = "1", PCG_Overridable))
+	UPROPERTY(EditAnywhere, BlueprintReadWrite, Category = Settings, AdvancedDisplay, meta = (ClampMin = "1", PCG_Overridable))
 	int32 MaxIterations = 100;
 
 	/** For EM, the maximum allowed difference between the last two iterations' "Log Likelihood"--which converges from positive infinity to 0 in relation to point-to-cluster probabilities.
 	 * It is exponential, so raising this number can offer faster iteration if exact precision isn't needed.
 	 */
-	UPROPERTY(EditAnywhere, BlueprintReadWrite, Category = "Settings|Advanced", meta = (ClampMin = "0.0001", EditCondition = "Algorithm == EPCGClusterAlgorithm::EM", EditConditionHides, PCG_Overridable))
+	UPROPERTY(EditAnywhere, BlueprintReadWrite, Category = Settings, AdvancedDisplay, meta = (ClampMin = "0.0001", EditCondition = "Algorithm == EPCGClusterAlgorithm::EM", EditConditionHides, PCG_Overridable))
 	double Tolerance = UE_DOUBLE_KINDA_SMALL_NUMBER;
 
 	/** Output the final location of the centroids or gaussians. */
-	UPROPERTY(EditAnywhere, BlueprintReadWrite, Category = "Settings|Advanced")
+	UPROPERTY(EditAnywhere, BlueprintReadWrite, Category = Settings, AdvancedDisplay)
 	bool bOutputFinalCentroids = false;
 
 	/** Output the element count for each of the final centroids. */
-	UPROPERTY(EditAnywhere, BlueprintReadWrite, Category = "Settings|Advanced", meta = (EditCondition = "bOutputFinalCentroids", EditConditionHides, PCG_Overridable))
+	UPROPERTY(EditAnywhere, BlueprintReadWrite, Category = Settings, AdvancedDisplay, meta = (EditCondition = "bOutputFinalCentroids", EditConditionHides, PCG_Overridable))
 	bool bOutputFinalCentroidElementCount = false;
 
 	/** Final centroid element count will be written to this attribute on the Final Centroid output. */
-	UPROPERTY(EditAnywhere, BlueprintReadWrite, Category = "Settings|Advanced", meta = (EditCondition = "bOutputFinalCentroidElementCount", EditConditionHides, PCG_Overridable))
+	UPROPERTY(EditAnywhere, BlueprintReadWrite, Category = Settings, AdvancedDisplay, meta = (EditCondition = "bOutputFinalCentroidElementCount", EditConditionHides, PCG_Overridable))
 	FPCGAttributePropertyOutputNoSourceSelector FinalCentroidElementCountAttribute;
 };
 
diff --git a/Engine/Plugins/PCG/Source/PCG/Public/Elements/PCGCreateCollisionData.h b/Engine/Plugins/PCG/Source/PCG/Public/Elements/PCGCreateCollisionData.h
--- a/Engine/Plugins/PCG/Source/PCG/Public/Elements/PCGCreateCollisionData.h
+++ b/Engine/Plugins/PCG/Source/PCG/Public/Elements/PCGCreateCollisionData.h
@@ -41,7 +41,7 @@ public:
 	UPROPERTY(BlueprintReadWrite, EditAnywhere, Category = Settings, meta = (PCG_Overridable, PCG_OverrideAliases="bUseComplexCollision"))
 	EPCGCollisionQueryFlag CollisionQueryFlag = EPCGCollisionQueryFlag::Simple;
 
-	UPROPERTY(BlueprintReadWrite, EditAnywhere, Category = "Settings|Advanced")
+	UPROPERTY(BlueprintReadWrite, EditAnywhere, Category = Settings, AdvancedDisplay)
 	bool bWarnIfAttributeCouldNotBeUsed = true;
 
 	UPROPERTY(BlueprintReadWrite, EditAnywhere, Category = "Settings|Debug")
diff --git a/Engine/Plugins/PCG/Source/PCG/Public/Elements/PCGCreateTargetActor.h b/Engine/Plugins/PCG/Source/PCG/Public/Elements/PCGCreateTargetActor.h
--- a/Engine/Plugins/PCG/Source/PCG/Public/Elements/PCGCreateTargetActor.h
+++ b/Engine/Plugins/PCG/Source/PCG/Public/Elements/PCGCreateTargetActor.h
@@ -89,7 +89,7 @@ public:
 	UPROPERTY(BlueprintReadWrite, EditAnywhere, Category = Streaming, DisplayName = "HLOD Settings", meta = (PCG_Overridable))
 	FPCGHLODSettings HLODSettings;
 
-	UPROPERTY(BlueprintReadWrite, EditAnywhere, Category = "Settings|Advanced")
+	UPROPERTY(BlueprintReadWrite, EditAnywhere, Category = Settings, AdvancedDisplay)
 	bool bDeleteActorsBeforeGeneration = false;
 
 	PCG_API void SetTemplateActorClass(const TSubclassOf<AActor>& InTemplateActorClass);
diff --git a/Engine/Plugins/PCG/Source/PCG/Public/Elements/PCGDataFromActor.h b/Engine/Plugins/PCG/Source/PCG/Public/Elements/PCGDataFromActor.h
--- a/Engine/Plugins/PCG/Source/PCG/Public/Elements/PCGDataFromActor.h
+++ b/Engine/Plugins/PCG/Source/PCG/Public/Elements/PCGDataFromActor.h
@@ -128,11 +128,11 @@ public:
 	bool bAlwaysRequeryActors = false;
 
 	/** Silence warnings that attribute names were sanitized to replace invalid characters. */
-	UPROPERTY(BlueprintReadWrite, EditAnywhere, Category = "Data Retrieval Settings|Advanced")
+	UPROPERTY(BlueprintReadWrite, EditAnywhere, Category = "Data Retrieval Settings", AdvancedDisplay)
 	bool bSilenceSanitizedAttributeNameWarnings = false;
 
 	/** Silence warnings that attribute names were rejected because they clash with reserved names. */
-	UPROPERTY(BlueprintReadWrite, EditAnywhere, Category = "Data Retrieval Settings|Advanced")
+	UPROPERTY(BlueprintReadWrite, EditAnywhere, Category = "Data Retrieval Settings", AdvancedDisplay)
 	bool bSilenceReservedAttributeNameWarnings = false;
 
 #if WITH_EDITORONLY_DATA
diff --git a/Engine/Plugins/PCG/Source/PCG/Public/Elements/PCGGetPropertyFromObjectPath.h b/Engine/Plugins/PCG/Source/PCG/Public/Elements/PCGGetPropertyFromObjectPath.h
--- a/Engine/Plugins/PCG/Source/PCG/Public/Elements/PCGGetPropertyFromObjectPath.h
+++ b/Engine/Plugins/PCG/Source/PCG/Public/Elements/PCGGetPropertyFromObjectPath.h
@@ -74,11 +74,11 @@ public:
 	bool bSynchronousLoad = false;
 
 	/** Opt-in option to create empty data when there is nothing to extract or property is not found, to have the same number of inputs than outputs. */
-	UPROPERTY(BlueprintReadWrite, EditAnywhere, Category = "Settings|Advanced")
+	UPROPERTY(BlueprintReadWrite, EditAnywhere, Category = Settings, AdvancedDisplay)
 	bool bPersistAllData = false;
 
 	/** Opt-in option to silence errors when the path is Empty or nothing to extract. */
-	UPROPERTY(BlueprintReadWrite, EditAnywhere, Category = "Settings|Advanced")
+	UPROPERTY(BlueprintReadWrite, EditAnywhere, Category = Settings, AdvancedDisplay)
 	bool bSilenceErrorOnEmptyObjectPath = false;
 };
 
diff --git a/Engine/Plugins/PCG/Source/PCG/Public/Elements/PCGMatchAndSetAttributes.h b/Engine/Plugins/PCG/Source/PCG/Public/Elements/PCGMatchAndSetAttributes.h
--- a/Engine/Plugins/PCG/Source/PCG/Public/Elements/PCGMatchAndSetAttributes.h
+++ b/Engine/Plugins/PCG/Source/PCG/Public/Elements/PCGMatchAndSetAttributes.h
@@ -118,10 +118,10 @@ public:
 	FPCGAttributePropertyInputSelector WeightAttribute;
 
 	/** Controls whether we will emit a warning and return nothing if there is no provided attribute set. */
-	UPROPERTY(BlueprintReadWrite, EditAnywhere, Category = "Settings|Advanced")
+	UPROPERTY(BlueprintReadWrite, EditAnywhere, Category = Settings, AdvancedDisplay)
 	bool bWarnIfNoMatchData = true;
 
-	UPROPERTY(BlueprintReadWrite, EditAnywhere, Category = "Settings|Advanced")
+	UPROPERTY(BlueprintReadWrite, EditAnywhere, Category = Settings, AdvancedDisplay)
 	bool bWarnOnAttributeCast = true;
 };
 
diff --git a/Engine/Plugins/PCG/Source/PCG/Public/Elements/PCGSkinnedMeshSpawner.h b/Engine/Plugins/PCG/Source/PCG/Public/Elements/PCGSkinnedMeshSpawner.h
--- a/Engine/Plugins/PCG/Source/PCG/Public/Elements/PCGSkinnedMeshSpawner.h
+++ b/Engine/Plugins/PCG/Source/PCG/Public/Elements/PCGSkinnedMeshSpawner.h
@@ -92,11 +92,11 @@ public:
 	bool bSynchronousLoad = false;
 
 	/** Opt-in option to silence errors when the property override attributes are not found. */
-	UPROPERTY(BlueprintReadWrite, EditAnywhere, Category = "Settings|Advanced")
+	UPROPERTY(BlueprintReadWrite, EditAnywhere, Category = Settings, AdvancedDisplay)
 	bool bSilenceOverrideAttributeNotFoundErrors = false;
 
 	/** Adds a warning to the node on repeated spawning with identical conditions (ie. same mesh descriptor at same spawn location, etc). */
-	UPROPERTY(BlueprintReadWrite, EditAnywhere, Category = "Settings|Advanced")
+	UPROPERTY(BlueprintReadWrite, EditAnywhere, Category = Settings, AdvancedDisplay)
 	bool bWarnOnIdenticalSpawn = true;
 
 protected:
diff --git a/Engine/Plugins/PCG/Source/PCG/Public/Elements/PCGSpawnActor.h b/Engine/Plugins/PCG/Source/PCG/Public/Elements/PCGSpawnActor.h
--- a/Engine/Plugins/PCG/Source/PCG/Public/Elements/PCGSpawnActor.h
+++ b/Engine/Plugins/PCG/Source/PCG/Public/Elements/PCGSpawnActor.h
@@ -87,10 +87,10 @@ public:
 	FName SpawnAttribute = NAME_None;
 
 	/** Adds a warning to the node on repeated spawning with identical conditions (ie. same actor at same spawn location, etc). */
-	UPROPERTY(BlueprintReadWrite, EditAnywhere, Category = "Settings|Advanced")
+	UPROPERTY(BlueprintReadWrite, EditAnywhere, Category = Settings, AdvancedDisplay)
 	bool bWarnOnIdenticalSpawn = true;
 
-	UPROPERTY(BlueprintReadWrite, EditAnywhere, Category = "Settings|Advanced")
+	UPROPERTY(BlueprintReadWrite, EditAnywhere, Category = Settings, AdvancedDisplay)
 	bool bDeleteActorsBeforeGeneration = false;
 
 	UPROPERTY(BlueprintReadWrite, EditAnywhere, Category = Streaming, meta = (PCG_Overridable))
diff --git a/Engine/Plugins/PCG/Source/PCG/Public/Elements/PCGStaticMeshSpawner.h b/Engine/Plugins/PCG/Source/PCG/Public/Elements/PCGStaticMeshSpawner.h
--- a/Engine/Plugins/PCG/Source/PCG/Public/Elements/PCGStaticMeshSpawner.h
+++ b/Engine/Plugins/PCG/Source/PCG/Public/Elements/PCGStaticMeshSpawner.h
@@ -106,15 +106,15 @@ public:
 	bool bSynchronousLoad = false;
 
 	/** Controls whether instances stemming from different data can end up in the same ISM. Performance warning if set to false. */
-	UPROPERTY(BlueprintReadWrite, EditAnywhere, Category = "Settings|Advanced", meta = (PCG_Overridable))
+	UPROPERTY(BlueprintReadWrite, EditAnywhere, Category = Settings, AdvancedDisplay, meta = (PCG_Overridable))
 	bool bAllowMergeDifferentDataInSameInstancedComponents = true;
 
 	/** Opt-in option to silence errors when the property override attributes are not found. */
-	UPROPERTY(BlueprintReadWrite, EditAnywhere, Category = "Settings|Advanced")
+	UPROPERTY(BlueprintReadWrite, EditAnywhere, Category = Settings, AdvancedDisplay)
 	bool bSilenceOverrideAttributeNotFoundErrors = false;
 
 	/** Adds a warning to the node on repeated spawning with identical conditions (ie. same mesh descriptor at same spawn location, etc). */
-	UPROPERTY(BlueprintReadWrite, EditAnywhere, Category = "Settings|Advanced")
+	UPROPERTY(BlueprintReadWrite, EditAnywhere, Category = Settings, AdvancedDisplay)
 	bool bWarnOnIdenticalSpawn = true;
 
 protected:
diff --git a/Engine/Plugins/PCGInterops/PCGGeometryScriptInterop/Source/PCGGeometryScriptInterop/Public/Elements/PCGPrimitiveCrossSection.h b/Engine/Plugins/PCGInterops/PCGGeometryScriptInterop/Source/PCGGeometryScriptInterop/Public/Elements/PCGPrimitiveCrossSection.h
--- a/Engine/Plugins/PCGInterops/PCGGeometryScriptInterop/Source/PCGGeometryScriptInterop/Public/Elements/PCGPrimitiveCrossSection.h
+++ b/Engine/Plugins/PCGInterops/PCGGeometryScriptInterop/Source/PCGGeometryScriptInterop/Public/Elements/PCGPrimitiveCrossSection.h
@@ -50,7 +50,7 @@ public:
 	int32 MinimumCoplanarVertices = 3;
 
 	/** A safeguard to prevent finding features on an overly complex mesh. */
-	UPROPERTY(BlueprintReadWrite, EditAnywhere, Category = "Settings|Advanced", meta = (PCG_Overridable, EditCondition = "GenerateCrossSectionMode != EPCGGenerateCrossSectionMode::TierSlicing", EditConditionHides))
+	UPROPERTY(BlueprintReadWrite, EditAnywhere, Category = Settings, AdvancedDisplay, meta = (PCG_Overridable, EditCondition = "GenerateCrossSectionMode != EPCGGenerateCrossSectionMode::TierSlicing", EditConditionHides))
 	int32 MaxMeshVertexCount = 2048;
 
 	/** Cull tiers that are within a specified threshold. */
