[PCG] Fix soft lock when inspecting a node upstream of a compute graph

#rnx

diff --git a/Engine/Plugins/PCG/Source/PCG/Private/Compute/DataInterfaces/PCGComputeDataInterface.cpp b/Engine/Plugins/PCG/Source/PCG/Private/Compute/DataInterfaces/PCGComputeDataInterface.cpp
--- a/Engine/Plugins/PCG/Source/PCG/Private/Compute/DataInterfaces/PCGComputeDataInterface.cpp
+++ b/Engine/Plugins/PCG/Source/PCG/Private/Compute/DataInterfaces/PCGComputeDataInterface.cpp
@@ -108,6 +108,12 @@ const UPCGSettings* UPCGComputeDataProvider::GetProducerSettings() const
 	return ResolvedProducerSettings;
 }
 
+void UPCGComputeDataProvider::SetProducerSettings(const UPCGSettings* InSettings)
+{
+	ProducerSettings = InSettings;
+	ResolvedProducerSettings = InSettings;
+}
+
 void UPCGComputeDataProvider::AddReferencedObjects(UObject* InThis, FReferenceCollector& Collector)
 {
 	Super::AddReferencedObjects(InThis, Collector);
diff --git a/Engine/Plugins/PCG/Source/PCG/Private/Compute/DataInterfaces/PCGComputeDataInterface.h b/Engine/Plugins/PCG/Source/PCG/Private/Compute/DataInterfaces/PCGComputeDataInterface.h
--- a/Engine/Plugins/PCG/Source/PCG/Private/Compute/DataInterfaces/PCGComputeDataInterface.h
+++ b/Engine/Plugins/PCG/Source/PCG/Private/Compute/DataInterfaces/PCGComputeDataInterface.h
@@ -76,6 +76,7 @@ public:
 	virtual void ReleaseTransientResources(const TCHAR* InReason = nullptr) {}
 
 	const UPCGSettings* GetProducerSettings() const;
+	void SetProducerSettings(const UPCGSettings* InSettings);
 
 	static void AddReferencedObjects(UObject* InThis, FReferenceCollector& Collector);
 
diff --git a/Engine/Plugins/PCG/Source/PCG/Private/Compute/DataInterfaces/PCGDataCollectionUploadDataInterface.cpp b/Engine/Plugins/PCG/Source/PCG/Private/Compute/DataInterfaces/PCGDataCollectionUploadDataInterface.cpp
--- a/Engine/Plugins/PCG/Source/PCG/Private/Compute/DataInterfaces/PCGDataCollectionUploadDataInterface.cpp
+++ b/Engine/Plugins/PCG/Source/PCG/Private/Compute/DataInterfaces/PCGDataCollectionUploadDataInterface.cpp
@@ -29,11 +29,13 @@ void UPCGDataProviderDataCollectionUpload::Initialize(const UComputeDataInterfac
 {
 	TRACE_CPUPROFILER_EVENT_SCOPE(UPCGDataProviderDataCollectionUpload::Initialize);
 
-	Super::Initialize(InDataInterface, InBinding, InInputMask, InOutputMask);
+	// @todo_pcg: Properly initializing this data provider causes unexpected behavior.
+	// Super::Initialize(InDataInterface, InBinding, InInputMask, InOutputMask);
 
 	const UPCGDataCollectionUploadDataInterface* DataInterface = CastChecked<UPCGDataCollectionUploadDataInterface>(InDataInterface);
 
 	Binding = CastChecked<UPCGDataBinding>(InBinding);
+	SetProducerSettings(DataInterface->GetProducerSettings());
 	GraphBindingIndex = DataInterface->GraphBindingIndex;
 
 	// Pick the data items from input data collection using any of the compute graph element virtual input pin labels.
