[PCG] Fix crash when gathering additional sources during GPU graph compilation

#rnx

diff --git a/Engine/Plugins/PCG/Source/PCG/Private/Compute/PCGComputeKernel.cpp b/Engine/Plugins/PCG/Source/PCG/Private/Compute/PCGComputeKernel.cpp
--- a/Engine/Plugins/PCG/Source/PCG/Private/Compute/PCGComputeKernel.cpp
+++ b/Engine/Plugins/PCG/Source/PCG/Private/Compute/PCGComputeKernel.cpp
@@ -14,6 +14,7 @@
 #include "Graph/PCGGPUGraphCompilationContext.h"
 
 #include "ShaderCompilerCore.h"
+#include "Async/Async.h"
 
 #include UE_INLINE_GENERATED_CPP_BY_NAME(PCGComputeKernel)
 
@@ -26,9 +27,15 @@ namespace PCGComputeKernel
 		true,
 		TEXT("Enables warnings when parameters are overidden on GPU nodes."));
 
-	static const FSoftObjectPath ShaderUtilsComputeSourcePath(TEXT("PCGComputeSource'/PCG/ComputeSources/PCGCS_ShaderUtils.PCGCS_ShaderUtils'"));
-	static const FSoftObjectPath ShaderUtilsInternalComputeSourcePath(TEXT("PCGComputeSource'/PCG/ComputeSources/PCGCS_ShaderUtilsInternal.PCGCS_ShaderUtilsInternal'"));
-	static const FSoftObjectPath QuaternionComputeSourcePath(TEXT("PCGComputeSource'/PCG/ComputeSources/PCGCS_Quaternion.PCGCS_Quaternion'"));
+	static const TStaticArray<FSoftObjectPath, 3> DefaultAdditionalSourcePaths =
+	{
+		FSoftObjectPath(TEXT("/Script/PCG.PCGComputeSource'/PCG/ComputeSources/PCGCS_ShaderUtils.PCGCS_ShaderUtils'")),
+		FSoftObjectPath(TEXT("/Script/PCG.PCGComputeSource'/PCG/ComputeSources/PCGCS_ShaderUtilsInternal.PCGCS_ShaderUtilsInternal'")),
+
+		// Note: PCGDataCollectionDataInterface.ush depends on the quaternion helpers, therefore all kernels also depend on the quaternion helpers.
+		// @todo_pcg: In the future quaternion compute source could be opt-in if the kernel does not manipulate point/attribute data.
+		FSoftObjectPath(TEXT("/Script/PCG.PCGComputeSource'/PCG/ComputeSources/PCGCS_Quaternion.PCGCS_Quaternion'"))
+	};
 }
 
 #if WITH_EDITOR
@@ -92,25 +99,51 @@ bool UPCGComputeKernel::AreKernelSettingsValid(FPCGContext* InContext) const
 #if WITH_EDITOR
 void UPCGComputeKernel::GatherAdditionalSources(TArray<TObjectPtr<UComputeSource>>& OutAdditionalSources) const
 {
-	auto GatherComputeSource = [&OutAdditionalSources](const FSoftObjectPath& InComputeSourcePath)
+	auto GatherSources = [this, &OutAdditionalSources, &DefaultAdditionalSourcePaths=PCGComputeKernel::DefaultAdditionalSourcePaths]()
 	{
-		TSoftObjectPtr<UPCGComputeSource> ComputeSourcePtr(InComputeSourcePath);
-		if (UPCGComputeSource* ComputeSource = ComputeSourcePtr.LoadSynchronous(); ensure(ComputeSource))
-		{
-			OutAdditionalSources.Add(ComputeSource);
-		}
-		else
+		check(IsInGameThread());
+
+		for (const FSoftObjectPath& AdditionalSourcePath : DefaultAdditionalSourcePaths)
 		{
-			UE_LOG(LogPCG, Error, TEXT("Compute source at '%s' could not be loaded."), *InComputeSourcePath.ToString());
+			TSoftObjectPtr<UPCGComputeSource> AdditionalSourcePtr(AdditionalSourcePath);
+			UPCGComputeSource* AdditionalSource = Cast<UPCGComputeSource>(AdditionalSourcePtr.LoadSynchronous());
+			
+			if (ensure(AdditionalSource))
+			{
+				OutAdditionalSources.Add(AdditionalSource);
+			}
+			else
+			{
+				UE_LOG(LogPCG, Error, TEXT("Default additional compute source '%s' could not be loaded for kernel '%s'."), *AdditionalSourcePath.ToString(), *GetName());
+			}
 		}
 	};
 
-	GatherComputeSource(PCGComputeKernel::ShaderUtilsComputeSourcePath);
-	GatherComputeSource(PCGComputeKernel::ShaderUtilsInternalComputeSourcePath);
+	if (!IsInGameThread())
+	{
+		// Loads must happen on the game thread, so we launch and wait for an async game thread task if necessary.
+		if (FEvent* LoadEvent = FPlatformProcess::GetSynchEventFromPool(/*bIsManualReset=*/true))
+		{
+			AsyncTask(ENamedThreads::GameThread, [&GatherSources, LoadEvent]()
+			{
+				GatherSources();
+				LoadEvent->Trigger();
+			});
 
-	// Note: PCGDataCollectionDataInterface.ush depends on the quaternion helpers, therefore all kernels also depend on the quaternion helpers.
-	// @todo_pcg: In the future quaternion compute source could be opt-in if the kernel does not manipulate point/attribute data.
-	GatherComputeSource(PCGComputeKernel::QuaternionComputeSourcePath);
+			LoadEvent->Wait();
+
+			FPlatformProcess::ReturnSynchEventToPool(LoadEvent);
+		}
+		else
+		{
+			UE_LOG(LogPCG, Error, TEXT("Default additional compute sources could not be loaded for kernel '%s'."), *GetName());
+		}
+	}
+	else
+	{
+		// If we're already on the game thread, just load the compute sources.
+		GatherSources();
+	}
 }
 #endif
 
