[PCG] Fix compute source asset changes not triggering refresh

#rnx

diff --git a/Engine/Plugins/PCG/Source/PCG/Private/Compute/Elements/PCGCustomHLSL.cpp b/Engine/Plugins/PCG/Source/PCG/Private/Compute/Elements/PCGCustomHLSL.cpp
--- a/Engine/Plugins/PCG/Source/PCG/Private/Compute/Elements/PCGCustomHLSL.cpp
+++ b/Engine/Plugins/PCG/Source/PCG/Private/Compute/Elements/PCGCustomHLSL.cpp
@@ -8,6 +8,7 @@
 #include "PCGModule.h"
 #include "PCGPoint.h"
 #include "PCGSubsystem.h"
+#include "Compute/PCGComputeSource.h"
 #include "Compute/Elements/PCGCustomHLSLKernel.h"
 #include "Data/PCGPointData.h"
 #include "Graph/PCGGPUGraphCompilationContext.h"
@@ -31,6 +32,16 @@ namespace PCGHLSLElement
 UPCGCustomHLSLSettings::UPCGCustomHLSLSettings()
 {
 	bExecuteOnGPU = true;
+
+#if WITH_EDITOR
+	if (!HasAnyFlags(RF_ClassDefaultObject) && !FPCGContext::IsInitializingSettings())
+	{
+		if (!UPCGComputeSource::OnModifiedDelegate.IsBoundToObject(this))
+		{
+			UPCGComputeSource::OnModifiedDelegate.AddUObject(this, &UPCGCustomHLSLSettings::OnComputeSourceModified);
+		}
+	}
+#endif
 }
 
 #if WITH_EDITOR
@@ -53,6 +64,13 @@ void UPCGCustomHLSLSettings::PostInitProperties()
 	UpdateAttributeKeys();
 	UpdateDeclarations();
 }
+
+void UPCGCustomHLSLSettings::BeginDestroy()
+{
+	UPCGComputeSource::OnModifiedDelegate.RemoveAll(this);
+
+	Super::BeginDestroy();
+}
 #endif
 
 TArray<FPCGPinProperties> UPCGCustomHLSLSettings::OutputPinProperties() const
@@ -979,6 +997,46 @@ void UPCGCustomHLSLSettings::UpdateAttributeKeys()
 	}
 }
 
+void UPCGCustomHLSLSettings::OnComputeSourceModified(const UPCGComputeSource* InModifiedComputeSource)
+{
+	TArray<UComputeSource*> ComputeSourcesToVisit;
+	TSet<UComputeSource*> VisitedComputeSources;
+
+	ComputeSourcesToVisit.Push(KernelSourceOverride);
+	ComputeSourcesToVisit.Append(AdditionalSources);
+
+	// Visit the entire network of additional sources to see if our source depends on the modified compute source.
+	bool bAnyMatch = false;
+
+	while (!ComputeSourcesToVisit.IsEmpty())
+	{
+		if (UComputeSource* ComputeSource = ComputeSourcesToVisit.Pop())
+		{
+			if (ComputeSource == InModifiedComputeSource)
+			{
+				bAnyMatch = true;
+				break;
+			}
+
+			VisitedComputeSources.Add(ComputeSource);
+
+			for (UComputeSource* AdditionalSource : ComputeSource->AdditionalSources)
+			{
+				if (!VisitedComputeSources.Contains(AdditionalSource))
+				{
+					ComputeSourcesToVisit.Add(AdditionalSource);
+				}
+			}
+		}
+	}
+
+	if (bAnyMatch)
+	{
+		// @todo_pcg: Revisit whether we can remove Structural from this (and other) source modifications.
+		OnSettingsChangedDelegate.Broadcast(this, EPCGChangeType::ShaderSource | EPCGChangeType::Structural);
+	}
+}
+
 TArray<FName> UPCGCustomHLSLSettings::GetInputPinNames() const
 {
 	TArray<FName> PinNames;
diff --git a/Engine/Plugins/PCG/Source/PCG/Private/Compute/Elements/PCGCustomHLSL.h b/Engine/Plugins/PCG/Source/PCG/Private/Compute/Elements/PCGCustomHLSL.h
--- a/Engine/Plugins/PCG/Source/PCG/Private/Compute/Elements/PCGCustomHLSL.h
+++ b/Engine/Plugins/PCG/Source/PCG/Private/Compute/Elements/PCGCustomHLSL.h
@@ -10,6 +10,7 @@
 #include "PCGCustomHLSL.generated.h"
 
 class UPCGComputeKernel;
+class UPCGComputeSource;
 class UPCGPin;
 
 class UComputeSource;
@@ -49,6 +50,7 @@ public:
 	//~Begin UObject interface
 	virtual void PostLoad() override;
 	virtual void PostInitProperties() override;
+	virtual void BeginDestroy() override;
 	//~End UObject interface
 #endif
 
@@ -110,6 +112,9 @@ protected:
 	void UpdatePinSettings();
 	void UpdateAttributeKeys();
 
+	/** Called when a compute source is modified to propagate graph refreshes. */
+	void OnComputeSourceModified(const UPCGComputeSource* InModifiedComputeSource);
+
 	/** List of all non-advanced input pin names. */
 	UFUNCTION()
 	TArray<FName> GetInputPinNames() const;
diff --git a/Engine/Plugins/PCG/Source/PCG/Private/Compute/PCGComputeSource.cpp b/Engine/Plugins/PCG/Source/PCG/Private/Compute/PCGComputeSource.cpp
--- a/Engine/Plugins/PCG/Source/PCG/Private/Compute/PCGComputeSource.cpp
+++ b/Engine/Plugins/PCG/Source/PCG/Private/Compute/PCGComputeSource.cpp
@@ -6,6 +6,24 @@
 
 #include UE_INLINE_GENERATED_CPP_BY_NAME(PCGComputeSource)
 
+#if WITH_EDITOR
+FOnPCGComputeSourceModified UPCGComputeSource::OnModifiedDelegate;
+
+void UPCGComputeSource::PostEditChangeProperty(FPropertyChangedEvent& PropertyChangedEvent)
+{
+	Super::PostEditChangeProperty(PropertyChangedEvent);
+
+	OnModifiedDelegate.Broadcast(this);
+}
+
+void UPCGComputeSource::PostEditUndo()
+{
+	Super::PostEditUndo();
+
+	OnModifiedDelegate.Broadcast(this);
+}
+#endif
+
 FString UPCGComputeSource::GetSource() const
 {
 #if WITH_EDITOR
@@ -29,12 +47,15 @@ FString UPCGComputeSource::GetVirtualPath() const
 }
 
 #if WITH_EDITOR
-void UPCGComputeSource::SetSource(const FString& InSource)
+void UPCGComputeSource::SetShaderText(const FString& NewText)
 {
-	if (Source != InSource)
+	if (Source != NewText)
 	{
 		Modify();
-		Source = InSource;
+
+		Source = NewText;
+
+		OnModifiedDelegate.Broadcast(this);
 	}
 }
 #endif
diff --git a/Engine/Plugins/PCG/Source/PCG/Public/Compute/PCGComputeSource.h b/Engine/Plugins/PCG/Source/PCG/Public/Compute/PCGComputeSource.h
--- a/Engine/Plugins/PCG/Source/PCG/Public/Compute/PCGComputeSource.h
+++ b/Engine/Plugins/PCG/Source/PCG/Public/Compute/PCGComputeSource.h
@@ -8,6 +8,11 @@
 
 #include "PCGComputeSource.generated.h"
 
+#if WITH_EDITOR
+class UPCGComputeSource;
+DECLARE_MULTICAST_DELEGATE_OneParam(FOnPCGComputeSourceModified, const UPCGComputeSource*);
+#endif // WITH_EDITOR
+
 #define UE_API PCG_API
 
 UCLASS(MinimalAPI)
@@ -16,6 +21,13 @@ class UPCGComputeSource : public UComputeSource, public IPCGNodeSourceTextProvid
 	GENERATED_BODY()
 
 public:
+#if WITH_EDITOR
+	//~Begin UObject interface
+	virtual void PostEditChangeProperty(FPropertyChangedEvent& PropertyChangedEvent) override;
+	virtual void PostEditUndo() override;
+	//~Begin UObject interface
+#endif
+
 	//~ Begin UComputeSource Interface.
 	UE_API FString GetSource() const override;
 	UE_API FString GetVirtualPath() const override;
@@ -27,11 +39,13 @@ public:
 	FString GetDeclarationsText() const override { return {}; }
 	FString GetShaderFunctionsText() const override { return {}; }
 	void SetShaderFunctionsText(const FString& NewFunctionsText) override {}
-	void SetShaderText(const FString& NewText) override { SetSource(NewText); }
+	void SetShaderText(const FString& NewText) override;
 	bool IsShaderTextReadOnly() const override { return false; }
 	//~End IPCGNodeSourceTextProvider interface
-
+	
 	UE_API void SetSource(const FString& InSource);
+
+	static FOnPCGComputeSourceModified OnModifiedDelegate;
 #endif
 
 #if WITH_EDITORONLY_DATA
diff --git a/Engine/Plugins/PCG/Source/PCGEditor/Private/Details/PCGComputeSourceDetails.cpp b/Engine/Plugins/PCG/Source/PCGEditor/Private/Details/PCGComputeSourceDetails.cpp
--- a/Engine/Plugins/PCG/Source/PCGEditor/Private/Details/PCGComputeSourceDetails.cpp
+++ b/Engine/Plugins/PCG/Source/PCGEditor/Private/Details/PCGComputeSourceDetails.cpp
@@ -95,7 +95,8 @@ void FPCGComputeSourceDetails::CustomizeDetails(IDetailLayoutBuilder& DetailBuil
 		.FillHeight(1.0f)
 		[
 			SNew(SBox)
-			.HeightOverride(500)
+			.MinDesiredHeight(200)
+			.MaxDesiredHeight(800)
 			[
 				SAssignNew(SourceTextBox, SPCGNodeSourceTextBox)
 				.Text(this, &FPCGComputeSourceDetails::GetSourceText)
