[PCG] Refactor subsystem into a base class and the original UPCGSubsystem and a new UPCGEditorSubsystem

#rnx

diff --git a/Engine/Plugins/PCG/Source/PCG/Private/Graph/PCGGraphExecutor.h b/Engine/Plugins/PCG/Source/PCG/Private/Graph/PCGGraphExecutor.h
--- a/Engine/Plugins/PCG/Source/PCG/Private/Graph/PCGGraphExecutor.h
+++ b/Engine/Plugins/PCG/Source/PCG/Private/Graph/PCGGraphExecutor.h
@@ -156,6 +156,7 @@ public:
 
 private:
 	friend UPCGSubsystem;
+	friend IPCGBaseSubsystem;
 
 	enum class EExecutionCacheDataType : int8
 	{
diff --git a/Engine/Plugins/PCG/Source/PCG/Private/PCGComponent.cpp b/Engine/Plugins/PCG/Source/PCG/Private/PCGComponent.cpp
--- a/Engine/Plugins/PCG/Source/PCG/Private/PCGComponent.cpp
+++ b/Engine/Plugins/PCG/Source/PCG/Private/PCGComponent.cpp
@@ -842,7 +842,7 @@ void UPCGComponent::PostProcessGraph(const FBox& InNewBounds, bool bInGenerated,
 	StopGenerationInProgress();
 
 	UPCGSubsystem* Subsystem = GetSubsystem();
-	Subsystem->OnPCGComponentGenerationDone.Broadcast(Subsystem, IsValidChecked(this) ? this : nullptr, EPCGGenerationStatus::Completed);
+	Subsystem->OnPCGSourceGenerationDone.Broadcast(Subsystem, IsValidChecked(this) ? this : nullptr, EPCGGenerationStatus::Completed);
 #endif
 }
 
@@ -978,7 +978,7 @@ void UPCGComponent::OnProcessGraphAborted(bool bQuiet, bool bCleanupUnusedResour
 	OnPCGGraphCancelledDelegate.Broadcast(this);
 
 #if WITH_EDITOR
-	Subsystem->OnPCGComponentGenerationDone.Broadcast(Subsystem, this, EPCGGenerationStatus::Aborted);
+	Subsystem->OnPCGSourceGenerationDone.Broadcast(Subsystem, this, EPCGGenerationStatus::Aborted);
 #endif
 
 	PCGComponent::BroadcastDynamicDelegate(OnPCGGraphCancelledExternal, this);
@@ -4221,6 +4221,12 @@ void FPCGComponentExecutionState::Cancel()
 	Component->CancelGeneration();
 }
 
+bool FPCGComponentExecutionState::IsGenerating() const
+{
+	check(Component);
+	return Component->IsGenerating();
+}
+
 #if WITH_EDITOR
 const PCGUtils::FExtraCapture& FPCGComponentExecutionState::GetExtraCapture() const
 {
@@ -4257,6 +4263,12 @@ void FPCGComponentExecutionState::RegisterDynamicTracking(const FPCGSelectionKey
 	check(Component);
 	Component->RegisterDynamicTracking(InKeysToSettings);
 }
+
+bool FPCGComponentExecutionState::IsRefreshInProgress() const
+{
+	check(Component);
+	return Component->IsRefreshInProgress();
+}
 #endif // WITH_EDITOR
 
 #undef LOCTEXT_NAMESPACE
diff --git a/Engine/Plugins/PCG/Source/PCG/Private/PCGGraphExecutionStateInterface.cpp b/Engine/Plugins/PCG/Source/PCG/Private/PCGGraphExecutionStateInterface.cpp
--- a/Engine/Plugins/PCG/Source/PCG/Private/PCGGraphExecutionStateInterface.cpp
+++ b/Engine/Plugins/PCG/Source/PCG/Private/PCGGraphExecutionStateInterface.cpp
@@ -2,7 +2,106 @@
 
 #include "PCGGraphExecutionStateInterface.h"
 
+#include "Editor/IPCGEditorModule.h"
+#include "Subsystems/PCGSubsystem.h"
+
+IPCGBaseSubsystem* IPCGGraphExecutionState::GetSubsystem() const
+{
+	IPCGBaseSubsystem* Subsystem = nullptr;
+	
+	if (UWorld* World = GetWorld())
+	{
+		Subsystem = UPCGSubsystem::GetInstance(World);
+	}
+#if WITH_EDITOR
+	else if (const IPCGEditorModule* EditorModule = IPCGEditorModule::Get())
+	{
+		Subsystem = EditorModule->GetEditorSubsystem();
+	}
+#endif // WITH_EDITOR
+
+	return Subsystem;
+}
+
+//////////////////////////////////////////////////////////////////////
+
 UPCGGraphExecutionSource::UPCGGraphExecutionSource(const FObjectInitializer& ObjectInitializer)
 	: Super(ObjectInitializer)
 {
+}
+
+//////////////////////////////////////////////////////////////////////
+
+FPCGSoftGraphExecutionSource::FPCGSoftGraphExecutionSource(const TSoftObjectPtr<UObject>& InSoftObjectPtr)
+	: SoftObjectPtr(InSoftObjectPtr)
+{
+	
+}
+
+FPCGSoftGraphExecutionSource::FPCGSoftGraphExecutionSource(IPCGGraphExecutionSource* InSource)
+	: FPCGSoftGraphExecutionSource(TSoftObjectPtr(Cast<UObject>(InSource)))
+{
+}
+
+FPCGSoftGraphExecutionSource::FPCGSoftGraphExecutionSource(const IPCGGraphExecutionSource* InSource)
+	: FPCGSoftGraphExecutionSource(const_cast<IPCGGraphExecutionSource*>(InSource))
+{
+}
+
+FPCGSoftGraphExecutionSource::FPCGSoftGraphExecutionSource(const FPCGSoftGraphExecutionSource& Other)
+{
+	SoftObjectPtr = Other.SoftObjectPtr;
+}
+
+FPCGSoftGraphExecutionSource::FPCGSoftGraphExecutionSource(FPCGSoftGraphExecutionSource&& Other)
+{
+	SoftObjectPtr = std::move(Other.SoftObjectPtr);
+}
+
+FPCGSoftGraphExecutionSource& FPCGSoftGraphExecutionSource::operator=(const FPCGSoftGraphExecutionSource& Other)
+{
+	SoftObjectPtr = Other.SoftObjectPtr;
+	return *this;
+}
+
+FPCGSoftGraphExecutionSource& FPCGSoftGraphExecutionSource::operator=(FPCGSoftGraphExecutionSource&& Other)
+{
+	SoftObjectPtr = std::move(Other.SoftObjectPtr);
+	return *this;
+}
+
+IPCGGraphExecutionSource* FPCGSoftGraphExecutionSource::Get() const
+{
+	if (!CachedWeakPtr.IsValid())
+	{
+		FGCScopeGuard ScopeGuard;
+		
+		UObject* Object = SoftObjectPtr.Get();
+		if (!Object || !Object->Implements<UPCGGraphExecutionSource>())
+		{
+			return nullptr;
+		}
+
+		IPCGGraphExecutionSource* InterfacePtr = CastChecked<IPCGGraphExecutionSource>(Object);
+		
+		{
+			UE::TScopeLock Lock(SpinLock);
+			if (!CachedWeakPtr.IsValid())
+			{
+				CachedWeakPtr = CastChecked<IPCGGraphExecutionSource>(Object);
+			}
+		}
+	}
+
+	return CachedWeakPtr.Get();
+}
+
+bool FPCGSoftGraphExecutionSource::operator==(const FPCGSoftGraphExecutionSource& Other) const
+{
+	return Other.SoftObjectPtr == SoftObjectPtr;
+}
+
+int32 GetTypeHash(const FPCGSoftGraphExecutionSource& This)
+{
+	return GetTypeHash(This.SoftObjectPtr);
 }
\ No newline at end of file
diff --git a/Engine/Plugins/PCG/Source/PCG/Private/PCGSubsystem.cpp b/Engine/Plugins/PCG/Source/PCG/Private/PCGSubsystem.cpp
deleted file mode 100644
diff --git a/Engine/Plugins/PCG/Source/PCG/Private/Subsystems/IPCGBaseSubsystem.cpp b/Engine/Plugins/PCG/Source/PCG/Private/Subsystems/IPCGBaseSubsystem.cpp
new file mode 100644
--- /dev/null
+++ b/Engine/Plugins/PCG/Source/PCG/Private/Subsystems/IPCGBaseSubsystem.cpp
@@ -0,0 +1,287 @@
+// Copyright Epic Games, Inc. All Rights Reserved.
+
+#include "Subsystems/IPCGBaseSubsystem.h"
+
+#include "PCGActorAndComponentMapping.h"
+#include "PCGComponent.h"
+#include "PCGGraph.h"
+#include "Graph/PCGGraphCache.h"
+#include "Graph/PCGGraphCompiler.h"
+#include "Graph/PCGGraphExecutor.h"
+#include "Helpers/PCGActorHelpers.h"
+#include "Helpers/PCGHelpers.h"
+#include "RuntimeGen/PCGGenSourceManager.h"
+
+#if WITH_EDITOR
+#include "Editor.h"
+#include "PackageSourceControlHelper.h"
+#include "Misc/ScopedSlowTask.h"
+#include "WorldPartition/WorldPartition.h"
+#include "WorldPartition/DataLayer/DataLayerManager.h"
+#endif // WITH_EDITOR
+
+void IPCGBaseSubsystem::InitializeBaseSubsystem()
+{
+	check(!GraphExecutor);
+	GraphExecutor = MakeShared<FPCGGraphExecutor>(GetSubsystemWorld());
+}
+
+void IPCGBaseSubsystem::DeinitializeBaseSubsystem()
+{
+	GraphExecutor = nullptr;
+}
+
+#if WITH_EDITOR
+void IPCGBaseSubsystem::OnScheduleGraph(const FPCGStackContext& StackContext)
+{
+	// nothing to do for now
+}
+#endif
+
+FPCGTaskId IPCGBaseSubsystem::ScheduleGraph(const FPCGScheduleGraphParams& InParams)
+{
+	if (InParams.ExecutionSource)
+	{
+		return GraphExecutor->ScheduleGraph(InParams);
+	}
+	else
+	{
+		return InvalidPCGTaskId;
+	}
+}
+
+FPCGTaskId IPCGBaseSubsystem::ScheduleGeneric(const FPCGScheduleGenericParams& InParams)
+{
+	check(GraphExecutor);
+	return GraphExecutor->ScheduleGeneric(InParams);
+}
+
+void IPCGBaseSubsystem::CancelGeneration(IPCGGraphExecutionSource* Source)
+{
+	CancelGeneration(Source, /*bCleanupUnusedResources=*/true);
+}
+
+double IPCGBaseSubsystem::GetTickEndTime() const
+{
+	return FPlatformTime::Seconds() + FPCGGraphExecutor::GetTickBudgetInSeconds();
+}
+
+double IPCGBaseSubsystem::Tick()
+{
+	double EndTime = GetTickEndTime();
+
+	// If we have any tasks to execute, schedule some
+	if (GraphExecutor)
+	{
+		GraphExecutor->Execute(EndTime);
+	}
+
+	return EndTime;
+}
+
+void IPCGBaseSubsystem::CancelGeneration(IPCGGraphExecutionSource* Source, bool bCleanupUnusedResources)
+{
+	check(GraphExecutor && IsInGameThread());
+	if (!Source || !Source->GetExecutionState().IsGenerating())
+	{
+		return;
+	}
+
+	CancelGenerationInternal(Source, bCleanupUnusedResources);
+
+	TArray<IPCGGraphExecutionSource*> CancelledExecutionSources = GraphExecutor->Cancel(Source);
+	for (IPCGGraphExecutionSource* CancelledExecutionSource : CancelledExecutionSources)
+	{
+		if (CancelledExecutionSource)
+		{
+			CancelledExecutionSource->GetExecutionState().OnGraphExecutionAborted(/*bQuiet=*/true, bCleanupUnusedResources);
+		}
+	}	
+}
+
+void IPCGBaseSubsystem::CancelGeneration(UPCGGraph* Graph)
+{
+	check(GraphExecutor);
+
+	if (!Graph)
+	{
+		return;
+	}
+
+	TArray<IPCGGraphExecutionSource*> CancelledExecutionSources = GraphExecutor->Cancel(Graph);
+	for (IPCGGraphExecutionSource* CancelledExecutionSource : CancelledExecutionSources)
+	{
+		if (ensure(CancelledExecutionSource))
+		{
+			CancelledExecutionSource->GetExecutionState().OnGraphExecutionAborted(/*bQuiet=*/true);
+		}
+	}
+}
+
+void IPCGBaseSubsystem::CancelAllGeneration()
+{
+	check(GraphExecutor);
+
+	TArray<IPCGGraphExecutionSource*> CancelledExecutionSources = GraphExecutor->CancelAll();
+	for (IPCGGraphExecutionSource* CancelledExecutionSource : CancelledExecutionSources)
+	{
+		if (ensure(CancelledExecutionSource))
+		{
+			CancelledExecutionSource->GetExecutionState().OnGraphExecutionAborted(/*bQuiet=*/true);
+		}
+	}
+}
+
+bool IPCGBaseSubsystem::IsGraphCurrentlyExecuting(UPCGGraph* Graph)
+{
+	check(GraphExecutor);
+
+	if (!Graph)
+	{
+		return false;
+	}
+
+	return GraphExecutor->IsGraphCurrentlyExecuting(Graph);
+}
+
+bool IPCGBaseSubsystem::IsAnyGraphCurrentlyExecuting() const
+{
+	return GraphExecutor && GraphExecutor->IsAnyGraphCurrentlyExecuting();
+}
+
+bool IPCGBaseSubsystem::IsGraphCacheDebuggingEnabled() const
+{
+	return GraphExecutor && GraphExecutor->IsGraphCacheDebuggingEnabled();
+}
+
+FPCGGraphCompiler* IPCGBaseSubsystem::GetGraphCompiler()
+{
+	return GraphExecutor ? GraphExecutor->GetCompiler() : nullptr;
+}
+
+UPCGComputeGraph* IPCGBaseSubsystem::GetComputeGraph(const UPCGGraph* InGraph, uint32 GridSize, uint32 ComputeGraphIndex)
+{
+	if (FPCGGraphCompiler* GraphCompiler = GetGraphCompiler())
+	{
+		return GraphCompiler->GetComputeGraph(InGraph, GridSize, ComputeGraphIndex);
+	}
+
+	return nullptr;
+}
+
+bool IPCGBaseSubsystem::GetOutputData(FPCGTaskId TaskId, FPCGDataCollection& OutData)
+{
+	check(GraphExecutor);
+	return GraphExecutor->GetOutputData(TaskId, OutData);
+}
+
+void IPCGBaseSubsystem::ClearOutputData(FPCGTaskId TaskId)
+{
+	check(GraphExecutor);
+	GraphExecutor->ClearOutputData(TaskId);
+}
+
+#if WITH_EDITOR
+void IPCGBaseSubsystem::NotifyGraphChanged(UPCGGraph* InGraph, EPCGChangeType ChangeType)
+{
+	if (GraphExecutor)
+	{
+		GraphExecutor->NotifyGraphChanged(InGraph, ChangeType);
+	}
+}
+
+void IPCGBaseSubsystem::CleanFromCache(const IPCGElement* InElement, const UPCGSettings* InSettings/*= nullptr*/)
+{
+	if (GraphExecutor)
+	{
+		GraphExecutor->GetCache().CleanFromCache(InElement, InSettings);
+	}
+}
+
+bool IPCGBaseSubsystem::GetStackContext(UPCGGraph* InGraph, uint32 InGridSize, bool bIsPartitioned, FPCGStackContext& OutStackContext)
+{
+	if (!InGraph)
+	{
+		return false;
+	}
+
+	// A non-partitioned component generally executes (original component or local component).
+	if(bIsPartitioned)
+	{
+		// A partitioned higen original component will execute if the graph has UB grid level.
+		if (InGraph->IsHierarchicalGenerationEnabled())
+		{
+			PCGHiGenGrid::FSizeArray GridSizes;
+			bool bHasUnbounded = false;
+			InGraph->GetGridSizes(GridSizes, bHasUnbounded);
+
+			if (!bHasUnbounded)
+			{
+				return false;
+			}
+		}
+		else
+		{
+			return false;
+		}
+	}
+
+	if (FPCGGraphCompiler* GraphCompiler = GetGraphCompiler())
+	{
+		GraphCompiler->GetCompiledTasks(InGraph, InGridSize, OutStackContext, /*bIsTopGraph=*/false);
+		return true;
+	}
+	else
+	{
+		return false;
+	}
+}
+
+bool IPCGBaseSubsystem::GetStackContext(const IPCGGraphExecutionSource* InSource, FPCGStackContext& OutStackContext)
+{
+	// @todo_pcg: Extend execution state.
+	const UPCGComponent* Component = Cast<UPCGComponent>(InSource);
+	const bool bIsPartitioned = Component && Component->IsPartitioned();
+	const uint32 GridSize = (!Component || bIsPartitioned) ?  PCGHiGenGrid::UnboundedGridSize() : Component->GetGenerationGridSize();
+	
+	return InSource && GetStackContext(InSource->GetExecutionState().GetGraph(),GridSize, bIsPartitioned, OutStackContext);
+}
+
+uint32 IPCGBaseSubsystem::GetGraphCacheEntryCount(IPCGElement* InElement) const
+{
+	return GraphExecutor ? GraphExecutor->GetGraphCacheEntryCount(InElement) : 0;
+}
+
+
+void IPCGBaseSubsystem::SetDisableClearResults(bool bInDisableClearResults)
+{
+	if (GraphExecutor)
+	{
+		GraphExecutor->SetDisableClearResults(bInDisableClearResults);
+	}
+}
+
+#endif // WITH_EDITOR
+
+IPCGGraphCache* IPCGBaseSubsystem::GetCache()
+{
+	return GraphExecutor ? &(GraphExecutor->GetCache()) : nullptr;
+}
+
+void IPCGBaseSubsystem::FlushCache()
+{
+	if (GraphExecutor && GraphExecutor->GetCompiler())
+	{
+		GraphExecutor->GetCache().ClearCache();
+		GraphExecutor->GetCompiler()->ClearCache();
+	}
+
+#if WITH_EDITOR
+	// Garbage collection is very seldom run in the editor, but we currently can consume a lot of memory in the cache.
+	const UWorld* World = GetSubsystemWorld();
+	if (!World || !World->IsGameWorld())
+	{
+		CollectGarbage(GARBAGE_COLLECTION_KEEPFLAGS, /*bPerformFullPurge=*/true);
+	}
+#endif
+}
diff --git a/Engine/Plugins/PCG/Source/PCG/Private/Subsystems/PCGSubsystem.cpp b/Engine/Plugins/PCG/Source/PCG/Private/Subsystems/PCGSubsystem.cpp
new file mode 100644
--- /dev/null
+++ b/Engine/Plugins/PCG/Source/PCG/Private/Subsystems/PCGSubsystem.cpp
@@ -0,0 +1,2029 @@
+// Copyright Epic Games, Inc. All Rights Reserved.
+
+#include "PCGSubsystem.h"
+
+#include "PCGActorAndComponentMapping.h"
+#include "PCGComponent.h"
+#include "PCGGraph.h"
+#include "PCGModule.h"
+#include "PCGWorldActor.h"
+#include "Graph/PCGGraphCache.h"
+#include "Graph/PCGGraphCompiler.h"
+#include "Graph/PCGGraphExecutor.h"
+#include "Graph/PCGStackContext.h"
+#include "Grid/PCGLandscapeCache.h"
+#include "Grid/PCGPartitionActor.h"
+#include "Helpers/PCGActorHelpers.h"
+#include "Helpers/PCGHelpers.h"
+#include "RuntimeGen/PCGGenSourceManager.h"
+#include "RuntimeGen/PCGRuntimeGenScheduler.h"
+
+#include "Algo/Transform.h"
+#include "Engine/Level.h"
+
+#include UE_INLINE_GENERATED_CPP_BY_NAME(PCGSubsystem)
+
+#if WITH_EDITOR
+#include "Editor/IPCGEditorModule.h"
+
+#include "Editor.h"
+#include "PackageSourceControlHelper.h"
+#include "ObjectTools.h"
+#include "Misc/ScopedSlowTask.h"
+#include "Subsystems/ActorEditorContextSubsystem.h"
+#include "WorldPartition/WorldPartition.h"
+#include "WorldPartition/WorldPartitionHelpers.h"
+#include "WorldPartition/DataLayer/DataLayerManager.h"
+#include "WorldPartition/LoaderAdapter/LoaderAdapterShape.h"
+#else
+#include "Engine/Engine.h"
+#include "Engine/World.h"
+#endif
+
+namespace PCGSubsystemConsole
+{
+	static FAutoConsoleCommand CommandFlushCache(
+		TEXT("pcg.FlushCache"),
+		TEXT("Clears the PCG results cache and compiled graph cache."),
+		FConsoleCommandDelegate::CreateLambda([]()
+			{
+				if (UPCGSubsystem* PCGSubsystem = UPCGSubsystem::GetSubsystemForCurrentWorld())
+				{
+					PCGSubsystem->FlushCache();
+				}
+			}));
+
+	TAutoConsoleVariable<int32> CVarPCGQuality(
+		TEXT("pcg.Quality"), 2,
+		TEXT("Selects the quality permutation of PCG which impacts Runtime Quality Branch/Select nodes.\n")
+		TEXT(" 0: Low\n")
+		TEXT(" 1: Medium\n")
+		TEXT(" 2: High\n")
+		TEXT(" 3: Epic\n")
+		TEXT(" 4: Cinematic\n"),
+		FConsoleVariableDelegate::CreateLambda([](IConsoleVariable* CVar)
+		{
+			if (UPCGSubsystem* Subsystem = UPCGSubsystem::GetSubsystemForCurrentWorld())
+			{
+				Subsystem->OnPCGQualityLevelChanged();
+			}
+		}),
+		ECVF_Scalability);
+
+	static FAutoConsoleCommand CommandRefreshRuntimeGen(
+		TEXT("pcg.RuntimeGeneration.Refresh"),
+		TEXT("Cleans up and re-generates all GenerateAtRuntime PCG components, including their partition actors."),
+		FConsoleCommandDelegate::CreateLambda([]()
+		{
+			if (UPCGSubsystem* PCGSubsystem = UPCGSubsystem::GetSubsystemForCurrentWorld())
+			{
+				PCGSubsystem->RefreshAllRuntimeGenComponents(EPCGChangeType::GenerationGrid);
+			}
+		}));
+
+#if WITH_EDITOR
+	static FAutoConsoleCommand CommandBuildLandscapeCache(
+		TEXT("pcg.BuildLandscapeCache"),
+		TEXT("Builds the landscape cache in the current world."),
+		FConsoleCommandDelegate::CreateLambda([]()
+			{
+				if (UPCGSubsystem* PCGSubsystem = UPCGSubsystem::GetSubsystemForCurrentWorld())
+				{
+					PCGSubsystem->BuildLandscapeCache();
+				}
+			}));
+
+	static FAutoConsoleCommand CommandClearLandscapeCache(
+		TEXT("pcg.ClearLandscapeCache"),
+		TEXT("Clear the landscape cache in the current world."),
+		FConsoleCommandDelegate::CreateLambda([]()
+			{
+				if (UPCGSubsystem* PCGSubsystem = UPCGSubsystem::GetSubsystemForCurrentWorld())
+				{
+					PCGSubsystem->ClearLandscapeCache();
+				}
+			}));
+
+	static TAutoConsoleVariable<bool> CVarRebuildLandscapeOnPIE(
+		TEXT("pcg.PIE.RebuildLandscapeOnPIE"),
+		true,
+		TEXT("Controls whether the landscape cache will be rebuilt on PIE"));
+
+	static FAutoConsoleCommand CommandDeleteCurrentPCGWorldActor(
+		TEXT("pcg.DeleteCurrentPCGWorldActor"),
+		TEXT("Deletes the PCG World Actor currently registered to the PCG Subsystem."),
+		FConsoleCommandDelegate::CreateLambda([]()
+		{
+			if(UPCGSubsystem* PCGSubsystem = UPCGSubsystem::GetSubsystemForCurrentWorld())
+			{
+				PCGSubsystem->DestroyCurrentPCGWorldActor();
+			}
+		}));
+
+	static FAutoConsoleCommand CommandDeleteAllPCGWorldActors(
+		TEXT("pcg.DeleteAllPCGWorldActors"),
+		TEXT("Deletes all PCG World Actors in current World.."),
+		FConsoleCommandDelegate::CreateLambda([]()
+		{
+			if (UPCGSubsystem* PCGSubsystem = UPCGSubsystem::GetSubsystemForCurrentWorld())
+			{
+				PCGSubsystem->DestroyAllPCGWorldActors();
+			}
+		}));
+#endif
+}
+
+#if WITH_EDITOR
+TSet<UWorld*> UPCGSubsystem::DisablePartitionActorCreationForWorld;
+
+struct FPCGPartitionActorLoaderAdapter : public FLoaderAdapterShape
+{
+public:
+	FPCGPartitionActorLoaderAdapter(UWorld* InWorld, const FBox& InBoundingBox, const FString& InLabel)
+		: FLoaderAdapterShape(InWorld, InBoundingBox, InLabel)
+	{ }
+
+	virtual bool PassActorDescFilter(const FWorldPartitionHandle& Actor) const override
+	{
+		return FLoaderAdapterShape::PassActorDescFilter(Actor) && Actor->GetActorNativeClass() && Actor->GetActorNativeClass()->IsChildOf<APCGPartitionActor>();
+	}
+};
+#endif
+
+UPCGSubsystem::UPCGSubsystem()
+	: UTickableWorldSubsystem()
+	, IPCGBaseSubsystem()
+	, ActorAndComponentMapping(new FPCGActorAndComponentMapping(this)) // Can't use MakeUnique because the ctor of FPCGActorAndComponentMapping is private
+{
+}
+
+void UPCGSubsystem::AddReferencedObjects(UObject* InThis, FReferenceCollector& Collector)
+{
+	Super::AddReferencedObjects(InThis, Collector);
+
+	UPCGSubsystem* PCGSubsystem = CastChecked<UPCGSubsystem>(InThis);
+	
+	if (FPCGRuntimeGenScheduler* RuntimeGenScheduler = PCGSubsystem->GetRuntimeGenScheduler())
+	{
+		RuntimeGenScheduler->AddReferencedObjects(Collector);
+	}
+}
+
+UPCGSubsystem* UPCGSubsystem::GetSubsystemForCurrentWorld()
+{
+	UWorld* World = nullptr;
+
+#if WITH_EDITOR
+	if (GEditor)
+	{
+		if (GEditor->PlayWorld)
+		{
+			World = GEditor->PlayWorld;
+		}
+		else
+		{
+			World = GEditor->GetEditorWorldContext().World();
+		}
+	}
+	else
+#endif
+	if (GEngine)
+	{
+		World = GEngine->GetCurrentPlayWorld();
+	}
+
+	return UPCGSubsystem::GetInstance(World);
+}
+
+void UPCGSubsystem::Deinitialize()
+{
+	// Cancel all tasks
+	// TODO
+	IPCGBaseSubsystem::DeinitializeBaseSubsystem();
+
+	delete RuntimeGenScheduler;
+	RuntimeGenScheduler = nullptr;
+
+	PCGWorldActor = nullptr;
+	bHasTickedOnce = false;
+
+	ActorAndComponentMapping->Deinitialize();
+
+	Super::Deinitialize();
+}
+
+void UPCGSubsystem::PostInitialize()
+{
+	Super::PostInitialize();
+
+	// Gather world pcg actor if it exists
+	if (!PCGWorldActor)
+	{
+		if (UWorld* World = GetWorld())
+		{
+			UPCGActorHelpers::ForEachActorInLevel<APCGWorldActor>(World->PersistentLevel, [this](AActor* InActor)
+			{
+				PCGWorldActor = Cast<APCGWorldActor>(InActor);
+				return PCGWorldActor == nullptr;
+			});
+		}
+	}
+
+	ActorAndComponentMapping->Initialize();
+
+	IPCGBaseSubsystem::InitializeBaseSubsystem();
+
+	// Initialize runtime generation scheduler
+	check(!RuntimeGenScheduler);
+	RuntimeGenScheduler = new FPCGRuntimeGenScheduler(GetWorld(), ActorAndComponentMapping.Get());
+}
+
+UPCGSubsystem* UPCGSubsystem::GetInstance(UWorld* World)
+{
+	if (World)
+	{
+		UPCGSubsystem* Subsystem = World->GetSubsystem<UPCGSubsystem>();
+		return (Subsystem && Subsystem->IsInitialized()) ? Subsystem : nullptr;
+	}
+	else
+	{
+		return nullptr;
+	}
+}
+
+void UPCGSubsystem::RegisterBeginTickAction(FTickAction&& Action)
+{
+	BeginTickActions.Emplace(Action);
+}
+
+#if WITH_EDITOR
+UPCGSubsystem* UPCGSubsystem::GetActiveEditorInstance()
+{
+	if (GEditor)
+	{
+		return GEditor->PlayWorld ? UPCGSubsystem::GetInstance(GEditor->PlayWorld.Get()) : UPCGSubsystem::GetInstance(GEditor->GetEditorWorldContext().World());
+	}
+
+	return nullptr;
+}
+
+void UPCGSubsystem::SetConstructionScriptSourceComponent(UPCGComponent* InComponent)
+{
+	if (InComponent)
+	{
+		if (AActor* Owner = InComponent->GetOwner())
+		{
+			PerActorConstructionScriptSourceComponents.FindOrAdd(Owner).Add(InComponent->GetFName(), InComponent);
+		}
+	}
+}
+
+bool UPCGSubsystem::RemoveAndCopyConstructionScriptSourceComponent(AActor* InComponentOwner, FName InComponentName, UPCGComponent*& OutSourceComponent)
+{
+	OutSourceComponent = nullptr;
+	if (FConstructionScriptSourceComponents* Found = PerActorConstructionScriptSourceComponents.Find(InComponentOwner))
+	{
+		TObjectKey<UPCGComponent> FoundComponent;
+		if (Found->RemoveAndCopyValue(InComponentName, FoundComponent))
+		{
+			OutSourceComponent = FoundComponent.ResolveObjectPtrEvenIfGarbage();
+			if (Found->IsEmpty())
+			{
+				PerActorConstructionScriptSourceComponents.Remove(InComponentOwner);
+			}
+		}
+	}
+
+	return OutSourceComponent != nullptr;
+}
+#endif
+
+void UPCGSubsystem::Tick(float DeltaSeconds)
+{
+	TRACE_CPUPROFILER_EVENT_SCOPE(UPCGSubsystem::Tick);
+
+	Super::Tick(DeltaSeconds);
+
+	ExecuteBeginTickActions();
+
+#if WITH_EDITOR
+	PerActorConstructionScriptSourceComponents.Empty();
+#endif
+
+	if (!bHasTickedOnce)
+	{
+#if WITH_EDITOR
+		if (PCGSubsystemConsole::CVarRebuildLandscapeOnPIE.GetValueOnAnyThread() && PCGHelpers::IsRuntimeOrPIE())
+		{
+			BuildLandscapeCache(/*bQuiet=*/true, /*bForceBuild=*/false);
+		}
+#endif
+
+		bHasTickedOnce = true;
+	}
+
+	// Lose references to landscape cache as needed
+	// This will also initialize the cache if it isn't already so needs to happen before the GraphExecutor->Executor call
+	if (PCGWorldActor && GetLandscapeCache())
+	{
+		GetLandscapeCache()->Tick(DeltaSeconds);
+	}
+
+	// Base subsystem can modify end time
+	const double EndTime = IPCGBaseSubsystem::Tick();
+	
+	ActorAndComponentMapping->Tick();
+
+	if (PCGWorldActor)
+	{
+		RuntimeGenScheduler->Tick(PCGWorldActor, EndTime);
+	}
+}
+
+APCGWorldActor* UPCGSubsystem::GetPCGWorldActor()
+{
+#if WITH_EDITOR
+	if (!PCGWorldActor && !PCGHelpers::IsRuntimeOrPIE())
+	{
+		FScopeLock Lock(&PCGWorldActorLock);
+
+		if (!PCGWorldActor)
+		{
+			PCGWorldActor = APCGWorldActor::CreatePCGWorldActor(GetWorld());
+		}
+	}
+#endif
+
+	return PCGWorldActor;
+}
+
+APCGWorldActor* UPCGSubsystem::GetPCGWorldActorForPartitionActor(APCGPartitionActor* InActor)
+{
+	APCGWorldActor* FoundWorldActor = GetPCGWorldActor();
+	if (FoundWorldActor)
+	{
+		return FoundWorldActor;
+	}
+	
+	// We're at runtime and we didn't find a World Actor. This can happen if the partition actor is not in the same level than the persistent
+	// level of the world. In that case, look for the PCG World Actor in the same level as the Partition Actor. If we find one, register it.
+	if (ensure(PCGHelpers::IsRuntimeOrPIE()) && ensure(InActor))
+	{
+		UPCGActorHelpers::ForEachActorInLevel<APCGWorldActor>(InActor->GetLevel(), [this](AActor* InActor)
+		{
+			FScopeLock Lock(&PCGWorldActorLock);
+			if (!PCGWorldActor)
+			{
+				PCGWorldActor = Cast<APCGWorldActor>(InActor);
+			}
+
+			return PCGWorldActor == nullptr;
+		});
+	}
+
+	return PCGWorldActor;
+}
+
+APCGWorldActor* UPCGSubsystem::FindPCGWorldActor()
+{
+	return PCGWorldActor;
+}
+
+int32 UPCGSubsystem::GetPCGQualityLevel()
+{
+	return PCGSubsystemConsole::CVarPCGQuality.GetValueOnAnyThread();
+}
+
+void UPCGSubsystem::OnPCGQualityLevelChanged()
+{
+	// Trigger deep refresh of all runtime generated components.
+	ActorAndComponentMapping->ForAllOriginalComponents([this](UPCGComponent* InComponent)
+	{
+		if (InComponent && InComponent->IsManagedByRuntimeGenSystem())
+		{
+#if WITH_EDITOR
+			if (UPCGGraph* Graph = InComponent->GetGraph())
+			{
+				Graph->OnPCGQualityLevelChanged();
+			}
+#else
+			RefreshRuntimeGenComponent(InComponent, EPCGChangeType::GenerationGrid);
+#endif
+
+		}
+	});
+}
+
+#if WITH_EDITOR
+void UPCGSubsystem::DestroyCurrentPCGWorldActor()
+{
+	if (PCGWorldActor)
+	{
+		PCGWorldActorLock.Lock();
+		PCGWorldActor->Destroy();
+		PCGWorldActor = nullptr;
+		PCGWorldActorLock.Unlock();
+	}
+}
+
+void UPCGSubsystem::DestroyAllPCGWorldActors()
+{
+	// Delete all PAs first to avoid leaving orphans behind.
+	DeleteSerializedPartitionActors(/*bOnlyDeleteUnused=*/false);
+
+	// Get rid of current PCG world actor first
+	DestroyCurrentPCGWorldActor();
+	
+	// Pick up any strays in the current world
+	TArray<APCGWorldActor*> ActorsToDestroy;
+	ForEachObjectWithOuter(GetWorld(), [&ActorsToDestroy](UObject* Object)
+	{
+		if (APCGWorldActor* WorldActor = Cast<APCGWorldActor>(Object))
+		{
+			if (IsValid(WorldActor))
+			{
+				ActorsToDestroy.Add(WorldActor);
+			}
+		}
+	});
+
+	for (APCGWorldActor* ActorToDestroy : ActorsToDestroy)
+	{
+		ActorToDestroy->Destroy();
+	}
+}
+
+void UPCGSubsystem::LogAbnormalComponentStates(bool bGroupByState) const
+{
+	TArray<UPCGComponent*> DeactivatedComponents;
+	TArray<UPCGComponent*> NotGeneratedComponents;
+	TArray<UPCGComponent*> DirtyGeneratedComponents;
+
+	UE_LOG(LogPCG, Log, TEXT("--- Logging Abnormal PCG Component States ---"));
+
+	UPCGActorHelpers::ForEachActorInWorld(GetWorld(), AActor::StaticClass(), [bGroupByState, &DeactivatedComponents, &NotGeneratedComponents, &DirtyGeneratedComponents](AActor* InActor)
+	{
+		if (!InActor || !IsValid(InActor))
+		{
+			return true;
+		}
+
+		TInlineComponentArray<UPCGComponent*, 1> PCGComponents;
+		InActor->GetComponents(PCGComponents);
+
+		for (UPCGComponent* PCGComponent : PCGComponents)
+		{
+			if (!PCGComponent->bActivated)
+			{
+				if (bGroupByState)
+				{
+					DeactivatedComponents.Add(PCGComponent);
+				}
+				else
+				{
+					UE_LOG(LogPCG, Log, TEXT("%s - %s - Deactivated Component"), *InActor->GetName(), *PCGComponent->GetName());
+				}
+			}
+			else if (!PCGComponent->bGenerated && PCGComponent->GenerationTrigger != EPCGComponentGenerationTrigger::GenerateAtRuntime)
+			{
+				if (bGroupByState)
+				{
+					NotGeneratedComponents.Add(PCGComponent);
+				}
+				else
+				{
+					UE_LOG(LogPCG, Log, TEXT("%s - %s - Not Generated Component"), *InActor->GetName(), *PCGComponent->GetName());
+				}
+			}
+			else if (PCGComponent->bDirtyGenerated)
+			{
+				if (bGroupByState)
+				{
+					DirtyGeneratedComponents.Add(PCGComponent);
+				}
+				else
+				{
+					UE_LOG(LogPCG, Log, TEXT("%s - %s - Dirty Generated Component "), *InActor->GetName(), *PCGComponent->GetName());
+				}
+			}
+		}
+
+		return true;
+	});
+
+	if (bGroupByState)
+	{
+		UE_LOG(LogPCG, Log, TEXT("--- Deactivated PCG Components ---"));
+		for (UPCGComponent* Component : DeactivatedComponents)
+		{
+			check(Component && Component->GetOwner());
+			UE_LOG(LogPCG, Log, TEXT("%s - %s"), *Component->GetOwner()->GetName(), *Component->GetName());
+		}
+
+		UE_LOG(LogPCG, Log, TEXT("--- Not Generated Components ---"));
+		for (UPCGComponent* Component : NotGeneratedComponents)
+		{
+			check(Component && Component->GetOwner());
+			UE_LOG(LogPCG, Log, TEXT("%s - %s"), *Component->GetOwner()->GetName(), *Component->GetName());
+		}
+
+		UE_LOG(LogPCG, Log, TEXT("--- Dirty Components ---"));
+		for (UPCGComponent* Component : DirtyGeneratedComponents)
+		{
+			check(Component && Component->GetOwner());
+			UE_LOG(LogPCG, Log, TEXT("%s - %s"), *Component->GetOwner()->GetName(), *Component->GetName());
+		}
+	}
+}
+
+#endif // WITH_EDITOR
+
+void UPCGSubsystem::RegisterPCGWorldActor(APCGWorldActor* InActor)
+{
+	// TODO: we should support merging or multi world actor support when relevant
+	if (!PCGWorldActor)
+	{
+		PCGWorldActor = InActor;
+	}
+	else if (InActor != PCGWorldActor)
+	{
+		PCGWorldActor->MergeFrom(InActor);
+	}
+}
+
+void UPCGSubsystem::UnregisterPCGWorldActor(APCGWorldActor* InActor)
+{
+	if (PCGWorldActor == InActor)
+	{
+		PCGWorldActor = nullptr;
+	}
+}
+
+void UPCGSubsystem::OnOriginalComponentRegistered(UPCGComponent* InComponent)
+{
+	if (RuntimeGenScheduler)
+	{
+		RuntimeGenScheduler->OnOriginalComponentRegistered(InComponent);
+	}
+}
+
+void UPCGSubsystem::OnOriginalComponentUnregistered(UPCGComponent* InComponent)
+{
+	if (RuntimeGenScheduler)
+	{
+		RuntimeGenScheduler->OnOriginalComponentUnregistered(InComponent);
+	}
+
+#if WITH_EDITOR
+	if (IPCGEditorModule* PCGEditorModule = IPCGEditorModule::Get())
+	{
+		PCGEditorModule->ClearExecutionMetadata(InComponent);
+	}
+#endif
+}
+
+UPCGLandscapeCache* UPCGSubsystem::GetLandscapeCache()
+{
+	APCGWorldActor* LandscapeCacheOwner = GetPCGWorldActor();
+	return LandscapeCacheOwner ? LandscapeCacheOwner->LandscapeCacheObject.Get() : nullptr;
+}
+
+FPCGTaskId UPCGSubsystem::ScheduleComponent(UPCGComponent* PCGComponent, EPCGHiGenGrid Grid, bool bForce, const TArray<FPCGTaskId>& InDependencies)
+{
+	check(GraphExecutor);
+
+	if (!PCGComponent)
+	{
+		return InvalidPCGTaskId;
+	}
+
+	bool bHasUnbounded = false;
+	PCGHiGenGrid::FSizeArray GridSizes;
+	ensure(PCGHelpers::GetGenerationGridSizes(PCGComponent->GetGraph(), GetPCGWorldActor(), GridSizes, bHasUnbounded));
+
+	// Create the PartitionActors if necessary. Skip if this is a runtime managed component, PAs are handled manually by the RuntimeGenScheduler.
+	// Editor only because we expect at runtime for PAs to already exist so they can properly be streamed in and out (creating them at runtime would leave them unmanaged and always loaded)
+	if (PCGComponent->IsPartitioned() && !PCGComponent->IsManagedByRuntimeGenSystem())
+	{
+#if WITH_EDITOR
+		if (!GridSizes.IsEmpty())
+		{
+			CreatePartitionActorsWithinBounds(PCGComponent, PCGComponent->GetGridBounds(), GridSizes);
+		}
+#endif // WITH_EDITOR
+
+		ActorAndComponentMapping->UpdateMappingPCGComponentPartitionActor(PCGComponent);
+	}
+
+	// Execution dependencies require a task to finish executing before the dependent task.
+	TArray<FPCGTaskId> ExecutionDependencyTasks;
+
+	// Data dependencies act as execution dependencies, but will also have their output consumed by the waiting task. For a component, this means
+	// it will store the output data into its managed resources, which, for an original component, should not include the local component generation tasks,
+	// since those resources should be managed locally.
+	TArray<FPCGTaskId> DataDependencyTasks;
+
+	// Schedule generation of original component if is is non-partitioned, or if it has nodes that will execute at the Unbounded level.
+	FPCGTaskId OriginalComponentTask = InvalidPCGTaskId;
+
+	bool bGeneratePCGComponent = false;
+	if (!PCGComponent->IsPartitioned())
+	{
+		// This component is either an unpartitioned original component or a local component. Generate if grid size matches preference (if provided).
+		bGeneratePCGComponent = (Grid == EPCGHiGenGrid::Uninitialized) || !!(Grid & PCGComponent->GetGenerationGrid());
+	}
+	else
+	{
+		// This component is a partitioned original component. Generate if the graph has unbounded nodes and if this grid matches preference (if provided).
+		bGeneratePCGComponent = bHasUnbounded && (!!(Grid & EPCGHiGenGrid::Unbounded) || Grid == EPCGHiGenGrid::Uninitialized);
+	}
+
+	if (bGeneratePCGComponent)
+	{
+		OriginalComponentTask = PCGComponent->CreateGenerateTask(bForce, InDependencies);
+		if (OriginalComponentTask != InvalidPCGTaskId)
+		{
+			DataDependencyTasks.Add(OriginalComponentTask);
+		}
+	}
+
+	// If the component is partitioned, we will forward the calls to its registered PCG Partition actors
+	if (PCGComponent->IsPartitioned() && PCGHiGenGrid::IsValidGridOrUninitialized(Grid))
+	{
+		// Local components depend on the original component (to ensure any data is available).
+		TArray<FPCGTaskId> Dependencies = InDependencies;
+		if (OriginalComponentTask != InvalidPCGTaskId)
+		{
+			Dependencies.Add(OriginalComponentTask);
+		}
+
+		auto LocalGenerateTask = [OriginalComponent = PCGComponent, Grid, &Dependencies, bForce, &GridSizes](UPCGComponent* LocalComponent)
+		{
+			if (!GridSizes.Contains(LocalComponent->GetGenerationGridSize()))
+			{
+				// Local component with invalid grid size. Grid sizes may have changed in graph.
+				return LocalComponent->CleanupLocal(/*bRemoveComponents=*/true, Dependencies);
+			}
+			else if (Grid != EPCGHiGenGrid::Uninitialized && !(Grid & LocalComponent->GetGenerationGrid()))
+			{
+				// Grid size does not match the given target grid, so skip.
+				return InvalidPCGTaskId;
+			}
+
+			// If the local component is currently generating, it's probably because it was requested by a refresh.
+			// Wait after this one instead
+			if (LocalComponent->IsGenerating())
+			{
+				return LocalComponent->CurrentGenerationTask;
+			}
+
+			// Ensure that the PCG actor match our original
+			LocalComponent->SetPropertiesFromOriginal(OriginalComponent);
+
+			FPCGTaskId LocalCleanupTaskId = InvalidPCGTaskId;
+			if (LocalComponent->bGenerated && !OriginalComponent->bGenerated)
+			{
+				// Detected a mismatch between the original component and the local component.
+				// Request a cleanup first
+				LocalCleanupTaskId = LocalComponent->CleanupLocal(/*bRemoveComponents=*/true, Dependencies);
+			}
+
+			TArray<FPCGTaskId> AdditionalDependencies;
+			const TArray<FPCGTaskId>* AllDependencies = &Dependencies;
+			if (LocalCleanupTaskId != InvalidPCGTaskId)
+			{
+				AdditionalDependencies.Reserve(Dependencies.Num() + 1);
+				AdditionalDependencies.Append(Dependencies);
+				AdditionalDependencies.Add(LocalCleanupTaskId);
+				AllDependencies = &AdditionalDependencies;
+			}
+
+			return LocalComponent->GenerateInternal(bForce, LocalComponent->GetGenerationGrid(), EPCGComponentGenerationTrigger::GenerateOnDemand, *AllDependencies);
+		};
+
+		ExecutionDependencyTasks.Append(ActorAndComponentMapping->DispatchToRegisteredLocalComponents(PCGComponent, LocalGenerateTask));
+	}
+
+	if (!ExecutionDependencyTasks.IsEmpty() || !DataDependencyTasks.IsEmpty())
+	{
+		TWeakObjectPtr<UPCGComponent> ComponentPtr(PCGComponent);
+
+		return GraphExecutor->ScheduleGenericWithContext([ComponentPtr](FPCGContext* Context)
+		{
+			if (UPCGComponent* Component = ComponentPtr.Get())
+			{
+				// If the component is not valid anymore, just early out.
+				if (!IsValid(Component))
+				{
+					return true;
+				}
+
+				const FBox NewBounds = Component->GetGridBounds();
+				Component->PostProcessGraph(NewBounds, /*bGenerate=*/true, Context);
+			}
+
+			return true;
+		}, PCGComponent, ExecutionDependencyTasks, DataDependencyTasks, /*bSupportBasePointDataInput=*/true);
+	}
+	else
+	{
+		UE_LOG(LogPCG, Error, TEXT("[ScheduleComponent] Didn't schedule any task."));
+		if (PCGComponent)
+		{
+			PCGComponent->OnProcessGraphAborted();
+		}
+		return InvalidPCGTaskId;
+	}
+}
+
+FPCGTaskId UPCGSubsystem::ScheduleCleanup(UPCGComponent* PCGComponent, bool bRemoveComponents, const TArray<FPCGTaskId>& Dependencies)
+{
+	if (!PCGComponent)
+	{
+		return InvalidPCGTaskId;
+	}
+
+	bool bHasUnbounded = false;
+	PCGHiGenGrid::FSizeArray GridSizes;
+	ensure(PCGHelpers::GetGenerationGridSizes(PCGComponent->GetGraph(), GetPCGWorldActor(), GridSizes, bHasUnbounded));
+
+	TArray<FPCGTaskId> AllTasks;
+
+	// Schedule cleanup of original component if is is non-partitioned, or if it has nodes that will execute at the Unbounded level.
+	if (!PCGComponent->IsPartitioned() || bHasUnbounded)
+	{
+		FPCGTaskId TaskId = PCGComponent->CreateCleanupTask(bRemoveComponents, Dependencies);
+		if (TaskId != InvalidPCGTaskId)
+		{
+			AllTasks.Add(TaskId);
+		}
+	}
+
+	// If the component is partitioned, we will forward the calls to its registered PCG Partition actors
+	if (PCGComponent->IsPartitioned())
+	{
+		auto LocalCleanupTask = [bRemoveComponents, &Dependencies](UPCGComponent* LocalComponent)
+		{
+			// If the local component is currently cleaning up, it's probably because it was requested by a refresh.
+			// Wait after this one instead
+			if (LocalComponent->IsCleaningUp())
+			{
+				return LocalComponent->CurrentCleanupTask;
+			}
+
+			// Always executes regardless of local component grid size - clean up as much as possible.
+			return LocalComponent->CleanupLocal(bRemoveComponents, Dependencies);
+		};
+
+		AllTasks.Append(ActorAndComponentMapping->DispatchToRegisteredLocalComponents(PCGComponent, LocalCleanupTask));
+	}
+
+	TWeakObjectPtr<UPCGComponent> ComponentPtr(PCGComponent);
+	auto PostCleanupTask = [this, bRemoveComponents, ComponentPtr]()
+	{
+		if (UPCGComponent* Component = ComponentPtr.Get())
+		{
+			// If the component is not valid anymore, just early out
+			if (!IsValid(Component))
+			{
+				return true;
+			}
+
+			Component->PostCleanupGraph(bRemoveComponents);
+
+			// Remove the local component mappings if requested and the component is partitioned. If 'bRemoveComponents' is false, that indicates we are doing a refresh, so destroying
+			// the component mappings is counterproductive.
+			if (bRemoveComponents && Component->IsPartitioned())
+			{
+				ActorAndComponentMapping->DeleteMappingPCGComponentPartitionActor(Component);
+			}
+		}
+
+		return true;
+	};
+
+	FPCGTaskId PostCleanupTaskId = InvalidPCGTaskId;
+
+	// If we have no tasks to do, just call PostCleanup immediately
+	// otherwise wait for all the tasks to be done to call PostCleanup.
+	if (AllTasks.IsEmpty())
+	{
+		PostCleanupTask();
+	}
+	else
+	{
+		PostCleanupTaskId = GraphExecutor->ScheduleGeneric(PostCleanupTask, PCGComponent, AllTasks);
+	}
+
+	return PostCleanupTaskId;
+}
+
+FPCGTaskId UPCGSubsystem::ScheduleGraph(UPCGComponent* SourceComponent, const TArray<FPCGTaskId>& Dependencies)
+{
+	if (SourceComponent)
+	{
+		return GraphExecutor->Schedule(SourceComponent, Dependencies);
+	}
+	else
+	{
+		return InvalidPCGTaskId;
+	}
+}
+
+FPCGTaskId UPCGSubsystem::ScheduleGraph(
+	UPCGGraph* Graph,
+	UPCGComponent* SourceComponent,
+	FPCGElementPtr PreGraphElement, 
+	FPCGElementPtr InputElement, 
+	const TArray<FPCGTaskId>& Dependencies, 
+	const FPCGStack* InFromStack, 
+	bool bAllowHierarchicalGeneration)
+{
+	return ScheduleGraph(FPCGScheduleGraphParams(Graph, SourceComponent, PreGraphElement, InputElement, Dependencies, InFromStack, bAllowHierarchicalGeneration));
+}
+
+ETickableTickType UPCGSubsystem::GetTickableTickType() const
+{
+	return IsTemplate() ? ETickableTickType::Never : ETickableTickType::Always;
+}
+
+TStatId UPCGSubsystem::GetStatId() const
+{
+	RETURN_QUICK_DECLARE_CYCLE_STAT(UPCGSubsystem, STATGROUP_Tickables);
+}
+
+FPCGTaskId UPCGSubsystem::ScheduleGeneric(TFunction<bool()> InOperation, UPCGComponent* SourceComponent, const TArray<FPCGTaskId>& TaskExecutionDependencies)
+{
+	check(GraphExecutor);
+	return GraphExecutor->ScheduleGeneric(InOperation, SourceComponent, TaskExecutionDependencies);
+}
+
+FPCGTaskId UPCGSubsystem::ScheduleGeneric(TFunction<bool()> InOperation, TFunction<void()> InAbortOperation, UPCGComponent* SourceComponent, const TArray<FPCGTaskId>& TaskExecutionDependencies)
+{
+	check(GraphExecutor);
+	return GraphExecutor->ScheduleGeneric(InOperation, InAbortOperation, SourceComponent, TaskExecutionDependencies);
+}
+
+FPCGTaskId UPCGSubsystem::ScheduleGenericWithContext(TFunction<bool(FPCGContext*)> InOperation, UPCGComponent* SourceComponent, const TArray<FPCGTaskId>& TaskExecutionDependencies, const TArray<FPCGTaskId>& TaskDataDependencies)
+{
+	check(GraphExecutor);
+	return GraphExecutor->ScheduleGenericWithContext(InOperation, SourceComponent, TaskExecutionDependencies, TaskDataDependencies);
+}
+
+FPCGTaskId UPCGSubsystem::ScheduleGenericWithContext(TFunction<bool(FPCGContext*)> InOperation, TFunction<void(FPCGContext*)> InAbortOperation, UPCGComponent* SourceComponent, const TArray<FPCGTaskId>& TaskExecutionDependencies, const TArray<FPCGTaskId>& TaskDataDependencies)
+{
+	check(GraphExecutor);
+	return GraphExecutor->ScheduleGenericWithContext(InOperation, InAbortOperation, SourceComponent, TaskExecutionDependencies, TaskDataDependencies);
+}
+
+
+void UPCGSubsystem::CancelGenerationInternal(IPCGGraphExecutionSource* Source, bool bCleanupUnusedResources)
+{
+	check(GraphExecutor && IsInGameThread());
+	UPCGComponent* Component = Cast<UPCGComponent>(Source);
+
+	if (!Component || !Component->IsGenerating())
+	{
+		return;
+	}
+
+	if (Component->IsPartitioned())
+	{
+		auto LocalCancel = [this, bCleanupUnusedResources](UPCGComponent* LocalComponent)
+		{
+			if (LocalComponent->IsGenerating())
+			{
+				CancelGeneration(LocalComponent, bCleanupUnusedResources);
+			}
+
+			return InvalidPCGTaskId;
+		};
+
+		ActorAndComponentMapping->DispatchToRegisteredLocalComponents(Component, LocalCancel);
+	}
+}
+
+void UPCGSubsystem::RefreshRuntimeGenComponent(UPCGComponent* RuntimeComponent, EPCGChangeType ChangeType)
+{
+	if (!ensure(RuntimeComponent && RuntimeComponent->IsManagedByRuntimeGenSystem()))
+	{
+		return;
+	}
+
+	if (ensure(RuntimeGenScheduler))
+	{
+		// Only need to remove PAs if the grid sizes have changed.
+		const bool bRemovePartitionActors = !!(ChangeType & EPCGChangeType::GenerationGrid);
+		RuntimeGenScheduler->RefreshComponent(RuntimeComponent, bRemovePartitionActors);
+	}
+}
+
+void UPCGSubsystem::RefreshAllRuntimeGenComponents(EPCGChangeType ChangeType)
+{
+	for (UPCGComponent* Component : GetAllRegisteredComponents())
+	{
+		if (Component->GenerationTrigger == EPCGComponentGenerationTrigger::GenerateAtRuntime)
+		{
+			RefreshRuntimeGenComponent(Component, ChangeType);
+		}
+	}
+}
+
+#if WITH_EDITOR
+void UPCGSubsystem::RefreshAllComponentsFiltered(const TFunction<bool(UPCGComponent*)>& ComponentFilter, EPCGChangeType ChangeType)
+{
+	for (UPCGComponent* Component : GetAllRegisteredComponents())
+	{
+		if (ComponentFilter(Component))
+		{
+			Component->Refresh(ChangeType);
+		}
+	}
+}
+#endif
+
+void UPCGSubsystem::ForAllRegisteredLocalComponents(UPCGComponent* InOriginalComponent, const TFunctionRef<void(UPCGComponent*)>& InFunc) const
+{
+	auto WrapperFunc = [&InFunc](UPCGComponent* Component) -> FPCGTaskId
+	{
+		InFunc(Component);
+		return InvalidPCGTaskId;
+	};
+
+	ActorAndComponentMapping->DispatchToRegisteredLocalComponents(InOriginalComponent, WrapperFunc);
+}
+
+void UPCGSubsystem::ForAllRegisteredIntersectingLocalComponents(UPCGComponent* InOriginalComponent, const FBoxCenterAndExtent& InBounds, const TFunctionRef<void(UPCGComponent*)>& InFunc) const
+{
+	check(InOriginalComponent);
+	const FBox Overlap = InOriginalComponent->GetGridBounds().Overlap(InBounds.GetBox());
+
+	// We reject overlaps with zero volume instead of simply checking Intersect(...) to avoid bounds which touch but do not overlap.
+	if (Overlap.GetVolume() <= 0)
+	{
+		return;
+	}
+
+	ActorAndComponentMapping->ForAllIntersectingPartitionActors(Overlap, [InOriginalComponent, &InFunc](APCGPartitionActor* Actor)
+	{
+		if (UPCGComponent* LocalComponent = Actor->GetLocalComponent(InOriginalComponent))
+		{
+			InFunc(LocalComponent);
+		}
+	});
+}
+
+void UPCGSubsystem::ForAllOverlappingComponentsInHierarchy(UPCGComponent* InComponent, const TFunctionRef<void(UPCGComponent*)>& InFunc) const
+{
+	UPCGComponent* OriginalComponent = InComponent->GetOriginalComponent();
+
+	ForAllRegisteredLocalComponents(OriginalComponent, [InComponent, &InFunc](UPCGComponent* InLocalComponent)
+	{
+		const FBox OtherBounds = InLocalComponent->GetGridBounds();
+		const FBox ThisBounds = InComponent->GetGridBounds();
+		const FBox Overlap = OtherBounds.Overlap(ThisBounds);
+		if (Overlap.GetVolume() > 0)
+		{
+			InFunc(InLocalComponent);
+		}
+	});
+}
+
+FPCGTaskId UPCGSubsystem::ForAllOverlappingCells(UPCGComponent* InComponent, const FBox& InBounds, const PCGHiGenGrid::FSizeArray& InGridSizes, bool bCanCreateActor, const TArray<FPCGTaskId>& Dependencies, TFunctionRef<FPCGTaskId(APCGPartitionActor*, const FBox&)> InFunc, TFunctionRef<FPCGTaskId(const FPCGGridDescriptor&, const FIntVector&, const FBox&)> InUnloadedFunc) const
+{
+	if (!GraphExecutor || !PCGWorldActor)
+	{
+		UE_LOG(LogPCG, Error, TEXT("[ForAllOverlappingCells] GraphExecutor or PCGWorldActor is null."));
+		return InvalidPCGTaskId;
+	}
+
+	PCGHiGenGrid::FSizeArray GridSizes = InGridSizes;
+
+	// We have no use for unbounded grids as this is a grid-centric function. Also discard invalid grid sizes.
+	GridSizes.Remove(static_cast<uint32>(EPCGHiGenGrid::Unbounded));
+	GridSizes.RemoveAll([](uint32 GridSize) { return !ensure(PCGHiGenGrid::IsValidGridSize(GridSize)); });
+
+	if (GridSizes.IsEmpty())
+	{
+		return InvalidPCGTaskId;
+	}
+
+	TArray<FPCGTaskId> CellTasks;
+	for (uint32 GridSize : GridSizes)
+	{
+		FPCGGridDescriptor Descriptor = InComponent->GetGridDescriptor(GridSize);
+		check(!Descriptor.IsRuntime());
+
+		// In case of 2D grid, we are clamping our bounds in Z to be within 0 and GridSize to create a 2D grid instead of 3D.
+		FBox ModifiedInBounds = InBounds;
+		if (Descriptor.Is2DGrid())
+		{
+			FVector MinBounds = InBounds.Min;
+			FVector MaxBounds = InBounds.Max;
+
+			MinBounds.Z = 0;
+			MaxBounds.Z = GridSize;
+			ModifiedInBounds = FBox(MinBounds, MaxBounds);
+		}
+
+		// Helper lambda to apply 'InFunc' to the specified grid cell. Only applies if the cell overlaps the modified bounds, and the actor can be found or created and is not unloaded.
+		auto ApplyOnCell = [this, &CellTasks, ModifiedInBounds, bCanCreateActor, &Descriptor, &InFunc, InUnloadedFunc](const FIntVector& CellCoord, const FBox& CellBounds)
+		{
+			const FBox IntersectedBounds = ModifiedInBounds.Overlap(CellBounds);
+
+			if (IntersectedBounds.IsValid)
+			{
+				APCGPartitionActor* Actor = FindOrCreatePCGPartitionActor(Descriptor, CellCoord, bCanCreateActor);
+#if WITH_EDITOR
+				if (!Actor && !GetWorld()->IsGameWorld() && ActorAndComponentMapping->DoesPartitionActorRecordExist(Descriptor, CellCoord))
+				{
+					FPCGTaskId ExecuteTaskId = InUnloadedFunc(Descriptor, CellCoord, IntersectedBounds);
+					
+					if (ExecuteTaskId != InvalidPCGTaskId)
+					{
+						CellTasks.Add(ExecuteTaskId);
+					}
+				}
+				else 
+#endif
+				if(Actor)
+				{
+					FPCGTaskId ExecuteTaskId = InFunc(Actor, IntersectedBounds);
+
+					if (ExecuteTaskId != InvalidPCGTaskId)
+					{
+						CellTasks.Add(ExecuteTaskId);
+					}
+				}
+			}
+		};
+
+		FIntVector MinCellCoords = UPCGActorHelpers::GetCellCoord(InBounds.Min, GridSize, Descriptor.Is2DGrid());
+		FIntVector MaxCellCoords = UPCGActorHelpers::GetCellCoord(InBounds.Max, GridSize, Descriptor.Is2DGrid());
+
+		// Apply 'InFunc' to all cells in the provided bounds.
+		for (int32 Z = MinCellCoords.Z; Z <= MaxCellCoords.Z; ++Z)
+		{
+			for (int32 Y = MinCellCoords.Y; Y <= MaxCellCoords.Y; ++Y)
+			{
+				for (int32 X = MinCellCoords.X; X <= MaxCellCoords.X; ++X)
+				{
+					FIntVector CellCoords(X, Y, Z);
+
+					const FVector Min = FVector(CellCoords) * GridSize;
+					const FVector Max = Min + FVector(GridSize);
+					FBox CellBounds(Min, Max);
+
+					ApplyOnCell(MoveTemp(CellCoords), MoveTemp(CellBounds));
+				}
+			}
+		}
+	}
+
+	// Create a dummy task to wait on dependencies, which creates a dummy task to wait on all cells.
+	if (!CellTasks.IsEmpty())
+	{
+		TArray<FPCGTaskId> AllDependencies(Dependencies);
+		AllDependencies.Append(CellTasks);
+		return GraphExecutor->ScheduleGeneric([]() { return true; }, /*InSourceComponent=*/nullptr, AllDependencies);
+	}
+	else
+	{
+		return InvalidPCGTaskId;
+	}
+}
+
+void UPCGSubsystem::CleanupLocalComponentsImmediate(UPCGComponent* InOriginalComponent, bool bRemoveComponents)
+{
+	if (!InOriginalComponent)
+	{
+		return;
+	}
+
+	auto LocalCleanupTask = [bRemoveComponents](UPCGComponent* LocalComponent)
+	{
+		if (ensure(LocalComponent) && !LocalComponent->IsCleaningUp())
+		{
+			LocalComponent->CleanupLocalImmediate(bRemoveComponents);
+		}
+
+		return InvalidPCGTaskId;
+	};
+
+	ActorAndComponentMapping->DispatchToRegisteredLocalComponents(InOriginalComponent, LocalCleanupTask);
+
+	// Remove the local component mappings if requested and the component is partitioned. If 'bRemoveComponents' is false, that indicates we are doing a refresh, so destroying
+	// the component mappings is counterproductive.
+	if (bRemoveComponents && InOriginalComponent->IsPartitioned())
+	{
+		ActorAndComponentMapping->DeleteMappingPCGComponentPartitionActor(InOriginalComponent);
+	}
+}
+
+// deprecated
+UPCGComponent* UPCGSubsystem::GetLocalComponent(uint32 GridSize, const FIntVector& CellCoords, const UPCGComponent* InOriginalComponent, bool bTransient) const
+{
+	check(InOriginalComponent);
+
+	FPCGGridDescriptor GridDescriptor = FPCGGridDescriptor()
+		.SetGridSize(GridSize)
+		.SetIsRuntime(bTransient)
+		.SetIs2DGrid(InOriginalComponent->Use2DGrid());
+
+	return GetLocalComponent(GridDescriptor, CellCoords, InOriginalComponent);
+}
+
+// deprecated
+APCGPartitionActor* UPCGSubsystem::GetRegisteredPCGPartitionActor(uint32 GridSize, const FIntVector& GridCoords, bool bRuntimeGenerated) const
+{
+	FPCGGridDescriptor GridDescriptor = FPCGGridDescriptor()
+		.SetGridSize(GridSize)
+		.SetIsRuntime(bRuntimeGenerated);
+
+	return GetRegisteredPCGPartitionActor(GridDescriptor, GridCoords);
+}
+
+// deprecated
+APCGPartitionActor* UPCGSubsystem::FindOrCreatePCGPartitionActor(const FGuid& Guid, uint32 GridSize, const FIntVector& GridCoords, bool bRuntimeGenerated, bool bCanCreateActor) const
+{
+	FPCGGridDescriptor GridDescriptor = FPCGGridDescriptor()
+		.SetGridSize(GridSize)
+		.SetIsRuntime(bRuntimeGenerated);
+
+	return FindOrCreatePCGPartitionActor(GridDescriptor, GridCoords, bCanCreateActor);
+}
+
+UPCGComponent* UPCGSubsystem::GetLocalComponent(const FPCGGridDescriptor& GridDescriptor, const FIntVector& CellCoords, const UPCGComponent* InOriginalComponent) const
+{
+	check(InOriginalComponent);
+	return ActorAndComponentMapping->GetLocalComponent(GridDescriptor, CellCoords, InOriginalComponent);
+}
+
+APCGPartitionActor* UPCGSubsystem::GetRegisteredPCGPartitionActor(const FPCGGridDescriptor& GridDescriptor, const FIntVector& GridCoords) const
+{
+	return ActorAndComponentMapping->GetPartitionActor(GridDescriptor, GridCoords);
+}
+
+APCGPartitionActor* UPCGSubsystem::FindOrCreatePCGPartitionActor(const FPCGGridDescriptor& GridDescriptor, const FIntVector& GridCoords, bool bCanCreateActor, bool bHideFromOutliner) const
+{
+	TRACE_CPUPROFILER_EVENT_SCOPE(UPCGActorHelpers::FindOrCreatePCGPartitionActor);
+
+	UWorld* World = GetWorld();
+
+	if (!World)
+	{
+		UE_LOG(LogPCG, Error, TEXT("[FindOrCreatePCGPartitionActor] World is null."));
+		return nullptr;
+	}
+
+	if (!PCGWorldActor)
+	{
+		UE_LOG(LogPCG, Error, TEXT("[FindOrCreatePCGPartitionActor] PCGWorldActor is null."));
+		return nullptr;
+	}
+
+	// Attempt to find an existing PA.
+	if (APCGPartitionActor* ExistingActor = GetRegisteredPCGPartitionActor(GridDescriptor, GridCoords))
+	{
+		return ExistingActor;
+	}
+	else if (!GridDescriptor.IsRuntime())
+	{
+		// In a Game World PAs need to be Pre-Existing. 
+		// The Original PCG Component will be marked as Generated on load see UPCGComponent::BeginPlay
+		if (GetWorld()->IsGameWorld())
+		{
+			return nullptr;
+		}
+
+#if WITH_EDITOR
+		// Check if there is already an unloaded actor for this cell. RuntimeGenerated PAs are never unloaded, so we ignore them.
+		if (ActorAndComponentMapping->DoesPartitionActorRecordExist(GridDescriptor, GridCoords))
+		{
+			return nullptr;
+		}
+#endif
+	}
+
+	// If FindOrCreatePCGPartitionActor is called on a Level while it is not fully registered.
+	// We can't create the actor as it may already exist but not have been registered yet.
+	if (!World->PersistentLevel->bAreComponentsCurrentlyRegistered)
+	{
+		return nullptr;
+	}
+
+#if WITH_EDITOR
+	// Do not try and create while executing a Undo/Redo because actor might already be in the process of being re-created by the transaction
+	if (GIsTransacting)
+	{
+		return nullptr;
+	}
+#endif
+	
+	if (!bCanCreateActor)
+	{
+		return nullptr;
+	}
+
+	FActorSpawnParameters SpawnParams;
+#if WITH_EDITOR
+	SpawnParams.Name = *APCGPartitionActor::GetPCGPartitionActorName(GridDescriptor, GridCoords);
+	SpawnParams.NameMode = FActorSpawnParameters::ESpawnActorNameMode::Required_Fatal;
+	SpawnParams.bHideFromSceneOutliner = bHideFromOutliner;
+#endif
+	if (GridDescriptor.IsRuntime())
+	{
+		SpawnParams.NameMode = FActorSpawnParameters::ESpawnActorNameMode::Requested;
+		SpawnParams.ObjectFlags |= RF_Transient;
+		SpawnParams.ObjectFlags &= ~RF_Transactional;
+	}
+		
+#if WITH_EDITOR
+	bool bPushedContext = false;
+	const UExternalDataLayerAsset* ExternalDataLayerAsset = nullptr;
+	
+	// No need to do any DataLayer assignment in a game world
+	if (!World->IsGameWorld())
+	{
+		TArray<TSoftObjectPtr<UDataLayerAsset>> DataLayerAssets;
+		
+		GridDescriptor.GetDataLayerAssets(DataLayerAssets, ExternalDataLayerAsset);
+
+		// Avoid relying on the Editor Context at all
+		UActorEditorContextSubsystem::Get()->PushContext();
+		bPushedContext = true;
+	}
+
+	ON_SCOPE_EXIT
+	{ 
+		if (bPushedContext)
+		{
+			UActorEditorContextSubsystem::Get()->PopContext();
+		}
+	};
+
+	// Specify EDL we want to use if any for spawning this actor
+	FScopedOverrideSpawningLevelMountPointObject EDLScope(ExternalDataLayerAsset);
+
+	// Handle the case where the actor already exists, but is in the undo stack (was deleted) (copied from ActorPartitionSubsystem, we should probably merge back into it at some point) 
+	if (SpawnParams.NameMode == FActorSpawnParameters::ESpawnActorNameMode::Required_Fatal)
+	{
+		if (UObject* ExistingObject = StaticFindObject(nullptr, World->PersistentLevel, *SpawnParams.Name.ToString()))
+		{
+			AActor* ExistingActor = CastChecked<AActor>(ExistingObject);
+			// This actor is expected to be invalid
+			check(!IsValidChecked(ExistingActor));
+			ExistingActor->Modify();
+
+			// Don't go through AActor::Rename here because we aren't changing outers (the actor's level). We just want to rename that actor 
+			// out of the way so we can spawn the new one in the exact same package, keeping the package name intact.
+			ExistingActor->UObject::Rename(nullptr, nullptr, REN_DontCreateRedirectors | REN_DoNotDirty | REN_NonTransactional);
+
+			// Reuse ActorGuid so that ActorDesc can be updated on save
+			SpawnParams.OverrideActorGuid = ExistingActor->GetActorGuid();
+		}
+	}
+#endif
+
+	const FVector CellCenter(FVector(GridCoords.X + 0.5, GridCoords.Y + 0.5, GridCoords.Z + 0.5) * GridDescriptor.GetGridSize());
+	APCGPartitionActor* NewActor = CastChecked<APCGPartitionActor>(World->SpawnActor(APCGPartitionActor::StaticClass(), &CellCenter, nullptr, SpawnParams));
+
+	if (GridDescriptor.IsRuntime())
+	{
+		NewActor->SetToRuntimeGenerated();
+	}
+
+#if WITH_EDITOR
+	NewActor->SetLockLocation(true);
+	NewActor->SetActorLabel(SpawnParams.Name.ToString());
+#endif
+
+	// Empty GUID if runtime generated, since transient PAs don't need one.
+	NewActor->PostCreation(GridDescriptor);
+
+	return NewActor;
+}
+
+TSet<TObjectPtr<APCGPartitionActor>> UPCGSubsystem::GetPCGComponentPartitionActorMappings(UPCGComponent* InComponent) const
+{
+	return ActorAndComponentMapping->GetPCGComponentPartitionActorMappings(InComponent);
+}
+
+FPCGGenSourceManager* UPCGSubsystem::GetGenSourceManager() const
+{
+	return RuntimeGenScheduler ? RuntimeGenScheduler->GenSourceManager : nullptr;
+}
+
+#if WITH_EDITOR
+
+FPCGTaskId UPCGSubsystem::ScheduleRefresh(UPCGComponent* Component, bool bForceRegen)
+{
+	check(Component && !Component->IsManagedByRuntimeGenSystem());
+
+	TWeakObjectPtr<UPCGComponent> ComponentPtr(Component);
+
+	auto RefreshTask = [ComponentPtr, bForceRegen]() {
+		if (UPCGComponent* Component = ComponentPtr.Get())
+		{
+			Component->OnRefresh(bForceRegen);
+		}
+		return true;
+	};
+
+	return GraphExecutor->ScheduleGeneric(RefreshTask, Component, {});
+}
+
+void UPCGSubsystem::DirtyGraph(UPCGComponent* Component, const FBox& InBounds, EPCGComponentDirtyFlag DirtyFlag)
+{
+	check(Component);
+
+	// Immediate operation
+	auto DirtyTask = [DirtyFlag](UPCGComponent* LocalComponent)
+	{
+		LocalComponent->DirtyGenerated(DirtyFlag);
+		return InvalidPCGTaskId;
+	};
+
+	ActorAndComponentMapping->DispatchToRegisteredLocalComponents(Component, DirtyTask);
+}
+
+void UPCGSubsystem::ClearPCGLink(UPCGComponent* InComponent, const FBox& InBounds, AActor* InNewActor)
+{
+	TWeakObjectPtr<AActor> NewActorPtr(InNewActor);
+	TWeakObjectPtr<UPCGComponent> ComponentPtr(InComponent);
+
+	auto MoveTask = [this, NewActorPtr, ComponentPtr](APCGPartitionActor* PCGActor, const FBox& InIntersectedBounds)
+	{
+		check(NewActorPtr.IsValid() && ComponentPtr.IsValid() && PCGActor != nullptr);
+
+		if (TObjectPtr<UPCGComponent> LocalComponent = PCGActor->GetLocalComponent(ComponentPtr.Get()))
+		{
+			LocalComponent->MoveResourcesToNewActor(NewActorPtr.Get(), /*bCreateChild=*/true);
+		}
+
+		return true;
+	};
+
+	auto ScheduleTask = [this, ComponentPtr, MoveTask](APCGPartitionActor* InPCGActor, const FBox& InIntersectedBounds)
+	{
+		auto MoveTaskInternal = [MoveTask, InPCGActor, InIntersectedBounds]()
+		{
+			return MoveTask(InPCGActor, InIntersectedBounds);
+		};
+
+		return GraphExecutor->ScheduleGeneric(MoveTaskInternal, ComponentPtr.Get(), /*TaskExecutionDependencies=*/{});
+	};
+
+	auto ScheduleUnloadedTask = [this, ComponentPtr, MoveTask](const FPCGGridDescriptor& InGridDescriptor, const FIntVector& InGridCoord, const FBox& InIntersectedBounds)
+	{
+		auto MoveTaskInternal = [this, MoveTask, InGridDescriptor, InGridCoord, InIntersectedBounds]()
+		{
+			TUniquePtr<IWorldPartitionActorLoaderInterface::ILoaderAdapter> LoaderAdapter = MakeUnique<FPCGPartitionActorLoaderAdapter>(GetWorld(), InIntersectedBounds, TEXT("UPCGSubsystem::ClearPCGLink"));
+			LoaderAdapter->Load();
+			if (APCGPartitionActor* PCGActor = FindOrCreatePCGPartitionActor(InGridDescriptor, InGridCoord, /*bCanCreateActor=*/false))
+			{
+				return MoveTask(PCGActor, InIntersectedBounds);
+			}
+
+			return true;
+		};
+
+		return GraphExecutor->ScheduleGeneric(MoveTaskInternal, ComponentPtr.Get(), /*TaskExecutionDependencies=*/{});
+	};
+
+	bool bHasUnbounded = false;
+	PCGHiGenGrid::FSizeArray GridSizes;
+	ensure(PCGHelpers::GetGenerationGridSizes(InComponent->GetGraph(), GetPCGWorldActor(), GridSizes, bHasUnbounded));
+
+	FPCGTaskId TaskId = InvalidPCGTaskId;
+	if (!GridSizes.IsEmpty())
+	{
+		// Load unloaded PAs to make sure Clear PCG Link is not partial
+		TaskId = ForAllOverlappingCells(InComponent, InBounds, GridSizes, /*bCanCreateActor=*/false, /*Dependencies=*/{}, ScheduleTask, ScheduleUnloadedTask);
+	}
+
+	// Verify if the NewActor has some components attached to its root or attached actors. If not, destroy it.
+	// Return false if the new actor is not valid or destroyed.
+	auto VerifyAndDestroyNewActor = [this, NewActorPtr]()
+	{
+		check(NewActorPtr.IsValid());
+
+		USceneComponent* RootComponent = NewActorPtr->GetRootComponent();
+		check(RootComponent);
+
+		AActor* NewActor = NewActorPtr.Get();
+
+		TArray<AActor*> AttachedActors;
+		NewActor->GetAttachedActors(AttachedActors);
+
+		if (RootComponent->GetNumChildrenComponents() == 0 && AttachedActors.IsEmpty())
+		{
+			GetWorld()->DestroyActor(NewActor);
+			return false;
+		}
+
+		return true;
+	};
+
+	if (TaskId != InvalidPCGTaskId)
+	{
+		auto CleanupTask = [this, ComponentPtr, VerifyAndDestroyNewActor]()
+		{
+
+			// If the new actor is valid, clean up the original component.
+			if (VerifyAndDestroyNewActor())
+			{
+				check(ComponentPtr.IsValid());
+				ComponentPtr->Cleanup(/*bRemoveComponents=*/true);
+			}
+
+			return true;
+		};
+
+		GraphExecutor->ScheduleGeneric(CleanupTask, InComponent, { TaskId });
+	}
+	else
+	{
+		VerifyAndDestroyNewActor();
+	}
+}
+
+void UPCGSubsystem::DeleteSerializedPartitionActors(bool bOnlyDeleteUnused, bool bOnlyChildren)
+{
+	TRACE_CPUPROFILER_EVENT_SCOPE(UPCGSubsystem::DeleteSerializedPartitionActors);
+
+	TSet<UPackage*> PackagesToCleanup;
+	UWorld* World = GetWorld();
+
+	if (!World)
+	{
+		return;
+	}
+
+	UWorldPartition* WorldPartition = World->GetWorldPartition();
+	TMap<FGuid, TArray<FGuid>> Attachments;
+
+	auto GetAttachments = [WorldPartition](const FGuid& ParentActor, const TMap<FGuid, TArray<FGuid>>& Attachments, auto&& GetAttachmentsRecursive, TArray<FGuid>& OutAttachedActors) -> void
+	{
+		if (const TArray<FGuid>* Attached = Attachments.Find(ParentActor))
+		{
+			OutAttachedActors.Append(*Attached);
+			for (const FGuid& AttachedGuid : *Attached)
+			{
+				GetAttachmentsRecursive(AttachedGuid, Attachments, GetAttachmentsRecursive, OutAttachedActors);
+			}
+		}
+	};
+
+	auto GatherAndDestroyActors = [this, &PackagesToCleanup, World, WorldPartition, bOnlyDeleteUnused, bOnlyChildren, &GetAttachments, &Attachments](AActor* Actor) -> bool
+	{
+		TObjectPtr<APCGPartitionActor> PartitionActor = CastChecked<APCGPartitionActor>(Actor);
+
+		// Do not delete RuntimeGen PAs or PAs with graph instances if we are only deleting unused PAs.
+		if (!PartitionActor->IsRuntimeGenerated() && (!bOnlyDeleteUnused || !PartitionActor->HasGraphInstances()))
+		{
+			TArray<AActor*> ActorsToDelete;
+
+			TArray<FWorldPartitionReference> ActorReferences;
+
+			// Load Generated Resources to delete them
+			TArray<TSoftObjectPtr<AActor>> ManagedActors = UPCGComponent::GetManagedActorPaths(PartitionActor);
+			for (const TSoftObjectPtr<AActor>& ManagedActorPath : ManagedActors)
+			{
+				// Test to see if actor is loaded first to support non World Partition worlds
+				AActor* ManagedActor = ManagedActorPath.Get();
+				if (!ManagedActor && WorldPartition)
+				{
+					if (const FWorldPartitionActorDescInstance* ActorDescInstance = WorldPartition->GetActorDescInstanceByPath(ManagedActorPath.ToSoftObjectPath()))
+					{
+						FWorldPartitionReference& ActorReference = ActorReferences.Emplace_GetRef(FWorldPartitionReference(ActorDescInstance->GetContainerInstance(), ActorDescInstance->GetGuid()));
+						ManagedActor = ActorReference.GetActor();
+					}
+				}
+
+				if (ManagedActor)
+				{
+					ActorsToDelete.Add(ManagedActor);
+				}
+			}
+			
+			// Load Attachments before getting them in the next code block, since loading an actor doesn't load it's attachments (the reference is child to parent)
+			if (WorldPartition)
+			{
+				TArray<FGuid> AttachedActors;
+				GetAttachments(Actor->GetActorGuid(), Attachments, GetAttachments, AttachedActors);
+				for (const FGuid& AttachedActor : AttachedActors)
+				{
+					ActorReferences.Add(FWorldPartitionReference(WorldPartition, AttachedActor));
+				}
+			}
+
+			// We might have actors that weren't saved as managed resources that are attached and have the proper tag
+			TArray<AActor*> AttachedActors;
+			PartitionActor->GetAttachedActors(AttachedActors, /*bResetArray=*/true, /*bRecursivelyIncludeAttachedActors=*/ true);
+			for (AActor* AttachedActor : AttachedActors)
+			{
+				if (AttachedActor->ActorHasTag(PCGHelpers::DefaultPCGActorTag))
+				{
+					if (!ActorsToDelete.Contains(AttachedActor))
+					{
+						ActorsToDelete.Add(AttachedActor);
+					}
+				}
+				else if (WorldPartition && !bOnlyChildren)
+				{
+					// If actor isn't getting deleted but is an attached actor and it's Partition Actor parent 
+					// will get deleted then Pin the actor so it stays loaded and modified for the user to save
+					WorldPartition->PinActors({ AttachedActor->GetActorGuid() });
+				}
+			}
+
+			if (!bOnlyChildren)
+			{
+				ActorsToDelete.Add(PartitionActor);
+			}
+
+			for (AActor* ActorToDelete : ActorsToDelete)
+			{
+				if (UPackage* ExternalPackage = ActorToDelete->GetExternalPackage())
+				{
+					// Since we aren't in a transaction (and this operation isnt undoable) make sure UPackage objects are no longer marked as RF_Standalone so it they can get properly GCed in (ObjectTools::CleanupAfterSuccessfulDelete)
+					// Without this call we end up spamming in UWorldPartition::OnGCPostReachabilityAnalysis
+					ForEachObjectWithPackage(ExternalPackage, [this, Actor](UObject* Object)
+					{
+						Object->ClearFlags(RF_Standalone);
+						return true;
+					}, false);
+
+					PackagesToCleanup.Add(ExternalPackage);
+				}
+
+				World->DestroyActor(ActorToDelete);
+			}
+		}
+
+		return true;
+	};
+
+	// First, clear selection otherwise it might crash
+	if (GEditor)
+	{
+		GEditor->SelectNone(true, true, false);
+		// Any reference in the Transaction buffer to the deleted actors will prevent them from being properly GCed so here we reset the transaction buffer
+		GEditor->ResetTransaction(NSLOCTEXT("PCGSubsystem", "DeletePartitionActorsResetTransaction", "Deleted PCG Actors"));
+	}
+
+	{
+		FScopedSlowTask DeleteTask(0, NSLOCTEXT("PCGSubsystem", "DeletePartitionActors", "Deleting PCG Actors..."));
+		DeleteTask.MakeDialog();
+
+		FWorldPartitionHelpers::FForEachActorWithLoadingResult LoadingResult;
+		{
+			TRACE_CPUPROFILER_EVENT_SCOPE(UPCGSubsystem::DeleteSerializedPartitionActors::ForEachActorInLevel);
+			if (WorldPartition)
+			{
+				// Gather Attach Parent information
+				FWorldPartitionHelpers::ForEachActorDescInstance(WorldPartition, [&Attachments](const FWorldPartitionActorDescInstance* ActorDescInstance)
+				{
+					if (ActorDescInstance->GetParentActor().IsValid())
+					{
+						Attachments.FindOrAdd(ActorDescInstance->GetParentActor()).Add(ActorDescInstance->GetGuid());
+					}
+
+					return true;
+				});
+
+				// Process Loaded Actors first (and unsaved actors that don't have an Actor Desc yet)
+				// Do not use UPCGActorHelpers::ForEachActorInLevel as GatherAndDestroy can end up modifying the Actors array (by loading actors) in a WP World
+				TSet<FGuid> ProcessedActors;
+				const TArray<AActor*> ActorsCopy(World->PersistentLevel->Actors);
+				for (AActor* Actor : ActorsCopy)
+				{
+					if (APCGPartitionActor* PartitionActor = Cast<APCGPartitionActor>(Actor))
+					{
+						ProcessedActors.Add(Actor->GetActorGuid());
+						GatherAndDestroyActors(Actor);
+					}
+				}
+				
+				FWorldPartitionHelpers::FForEachActorWithLoadingParams ForEachActorWithLoadingParams;
+				ForEachActorWithLoadingParams.bKeepReferences = true;
+				ForEachActorWithLoadingParams.ActorClasses = { APCGPartitionActor::StaticClass() };
+
+				// Load and Process remaining actors
+				FWorldPartitionHelpers::ForEachActorWithLoading(WorldPartition, [&GatherAndDestroyActors, &ProcessedActors](const FWorldPartitionActorDescInstance* ActorDescInstance)
+				{
+					if(AActor* Actor = ActorDescInstance->GetActor(); Actor && !ProcessedActors.Contains(Actor->GetActorGuid()))
+					{
+						GatherAndDestroyActors(Actor);
+					}
+					return true;
+				},
+				ForEachActorWithLoadingParams);
+			}
+			else
+			{
+				UPCGActorHelpers::ForEachActorInLevel<APCGPartitionActor>(World->PersistentLevel, GatherAndDestroyActors);
+			}
+		}
+
+		if (PackagesToCleanup.Num() > 0)
+		{
+			ObjectTools::CleanupAfterSuccessfulDelete(PackagesToCleanup.Array(), /*bPerformanceReferenceCheck=*/true);
+		}
+
+		// Non World Partition Levels might have deleted actors without saving anything and we need to GC so that Partition Actors can be created again (avoid name clash)
+		CollectGarbage(GARBAGE_COLLECTION_KEEPFLAGS);
+	}
+}
+
+void UPCGSubsystem::PropagateEditingModeToLocalComponents(UPCGComponent* InComponent, EPCGEditorDirtyMode EditingMode)
+{
+	if (ensure(InComponent && InComponent->IsPartitioned()))
+	{
+		FBox Bounds = ActorAndComponentMapping->PartitionedOctree.GetBounds(InComponent);
+		if (!Bounds.IsValid)
+		{
+			return;
+		}
+
+		ActorAndComponentMapping->ForAllIntersectingPartitionActors(Bounds, [InComponent, EditingMode](APCGPartitionActor* Actor)
+		{
+			Actor->ChangeTransientState(InComponent, EditingMode);
+		});
+	}
+}
+
+void UPCGSubsystem::BuildLandscapeCache(bool bQuiet, bool bForceBuild)
+{
+	TRACE_CPUPROFILER_EVENT_SCOPE(UPCGSubsystem::BuildLandscapeCache);
+	if (UPCGLandscapeCache* LandscapeCache = GetLandscapeCache())
+	{
+		if (bForceBuild || LandscapeCache->SerializationMode != EPCGLandscapeCacheSerializationMode::NeverSerialize)
+		{
+			LandscapeCache->PrimeCache();
+		}
+	}
+	else if(!bQuiet)
+	{
+		UE_LOG(LogPCG, Error, TEXT("Unable to build landscape cache because either the world is null or there is no PCG world actor"));
+	}
+}
+
+void UPCGSubsystem::ClearLandscapeCache()
+{
+	if (UPCGLandscapeCache* LandscapeCache = GetLandscapeCache())
+	{
+		LandscapeCache->ClearCache();
+	}
+}
+
+void UPCGSubsystem::GenerateAllPCGComponents(bool bForce) const
+{
+	for (UPCGComponent* Component : ActorAndComponentMapping->GetAllRegisteredComponents())
+	{
+		Component->Generate(bForce);
+	}
+}
+
+void UPCGSubsystem::CleanupAllPCGComponents(bool bPurge) const
+{
+	for (UPCGComponent* Component : ActorAndComponentMapping->GetAllRegisteredComponents())
+	{
+		if (bPurge)
+		{
+			Component->CleanupLocalDeleteAllGeneratedObjects({});
+		}
+		else
+		{
+			Component->Cleanup(/*bRemoveComponents=*/true);
+		}
+	}
+}
+
+const FPCGNodeVisualLogs& UPCGSubsystem::GetNodeVisualLogs() const
+{
+	if (IPCGEditorModule* PCGEditorModule = IPCGEditorModule::Get())
+	{
+		return PCGEditorModule->GetNodeVisualLogs();
+	}
+	else
+	{
+		// Not meant to be used but to respect the existing API return a placeholder log
+		static FPCGNodeVisualLogs Placeholder;
+		return Placeholder;
+	}
+}
+
+FPCGNodeVisualLogs& UPCGSubsystem::GetNodeVisualLogsMutable()
+{
+	if (IPCGEditorModule* PCGEditorModule = IPCGEditorModule::Get())
+	{
+		return PCGEditorModule->GetNodeVisualLogsMutable();
+	}
+	else
+	{
+		// Not meant to be used but to respect the existing API return a placeholder log
+		static FPCGNodeVisualLogs Placeholder;
+		return Placeholder;
+	}
+}
+
+void UPCGSubsystem::NotifyLandscapeEditModeExited()
+{
+	ActorAndComponentMapping->NotifyLandscapeEditModeExited();
+}
+
+/* UE_DEPRECATED(5.6) */
+void UPCGSubsystem::ClearExecutionMetadata(const FPCGStack& BaseStack)
+{
+	if (IPCGEditorModule* PCGEditorModule = IPCGEditorModule::Get())
+	{
+		if (UPCGComponent* Component = const_cast<UPCGComponent*>(BaseStack.GetRootComponent()))
+		{
+			PCGEditorModule->ClearExecutionMetadata(Component);
+		}
+	}
+}
+
+void UPCGSubsystem::ClearExecutionMetadata(UPCGComponent* InComponent)
+{
+	if (IPCGEditorModule* PCGEditorModule = IPCGEditorModule::Get())
+	{
+		PCGEditorModule->ClearExecutionMetadata(InComponent);
+	}
+}
+
+/* UE_DEPRECATED(5.6) */
+TArray<FPCGStack> UPCGSubsystem::GetExecutedStacks(const UPCGComponent* InComponent, const UPCGGraph* InSubgraph, bool bOnlyWithSubgraphAsCurrentFrame)
+{
+	if (IPCGEditorModule* PCGEditorModule = IPCGEditorModule::Get())
+	{
+		TArray<FPCGStackSharedPtr> StackPtrs = PCGEditorModule->GetExecutedStacksPtrs(InComponent, InSubgraph, bOnlyWithSubgraphAsCurrentFrame);
+		TArray<FPCGStack> Stacks;
+		Stacks.Reserve(StackPtrs.Num());
+		Algo::Transform(StackPtrs, Stacks, [](const FPCGStackSharedPtr& StackPtr) { return *StackPtr; });
+
+		return Stacks;
+	}
+	else
+	{
+		return {};
+	}
+}
+
+/* UE_DEPRECATED(5.6) */
+TArray<FPCGStack> UPCGSubsystem::GetExecutedStacks(const FPCGStack& BeginningWithStack)
+{
+	IPCGEditorModule* PCGEditorModule = IPCGEditorModule::Get();
+	if (!PCGEditorModule)
+	{
+		return {};
+	}
+
+	TArray<FPCGStackSharedPtr> StackPtrs = PCGEditorModule->GetExecutedStacksPtrs(BeginningWithStack);
+	TArray<FPCGStack> Stacks;
+
+	Stacks.Reserve(StackPtrs.Num());
+
+	for (FPCGStackSharedPtr StackPtr : StackPtrs)
+	{
+		Stacks.Add(*StackPtr);
+	}
+
+	return Stacks;
+}
+
+/* UE_DEPRECATED(5.6) */
+void UPCGSubsystem::ClearExecutedStacks(const FPCGStack& BeginningWithStack)
+{
+PRAGMA_DISABLE_DEPRECATION_WARNINGS
+	return ClearExecutedStacks(BeginningWithStack.GetRootComponent());
+PRAGMA_ENABLE_DEPRECATION_WARNINGS
+}
+
+/* UE_DEPRECATED(5.6) */
+void UPCGSubsystem::ClearExecutedStacks(const UPCGComponent* InRootComponent)
+{
+	if (IPCGEditorModule* PCGEditorModule = IPCGEditorModule::Get())
+	{
+		PCGEditorModule->ClearExecutedStacks(InRootComponent);
+	}
+}
+
+/* UE_DEPRECATED(5.6) */
+void UPCGSubsystem::ClearExecutedStacks(const UPCGGraph* InContainingGraph)
+{
+	if (IPCGEditorModule* PCGEditorModule = IPCGEditorModule::Get())
+	{
+		PCGEditorModule->ClearExecutedStacks(InContainingGraph);
+	}
+}
+
+void UPCGSubsystem::CreateMissingPartitionActors()
+{
+	if (!PCGHelpers::IsRuntimeOrPIE())
+	{
+		ActorAndComponentMapping->ForAllOriginalComponents([this](UPCGComponent* PCGComponent)
+		{
+			if (PCGComponent->IsPartitioned() && !PCGComponent->IsManagedByRuntimeGenSystem())
+			{
+				bool bHasUnbounded = false;
+				PCGHiGenGrid::FSizeArray GridSizes;
+				ensure(PCGHelpers::GetGenerationGridSizes(PCGComponent->GetGraph(), GetPCGWorldActor(), GridSizes, bHasUnbounded));
+				if (!GridSizes.IsEmpty())
+				{
+					CreatePartitionActorsWithinBounds(PCGComponent, PCGComponent->GetGridBounds(), GridSizes);
+				}
+				ActorAndComponentMapping->UpdateMappingPCGComponentPartitionActor(PCGComponent);
+			}
+		});
+	}
+}
+
+void UPCGSubsystem::CreatePartitionActorsWithinBounds(UPCGComponent* InComponent, const FBox& InBounds, const PCGHiGenGrid::FSizeArray& InGridSizes)
+{
+	UWorld* World = GetWorld();
+	if (!PCGHelpers::IsRuntimeOrPIE() && !UE::GetIsEditorLoadingPackage() && !IsPartitionActorCreationDisabledForWorld(World))
+	{
+		// We can't spawn actors if we are running constructions scripts, asserting when we try to get the actor with the WP API.
+		// We should never enter this if we are in a construction script. If the ensure is hit, we need to fix it.
+		if (ensure(World && !World->bIsRunningConstructionScript))
+		{
+			ForAllOverlappingCells(InComponent, InBounds, InGridSizes, /*bCanCreateActor=*/true, {}, [](APCGPartitionActor*, const FBox&) { return InvalidPCGTaskId; });
+		}
+	}
+}
+
+void UPCGSubsystem::UpdateMappingPCGComponentPartitionActor(UPCGComponent* InComponent) 
+{ 
+	ActorAndComponentMapping->UpdateMappingPCGComponentPartitionActor(InComponent); 
+}
+
+void UPCGSubsystem::OnPCGGraphCancelled(UPCGComponent* InComponent)
+{
+	ActorAndComponentMapping->OnPCGGraphCancelled(InComponent);
+}
+
+void UPCGSubsystem::OnPCGGraphStartGenerating(UPCGComponent* InComponent)
+{
+	ActorAndComponentMapping->OnPCGGraphStartsGenerating(InComponent);
+}
+
+void UPCGSubsystem::OnPCGGraphGenerated(UPCGComponent* InComponent)
+{
+	ActorAndComponentMapping->OnPCGGraphGeneratedOrCleaned(InComponent);
+}
+
+void UPCGSubsystem::OnPCGGraphCleaned(UPCGComponent* InComponent)
+{
+	ActorAndComponentMapping->OnPCGGraphGeneratedOrCleaned(InComponent);
+}
+
+#endif // WITH_EDITOR
+
+UPCGData* UPCGSubsystem::GetPCGData(FPCGTaskId InGraphExecutionTaskId)
+{
+	return GraphExecutor ? GraphExecutor->GetPCGData(InGraphExecutionTaskId) : nullptr;
+}
+
+UPCGData* UPCGSubsystem::GetInputPCGData(FPCGTaskId InGraphExecutionTaskId)
+{
+	return GraphExecutor ? GraphExecutor->GetInputPCGData(InGraphExecutionTaskId) : nullptr;
+}
+
+UPCGData* UPCGSubsystem::GetActorPCGData(FPCGTaskId InGraphExecutionTaskId)
+{
+	return GraphExecutor ? GraphExecutor->GetActorPCGData(InGraphExecutionTaskId) : nullptr;
+}
+
+UPCGData* UPCGSubsystem::GetLandscapePCGData(FPCGTaskId InGraphExecutionTaskId)
+{
+	return GraphExecutor ? GraphExecutor->GetLandscapePCGData(InGraphExecutionTaskId) : nullptr;
+}
+
+UPCGData* UPCGSubsystem::GetLandscapeHeightPCGData(FPCGTaskId InGraphExecutionTaskId)
+{
+	return GraphExecutor ? GraphExecutor->GetLandscapeHeightPCGData(InGraphExecutionTaskId) : nullptr;
+}
+
+UPCGData* UPCGSubsystem::GetOriginalActorPCGData(FPCGTaskId InGraphExecutionTaskId)
+{
+	return GraphExecutor ? GraphExecutor->GetOriginalActorPCGData(InGraphExecutionTaskId) : nullptr;
+}
+
+void UPCGSubsystem::SetPCGData(FPCGTaskId InGraphExecutionTaskId, UPCGData* InData)
+{
+	if (GraphExecutor)
+	{
+		GraphExecutor->SetPCGData(InGraphExecutionTaskId, InData);
+	}
+}
+
+void UPCGSubsystem::SetInputPCGData(FPCGTaskId InGraphExecutionTaskId, UPCGData* InData)
+{
+	if (GraphExecutor)
+	{
+		GraphExecutor->SetInputPCGData(InGraphExecutionTaskId, InData);
+	}
+}
+
+void UPCGSubsystem::SetActorPCGData(FPCGTaskId InGraphExecutionTaskId, UPCGData* InData)
+{
+	if (GraphExecutor)
+	{
+		GraphExecutor->SetActorPCGData(InGraphExecutionTaskId, InData);
+	}
+}
+
+void UPCGSubsystem::SetLandscapePCGData(FPCGTaskId InGraphExecutionTaskId, UPCGData* InData)
+{
+	if (GraphExecutor)
+	{
+		GraphExecutor->SetLandscapePCGData(InGraphExecutionTaskId, InData);
+	}
+}
+
+void UPCGSubsystem::SetLandscapeHeightPCGData(FPCGTaskId InGraphExecutionTaskId, UPCGData* InData)
+{
+	if (GraphExecutor)
+	{
+		GraphExecutor->SetLandscapeHeightPCGData(InGraphExecutionTaskId, InData);
+	}
+}
+
+void UPCGSubsystem::SetOriginalActorPCGData(FPCGTaskId InGraphExecutionTaskId, UPCGData* InData)
+{
+	if (GraphExecutor)
+	{
+		GraphExecutor->SetOriginalActorPCGData(InGraphExecutionTaskId, InData);
+	}
+}
+
+void UPCGSubsystem::ExecuteBeginTickActions()
+{
+	TArray<FTickAction> Actions = MoveTemp(BeginTickActions);
+	BeginTickActions.Reset();
+
+	for (FTickAction& Action : Actions)
+	{
+		Action();
+	}
+}
+
+bool UPCGSubsystem::RegisterOrUpdatePCGComponent(UPCGComponent* InComponent, bool bDoActorMapping)
+{
+	return ActorAndComponentMapping->RegisterOrUpdatePCGComponent(InComponent, bDoActorMapping);
+}
+
+void UPCGSubsystem::UnregisterPartitionActor(APCGPartitionActor* InActor)
+{
+	ActorAndComponentMapping->UnregisterPartitionActor(InActor);
+}
+
+TSet<UPCGComponent*> UPCGSubsystem::GetAllRegisteredPartitionedComponents() const
+{
+	return ActorAndComponentMapping->GetAllRegisteredPartitionedComponents();
+}
+
+TSet<UPCGComponent*> UPCGSubsystem::GetAllRegisteredComponents() const
+{
+	return ActorAndComponentMapping->GetAllRegisteredComponents();
+}
+
+void UPCGSubsystem::ForAllIntersectingPartitionedComponents(const FBoxCenterAndExtent& InBounds, TFunctionRef<void(UPCGComponent*)> InFunc) const
+{
+	return ActorAndComponentMapping->ForAllIntersectingPartitionedComponents(InBounds, std::move(InFunc));
+}
+
+TArray<UPCGComponent*> UPCGSubsystem::GetAllIntersectingComponents(const FBoxCenterAndExtent& InBounds) const
+{
+	return ActorAndComponentMapping->GetAllIntersectingComponents(InBounds);
+}
+
+bool UPCGSubsystem::RemapPCGComponent(const UPCGComponent* OldComponent, UPCGComponent* NewComponent, bool bDoActorMapping)
+{
+	return ActorAndComponentMapping->RemapPCGComponent(OldComponent, NewComponent, bDoActorMapping);
+}
+
+void UPCGSubsystem::UnregisterPCGComponent(UPCGComponent* InComponent, bool bForce)
+{
+	ActorAndComponentMapping->UnregisterPCGComponent(InComponent, bForce);
+}
+
+void UPCGSubsystem::RegisterPartitionActor(APCGPartitionActor* InActor)
+{
+	ActorAndComponentMapping->RegisterPartitionActor(InActor);
+}
+
+#if WITH_EDITOR
+void UPCGSubsystem::UpdateComponentTracking(UPCGComponent* InComponent, bool bShouldDirtyActors, const TArray<FPCGSelectionKey>* OptionalChangedKeys)
+{
+	ActorAndComponentMapping->UpdateTracking(InComponent, bShouldDirtyActors, OptionalChangedKeys);
+}
+
+void UPCGSubsystem::ResetPartitionActorsMap()
+{
+	ActorAndComponentMapping->ResetPartitionActorsMap();
+}
+#endif // WITH_EDITOR
diff --git a/Engine/Plugins/PCG/Source/PCG/Public/Editor/IPCGEditorModule.h b/Engine/Plugins/PCG/Source/PCG/Public/Editor/IPCGEditorModule.h
--- a/Engine/Plugins/PCG/Source/PCG/Public/Editor/IPCGEditorModule.h
+++ b/Engine/Plugins/PCG/Source/PCG/Public/Editor/IPCGEditorModule.h
@@ -9,6 +9,7 @@
 #include "Internationalization/Text.h"
 #include "Modules/ModuleInterface.h"
 
+class IPCGBaseSubsystem;
 struct FPCGStackContext;
 class UPCGComponent;
 
@@ -26,6 +27,7 @@ public:
 	virtual FPCGNodeVisualLogs& GetNodeVisualLogsMutable() = 0;
 	virtual bool CanSelectPartitionActors() const = 0;
 	virtual TSharedPtr<class ISyntaxTokenizer> CreateHLSLSyntaxTokenizer(const FPCGSyntaxTokenizerParams& InParams) const = 0;
+	virtual IPCGBaseSubsystem* GetEditorSubsystem() const = 0;
 #endif
 
 protected:
diff --git a/Engine/Plugins/PCG/Source/PCG/Public/PCGComponent.h b/Engine/Plugins/PCG/Source/PCG/Public/PCGComponent.h
--- a/Engine/Plugins/PCG/Source/PCG/Public/PCGComponent.h
+++ b/Engine/Plugins/PCG/Source/PCG/Public/PCGComponent.h
@@ -107,6 +107,7 @@ public:
 	virtual UPCGGraphInstance* GetGraphInstance() const override;
 	virtual void OnGraphExecutionAborted(bool bQuiet, bool bCleanupUnusedResources) override;
 	virtual void Cancel() override;
+	virtual bool IsGenerating() const override;
 
 #if WITH_EDITOR
 	virtual const PCGUtils::FExtraCapture& GetExtraCapture() const override;
@@ -117,6 +118,8 @@ public:
 
 	virtual void RegisterDynamicTracking(const UPCGSettings* InSettings, const TArrayView<TPair<FPCGSelectionKey, bool>>& InDynamicKeysAndCulling) override;
 	virtual void RegisterDynamicTracking(const FPCGSelectionKeyToSettingsMap& InKeysToSettings) override;
+	
+	virtual bool IsRefreshInProgress() const override;
 #endif // WITH_EDITOR
 
 private:
diff --git a/Engine/Plugins/PCG/Source/PCG/Public/PCGGraphExecutionStateInterface.h b/Engine/Plugins/PCG/Source/PCG/Public/PCGGraphExecutionStateInterface.h
--- a/Engine/Plugins/PCG/Source/PCG/Public/PCGGraphExecutionStateInterface.h
+++ b/Engine/Plugins/PCG/Source/PCG/Public/PCGGraphExecutionStateInterface.h
@@ -6,13 +6,15 @@
 #include "Utils/PCGExtraCapture.h"
 
 #include "CoreMinimal.h"
+#include "Misc/SpinLock.h"
 #include "UObject/Interface.h"
 
 #include "PCGGraphExecutionStateInterface.generated.h"
 
-class UPCGData;
+class IPCGBaseSubsystem;
 class FPCGGraphExecutionInspection;
 class IPCGGraphExecutionSource;
+class UPCGData;
 class UPCGGraph;
 class UPCGGraphInstance;
 
@@ -62,6 +64,12 @@ public:
 	/** Notify ExecutionState that its execution is being aborted */
 	virtual void OnGraphExecutionAborted(bool bQuiet = false, bool bCleanupUnusedResources = true) = 0;
 
+	/** Indicates if the current source is generating */
+	virtual bool IsGenerating() const = 0;
+
+	/** Returns the subsystem associated with this source. By default, it will return the UPCGSubsystem if the source has a world, otherwise, the EditorSubsystem (in editor only, obviously). */
+	virtual IPCGBaseSubsystem* GetSubsystem() const;
+
 #if WITH_EDITOR
 
 	/** Returns the FExtraCapture object for this ExecutionState */
@@ -80,6 +88,9 @@ public:
 	/** Register multiple tracking dependencies, so ExecutionState can be updated when they change */
 	virtual void RegisterDynamicTracking(const FPCGSelectionKeyToSettingsMap& InKeysToSettings) = 0;
 
+	/** Indicates if the current source is waiting for a refresh */
+	virtual bool IsRefreshInProgress() const = 0;
+
 #endif
 };
 
@@ -100,3 +111,37 @@ public:
 	virtual IPCGGraphExecutionState& GetExecutionState() = 0;
 	virtual const IPCGGraphExecutionState& GetExecutionState() const = 0;
 };
+
+USTRUCT()
+struct FPCGSoftGraphExecutionSource
+{
+	GENERATED_BODY()
+
+	FPCGSoftGraphExecutionSource() = default;
+	PCG_API explicit FPCGSoftGraphExecutionSource(const TSoftObjectPtr<UObject>& InSoftObjectPtr);
+	PCG_API FPCGSoftGraphExecutionSource(IPCGGraphExecutionSource* InSource);
+	PCG_API FPCGSoftGraphExecutionSource(const IPCGGraphExecutionSource* InSource);
+
+	// Need the copy/move constructor/assignment because of the lock.
+	PCG_API FPCGSoftGraphExecutionSource(const FPCGSoftGraphExecutionSource& Other);
+	PCG_API FPCGSoftGraphExecutionSource(FPCGSoftGraphExecutionSource&& Other);
+	PCG_API FPCGSoftGraphExecutionSource& operator=(const FPCGSoftGraphExecutionSource& Other);
+	PCG_API FPCGSoftGraphExecutionSource& operator=(FPCGSoftGraphExecutionSource&& Other);
+
+	PCG_API IPCGGraphExecutionSource* Get() const;
+	
+	bool IsValid() const { return Get() != nullptr; }
+	
+	IPCGGraphExecutionSource* operator->() const { return Get(); }
+
+	PCG_API bool operator==(const FPCGSoftGraphExecutionSource& Other) const;
+
+	PCG_API friend int32 GetTypeHash(const FPCGSoftGraphExecutionSource& This);
+	
+private:
+	UPROPERTY()
+	TSoftObjectPtr<UObject> SoftObjectPtr;
+
+	mutable UE::FSpinLock SpinLock;
+	mutable TWeakInterfacePtr<IPCGGraphExecutionSource> CachedWeakPtr;
+};
diff --git a/Engine/Plugins/PCG/Source/PCG/Public/PCGSubsystem.h b/Engine/Plugins/PCG/Source/PCG/Public/PCGSubsystem.h
--- a/Engine/Plugins/PCG/Source/PCG/Public/PCGSubsystem.h
+++ b/Engine/Plugins/PCG/Source/PCG/Public/PCGSubsystem.h
@@ -2,480 +2,4 @@
 
 #pragma once
 
-#include "Subsystems/WorldSubsystem.h"
-
-#include "PCGCommon.h"
-#include "Elements/PCGActorSelector.h"
-#include "Grid/PCGGridDescriptor.h"
-#include "Grid/PCGComponentOctree.h"
-#include "Utils/PCGNodeVisualLogs.h"
-
-#include "UObject/ObjectKey.h"
-
-#include "PCGSubsystem.generated.h"
-
-#define UE_API PCG_API
-
-class APCGPartitionActor;
-class APCGWorldActor;
-class FPCGActorAndComponentMapping;
-class FPCGGenSourceManager;
-class FPCGRuntimeGenScheduler;
-class UPCGComputeGraph;
-class UPCGData;
-class UPCGGraph;
-class UPCGLandscapeCache;
-
-enum class EPCGComponentDirtyFlag : uint8;
-enum class ETickableTickType : uint8;
-
-class IPCGGraphCache;
-class FPCGGraphCompiler;
-class FPCGGraphExecutor;
-struct FPCGContext;
-struct FPCGDataCollection;
-struct FPCGScheduleGenericParams;
-struct FPCGStack;
-class UPCGSettings;
-
-class IPCGElement;
-typedef TSharedPtr<IPCGElement, ESPMode::ThreadSafe> FPCGElementPtr;
-
-class UWorld;
-
-#if WITH_EDITOR
-/** Deprecated - use FPCGOnPCGComponentUnregistered */
-DECLARE_MULTICAST_DELEGATE(FPCGOnComponentUnregistered);
-/** Deprecated - use FPCGOnPCGComponentGenerationDone */
-DECLARE_MULTICAST_DELEGATE_OneParam(FPCGOnComponentGenerationCompleteOrCancelled, UPCGSubsystem*);
-
-DECLARE_MULTICAST_DELEGATE_OneParam(FPCGOnPCGComponentUnregistered, UPCGComponent*);
-DECLARE_MULTICAST_DELEGATE_ThreeParams(FPCGOnPCGComponentGenerationDone, UPCGSubsystem*, UPCGComponent*, EPCGGenerationStatus);
-#endif // WITH_EDITOR
-
-/**
-* UPCGSubsystem
-*/
-UCLASS(MinimalAPI)
-class UPCGSubsystem : public UTickableWorldSubsystem
-{
-	GENERATED_BODY()
-
-public:
-	friend class UPCGComponent;
-	friend FPCGActorAndComponentMapping;
-	friend struct FPCGWorldPartitionBuilder;
-
-	UE_API UPCGSubsystem();
-
-	/** Add UObject references for GC */
-	static UE_API void AddReferencedObjects(UObject* InThis, FReferenceCollector& Collector);
-
-	/** To be used when a PCG component can not have a world anymore, to unregister itself. */
-	static UE_API UPCGSubsystem* GetSubsystemForCurrentWorld();
-
-	//~ Begin USubsystem Interface.
-	UE_API virtual void Deinitialize() override;
-	//~ End USubsystem Interface
-
-	//~ Begin UWorldSubsystem Interface.
-	UE_API virtual void PostInitialize() override;
-	// need UpdateStreamingState? 
-	//~ End UWorldSubsystem Interface.
-
-	//~ Begin FTickableGameObject
-	UE_API virtual void Tick(float DeltaSeconds) override;
-	virtual bool IsTickableInEditor() const override { return true; }
-	UE_API virtual ETickableTickType GetTickableTickType() const override;
-	UE_API virtual TStatId GetStatId() const override;
-	//~ End FTickableGameObject
-
-	/** Will return the subsystem from the World if it exists and if it is initialized */
-	static UE_API UPCGSubsystem* GetInstance(UWorld* World);
-
-	/** Adds an action that will be executed once at the beginning of this subsystem's next Tick(). */
-	using FTickAction = TFunction<void()>;
-
-	UE_DEPRECATED(5.6, "Use FPCGModule::ExecuteNextTick instead")
-	UE_API void RegisterBeginTickAction(FTickAction&& Action);
-
-#if WITH_EDITOR
-	/** Returns PIE world if it is active, otherwise returns editor world. */
-	static UE_API UPCGSubsystem* GetActiveEditorInstance();
-
-	UE_API void SetConstructionScriptSourceComponent(UPCGComponent* InComponent);
-	UE_API bool RemoveAndCopyConstructionScriptSourceComponent(AActor* InComponentOwner, FName InComponentName, UPCGComponent*& OutSourceComponent);
-#endif
-
-	/** Subsystem must not be used without this condition being true. */
-	bool IsInitialized() const { return GraphExecutor != nullptr; }
-
-	UE_API APCGWorldActor* GetPCGWorldActor();
-	UE_API APCGWorldActor* FindPCGWorldActor();
-
-	/** Returns current quality level between Low (0) and Cinematic (4). */
-	static UE_API int32 GetPCGQualityLevel();
-	UE_API void OnPCGQualityLevelChanged();
-
-#if WITH_EDITOR
-	UE_API void DestroyAllPCGWorldActors();
-	UE_API void DestroyCurrentPCGWorldActor();
-	UE_API void LogAbnormalComponentStates(bool bGroupByState) const;
-#endif
-	UE_API void RegisterPCGWorldActor(APCGWorldActor* InActor);
-	UE_API void UnregisterPCGWorldActor(APCGWorldActor* InActor);
-
-	UE_API void OnOriginalComponentRegistered(UPCGComponent* InComponent);
-	UE_API void OnOriginalComponentUnregistered(UPCGComponent* InComponent);
-
-	/** Called by graph executor when a graph is scheduled. */
-	UE_API void OnScheduleGraph(const FPCGStackContext& StackContext);
-
-	UE_API UPCGLandscapeCache* GetLandscapeCache();
-
-	/** Schedule graph(owner->graph) */
-	UE_API FPCGTaskId ScheduleComponent(UPCGComponent* PCGComponent, EPCGHiGenGrid Grid, bool bForce, const TArray<FPCGTaskId>& InDependencies);
-
-	/** Schedule cleanup(owner->graph). */
-	UE_API FPCGTaskId ScheduleCleanup(UPCGComponent* PCGComponent, bool bRemoveComponents, const TArray<FPCGTaskId>& Dependencies);
-
-	// Schedule graph (used internally for dynamic subgraph execution)
-	UE_API FPCGTaskId ScheduleGraph(const FPCGScheduleGraphParams& InParams);
-		
-	UE_API FPCGTaskId ScheduleGraph(
-		UPCGGraph* Graph,
-		UPCGComponent* SourceComponent,
-		FPCGElementPtr PreGraphElement,
-		FPCGElementPtr InputElement,
-		const TArray<FPCGTaskId>& Dependencies,
-		const FPCGStack* InFromStack,
-		bool bAllowHierarchicalGeneration);
-
-	// Schedule graph (used internally for dynamic subgraph execution)
-	UE_API FPCGTaskId ScheduleGraph(UPCGComponent* SourceComponent, const TArray<FPCGTaskId>& Dependencies);
-
-	// General job scheduling
-	UE_API FPCGTaskId ScheduleGeneric(const FPCGScheduleGenericParams& InParams);
-
-	/** General job scheduling
-	*  @param InOperation:               Callback that returns true if the task is done, false otherwise.
-	*  @param SourceComponent:           PCG component associated with this task. Can be null.
-	*  @param TaskExecutionDependencies: Task will wait on these tasks to execute and won't take their output data as input.
-	*/
-	UE_API FPCGTaskId ScheduleGeneric(TFunction<bool()> InOperation, UPCGComponent* SourceComponent, const TArray<FPCGTaskId>& TaskExecutionDependencies);
-
-	/** General job scheduling
-	*  @param InOperation:               Callback that returns true if the task is done, false otherwise.
-	*  @param InAbortOperation:          Callback that will be called if the generic task is cancelled for any reason.
-	*  @param SourceComponent:           PCG component associated with this task. Can be null.
-	*  @param TaskExecutionDependencies: Task will wait on these tasks to execute and won't take their output data as input.
-	*/
-	UE_API FPCGTaskId ScheduleGeneric(TFunction<bool()> InOperation, TFunction<void()> InAbortOperation, UPCGComponent* SourceComponent, const TArray<FPCGTaskId>& TaskExecutionDependencies);
-
-	/** General job scheduling with context
-	*  @param InOperation:               Callback that takes a Context as argument and returns true if the task is done, false otherwise.
-	*  @param SourceComponent:           PCG component associated with this task. Can be null.
-	*  @param TaskExecutionDependencies: Task will wait on these tasks to execute and won't take their output data as input.
-	*  @param TaskDataDependencies:      Task will wait on these tasks to execute and will take their output data as input.
-	*/
-	UE_API FPCGTaskId ScheduleGenericWithContext(TFunction<bool(FPCGContext*)> InOperation, UPCGComponent* SourceComponent, const TArray<FPCGTaskId>& TaskExecutionDependencies, const TArray<FPCGTaskId>& TaskDataDependencies);
-
-	/** General job scheduling with context
-	*  @param InOperation:               Callback that takes a Context as argument and returns true if the task is done, false otherwise.
-	*  @param InAbortOperation:          Callback that will be called if the generic task is cancelled for any reason.
-	*  @param SourceComponent:           PCG component associated with this task. Can be null.
-	*  @param TaskExecutionDependencies: Task will wait on these tasks to execute and won't take their output data as input.
-	*  @param TaskDataDependencies:      Task will wait on these tasks to execute and will take their output data as input.
-	*/
-	UE_API FPCGTaskId ScheduleGenericWithContext(TFunction<bool(FPCGContext*)> InOperation, TFunction<void(FPCGContext*)> InAbortOperation, UPCGComponent* SourceComponent, const TArray<FPCGTaskId>& TaskExecutionDependencies, const TArray<FPCGTaskId>& TaskDataDependencies);
-
-	/** Cancels currently running generation */
-	UE_API void CancelGeneration(UPCGComponent* Component);
-
-	/** Cancels currently running generation on given graph */
-	UE_API void CancelGeneration(UPCGGraph* Graph);
-
-	/** Asks the runtime generation scheduler to refresh a given GenerateAtRuntime component. ChangeType should be 'GenerationGrid' to perform a full cleanup of PAs and local components. */
-	UE_API void RefreshRuntimeGenComponent(UPCGComponent* RuntimeComponent, EPCGChangeType ChangeType = EPCGChangeType::None);
-
-	/** Asks the runtime generation scheduler to refresh all GenerateAtRuntime components. ChangeType should be 'GenerationGrid' to perform a full cleanup of PAs and local components. */
-	UE_API void RefreshAllRuntimeGenComponents(EPCGChangeType ChangeType = EPCGChangeType::None);
-
-#if WITH_EDITOR
-	/** Refresh all components selected by the filter (runtime generated or otherwise). */
-	UE_API void RefreshAllComponentsFiltered(const TFunction<bool(UPCGComponent*)>& ComponentFilter, EPCGChangeType ChangeType = EPCGChangeType::None);
-#endif
-
-	FPCGRuntimeGenScheduler* GetRuntimeGenScheduler() const { return RuntimeGenScheduler; }
-
-	/** Returns true if there are any tasks for this graph currently scheduled or executing. */
-	UE_API bool IsGraphCurrentlyExecuting(UPCGGraph* Graph);
-
-	/** Returns true if any task is scheduled or executing for any graph */
-	UE_API bool IsAnyGraphCurrentlyExecuting() const;
-
-	/** Cancels everything running */
-	UE_API void CancelAllGeneration();
-
-	/** Gets the output data for a given task */
-	UE_API bool GetOutputData(FPCGTaskId InTaskId, FPCGDataCollection& OutData);
-
-	/** Clears the output data for a given task. Should only be called on tasks with bNeedsManualClear set to true */
-	UE_API void ClearOutputData(FPCGTaskId InTaskId);
-
-	/** Register a new PCG Component or update it, will be added to the octree if it doesn't exists yet. Returns true if it was added/updated. Thread safe */
-	UE_API bool RegisterOrUpdatePCGComponent(UPCGComponent* InComponent, bool bDoActorMapping = true);
-
-	/** In case of BP Actors, we need to remap the old component destroyed by the construction script to the new one. Returns true if re-mapping succeeded. */
-	UE_API bool RemapPCGComponent(const UPCGComponent* OldComponent, UPCGComponent* NewComponent, bool bDoActorMapping);
-
-	/** Unregister a PCG Component, will be removed from the octree. Can force it, if we have a delayed unregister. Thread safe */
-	UE_API void UnregisterPCGComponent(UPCGComponent* InComponent, bool bForce = false);
-
-	/** Register a new Partition actor, will be added to a map and will query all intersecting volume to bind to them if asked. Thread safe */
-	UE_API void RegisterPartitionActor(APCGPartitionActor* InActor);
-
-	/** Unregister a Partition actor, will be removed from the map and remove itself to all intersecting volumes. Thread safe */
-	UE_API void UnregisterPartitionActor(APCGPartitionActor* InActor);
-
-	UE_API TSet<UPCGComponent*> GetAllRegisteredPartitionedComponents() const;
-	UE_API TSet<UPCGComponent*> GetAllRegisteredComponents() const;
-
-	/** Returns the interface to the cache, required for element per-data caching */
-	UE_API IPCGGraphCache* GetCache();
-
-	/** Flushes the graph cache completely, use only for debugging */
-	UE_API void FlushCache();
-
-	/** Call the InFunc function to all local component registered to the original component. Thread safe*/
-	UE_API void ForAllRegisteredLocalComponents(UPCGComponent* InOriginalComponent, const TFunctionRef<void(UPCGComponent*)>& InFunc) const;
-
-	/** Call the InFunc function to all local component registered to the original component within some bounds. Thread safe*/
-	UE_API void ForAllRegisteredIntersectingLocalComponents(UPCGComponent* InOriginalComponent, const FBoxCenterAndExtent& InBounds, const TFunctionRef<void(UPCGComponent*)>& InFunc) const;
-
-	/** Get all components in specified bounds. */
-	UE_API void ForAllIntersectingPartitionedComponents(const FBoxCenterAndExtent& InBounds, TFunctionRef<void(UPCGComponent*)> InFunc) const;
-
-	/** Gather all the PCG components within some bounds. */
-	UE_API TArray<UPCGComponent*> GetAllIntersectingComponents(const FBoxCenterAndExtent& InBounds) const;
-
-	/** Traverses the hierarchy associated with the given component and calls InFunc for each overlapping component. */
-	UE_API void ForAllOverlappingComponentsInHierarchy(UPCGComponent* InComponent, const TFunctionRef<void(UPCGComponent*)>& InFunc) const;
-
-	/**
-	 * Call InFunc to all partition grid cells matching 'InGridSizes' and overlapping with 'InBounds'. 'InFunc' can schedule work or execute immediately.
-	 * 'InGridSizes' should be sorted in descending order. If 'bCanCreateActor' is true, it will create the partition actor at that cell if necessary.
-	 */
-	UE_DEPRECATED(5.5, "Use version with UPCGComponent")
-	FPCGTaskId ForAllOverlappingCells(const FBox& InBounds, const PCGHiGenGrid::FSizeArray& InGridSizes, bool bCanCreateActor, const TArray<FPCGTaskId>& Dependencies, TFunctionRef<FPCGTaskId(APCGPartitionActor*, const FBox&)> InFunc) const { return InvalidPCGTaskId;  }
-
-	UE_API FPCGTaskId ForAllOverlappingCells(
-		UPCGComponent* InPCGComponent,
-		const FBox& InBounds,
-		const PCGHiGenGrid::FSizeArray& InGridSizes,
-		bool bCanCreateActor,
-		const TArray<FPCGTaskId>& Dependencies,
-		TFunctionRef<FPCGTaskId(APCGPartitionActor*, const FBox&)> InFunc,
-		TFunctionRef<FPCGTaskId(const FPCGGridDescriptor&, const FIntVector&, const FBox&)> InUnloadedFunc = [](const FPCGGridDescriptor&, const FIntVector&, const FBox&) { return InvalidPCGTaskId; }) const;
-		
-	/** Immediately cleanup the local components associated with an original component. */
-	UE_API void CleanupLocalComponentsImmediate(UPCGComponent* InOriginalComponent, bool bRemoveComponents);
-
-	UE_DEPRECATED(5.5, "Use FPCGGridDescriptor version")
-	UE_API UPCGComponent* GetLocalComponent(uint32 GridSize, const FIntVector& CellCoords, const UPCGComponent* InOriginalComponent, bool bTransient = false) const;
-
-	UE_DEPRECATED(5.5, "Use FPCGGridDescriptor version")
-	UE_API APCGPartitionActor* GetRegisteredPCGPartitionActor(uint32 GridSize, const FIntVector& GridCoords, bool bRuntimeGenerated = false) const;
-
-	UE_DEPRECATED(5.5, "Use FPCGGridDescriptor verison")
-	UE_API APCGPartitionActor* FindOrCreatePCGPartitionActor(const FGuid& Guid, uint32 GridSize, const FIntVector& GridCoords, bool bRuntimeGenerated, bool bCanCreateActor = true) const;
-
-	/** Retrieves a local component using grid descriptor and grid coordinates, returns nullptr if no such component is found. */
-	UE_API UPCGComponent* GetLocalComponent(const FPCGGridDescriptor& GridDescriptor, const FIntVector& CellCoords, const UPCGComponent* InOriginalComponent) const;
-
-	/** Retrieves a registered partition actor using grid size and grid coordinates, returns nullptr if no such partition actor is found. */
-	UE_API APCGPartitionActor* GetRegisteredPCGPartitionActor(const FPCGGridDescriptor& GridDescriptor, const FIntVector& GridCoords) const;
-
-	/** Creates a new partition actor if one does not already exist with the same grid size, coords, and generation mode. */
-	UE_API APCGPartitionActor* FindOrCreatePCGPartitionActor(const FPCGGridDescriptor& GridDescriptor, const FIntVector& GridCoords, bool bCanCreateActor = true, bool bHideFromOutliner = false) const;
-	
-	/** Retrieves partition actors for this original component */
-	UE_API TSet<TObjectPtr<APCGPartitionActor>> GetPCGComponentPartitionActorMappings(UPCGComponent* InComponent) const;
-
-	/** True if graph cache debugging is enabled. */
-	UE_API bool IsGraphCacheDebuggingEnabled() const;
-
-	UE_API FPCGGenSourceManager* GetGenSourceManager() const;
-	UE_API FPCGGraphCompiler* GetGraphCompiler();
-	UE_API UPCGComputeGraph* GetComputeGraph(const UPCGGraph* InGraph, uint32 GridSize, uint32 ComputeGraphIndex);
-
-#if WITH_EDITOR
-public:
-	/** Schedule refresh on the current or next frame */
-	UE_API FPCGTaskId ScheduleRefresh(UPCGComponent* SourceComponent, bool bForceRefresh);
-
-	/** Immediately dirties the partition actors in the given bounds */
-	UE_API void DirtyGraph(UPCGComponent* Component, const FBox& InBounds, EPCGComponentDirtyFlag DirtyFlag);
-
-	/** Delete serialized partition actors in the level. If 'bOnlyDeleteUnused' is true, only PAs with no graph instances will be deleted. */
-	UE_API void DeleteSerializedPartitionActors(bool bOnlyDeleteUnused, bool bOnlyChildren = false);
-
-	/** Propagate to the graph compiler graph changes */
-	UE_API void NotifyGraphChanged(UPCGGraph* InGraph, EPCGChangeType ChangeType);
-
-	/** Update the tracking on a given component. */
-	UE_API void UpdateComponentTracking(UPCGComponent* InComponent, bool bShouldDirtyActors, const TArray<FPCGSelectionKey>* OptionalChangedKeys = nullptr);
-
-	/** Propagates transient state change from an original component to the relevant partition actors */
-	UE_API void PropagateEditingModeToLocalComponents(UPCGComponent* InOriginalComponent, EPCGEditorDirtyMode EditingMode);
-
-	/** Cleans up the graph cache on an element basis. InSettings is used for debugging and is optional. */
-	UE_API void CleanFromCache(const IPCGElement* InElement, const UPCGSettings* InSettings = nullptr);
-
-	/** Move all resources from sub actors to a new actor */
-	UE_API void ClearPCGLink(UPCGComponent* InComponent, const FBox& InBounds, AActor* InNewActor);
-
-	/** If the partition grid size change, call this to empty the Partition actors map */
-	UE_API void ResetPartitionActorsMap();
-
-	/** Builds the landscape data cache. If force build is true, then it will build the cache even if it is never serialized. */
-	UE_API void BuildLandscapeCache(bool bQuiet = false, bool bForceBuild = true);
-
-	/** Clears the landscape data cache */
-	UE_API void ClearLandscapeCache();
-
-	/** Gets the execution stack information for the given component (depending on partitioning, grid size, etc.) but with no component frames. */
-	UE_API bool GetStackContext(const UPCGComponent* InComponent, FPCGStackContext& OutStackContext);
-
-	/** Gets the base execution stack information for a specific graph & grid size. */
-	UE_API bool GetStackContext(UPCGGraph* InGraph, uint32 GridSize, bool bIsPartitioned, FPCGStackContext& OutStackContext);
-
-	/** Returns how many times InElement is present in the cache. */
-	UE_API uint32 GetGraphCacheEntryCount(IPCGElement* InElement) const;
-
-	/** Will gather all the components registered, and ask for generate. */
-	UE_API void GenerateAllPCGComponents(bool bForce) const;
-
-	/** Will gather all the components registered, and ask for cleanup. */
-	UE_API void CleanupAllPCGComponents(bool bPurge) const;
-
-	/** Get graph warnings and errors for all nodes. */
-	UE_DEPRECATED(5.6, "Use IPCGEditorModule::GetNodeVisualLogs instead")
-	UE_API const FPCGNodeVisualLogs& GetNodeVisualLogs() const;
-	
-	UE_DEPRECATED(5.6, "Use IPCGEditorModule::GetNodeVisualLogsMutable instead")
-	UE_API FPCGNodeVisualLogs& GetNodeVisualLogsMutable();
-
-	/** Notify that we exited the Landscape edit mode. */
-	UE_API void NotifyLandscapeEditModeExited();
-
-	UE_DEPRECATED(5.6, "No longer supported")
-	UE_API void ClearExecutionMetadata(UPCGComponent* InComponent);
-
-	UE_DEPRECATED(5.6, "No longer supported")
-	UE_API void ClearExecutionMetadata(const FPCGStack& BaseStack);
-
-	UE_DEPRECATED(5.6, "No longer supported")
-	UE_API void ClearExecutedStacks(const UPCGComponent* InRootComponent);
-
-	UE_DEPRECATED(5.6, "No longer supported")
-	UE_API void ClearExecutedStacks(const UPCGGraph* InContainingGraph);
-
-	UE_DEPRECATED(5.6, "No longer supported")
-	UE_API TArray<FPCGStack> GetExecutedStacks(const UPCGComponent* InComponent, const UPCGGraph* InSubgraph, bool bOnlyWithSubgraphAsCurrentFrame = true);
-
-	UE_DEPRECATED(5.6, "No longer supported")
-	UE_API TArray<FPCGStack> GetExecutedStacks(const FPCGStack& BeginningWithStack);
-
-	UE_DEPRECATED(5.6, "No longer supported")
-	UE_API void ClearExecutedStacks(const FPCGStack& BeginningWithStack);
-
-	UE_DEPRECATED(5.5, "Deprecated in favor of OnPCGComponentUnregistered, will not be notified anymore")
-	FPCGOnComponentUnregistered OnComponentUnregistered;
-
-	UE_DEPRECATED(5.5, "Deprecated in favor of OnPCGComponentGenerationDone, will not be notified anymore")
-	FPCGOnComponentGenerationCompleteOrCancelled OnComponentGenerationCompleteOrCancelled;
-
-	FPCGOnPCGComponentUnregistered OnPCGComponentUnregistered;
-	FPCGOnPCGComponentGenerationDone OnPCGComponentGenerationDone;
-
-	UE_API void CreateMissingPartitionActors();
-private:
-	enum class EOperation : uint32
-	{
-		Partition,
-		Unpartition,
-		Generate
-	};
-
-	UE_API void OnPCGGraphCancelled(UPCGComponent* InComponent);
-	UE_API void OnPCGGraphStartGenerating(UPCGComponent* InComponent);
-	UE_API void OnPCGGraphGenerated(UPCGComponent* InComponent);
-	UE_API void OnPCGGraphCleaned(UPCGComponent* InComponent);
-
-	UE_API void CreatePartitionActorsWithinBounds(UPCGComponent* InComponent, const FBox& InBounds, const PCGHiGenGrid::FSizeArray& InGridSizes);
-	UE_API void UpdateMappingPCGComponentPartitionActor(UPCGComponent* InComponent);
-#endif // WITH_EDITOR
-	
-private:
-	/** When registering PAs, we might not have a PCG World Actor already registered (at runtime). In that case we'll look for a World Actor in the same level.*/
-	UE_API APCGWorldActor* GetPCGWorldActorForPartitionActor(APCGPartitionActor* InActor);
-
-	UE_API void CancelGeneration(UPCGComponent* Component, bool bCleanupUnusedResources);
-
-	UE_API UPCGData* GetPCGData(FPCGTaskId InGraphExecutionTaskId);
-	UE_API UPCGData* GetInputPCGData(FPCGTaskId InGraphExecutionTaskId);
-	UE_API UPCGData* GetActorPCGData(FPCGTaskId InGraphExecutionTaskId);
-	UE_API UPCGData* GetLandscapePCGData(FPCGTaskId InGraphExecutionTaskId);
-	UE_API UPCGData* GetLandscapeHeightPCGData(FPCGTaskId InGraphExecutionTaskId);
-	UE_API UPCGData* GetOriginalActorPCGData(FPCGTaskId InGraphExecutionTaskId);
-
-	UE_API void SetPCGData(FPCGTaskId InGraphExecutionTaskId, UPCGData* InData);
-	UE_API void SetInputPCGData(FPCGTaskId InGraphExecutionTaskId, UPCGData* InData);
-	UE_API void SetActorPCGData(FPCGTaskId InGraphExecutionTaskId, UPCGData* InData);
-	UE_API void SetLandscapePCGData(FPCGTaskId InGraphExecutionTaskId, UPCGData* InData);
-	UE_API void SetLandscapeHeightPCGData(FPCGTaskId InGraphExecutionTaskId, UPCGData* InData);
-	UE_API void SetOriginalActorPCGData(FPCGTaskId InGraphExecutionTaskId, UPCGData* InData);
-
-	UE_API void ExecuteBeginTickActions();
-
-	APCGWorldActor* PCGWorldActor = nullptr;
-	TSharedPtr<FPCGGraphExecutor> GraphExecutor;
-	FPCGRuntimeGenScheduler* RuntimeGenScheduler = nullptr;
-	bool bHasTickedOnce = false;
-	TUniquePtr<FPCGActorAndComponentMapping> ActorAndComponentMapping;
-
-	/** Functions will be executed at the beginning of the tick and then removed from this array. */
-	TArray<FTickAction> BeginTickActions;
-
-	FCriticalSection PCGWorldActorLock;
-
-#if WITH_EDITOR
-	using FConstructionScriptSourceComponents = TMap<FName, TObjectKey<UPCGComponent>>;
-	TMap<TObjectKey<AActor>, FConstructionScriptSourceComponents> PerActorConstructionScriptSourceComponents;
-
-	static UE_API TSet<UWorld*> DisablePartitionActorCreationForWorld;
-
-	// Used by UPCGWorldPartitonBuilder to disable PA creation while outside of a certain scope
-	static void SetDisablePartitionActorCreationForWorld(UWorld* InWorld, bool bDisable) 
-	{ 
-		if (bDisable)
-		{
-			DisablePartitionActorCreationForWorld.Add(InWorld);
-		}
-		else
-		{
-			DisablePartitionActorCreationForWorld.Remove(InWorld);
-		}
-	}
-
-	static bool IsPartitionActorCreationDisabledForWorld(UWorld* InWorld)
-	{
-		return DisablePartitionActorCreationForWorld.Contains(InWorld);
-	}
-
-	UE_API void SetDisableClearResults(bool bInDisableClearResults);
-#endif
-};
-
-#undef UE_API
+#include "Subsystems/PCGSubsystem.h"
\ No newline at end of file
diff --git a/Engine/Plugins/PCG/Source/PCG/Public/Subsystems/IPCGBaseSubsystem.h b/Engine/Plugins/PCG/Source/PCG/Public/Subsystems/IPCGBaseSubsystem.h
new file mode 100644
--- /dev/null
+++ b/Engine/Plugins/PCG/Source/PCG/Public/Subsystems/IPCGBaseSubsystem.h
@@ -0,0 +1,127 @@
+﻿// Copyright Epic Games, Inc. All Rights Reserved.
+
+#pragma once
+
+#include "Subsystems/WorldSubsystem.h"
+
+#include "PCGCommon.h"
+#include "Elements/PCGActorSelector.h"
+#include "Grid/PCGComponentOctree.h"
+#include "Utils/PCGNodeVisualLogs.h"
+#include "UObject/Interface.h"
+
+#include "UObject/ObjectKey.h"
+
+class UPCGComputeGraph;
+class UPCGData;
+class UPCGGraph;
+
+class IPCGBaseSubsystem;
+class IPCGElement;
+class IPCGGraphCache;
+class FPCGGraphCompiler;
+class FPCGGraphExecutor;
+struct FPCGDataCollection;
+struct FPCGScheduleGenericParams;
+struct FPCGStack;
+class UPCGSettings;
+class UWorld;
+
+#if WITH_EDITOR
+DECLARE_MULTICAST_DELEGATE_ThreeParams(FPCGOnPCGSourceGenerationDone, IPCGBaseSubsystem*, IPCGGraphExecutionSource*, EPCGGenerationStatus);
+#endif // WITH_EDITOR
+
+class IPCGBaseSubsystem
+{
+
+public:
+	PCG_API virtual ~IPCGBaseSubsystem() = default;
+
+	PCG_API virtual UWorld* GetSubsystemWorld() const { return nullptr; }
+
+	PCG_API void InitializeBaseSubsystem();
+	PCG_API void DeinitializeBaseSubsystem();
+	
+	/** Subsystem must not be used without this condition being true. */
+	// @todo_pcg: Hides function in WorldSubsystem
+	//bool IsInitialized() const { return GraphExecutor != nullptr; }
+
+	/** Called by graph executor when a graph is scheduled. */
+	PCG_API void OnScheduleGraph(const FPCGStackContext& StackContext);
+
+	// Schedule graph (used internally for dynamic subgraph execution)
+	PCG_API FPCGTaskId ScheduleGraph(const FPCGScheduleGraphParams& InParams);
+
+	// General job scheduling
+	PCG_API FPCGTaskId ScheduleGeneric(const FPCGScheduleGenericParams& InParams);
+	
+	/** Cancels currently running generation */
+	PCG_API void CancelGeneration(IPCGGraphExecutionSource* Source);
+
+	/** Cancels currently running generation on given graph */
+	PCG_API void CancelGeneration(UPCGGraph* Graph);
+
+	/** Returns true if there are any tasks for this graph currently scheduled or executing. */
+	PCG_API bool IsGraphCurrentlyExecuting(UPCGGraph* Graph);
+
+	/** Returns true if any task is scheduled or executing for any graph */
+	PCG_API bool IsAnyGraphCurrentlyExecuting() const;
+
+	/** Cancels everything running */
+	PCG_API void CancelAllGeneration();
+
+	/** Gets the output data for a given task */
+	PCG_API bool GetOutputData(FPCGTaskId InTaskId, FPCGDataCollection& OutData);
+
+	/** Clears the output data for a given task. Should only be called on tasks with bNeedsManualClear set to true */
+	PCG_API void ClearOutputData(FPCGTaskId InTaskId);
+
+	/** Returns the interface to the cache, required for element per-data caching */
+	PCG_API IPCGGraphCache* GetCache();
+
+	/** Flushes the graph cache completely, use only for debugging */
+	PCG_API void FlushCache();
+
+	/** True if graph cache debugging is enabled. */
+	PCG_API bool IsGraphCacheDebuggingEnabled() const;
+	
+	PCG_API FPCGGraphCompiler* GetGraphCompiler();
+	PCG_API UPCGComputeGraph* GetComputeGraph(const UPCGGraph* InGraph, uint32 GridSize, uint32 ComputeGraphIndex);
+
+#if WITH_EDITOR
+public:
+	/** Propagate to the graph compiler graph changes */
+	PCG_API void NotifyGraphChanged(UPCGGraph* InGraph, EPCGChangeType ChangeType);
+
+	/** Cleans up the graph cache on an element basis. InSettings is used for debugging and is optional. */
+	PCG_API void CleanFromCache(const IPCGElement* InElement, const UPCGSettings* InSettings = nullptr);
+
+	/** Gets the execution stack information for the given component (depending on partitioning, grid size, etc.) but with no component frames. */
+	PCG_API bool GetStackContext(const IPCGGraphExecutionSource* InSource, FPCGStackContext& OutStackContext);
+
+	/** Gets the base execution stack information for a specific graph & grid size. */
+	PCG_API bool GetStackContext(UPCGGraph* InGraph, uint32 GridSize, bool bIsPartitioned, FPCGStackContext& OutStackContext);
+
+	/** Returns how many times InElement is present in the cache. */
+	PCG_API uint32 GetGraphCacheEntryCount(IPCGElement* InElement) const;
+	
+	FPCGOnPCGSourceGenerationDone OnPCGSourceGenerationDone;
+#endif // WITH_EDITOR
+	
+protected:
+	/** Returns the expected end time of the tick. By default, it's the current time + the budget of the graph executor. */
+	PCG_API virtual double GetTickEndTime() const;
+
+	/** Tick the graph executor. Return expected end time. */
+	PCG_API double Tick();
+	
+	PCG_API void CancelGeneration(IPCGGraphExecutionSource* Source, bool bCleanupUnusedResources);
+	
+	PCG_API virtual void CancelGenerationInternal(IPCGGraphExecutionSource* Source, bool bCleanupUnusedResources) {};
+	
+	TSharedPtr<FPCGGraphExecutor> GraphExecutor;
+
+#if WITH_EDITOR
+	PCG_API void SetDisableClearResults(bool bInDisableClearResults);
+#endif
+};
diff --git a/Engine/Plugins/PCG/Source/PCG/Public/Subsystems/PCGSubsystem.h b/Engine/Plugins/PCG/Source/PCG/Public/Subsystems/PCGSubsystem.h
new file mode 100644
--- /dev/null
+++ b/Engine/Plugins/PCG/Source/PCG/Public/Subsystems/PCGSubsystem.h
@@ -0,0 +1,424 @@
+// Copyright Epic Games, Inc. All Rights Reserved.
+
+#pragma once
+
+#include "Subsystems/IPCGBaseSubsystem.h"
+#include "Subsystems/WorldSubsystem.h"
+
+#include "PCGCommon.h"
+#include "Elements/PCGActorSelector.h"
+#include "Grid/PCGGridDescriptor.h"
+#include "Grid/PCGComponentOctree.h"
+#include "Utils/PCGNodeVisualLogs.h"
+
+#include "UObject/ObjectKey.h"
+
+#include "PCGSubsystem.generated.h"
+
+#define UE_API PCG_API
+
+class APCGPartitionActor;
+class APCGWorldActor;
+class FPCGActorAndComponentMapping;
+class FPCGGenSourceManager;
+class FPCGRuntimeGenScheduler;
+class UPCGComputeGraph;
+class UPCGData;
+class UPCGGraph;
+class UPCGLandscapeCache;
+
+enum class EPCGComponentDirtyFlag : uint8;
+enum class ETickableTickType : uint8;
+
+class IPCGGraphCache;
+class FPCGGraphCompiler;
+class FPCGGraphExecutor;
+struct FPCGContext;
+struct FPCGDataCollection;
+struct FPCGScheduleGenericParams;
+struct FPCGStack;
+class UPCGSettings;
+
+class IPCGElement;
+typedef TSharedPtr<IPCGElement, ESPMode::ThreadSafe> FPCGElementPtr;
+
+class UWorld;
+
+#if WITH_EDITOR
+/** Deprecated - use FPCGOnPCGComponentUnregistered */
+DECLARE_MULTICAST_DELEGATE(FPCGOnComponentUnregistered);
+/** Deprecated - use FPCGOnPCGSourceGenerationDone */
+DECLARE_MULTICAST_DELEGATE_OneParam(FPCGOnComponentGenerationCompleteOrCancelled, UPCGSubsystem*);
+
+DECLARE_MULTICAST_DELEGATE_OneParam(FPCGOnPCGComponentUnregistered, UPCGComponent*);
+
+/** Deprecated - use FPCGOnPCGSourceGenerationDone */
+DECLARE_MULTICAST_DELEGATE_ThreeParams(FPCGOnPCGComponentGenerationDone, UPCGSubsystem*, UPCGComponent*, EPCGGenerationStatus);
+#endif // WITH_EDITOR
+
+/**
+* UPCGSubsystem
+*/
+UCLASS(MinimalAPI)
+class UPCGSubsystem : public UTickableWorldSubsystem, public IPCGBaseSubsystem
+{
+	GENERATED_BODY()
+
+public:
+	friend class UPCGComponent;
+	friend FPCGActorAndComponentMapping;
+	friend struct FPCGWorldPartitionBuilder;
+
+	UE_API UPCGSubsystem();
+
+	/** Add UObject references for GC */
+	static PCG_API void AddReferencedObjects(UObject* InThis, FReferenceCollector& Collector);
+
+	/** To be used when a PCG component can not have a world anymore, to unregister itself. */
+	static UE_API UPCGSubsystem* GetSubsystemForCurrentWorld();
+
+	virtual PCG_API UWorld* GetSubsystemWorld() const override { return GetWorld(); }
+
+	//~ Begin USubsystem Interface.
+	UE_API virtual void Deinitialize() override;
+	//~ End USubsystem Interface
+
+	//~ Begin UWorldSubsystem Interface.
+	UE_API virtual void PostInitialize() override;
+	// need UpdateStreamingState? 
+	//~ End UWorldSubsystem Interface.
+
+	//~ Begin FTickableGameObject
+	UE_API virtual void Tick(float DeltaSeconds) override;
+	virtual bool IsTickableInEditor() const override { return true; }
+	UE_API virtual ETickableTickType GetTickableTickType() const override;
+	UE_API virtual TStatId GetStatId() const override;
+	//~ End FTickableGameObject
+
+	/** Will return the subsystem from the World if it exists and if it is initialized */
+	static UE_API UPCGSubsystem* GetInstance(UWorld* World);
+
+	/** Adds an action that will be executed once at the beginning of this subsystem's next Tick(). */
+	using FTickAction = TFunction<void()>;
+
+	UE_DEPRECATED(5.6, "Use FPCGModule::ExecuteNextTick instead")
+	UE_API void RegisterBeginTickAction(FTickAction&& Action);
+
+#if WITH_EDITOR
+	/** Returns PIE world if it is active, otherwise returns editor world. */
+	static UE_API UPCGSubsystem* GetActiveEditorInstance();
+
+	UE_API void SetConstructionScriptSourceComponent(UPCGComponent* InComponent);
+	UE_API bool RemoveAndCopyConstructionScriptSourceComponent(AActor* InComponentOwner, FName InComponentName, UPCGComponent*& OutSourceComponent);
+#endif
+
+	UE_API APCGWorldActor* GetPCGWorldActor();
+	UE_API APCGWorldActor* FindPCGWorldActor();
+
+	/** Returns current quality level between Low (0) and Cinematic (4). */
+	static UE_API int32 GetPCGQualityLevel();
+	UE_API void OnPCGQualityLevelChanged();
+
+#if WITH_EDITOR
+	UE_API void DestroyAllPCGWorldActors();
+	UE_API void DestroyCurrentPCGWorldActor();
+	UE_API void LogAbnormalComponentStates(bool bGroupByState) const;
+#endif
+	UE_API void RegisterPCGWorldActor(APCGWorldActor* InActor);
+	UE_API void UnregisterPCGWorldActor(APCGWorldActor* InActor);
+
+	UE_API void OnOriginalComponentRegistered(UPCGComponent* InComponent);
+	UE_API void OnOriginalComponentUnregistered(UPCGComponent* InComponent);
+
+	UE_API UPCGLandscapeCache* GetLandscapeCache();
+
+	/** Schedule graph(owner->graph) */
+	UE_API FPCGTaskId ScheduleComponent(UPCGComponent* PCGComponent, EPCGHiGenGrid Grid, bool bForce, const TArray<FPCGTaskId>& InDependencies);
+
+	/** Schedule cleanup(owner->graph). */
+	UE_API FPCGTaskId ScheduleCleanup(UPCGComponent* PCGComponent, bool bRemoveComponents, const TArray<FPCGTaskId>& Dependencies);
+
+	using IPCGBaseSubsystem::ScheduleGraph;
+	
+	UE_API FPCGTaskId ScheduleGraph(
+		UPCGGraph* Graph,
+		UPCGComponent* SourceComponent,
+		FPCGElementPtr PreGraphElement,
+		FPCGElementPtr InputElement,
+		const TArray<FPCGTaskId>& Dependencies,
+		const FPCGStack* InFromStack,
+		bool bAllowHierarchicalGeneration);
+
+	// Schedule graph (used internally for dynamic subgraph execution)
+	UE_API FPCGTaskId ScheduleGraph(UPCGComponent* SourceComponent, const TArray<FPCGTaskId>& Dependencies);
+
+	using IPCGBaseSubsystem::ScheduleGeneric;
+
+	/** General job scheduling
+	*  @param InOperation:               Callback that returns true if the task is done, false otherwise.
+	*  @param SourceComponent:           PCG component associated with this task. Can be null.
+	*  @param TaskExecutionDependencies: Task will wait on these tasks to execute and won't take their output data as input.
+	*/
+	UE_API FPCGTaskId ScheduleGeneric(TFunction<bool()> InOperation, UPCGComponent* SourceComponent, const TArray<FPCGTaskId>& TaskExecutionDependencies);
+
+	/** General job scheduling
+	*  @param InOperation:               Callback that returns true if the task is done, false otherwise.
+	*  @param InAbortOperation:          Callback that will be called if the generic task is cancelled for any reason.
+	*  @param SourceComponent:           PCG component associated with this task. Can be null.
+	*  @param TaskExecutionDependencies: Task will wait on these tasks to execute and won't take their output data as input.
+	*/
+	UE_API FPCGTaskId ScheduleGeneric(TFunction<bool()> InOperation, TFunction<void()> InAbortOperation, UPCGComponent* SourceComponent, const TArray<FPCGTaskId>& TaskExecutionDependencies);
+
+	/** General job scheduling with context
+	*  @param InOperation:               Callback that takes a Context as argument and returns true if the task is done, false otherwise.
+	*  @param SourceComponent:           PCG component associated with this task. Can be null.
+	*  @param TaskExecutionDependencies: Task will wait on these tasks to execute and won't take their output data as input.
+	*  @param TaskDataDependencies:      Task will wait on these tasks to execute and will take their output data as input.
+	*/
+	UE_API FPCGTaskId ScheduleGenericWithContext(TFunction<bool(FPCGContext*)> InOperation, UPCGComponent* SourceComponent, const TArray<FPCGTaskId>& TaskExecutionDependencies, const TArray<FPCGTaskId>& TaskDataDependencies);
+
+	/** General job scheduling with context
+	*  @param InOperation:               Callback that takes a Context as argument and returns true if the task is done, false otherwise.
+	*  @param InAbortOperation:          Callback that will be called if the generic task is cancelled for any reason.
+	*  @param SourceComponent:           PCG component associated with this task. Can be null.
+	*  @param TaskExecutionDependencies: Task will wait on these tasks to execute and won't take their output data as input.
+	*  @param TaskDataDependencies:      Task will wait on these tasks to execute and will take their output data as input.
+	*/
+	UE_API FPCGTaskId ScheduleGenericWithContext(TFunction<bool(FPCGContext*)> InOperation, TFunction<void(FPCGContext*)> InAbortOperation, UPCGComponent* SourceComponent, const TArray<FPCGTaskId>& TaskExecutionDependencies, const TArray<FPCGTaskId>& TaskDataDependencies);
+
+	/** Asks the runtime generation scheduler to refresh a given GenerateAtRuntime component. ChangeType should be 'GenerationGrid' to perform a full cleanup of PAs and local components. */
+	UE_API void RefreshRuntimeGenComponent(UPCGComponent* RuntimeComponent, EPCGChangeType ChangeType = EPCGChangeType::None);
+
+	/** Asks the runtime generation scheduler to refresh all GenerateAtRuntime components. ChangeType should be 'GenerationGrid' to perform a full cleanup of PAs and local components. */
+	UE_API void RefreshAllRuntimeGenComponents(EPCGChangeType ChangeType = EPCGChangeType::None);
+
+#if WITH_EDITOR
+	/** Refresh all components selected by the filter (runtime generated or otherwise). */
+	UE_API void RefreshAllComponentsFiltered(const TFunction<bool(UPCGComponent*)>& ComponentFilter, EPCGChangeType ChangeType = EPCGChangeType::None);
+#endif
+
+	FPCGRuntimeGenScheduler* GetRuntimeGenScheduler() const { return RuntimeGenScheduler; }
+
+	/** Register a new PCG Component or update it, will be added to the octree if it doesn't exists yet. Returns true if it was added/updated. Thread safe */
+	UE_API bool RegisterOrUpdatePCGComponent(UPCGComponent* InComponent, bool bDoActorMapping = true);
+
+	/** In case of BP Actors, we need to remap the old component destroyed by the construction script to the new one. Returns true if re-mapping succeeded. */
+	UE_API bool RemapPCGComponent(const UPCGComponent* OldComponent, UPCGComponent* NewComponent, bool bDoActorMapping);
+
+	/** Unregister a PCG Component, will be removed from the octree. Can force it, if we have a delayed unregister. Thread safe */
+	UE_API void UnregisterPCGComponent(UPCGComponent* InComponent, bool bForce = false);
+
+	/** Register a new Partition actor, will be added to a map and will query all intersecting volume to bind to them if asked. Thread safe */
+	UE_API void RegisterPartitionActor(APCGPartitionActor* InActor);
+
+	/** Unregister a Partition actor, will be removed from the map and remove itself to all intersecting volumes. Thread safe */
+	UE_API void UnregisterPartitionActor(APCGPartitionActor* InActor);
+
+	UE_API TSet<UPCGComponent*> GetAllRegisteredPartitionedComponents() const;
+	UE_API TSet<UPCGComponent*> GetAllRegisteredComponents() const;
+
+	/** Call the InFunc function to all local component registered to the original component. Thread safe*/
+	UE_API void ForAllRegisteredLocalComponents(UPCGComponent* InOriginalComponent, const TFunctionRef<void(UPCGComponent*)>& InFunc) const;
+
+	/** Call the InFunc function to all local component registered to the original component within some bounds. Thread safe*/
+	UE_API void ForAllRegisteredIntersectingLocalComponents(UPCGComponent* InOriginalComponent, const FBoxCenterAndExtent& InBounds, const TFunctionRef<void(UPCGComponent*)>& InFunc) const;
+
+	/** Get all components in specified bounds. */
+	UE_API void ForAllIntersectingPartitionedComponents(const FBoxCenterAndExtent& InBounds, TFunctionRef<void(UPCGComponent*)> InFunc) const;
+
+	/** Gather all the PCG components within some bounds. */
+	UE_API TArray<UPCGComponent*> GetAllIntersectingComponents(const FBoxCenterAndExtent& InBounds) const;
+
+	/** Traverses the hierarchy associated with the given component and calls InFunc for each overlapping component. */
+	UE_API void ForAllOverlappingComponentsInHierarchy(UPCGComponent* InComponent, const TFunctionRef<void(UPCGComponent*)>& InFunc) const;
+
+	/**
+	 * Call InFunc to all partition grid cells matching 'InGridSizes' and overlapping with 'InBounds'. 'InFunc' can schedule work or execute immediately.
+	 * 'InGridSizes' should be sorted in descending order. If 'bCanCreateActor' is true, it will create the partition actor at that cell if necessary.
+	 */
+	UE_DEPRECATED(5.5, "Use version with UPCGComponent")
+	FPCGTaskId ForAllOverlappingCells(const FBox& InBounds, const PCGHiGenGrid::FSizeArray& InGridSizes, bool bCanCreateActor, const TArray<FPCGTaskId>& Dependencies, TFunctionRef<FPCGTaskId(APCGPartitionActor*, const FBox&)> InFunc) const { return InvalidPCGTaskId;  }
+
+	UE_API FPCGTaskId ForAllOverlappingCells(
+		UPCGComponent* InPCGComponent,
+		const FBox& InBounds,
+		const PCGHiGenGrid::FSizeArray& InGridSizes,
+		bool bCanCreateActor,
+		const TArray<FPCGTaskId>& Dependencies,
+		TFunctionRef<FPCGTaskId(APCGPartitionActor*, const FBox&)> InFunc,
+		TFunctionRef<FPCGTaskId(const FPCGGridDescriptor&, const FIntVector&, const FBox&)> InUnloadedFunc = [](const FPCGGridDescriptor&, const FIntVector&, const FBox&) { return InvalidPCGTaskId; }) const;
+		
+	/** Immediately cleanup the local components associated with an original component. */
+	UE_API void CleanupLocalComponentsImmediate(UPCGComponent* InOriginalComponent, bool bRemoveComponents);
+
+	UE_DEPRECATED(5.5, "Use FPCGGridDescriptor version")
+	UE_API UPCGComponent* GetLocalComponent(uint32 GridSize, const FIntVector& CellCoords, const UPCGComponent* InOriginalComponent, bool bTransient = false) const;
+
+	UE_DEPRECATED(5.5, "Use FPCGGridDescriptor version")
+	UE_API APCGPartitionActor* GetRegisteredPCGPartitionActor(uint32 GridSize, const FIntVector& GridCoords, bool bRuntimeGenerated = false) const;
+
+	UE_DEPRECATED(5.5, "Use FPCGGridDescriptor verison")
+	UE_API APCGPartitionActor* FindOrCreatePCGPartitionActor(const FGuid& Guid, uint32 GridSize, const FIntVector& GridCoords, bool bRuntimeGenerated, bool bCanCreateActor = true) const;
+
+	/** Retrieves a local component using grid descriptor and grid coordinates, returns nullptr if no such component is found. */
+	UE_API UPCGComponent* GetLocalComponent(const FPCGGridDescriptor& GridDescriptor, const FIntVector& CellCoords, const UPCGComponent* InOriginalComponent) const;
+
+	/** Retrieves a registered partition actor using grid size and grid coordinates, returns nullptr if no such partition actor is found. */
+	UE_API APCGPartitionActor* GetRegisteredPCGPartitionActor(const FPCGGridDescriptor& GridDescriptor, const FIntVector& GridCoords) const;
+
+	/** Creates a new partition actor if one does not already exist with the same grid size, coords, and generation mode. */
+	UE_API APCGPartitionActor* FindOrCreatePCGPartitionActor(const FPCGGridDescriptor& GridDescriptor, const FIntVector& GridCoords, bool bCanCreateActor = true, bool bHideFromOutliner = false) const;
+	
+	/** Retrieves partition actors for this original component */
+	UE_API TSet<TObjectPtr<APCGPartitionActor>> GetPCGComponentPartitionActorMappings(UPCGComponent* InComponent) const;
+
+	UE_API FPCGGenSourceManager* GetGenSourceManager() const;
+
+#if WITH_EDITOR
+public:
+	/** Schedule refresh on the current or next frame */
+	UE_API FPCGTaskId ScheduleRefresh(UPCGComponent* SourceComponent, bool bForceRefresh);
+
+	/** Immediately dirties the partition actors in the given bounds */
+	UE_API void DirtyGraph(UPCGComponent* Component, const FBox& InBounds, EPCGComponentDirtyFlag DirtyFlag);
+
+	/** Delete serialized partition actors in the level. If 'bOnlyDeleteUnused' is true, only PAs with no graph instances will be deleted. */
+	UE_API void DeleteSerializedPartitionActors(bool bOnlyDeleteUnused, bool bOnlyChildren = false);
+
+	/** Update the tracking on a given component. */
+	UE_API void UpdateComponentTracking(UPCGComponent* InComponent, bool bShouldDirtyActors, const TArray<FPCGSelectionKey>* OptionalChangedKeys = nullptr);
+
+	/** Propagates transient state change from an original component to the relevant partition actors */
+	UE_API void PropagateEditingModeToLocalComponents(UPCGComponent* InOriginalComponent, EPCGEditorDirtyMode EditingMode);
+
+	/** Move all resources from sub actors to a new actor */
+	UE_API void ClearPCGLink(UPCGComponent* InComponent, const FBox& InBounds, AActor* InNewActor);
+
+	/** If the partition grid size change, call this to empty the Partition actors map */
+	UE_API void ResetPartitionActorsMap();
+
+	/** Builds the landscape data cache. If force build is true, then it will build the cache even if it is never serialized. */
+	UE_API void BuildLandscapeCache(bool bQuiet = false, bool bForceBuild = true);
+
+	/** Clears the landscape data cache */
+	UE_API void ClearLandscapeCache();
+
+	/** Will gather all the components registered, and ask for generate. */
+	UE_API void GenerateAllPCGComponents(bool bForce) const;
+
+	/** Will gather all the components registered, and ask for cleanup. */
+	UE_API void CleanupAllPCGComponents(bool bPurge) const;
+
+	/** Get graph warnings and errors for all nodes. */
+	UE_DEPRECATED(5.6, "Use IPCGEditorModule::GetNodeVisualLogs instead")
+	UE_API const FPCGNodeVisualLogs& GetNodeVisualLogs() const;
+	
+	UE_DEPRECATED(5.6, "Use IPCGEditorModule::GetNodeVisualLogsMutable instead")
+	UE_API FPCGNodeVisualLogs& GetNodeVisualLogsMutable();
+
+	/** Notify that we exited the Landscape edit mode. */
+	UE_API void NotifyLandscapeEditModeExited();
+
+	UE_DEPRECATED(5.6, "No longer supported")
+	UE_API void ClearExecutionMetadata(UPCGComponent* InComponent);
+
+	UE_DEPRECATED(5.6, "No longer supported")
+	UE_API void ClearExecutionMetadata(const FPCGStack& BaseStack);
+
+	UE_DEPRECATED(5.6, "No longer supported")
+	UE_API void ClearExecutedStacks(const UPCGComponent* InRootComponent);
+
+	UE_DEPRECATED(5.6, "No longer supported")
+	UE_API void ClearExecutedStacks(const UPCGGraph* InContainingGraph);
+
+	UE_DEPRECATED(5.6, "No longer supported")
+	UE_API TArray<FPCGStack> GetExecutedStacks(const UPCGComponent* InComponent, const UPCGGraph* InSubgraph, bool bOnlyWithSubgraphAsCurrentFrame = true);
+
+	UE_DEPRECATED(5.6, "No longer supported")
+	UE_API TArray<FPCGStack> GetExecutedStacks(const FPCGStack& BeginningWithStack);
+
+	UE_DEPRECATED(5.6, "No longer supported")
+	UE_API void ClearExecutedStacks(const FPCGStack& BeginningWithStack);
+
+	UE_DEPRECATED(5.5, "Deprecated in favor of OnPCGComponentUnregistered, will not be notified anymore")
+	FPCGOnComponentUnregistered OnComponentUnregistered;
+
+	UE_DEPRECATED(5.5, "Deprecated in favor of OnPCGComponentGenerationDone, will not be notified anymore")
+	FPCGOnComponentGenerationCompleteOrCancelled OnComponentGenerationCompleteOrCancelled;
+
+	FPCGOnPCGComponentUnregistered OnPCGComponentUnregistered;
+	
+	UE_DEPRECATED(5.7, "Deprecated in favor of OnPCGSourceGenerationDone, will not be notified anymore")
+	FPCGOnPCGComponentGenerationDone OnPCGComponentGenerationDone;
+
+	UE_API void CreateMissingPartitionActors();
+private:
+	UE_API void OnPCGGraphCancelled(UPCGComponent* InComponent);
+	UE_API void OnPCGGraphStartGenerating(UPCGComponent* InComponent);
+	UE_API void OnPCGGraphGenerated(UPCGComponent* InComponent);
+	UE_API void OnPCGGraphCleaned(UPCGComponent* InComponent);
+
+	UE_API void CreatePartitionActorsWithinBounds(UPCGComponent* InComponent, const FBox& InBounds, const PCGHiGenGrid::FSizeArray& InGridSizes);
+	UE_API void UpdateMappingPCGComponentPartitionActor(UPCGComponent* InComponent);
+#endif // WITH_EDITOR
+	
+private:
+	/** When registering PAs, we might not have a PCG World Actor already registered (at runtime). In that case we'll look for a World Actor in the same level.*/
+	UE_API APCGWorldActor* GetPCGWorldActorForPartitionActor(APCGPartitionActor* InActor);
+
+	UE_API UPCGData* GetPCGData(FPCGTaskId InGraphExecutionTaskId);
+	UE_API UPCGData* GetInputPCGData(FPCGTaskId InGraphExecutionTaskId);
+	UE_API UPCGData* GetActorPCGData(FPCGTaskId InGraphExecutionTaskId);
+	UE_API UPCGData* GetLandscapePCGData(FPCGTaskId InGraphExecutionTaskId);
+	UE_API UPCGData* GetLandscapeHeightPCGData(FPCGTaskId InGraphExecutionTaskId);
+	UE_API UPCGData* GetOriginalActorPCGData(FPCGTaskId InGraphExecutionTaskId);
+
+	UE_API void SetPCGData(FPCGTaskId InGraphExecutionTaskId, UPCGData* InData);
+	UE_API void SetInputPCGData(FPCGTaskId InGraphExecutionTaskId, UPCGData* InData);
+	UE_API void SetActorPCGData(FPCGTaskId InGraphExecutionTaskId, UPCGData* InData);
+	UE_API void SetLandscapePCGData(FPCGTaskId InGraphExecutionTaskId, UPCGData* InData);
+	UE_API void SetLandscapeHeightPCGData(FPCGTaskId InGraphExecutionTaskId, UPCGData* InData);
+	UE_API void SetOriginalActorPCGData(FPCGTaskId InGraphExecutionTaskId, UPCGData* InData);
+
+	UE_API void ExecuteBeginTickActions();
+
+	APCGWorldActor* PCGWorldActor = nullptr;
+	FPCGRuntimeGenScheduler* RuntimeGenScheduler = nullptr;
+	bool bHasTickedOnce = false;
+	TUniquePtr<FPCGActorAndComponentMapping> ActorAndComponentMapping;
+
+	/** Functions will be executed at the beginning of the tick and then removed from this array. */
+	TArray<FTickAction> BeginTickActions;
+
+	FCriticalSection PCGWorldActorLock;
+
+#if WITH_EDITOR
+	using FConstructionScriptSourceComponents = TMap<FName, TObjectKey<UPCGComponent>>;
+	TMap<TObjectKey<AActor>, FConstructionScriptSourceComponents> PerActorConstructionScriptSourceComponents;
+
+	static UE_API TSet<UWorld*> DisablePartitionActorCreationForWorld;
+
+	// Used by UPCGWorldPartitonBuilder to disable PA creation while outside of a certain scope
+	static void SetDisablePartitionActorCreationForWorld(UWorld* InWorld, bool bDisable) 
+	{ 
+		if (bDisable)
+		{
+			DisablePartitionActorCreationForWorld.Add(InWorld);
+		}
+		else
+		{
+			DisablePartitionActorCreationForWorld.Remove(InWorld);
+		}
+	}
+
+	static bool IsPartitionActorCreationDisabledForWorld(UWorld* InWorld)
+	{
+		return DisablePartitionActorCreationForWorld.Contains(InWorld);
+	}
+#endif
+
+protected:
+	PCG_API virtual void CancelGenerationInternal(IPCGGraphExecutionSource* Source, bool bCleanupUnusedResources) override;
+};
+
+#undef UE_API
diff --git a/Engine/Plugins/PCG/Source/PCGEditor/Private/PCGEditor.cpp b/Engine/Plugins/PCG/Source/PCGEditor/Private/PCGEditor.cpp
--- a/Engine/Plugins/PCG/Source/PCGEditor/Private/PCGEditor.cpp
+++ b/Engine/Plugins/PCG/Source/PCGEditor/Private/PCGEditor.cpp
@@ -356,7 +356,7 @@ void FPCGEditor::UpdateDebugAfterComponentSelection(UPCGComponent* InOldComponen
 		// GenerateAtRuntime components should be refreshed through the runtime gen scheduler.
 		if (Component->IsManagedByRuntimeGenSystem())
 		{
-			if (UPCGSubsystem* Subsystem = GetSubsystem())
+			if (UPCGSubsystem* Subsystem = GetWorldSubsystem())
 			{
 				// We don't want to do a full cleanup if we're setting the debug object, since full cleanup destroys the component, which is the debug object itself!
 				Subsystem->RefreshRuntimeGenComponent(Component);
@@ -1064,7 +1064,7 @@ void FPCGEditor::OnForceGraphRegeneration_Clicked()
 		FModifierKeysState ModifierKeys = FSlateApplication::Get().GetModifierKeys();
 		if (ModifierKeys.IsControlDown())
 		{
-			if (UPCGSubsystem* Subsystem = GetSubsystem())
+			if (IPCGBaseSubsystem* Subsystem = GetSubsystem())
 			{
 				Subsystem->FlushCache();
 			}
@@ -1080,7 +1080,7 @@ void FPCGEditor::OnForceGraphRegeneration_Clicked()
 
 void FPCGEditor::OnCancelExecution_Clicked()
 {
-	UPCGSubsystem* Subsystem = GetSubsystem();
+	IPCGBaseSubsystem* Subsystem = GetSubsystem();
 	if (PCGEditorGraph && Subsystem)
 	{
 		Subsystem->CancelGeneration(PCGEditorGraph->GetPCGGraph());
@@ -1089,7 +1089,7 @@ void FPCGEditor::OnCancelExecution_Clicked()
 
 bool FPCGEditor::IsCurrentlyGenerating() const
 {
-	UPCGSubsystem* Subsystem = GetSubsystem();
+	IPCGBaseSubsystem* Subsystem = GetSubsystem();
 	if (PCGGraphBeingEdited && Subsystem)
 	{
 		return Subsystem->IsGraphCurrentlyExecuting(PCGGraphBeingEdited);
@@ -3076,8 +3076,11 @@ void FPCGEditor::OnComponentUnregistered(UPCGComponent* Component)
 	}
 }
 
-void FPCGEditor::OnComponentGenerationDone(UPCGSubsystem* Subsystem, UPCGComponent* Component, EPCGGenerationStatus Status)
+void FPCGEditor::OnSourceGenerationDone(IPCGBaseSubsystem* Subsystem, IPCGGraphExecutionSource* Source, EPCGGenerationStatus Status)
 {
+	// @todo_pcg Should support any source
+	UPCGComponent* Component = Cast<UPCGComponent>(Source);
+	
 	// We want to refresh if the component that is done generating has generated the current graph being edited,
 	// or if it is the root of the current stack being inspected (for subgraphs to also be refreshed).
 	// If we don't have a component, we refresh nonetheless.
@@ -3121,7 +3124,12 @@ void FPCGEditor::OnComponentGenerationDone(UPCGSubsystem* Subsystem, UPCGCompone
 	}
 }
 
-UPCGSubsystem* FPCGEditor::GetSubsystem()
+IPCGBaseSubsystem* FPCGEditor::GetSubsystem() const
+{
+	return GetWorldSubsystem();
+}
+
+UPCGSubsystem* FPCGEditor::GetWorldSubsystem()
 {
 	UWorld* World = (GEditor ? (GEditor->PlayWorld ? GEditor->PlayWorld.Get() : GEditor->GetEditorWorldContext().World()) : nullptr);
 	return UPCGSubsystem::GetInstance(World);
@@ -3134,7 +3142,7 @@ void FPCGEditor::RegisterDelegatesForWorld(UWorld* World)
 	if (UPCGSubsystem* Subsystem = UPCGSubsystem::GetInstance(World))
 	{
 		Subsystem->OnPCGComponentUnregistered.AddRaw(this, &FPCGEditor::OnComponentUnregistered);
-		Subsystem->OnPCGComponentGenerationDone.AddRaw(this, &FPCGEditor::OnComponentGenerationDone);
+		Subsystem->OnPCGSourceGenerationDone.AddRaw(this, &FPCGEditor::OnSourceGenerationDone);
 	}
 }
 
@@ -3143,7 +3151,7 @@ void FPCGEditor::UnregisterDelegatesForWorld(UWorld* World)
 	if (UPCGSubsystem* Subsystem = UPCGSubsystem::GetInstance(World))
 	{
 		Subsystem->OnPCGComponentUnregistered.RemoveAll(this);
-		Subsystem->OnPCGComponentGenerationDone.RemoveAll(this);
+		Subsystem->OnPCGSourceGenerationDone.RemoveAll(this);
 	}
 }
 
diff --git a/Engine/Plugins/PCG/Source/PCGEditor/Private/PCGEditor.h b/Engine/Plugins/PCG/Source/PCGEditor/Private/PCGEditor.h
--- a/Engine/Plugins/PCG/Source/PCGEditor/Private/PCGEditor.h
+++ b/Engine/Plugins/PCG/Source/PCGEditor/Private/PCGEditor.h
@@ -12,6 +12,7 @@
 enum class ECheckBoxState : uint8;
 namespace ETextCommit { enum Type : int; }
 
+class IPCGBaseSubsystem;
 class FUICommandList;
 class SGraphEditor;
 class SGraphEditorActionMenu;
@@ -31,6 +32,7 @@ class UPCGComponent;
 class UPCGEditorGraph;
 class UPCGEditorGraphNodeBase;
 class UPCGGraph;
+class UPCGSubsystem;
 struct FPCGCompilerDiagnostics;
 
 DECLARE_MULTICAST_DELEGATE_OneParam(FOnInspectedStackChanged, const FPCGStack&);
@@ -111,7 +113,8 @@ public:
 	bool IsPanelCurrentlyForeground(EPCGEditorPanel Panel) const;
 
 	/** Helper to get to the subsystem. */
-	static class UPCGSubsystem* GetSubsystem();
+	virtual IPCGBaseSubsystem* GetSubsystem() const;
+	static UPCGSubsystem* GetWorldSubsystem();
 
 	// ~Begin IToolkit interface
 	virtual void RegisterTabSpawners(const TSharedRef<class FTabManager>& InTabManager) override;
@@ -401,7 +404,7 @@ private:
 	void OnComponentUnregistered(UPCGComponent* Component);
 
 	/** Called when a component finishes executing. Useful for updating debugging tools/UIs. */
-	void OnComponentGenerationDone(UPCGSubsystem* Subsystem, UPCGComponent* Component, EPCGGenerationStatus Status);
+	void OnSourceGenerationDone(IPCGBaseSubsystem* Subsystem, IPCGGraphExecutionSource* Source, EPCGGenerationStatus Status);
 
 	/** Trigger any generation required to ensure debug display is up to date. */
 	void UpdateDebugAfterComponentSelection(UPCGComponent* InOldComponent, UPCGComponent* InNewComponent, bool bNewComponentStartedInspecting);
diff --git a/Engine/Plugins/PCG/Source/PCGEditor/Private/PCGEditorModule.cpp b/Engine/Plugins/PCG/Source/PCGEditor/Private/PCGEditorModule.cpp
--- a/Engine/Plugins/PCG/Source/PCGEditor/Private/PCGEditorModule.cpp
+++ b/Engine/Plugins/PCG/Source/PCGEditor/Private/PCGEditorModule.cpp
@@ -46,6 +46,7 @@
 #include "Details/PCGInstancedPropertyBagOverrideDetails.h"
 #include "Details/PCGVolumeDetails.h"
 #include "Nodes/PCGEditorGraphNodeFactory.h"
+#include "Subsystems/PCGEditorSubsystem.h"
 
 #include "ContentBrowserMenuContexts.h"
 #include "ContentBrowserModule.h"
@@ -267,6 +268,11 @@ TSharedPtr<ISyntaxTokenizer> FPCGEditorModule::CreateHLSLSyntaxTokenizer(const F
 	return MakeShared<FPCGHLSLSyntaxTokenizer>(InParams);
 }
 
+IPCGBaseSubsystem* FPCGEditorModule::GetEditorSubsystem() const
+{
+	return GEditor ? GEditor->GetEditorSubsystem<UPCGEditorSubsystem>() : nullptr;
+}
+
 void FPCGEditorModule::RegisterOnEditorModeChange()
 {
 	// Have a callback that catches changes in the Editor modes, to catch when we exit the landscape edit mode.
diff --git a/Engine/Plugins/PCG/Source/PCGEditor/Private/Subsystems/PCGEditorSubsystem.cpp b/Engine/Plugins/PCG/Source/PCGEditor/Private/Subsystems/PCGEditorSubsystem.cpp
new file mode 100644
--- /dev/null
+++ b/Engine/Plugins/PCG/Source/PCGEditor/Private/Subsystems/PCGEditorSubsystem.cpp
@@ -0,0 +1,34 @@
+﻿// Copyright Epic Games, Inc. All Rights Reserved.
+
+#include "Subsystems/PCGEditorSubsystem.h"
+
+#if WITH_EDITOR
+
+void UPCGEditorSubsystem::Initialize(FSubsystemCollectionBase& Collection)
+{
+	Super::Initialize(Collection);
+	IPCGBaseSubsystem::InitializeBaseSubsystem();
+}
+
+void UPCGEditorSubsystem::Deinitialize()
+{
+	IPCGBaseSubsystem::DeinitializeBaseSubsystem();
+	Super::Deinitialize();
+}
+
+void UPCGEditorSubsystem::Tick(float DeltaTime)
+{
+	IPCGBaseSubsystem::Tick();
+}
+
+ETickableTickType UPCGEditorSubsystem::GetTickableTickType() const
+{ 
+	return ETickableTickType::Always; 
+}
+
+TStatId UPCGEditorSubsystem::GetStatId() const
+{
+	RETURN_QUICK_DECLARE_CYCLE_STAT(FPCGEditorSubsystemTickFunction, STATGROUP_Tickables);
+}
+
+#endif // WITH_EDITOR
diff --git a/Engine/Plugins/PCG/Source/PCGEditor/Private/Widgets/SPCGEditorGraphDebugObjectTree.cpp b/Engine/Plugins/PCG/Source/PCGEditor/Private/Widgets/SPCGEditorGraphDebugObjectTree.cpp
--- a/Engine/Plugins/PCG/Source/PCGEditor/Private/Widgets/SPCGEditorGraphDebugObjectTree.cpp
+++ b/Engine/Plugins/PCG/Source/PCGEditor/Private/Widgets/SPCGEditorGraphDebugObjectTree.cpp
@@ -1142,7 +1142,7 @@ void SPCGEditorGraphDebugObjectTree::RefreshTree()
 		return;
 	}
 
-	UPCGSubsystem* Subsystem = PCGEditor.Pin()->GetSubsystem();
+	IPCGBaseSubsystem* Subsystem = PCGEditor.Pin()->GetSubsystem();
 	if (!Subsystem)
 	{
 		return;
diff --git a/Engine/Plugins/PCG/Source/PCGEditor/Public/PCGEditorModule.h b/Engine/Plugins/PCG/Source/PCGEditor/Public/PCGEditorModule.h
--- a/Engine/Plugins/PCG/Source/PCGEditor/Public/PCGEditorModule.h
+++ b/Engine/Plugins/PCG/Source/PCGEditor/Public/PCGEditorModule.h
@@ -35,6 +35,7 @@ public:
 	virtual FPCGNodeVisualLogs& GetNodeVisualLogsMutable() override { return NodeVisualLogs; }
 	virtual bool CanSelectPartitionActors() const override;
 	virtual TSharedPtr<class ISyntaxTokenizer> CreateHLSLSyntaxTokenizer(const FPCGSyntaxTokenizerParams& InParams) const override;
+	virtual IPCGBaseSubsystem* GetEditorSubsystem() const override;
 
 protected:
 	virtual void OnScheduleGraph(const FPCGStackContext& StackContext) override;
diff --git a/Engine/Plugins/PCG/Source/PCGEditor/Public/Subsystems/PCGEditorSubsystem.h b/Engine/Plugins/PCG/Source/PCGEditor/Public/Subsystems/PCGEditorSubsystem.h
new file mode 100644
--- /dev/null
+++ b/Engine/Plugins/PCG/Source/PCGEditor/Public/Subsystems/PCGEditorSubsystem.h
@@ -0,0 +1,30 @@
+﻿// Copyright Epic Games, Inc. All Rights Reserved.
+
+#pragma once
+
+#if WITH_EDITOR
+
+#include "EditorSubsystem.h"
+#include "TickableEditorObject.h"
+
+#include "Subsystems/IPCGBaseSubsystem.h"
+
+#include "PCGEditorSubsystem.generated.h"
+
+UCLASS(MinimalAPI)
+class UPCGEditorSubsystem : public UEditorSubsystem, public IPCGBaseSubsystem, public FTickableEditorObject
+{
+	GENERATED_BODY()
+
+	PCGEDITOR_API virtual void Initialize(FSubsystemCollectionBase& Collection) override;
+	PCGEDITOR_API virtual void Deinitialize() override;
+
+protected:
+	//** FTickableEditorObject Interface
+	PCGEDITOR_API virtual void Tick(float DeltaTime) override;
+	PCGEDITOR_API virtual ETickableTickType GetTickableTickType() const override ;
+	PCGEDITOR_API virtual TStatId GetStatId() const override;
+	//~ END FTickableEditorObject Interface
+};
+
+#endif // WITH_EDITOR
\ No newline at end of file
