[PCG] define EPCGNativeProperties enum traits to make code more robust and type safe
- also optimizes some code paths through templating

diff --git a/Engine/Plugins/PCG/Source/PCG/Private/Data/PCGBasePointData.cpp b/Engine/Plugins/PCG/Source/PCG/Private/Data/PCGBasePointData.cpp
--- a/Engine/Plugins/PCG/Source/PCG/Private/Data/PCGBasePointData.cpp
+++ b/Engine/Plugins/PCG/Source/PCG/Private/Data/PCGBasePointData.cpp
@@ -144,20 +144,21 @@ namespace PCGBasePointData
 		switch (NativeProperty)
 		{
 		case EPCGPointNativeProperties::Transform:
-			return MakeUnique<FPCGNativePointPropertyAccessor<FTransform>>(PointData, NativeProperty);
+			return MakeUnique<FPCGNativePointPropertyEnumAccessor<EPCGPointNativeProperties::Transform>>(PointData);
 		case EPCGPointNativeProperties::Density:
+			return MakeUnique<FPCGNativePointPropertyEnumAccessor<EPCGPointNativeProperties::Density>>(PointData);
 		case EPCGPointNativeProperties::Steepness:
-			// The values are floats but we want to remain compatible with existing UPCGPointData Steepnees/Density property accessors which are of type double 
-			return MakeUnique<FPCGNativePointPropertyAccessor<double, float>>(PointData, NativeProperty);
+			return MakeUnique<FPCGNativePointPropertyEnumAccessor<EPCGPointNativeProperties::Steepness>>(PointData);
 		case EPCGPointNativeProperties::BoundsMin:
+			return MakeUnique<FPCGNativePointPropertyEnumAccessor<EPCGPointNativeProperties::BoundsMin>>(PointData);
 		case EPCGPointNativeProperties::BoundsMax:
-			return MakeUnique<FPCGNativePointPropertyAccessor<FVector>>(PointData, NativeProperty);
+			return MakeUnique<FPCGNativePointPropertyEnumAccessor<EPCGPointNativeProperties::BoundsMax>>(PointData);
 		case EPCGPointNativeProperties::Color:
-			return MakeUnique<FPCGNativePointPropertyAccessor<FVector4>>(PointData, NativeProperty);
+			return MakeUnique<FPCGNativePointPropertyEnumAccessor<EPCGPointNativeProperties::Color>>(PointData);
 		case EPCGPointNativeProperties::Seed:
-			return MakeUnique<FPCGNativePointPropertyAccessor<int32>>(PointData, NativeProperty);
+			return MakeUnique<FPCGNativePointPropertyEnumAccessor<EPCGPointNativeProperties::Seed>>(PointData);
 		case EPCGPointNativeProperties::MetadataEntry:
-			return MakeUnique<FPCGNativePointPropertyAccessor<int64>>(PointData, NativeProperty);
+			return MakeUnique<FPCGNativePointPropertyEnumAccessor<EPCGPointNativeProperties::MetadataEntry>>(PointData);
 		default:
 			UE_LOG(LogPCG, Error, TEXT("EPCGPointNativeProperty value '%d' does not exist."), NativeProperty);
 			return TUniquePtr<IPCGAttributeAccessor>();
@@ -169,20 +170,21 @@ namespace PCGBasePointData
 		switch (NativeProperty)
 		{
 		case EPCGPointNativeProperties::Transform:
-			return MakeUnique<FPCGNativePointPropertyConstAccessor<FTransform>>(PointData, NativeProperty);
+			return MakeUnique<FPCGNativePointPropertyEnumConstAccessor<EPCGPointNativeProperties::Transform>>(PointData);
 		case EPCGPointNativeProperties::Density:
+			return MakeUnique<FPCGNativePointPropertyEnumConstAccessor<EPCGPointNativeProperties::Density>>(PointData);
 		case EPCGPointNativeProperties::Steepness:
-			// The values are floats but we want to remain compatible with existing UPCGPointData Steepnees/Density property accessors which are of type double 
-			return MakeUnique<FPCGNativePointPropertyConstAccessor<double, float>>(PointData, NativeProperty);
+			return MakeUnique<FPCGNativePointPropertyEnumConstAccessor<EPCGPointNativeProperties::Steepness>>(PointData);
 		case EPCGPointNativeProperties::BoundsMin:
+			return MakeUnique<FPCGNativePointPropertyEnumConstAccessor<EPCGPointNativeProperties::BoundsMin>>(PointData);
 		case EPCGPointNativeProperties::BoundsMax:
-			return MakeUnique<FPCGNativePointPropertyConstAccessor<FVector>>(PointData, NativeProperty);
+			return MakeUnique<FPCGNativePointPropertyEnumConstAccessor<EPCGPointNativeProperties::BoundsMax>>(PointData);
 		case EPCGPointNativeProperties::Color:
-			return MakeUnique<FPCGNativePointPropertyConstAccessor<FVector4>>(PointData, NativeProperty);
+			return MakeUnique<FPCGNativePointPropertyEnumConstAccessor<EPCGPointNativeProperties::Color>>(PointData);
 		case EPCGPointNativeProperties::Seed:
-			return MakeUnique<FPCGNativePointPropertyConstAccessor<int32>>(PointData, NativeProperty);
+			return MakeUnique<FPCGNativePointPropertyEnumConstAccessor<EPCGPointNativeProperties::Seed>>(PointData);
 		case EPCGPointNativeProperties::MetadataEntry:
-			return MakeUnique<FPCGNativePointPropertyConstAccessor<int64>>(PointData, NativeProperty);
+			return MakeUnique<FPCGNativePointPropertyEnumConstAccessor<EPCGPointNativeProperties::MetadataEntry>>(PointData);
 		default:
 			UE_LOG(LogPCG, Error, TEXT("EPCGPointNativeProperty value '%d' does not exist."), NativeProperty);
 			return TUniquePtr<IPCGAttributeAccessor>();
@@ -193,14 +195,14 @@ namespace PCGBasePointData
 	{
 		if (Name == PCGPointCustomPropertyNames::ExtentsName)
 		{
-			return MakeUnique<FPCGCustomPointPropertyAccessor<FVector, TPCGValueRange<FVector>, TPCGValueRange<FVector>>>(
+			return MakeUnique<FPCGCustomPointPropertyAccessor<FVector, FPCGPointBoundsMin::ValueRange, FPCGPointBoundsMax::ValueRange>>(
 				PointData,
-				[](int32 Index, FVector& OutValue, const TPCGValueRange<FVector>& BoundsMin, const TPCGValueRange<FVector>& BoundsMax)
+				[](int32 Index, FVector& OutValue, const FPCGPointBoundsMin::ValueRange& BoundsMin, const FPCGPointBoundsMax::ValueRange& BoundsMax)
 				{
 					OutValue = PCGPointHelpers::GetExtents(BoundsMin[Index], BoundsMax[Index]);
 					return true;
 				},
-				[](int32 Index, const FVector& InValue, TPCGValueRange<FVector> BoundsMin, TPCGValueRange<FVector> BoundsMax)
+				[](int32 Index, const FVector& InValue, FPCGPointBoundsMin::ValueRange BoundsMin, FPCGPointBoundsMax::ValueRange BoundsMax)
 				{
 					PCGPointHelpers::SetExtents(InValue, BoundsMin[Index], BoundsMax[Index]);
 					return true;
@@ -211,14 +213,14 @@ namespace PCGBasePointData
 		}
 		else if (Name == PCGPointCustomPropertyNames::LocalCenterName)
 		{
-			return MakeUnique<FPCGCustomPointPropertyAccessor<FVector, TPCGValueRange<FVector>, TPCGValueRange<FVector>>>(
+			return MakeUnique<FPCGCustomPointPropertyAccessor<FVector, FPCGPointBoundsMin::ValueRange, FPCGPointBoundsMax::ValueRange>>(
 				PointData,
-				[](int32 Index, FVector& OutValue, const TPCGValueRange<FVector>& BoundsMin, const TPCGValueRange<FVector>& BoundsMax)
+				[](int32 Index, FVector& OutValue, const FPCGPointBoundsMin::ValueRange& BoundsMin, const FPCGPointBoundsMax::ValueRange& BoundsMax)
 				{
 					OutValue = PCGPointHelpers::GetLocalCenter(BoundsMin[Index], BoundsMax[Index]);
 					return true;
 				},
-				[](int32 Index, const FVector& InValue, TPCGValueRange<FVector> BoundsMin, TPCGValueRange<FVector> BoundsMax)
+				[](int32 Index, const FVector& InValue, FPCGPointBoundsMin::ValueRange BoundsMin, FPCGPointBoundsMax::ValueRange BoundsMax)
 				{
 					PCGPointHelpers::SetLocalCenter(InValue, BoundsMin[Index], BoundsMax[Index]);
 					return true;
@@ -229,14 +231,14 @@ namespace PCGBasePointData
 		}
 		else if (Name == PCGPointCustomPropertyNames::PositionName)
 		{
-			return MakeUnique<FPCGCustomPointPropertyAccessor<FVector, TPCGValueRange<FTransform>>>(
+			return MakeUnique<FPCGCustomPointPropertyAccessor<FVector, FPCGPointTransform::ValueRange>>(
 				PointData,
-				[](int32 Index, FVector& OutValue, const TPCGValueRange<FTransform>& Transform)
+				[](int32 Index, FVector& OutValue, const FPCGPointTransform::ValueRange& Transform)
 				{
 					OutValue = Transform[Index].GetLocation();
 					return true;
 				},
-				[](int32 Index, const FVector& InValue, TPCGValueRange<FTransform> Transform)
+				[](int32 Index, const FVector& InValue, FPCGPointTransform::ValueRange Transform)
 				{
 					Transform[Index].SetLocation(InValue);
 					return true;
@@ -246,14 +248,14 @@ namespace PCGBasePointData
 		}
 		else if (Name == PCGPointCustomPropertyNames::RotationName)
 		{
-			return MakeUnique<FPCGCustomPointPropertyAccessor<FQuat, TPCGValueRange<FTransform>>>(
+			return MakeUnique<FPCGCustomPointPropertyAccessor<FQuat, FPCGPointTransform::ValueRange>>(
 				PointData,
-				[](int32 Index, FQuat& OutValue, const TPCGValueRange<FTransform>& Transform)
+				[](int32 Index, FQuat& OutValue, const FPCGPointTransform::ValueRange& Transform)
 				{
 					OutValue = Transform[Index].GetRotation();
 					return true;
 				},
-				[](int32 Index, const FQuat& InValue, TPCGValueRange<FTransform> Transform)
+				[](int32 Index, const FQuat& InValue, FPCGPointTransform::ValueRange Transform)
 				{
 					Transform[Index].SetRotation(InValue);
 					return true;
@@ -263,14 +265,14 @@ namespace PCGBasePointData
 		}
 		else if (Name == PCGPointCustomPropertyNames::ScaleName)
 		{
-			return MakeUnique<FPCGCustomPointPropertyAccessor<FVector, TPCGValueRange<FTransform>>>(
+			return MakeUnique<FPCGCustomPointPropertyAccessor<FVector, FPCGPointTransform::ValueRange>>(
 				PointData,
-				[](int32 Index, FVector& OutValue, const TPCGValueRange<FTransform>& Transform)
+				[](int32 Index, FVector& OutValue, const FPCGPointTransform::ValueRange& Transform)
 				{
 					OutValue = Transform[Index].GetScale3D();
 					return true;
 				},
-				[](int32 Index, const FVector& InValue, TPCGValueRange<FTransform> Transform)
+				[](int32 Index, const FVector& InValue, FPCGPointTransform::ValueRange Transform)
 				{
 					Transform[Index].SetScale3D(InValue);
 					return true;
@@ -286,9 +288,9 @@ namespace PCGBasePointData
 	{
 		if (Name == PCGPointCustomPropertyNames::ExtentsName)
 		{
-			return MakeUnique<FPCGCustomPointPropertyAccessor<FVector, TConstPCGValueRange<FVector>, TConstPCGValueRange<FVector>>>(
+			return MakeUnique<FPCGCustomPointPropertyAccessor<FVector, FPCGPointBoundsMin::ConstValueRange, FPCGPointBoundsMax::ConstValueRange>>(
 				PointData,
-				[](int32 Index, FVector& OutValue, const TConstPCGValueRange<FVector>& BoundsMin, const TConstPCGValueRange<FVector>& BoundsMax)
+				[](int32 Index, FVector& OutValue, const FPCGPointBoundsMin::ConstValueRange& BoundsMin, const FPCGPointBoundsMax::ConstValueRange& BoundsMax)
 				{
 					OutValue = PCGPointHelpers::GetExtents(BoundsMin[Index], BoundsMax[Index]);
 					return true;
@@ -299,9 +301,9 @@ namespace PCGBasePointData
 		}
 		else if (Name == PCGPointCustomPropertyNames::LocalCenterName)
 		{
-			return MakeUnique<FPCGCustomPointPropertyAccessor<FVector, TConstPCGValueRange<FVector>, TConstPCGValueRange<FVector>>>(
+			return MakeUnique<FPCGCustomPointPropertyAccessor<FVector, FPCGPointBoundsMin::ConstValueRange, FPCGPointBoundsMax::ConstValueRange>>(
 				PointData,
-				[](int32 Index, FVector& OutValue, const TConstPCGValueRange<FVector>& BoundsMin, const TConstPCGValueRange<FVector>& BoundsMax)
+				[](int32 Index, FVector& OutValue, const FPCGPointBoundsMin::ConstValueRange& BoundsMin, const FPCGPointBoundsMax::ConstValueRange& BoundsMax)
 				{
 					OutValue = PCGPointHelpers::GetLocalCenter(BoundsMin[Index], BoundsMax[Index]);
 					return true;
@@ -312,9 +314,9 @@ namespace PCGBasePointData
 		}
 		else if (Name == PCGPointCustomPropertyNames::PositionName)
 		{
-			return MakeUnique<FPCGCustomPointPropertyAccessor<FVector, TConstPCGValueRange<FTransform>>>(
+			return MakeUnique<FPCGCustomPointPropertyAccessor<FVector, FPCGPointTransform::ConstValueRange>>(
 				PointData,
-				[](int32 Index, FVector& OutValue, const TConstPCGValueRange<FTransform>& Transform)
+				[](int32 Index, FVector& OutValue, const FPCGPointTransform::ConstValueRange& Transform)
 				{
 					OutValue = Transform[Index].GetLocation();
 					return true;
@@ -324,9 +326,9 @@ namespace PCGBasePointData
 		}
 		else if (Name == PCGPointCustomPropertyNames::RotationName)
 		{
-			return MakeUnique<FPCGCustomPointPropertyAccessor<FQuat, TConstPCGValueRange<FTransform>>>(
+			return MakeUnique<FPCGCustomPointPropertyAccessor<FQuat, FPCGPointTransform::ConstValueRange>>(
 				PointData,
-				[](int32 Index, FQuat& OutValue, const TConstPCGValueRange<FTransform>& Transform)
+				[](int32 Index, FQuat& OutValue, const FPCGPointTransform::ConstValueRange& Transform)
 				{
 					OutValue = Transform[Index].GetRotation();
 					return true;
@@ -336,9 +338,9 @@ namespace PCGBasePointData
 		}
 		else if (Name == PCGPointCustomPropertyNames::ScaleName)
 		{
-			return MakeUnique<FPCGCustomPointPropertyAccessor<FVector, TConstPCGValueRange<FTransform>>>(
+			return MakeUnique<FPCGCustomPointPropertyAccessor<FVector, FPCGPointTransform::ConstValueRange>>(
 				PointData,
-				[](int32 Index, FVector& OutValue, const TConstPCGValueRange<FTransform>& Transform)
+				[](int32 Index, FVector& OutValue, const FPCGPointTransform::ConstValueRange& Transform)
 				{
 					OutValue = Transform[Index].GetScale3D();
 					return true;
@@ -348,9 +350,9 @@ namespace PCGBasePointData
 		}
 		else if (Name == PCGPointCustomPropertyNames::LocalSizeName)
 		{
-			return MakeUnique<FPCGCustomPointPropertyAccessor<FVector, TConstPCGValueRange<FVector>, TConstPCGValueRange<FVector>>>(
+			return MakeUnique<FPCGCustomPointPropertyAccessor<FVector, FPCGPointBoundsMin::ConstValueRange, FPCGPointBoundsMax::ConstValueRange>>(
 				PointData,
-				[](int32 Index, FVector& OutValue, const TConstPCGValueRange<FVector>& BoundsMin, const TConstPCGValueRange<FVector>& BoundsMax)
+				[](int32 Index, FVector& OutValue, const FPCGPointBoundsMin::ConstValueRange& BoundsMin, const FPCGPointBoundsMax::ConstValueRange& BoundsMax)
 				{
 					OutValue = PCGPointHelpers::GetLocalSize(BoundsMin[Index], BoundsMax[Index]);
 					return true;
@@ -361,9 +363,9 @@ namespace PCGBasePointData
 		}
 		else if (Name == PCGPointCustomPropertyNames::ScaledLocalSizeName)
 		{
-			return MakeUnique<FPCGCustomPointPropertyAccessor<FVector, TConstPCGValueRange<FTransform>, TConstPCGValueRange<FVector>, TConstPCGValueRange<FVector>>>(
+			return MakeUnique<FPCGCustomPointPropertyAccessor<FVector, FPCGPointTransform::ConstValueRange, FPCGPointBoundsMin::ConstValueRange, FPCGPointBoundsMax::ConstValueRange>>(
 				PointData,
-				[](int32 Index, FVector& OutValue, const TConstPCGValueRange<FTransform>& Transform, const TConstPCGValueRange<FVector>& BoundsMin, const TConstPCGValueRange<FVector>& BoundsMax)
+				[](int32 Index, FVector& OutValue, const FPCGPointTransform::ConstValueRange& Transform, const FPCGPointBoundsMin::ConstValueRange& BoundsMin, const FPCGPointBoundsMax::ConstValueRange& BoundsMax)
 				{
 					OutValue = PCGPointHelpers::GetScaledLocalSize(Transform[Index], BoundsMin[Index], BoundsMax[Index]);
 					return true;
@@ -754,59 +756,44 @@ void UPCGBasePointData::RebuildOctree() const
 	bOctreeIsDirty = false;
 }
 
-void UPCGBasePointData::SetTransform(const FTransform& InTransform)
+void UPCGBasePointData::SetTransform(const FPCGPointTransform::Type& InTransform)
 {
-	FreeProperties(EPCGPointNativeProperties::Transform);
-	TPCGValueRange<FTransform> ValueRange = GetTransformValueRange(/*bAllocate=*/false);
-	ValueRange.Set(InTransform);
+	SetPointPropertyValues<FPCGPointTransform::EnumValue>(InTransform);
 }
 
-void UPCGBasePointData::SetDensity(float InDensity)
+void UPCGBasePointData::SetDensity(FPCGPointDensity::Type InDensity)
 {
-	FreeProperties(EPCGPointNativeProperties::Density);
-	TPCGValueRange<float> ValueRange = GetDensityValueRange(/*bAllocate=*/false);
-	ValueRange.Set(InDensity);
+	SetPointPropertyValues<FPCGPointDensity::EnumValue>(InDensity);
 }
 
-void UPCGBasePointData::SetBoundsMin(const FVector& InBoundsMin)
+void UPCGBasePointData::SetBoundsMin(const FPCGPointBoundsMin::Type& InBoundsMin)
 {
-	FreeProperties(EPCGPointNativeProperties::BoundsMin);
-	TPCGValueRange<FVector> ValueRange = GetBoundsMinValueRange(/*bAllocate=*/false);
-	ValueRange.Set(InBoundsMin);
+	SetPointPropertyValues<FPCGPointBoundsMin::EnumValue>(InBoundsMin);
 }
 
-void UPCGBasePointData::SetBoundsMax(const FVector& InBoundsMax)
+void UPCGBasePointData::SetBoundsMax(const FPCGPointBoundsMax::Type& InBoundsMax)
 {
-	FreeProperties(EPCGPointNativeProperties::BoundsMax);
-	TPCGValueRange<FVector> ValueRange = GetBoundsMaxValueRange(/*bAllocate=*/false);
-	ValueRange.Set(InBoundsMax);
+	SetPointPropertyValues<FPCGPointBoundsMax::EnumValue>(InBoundsMax);
 }
 
-void UPCGBasePointData::SetColor(const FVector4& InColor)
+void UPCGBasePointData::SetColor(const FPCGPointColor::Type& InColor)
 {
-	FreeProperties(EPCGPointNativeProperties::Color);
-	TPCGValueRange<FVector4> ValueRange = GetColorValueRange(/*bAllocate=*/false);
-	ValueRange.Set(InColor);
+	SetPointPropertyValues<FPCGPointColor::EnumValue>(InColor);
 }
 
-void UPCGBasePointData::SetSteepness(float InSteepness)
+void UPCGBasePointData::SetSteepness(FPCGPointSteepness::Type InSteepness)
 {
-	FreeProperties(EPCGPointNativeProperties::Steepness);
-	TPCGValueRange<float> ValueRange = GetSteepnessValueRange(/*bAllocate=*/false);
-	ValueRange.Set(InSteepness);
+	SetPointPropertyValues<FPCGPointSteepness::EnumValue>(InSteepness);
 }
 
-void UPCGBasePointData::SetSeed(int32 InSeed)
+void UPCGBasePointData::SetSeed(FPCGPointSeed::Type InSeed)
 {
-	TPCGValueRange<int32> ValueRange = GetSeedValueRange(/*bAllocate=*/false);
-	ValueRange.Set(InSeed);
+	SetPointPropertyValues<FPCGPointSeed::EnumValue>(InSeed);
 }
 
-void UPCGBasePointData::SetMetadataEntry(int64 InMetadataEntry)
+void UPCGBasePointData::SetMetadataEntry(FPCGPointMetadataEntry::Type InMetadataEntry)
 {
-	FreeProperties(EPCGPointNativeProperties::MetadataEntry);
-	TPCGValueRange<int64> ValueRange = GetMetadataEntryValueRange(/*bAllocate=*/false);
-	ValueRange.Set(InMetadataEntry);
+	SetPointPropertyValues<FPCGPointMetadataEntry::EnumValue>(InMetadataEntry);
 }
 
 void UPCGBasePointData::SetExtents(const FVector& InExtents)
@@ -859,52 +846,44 @@ void UPCGBasePointData::SetLocalCenter(const FVector& InLocalCenter)
 	}
 }
 
-const FTransform& UPCGBasePointData::GetTransform(int32 InPointIndex) const
+const FPCGPointTransform::Type& UPCGBasePointData::GetTransform(int32 InPointIndex) const
 {
-	const TConstPCGValueRange<FTransform> ValueRange = GetConstTransformValueRange();
-	return ValueRange[InPointIndex];
+	return GetPointPropertyValue<FPCGPointTransform::EnumValue>(InPointIndex);
 }
 
-float UPCGBasePointData::GetDensity(int32 InPointIndex) const
+const FPCGPointDensity::Type& UPCGBasePointData::GetDensity(int32 InPointIndex) const
 {
-	const TConstPCGValueRange<float> ValueRange = GetConstDensityValueRange();
-	return ValueRange[InPointIndex];
+	return GetPointPropertyValue<FPCGPointDensity::EnumValue>(InPointIndex);
 }
 
-const FVector& UPCGBasePointData::GetBoundsMin(int32 InPointIndex) const
+const FPCGPointBoundsMin::Type& UPCGBasePointData::GetBoundsMin(int32 InPointIndex) const
 {
-	const TConstPCGValueRange<FVector> ValueRange = GetConstBoundsMinValueRange();
-	return ValueRange[InPointIndex];
+	return GetPointPropertyValue<FPCGPointBoundsMin::EnumValue>(InPointIndex);
 }
 
-const FVector& UPCGBasePointData::GetBoundsMax(int32 InPointIndex) const
+const FPCGPointBoundsMax::Type& UPCGBasePointData::GetBoundsMax(int32 InPointIndex) const
 {
-	const TConstPCGValueRange<FVector> ValueRange = GetConstBoundsMaxValueRange();
-	return ValueRange[InPointIndex];
+	return GetPointPropertyValue<FPCGPointBoundsMax::EnumValue>(InPointIndex);
 }
 
-const FVector4& UPCGBasePointData::GetColor(int32 InPointIndex) const
+const FPCGPointColor::Type& UPCGBasePointData::GetColor(int32 InPointIndex) const
 {
-	const TConstPCGValueRange<FVector4> ValueRange = GetConstColorValueRange();
-	return ValueRange[InPointIndex];
+	return GetPointPropertyValue<FPCGPointColor::EnumValue>(InPointIndex);
 }
 
-float UPCGBasePointData::GetSteepness(int32 InPointIndex) const
+const FPCGPointSteepness::Type& UPCGBasePointData::GetSteepness(int32 InPointIndex) const
 {
-	const TConstPCGValueRange<float> ValueRange = GetConstSteepnessValueRange();
-	return ValueRange[InPointIndex];
+	return GetPointPropertyValue<FPCGPointSteepness::EnumValue>(InPointIndex);
 }
 
-int32 UPCGBasePointData::GetSeed(int32 InPointIndex) const
+const FPCGPointSeed::Type& UPCGBasePointData::GetSeed(int32 InPointIndex) const
 {
-	const TConstPCGValueRange<int32> ValueRange = GetConstSeedValueRange();
-	return ValueRange[InPointIndex];
+	return GetPointPropertyValue<FPCGPointSeed::EnumValue>(InPointIndex);
 }
 
-int64 UPCGBasePointData::GetMetadataEntry(int32 InPointIndex) const
+const FPCGPointMetadataEntry::Type& UPCGBasePointData::GetMetadataEntry(int32 InPointIndex) const
 {
-	const TConstPCGValueRange<int64> ValueRange = GetConstMetadataEntryValueRange();
-	return ValueRange[InPointIndex];
+	return GetPointPropertyValue<FPCGPointMetadataEntry::EnumValue>(InPointIndex);
 }
 
 FBoxSphereBounds UPCGBasePointData::GetDensityBounds(int32 InPointIndex) const
diff --git a/Engine/Plugins/PCG/Source/PCG/Private/Data/PCGPointArrayData.cpp b/Engine/Plugins/PCG/Source/PCG/Private/Data/PCGPointArrayData.cpp
--- a/Engine/Plugins/PCG/Source/PCG/Private/Data/PCGPointArrayData.cpp
+++ b/Engine/Plugins/PCG/Source/PCG/Private/Data/PCGPointArrayData.cpp
@@ -63,42 +63,42 @@ void UPCGPointArrayData::CopyPropertiesTo(UPCGBasePointData* To, int32 ReadStart
 
 		if (EnumHasAllFlags(Properties, EPCGPointNativeProperties::Transform))
 		{
-			GetProperty<FTransform>(EPCGPointNativeProperties::Transform)->CopyTo(*PointArrayData->GetProperty<FTransform>(EPCGPointNativeProperties::Transform, /*bWithInheritance=*/false), ReadStartIndex, WriteStartIndex, Count);
+			GetProperty<EPCGPointNativeProperties::Transform>()->CopyTo(*PointArrayData->GetProperty<EPCGPointNativeProperties::Transform>(/*bWithInheritance=*/false), ReadStartIndex, WriteStartIndex, Count);
 		}
 
 		if (EnumHasAllFlags(Properties, EPCGPointNativeProperties::Density))
 		{
-			GetProperty<float>(EPCGPointNativeProperties::Density)->CopyTo(*PointArrayData->GetProperty<float>(EPCGPointNativeProperties::Density, /*bWithInheritance=*/false), ReadStartIndex, WriteStartIndex, Count);
+			GetProperty<EPCGPointNativeProperties::Density>()->CopyTo(*PointArrayData->GetProperty<EPCGPointNativeProperties::Density>(/*bWithInheritance=*/false), ReadStartIndex, WriteStartIndex, Count);
 		}
 
 		if (EnumHasAllFlags(Properties, EPCGPointNativeProperties::BoundsMin))
 		{
-			GetProperty<FVector>(EPCGPointNativeProperties::BoundsMin)->CopyTo(*PointArrayData->GetProperty<FVector>(EPCGPointNativeProperties::BoundsMin, /*bWithInheritance=*/false), ReadStartIndex, WriteStartIndex, Count);
+			GetProperty<EPCGPointNativeProperties::BoundsMin>()->CopyTo(*PointArrayData->GetProperty<EPCGPointNativeProperties::BoundsMin>(/*bWithInheritance=*/false), ReadStartIndex, WriteStartIndex, Count);
 		}
 
 		if (EnumHasAllFlags(Properties, EPCGPointNativeProperties::BoundsMax))
 		{
-			GetProperty<FVector>(EPCGPointNativeProperties::BoundsMax)->CopyTo(*PointArrayData->GetProperty<FVector>(EPCGPointNativeProperties::BoundsMax, /*bWithInheritance=*/false), ReadStartIndex, WriteStartIndex, Count);
+			GetProperty<EPCGPointNativeProperties::BoundsMax>()->CopyTo(*PointArrayData->GetProperty<EPCGPointNativeProperties::BoundsMax>(/*bWithInheritance=*/false), ReadStartIndex, WriteStartIndex, Count);
 		}
 
 		if (EnumHasAllFlags(Properties, EPCGPointNativeProperties::Color))
 		{
-			GetProperty<FVector4>(EPCGPointNativeProperties::Color)->CopyTo(*PointArrayData->GetProperty<FVector4>(EPCGPointNativeProperties::Color, /*bWithInheritance=*/false), ReadStartIndex, WriteStartIndex, Count);
+			GetProperty<EPCGPointNativeProperties::Color>()->CopyTo(*PointArrayData->GetProperty<EPCGPointNativeProperties::Color>(/*bWithInheritance=*/false), ReadStartIndex, WriteStartIndex, Count);
 		}
 
 		if (EnumHasAllFlags(Properties, EPCGPointNativeProperties::Steepness))
 		{
-			GetProperty<float>(EPCGPointNativeProperties::Steepness)->CopyTo(*PointArrayData->GetProperty<float>(EPCGPointNativeProperties::Steepness, /*bWithInheritance=*/false), ReadStartIndex, WriteStartIndex, Count);
+			GetProperty<EPCGPointNativeProperties::Steepness>()->CopyTo(*PointArrayData->GetProperty<EPCGPointNativeProperties::Steepness>(/*bWithInheritance=*/false), ReadStartIndex, WriteStartIndex, Count);
 		}
 
 		if (EnumHasAllFlags(Properties, EPCGPointNativeProperties::Seed))
 		{
-			GetProperty<int32>(EPCGPointNativeProperties::Seed)->CopyTo(*PointArrayData->GetProperty<int32>(EPCGPointNativeProperties::Seed, /*bWithInheritance=*/false), ReadStartIndex, WriteStartIndex, Count);
+			GetProperty<EPCGPointNativeProperties::Seed>()->CopyTo(*PointArrayData->GetProperty<EPCGPointNativeProperties::Seed>(/*bWithInheritance=*/false), ReadStartIndex, WriteStartIndex, Count);
 		}
 
 		if (EnumHasAllFlags(Properties, EPCGPointNativeProperties::MetadataEntry))
 		{
-			GetProperty<int64>(EPCGPointNativeProperties::MetadataEntry)->CopyTo(*PointArrayData->GetProperty<int64>(EPCGPointNativeProperties::MetadataEntry, /*bWithInheritance=*/false), ReadStartIndex, WriteStartIndex, Count);
+			GetProperty<EPCGPointNativeProperties::MetadataEntry>()->CopyTo(*PointArrayData->GetProperty<EPCGPointNativeProperties::MetadataEntry>(/*bWithInheritance=*/false), ReadStartIndex, WriteStartIndex, Count);
 		}
 	}
 	else
@@ -153,42 +153,42 @@ void UPCGPointArrayData::FlattenPropertiesIfNeeded(EPCGPointNativeProperties Pro
 {	
 	if (EnumHasAnyFlags(Properties, EPCGPointNativeProperties::Transform))
 	{
-		FlattenPropertyIfNeeded<FTransform>(EPCGPointNativeProperties::Transform);
+		FlattenPropertyIfNeeded<EPCGPointNativeProperties::Transform>();
 	}
 	
 	if (EnumHasAnyFlags(Properties, EPCGPointNativeProperties::Steepness))
 	{
-		FlattenPropertyIfNeeded<float>(EPCGPointNativeProperties::Steepness);
+		FlattenPropertyIfNeeded<EPCGPointNativeProperties::Steepness>();
 	}
 
 	if (EnumHasAnyFlags(Properties, EPCGPointNativeProperties::BoundsMin))
 	{
-		FlattenPropertyIfNeeded<FVector>(EPCGPointNativeProperties::BoundsMin);
+		FlattenPropertyIfNeeded<EPCGPointNativeProperties::BoundsMin>();
 	}
 
 	if (EnumHasAnyFlags(Properties, EPCGPointNativeProperties::BoundsMax))
 	{
-		FlattenPropertyIfNeeded<FVector>(EPCGPointNativeProperties::BoundsMax);
+		FlattenPropertyIfNeeded<EPCGPointNativeProperties::BoundsMax>();
 	}
 
 	if (EnumHasAnyFlags(Properties, EPCGPointNativeProperties::Color))
 	{
-		FlattenPropertyIfNeeded<FVector4>(EPCGPointNativeProperties::Color);
+		FlattenPropertyIfNeeded<EPCGPointNativeProperties::Color>();
 	}
 
 	if (EnumHasAnyFlags(Properties, EPCGPointNativeProperties::Density))
 	{
-		FlattenPropertyIfNeeded<float>(EPCGPointNativeProperties::Density);
+		FlattenPropertyIfNeeded<EPCGPointNativeProperties::Density>();
 	}
 
 	if (EnumHasAnyFlags(Properties, EPCGPointNativeProperties::Seed))
 	{
-		FlattenPropertyIfNeeded<int32>(EPCGPointNativeProperties::Seed);
+		FlattenPropertyIfNeeded<EPCGPointNativeProperties::Seed>();
 	}
 
 	if (EnumHasAnyFlags(Properties, EPCGPointNativeProperties::MetadataEntry))
 	{
-		FlattenPropertyIfNeeded<int64>(EPCGPointNativeProperties::MetadataEntry);
+		FlattenPropertyIfNeeded<EPCGPointNativeProperties::MetadataEntry>();
 	}
 }
 
@@ -233,100 +233,101 @@ void UPCGPointArrayData::CopyUnallocatedPropertiesFrom(const UPCGBasePointData*
 
 	if (const UPCGPointArrayData* InPointArrayData = Cast<UPCGPointArrayData>(InPointData))
 	{
-		InPointArrayData->GetProperty<FTransform>(EPCGPointNativeProperties::Transform)->CopyUnallocatedProperty(*GetProperty<FTransform>(EPCGPointNativeProperties::Transform, /*bWithInheritance=*/false));
-		InPointArrayData->GetProperty<float>(EPCGPointNativeProperties::Density)->CopyUnallocatedProperty(*GetProperty<float>(EPCGPointNativeProperties::Density, /*bWithInheritance=*/false));
-		InPointArrayData->GetProperty<FVector>(EPCGPointNativeProperties::BoundsMin)->CopyUnallocatedProperty(*GetProperty<FVector>(EPCGPointNativeProperties::BoundsMin, /*bWithInheritance=*/false));
-		InPointArrayData->GetProperty<FVector>(EPCGPointNativeProperties::BoundsMax)->CopyUnallocatedProperty(*GetProperty<FVector>(EPCGPointNativeProperties::BoundsMax, /*bWithInheritance=*/false));
-		InPointArrayData->GetProperty<FVector4>(EPCGPointNativeProperties::Color)->CopyUnallocatedProperty(*GetProperty<FVector4>(EPCGPointNativeProperties::Color, /*bWithInheritance=*/false));
-		InPointArrayData->GetProperty<float>(EPCGPointNativeProperties::Steepness)->CopyUnallocatedProperty(*GetProperty<float>(EPCGPointNativeProperties::Steepness, /*bWithInheritance=*/false));
-		InPointArrayData->GetProperty<int32>(EPCGPointNativeProperties::Seed)->CopyUnallocatedProperty(*GetProperty<int32>(EPCGPointNativeProperties::Seed, /*bWithInheritance=*/false));
-		InPointArrayData->GetProperty<int64>(EPCGPointNativeProperties::MetadataEntry)->CopyUnallocatedProperty(*GetProperty<int64>(EPCGPointNativeProperties::MetadataEntry, /*bWithInheritance=*/false));
+		InPointArrayData->GetProperty<EPCGPointNativeProperties::Transform>()->CopyUnallocatedProperty(*GetProperty<EPCGPointNativeProperties::Transform>(/*bWithInheritance=*/false));
+		InPointArrayData->GetProperty<EPCGPointNativeProperties::Density>()->CopyUnallocatedProperty(*GetProperty<EPCGPointNativeProperties::Density>(/*bWithInheritance=*/false));
+		InPointArrayData->GetProperty<EPCGPointNativeProperties::BoundsMin>()->CopyUnallocatedProperty(*GetProperty<EPCGPointNativeProperties::BoundsMin>(/*bWithInheritance=*/false));
+		InPointArrayData->GetProperty<EPCGPointNativeProperties::BoundsMax>()->CopyUnallocatedProperty(*GetProperty<EPCGPointNativeProperties::BoundsMax>(/*bWithInheritance=*/false));
+		InPointArrayData->GetProperty<EPCGPointNativeProperties::Color>()->CopyUnallocatedProperty(*GetProperty<EPCGPointNativeProperties::Color>(/*bWithInheritance=*/false));
+		InPointArrayData->GetProperty<EPCGPointNativeProperties::Steepness>()->CopyUnallocatedProperty(*GetProperty<EPCGPointNativeProperties::Steepness>(/*bWithInheritance=*/false));
+		InPointArrayData->GetProperty<EPCGPointNativeProperties::Seed>()->CopyUnallocatedProperty(*GetProperty<EPCGPointNativeProperties::Seed>(/*bWithInheritance=*/false));
+		InPointArrayData->GetProperty<EPCGPointNativeProperties::MetadataEntry>()->CopyUnallocatedProperty(*GetProperty<EPCGPointNativeProperties::MetadataEntry>(/*bWithInheritance=*/false));
 	}
 }
 
-TPCGValueRange<FTransform> UPCGPointArrayData::GetTransformValueRange(bool bAllocate)
+FPCGPointTransform::ValueRange UPCGPointArrayData::GetTransformValueRange(bool bAllocate)
 { 
-	FlattenPropertiesIfNeeded(EPCGPointNativeProperties::Transform);
+	FlattenPropertyIfNeeded<EPCGPointNativeProperties::Transform>();
 	return PointArray.GetTransformValueRange(bAllocate); 
 }
 
-TPCGValueRange<float> UPCGPointArrayData::GetDensityValueRange(bool bAllocate)
+FPCGPointDensity::ValueRange UPCGPointArrayData::GetDensityValueRange(bool bAllocate)
 { 
-	FlattenPropertiesIfNeeded(EPCGPointNativeProperties::Density);
+	FlattenPropertyIfNeeded<EPCGPointNativeProperties::Density>();
 	return PointArray.GetDensityValueRange(bAllocate); 
 }
 
-TPCGValueRange<FVector> UPCGPointArrayData::GetBoundsMinValueRange(bool bAllocate)
+FPCGPointBoundsMin::ValueRange UPCGPointArrayData::GetBoundsMinValueRange(bool bAllocate)
 { 
-	FlattenPropertiesIfNeeded(EPCGPointNativeProperties::BoundsMin);
+	FlattenPropertyIfNeeded<EPCGPointNativeProperties::BoundsMin>();
 	return PointArray.GetBoundsMinValueRange(bAllocate); 
 }
 
-TPCGValueRange<FVector> UPCGPointArrayData::GetBoundsMaxValueRange(bool bAllocate)
+FPCGPointBoundsMax::ValueRange UPCGPointArrayData::GetBoundsMaxValueRange(bool bAllocate)
 {
-	FlattenPropertiesIfNeeded(EPCGPointNativeProperties::BoundsMax);
+	FlattenPropertyIfNeeded<EPCGPointNativeProperties::BoundsMax>();
 	return PointArray.GetBoundsMaxValueRange(bAllocate); 
 }
 
-TPCGValueRange<FVector4> UPCGPointArrayData::GetColorValueRange(bool bAllocate)
+FPCGPointColor::ValueRange UPCGPointArrayData::GetColorValueRange(bool bAllocate)
 { 
-	FlattenPropertiesIfNeeded(EPCGPointNativeProperties::Color);
+	FlattenPropertyIfNeeded<EPCGPointNativeProperties::Color>();
 	return PointArray.GetColorValueRange(bAllocate); 
 }
 
-TPCGValueRange<float> UPCGPointArrayData::GetSteepnessValueRange(bool bAllocate)
+FPCGPointSteepness::ValueRange UPCGPointArrayData::GetSteepnessValueRange(bool bAllocate)
 { 
-	FlattenPropertiesIfNeeded(EPCGPointNativeProperties::Steepness);
+	FlattenPropertyIfNeeded<EPCGPointNativeProperties::Steepness>();
 	return PointArray.GetSteepnessValueRange(bAllocate); 
 }
 
-TPCGValueRange<int32> UPCGPointArrayData::GetSeedValueRange(bool bAllocate)
+FPCGPointSeed::ValueRange UPCGPointArrayData::GetSeedValueRange(bool bAllocate)
 { 
-	FlattenPropertiesIfNeeded(EPCGPointNativeProperties::Seed);
+	FlattenPropertyIfNeeded<EPCGPointNativeProperties::Seed>();
 	return PointArray.GetSeedValueRange(bAllocate); 
 }
-TPCGValueRange<int64> UPCGPointArrayData::GetMetadataEntryValueRange(bool bAllocate)
+
+FPCGPointMetadataEntry::ValueRange UPCGPointArrayData::GetMetadataEntryValueRange(bool bAllocate)
 { 
-	FlattenPropertiesIfNeeded(EPCGPointNativeProperties::MetadataEntry);
+	FlattenPropertyIfNeeded<EPCGPointNativeProperties::MetadataEntry>();
 	return PointArray.GetMetadataEntryValueRange(bAllocate); 
 }
 
-TConstPCGValueRange<FTransform> UPCGPointArrayData::GetConstTransformValueRange() const
+FPCGPointTransform::ConstValueRange UPCGPointArrayData::GetConstTransformValueRange() const
 { 
-	return GetProperty<FTransform>(EPCGPointNativeProperties::Transform)->GetConstValueRange();
+	return GetProperty<EPCGPointNativeProperties::Transform>()->GetConstValueRange();
 }
 
-TConstPCGValueRange<float> UPCGPointArrayData::GetConstDensityValueRange() const
+FPCGPointDensity::ConstValueRange UPCGPointArrayData::GetConstDensityValueRange() const
 { 
-	return GetProperty<float>(EPCGPointNativeProperties::Density)->GetConstValueRange();
+	return GetProperty<EPCGPointNativeProperties::Density>()->GetConstValueRange();
 }
 
-TConstPCGValueRange<FVector> UPCGPointArrayData::GetConstBoundsMinValueRange() const
+FPCGPointBoundsMin::ConstValueRange UPCGPointArrayData::GetConstBoundsMinValueRange() const
 { 
-	return GetProperty<FVector>(EPCGPointNativeProperties::BoundsMin)->GetConstValueRange();
+	return GetProperty<EPCGPointNativeProperties::BoundsMin>()->GetConstValueRange();
 }
 
-TConstPCGValueRange<FVector> UPCGPointArrayData::GetConstBoundsMaxValueRange() const
+FPCGPointBoundsMax::ConstValueRange UPCGPointArrayData::GetConstBoundsMaxValueRange() const
 { 
-	return GetProperty<FVector>(EPCGPointNativeProperties::BoundsMax)->GetConstValueRange();
+	return GetProperty<EPCGPointNativeProperties::BoundsMax>()->GetConstValueRange();
 }
 
-TConstPCGValueRange<FVector4> UPCGPointArrayData::GetConstColorValueRange() const
+FPCGPointColor::ConstValueRange UPCGPointArrayData::GetConstColorValueRange() const
 { 
-	return GetProperty<FVector4>(EPCGPointNativeProperties::Color)->GetConstValueRange();
+	return GetProperty<EPCGPointNativeProperties::Color>()->GetConstValueRange();
 }
 
-TConstPCGValueRange<float> UPCGPointArrayData::GetConstSteepnessValueRange() const
+FPCGPointSteepness::ConstValueRange UPCGPointArrayData::GetConstSteepnessValueRange() const
 { 
-	return GetProperty<float>(EPCGPointNativeProperties::Steepness)->GetConstValueRange();
+	return GetProperty<EPCGPointNativeProperties::Steepness>()->GetConstValueRange();
 }
 
-TConstPCGValueRange<int32> UPCGPointArrayData::GetConstSeedValueRange() const
+FPCGPointSeed::ConstValueRange UPCGPointArrayData::GetConstSeedValueRange() const
 { 
-	return GetProperty<int32>(EPCGPointNativeProperties::Seed)->GetConstValueRange();
+	return GetProperty<EPCGPointNativeProperties::Seed>()->GetConstValueRange();
 }
 
-TConstPCGValueRange<int64> UPCGPointArrayData::GetConstMetadataEntryValueRange() const
+FPCGPointMetadataEntry::ConstValueRange UPCGPointArrayData::GetConstMetadataEntryValueRange() const
 { 
-	return GetProperty<int64>(EPCGPointNativeProperties::MetadataEntry)->GetConstValueRange();
+	return GetProperty<EPCGPointNativeProperties::MetadataEntry>()->GetConstValueRange();
 }
\ No newline at end of file
diff --git a/Engine/Plugins/PCG/Source/PCG/Private/Data/PCGPointData.cpp b/Engine/Plugins/PCG/Source/PCG/Private/Data/PCGPointData.cpp
--- a/Engine/Plugins/PCG/Source/PCG/Private/Data/PCGPointData.cpp
+++ b/Engine/Plugins/PCG/Source/PCG/Private/Data/PCGPointData.cpp
@@ -16,10 +16,10 @@
 
 #include UE_INLINE_GENERATED_CPP_BY_NAME(PCGPointData)
 
-#define PCG_POINTDATA_MAKERANGE_BASE(RangeType, MakeViewType, Type, Field) RangeType<Type>(MakeViewType(sizeof(FPCGPoint), Points.IsEmpty() ? nullptr : &Points[0].Field, Points.Num()));
+#define PCG_POINTDATA_MAKERANGE_BASE(RangeType, MakeViewType, Property, Field) RangeType<TPCGPointNativeProperty<Property>::Type>(MakeViewType(sizeof(FPCGPoint), Points.IsEmpty() ? nullptr : &Points[0].Field, Points.Num()));
 
-#define PCG_POINTDATA_MAKERANGE(Type, Field) PCG_POINTDATA_MAKERANGE_BASE(TPCGValueRange, MakeStridedView, Type, Field)
-#define PCG_POINTDATA_MAKECONSTRANGE(Type, Field) PCG_POINTDATA_MAKERANGE_BASE(TConstPCGValueRange, MakeConstStridedView, Type, Field)
+#define PCG_POINTDATA_MAKERANGE(Property, Field) PCG_POINTDATA_MAKERANGE_BASE(TPCGValueRange, MakeStridedView, Property, Field)
+#define PCG_POINTDATA_MAKECONSTRANGE(Property, Field) PCG_POINTDATA_MAKERANGE_BASE(TConstPCGValueRange, MakeConstStridedView, Property, Field)
 
 namespace PCGPointHelpers
 {
@@ -245,84 +245,84 @@ TArray<FTransform> UPCGPointData::GetTransformsCopy() const
 	return Transforms;
 }
 
-TPCGValueRange<FTransform> UPCGPointData::GetTransformValueRange(bool bAllocate)
+FPCGPointTransform::ValueRange UPCGPointData::GetTransformValueRange(bool bAllocate)
 {
-	return PCG_POINTDATA_MAKERANGE(FTransform, Transform);
+	return PCG_POINTDATA_MAKERANGE(EPCGPointNativeProperties::Transform, Transform);
 }
 
-TPCGValueRange<float> UPCGPointData::GetDensityValueRange(bool bAllocate)
+FPCGPointDensity::ValueRange UPCGPointData::GetDensityValueRange(bool bAllocate)
 {
-	return PCG_POINTDATA_MAKERANGE(float, Density);
+	return PCG_POINTDATA_MAKERANGE(EPCGPointNativeProperties::Density, Density);
 }
 
-TPCGValueRange<FVector> UPCGPointData::GetBoundsMinValueRange(bool bAllocate)
+FPCGPointBoundsMin::ValueRange UPCGPointData::GetBoundsMinValueRange(bool bAllocate)
 {
-	return PCG_POINTDATA_MAKERANGE(FVector, BoundsMin);
+	return PCG_POINTDATA_MAKERANGE(EPCGPointNativeProperties::BoundsMin, BoundsMin);
 }
 
-TPCGValueRange<FVector> UPCGPointData::GetBoundsMaxValueRange(bool bAllocate)
+FPCGPointBoundsMax::ValueRange UPCGPointData::GetBoundsMaxValueRange(bool bAllocate)
 {
-	return PCG_POINTDATA_MAKERANGE(FVector, BoundsMax);
+	return PCG_POINTDATA_MAKERANGE(EPCGPointNativeProperties::BoundsMax, BoundsMax);
 }
 
-TPCGValueRange<FVector4> UPCGPointData::GetColorValueRange(bool bAllocate)
+FPCGPointColor::ValueRange UPCGPointData::GetColorValueRange(bool bAllocate)
 {
-	return PCG_POINTDATA_MAKERANGE(FVector4, Color);
+	return PCG_POINTDATA_MAKERANGE(EPCGPointNativeProperties::Color, Color);
 }
 
-TPCGValueRange<float> UPCGPointData::GetSteepnessValueRange(bool bAllocate)
+FPCGPointSteepness::ValueRange UPCGPointData::GetSteepnessValueRange(bool bAllocate)
 {
-	return PCG_POINTDATA_MAKERANGE(float, Steepness);
+	return PCG_POINTDATA_MAKERANGE(EPCGPointNativeProperties::Steepness, Steepness);
 }
 
-TPCGValueRange<int32> UPCGPointData::GetSeedValueRange(bool bAllocate)
+FPCGPointSeed::ValueRange UPCGPointData::GetSeedValueRange(bool bAllocate)
 {
-	return PCG_POINTDATA_MAKERANGE(int32, Seed);
+	return PCG_POINTDATA_MAKERANGE(EPCGPointNativeProperties::Seed, Seed);
 }
 
-TPCGValueRange<int64> UPCGPointData::GetMetadataEntryValueRange(bool bAllocate)
+FPCGPointMetadataEntry::ValueRange UPCGPointData::GetMetadataEntryValueRange(bool bAllocate)
 {
-	return PCG_POINTDATA_MAKERANGE(int64, MetadataEntry);
+	return PCG_POINTDATA_MAKERANGE(EPCGPointNativeProperties::MetadataEntry, MetadataEntry);
 }
 
-TConstPCGValueRange<FTransform> UPCGPointData::GetConstTransformValueRange() const
+FPCGPointTransform::ConstValueRange UPCGPointData::GetConstTransformValueRange() const
 {
-	return PCG_POINTDATA_MAKECONSTRANGE(FTransform, Transform);
+	return PCG_POINTDATA_MAKECONSTRANGE(EPCGPointNativeProperties::Transform, Transform);
 }
 
-TConstPCGValueRange<float> UPCGPointData::GetConstDensityValueRange() const
+FPCGPointDensity::ConstValueRange UPCGPointData::GetConstDensityValueRange() const
 {
-	return PCG_POINTDATA_MAKECONSTRANGE(float, Density);
+	return PCG_POINTDATA_MAKECONSTRANGE(EPCGPointNativeProperties::Density, Density);
 }
 
-TConstPCGValueRange<FVector> UPCGPointData::GetConstBoundsMinValueRange() const
+FPCGPointBoundsMin::ConstValueRange UPCGPointData::GetConstBoundsMinValueRange() const
 {
-	return PCG_POINTDATA_MAKECONSTRANGE(FVector, BoundsMin);
+	return PCG_POINTDATA_MAKECONSTRANGE(EPCGPointNativeProperties::BoundsMin, BoundsMin);
 }
 
-TConstPCGValueRange<FVector> UPCGPointData::GetConstBoundsMaxValueRange() const
+FPCGPointBoundsMax::ConstValueRange UPCGPointData::GetConstBoundsMaxValueRange() const
 {
-	return PCG_POINTDATA_MAKECONSTRANGE(FVector, BoundsMax);
+	return PCG_POINTDATA_MAKECONSTRANGE(EPCGPointNativeProperties::BoundsMax, BoundsMax);
 }
 
-TConstPCGValueRange<FVector4> UPCGPointData::GetConstColorValueRange() const
+FPCGPointColor::ConstValueRange UPCGPointData::GetConstColorValueRange() const
 {
-	return PCG_POINTDATA_MAKECONSTRANGE(FVector4, Color);
+	return PCG_POINTDATA_MAKECONSTRANGE(EPCGPointNativeProperties::Color, Color);
 }
 
-TConstPCGValueRange<float> UPCGPointData::GetConstSteepnessValueRange() const
+FPCGPointSteepness::ConstValueRange UPCGPointData::GetConstSteepnessValueRange() const
 {
-	return PCG_POINTDATA_MAKECONSTRANGE(float, Steepness);
+	return PCG_POINTDATA_MAKECONSTRANGE(EPCGPointNativeProperties::Steepness, Steepness);
 }
 
-TConstPCGValueRange<int32> UPCGPointData::GetConstSeedValueRange() const
+FPCGPointSeed::ConstValueRange UPCGPointData::GetConstSeedValueRange() const
 {
-	return PCG_POINTDATA_MAKECONSTRANGE(int32, Seed);
+	return PCG_POINTDATA_MAKECONSTRANGE(EPCGPointNativeProperties::Seed, Seed);
 }
 
-TConstPCGValueRange<int64> UPCGPointData::GetConstMetadataEntryValueRange() const
+FPCGPointMetadataEntry::ConstValueRange UPCGPointData::GetConstMetadataEntryValueRange() const
 {
-	return PCG_POINTDATA_MAKECONSTRANGE(int64, MetadataEntry);
+	return PCG_POINTDATA_MAKECONSTRANGE(EPCGPointNativeProperties::MetadataEntry, MetadataEntry);
 }
 
 void UPCGPointData::SetPoints(const TArray<FPCGPoint>& InPoints)
diff --git a/Engine/Plugins/PCG/Source/PCG/Private/Tests/Data/PCGPointArrayDataTest.cpp b/Engine/Plugins/PCG/Source/PCG/Private/Tests/Data/PCGPointArrayDataTest.cpp
--- a/Engine/Plugins/PCG/Source/PCG/Private/Tests/Data/PCGPointArrayDataTest.cpp
+++ b/Engine/Plugins/PCG/Source/PCG/Private/Tests/Data/PCGPointArrayDataTest.cpp
@@ -10,13 +10,13 @@ class FPCGPointArrayDataTestBaseClass : public FPCGTestBaseClass
 public:
 	using FPCGTestBaseClass::FPCGTestBaseClass;
 
-	template <typename T>
-	bool TestEqualRange(const UPCGBasePointData* PointDataA, const UPCGBasePointData* PointDataB, const EPCGPointNativeProperties& NativeProperty)
+	template <EPCGPointNativeProperties NativeProperty>
+	bool TestEqualRange(const UPCGBasePointData* PointDataA, const UPCGBasePointData* PointDataB)
 	{		
-		const TConstPCGValueRange<T> RangeA = PointDataA->GetConstValueRange<T>(NativeProperty);
-		const TConstPCGValueRange<T> RangeB = PointDataB->GetConstValueRange<T>(NativeProperty);
+		const typename TPCGPointNativeProperty<NativeProperty>::ConstValueRange RangeA = PointDataA->GetConstValueRange<NativeProperty>();
+		const typename TPCGPointNativeProperty<NativeProperty>::ConstValueRange RangeB = PointDataB->GetConstValueRange<NativeProperty>();
 
-		const FString PropertyName = StaticEnum<EPCGPointNativeProperties>()->GetNameStringByValue((int64)NativeProperty);
+		const FString PropertyName = StaticEnum<EPCGPointNativeProperties>()->GetNameStringByValue((int64)TPCGPointNativeProperty<NativeProperty>::EnumValue);
 		UTEST_EQUAL(PCG_FMT("Property %s : RangeA.Num() == RangeB.Num()", *PropertyName), RangeA.Num(), RangeB.Num());
 
 		for (int32 i = 0; i < RangeA.Num(); ++i)
@@ -29,42 +29,42 @@ public:
 
 	bool TestEqualData(const UPCGBasePointData* PointDataA, const UPCGBasePointData* PointDataB)
 	{
-		if (!TestEqualRange<FTransform>(PointDataA, PointDataB, EPCGPointNativeProperties::Transform))
+		if (!TestEqualRange<EPCGPointNativeProperties::Transform>(PointDataA, PointDataB))
 		{
 			return false;
 		}
 
-		if (!TestEqualRange<float>(PointDataA, PointDataB, EPCGPointNativeProperties::Density))
+		if (!TestEqualRange<EPCGPointNativeProperties::Density>(PointDataA, PointDataB))
 		{
 			return false;
 		}
 
-		if (!TestEqualRange<FVector>(PointDataA, PointDataB, EPCGPointNativeProperties::BoundsMin))
+		if (!TestEqualRange<EPCGPointNativeProperties::BoundsMin>(PointDataA, PointDataB))
 		{
 			return false;
 		}
 
-		if (!TestEqualRange<FVector>(PointDataA, PointDataB, EPCGPointNativeProperties::BoundsMax))
+		if (!TestEqualRange<EPCGPointNativeProperties::BoundsMax>(PointDataA, PointDataB))
 		{
 			return false;
 		}
 
-		if (!TestEqualRange<FVector4>(PointDataA, PointDataB, EPCGPointNativeProperties::Color))
+		if (!TestEqualRange<EPCGPointNativeProperties::Color>(PointDataA, PointDataB))
 		{
 			return false;
 		}
 
-		if (!TestEqualRange<float>(PointDataA, PointDataB, EPCGPointNativeProperties::Steepness))
+		if (!TestEqualRange<EPCGPointNativeProperties::Steepness>(PointDataA, PointDataB))
 		{
 			return false;
 		}
 
-		if (!TestEqualRange<int32>(PointDataA, PointDataB, EPCGPointNativeProperties::Seed))
+		if (!TestEqualRange<EPCGPointNativeProperties::Seed>(PointDataA, PointDataB))
 		{
 			return false;
 		}
 
-		if (!TestEqualRange<int64>(PointDataA, PointDataB, EPCGPointNativeProperties::MetadataEntry))
+		if (!TestEqualRange<EPCGPointNativeProperties::MetadataEntry>(PointDataA, PointDataB))
 		{
 			return false;
 		}
@@ -369,22 +369,22 @@ bool FPCGPointArrayDataCopyPropertyOverInheritedPropertyTest::RunTest(const FStr
 	OtherData->CopyPropertiesTo(ChildData, 0, 0, OtherData->GetNumPoints(), EPCGPointNativeProperties::Color | EPCGPointNativeProperties::Seed);
 	TestTrue("ChildData has parent (after CopyPropertiesTo)", ChildData->HasSpatialDataParent());
 
-	if (!TestEqualRange<FVector4>(ChildData, OtherData, EPCGPointNativeProperties::Color))
+	if (!TestEqualRange<EPCGPointNativeProperties::Color>(ChildData, OtherData))
 	{
 		return false;
 	}
 
-	if (!TestEqualRange<int32>(ChildData, OtherData, EPCGPointNativeProperties::Seed))
+	if (!TestEqualRange<EPCGPointNativeProperties::Seed>(ChildData, OtherData))
 	{
 		return false;
 	}
 
-	if (!TestEqualRange<FTransform>(ChildData, ParentData, EPCGPointNativeProperties::Transform))
+	if (!TestEqualRange<EPCGPointNativeProperties::Transform>(ChildData, ParentData))
 	{
 		return false;
 	}
 
-	if (!TestEqualRange<float>(ChildData, ParentData, EPCGPointNativeProperties::Density))
+	if (!TestEqualRange<EPCGPointNativeProperties::Density>(ChildData, ParentData))
 	{
 		return false;
 	}
diff --git a/Engine/Plugins/PCG/Source/PCG/Private/Tests/Determinism/PCGDeterminismTestsCommon.cpp b/Engine/Plugins/PCG/Source/PCG/Private/Tests/Determinism/PCGDeterminismTestsCommon.cpp
--- a/Engine/Plugins/PCG/Source/PCG/Private/Tests/Determinism/PCGDeterminismTestsCommon.cpp
+++ b/Engine/Plugins/PCG/Source/PCG/Private/Tests/Determinism/PCGDeterminismTestsCommon.cpp
@@ -966,16 +966,16 @@ namespace PCGDeterminismTests
 		return false;
 	}
 
-	template<class T>
-	bool CompareProperty(const UPCGBasePointData* FirstPointData, const UPCGBasePointData* SecondPointData, EPCGPointNativeProperties Property)
+	template<EPCGPointNativeProperties Property>
+	bool CompareProperty(const UPCGBasePointData* FirstPointData, const UPCGBasePointData* SecondPointData)
 	{
-		const TConstPCGValueRange<T> FirstRange = FirstPointData->GetConstValueRange<T>(Property);
-		const TConstPCGValueRange<T> SecondRange = SecondPointData->GetConstValueRange<T>(Property);
+		const typename TPCGPointNativeProperty<Property>::ConstValueRange FirstRange = FirstPointData->GetConstValueRange<Property>();
+		const typename TPCGPointNativeProperty<Property>::ConstValueRange SecondRange = SecondPointData->GetConstValueRange<Property>();
 		check(FirstRange.Num() == SecondRange.Num());
 
 		for (int32 Index = 0; Index < FirstRange.Num(); ++Index)
 		{
-			if (!PCG::Private::MetadataTraits<T>::Equal(FirstRange[Index], SecondRange[Index]))
+			if (!PCG::Private::MetadataTraits<typename TPCGPointNativeProperty<Property>::Type>::Equal(FirstRange[Index], SecondRange[Index]))
 			{
 				return false;
 			}
@@ -1001,32 +1001,32 @@ namespace PCGDeterminismTests
 			return false;
 		}
 
-		if (!CompareProperty<float>(FirstPointData, SecondPointData, EPCGPointNativeProperties::Steepness))
+		if (!CompareProperty<EPCGPointNativeProperties::Steepness>(FirstPointData, SecondPointData))
 		{
 			return false;
 		}
 
-		if (!CompareProperty<float>(FirstPointData, SecondPointData, EPCGPointNativeProperties::Density))
+		if (!CompareProperty<EPCGPointNativeProperties::Density>(FirstPointData, SecondPointData))
 		{
 			return false;
 		}
 		
-		if (!CompareProperty<FVector>(FirstPointData, SecondPointData, EPCGPointNativeProperties::BoundsMin))
+		if (!CompareProperty<EPCGPointNativeProperties::BoundsMin>(FirstPointData, SecondPointData))
 		{
 			return false;
 		}
 		
-		if (!CompareProperty<FVector>(FirstPointData, SecondPointData, EPCGPointNativeProperties::BoundsMax))
+		if (!CompareProperty<EPCGPointNativeProperties::BoundsMax>(FirstPointData, SecondPointData))
 		{
 			return false;
 		}
 		
-		if (!CompareProperty<FVector4>(FirstPointData, SecondPointData, EPCGPointNativeProperties::Color))
+		if (!CompareProperty<EPCGPointNativeProperties::Color>(FirstPointData, SecondPointData))
 		{
 			return false;
 		}
 		
-		if (!CompareProperty<FTransform>(FirstPointData, SecondPointData, EPCGPointNativeProperties::Transform))
+		if (!CompareProperty<EPCGPointNativeProperties::Transform>(FirstPointData, SecondPointData))
 		{
 			return false;
 		}
diff --git a/Engine/Plugins/PCG/Source/PCG/Private/Tests/PCGValueRangeTest.cpp b/Engine/Plugins/PCG/Source/PCG/Private/Tests/PCGValueRangeTest.cpp
--- a/Engine/Plugins/PCG/Source/PCG/Private/Tests/PCGValueRangeTest.cpp
+++ b/Engine/Plugins/PCG/Source/PCG/Private/Tests/PCGValueRangeTest.cpp
@@ -14,8 +14,8 @@ class FPCGValueRangeTestBaseClass : public FPCGTestBaseClass
 public:
 	using FPCGTestBaseClass::FPCGTestBaseClass;
 
-	template <typename T>
-	bool TestRange(const EPCGPointNativeProperties& NativeProperty)
+	template <EPCGPointNativeProperties Property>
+	bool TestRange()
 	{
 		const int32 NumPoints = 100;
 		const int32 Seed = 42;
@@ -23,8 +23,8 @@ public:
 		UPCGPointData* PointDataA = PCGTestsCommon::CreateRandomPointData<UPCGPointData>(NumPoints, Seed, /*bRandomDensity=*/true);
 		UPCGPointArrayData* PointDataB = PCGTestsCommon::CreateRandomPointData<UPCGPointArrayData>(NumPoints, Seed, /*bRandomDensity=*/true);
 
-		const TConstPCGValueRange<T> RangeA = PointDataA->GetConstValueRange<T>(NativeProperty);
-		const TConstPCGValueRange<T> RangeB = PointDataB->GetConstValueRange<T>(NativeProperty);
+		const typename TPCGPointNativeProperty<Property>::ConstValueRange RangeA = PointDataA->GetConstValueRange<Property>();
+		const typename TPCGPointNativeProperty<Property>::ConstValueRange RangeB = PointDataB->GetConstValueRange<Property>();
 
 		UTEST_EQUAL("RangeA.Num() == RangeB.Num()", RangeA.Num(), RangeB.Num());
 
@@ -114,42 +114,42 @@ bool FPCGValueRangeTest_Base::RunTest(const FString& Parameters)
 
 bool FPCGValueRangeTest_PointData_Transform::RunTest(const FString& Parameters)
 {
-	return TestRange<FTransform>(EPCGPointNativeProperties::Transform);
+	return TestRange<EPCGPointNativeProperties::Transform>();
 }
 
 bool FPCGValueRangeTest_PointData_Steepness::RunTest(const FString& Parameters)
 {
-	return TestRange<float>(EPCGPointNativeProperties::Steepness);
+	return TestRange<EPCGPointNativeProperties::Steepness>();
 }
 
 bool FPCGValueRangeTest_PointData_BoundsMin::RunTest(const FString& Parameters)
 {
-	return TestRange<FVector>(EPCGPointNativeProperties::BoundsMin);
+	return TestRange<EPCGPointNativeProperties::BoundsMin>();
 }
 
 bool FPCGValueRangeTest_PointData_BoundsMax::RunTest(const FString& Parameters)
 {
-	return TestRange<FVector>(EPCGPointNativeProperties::BoundsMax);
+	return TestRange<EPCGPointNativeProperties::BoundsMax>();
 }
 
 bool FPCGValueRangeTest_PointData_Color::RunTest(const FString& Parameters)
 {
-	return TestRange<FVector4>(EPCGPointNativeProperties::Color);
+	return TestRange<EPCGPointNativeProperties::Color>();
 }
 
 bool FPCGValueRangeTest_PointData_Seed::RunTest(const FString& Parameters)
 {
-	return TestRange<int32>(EPCGPointNativeProperties::Seed);
+	return TestRange<EPCGPointNativeProperties::Seed>();
 }
 
 bool FPCGValueRangeTest_PointData_Density::RunTest(const FString& Parameters)
 {
-	return TestRange<float>(EPCGPointNativeProperties::Density);
+	return TestRange<EPCGPointNativeProperties::Density>();
 }
 
 bool FPCGValueRangeTest_PointData_MetadataEntry::RunTest(const FString& Parameters)
 {
-	return TestRange<int64>(EPCGPointNativeProperties::MetadataEntry);
+	return TestRange<EPCGPointNativeProperties::MetadataEntry>();
 }
 
 #endif // WITH_EDITOR
\ No newline at end of file
diff --git a/Engine/Plugins/PCG/Source/PCG/Public/Data/PCGBasePointData.h b/Engine/Plugins/PCG/Source/PCG/Public/Data/PCGBasePointData.h
--- a/Engine/Plugins/PCG/Source/PCG/Public/Data/PCGBasePointData.h
+++ b/Engine/Plugins/PCG/Source/PCG/Public/Data/PCGBasePointData.h
@@ -85,50 +85,74 @@ public:
 	 * 
 	 * @param bAllocate In the case where we return a non-const range by default the memory will be allocated but in some specific case we might not want to allocate (for SingleValue() ranges)
 	 */
+
+	// Keeping for backward compatibility (prefer using enum template version)
 	template <typename T>
 	TPCGValueRange<T> GetValueRange(EPCGPointNativeProperties NativeProperty, bool bAllocate = true);
 
-	virtual TPCGValueRange<FTransform> GetTransformValueRange(bool bAllocate = true) PURE_VIRTUAL(UPCGBasePointData::GetTransformValueRange, return TPCGValueRange<FTransform>(););
-	virtual TPCGValueRange<float> GetDensityValueRange(bool bAllocate = true) PURE_VIRTUAL(UPCGBasePointData::GetDensityValueRange, return TPCGValueRange<float>(););
-	virtual TPCGValueRange<FVector> GetBoundsMinValueRange(bool bAllocate = true) PURE_VIRTUAL(UPCGBasePointData::GetBoundsMinValueRange, return TPCGValueRange<FVector>(););
-	virtual TPCGValueRange<FVector> GetBoundsMaxValueRange(bool bAllocate = true) PURE_VIRTUAL(UPCGBasePointData::GetBoundsMaxValueRange, return TPCGValueRange<FVector>(););
-	virtual TPCGValueRange<FVector4> GetColorValueRange(bool bAllocate = true) PURE_VIRTUAL(UPCGBasePointData::GetColorValueRange, return TPCGValueRange<FVector4>(););
-	virtual TPCGValueRange<float> GetSteepnessValueRange(bool bAllocate = true) PURE_VIRTUAL(UPCGBasePointData::GetSteepnessValueRange, return TPCGValueRange<float>(););
-	virtual TPCGValueRange<int32> GetSeedValueRange(bool bAllocate = true) PURE_VIRTUAL(UPCGBasePointData::GetSeedValueRange, return TPCGValueRange<int32>(););
-	virtual TPCGValueRange<int64> GetMetadataEntryValueRange(bool bAllocate = true) PURE_VIRTUAL(UPCGBasePointData::GetMetadataEntryValueRange, return TPCGValueRange<int64>(););
+	template<EPCGPointNativeProperties Property> 
+	TPCGPointNativeProperty<Property>::ValueRange GetValueRange(bool bAllocate);
+
+	virtual FPCGPointTransform::ValueRange GetTransformValueRange(bool bAllocate = true) PURE_VIRTUAL(UPCGBasePointData::GetTransformValueRange, return FPCGPointTransform::ValueRange(););
+	virtual FPCGPointDensity::ValueRange GetDensityValueRange(bool bAllocate = true) PURE_VIRTUAL(UPCGBasePointData::GetDensityValueRange, return FPCGPointDensity::ValueRange(););
+	virtual FPCGPointBoundsMin::ValueRange GetBoundsMinValueRange(bool bAllocate = true) PURE_VIRTUAL(UPCGBasePointData::GetBoundsMinValueRange, return FPCGPointBoundsMin::ValueRange(););
+	virtual FPCGPointBoundsMax::ValueRange GetBoundsMaxValueRange(bool bAllocate = true) PURE_VIRTUAL(UPCGBasePointData::GetBoundsMaxValueRange, return FPCGPointBoundsMax::ValueRange(););
+	virtual FPCGPointColor::ValueRange GetColorValueRange(bool bAllocate = true) PURE_VIRTUAL(UPCGBasePointData::GetColorValueRange, return FPCGPointColor::ValueRange(););
+	virtual FPCGPointSteepness::ValueRange GetSteepnessValueRange(bool bAllocate = true) PURE_VIRTUAL(UPCGBasePointData::GetSteepnessValueRange, return FPCGPointSteepness::ValueRange(););
+	virtual FPCGPointSeed::ValueRange GetSeedValueRange(bool bAllocate = true) PURE_VIRTUAL(UPCGBasePointData::GetSeedValueRange, return FPCGPointSeed::ValueRange(););
+	virtual FPCGPointMetadataEntry::ValueRange GetMetadataEntryValueRange(bool bAllocate = true) PURE_VIRTUAL(UPCGBasePointData::GetMetadataEntryValueRange, return FPCGPointMetadataEntry::ValueRange(););
 
+	// Keeping for backward compatibility (prefer using enum template version)
 	template <typename T>
 	TConstPCGValueRange<T> GetConstValueRange(EPCGPointNativeProperties NativeProperty) const;
 
-	virtual TConstPCGValueRange<FTransform> GetConstTransformValueRange() const PURE_VIRTUAL(UPCGBasePointData::GetConstTransformValueRange, return TConstPCGValueRange<FTransform>(););
-	virtual TConstPCGValueRange<float> GetConstDensityValueRange() const PURE_VIRTUAL(UPCGBasePointData::GetConstDensityValueRange, return TConstPCGValueRange<float>(););
-	virtual TConstPCGValueRange<FVector> GetConstBoundsMinValueRange() const PURE_VIRTUAL(UPCGBasePointData::GetConstBoundsMinValueRange, return TConstPCGValueRange<FVector>(););
-	virtual TConstPCGValueRange<FVector> GetConstBoundsMaxValueRange() const PURE_VIRTUAL(UPCGBasePointData::GetConstBoundsMaxValueRange, return TConstPCGValueRange<FVector>(););
-	virtual TConstPCGValueRange<FVector4> GetConstColorValueRange() const PURE_VIRTUAL(UPCGBasePointData::GetConstColorValueRange, return TConstPCGValueRange<FVector4>(););
-	virtual TConstPCGValueRange<float> GetConstSteepnessValueRange() const PURE_VIRTUAL(UPCGBasePointData::GetConstSteepnessValueRange, return TConstPCGValueRange<float>(););
-	virtual TConstPCGValueRange<int32> GetConstSeedValueRange() const PURE_VIRTUAL(UPCGBasePointData::GetConstSeedValueRange, return TConstPCGValueRange<int32>(););
-	virtual TConstPCGValueRange<int64> GetConstMetadataEntryValueRange() const PURE_VIRTUAL(UPCGBasePointData::GetConstMetadataEntryValueRange, return TConstPCGValueRange<int64>(););
-		
-	UE_API virtual void SetTransform(const FTransform& InTransform);
-	UE_API virtual void SetDensity(float InDensity);
-	UE_API virtual void SetBoundsMin(const FVector& InBoundsMin);
-	UE_API virtual void SetBoundsMax(const FVector& InBoundsMax);
-	UE_API virtual void SetColor(const FVector4& InColor);
-	UE_API virtual void SetSteepness(float InSteepness);
-	UE_API virtual void SetSeed(int32 InSeed);
-	UE_API virtual void SetMetadataEntry(int64 InMetadataEntry);
+	template<EPCGPointNativeProperties Property>
+	TPCGPointNativeProperty<Property>::ConstValueRange GetConstValueRange() const;
+
+	virtual FPCGPointTransform::ConstValueRange GetConstTransformValueRange() const PURE_VIRTUAL(UPCGBasePointData::GetConstTransformValueRange, return FPCGPointTransform::ConstValueRange(););
+	virtual FPCGPointDensity::ConstValueRange GetConstDensityValueRange() const PURE_VIRTUAL(UPCGBasePointData::GetConstDensityValueRange, return FPCGPointDensity::ConstValueRange(););
+	virtual FPCGPointBoundsMin::ConstValueRange GetConstBoundsMinValueRange() const PURE_VIRTUAL(UPCGBasePointData::GetConstBoundsMinValueRange, return FPCGPointBoundsMin::ConstValueRange(););
+	virtual FPCGPointBoundsMax::ConstValueRange GetConstBoundsMaxValueRange() const PURE_VIRTUAL(UPCGBasePointData::GetConstBoundsMaxValueRange, return FPCGPointBoundsMax::ConstValueRange(););
+	virtual FPCGPointColor::ConstValueRange GetConstColorValueRange() const PURE_VIRTUAL(UPCGBasePointData::GetConstColorValueRange, return FPCGPointColor::ConstValueRange(););
+	virtual FPCGPointSteepness::ConstValueRange GetConstSteepnessValueRange() const PURE_VIRTUAL(UPCGBasePointData::GetConstSteepnessValueRange, return FPCGPointSteepness::ConstValueRange(););
+	virtual FPCGPointSeed::ConstValueRange GetConstSeedValueRange() const PURE_VIRTUAL(UPCGBasePointData::GetConstSeedValueRange, return FPCGPointSeed::ConstValueRange(););
+	virtual FPCGPointMetadataEntry::ConstValueRange GetConstMetadataEntryValueRange() const PURE_VIRTUAL(UPCGBasePointData::GetConstMetadataEntryValueRange, return FPCGPointMetadataEntry::ConstValueRange(););
+
+	template<EPCGPointNativeProperties Property>
+	void SetPointPropertyValues(const typename TPCGPointNativeProperty<Property>::Type& InValue)
+	{
+		FreeProperties(TPCGPointNativeProperty<Property>::EnumValue);
+		typename TPCGPointNativeProperty<Property>::ValueRange PropertyValueRange = GetValueRange<TPCGPointNativeProperty<Property>::EnumValue>(/*bAllocate=*/false);
+		PropertyValueRange.Set(InValue);
+	}
+
+	template<EPCGPointNativeProperties Property>
+	const typename TPCGPointNativeProperty<Property>::Type& GetPointPropertyValue(int32 InPointIndex) const
+	{
+		const typename TPCGPointNativeProperty<Property>::ConstValueRange ConstPropertyValueRange = GetConstValueRange<TPCGPointNativeProperty<Property>::EnumValue>();
+		return ConstPropertyValueRange[InPointIndex];
+	}
+
+	UE_API virtual void SetTransform(const FPCGPointTransform::Type& InTransform);
+	UE_API virtual void SetDensity(FPCGPointDensity::Type InDensity);
+	UE_API virtual void SetBoundsMin(const FPCGPointBoundsMin::Type& InBoundsMin);
+	UE_API virtual void SetBoundsMax(const FPCGPointBoundsMax::Type& InBoundsMax);
+	UE_API virtual void SetColor(const FPCGPointColor::Type& InColor);
+	UE_API virtual void SetSteepness(FPCGPointSteepness::Type InSteepness);
+	UE_API virtual void SetSeed(FPCGPointSeed::Type InSeed);
+	UE_API virtual void SetMetadataEntry(FPCGPointMetadataEntry::Type InMetadataEntry);
 
 	UE_API virtual void SetExtents(const FVector& InExtents);
 	UE_API virtual void SetLocalCenter(const FVector& InLocalCenter);
 
-	UE_API virtual const FTransform& GetTransform(int32 InPointIndex) const;
-	UE_API virtual float GetDensity(int32 InPointIndex) const;
-	UE_API virtual const FVector& GetBoundsMin(int32 InPointIndex) const;
-	UE_API virtual const FVector& GetBoundsMax(int32 InPointIndex) const;
-	UE_API virtual const FVector4& GetColor(int32 InPointIndex) const;
-	UE_API virtual float GetSteepness(int32 InPointIndex) const;
-	UE_API virtual int32 GetSeed(int32 InPointIndex) const;
-	UE_API virtual int64 GetMetadataEntry(int32 InPointIndex) const;
+	UE_API virtual const FPCGPointTransform::Type& GetTransform(int32 InPointIndex) const;
+	UE_API virtual const FPCGPointDensity::Type& GetDensity(int32 InPointIndex) const;
+	UE_API virtual const FPCGPointBoundsMin::Type& GetBoundsMin(int32 InPointIndex) const;
+	UE_API virtual const FPCGPointBoundsMax::Type& GetBoundsMax(int32 InPointIndex) const;
+	UE_API virtual const FPCGPointColor::Type& GetColor(int32 InPointIndex) const;
+	UE_API virtual const FPCGPointSteepness::Type& GetSteepness(int32 InPointIndex) const;
+	UE_API virtual const FPCGPointSeed::Type& GetSeed(int32 InPointIndex) const;
+	UE_API virtual const FPCGPointMetadataEntry::Type& GetMetadataEntry(int32 InPointIndex) const;
 
 	UE_API virtual FBoxSphereBounds GetDensityBounds(int32 InPointIndex) const;
 	UE_API virtual FBox GetLocalDensityBounds(int32 InPointIndex) const;
@@ -166,42 +190,42 @@ public:
 			PropertiesToAllocate |= PointInputData->GetAllocatedProperties();
 		}
 
-		if (!EnumHasAnyFlags(PropertiesToAllocate, EPCGPointNativeProperties::Transform) && NeedToAllocateSingleValueProperty<FTransform>(EPCGPointNativeProperties::Transform, PointDatas))
+		if (!EnumHasAnyFlags(PropertiesToAllocate, EPCGPointNativeProperties::Transform) && NeedToAllocateSingleValueProperty<EPCGPointNativeProperties::Transform>(PointDatas))
 		{
 			PropertiesToAllocate |= EPCGPointNativeProperties::Transform;
 		}
 
-		if (!EnumHasAnyFlags(PropertiesToAllocate, EPCGPointNativeProperties::Density) && NeedToAllocateSingleValueProperty<float>(EPCGPointNativeProperties::Density, PointDatas))
+		if (!EnumHasAnyFlags(PropertiesToAllocate, EPCGPointNativeProperties::Density) && NeedToAllocateSingleValueProperty<EPCGPointNativeProperties::Density>(PointDatas))
 		{
 			PropertiesToAllocate |= EPCGPointNativeProperties::Density;
 		}
 
-		if (!EnumHasAnyFlags(PropertiesToAllocate, EPCGPointNativeProperties::BoundsMin) && NeedToAllocateSingleValueProperty<FVector>(EPCGPointNativeProperties::BoundsMin, PointDatas))
+		if (!EnumHasAnyFlags(PropertiesToAllocate, EPCGPointNativeProperties::BoundsMin) && NeedToAllocateSingleValueProperty<EPCGPointNativeProperties::BoundsMin>(PointDatas))
 		{
 			PropertiesToAllocate |= EPCGPointNativeProperties::BoundsMin;
 		}
 
-		if (!EnumHasAnyFlags(PropertiesToAllocate, EPCGPointNativeProperties::BoundsMax) && NeedToAllocateSingleValueProperty<FVector>(EPCGPointNativeProperties::BoundsMax, PointDatas))
+		if (!EnumHasAnyFlags(PropertiesToAllocate, EPCGPointNativeProperties::BoundsMax) && NeedToAllocateSingleValueProperty<EPCGPointNativeProperties::BoundsMax>(PointDatas))
 		{
 			PropertiesToAllocate |= EPCGPointNativeProperties::BoundsMax;
 		}
 
-		if (!EnumHasAnyFlags(PropertiesToAllocate, EPCGPointNativeProperties::Color) && NeedToAllocateSingleValueProperty<FVector4>(EPCGPointNativeProperties::Color, PointDatas))
+		if (!EnumHasAnyFlags(PropertiesToAllocate, EPCGPointNativeProperties::Color) && NeedToAllocateSingleValueProperty<EPCGPointNativeProperties::Color>(PointDatas))
 		{
 			PropertiesToAllocate |= EPCGPointNativeProperties::Color;
 		}
 
-		if (!EnumHasAnyFlags(PropertiesToAllocate, EPCGPointNativeProperties::Steepness) && NeedToAllocateSingleValueProperty<float>(EPCGPointNativeProperties::Steepness, PointDatas))
+		if (!EnumHasAnyFlags(PropertiesToAllocate, EPCGPointNativeProperties::Steepness) && NeedToAllocateSingleValueProperty<EPCGPointNativeProperties::Steepness>(PointDatas))
 		{
 			PropertiesToAllocate |= EPCGPointNativeProperties::Steepness;
 		}
 
-		if (!EnumHasAnyFlags(PropertiesToAllocate, EPCGPointNativeProperties::Seed) && NeedToAllocateSingleValueProperty<int32>(EPCGPointNativeProperties::Seed, PointDatas))
+		if (!EnumHasAnyFlags(PropertiesToAllocate, EPCGPointNativeProperties::Seed) && NeedToAllocateSingleValueProperty<EPCGPointNativeProperties::Seed>(PointDatas))
 		{
 			PropertiesToAllocate |= EPCGPointNativeProperties::Seed;
 		}
 
-		if (!EnumHasAnyFlags(PropertiesToAllocate, EPCGPointNativeProperties::MetadataEntry) && NeedToAllocateSingleValueProperty<int64>(EPCGPointNativeProperties::MetadataEntry, PointDatas))
+		if (!EnumHasAnyFlags(PropertiesToAllocate, EPCGPointNativeProperties::MetadataEntry) && NeedToAllocateSingleValueProperty<EPCGPointNativeProperties::MetadataEntry>(PointDatas))
 		{
 			PropertiesToAllocate |= EPCGPointNativeProperties::MetadataEntry;
 		}
@@ -239,21 +263,22 @@ protected:
 
 private:
 
-	template<class T>
-	static bool NeedToAllocateSingleValueProperty(EPCGPointNativeProperties InProperty, const TArray<const UPCGBasePointData*>& PointDatas)
+	template<EPCGPointNativeProperties Property>
+	static bool NeedToAllocateSingleValueProperty(const TArray<const UPCGBasePointData*>& PointDatas)
 	{
+		using PropertyType = typename TPCGPointNativeProperty<Property>::Type;
 		// Here we are comparing single values (non-allocated properties) to see if we have single values that differ between inputs.
 		// If they do differ, we need to allocate those properties because the output will have multiple values.
-		TOptional<const T> DefaultValue;
+		TOptional<const PropertyType> DefaultValue;
 		for (int32 i = 0; i < PointDatas.Num(); ++i)
 		{
-			const TConstPCGValueRange<T> ValueRange = PointDatas[i]->GetConstValueRange<T>(InProperty);
-			TOptional<const T> SingleValue = ValueRange.GetSingleValue();
+			const typename TPCGPointNativeProperty<Property>::ConstValueRange ValueRange = PointDatas[i]->GetConstValueRange<Property>();
+			TOptional<const PropertyType> SingleValue = ValueRange.GetSingleValue();
 			if (!DefaultValue.IsSet())
 			{
 				DefaultValue = SingleValue;
 			}
-			else if (DefaultValue.IsSet() && SingleValue.IsSet() && !PCG::Private::MetadataTraits<T>::Equal(DefaultValue.GetValue(), SingleValue.GetValue()))
+			else if (DefaultValue.IsSet() && SingleValue.IsSet() && !PCG::Private::MetadataTraits<PropertyType>::Equal(DefaultValue.GetValue(), SingleValue.GetValue()))
 			{
 				return true;
 			}
@@ -272,6 +297,90 @@ protected:
 	mutable bool bBoundsAreDirty = true;
 };
 
+template<EPCGPointNativeProperties Property>
+inline TPCGPointNativeProperty<Property>::ValueRange UPCGBasePointData::GetValueRange(bool bAllocate)
+{
+	if constexpr (Property == EPCGPointNativeProperties::Transform)
+	{
+		return GetTransformValueRange(bAllocate);
+	}
+	else if constexpr (Property == EPCGPointNativeProperties::Density)
+	{
+		return GetDensityValueRange(bAllocate);
+	}
+	else if constexpr (Property == EPCGPointNativeProperties::Steepness)
+	{
+		return GetSteepnessValueRange(bAllocate);
+	}
+	else if constexpr (Property == EPCGPointNativeProperties::BoundsMin)
+	{
+		return GetBoundsMinValueRange(bAllocate);
+	}
+	else if constexpr (Property == EPCGPointNativeProperties::BoundsMax)
+	{
+		return GetBoundsMaxValueRange(bAllocate);
+	}
+	else if constexpr (Property == EPCGPointNativeProperties::Color)
+	{
+		return GetColorValueRange(bAllocate);
+	}
+	else if constexpr (Property == EPCGPointNativeProperties::Seed)
+	{
+		return GetSeedValueRange(bAllocate);
+	}
+	else if constexpr (Property == EPCGPointNativeProperties::MetadataEntry)
+	{
+		return GetMetadataEntryValueRange(bAllocate);
+	}
+	else
+	{
+		static_assert(TPCGPointUnsupportedProperty<Property>::Value, "Unsupported EPCGPointNativeProperties value");
+		return TPCGPointNativeProperty<Property>::ValueRange();
+	}
+}
+
+template<EPCGPointNativeProperties Property>
+inline TPCGPointNativeProperty<Property>::ConstValueRange UPCGBasePointData::GetConstValueRange() const
+{
+	if constexpr (Property == EPCGPointNativeProperties::Transform)
+	{
+		return GetConstTransformValueRange();
+	}
+	else if constexpr (Property == EPCGPointNativeProperties::Density)
+	{
+		return GetConstDensityValueRange();
+	}
+	else if constexpr (Property == EPCGPointNativeProperties::Steepness)
+	{
+		return GetConstSteepnessValueRange();
+	}
+	else if constexpr (Property == EPCGPointNativeProperties::BoundsMin)
+	{
+		return GetConstBoundsMinValueRange();
+	}
+	else if constexpr (Property == EPCGPointNativeProperties::BoundsMax)
+	{
+		return GetConstBoundsMaxValueRange();
+	}
+	else if constexpr (Property == EPCGPointNativeProperties::Color)
+	{
+		return GetConstColorValueRange();
+	}
+	else if constexpr (Property == EPCGPointNativeProperties::Seed)
+	{
+		return GetConstSeedValueRange();
+	}
+	else if constexpr (Property == EPCGPointNativeProperties::MetadataEntry)
+	{
+		return GetConstMetadataEntryValueRange();
+	}
+	else
+	{
+		static_assert(TPCGPointUnsupportedProperty<Property>::Value, "Unsupported EPCGPointNativeProperties value");
+		return TPCGPointNativeProperty<Property>::ConstValueRange();
+	}
+}
+
 template <typename T>
 inline TPCGValueRange<T> UPCGBasePointData::GetValueRange(EPCGPointNativeProperties NativeProperty, bool bAllocate)
 {
diff --git a/Engine/Plugins/PCG/Source/PCG/Public/Data/PCGPointArrayData.h b/Engine/Plugins/PCG/Source/PCG/Public/Data/PCGPointArrayData.h
--- a/Engine/Plugins/PCG/Source/PCG/Public/Data/PCGPointArrayData.h
+++ b/Engine/Plugins/PCG/Source/PCG/Public/Data/PCGPointArrayData.h
@@ -44,23 +44,23 @@ public:
 
 	virtual TArray<FTransform> GetTransformsCopy() const { return PointArray.GetTransformCopy(); }
 
-	UE_API virtual TPCGValueRange<FTransform> GetTransformValueRange(bool bAllocate = true) override;
-	UE_API virtual TPCGValueRange<float> GetDensityValueRange(bool bAllocate = true) override;
-	UE_API virtual TPCGValueRange<FVector> GetBoundsMinValueRange(bool bAllocate = true) override;
-	UE_API virtual TPCGValueRange<FVector> GetBoundsMaxValueRange(bool bAllocate = true) override;
-	UE_API virtual TPCGValueRange<FVector4> GetColorValueRange(bool bAllocate = true) override;
-	UE_API virtual TPCGValueRange<float> GetSteepnessValueRange(bool bAllocate = true) override;
-	UE_API virtual TPCGValueRange<int32> GetSeedValueRange(bool bAllocate = true) override;
-	UE_API virtual TPCGValueRange<int64> GetMetadataEntryValueRange(bool bAllocate = true) override;
-
-	UE_API virtual TConstPCGValueRange<FTransform> GetConstTransformValueRange() const override;
-	UE_API virtual TConstPCGValueRange<float> GetConstDensityValueRange() const override;
-	UE_API virtual TConstPCGValueRange<FVector> GetConstBoundsMinValueRange() const override;
-	UE_API virtual TConstPCGValueRange<FVector> GetConstBoundsMaxValueRange() const override;
-	UE_API virtual TConstPCGValueRange<FVector4> GetConstColorValueRange() const override;
-	UE_API virtual TConstPCGValueRange<float> GetConstSteepnessValueRange() const override;
-	UE_API virtual TConstPCGValueRange<int32> GetConstSeedValueRange() const override;
-	UE_API virtual TConstPCGValueRange<int64> GetConstMetadataEntryValueRange() const override;
+	UE_API virtual FPCGPointTransform::ValueRange GetTransformValueRange(bool bAllocate = true) override;
+	UE_API virtual FPCGPointDensity::ValueRange GetDensityValueRange(bool bAllocate = true) override;
+	UE_API virtual FPCGPointBoundsMin::ValueRange GetBoundsMinValueRange(bool bAllocate = true) override;
+	UE_API virtual FPCGPointBoundsMax::ValueRange GetBoundsMaxValueRange(bool bAllocate = true) override;
+	UE_API virtual FPCGPointColor::ValueRange GetColorValueRange(bool bAllocate = true) override;
+	UE_API virtual FPCGPointSteepness::ValueRange GetSteepnessValueRange(bool bAllocate = true) override;
+	UE_API virtual FPCGPointSeed::ValueRange GetSeedValueRange(bool bAllocate = true) override;
+	UE_API virtual FPCGPointMetadataEntry::ValueRange GetMetadataEntryValueRange(bool bAllocate = true) override;
+
+	UE_API virtual FPCGPointTransform::ConstValueRange GetConstTransformValueRange() const override;
+	UE_API virtual FPCGPointDensity::ConstValueRange GetConstDensityValueRange() const override;
+	UE_API virtual FPCGPointBoundsMin::ConstValueRange GetConstBoundsMinValueRange() const override;
+	UE_API virtual FPCGPointBoundsMax::ConstValueRange GetConstBoundsMaxValueRange() const override;
+	UE_API virtual FPCGPointColor::ConstValueRange GetConstColorValueRange() const override;
+	UE_API virtual FPCGPointSteepness::ConstValueRange GetConstSteepnessValueRange() const override;
+	UE_API virtual FPCGPointSeed::ConstValueRange GetConstSeedValueRange() const override;
+	UE_API virtual FPCGPointMetadataEntry::ConstValueRange GetConstMetadataEntryValueRange() const override;
 	//~End UPCGBasePointData Interface
 		
 protected:
@@ -69,150 +69,121 @@ protected:
 private:
 	UE_API void FlattenPropertiesIfNeeded(EPCGPointNativeProperties Property = EPCGPointNativeProperties::All);
 
-	template<class T>
-	FPCGPointArrayProperty<T>* GetProperty(EPCGPointNativeProperties Property, bool bWithInheritance = true)
+	template<EPCGPointNativeProperties Property>
+	FPCGPointArrayProperty<typename TPCGPointNativeProperty<Property>::Type>* GetProperty(bool bWithInheritance = true)
 	{
 		if (bWithInheritance)
 		{
-			if (EnumHasAnyFlags(InheritedProperties, Property))
+			if (EnumHasAnyFlags(InheritedProperties, TPCGPointNativeProperty<Property>::EnumValue))
 			{
 				check(ParentData);
-				return ParentData->GetProperty<T>(Property, bWithInheritance);
+				return ParentData->GetProperty<Property>(bWithInheritance);
 			}
 		}
 
-		if constexpr (std::is_same_v<T, FTransform>)
+		if constexpr (Property == EPCGPointNativeProperties::Transform)
 		{
-			if (Property == EPCGPointNativeProperties::Transform)
-			{
-				return &PointArray.Transform;
-			}
+			return &PointArray.Transform;
 		}
-		else if constexpr (std::is_same_v<T, float>)
+		else if constexpr (Property == EPCGPointNativeProperties::Density)
 		{
-			if (Property == EPCGPointNativeProperties::Density)
-			{
-				return &PointArray.Density;
-			}
-			else if (Property == EPCGPointNativeProperties::Steepness)
-			{
-				return &PointArray.Steepness;
-			}
+			return &PointArray.Density;
 		}
-		else if constexpr (std::is_same_v<T, FVector>)
+		else if constexpr (Property == EPCGPointNativeProperties::Steepness)
 		{
-			if (Property == EPCGPointNativeProperties::BoundsMin)
-			{
-				return &PointArray.BoundsMin;
-			}
-			else if (Property == EPCGPointNativeProperties::BoundsMax)
-			{
-				return &PointArray.BoundsMax;
-			}
+			return &PointArray.Steepness;
 		}
-		else if constexpr (std::is_same_v<T, FVector4>)
+		else if constexpr (Property == EPCGPointNativeProperties::BoundsMin)
 		{
-			if (Property == EPCGPointNativeProperties::Color)
-			{
-				return &PointArray.Color;
-			}
+			return &PointArray.BoundsMin;
 		}
-		else if constexpr (std::is_same_v<T, int32>)
+		else if constexpr (Property == EPCGPointNativeProperties::BoundsMax)
 		{
-			if (Property == EPCGPointNativeProperties::Seed)
-			{
-				return &PointArray.Seed;
-			}
+			return &PointArray.BoundsMax;
 		}
-		else if constexpr (std::is_same_v<T, int64>)
+		else if constexpr (Property == EPCGPointNativeProperties::Color)
 		{
-			if (Property == EPCGPointNativeProperties::MetadataEntry)
-			{
-				return &PointArray.MetadataEntry;
-			}
+			return &PointArray.Color;
+		}
+		else if constexpr (Property == EPCGPointNativeProperties::Seed)
+		{
+			return &PointArray.Seed;
+		}
+		else if constexpr (Property == EPCGPointNativeProperties::MetadataEntry)
+		{
+			return &PointArray.MetadataEntry;
+		}
+		else
+		{
+			static_assert(TPCGPointUnsupportedProperty<Property>::Value, "Unsupported EPCGPointNativeProperties value");
+			return nullptr;
 		}
-
-		check(false);
-		return nullptr;
 	}
 
-	template<class T>
-	const FPCGPointArrayProperty<T>* GetProperty(EPCGPointNativeProperties Property, bool bWithInheritance = true) const
+	template<EPCGPointNativeProperties Property>
+	const FPCGPointArrayProperty<typename TPCGPointNativeProperty<Property>::Type>* GetProperty(bool bWithInheritance = true) const
 	{
 		if (bWithInheritance)
 		{
-			if (EnumHasAnyFlags(InheritedProperties, Property))
+			if (EnumHasAnyFlags(InheritedProperties, TPCGPointNativeProperty<Property>::EnumValue))
 			{
 				check(ParentData);
-				return ParentData->GetProperty<T>(Property, bWithInheritance);
+				return ParentData->GetProperty<Property>(bWithInheritance);
 			}
 		}
 
-		if constexpr (std::is_same_v<T, FTransform>)
+		if constexpr (Property == EPCGPointNativeProperties::Transform)
 		{
-			if (Property == EPCGPointNativeProperties::Transform)
-			{
-				return &PointArray.Transform;
-			}
+			return &PointArray.Transform;
 		}
-		else if constexpr (std::is_same_v<T, float>)
+		else if constexpr (Property == EPCGPointNativeProperties::Density)
 		{
-			if (Property == EPCGPointNativeProperties::Density)
-			{
-				return &PointArray.Density;
-			}
-			else if (Property == EPCGPointNativeProperties::Steepness)
-			{
-				return &PointArray.Steepness;
-			}
+			return &PointArray.Density;
 		}
-		else if constexpr (std::is_same_v<T, FVector>)
+		else if constexpr (Property == EPCGPointNativeProperties::Steepness)
 		{
-			if (Property == EPCGPointNativeProperties::BoundsMin)
-			{
-				return &PointArray.BoundsMin;
-			}
-			else if (Property == EPCGPointNativeProperties::BoundsMax)
-			{
-				return &PointArray.BoundsMax;
-			}
+			return &PointArray.Steepness;
 		}
-		else if constexpr (std::is_same_v<T, FVector4>)
+		else if constexpr (Property == EPCGPointNativeProperties::BoundsMin)
 		{
-			if (Property == EPCGPointNativeProperties::Color)
-			{
-				return &PointArray.Color;
-			}
+			return &PointArray.BoundsMin;
 		}
-		else if constexpr (std::is_same_v<T, int32>)
+		else if constexpr (Property == EPCGPointNativeProperties::BoundsMax)
 		{
-			if (Property == EPCGPointNativeProperties::Seed)
-			{
-				return &PointArray.Seed;
-			}
+			return &PointArray.BoundsMax;
 		}
-		else if constexpr (std::is_same_v<T, int64>)
+		else if constexpr (Property == EPCGPointNativeProperties::Color)
 		{
-			if (Property == EPCGPointNativeProperties::MetadataEntry)
-			{
-				return &PointArray.MetadataEntry;
-			}
+			return &PointArray.Color;
+		}
+		else if constexpr (Property == EPCGPointNativeProperties::Seed)
+		{
+			return &PointArray.Seed;
+		}
+		else if constexpr (Property == EPCGPointNativeProperties::MetadataEntry)
+		{
+			return &PointArray.MetadataEntry;
+		}
+		else
+		{
+			static_assert(TPCGPointUnsupportedProperty<Property>::Value, "Unsupported EPCGPointNativeProperties value");
+			return nullptr;
 		}
-
-		check(false);
-		return nullptr;
 	}
 
-	template<class T>
-	bool FlattenPropertyIfNeeded(EPCGPointNativeProperties NativeProperty)
+	template<EPCGPointNativeProperties NativeProperty>
+	bool FlattenPropertyIfNeeded()
 	{
-		if (!EnumHasAnyFlags(InheritedProperties, NativeProperty))
+		using PropertyType = typename TPCGPointNativeProperty<NativeProperty>::Type;
+		EPCGPointNativeProperties EnumValue = TPCGPointNativeProperty<NativeProperty>::EnumValue;
+
+		if (!EnumHasAnyFlags(InheritedProperties, EnumValue))
 		{
 			return false;
 		}
 
-		const FPCGPointArrayProperty<T>* InheritedProperty = GetProperty<T>(NativeProperty, /*bWithInheritance=*/true);
-		FPCGPointArrayProperty<T>* Property = GetProperty<T>(NativeProperty, /*bWithInheritance=*/false);
+		const FPCGPointArrayProperty<PropertyType>* InheritedProperty = GetProperty<NativeProperty>(/*bWithInheritance=*/true);
+		FPCGPointArrayProperty<PropertyType>* Property = GetProperty<NativeProperty>(/*bWithInheritance=*/false);
 
 		// Should not be equal since InheritedProperties contains that property
 		if (ensure(InheritedProperty != Property))
@@ -226,7 +197,7 @@ private:
 			InheritedProperty->CopyTo(*Property, 0, 0, Property->Num());
 		}
 
-		EnumRemoveFlags(InheritedProperties, NativeProperty);
+		EnumRemoveFlags(InheritedProperties, EnumValue);
 
 		if (InheritedProperties == EPCGPointNativeProperties::None)
 		{
diff --git a/Engine/Plugins/PCG/Source/PCG/Public/Data/PCGPointData.h b/Engine/Plugins/PCG/Source/PCG/Public/Data/PCGPointData.h
--- a/Engine/Plugins/PCG/Source/PCG/Public/Data/PCGPointData.h
+++ b/Engine/Plugins/PCG/Source/PCG/Public/Data/PCGPointData.h
@@ -93,23 +93,23 @@ UE_API PRAGMA_ENABLE_DEPRECATION_WARNINGS
 
 	UE_API virtual TArray<FTransform> GetTransformsCopy() const override;
 
-	UE_API virtual TPCGValueRange<FTransform> GetTransformValueRange(bool bAllocate = true) override;
-	UE_API virtual TPCGValueRange<float> GetDensityValueRange(bool bAllocate = true) override;
-	UE_API virtual TPCGValueRange<FVector> GetBoundsMinValueRange(bool bAllocate = true) override;
-	UE_API virtual TPCGValueRange<FVector> GetBoundsMaxValueRange(bool bAllocate = true) override;
-	UE_API virtual TPCGValueRange<FVector4> GetColorValueRange(bool bAllocate = true) override;
-	UE_API virtual TPCGValueRange<float> GetSteepnessValueRange(bool bAllocate = true) override;
-	UE_API virtual TPCGValueRange<int32> GetSeedValueRange(bool bAllocate = true) override;
-	UE_API virtual TPCGValueRange<int64> GetMetadataEntryValueRange(bool bAllocate = true) override;
-
-	UE_API virtual TConstPCGValueRange<FTransform> GetConstTransformValueRange() const override;
-	UE_API virtual TConstPCGValueRange<float> GetConstDensityValueRange() const override;
-	UE_API virtual TConstPCGValueRange<FVector> GetConstBoundsMinValueRange() const override;
-	UE_API virtual TConstPCGValueRange<FVector> GetConstBoundsMaxValueRange() const override;
-	UE_API virtual TConstPCGValueRange<FVector4> GetConstColorValueRange() const override;
-	UE_API virtual TConstPCGValueRange<float> GetConstSteepnessValueRange() const override;
-	UE_API virtual TConstPCGValueRange<int32> GetConstSeedValueRange() const override;
-	UE_API virtual TConstPCGValueRange<int64> GetConstMetadataEntryValueRange() const override;
+	UE_API virtual FPCGPointTransform::ValueRange GetTransformValueRange(bool bAllocate = true) override;
+	UE_API virtual FPCGPointDensity::ValueRange GetDensityValueRange(bool bAllocate = true) override;
+	UE_API virtual FPCGPointBoundsMin::ValueRange GetBoundsMinValueRange(bool bAllocate = true) override;
+	UE_API virtual FPCGPointBoundsMax::ValueRange GetBoundsMaxValueRange(bool bAllocate = true) override;
+	UE_API virtual FPCGPointColor::ValueRange GetColorValueRange(bool bAllocate = true) override;
+	UE_API virtual FPCGPointSteepness::ValueRange GetSteepnessValueRange(bool bAllocate = true) override;
+	UE_API virtual FPCGPointSeed::ValueRange GetSeedValueRange(bool bAllocate = true) override;
+	UE_API virtual FPCGPointMetadataEntry::ValueRange GetMetadataEntryValueRange(bool bAllocate = true) override;
+
+	UE_API virtual FPCGPointTransform::ConstValueRange GetConstTransformValueRange() const override;
+	UE_API virtual FPCGPointDensity::ConstValueRange GetConstDensityValueRange() const override;
+	UE_API virtual FPCGPointBoundsMin::ConstValueRange GetConstBoundsMinValueRange() const override;
+	UE_API virtual FPCGPointBoundsMax::ConstValueRange GetConstBoundsMaxValueRange() const override;
+	UE_API virtual FPCGPointColor::ConstValueRange GetConstColorValueRange() const override;
+	UE_API virtual FPCGPointSteepness::ConstValueRange GetConstSteepnessValueRange() const override;
+	UE_API virtual FPCGPointSeed::ConstValueRange GetConstSeedValueRange() const override;
+	UE_API virtual FPCGPointMetadataEntry::ConstValueRange GetConstMetadataEntryValueRange() const override;
 	
 	UE_API virtual void MoveRange(int32 RangeStartIndex, int32 MoveToIndex, int32 NumElements) override;
 
diff --git a/Engine/Plugins/PCG/Source/PCG/Public/Metadata/Accessors/PCGCustomAccessor.h b/Engine/Plugins/PCG/Source/PCG/Public/Metadata/Accessors/PCGCustomAccessor.h
--- a/Engine/Plugins/PCG/Source/PCG/Public/Metadata/Accessors/PCGCustomAccessor.h
+++ b/Engine/Plugins/PCG/Source/PCG/Public/Metadata/Accessors/PCGCustomAccessor.h
@@ -339,7 +339,7 @@ namespace PCGCustomAccessor
 }
 
 template <typename T, typename R = T>
-class FPCGNativePointPropertyConstAccessor : public IPCGAttributeAccessorT<FPCGNativePointPropertyConstAccessor<T, R>>
+class UE_DEPRECATED(5.7, "Use FPCGNativePointPropertyEnumConstAccessor instead") FPCGNativePointPropertyConstAccessor : public IPCGAttributeAccessorT<FPCGNativePointPropertyConstAccessor<T, R>>
 {
 public:
 	using Type = T;
@@ -372,7 +372,7 @@ private:
 };
 
 template <typename T, typename R = T>
-class FPCGNativePointPropertyAccessor : public IPCGAttributeAccessorT<FPCGNativePointPropertyAccessor<T, R>>
+class UE_DEPRECATED(5.7, "Use FPCGNativePointPropertyEnumAccessor instead") FPCGNativePointPropertyAccessor : public IPCGAttributeAccessorT<FPCGNativePointPropertyAccessor<T, R>>
 {
 public:
 	using Type = T;
@@ -431,6 +431,101 @@ private:
 	TPCGValueRange<R> ValueRange;
 };
 
+template <EPCGPointNativeProperties Property>
+class FPCGNativePointPropertyEnumConstAccessor : public IPCGAttributeAccessorT<FPCGNativePointPropertyEnumConstAccessor<Property>>
+{
+public:
+	using Type = TPCGPointNativeProperty<Property>::AccessorType;
+	using PropertyType = TPCGPointNativeProperty<Property>::Type;
+	using Super = IPCGAttributeAccessorT<FPCGNativePointPropertyEnumConstAccessor<Property>>;
+
+	FPCGNativePointPropertyEnumConstAccessor(const UPCGBasePointData* InPointData)
+		: Super(/*bReadOnly=*/true)
+	{
+		ValueRange = InPointData->GetConstValueRange<Property>();
+	}
+
+	FPCGNativePointPropertyEnumConstAccessor(UPCGBasePointData* InPointData)
+		: Super(/*bReadOnly=*/true)
+	{
+		ValueRange = InPointData->GetConstValueRange<Property>();
+	}
+
+	bool GetRangeImpl(TArrayView<Type> OutValues, int32 StartIndex, const IPCGAttributeAccessorKeys& Keys) const
+	{
+		return PCGCustomAccessor::GetRange(OutValues, StartIndex, Keys, ValueRange);
+	}
+
+	bool SetRangeImpl(TArrayView<const Type> InValues, int32 StartIndex, IPCGAttributeAccessorKeys& Keys, EPCGAttributeAccessorFlags Flags)
+	{
+		return false;
+	}
+
+private:
+	typename TPCGPointNativeProperty<Property>::ConstValueRange ValueRange;
+};
+
+template <EPCGPointNativeProperties Property>
+class FPCGNativePointPropertyEnumAccessor : public IPCGAttributeAccessorT<FPCGNativePointPropertyEnumAccessor<Property>>
+{
+public:
+	using Type = TPCGPointNativeProperty<Property>::AccessorType;
+	using PropertyType = TPCGPointNativeProperty<Property>::Type;
+	using Super = IPCGAttributeAccessorT<FPCGNativePointPropertyEnumAccessor<Property>>;
+
+	FPCGNativePointPropertyEnumAccessor(const UPCGBasePointData* InPointData)
+		: Super(/*bReadOnly=*/false)
+	{
+		ValueRange = const_cast<UPCGBasePointData*>(InPointData)->GetValueRange<Property>(/*bAllocate=*/ true);
+	}
+
+	FPCGNativePointPropertyEnumAccessor(UPCGBasePointData* InPointData)
+		: Super(/*bReadOnly=*/false)
+	{
+		ValueRange = InPointData->GetValueRange<Property>(/*bAllocate=*/true);
+	}
+
+	bool GetRangeImpl(TArrayView<Type> OutValues, int32 StartIndex, const IPCGAttributeAccessorKeys& Keys) const
+	{
+		return PCGCustomAccessor::GetRange<Type>(OutValues, StartIndex, Keys, ValueRange);
+	}
+
+	bool SetRangeImpl(TArrayView<const Type> InValues, int32 StartIndex, IPCGAttributeAccessorKeys& Keys, EPCGAttributeAccessorFlags Flags)
+	{
+		if (Keys.IsReadOnly())
+		{
+			return false;
+		}
+
+		TArray<int32> KeyIndices;
+		bool bContiguous = false;
+		if (!Keys.GetKeyIndices(StartIndex, InValues.Num(), KeyIndices, bContiguous))
+		{
+			return false;
+		}
+
+		if (bContiguous)
+		{
+			for (int32 InIndex = 0; InIndex < InValues.Num(); ++InIndex)
+			{
+				ValueRange[StartIndex + InIndex] = (PropertyType)InValues[InIndex];
+			}
+		}
+		else
+		{
+			for (int32 InIndex = 0; InIndex < InValues.Num(); ++InIndex)
+			{
+				ValueRange[KeyIndices[InIndex]] = (PropertyType)InValues[InIndex];
+			}
+		}
+
+		return true;
+	}
+
+private:
+	typename TPCGPointNativeProperty<Property>::ValueRange ValueRange;
+};
+
 template <typename T, typename... Args>
 using FPointCustomPropertyGetter = TFunction<bool(int32, T&, Args...)>;
 
diff --git a/Engine/Plugins/PCG/Source/PCG/Public/PCGCommon.h b/Engine/Plugins/PCG/Source/PCG/Public/PCGCommon.h
--- a/Engine/Plugins/PCG/Source/PCG/Public/PCGCommon.h
+++ b/Engine/Plugins/PCG/Source/PCG/Public/PCGCommon.h
@@ -2,6 +2,8 @@
 
 #pragma once
 
+#include "PCGPointPropertiesTraits.h"
+
 #include "HAL/IConsoleManager.h"
 
 #include "PCGCommon.generated.h"
@@ -64,23 +66,6 @@ namespace PCGPointCustomPropertyNames
 	const FName ScaledLocalSizeName = TEXT("ScaledLocalSize");
 }
 
-// Value names need to match EPCGPointProperties
-UENUM(meta = (Bitflags))
-enum class EPCGPointNativeProperties : uint32
-{
-	None = 0 UMETA(Hidden),
-	Transform = 1 << 0,
-	Density = 1 << 1,
-	BoundsMin = 1 << 2,
-	BoundsMax = 1 << 3,
-	Color = 1 << 4,
-	Steepness = 1 << 5,
-	Seed = 1 << 6,
-	MetadataEntry = 1 << 7,
-	All = Transform | Density | BoundsMin | BoundsMax | Color | Steepness | Seed | MetadataEntry,
-};
-ENUM_CLASS_FLAGS(EPCGPointNativeProperties);
-
 UENUM(meta = (Bitflags))
 enum class EPCGChangeType : uint32
 {
diff --git a/Engine/Plugins/PCG/Source/PCG/Public/PCGPointPropertiesTraits.h b/Engine/Plugins/PCG/Source/PCG/Public/PCGPointPropertiesTraits.h
new file mode 100644
--- /dev/null
+++ b/Engine/Plugins/PCG/Source/PCG/Public/PCGPointPropertiesTraits.h
@@ -0,0 +1,133 @@
+// Copyright Epic Games, Inc. All Rights Reserved.
+
+#pragma once
+
+#include "Utils/PCGValueRange.h"
+
+#include "PCGPointPropertiesTraits.generated.h"
+
+// Value names need to match EPCGPointProperties
+UENUM(meta = (Bitflags))
+enum class EPCGPointNativeProperties : uint32
+{
+	None = 0 UMETA(Hidden),
+	Transform = 1 << 0,
+	Density = 1 << 1,
+	BoundsMin = 1 << 2,
+	BoundsMax = 1 << 3,
+	Color = 1 << 4,
+	Steepness = 1 << 5,
+	Seed = 1 << 6,
+	MetadataEntry = 1 << 7,
+	All = Transform | Density | BoundsMin | BoundsMax | Color | Steepness | Seed | MetadataEntry,
+};
+ENUM_CLASS_FLAGS(EPCGPointNativeProperties);
+
+template<EPCGPointNativeProperties Property>
+struct TPCGBasePointNativeProperty
+{
+	static constexpr EPCGPointNativeProperties EnumValue = Property;
+};
+
+template<EPCGPointNativeProperties Property>
+struct TPCGPointUnsupportedProperty
+{
+	static constexpr bool Value = false;
+};
+
+template<EPCGPointNativeProperties Property>
+struct TPCGPointNativeProperty : public TPCGBasePointNativeProperty<Property>
+{
+	using Type = void;
+	using AccessorType = Type;
+	using ValueRange = TPCGValueRange<Type>;
+	using ConstValueRange = TConstPCGValueRange<Type>;
+};
+
+template<>
+struct TPCGPointNativeProperty<EPCGPointNativeProperties::Transform> : public TPCGBasePointNativeProperty<EPCGPointNativeProperties::Transform>
+{
+	using Type = FTransform;
+	using AccessorType = Type;
+	using ValueRange = TPCGValueRange<Type>;
+	using ConstValueRange = TConstPCGValueRange<Type>;
+};
+
+using FPCGPointTransform = TPCGPointNativeProperty<EPCGPointNativeProperties::Transform>;
+
+template<>
+struct TPCGPointNativeProperty<EPCGPointNativeProperties::Density> : public TPCGBasePointNativeProperty<EPCGPointNativeProperties::Density>
+{
+	using Type = float;
+	using AccessorType = double;
+	using ValueRange = TPCGValueRange<Type>;
+	using ConstValueRange = TConstPCGValueRange<Type>;
+};
+
+using FPCGPointDensity = TPCGPointNativeProperty<EPCGPointNativeProperties::Density>;
+
+template<>
+struct TPCGPointNativeProperty<EPCGPointNativeProperties::BoundsMin> : public TPCGBasePointNativeProperty<EPCGPointNativeProperties::BoundsMin>
+{
+	using Type = FVector;
+	using AccessorType = Type;
+	using ValueRange = TPCGValueRange<Type>;
+	using ConstValueRange = TConstPCGValueRange<Type>;
+};
+
+using FPCGPointBoundsMin = TPCGPointNativeProperty<EPCGPointNativeProperties::BoundsMin>;
+
+template<>
+struct TPCGPointNativeProperty<EPCGPointNativeProperties::BoundsMax> : public TPCGBasePointNativeProperty<EPCGPointNativeProperties::BoundsMax>
+{
+	using Type = FVector;
+	using AccessorType = Type;
+	using ValueRange = TPCGValueRange<Type>;
+	using ConstValueRange = TConstPCGValueRange<Type>;
+};
+
+using FPCGPointBoundsMax = TPCGPointNativeProperty<EPCGPointNativeProperties::BoundsMax>;
+
+template<>
+struct TPCGPointNativeProperty<EPCGPointNativeProperties::Color> : public TPCGBasePointNativeProperty<EPCGPointNativeProperties::Color>
+{
+	using Type = FVector4;
+	using AccessorType = Type;
+	using ValueRange = TPCGValueRange<Type>;
+	using ConstValueRange = TConstPCGValueRange<Type>;
+};
+
+using FPCGPointColor = TPCGPointNativeProperty<EPCGPointNativeProperties::Color>;
+
+template<>
+struct TPCGPointNativeProperty<EPCGPointNativeProperties::Steepness> : public TPCGBasePointNativeProperty<EPCGPointNativeProperties::Steepness>
+{
+	using Type = float;
+	using AccessorType = double;
+	using ValueRange = TPCGValueRange<Type>;
+	using ConstValueRange = TConstPCGValueRange<Type>;
+};
+
+using FPCGPointSteepness = TPCGPointNativeProperty<EPCGPointNativeProperties::Steepness>;
+
+template<>
+struct TPCGPointNativeProperty<EPCGPointNativeProperties::Seed> : public TPCGBasePointNativeProperty<EPCGPointNativeProperties::Seed>
+{
+	using Type = int32;
+	using AccessorType = Type;
+	using ValueRange = TPCGValueRange<Type>;
+	using ConstValueRange = TConstPCGValueRange<Type>;
+};
+
+using FPCGPointSeed = TPCGPointNativeProperty<EPCGPointNativeProperties::Seed>;
+
+template<>
+struct TPCGPointNativeProperty<EPCGPointNativeProperties::MetadataEntry> : public TPCGBasePointNativeProperty<EPCGPointNativeProperties::MetadataEntry>
+{
+	using Type = int64;
+	using AccessorType = Type;
+	using ValueRange = TPCGValueRange<Type>;
+	using ConstValueRange = TConstPCGValueRange<Type>;
+};
+
+using FPCGPointMetadataEntry = TPCGPointNativeProperty<EPCGPointNativeProperties::MetadataEntry>;
\ No newline at end of file
