[PCG] Fix Android compilation

* On behalf of [REDACTED_HANDLE]

#rnx

diff --git a/Engine/Plugins/PCG/Source/PCG/Private/Data/PCGPointData.cpp b/Engine/Plugins/PCG/Source/PCG/Private/Data/PCGPointData.cpp
--- a/Engine/Plugins/PCG/Source/PCG/Private/Data/PCGPointData.cpp
+++ b/Engine/Plugins/PCG/Source/PCG/Private/Data/PCGPointData.cpp
@@ -16,7 +16,7 @@
 
 #include UE_INLINE_GENERATED_CPP_BY_NAME(PCGPointData)
 
-#define PCG_POINTDATA_MAKERANGE_BASE(RangeType, MakeViewType, Property, Field) RangeType<TPCGPointNativeProperty<Property>::Type>(MakeViewType(sizeof(FPCGPoint), Points.IsEmpty() ? nullptr : &Points[0].Field, Points.Num()));
+#define PCG_POINTDATA_MAKERANGE_BASE(RangeType, MakeViewType, Property, Field) RangeType<typename TPCGPointNativeProperty<Property>::Type>(MakeViewType(sizeof(FPCGPoint), Points.IsEmpty() ? nullptr : &Points[0].Field, Points.Num()));
 
 #define PCG_POINTDATA_MAKERANGE(Property, Field) PCG_POINTDATA_MAKERANGE_BASE(TPCGValueRange, MakeStridedView, Property, Field)
 #define PCG_POINTDATA_MAKECONSTRANGE(Property, Field) PCG_POINTDATA_MAKERANGE_BASE(TConstPCGValueRange, MakeConstStridedView, Property, Field)
diff --git a/Engine/Plugins/PCG/Source/PCG/Public/Data/PCGBasePointData.h b/Engine/Plugins/PCG/Source/PCG/Public/Data/PCGBasePointData.h
--- a/Engine/Plugins/PCG/Source/PCG/Public/Data/PCGBasePointData.h
+++ b/Engine/Plugins/PCG/Source/PCG/Public/Data/PCGBasePointData.h
@@ -91,7 +91,7 @@ public:
 	TPCGValueRange<T> GetValueRange(EPCGPointNativeProperties NativeProperty, bool bAllocate = true);
 
 	template<EPCGPointNativeProperties Property> 
-	TPCGPointNativeProperty<Property>::ValueRange GetValueRange(bool bAllocate);
+	typename TPCGPointNativeProperty<Property>::ValueRange GetValueRange(bool bAllocate);
 
 	virtual FPCGPointTransform::ValueRange GetTransformValueRange(bool bAllocate = true) PURE_VIRTUAL(UPCGBasePointData::GetTransformValueRange, return FPCGPointTransform::ValueRange(););
 	virtual FPCGPointDensity::ValueRange GetDensityValueRange(bool bAllocate = true) PURE_VIRTUAL(UPCGBasePointData::GetDensityValueRange, return FPCGPointDensity::ValueRange(););
@@ -107,7 +107,7 @@ public:
 	TConstPCGValueRange<T> GetConstValueRange(EPCGPointNativeProperties NativeProperty) const;
 
 	template<EPCGPointNativeProperties Property>
-	TPCGPointNativeProperty<Property>::ConstValueRange GetConstValueRange() const;
+	typename TPCGPointNativeProperty<Property>::ConstValueRange GetConstValueRange() const;
 
 	virtual FPCGPointTransform::ConstValueRange GetConstTransformValueRange() const PURE_VIRTUAL(UPCGBasePointData::GetConstTransformValueRange, return FPCGPointTransform::ConstValueRange(););
 	virtual FPCGPointDensity::ConstValueRange GetConstDensityValueRange() const PURE_VIRTUAL(UPCGBasePointData::GetConstDensityValueRange, return FPCGPointDensity::ConstValueRange(););
@@ -298,7 +298,7 @@ protected:
 };
 
 template<EPCGPointNativeProperties Property>
-inline TPCGPointNativeProperty<Property>::ValueRange UPCGBasePointData::GetValueRange(bool bAllocate)
+typename TPCGPointNativeProperty<Property>::ValueRange UPCGBasePointData::GetValueRange(bool bAllocate)
 {
 	if constexpr (Property == EPCGPointNativeProperties::Transform)
 	{
@@ -340,7 +340,7 @@ inline TPCGPointNativeProperty<Property>::ValueRange UPCGBasePointData::GetValue
 }
 
 template<EPCGPointNativeProperties Property>
-inline TPCGPointNativeProperty<Property>::ConstValueRange UPCGBasePointData::GetConstValueRange() const
+typename TPCGPointNativeProperty<Property>::ConstValueRange UPCGBasePointData::GetConstValueRange() const
 {
 	if constexpr (Property == EPCGPointNativeProperties::Transform)
 	{
diff --git a/Engine/Plugins/PCG/Source/PCG/Public/Metadata/Accessors/PCGCustomAccessor.h b/Engine/Plugins/PCG/Source/PCG/Public/Metadata/Accessors/PCGCustomAccessor.h
--- a/Engine/Plugins/PCG/Source/PCG/Public/Metadata/Accessors/PCGCustomAccessor.h
+++ b/Engine/Plugins/PCG/Source/PCG/Public/Metadata/Accessors/PCGCustomAccessor.h
@@ -435,8 +435,8 @@ template <EPCGPointNativeProperties Property>
 class FPCGNativePointPropertyEnumConstAccessor : public IPCGAttributeAccessorT<FPCGNativePointPropertyEnumConstAccessor<Property>>
 {
 public:
-	using Type = TPCGPointNativeProperty<Property>::AccessorType;
-	using PropertyType = TPCGPointNativeProperty<Property>::Type;
+	using Type = typename TPCGPointNativeProperty<Property>::AccessorType;
+	using PropertyType = typename TPCGPointNativeProperty<Property>::Type;
 	using Super = IPCGAttributeAccessorT<FPCGNativePointPropertyEnumConstAccessor<Property>>;
 
 	FPCGNativePointPropertyEnumConstAccessor(const UPCGBasePointData* InPointData)
@@ -469,8 +469,8 @@ template <EPCGPointNativeProperties Property>
 class FPCGNativePointPropertyEnumAccessor : public IPCGAttributeAccessorT<FPCGNativePointPropertyEnumAccessor<Property>>
 {
 public:
-	using Type = TPCGPointNativeProperty<Property>::AccessorType;
-	using PropertyType = TPCGPointNativeProperty<Property>::Type;
+	using Type = typename TPCGPointNativeProperty<Property>::AccessorType;
+	using PropertyType = typename TPCGPointNativeProperty<Property>::Type;
 	using Super = IPCGAttributeAccessorT<FPCGNativePointPropertyEnumAccessor<Property>>;
 
 	FPCGNativePointPropertyEnumAccessor(const UPCGBasePointData* InPointData)
