Work around android compiler crash
#rnx

diff --git a/Engine/Plugins/PCG/Source/PCG/Private/PCGGraphExecutionStateInterface.cpp b/Engine/Plugins/PCG/Source/PCG/Private/PCGGraphExecutionStateInterface.cpp
--- a/Engine/Plugins/PCG/Source/PCG/Private/PCGGraphExecutionStateInterface.cpp
+++ b/Engine/Plugins/PCG/Source/PCG/Private/PCGGraphExecutionStateInterface.cpp
@@ -39,7 +39,7 @@ FPCGSoftGraphExecutionSource::FPCGSoftGraphExecutionSource(const TSoftObjectPtr<
 }
 
 FPCGSoftGraphExecutionSource::FPCGSoftGraphExecutionSource(IPCGGraphExecutionSource* InSource)
-	: FPCGSoftGraphExecutionSource(TSoftObjectPtr(Cast<UObject>(InSource)))
+	: SoftObjectPtr(Cast<UObject>(InSource))
 {
 }
 
