[PCG] Forced plugin/engine template visibility in the create graph from template dialog

#rnx

diff --git a/Engine/Plugins/PCG/Source/PCGEditor/Private/PCGEditorUtils.cpp b/Engine/Plugins/PCG/Source/PCGEditor/Private/PCGEditorUtils.cpp
--- a/Engine/Plugins/PCG/Source/PCGEditor/Private/PCGEditorUtils.cpp
+++ b/Engine/Plugins/PCG/Source/PCGEditor/Private/PCGEditorUtils.cpp
@@ -258,6 +258,8 @@ bool PCGEditorUtils::PickGraphTemplate(FAssetData& OutAssetData, const FText& Ti
 	AssetPickerConfig.OnAssetSelected = FOnAssetSelected::CreateLambda(HandleAssetSelected);
 	AssetPickerConfig.SelectionMode = ESelectionMode::Single;
 	AssetPickerConfig.bAllowNullSelection = false;
+	AssetPickerConfig.bForceShowEngineContent = true;
+	AssetPickerConfig.bForceShowPluginContent = true;
 	AssetPickerConfig.InitialAssetViewType = EAssetViewType::List;
 	AssetPickerConfig.InitialThumbnailSize = EThumbnailSize::Small;
 	AssetPickerConfig.Filter.ClassPaths.Add(UPCGGraph::StaticClass()->GetClassPathName());
