[PCG] Replace PCGComponent by Execution source in PCG Editor where it is possible

* Some places are still dependent on "IsPartitioned", "GetOwner" and "GetGenerationGrid". We might want to move those ito the interface.
* We should also add the concept of Generate/Refresh on the interface. Not clear if CleanUp is necessary at that point
* Also start some export of PCGEditor symbols, even if it is still private for now

#rnx

diff --git a/Engine/Plugins/PCG/Source/PCG/Private/Graph/PCGGraphExecutor.cpp b/Engine/Plugins/PCG/Source/PCG/Private/Graph/PCGGraphExecutor.cpp
--- a/Engine/Plugins/PCG/Source/PCG/Private/Graph/PCGGraphExecutor.cpp
+++ b/Engine/Plugins/PCG/Source/PCG/Private/Graph/PCGGraphExecutor.cpp
@@ -2877,13 +2877,14 @@ bool FPCGPreGraphElement::ExecuteInternal(FPCGContext* Context) const
 	check(Context);
 
 	check(IsInGameThread());
+	const bool bIsValid = Context->ExecutionSource.IsValid();
 	UPCGComponent* Component = Cast<UPCGComponent>(Context->ExecutionSource.Get());
 
 	// Early out if the component has been deleted/is invalid
 	if (!Component)
 	{
 		// If the component should exist but it doesn't (which is all the time here, previously we checked for it), then this should be cancelled
-		Context->OutputData.bCancelExecution = true;
+		Context->OutputData.bCancelExecution = !bIsValid;
 		return true;
 	}
 
@@ -2910,13 +2911,14 @@ bool FPCGFetchInputElement::ExecuteInternal(FPCGContext* Context) const
 	Context->OutputData = Context->InputData;
 
 	// Second: fetch the inputs provided by the component
+	const bool bIsValid = Context->ExecutionSource.IsValid();
 	UPCGComponent* Component = Cast<UPCGComponent>(Context->ExecutionSource.Get());
 	
 	// Early out if the component has been deleted/is invalid
 	if (!Component)
 	{
 		// If the component should exist but it doesn't (which is all the time here, previously we checked for it), then this should be cancelled
-		Context->OutputData.bCancelExecution = true;
+		Context->OutputData.bCancelExecution = !bIsValid;
 		return true;
 	}
 
diff --git a/Engine/Plugins/PCG/Source/PCG/Private/Graph/PCGStackContext.cpp b/Engine/Plugins/PCG/Source/PCG/Private/Graph/PCGStackContext.cpp
--- a/Engine/Plugins/PCG/Source/PCG/Private/Graph/PCGStackContext.cpp
+++ b/Engine/Plugins/PCG/Source/PCG/Private/Graph/PCGStackContext.cpp
@@ -173,6 +173,11 @@ const UPCGComponent* FPCGStack::GetRootComponent() const
 	return StackFrames.IsEmpty() ? nullptr : StackFrames[0].GetObject_AnyThread<UPCGComponent>();
 }
 
+const IPCGGraphExecutionSource* FPCGStack::GetRootSource() const
+{
+	return StackFrames.IsEmpty() ? nullptr : StackFrames[0].GetObject_AnyThread<IPCGGraphExecutionSource>();
+}
+
 const UPCGGraph* FPCGStack::GetRootGraph(int32* OutRootFrameIndex) const
 {
 	if (GetStackFrames().Num() > 0)
diff --git a/Engine/Plugins/PCG/Source/PCG/Private/PCGGraphExecutionStateInterface.cpp b/Engine/Plugins/PCG/Source/PCG/Private/PCGGraphExecutionStateInterface.cpp
--- a/Engine/Plugins/PCG/Source/PCG/Private/PCGGraphExecutionStateInterface.cpp
+++ b/Engine/Plugins/PCG/Source/PCG/Private/PCGGraphExecutionStateInterface.cpp
@@ -35,7 +35,6 @@ UPCGGraphExecutionSource::UPCGGraphExecutionSource(const FObjectInitializer& Obj
 FPCGSoftGraphExecutionSource::FPCGSoftGraphExecutionSource(const TSoftObjectPtr<UObject>& InSoftObjectPtr)
 	: SoftObjectPtr(InSoftObjectPtr)
 {
-	
 }
 
 FPCGSoftGraphExecutionSource::FPCGSoftGraphExecutionSource(IPCGGraphExecutionSource* InSource)
@@ -60,16 +59,28 @@ FPCGSoftGraphExecutionSource::FPCGSoftGraphExecutionSource(FPCGSoftGraphExecutio
 
 FPCGSoftGraphExecutionSource& FPCGSoftGraphExecutionSource::operator=(const FPCGSoftGraphExecutionSource& Other)
 {
+	Reset();
+
 	SoftObjectPtr = Other.SoftObjectPtr;
 	return *this;
 }
 
 FPCGSoftGraphExecutionSource& FPCGSoftGraphExecutionSource::operator=(FPCGSoftGraphExecutionSource&& Other)
 {
+	Reset();
+
 	SoftObjectPtr = std::move(Other.SoftObjectPtr);
 	return *this;
 }
 
+FPCGSoftGraphExecutionSource& FPCGSoftGraphExecutionSource::operator=(const IPCGGraphExecutionSource* InSource)
+{
+	Reset();
+
+	*this = FPCGSoftGraphExecutionSource(InSource);
+	return *this;
+}
+
 IPCGGraphExecutionSource* FPCGSoftGraphExecutionSource::Get() const
 {
 	if (!CachedWeakPtr.IsValid())
@@ -96,6 +107,21 @@ IPCGGraphExecutionSource* FPCGSoftGraphExecutionSource::Get() const
 	return CachedWeakPtr.Get();
 }
 
+UObject* FPCGSoftGraphExecutionSource::GetObject() const
+{
+	return SoftObjectPtr.Get();
+}
+
+void FPCGSoftGraphExecutionSource::Reset()
+{
+	SoftObjectPtr = nullptr;
+	if (CachedWeakPtr.IsValid())
+	{
+		UE::TScopeLock Lock(SpinLock);
+		CachedWeakPtr.Reset();
+	}
+}
+
 bool FPCGSoftGraphExecutionSource::operator==(const FPCGSoftGraphExecutionSource& Other) const
 {
 	return Other.SoftObjectPtr == SoftObjectPtr;
diff --git a/Engine/Plugins/PCG/Source/PCG/Private/Tests/Determinism/PCGDeterminismTestsCommon.cpp b/Engine/Plugins/PCG/Source/PCG/Private/Tests/Determinism/PCGDeterminismTestsCommon.cpp
--- a/Engine/Plugins/PCG/Source/PCG/Private/Tests/Determinism/PCGDeterminismTestsCommon.cpp
+++ b/Engine/Plugins/PCG/Source/PCG/Private/Tests/Determinism/PCGDeterminismTestsCommon.cpp
@@ -70,9 +70,17 @@ namespace PCGDeterminismTests
 
 #if WITH_EDITOR
 	void RunDeterminismTest(const UPCGGraph* InPCGGraph, UPCGComponent* InPCGComponent, FDeterminismTestResult& OutResult)
+	{
+		RunDeterminismTest(InPCGGraph, static_cast<IPCGGraphExecutionSource*>(InPCGComponent), OutResult);
+	}
+	
+	void RunDeterminismTest(const UPCGGraph* InPCGGraph, IPCGGraphExecutionSource* InSource, FDeterminismTestResult& OutResult)
 	{
 		TRACE_CPUPROFILER_EVENT_SCOPE(PCGDeterminismTests::RunDeterminismTest::Graph);
 
+		// @todo_pcg: Support Determinism test on all sources.  
+		UPCGComponent* InPCGComponent = Cast<UPCGComponent>(InSource);
+		
 		if (!InPCGGraph || !InPCGComponent)
 		{
 			OutResult.DataTypesTested = EPCGDataType::None;
diff --git a/Engine/Plugins/PCG/Source/PCG/Private/Utils/PCGNodeVisualLogs.cpp b/Engine/Plugins/PCG/Source/PCG/Private/Utils/PCGNodeVisualLogs.cpp
--- a/Engine/Plugins/PCG/Source/PCG/Private/Utils/PCGNodeVisualLogs.cpp
+++ b/Engine/Plugins/PCG/Source/PCG/Private/Utils/PCGNodeVisualLogs.cpp
@@ -329,11 +329,14 @@ void FPCGNodeVisualLogs::ClearLogs(const FPCGStack& InPCGStack)
 	}
 }
 
-void FPCGNodeVisualLogs::ClearLogs(const UPCGComponent* InComponent)
+void FPCGNodeVisualLogs::ClearLogs(const IPCGGraphExecutionSource* InSource)
 {
-	FPCGStack Stack;
-	Stack.PushFrame(InComponent);
-	ClearLogs(Stack);
+	if (const UObject* SourceObject = Cast<UObject>(InSource))
+	{
+		FPCGStack Stack;
+		Stack.PushFrame(SourceObject);
+		ClearLogs(Stack);
+	}
 }
 
 #endif // WITH_EDITOR
diff --git a/Engine/Plugins/PCG/Source/PCG/Public/Editor/IPCGEditorModule.h b/Engine/Plugins/PCG/Source/PCG/Public/Editor/IPCGEditorModule.h
--- a/Engine/Plugins/PCG/Source/PCG/Public/Editor/IPCGEditorModule.h
+++ b/Engine/Plugins/PCG/Source/PCG/Public/Editor/IPCGEditorModule.h
@@ -42,14 +42,14 @@ protected:
 	virtual void OnGraphPreSave(UPCGGraph* Graph, FObjectPreSaveContext ObjectSaveContext) = 0;
 
 	/** Clear any data collected during execution, normally called prior to generating the component. */
-	virtual void ClearExecutionMetadata(UPCGComponent* InComponent) = 0;
+	virtual void ClearExecutionMetadata(IPCGGraphExecutionSource* InSource) = 0;
 
-	virtual void ClearExecutedStacks(const UPCGComponent* InRootComponent) = 0;
+	virtual void ClearExecutedStacks(const IPCGGraphExecutionSource* InRootSource) = 0;
 	virtual void ClearExecutedStacks(const UPCGGraph* InContainingGraph) = 0;
 
 	/** Get a list of stacks that were executed during the last execution. */
-	virtual	TArray<FPCGStackSharedPtr> GetExecutedStacksPtrs(const FPCGStack& BeginningWithStack) = 0;
-	virtual TArray<FPCGStackSharedPtr> GetExecutedStacksPtrs(const UPCGComponent* InComponent, const UPCGGraph* InSubgraph, bool bOnlyWithSubgraphAsCurrentFrame = true) = 0;
+	virtual TArray<FPCGStackSharedPtr> GetExecutedStacksPtrs(const FPCGStack& BeginningWithStack) = 0;
+	virtual TArray<FPCGStackSharedPtr> GetExecutedStacksPtrs(const IPCGGraphExecutionSource* InSource, const UPCGGraph* InSubgraph, bool bOnlyWithSubgraphAsCurrentFrame = true) = 0;
 
 	virtual void NotifyGraphChanged(UPCGGraph* InGraph, EPCGChangeType ChangeType) = 0;
 
diff --git a/Engine/Plugins/PCG/Source/PCG/Public/Graph/PCGStackContext.h b/Engine/Plugins/PCG/Source/PCG/Public/Graph/PCGStackContext.h
--- a/Engine/Plugins/PCG/Source/PCG/Public/Graph/PCGStackContext.h
+++ b/Engine/Plugins/PCG/Source/PCG/Public/Graph/PCGStackContext.h
@@ -155,6 +155,9 @@ public:
 	/** Component given by first stack frame. */
 	UE_API const UPCGComponent* GetRootComponent() const;
 
+	/** Execution source given by first stack frame. */
+	UE_API const IPCGGraphExecutionSource* GetRootSource() const;
+
 	/** First (top) graph frame in stack (or null if no graph frames present). */
 	UE_API const UPCGGraph* GetRootGraph(int32* OutRootFrameIndex = nullptr) const;
 
diff --git a/Engine/Plugins/PCG/Source/PCG/Public/PCGGraphExecutionInspection.h b/Engine/Plugins/PCG/Source/PCG/Public/PCGGraphExecutionInspection.h
--- a/Engine/Plugins/PCG/Source/PCG/Public/PCGGraphExecutionInspection.h
+++ b/Engine/Plugins/PCG/Source/PCG/Public/PCGGraphExecutionInspection.h
@@ -17,7 +17,6 @@ class FPCGGraphExecutionInspection
 {
 public:
 	friend struct FPCGComponentInstanceData;
-	friend class UPCGComponent;
 
 	/* Retrieves the executed nodes information */
 	struct FNodeExecutedNotificationData
@@ -61,9 +60,9 @@ public:
 	UE_API bool DidNodeTriggerGPUToCPUReadback(const UPCGNode* InNode, const FPCGStack& Stack) const;
 	UE_API bool DidNodeTriggerCPUToGPUUpload(const UPCGNode* InNode, const FPCGStack& Stack) const;
 
-private:
 	UE_API void AddReferencedObjects(FReferenceCollector& Collector);
 
+private:
 	int32 InspectionCounter = 0;
 
 	TMap<FPCGStack, FPCGDataCollection> InspectionCache;
diff --git a/Engine/Plugins/PCG/Source/PCG/Public/PCGGraphExecutionStateInterface.h b/Engine/Plugins/PCG/Source/PCG/Public/PCGGraphExecutionStateInterface.h
--- a/Engine/Plugins/PCG/Source/PCG/Public/PCGGraphExecutionStateInterface.h
+++ b/Engine/Plugins/PCG/Source/PCG/Public/PCGGraphExecutionStateInterface.h
@@ -68,7 +68,7 @@ public:
 	virtual bool IsGenerating() const = 0;
 
 	/** Returns the subsystem associated with this source. By default, it will return the UPCGSubsystem if the source has a world, otherwise, the EditorSubsystem (in editor only, obviously). */
-	virtual IPCGBaseSubsystem* GetSubsystem() const;
+	PCG_API virtual IPCGBaseSubsystem* GetSubsystem() const;
 
 #if WITH_EDITOR
 
@@ -128,7 +128,17 @@ struct FPCGSoftGraphExecutionSource
 	PCG_API FPCGSoftGraphExecutionSource& operator=(const FPCGSoftGraphExecutionSource& Other);
 	PCG_API FPCGSoftGraphExecutionSource& operator=(FPCGSoftGraphExecutionSource&& Other);
 
+	/** Can assign an execution source. Will reset the cached weak pointer. */
+	PCG_API FPCGSoftGraphExecutionSource& operator=(const IPCGGraphExecutionSource* InSource);
+
+	/** De-reference the weak pointer if it is valid, otherwise will resolve the soft pointer and update the cached weak. Threadsafe. */
 	PCG_API IPCGGraphExecutionSource* Get() const;
+
+	/** Resolve the soft pointer as a UObject. */
+	PCG_API UObject* GetObject() const;
+
+	/** Reset the soft pointer to null and invalidate the cached weak. */
+	PCG_API void Reset();
 	
 	bool IsValid() const { return Get() != nullptr; }
 	
diff --git a/Engine/Plugins/PCG/Source/PCG/Public/Tests/Determinism/PCGDeterminismTestsCommon.h b/Engine/Plugins/PCG/Source/PCG/Public/Tests/Determinism/PCGDeterminismTestsCommon.h
--- a/Engine/Plugins/PCG/Source/PCG/Public/Tests/Determinism/PCGDeterminismTestsCommon.h
+++ b/Engine/Plugins/PCG/Source/PCG/Public/Tests/Determinism/PCGDeterminismTestsCommon.h
@@ -142,7 +142,7 @@ namespace PCGDeterminismTests
 
 #if WITH_EDITOR
 	/** Validates if a PCGGraph is deterministic */
-	PCG_API void RunDeterminismTest(const UPCGGraph* InPCGGraph, UPCGComponent* InPCGComponent, FDeterminismTestResult& OutResult);
+	PCG_API void RunDeterminismTest(const UPCGGraph* InPCGGraph, IPCGGraphExecutionSource* InSource, FDeterminismTestResult& OutResult);
 #endif
 
 	/** Validates all the generic determinism tests for any given node */
diff --git a/Engine/Plugins/PCG/Source/PCG/Public/Utils/PCGNodeVisualLogs.h b/Engine/Plugins/PCG/Source/PCG/Public/Utils/PCGNodeVisualLogs.h
--- a/Engine/Plugins/PCG/Source/PCG/Public/Utils/PCGNodeVisualLogs.h
+++ b/Engine/Plugins/PCG/Source/PCG/Public/Utils/PCGNodeVisualLogs.h
@@ -69,8 +69,8 @@ public:
 	/** Clear all errors and warnings that occurred while executing stacks beginning with the given stack. */
 	UE_API void ClearLogs(const FPCGStack& InPCGStack);
 
-	/** Clear all errors and warnings corresponding to the given component. */
-	UE_API void ClearLogs(const UPCGComponent* InComponent);
+	/** Clear all errors and warnings corresponding to the given execution source. */
+	UE_API void ClearLogs(const IPCGGraphExecutionSource* InSource);
 
 	/** Generates summary text from log list, with optional component pointers to show actor label. */
 	static UE_API FText GetSummaryText(const FPCGPerNodeVisualLogs& InLogs, const TArray<const UPCGComponent*>* InComponents, ELogVerbosity::Type* OutMinimumVerbosity = nullptr);
diff --git a/Engine/Plugins/PCG/Source/PCGEditor/PCGEditor.Build.cs b/Engine/Plugins/PCG/Source/PCGEditor/PCGEditor.Build.cs
--- a/Engine/Plugins/PCG/Source/PCGEditor/PCGEditor.Build.cs
+++ b/Engine/Plugins/PCG/Source/PCGEditor/PCGEditor.Build.cs
@@ -15,7 +15,8 @@ namespace UnrealBuildTool.Rules
 					"Projects",
 					"Engine",
 					"CoreUObject",
-					"PlacementMode",
+					"PlacementMode", 
+					"PCG"
 				});
 
 			if (Target.WithAutomationTests)
@@ -52,7 +53,6 @@ namespace UnrealBuildTool.Rules
 					"InputCore",
 					"Kismet",
 					"KismetWidgets",
-					"PCG",
 					"PropertyEditor",
 					"RHI",
 					"RenderCore",
diff --git a/Engine/Plugins/PCG/Source/PCGEditor/Private/Nodes/PCGEditorGraphNodeBase.cpp b/Engine/Plugins/PCG/Source/PCGEditor/Private/Nodes/PCGEditorGraphNodeBase.cpp
--- a/Engine/Plugins/PCG/Source/PCGEditor/Private/Nodes/PCGEditorGraphNodeBase.cpp
+++ b/Engine/Plugins/PCG/Source/PCGEditor/Private/Nodes/PCGEditorGraphNodeBase.cpp
@@ -54,7 +54,7 @@ namespace PCGEditorGraphNodeBase
 	}
 
 	/** Whether this node was culled during graph compilation or during graph execution. */
-	bool ShouldDisplayAsActive(const UPCGEditorGraphNodeBase* InNode, const UPCGComponent* InComponentBeingDebugged, const FPCGStack* InStackBeingInspected)
+	bool ShouldDisplayAsActive(const UPCGEditorGraphNodeBase* InNode, const IPCGGraphExecutionSource* InSourceBeingDebugged, const FPCGStack* InStackBeingInspected)
 	{
 		if (!InNode)
 		{
@@ -69,7 +69,7 @@ namespace PCGEditorGraphNodeBase
 
 		// Don't display as culled while component is executing or about to refresh as nodes will flash to culled state and back
 		// which looks disturbing.
-		if (!InComponentBeingDebugged || InComponentBeingDebugged->IsGenerating() || InComponentBeingDebugged->IsRefreshInProgress())
+		if (!InSourceBeingDebugged || InSourceBeingDebugged->GetExecutionState().IsGenerating() || InSourceBeingDebugged->GetExecutionState().IsRefreshInProgress())
 		{
 			return true;
 		}
@@ -84,7 +84,7 @@ namespace PCGEditorGraphNodeBase
 			if (!Settings || !Settings->IsA<UPCGRerouteSettings>())
 			{
 				// Task will be displayed as active if it was executed.
-				return InComponentBeingDebugged->GetExecutionState().GetInspection().WasNodeExecuted(PCGNode, *InStackBeingInspected);
+				return InSourceBeingDebugged->GetExecutionState().GetInspection().WasNodeExecuted(PCGNode, *InStackBeingInspected);
 			}
 			else
 			{
@@ -94,7 +94,7 @@ namespace PCGEditorGraphNodeBase
 					const UPCGNode* DeclarationPCGNode = RerouteUsageSettings->Declaration ? Cast<UPCGNode>(RerouteUsageSettings->Declaration->GetOuter()) : nullptr;
 					const UPCGEditorGraph* EditorGraph = Cast<UPCGEditorGraph>(InNode->GetOuter());
 					const UPCGEditorGraphNodeBase* DeclarationNode = EditorGraph ? EditorGraph->GetEditorNodeFromPCGNode(DeclarationPCGNode) : nullptr;
-					return !DeclarationNode || ShouldDisplayAsActive(DeclarationNode, InComponentBeingDebugged, InStackBeingInspected);
+					return !DeclarationNode || ShouldDisplayAsActive(DeclarationNode, InSourceBeingDebugged, InStackBeingInspected);
 				}
 
 				// Special case - reroute culled state is evaluated here based on upstream connections. Reroutes are always culled/never executed, but still need
@@ -107,7 +107,7 @@ namespace PCGEditorGraphNodeBase
 						{
 							if (const UPCGEditorGraphNodeBase* UpstreamNode = LinkedPin ? Cast<UPCGEditorGraphNodeBase>(LinkedPin->GetOwningNode()) : nullptr)
 							{
-								const bool bUpstreamNodeActive = ShouldDisplayAsActive(UpstreamNode, InComponentBeingDebugged, InStackBeingInspected);
+								const bool bUpstreamNodeActive = ShouldDisplayAsActive(UpstreamNode, InSourceBeingDebugged, InStackBeingInspected);
 								const bool bUpstreamPinActive = UpstreamNode->IsOutputPinActive(LinkedPin);
 
 								if (bUpstreamNodeActive && bUpstreamPinActive)
@@ -692,19 +692,22 @@ EPCGChangeType UPCGEditorGraphNodeBase::UpdateErrorsAndWarnings()
 	return bStateChanged ? EPCGChangeType::Cosmetic : EPCGChangeType::None;
 }
 
-EPCGChangeType UPCGEditorGraphNodeBase::UpdateStructuralVisualization(UPCGComponent* InComponentBeingDebugged, const FPCGStack* InStackBeingInspected, bool bNewlyPlaced)
+EPCGChangeType UPCGEditorGraphNodeBase::UpdateStructuralVisualization(IPCGGraphExecutionSource* InSourceBeingDebugged, const FPCGStack* InStackBeingInspected, bool bNewlyPlaced)
 {
 	const UPCGGraph* Graph = PCGNode ? PCGNode->GetGraph() : nullptr;
-	if (!Graph)
+	if (!Graph || !InSourceBeingDebugged)
 	{
 		return EPCGChangeType::None;
 	}
 
-	const bool bInspecting = InComponentBeingDebugged && InStackBeingInspected && !InStackBeingInspected->GetStackFrames().IsEmpty();
+	const UPCGComponent* ComponentBeingDebugged = Cast<const UPCGComponent>(InSourceBeingDebugged);
+
+	const FPCGGraphExecutionInspection& Inspection = InSourceBeingDebugged->GetExecutionState().GetInspection();
+	const bool bInspecting = InStackBeingInspected && !InStackBeingInspected->GetStackFrames().IsEmpty();
 
 	EPCGChangeType ChangeType = EPCGChangeType::None;
 
-	const uint64 NewInactiveMask = bInspecting ? InComponentBeingDebugged->GetNodeInactivePinMask(PCGNode, *InStackBeingInspected) : 0;
+	const uint64 NewInactiveMask = bInspecting ? Inspection.GetNodeInactivePinMask(PCGNode, *InStackBeingInspected) : 0;
 	if (NewInactiveMask != InactiveOutputPinMask)
 	{
 		InactiveOutputPinMask = NewInactiveMask;
@@ -718,10 +721,10 @@ EPCGChangeType UPCGEditorGraphNodeBase::UpdateStructuralVisualization(UPCGCompon
 	// Set the inspected grid size - this is used for grid size visualization.
 	uint32 InspectingGridSize = PCGHiGenGrid::UninitializedGridSize();
 	EPCGHiGenGrid InspectingGrid = EPCGHiGenGrid::Uninitialized;
-	if (TopGraph && TopGraph->IsHierarchicalGenerationEnabled() && InComponentBeingDebugged && (InComponentBeingDebugged->IsPartitioned() || InComponentBeingDebugged->IsLocalComponent()))
+	if (TopGraph && TopGraph->IsHierarchicalGenerationEnabled() && ComponentBeingDebugged && (ComponentBeingDebugged->IsPartitioned() || ComponentBeingDebugged->IsLocalComponent()))
 	{
-		InspectingGridSize = InComponentBeingDebugged->GetGenerationGridSize();
-		InspectingGrid = InComponentBeingDebugged->GetGenerationGrid();
+		InspectingGridSize = ComponentBeingDebugged->GetGenerationGridSize();
+		InspectingGrid = ComponentBeingDebugged->GetGenerationGrid();
 	}
 
 	if (InspectedGenerationGrid != InspectingGrid)
@@ -740,14 +743,14 @@ EPCGChangeType UPCGEditorGraphNodeBase::UpdateStructuralVisualization(UPCGCompon
 		bShouldDisplayAsDisabled = (Graph != TopGraph) || !TopGraph->IsHierarchicalGenerationEnabled();
 
 		// If we're inspecting a component, it must either be a partitioned OC or an LC (because higen requires partitioning).
-		if (!bShouldDisplayAsDisabled && InComponentBeingDebugged)
+		if (!bShouldDisplayAsDisabled && ComponentBeingDebugged)
 		{
-			bShouldDisplayAsDisabled = !InComponentBeingDebugged->IsPartitioned() && !InComponentBeingDebugged->IsLocalComponent();
+			bShouldDisplayAsDisabled = !ComponentBeingDebugged->IsPartitioned() && !ComponentBeingDebugged->IsLocalComponent();
 		}
 	}
 
 	// Don't do culling visualization on newly placed nodes. Let the execution complete notification update that.
-	bool bIsCulled = !bNewlyPlaced && !PCGEditorGraphNodeBase::ShouldDisplayAsActive(this, InComponentBeingDebugged, InStackBeingInspected);
+	bool bIsCulled = !bNewlyPlaced && !PCGEditorGraphNodeBase::ShouldDisplayAsActive(this, InSourceBeingDebugged, InStackBeingInspected);
 
 	EPCGHiGenGrid ThisGrid = EPCGHiGenGrid::Uninitialized;
 
@@ -806,14 +809,14 @@ EPCGChangeType UPCGEditorGraphNodeBase::UpdateStructuralVisualization(UPCGCompon
 	return ChangeType;
 }
 
-EPCGChangeType UPCGEditorGraphNodeBase::UpdateGPUVisualization(class UPCGComponent* InComponentBeingDebugged, const FPCGStack* InStackBeingInspected)
+EPCGChangeType UPCGEditorGraphNodeBase::UpdateGPUVisualization(IPCGGraphExecutionSource* InSourceBeingDebugged, const FPCGStack* InStackBeingInspected)
 {
 	EPCGChangeType ChangeType = EPCGChangeType::None;
 
 	const UPCGNode* Node = GetPCGNode();
 
-	const bool bTriggeredUpload = InComponentBeingDebugged && InStackBeingInspected && Node
-		&& InComponentBeingDebugged->GetExecutionState().GetInspection().DidNodeTriggerCPUToGPUUpload(Node, *InStackBeingInspected);
+	const bool bTriggeredUpload = InSourceBeingDebugged && InStackBeingInspected && Node
+		&& InSourceBeingDebugged->GetExecutionState().GetInspection().DidNodeTriggerCPUToGPUUpload(Node, *InStackBeingInspected);
 
 	if (bTriggeredUpload != GetTriggeredGPUUpload())
 	{
@@ -821,8 +824,8 @@ EPCGChangeType UPCGEditorGraphNodeBase::UpdateGPUVisualization(class UPCGCompone
 		ChangeType |= EPCGChangeType::Cosmetic;
 	}
 
-	const bool bTriggeredReadback = InComponentBeingDebugged && InStackBeingInspected && Node
-		&& InComponentBeingDebugged->GetExecutionState().GetInspection().DidNodeTriggerGPUToCPUReadback(Node, *InStackBeingInspected);
+	const bool bTriggeredReadback = InSourceBeingDebugged && InStackBeingInspected && Node
+		&& InSourceBeingDebugged->GetExecutionState().GetInspection().DidNodeTriggerGPUToCPUReadback(Node, *InStackBeingInspected);
 
 	if (bTriggeredReadback != GetTriggeredGPUReadback())
 	{
diff --git a/Engine/Plugins/PCG/Source/PCGEditor/Private/Nodes/PCGEditorGraphNodeBase.h b/Engine/Plugins/PCG/Source/PCGEditor/Private/Nodes/PCGEditorGraphNodeBase.h
--- a/Engine/Plugins/PCG/Source/PCGEditor/Private/Nodes/PCGEditorGraphNodeBase.h
+++ b/Engine/Plugins/PCG/Source/PCGEditor/Private/Nodes/PCGEditorGraphNodeBase.h
@@ -17,7 +17,6 @@ struct FPCGPreconfiguredInfo;
 struct FPCGStack;
 
 class IPCGSettingsDefaultValueProvider;
-class UPCGComponent;
 class UPCGGraph;
 class UPCGNode;
 class UPCGPin;
@@ -91,9 +90,9 @@ public:
 	EPCGChangeType UpdateErrorsAndWarnings();
 
 	/** If the currently inspected grid size is smaller than the grid size of this node, display transparent. */
-	EPCGChangeType UpdateStructuralVisualization(class UPCGComponent* InComponentBeingDebugged, const FPCGStack* InStackBeingInspected, bool bNewlyPlaced = false);
+	EPCGChangeType UpdateStructuralVisualization(IPCGGraphExecutionSource* InSourceBeingDebugged, const FPCGStack* InStackBeingInspected, bool bNewlyPlaced = false);
 
-	EPCGChangeType UpdateGPUVisualization(class UPCGComponent* InComponentBeingDebugged, const FPCGStack* InStackBeingInspected);
+	EPCGChangeType UpdateGPUVisualization(IPCGGraphExecutionSource* InSourceBeingDebugged, const FPCGStack* InStackBeingInspected);
 
 	bool CanUserAddRemoveDynamicInputPins() const;
 
diff --git a/Engine/Plugins/PCG/Source/PCGEditor/Private/PCGEditor.cpp b/Engine/Plugins/PCG/Source/PCGEditor/Private/PCGEditor.cpp
--- a/Engine/Plugins/PCG/Source/PCGEditor/Private/PCGEditor.cpp
+++ b/Engine/Plugins/PCG/Source/PCGEditor/Private/PCGEditor.cpp
@@ -235,8 +235,8 @@ void FPCGEditor::SetStackBeingInspected(const FPCGStack& FullStack)
 		return;
 	}
 
-	UPCGComponent* LastComponent = LastValidPCGComponentBeingInspected.Get();
-	UPCGComponent* NewComponent = const_cast<UPCGComponent*>(FullStack.GetRootComponent());
+	IPCGGraphExecutionSource* LastComponent = LastValidPCGSourceBeingInspected.Get();
+	IPCGGraphExecutionSource* NewComponent = const_cast<IPCGGraphExecutionSource*>(FullStack.GetRootSource());
 
 	if (NewComponent && NewComponent != LastComponent)
 	{
@@ -245,7 +245,7 @@ void FPCGEditor::SetStackBeingInspected(const FPCGStack& FullStack)
 			LastComponent->GetExecutionState().GetInspection().DisableInspection();
 		}
 
-		LastValidPCGComponentBeingInspected = NewComponent;
+		LastValidPCGSourceBeingInspected = NewComponent;
 	}
 
 	if (PCGGraphBeingEdited)
@@ -258,7 +258,7 @@ void FPCGEditor::SetStackBeingInspected(const FPCGStack& FullStack)
 		PCGGraphBeingEdited->EnableInspection(StackBeingInspected);
 	}
 
-	PCGComponentBeingInspected = NewComponent;
+	PCGSourceBeingInspected = NewComponent;
 
 	StackBeingInspected = FullStack;
 	OnInspectedStackChangedDelegate.Broadcast(StackBeingInspected);
@@ -266,14 +266,14 @@ void FPCGEditor::SetStackBeingInspected(const FPCGStack& FullStack)
 	UpdateAfterInspectedStackChanged();
 }
 
-void FPCGEditor::OnComponentGenerated(UPCGComponent* InComponent)
+void FPCGEditor::OnSourceGenerated(IPCGGraphExecutionSource* InSource)
 {
 	if (DebugObjectTreeWidget)
 	{
 		DebugObjectTreeWidget->RequestRefresh();
 	}
 
-	if(InComponent == GetPCGComponentBeingInspected())
+	if(InSource == GetPCGSourceBeingInspected())
 	{
 		TRACE_CPUPROFILER_EVENT_SCOPE(FPCGEditor::SetStackBeingInspected::BroadcastStackBeingInspected);
 		OnInspectedStackChangedDelegate.Broadcast(StackBeingInspected);
@@ -296,19 +296,20 @@ bool FPCGEditor::OnValidateNodeTitle(const FText& NewName, UEdGraphNode* GraphNo
 
 void FPCGEditor::UpdateAfterInspectedStackChanged()
 {
-	UPCGComponent* Component = PCGComponentBeingInspected.Get();
+	IPCGGraphExecutionSource* Source = PCGSourceBeingInspected.Get();
+	UPCGComponent* Component = Cast<UPCGComponent>(PCGSourceBeingInspected.Get());
 
-	if (Component)
+	if (Source)
 	{
 		// Implementation note: if we're inspecting and have not pre-run the graph, then it probably makes sense to enable inspection by default. 
 		// TODO This could be selected with a cvar though.
-		const bool bHasBeenGeneratedThisSession = Component->bGenerated && Component->WasGeneratedThisSession();
-		const bool bWasInspecting = Component->GetExecutionState().GetInspection().IsInspecting();
+		const bool bHasBeenGeneratedThisSession = Component && Component->bGenerated && Component->WasGeneratedThisSession();
+		const bool bWasInspecting = Source->GetExecutionState().GetInspection().IsInspecting();
 		const bool bNeedsInspection = Algo::AnyOf(AttributesWidgets, [](const TSharedPtr<SPCGEditorGraphAttributeListView>& ALV) { return ALV->GetNodeBeingInspected() != nullptr; });
 
 		if (!bHasBeenGeneratedThisSession || (bNeedsInspection && !bWasInspecting))
 		{
-			Component->GetExecutionState().GetInspection().EnableInspection();
+			Source->GetExecutionState().GetInspection().EnableInspection();
 			UpdateDebugAfterComponentSelection(Component, Component, true);
 		}
 	}
@@ -320,8 +321,8 @@ void FPCGEditor::UpdateAfterInspectedStackChanged()
 		{
 			// Update now that component has changed. Will fire OnNodeChanged if necessary.
 			EPCGChangeType ChangeType = PCGNode->UpdateErrorsAndWarnings();
-			ChangeType |= PCGNode->UpdateStructuralVisualization(Component, &StackBeingInspected);
-			ChangeType |= PCGNode->UpdateGPUVisualization(Component, &StackBeingInspected);
+			ChangeType |= PCGNode->UpdateStructuralVisualization(Source, &StackBeingInspected);
+			ChangeType |= PCGNode->UpdateGPUVisualization(Source, &StackBeingInspected);
 
 			if (ChangeType != EPCGChangeType::None)
 			{
@@ -356,10 +357,10 @@ void FPCGEditor::UpdateDebugAfterComponentSelection(UPCGComponent* InOldComponen
 		// GenerateAtRuntime components should be refreshed through the runtime gen scheduler.
 		if (Component->IsManagedByRuntimeGenSystem())
 		{
-			if (UPCGSubsystem* Subsystem = GetWorldSubsystem())
+			if (UPCGSubsystem* WorldSubsystem = GetWorldSubsystem())
 			{
 				// We don't want to do a full cleanup if we're setting the debug object, since full cleanup destroys the component, which is the debug object itself!
-				Subsystem->RefreshRuntimeGenComponent(Component);
+				WorldSubsystem->RefreshRuntimeGenComponent(Component);
 			}
 		}
 		else
@@ -1225,19 +1226,19 @@ void FPCGEditor::OnDeterminismNodeTest() const
 
 bool FPCGEditor::CanRunDeterminismGraphTest() const
 {
-	return PCGEditorGraph && PCGComponentBeingInspected.IsValid();
+	return PCGEditorGraph && PCGSourceBeingInspected.IsValid();
 }
 
 void FPCGEditor::OnDeterminismGraphTest() const
 {
 	check(GraphEditorWidget.IsValid());
 
-	if (!DeterminismWidget.IsValid() || !DeterminismWidget->WidgetIsConstructed() || !PCGGraphBeingEdited || !PCGComponentBeingInspected.IsValid())
+	if (!DeterminismWidget.IsValid() || !DeterminismWidget->WidgetIsConstructed() || !PCGGraphBeingEdited || !PCGSourceBeingInspected.IsValid())
 	{
 		return;
 	}
 
-	if (PCGComponentBeingInspected->GetGraph() != PCGGraphBeingEdited)
+	if (PCGSourceBeingInspected->GetExecutionState().GetGraph() != PCGGraphBeingEdited)
 	{
 		// TODO: Should we alert the user more directly or disable this altogether?
 		UE_LOG(LogPCGEditor, Warning, TEXT("Running Determinism on a PCG Component with different/no attached PCG Graph"));
@@ -1253,9 +1254,9 @@ void FPCGEditor::OnDeterminismGraphTest() const
 	TestResult->Index = 0;
 	TestResult->TestResultTitle = TEXT("Full Graph Test");
 	TestResult->TestResultName = PCGGraphBeingEdited->GetName();
-	TestResult->Seed = PCGComponentBeingInspected->Seed;
+	TestResult->Seed = PCGSourceBeingInspected->GetExecutionState().GetSeed();
 
-	PCGDeterminismTests::RunDeterminismTest(PCGGraphBeingEdited, PCGComponentBeingInspected.Get(), *TestResult);
+	PCGDeterminismTests::RunDeterminismTest(PCGGraphBeingEdited, PCGSourceBeingInspected.Get(), *TestResult);
 
 	DeterminismWidget->AddItem(TestResult);
 	DeterminismWidget->AddDetailsColumn();
@@ -2823,19 +2824,19 @@ void FPCGEditor::OnClose()
 
 	FAssetEditorToolkit::OnClose();
 
-	if (PCGComponentBeingInspected.IsValid())
+	if (PCGSourceBeingInspected.IsValid())
 	{
-		if (PCGComponentBeingInspected->GetExecutionState().GetInspection().IsInspecting())
+		if (PCGSourceBeingInspected->GetExecutionState().GetInspection().IsInspecting())
 		{
-			PCGComponentBeingInspected->GetExecutionState().GetInspection().DisableInspection();
+			PCGSourceBeingInspected->GetExecutionState().GetInspection().DisableInspection();
 		}
 	}
 
-	if (LastValidPCGComponentBeingInspected.IsValid())
+	if (LastValidPCGSourceBeingInspected.IsValid())
 	{
-		if (LastValidPCGComponentBeingInspected->GetExecutionState().GetInspection().IsInspecting())
+		if (LastValidPCGSourceBeingInspected->GetExecutionState().GetInspection().IsInspecting())
 		{
-			LastValidPCGComponentBeingInspected->GetExecutionState().GetInspection().DisableInspection();
+			LastValidPCGSourceBeingInspected->GetExecutionState().GetInspection().DisableInspection();
 		}
 	}
 
@@ -3078,19 +3079,16 @@ void FPCGEditor::OnComponentUnregistered(UPCGComponent* Component)
 
 void FPCGEditor::OnSourceGenerationDone(IPCGBaseSubsystem* Subsystem, IPCGGraphExecutionSource* Source, EPCGGenerationStatus Status)
 {
-	// @todo_pcg Should support any source
-	UPCGComponent* Component = Cast<UPCGComponent>(Source);
-	
 	// We want to refresh if the component that is done generating has generated the current graph being edited,
 	// or if it is the root of the current stack being inspected (for subgraphs to also be refreshed).
 	// If we don't have a component, we refresh nonetheless.
-	bool bShouldRefresh = !Component || StackBeingInspected.GetRootComponent() == Component || Component->GetGraph() == PCGGraphBeingEdited;
+	bool bShouldRefresh = !Source || StackBeingInspected.GetRootSource() == Source || Source->GetExecutionState().GetGraph() == PCGGraphBeingEdited;
 
 	// Additionally, if we are not inspecting but the component that's done executing contains this graph, then we should also update.
 	IPCGEditorModule* PCGEditorModule = IPCGEditorModule::Get();
-	if (!bShouldRefresh && PCGComponentBeingInspected.IsNull() && PCGGraphBeingEdited && PCGEditorModule)
+	if (!bShouldRefresh && PCGSourceBeingInspected.IsValid() && PCGGraphBeingEdited && PCGEditorModule)
 	{
-		TArray<FPCGStackSharedPtr> ExecutedStacksContainingThisGraph = PCGEditorModule->GetExecutedStacksPtrs(Component, PCGGraphBeingEdited);
+		TArray<FPCGStackSharedPtr> ExecutedStacksContainingThisGraph = PCGEditorModule->GetExecutedStacksPtrs(Source, PCGGraphBeingEdited);
 		bShouldRefresh |= !ExecutedStacksContainingThisGraph.IsEmpty();
 	}
 
@@ -3099,7 +3097,7 @@ void FPCGEditor::OnSourceGenerationDone(IPCGBaseSubsystem* Subsystem, IPCGGraphE
 		return;
 	}
 
-	OnComponentGenerated(Component);
+	OnSourceGenerated(Source);
 
 	const bool CacheDebuggingEnabled = Subsystem && Subsystem->IsGraphCacheDebuggingEnabled();
 
@@ -3113,8 +3111,8 @@ void FPCGEditor::OnSourceGenerationDone(IPCGBaseSubsystem* Subsystem, IPCGGraphE
 			EPCGChangeType ChangeType = CacheDebuggingEnabled ? EPCGChangeType::Cosmetic : EPCGChangeType::None;
 
 			ChangeType |= PCGEditorGraphNode->UpdateErrorsAndWarnings();
-			ChangeType |= PCGEditorGraphNode->UpdateStructuralVisualization(PCGComponentBeingInspected.Get(), &StackBeingInspected);
-			ChangeType |= PCGEditorGraphNode->UpdateGPUVisualization(PCGComponentBeingInspected.Get(), &StackBeingInspected);
+			ChangeType |= PCGEditorGraphNode->UpdateStructuralVisualization(PCGSourceBeingInspected.Get(), &StackBeingInspected);
+			ChangeType |= PCGEditorGraphNode->UpdateGPUVisualization(PCGSourceBeingInspected.Get(), &StackBeingInspected);
 
 			if (ChangeType != EPCGChangeType::None)
 			{
@@ -3246,7 +3244,7 @@ void FPCGEditor::RefreshViewsOnLevelChange()
 	}
 }
 
-TSharedRef<FTabManager::FLayout> FPCGEditor::GetDefaultLayout()
+TSharedRef<FTabManager::FLayout> FPCGEditor::GetDefaultLayout() const
 {
 	return FTabManager::NewLayout("Standalone_PCGGraphEditor_DefaultLayout_v1.0")
 		->AddArea // Main PCG Graph Editor Area
diff --git a/Engine/Plugins/PCG/Source/PCGEditor/Private/PCGEditor.h b/Engine/Plugins/PCG/Source/PCGEditor/Private/PCGEditor.h
--- a/Engine/Plugins/PCG/Source/PCGEditor/Private/PCGEditor.h
+++ b/Engine/Plugins/PCG/Source/PCGEditor/Private/PCGEditor.h
@@ -67,30 +67,31 @@ class FPCGEditor : public FAssetEditorToolkit, public FGCObject, public FSelfReg
 {
 public:
 	/** Edits the specified PCGGraph */
-	void Initialize(const EToolkitMode::Type InMode, const TSharedPtr<class IToolkitHost>& InToolkitHost, UPCGGraph* InPCGGraph);
+	PCGEDITOR_API void Initialize(const EToolkitMode::Type InMode, const TSharedPtr<class IToolkitHost>& InToolkitHost, UPCGGraph* InPCGGraph);
 
 	/** Get the PCG editor graph being edited. */
-	UPCGEditorGraph* GetPCGEditorGraph();
+	PCGEDITOR_API UPCGEditorGraph* GetPCGEditorGraph();
 
 	/** Gets/Creates the PCG graph editor for a given PCG graph */
-	static UPCGEditorGraph* GetPCGEditorGraph(UPCGGraph* InGraph);
-	static UPCGEditorGraph* GetPCGEditorGraph(const UPCGNode* InNode);
-	static UPCGEditorGraph* GetPCGEditorGraph(const UPCGSettings* InSettings);
+	PCGEDITOR_API static UPCGEditorGraph* GetPCGEditorGraph(UPCGGraph* InGraph);
+	PCGEDITOR_API static UPCGEditorGraph* GetPCGEditorGraph(const UPCGNode* InNode);
+	PCGEDITOR_API static UPCGEditorGraph* GetPCGEditorGraph(const UPCGSettings* InSettings);
 
 	/** Get the PCG graph being edited. */
 	const UPCGGraph* GetPCGGraph() { return PCGGraphBeingEdited; }
 
 	/** Sets the execution stack that want to inspect. */
-	void SetStackBeingInspected(const FPCGStack& FullStack);
+	PCGEDITOR_API void SetStackBeingInspected(const FPCGStack& FullStack);
 
 	/** Sets the execution stack from another editor, which will set directly in the debug object tree view. */
-	void SetStackBeingInspectedFromAnotherEditor(const FPCGStack& FullStack);
+	PCGEDITOR_API void SetStackBeingInspectedFromAnotherEditor(const FPCGStack& FullStack);
 
 	/** Clear current inspection. */
 	void ClearStackBeingInspected();
 
-	/** Gets the PCG component we are debugging */
+	/** Gets the PCG source we are debugging */
 	UPCGComponent* GetPCGComponentBeingInspected() const { return const_cast<UPCGComponent*>(StackBeingInspected.GetRootComponent()); }
+	IPCGGraphExecutionSource* GetPCGSourceBeingInspected() const { return const_cast<IPCGGraphExecutionSource*>(StackBeingInspected.GetRootSource()); }
 	
 	/** Gets the PCG stack we are inspecting */
 	const FPCGStack* GetStackBeingInspected() const;
@@ -98,31 +99,31 @@ public:
 	void SetSourceEditorTargetObject(UObject* InObject);
 
 	/** Focus the graph view on a specific node */
-	void JumpToNode(const UEdGraphNode* InNode);
-	void JumpToNode(const UPCGNode* InNode);
+	PCGEDITOR_API void JumpToNode(const UEdGraphNode* InNode);
+	PCGEDITOR_API void JumpToNode(const UPCGNode* InNode);
 
 	/** Get the TabID of the editor panel. */
-	FName GetPanelID(EPCGEditorPanel Panel) const;
+	PCGEDITOR_API FName GetPanelID(EPCGEditorPanel Panel) const;
 	/** Focuses the user on a specific panel and flashes the tab. */
-	void BringFocusToPanel(EPCGEditorPanel Panel) const;
+	PCGEDITOR_API void BringFocusToPanel(EPCGEditorPanel Panel) const;
 	/** Attempts to close the specific panel if it's open. */
-	void CloseGraphPanel(EPCGEditorPanel Panel) const;
+	PCGEDITOR_API void CloseGraphPanel(EPCGEditorPanel Panel) const;
 	/** Returns true if the selected tab is currently open. */
-	bool IsPanelCurrentlyOpen(EPCGEditorPanel Panel) const;
+	PCGEDITOR_API bool IsPanelCurrentlyOpen(EPCGEditorPanel Panel) const;
 	/** Returns true if the selected tab is currently open and focused. */
-	bool IsPanelCurrentlyForeground(EPCGEditorPanel Panel) const;
+	PCGEDITOR_API bool IsPanelCurrentlyForeground(EPCGEditorPanel Panel) const;
 
 	/** Helper to get to the subsystem. */
-	virtual IPCGBaseSubsystem* GetSubsystem() const;
-	static UPCGSubsystem* GetWorldSubsystem();
+	PCGEDITOR_API virtual IPCGBaseSubsystem* GetSubsystem() const;
+	PCGEDITOR_API static UPCGSubsystem* GetWorldSubsystem();
 
 	// ~Begin IToolkit interface
-	virtual void RegisterTabSpawners(const TSharedRef<class FTabManager>& InTabManager) override;
-	virtual void UnregisterTabSpawners(const TSharedRef<class FTabManager>& InTabManager) override;
+	PCGEDITOR_API virtual void RegisterTabSpawners(const TSharedRef<class FTabManager>& InTabManager) override;
+	PCGEDITOR_API virtual void UnregisterTabSpawners(const TSharedRef<class FTabManager>& InTabManager) override;
 	// ~End IToolkit interface
 
 	// ~Begin FGCObject interface
-	virtual void AddReferencedObjects(FReferenceCollector& Collector) override;
+	PCGEDITOR_API virtual void AddReferencedObjects(FReferenceCollector& Collector) override;
 	virtual FString GetReferencerName() const override
 	{
 		return TEXT("FPCGEditor");
@@ -130,18 +131,18 @@ public:
 	// ~End FGCObject interface
 	
 	// ~Begin FEditorUndoClient interface
-	virtual bool MatchesContext(const FTransactionContext& InContext, const TArray<TPair<UObject*, FTransactionObjectEvent>>& TransactionObjectContexts) const override;
-	virtual void PostUndo(bool bSuccess) override;
+	PCGEDITOR_API virtual bool MatchesContext(const FTransactionContext& InContext, const TArray<TPair<UObject*, FTransactionObjectEvent>>& TransactionObjectContexts) const override;
+	PCGEDITOR_API virtual void PostUndo(bool bSuccess) override;
 	virtual void PostRedo(bool bSuccess) override { PostUndo(bSuccess); }
 	// ~End FEditorUndoClient interface
 
 	// ~Begin FAssetEditorToolkit interface
-	virtual FName GetToolkitFName() const override;
-	virtual FText GetBaseToolkitName() const override;
-	virtual FLinearColor GetWorldCentricTabColorScale() const override;
-	virtual FString GetWorldCentricTabPrefix() const override;
-	virtual void OnClose() override;
-	virtual void InitToolMenuContext(FToolMenuContext& MenuContext) override;
+	PCGEDITOR_API virtual FName GetToolkitFName() const override;
+	PCGEDITOR_API virtual FText GetBaseToolkitName() const override;
+	PCGEDITOR_API virtual FLinearColor GetWorldCentricTabColorScale() const override;
+	PCGEDITOR_API virtual FString GetWorldCentricTabPrefix() const override;
+	PCGEDITOR_API virtual void OnClose() override;
+	PCGEDITOR_API virtual void InitToolMenuContext(FToolMenuContext& MenuContext) override;
 	// ~End FAssetEditorToolkit interface
 
 	/**
@@ -166,7 +167,7 @@ public:
 	/** Run the determinism test on the current graph */
 	void OnDeterminismGraphTest() const;
 
-private:
+protected:
 	/** Register PCG specific toolbar for the editor */
 	void RegisterToolbar() const;
 
@@ -378,7 +379,7 @@ private:
 	void OnSelectedNodesChanged(const TSet<UObject*>& NewSelection);
 
 	/** Called when the component inspected is generated/cleaned */
-	void OnComponentGenerated(UPCGComponent* InComponent);
+	void OnSourceGenerated(IPCGGraphExecutionSource* InSource);
 
 	/** Called to validate a node title change. */
 	bool OnValidateNodeTitle(const FText& NewName, UEdGraphNode* GraphNode, FText& OutErrorMessage);
@@ -404,7 +405,7 @@ private:
 	void OnComponentUnregistered(UPCGComponent* Component);
 
 	/** Called when a component finishes executing. Useful for updating debugging tools/UIs. */
-	void OnSourceGenerationDone(IPCGBaseSubsystem* Subsystem, IPCGGraphExecutionSource* Source, EPCGGenerationStatus Status);
+	PCGEDITOR_API void OnSourceGenerationDone(IPCGBaseSubsystem* Subsystem, IPCGGraphExecutionSource* Source, EPCGGenerationStatus Status);
 
 	/** Trigger any generation required to ensure debug display is up to date. */
 	void UpdateDebugAfterComponentSelection(UPCGComponent* InOldComponent, UPCGComponent* InNewComponent, bool bNewComponentStartedInspecting);
@@ -440,7 +441,7 @@ private:
 	FText GetAttributesTabLabel(int AttributesIndex);
 	FText GetViewportTabLabel(int ViewportIndex);
 
-	static TSharedRef<FTabManager::FLayout> GetDefaultLayout();
+	PCGEDITOR_API virtual TSharedRef<FTabManager::FLayout> GetDefaultLayout() const;
 
 	TSharedPtr<SGraphEditor> GraphEditorWidget;
 	TArray<TSharedPtr<SPCGEditorGraphDetailsView>> PropertyDetailsWidgets;
@@ -459,8 +460,8 @@ private:
 	TObjectPtr<UPCGGraph> PCGGraphBeingEdited = nullptr;
 	UPCGEditorGraph* PCGEditorGraph = nullptr;
 
-	TSoftObjectPtr<UPCGComponent> PCGComponentBeingInspected;
-	// Implementation note: we'll keep the last valid component inspected so we don't un-inspect on spurious selection changes
-	TSoftObjectPtr<UPCGComponent> LastValidPCGComponentBeingInspected;
+	FPCGSoftGraphExecutionSource PCGSourceBeingInspected;
+	// Implementation note: we'll keep the last valid source inspected so we don't un-inspect on spurious selection changes
+	FPCGSoftGraphExecutionSource LastValidPCGSourceBeingInspected;
 	FPCGStack StackBeingInspected;
 };
diff --git a/Engine/Plugins/PCG/Source/PCGEditor/Private/PCGEditorGraph.cpp b/Engine/Plugins/PCG/Source/PCGEditor/Private/PCGEditorGraph.cpp
--- a/Engine/Plugins/PCG/Source/PCGEditor/Private/PCGEditorGraph.cpp
+++ b/Engine/Plugins/PCG/Source/PCGEditor/Private/PCGEditorGraph.cpp
@@ -213,15 +213,15 @@ void UPCGEditorGraph::ReplicateExtraNodes() const
 	}
 }
 
-void UPCGEditorGraph::UpdateVisualizations(UPCGComponent* PCGComponentBeingInspected, const FPCGStack* PCGStackBeingInspected)
+void UPCGEditorGraph::UpdateVisualizations(IPCGGraphExecutionSource* PCGSourceBeingInspected, const FPCGStack* PCGStackBeingInspected)
 {
 	for (UEdGraphNode* EditorNode : Nodes)
 	{
 		if (UPCGEditorGraphNodeBase* PCGEditorNode = Cast<UPCGEditorGraphNodeBase>(EditorNode))
 		{
 			EPCGChangeType ChangeType = EPCGChangeType::None;
-			ChangeType |= PCGEditorNode->UpdateStructuralVisualization(PCGComponentBeingInspected, PCGStackBeingInspected);
-			ChangeType |= PCGEditorNode->UpdateGPUVisualization(PCGComponentBeingInspected, PCGStackBeingInspected);
+			ChangeType |= PCGEditorNode->UpdateStructuralVisualization(PCGSourceBeingInspected, PCGStackBeingInspected);
+			ChangeType |= PCGEditorNode->UpdateGPUVisualization(PCGSourceBeingInspected, PCGStackBeingInspected);
 
 			if (ChangeType != EPCGChangeType::None)
 			{
diff --git a/Engine/Plugins/PCG/Source/PCGEditor/Private/PCGEditorGraph.h b/Engine/Plugins/PCG/Source/PCGEditor/Private/PCGEditorGraph.h
--- a/Engine/Plugins/PCG/Source/PCGEditor/Private/PCGEditorGraph.h
+++ b/Engine/Plugins/PCG/Source/PCGEditor/Private/PCGEditorGraph.h
@@ -51,7 +51,7 @@ public:
 	TWeakPtr<FPCGEditor> GetEditor() const { return PCGEditor; }
 
 	/** Updates the grid size visualization in the editor. */
-	void UpdateVisualizations(UPCGComponent* PCGComponentBeingInspected, const FPCGStack* PCGStackBeingInspected);
+	void UpdateVisualizations(IPCGGraphExecutionSource* PCGSourceBeingInspected, const FPCGStack* PCGStackBeingInspected);
 
 	/** Returns the PCG editor graph node corresponding to the given PCG node. */
 	const UPCGEditorGraphNodeBase* GetEditorNodeFromPCGNode(const UPCGNode* InPCGNode) const;
diff --git a/Engine/Plugins/PCG/Source/PCGEditor/Private/PCGEditorGraphSchema.cpp b/Engine/Plugins/PCG/Source/PCGEditor/Private/PCGEditorGraphSchema.cpp
--- a/Engine/Plugins/PCG/Source/PCGEditor/Private/PCGEditorGraphSchema.cpp
+++ b/Engine/Plugins/PCG/Source/PCGEditor/Private/PCGEditorGraphSchema.cpp
@@ -1056,9 +1056,9 @@ bool FPCGEditorConnectionDrawingPolicy::UpdateParamsIfDebugging(UEdGraphPin* Out
 
 	const FPCGEditor* Editor = Graph->GetEditor().Pin().Get();
 	const FPCGStack* PCGStack = Editor ? Editor->GetStackBeingInspected() : nullptr;
-	UPCGComponent* PCGComponent = Editor ? Editor->GetPCGComponentBeingInspected() : nullptr;
+	IPCGGraphExecutionSource* PCGSource = Editor ? Editor->GetPCGSourceBeingInspected() : nullptr;
 
-	if (!Editor || !PCGStack || !PCGComponent || !PCGComponent->GetExecutionState().GetInspection().IsInspecting())
+	if (!Editor || !PCGStack || !PCGSource || !PCGSource->GetExecutionState().GetInspection().IsInspecting())
 	{
 		return false;
 	}
@@ -1069,7 +1069,7 @@ bool FPCGEditorConnectionDrawingPolicy::UpdateParamsIfDebugging(UEdGraphPin* Out
 	StackFrames.Emplace(NodeToInspect);
 	StackFrames.Emplace(PinToInspect);
 
-	if (const FPCGDataCollection* DataCollection = PCGComponent->GetExecutionState().GetInspection().GetInspectionData(Stack))
+	if (const FPCGDataCollection* DataCollection = PCGSource->GetExecutionState().GetInspection().GetInspectionData(Stack))
 	{
 		if (DataCollection->TaggedData.Num() > 1)
 		{
diff --git a/Engine/Plugins/PCG/Source/PCGEditor/Private/PCGEditorModule.cpp b/Engine/Plugins/PCG/Source/PCGEditor/Private/PCGEditorModule.cpp
--- a/Engine/Plugins/PCG/Source/PCGEditor/Private/PCGEditorModule.cpp
+++ b/Engine/Plugins/PCG/Source/PCGEditor/Private/PCGEditorModule.cpp
@@ -4,6 +4,7 @@
 
 #include "PCGComponent.h"
 #include "PCGEngineSettings.h"
+#include "PCGGraphExecutionStateInterface.h"
 #include "PCGModule.h"
 #include "PCGParamData.h"
 #include "PCGSubsystem.h"
@@ -829,17 +830,17 @@ void FPCGEditorModule::OnScheduleGraph(const FPCGStackContext& StackContext)
 		GetNodeVisualLogsMutable().ClearLogs(Stack);
 	}
 	
-	UPCGComponent* RootComponent = nullptr;
+	IPCGGraphExecutionSource* RootSource = nullptr;
 
-	// Flush out all stacks that begin from the component / top graph as the existing dynamic
+	// Flush out all stacks that begin from the source / top graph as the existing dynamic
 	// stacks may not be occur during the next execution.
 	if (const FPCGStack* BaseStack = StackContext.GetStack(0))
 	{
-		RootComponent = const_cast<UPCGComponent*>(BaseStack->GetRootComponent());
+		RootSource = const_cast<IPCGGraphExecutionSource*>(BaseStack->GetRootSource());
 
 		if (BaseStack->IsCurrentFrameInRootGraph())
 		{
-			ClearExecutionMetadata(RootComponent);
+			ClearExecutionMetadata(RootSource);
 		}
 	}
 
@@ -848,8 +849,6 @@ void FPCGEditorModule::OnScheduleGraph(const FPCGStackContext& StackContext)
 		TRACE_CPUPROFILER_EVENT_SCOPE(UPCGSubsystem::OnScheduleGraph::RecordExecutedStacks);
 		FWriteScopeLock Lock(ExecutedStacksLock);
 
-		TSoftObjectPtr<UPCGComponent> RootComponentSoftPtr(RootComponent);
-
 		for (const FPCGStack& ExecutedStack : StackContext.GetStacks())
 		{
 			FPCGStackSharedPtr SharedStack(MakeShared<FPCGStack>(ExecutedStack));
@@ -859,9 +858,9 @@ void FPCGEditorModule::OnScheduleGraph(const FPCGStackContext& StackContext)
 			// Push to the set
 			ExecutedStacks.Add(SharedStack, &bIsAlreadyInSet);
 
-			if (!bIsAlreadyInSet && RootComponent)
+			if (!bIsAlreadyInSet && RootSource)
 			{
-				ExecutedStacksPerComponent.FindOrAdd(RootComponentSoftPtr).Add(SharedStack);
+				ExecutedStacksPerSource.FindOrAdd(RootSource).Add(SharedStack);
 			}
 		}
 	}
@@ -881,32 +880,32 @@ void FPCGEditorModule::OnGraphPreSave(UPCGGraph* Graph, FObjectPreSaveContext Ob
 	}
 }
 
-void FPCGEditorModule::ClearExecutionMetadata(UPCGComponent* InComponent)
+void FPCGEditorModule::ClearExecutionMetadata(IPCGGraphExecutionSource* InSource)
 {
-	if (InComponent)
+	if (InSource)
 	{
-		ClearExecutedStacks(InComponent);
-		InComponent->GetExecutionState().GetInspection().ClearInspectionData();
+		ClearExecutedStacks(InSource);
+		InSource->GetExecutionState().GetInspection().ClearInspectionData();
 
-		GetNodeVisualLogsMutable().ClearLogs(InComponent);
+		GetNodeVisualLogsMutable().ClearLogs(InSource);
 	}
 }
 
-void FPCGEditorModule::ClearExecutedStacks(const UPCGComponent* InRootComponent)
+void FPCGEditorModule::ClearExecutedStacks(const IPCGGraphExecutionSource* InRootSource)
 {
 	TRACE_CPUPROFILER_EVENT_SCOPE(FPCGEditorModule::ClearExecutedStacksWithComponent);
-	if (InRootComponent)
+	if (InRootSource)
 	{
 		FWriteScopeLock Lock(ExecutedStacksLock);
 
-		if (TSet<FPCGStackSharedPtr>* AllStackForComponent = ExecutedStacksPerComponent.Find(InRootComponent))
+		if (TSet<FPCGStackSharedPtr>* AllStackForComponent = ExecutedStacksPerSource.Find(InRootSource))
 		{
 			for (const FPCGStackSharedPtr& Stack : *AllStackForComponent)
 			{
 				ExecutedStacks.Remove(Stack);
 			}
 
-			ExecutedStacksPerComponent.Remove(InRootComponent);
+			ExecutedStacksPerSource.Remove(InRootSource);
 		}
 	}
 }
@@ -929,7 +928,7 @@ void FPCGEditorModule::ClearExecutedStacks(const UPCGGraph* InContainingGraph)
 		{
 			ExecutedStacks.Remove(StackPtr);
 
-			if (TSet<FPCGStackSharedPtr>* StacksForComponent = ExecutedStacksPerComponent.Find(StackPtr->GetRootComponent()))
+			if (TSet<FPCGStackSharedPtr>* StacksForComponent = ExecutedStacksPerSource.Find(StackPtr->GetRootSource()))
 			{
 				StacksForComponent->Remove(StackPtr);
 			}
@@ -943,7 +942,7 @@ TArray<FPCGStackSharedPtr> FPCGEditorModule::GetExecutedStacksPtrs(const FPCGSta
 	TArray<FPCGStackSharedPtr> MatchingStacks;
 
 	FReadScopeLock Lock(ExecutedStacksLock);
-	TSet<FPCGStackSharedPtr>* StacksToTest = ExecutedStacksPerComponent.Find(BeginningWithStack.GetRootComponent());
+	TSet<FPCGStackSharedPtr>* StacksToTest = ExecutedStacksPerSource.Find(BeginningWithStack.GetRootSource());
 
 	// Stack doesn't have a root component (unlikely), we'll check everything then
 	if (!StacksToTest)
@@ -962,7 +961,7 @@ TArray<FPCGStackSharedPtr> FPCGEditorModule::GetExecutedStacksPtrs(const FPCGSta
 	return MatchingStacks;
 }
 
-TArray<FPCGStackSharedPtr> FPCGEditorModule::GetExecutedStacksPtrs(const UPCGComponent* InComponent, const UPCGGraph* InSubgraph, bool bOnlyWithSubgraphAsCurrentFrame)
+TArray<FPCGStackSharedPtr> FPCGEditorModule::GetExecutedStacksPtrs(const IPCGGraphExecutionSource* InSource, const UPCGGraph* InSubgraph, bool bOnlyWithSubgraphAsCurrentFrame)
 {
 	TRACE_CPUPROFILER_EVENT_SCOPE(FPCGEditorModule::GetExecutedStacksPtrs);
 	TArray<FPCGStackSharedPtr> MatchingStacks;
@@ -971,9 +970,9 @@ TArray<FPCGStackSharedPtr> FPCGEditorModule::GetExecutedStacksPtrs(const UPCGCom
 	TSet<FPCGStackSharedPtr>* ComponentSet = nullptr;
 	TSet<FPCGStackSharedPtr>* GraphSet = nullptr;
 
-	if (InComponent)
+	if (InSource)
 	{
-		ComponentSet = ExecutedStacksPerComponent.Find(InComponent);
+		ComponentSet = ExecutedStacksPerSource.Find(InSource);
 	}
 
 	if (!ComponentSet)
diff --git a/Engine/Plugins/PCG/Source/PCGEditor/Private/Widgets/SPCGEditorGraphAttributeListView.cpp b/Engine/Plugins/PCG/Source/PCGEditor/Private/Widgets/SPCGEditorGraphAttributeListView.cpp
--- a/Engine/Plugins/PCG/Source/PCGEditor/Private/Widgets/SPCGEditorGraphAttributeListView.cpp
+++ b/Engine/Plugins/PCG/Source/PCGEditor/Private/Widgets/SPCGEditorGraphAttributeListView.cpp
@@ -715,9 +715,9 @@ TSharedRef<SHeaderRow> SPCGEditorGraphAttributeListView::CreateHeaderRowWidget()
 	return SNew(SHeaderRow);
 }
 
-TWeakObjectPtr<UPCGComponent> SPCGEditorGraphAttributeListView::GetPCGComponent() const
+TWeakInterfacePtr<IPCGGraphExecutionSource> SPCGEditorGraphAttributeListView::GetExecutionSource() const
 {
-	return PCGEditorPtr.IsValid() ? PCGEditorPtr.Pin()->GetPCGComponentBeingInspected() : nullptr;
+	return PCGEditorPtr.IsValid() ? PCGEditorPtr.Pin()->GetPCGSourceBeingInspected() : nullptr;
 }
 
 void SPCGEditorGraphAttributeListView::OnInspectedStackChanged(const FPCGStack& InPCGStack)
@@ -780,15 +780,10 @@ void SPCGEditorGraphAttributeListView::AddReferencedObjects(FReferenceCollector&
 	}
 }
 
-void SPCGEditorGraphAttributeListView::OnGenerateUpdated(UPCGComponent* /*InPCGComponent*/)
-{
-	RequestRefresh();
-}
-
 const FPCGDataCollection* SPCGEditorGraphAttributeListView::GetInspectionData(const TSharedPtr<FPinComboBoxItem>& EditorPin) const
 {
-	TWeakObjectPtr<UPCGComponent> PCGComponent = GetPCGComponent();
-	if (!PCGComponent.IsValid())
+	TWeakInterfacePtr<IPCGGraphExecutionSource> ExecutionSource = GetExecutionSource();
+	if (!ExecutionSource.IsValid())
 	{
 		return nullptr;
 	}
@@ -835,7 +830,7 @@ const FPCGDataCollection* SPCGEditorGraphAttributeListView::GetInspectionData(co
 	StackFrames.Emplace(PCGNode);
 	StackFrames.Emplace(Pin);
 
-	return PCGComponent->GetExecutionState().GetInspection().GetInspectionData(Stack);
+	return ExecutionSource->GetExecutionState().GetInspection().GetInspectionData(Stack);
 }
 
 const FPCGDataCollection* SPCGEditorGraphAttributeListView::GetInspectionData() const
diff --git a/Engine/Plugins/PCG/Source/PCGEditor/Private/Widgets/SPCGEditorGraphAttributeListView.h b/Engine/Plugins/PCG/Source/PCGEditor/Private/Widgets/SPCGEditorGraphAttributeListView.h
--- a/Engine/Plugins/PCG/Source/PCGEditor/Private/Widgets/SPCGEditorGraphAttributeListView.h
+++ b/Engine/Plugins/PCG/Source/PCGEditor/Private/Widgets/SPCGEditorGraphAttributeListView.h
@@ -11,6 +11,7 @@
 #include "Misc/TextFilterExpressionEvaluator.h"
 #include "Tasks/Task.h"
 #include "UObject/ObjectKey.h"
+#include "UObject/WeakInterfacePtr.h"
 #include "Widgets/Views/SListView.h"
 #include "Widgets/Views/STableRow.h"
 
@@ -23,7 +24,6 @@ class FPCGMetadataAttributeBase;
 class FUICommandList;
 class SPCGEditorGraphAttributeListView;
 class SPCGEditorViewport;
-class UPCGComponent;
 class UPCGData;
 class UPCGEditorGraphNodeBase;
 class UPCGMetadata;
@@ -173,8 +173,6 @@ private:
 
 	void OnInspectedStackChanged(const FPCGStack& InPCGStack);
 
-	void OnGenerateUpdated(UPCGComponent* InPCGComponent);
-
 	void ResizeColumnToMaxWidth(FName InColumnId);
 	void ResetColumnsWidthToDefault();
 	void ExpendAllColumnToMaxWidth();
@@ -268,7 +266,7 @@ private:
 	bool IsViewportOpen() const;
 
 	/** Returns the current PCG component */
-	TWeakObjectPtr<UPCGComponent> GetPCGComponent() const;
+	TWeakInterfacePtr<IPCGGraphExecutionSource> GetExecutionSource() const;
 
 	/** Pointer back to the PCG editor that owns us */
 	TWeakPtr<FPCGEditor> PCGEditorPtr;
diff --git a/Engine/Plugins/PCG/Source/PCGEditor/Private/Widgets/SPCGEditorGraphDebugObjectTree.cpp b/Engine/Plugins/PCG/Source/PCGEditor/Private/Widgets/SPCGEditorGraphDebugObjectTree.cpp
--- a/Engine/Plugins/PCG/Source/PCGEditor/Private/Widgets/SPCGEditorGraphDebugObjectTree.cpp
+++ b/Engine/Plugins/PCG/Source/PCGEditor/Private/Widgets/SPCGEditorGraphDebugObjectTree.cpp
@@ -249,11 +249,13 @@ const FSlateBrush* FPCGEditorGraphDebugObjectItem_Actor::GetIcon() const
 	return FSlateIconFinder::FindIconBrushForClass(Object ? Object->GetClass() : AActor::StaticClass());
 }
 
-FString FPCGEditorGraphDebugObjectItem_PCGComponent::GetLabel(bool bForSorting) const
+FString FPCGEditorGraphDebugObjectItem_PCGSource::GetLabel(bool bForSorting) const
 {
-	if (PCGComponent.IsValid() && PCGGraph.IsValid())
+	//@todo_pcg: Should be part of the interface?
+	const UObject* Source = PCGSource.GetObject();
+	if (Source && PCGGraph.IsValid())
 	{
-		return PCGEditorGraphDebugObjectTree::GetStringFromName(PCGComponent->GetFName(), bForSorting) + FString(TEXT(" - ")) + PCGEditorGraphDebugObjectTree::GetStringFromName(PCGGraph->GetFName(), bForSorting);
+		return PCGEditorGraphDebugObjectTree::GetStringFromName(Source->GetFName(), bForSorting) + FString(TEXT(" - ")) + PCGEditorGraphDebugObjectTree::GetStringFromName(PCGGraph->GetFName(), bForSorting);
 	}
 
 	return FString();
@@ -568,8 +570,8 @@ void SPCGEditorGraphDebugObjectTree::Tick(const FGeometry& AllottedGeometry, con
 	{
 		// Updating the tree while the inspected component is generating can be bad as the selection can
 		// be lost. Don't change the tree if we're inspecting something that is generating.
-		const UPCGComponent* InspectedComponent = SelectedStack.GetRootComponent();
-		if(InspectedComponent == nullptr || !InspectedComponent->IsGenerating())
+		const IPCGGraphExecutionSource* InspectedSource = SelectedStack.GetRootSource();
+		if(InspectedSource == nullptr || !InspectedSource->GetExecutionState().IsGenerating())
 		{
 			bNeedsRefresh = false;
 			RefreshTree();
@@ -856,7 +858,7 @@ void SPCGEditorGraphDebugObjectTree::UpdateIsSetDebugObjectFromSelectionEnabled(
 }
 
 void SPCGEditorGraphDebugObjectTree::AddStacksToTree(const TArray<FPCGStackSharedPtr>& Stacks,
-	TMap<AActor*, TSharedPtr<FPCGEditorGraphDebugObjectItem_Actor>>& InOutActorItems,
+	TMap<UObject*, FPCGEditorGraphDebugObjectItemPtr>& InOutOwnerItems,
 	TMap<const FPCGStack, FPCGEditorGraphDebugObjectItemPtr>& InOutStackToItem)
 {
 	TRACE_CPUPROFILER_EVENT_SCOPE(SPCGEditorGraphDebugObjectTree::AddStacksToTree);
@@ -868,12 +870,12 @@ void SPCGEditorGraphDebugObjectTree::AddStacksToTree(const TArray<FPCGStackShare
 
 	for (const FPCGStackSharedPtr& Stack : Stacks)
 	{
-		AddStackToTree(*Stack, GraphBeingEdited, InOutActorItems, InOutStackToItem);
+		AddStackToTree(*Stack, GraphBeingEdited, InOutOwnerItems, InOutStackToItem);
 	}
 }
 
 void SPCGEditorGraphDebugObjectTree::AddStacksToTree(const TArray<FPCGStack>& Stacks,
-	TMap<AActor*, TSharedPtr<FPCGEditorGraphDebugObjectItem_Actor>>& InOutActorItems,
+	TMap<UObject*, FPCGEditorGraphDebugObjectItemPtr>& InOutOwnerItems,
 	TMap<const FPCGStack, FPCGEditorGraphDebugObjectItemPtr>& InOutStackToItem)
 {
 	TRACE_CPUPROFILER_EVENT_SCOPE(SPCGEditorGraphDebugObjectTree::AddStacksToTree);
@@ -885,13 +887,13 @@ void SPCGEditorGraphDebugObjectTree::AddStacksToTree(const TArray<FPCGStack>& St
 
 	for (const FPCGStack& Stack : Stacks)
 	{
-		AddStackToTree(Stack, GraphBeingEdited, InOutActorItems, InOutStackToItem);
+		AddStackToTree(Stack, GraphBeingEdited, InOutOwnerItems, InOutStackToItem);
 	}
 }
 
 void SPCGEditorGraphDebugObjectTree::AddStackToTree(const FPCGStack& Stack,
 	const UPCGGraph* GraphBeingEdited,
-	TMap<AActor*, TSharedPtr<FPCGEditorGraphDebugObjectItem_Actor>>& InOutActorItems,
+	TMap<UObject*, FPCGEditorGraphDebugObjectItemPtr>& InOutOwnerItems,
 	TMap<const FPCGStack, FPCGEditorGraphDebugObjectItemPtr>& InOutStackToItem)
 {
 	check(GraphBeingEdited);
@@ -905,8 +907,9 @@ void SPCGEditorGraphDebugObjectTree::AddStackToTree(const FPCGStack& Stack,
 			return;
 		}
 
-		UPCGComponent* PCGComponent = const_cast<UPCGComponent*>(Stack.GetRootComponent());
-		if (!PCGComponent)
+		IPCGGraphExecutionSource* PCGSource = const_cast<IPCGGraphExecutionSource*>(Stack.GetRootSource());
+		
+		if (!PCGSource)
 		{
 			return;
 		}
@@ -914,7 +917,7 @@ void SPCGEditorGraphDebugObjectTree::AddStackToTree(const FPCGStack& Stack,
 		// Prevent duplicate entries from the editor world while in PIE.
 		if (PCGHelpers::IsRuntimeOrPIE())
 		{
-			if (UWorld* World = PCGComponent->GetWorld())
+			if (UWorld* World = PCGSource->GetExecutionState().GetWorld())
 			{
 				if (!World->IsGameWorld())
 				{
@@ -939,26 +942,29 @@ void SPCGEditorGraphDebugObjectTree::AddStackToTree(const FPCGStack& Stack,
 
 			// Display grayed out if no inspection data has been stored for this node, and node does not run on GPU,
 			// because we don't opportunistically store inspection data for GPU nodes.
-			bDisplayGrayedOut = !bGPUNode && !PCGComponent->HasNodeProducedData(PCGNodeBeingInspected.Get(), Stack);
-		}
-
-		AActor* Actor = PCGComponent->GetOwner();
-		if (!Actor)
-		{
-			return;
+			bDisplayGrayedOut = !bGPUNode && !PCGSource->GetExecutionState().GetInspection().HasNodeProducedData(PCGNodeBeingInspected.Get(), Stack);
 		}
 
-		// Add actor item if not already added.
-		FPCGEditorGraphDebugObjectItemPtr ActorItem;
-		if (TSharedPtr<FPCGEditorGraphDebugObjectItem_Actor>* FoundActorItem = InOutActorItems.Find(Actor))
+		// @todo_pcg: GetOwner might need to be part of the execution source interface.
+		UPCGComponent* PCGComponent = Cast<UPCGComponent>(PCGSource);
+		FPCGEditorGraphDebugObjectItemPtr OwnerItem;
+		if (AActor* Actor = PCGComponent ? PCGComponent->GetOwner() : nullptr)
 		{
-			ActorItem = *FoundActorItem;
-			ActorItem->UpdateGrayedOut(bDisplayGrayedOut);
+			// Add actor item if not already added.
+			if (TSharedPtr<FPCGEditorGraphDebugObjectItem>* FoundActorItem = InOutOwnerItems.Find(Actor))
+			{
+				OwnerItem = *FoundActorItem;
+				OwnerItem->UpdateGrayedOut(bDisplayGrayedOut);
+			}
+			else
+			{
+				OwnerItem = InOutOwnerItems.Emplace(Actor, MakeShared<FPCGEditorGraphDebugObjectItem_Actor>(Actor, bDisplayGrayedOut));
+				AllGraphItems.Add(OwnerItem);
+			}
 		}
-		else
+		else if (PCGComponent)
 		{
-			ActorItem = InOutActorItems.Emplace(Actor, MakeShared<FPCGEditorGraphDebugObjectItem_Actor>(Actor, bDisplayGrayedOut));
-			AllGraphItems.Add(ActorItem);
+			return;
 		}
 
 		const TArray<FPCGStackFrame>& StackFrames = Stack.GetStackFrames();
@@ -1076,11 +1082,20 @@ void SPCGEditorGraphDebugObjectTree::AddStackToTree(const FPCGStack& Stack,
 					{
 						FPCGEditorGraphDebugObjectItemPtr TopGraphItem = InOutStackToItem.Emplace(
 							GraphStack,
-							MakeShared<FPCGEditorGraphDebugObjectItem_PCGComponent>(PCGComponent, StackGraph, GraphStack, bIsDebuggable, bDisplayGrayedOut));
+							MakeShared<FPCGEditorGraphDebugObjectItem_PCGSource>(PCGSource, StackGraph, GraphStack, bIsDebuggable, bDisplayGrayedOut));
 
 						AllGraphItems.Add(TopGraphItem);
 
-						ActorItem->AddChild(TopGraphItem.ToSharedRef());
+						if (OwnerItem)
+						{
+							OwnerItem->AddChild(TopGraphItem.ToSharedRef());
+						}
+						else
+						{
+							// If we don't have a component but an execution source, this becomes our owner item.
+							InOutOwnerItems.Emplace(Cast<UObject>(PCGSource), TopGraphItem);
+							OwnerItem = TopGraphItem;
+						}
 					}
 				}
 				// Previous stack was node, therefore static subgraph.
@@ -1148,37 +1163,43 @@ void SPCGEditorGraphDebugObjectTree::RefreshTree()
 		return;
 	}
 
-	TArray<UObject*> PCGComponents;
-	GetObjectsOfClass(UPCGComponent::StaticClass(), PCGComponents, /*bIncludeDerivedClasses=*/ true);
 
-	if (!PCGComponents.IsEmpty())
+	// @todo_pcg: Support other executions sources? GetObjectsOfClass(UPCGGraphExecutionSource::StaticClass()) doesn't work.
+	TArray<UObject*> PCGSources;
+	GetObjectsOfClass(UPCGComponent::StaticClass(), PCGSources, /*bIncludeDerivedClasses=*/ true);
+
+	if (!PCGSources.IsEmpty())
 	{
-		TMap<AActor*, TSharedPtr<FPCGEditorGraphDebugObjectItem_Actor>> ActorItems;
+		TMap<UObject*, TSharedPtr<FPCGEditorGraphDebugObjectItem>> OwnerItems;
 		TMap<const FPCGStack, FPCGEditorGraphDebugObjectItemPtr> StackToItem;
 
 		using FGraphContextMapKey = TTuple<UPCGGraph*, uint32 /*GenerationGridSize*/, bool /*IsPartitioned*/>;
 		using FGraphContextMapValue = TTuple<FPCGStackContext, bool /*bContainsGraph*/>;
 		TMap<FGraphContextMapKey, FGraphContextMapValue> GraphToContexts;
 
-		for (UObject* PCGComponentObject : PCGComponents)
+		for (UObject* PCGSourceObject : PCGSources)
 		{
-			if (!IsValid(PCGComponentObject))
+			if (!IsValid(PCGSourceObject))
 			{
 				continue;
 			}
 
-			UPCGComponent* PCGComponent = Cast<UPCGComponent>(PCGComponentObject);
-			if (!PCGComponent || !PCGComponent->IsRegistered() || !PCGComponent->GetGraph())
+			IPCGGraphExecutionSource* PCGSource = Cast<IPCGGraphExecutionSource>(PCGSourceObject);
+			UPCGComponent* PCGComponent = Cast<UPCGComponent>(PCGSource);
+			
+			if (!PCGSource || (PCGComponent && !PCGComponent->IsRegistered()) || !PCGSource->GetExecutionState().GetGraph())
 			{
 				continue;
 			}
 
+			const bool bIsPartitioned = PCGComponent && PCGComponent->IsPartitioned();
+
 			// Process static stacks that can be read from the compiled graph.
 			{
 				//TRACE_CPUPROFILER_EVENT_SCOPE(SPCGEditorGraphDebugObjectTree::RefreshTree::StaticStacks);
 				// Considering it's likely for multiple components to have the same graph, we'll do a per-graph check first prior to building the component + graph stack.
-				UPCGGraph* ComponentGraph = PCGComponent->GetGraph();
-				FGraphContextMapKey MapKey(ComponentGraph, PCGComponent->IsPartitioned() ? PCGHiGenGrid::UnboundedGridSize() : PCGComponent->GetGenerationGridSize(), PCGComponent->IsPartitioned());
+				UPCGGraph* ComponentGraph = PCGSource->GetExecutionState().GetGraph();
+				FGraphContextMapKey MapKey(ComponentGraph, bIsPartitioned ? PCGHiGenGrid::UnboundedGridSize() : PCGComponent->GetGenerationGridSize(), bIsPartitioned);
 				FGraphContextMapValue* MapValue = GraphToContexts.Find(MapKey);
 
 				if (!MapValue)
@@ -1197,9 +1218,9 @@ void SPCGEditorGraphDebugObjectTree::RefreshTree()
 				if (MapValue->Get<1>())
 				{
 					FPCGStack ComponentStack;
-					ComponentStack.PushFrame(PCGComponent);
+					ComponentStack.PushFrame(PCGSourceObject);
 					FPCGStackContext ComponentContext(MapValue->Get<0>(), ComponentStack);
-					AddStacksToTree(ComponentContext.GetStacks(), ActorItems, StackToItem);
+					AddStacksToTree(ComponentContext.GetStacks(), OwnerItems, StackToItem);
 				}
 			}
 		}
@@ -1211,12 +1232,12 @@ void SPCGEditorGraphDebugObjectTree::RefreshTree()
 			TRACE_CPUPROFILER_EVENT_SCOPE(SPCGEditorGraphDebugObjectTree::RefreshTree::DynamicStacks);
 			TArray<FPCGStackSharedPtr> GraphStacks = PCGEditorModule->GetExecutedStacksPtrs(nullptr, PCGGraph, /*bOnlyWithSubgraphAsCurrentFrame=*/!bShowDownstream);
 			// TODO: do per component filtering?
-			AddStacksToTree(GraphStacks, ActorItems, StackToItem);
+			AddStacksToTree(GraphStacks, OwnerItems, StackToItem);
 		}
 
-		for (TPair<AActor*, TSharedPtr<FPCGEditorGraphDebugObjectItem_Actor>>& ActorItem : ActorItems)
+		for (TPair<UObject*, TSharedPtr<FPCGEditorGraphDebugObjectItem>>& OwnerItem : OwnerItems)
 		{
-			RootItems.Add(MoveTemp(ActorItem.Value));
+			RootItems.Add(MoveTemp(OwnerItem.Value));
 		}
 	}
 
@@ -1279,18 +1300,20 @@ void SPCGEditorGraphDebugObjectTree::RestoreTreeState()
 
 			if (ItemStack && SelectedGraph.Get() == ItemStack->GetRootGraph())
 			{
-				const UPCGComponent* RootComponent = ItemStack->GetRootComponent();
+				const IPCGGraphExecutionSource* RootSource = ItemStack->GetRootSource();
+				const UPCGComponent* RootComponent = Cast<UPCGComponent>(RootSource);
 				const AActor* RootOwner = RootComponent ? RootComponent->GetOwner() : nullptr;
 				const APCGPartitionActor* RootPartitionActor = Cast<APCGPartitionActor>(RootOwner);
 
 				if (RootComponent && RootPartitionActor)
 				{
+					const UPCGComponent* SelectedOriginalComponent = Cast<UPCGComponent>(SelectedOriginalSource.Get());
 					// For local components, we can fuzzy match as long as the GridSize, GridCoord, OriginalComponent, and ExecutionDomain are the same.
 					// This is equivalent to saying they are on the same partition actor and come from the same original component.
 					bFuzzyMatch = SelectedGridSize == RootComponent->GetGenerationGridSize()
 						&& SelectedGridCoord == RootPartitionActor->GetGridCoord()
-						&& SelectedOriginalComponent.Get() == RootPartitionActor->GetOriginalComponent(RootComponent)
-						&& (SelectedOriginalComponent.IsValid() && SelectedOriginalComponent->IsManagedByRuntimeGenSystem() == RootComponent->IsManagedByRuntimeGenSystem());
+						&& SelectedOriginalComponent == RootPartitionActor->GetOriginalComponent(RootComponent)
+						&& (IsValid(SelectedOriginalComponent) && SelectedOriginalComponent->IsManagedByRuntimeGenSystem() == RootComponent->IsManagedByRuntimeGenSystem());
 				}
 				else
 				{
@@ -1403,26 +1426,28 @@ void SPCGEditorGraphDebugObjectTree::OnSelectionChanged(FPCGEditorGraphDebugObje
 	SelectedOwner = nullptr;
 	SelectedGridSize = PCGHiGenGrid::UnboundedGridSize();
 	SelectedGridCoord = FIntVector::ZeroValue;
-	SelectedOriginalComponent = nullptr;
+	SelectedOriginalSource.Reset();
 
 	if (const FPCGStack* Stack = InItem ? InItem->GetPCGStack() : nullptr)
 	{
 		SelectedStack = *Stack;
 		SelectedGraph = SelectedStack.GetRootGraph();
 
-		if (const UPCGComponent* RootComponent = SelectedStack.GetRootComponent())
+		if (const IPCGGraphExecutionSource* RootSource = SelectedStack.GetRootSource())
 		{
-			SelectedOwner = RootComponent->GetOwner();
+			const UPCGComponent* RootComponent = Cast<UPCGComponent>(RootSource);
+			SelectedOwner = RootComponent ? RootComponent->GetOwner() : nullptr;
 
 			if (const APCGPartitionActor* PartitionActor = Cast<APCGPartitionActor>(SelectedOwner.Get()))
 			{
+				check(RootComponent);
 				SelectedGridSize = RootComponent->GetGenerationGridSize();
 				SelectedGridCoord = PartitionActor->GetGridCoord();
-				SelectedOriginalComponent = PartitionActor->GetOriginalComponent(RootComponent);
+				SelectedOriginalSource = PartitionActor->GetOriginalComponent(RootComponent);
 			}
 			else
 			{
-				SelectedOriginalComponent = const_cast<UPCGComponent*>(RootComponent);
+				SelectedOriginalSource = RootSource;
 			}
 		}
 	}
diff --git a/Engine/Plugins/PCG/Source/PCGEditor/Private/Widgets/SPCGEditorGraphDebugObjectTree.h b/Engine/Plugins/PCG/Source/PCGEditor/Private/Widgets/SPCGEditorGraphDebugObjectTree.h
--- a/Engine/Plugins/PCG/Source/PCGEditor/Private/Widgets/SPCGEditorGraphDebugObjectTree.h
+++ b/Engine/Plugins/PCG/Source/PCGEditor/Private/Widgets/SPCGEditorGraphDebugObjectTree.h
@@ -98,17 +98,17 @@ protected:
 	FPCGStack PCGStack;
 };
 
-class FPCGEditorGraphDebugObjectItem_PCGComponent : public FPCGEditorGraphDebugObjectItem
+class FPCGEditorGraphDebugObjectItem_PCGSource : public FPCGEditorGraphDebugObjectItem
 {
 public:
-	explicit FPCGEditorGraphDebugObjectItem_PCGComponent(
-		UPCGComponent* InPCGComponent,
+	explicit FPCGEditorGraphDebugObjectItem_PCGSource(
+		IPCGGraphExecutionSource* InSource,
 		const UPCGGraph* InPCGGraph,
 		const FPCGStack& InPCGStack,
 		bool bInIsDebuggable,
 		bool bInHasInspectionData)
 		: FPCGEditorGraphDebugObjectItem(bInHasInspectionData)
-		, PCGComponent(InPCGComponent)
+		, PCGSource(InSource)
 		, PCGGraph(InPCGGraph)
 		, PCGStack(InPCGStack)
 		, bIsDebuggable(bInIsDebuggable)
@@ -122,7 +122,7 @@ public:
 protected:
 	virtual FPCGStack* GetMutablePCGStack() override { return &PCGStack; }
 
-	TSoftObjectPtr<UPCGComponent> PCGComponent;
+	FPCGSoftGraphExecutionSource PCGSource;
 	TSoftObjectPtr<const UPCGGraph> PCGGraph;
 
 	FPCGStack PCGStack;
@@ -285,15 +285,15 @@ private:
 
 	void AddStacksToTree(
 		const TArray<FPCGStackSharedPtr>& Stacks,
-		TMap<AActor*, TSharedPtr<FPCGEditorGraphDebugObjectItem_Actor>>& InOutActorItems,
+		TMap<UObject*, FPCGEditorGraphDebugObjectItemPtr>& InOutOwnerItems,
 		TMap<const FPCGStack, FPCGEditorGraphDebugObjectItemPtr>& InOutStackToItem);
 	void AddStacksToTree(
 		const TArray<FPCGStack>& Stacks,
-		TMap<AActor*, TSharedPtr<FPCGEditorGraphDebugObjectItem_Actor>>& InOutActorItems,
+		TMap<UObject*, FPCGEditorGraphDebugObjectItemPtr>& InOutOwnerItems,
 		TMap<const FPCGStack, FPCGEditorGraphDebugObjectItemPtr>& InOutStackToItem);
 	void AddStackToTree(const FPCGStack& Stack,
 		const UPCGGraph* GraphBeingEdited,
-		TMap<AActor*, TSharedPtr<FPCGEditorGraphDebugObjectItem_Actor>>& InOutActorItems,
+		TMap<UObject*, FPCGEditorGraphDebugObjectItemPtr>& InOutOwnerItems,
 		TMap<const FPCGStack, FPCGEditorGraphDebugObjectItemPtr>& InOutStackToItem);
 
 	void OnEditorSelectionChanged(UObject* InObject);
@@ -358,7 +358,7 @@ private:
 	TSoftObjectPtr<const AActor> SelectedOwner = nullptr;
 	uint32 SelectedGridSize = PCGHiGenGrid::UnboundedGridSize();
 	FIntVector SelectedGridCoord = FIntVector::ZeroValue;
-	TSoftObjectPtr<UPCGComponent> SelectedOriginalComponent = nullptr;
+	FPCGSoftGraphExecutionSource SelectedOriginalSource;
 	TSoftObjectPtr<const UPCGNode> PCGNodeBeingInspected = nullptr;
 
 	/** Controls whether only stacks containing errors/warnings should be shown. */
diff --git a/Engine/Plugins/PCG/Source/PCGEditor/Private/Widgets/SPCGEditorGraphProfilingView.h b/Engine/Plugins/PCG/Source/PCGEditor/Private/Widgets/SPCGEditorGraphProfilingView.h
--- a/Engine/Plugins/PCG/Source/PCGEditor/Private/Widgets/SPCGEditorGraphProfilingView.h
+++ b/Engine/Plugins/PCG/Source/PCGEditor/Private/Widgets/SPCGEditorGraphProfilingView.h
@@ -12,7 +12,6 @@
 
 class FPCGEditor;
 class STableViewBase;
-class UPCGComponent;
 class UPCGEditorGraphNode;
 class UPCGEditorGraph;
 class UPCGNode;
diff --git a/Engine/Plugins/PCG/Source/PCGEditor/Public/PCGEditorModule.h b/Engine/Plugins/PCG/Source/PCGEditor/Public/PCGEditorModule.h
--- a/Engine/Plugins/PCG/Source/PCGEditor/Public/PCGEditorModule.h
+++ b/Engine/Plugins/PCG/Source/PCGEditor/Public/PCGEditorModule.h
@@ -3,6 +3,9 @@
 #pragma once
 
 #include "Editor/IPCGEditorModule.h"
+
+#include "PCGGraphExecutionStateInterface.h"
+
 #include "Utils/PCGNodeVisualLogs.h"
 
 #include "AssetTypeCategories.h"
@@ -40,11 +43,11 @@ public:
 protected:
 	virtual void OnScheduleGraph(const FPCGStackContext& StackContext) override;
 	virtual void OnGraphPreSave(UPCGGraph* Graph, FObjectPreSaveContext ObjectSaveContext) override;
-	virtual void ClearExecutionMetadata(UPCGComponent* InComponent) override;
-	virtual void ClearExecutedStacks(const UPCGComponent* InRootComponent) override;
+	virtual void ClearExecutionMetadata(IPCGGraphExecutionSource* InSource) override;
+	virtual void ClearExecutedStacks(const IPCGGraphExecutionSource* InRootSource) override;
 	virtual void ClearExecutedStacks(const UPCGGraph* InContainingGraph) override;
 	virtual	TArray<FPCGStackSharedPtr> GetExecutedStacksPtrs(const FPCGStack& BeginningWithStack) override;
-	virtual TArray<FPCGStackSharedPtr> GetExecutedStacksPtrs(const UPCGComponent* InComponent, const UPCGGraph* InSubgraph, bool bOnlyWithSubgraphAsCurrentFrame) override;
+	virtual TArray<FPCGStackSharedPtr> GetExecutedStacksPtrs(const IPCGGraphExecutionSource* InSource, const UPCGGraph* InSubgraph, bool bOnlyWithSubgraphAsCurrentFrame) override;
 	virtual void NotifyGraphChanged(UPCGGraph* InGraph, EPCGChangeType ChangeType) override;
 	// ~End IPCGEditorModule implementation
 
@@ -85,6 +88,6 @@ protected:
 
 	/** A record of stacks that were executed. Used to populate debugging tool UIs. */
 	TSet<FPCGStackSharedPtr> ExecutedStacks;
-	TMap<TSoftObjectPtr<UPCGComponent>, TSet<FPCGStackSharedPtr>> ExecutedStacksPerComponent;
+	TMap<FPCGSoftGraphExecutionSource, TSet<FPCGStackSharedPtr>> ExecutedStacksPerSource;
 	mutable FRWLock ExecutedStacksLock;
 };
