PCG Spline Attribute Channel Sampling (1) - Integrate FSpline into FPCGSplineStruct.

Using the same approach as used for USplineComponent, FPCGSplineStruct now has 2 related spline data structures which are selected using the 'PCGSplineStruct.UseSplineCurves' CVar. By setting this to false, FSpline is the authoritative structure.

This change should not result in any functional changes. The change is a prerequisite for the deprecation of SplineCurves (tbd on when to do this) and for the spline sampler to sample spline attribute channels (a forthcoming CL)

diff --git a/Engine/Plugins/PCG/Source/PCG/Private/Data/PCGSplineData.cpp b/Engine/Plugins/PCG/Source/PCG/Private/Data/PCGSplineData.cpp
--- a/Engine/Plugins/PCG/Source/PCG/Private/Data/PCGSplineData.cpp
+++ b/Engine/Plugins/PCG/Source/PCG/Private/Data/PCGSplineData.cpp
@@ -173,7 +173,7 @@ FTransform UPCGSplineData::GetTransformAtDistance(int SegmentIndex, FVector::FRe
 FVector::FReal UPCGSplineData::GetCurvatureAtDistance(int SegmentIndex, FVector::FReal Distance) const
 {
 	const float FullDistance = SplineStruct.GetDistanceAlongSplineAtSplinePoint(SegmentIndex) + Distance;
-	const float Param = SplineStruct.GetSplineRepramTable().Eval(FullDistance, 0.0f);
+	const float Param = SplineStruct.GetInputKeyAtDistanceAlongSpline(FullDistance);
 
 	// Since we need the first derivative (e.g. very similar to direction) to have its norm, we'll get the value directly
 	const FVector FirstDerivative = SplineStruct.GetSplinePointsPosition().EvalDerivative(Param, FVector::ZeroVector);
@@ -194,7 +194,7 @@ FVector::FReal UPCGSplineData::GetCurvatureAtDistance(int SegmentIndex, FVector:
 float UPCGSplineData::GetInputKeyAtDistance(int SegmentIndex, FVector::FReal Distance) const
 {
 	const float FullDistance = GetDistanceAtSegmentStart(SegmentIndex) + Distance;
-	return SplineStruct.GetSplineRepramTable().Eval(FullDistance, 0.0f);
+	return SplineStruct.GetInputKeyAtDistanceAlongSpline(FullDistance);
 }
 
 void UPCGSplineData::GetTangentsAtSegmentStart(int SegmentIndex, FVector& OutArriveTangent, FVector& OutLeaveTangent) const
@@ -281,7 +281,6 @@ void UPCGSplineData::WriteMetadataToPoint(float InputKey, FPCGPoint& OutPoint, U
 TUniquePtr<IPCGAttributeAccessor> UPCGSplineData::CreateStaticAccessor(const FPCGAttributePropertySelector& InSelector, bool bQuiet)
 {
 	static const FStructProperty* PCGSplineStructProperty = CastFieldChecked<FStructProperty>(UPCGSplineData::StaticClass()->FindPropertyByName(GET_MEMBER_NAME_CHECKED(UPCGSplineData, SplineStruct)));
-	static const FProperty* SplineCurveProperty = FPCGSplineStruct::StaticStruct()->FindPropertyByName(TEXT("SplineCurves"));
 
 	const FName DomainName = InSelector.GetDomainName();
 
@@ -321,28 +320,17 @@ TUniquePtr<IPCGAttributeAccessor> UPCGSplineData::CreateStaticAccessor(const FPC
 		{
 			return MakeUnique<FPCGControlPointsAccessor<FTransform, EPCGControlPointsAccessorTarget::Transform, /*bWorldCoordinates=*/false>>(PCGSplineStructProperty);
 		}
-		else if (PropertyName == TEXT("LeaveTangent") || PropertyName == TEXT("ArriveTangent") || PropertyName == TEXT("InterpType"))
+		else if (PropertyName == TEXT("ArriveTangent"))
 		{
-			TArray<const FProperty*> ExtraProperties = {PCGSplineStructProperty, SplineCurveProperty};
-			
-			// Implementation note:
-			// Tangents and Interpolation type are all linked to the Position property on FSplineCurves, even if technically the rotation and the scale also have tangents and interpolation type.
-			// But in the spline details view, and everywhere we have tangents and interpolation, they are always related to the Position interp curve.
-			// If we need access to the rotation/scale tangents/interpolation modes, we can add them later on.
-			const FStructProperty* PositionProperty = CastField<FStructProperty>(FSplineCurves::StaticStruct()->FindPropertyByName(TEXT("Position")));
-			check(PositionProperty);
-			if (PropertyName == TEXT("LeaveTangent"))
-			{
-				return MakeUnique<FPCGInterpCurveAccessor<FInterpCurveVector, EPCGInterpCurveAccessorTarget::LeaveTangent>>(PositionProperty, std::move(ExtraProperties));
-			}
-			else if (PropertyName == TEXT("ArriveTangent"))
-			{
-				return MakeUnique<FPCGInterpCurveAccessor<FInterpCurveVector, EPCGInterpCurveAccessorTarget::ArriveTangent>>(PositionProperty, std::move(ExtraProperties));
-			}
-			else
-			{
-				return MakeUnique<FPCGInterpCurveAccessor<FInterpCurveVector, EPCGInterpCurveAccessorTarget::InterpMode>>(PositionProperty, std::move(ExtraProperties));
-			}
+			return MakeUnique<FPCGControlPointsAccessor<FVector, EPCGControlPointsAccessorTarget::ArriveTangent, /*bWorldCoordinates=*/false>>(PCGSplineStructProperty);
+		}
+		else if (PropertyName == TEXT("LeaveTangent"))
+		{
+			return MakeUnique<FPCGControlPointsAccessor<FVector, EPCGControlPointsAccessorTarget::LeaveTangent, /*bWorldCoordinates=*/false>>(PCGSplineStructProperty);
+		}
+		else if (PropertyName == TEXT("InterpType"))
+		{
+			return MakeUnique<FPCGControlPointsAccessor<FPCGEnumPropertyAccessor::Type, EPCGControlPointsAccessorTarget::InterpMode, /*bWorldCoordinates=*/false>>(PCGSplineStructProperty);
 		}
 
 		if (!bQuiet)
@@ -367,6 +355,58 @@ TUniquePtr<IPCGAttributeAccessor> UPCGSplineData::CreateStaticAccessor(const FPC
 	return {};
 }
 
+TUniquePtr<IPCGAttributeAccessorKeys> UPCGSplineData::CreateAccessorKeys(const FPCGAttributePropertySelector& InSelector, bool bQuiet)
+{
+    const EPCGAttributePropertySelection Selection = InSelector.GetSelection();
+	const FPCGMetadataDomainID DomainID = GetMetadataDomainIDFromSelector(InSelector);
+
+	// Global data
+	if (DomainID == PCGMetadataDomainID::Data && (Selection == EPCGAttributePropertySelection::Property || Selection == EPCGAttributePropertySelection::ExtraProperty))
+	{
+		return MakeUnique<FPCGAttributeAccessorKeysSplineData>(this, /*bInGlobalData=*/true);
+	}
+
+	// Control point data
+	if ((DomainID.IsDefault() || DomainID == PCGMetadataDomainID::Elements) && (Selection == EPCGAttributePropertySelection::Property || Selection == EPCGAttributePropertySelection::ExtraProperty))
+	{
+		return MakeUnique<FPCGAttributeAccessorKeysSplineData>(this, /*bInGlobalData=*/false);
+	}
+
+	// Metadata in the control point domain
+	if ((DomainID.IsDefault() || DomainID == PCGMetadataDomainID::Elements) && Selection == EPCGAttributePropertySelection::Attribute)
+	{
+		return MakeUnique<FPCGAttributeAccessorKeysSplineDataEntries>(this);
+	}
+	
+	return {};
+}
+
+TUniquePtr<const IPCGAttributeAccessorKeys> UPCGSplineData::CreateConstAccessorKeys(const FPCGAttributePropertySelector& InSelector, bool bQuiet) const
+{
+	const EPCGAttributePropertySelection Selection = InSelector.GetSelection();
+	const FPCGMetadataDomainID DomainID = GetMetadataDomainIDFromSelector(InSelector);
+
+	// Global data
+	if (DomainID == PCGMetadataDomainID::Data && (Selection == EPCGAttributePropertySelection::Property || Selection == EPCGAttributePropertySelection::ExtraProperty))
+	{
+		return MakeUnique<const FPCGAttributeAccessorKeysSplineData>(this, /*bInGlobalData=*/true);
+	}
+
+	// Control point data
+	if ((DomainID.IsDefault() || DomainID == PCGMetadataDomainID::Elements) && (Selection == EPCGAttributePropertySelection::Property || Selection == EPCGAttributePropertySelection::ExtraProperty))
+	{
+		return MakeUnique<const FPCGAttributeAccessorKeysSplineData>(this, /*bInGlobalData=*/false);
+	}
+
+	// Metadata in the control point domain
+	if ((DomainID.IsDefault() || DomainID == PCGMetadataDomainID::Elements) && Selection == EPCGAttributePropertySelection::Attribute)
+	{
+		return MakeUnique<const FPCGAttributeAccessorKeysSplineDataEntries>(this);
+	}
+	
+	return {};
+}
+
 FPCGAttributeAccessorMethods UPCGSplineData::GetSplineAccessorMethods()
 {
 	auto CreateAccessorFunc = [](UPCGData* InData, const FPCGAttributePropertySelector& InSelector, bool bQuiet) -> TUniquePtr<IPCGAttributeAccessor>
@@ -381,44 +421,12 @@ FPCGAttributeAccessorMethods UPCGSplineData::GetSplineAccessorMethods()
 
 	auto CreateAccessorKeysFunc = [](UPCGData* InData, const FPCGAttributePropertySelector& InSelector, bool bQuiet) -> TUniquePtr<IPCGAttributeAccessorKeys>
 	{
-		UPCGSplineData* SplineData = CastChecked<UPCGSplineData>(InData);
-		const EPCGAttributePropertySelection Selection = InSelector.GetSelection();
-		
-		if (Selection == EPCGAttributePropertySelection::Property || Selection == EPCGAttributePropertySelection::ExtraProperty)
-		{
-			return MakeUnique<FPCGAttributeAccessorKeysSplineData>(SplineData, InSelector.GetDomainName() == PCGDataConstants::DataDomainName);
-		}
-		else if (Selection == EPCGAttributePropertySelection::Attribute)
-		{
-			const FPCGMetadataDomainID DomainID = InData->GetMetadataDomainIDFromSelector(InSelector);
-			if (DomainID.IsDefault() || DomainID == PCGMetadataDomainID::Elements)
-			{
-				return MakeUnique<FPCGAttributeAccessorKeysSplineDataEntries>(SplineData);
-			}
-		}
-
-		return {};
+		return CastChecked<UPCGSplineData>(InData)->CreateAccessorKeys(InSelector, bQuiet);
 	};
 	
 	auto CreateConstAccessorKeysFunc = [](const UPCGData* InData, const FPCGAttributePropertySelector& InSelector, bool bQuiet) -> TUniquePtr<const IPCGAttributeAccessorKeys>
 	{
-		const UPCGSplineData* SplineData = CastChecked<UPCGSplineData>(InData);
-		
-		const EPCGAttributePropertySelection Selection = InSelector.GetSelection();
-		if (Selection == EPCGAttributePropertySelection::Property || Selection == EPCGAttributePropertySelection::ExtraProperty)
-		{
-			return MakeUnique<const FPCGAttributeAccessorKeysSplineData>(SplineData, InSelector.GetDomainName() == PCGDataConstants::DataDomainName);
-		}
-		else if (Selection == EPCGAttributePropertySelection::Attribute)
-		{
-			const FPCGMetadataDomainID DomainID = InData->GetMetadataDomainIDFromSelector(InSelector);
-			if (DomainID.IsDefault() || DomainID == PCGMetadataDomainID::Elements)
-			{
-				return MakeUnique<FPCGAttributeAccessorKeysSplineDataEntries>(SplineData);
-			}
-		}
-
-		return {};
+		return CastChecked<UPCGSplineData>(InData)->CreateConstAccessorKeys(InSelector, bQuiet);
 	};
 	
 	FPCGAttributeAccessorMethods Methods
diff --git a/Engine/Plugins/PCG/Source/PCG/Private/Data/PCGSplineStruct.cpp b/Engine/Plugins/PCG/Source/PCG/Private/Data/PCGSplineStruct.cpp
--- a/Engine/Plugins/PCG/Source/PCG/Private/Data/PCGSplineStruct.cpp
+++ b/Engine/Plugins/PCG/Source/PCG/Private/Data/PCGSplineStruct.cpp
@@ -8,6 +8,13 @@
 
 namespace PCGSplineStruct
 {
+	bool GUseSplineCurves = true;
+	FAutoConsoleVariableRef CVarUseSplineCurves(
+		TEXT("PCGSplineStruct.UseSplineCurves"),
+		GUseSplineCurves,
+		TEXT("When true, SplineCurves is the 'source of truth'.")
+	);
+	
 	static int32 UpperBound(const TArray<FInterpCurvePoint<FVector>>& SplinePoints, float Value)
 	{
 		int32 Count = SplinePoints.Num();
@@ -31,8 +38,8 @@ namespace PCGSplineStruct
 	}
 
 	// Note: copied verbatim from USplineComponent::CalcBounds
-	static FBoxSphereBounds CalcBounds(const FSplineCurves& SplineCurves, bool bClosedLoop, const FTransform& LocalToWorld)
-	{
+    static FBoxSphereBounds CalcBounds(const FSplineCurves& SplineCurves, bool bClosedLoop, const FTransform& LocalToWorld)
+    {
 //#if SPLINE_FAST_BOUNDS_CALCULATION
 //		FBox BoundingBox(0);
 //		for (const auto& InterpPoint : SplineCurves.Position.Points)
@@ -42,7 +49,54 @@ namespace PCGSplineStruct
 //
 //		return FBoxSphereBounds(BoundingBox.TransformBy(LocalToWorld));
 //#else
-		const int32 NumPoints = SplineCurves.Position.Points.Num();
+    	const int32 NumPoints = SplineCurves.Position.Points.Num();
+    	const int32 NumSegments = bClosedLoop ? NumPoints : NumPoints - 1;
+
+    	FVector Min(WORLD_MAX);
+    	FVector Max(-WORLD_MAX);
+    	if (NumSegments > 0)
+    	{
+    		for (int32 Index = 0; Index < NumSegments; Index++)
+    		{
+    			const bool bLoopSegment = (Index == NumPoints - 1);
+    			const int32 NextIndex = bLoopSegment ? 0 : Index + 1;
+    			const FInterpCurvePoint<FVector>& ThisInterpPoint = SplineCurves.Position.Points[Index];
+    			FInterpCurvePoint<FVector> NextInterpPoint = SplineCurves.Position.Points[NextIndex];
+    			if (bLoopSegment)
+    			{
+    				NextInterpPoint.InVal = ThisInterpPoint.InVal + SplineCurves.Position.LoopKeyOffset;
+    			}
+
+    			CurveVectorFindIntervalBounds(ThisInterpPoint, NextInterpPoint, Min, Max);
+    		}
+    	}
+    	else if (NumPoints == 1)
+    	{
+    		Min = Max = SplineCurves.Position.Points[0].OutVal;
+    	}
+    	else
+    	{
+    		Min = FVector::ZeroVector;
+    		Max = FVector::ZeroVector;
+    	}
+
+    	return FBoxSphereBounds(FBox(Min, Max).TransformBy(LocalToWorld));
+//#endif
+    }
+	
+	// Note: copied verbatim from USplineComponent::CalcBounds
+	static FBoxSphereBounds CalcBounds(const FSpline& Spline, bool bClosedLoop, const FTransform& LocalToWorld)
+	{
+//#if SPLINE_FAST_BOUNDS_CALCULATION
+//		FBox BoundingBox(0);
+//		for (const auto& InterpPoint : Spline.Position.Points)
+//		{
+//			BoundingBox += InterpPoint.OutVal;
+//		}
+//
+//		return FBoxSphereBounds(BoundingBox.TransformBy(LocalToWorld));
+//#else
+		const int32 NumPoints = Spline.GetNumControlPoints();
 		const int32 NumSegments = bClosedLoop ? NumPoints : NumPoints - 1;
 
 		FVector Min(WORLD_MAX);
@@ -53,11 +107,11 @@ namespace PCGSplineStruct
 			{
 				const bool bLoopSegment = (Index == NumPoints - 1);
 				const int32 NextIndex = bLoopSegment ? 0 : Index + 1;
-				const FInterpCurvePoint<FVector>& ThisInterpPoint = SplineCurves.Position.Points[Index];
-				FInterpCurvePoint<FVector> NextInterpPoint = SplineCurves.Position.Points[NextIndex];
+				const FInterpCurvePoint<FVector>& ThisInterpPoint = Spline.GetSplinePointsPosition().Points[Index];
+				FInterpCurvePoint<FVector> NextInterpPoint = Spline.GetSplinePointsPosition().Points[NextIndex];
 				if (bLoopSegment)
 				{
-					NextInterpPoint.InVal = ThisInterpPoint.InVal + SplineCurves.Position.LoopKeyOffset;
+					NextInterpPoint.InVal = ThisInterpPoint.InVal + Spline.GetSplinePointsPosition().LoopKeyOffset;
 				}
 
 				CurveVectorFindIntervalBounds(ThisInterpPoint, NextInterpPoint, Min, Max);
@@ -65,7 +119,7 @@ namespace PCGSplineStruct
 		}
 		else if (NumPoints == 1)
 		{
-			Min = Max = SplineCurves.Position.Points[0].OutVal;
+			Min = Max = Spline.GetSplinePointsPosition().Points[0].OutVal;
 		}
 		else
 		{
@@ -81,7 +135,10 @@ namespace PCGSplineStruct
 void FPCGSplineStruct::Initialize(const USplineComponent* InSplineComponent)
 {
 	check(InSplineComponent);
-	SplineCurves = InSplineComponent->GetSplineCurves();
+
+	WarninglessSplineCurves() = InSplineComponent->GetSplineCurves();
+	Spline = InSplineComponent->GetSpline();
+	
 	Transform = InSplineComponent->GetComponentTransform();
 	DefaultUpVector = InSplineComponent->DefaultUpVector;
 	ReparamStepsPerSegment = InSplineComponent->ReparamStepsPerSegment;
@@ -102,8 +159,16 @@ void FPCGSplineStruct::Initialize(const TArray<FSplinePoint>& InSplinePoints, bo
 	bClosedLoop = bIsClosedLoop;
 	AddPoints(InSplinePoints, true);
 
-	Bounds = PCGSplineStruct::CalcBounds(SplineCurves, bClosedLoop, InTransform);
-	LocalBounds = PCGSplineStruct::CalcBounds(SplineCurves, bClosedLoop, FTransform::Identity);
+	if (ShouldUseSplineCurves())
+	{
+		Bounds = PCGSplineStruct::CalcBounds(WarninglessSplineCurves(), bClosedLoop, InTransform);
+		LocalBounds = PCGSplineStruct::CalcBounds(WarninglessSplineCurves(), bClosedLoop, FTransform::Identity);
+	}
+	else
+	{
+		Bounds = PCGSplineStruct::CalcBounds(Spline, bClosedLoop, InTransform);
+		LocalBounds = PCGSplineStruct::CalcBounds(Spline, bClosedLoop, FTransform::Identity);
+	}
 
 	if (!InOptionalEntryKeys.IsEmpty() && InSplinePoints.Num() == InOptionalEntryKeys.Num())
 	{
@@ -125,7 +190,7 @@ void FPCGSplineStruct::ApplyTo(USplineComponent* InSplineComponent) const
 	InSplineComponent->DefaultUpVector = DefaultUpVector;
 	InSplineComponent->ReparamStepsPerSegment = ReparamStepsPerSegment;
 
-	InSplineComponent->SetSpline(SplineCurves);
+	InSplineComponent->SetSpline(Spline);
 	InSplineComponent->bStationaryEndpoints = false;
 	// TODO: metadata? might not be needed
 	InSplineComponent->SetClosedLoop(bClosedLoop);
@@ -135,34 +200,37 @@ void FPCGSplineStruct::ApplyTo(USplineComponent* InSplineComponent) const
 
 void FPCGSplineStruct::AddPoint(const FSplinePoint& InSplinePoint, bool bUpdateSpline)
 {
-	const int32 Index = PCGSplineStruct::UpperBound(SplineCurves.Position.Points, InSplinePoint.InputKey);
+	const TArray<FInterpCurvePoint<FVector>>& SourcePoints = ShouldUseSplineCurves()
+		? WarninglessSplineCurves().Position.Points
+		: Spline.GetSplinePointsPosition().Points;
+	
+	const int32 Index = PCGSplineStruct::UpperBound(SourcePoints, InSplinePoint.InputKey);
 
-	SplineCurves.Position.Points.Insert(FInterpCurvePoint<FVector>(
+	Spline.AddPoint(InSplinePoint);
+
+	WarninglessSplineCurves().Position.Points.Insert(FInterpCurvePoint<FVector>(
 		InSplinePoint.InputKey,
 		InSplinePoint.Position,
 		InSplinePoint.ArriveTangent,
 		InSplinePoint.LeaveTangent,
 		ConvertSplinePointTypeToInterpCurveMode(InSplinePoint.Type)
-		),
-		Index);
+	), Index);
 
-	SplineCurves.Rotation.Points.Insert(FInterpCurvePoint<FQuat>(
+	WarninglessSplineCurves().Rotation.Points.Insert(FInterpCurvePoint<FQuat>(
 		InSplinePoint.InputKey,
 		InSplinePoint.Rotation.Quaternion(),
 		FQuat::Identity,
 		FQuat::Identity,
 		CIM_CurveAuto
-		),
-		Index);
+	), Index);
 
-	SplineCurves.Scale.Points.Insert(FInterpCurvePoint<FVector>(
+	WarninglessSplineCurves().Scale.Points.Insert(FInterpCurvePoint<FVector>(
 		InSplinePoint.InputKey,
 		InSplinePoint.Scale,
 		FVector::ZeroVector,
 		FVector::ZeroVector,
 		CIM_CurveAuto
-		),
-		Index);
+	), Index);
 
 	if (!ControlPointsEntryKeys.IsEmpty())
 	{
@@ -177,10 +245,6 @@ void FPCGSplineStruct::AddPoint(const FSplinePoint& InSplinePoint, bool bUpdateS
 
 void FPCGSplineStruct::AddPoints(const TArray<FSplinePoint>& InSplinePoints, bool bUpdateSpline)
 {
-	SplineCurves.Position.Points.Reserve(SplineCurves.Position.Points.Num() + InSplinePoints.Num());
-	SplineCurves.Rotation.Points.Reserve(SplineCurves.Rotation.Points.Num() + InSplinePoints.Num());
-	SplineCurves.Scale.Points.Reserve(SplineCurves.Scale.Points.Num() + InSplinePoints.Num());
-	
 	if (!ControlPointsEntryKeys.IsEmpty())
 	{
 		ControlPointsEntryKeys.Reserve(ControlPointsEntryKeys.Num() + InSplinePoints.Num());
@@ -197,36 +261,125 @@ void FPCGSplineStruct::AddPoints(const TArray<FSplinePoint>& InSplinePoints, boo
 	}
 }
 
+void FPCGSplineStruct::SetLocation(int32 Index, const FVector& InLocation)
+{
+	WarninglessSplineCurves().Position.Points[Index].OutVal = InLocation;
+	Spline.SetLocation(Index, InLocation);
+}
+
+FVector FPCGSplineStruct::GetLocation(const int32 Index) const
+{
+	return ShouldUseSplineCurves()
+		? WarninglessSplineCurves().Position.Points[Index].OutVal
+		: Spline.GetLocation(Index);
+}
+	
+void FPCGSplineStruct::SetArriveTangent(const int32 Index, const FVector& InArriveTangent)
+{
+	WarninglessSplineCurves().Position.Points[Index].ArriveTangent = InArriveTangent;
+	Spline.SetInTangent(Index, InArriveTangent);
+}
+
+FVector FPCGSplineStruct::GetArriveTangent(const int32 Index) const
+{
+	return ShouldUseSplineCurves()
+		? WarninglessSplineCurves().Position.Points[Index].ArriveTangent
+		: Spline.GetInTangent(Index);
+}
+	
+void FPCGSplineStruct::SetLeaveTangent(const int32 Index, const FVector& InLeaveTangent)
+{
+	WarninglessSplineCurves().Position.Points[Index].LeaveTangent = InLeaveTangent;
+	Spline.SetOutTangent(Index, InLeaveTangent);
+}
+
+FVector FPCGSplineStruct::GetLeaveTangent(const int32 Index) const
+{
+	return ShouldUseSplineCurves()
+		? WarninglessSplineCurves().Position.Points[Index].LeaveTangent
+		: Spline.GetOutTangent(Index);
+}
+
+void FPCGSplineStruct::SetRotation(int32 Index, const FQuat& InRotation)
+{
+	WarninglessSplineCurves().Rotation.Points[Index].OutVal = InRotation;
+	Spline.SetRotation(Index, InRotation);
+}
+
+FQuat FPCGSplineStruct::GetRotation(const int32 Index) const
+{
+	return ShouldUseSplineCurves()
+		? WarninglessSplineCurves().Rotation.Points[Index].OutVal
+		: Spline.GetRotation(Index);
+}
+	
+void FPCGSplineStruct::SetScale(int32 Index, const FVector& InScale)
+{
+	WarninglessSplineCurves().Scale.Points[Index].OutVal = InScale;
+	Spline.SetScale(Index, InScale);
+}
+
+FVector FPCGSplineStruct::GetScale(const int32 Index) const
+{
+	return ShouldUseSplineCurves()
+		? WarninglessSplineCurves().Scale.Points[Index].OutVal
+		: Spline.GetScale(Index);
+}
+	
+void FPCGSplineStruct::SetSplinePointType(int32 Index, EInterpCurveMode Type)
+{
+	WarninglessSplineCurves().Position.Points[Index].InterpMode = Type;
+	Spline.SetSplinePointType(Index, Type);
+}
+
+EInterpCurveMode FPCGSplineStruct::GetSplinePointType(int32 Index) const
+{
+	return ShouldUseSplineCurves()
+		? WarninglessSplineCurves().Position.Points[Index].InterpMode.GetValue()
+		: Spline.GetSplinePointType(Index);
+}
+
 void FPCGSplineStruct::UpdateSpline()
 {
 	const bool bLoopPositionOverride = false;
 	const bool bStationaryEndpoints = false;
 	const float LoopPosition = 0.0f;
 
-	SplineCurves.UpdateSpline(bClosedLoop, bStationaryEndpoints, ReparamStepsPerSegment, bLoopPositionOverride, LoopPosition, Transform.GetScale3D());
+	FSpline::FUpdateSplineParams Params{bClosedLoop, bStationaryEndpoints, ReparamStepsPerSegment, bLoopPositionOverride, LoopPosition};
+
+	WarninglessSplineCurves().UpdateSpline(Params.bClosedLoop, Params.bStationaryEndpoints, Params.ReparamStepsPerSegment, Params.bLoopPositionOverride, Params.LoopPosition, Params.Scale3D);
+	Spline.UpdateSpline(Params);
 }
 
 int FPCGSplineStruct::GetNumberOfSplineSegments() const
 {
-	const int32 NumPoints = SplineCurves.Position.Points.Num();
-	return (bClosedLoop ? NumPoints : NumPoints - 1);
+	const int32 NumPoints = GetNumberOfPoints();
+	return (bClosedLoop ? NumPoints : FMath::Max(0, NumPoints - 1));
 }
 
 int FPCGSplineStruct::GetNumberOfPoints() const
 {
-	return SplineCurves.Position.Points.Num();
+	return ShouldUseSplineCurves()
+		? WarninglessSplineCurves().Position.Points.Num()
+		: Spline.GetNumControlPoints();
 }
 
 FVector::FReal FPCGSplineStruct::GetSplineLength() const
 {
-	return SplineCurves.GetSplineLength();
+	return ShouldUseSplineCurves()
+		? WarninglessSplineCurves().GetSplineLength()
+		: Spline.GetSplineLength();
 }
 
 FBox FPCGSplineStruct::GetBounds() const
 {
 	// See USplineComponent::CalcBounds
-	const int32 NumPoints = SplineCurves.Position.Points.Num();
-	const int32 NumSegments = bClosedLoop ? NumPoints : NumPoints - 1;
+	const int32 NumPoints = GetNumberOfPoints();
+	const int32 NumSegments = GetNumberOfSplineSegments();
+	
+	const FInterpCurveVector& SplinePointsPosition = ShouldUseSplineCurves()
+		? WarninglessSplineCurves().Position
+		: Spline.GetSplinePointsPosition();
 
 	FVector Min(WORLD_MAX);
 	FVector Max(-WORLD_MAX);
@@ -236,11 +389,11 @@ FBox FPCGSplineStruct::GetBounds() const
 		{
 			const bool bLoopSegment = (Index == NumPoints - 1);
 			const int32 NextIndex = bLoopSegment ? 0 : Index + 1;
-			const FInterpCurvePoint<FVector>& ThisInterpPoint = SplineCurves.Position.Points[Index];
-			FInterpCurvePoint<FVector> NextInterpPoint = SplineCurves.Position.Points[NextIndex];
+			const FInterpCurvePoint<FVector>& ThisInterpPoint = SplinePointsPosition.Points[Index];
+			FInterpCurvePoint<FVector> NextInterpPoint = SplinePointsPosition.Points[NextIndex];
 			if (bLoopSegment)
 			{
-				NextInterpPoint.InVal = ThisInterpPoint.InVal /* + SplineCurves.Position.LoopKeyOffset*/;
+				NextInterpPoint.InVal = ThisInterpPoint.InVal /* + SplinePointsPosition.LoopKeyOffset*/;
 			}
 
 			CurveVectorFindIntervalBounds(ThisInterpPoint, NextInterpPoint, Min, Max);
@@ -248,7 +401,7 @@ FBox FPCGSplineStruct::GetBounds() const
 	}
 	else if (NumPoints == 1)
 	{
-		Min = Max = SplineCurves.Position.Points[0].OutVal;
+		Min = Max = SplinePointsPosition.Points[0].OutVal;
 	}
 	else
 	{
@@ -261,34 +414,50 @@ FBox FPCGSplineStruct::GetBounds() const
 
 const FInterpCurveVector& FPCGSplineStruct::GetSplinePointsScale() const
 {
-	return SplineCurves.Scale;
+	return ShouldUseSplineCurves()
+		? WarninglessSplineCurves().Scale
+		: Spline.GetSplinePointsScale();
 }
 
 const FInterpCurveVector& FPCGSplineStruct::GetSplinePointsPosition() const
 {
-	return SplineCurves.Position;
-}
-
-const FInterpCurveFloat& FPCGSplineStruct::GetSplineRepramTable() const
-{
-	return SplineCurves.ReparamTable;
+	return ShouldUseSplineCurves()
+		? WarninglessSplineCurves().Position
+		: Spline.GetSplinePointsPosition();
 }
 
 const FInterpCurveQuat& FPCGSplineStruct::GetSplinePointsRotation() const
 {
-	return SplineCurves.Rotation;
+	return ShouldUseSplineCurves()
+		? WarninglessSplineCurves().Rotation
+		: Spline.GetSplinePointsRotation();
 }
 
 FVector::FReal FPCGSplineStruct::GetDistanceAlongSplineAtSplinePoint(int32 PointIndex) const
 {
-	const int32 NumPoints = SplineCurves.Position.Points.Num();
-	const int32 NumSegments = bClosedLoop ? NumPoints : NumPoints - 1;
-
-	// Ensure that if the reparam table is not prepared yet we don't attempt to access it. This can happen
-	// early in the construction of the spline component object.
-	if ((PointIndex >= 0) && (PointIndex < NumSegments + 1) && ((PointIndex * ReparamStepsPerSegment) < SplineCurves.ReparamTable.Points.Num()))
+	if (IsClosedLoop() && PointIndex == GetNumberOfPoints())
+	{
+		// Special case, if we are closed and the index here is 1 past the last valid point, the length is the full spline.
+		return GetSplineLength();
+	}
+	
+	if (ShouldUseSplineCurves())
+	{
+		const int32 NumPoints = GetNumberOfPoints();
+		const int32 NumSegments = bClosedLoop ? NumPoints : NumPoints - 1;
+		const int32 NumReparamPoints = WarninglessSplineCurves().ReparamTable.Points.Num();
+        
+		// Ensure that if the reparam table is not prepared yet we don't attempt to access it. This can happen
+		// early in the construction of the spline component object.
+		if ((PointIndex >= 0) && (PointIndex < NumSegments + 1) && ((PointIndex * ReparamStepsPerSegment) < NumReparamPoints))
+		{
+			return WarninglessSplineCurves().ReparamTable.Points[PointIndex * ReparamStepsPerSegment].InVal;
+		}
+	}
+	else
 	{
-		return SplineCurves.ReparamTable.Points[PointIndex * ReparamStepsPerSegment].InVal;
+		const float ParameterAtIndex = Spline.GetParameterAtIndex(PointIndex);
+		return Spline.GetDistanceAtParameter(ParameterAtIndex);
 	}
 
 	return 0.0f;
@@ -296,16 +465,36 @@ FVector::FReal FPCGSplineStruct::GetDistanceAlongSplineAtSplinePoint(int32 Point
 
 FVector FPCGSplineStruct::GetLocationAtDistanceAlongSpline(FVector::FReal Distance, ESplineCoordinateSpace::Type CoordinateSpace) const
 {
-	const float Param = SplineCurves.ReparamTable.Eval(Distance, 0.0f);
+	const float Param = ShouldUseSplineCurves()
+		? WarninglessSplineCurves().ReparamTable.Eval(Distance, 0.0f)
+		: Spline.GetParameterAtDistance(Distance);
+	
 	return GetLocationAtSplineInputKey(Param, CoordinateSpace);
 }
 
 FTransform FPCGSplineStruct::GetTransformAtDistanceAlongSpline(FVector::FReal Distance, ESplineCoordinateSpace::Type CoordinateSpace, bool bUseScale) const
 {
-	const float Param = SplineCurves.ReparamTable.Eval(Distance, 0.0f);
+	const float Param = ShouldUseSplineCurves()
+		? WarninglessSplineCurves().ReparamTable.Eval(Distance, 0.0f)
+		: Spline.GetParameterAtDistance(Distance);
+	
 	return GetTransformAtSplineInputKey(Param, CoordinateSpace, bUseScale);
 }
 
+float FPCGSplineStruct::GetInputKeyAtDistanceAlongSpline(FVector::FReal Distance) const
+{
+	const int32 NumPoints = GetNumberOfPoints();
+
+	if (NumPoints < 2)
+	{
+		return 0.0f;
+	}
+	
+	return ShouldUseSplineCurves()
+		? WarninglessSplineCurves().ReparamTable.Eval(Distance, 0.0f)
+		: Spline.GetParameterAtDistance(Distance);
+}
+
 FVector FPCGSplineStruct::GetRightVectorAtSplineInputKey(float InKey, ESplineCoordinateSpace::Type CoordinateSpace) const
 {
 	const FQuat Quat = GetQuaternionAtSplineInputKey(InKey, ESplineCoordinateSpace::Local);
@@ -339,32 +528,41 @@ float FPCGSplineStruct::FindInputKeyClosestToWorldLocation(const FVector& WorldL
 {
 	const FVector LocalLocation = Transform.InverseTransformPosition(WorldLocation);
 	float Dummy;
-	return SplineCurves.Position.InaccurateFindNearest(LocalLocation, Dummy);
+	return ShouldUseSplineCurves()
+		? WarninglessSplineCurves().Position.FindNearest(LocalLocation, Dummy)
+		: Spline.FindNearest(LocalLocation, Dummy);
 }
 
 TTuple<int, float> FPCGSplineStruct::GetSegmentStartIndexAndKeyAtInputKey(float InKey) const
 {
-	const int32 Index = SplineCurves.Position.GetPointIndexForInputValue(InKey);
+	const int32 Index = ShouldUseSplineCurves()
+		? WarninglessSplineCurves().Position.GetPointIndexForInputValue(InKey)
+		: Spline.GetSplinePointsPosition().GetPointIndexForInputValue(InKey);
 	return {Index, GetInputKeyAtSegmentStart(Index)};
 }
 
 float FPCGSplineStruct::GetInputKeyAtSegmentStart(int InSegmentIndex) const
 {
-	if (IsClosedLoop() && !SplineCurves.Position.Points.IsEmpty() && InSegmentIndex >= SplineCurves.Position.Points.Num())
+	const auto& SplinePointsPosition = ShouldUseSplineCurves()
+		? WarninglessSplineCurves().Position
+		: Spline.GetSplinePointsPosition();
+
+	if (IsClosedLoop() && !SplinePointsPosition.Points.IsEmpty() && InSegmentIndex >= SplinePointsPosition.Points.Num())
 	{
 		// In case of a closed loop, the last point is the first point, and the input key is the last point + LoopKeyOffset
-		return SplineCurves.Position.Points.Last().InVal + SplineCurves.Position.LoopKeyOffset;
+		return SplinePointsPosition.Points.Last().InVal + SplinePointsPosition.LoopKeyOffset;
 	}
 	else
 	{
-		return SplineCurves.Position.Points.IsValidIndex(InSegmentIndex) ? SplineCurves.Position.Points[InSegmentIndex].InVal : 0.0f;
+		return SplinePointsPosition.Points.IsValidIndex(InSegmentIndex) ? SplinePointsPosition.Points[InSegmentIndex].InVal : 0.0f;
 	}
 }
 
 void FPCGSplineStruct::AllocateMetadataEntries()
 {
 	// Add robustness to cleanup everything if we have a mismatch between the number of points and the number of entry keys
-	const int32 NumPoints = SplineCurves.Position.Points.Num();
+	const int32 NumPoints = GetNumberOfPoints();
+	
 	if (ControlPointsEntryKeys.Num() != NumPoints)
 	{
 		ControlPointsEntryKeys.Empty(NumPoints);
@@ -382,7 +580,9 @@ void FPCGSplineStruct::AllocateMetadataEntries()
 
 FVector FPCGSplineStruct::GetLocationAtSplineInputKey(float InKey, ESplineCoordinateSpace::Type CoordinateSpace) const
 {
-	FVector Location = SplineCurves.Position.Eval(InKey, FVector::ZeroVector);
+	FVector Location = ShouldUseSplineCurves()
+		? WarninglessSplineCurves().Position.Eval(InKey, FVector::ZeroVector)
+		: Spline.Evaluate(InKey);
 
 	if (CoordinateSpace == ESplineCoordinateSpace::World)
 	{
@@ -394,10 +594,15 @@ FVector FPCGSplineStruct::GetLocationAtSplineInputKey(float InKey, ESplineCoordi
 
 FQuat FPCGSplineStruct::GetQuaternionAtSplineInputKey(float InKey, ESplineCoordinateSpace::Type CoordinateSpace) const
 {
-	FQuat Quat = SplineCurves.Rotation.Eval(InKey, FQuat::Identity);
+	FQuat Quat = ShouldUseSplineCurves()
+		? WarninglessSplineCurves().Rotation.Eval(InKey, FQuat::Identity)
+		: Spline.EvaluateRotation(InKey);
 	Quat.Normalize();
 
-	const FVector Direction = SplineCurves.Position.EvalDerivative(InKey, FVector::ZeroVector).GetSafeNormal();
+	FVector Direction = ShouldUseSplineCurves()
+		? WarninglessSplineCurves().Position.EvalDerivative(InKey, FVector::ZeroVector).GetSafeNormal()
+		: Spline.EvaluateDerivative(InKey).GetSafeNormal();
+	
 	const FVector UpVector = Quat.RotateVector(DefaultUpVector);
 
 	FQuat Rot = (FRotationMatrix::MakeFromXZ(Direction, UpVector)).ToQuat();
@@ -412,7 +617,9 @@ FQuat FPCGSplineStruct::GetQuaternionAtSplineInputKey(float InKey, ESplineCoordi
 
 FVector FPCGSplineStruct::GetScaleAtSplineInputKey(float InKey) const
 {
-	return SplineCurves.Scale.Eval(InKey, FVector(1.0f));
+	return ShouldUseSplineCurves()
+		? WarninglessSplineCurves().Scale.Eval(InKey, FVector(1.0f))
+		: Spline.GetSplinePointsScale().Eval(InKey, FVector(1.0f));
 }
 
 /** Taken from USplineComponent::ConvertSplineSegmentToPolyLine. */
@@ -422,40 +629,40 @@ bool FPCGSplineStruct::ConvertSplineSegmentToPolyLine(int32 SplinePointStartInde
 	
 	TArray<double> DummyDistances;
 
-    const double StartDist = GetDistanceAlongSplineAtSplinePoint(SplinePointStartIndex);
-    const double StopDist = GetDistanceAlongSplineAtSplinePoint(SplinePointStartIndex + 1);
+	const double StartDist = GetDistanceAlongSplineAtSplinePoint(SplinePointStartIndex);
+	const double StopDist = GetDistanceAlongSplineAtSplinePoint(SplinePointStartIndex + 1);
 
-    const int32 NumLines = 2; // Dichotomic subdivision of the spline segment
-    double Dist = StopDist - StartDist;
-    double SubstepSize = Dist / NumLines;
-    if (SubstepSize == 0.0)
-    {
-    	// There is no distance to cover, so handle the segment with a single point
-    	OutPoints.Add(GetLocationAtDistanceAlongSpline(StopDist, CoordinateSpace));
-    	return true;
-    }
+	const int32 NumLines = 2; // Dichotomic subdivision of the spline segment
+	double Dist = StopDist - StartDist;
+	double SubstepSize = Dist / NumLines;
+	if (SubstepSize == 0.0)
+	{
+		// There is no distance to cover, so handle the segment with a single point
+		OutPoints.Add(GetLocationAtDistanceAlongSpline(StopDist, CoordinateSpace));
+		return true;
+	}
 
-    double SubstepStartDist = StartDist;
-    for (int32 i = 0; i < NumLines; ++i)
-    {
-    	double SubstepEndDist = SubstepStartDist + SubstepSize;
-    	TArray<FVector> NewPoints;
-    	DummyDistances.Reset();
-    	// Recursively sub-divide each segment until the requested precision is reached :
-    	if (DivideSplineIntoPolylineRecursiveWithDistancesHelper(SubstepStartDist, SubstepEndDist, CoordinateSpace, MaxSquareDistanceFromSpline, NewPoints, DummyDistances))
-    	{
-    		if (OutPoints.Num() > 0)
-    		{
-    			check(OutPoints.Last() == NewPoints[0]); // our last point must be the same as the new segment's first
-    			OutPoints.RemoveAt(OutPoints.Num() - 1);
-    		}
-    		OutPoints.Append(NewPoints);
-    	}
+	double SubstepStartDist = StartDist;
+	for (int32 i = 0; i < NumLines; ++i)
+	{
+		double SubstepEndDist = SubstepStartDist + SubstepSize;
+		TArray<FVector> NewPoints;
+		DummyDistances.Reset();
+		// Recursively sub-divide each segment until the requested precision is reached :
+		if (DivideSplineIntoPolylineRecursiveWithDistancesHelper(SubstepStartDist, SubstepEndDist, CoordinateSpace, MaxSquareDistanceFromSpline, NewPoints, DummyDistances))
+		{
+			if (OutPoints.Num() > 0)
+			{
+				check(OutPoints.Last() == NewPoints[0]); // our last point must be the same as the new segment's first
+				OutPoints.RemoveAt(OutPoints.Num() - 1);
+			}
+			OutPoints.Append(NewPoints);
+		}
 
-    	SubstepStartDist = SubstepEndDist;
-    }
+		SubstepStartDist = SubstepEndDist;
+	}
 
-    return (OutPoints.Num() > 0);
+	return (OutPoints.Num() > 0);
 }
 
 /** Taken from USplineComponent::ConvertSplineToPolyLine. */
@@ -530,4 +737,24 @@ bool FPCGSplineStruct::DivideSplineIntoPolylineRecursiveWithDistancesHelper(floa
 
 	check(OutPoints.Num() == OutDistancesAlongSpline.Num())
 	return (OutPoints.Num() > 0);
+}
+
+void FPCGSplineStruct::PostSerialize(const FArchive& Ar)
+{
+#if WITH_EDITOR
+	if (Ar.IsLoading() && Ar.IsPersistent() && !Ar.HasAnyPortFlags(PPF_Duplicate | PPF_DuplicateForPIE))
+	{
+		// This is kind of a hack. It will be true for every FPCGSplineStruct which did not have an FSpline before but does not.
+		// It is also going to be true if we already had a Spline but did not save it while it was configured to actually do anything.
+		if (!Spline.WasEnabled())
+		{
+			Spline = WarninglessSplineCurves();
+		}
+	}
+#endif // WITH_EDITOR
+}
+
+bool FPCGSplineStruct::ShouldUseSplineCurves() const
+{
+	return PCGSplineStruct::GUseSplineCurves || Spline.IsEnabled() == false;
 }
\ No newline at end of file
diff --git a/Engine/Plugins/PCG/Source/PCG/Public/Data/PCGSplineData.h b/Engine/Plugins/PCG/Source/PCG/Public/Data/PCGSplineData.h
--- a/Engine/Plugins/PCG/Source/PCG/Public/Data/PCGSplineData.h
+++ b/Engine/Plugins/PCG/Source/PCG/Public/Data/PCGSplineData.h
@@ -89,6 +89,11 @@ public:
 	/** Static helper to create an accessor on a data that doesn't yet exist, as accessors for spline data don't rely on existing data. */
 	static UE_API TUniquePtr<IPCGAttributeAccessor> CreateStaticAccessor(const FPCGAttributePropertySelector& InSelector, bool bQuiet = false);
 
+	/** Helper to create an accessor keys for data that doesn't yet exist. */
+	UE_API TUniquePtr<IPCGAttributeAccessorKeys> CreateAccessorKeys(const FPCGAttributePropertySelector& InSelector, bool bQuiet = false);
+	UE_API TUniquePtr<const IPCGAttributeAccessorKeys> CreateConstAccessorKeys(const FPCGAttributePropertySelector& InSelector, bool bQuiet = false) const;
+
+	
 	// Get the functions to the accessor factory
 	static UE_API FPCGAttributeAccessorMethods GetSplineAccessorMethods();
 
diff --git a/Engine/Plugins/PCG/Source/PCG/Public/Data/PCGSplineStruct.h b/Engine/Plugins/PCG/Source/PCG/Public/Data/PCGSplineStruct.h
--- a/Engine/Plugins/PCG/Source/PCG/Public/Data/PCGSplineStruct.h
+++ b/Engine/Plugins/PCG/Source/PCG/Public/Data/PCGSplineStruct.h
@@ -39,7 +39,10 @@ enum class EPCGControlPointsAccessorTarget
 	Location,
 	Rotation,
 	Scale,
-	Transform
+	Transform,
+	ArriveTangent,
+	LeaveTangent,
+	InterpMode
 };
 
 /** Subset of the Spline Component API in a standalone struct */
@@ -57,6 +60,25 @@ struct FPCGSplineStruct
 	// Spline-related methods
 	UE_API void AddPoint(const FSplinePoint& InSplinePoint, bool bUpdateSpline); 
 	UE_API void AddPoints(const TArray<FSplinePoint>& InSplinePoints, bool bUpdateSpline);
+
+	UE_API void SetLocation(int32 Index, const FVector& InLocation);
+	UE_API FVector GetLocation(const int32 Index) const;
+	
+	UE_API void SetArriveTangent(const int32 Index, const FVector& InArriveTangent);
+	UE_API FVector GetArriveTangent(const int32 Index) const;
+	
+	UE_API void SetLeaveTangent(const int32 Index, const FVector& LeaveTangent);
+	UE_API FVector GetLeaveTangent(const int32 Index) const;
+
+	UE_API void SetRotation(int32 Index, const FQuat& InRotation);
+	UE_API FQuat GetRotation(const int32 Index) const;
+	
+	UE_API void SetScale(int32 Index, const FVector& InScale);
+	UE_API FVector GetScale(const int32 Index) const;
+	
+	UE_API void SetSplinePointType(int32 Index, EInterpCurveMode Type);
+	UE_API EInterpCurveMode GetSplinePointType(int32 Index) const;
+	
 	UE_API void UpdateSpline();
 
 	UE_API int GetNumberOfSplineSegments() const;
@@ -68,7 +90,6 @@ struct FPCGSplineStruct
 	UE_API const FInterpCurveVector& GetSplinePointsScale() const;
 	UE_API const FInterpCurveQuat& GetSplinePointsRotation() const;
 	UE_API const FInterpCurveVector& GetSplinePointsPosition() const;
-	UE_API const FInterpCurveFloat& GetSplineRepramTable() const;
 
 	UE_API FVector GetRightVectorAtSplineInputKey(float InKey, ESplineCoordinateSpace::Type CoordinateSpace) const;
 	UE_API FTransform GetTransformAtSplineInputKey(float InKey, ESplineCoordinateSpace::Type CoordinateSpace, bool bUseScale = false) const;
@@ -79,6 +100,7 @@ struct FPCGSplineStruct
 	UE_API FVector::FReal GetDistanceAlongSplineAtSplinePoint(int32 PointIndex) const;
 	UE_API FVector GetLocationAtDistanceAlongSpline(FVector::FReal Distance, ESplineCoordinateSpace::Type CoordinateSpace) const;
 	UE_API FTransform GetTransformAtDistanceAlongSpline(FVector::FReal Distance, ESplineCoordinateSpace::Type CoordinateSpace, bool bUseScale = false) const;
+	UE_API float GetInputKeyAtDistanceAlongSpline(FVector::FReal Distance) const;
 
 	/** Given a threshold, returns a list of vertices along the spline segment that, treated as a list of segments (polyline), matches the spline shape. Taken from USplineComponent. */
 	UE_API bool ConvertSplineSegmentToPolyLine(int32 SplinePointStartIndex, ESplineCoordinateSpace::Type CoordinateSpace, const float MaxSquareDistanceFromSpline, TArray<FVector>& OutPoints) const;
@@ -101,6 +123,8 @@ struct FPCGSplineStruct
 	UE_API void AllocateMetadataEntries();
 	TConstArrayView<PCGMetadataEntryKey> GetConstControlPointsEntryKeys() const { return ControlPointsEntryKeys; }
 	TArrayView<PCGMetadataEntryKey> GetMutableControlPointsEntryKeys() { return ControlPointsEntryKeys; }
+
+	UE_API void PostSerialize(const FArchive& Ar);
 	
 	// Replaces the component transform
 	UPROPERTY()
@@ -122,6 +146,7 @@ struct FPCGSplineStruct
 	FBoxSphereBounds Bounds = FBoxSphereBounds(EForceInit::ForceInit);
 
 protected:
+	// todo: deprecate, serialize into FSpline
 	UPROPERTY()
 	FSplineCurves SplineCurves;
 	
@@ -129,12 +154,38 @@ protected:
 	TArray<int64> ControlPointsEntryKeys; // Needs to be int64 for UHT, but it is a PCGMetadataEntryKey
 
 private:
+	UPROPERTY()
+	FSpline Spline;
+
+	UE_API bool ShouldUseSplineCurves() const;
+	
+	// Remove when appending _DEPRECATED to SplineCurves.
+	FSplineCurves& WarninglessSplineCurves()
+	{
+PRAGMA_DISABLE_DEPRECATION_WARNINGS
+		return SplineCurves;
+PRAGMA_ENABLE_DEPRECATION_WARNINGS
+	}
+
+	// Remove when appending _DEPRECATED to SplineCurves.
+	const FSplineCurves& WarninglessSplineCurves() const
+	{
+PRAGMA_DISABLE_DEPRECATION_WARNINGS
+		return SplineCurves;
+PRAGMA_ENABLE_DEPRECATION_WARNINGS
+	}
+	
 	// Internal helper function called by ConvertSplineSegmentToPolyLine -- assumes the input is within a half-segment, so testing the distance to midpoint will be an accurate guide to subdivision. Taken from USplineComponent.
 	UE_API bool DivideSplineIntoPolylineRecursiveWithDistancesHelper(float StartDistanceAlongSpline, float EndDistanceAlongSpline, ESplineCoordinateSpace::Type CoordinateSpace, const float MaxSquareDistanceFromSpline, TArray<FVector>& OutPoints, TArray<double>& OutDistancesAlongSpline) const;
+};
 
-	// Spline accessors need to access the Spline Curves, so make it friend.
-	template <typename T, EPCGControlPointsAccessorTarget Target, bool bWorldCoordinates>
-	friend class FPCGControlPointsAccessor;
+template<>
+struct TStructOpsTypeTraits<FPCGSplineStruct> : public TStructOpsTypeTraitsBase2<FPCGSplineStruct>
+{
+	enum
+	{
+		WithPostSerialize = true,
+	};
 };
 
 #undef UE_API
diff --git a/Engine/Plugins/PCG/Source/PCG/Public/Metadata/Accessors/PCGSplineAccessor.h b/Engine/Plugins/PCG/Source/PCG/Public/Metadata/Accessors/PCGSplineAccessor.h
--- a/Engine/Plugins/PCG/Source/PCG/Public/Metadata/Accessors/PCGSplineAccessor.h
+++ b/Engine/Plugins/PCG/Source/PCG/Public/Metadata/Accessors/PCGSplineAccessor.h
@@ -7,150 +7,12 @@
 
 #include "Components/SplineComponent.h"
 
-enum class EPCGInterpCurveAccessorTarget
-{
-	Value,
-	ArriveTangent,
-	LeaveTangent,
-	InterpMode
-};
-
 enum class EPCGSplineAccessorTarget
 {
 	Transform,
 	ClosedLoop
 };
 
-/**
-* Templated accessor for any interp curve. It's important that the keys only have a single value, the struct that holds the spline curve, since
-* the interp curve is basically an array.
-* Keys supported: PCGSplineData, FPCGSplineStruct, FSplineCurves
-*/
-template <typename CurveType,
-	EPCGInterpCurveAccessorTarget Target = EPCGInterpCurveAccessorTarget::Value,
-	std::enable_if_t<std::is_base_of_v<FInterpCurve<typename CurveType::ElementType>, CurveType>, int> = 0>
-class FPCGInterpCurveAccessor : public IPCGAttributeAccessorT<FPCGInterpCurveAccessor<CurveType, Target>>, IPCGPropertyChain
-{
-public:
-	// The underlying type is either the Curve element type or the same as the enum property accessor if the target is the InterpMode, as it is an enum below.
-	using Type = std::conditional_t<Target == EPCGInterpCurveAccessorTarget::InterpMode, FPCGEnumPropertyAccessor::Type, typename CurveType::ElementType>;
-	using Super = IPCGAttributeAccessorT<FPCGInterpCurveAccessor<CurveType, Target>>;
-
-	FPCGInterpCurveAccessor(const FStructProperty* InProperty, TArray<const FProperty*>&& ExtraProperties = {})
-		: Super(/*bInReadOnly=*/ false)
-		, IPCGPropertyChain(InProperty, std::move(ExtraProperties))
-	{
-		static_assert(PCG::Private::IsPCGType<Type>());		
-		check(InProperty && InProperty->Struct->IsChildOf(TBaseStructure<CurveType>::Get()));
-		TopPropertyStruct = GetTopPropertyStruct();
-	}
-
-	bool GetRangeImpl(TArrayView<Type> OutValues, int32 Index, const IPCGAttributeAccessorKeys& Keys) const
-	{
-		// Spline struct contains an array, so there is only a single address, to the array.
-		const void* ContainerKeys = nullptr;
-		TArrayView<const void*> ContainerKeysView(&ContainerKeys, 1);
-		if (!Keys.GetKeys(Index, ContainerKeysView))
-		{
-			return false;
-		}
-
-		// Validation to not access keys that are not the expected type. Done after the GetKeys, as we also want to discard other type of incompatible
-		// keys (like a Default Metadata entry key)
-		if (!ensure(Keys.IsClassSupported(TopPropertyStruct)))
-		{
-			return false;
-		}
-
-		PCGPropertyAccessor::AddressOffset(GetPropertyChain(), ContainerKeysView);
-
-		const CurveType& InterpCurve = *static_cast<const CurveType*>(ContainerKeys);
-		const int32 NumPoints = InterpCurve.Points.Num();
-		
-		if (NumPoints == 0)
-		{
-			return false;
-		}
-		
-		for (int32 i = 0; i < OutValues.Num(); ++i)
-		{
-			const int32 CurrIndex = (Index + i) % NumPoints;
-			if constexpr (Target == EPCGInterpCurveAccessorTarget::InterpMode)
-			{
-				OutValues[i] = static_cast<Type>(ConvertInterpCurveModeToSplinePointType(InterpCurve.Points[CurrIndex].InterpMode));
-			}
-			else if constexpr (Target == EPCGInterpCurveAccessorTarget::ArriveTangent)
-			{
-				OutValues[i] = InterpCurve.Points[CurrIndex].ArriveTangent;
-			}
-			else if constexpr (Target == EPCGInterpCurveAccessorTarget::LeaveTangent)
-			{
-				OutValues[i] = InterpCurve.Points[CurrIndex].LeaveTangent;
-			}
-			else // if constexpr (Target == EPCGInterpCurveAccessorTarget::Value)
-			{
-				OutValues[i] = InterpCurve.Points[CurrIndex].OutVal;
-			}
-		}
-
-		return true;
-	}
-
-	bool SetRangeImpl(TArrayView<const Type> InValues, int32 Index, IPCGAttributeAccessorKeys& Keys, EPCGAttributeAccessorFlags)
-	{		
-		// Spline struct contains an array, so there is only a single address, to the array.
-		void* ContainerKeys = nullptr;
-		TArrayView<void*> ContainerKeysView(&ContainerKeys, 1);
-		if (!Keys.GetKeys(Index, ContainerKeysView))
-		{
-			return false;
-		}
-
-		// Validation to not access keys that are not the expected type. Done after the GetKeys, as we also want to discard other type of incompatible
-		// keys (like a Default Metadata entry key)
-		if (!ensure(Keys.IsClassSupported(TopPropertyStruct)))
-		{
-			return false;
-		}
-
-		PCGPropertyAccessor::AddressOffset(GetPropertyChain(), ContainerKeysView);
-
-		CurveType& InterpCurve = *static_cast<CurveType*>(ContainerKeys);
-		const int32 NumPoints = InterpCurve.Points.Num();
-		if (NumPoints == 0)
-		{
-			return false;
-		}
-		
-		for (int32 i = 0; i < InValues.Num(); ++i)
-		{
-			const int32 CurrIndex = (Index + i) % NumPoints;
-
-			if constexpr (Target == EPCGInterpCurveAccessorTarget::InterpMode)
-			{
-				InterpCurve.Points[CurrIndex].InterpMode = ConvertSplinePointTypeToInterpCurveMode(static_cast<ESplinePointType::Type>(InValues[i]));
-			}
-			else if constexpr (Target == EPCGInterpCurveAccessorTarget::ArriveTangent)
-			{
-				InterpCurve.Points[CurrIndex].ArriveTangent = InValues[i];
-			}
-			else if constexpr (Target == EPCGInterpCurveAccessorTarget::LeaveTangent)
-			{
-				InterpCurve.Points[CurrIndex].LeaveTangent = InValues[i];
-			}
-			else // if constexpr (Target == EPCGInterpCurveAccessorTarget::Value)
-			{
-				InterpCurve.Points[CurrIndex].OutVal = InValues[i];
-			}
-		}
-
-		return true;
-	}
-
-private:
-	const UStruct* TopPropertyStruct = nullptr;
-};
-
 /**
 * Templated accessor for location/rotation/scale in world coordinates. It's important that the keys only have a single value, the struct that holds the spline data/struct, since
 * there is a single transform per spline.
@@ -193,12 +55,9 @@ public:
 		PCGPropertyAccessor::AddressOffset(GetPropertyChain(), ContainerKeysView);
 
 		const FPCGSplineStruct& SplineStruct = *static_cast<const FPCGSplineStruct*>(ContainerKeys);
-		const FInterpCurveVector& Positions = SplineStruct.GetSplinePointsPosition();
-		const FInterpCurveQuat& Rotations = SplineStruct.GetSplinePointsRotation();
-		const FInterpCurveVector& Scale = SplineStruct.GetSplinePointsScale();
 		const FTransform SplineTransform = SplineStruct.GetTransform();
 
-		const int32 NumPoints = Positions.Points.Num();
+		const int32 NumPoints = SplineStruct.GetNumberOfPoints();
 		if (NumPoints == 0)
 		{
 			return false;
@@ -211,7 +70,7 @@ public:
 			if constexpr (Target == EPCGControlPointsAccessorTarget::Location)
 			{
 				static_assert(std::is_same_v<Type, FVector>);
-				OutValues[i] = Positions.Points[CurrIndex].OutVal;
+				OutValues[i] = SplineStruct.GetLocation(CurrIndex);
 				if constexpr (bWorldCoordinates)
 				{
 					OutValues[i] = SplineTransform.TransformPosition(OutValues[i]);
@@ -220,7 +79,7 @@ public:
 			else if constexpr (Target == EPCGControlPointsAccessorTarget::Rotation)
 			{
 				static_assert(std::is_same_v<Type, FQuat>);
-				OutValues[i] = Rotations.Points[CurrIndex].OutVal;
+				OutValues[i] = SplineStruct.GetRotation(CurrIndex);
 				if constexpr (bWorldCoordinates)
 				{
 					OutValues[i] = SplineTransform.TransformRotation(OutValues[i]);
@@ -230,7 +89,7 @@ public:
 			{
 				
 				static_assert(std::is_same_v<Type, FVector>);
-				OutValues[i] = Scale.Points[CurrIndex].OutVal;
+				OutValues[i] = SplineStruct.GetScale(CurrIndex);
 				if constexpr (bWorldCoordinates)
 				{
 					OutValues[i] = (FTransform(FQuat::Identity, FVector::ZeroVector, OutValues[i]) * SplineTransform).GetScale3D();
@@ -239,12 +98,35 @@ public:
 			else if constexpr (Target == EPCGControlPointsAccessorTarget::Transform)
 			{
 				static_assert(std::is_same_v<Type, FTransform>);
-				OutValues[i] = FTransform(Rotations.Points[CurrIndex].OutVal, Positions.Points[CurrIndex].OutVal, Scale.Points[CurrIndex].OutVal);
+				OutValues[i] = FTransform(SplineStruct.GetRotation(CurrIndex), SplineStruct.GetLocation(CurrIndex), SplineStruct.GetScale(CurrIndex));
 				if constexpr (bWorldCoordinates)
 				{
 					OutValues[i] = OutValues[i] * SplineTransform;
 				}
 			}
+			else if constexpr (Target == EPCGControlPointsAccessorTarget::ArriveTangent)
+			{
+				static_assert(std::is_same_v<Type, FVector>);
+				OutValues[i] = SplineStruct.GetArriveTangent(CurrIndex);
+				if constexpr (bWorldCoordinates)
+				{
+					OutValues[i] = SplineTransform.TransformVector(OutValues[i]);
+				}
+			}
+			else if constexpr (Target == EPCGControlPointsAccessorTarget::LeaveTangent)
+			{
+				static_assert(std::is_same_v<Type, FVector>);
+				OutValues[i] = SplineStruct.GetLeaveTangent(CurrIndex);
+				if constexpr (bWorldCoordinates)
+				{
+					OutValues[i] = SplineTransform.TransformVector(OutValues[i]);
+				}
+			}
+			else if constexpr (Target == EPCGControlPointsAccessorTarget::InterpMode)
+			{
+				static_assert(std::is_same_v<Type, FPCGEnumPropertyAccessor::Type>);
+				OutValues[i] = static_cast<Type>(ConvertInterpCurveModeToSplinePointType(SplineStruct.GetSplinePointType(CurrIndex)));
+			}
 			else
 			{
 				// Pitfall static assert
@@ -277,7 +159,7 @@ public:
 		FPCGSplineStruct& SplineStruct = *static_cast<FPCGSplineStruct*>(ContainerKeys);
 		const FTransform InverseSplineTransform = SplineStruct.GetTransform().Inverse();
 		
-		const int32 NumPoints = SplineStruct.SplineCurves.Position.Points.Num();
+		const int32 NumPoints = SplineStruct.GetNumberOfPoints();
 		if (NumPoints == 0)
 		{
 			return false;
@@ -290,33 +172,32 @@ public:
 			if constexpr (Target == EPCGControlPointsAccessorTarget::Location)
 			{
 				static_assert(std::is_same_v<Type, FVector>);
-				FVector& OutPosition = SplineStruct.SplineCurves.Position.Points[CurrIndex].OutVal;
-				OutPosition = InValues[i];
+				FVector OutPosition = InValues[i];
 				if constexpr (bWorldCoordinates)
 				{
 					OutPosition = InverseSplineTransform.TransformPosition(OutPosition);
 				}
+				SplineStruct.SetLocation(CurrIndex, OutPosition);
 			}
 			else if constexpr (Target == EPCGControlPointsAccessorTarget::Rotation)
 			{
 				static_assert(std::is_same_v<Type, FQuat>);
-				FQuat& OutRotation = SplineStruct.SplineCurves.Rotation.Points[CurrIndex].OutVal;
-				OutRotation = InValues[i];
+				FQuat OutRotation = InValues[i];
 				if constexpr (bWorldCoordinates)
 				{
 					OutRotation = InverseSplineTransform.TransformRotation(OutRotation);
 				}
+				SplineStruct.SetRotation(CurrIndex, OutRotation);
 			}
 			else if constexpr (Target == EPCGControlPointsAccessorTarget::Scale)
 			{
 				static_assert(std::is_same_v<Type, FVector>);
-				FVector& OutScale = SplineStruct.SplineCurves.Scale.Points[CurrIndex].OutVal;
-				OutScale = InValues[i];
-				
+				FVector OutScale = InValues[i];
 				if constexpr (bWorldCoordinates)
 				{
 					OutScale = (FTransform(FQuat::Identity, FVector::ZeroVector, OutScale) * InverseSplineTransform).GetScale3D();
 				}
+				SplineStruct.SetScale(CurrIndex, OutScale);
 			}
 			else if constexpr (Target == EPCGControlPointsAccessorTarget::Transform)
 			{
@@ -326,10 +207,34 @@ public:
 				{
 					Transform = Transform * InverseSplineTransform;
 				}
-
-				SplineStruct.SplineCurves.Position.Points[CurrIndex].OutVal = InValues[i].GetLocation();
-				SplineStruct.SplineCurves.Rotation.Points[CurrIndex].OutVal = InValues[i].GetRotation();
-				SplineStruct.SplineCurves.Scale.Points[CurrIndex].OutVal = InValues[i].GetScale3D();
+				SplineStruct.SetLocation(CurrIndex, Transform.GetLocation());
+				SplineStruct.SetRotation(CurrIndex, Transform.GetRotation());
+				SplineStruct.SetScale(CurrIndex, Transform.GetScale3D());
+			}
+			else if constexpr (Target == EPCGControlPointsAccessorTarget::ArriveTangent)
+			{
+				static_assert(std::is_same_v<Type, FVector>);
+				FVector OutArriveTangent = InValues[i];
+				if constexpr (bWorldCoordinates)
+				{
+					OutArriveTangent = InverseSplineTransform.TransformVector(OutArriveTangent);
+				}
+				SplineStruct.SetArriveTangent(CurrIndex, OutArriveTangent);
+			}
+			else if constexpr (Target == EPCGControlPointsAccessorTarget::LeaveTangent)
+			{
+				static_assert(std::is_same_v<Type, FVector>);
+				FVector OutLeaveTangent = InValues[i];
+				if constexpr (bWorldCoordinates)
+				{
+					OutLeaveTangent = InverseSplineTransform.TransformVector(OutLeaveTangent);
+				}
+				SplineStruct.SetLeaveTangent(CurrIndex, OutLeaveTangent);
+			}
+			else if constexpr (Target == EPCGControlPointsAccessorTarget::InterpMode)
+			{
+				static_assert(std::is_same_v<Type, FPCGEnumPropertyAccessor::Type>);
+				SplineStruct.SetSplinePointType(CurrIndex, ConvertSplinePointTypeToInterpCurveMode(static_cast<ESplinePointType::Type>(InValues[i])));
 			}
 			else
 			{
diff --git a/Engine/Plugins/PCG/Source/PCGEditor/Private/DataVisualizations/PCGSplineDataVisualization.cpp b/Engine/Plugins/PCG/Source/PCGEditor/Private/DataVisualizations/PCGSplineDataVisualization.cpp
--- a/Engine/Plugins/PCG/Source/PCGEditor/Private/DataVisualizations/PCGSplineDataVisualization.cpp
+++ b/Engine/Plugins/PCG/Source/PCGEditor/Private/DataVisualizations/PCGSplineDataVisualization.cpp
@@ -31,9 +31,7 @@ FPCGTableVisualizerInfo IPCGSplineDataVisualization::GetTableVisualizerInfoWithD
 	if (DomainID == PCGMetadataDomainID::Data)
 	{
 		// Column Sorting
-		FPCGAttributePropertySelector IndexSelector = FPCGAttributePropertySelector::CreateExtraPropertySelector(EPCGExtraProperties::Index);
-		IndexSelector.SetDomainName(PCGDataConstants::DataDomainName);
-		
+		FPCGAttributePropertySelector IndexSelector = FPCGAttributePropertySelector::CreateExtraPropertySelector(EPCGExtraProperties::Index, PCGDataConstants::DataDomainName);
 		AddColumnInfo(Info, SplineData, IndexSelector);
 		Info.SortingColumn = Info.ColumnInfos.Last().Id;
 
@@ -47,27 +45,28 @@ FPCGTableVisualizerInfo IPCGSplineDataVisualization::GetTableVisualizerInfoWithD
 		return Info;
 	}
 	
-	// Column Sorting
-	AddColumnInfo(Info, SplineData, FPCGAttributePropertySelector::CreateExtraPropertySelector(EPCGExtraProperties::Index));
-	Info.SortingColumn = Info.ColumnInfos.Last().Id;
-	
-	// Columns
-	AddPropertyEnumColumnInfo<FVector>(Info, SplineData, EPCGSplineStructProperties::Position);
-	AddPropertyEnumColumnInfo<FRotator>(Info, SplineData, EPCGSplineStructProperties::Rotation);
-	AddPropertyEnumColumnInfo<FVector>(Info, SplineData, EPCGSplineStructProperties::Scale);
-	AddPropertyEnumColumnInfo<FVector>(Info, SplineData, EPCGSplineStructProperties::LocalPosition);
-	AddPropertyEnumColumnInfo<FRotator>(Info, SplineData, EPCGSplineStructProperties::LocalRotation);
-	AddPropertyEnumColumnInfo<FVector>(Info, SplineData, EPCGSplineStructProperties::LocalScale);
-	AddPropertyEnumColumnInfo<FVector>(Info, SplineData, EPCGSplineStructProperties::ArriveTangent);
-	AddPropertyEnumColumnInfo<FVector>(Info, SplineData, EPCGSplineStructProperties::LeaveTangent);
-	AddPropertyEnumColumnInfo<int32>(Info, SplineData, EPCGSplineStructProperties::InterpType);
-	
-	// Add Metadata Columns
-	CreateMetadataColumnInfos(SplineData, Info, PCGMetadataDomainID::Elements);
-	
-	// Focus on data behavior
 	if (DomainID == PCGMetadataDomainID::Elements)
 	{
+		// Column Sorting
+		FPCGAttributePropertySelector IndexSelector = FPCGAttributePropertySelector::CreateExtraPropertySelector(EPCGExtraProperties::Index, PCGSplineData::ControlPointDomainName);
+		AddColumnInfo(Info, SplineData, IndexSelector);
+		Info.SortingColumn = Info.ColumnInfos.Last().Id;
+	
+		// Columns
+		AddPropertyEnumColumnInfo<FVector>(Info, SplineData, EPCGSplineStructProperties::Position);
+		AddPropertyEnumColumnInfo<FRotator>(Info, SplineData, EPCGSplineStructProperties::Rotation);
+		AddPropertyEnumColumnInfo<FVector>(Info, SplineData, EPCGSplineStructProperties::Scale);
+		AddPropertyEnumColumnInfo<FVector>(Info, SplineData, EPCGSplineStructProperties::LocalPosition);
+		AddPropertyEnumColumnInfo<FRotator>(Info, SplineData, EPCGSplineStructProperties::LocalRotation);
+		AddPropertyEnumColumnInfo<FVector>(Info, SplineData, EPCGSplineStructProperties::LocalScale);
+		AddPropertyEnumColumnInfo<FVector>(Info, SplineData, EPCGSplineStructProperties::ArriveTangent);
+		AddPropertyEnumColumnInfo<FVector>(Info, SplineData, EPCGSplineStructProperties::LeaveTangent);
+		AddPropertyEnumColumnInfo<int32>(Info, SplineData, EPCGSplineStructProperties::InterpType);
+	
+		// Add Metadata Columns
+		CreateMetadataColumnInfos(SplineData, Info, PCGMetadataDomainID::Elements);
+
+		// Focus on data behavior
 		Info.FocusOnDataCallback = [](const UPCGData* Data, TArrayView<const int> Indices)
 		{
 			if (const UPCGSplineData* SplineData = Cast<UPCGSplineData>(Data))
diff --git a/Engine/Source/Runtime/Engine/Classes/Components/SplineComponent.h b/Engine/Source/Runtime/Engine/Classes/Components/SplineComponent.h
--- a/Engine/Source/Runtime/Engine/Classes/Components/SplineComponent.h
+++ b/Engine/Source/Runtime/Engine/Classes/Components/SplineComponent.h
@@ -364,6 +364,7 @@ public:
 
 	// Safe for caching spline data in the legacy format. Will auto convert data if necessary.
 	ENGINE_API FSplineCurves GetSplineCurves() const;
+	ENGINE_API FSpline GetSpline() const;
 	ENGINE_API void SetSpline(const FSpline& InSpline);
 	ENGINE_API void SetSpline(const FSplineCurves& InSplineCurves);
 	
diff --git a/Engine/Source/Runtime/Engine/Classes/Curves/Spline.h b/Engine/Source/Runtime/Engine/Classes/Curves/Spline.h
--- a/Engine/Source/Runtime/Engine/Classes/Curves/Spline.h
+++ b/Engine/Source/Runtime/Engine/Classes/Curves/Spline.h
@@ -156,6 +156,8 @@ private:
 	// Not ideal, but allows us to automatically populate the proxy
 	// at serialize time when we might otherwise not be able to.
 	friend class USplineComponent;
+	friend struct FPCGSplineStruct;
+	
 #if WITH_EDITOR
 	uint8 PreviousImplementation;
 #endif
diff --git a/Engine/Source/Runtime/Engine/Private/Components/SplineComponent.cpp b/Engine/Source/Runtime/Engine/Private/Components/SplineComponent.cpp
--- a/Engine/Source/Runtime/Engine/Private/Components/SplineComponent.cpp
+++ b/Engine/Source/Runtime/Engine/Private/Components/SplineComponent.cpp
@@ -35,7 +35,7 @@ namespace UE::SplineComponent
 {
 
 bool GUseSplineCurves = true;
-FAutoConsoleVariableRef CVarEnableDaySequenceModifierDeviceCustomization(
+FAutoConsoleVariableRef CVarUseSplineCurves(
 	TEXT("SplineComponent.UseSplineCurves"),
 	GUseSplineCurves,
 	TEXT("When true, SplineCurves is the 'source of truth'.")
@@ -2566,6 +2566,20 @@ FSplineCurves USplineComponent::GetSplineCurves() const
 		: FSplineCurves::FromSplineInterface(Spline);
 }
 
+FSpline USplineComponent::GetSpline() const
+{
+	if (UE::SplineComponent::GUseSplineCurves)
+	{
+		FSpline Intermediate;
+		Intermediate = SplineCurves;
+		return Intermediate;
+	}
+	else
+	{
+		return Spline;
+	}
+}
+
 void USplineComponent::SetSpline(const FSpline& InSpline)
 {
 	WarninglessSplineCurves() = FSplineCurves::FromSplineInterface(InSpline);
