[PCG] Added option to filter empty data from the output of the Attribute Filter node

diff --git a/Engine/Plugins/PCG/Source/PCG/Private/Elements/PCGAttributeFilter.cpp b/Engine/Plugins/PCG/Source/PCG/Private/Elements/PCGAttributeFilter.cpp
--- a/Engine/Plugins/PCG/Source/PCG/Private/Elements/PCGAttributeFilter.cpp
+++ b/Engine/Plugins/PCG/Source/PCG/Private/Elements/PCGAttributeFilter.cpp
@@ -473,7 +473,7 @@ TArray<FPCGPinProperties> UPCGAttributeFilteringRangeSettings::OutputPinProperti
 // FPCGAttributeFilterElementBase
 ////////////////////////////////////////
 
-bool FPCGAttributeFilterElementBase::DoFiltering(FPCGContext* Context, EPCGAttributeFilterOperator InOperation, const FPCGAttributePropertyInputSelector& InTargetAttribute, bool bHasSpatialToPointDeprecation, bool bWarnOnDataMissingAttribute, const FPCGAttributeFilterThresholdSettings& FirstThreshold, const FPCGAttributeFilterThresholdSettings* SecondThreshold) const
+bool FPCGAttributeFilterElementBase::DoFiltering(FPCGContext* Context, EPCGAttributeFilterOperator InOperation, const FPCGAttributePropertyInputSelector& InTargetAttribute, bool bHasSpatialToPointDeprecation, bool bWarnOnDataMissingAttribute, bool bAlwaysGenerateOutputData, const FPCGAttributeFilterThresholdSettings& FirstThreshold, const FPCGAttributeFilterThresholdSettings* SecondThreshold) const
 {
 	TRACE_CPUPROFILER_EVENT_SCOPE(FPCGAttributeFilterElementBase::DoFiltering);
 	check(Context);
@@ -629,52 +629,24 @@ bool FPCGAttributeFilterElementBase::DoFiltering(FPCGContext* Context, EPCGAttri
 			continue;
 		}
 
-		UPCGData* InFilterData = nullptr;
-		UPCGData* OutFilterData = nullptr;
-
 		if (OperationData.bIsInputPointData)
 		{
 			const UPCGBasePointData* OriginalPointData = CastChecked<UPCGBasePointData>(OriginalData);
-			UPCGBasePointData* InFilterPointData = FPCGContext::NewPointData_AnyThread(Context);
-			UPCGBasePointData* OutFilterPointData = FPCGContext::NewPointData_AnyThread(Context);
 
 			OperationData.OriginalPointData = OriginalPointData;
 
-			FPCGInitializeFromDataParams InitializeFromDataParams(OriginalPointData);
-			InitializeFromDataParams.bInheritSpatialData = false;
-
-			InFilterPointData->InitializeFromDataWithParams(InitializeFromDataParams);
-			OutFilterPointData->InitializeFromDataWithParams(InitializeFromDataParams);
-
-			OperationData.InFilterPointData = InFilterPointData;
-			OperationData.OutFilterPointData = OutFilterPointData;
-
 			// Will be set individually in batches
 			OperationData.FilterBitArray.SetNumUninitialized(OriginalPointData->GetNumPoints());
-
-			InFilterData = InFilterPointData;
-			OutFilterData = OutFilterPointData;
 		}
 		else
 		{
 			// Param data
 			const UPCGParamData* OriginalParamData = CastChecked<UPCGParamData>(OriginalData);
-			UPCGParamData* InFilterParamData = FPCGContext::NewObject_AnyThread<UPCGParamData>(Context);
-			UPCGParamData* OutFilterParamData = FPCGContext::NewObject_AnyThread<UPCGParamData>(Context);
 
 			OperationData.OriginalMetadata = OriginalParamData->Metadata;
 
-			// Add all attributes from the original param data, but without any entry/value. It will be added later.
-			OperationData.InFilterMetadata = InFilterParamData->Metadata;
-			OperationData.OutFilterMetadata = OutFilterParamData->Metadata;
-			OperationData.InFilterMetadata->AddAttributesFiltered(OriginalParamData->Metadata, TSet<FName>(), EPCGMetadataFilterMode::ExcludeAttributes);
-			OperationData.OutFilterMetadata->AddAttributesFiltered(OriginalParamData->Metadata, TSet<FName>(), EPCGMetadataFilterMode::ExcludeAttributes);
-
 			// Will be set individually in batches
 			OperationData.FilterBitArray.SetNumUninitialized(OriginalParamData->Metadata ? OriginalParamData->Metadata->GetItemCountForChild() : 0);
-
-			InFilterData = InFilterParamData;
-			OutFilterData = OutFilterParamData;
 		}
 
 		auto Operation = [&Operator, &FirstThresholdInfo, &SecondThresholdInfo, &TargetAccessor, &TargetKeys, &FirstThreshold, &SecondThreshold, &OperationData](auto Dummy) -> bool
@@ -787,36 +759,59 @@ bool FPCGAttributeFilterElementBase::DoFiltering(FPCGContext* Context, EPCGAttri
 
 		if (PCGMetadataAttribute::CallbackWithRightType(TargetAccessor->GetUnderlyingType(), Operation))
 		{
+			UPCGData* InFilterData = nullptr;
+			UPCGData* OutFilterData = nullptr;
+
 			if (OperationData.bIsInputPointData)
 			{
 				const int32 NumInFilterPoints = OperationData.FilterBitArray.CountSetBits();
 				const int32 NumOutFilterPoints = OperationData.OriginalPointData->GetNumPoints() - NumInFilterPoints;
 
-				int32 InFilterWriteIndex = 0;
-				OperationData.InFilterPointData->SetNumPoints(NumInFilterPoints);
-				OperationData.InFilterPointData->AllocateProperties(OperationData.OriginalPointData->GetAllocatedProperties());
-				OperationData.InFilterPointData->CopyUnallocatedPropertiesFrom(OperationData.OriginalPointData);
+				FPCGInitializeFromDataParams InitializeFromDataParams(OperationData.OriginalPointData);
+				InitializeFromDataParams.bInheritSpatialData = false;
+				const FConstPCGPointValueRanges OriginalRanges(OperationData.OriginalPointData);
 
-				int32 OutFilterWriteIndex = 0;
-				OperationData.OutFilterPointData->SetNumPoints(NumOutFilterPoints);
-				OperationData.OutFilterPointData->AllocateProperties(OperationData.OriginalPointData->GetAllocatedProperties());
-				OperationData.OutFilterPointData->CopyUnallocatedPropertiesFrom(OperationData.OriginalPointData);
+				// Note: Splitting the set vs unset values is cleaner and likely more optimal, due to branch prediction.
 
-				const FConstPCGPointValueRanges OriginalRanges(OperationData.OriginalPointData);
-				FPCGPointValueRanges InFilterRanges(OperationData.InFilterPointData, /*bAllocate=*/false);
-				FPCGPointValueRanges OutFilterRanges(OperationData.OutFilterPointData, /*bAllocate=*/false);
-								
-				for (int32 Index = 0; Index < OperationData.FilterBitArray.Num(); ++Index)
+				if (NumInFilterPoints > 0 || bAlwaysGenerateOutputData)
 				{
-					if (OperationData.FilterBitArray[Index])
+					UPCGBasePointData* InFilterPointData = FPCGContext::NewPointData_AnyThread(Context);
+					InFilterPointData->InitializeFromDataWithParams(InitializeFromDataParams);
+					InFilterData = InFilterPointData;
+
+					InFilterPointData->SetNumPoints(NumInFilterPoints);
+					InFilterPointData->AllocateProperties(OperationData.OriginalPointData->GetAllocatedProperties());
+					InFilterPointData->CopyUnallocatedPropertiesFrom(OperationData.OriginalPointData);
+
+					FPCGPointValueRanges InFilterRanges(InFilterPointData, /*bAllocate=*/false);
+
+					int32 InFilterWriteIndex = 0;
+					for (TConstSetBitIterator It(OperationData.FilterBitArray); It; ++It)
 					{
-						InFilterRanges.SetFromValueRanges(InFilterWriteIndex, OriginalRanges, Index);						
-						++InFilterWriteIndex;
+						InFilterRanges.SetFromValueRanges(InFilterWriteIndex++, OriginalRanges, It.GetIndex());
 					}
-					else
+				}
+
+				if (NumOutFilterPoints > 0 || bAlwaysGenerateOutputData)
+				{
+					UPCGBasePointData* OutFilterPointData = FPCGContext::NewPointData_AnyThread(Context);
+					OutFilterPointData->InitializeFromDataWithParams(InitializeFromDataParams);
+					OutFilterData = OutFilterPointData;
+
+					OutFilterPointData->SetNumPoints(NumOutFilterPoints);
+					OutFilterPointData->AllocateProperties(OperationData.OriginalPointData->GetAllocatedProperties());
+					OutFilterPointData->CopyUnallocatedPropertiesFrom(OperationData.OriginalPointData);
+
+					FPCGPointValueRanges OutFilterRanges(OutFilterPointData, /*bAllocate=*/false);
+
+					// @todo_pcg: Could write a TConstUnsetBitIterator in BitArray.h.
+					int32 OutFilterWriteIndex = 0;
+					for (int32 Index = 0; Index < OperationData.FilterBitArray.Num(); ++Index)
 					{
-						OutFilterRanges.SetFromValueRanges(OutFilterWriteIndex, OriginalRanges, Index);
-						++OutFilterWriteIndex;
+						if (!OperationData.FilterBitArray[Index])
+						{
+							OutFilterRanges.SetFromValueRanges(OutFilterWriteIndex++, OriginalRanges, Index);
+						}
 					}
 				}
 			}
@@ -849,20 +844,42 @@ bool FPCGAttributeFilterElementBase::DoFiltering(FPCGContext* Context, EPCGAttri
 						}
 					}
 
-					OperationData.InFilterMetadata->SetAttributes(InEntryKeys, OperationData.OriginalMetadata, /*InOutOptionalKeys=*/nullptr, Context);
-					OperationData.OutFilterMetadata->SetAttributes(OutEntryKeys, OperationData.OriginalMetadata, /*InOutOptionalKeys=*/nullptr, Context);
+					// Add all attributes from the original param data, but without any entry/value. It will be added later.
+					if (NumInFilterMetadata > 0 || bAlwaysGenerateOutputData)
+					{
+						UPCGParamData* InFilterParamData = FPCGContext::NewObject_AnyThread<UPCGParamData>(Context);
+						UPCGMetadata* InFilterMetadata = InFilterParamData->Metadata;
+						InFilterMetadata->AddAttributesFiltered(OperationData.OriginalMetadata, TSet<FName>(), EPCGMetadataFilterMode::ExcludeAttributes);
+						InFilterMetadata->SetAttributes(InEntryKeys, OperationData.OriginalMetadata, /*InOutOptionalKeys=*/nullptr, Context);
+						InFilterData = InFilterParamData;
+					}
+
+					if (NumOutFilterMetadata > 0 || bAlwaysGenerateOutputData)
+					{
+						UPCGParamData* OutFilterParamData = FPCGContext::NewObject_AnyThread<UPCGParamData>(Context);
+						UPCGMetadata* OutFilterMetadata = OutFilterParamData->Metadata;
+						OutFilterMetadata->AddAttributesFiltered(OperationData.OriginalMetadata, TSet<FName>(), EPCGMetadataFilterMode::ExcludeAttributes);
+						OutFilterMetadata->SetAttributes(OutEntryKeys, OperationData.OriginalMetadata, /*InOutOptionalKeys=*/nullptr, Context);
+						OutFilterData = OutFilterParamData;
+					}
 				}
 			}
 
-			FPCGTaggedData& InFilterOutput = Outputs.Add_GetRef(Input);
-			InFilterOutput.Pin = PCGPinConstants::DefaultInFilterLabel;
-			InFilterOutput.Data = InFilterData;
-			InFilterOutput.Tags = Input.Tags;
+			if (InFilterData)
+			{
+				FPCGTaggedData& InFilterOutput = Outputs.Add_GetRef(Input);
+				InFilterOutput.Pin = PCGPinConstants::DefaultInFilterLabel;
+				InFilterOutput.Data = InFilterData;
+				InFilterOutput.Tags = Input.Tags;
+			}
 
-			FPCGTaggedData& OutFilterOutput = Outputs.Add_GetRef(Input);
-			OutFilterOutput.Pin = PCGPinConstants::DefaultOutFilterLabel;
-			OutFilterOutput.Data = OutFilterData;
-			OutFilterOutput.Tags = Input.Tags;
+			if (OutFilterData)
+			{
+				FPCGTaggedData& OutFilterOutput = Outputs.Add_GetRef(Input);
+				OutFilterOutput.Pin = PCGPinConstants::DefaultOutFilterLabel;
+				OutFilterOutput.Data = OutFilterData;
+				OutFilterOutput.Tags = Input.Tags;
+			}
 		}
 		else
 		{
@@ -903,7 +920,7 @@ bool FPCGAttributeFilterElement::ExecuteInternal(FPCGContext* Context) const
 	ThresholdSettings.ThresholdAttribute = Settings->ThresholdAttribute;
 	ThresholdSettings.AttributeTypes = Settings->AttributeTypes;
 
-	return DoFiltering(Context, Settings->Operator, Settings->TargetAttribute, Settings->bHasSpatialToPointDeprecation, Settings->bWarnOnDataMissingAttribute, ThresholdSettings);
+	return DoFiltering(Context, Settings->Operator, Settings->TargetAttribute, Settings->bHasSpatialToPointDeprecation, Settings->bWarnOnDataMissingAttribute, Settings->bGenerateOutputDataEvenIfEmpty, ThresholdSettings);
 }
 
 ////////////////////////////////////////
@@ -928,7 +945,7 @@ bool FPCGAttributeFilterRangeElement::ExecuteInternal(FPCGContext* Context) cons
 	const UPCGAttributeFilteringRangeSettings* Settings = Context->GetInputSettings<UPCGAttributeFilteringRangeSettings>();
 	check(Settings);
 
-	return DoFiltering(Context, EPCGAttributeFilterOperator::InRange, Settings->TargetAttribute, Settings->bHasSpatialToPointDeprecation, Settings->bWarnOnDataMissingAttribute, Settings->MinThreshold, &Settings->MaxThreshold);
+	return DoFiltering(Context, EPCGAttributeFilterOperator::InRange, Settings->TargetAttribute, Settings->bHasSpatialToPointDeprecation, Settings->bWarnOnDataMissingAttribute, Settings->bGenerateOutputDataEvenIfEmpty, Settings->MinThreshold, &Settings->MaxThreshold);
 }
 
 #undef LOCTEXT_NAMESPACE
diff --git a/Engine/Plugins/PCG/Source/PCG/Private/Tests/Elements/PCGAttributeFilterTest.cpp b/Engine/Plugins/PCG/Source/PCG/Private/Tests/Elements/PCGAttributeFilterTest.cpp
--- a/Engine/Plugins/PCG/Source/PCG/Private/Tests/Elements/PCGAttributeFilterTest.cpp
+++ b/Engine/Plugins/PCG/Source/PCG/Private/Tests/Elements/PCGAttributeFilterTest.cpp
@@ -6,6 +6,7 @@
 #include "PCGContext.h"
 #include "PCGParamData.h"
 #include "PCGPoint.h"
+#include "Algo/Count.h"
 #include "Data/PCGPointData.h"
 #include "Elements/PCGAttributeFilter.h"
 #include "Metadata/PCGMetadata.h"
@@ -19,6 +20,14 @@ IMPLEMENT_CUSTOM_SIMPLE_AUTOMATION_TEST(FPCGPointFilterDensityRange, FPCGTestBas
 IMPLEMENT_CUSTOM_SIMPLE_AUTOMATION_TEST(FPCGAttributeFilterInt, FPCGTestBaseClass, "Plugins.PCG.AttributeFilter.Params.Int", PCGTestsCommon::TestFlags)
 IMPLEMENT_CUSTOM_SIMPLE_AUTOMATION_TEST(FPCGAttributeFilterIntRange, FPCGTestBaseClass, "Plugins.PCG.AttributeFilter.Params.IntRange", PCGTestsCommon::TestFlags)
 IMPLEMENT_CUSTOM_SIMPLE_AUTOMATION_TEST(FPCGAttributeFilterSkipTestBug, FPCGTestBaseClass, "Plugins.PCG.AttributeFilter.SkipTestBug", PCGTestsCommon::TestFlags)
+IMPLEMENT_CUSTOM_SIMPLE_AUTOMATION_TEST(FPCGAttributeFilter_Points_GenerateEmptyOutput, FPCGTestBaseClass, "Plugins.PCG.AttributeFilter.Points.GenerateEmptyOutput", PCGTestsCommon::TestFlags)
+IMPLEMENT_CUSTOM_SIMPLE_AUTOMATION_TEST(FPCGAttributeFilter_PointsRange_GenerateEmptyOutput, FPCGTestBaseClass, "Plugins.PCG.AttributeFilter.PointsRange.GenerateEmptyOutput", PCGTestsCommon::TestFlags)
+IMPLEMENT_CUSTOM_SIMPLE_AUTOMATION_TEST(FPCGAttributeFilter_Attributes_GenerateEmptyOutput, FPCGTestBaseClass, "Plugins.PCG.AttributeFilter.Params.GenerateEmptyOutput", PCGTestsCommon::TestFlags)
+IMPLEMENT_CUSTOM_SIMPLE_AUTOMATION_TEST(FPCGAttributeFilter_AttributesRange_GenerateEmptyOutput, FPCGTestBaseClass, "Plugins.PCG.AttributeFilter.ParamsRange.GenerateEmptyOutput", PCGTestsCommon::TestFlags)
+IMPLEMENT_CUSTOM_SIMPLE_AUTOMATION_TEST(FPCGAttributeFilter_Points_NoGenerateEmptyOutput, FPCGTestBaseClass, "Plugins.PCG.AttributeFilter.Points.NoGenerateEmptyOutput", PCGTestsCommon::TestFlags)
+IMPLEMENT_CUSTOM_SIMPLE_AUTOMATION_TEST(FPCGAttributeFilter_PointsRange_NoGenerateEmptyOutput, FPCGTestBaseClass, "Plugins.PCG.AttributeFilter.PointsRange.NoGenerateEmptyOutput", PCGTestsCommon::TestFlags)
+IMPLEMENT_CUSTOM_SIMPLE_AUTOMATION_TEST(FPCGAttributeFilter_Attributes_NoGenerateEmptyOutput, FPCGTestBaseClass, "Plugins.PCG.AttributeFilter.Params.NoGenerateEmptyOutput", PCGTestsCommon::TestFlags)
+IMPLEMENT_CUSTOM_SIMPLE_AUTOMATION_TEST(FPCGAttributeFilter_AttributesRange_NoGenerateEmptyOutput, FPCGTestBaseClass, "Plugins.PCG.AttributeFilter.ParamsRange.NoGenerateEmptyOutput", PCGTestsCommon::TestFlags)
 
 namespace PCGPointFilterTest
 {
@@ -46,6 +55,19 @@ namespace PCGPointFilterTest
 		return PointData;
 	}
 
+	UPCGBasePointData* GeneratePointDataWithSingleDensity(const float Density)
+	{
+		UPCGBasePointData* PointData = PCGTestsCommon::CreateEmptyBasePointData();
+
+		PointData->SetNumPoints(1);
+		PointData->AllocateProperties(EPCGPointNativeProperties::Density);
+
+		FPCGPointValueRanges ValueRanges(PointData, /*bAllocate=*/false);
+		ValueRanges.DensityRange[0] = Density;
+
+		return PointData;
+	}
+
 	static const FName IntAttributeName = TEXT("Int");
 
 	UPCGParamData* GenerateAttributeSetWithRandomInt(int32 InNumEntries, int32 InRandomSeed)
@@ -65,6 +87,39 @@ namespace PCGPointFilterTest
 
 		return ParamData;
 	}
+
+	static const FName DoubleAttributeName = TEXT("Double");
+
+	UPCGParamData* GenerateAttributeSetWithSingleValue(const double Value)
+	{
+		UPCGParamData* ParamData = NewObject<UPCGParamData>();
+		check(ParamData && ParamData->Metadata);
+		FPCGMetadataAttribute<double>* Attribute = ParamData->Metadata->CreateAttribute<double>(DoubleAttributeName, 0, true, false);
+		check(Attribute)
+
+		const PCGMetadataEntryKey Key = ParamData->Metadata->AddEntry();
+		Attribute->SetValue(Key, Value);
+
+		return ParamData;
+	}
+
+	int GetNumEmptyData(const TArray<FPCGTaggedData>& InDataArray)
+	{
+		return Algo::CountIf(InDataArray, [](const FPCGTaggedData& Data)
+		{
+			if (const UPCGBasePointData* PointData = Cast<UPCGBasePointData>(Data.Data))
+			{
+				return PointData->IsEmpty();
+			}
+			else if (const UPCGParamData* ParamData = Cast<UPCGParamData>(Data.Data))
+			{
+				check(ParamData->Metadata)
+				return ParamData->Metadata->GetAttributeCount() == 0 || ParamData->Metadata->GetItemCountForChild() == 0;
+			}
+
+			return false;
+		});
+	}
 }
 
 bool FPCGPointFilterDensity::RunTest(const FString& Parameters)
@@ -386,5 +441,547 @@ bool FPCGAttributeFilterSkipTestBug::RunTest(const FString& Parameters)
 	UTEST_EQUAL(TEXT("InFilter data has the right number of points"), InFilterPointData->GetNumPoints(), HalfNumPoints);
 	UTEST_EQUAL(TEXT("OutFilter data has the right number of points"), OutFilterPointData->GetNumPoints(), HalfNumPoints);
 
+	return true;
+}
+
+bool FPCGAttributeFilter_Points_GenerateEmptyOutput::RunTest(const FString& Parameters)
+{
+	PCGTestsCommon::FTestData TestData;
+	UPCGAttributeFilteringSettings* Settings = PCGTestsCommon::GenerateSettings<UPCGAttributeFilteringSettings>(TestData);
+	check(Settings);
+
+	const FPCGElementPtr TestElement = TestData.Settings->GetElement();
+
+	Settings->Operator = EPCGAttributeFilterOperator::Greater;
+	Settings->TargetAttribute.SetPointProperty(EPCGPointProperties::Density);
+	Settings->bUseConstantThreshold = true;
+	Settings->AttributeTypes.Type = EPCGMetadataTypes::Float;
+
+	// Generate 4 point data, each with 1 point that has a normalized density based on its data index: .25, .5, .75, 1.0
+	for (int32 I = 0; I < 4; ++I)
+	{
+		FPCGTaggedData& TaggedData = TestData.InputData.TaggedData.Emplace_GetRef(FPCGTaggedData());
+		TaggedData.Pin = PCGPinConstants::DefaultInputLabel;
+		TaggedData.Data = PCGPointFilterTest::GeneratePointDataWithSingleDensity((1.0 + I) / 4.0);
+	}
+
+	Settings->bGenerateOutputDataEvenIfEmpty = true;
+
+	// All are > 0.0
+	Settings->AttributeTypes.FloatValue = 0.f;
+	TUniquePtr<FPCGContext> Context = TestData.InitializeTestContext();
+	while (!TestElement->Execute(Context.Get())) {}
+
+	// When always generating output, all data should exist in both
+	TArray<FPCGTaggedData> InFilterOutput = Context->OutputData.GetInputsByPin(PCGPointFilterTest::InsideFilterLabel);
+	TArray<FPCGTaggedData> OutFilterOutput = Context->OutputData.GetInputsByPin(PCGPointFilterTest::OutsideFilterLabel);
+
+	UTEST_EQUAL(TEXT("InFilter pin has correct number of outputs"), InFilterOutput.Num(), 4);
+	UTEST_EQUAL(TEXT("OutFilter pin has correct number of outputs"), OutFilterOutput.Num(), 4);
+
+	UTEST_EQUAL(TEXT("InFilter data has correct number of empty outputs"), PCGPointFilterTest::GetNumEmptyData(InFilterOutput), 0);
+	UTEST_EQUAL(TEXT("OutFilter data has correct number of empty outputs"), PCGPointFilterTest::GetNumEmptyData(OutFilterOutput), 4);
+
+	// None are > 1.0
+	Settings->AttributeTypes.FloatValue = 1.f;
+	Context = TestData.InitializeTestContext();
+	while (!TestElement->Execute(Context.Get())) {}
+
+	InFilterOutput = Context->OutputData.GetInputsByPin(PCGPointFilterTest::InsideFilterLabel);
+	OutFilterOutput = Context->OutputData.GetInputsByPin(PCGPointFilterTest::OutsideFilterLabel);
+
+	UTEST_EQUAL(TEXT("InFilter pin has correct number of outputs"), InFilterOutput.Num(), 4);
+	UTEST_EQUAL(TEXT("OutFilter pin has correct number of outputs"), OutFilterOutput.Num(), 4);
+
+	UTEST_EQUAL(TEXT("InFilter data has correct number of empty outputs"), PCGPointFilterTest::GetNumEmptyData(InFilterOutput), 4);
+	UTEST_EQUAL(TEXT("OutFilter data has correct number of empty outputs"), PCGPointFilterTest::GetNumEmptyData(OutFilterOutput), 0);
+
+	return true;
+}
+
+bool FPCGAttributeFilter_PointsRange_GenerateEmptyOutput::RunTest(const FString& Parameters)
+{
+	PCGTestsCommon::FTestData TestData;
+	UPCGAttributeFilteringRangeSettings* Settings = PCGTestsCommon::GenerateSettings<UPCGAttributeFilteringRangeSettings>(TestData);
+	check(Settings);
+
+	const FPCGElementPtr TestElement = TestData.Settings->GetElement();
+
+	Settings->TargetAttribute.SetPointProperty(EPCGPointProperties::Density);
+	Settings->MinThreshold.bUseConstantThreshold = true;
+	Settings->MaxThreshold.bUseConstantThreshold = true;
+	Settings->MinThreshold.AttributeTypes.Type = EPCGMetadataTypes::Float;
+	Settings->MaxThreshold.AttributeTypes.Type = EPCGMetadataTypes::Float;
+	Settings->MinThreshold.bInclusive = true;
+	Settings->MaxThreshold.bInclusive = true;
+
+	// Generate 4 point data, each with 1 point that has a normalized density based on its data index: 0, .25, .5, 0.75
+	for (int32 I = 0; I < 4; ++I)
+	{
+		FPCGTaggedData& TaggedData = TestData.InputData.TaggedData.Emplace_GetRef(FPCGTaggedData());
+		TaggedData.Pin = PCGPinConstants::DefaultInputLabel;
+		TaggedData.Data = PCGPointFilterTest::GeneratePointDataWithSingleDensity((1.0 + I) / 4.0);
+	}
+
+	// Check unfiltered empty data.
+	Settings->bGenerateOutputDataEvenIfEmpty = true;
+
+	Settings->MinThreshold.AttributeTypes.FloatValue = 0.0f;
+	Settings->MaxThreshold.AttributeTypes.FloatValue = 1.0f;
+	TUniquePtr<FPCGContext> Context = TestData.InitializeTestContext();
+	while (!TestElement->Execute(Context.Get())) {}
+
+	// When always generating output, all data should exist in both
+	TArray<FPCGTaggedData> InFilterOutput = Context->OutputData.GetInputsByPin(PCGPointFilterTest::InsideFilterLabel);
+	TArray<FPCGTaggedData> OutFilterOutput = Context->OutputData.GetInputsByPin(PCGPointFilterTest::OutsideFilterLabel);
+
+	// All, inclusive.
+	UTEST_EQUAL(TEXT("InFilter pin has correct number of outputs"), InFilterOutput.Num(), 4);
+	UTEST_EQUAL(TEXT("OutFilter pin has correct number of outputs"), OutFilterOutput.Num(), 4);
+
+	UTEST_EQUAL(TEXT("InFilter data has correct number of empty outputs"), PCGPointFilterTest::GetNumEmptyData(InFilterOutput), 0);
+	UTEST_EQUAL(TEXT("OutFilter data has correct number of empty outputs"), PCGPointFilterTest::GetNumEmptyData(OutFilterOutput), 4);
+
+	Settings->MinThreshold.AttributeTypes.FloatValue = 0.0f;
+	Settings->MaxThreshold.AttributeTypes.FloatValue = 0.0f;
+	Context = TestData.InitializeTestContext();
+	while (!TestElement->Execute(Context.Get())) {}
+
+	// When always generating output, all data should exist in both
+	InFilterOutput = Context->OutputData.GetInputsByPin(PCGPointFilterTest::InsideFilterLabel);
+	OutFilterOutput = Context->OutputData.GetInputsByPin(PCGPointFilterTest::OutsideFilterLabel);
+
+	// None, inclusive.
+	UTEST_EQUAL(TEXT("InFilter pin has correct number of outputs"), InFilterOutput.Num(), 4);
+	UTEST_EQUAL(TEXT("OutFilter pin has correct number of outputs"), OutFilterOutput.Num(), 4);
+
+	UTEST_EQUAL(TEXT("InFilter data has correct number of empty outputs"), PCGPointFilterTest::GetNumEmptyData(InFilterOutput), 4);
+	UTEST_EQUAL(TEXT("OutFilter data has correct number of empty outputs"), PCGPointFilterTest::GetNumEmptyData(OutFilterOutput), 0);
+
+	return true;
+}
+
+bool FPCGAttributeFilter_Attributes_GenerateEmptyOutput::RunTest(const FString& Parameters)
+{
+	PCGTestsCommon::FTestData TestData;
+	UPCGAttributeFilteringSettings* Settings = PCGTestsCommon::GenerateSettings<UPCGAttributeFilteringSettings>(TestData);
+	check(Settings);
+
+	const FPCGElementPtr TestElement = TestData.Settings->GetElement();
+
+	Settings->Operator = EPCGAttributeFilterOperator::Greater;
+	Settings->TargetAttribute.SetAttributeName(PCGPointFilterTest::DoubleAttributeName);
+	Settings->AttributeTypes.Type = EPCGMetadataTypes::Double;
+	Settings->bUseConstantThreshold = true;
+
+	// Generate 4 point data, each with 1 point that has a normalized value based on its data index: .25, .5, .75, 1.0
+	for (int32 I = 0; I < 4; ++I)
+	{
+		FPCGTaggedData& TaggedData = TestData.InputData.TaggedData.Emplace_GetRef(FPCGTaggedData());
+		TaggedData.Pin = PCGPinConstants::DefaultInputLabel;
+		TaggedData.Data = PCGPointFilterTest::GenerateAttributeSetWithSingleValue((1.0 + I) / 4.0);
+	}
+
+	// Check unfiltered empty data.
+	Settings->bGenerateOutputDataEvenIfEmpty = true;
+
+	// All are > 0.0
+	Settings->AttributeTypes.DoubleValue = 0.f;
+	TUniquePtr<FPCGContext> Context = TestData.InitializeTestContext();
+	while (!TestElement->Execute(Context.Get())) {}
+
+	// When always generating output, all data should exist in both
+	TArray<FPCGTaggedData> InFilterOutput = Context->OutputData.GetInputsByPin(PCGPointFilterTest::InsideFilterLabel);
+	TArray<FPCGTaggedData> OutFilterOutput = Context->OutputData.GetInputsByPin(PCGPointFilterTest::OutsideFilterLabel);
+
+	UTEST_EQUAL(TEXT("InFilter pin has correct number of outputs"), InFilterOutput.Num(), 4);
+	UTEST_EQUAL(TEXT("OutFilter pin has correct number of outputs"), OutFilterOutput.Num(), 4);
+
+	UTEST_EQUAL(TEXT("InFilter data has correct number of empty outputs"), PCGPointFilterTest::GetNumEmptyData(InFilterOutput), 0);
+	UTEST_EQUAL(TEXT("OutFilter data has correct number of empty outputs"), PCGPointFilterTest::GetNumEmptyData(OutFilterOutput), 4);
+
+	// None are > 1.0
+	Settings->AttributeTypes.DoubleValue = 1.f;
+	Context = TestData.InitializeTestContext();
+	while (!TestElement->Execute(Context.Get())) {}
+
+	// When always generating output, all data should exist in both
+	InFilterOutput = Context->OutputData.GetInputsByPin(PCGPointFilterTest::InsideFilterLabel);
+	OutFilterOutput = Context->OutputData.GetInputsByPin(PCGPointFilterTest::OutsideFilterLabel);
+
+	UTEST_EQUAL(TEXT("InFilter pin has correct number of outputs"), InFilterOutput.Num(), 4);
+	UTEST_EQUAL(TEXT("OutFilter pin has correct number of outputs"), OutFilterOutput.Num(), 4);
+
+	UTEST_EQUAL(TEXT("InFilter data has correct number of empty outputs"), PCGPointFilterTest::GetNumEmptyData(InFilterOutput), 4);
+	UTEST_EQUAL(TEXT("OutFilter data has correct number of empty outputs"), PCGPointFilterTest::GetNumEmptyData(OutFilterOutput), 0);
+
+	return true;
+}
+
+bool FPCGAttributeFilter_AttributesRange_GenerateEmptyOutput::RunTest(const FString& Parameters)
+{
+	PCGTestsCommon::FTestData TestData;
+	UPCGAttributeFilteringRangeSettings* Settings = PCGTestsCommon::GenerateSettings<UPCGAttributeFilteringRangeSettings>(TestData);
+	check(Settings);
+
+	const FPCGElementPtr TestElement = TestData.Settings->GetElement();
+
+	Settings->TargetAttribute.SetAttributeName(PCGPointFilterTest::DoubleAttributeName);
+	Settings->MinThreshold.bUseConstantThreshold = true;
+	Settings->MaxThreshold.bUseConstantThreshold = true;
+	Settings->MinThreshold.AttributeTypes.Type = EPCGMetadataTypes::Double;
+	Settings->MaxThreshold.AttributeTypes.Type = EPCGMetadataTypes::Double;
+	Settings->MinThreshold.bInclusive = true;
+	Settings->MaxThreshold.bInclusive = true;
+
+	// Generate 4 point data, each with 1 point that has a normalized value based on its data index: .25, .5, .75, 1.0
+	for (int32 I = 0; I < 4; ++I)
+	{
+		FPCGTaggedData& TaggedData = TestData.InputData.TaggedData.Emplace_GetRef(FPCGTaggedData());
+		TaggedData.Pin = PCGPinConstants::DefaultInputLabel;
+		TaggedData.Data = PCGPointFilterTest::GenerateAttributeSetWithSingleValue((1.0 + I) / 4.0);
+	}
+
+	// Check unfiltered empty data.
+	Settings->bGenerateOutputDataEvenIfEmpty = true;
+
+	Settings->MinThreshold.AttributeTypes.DoubleValue = 0.0f;
+	Settings->MaxThreshold.AttributeTypes.DoubleValue = 1.0f;
+	TUniquePtr<FPCGContext> Context = TestData.InitializeTestContext();
+	while (!TestElement->Execute(Context.Get())) {}
+
+	// When always generating output, all data should exist in both
+	TArray<FPCGTaggedData> InFilterOutput = Context->OutputData.GetInputsByPin(PCGPointFilterTest::InsideFilterLabel);
+	TArray<FPCGTaggedData> OutFilterOutput = Context->OutputData.GetInputsByPin(PCGPointFilterTest::OutsideFilterLabel);
+
+	// All, inclusive.
+	UTEST_EQUAL(TEXT("InFilter pin has correct number of outputs"), InFilterOutput.Num(), 4);
+	UTEST_EQUAL(TEXT("OutFilter pin has correct number of outputs"), OutFilterOutput.Num(), 4);
+
+	UTEST_EQUAL(TEXT("InFilter data has correct number of empty outputs"), PCGPointFilterTest::GetNumEmptyData(InFilterOutput), 0);
+	UTEST_EQUAL(TEXT("OutFilter data has correct number of empty outputs"), PCGPointFilterTest::GetNumEmptyData(OutFilterOutput), 4);
+
+	Settings->MinThreshold.AttributeTypes.DoubleValue = 0.0f;
+	Settings->MaxThreshold.AttributeTypes.DoubleValue = 0.0f;
+	Context = TestData.InitializeTestContext();
+	while (!TestElement->Execute(Context.Get())) {}
+
+	// When always generating output, all data should exist in both
+	InFilterOutput = Context->OutputData.GetInputsByPin(PCGPointFilterTest::InsideFilterLabel);
+	OutFilterOutput = Context->OutputData.GetInputsByPin(PCGPointFilterTest::OutsideFilterLabel);
+
+	// None, inclusive.
+	UTEST_EQUAL(TEXT("InFilter pin has correct number of outputs"), InFilterOutput.Num(), 4);
+	UTEST_EQUAL(TEXT("OutFilter pin has correct number of outputs"), OutFilterOutput.Num(), 4);
+
+	UTEST_EQUAL(TEXT("InFilter data has correct number of empty outputs"), PCGPointFilterTest::GetNumEmptyData(InFilterOutput), 4);
+	UTEST_EQUAL(TEXT("OutFilter data has correct number of empty outputs"), PCGPointFilterTest::GetNumEmptyData(OutFilterOutput), 0);
+
+	return true;
+}
+
+bool FPCGAttributeFilter_Points_NoGenerateEmptyOutput::RunTest(const FString& Parameters)
+{
+	PCGTestsCommon::FTestData TestData;
+	UPCGAttributeFilteringSettings* Settings = PCGTestsCommon::GenerateSettings<UPCGAttributeFilteringSettings>(TestData);
+	check(Settings);
+
+	const FPCGElementPtr TestElement = TestData.Settings->GetElement();
+
+	Settings->Operator = EPCGAttributeFilterOperator::Greater;
+	Settings->TargetAttribute.SetPointProperty(EPCGPointProperties::Density);
+	Settings->bUseConstantThreshold = true;
+	Settings->AttributeTypes.Type = EPCGMetadataTypes::Float;
+
+	// Generate 4 point data, each with 1 point that has a normalized density based on its data index: .25, .5, .75, 1.0
+	for (int32 I = 0; I < 4; ++I)
+	{
+		FPCGTaggedData& TaggedData = TestData.InputData.TaggedData.Emplace_GetRef(FPCGTaggedData());
+		TaggedData.Pin = PCGPinConstants::DefaultInputLabel;
+		TaggedData.Data = PCGPointFilterTest::GeneratePointDataWithSingleDensity((1.0 + I) / 4.0);
+	}
+
+	Settings->bGenerateOutputDataEvenIfEmpty = false;
+
+	// All > 0.0
+	Settings->AttributeTypes.FloatValue = 0.f;
+	TUniquePtr<FPCGContext> Context = TestData.InitializeTestContext();
+	while (!TestElement->Execute(Context.Get())) {}
+
+	TArray<FPCGTaggedData> InFilterOutput = Context->OutputData.GetInputsByPin(PCGPointFilterTest::InsideFilterLabel);
+	TArray<FPCGTaggedData> OutFilterOutput = Context->OutputData.GetInputsByPin(PCGPointFilterTest::OutsideFilterLabel);
+
+	// .25, .5, .75 and 1 > 0
+	UTEST_EQUAL(TEXT("InFilter pin has correct number of outputs"), InFilterOutput.Num(), 4);
+	UTEST_EQUAL(TEXT("OutFilter pin has correct number of outputs"), OutFilterOutput.Num(), 0);
+
+	UTEST_EQUAL(TEXT("InFilter data has correct number of empty outputs"), PCGPointFilterTest::GetNumEmptyData(InFilterOutput), 0);
+
+	Settings->AttributeTypes.FloatValue = 0.333f;
+	Context = TestData.InitializeTestContext();
+	while (!TestElement->Execute(Context.Get())) {}
+	InFilterOutput = Context->OutputData.GetInputsByPin(PCGPointFilterTest::InsideFilterLabel);
+	OutFilterOutput = Context->OutputData.GetInputsByPin(PCGPointFilterTest::OutsideFilterLabel);
+
+	// .5, .75, 1 > 0.333
+	UTEST_EQUAL(TEXT("InFilter pin has correct number of outputs"), InFilterOutput.Num(), 3);
+	UTEST_EQUAL(TEXT("OutFilter pin has correct number of outputs"), OutFilterOutput.Num(), 1);
+
+	UTEST_EQUAL(TEXT("InFilter data has correct number of empty outputs"), PCGPointFilterTest::GetNumEmptyData(InFilterOutput), 0);
+	UTEST_EQUAL(TEXT("OutFilter data has correct number of empty outputs"), PCGPointFilterTest::GetNumEmptyData(OutFilterOutput), 0);
+
+	Settings->AttributeTypes.FloatValue = 1.f;
+	Context = TestData.InitializeTestContext();
+	while (!TestElement->Execute(Context.Get())) {}
+	InFilterOutput = Context->OutputData.GetInputsByPin(PCGPointFilterTest::InsideFilterLabel);
+	OutFilterOutput = Context->OutputData.GetInputsByPin(PCGPointFilterTest::OutsideFilterLabel);
+
+	// None > 1
+	UTEST_EQUAL(TEXT("InFilter pin has correct number of outputs"), InFilterOutput.Num(), 0);
+	UTEST_EQUAL(TEXT("OutFilter pin has correct number of outputs"), OutFilterOutput.Num(), 4);
+
+	UTEST_EQUAL(TEXT("OutFilter data has correct number of empty outputs"), PCGPointFilterTest::GetNumEmptyData(OutFilterOutput), 0);
+
+	return true;
+}
+
+bool FPCGAttributeFilter_PointsRange_NoGenerateEmptyOutput::RunTest(const FString& Parameters)
+{
+	PCGTestsCommon::FTestData TestData;
+	UPCGAttributeFilteringRangeSettings* Settings = PCGTestsCommon::GenerateSettings<UPCGAttributeFilteringRangeSettings>(TestData);
+	check(Settings);
+
+	const FPCGElementPtr TestElement = TestData.Settings->GetElement();
+
+	Settings->TargetAttribute.SetPointProperty(EPCGPointProperties::Density);
+	Settings->MinThreshold.bUseConstantThreshold = true;
+	Settings->MaxThreshold.bUseConstantThreshold = true;
+	Settings->MinThreshold.AttributeTypes.Type = EPCGMetadataTypes::Float;
+	Settings->MaxThreshold.AttributeTypes.Type = EPCGMetadataTypes::Float;
+	Settings->MinThreshold.bInclusive = true;
+	Settings->MaxThreshold.bInclusive = true;
+
+	// Generate 4 point data, each with 1 point that has a normalized density based on its data index: 0, .25, .5, 0.75
+	for (int32 I = 0; I < 4; ++I)
+	{
+		FPCGTaggedData& TaggedData = TestData.InputData.TaggedData.Emplace_GetRef(FPCGTaggedData());
+		TaggedData.Pin = PCGPinConstants::DefaultInputLabel;
+		TaggedData.Data = PCGPointFilterTest::GeneratePointDataWithSingleDensity((1.0 + I) / 4.0);
+	}
+
+	// Check unfiltered empty data.
+	Settings->bGenerateOutputDataEvenIfEmpty = false;
+
+	Settings->MinThreshold.AttributeTypes.FloatValue = 0.0f;
+	Settings->MaxThreshold.AttributeTypes.FloatValue = 1.0f;
+	TUniquePtr<FPCGContext> Context = TestData.InitializeTestContext();
+	while (!TestElement->Execute(Context.Get())) {}
+
+	// When always generating output, all data should exist in both
+	TArray<FPCGTaggedData> InFilterOutput = Context->OutputData.GetInputsByPin(PCGPointFilterTest::InsideFilterLabel);
+	TArray<FPCGTaggedData> OutFilterOutput = Context->OutputData.GetInputsByPin(PCGPointFilterTest::OutsideFilterLabel);
+
+	// All, inclusive
+	UTEST_EQUAL(TEXT("InFilter pin has correct number of outputs"), InFilterOutput.Num(), 4);
+	UTEST_EQUAL(TEXT("OutFilter pin has correct number of outputs"), OutFilterOutput.Num(), 0);
+
+	Settings->MinThreshold.AttributeTypes.FloatValue = 0.2f;
+	Settings->MaxThreshold.AttributeTypes.FloatValue = 0.6f;
+	Context = TestData.InitializeTestContext();
+	while (!TestElement->Execute(Context.Get())) {}
+	InFilterOutput = Context->OutputData.GetInputsByPin(PCGPointFilterTest::InsideFilterLabel);
+	OutFilterOutput = Context->OutputData.GetInputsByPin(PCGPointFilterTest::OutsideFilterLabel);
+
+	// .2f <= .25, .5 <= .6
+	UTEST_EQUAL(TEXT("InFilter pin has correct number of outputs"), InFilterOutput.Num(), 2);
+	UTEST_EQUAL(TEXT("OutFilter pin has correct number of outputs"), OutFilterOutput.Num(), 2);
+
+	UTEST_EQUAL(TEXT("InFilter data has correct number of empty outputs"), PCGPointFilterTest::GetNumEmptyData(InFilterOutput), 0);
+	UTEST_EQUAL(TEXT("OutFilter data has correct number of empty outputs"), PCGPointFilterTest::GetNumEmptyData(OutFilterOutput), 0);
+
+	Settings->MinThreshold.AttributeTypes.FloatValue = 1.0f;
+	Settings->MaxThreshold.AttributeTypes.FloatValue = 1.0f;
+	Context = TestData.InitializeTestContext();
+	while (!TestElement->Execute(Context.Get())) {}
+	InFilterOutput = Context->OutputData.GetInputsByPin(PCGPointFilterTest::InsideFilterLabel);
+	OutFilterOutput = Context->OutputData.GetInputsByPin(PCGPointFilterTest::OutsideFilterLabel);
+
+	// 1 >= 1, inclusive
+	UTEST_EQUAL(TEXT("InFilter pin has correct number of outputs"), InFilterOutput.Num(), 1);
+	UTEST_EQUAL(TEXT("OutFilter pin has correct number of outputs"), OutFilterOutput.Num(), 3);
+
+	UTEST_EQUAL(TEXT("InFilter data has correct number of empty outputs"), PCGPointFilterTest::GetNumEmptyData(InFilterOutput), 0);
+	UTEST_EQUAL(TEXT("OutFilter data has correct number of empty outputs"), PCGPointFilterTest::GetNumEmptyData(OutFilterOutput), 0);
+
+	Settings->MinThreshold.bInclusive = false;
+	Settings->MaxThreshold.bInclusive = false;
+
+	Settings->MinThreshold.AttributeTypes.FloatValue = 1.0f;
+	Settings->MaxThreshold.AttributeTypes.FloatValue = 1.0f;
+	Context = TestData.InitializeTestContext();
+	while (!TestElement->Execute(Context.Get())) {}
+	InFilterOutput = Context->OutputData.GetInputsByPin(PCGPointFilterTest::InsideFilterLabel);
+	OutFilterOutput = Context->OutputData.GetInputsByPin(PCGPointFilterTest::OutsideFilterLabel);
+
+	// None > 1, exclusive.
+	UTEST_EQUAL(TEXT("InFilter pin has correct number of outputs"), InFilterOutput.Num(), 0);
+	UTEST_EQUAL(TEXT("OutFilter pin has correct number of outputs"), OutFilterOutput.Num(), 4);
+
+	UTEST_EQUAL(TEXT("OutFilter data has correct number of empty outputs"), PCGPointFilterTest::GetNumEmptyData(OutFilterOutput), 0);
+
+	return true;
+}
+
+bool FPCGAttributeFilter_Attributes_NoGenerateEmptyOutput::RunTest(const FString& Parameters)
+{
+	PCGTestsCommon::FTestData TestData;
+	UPCGAttributeFilteringSettings* Settings = PCGTestsCommon::GenerateSettings<UPCGAttributeFilteringSettings>(TestData);
+	check(Settings);
+
+	const FPCGElementPtr TestElement = TestData.Settings->GetElement();
+
+	Settings->Operator = EPCGAttributeFilterOperator::Greater;
+	Settings->TargetAttribute.SetAttributeName(PCGPointFilterTest::DoubleAttributeName);
+	Settings->AttributeTypes.Type = EPCGMetadataTypes::Double;
+	Settings->bUseConstantThreshold = true;
+
+	// Generate 4 point data, each with 1 point that has a normalized value based on its data index: .25, .5, .75, 1.0
+	for (int32 I = 0; I < 4; ++I)
+	{
+		FPCGTaggedData& TaggedData = TestData.InputData.TaggedData.Emplace_GetRef(FPCGTaggedData());
+		TaggedData.Pin = PCGPinConstants::DefaultInputLabel;
+		TaggedData.Data = PCGPointFilterTest::GenerateAttributeSetWithSingleValue((1.0 + I) / 4.0);
+	}
+
+	Settings->bGenerateOutputDataEvenIfEmpty = false;
+
+	Settings->AttributeTypes.DoubleValue = 0.f;
+	TUniquePtr<FPCGContext> Context = TestData.InitializeTestContext();
+	while (!TestElement->Execute(Context.Get())) {}
+
+	// When always generating output, all data should exist in both
+	TArray<FPCGTaggedData> InFilterOutput = Context->OutputData.GetInputsByPin(PCGPointFilterTest::InsideFilterLabel);
+	TArray<FPCGTaggedData> OutFilterOutput = Context->OutputData.GetInputsByPin(PCGPointFilterTest::OutsideFilterLabel);
+
+	// All
+	UTEST_EQUAL(TEXT("InFilter pin has correct number of outputs"), InFilterOutput.Num(), 4);
+	UTEST_EQUAL(TEXT("OutFilter pin has correct number of outputs"), OutFilterOutput.Num(), 0);
+
+	UTEST_EQUAL(TEXT("InFilter data has correct number of empty outputs"), PCGPointFilterTest::GetNumEmptyData(InFilterOutput), 0);
+
+	Settings->AttributeTypes.DoubleValue = 0.333f;
+	Context = TestData.InitializeTestContext();
+	while (!TestElement->Execute(Context.Get())) {}
+	InFilterOutput = Context->OutputData.GetInputsByPin(PCGPointFilterTest::InsideFilterLabel);
+	OutFilterOutput = Context->OutputData.GetInputsByPin(PCGPointFilterTest::OutsideFilterLabel);
+
+	// .5, .75, 1 > .333
+	UTEST_EQUAL(TEXT("InFilter pin has correct number of outputs"), InFilterOutput.Num(), 3);
+	UTEST_EQUAL(TEXT("OutFilter pin has correct number of outputs"), OutFilterOutput.Num(), 1);
+
+	UTEST_EQUAL(TEXT("InFilter data has correct number of empty outputs"), PCGPointFilterTest::GetNumEmptyData(InFilterOutput), 0);
+	UTEST_EQUAL(TEXT("OutFilter data has correct number of empty outputs"), PCGPointFilterTest::GetNumEmptyData(OutFilterOutput), 0);
+
+	Settings->AttributeTypes.DoubleValue = 1.f;
+	Context = TestData.InitializeTestContext();
+	while (!TestElement->Execute(Context.Get())) {}
+	InFilterOutput = Context->OutputData.GetInputsByPin(PCGPointFilterTest::InsideFilterLabel);
+	OutFilterOutput = Context->OutputData.GetInputsByPin(PCGPointFilterTest::OutsideFilterLabel);
+
+	// None > 1
+	UTEST_EQUAL(TEXT("InFilter pin has correct number of outputs"), InFilterOutput.Num(), 0);
+	UTEST_EQUAL(TEXT("OutFilter pin has correct number of outputs"), OutFilterOutput.Num(), 4);
+
+	UTEST_EQUAL(TEXT("OutFilter data has correct number of empty outputs"), PCGPointFilterTest::GetNumEmptyData(OutFilterOutput), 0);
+
+	return true;
+}
+
+bool FPCGAttributeFilter_AttributesRange_NoGenerateEmptyOutput::RunTest(const FString& Parameters)
+{
+	PCGTestsCommon::FTestData TestData;
+	UPCGAttributeFilteringRangeSettings* Settings = PCGTestsCommon::GenerateSettings<UPCGAttributeFilteringRangeSettings>(TestData);
+	check(Settings);
+
+	const FPCGElementPtr TestElement = TestData.Settings->GetElement();
+
+	Settings->TargetAttribute.SetAttributeName(PCGPointFilterTest::DoubleAttributeName);
+	Settings->MinThreshold.bUseConstantThreshold = true;
+	Settings->MaxThreshold.bUseConstantThreshold = true;
+	Settings->MinThreshold.AttributeTypes.Type = EPCGMetadataTypes::Double;
+	Settings->MaxThreshold.AttributeTypes.Type = EPCGMetadataTypes::Double;
+	Settings->MinThreshold.bInclusive = true;
+	Settings->MaxThreshold.bInclusive = true;
+
+	// Generate 4 point data, each with 1 point that has a normalized value based on its data index: .25, .5, .75, 1.0
+	for (int32 I = 0; I < 4; ++I)
+	{
+		FPCGTaggedData& TaggedData = TestData.InputData.TaggedData.Emplace_GetRef(FPCGTaggedData());
+		TaggedData.Pin = PCGPinConstants::DefaultInputLabel;
+		TaggedData.Data = PCGPointFilterTest::GenerateAttributeSetWithSingleValue((1.0 + I) / 4.0);
+	}
+
+	// Check unfiltered empty data.
+	Settings->bGenerateOutputDataEvenIfEmpty = false;
+
+	Settings->MinThreshold.AttributeTypes.DoubleValue = 0.0f;
+	Settings->MaxThreshold.AttributeTypes.DoubleValue = 1.0f;
+	TUniquePtr<FPCGContext> Context = TestData.InitializeTestContext();
+	while (!TestElement->Execute(Context.Get())) {}
+
+	// When always generating output, all data should exist in both
+	TArray<FPCGTaggedData> InFilterOutput = Context->OutputData.GetInputsByPin(PCGPointFilterTest::InsideFilterLabel);
+	TArray<FPCGTaggedData> OutFilterOutput = Context->OutputData.GetInputsByPin(PCGPointFilterTest::OutsideFilterLabel);
+
+	// All, inclusive.
+	UTEST_EQUAL(TEXT("InFilter pin has correct number of outputs"), InFilterOutput.Num(), 4);
+	UTEST_EQUAL(TEXT("OutFilter pin has correct number of outputs"), OutFilterOutput.Num(), 0);
+
+	UTEST_EQUAL(TEXT("InFilter data has correct number of empty outputs"), PCGPointFilterTest::GetNumEmptyData(InFilterOutput), 0);
+
+	Settings->MinThreshold.AttributeTypes.DoubleValue = 0.2f;
+	Settings->MaxThreshold.AttributeTypes.DoubleValue = 0.6f;
+	Context = TestData.InitializeTestContext();
+	while (!TestElement->Execute(Context.Get())) {}
+	InFilterOutput = Context->OutputData.GetInputsByPin(PCGPointFilterTest::InsideFilterLabel);
+	OutFilterOutput = Context->OutputData.GetInputsByPin(PCGPointFilterTest::OutsideFilterLabel);
+
+	// 0.2 < .25, .5 < .6
+	UTEST_EQUAL(TEXT("InFilter pin has correct number of outputs"), InFilterOutput.Num(), 2);
+	UTEST_EQUAL(TEXT("OutFilter pin has correct number of outputs"), OutFilterOutput.Num(), 2);
+
+	UTEST_EQUAL(TEXT("InFilter data has correct number of empty outputs"), PCGPointFilterTest::GetNumEmptyData(InFilterOutput), 0);
+	UTEST_EQUAL(TEXT("OutFilter data has correct number of empty outputs"), PCGPointFilterTest::GetNumEmptyData(OutFilterOutput), 0);
+
+	Settings->MinThreshold.AttributeTypes.DoubleValue = 1.0f;
+	Settings->MaxThreshold.AttributeTypes.DoubleValue = 1.0f;
+	Context = TestData.InitializeTestContext();
+	while (!TestElement->Execute(Context.Get())) {}
+	InFilterOutput = Context->OutputData.GetInputsByPin(PCGPointFilterTest::InsideFilterLabel);
+	OutFilterOutput = Context->OutputData.GetInputsByPin(PCGPointFilterTest::OutsideFilterLabel);
+
+	// 1 >= 1, inclusive.
+	UTEST_EQUAL(TEXT("InFilter pin has correct number of outputs"), InFilterOutput.Num(), 1);
+	UTEST_EQUAL(TEXT("OutFilter pin has correct number of outputs"), OutFilterOutput.Num(), 3);
+
+	UTEST_EQUAL(TEXT("InFilter data has correct number of empty outputs"), PCGPointFilterTest::GetNumEmptyData(InFilterOutput), 0);
+	UTEST_EQUAL(TEXT("OutFilter data has correct number of empty outputs"), PCGPointFilterTest::GetNumEmptyData(OutFilterOutput), 0);
+
+	Settings->MinThreshold.bInclusive = false;
+	Settings->MaxThreshold.bInclusive = false;
+
+	Settings->MinThreshold.AttributeTypes.DoubleValue = 1.0f;
+	Settings->MaxThreshold.AttributeTypes.DoubleValue = 1.0f;
+	Context = TestData.InitializeTestContext();
+	while (!TestElement->Execute(Context.Get())) {}
+	InFilterOutput = Context->OutputData.GetInputsByPin(PCGPointFilterTest::InsideFilterLabel);
+	OutFilterOutput = Context->OutputData.GetInputsByPin(PCGPointFilterTest::OutsideFilterLabel);
+
+	// None >= 1, exclusive.
+	UTEST_EQUAL(TEXT("InFilter pin has correct number of outputs"), InFilterOutput.Num(), 0);
+	UTEST_EQUAL(TEXT("OutFilter pin has correct number of outputs"), OutFilterOutput.Num(), 4);
+
+	UTEST_EQUAL(TEXT("OutFilter data has correct number of empty outputs"), PCGPointFilterTest::GetNumEmptyData(OutFilterOutput), 0);
+
 	return true;
 }
\ No newline at end of file
diff --git a/Engine/Plugins/PCG/Source/PCG/Public/Elements/PCGAttributeFilter.h b/Engine/Plugins/PCG/Source/PCG/Public/Elements/PCGAttributeFilter.h
--- a/Engine/Plugins/PCG/Source/PCG/Public/Elements/PCGAttributeFilter.h
+++ b/Engine/Plugins/PCG/Source/PCG/Public/Elements/PCGAttributeFilter.h
@@ -123,6 +123,10 @@ public:
 	UPROPERTY(BlueprintReadWrite, EditAnywhere, Category = Settings)
 	bool bWarnOnDataMissingAttribute = true;
 
+	/** Always generate output data (possibly empty) for both the in and out filters, even if all/none of the elements were filtered. */
+	UPROPERTY(BlueprintReadWrite, EditAnywhere, Category = Settings)
+	bool bGenerateOutputDataEvenIfEmpty = true;
+
 	// Hidden value to indicate that Spatial -> Point deprecation is on where pins are not explicitly points.
 	UPROPERTY()
 	bool bHasSpatialToPointDeprecation = false;
@@ -191,6 +195,10 @@ public:
 	UPROPERTY(BlueprintReadWrite, EditAnywhere, Category = Settings)
 	bool bWarnOnDataMissingAttribute = true;
 
+	/** Always generate output data (possibly empty) for both the in and out filters, even if all/none of the elements were filtered. */
+	UPROPERTY(BlueprintReadWrite, EditAnywhere, Category = Settings)
+	bool bGenerateOutputDataEvenIfEmpty = true;
+
 	// Hidden value to indicate that Spatial -> Point deprecation is on where pins are not explicitly points.
 	UPROPERTY()
 	bool bHasSpatialToPointDeprecation = false;
@@ -203,7 +211,7 @@ protected:
 	virtual bool SupportsBasePointDataInputs(FPCGContext* InContext) const override { return true; }
 
 	/** Performs the filter operation. Returns true when the operation is complete. */
-	bool DoFiltering(FPCGContext* Context, EPCGAttributeFilterOperator InOperation, const FPCGAttributePropertyInputSelector& TargetAttribute, bool bHasSpatialToPointDeprecation, bool bWarnOnDataMissingAttribute, const FPCGAttributeFilterThresholdSettings& FirstThreshold, const FPCGAttributeFilterThresholdSettings* SecondThreshold = nullptr) const;
+	bool DoFiltering(FPCGContext* Context, EPCGAttributeFilterOperator InOperation, const FPCGAttributePropertyInputSelector& InTargetAttribute, bool bHasSpatialToPointDeprecation, bool bWarnOnDataMissingAttribute, bool bAlwaysGenerateOutputData, const FPCGAttributeFilterThresholdSettings& FirstThreshold, const FPCGAttributeFilterThresholdSettings* SecondThreshold = nullptr) const;
 };
 
 class FPCGAttributeFilterElement : public FPCGAttributeFilterElementBase
