[PCG] Support per-data GPU attribute value counting, rather than across all data.

#rnx

diff --git a/Engine/Plugins/PCG/Shaders/Private/BuiltInKernels/PCGCountUniqueAttributeValues.usf b/Engine/Plugins/PCG/Shaders/Private/BuiltInKernels/PCGCountUniqueAttributeValues.usf
--- a/Engine/Plugins/PCG/Shaders/Private/BuiltInKernels/PCGCountUniqueAttributeValues.usf
+++ b/Engine/Plugins/PCG/Shaders/Private/BuiltInKernels/PCGCountUniqueAttributeValues.usf
@@ -17,34 +17,48 @@ void PCGCountUniqueAttributeValuesCS(uint3 GroupId : SV_GroupID, uint GroupIndex
 
 	const uint ThreadIndex = GetUnWrappedDispatchThreadId(GroupId, GroupIndex, 64);
 
-	uint In_DataIndex, In_ElementIndex;
-	if (!In_GetThreadData(ThreadIndex, In_DataIndex, In_ElementIndex))
+	// Pre-job - if this thread maps to any element in the output, initialize the value key. The dispatch thread count will be large enough to ensure
+	// all output elements are initialized.
 	{
-		return;
+		uint Out_DataIndex;
+		uint Out_ElementIndex;
+		if (Out_GetThreadData(ThreadIndex, Out_DataIndex, Out_ElementIndex))
+		{
+			// Write value of attribute. We currently have an output element per possible string key value.
+ 			const int OutputValueAttributeId = CountUniqueValues_GetOutputValueAttributeId();
+ 			Out_SetInt(Out_DataIndex, Out_ElementIndex, OutputValueAttributeId, Out_ElementIndex);
+		}
 	}
 
-	if (In_IsPointRemoved(In_DataIndex, In_ElementIndex))
+	// Main job - one thread per input data element. Read the string key value and atomic increment the output counter.
 	{
-		return;
-	}
+		uint In_DataIndex, In_ElementIndex;
+		if (!In_GetThreadData(ThreadIndex, In_DataIndex, In_ElementIndex))
+		{
+			return;
+		}
 
-	const int AttributeToCountId = CountUniqueValues_GetAttributeToCountId();
+		if (In_IsPointRemoved(In_DataIndex, In_ElementIndex))
+		{
+			return;
+		}
 
-	const int AttributeValue = In_GetStringKey(In_DataIndex, In_ElementIndex, AttributeToCountId);
+		const int AttributeToCountId = CountUniqueValues_GetAttributeToCountId();
 
-	const uint Out_DataIndex = 0;
-	const uint Out_ElementIndex = CountUniqueValues_GetValueIndex(AttributeValue);
+		const int AttributeValue = In_GetStringKey(In_DataIndex, In_ElementIndex, AttributeToCountId);
 
-	// Write value of attribute.
-	const int OutputValueAttributeId = CountUniqueValues_GetOutputValueAttributeId();
-	Out_SetInt(Out_DataIndex, Out_ElementIndex, OutputValueAttributeId, AttributeValue);
+		const uint Out_DataIndex = CountUniqueValues_GetEmitPerDataCounts() ? In_DataIndex : 0;
+		const uint Out_ElementIndex = CountUniqueValues_GetValueIndex(Out_DataIndex, AttributeValue);
 
-	const int OutputCountAttributeId = CountUniqueValues_GetOutputCountAttributeId();
-	const int ValueIndex = CountUniqueValues_GetValueIndex(AttributeValue);
+		if (Out_ElementIndex == (uint)-1)
+		{
+			// Should not happen.
+			return;
+		}
 
-	// Increment count of this attribute.
-	if (ValueIndex != -1)
-	{
-		Out_AtomicAddInt(Out_DataIndex, ValueIndex, OutputCountAttributeId, 1);
+		const int OutputCountAttributeId = CountUniqueValues_GetOutputCountAttributeId();
+
+		// Increment count of this attribute.
+		Out_AtomicAddInt(Out_DataIndex, Out_ElementIndex, OutputCountAttributeId, 1);
 	}
 }
diff --git a/Engine/Plugins/PCG/Source/PCG/Internal/Compute/PCGDataForGPU.h b/Engine/Plugins/PCG/Source/PCG/Internal/Compute/PCGDataForGPU.h
--- a/Engine/Plugins/PCG/Source/PCG/Internal/Compute/PCGDataForGPU.h
+++ b/Engine/Plugins/PCG/Source/PCG/Internal/Compute/PCGDataForGPU.h
@@ -137,7 +137,7 @@ public:
 
 	void AddUniqueStringKeys(const TArray<int32>& InOtherStringKeys);
 
-	void SetStringKeys(const TArrayView<int32>& InStringKeys);
+	void SetStringKeys(const TConstArrayView<int32>& InStringKeys);
 
 	bool operator==(const FPCGKernelAttributeDesc& Other) const;
 
@@ -216,7 +216,7 @@ struct FPCGDataCollectionDesc
 
 	/** Get description of first attribute with matching identifier in input data. Returns true if such an attribute found and also signals whether multiple matching attributes
 	* with conflicting names are present. */
-	bool GetAttributeDesc(FPCGAttributeIdentifier InAttributeIdentifier, FPCGKernelAttributeDesc& OutAttributeDesc, bool &bOutConflictingTypesFound) const;
+	bool GetAttributeDesc(FPCGAttributeIdentifier InAttributeIdentifier, FPCGKernelAttributeDesc& OutAttributeDesc, bool& bOutConflictingTypesFound, bool& bOutPresentOnAllData) const;
 
 	bool ContainsAttributeOnAnyData(FPCGAttributeIdentifier InAttributeIdentifier) const;
 
@@ -225,6 +225,8 @@ struct FPCGDataCollectionDesc
 
 	void GetUniqueStringKeyValues(int32 InAttributeId, TArray<int32>& OutUniqueStringKeys) const;
 
+	int GetNumStringKeyValues(int32 InAttributeId) const;
+
 	/** Description of each data in this data collection. */
 	TArray<FPCGDataDesc> DataDescs;
 };
diff --git a/Engine/Plugins/PCG/Source/PCG/Private/Compute/BuiltInKernels/PCGCountUniqueAttributeValuesKernel.cpp b/Engine/Plugins/PCG/Source/PCG/Private/Compute/BuiltInKernels/PCGCountUniqueAttributeValuesKernel.cpp
--- a/Engine/Plugins/PCG/Source/PCG/Private/Compute/BuiltInKernels/PCGCountUniqueAttributeValuesKernel.cpp
+++ b/Engine/Plugins/PCG/Source/PCG/Private/Compute/BuiltInKernels/PCGCountUniqueAttributeValuesKernel.cpp
@@ -2,7 +2,6 @@
 
 #include "Compute/BuiltInKernels/PCGCountUniqueAttributeValuesKernel.h"
 
-#include "PCGComponent.h"
 #include "PCGContext.h"
 #include "PCGModule.h"
 #include "Compute/Elements/PCGComputeGraphElement.h"
@@ -10,11 +9,10 @@
 #include "Compute/PCGDataBinding.h"
 #include "Compute/PCGPinPropertiesGPU.h"
 #include "Compute/DataInterfaces/BuiltInKernels/PCGCountUniqueAttributeValuesDataInterface.h"
-#include "Elements/PCGStaticMeshSpawner.h"
 #include "Graph/PCGGPUGraphCompilationContext.h"
-#include "MeshSelectors/PCGMeshSelectorByAttribute.h"
 
 #include "ShaderCompilerCore.h"
+#include "Algo/MaxElement.h"
 
 #include UE_INLINE_GENERATED_CPP_BY_NAME(PCGCountUniqueAttributeValuesKernel)
 
@@ -47,8 +45,9 @@ bool UPCGCountUniqueAttributeValuesKernel::IsKernelDataValid(FPCGContext* InCont
 
 		FPCGKernelAttributeDesc AttributeDesc;
 		bool bConflictingTypesFound = false;
-
-		if (!InputDataDesc->GetAttributeDesc(AttributeName, AttributeDesc, bConflictingTypesFound))
+		bool bPresentOnAllData = false;
+		InputDataDesc->GetAttributeDesc(AttributeName, AttributeDesc, bConflictingTypesFound, bPresentOnAllData);
+		if (!bPresentOnAllData)
 		{
 			if (!InputDataDesc->DataDescs.IsEmpty())
 			{
@@ -98,19 +97,51 @@ bool UPCGCountUniqueAttributeValuesKernel::ComputeOutputBindingDataDesc(const UP
 	const FPCGKernelPin InputKernelPin(KernelIndex, PCGPinConstants::DefaultInputLabel, /*bIsInput=*/true);
 	ensure(InBinding->ComputeKernelPinDataDesc(InputKernelPin, &InputDesc));
 
+	if (InputDesc.DataDescs.IsEmpty())
+	{
+		return true;
+	}
+
 	const int32 AttributeId = InBinding->GetAttributeId(AttributeName, EPCGKernelAttributeType::StringKey);
 
 	if (AttributeId != INDEX_NONE)
 	{
+		// Base data description object, copied N times below to populate the final description.
+		FPCGDataDesc CountDataDesc;
+		CountDataDesc.Type = EPCGDataType::Param;
+		CountDataDesc.AddAttribute(FPCGKernelAttributeKey(PCGCountUniqueAttributeValuesConstants::ValueAttributeName, EPCGKernelAttributeType::Int), InBinding);
+		CountDataDesc.AddAttribute(FPCGKernelAttributeKey(PCGCountUniqueAttributeValuesConstants::ValueCountAttributeName, EPCGKernelAttributeType::Int), InBinding);
+
+		// Compute a set of (value, count) pairs across all data.
 		TArray<int32> UniqueStringKeyValues;
 		InputDesc.GetUniqueStringKeyValues(AttributeId, UniqueStringKeyValues);
 
-		FPCGDataDesc& OutputDataDesc = OutDataDesc.DataDescs.Emplace_GetRef();
-		OutputDataDesc.Type = EPCGDataType::Param;
-		OutputDataDesc.ElementCount = UniqueStringKeyValues.Num();
+		// If the highest string key value is 3, we'll allocate 4 counters. String key value 0 is reserved for empty/null string.
+		const int32* MaxStringKeyValue = Algo::MaxElement(UniqueStringKeyValues);
+		CountDataDesc.ElementCount = (MaxStringKeyValue ? *MaxStringKeyValue : 0) + 1;
 
-		OutputDataDesc.AddAttribute(FPCGKernelAttributeKey(PCGCountUniqueAttributeValuesConstants::ValueAttributeName, EPCGKernelAttributeType::Int), InBinding);
-		OutputDataDesc.AddAttribute(FPCGKernelAttributeKey(PCGCountUniqueAttributeValuesConstants::ValueCountAttributeName, EPCGKernelAttributeType::Int), InBinding);
+		if (bEmitPerDataCounts)
+		{
+			// Output data is array of (value, count) pairs for each data.
+			for (int InputIndex = 0; InputIndex < InputDesc.DataDescs.Num(); ++InputIndex)
+			{
+				const FPCGDataDesc& DataDesc = InputDesc.DataDescs[InputIndex];
+
+				// Data descriptions have multiple collection objects so worth to move last element for efficiency, otherwise copy.
+				if (InputIndex < InputDesc.DataDescs.Num() - 1)
+				{
+					OutDataDesc.DataDescs.Add(CountDataDesc);
+				}
+				else
+				{
+					OutDataDesc.DataDescs.Add(MoveTemp(CountDataDesc));
+				}
+			}
+		}
+		else
+		{
+			OutDataDesc.DataDescs.Add(MoveTemp(CountDataDesc));
+		}
 	}
 
 	return true;
@@ -118,20 +149,20 @@ bool UPCGCountUniqueAttributeValuesKernel::ComputeOutputBindingDataDesc(const UP
 
 int UPCGCountUniqueAttributeValuesKernel::ComputeThreadCount(const UPCGDataBinding* InBinding) const
 {
-	int ThreadCount = 0;
+	if (const FPCGDataCollectionDesc* InputPinDesc = InBinding->GetCachedKernelPinDataDesc(this, PCGPinConstants::DefaultInputLabel, /*bIsInput=*/true); ensure(InputPinDesc))
+	{
+		const int ElementCount = InputPinDesc->ComputeDataElementCount(EPCGDataType::Point);
 
-	// Execute a thread for every element in the input data.
-	const FPCGDataCollectionDesc* OutputPinDesc = InBinding->GetCachedKernelPinDataDesc(this, PCGPinConstants::DefaultInputLabel, /*bIsInput=*/true);
+		const int32 PartitionAttributeId = InBinding->GetAttributeId(AttributeName, EPCGKernelAttributeType::StringKey);
+		const int ValueCount = (PartitionAttributeId != INDEX_NONE) ? InputPinDesc->GetNumStringKeyValues(PartitionAttributeId) : 0;
 
-	if (ensure(OutputPinDesc))
+		// One thread per element, but also make sure we have enough threads to write each unique attribute value (rather than leaving uninitialized).
+		return FMath::Max(ElementCount, ValueCount);
+	}
+	else
 	{
-		for (const FPCGDataDesc& DataDesc : OutputPinDesc->DataDescs)
-		{
-			ThreadCount += DataDesc.ElementCount;
-		}
+		return 0;
 	}
-
-	return ThreadCount;
 }
 
 bool UPCGCountUniqueAttributeValuesKernel::DoesOutputPinRequireZeroInitialization(FName InOutputPinLabel) const
@@ -155,6 +186,7 @@ void UPCGCountUniqueAttributeValuesKernel::CreateAdditionalInputDataInterfaces(F
 	TObjectPtr<UPCGCountUniqueAttributeValuesDataInterface> KernelDI = InOutContext.NewObject_AnyThread<UPCGCountUniqueAttributeValuesDataInterface>(InObjectOuter);
 	KernelDI->ProducerKernel = this;
 	KernelDI->AttributeToCountName = AttributeName;
+	KernelDI->bEmitPerDataCounts = bEmitPerDataCounts;
 
 	OutDataInterfaces.Add(KernelDI);
 }
diff --git a/Engine/Plugins/PCG/Source/PCG/Private/Compute/BuiltInKernels/PCGCountUniqueAttributeValuesKernel.h b/Engine/Plugins/PCG/Source/PCG/Private/Compute/BuiltInKernels/PCGCountUniqueAttributeValuesKernel.h
--- a/Engine/Plugins/PCG/Source/PCG/Private/Compute/BuiltInKernels/PCGCountUniqueAttributeValuesKernel.h
+++ b/Engine/Plugins/PCG/Source/PCG/Private/Compute/BuiltInKernels/PCGCountUniqueAttributeValuesKernel.h
@@ -20,6 +20,9 @@ class UPCGCountUniqueAttributeValuesKernel : public UPCGComputeKernel
 public:
 	void SetAttributeName(FName InAttributeName) { AttributeName = InAttributeName; }
 
+	/** Whether to produce one set of element counts that count across all input data, rather than producing counters per input data. */
+	virtual void SetEmitPerDataCounts(bool bInEmitPerDataCounts) { bEmitPerDataCounts = bInEmitPerDataCounts; }
+
 	virtual bool IsKernelDataValid(FPCGContext* InContext) const override;
 	virtual bool ComputeOutputBindingDataDesc(const UPCGComputeGraph* InGraph, FName InOutputPinLabel, UPCGDataBinding* InBinding, FPCGDataCollectionDesc& OutDataDesc) const override;
 	virtual int ComputeThreadCount(const UPCGDataBinding* Binding) const override;
@@ -39,4 +42,8 @@ public:
 public:
 	UPROPERTY()
 	FName AttributeName;
+
+	/** Produce one set of element counts that count across all input data, rather than producing counters per input data. */
+	UPROPERTY()
+	bool bEmitPerDataCounts = true;
 };
diff --git a/Engine/Plugins/PCG/Source/PCG/Private/Compute/BuiltInKernels/PCGDataNumKernel.cpp b/Engine/Plugins/PCG/Source/PCG/Private/Compute/BuiltInKernels/PCGDataNumKernel.cpp
--- a/Engine/Plugins/PCG/Source/PCG/Private/Compute/BuiltInKernels/PCGDataNumKernel.cpp
+++ b/Engine/Plugins/PCG/Source/PCG/Private/Compute/BuiltInKernels/PCGDataNumKernel.cpp
@@ -57,7 +57,8 @@ bool UPCGDataNumKernel::IsKernelDataValid(FPCGContext* InContext) const
 		{
 			FPCGKernelAttributeDesc AttributeDesc;
 			bool bConflictingTypesFound = false;
-			if (OutputDataDesc->GetAttributeDesc(OutCountAttributeKey.Identifier, AttributeDesc, bConflictingTypesFound))
+			bool bPresentOnAllData = false;
+			if (OutputDataDesc->GetAttributeDesc(OutCountAttributeKey.Identifier, AttributeDesc, bConflictingTypesFound, bPresentOnAllData))
 			{
 				if (OutCountAttributeKey.Type == EPCGKernelAttributeType::Invalid)
 				{
diff --git a/Engine/Plugins/PCG/Source/PCG/Private/Compute/DataInterfaces/BuiltInKernels/PCGCountUniqueAttributeValuesDataInterface.cpp b/Engine/Plugins/PCG/Source/PCG/Private/Compute/DataInterfaces/BuiltInKernels/PCGCountUniqueAttributeValuesDataInterface.cpp
--- a/Engine/Plugins/PCG/Source/PCG/Private/Compute/DataInterfaces/BuiltInKernels/PCGCountUniqueAttributeValuesDataInterface.cpp
+++ b/Engine/Plugins/PCG/Source/PCG/Private/Compute/DataInterfaces/BuiltInKernels/PCGCountUniqueAttributeValuesDataInterface.cpp
@@ -34,14 +34,20 @@ void UPCGCountUniqueAttributeValuesDataInterface::GetSupportedInputs(TArray<FSha
 	OutFunctions.AddDefaulted_GetRef()
 		.SetName(TEXT("CountUniqueValues_GetValueIndex"))
 		.AddReturnType(EShaderFundamentalType::Uint)
+		.AddParam(EShaderFundamentalType::Uint) // InOutputDataIndex
 		.AddParam(EShaderFundamentalType::Uint); // InAttributeValue
+
+	OutFunctions.AddDefaulted_GetRef()
+		.SetName(TEXT("CountUniqueValues_GetEmitPerDataCounts"))
+		.AddReturnType(EShaderFundamentalType::Bool);
 }
 
 BEGIN_SHADER_PARAMETER_STRUCT(FPCGCountUniqueAttributeValuesDataInterfaceParameters,)
 	SHADER_PARAMETER(int32, AttributeToCountId)
 	SHADER_PARAMETER(int32, OutputValueAttributeId)
 	SHADER_PARAMETER(int32, OutputCountAttributeId)
-	SHADER_PARAMETER_RDG_BUFFER_SRV(StructuredBuffer<uint32>, AttributeValueToValueIndex)
+	SHADER_PARAMETER(uint32, EmitPerDataCounts)
+	SHADER_PARAMETER(int32, MaxAttributeValueAcrossAllData)
 END_SHADER_PARAMETER_STRUCT()
 
 void UPCGCountUniqueAttributeValuesDataInterface::GetShaderParameters(TCHAR const* UID, FShaderParametersMetadataBuilder& InOutBuilder, FShaderParametersMetadataAllocations& InOutAllocations) const
@@ -57,10 +63,11 @@ void UPCGCountUniqueAttributeValuesDataInterface::GetHLSL(FString& OutHLSL, FStr
 	};
 
 	OutHLSL += FString::Format(TEXT(
-		"StructuredBuffer<uint> {DataInterfaceName}_AttributeValueToValueIndex;\n"
 		"int {DataInterfaceName}_AttributeToCountId;\n"
 		"int {DataInterfaceName}_OutputValueAttributeId;\n"
 		"int {DataInterfaceName}_OutputCountAttributeId;\n"
+		"uint {DataInterfaceName}_EmitPerDataCounts;\n"
+		"int {DataInterfaceName}_MaxAttributeValueAcrossAllData;\n"
 		"\n"
 		"int CountUniqueValues_GetAttributeToCountId_{DataInterfaceName}()\n"
 		"{\n"
@@ -77,11 +84,16 @@ void UPCGCountUniqueAttributeValuesDataInterface::GetHLSL(FString& OutHLSL, FStr
 		"	return {DataInterfaceName}_OutputCountAttributeId;\n"
 		"}\n"
 		"\n"
-		"uint CountUniqueValues_GetValueIndex_{DataInterfaceName}(int InAttributeValue)\n"
+		"uint CountUniqueValues_GetValueIndex_{DataInterfaceName}(uint InOutputDataIndex, int InAttributeValue)\n"
+		"{\n"
+		"	return (uint)( (int)InOutputDataIndex * {DataInterfaceName}_MaxAttributeValueAcrossAllData + InAttributeValue );\n"
+		"}\n"
+		"\n"
+		"bool CountUniqueValues_GetEmitPerDataCounts_{DataInterfaceName}()\n"
 		"{\n"
-		"	return {DataInterfaceName}_AttributeValueToValueIndex[InAttributeValue];\n"
+		"	return {DataInterfaceName}_EmitPerDataCounts > 0;\n"
 		"}\n"
-	), TemplateArgs);
+		), TemplateArgs);
 }
 
 UComputeDataProvider* UPCGCountUniqueAttributeValuesDataInterface::CreateDataProvider() const
@@ -99,6 +111,7 @@ void UPCGCountUniqueAttributeValuesDataProvider::Initialize(const UComputeDataIn
 	Kernel = DataInterface->ProducerKernel;
 
 	AttributeToCountName = DataInterface->AttributeToCountName;
+	bEmitPerDataCounts = DataInterface->bEmitPerDataCounts;
 }
 
 bool UPCGCountUniqueAttributeValuesDataProvider::PrepareForExecute_GameThread(UPCGDataBinding* InBinding)
@@ -125,16 +138,46 @@ bool UPCGCountUniqueAttributeValuesDataProvider::PrepareForExecute_GameThread(UP
 
 	AttributeToCountId = InBinding->GetAttributeId(AttributeToCountName, EPCGKernelAttributeType::StringKey);
 
+	TArray<TArray<int32>> UniqueStringKeyValuesPerOutputData;
+
 	if (AttributeToCountId != INDEX_NONE)
 	{
-		InputDataDesc->GetUniqueStringKeyValues(AttributeToCountId, UniqueStringKeyValues);
+		if (bEmitPerDataCounts)
+		{
+			// Gather unique string key values for each data.
+			for (int InputIndex = 0; InputIndex < InputDataDesc->DataDescs.Num(); ++InputIndex)
+			{
+				const FPCGDataDesc& DataDesc = InputDataDesc->DataDescs[InputIndex];
+
+				TArray<int32> UniqueStringKeyValues;
+
+				for (const FPCGKernelAttributeDesc& AttributeDesc : DataDesc.AttributeDescs)
+				{
+					if (AttributeDesc.AttributeId == AttributeToCountId)
+					{
+						UniqueStringKeyValues = AttributeDesc.GetUniqueStringKeys();
+						break;
+					}
+				}
+
+				UniqueStringKeyValuesPerOutputData.Emplace_GetRef() = MoveTemp(UniqueStringKeyValues);
+			}
+		}
+		else
+		{
+			// Gather unique string key values across all data.
+			InputDataDesc->GetUniqueStringKeyValues(AttributeToCountId, UniqueStringKeyValuesPerOutputData.Emplace_GetRef());
+		}
 	}
 
-	MaxAttributeValue = INDEX_NONE;
+	MaxAttributeValueAcrossAllData = INDEX_NONE;
 
-	for (int32 Key : UniqueStringKeyValues)
+	for (const TArray<int32>& UniqueStringKeyValues : UniqueStringKeyValuesPerOutputData)
 	{
-		MaxAttributeValue = FMath::Max(MaxAttributeValue, Key);
+		for (int32 Value : UniqueStringKeyValues)
+		{
+			MaxAttributeValueAcrossAllData = FMath::Max(MaxAttributeValueAcrossAllData, Value);
+		}
 	}
 
 	return true;
@@ -144,10 +187,10 @@ FComputeDataProviderRenderProxy* UPCGCountUniqueAttributeValuesDataProvider::Get
 {
 	FPCGCountUniqueAttributeValuesProviderProxy::FCountUniqueValuesData_RenderThread Data;
 	Data.AttributeToCountId = AttributeToCountId;
-	Data.MaxAttributeValue = MaxAttributeValue;
+	Data.MaxAttributeValueAcrossAllData = MaxAttributeValueAcrossAllData;
 	Data.OutputValueAttributeId = OutputValueAttributeId;
 	Data.OutputCountAttributeId = OutputCountAttributeId;
-	Data.UniqueStringKeyValues = UniqueStringKeyValues;
+	Data.bEmitPerDataCounts = bEmitPerDataCounts;
 
 	return new FPCGCountUniqueAttributeValuesProviderProxy(MoveTemp(Data));
 }
@@ -159,8 +202,8 @@ void UPCGCountUniqueAttributeValuesDataProvider::Reset()
 	AttributeToCountId = INDEX_NONE;
 	OutputValueAttributeId = INDEX_NONE;
 	OutputCountAttributeId = INDEX_NONE;
-	MaxAttributeValue = INDEX_NONE;
-	UniqueStringKeyValues.Empty();
+	bEmitPerDataCounts = true;
+	MaxAttributeValueAcrossAllData = INDEX_NONE;
 
 	Super::Reset();
 }
@@ -170,30 +213,6 @@ bool FPCGCountUniqueAttributeValuesProviderProxy::IsValid(FValidationData const&
 	return (InValidationData.ParameterStructSize == sizeof(FParameters));
 }
 
-void FPCGCountUniqueAttributeValuesProviderProxy::AllocateResources(FRDGBuilder& GraphBuilder, FAllocationData const& InAllocationData)
-{
-	const int32 NumIndicesToAllocate = FMath::Max(1, Data.MaxAttributeValue + 1);
-
-	{
-		FRDGBufferDesc Desc = FRDGBufferDesc::CreateStructuredDesc(sizeof(uint32), NumIndicesToAllocate);
-		Desc.Usage |= BUF_SourceCopy;
-
-		AttributeValueToValueIndex = GraphBuilder.CreateBuffer(Desc, TEXT("PCGAttributeValueToValueIndex"));
-		AttributeValueToValueIndexSRV = GraphBuilder.CreateSRV(AttributeValueToValueIndex);
-		
-		TArray<uint32> ValueIndices;
-		ValueIndices.SetNumZeroed(NumIndicesToAllocate);
-
-		for (int Value = 0; Value <= Data.MaxAttributeValue; ++Value)
-		{
-			const int32 ValueIndex = Data.UniqueStringKeyValues.IndexOfByKey(Value);
-			ValueIndices[Value] = static_cast<uint32>(ValueIndex);
-		}
-
-		GraphBuilder.QueueBufferUpload(AttributeValueToValueIndex, MakeArrayView(ValueIndices));
-	}
-}
-
 void FPCGCountUniqueAttributeValuesProviderProxy::GatherDispatchData(FDispatchData const& InDispatchData)
 {
 	const TStridedView<FParameters> ParameterArray = MakeStridedParameterView<FParameters>(InDispatchData);
@@ -204,6 +223,7 @@ void FPCGCountUniqueAttributeValuesProviderProxy::GatherDispatchData(FDispatchDa
 		Parameters.AttributeToCountId = Data.AttributeToCountId;
 		Parameters.OutputValueAttributeId = Data.OutputValueAttributeId;
 		Parameters.OutputCountAttributeId = Data.OutputCountAttributeId;
-		Parameters.AttributeValueToValueIndex = AttributeValueToValueIndexSRV;
+		Parameters.EmitPerDataCounts = Data.bEmitPerDataCounts ? 1u : 0u;
+		Parameters.MaxAttributeValueAcrossAllData = Data.MaxAttributeValueAcrossAllData + 1;
 	}
 }
diff --git a/Engine/Plugins/PCG/Source/PCG/Private/Compute/DataInterfaces/BuiltInKernels/PCGCountUniqueAttributeValuesDataInterface.h b/Engine/Plugins/PCG/Source/PCG/Private/Compute/DataInterfaces/BuiltInKernels/PCGCountUniqueAttributeValuesDataInterface.h
--- a/Engine/Plugins/PCG/Source/PCG/Private/Compute/DataInterfaces/BuiltInKernels/PCGCountUniqueAttributeValuesDataInterface.h
+++ b/Engine/Plugins/PCG/Source/PCG/Private/Compute/DataInterfaces/BuiltInKernels/PCGCountUniqueAttributeValuesDataInterface.h
@@ -27,6 +27,9 @@ public:
 public:
 	UPROPERTY()
 	FName AttributeToCountName;
+
+	UPROPERTY()
+	bool bEmitPerDataCounts = true;
 };
 
 UCLASS()
@@ -49,23 +52,17 @@ public:
 	UPROPERTY()
 	TObjectPtr<const UPCGComputeKernel> Kernel;
 
-	UPROPERTY()
 	FName AttributeToCountName;
 
-	UPROPERTY()
 	int32 AttributeToCountId = INDEX_NONE;
 
-	UPROPERTY()
 	int32 OutputValueAttributeId = INDEX_NONE;
 
-	UPROPERTY()
 	int32 OutputCountAttributeId = INDEX_NONE;
 
-	UPROPERTY()
-	int32 MaxAttributeValue = INDEX_NONE;
+	bool bEmitPerDataCounts = true;
 
-	UPROPERTY()
-	TArray<int32> UniqueStringKeyValues;
+	int32 MaxAttributeValueAcrossAllData = INDEX_NONE;
 };
 
 class FPCGCountUniqueAttributeValuesProviderProxy : public FComputeDataProviderRenderProxy
@@ -76,8 +73,8 @@ public:
 		int32 AttributeToCountId = INDEX_NONE;
 		int32 OutputValueAttributeId = INDEX_NONE;
 		int32 OutputCountAttributeId = INDEX_NONE;
-		int32 MaxAttributeValue = INDEX_NONE;
-		TArray<int32> UniqueStringKeyValues;
+		bool bEmitPerDataCounts = true;
+		int32 MaxAttributeValueAcrossAllData = INDEX_NONE;
 	};
 
 	FPCGCountUniqueAttributeValuesProviderProxy(FCountUniqueValuesData_RenderThread InData)
@@ -86,7 +83,6 @@ public:
 
 	//~ Begin FComputeDataProviderRenderProxy Interface
 	bool IsValid(FValidationData const& InValidationData) const override;
-	void AllocateResources(FRDGBuilder& GraphBuilder, FAllocationData const& InAllocationData);
 	void GatherDispatchData(FDispatchData const& InDispatchData) override;
 	//~ End FComputeDataProviderRenderProxy Interface
 
@@ -94,7 +90,4 @@ protected:
 	using FParameters = FPCGCountUniqueAttributeValuesDataInterfaceParameters;
 
 	FCountUniqueValuesData_RenderThread Data;
-
-	FRDGBufferRef AttributeValueToValueIndex = nullptr;
-	FRDGBufferSRVRef AttributeValueToValueIndexSRV = nullptr;
 };
diff --git a/Engine/Plugins/PCG/Source/PCG/Private/Compute/PCGDataForGPU.cpp b/Engine/Plugins/PCG/Source/PCG/Private/Compute/PCGDataForGPU.cpp
--- a/Engine/Plugins/PCG/Source/PCG/Private/Compute/PCGDataForGPU.cpp
+++ b/Engine/Plugins/PCG/Source/PCG/Private/Compute/PCGDataForGPU.cpp
@@ -585,7 +585,8 @@ namespace PCGDataForGPUHelpers
 			// find in the input data similar to the CPU SM Spawner.
 			FPCGKernelAttributeDesc AttributeDesc;
 			bool bConflictingTypesInData = false;
-			InDataCollectionDescription->GetAttributeDesc(AttributeName, AttributeDesc, bConflictingTypesInData);
+			bool bPresentOnAllData = false;
+			InDataCollectionDescription->GetAttributeDesc(AttributeName, AttributeDesc, bConflictingTypesInData, bPresentOnAllData);
 
 			if (bConflictingTypesInData)
 			{
@@ -728,7 +729,7 @@ void FPCGKernelAttributeDesc::AddUniqueStringKeys(const TArray<int32>& InOtherSt
 	}
 }
 
-void FPCGKernelAttributeDesc::SetStringKeys(const TArrayView<int32>& InStringKeys)
+void FPCGKernelAttributeDesc::SetStringKeys(const TConstArrayView<int32>& InStringKeys)
 {
 	UniqueStringKeys = InStringKeys;
 }
@@ -1666,18 +1667,25 @@ void FPCGDataCollectionDesc::Combine(const FPCGDataCollectionDesc& Other)
 	DataDescs.Append(Other.DataDescs);
 }
 
-bool FPCGDataCollectionDesc::GetAttributeDesc(FPCGAttributeIdentifier InAttributeIdentifier, FPCGKernelAttributeDesc& OutAttributeDesc, bool& bOutConflictingTypesFound) const
+bool FPCGDataCollectionDesc::GetAttributeDesc(FPCGAttributeIdentifier InAttributeIdentifier, FPCGKernelAttributeDesc& OutAttributeDesc, bool& bOutConflictingTypesFound, bool& bOutPresentOnAllData) const
 {
 	// Will be set to the type of the first attribute across all the data which has a matching name, and then is used to detect if subsequent
 	// attributes are found with conflicting type.
 	EPCGKernelAttributeType FoundType = EPCGKernelAttributeType::Invalid;
+	bOutPresentOnAllData = false;
+
+	bool bFoundOnAllData = true;
 
 	for (const FPCGDataDesc& DataDesc : DataDescs)
 	{
+		bool bFoundOnData = false;
+
 		for (const FPCGKernelAttributeDesc& AttributeDesc : DataDesc.AttributeDescs)
 		{
 			if (AttributeDesc.AttributeKey.Identifier == InAttributeIdentifier && ensure(AttributeDesc.AttributeKey.Type != EPCGKernelAttributeType::Invalid))
 			{
+				bFoundOnData = true;
+
 				if (FoundType == EPCGKernelAttributeType::Invalid)
 				{
 					// Take the first matching attribute.
@@ -1694,11 +1702,15 @@ bool FPCGDataCollectionDesc::GetAttributeDesc(FPCGAttributeIdentifier InAttribut
 				}
 			}
 		}
+
+		bFoundOnAllData &= bFoundOnData;
 	}
 
 	// Signal no conflict found.
 	bOutConflictingTypesFound = false;
 
+	bOutPresentOnAllData = bFoundOnAllData;
+
 	// Signal attribute found or not.
 	return FoundType != EPCGKernelAttributeType::Invalid;
 }
@@ -1738,6 +1750,14 @@ void FPCGDataCollectionDesc::GetUniqueStringKeyValues(int32 InAttributeId, TArra
 	}
 }
 
+int FPCGDataCollectionDesc::GetNumStringKeyValues(int32 InAttributeId) const
+{
+	TArray<int32> UniqueStringKeys;
+	GetUniqueStringKeyValues(InAttributeId, UniqueStringKeys);
+
+	return UniqueStringKeys.Num();
+}
+
 FPCGProxyForGPUDataCollection::FPCGProxyForGPUDataCollection(TRefCountPtr<FRDGPooledBuffer> InBuffer, uint32 InBufferSizeBytes, const FPCGDataCollectionDesc& InDescription, const TArray<FString>& InStringTable)
 	: Buffer(InBuffer)
 	, BufferSizeBytes(InBufferSizeBytes)
diff --git a/Engine/Plugins/PCG/Source/PCG/Private/Elements/Metadata/PCGMetadataPartitionKernel.cpp b/Engine/Plugins/PCG/Source/PCG/Private/Elements/Metadata/PCGMetadataPartitionKernel.cpp
--- a/Engine/Plugins/PCG/Source/PCG/Private/Elements/Metadata/PCGMetadataPartitionKernel.cpp
+++ b/Engine/Plugins/PCG/Source/PCG/Private/Elements/Metadata/PCGMetadataPartitionKernel.cpp
@@ -43,7 +43,8 @@ bool UPCGMetadataPartitionKernel::IsKernelDataValid(FPCGContext* InContext) cons
 
 		FPCGKernelAttributeDesc AttributeDesc;
 		bool bConflictingTypesInData = false;
-		const bool bAttributePresent = InputDesc->GetAttributeDesc(AttributeName, AttributeDesc, bConflictingTypesInData);
+		bool bPresentOnAllData = false;
+		const bool bAttributePresent = InputDesc->GetAttributeDesc(AttributeName, AttributeDesc, bConflictingTypesInData, bPresentOnAllData);
 
 		if (!bAttributePresent)
 		{
diff --git a/Engine/Plugins/PCG/Source/PCG/Private/Elements/PCGSkinnedMeshSpawner.cpp b/Engine/Plugins/PCG/Source/PCG/Private/Elements/PCGSkinnedMeshSpawner.cpp
--- a/Engine/Plugins/PCG/Source/PCG/Private/Elements/PCGSkinnedMeshSpawner.cpp
+++ b/Engine/Plugins/PCG/Source/PCG/Private/Elements/PCGSkinnedMeshSpawner.cpp
@@ -59,6 +59,8 @@ void UPCGSkinnedMeshSpawnerSettings::CreateKernels(FPCGGPUCompilationContext& In
 		UPCGCountUniqueAttributeValuesKernel* CountKernel = InOutContext.NewObject_AnyThread<UPCGCountUniqueAttributeValuesKernel>(InObjectOuter);
 		CountKernel->Initialize(KernelParams);
 		CountKernel->SetAttributeName(MeshSelectorParameters->MeshAttribute.GetName());
+		// We operate across all input data rather than spawning for each input data separately.
+		CountKernel->SetEmitPerDataCounts(false);
 
 		OutKernels.Add(CountKernel);
 
diff --git a/Engine/Plugins/PCG/Source/PCG/Private/Elements/PCGStaticMeshSpawner.cpp b/Engine/Plugins/PCG/Source/PCG/Private/Elements/PCGStaticMeshSpawner.cpp
--- a/Engine/Plugins/PCG/Source/PCG/Private/Elements/PCGStaticMeshSpawner.cpp
+++ b/Engine/Plugins/PCG/Source/PCG/Private/Elements/PCGStaticMeshSpawner.cpp
@@ -67,6 +67,8 @@ void UPCGStaticMeshSpawnerSettings::CreateKernels(FPCGGPUCompilationContext& InO
 		UPCGCountUniqueAttributeValuesKernel* CountKernel = InOutContext.NewObject_AnyThread<UPCGCountUniqueAttributeValuesKernel>(InObjectOuter);
 		CountKernel->Initialize(KernelParams);
 		CountKernel->SetAttributeName(Selector->AttributeName);
+		// We operate across all input data rather than spawning for each input data separately.
+		CountKernel->SetEmitPerDataCounts(false);
 
 		OutKernels.Add(CountKernel);
 
