[PCG] GPU Attribute Set Processor node

[RN] Added GPU processor node for Attribute Sets

diff --git a/Engine/Plugins/PCG/Source/PCG/Private/Compute/Elements/PCGCustomHLSL.cpp b/Engine/Plugins/PCG/Source/PCG/Private/Compute/Elements/PCGCustomHLSL.cpp
--- a/Engine/Plugins/PCG/Source/PCG/Private/Compute/Elements/PCGCustomHLSL.cpp
+++ b/Engine/Plugins/PCG/Source/PCG/Private/Compute/Elements/PCGCustomHLSL.cpp
@@ -905,6 +905,10 @@ void UPCGCustomHLSLSettings::UpdatePinSettings()
 				Properties.bAllowMultipleData = false;
 				Properties.bAllowEditMultipleData = false;
 			}
+			else if (KernelType == EPCGKernelType::AttributeSetProcessor)
+			{
+				Properties.AllowedTypes = EPCGDataType::Param;
+			}
 		}
 
 		if (!!(Properties.AllowedTypes & EPCGDataType::Landscape))
@@ -952,6 +956,11 @@ void UPCGCustomHLSLSettings::UpdatePinSettings()
 				bInitModeDrivenByKernel = true;
 				Properties.AllowedTypes = EPCGDataType::BaseTexture;
 			}
+			else if (KernelType == EPCGKernelType::AttributeSetProcessor)
+			{
+				bInitModeDrivenByKernel = true;
+				Properties.AllowedTypes = EPCGDataType::Param;
+			}
 
 			Properties.bShowPropertiesGPU = !IsTextureKernel();
 		}
diff --git a/Engine/Plugins/PCG/Source/PCG/Private/Compute/Elements/PCGCustomHLSL.h b/Engine/Plugins/PCG/Source/PCG/Private/Compute/Elements/PCGCustomHLSL.h
--- a/Engine/Plugins/PCG/Source/PCG/Private/Compute/Elements/PCGCustomHLSL.h
+++ b/Engine/Plugins/PCG/Source/PCG/Private/Compute/Elements/PCGCustomHLSL.h
@@ -24,6 +24,7 @@ enum class EPCGKernelType : uint8
 	TextureProcessor UMETA(Tooltip = "Kernel executes on each texel in the first input pin."),
 	TextureGenerator UMETA(Tooltip = "Kernel executes for each texel in a fixed size texture, configurable on node."),
 	Custom UMETA(Tooltip = "Execution thread counts and output buffer sizes configurable on node. All data read/write indices must be manually bounds checked."),
+	AttributeSetProcessor UMETA(Tooltip = "Kernel executes on each element of the attribute sets in the first input pin."),
 };
 
 /** Total number of threads that will be dispatched for this kernel. */
@@ -132,7 +133,7 @@ protected:
 	/** Will the ThreadCountMultiplier value be applied when calculating the dispatch thread count. */
 	bool IsThreadCountMultiplierInUse() const { return KernelType == EPCGKernelType::Custom && DispatchThreadCount != EPCGDispatchThreadCount::Fixed; }
 
-	bool IsProcessorKernel() const { return KernelType == EPCGKernelType::PointProcessor || KernelType == EPCGKernelType::TextureProcessor; }
+	bool IsProcessorKernel() const { return KernelType == EPCGKernelType::PointProcessor || KernelType == EPCGKernelType::TextureProcessor || KernelType == EPCGKernelType::AttributeSetProcessor; }
 	bool IsGeneratorKernel() const { return KernelType == EPCGKernelType::PointGenerator || KernelType == EPCGKernelType::TextureGenerator; }
 	bool IsTextureKernel() const { return KernelType == EPCGKernelType::TextureProcessor || KernelType == EPCGKernelType::TextureGenerator; }
 	bool IsPointKernel() const { return KernelType == EPCGKernelType::PointProcessor || KernelType == EPCGKernelType::PointGenerator; }
diff --git a/Engine/Plugins/PCG/Source/PCG/Private/Compute/Elements/PCGCustomHLSLKernel.cpp b/Engine/Plugins/PCG/Source/PCG/Private/Compute/Elements/PCGCustomHLSLKernel.cpp
--- a/Engine/Plugins/PCG/Source/PCG/Private/Compute/Elements/PCGCustomHLSLKernel.cpp
+++ b/Engine/Plugins/PCG/Source/PCG/Private/Compute/Elements/PCGCustomHLSLKernel.cpp
@@ -511,6 +511,27 @@ FString UPCGCustomHLSLKernel::GetCookedSource(FPCGGPUCompilationContext& InOutCo
 				{ OutputPin->Label.ToString() });
 		}
 	}
+	else if (CustomHLSLSettings->KernelType == EPCGKernelType::AttributeSetProcessor)
+	{
+		const FPCGPinProperties* InputPin = GetFirstInputPin();
+		const FPCGPinPropertiesGPU* OutputPin = GetFirstOutputPin();
+
+		if (InputPin && OutputPin)
+		{
+			KernelSpecificPreamble += TEXT("    uint ElementIndex; // Assumption - element index identical in input and output data for processor.\n");
+
+			AddThreadInfoForPin(InputPin);
+			AddThreadInfoForPin(OutputPin);
+
+			// Automatically copy value of all attributes for this element.
+			// TODO pass in IDs of attributes that are actually present.
+			KernelSpecificPreamble += FString::Format(TEXT(
+				"\n"
+				"    // Attribute Set processor always initializes outputs by copying input data elements.\n"
+				"    PCG_COPY_ALL_ATTRIBUTES_TO_OUTPUT({1}, {0}, {1}_DataIndex, ElementIndex, {0}_DataIndex, ElementIndex);\n"),
+				{ InputPin->Label.ToString(), OutputPin->Label.ToString() });
+		}
+	}
 
 	FString Result;
 
@@ -1164,6 +1185,19 @@ bool UPCGCustomHLSLKernel::PerformStaticValidation()
 				return false;
 			}
 		}
+		else if (bIsFirstInputPin && CustomHLSLSettings->KernelType == EPCGKernelType::AttributeSetProcessor)
+		{
+			if (!(Properties.AllowedTypes & EPCGDataType::Param))
+			{
+#if PCG_KERNEL_LOGGING_ENABLED
+				StaticLogEntries.Emplace(FText::Format(
+					LOCTEXT("InvalidNonAttributeSetPrimaryInput", "'Attribute Set Processor' nodes require primary input pin to be of type 'Attribute Set', but found '{0}'."),
+					FText::FromString(PCGCustomHLSLKernel::GetDataTypeString(Properties.AllowedTypes))),
+					EPCGKernelLogVerbosity::Error);
+#endif
+				return false;
+			}
+		}
 
 		if (!PCGComputeHelpers::IsTypeAllowedAsInput(Properties.AllowedTypes))
 		{
@@ -1204,7 +1238,7 @@ bool UPCGCustomHLSLKernel::PerformStaticValidation()
 				return false;
 			}
 		}
-		else if (bIsFirstInputPin && CustomHLSLSettings->IsTextureKernel())
+		else if (bIsFirstOutputPin && CustomHLSLSettings->IsTextureKernel())
 		{
 			if (!(Properties.AllowedTypes & EPCGDataType::BaseTexture))
 			{
@@ -1217,6 +1251,19 @@ bool UPCGCustomHLSLKernel::PerformStaticValidation()
 				return false;
 			}
 		}
+		else if (bIsFirstOutputPin && CustomHLSLSettings->KernelType == EPCGKernelType::AttributeSetProcessor)
+		{
+			if (!(Properties.AllowedTypes & EPCGDataType::Param))
+			{
+#if PCG_KERNEL_LOGGING_ENABLED
+				StaticLogEntries.Emplace(FText::Format(
+					LOCTEXT("InvalidNonAttributeSetPrimaryOutput", "'Attribute Set Processor' nodes require primary output pin to be of type 'Attribute Set', but found '{0}'."),
+					FText::FromString(PCGCustomHLSLKernel::GetDataTypeString(Properties.AllowedTypes))),
+					EPCGKernelLogVerbosity::Error);
+#endif
+				return false;
+			}
+		}
 
 		if (!PCGComputeHelpers::IsTypeAllowedAsOutput(Properties.AllowedTypes))
 		{
