[PCG] Add a generic editor execution source

#rnx
#okforgihub public

diff --git a/Engine/Plugins/PCG/Source/PCG/Private/Editor/PCGEditorDefaultExecutionSource.cpp b/Engine/Plugins/PCG/Source/PCG/Private/Editor/PCGEditorDefaultExecutionSource.cpp
new file mode 100644
--- /dev/null
+++ b/Engine/Plugins/PCG/Source/PCG/Private/Editor/PCGEditorDefaultExecutionSource.cpp
@@ -0,0 +1,240 @@
+﻿// Copyright Epic Games, Inc. All Rights Reserved.
+
+#include "Editor/PCGEditorDefaultExecutionSource.h"
+
+#include "PCGSubgraph.h"
+#include "Subsystems/IPCGBaseSubsystem.h"
+
+UPCGData* FPCGEditorDefaultExecutionState::GetSelfData() const
+{
+	return nullptr;
+}
+
+int32 FPCGEditorDefaultExecutionState::GetSeed() const
+{
+	check(Source);
+	return Source->Seed;
+}
+
+FString FPCGEditorDefaultExecutionState::GetDebugName() const
+{
+	return TEXT("PCGEditorDefaultExecutionSource");
+}
+
+FTransform FPCGEditorDefaultExecutionState::GetTransform() const
+{
+	return FTransform::Identity;
+}
+
+UWorld* FPCGEditorDefaultExecutionState::GetWorld() const
+{
+	return nullptr;
+}
+
+bool FPCGEditorDefaultExecutionState::HasAuthority() const
+{
+	return false;
+}
+
+FBox FPCGEditorDefaultExecutionState::GetBounds() const
+{
+	return FBox(-FVector::One(), FVector::One());
+}
+
+UPCGGraph* FPCGEditorDefaultExecutionState::GetGraph() const
+{
+	check(Source);
+	return Source->Graph;
+}
+
+UPCGGraphInstance* FPCGEditorDefaultExecutionState::GetGraphInstance() const
+{
+	return nullptr;
+}
+
+void FPCGEditorDefaultExecutionState::OnGraphExecutionAborted(bool bQuiet, bool bCleanupUnusedResources)
+{
+}
+
+void FPCGEditorDefaultExecutionState::Cancel()
+{
+}
+
+bool FPCGEditorDefaultExecutionState::IsGenerating() const
+{
+	check(Source);
+	return Source->CurrentGenerationTask != InvalidPCGTaskId;
+}
+
+#if WITH_EDITOR
+const PCGUtils::FExtraCapture& FPCGEditorDefaultExecutionState::GetExtraCapture() const
+{
+	check(Source);
+	return Source->ExtraCapture;
+}
+
+PCGUtils::FExtraCapture& FPCGEditorDefaultExecutionState::GetExtraCapture()
+{
+	check(Source);
+	return Source->ExtraCapture;
+}
+
+const FPCGGraphExecutionInspection& FPCGEditorDefaultExecutionState::GetInspection() const
+{
+	check(Source);
+	return Source->Inspection;
+}
+
+FPCGGraphExecutionInspection& FPCGEditorDefaultExecutionState::GetInspection()
+{
+	check(Source);
+	return Source->Inspection;
+}
+#endif // WITH_EDITOR
+
+void UPCGEditorDefaultExecutionSource::SetGraph(UPCGGraph* InGraph)
+{
+	if (InGraph == Graph)
+	{
+		return;
+	}
+
+#if WITH_EDITOR
+	if (Graph)
+	{
+		Graph->OnGraphChangedDelegate.RemoveAll(this);
+	}
+#endif // WITH_EDITOR
+	
+	Graph = InGraph;
+
+#if WITH_EDITOR
+	if (Graph)
+	{
+		Graph->OnGraphChangedDelegate.AddUObject(this, &UPCGEditorDefaultExecutionSource::OnGraphChanged);
+	}
+#endif // WITH_EDITOR
+}
+
+#if WITH_EDITOR
+void UPCGEditorDefaultExecutionSource::OnGraphChanged(UPCGGraphInterface* InGraph, EPCGChangeType ChangeType)
+{
+	if (InGraph != Graph)
+	{
+		return;
+	}
+
+	if (ChangeType == EPCGChangeType::Cosmetic ||
+		ChangeType == EPCGChangeType::GraphCustomization ||
+		ChangeType == EPCGChangeType::None)
+	{
+		// If it is a cosmetic change (or no change), nothing to do
+		return;
+	}
+
+	if (!InGraph || !InGraph->GetGraph())
+	{
+		return;
+	}
+
+	Generate();
+}
+#endif // WITH_EDITOR
+
+void UPCGEditorDefaultExecutionSource::Generate()
+{
+	if (CurrentGenerationTask != InvalidPCGTaskId)
+	{
+		return;
+	}
+
+	if (IPCGBaseSubsystem* Subsystem = State.GetSubsystem())
+	{
+		FPCGScheduleGraphParams GraphParams{
+			/*InGraph=*/ GetGraph(),
+			/*InExecutionSource=*/ this,
+			/*InPreGraphElement=*/ MakeShared<FPCGInputForwardingElement>(FPCGDataCollection{}),
+			/*InInputElement=*/ MakeShared<FPCGInputForwardingElement>(FPCGDataCollection{}),
+			/*InExternalDependencies*/ {},
+			/*InFromStack*/ nullptr,
+			/*bInAllowHierarchicalGeneration*/ false};
+		
+		const FPCGTaskId GenerationTaskId = Subsystem->ScheduleGraph(GraphParams);
+		if (GenerationTaskId != InvalidPCGTaskId)
+		{
+			FPCGScheduleGenericParams Params{
+				/*InOperation=*/ [this, WeakSubsystem = TWeakInterfacePtr(Subsystem)](FPCGContext*) -> bool
+				{
+					CurrentGenerationTask = InvalidPCGTaskId;
+					if (WeakSubsystem.IsValid())
+					{
+#if WITH_EDITOR
+						WeakSubsystem->OnPCGSourceGenerationDone.Broadcast(WeakSubsystem.Get(), this, EPCGGenerationStatus::Completed);
+#endif // WITH_EDITOR
+					}
+					
+					return true;
+				},
+				/*InAbortOperation=*/ [this, WeakSubsystem = TWeakInterfacePtr(Subsystem)](FPCGContext*) -> bool
+				{
+					CurrentGenerationTask = InvalidPCGTaskId;
+					if (WeakSubsystem.IsValid())
+					{
+#if WITH_EDITOR
+						WeakSubsystem->OnPCGSourceGenerationDone.Broadcast(WeakSubsystem.Get(), this, EPCGGenerationStatus::Aborted);
+#endif // WITH_EDITOR
+					}
+					
+					return true;
+				},
+				/*InExecutionSource=*/this,
+				/*InExecutionDependencies=*/{GenerationTaskId}
+			};
+			
+			CurrentGenerationTask = Subsystem->ScheduleGeneric(Params);
+		}
+	}
+}
+
+void UPCGEditorDefaultExecutionSource::Sunset()
+{
+#if WITH_EDITOR
+	if (IsValid(Graph))
+	{
+		check(Graph);
+		Graph->OnGraphChangedDelegate.RemoveAll(this);
+	}
+#endif // WITH_EDITOR
+}
+
+UPCGEditorDefaultExecutionSource::UPCGEditorDefaultExecutionSource()
+	: State(this)
+{
+}
+
+UPCGEditorDefaultExecutionSource::~UPCGEditorDefaultExecutionSource()
+{
+	Sunset();
+}
+
+void UPCGEditorDefaultExecutionSource::Initialize(const FPCGEditorDefaultExecutionSourceParams& InParams)
+{
+	SetGraph(InParams.Graph);
+	SetSeed(InParams.Seed);
+}
+
+void UPCGEditorDefaultExecutionSource::SetSeed(int32 InSeed)
+{
+	Seed = InSeed;
+}
+
+void UPCGEditorDefaultExecutionSource::AddReferencedObjects(UObject* InThis, FReferenceCollector& Collector)
+{
+	UPCGEditorDefaultExecutionSource* This = CastChecked<UPCGEditorDefaultExecutionSource>(InThis);
+	
+#if WITH_EDITOR
+	This->Inspection.AddReferencedObjects(Collector);
+#endif // WITH_EDITOR
+
+	Super::AddReferencedObjects(This, Collector);
+}
diff --git a/Engine/Plugins/PCG/Source/PCG/Public/Editor/PCGEditorDefaultExecutionSource.h b/Engine/Plugins/PCG/Source/PCG/Public/Editor/PCGEditorDefaultExecutionSource.h
new file mode 100644
--- /dev/null
+++ b/Engine/Plugins/PCG/Source/PCG/Public/Editor/PCGEditorDefaultExecutionSource.h
@@ -0,0 +1,101 @@
+﻿// Copyright Epic Games, Inc. All Rights Reserved.
+
+#pragma once
+
+#include "PCGGraphExecutionInspection.h"
+#include "PCGGraphExecutionStateInterface.h"
+
+#include "PCGEditorDefaultExecutionSource.generated.h"
+
+class UPCGEditorDefaultExecutionSource;
+class UPCGGraph;
+class UPCGGraphInterface;
+
+class FPCGEditorDefaultExecutionState : public IPCGGraphExecutionState
+{
+public:
+	FPCGEditorDefaultExecutionState() = default;
+	explicit FPCGEditorDefaultExecutionState(UPCGEditorDefaultExecutionSource* InSource) : Source(InSource){}
+	
+	PCG_API virtual UPCGData* GetSelfData() const override;
+	PCG_API virtual int32 GetSeed() const override;
+	PCG_API virtual FString GetDebugName() const override;
+	PCG_API virtual FTransform GetTransform() const override;
+	PCG_API virtual UWorld* GetWorld() const override;
+	PCG_API virtual bool HasAuthority() const override;
+	PCG_API virtual FBox GetBounds() const override;
+	PCG_API virtual UPCGGraph* GetGraph() const override;
+	PCG_API virtual UPCGGraphInstance* GetGraphInstance() const override;
+	PCG_API virtual void OnGraphExecutionAborted(bool bQuiet, bool bCleanupUnusedResources) override;
+	PCG_API virtual void Cancel() override;
+	PCG_API virtual bool IsGenerating() const override;
+
+#if WITH_EDITOR
+	PCG_API virtual const PCGUtils::FExtraCapture& GetExtraCapture() const override;
+	PCG_API virtual PCGUtils::FExtraCapture& GetExtraCapture() override;
+
+	PCG_API virtual const FPCGGraphExecutionInspection& GetInspection() const override;
+	PCG_API virtual FPCGGraphExecutionInspection& GetInspection() override;
+
+	// No dynamic tracking for editor resources (at the moment)
+	virtual void RegisterDynamicTracking(const UPCGSettings* InSettings, const TArrayView<TPair<FPCGSelectionKey, bool>>& InDynamicKeysAndCulling) override {};
+	virtual void RegisterDynamicTracking(const FPCGSelectionKeyToSettingsMap& InKeysToSettings) override {};
+	
+	virtual bool IsRefreshInProgress() const override {return IsGenerating();}
+#endif // WITH_EDITOR
+
+private:
+	UPCGEditorDefaultExecutionSource* Source = nullptr;
+};
+
+struct FPCGEditorDefaultExecutionSourceParams
+{
+	UPCGGraph* Graph = nullptr;
+	int32 Seed = 42;
+};
+
+UCLASS(MinimalAPI)
+class UPCGEditorDefaultExecutionSource : public UObject, public IPCGGraphExecutionSource
+{
+	friend FPCGEditorDefaultExecutionState;
+	
+	GENERATED_BODY()
+	
+public:
+	PCG_API UPCGEditorDefaultExecutionSource();
+	PCG_API ~UPCGEditorDefaultExecutionSource();
+
+	PCG_API void Initialize(const FPCGEditorDefaultExecutionSourceParams& InParams);
+
+	PCG_API static void AddReferencedObjects(UObject* InThis, FReferenceCollector& Collector);
+	
+	virtual IPCGGraphExecutionState& GetExecutionState() override { return State; }
+	virtual const IPCGGraphExecutionState& GetExecutionState() const override { return State; }
+
+	PCG_API void SetGraph(UPCGGraph* InGraph);
+	UPCGGraph* GetGraph() const { return Graph.Get(); }
+
+	PCG_API void SetSeed(int32 InSeed);
+
+#if WITH_EDITOR
+	PCG_API void OnGraphChanged(UPCGGraphInterface* Graph, EPCGChangeType ChangeType);
+#endif // WITH_EDITOR
+	
+	PCG_API void Generate();
+
+	PCG_API void Sunset();
+
+private:
+	FPCGTaskId CurrentGenerationTask = InvalidPCGTaskId;
+	
+	UPROPERTY()
+	TObjectPtr<UPCGGraph> Graph;
+	
+	FPCGEditorDefaultExecutionState State;
+	int32 Seed = 42;
+
+#if WITH_EDITOR
+	PCGUtils::FExtraCapture ExtraCapture;
+	FPCGGraphExecutionInspection Inspection;
+#endif // WITH_EDITOR
+};
\ No newline at end of file
diff --git a/Engine/Plugins/PCG/Source/PCG/Public/Subsystems/IPCGBaseSubsystem.h b/Engine/Plugins/PCG/Source/PCG/Public/Subsystems/IPCGBaseSubsystem.h
--- a/Engine/Plugins/PCG/Source/PCG/Public/Subsystems/IPCGBaseSubsystem.h
+++ b/Engine/Plugins/PCG/Source/PCG/Public/Subsystems/IPCGBaseSubsystem.h
@@ -12,6 +12,8 @@
 
 #include "UObject/ObjectKey.h"
 
+#include "IPCGBaseSubsystem.generated.h"
+
 class UPCGComputeGraph;
 class UPCGData;
 class UPCGGraph;
@@ -31,12 +33,17 @@ class UWorld;
 DECLARE_MULTICAST_DELEGATE_ThreeParams(FPCGOnPCGSourceGenerationDone, IPCGBaseSubsystem*, IPCGGraphExecutionSource*, EPCGGenerationStatus);
 #endif // WITH_EDITOR
 
+UINTERFACE()
+class UPCGBaseSubsystem : public UInterface
+{
+	GENERATED_BODY()
+};
+
 class IPCGBaseSubsystem
 {
+	GENERATED_BODY()
 
 public:
-	virtual ~IPCGBaseSubsystem() = default;
-
 	virtual UWorld* GetSubsystemWorld() const { return nullptr; }
 
 	PCG_API void InitializeBaseSubsystem();
diff --git a/Engine/Plugins/PCG/Source/PCGEditor/Private/Subsystems/PCGEditorSubsystem.cpp b/Engine/Plugins/PCG/Source/PCGEditor/Private/Subsystems/PCGEditorSubsystem.cpp
--- a/Engine/Plugins/PCG/Source/PCGEditor/Private/Subsystems/PCGEditorSubsystem.cpp
+++ b/Engine/Plugins/PCG/Source/PCGEditor/Private/Subsystems/PCGEditorSubsystem.cpp
@@ -2,6 +2,8 @@
 
 #include "Subsystems/PCGEditorSubsystem.h"
 
+#include "Editor/PCGEditorDefaultExecutionSource.h"
+
 #if WITH_EDITOR
 
 void UPCGEditorSubsystem::Initialize(FSubsystemCollectionBase& Collection)
@@ -16,6 +18,18 @@ void UPCGEditorSubsystem::Deinitialize()
 	Super::Deinitialize();
 }
 
+UPCGEditorDefaultExecutionSource* UPCGEditorSubsystem::CreateEditorExecutionSource(const FPCGEditorDefaultExecutionSourceParams& InParams)
+{
+	if (!InParams.Graph)
+	{
+		return nullptr;
+	}
+
+	UPCGEditorDefaultExecutionSource* NewSource = NewObject<UPCGEditorDefaultExecutionSource>(this);
+	NewSource->Initialize(InParams);
+	return NewSource;
+}
+
 void UPCGEditorSubsystem::Tick(float DeltaTime)
 {
 	IPCGBaseSubsystem::Tick();
diff --git a/Engine/Plugins/PCG/Source/PCGEditor/Public/Subsystems/PCGEditorSubsystem.h b/Engine/Plugins/PCG/Source/PCGEditor/Public/Subsystems/PCGEditorSubsystem.h
--- a/Engine/Plugins/PCG/Source/PCGEditor/Public/Subsystems/PCGEditorSubsystem.h
+++ b/Engine/Plugins/PCG/Source/PCGEditor/Public/Subsystems/PCGEditorSubsystem.h
@@ -2,8 +2,6 @@
 
 #pragma once
 
-#if WITH_EDITOR
-
 #include "EditorSubsystem.h"
 #include "TickableEditorObject.h"
 
@@ -11,20 +9,24 @@
 
 #include "PCGEditorSubsystem.generated.h"
 
+struct FPCGEditorDefaultExecutionSourceParams;
+class UPCGEditorDefaultExecutionSource;
+
 UCLASS(MinimalAPI)
 class UPCGEditorSubsystem : public UEditorSubsystem, public IPCGBaseSubsystem, public FTickableEditorObject
 {
 	GENERATED_BODY()
-
+	
+public:
 	PCGEDITOR_API virtual void Initialize(FSubsystemCollectionBase& Collection) override;
 	PCGEDITOR_API virtual void Deinitialize() override;
 
+	PCGEDITOR_API UPCGEditorDefaultExecutionSource* CreateEditorExecutionSource(const FPCGEditorDefaultExecutionSourceParams& InParams);
+
 protected:
 	//** FTickableEditorObject Interface
 	PCGEDITOR_API virtual void Tick(float DeltaTime) override;
 	PCGEDITOR_API virtual ETickableTickType GetTickableTickType() const override ;
 	PCGEDITOR_API virtual TStatId GetStatId() const override;
 	//~ END FTickableEditorObject Interface
-};
-
-#endif // WITH_EDITOR
\ No newline at end of file
+};
\ No newline at end of file
