[PCG] Exposition of editor classes

* PCGEditor and PCGEditorSchema are now public and extensible
* PCGEditor can now proivide the schema it wants to use, what's its appearance and how it registers tabs.
* Schema can now override how the filtering of the node is done

#rnx
#okforgihub public

diff --git a/Engine/Plugins/PCG/Source/PCG/Private/Editor/PCGGraphCustomization.cpp b/Engine/Plugins/PCG/Source/PCG/Private/Editor/PCGGraphCustomization.cpp
--- a/Engine/Plugins/PCG/Source/PCG/Private/Editor/PCGGraphCustomization.cpp
+++ b/Engine/Plugins/PCG/Source/PCG/Private/Editor/PCGGraphCustomization.cpp
@@ -73,4 +73,41 @@ bool FPCGGraphEditorCustomization::FilterSubgraph(const FSoftObjectPath& InSubgr
 	{
 		return bFilterSubgraphs && SubgraphFilterType == EPCGGraphEditorFiltering::Allow;
 	} 
-}
\ No newline at end of file
+}
+
+bool FPCGGraphEditorCustomization::FiltersSettings() const
+{
+	if (bFilterSettings)
+	{
+		return true;
+	}
+	else if (const FPCGGraphEditorCustomization* Parent = GetParent())
+	{
+		return Parent->FiltersSettings();
+	}
+	else
+	{
+		return false;
+	}
+}
+
+bool FPCGGraphEditorCustomization::FilterSettings(const TSubclassOf<UPCGSettings>& InSettingsClass) const
+{
+	if (bFilterSettings)
+	{
+		if (FilteredSettingsTypes.ContainsByPredicate([InSettingsClass](const TSubclassOf<UPCGSettings>& Other) { return InSettingsClass->IsChildOf(Other);}))
+		{
+			return SubgraphFilterType != EPCGGraphEditorFiltering::Allow;
+		}
+	}
+
+	// Otherwise, delegate filtering to parent if any
+	if (const FPCGGraphEditorCustomization* Parent = GetParent())
+	{
+		return Parent->FilterSettings(InSettingsClass);
+	}
+	else
+	{
+		return bFilterSettings && SubgraphFilterType == EPCGGraphEditorFiltering::Allow;
+	} 
+}
diff --git a/Engine/Plugins/PCG/Source/PCG/Public/Editor/PCGGraphCustomization.h b/Engine/Plugins/PCG/Source/PCG/Public/Editor/PCGGraphCustomization.h
--- a/Engine/Plugins/PCG/Source/PCG/Public/Editor/PCGGraphCustomization.h
+++ b/Engine/Plugins/PCG/Source/PCG/Public/Editor/PCGGraphCustomization.h
@@ -50,11 +50,24 @@ struct FPCGGraphEditorCustomization
 	UPROPERTY(EditAnywhere, Category = Filtering, meta = (EditCondition = "bFilterSubgraphs"))
 	TSet<TSoftObjectPtr<UPCGGraph>> FilteredSubgraphTypes;
 
+	UPROPERTY(EditAnywhere, Category = Filtering)
+	bool bFilterSettings = false;
+
+	UPROPERTY(EditAnywhere, Category = Filtering, meta = (EditCondition = "bFilterSettings"))
+	EPCGGraphEditorFiltering SettingsFilterType = EPCGGraphEditorFiltering::Allow;
+
+	UPROPERTY(EditAnywhere, Category = Filtering, meta = (EditCondition = "bFilterSettings"))
+	TArray<TSubclassOf<UPCGSettings>> FilteredSettingsTypes;
+
 	UE_API bool Accepts(const FText& InCategory) const;
 	UE_API bool FiltersSubgraphs() const;
 	// Returns false if path is accepted by the customization
 	UE_API bool FilterSubgraph(const FSoftObjectPath& InSubgraphPath) const;
 
+	UE_API bool FiltersSettings() const;
+	// Returns false if settings is accepted by the customization
+	UE_API bool FilterSettings(const TSubclassOf<UPCGSettings>& InSettingsClass) const;
+
 private:
 	const FPCGGraphEditorCustomization* GetParent() const;
 };
diff --git a/Engine/Plugins/PCG/Source/PCGEditor/Private/Nodes/PCGEditorGraphNodeBase.cpp b/Engine/Plugins/PCG/Source/PCGEditor/Private/Nodes/PCGEditorGraphNodeBase.cpp
--- a/Engine/Plugins/PCG/Source/PCGEditor/Private/Nodes/PCGEditorGraphNodeBase.cpp
+++ b/Engine/Plugins/PCG/Source/PCGEditor/Private/Nodes/PCGEditorGraphNodeBase.cpp
@@ -20,9 +20,9 @@
 #include "PCGEditorCommands.h"
 #include "PCGEditorCommon.h"
 #include "PCGEditorGraph.h"
-#include "PCGEditorGraphSchema.h"
 #include "PCGEditorSettings.h"
 #include "PCGEditorStyle.h"
+#include "Schema/PCGEditorGraphSchema.h"
 
 #include "CoreGlobals.h"
 #include "GraphEditorActions.h"
diff --git a/Engine/Plugins/PCG/Source/PCGEditor/Private/PCGEditor.cpp b/Engine/Plugins/PCG/Source/PCGEditor/Private/PCGEditor.cpp
--- a/Engine/Plugins/PCG/Source/PCGEditor/Private/PCGEditor.cpp
+++ b/Engine/Plugins/PCG/Source/PCGEditor/Private/PCGEditor.cpp
@@ -19,8 +19,6 @@
 
 #include "PCGEditorCommands.h"
 #include "PCGEditorGraph.h"
-#include "PCGEditorGraphSchema.h"
-#include "PCGEditorGraphSchemaActions.h"
 #include "PCGEditorMenuContext.h"
 #include "PCGEditorSettings.h"
 #include "PCGEditorStyle.h"
@@ -29,6 +27,8 @@
 #include "Nodes/PCGEditorGraphNodeInput.h"
 #include "Nodes/PCGEditorGraphNodeOutput.h"
 #include "Nodes/PCGEditorGraphNodeReroute.h"
+#include "Schema/PCGEditorGraphSchema.h"
+#include "Schema/PCGEditorGraphSchemaActions.h"
 
 #include "Widgets/SPCGEditorGraphActionWidget.h"
 #include "Widgets/SPCGEditorGraphAttributeListView.h"
@@ -70,6 +70,7 @@
 #include "Misc/MessageDialog.h"
 #include "Misc/TransactionObjectEvent.h"
 #include "Preferences/UnrealEdOptions.h"
+#include "Subsystems/PCGEditorSubsystem.h"
 #include "Widgets/Docking/SDockTab.h"
 #include "Widgets/Input/SEditableTextBox.h"
 #include "Widgets/Notifications/SNotificationList.h"
@@ -104,23 +105,36 @@ namespace FPCGEditor_private
 		FName(TEXT("Viewport4")) };
 
 	const FText UserParamsTabName = LOCTEXT("UserParamsTab", "Graph Parameters");
+
+	const FName AllDefaultPanels[] =
+	{
+		GraphEditorID,
+		PropertyDetailsID[0],
+		PropertyDetailsID[1],
+		PropertyDetailsID[2],
+		PropertyDetailsID[3],
+		PaletteID,
+		DebugObjectID,
+		AttributesID[0],
+		AttributesID[1],
+		AttributesID[2],
+		AttributesID[3],
+		FindID,
+		DeterminismID,
+		ProfilingID,
+		LogID,
+		HLSLSourceID,
+		UserParamsID,
+		ViewportID[0],
+		ViewportID[1],
+		ViewportID[2],
+		ViewportID[3]
+	};
 }
 
 UPCGEditorGraph* FPCGEditor::GetPCGEditorGraph(UPCGGraph* InGraph)
 {
-	if (!InGraph)
-	{
-		return nullptr;
-	}
-
-	if (!InGraph->PCGEditorGraph)
-	{
-		InGraph->PCGEditorGraph = NewObject<UPCGEditorGraph>(InGraph, UPCGEditorGraph::StaticClass(), NAME_None, RF_Transactional | RF_Transient);
-		InGraph->PCGEditorGraph->Schema = UPCGEditorGraphSchema::StaticClass();
-		InGraph->PCGEditorGraph->InitFromNodeGraph(InGraph);
-	}
-
-	return InGraph->PCGEditorGraph;
+	return InGraph ? InGraph->PCGEditorGraph : nullptr;
 }
 
 UPCGEditorGraph* FPCGEditor::GetPCGEditorGraph(const UPCGNode* InNode)
@@ -140,7 +154,18 @@ void FPCGEditor::Initialize(const EToolkitMode::Type InMode, const TSharedPtr<cl
 	PCGGraphBeingEdited = InPCGGraph;
 
 	// Initializes the UPCGEditorGraph if needed
-	GetPCGEditorGraph(InPCGGraph);
+	const TSubclassOf<UPCGEditorGraphSchema> SchemaClass = GetSchemaClass();
+	const bool bShouldUpdateSchema = !InPCGGraph->PCGEditorGraph || InPCGGraph->PCGEditorGraph->Schema != SchemaClass;
+	if (!InPCGGraph->PCGEditorGraph)
+	{
+		InPCGGraph->PCGEditorGraph = NewObject<UPCGEditorGraph>(InPCGGraph, UPCGEditorGraph::StaticClass(), NAME_None, RF_Transactional | RF_Transient);
+	}
+
+	if (bShouldUpdateSchema)
+	{
+		InPCGGraph->PCGEditorGraph->Schema = SchemaClass;
+		InPCGGraph->PCGEditorGraph->InitFromNodeGraph(InPCGGraph);
+	}
 
 	PCGGraphBeingEdited->PCGEditorGraph->SetEditor(SharedThis(this));
 	PCGEditorGraph = PCGGraphBeingEdited->PCGEditorGraph;
@@ -202,6 +227,11 @@ void FPCGEditor::Initialize(const EToolkitMode::Type InMode, const TSharedPtr<cl
 
 		// In case the editor is opened while in PIE, we should try setting up callbacks for the PIE world.
 		RegisterDelegatesForWorld(GEditor->PlayWorld.Get());
+
+		if (UPCGEditorSubsystem* EditorSubsystem = GEditor->GetEditorSubsystem<UPCGEditorSubsystem>())
+		{
+			EditorSubsystem->OnPCGSourceGenerationDone.AddRaw(this, &FPCGEditor::OnSourceGenerationDone);
+		}
 	}
 
 	// Clear inspection flag on all nodes.
@@ -511,37 +541,44 @@ FName FPCGEditor::GetPanelID(const EPCGEditorPanel Panel) const
 	}
 }
 
-void FPCGEditor::BringFocusToPanel(const EPCGEditorPanel Panel) const
+void FPCGEditor::BringFocusToPanel(FName PanelID) const
 {
-	const FName PanelID = GetPanelID(Panel);
 	if (PanelID != NAME_None)
 	{
-		const TSharedPtr<SDockTab> Tab = TabManager->TryInvokeTab(PanelID);
-		Tab->DrawAttention(); // Bring the panel to focus and flash the tab
+		if (const TSharedPtr<SDockTab> Tab = TabManager->TryInvokeTab(PanelID))
+		{
+			Tab->DrawAttention(); // Bring the panel to focus and flash the tab
+		}
 	}
 }
 
-void FPCGEditor::CloseGraphPanel(const EPCGEditorPanel Panel) const
+void FPCGEditor::CloseGraphPanel(FName PanelID) const
 {
-	const FName PanelID = GetPanelID(Panel);
 	if (PanelID != NAME_None)
 	{
-		const TSharedPtr<SDockTab> Tab = TabManager->FindExistingLiveTab(PanelID);
-		Tab->RequestCloseTab();
+		if (const TSharedPtr<SDockTab> Tab = TabManager->FindExistingLiveTab(PanelID))
+		{
+			Tab->RequestCloseTab();
+		}
 	}
 }
 
-bool FPCGEditor::IsPanelCurrentlyOpen(const EPCGEditorPanel Panel) const
+bool FPCGEditor::IsPanelCurrentlyOpen(FName PanelID) const
 {
-	return TabManager.IsValid() && TabManager->FindExistingLiveTab(GetPanelID(Panel));
+	return TabManager.IsValid() && TabManager->FindExistingLiveTab(PanelID);
 }
 
-bool FPCGEditor::IsPanelCurrentlyForeground(const EPCGEditorPanel Panel) const
+bool FPCGEditor::IsPanelCurrentlyForeground(FName PanelID) const
 {
-		const TSharedPtr<SDockTab> DockTab = TabManager.IsValid() ? TabManager->FindExistingLiveTab(GetPanelID(Panel)) : nullptr;
+		const TSharedPtr<SDockTab> DockTab = TabManager.IsValid() ? TabManager->FindExistingLiveTab(PanelID) : nullptr;
 		return DockTab.IsValid() && DockTab->IsForeground();
 }
 
+bool FPCGEditor::IsPanelAvailable(const FName PanelID) const
+{
+	return Algo::AnyOf(FPCGEditor_private::AllDefaultPanels, [PanelID](const FName& It) { return It == PanelID; });
+}
+
 void FPCGEditor::RegisterTabSpawners(const TSharedRef<FTabManager>& InTabManager)
 {
 	WorkspaceMenuCategory = InTabManager->AddLocalWorkspaceMenuCategory(LOCTEXT("WorkspaceMenu_PCGEditor", "PCG Editor"));
@@ -552,90 +589,52 @@ void FPCGEditor::RegisterTabSpawners(const TSharedRef<FTabManager>& InTabManager
 
 	FAssetEditorToolkit::RegisterTabSpawners(InTabManager);
 
-	//TODO: Add Icons
-	InTabManager->RegisterTabSpawner(FPCGEditor_private::GraphEditorID, FOnSpawnTab::CreateSP(this, &FPCGEditor::SpawnTab_GraphEditor))
-		.SetDisplayName(LOCTEXT("GraphTab", "Graph"))
-		.SetGroup(WorkspaceMenuCategoryRef);
-
-	InTabManager->RegisterTabSpawner(FPCGEditor_private::PropertyDetailsID[0], FOnSpawnTab::CreateSP(this, &FPCGEditor::SpawnTab_PropertyDetails, 0))
-		.SetDisplayName(LOCTEXT("DetailsTab1", "Details 1"))
-		.SetGroup(DetailsGroup);
-
-	InTabManager->RegisterTabSpawner(FPCGEditor_private::PropertyDetailsID[1], FOnSpawnTab::CreateSP(this, &FPCGEditor::SpawnTab_PropertyDetails, 1))
-		.SetDisplayName(LOCTEXT("DetailsTab2", "Details 2"))
-		.SetGroup(DetailsGroup);
-
-	InTabManager->RegisterTabSpawner(FPCGEditor_private::PropertyDetailsID[2], FOnSpawnTab::CreateSP(this, &FPCGEditor::SpawnTab_PropertyDetails, 2))
-		.SetDisplayName(LOCTEXT("DetailsTab3", "Details 3"))
-		.SetGroup(DetailsGroup);
-
-	InTabManager->RegisterTabSpawner(FPCGEditor_private::PropertyDetailsID[3], FOnSpawnTab::CreateSP(this, &FPCGEditor::SpawnTab_PropertyDetails, 3))
-		.SetDisplayName(LOCTEXT("DetailsTab4", "Details 4"))
-		.SetGroup(DetailsGroup);
-
-	InTabManager->RegisterTabSpawner(FPCGEditor_private::PaletteID, FOnSpawnTab::CreateSP(this, &FPCGEditor::SpawnTab_Palette))
-		.SetDisplayName(LOCTEXT("PaletteTab", "Palette"))
-		.SetGroup(WorkspaceMenuCategoryRef);
-
-	InTabManager->RegisterTabSpawner(FPCGEditor_private::DebugObjectID, FOnSpawnTab::CreateSP(this, &FPCGEditor::SpawnTab_DebugObjectTree))
-		.SetDisplayName(LOCTEXT("DebugTab", "Debug Object Tree"))
-		.SetGroup(WorkspaceMenuCategoryRef);
-
-	InTabManager->RegisterTabSpawner(FPCGEditor_private::AttributesID[0], FOnSpawnTab::CreateSP(this, &FPCGEditor::SpawnTab_Attributes, 0))
-		.SetDisplayName(LOCTEXT("AttributesTab1", "Attributes 1"))
-		.SetGroup(AttributesGroup);
-
-	InTabManager->RegisterTabSpawner(FPCGEditor_private::AttributesID[1], FOnSpawnTab::CreateSP(this, &FPCGEditor::SpawnTab_Attributes, 1))
-		.SetDisplayName(LOCTEXT("AttributesTab2", "Attributes 2"))
-		.SetGroup(AttributesGroup);
-
-	InTabManager->RegisterTabSpawner(FPCGEditor_private::AttributesID[2], FOnSpawnTab::CreateSP(this, &FPCGEditor::SpawnTab_Attributes, 2))
-		.SetDisplayName(LOCTEXT("AttributesTab3", "Attributes 3"))
-		.SetGroup(AttributesGroup);
-
-	InTabManager->RegisterTabSpawner(FPCGEditor_private::AttributesID[3], FOnSpawnTab::CreateSP(this, &FPCGEditor::SpawnTab_Attributes, 3))
-		.SetDisplayName(LOCTEXT("AttributesTab4", "Attributes 4"))
-		.SetGroup(AttributesGroup);
-
-	InTabManager->RegisterTabSpawner(FPCGEditor_private::FindID, FOnSpawnTab::CreateSP(this, &FPCGEditor::SpawnTab_Find))
-		.SetDisplayName(LOCTEXT("FindTab", "Find"))
-		.SetGroup(WorkspaceMenuCategoryRef);
+	using SpawnOnTabCallback = decltype(FOnSpawnTab::CreateSP(this, &FPCGEditor::SpawnTab_GraphEditor));
 
-	InTabManager->RegisterTabSpawner(FPCGEditor_private::DeterminismID, FOnSpawnTab::CreateSP(this, &FPCGEditor::SpawnTab_Determinism))
-		.SetDisplayName(LOCTEXT("DeterminismTab", "Determinism"))
-		.SetGroup(WorkspaceMenuCategoryRef);
-
-	InTabManager->RegisterTabSpawner(FPCGEditor_private::ProfilingID, FOnSpawnTab::CreateSP(this, &FPCGEditor::SpawnTab_Profiling))
-		.SetDisplayName(LOCTEXT("ProfilingTab", "Profiling"))
-		.SetGroup(WorkspaceMenuCategoryRef);
-
-	InTabManager->RegisterTabSpawner(FPCGEditor_private::LogID, FOnSpawnTab::CreateSP(this, &FPCGEditor::SpawnTab_Log))
-		.SetDisplayName(LOCTEXT("LogCaptureTab", "Log Capture"))
-		.SetGroup(WorkspaceMenuCategoryRef);
-
-	InTabManager->RegisterTabSpawner(FPCGEditor_private::HLSLSourceID, FOnSpawnTab::CreateSP(this, &FPCGEditor::SpawnTab_NodeSource))
-		.SetDisplayName(LOCTEXT("NodeSourceTab", "HLSL Source"))
-		.SetGroup(WorkspaceMenuCategoryRef);
-
-	InTabManager->RegisterTabSpawner(FPCGEditor_private::UserParamsID, FOnSpawnTab::CreateSP(this, &FPCGEditor::SpawnTab_UserParams))
-		.SetDisplayName(FPCGEditor_private::UserParamsTabName)
-		.SetGroup(WorkspaceMenuCategoryRef);
-
-	InTabManager->RegisterTabSpawner(FPCGEditor_private::ViewportID[0], FOnSpawnTab::CreateSP(this, &FPCGEditor::SpawnTab_Viewport, 0))
-		.SetDisplayName(LOCTEXT("ViewportTab1", "Data Viewport 1"))
-		.SetGroup(ViewportGroup);
+	struct TabInfo
+	{
+		FName ID;
+		SpawnOnTabCallback Callback;
+		FText DisplayName;
+		TSharedRef<FWorkspaceItem> Group;
+	};
 
-	InTabManager->RegisterTabSpawner(FPCGEditor_private::ViewportID[1], FOnSpawnTab::CreateSP(this, &FPCGEditor::SpawnTab_Viewport, 1))
-		.SetDisplayName(LOCTEXT("ViewportTab2", "Data Viewport 2"))
-		.SetGroup(ViewportGroup);
+	TabInfo DefaultTabInfo[] =
+	{
+		{FPCGEditor_private::GraphEditorID, FOnSpawnTab::CreateSP(this, &FPCGEditor::SpawnTab_GraphEditor), LOCTEXT("GraphTab", "Graph"), WorkspaceMenuCategoryRef},
+		{FPCGEditor_private::PropertyDetailsID[0], FOnSpawnTab::CreateSP(this, &FPCGEditor::SpawnTab_PropertyDetails, 0), LOCTEXT("DetailsTab1", "Details 1"), DetailsGroup},
+		{FPCGEditor_private::PropertyDetailsID[1], FOnSpawnTab::CreateSP(this, &FPCGEditor::SpawnTab_PropertyDetails, 1), LOCTEXT("DetailsTab2", "Details 2"), DetailsGroup},
+		{FPCGEditor_private::PropertyDetailsID[2], FOnSpawnTab::CreateSP(this, &FPCGEditor::SpawnTab_PropertyDetails, 2), LOCTEXT("DetailsTab3", "Details 3"), DetailsGroup},
+		{FPCGEditor_private::PropertyDetailsID[3], FOnSpawnTab::CreateSP(this, &FPCGEditor::SpawnTab_PropertyDetails, 3), LOCTEXT("DetailsTab4", "Details 4"), DetailsGroup},
+		{FPCGEditor_private::PaletteID, FOnSpawnTab::CreateSP(this, &FPCGEditor::SpawnTab_Palette), LOCTEXT("PaletteTab", "Palette"), WorkspaceMenuCategoryRef},
+		{FPCGEditor_private::DebugObjectID, FOnSpawnTab::CreateSP(this, &FPCGEditor::SpawnTab_DebugObjectTree), LOCTEXT("DebugTab", "Debug Object Tree"), WorkspaceMenuCategoryRef},
+		{FPCGEditor_private::AttributesID[0], FOnSpawnTab::CreateSP(this, &FPCGEditor::SpawnTab_Attributes, 0), LOCTEXT("AttributesTab1", "Attributes 1"), AttributesGroup},
+		{FPCGEditor_private::AttributesID[1], FOnSpawnTab::CreateSP(this, &FPCGEditor::SpawnTab_Attributes, 1), LOCTEXT("AttributesTab2", "Attributes 2"), AttributesGroup},
+		{FPCGEditor_private::AttributesID[2], FOnSpawnTab::CreateSP(this, &FPCGEditor::SpawnTab_Attributes, 2), LOCTEXT("AttributesTab3", "Attributes 3"), AttributesGroup},
+		{FPCGEditor_private::AttributesID[3], FOnSpawnTab::CreateSP(this, &FPCGEditor::SpawnTab_Attributes, 3), LOCTEXT("AttributesTab4", "Attributes 4"), AttributesGroup},
+		{FPCGEditor_private::FindID, FOnSpawnTab::CreateSP(this, &FPCGEditor::SpawnTab_Find), LOCTEXT("FindTab", "Find"), WorkspaceMenuCategoryRef},
+		{FPCGEditor_private::DeterminismID, FOnSpawnTab::CreateSP(this, &FPCGEditor::SpawnTab_Determinism), LOCTEXT("DeterminismTab", "Determinism"), WorkspaceMenuCategoryRef},
+		{FPCGEditor_private::ProfilingID, FOnSpawnTab::CreateSP(this, &FPCGEditor::SpawnTab_Profiling), LOCTEXT("ProfilingTab", "Profiling"), WorkspaceMenuCategoryRef},
+		{FPCGEditor_private::LogID, FOnSpawnTab::CreateSP(this, &FPCGEditor::SpawnTab_Log), LOCTEXT("LogCaptureTab", "Log Capture"), WorkspaceMenuCategoryRef},
+		{FPCGEditor_private::HLSLSourceID, FOnSpawnTab::CreateSP(this, &FPCGEditor::SpawnTab_NodeSource), LOCTEXT("NodeSourceTab", "HLSL Source"), WorkspaceMenuCategoryRef},
+		{FPCGEditor_private::UserParamsID, FOnSpawnTab::CreateSP(this, &FPCGEditor::SpawnTab_UserParams), FPCGEditor_private::UserParamsTabName, WorkspaceMenuCategoryRef},
+		{FPCGEditor_private::ViewportID[0], FOnSpawnTab::CreateSP(this, &FPCGEditor::SpawnTab_Viewport, 0), LOCTEXT("ViewportTab1", "Viewport 1"), ViewportGroup},
+		{FPCGEditor_private::ViewportID[1], FOnSpawnTab::CreateSP(this, &FPCGEditor::SpawnTab_Viewport, 1), LOCTEXT("ViewportTab2", "Viewport 2"), ViewportGroup},
+		{FPCGEditor_private::ViewportID[2], FOnSpawnTab::CreateSP(this, &FPCGEditor::SpawnTab_Viewport, 2), LOCTEXT("ViewportTab3", "Viewport 3"), ViewportGroup},
+		{FPCGEditor_private::ViewportID[3], FOnSpawnTab::CreateSP(this, &FPCGEditor::SpawnTab_Viewport, 3), LOCTEXT("ViewportTab4", "Viewport 4"), ViewportGroup},
+	};
 
-	InTabManager->RegisterTabSpawner(FPCGEditor_private::ViewportID[2], FOnSpawnTab::CreateSP(this, &FPCGEditor::SpawnTab_Viewport, 2))
-		.SetDisplayName(LOCTEXT("ViewportTab3", "Data Viewport 3"))
-		.SetGroup(ViewportGroup);
+	for (const TabInfo& Info : DefaultTabInfo)
+	{
+		if (!IsPanelAvailable(Info.ID))
+		{
+			continue;
+		}
 
-	InTabManager->RegisterTabSpawner(FPCGEditor_private::ViewportID[3], FOnSpawnTab::CreateSP(this, &FPCGEditor::SpawnTab_Viewport, 3))
-		.SetDisplayName(LOCTEXT("ViewportTab4", "Data Viewport 4"))
-		.SetGroup(ViewportGroup);
+		InTabManager->RegisterTabSpawner(Info.ID, Info.Callback)
+		.SetDisplayName(Info.DisplayName)
+		.SetGroup(Info.Group);
+	}
 }
 
 void FPCGEditor::UnregisterTabSpawners(const TSharedRef<class FTabManager>& InTabManager)
@@ -784,12 +783,15 @@ void FPCGEditor::RegisterToolbar() const
 			TAttribute<FText>(),
 			TAttribute<FText>(),
 			FSlateIcon(FPCGEditorStyle::Get().GetStyleSetName(), "PCG.Command.StopRegen")));
-				
-		Section.AddEntry(FToolMenuEntry::InitToolBarButton(
-			FPCGEditorCommands::Get().OpenDebugObjectTreeTab,
-			TAttribute<FText>(),
-			TAttribute<FText>(),
-			FSlateIcon(FPCGEditorStyle::Get().GetStyleSetName(), "PCG.Command.OpenDebugTreeTab")));
+
+		if (IsPanelAvailable(FPCGEditor_private::DebugObjectID))
+		{
+			Section.AddEntry(FToolMenuEntry::InitToolBarButton(
+				FPCGEditorCommands::Get().OpenDebugObjectTreeTab,
+				TAttribute<FText>(),
+				TAttribute<FText>(),
+				FSlateIcon(FPCGEditorStyle::Get().GetStyleSetName(), "PCG.Command.OpenDebugTreeTab")));
+		}
 
 		Section.AddSeparator(NAME_None);
 
@@ -1269,6 +1271,19 @@ void FPCGEditor::OnDeterminismGraphTest() const
 	}
 }
 
+TSubclassOf<UPCGEditorGraphSchema> FPCGEditor::GetSchemaClass() const
+{
+	return UPCGEditorGraphSchema::StaticClass();
+}
+
+TAttribute<FGraphAppearanceInfo> FPCGEditor::GetAppearanceInfo() const
+{
+	FGraphAppearanceInfo AppearanceInfo;
+	AppearanceInfo.CornerText = LOCTEXT("PCGGraphEditorCornerText", "PCG Graph");
+
+	return AppearanceInfo;
+}
+
 void FPCGEditor::OnEditGraphSettings()
 {
 	check(GraphEditorWidget)
@@ -2782,9 +2797,6 @@ TSharedRef<SGraphEditor> FPCGEditor::CreateGraphEditorWidget()
 		FExecuteAction::CreateSP(this, &FPCGEditor::OnDistributeNodesV)
 	);
 
-	FGraphAppearanceInfo AppearanceInfo;
-	AppearanceInfo.CornerText = LOCTEXT("PCGGraphEditorCornerText", "PCG Graph");
-
 	SGraphEditor::FGraphEditorEvents InEvents;
 	InEvents.OnSelectionChanged = SGraphEditor::FOnSelectionChanged::CreateSP(this, &FPCGEditor::OnSelectedNodesChanged);
 	InEvents.OnVerifyTextCommit = FOnNodeVerifyTextCommit::CreateSP(this, &FPCGEditor::OnValidateNodeTitle);
@@ -2796,7 +2808,7 @@ TSharedRef<SGraphEditor> FPCGEditor::CreateGraphEditorWidget()
 	return SNew(SGraphEditor)
 		.AdditionalCommands(GraphEditorCommands)
 		.IsEditable(true)
-		.Appearance(AppearanceInfo)
+		.Appearance(GetAppearanceInfo())
 		.GraphToEdit(PCGEditorGraph)
 		.GraphEvents(InEvents)
 		.ShowGraphStateOverlay(false);
@@ -2858,6 +2870,11 @@ void FPCGEditor::OnClose()
 
 	if (GEditor)
 	{
+		if (UPCGEditorSubsystem* EditorSubsystem = GEditor->GetEditorSubsystem<UPCGEditorSubsystem>())
+		{
+			EditorSubsystem->OnPCGSourceGenerationDone.RemoveAll(this);
+		}
+
 		UnregisterDelegatesForWorld(GEditor->GetEditorWorldContext().World());
 		UnregisterDelegatesForWorld(GEditor->PlayWorld.Get());
 	}
diff --git a/Engine/Plugins/PCG/Source/PCGEditor/Private/PCGEditor.h b/Engine/Plugins/PCG/Source/PCGEditor/Private/PCGEditor.h
deleted file mode 100644
diff --git a/Engine/Plugins/PCG/Source/PCGEditor/Private/PCGEditorCommands.cpp b/Engine/Plugins/PCG/Source/PCGEditor/Private/PCGEditorCommands.cpp
--- a/Engine/Plugins/PCG/Source/PCGEditor/Private/PCGEditorCommands.cpp
+++ b/Engine/Plugins/PCG/Source/PCGEditor/Private/PCGEditorCommands.cpp
@@ -2,7 +2,7 @@
 
 #include "PCGEditorCommands.h"
 
-#include "PCGEditorGraphSchemaActions.h"
+#include "Schema/PCGEditorGraphSchemaActions.h"
 
 #include "PCGSettings.h"
 
diff --git a/Engine/Plugins/PCG/Source/PCGEditor/Private/PCGEditorCommands.h b/Engine/Plugins/PCG/Source/PCGEditor/Private/PCGEditorCommands.h
deleted file mode 100644
diff --git a/Engine/Plugins/PCG/Source/PCGEditor/Private/PCGEditorCommon.h b/Engine/Plugins/PCG/Source/PCGEditor/Private/PCGEditorCommon.h
deleted file mode 100644
diff --git a/Engine/Plugins/PCG/Source/PCGEditor/Private/PCGEditorGraph.cpp b/Engine/Plugins/PCG/Source/PCGEditor/Private/PCGEditorGraph.cpp
--- a/Engine/Plugins/PCG/Source/PCGEditor/Private/PCGEditorGraph.cpp
+++ b/Engine/Plugins/PCG/Source/PCGEditor/Private/PCGEditorGraph.cpp
@@ -3,7 +3,6 @@
 #include "PCGEditorGraph.h"
 
 #include "PCGEdge.h"
-#include "PCGEditorGraphSchemaActions.h"
 #include "PCGEditorModule.h"
 #include "PCGGraph.h"
 #include "PCGPin.h"
@@ -15,6 +14,7 @@
 #include "Nodes/PCGEditorGraphNodeInput.h"
 #include "Nodes/PCGEditorGraphNodeOutput.h"
 #include "Nodes/PCGEditorGraphNodeReroute.h"
+#include "Schema/PCGEditorGraphSchemaActions.h"
 
 #include "EdGraph/EdGraphPin.h"
 
diff --git a/Engine/Plugins/PCG/Source/PCGEditor/Private/PCGEditorGraphSchema.cpp b/Engine/Plugins/PCG/Source/PCGEditor/Private/PCGEditorGraphSchema.cpp
--- a/Engine/Plugins/PCG/Source/PCGEditor/Private/PCGEditorGraphSchema.cpp
+++ b/Engine/Plugins/PCG/Source/PCGEditor/Private/PCGEditorGraphSchema.cpp
@@ -1,6 +1,6 @@
 // Copyright Epic Games, Inc. All Rights Reserved.
 
-#include "PCGEditorGraphSchema.h"
+#include "Schema/PCGEditorGraphSchema.h"
 
 #include "PCGComponent.h"
 #include "PCGDataAsset.h"
@@ -18,12 +18,12 @@
 #include "PCGEditor.h"
 #include "PCGEditorCommon.h"
 #include "PCGEditorGraph.h"
-#include "PCGEditorGraphSchemaActions.h"
 #include "PCGEditorSettings.h"
 #include "PCGEditorUtils.h"
 #include "Helpers/PCGPropertyHelpers.h"
 #include "Nodes/PCGEditorGraphNodeBase.h"
 #include "Nodes/PCGEditorGraphNodeReroute.h"
+#include "Schema/PCGEditorGraphSchemaActions.h"
 
 #include "AssetRegistry/AssetData.h"
 #include "AssetRegistry/AssetRegistryModule.h"
@@ -37,14 +37,11 @@
 
 #define LOCTEXT_NAMESPACE "PCGEditorGraphSchema"
 
-FPCGActionsFilter::FPCGActionsFilter(const UEdGraph* InEdGraph, EPCGElementType InElementFilterType)
+FPCGActionsFilter::FPCGActionsFilter(const UEdGraph* InEdGraph, EPCGElementType InElementFilterType, FPCGGraphEditorCustomization InCustomization)
 	: FilterType(InElementFilterType)
 	, Graph(Cast<UPCGEditorGraph>(InEdGraph))
+	, Customization(MoveTemp(InCustomization))
 {
-	if (Graph && Graph->GetPCGGraph())
-	{
-		Customization = Graph->GetPCGGraph()->GraphCustomization;
-	}
 }
 
 bool FPCGActionsFilter::Accepts(const FText& InCategory) const
@@ -57,10 +54,20 @@ void UPCGEditorGraphSchema::GetPaletteActions(FGraphActionMenuBuilder& ActionMen
 	GetGraphActions(ActionMenuBuilder, ActionFilter, /*bIsContextual=*/false);
 }
 
+const FPCGGraphEditorCustomization& UPCGEditorGraphSchema::GetGraphEditorCustomization(const UEdGraph* InEdGraph) const
+{
+	static const FPCGGraphEditorCustomization Default{};
+	const UPCGEditorGraph* PCGGraph = Cast<UPCGEditorGraph>(InEdGraph);
+	return (PCGGraph && PCGGraph->GetPCGGraph()) ? PCGGraph->GetPCGGraph()->GraphCustomization : Default;
+}
+
 void UPCGEditorGraphSchema::GetGraphContextActions(FGraphContextMenuBuilder& ContextMenuBuilder) const
 {
 	Super::GetGraphContextActions(ContextMenuBuilder);
-	GetGraphActions(ContextMenuBuilder, FPCGActionsFilter(ContextMenuBuilder.CurrentGraph), /*bIsContextual=*/true);
+
+	FPCGActionsFilter ActionFilter(ContextMenuBuilder.CurrentGraph, GetElementTypeFiltering(), GetGraphEditorCustomization(ContextMenuBuilder.CurrentGraph));
+
+	GetGraphActions(ContextMenuBuilder, ActionFilter, /*bIsContextual=*/true);
 }
 
 void UPCGEditorGraphSchema::GetGraphActions(FGraphActionMenuBuilder& ActionMenuBuilder, const FPCGActionsFilter& ActionFilter, bool bIsContextual) const
@@ -413,10 +420,11 @@ void UPCGEditorGraphSchema::GetNativeElementActions(FGraphActionMenuBuilder& Act
 	TArray<UClass*> SettingsClasses;
 	for (TObjectIterator<UClass> It; It; ++It)
 	{
-		UClass* Class = *It;
+		TSubclassOf<UPCGSettings> Class(*It);
 
-		if (Class->IsChildOf(UPCGSettings::StaticClass()) &&
-			!Class->HasAnyClassFlags(CLASS_Abstract | CLASS_Deprecated | CLASS_Hidden))
+		if (*Class &&
+			!Class->HasAnyClassFlags(CLASS_Abstract | CLASS_Deprecated | CLASS_Hidden) &&
+			(!ActionFilter.Customization.FiltersSettings() || !ActionFilter.Customization.FilterSettings(Class)))
 		{
 			SettingsClasses.Add(Class);
 		}
diff --git a/Engine/Plugins/PCG/Source/PCGEditor/Private/PCGEditorGraphSchema.h b/Engine/Plugins/PCG/Source/PCGEditor/Private/PCGEditorGraphSchema.h
deleted file mode 100644
diff --git a/Engine/Plugins/PCG/Source/PCGEditor/Private/PCGEditorGraphSchemaActions.cpp b/Engine/Plugins/PCG/Source/PCGEditor/Private/PCGEditorGraphSchemaActions.cpp
--- a/Engine/Plugins/PCG/Source/PCGEditor/Private/PCGEditorGraphSchemaActions.cpp
+++ b/Engine/Plugins/PCG/Source/PCGEditor/Private/PCGEditorGraphSchemaActions.cpp
@@ -1,14 +1,14 @@
 // Copyright Epic Games, Inc. All Rights Reserved.
 
-#include "PCGEditorGraphSchemaActions.h"
+#include "Schema/PCGEditorGraphSchemaActions.h"
 
 #include "PCGEditorCommon.h"
 #include "PCGEditorGraph.h"
-#include "PCGEditorGraphSchema.h"
 #include "PCGEditorModule.h"
 #include "Nodes/PCGEditorGraphNode.h"
 #include "Nodes/PCGEditorGraphNodeComment.h"
 #include "Nodes/PCGEditorGraphNodeReroute.h"
+#include "Schema/PCGEditorGraphSchema.h"
 
 #include "PCGGraph.h"
 #include "PCGNode.h"
diff --git a/Engine/Plugins/PCG/Source/PCGEditor/Private/PCGEditorGraphSchemaActions.h b/Engine/Plugins/PCG/Source/PCGEditor/Private/PCGEditorGraphSchemaActions.h
deleted file mode 100644
diff --git a/Engine/Plugins/PCG/Source/PCGEditor/Private/PCGSettingsDragDropAction.cpp b/Engine/Plugins/PCG/Source/PCGEditor/Private/PCGSettingsDragDropAction.cpp
--- a/Engine/Plugins/PCG/Source/PCGEditor/Private/PCGSettingsDragDropAction.cpp
+++ b/Engine/Plugins/PCG/Source/PCGEditor/Private/PCGSettingsDragDropAction.cpp
@@ -3,8 +3,8 @@
 #include "PCGSettingsDragDropAction.h"
 
 #include "PCGEditorGraph.h"
-#include "PCGEditorGraphSchema.h"
-#include "PCGEditorGraphSchemaActions.h"
+#include "Schema/PCGEditorGraphSchema.h"
+#include "Schema/PCGEditorGraphSchemaActions.h"
 
 FReply FPCGSettingsDragDropAction::DroppedOnPanel(const TSharedRef<class SWidget>& Panel, const FVector2f& ScreenPosition, const FVector2f& GraphPosition, UEdGraph& Graph)
 {
diff --git a/Engine/Plugins/PCG/Source/PCGEditor/Private/Widgets/SPCGEditorGraphNodePalette.cpp b/Engine/Plugins/PCG/Source/PCGEditor/Private/Widgets/SPCGEditorGraphNodePalette.cpp
--- a/Engine/Plugins/PCG/Source/PCGEditor/Private/Widgets/SPCGEditorGraphNodePalette.cpp
+++ b/Engine/Plugins/PCG/Source/PCGEditor/Private/Widgets/SPCGEditorGraphNodePalette.cpp
@@ -6,10 +6,10 @@
 
 #include "PCGEditor.h"
 #include "PCGEditorGraph.h"
-#include "PCGEditorGraphSchema.h"
-#include "PCGEditorGraphSchemaActions.h"
 #include "PCGEditorUtils.h"
 #include "PCGSettingsDragDropAction.h"
+#include "Schema/PCGEditorGraphSchema.h"
+#include "Schema/PCGEditorGraphSchemaActions.h"
 #include "Widgets/SPCGEditorGraphActionWidget.h"
 
 #include "SEnumCombo.h"
@@ -130,7 +130,11 @@ void SPCGEditorGraphNodePalette::CollectAllActions(FGraphActionListBuilderBase&
 	const UPCGEditorGraphSchema* PCGSchema = GetDefault<UPCGEditorGraphSchema>();
 
 	FGraphActionMenuBuilder ActionMenuBuilder;
-	PCGSchema->GetPaletteActions(ActionMenuBuilder, FPCGActionsFilter(PCGEditor.Pin()->GetPCGEditorGraph(), ElementType));
+	FPCGEditor* Editor = PCGEditor.Pin().Get();
+	check(Editor);
+	FPCGActionsFilter ActionFilter(Editor->GetPCGEditorGraph(), ElementType & PCGSchema->GetElementTypeFiltering(), PCGSchema->GetGraphEditorCustomization(Editor->GetPCGEditorGraph()));
+	
+	PCGSchema->GetPaletteActions(ActionMenuBuilder, ActionFilter);
 	OutAllActions.Append(ActionMenuBuilder);
 }
 
diff --git a/Engine/Plugins/PCG/Source/PCGEditor/Public/PCGEditor.h b/Engine/Plugins/PCG/Source/PCGEditor/Public/PCGEditor.h
new file mode 100644
--- /dev/null
+++ b/Engine/Plugins/PCG/Source/PCGEditor/Public/PCGEditor.h
@@ -0,0 +1,487 @@
+// Copyright Epic Games, Inc. All Rights Reserved.
+
+#pragma once
+
+#include "PCGContext.h"
+#include "Graph/PCGStackContext.h"
+
+#include "EditorUndoClient.h"
+#include "GraphEditor.h"
+#include "Schema/PCGEditorGraphSchema.h"
+#include "Toolkits/AssetEditorToolkit.h"
+
+enum class ECheckBoxState : uint8;
+namespace ETextCommit { enum Type : int; }
+
+class IPCGBaseSubsystem;
+class FUICommandList;
+class SGraphEditor;
+class SGraphEditorActionMenu;
+class SPCGEditorGraphAttributeListView;
+class SPCGEditorGraphDebugObjectTree;
+class SPCGEditorGraphDetailsView;
+class SPCGEditorGraphDeterminismListView;
+class SPCGEditorGraphFind;
+class SPCGEditorGraphLogView;
+class SPCGEditorGraphNodePalette;
+class SPCGEditorGraphProfilingView;
+class SPCGEditorNodeSource;
+class SPCGEditorGraphUserParametersView;
+class SPCGEditorViewport;
+class UEdGraphNode;
+class UPCGComponent;
+class UPCGEditorGraph;
+class UPCGEditorGraphNodeBase;
+class UPCGGraph;
+class UPCGSubsystem;
+struct FPCGCompilerDiagnostics;
+
+DECLARE_MULTICAST_DELEGATE_OneParam(FOnInspectedStackChanged, const FPCGStack&);
+DECLARE_MULTICAST_DELEGATE_OneParam(FOnInspectedNodeChanged, UPCGEditorGraphNodeBase*);
+
+enum class EPCGEditorPanel
+{
+	Attributes1,
+	Attributes2,
+	Attributes3,
+	Attributes4,
+	DebugObjectTree,
+	Determinism,
+	Find,
+	GraphEditor,
+	Log,
+	NodePalette,
+	NodeSource,
+	Profiling,
+	PropertyDetails1,
+	PropertyDetails2,
+	PropertyDetails3,
+	PropertyDetails4,
+	UserParams,
+	Viewport1,
+	Viewport2,
+	Viewport3,
+	Viewport4
+};
+
+class FPCGEditor : public FAssetEditorToolkit, public FGCObject, public FSelfRegisteringEditorUndoClient
+{
+public:
+	/** Edits the specified PCGGraph */
+	PCGEDITOR_API void Initialize(const EToolkitMode::Type InMode, const TSharedPtr<class IToolkitHost>& InToolkitHost, UPCGGraph* InPCGGraph);
+
+	/** Get the PCG editor graph being edited. */
+	PCGEDITOR_API UPCGEditorGraph* GetPCGEditorGraph();
+
+	/** Gets/Creates the PCG graph editor for a given PCG graph */
+	PCGEDITOR_API static UPCGEditorGraph* GetPCGEditorGraph(UPCGGraph* InGraph);
+	PCGEDITOR_API static UPCGEditorGraph* GetPCGEditorGraph(const UPCGNode* InNode);
+	PCGEDITOR_API static UPCGEditorGraph* GetPCGEditorGraph(const UPCGSettings* InSettings);
+
+	/** Get the PCG graph being edited. */
+	const UPCGGraph* GetPCGGraph() { return PCGGraphBeingEdited; }
+
+	/** Sets the execution stack that want to inspect. */
+	PCGEDITOR_API void SetStackBeingInspected(const FPCGStack& FullStack);
+
+	/** Sets the execution stack from another editor, which will set directly in the debug object tree view. */
+	PCGEDITOR_API void SetStackBeingInspectedFromAnotherEditor(const FPCGStack& FullStack);
+
+	/** Clear current inspection. */
+	PCGEDITOR_API void ClearStackBeingInspected();
+
+	/** Gets the PCG source we are debugging */
+	UPCGComponent* GetPCGComponentBeingInspected() const { return const_cast<UPCGComponent*>(StackBeingInspected.GetRootComponent()); }
+	IPCGGraphExecutionSource* GetPCGSourceBeingInspected() const { return const_cast<IPCGGraphExecutionSource*>(StackBeingInspected.GetRootSource()); }
+	
+	/** Gets the PCG stack we are inspecting */
+	PCGEDITOR_API const FPCGStack* GetStackBeingInspected() const;
+
+	PCGEDITOR_API void SetSourceEditorTargetObject(UObject* InObject);
+
+	/** Focus the graph view on a specific node */
+	PCGEDITOR_API void JumpToNode(const UEdGraphNode* InNode);
+	PCGEDITOR_API void JumpToNode(const UPCGNode* InNode);
+
+	/** Get the TabID of the editor panel. */
+	PCGEDITOR_API FName GetPanelID(EPCGEditorPanel Panel) const;
+	/** Focuses the user on a specific panel and flashes the tab. */
+	void BringFocusToPanel(EPCGEditorPanel Panel) const { return BringFocusToPanel(GetPanelID(Panel)); }
+	/** Attempts to close the specific panel if it's open. */
+	void CloseGraphPanel(EPCGEditorPanel Panel) const { CloseGraphPanel(GetPanelID(Panel)); }
+	/** Returns true if the selected tab is currently open. */
+	bool IsPanelCurrentlyOpen(EPCGEditorPanel Panel) const { return IsPanelCurrentlyOpen(GetPanelID(Panel)); }
+	/** Returns true if the selected tab is currently open and focused. */
+	bool IsPanelCurrentlyForeground(EPCGEditorPanel Panel) const { return IsPanelCurrentlyForeground(GetPanelID(Panel)); }
+
+	/** Focuses the user on a specific panel and flashes the tab. */
+	PCGEDITOR_API void BringFocusToPanel(const FName PanelID) const;
+	/** Attempts to close the specific panel if it's open. */
+	PCGEDITOR_API void CloseGraphPanel(const FName PanelID) const;
+	/** Returns true if the selected tab is currently open. */
+	PCGEDITOR_API bool IsPanelCurrentlyOpen(const FName PanelID) const;
+	/** Returns true if the selected tab is currently open and focused. */
+	PCGEDITOR_API bool IsPanelCurrentlyForeground(const FName PanelID) const;
+
+	/** Returns true if the panel is available. */
+	PCGEDITOR_API virtual bool IsPanelAvailable(const FName PanelID) const;
+
+	/** Helper to get to the subsystem. */
+	PCGEDITOR_API virtual IPCGBaseSubsystem* GetSubsystem() const;
+	PCGEDITOR_API static UPCGSubsystem* GetWorldSubsystem();
+
+	// ~Begin IToolkit interface
+	PCGEDITOR_API virtual void RegisterTabSpawners(const TSharedRef<class FTabManager>& InTabManager) override;
+	PCGEDITOR_API virtual void UnregisterTabSpawners(const TSharedRef<class FTabManager>& InTabManager) override;
+	// ~End IToolkit interface
+
+	// ~Begin FGCObject interface
+	PCGEDITOR_API virtual void AddReferencedObjects(FReferenceCollector& Collector) override;
+	virtual FString GetReferencerName() const override
+	{
+		return TEXT("FPCGEditor");
+	}
+	// ~End FGCObject interface
+	
+	// ~Begin FEditorUndoClient interface
+	PCGEDITOR_API virtual bool MatchesContext(const FTransactionContext& InContext, const TArray<TPair<UObject*, FTransactionObjectEvent>>& TransactionObjectContexts) const override;
+	PCGEDITOR_API virtual void PostUndo(bool bSuccess) override;
+	virtual void PostRedo(bool bSuccess) override { PostUndo(bSuccess); }
+	// ~End FEditorUndoClient interface
+
+	// ~Begin FAssetEditorToolkit interface
+	PCGEDITOR_API virtual FName GetToolkitFName() const override;
+	PCGEDITOR_API virtual FText GetBaseToolkitName() const override;
+	PCGEDITOR_API virtual FLinearColor GetWorldCentricTabColorScale() const override;
+	PCGEDITOR_API virtual FString GetWorldCentricTabPrefix() const override;
+	PCGEDITOR_API virtual void OnClose() override;
+	PCGEDITOR_API virtual void InitToolMenuContext(FToolMenuContext& MenuContext) override;
+	// ~End FAssetEditorToolkit interface
+
+	/**
+	 * Handles spawning a graph node in the current graph using the passed in chord.
+	 *
+	 * @param	InChord		Chord that was just performed.
+	 * @param	InPosition	Current cursor position.
+	 * @param	InGraph		Graph that chord was performed in.
+	 *
+	 * @return	FReply	Whether chord was handled.
+	 */
+	UE_DEPRECATED(5.6, "Please use the version of the function accepting FVector2f.")
+	FReply OnSpawnNodeByShortcut(FInputChord InChord, const FVector2D& InPosition, UPCGEditorGraph* InGraph);
+	PCGEDITOR_API FReply OnSpawnNodeByShortcut(FInputChord InChord, const FVector2f& InPosition, UPCGEditorGraph* InGraph);
+
+	FActionMenuContent OnCreateActionMenuContent(UEdGraph* InGraph, const FVector2f& Location, const TArray<UEdGraphPin*>& InDraggedPins, bool bAutoExpand, SGraphEditor::FActionMenuClosed OnMenuClosed);
+
+	FOnInspectedStackChanged OnInspectedStackChangedDelegate;
+
+	/** Can determinism be tested on the current graph */
+	bool CanRunDeterminismGraphTest() const;
+	/** Run the determinism test on the current graph */
+	void OnDeterminismGraphTest() const;
+
+	// Can override the schema used for this editor. By default, it's UPCGEditorSchema.
+	PCGEDITOR_API virtual TSubclassOf<UPCGEditorGraphSchema> GetSchemaClass() const;
+
+protected:
+	PCGEDITOR_API virtual TAttribute<FGraphAppearanceInfo> GetAppearanceInfo() const;
+
+	/** Register PCG specific toolbar for the editor */
+	PCGEDITOR_API virtual void RegisterToolbar() const;
+
+	/** Bind commands to delegates */
+	PCGEDITOR_API virtual void BindCommands();
+
+	PCGEDITOR_API virtual void OnGraphChanged(UPCGGraphInterface* InGraph, EPCGChangeType ChangeType);
+
+	/** Called when the selection changes in the GraphEditor */
+	PCGEDITOR_API virtual void OnSelectedNodesChanged(const TSet<UObject*>& NewSelection);
+
+	/**
+	 * Called when a node is double clicked
+	 *
+	 * @param Node - The Node that was clicked
+	 */
+	PCGEDITOR_API virtual void OnNodeDoubleClicked(UEdGraphNode* Node);
+
+private:
+	/** Bring up the find tab */
+	void OnFind();
+
+	/** Bring up the first free details view, or if they are all locked, the first details view */
+	void OpenDetailsView();
+
+	/** Called when a details view tab is closed */
+	void OnDetailsViewTabClosed(TSharedRef<SDockTab> DockTab, int Index);
+
+	/** Called when an attribute list view tab is closed */
+	void OnAttributeListViewTabClosed(TSharedRef<SDockTab> DockTab, int Index);
+
+	/** Called when a viewport view tab is closed */
+	void OnViewportViewTabClosed(TSharedRef<SDockTab> DockTab, int Index);
+
+	/** Enable/Disable automatic PCG node generation */
+	void OnPauseAutomaticRegeneration_Clicked();
+	/** Has the user paused automatic regeneration in the Graph Editor */
+	bool IsAutomaticRegenerationPaused() const;
+
+	/** Force a regeneration by invoking the graph notifications  */
+	void OnForceGraphRegeneration_Clicked();
+
+	/** Toggle node inspection state for selected nodes */
+	void OnToggleInspected();
+	/** Whether we can toggle inspection of selected nodes */
+	bool CanToggleInspected() const;
+	/** Whether selected nodes are inspected or not */
+	ECheckBoxState GetInspectedCheckState() const;
+
+	void UpdateAfterInspectedStackChanged();
+	
+	/** Toggle node enabled state for selected nodes */
+	void OnToggleEnabled();
+	/** Whether we can toggle enabled state of selected nodes */
+	bool CanToggleEnabled() const;
+	/** Whether selected nodes are enabled or not */
+	ECheckBoxState GetEnabledCheckState() const;
+	
+	/** Toggle node debug state for selected nodes */
+	void OnToggleDebug();
+	/** Whether we can toggle debug state of selected nodes */
+	bool CanToggleDebug() const;
+	/** Whether selected nodes are being debugged or not */
+	ECheckBoxState GetDebugCheckState() const;
+
+	/** Enable node debug state for selected nodes and disable for others */
+	void OnDebugOnlySelected();
+
+	/** Disable node debug state for all nodes */
+	void OnDisableDebugOnAllNodes();
+
+	/** Cancels the current execution of the selected graph */
+	void OnCancelExecution_Clicked();
+
+	/** Returns true if inspected graph is currently scheduled or executing */
+	bool IsCurrentlyGenerating() const;
+
+	/** Returns true if the debug object tree tab is not currently open. */
+	bool IsDebugObjectTreeTabClosed() const;
+
+	/** Opens the debug object tree tab if it is not open already. */
+	void OnOpenDebugObjectTreeTab_Clicked();
+
+	/** Can determinism be tested on the selected node(s) */
+	bool CanRunDeterminismNodeTest() const;
+	/** Run the determinism test on the selected node(s) */
+	void OnDeterminismNodeTest() const;
+
+	/** Open details view for the PCG object being edited */
+	void OnEditGraphSettings();
+	/** Whether the PCG object being edited is opened in details view or not */
+	bool IsEditGraphSettingsToggled() const;
+
+	/** Open panel to view and edit the graph parameters. */
+	void OnToggleGraphParamsPanel() const;
+	/** Is the graph params panel is open. */
+	bool IsToggleGraphParamsToggled() const;
+
+	/** Select every node in the graph */
+	void SelectAllNodes();
+	/** Whether we can select every node */
+	bool CanSelectAllNodes() const;
+
+	/** Delete all selected nodes in the graph */
+	void DeleteSelectedNodes();
+	/** Whether we can delete all selected nodes */
+	bool CanDeleteSelectedNodes() const;
+
+	/** Copy all selected nodes in the graph */
+	void CopySelectedNodes();
+	/** Whether we can copy all selected nodes */
+	bool CanCopySelectedNodes() const;
+
+	/** Cut all selected nodes in the graph */
+	void CutSelectedNodes();
+	/** Whether we can cut all selected nodes */
+	bool CanCutSelectedNodes() const;
+
+	/** Paste nodes in the graph */
+	void PasteNodes();
+	/** Paste nodes in the graph at location*/
+	void PasteNodesHere(const FVector2D& Location);
+	/** Whether we can paste nodes */
+	bool CanPasteNodes() const;
+
+	/** Duplicate the currently selected nodes */
+	void DuplicateNodes();
+	/** Whether we are able to duplicate the currently selected nodes */
+	bool CanDuplicateNodes() const;
+
+	/** Exports node settings to assets */
+	void OnExportNodes();
+
+	/** Whether we are able to export the currently selected nodes */
+	bool CanExportNodes() const;
+
+	/** Converts instanced nodes to independent nodes */
+	void OnConvertToStandaloneNodes();
+
+	/** Whether we are able to convert the selected nodes to standalone */
+	bool CanConvertToStandaloneNodes() const;
+
+	/** Collapse the currently selected nodes in a subgraph */
+	void OnCollapseNodesInSubgraph();
+	/** Whether we can collapse nodes in a subgraph */
+	bool CanCollapseNodesInSubgraph() const;
+
+	/** User is attempting to add a dynamic source pin to a node */
+	void OnAddDynamicInputPin();
+	/** Whether the user can add a dynamic source pin to a node */
+	bool CanAddDynamicInputPin() const;
+
+	/** User is attempting to rename a node */
+	void OnRenameNode();
+	/** Whether the user can rename the selected node */
+	bool CanRenameNode() const;
+
+	/** Selects the associated usages of a given reroute declaration */
+	void OnSelectNamedRerouteUsages();
+	/** Whether the user can find the usages from the selection */
+	bool CanSelectNamedRerouteUsages() const;
+
+	/** Selects the associated declaration of a given reroute usage */
+	void OnSelectNamedRerouteDeclaration();
+	/** Whether the user can find the declaration from the selection */
+	bool CanSelectNamedRerouteDeclaration() const;
+
+	/** Jumps to source definition for the selected nodes. */
+	void OnJumpToSource();
+
+	/** Internal method that validates a few things (& logs errors) prior to executing actions. */
+	bool InternalValidationOnAction();
+
+	/** Finds editor graph node that matches the provided PCG node */
+	UPCGEditorGraphNodeBase* GetEditorNode(const UPCGNode* InNode);
+
+	void OnAlignTop();
+	void OnAlignMiddle();
+	void OnAlignBottom();
+	void OnAlignLeft();
+	void OnAlignCenter();
+	void OnAlignRight();
+	void OnStraightenConnections();
+	void OnDistributeNodesH();
+	void OnDistributeNodesV();
+	void OnCreateComment();
+
+	/** Create new graph editor widget */
+	TSharedRef<SGraphEditor> CreateGraphEditorWidget();
+
+	/** Create new palette widget */
+	TSharedRef<SPCGEditorGraphNodePalette> CreatePaletteWidget();
+
+	/** Create new debug object tree widget */
+	TSharedRef<SPCGEditorGraphDebugObjectTree> CreateDebugObjectTreeWidget();
+
+	/** Create new find widget */
+	TSharedRef<SPCGEditorGraphFind> CreateFindWidget();
+
+	/** Create new attributes widget */
+	TSharedRef<SPCGEditorGraphAttributeListView> CreateAttributesWidget();
+
+	/** Create a new determinism tab widget */
+	TSharedRef<SPCGEditorGraphDeterminismListView> CreateDeterminismWidget();
+
+	/** Create a new profiling tab widget */
+	TSharedRef<SPCGEditorGraphProfilingView> CreateProfilingWidget();
+
+	/** Create a new log capture tab widget */
+	TSharedRef<SPCGEditorGraphLogView> CreateLogWidget();
+
+	/** Create a new node source editor tab widget */
+	TSharedRef<SPCGEditorNodeSource> CreateNodeSourceWidget();
+
+	/** Create a new user graph parameters tab widget */
+	TSharedRef<SPCGEditorGraphUserParametersView> CreateGraphParamsWidget();
+
+	/** Create a new viewport widget */
+	TSharedRef<SPCGEditorViewport> CreateViewportWidget();
+
+	/** Called when the component inspected is generated/cleaned */
+	void OnSourceGenerated(IPCGGraphExecutionSource* InSource);
+
+	/** Called to validate a node title change. */
+	bool OnValidateNodeTitle(const FText& NewName, UEdGraphNode* GraphNode, FText& OutErrorMessage);
+
+	/** Called when the title of a node is changed */
+	void OnNodeTitleCommitted(const FText& NewText, ETextCommit::Type CommitInfo, UEdGraphNode* NodeBeingChanged);
+
+	/**
+	 * Try to jump to a given class (if allowed)
+	 *
+	 * @param Class - The Class to jump to
+	 */
+	void JumpToDefinition(const UClass* Class) const;
+
+	/** Called when a PCG component unregisters. */
+	void OnComponentUnregistered(UPCGComponent* Component);
+
+	/** Called when a component finishes executing. Useful for updating debugging tools/UIs. */
+	void OnSourceGenerationDone(IPCGBaseSubsystem* Subsystem, IPCGGraphExecutionSource* Source, EPCGGenerationStatus Status);
+
+	/** Trigger any generation required to ensure debug display is up to date. */
+	void UpdateDebugAfterComponentSelection(UPCGComponent* InOldComponent, UPCGComponent* InNewComponent, bool bNewComponentStartedInspecting);
+
+	void RegisterDelegatesForWorld(UWorld* World);
+	void UnregisterDelegatesForWorld(UWorld* World);
+
+	void OnNodeSourceCompiled(const UPCGNode* InNode, const FPCGCompilerDiagnostics& InDiagnostics);
+
+	void OnMapChanged(UWorld* InWorld, EMapChangeType InMapChangedType);
+	void OnPostPIEStarted(bool bIsSimulating);
+	void OnEndPIE(bool bIsSimulating);
+	void OnLevelActorDeleted(AActor* InActor);
+
+	void RefreshViewsOnLevelChange();
+
+	TSharedRef<SDockTab> SpawnTab_GraphEditor(const FSpawnTabArgs& Args);
+	TSharedRef<SDockTab> SpawnTab_PropertyDetails(const FSpawnTabArgs& Args, int PropertyDetailsIndex);
+	TSharedRef<SDockTab> SpawnTab_Palette(const FSpawnTabArgs& Args);
+	TSharedRef<SDockTab> SpawnTab_DebugObjectTree(const FSpawnTabArgs& Args);
+	TSharedRef<SDockTab> SpawnTab_Attributes(const FSpawnTabArgs& Args, int AttributesIndex);
+	TSharedRef<SDockTab> SpawnTab_Find(const FSpawnTabArgs& Args);
+	TSharedRef<SDockTab> SpawnTab_Determinism(const FSpawnTabArgs& Args);
+	TSharedRef<SDockTab> SpawnTab_Profiling(const FSpawnTabArgs& Args);
+	TSharedRef<SDockTab> SpawnTab_Log(const FSpawnTabArgs& Args);
+	TSharedRef<SDockTab> SpawnTab_NodeSource(const FSpawnTabArgs& Args);
+	TSharedRef<SDockTab> SpawnTab_UserParams(const FSpawnTabArgs& Args);
+	TSharedRef<SDockTab> SpawnTab_Viewport(const FSpawnTabArgs& Args, int ViewportIndex);
+
+	FText GetDetailsTabLabel(int DetailsIndex);
+	FText GetDetailsViewObjectName(int DetailsIndex);
+	FText GetAttributesTabLabel(int AttributesIndex);
+	FText GetViewportTabLabel(int ViewportIndex);
+
+	PCGEDITOR_API virtual TSharedRef<FTabManager::FLayout> GetDefaultLayout() const;
+
+	TSharedPtr<SGraphEditor> GraphEditorWidget;
+	TArray<TSharedPtr<SPCGEditorGraphDetailsView>> PropertyDetailsWidgets;
+	TSharedPtr<SPCGEditorGraphNodePalette> PaletteWidget;
+	TSharedPtr<SPCGEditorGraphDebugObjectTree> DebugObjectTreeWidget;
+	TSharedPtr<SPCGEditorGraphFind> FindWidget;
+	TArray<TSharedPtr<SPCGEditorGraphAttributeListView>> AttributesWidgets;
+	TSharedPtr<SPCGEditorGraphDeterminismListView> DeterminismWidget;
+	TSharedPtr<SPCGEditorGraphProfilingView> ProfilingWidget;
+	TSharedPtr<SPCGEditorGraphLogView> LogWidget;
+	TSharedPtr<SPCGEditorNodeSource> NodeSourceWidget;
+	TSharedPtr<SPCGEditorGraphUserParametersView> UserParamsWidget;
+
+	TSharedPtr<FUICommandList> GraphEditorCommands;
+
+	TObjectPtr<UPCGGraph> PCGGraphBeingEdited = nullptr;
+	UPCGEditorGraph* PCGEditorGraph = nullptr;
+
+	FPCGSoftGraphExecutionSource PCGSourceBeingInspected;
+	// Implementation note: we'll keep the last valid source inspected so we don't un-inspect on spurious selection changes
+	FPCGSoftGraphExecutionSource LastValidPCGSourceBeingInspected;
+	FPCGStack StackBeingInspected;
+};
diff --git a/Engine/Plugins/PCG/Source/PCGEditor/Public/PCGEditorCommands.h b/Engine/Plugins/PCG/Source/PCGEditor/Public/PCGEditorCommands.h
new file mode 100644
--- /dev/null
+++ b/Engine/Plugins/PCG/Source/PCGEditor/Public/PCGEditorCommands.h
@@ -0,0 +1,99 @@
+// Copyright Epic Games, Inc. All Rights Reserved.
+
+#pragma once
+
+#include "PCGSettings.h"
+#include "Framework/Commands/Commands.h"
+#include "Styling/AppStyle.h"
+
+class UPCGSettings;
+class UEdGraph;
+struct FEdGraphSchemaAction;
+
+class FPCGEditorCommands : public TCommands<FPCGEditorCommands>
+{
+public:
+	PCGEDITOR_API FPCGEditorCommands();
+
+	// ~Begin TCommands<> interface
+	PCGEDITOR_API virtual void RegisterCommands() override;
+	// ~End TCommands<> interface
+
+	TSharedPtr<FUICommandInfo> CollapseNodes;
+	TSharedPtr<FUICommandInfo> ExportNodes;
+	TSharedPtr<FUICommandInfo> ConvertToStandaloneNodes;
+	TSharedPtr<FUICommandInfo> Find;
+	TSharedPtr<FUICommandInfo> ShowSelectedDetails;
+	TSharedPtr<FUICommandInfo> PauseAutoRegeneration;
+	TSharedPtr<FUICommandInfo> ForceGraphRegeneration;
+	TSharedPtr<FUICommandInfo> OpenDebugObjectTreeTab;
+	TSharedPtr<FUICommandInfo> RunDeterminismNodeTest;
+	TSharedPtr<FUICommandInfo> RunDeterminismGraphTest;
+	TSharedPtr<FUICommandInfo> EditGraphSettings;
+	TSharedPtr<FUICommandInfo> ToggleGraphParams;
+	TSharedPtr<FUICommandInfo> CancelExecution;
+	TSharedPtr<FUICommandInfo> ToggleEnabled;
+	TSharedPtr<FUICommandInfo> ToggleDebug;
+	TSharedPtr<FUICommandInfo> DebugOnlySelected;
+	TSharedPtr<FUICommandInfo> DisableDebugOnAllNodes;
+	TSharedPtr<FUICommandInfo> ToggleInspect;
+	TSharedPtr<FUICommandInfo> AddSourcePin;
+	TSharedPtr<FUICommandInfo> RenameNode;
+	TSharedPtr<FUICommandInfo> SelectNamedRerouteUsages;
+	TSharedPtr<FUICommandInfo> SelectNamedRerouteDeclaration;
+	TSharedPtr<FUICommandInfo> JumpToSource;
+};
+
+struct FPCGSpawnNodeCommandInfo
+{
+	explicit FPCGSpawnNodeCommandInfo(TSubclassOf<UPCGSettings> PCGSettingsClass) : PCGSettingsClass(PCGSettingsClass) {}
+
+	/** Optional preconfigured settings info for a specific configuration of the node. */
+	FPCGPreConfiguredSettingsInfo PreconfiguredInfo;
+
+	/** Holds the UI Command to verify chords for this action are held. */
+	TSharedPtr<FUICommandInfo> CommandInfo;
+
+	/**
+	 * Creates an action to be used for placing a node into the graph.
+	 *
+	 * @return A fully prepared action containing the information to spawn the node.
+	 */
+	TSharedPtr<FEdGraphSchemaAction> GetAction() const;
+
+	/**
+	 * Gets the PCG Settings class of the spawn node action.
+	 *
+	 * @return The PCG Settings class of the spawn node action.
+	 */
+	UClass* GetClass() const { return PCGSettingsClass; }
+
+private:
+	/** Type of settings class (node) to spawn. */
+	TSubclassOf<UPCGSettings> PCGSettingsClass;
+};
+
+/** Handles spawn node commands for the PCG Graph Editor. */
+class FPCGEditorSpawnNodeCommands : public TCommands<FPCGEditorSpawnNodeCommands>
+{
+public:
+	FPCGEditorSpawnNodeCommands()
+		: TCommands(TEXT("PCGEditorSpawnNodes"), NSLOCTEXT("PCGEditorSpawnNodes", "PCGEditorSpawnNodes", "PCG Editor - Spawn Nodes"), NAME_None, FAppStyle::GetAppStyleSetName())
+	{}
+
+	// ~Begin TCommands<> interface
+	PCGEDITOR_API virtual void RegisterCommands() override;
+	// ~End TCommands<> interface
+
+	/**
+	 * Returns a graph action assigned to the passed in chord
+	 *
+	 * @param InChord		The chord to use for lookup.
+	 * @return				A shared pointer to the schema action.
+	 */
+	PCGEDITOR_API TSharedPtr<FEdGraphSchemaAction> GetGraphActionByChord(const FInputChord& InChord) const;
+
+private:
+	/** An array of all the possible commands for spawning nodes. */
+	TArray<TSharedPtr<FPCGSpawnNodeCommandInfo>> SpawnNodeCommands;
+};
diff --git a/Engine/Plugins/PCG/Source/PCGEditor/Public/PCGEditorCommon.h b/Engine/Plugins/PCG/Source/PCGEditor/Public/PCGEditorCommon.h
new file mode 100644
--- /dev/null
+++ b/Engine/Plugins/PCG/Source/PCGEditor/Public/PCGEditorCommon.h
@@ -0,0 +1,61 @@
+// Copyright Epic Games, Inc. All Rights Reserved.
+
+#pragma once
+
+#include "HAL/IConsoleManager.h"
+#include "Misc/AssetCategoryPath.h"
+
+#include "PCGEditorCommon.generated.h"
+
+#define LOCTEXT_NAMESPACE "PCGEditorCommon"
+
+UENUM()
+enum class EPCGElementType : uint8
+{
+	Native = 1 << 0,
+	Blueprint = 1 << 1,
+	Subgraph = 1 << 2,
+	Settings = 1 << 3,
+	Asset = 1 << 4,
+	Other = 1 << 5,
+	All = (1 << 6) - 1
+};
+ENUM_CLASS_FLAGS(EPCGElementType);
+
+/** Used to make sure UHT generates properly */
+USTRUCT()
+struct FPCGEditorCommonDummyStruct
+{
+	GENERATED_BODY()
+};
+
+namespace FPCGEditorCommon
+{
+	const FString ContextIdentifier = TEXT("PCGEditorContext");
+
+	const FName SpatialDataType = FName(TEXT("Spatial Data"));
+	const FName ParamDataType = FName(TEXT("Attribute Set"));
+	const FName SettingsDataType = FName(TEXT("Settings Data"));
+	const FName OtherDataType = FName(TEXT("Other Data"));
+
+	const FName ConcreteDataType = FName(TEXT("Concrete Data"));
+	const FName PointDataType = FName(TEXT("Point Data"));
+	const FName PolyLineDataType = FName(TEXT("Poly Line Data"));
+	const FName SurfaceDataType = FName(TEXT("Surface Data"));
+	const FName LandscapeDataType = FName(TEXT("Landscape Data"));
+	const FName BaseTextureDataType = FName(TEXT("Base Texture Data"));
+	const FName TextureDataType = FName(TEXT("Texture Data"));
+	const FName RenderTargetDataType = FName(TEXT("Render Target Data"));
+	const FName VirtualTextureDataType = FName(TEXT("Virtual Texture Data"));
+	const FName VolumeDataType = FName(TEXT("Volume Data"));
+	const FName PrimitiveDataType = FName(TEXT("Primitive Data"));
+	const FName DynamicMeshDataType = FName(TEXT("Dynamic Mesh Data"));
+	const FName StaticMeshResourceDataType = FName(TEXT("Static Mesh Resource Data"));
+
+	extern TAutoConsoleVariable<bool> CVarShowAdvancedAttributesFields;
+
+	const FAssetCategoryPath PCGAssetCategoryPath(LOCTEXT("PCGCategory", "PCG"));
+	const FAssetCategoryPath PCGAdvancedAssetCategoryPath(FPCGEditorCommon::PCGAssetCategoryPath, LOCTEXT("PCGAvancedCategory", "Advanced"));
+};
+
+#undef LOCTEXT_NAMESPACE
diff --git a/Engine/Plugins/PCG/Source/PCGEditor/Public/Schema/PCGEditorGraphSchema.h b/Engine/Plugins/PCG/Source/PCGEditor/Public/Schema/PCGEditorGraphSchema.h
new file mode 100644
--- /dev/null
+++ b/Engine/Plugins/PCG/Source/PCGEditor/Public/Schema/PCGEditorGraphSchema.h
@@ -0,0 +1,86 @@
+// Copyright Epic Games, Inc. All Rights Reserved.
+
+#pragma once
+
+#include "EdGraph/EdGraphSchema.h"
+
+#include "PCGEditorCommon.h"
+#include "Editor/PCGGraphCustomization.h"
+
+#include "ConnectionDrawingPolicy.h"
+#include "StructUtils/PropertyBag.h"
+
+#include "PCGEditorGraphSchema.generated.h"
+
+enum class EPCGElementType : uint8;
+
+class UPCGEditorGraph;
+
+struct FPCGActionsFilter
+{
+	FPCGActionsFilter(const UEdGraph* InEdGraph, EPCGElementType InElementFilterType, FPCGGraphEditorCustomization InCustomization);
+
+	bool Accepts(const FText& InCategory) const;
+
+	EPCGElementType FilterType = EPCGElementType::All;
+	const UPCGEditorGraph* Graph = nullptr;
+	FPCGGraphEditorCustomization Customization;
+};
+
+UCLASS(MinimalAPI)
+class UPCGEditorGraphSchema : public UEdGraphSchema
+{
+	GENERATED_BODY()
+
+public:
+	PCGEDITOR_API void GetPaletteActions(FGraphActionMenuBuilder& ActionMenuBuilder, const FPCGActionsFilter& ActionFilter) const;
+
+	// Return the element type filtering for the current editor. By default it is all.
+	virtual EPCGElementType GetElementTypeFiltering() const { return EPCGElementType::All; };
+
+	// Return the graph customization for this editor. By default it is the one provided by the graph.
+	PCGEDITOR_API virtual const FPCGGraphEditorCustomization& GetGraphEditorCustomization(const UEdGraph* InEdGraph) const;
+
+	//~ Begin EdGraphSchema Interface
+	PCGEDITOR_API virtual void GetGraphContextActions(FGraphContextMenuBuilder& ContextMenuBuilder) const override;
+	PCGEDITOR_API virtual FLinearColor GetPinTypeColor(const FEdGraphPinType& PinType) const override;
+	PCGEDITOR_API virtual const FPinConnectionResponse CanCreateConnection(const UEdGraphPin* A, const UEdGraphPin* B) const override;
+	PCGEDITOR_API virtual bool TryCreateConnection(UEdGraphPin* A, UEdGraphPin* B) const override;
+	PCGEDITOR_API virtual void BreakPinLinks(UEdGraphPin& TargetPin, bool bSendsNodeNotification) const override;
+	PCGEDITOR_API virtual void BreakSinglePinLink(UEdGraphPin* SourcePin, UEdGraphPin* TargetPin) const override;
+	PCGEDITOR_API virtual FConnectionDrawingPolicy* CreateConnectionDrawingPolicy(int32 InBackLayerID, int32 InFrontLayerID, float InZoomFactor, const FSlateRect& InClippingRect, class FSlateWindowElementList& InDrawElements, class UEdGraph* InGraphObj) const override;
+	PCGEDITOR_API virtual void DroppedAssetsOnGraph(const TArray<FAssetData>& Assets, const FVector2f& GraphPosition, UEdGraph* Graph) const override;
+	PCGEDITOR_API virtual void GetAssetsGraphHoverMessage(const TArray<FAssetData>& Assets, const UEdGraph* HoverGraph, FString& OutTooltipText, bool& OutOkIcon) const override;
+	PCGEDITOR_API virtual void OnPinConnectionDoubleCicked(UEdGraphPin* PinA, UEdGraphPin* PinB, const FVector2f& GraphPosition) const override;
+	//~ End EdGraphSchema Interface
+
+	PCGEDITOR_API virtual const FSlateBrush* GetMetadataTypeSlateBrush(EPCGContainerType ContainerType) const;
+	PCGEDITOR_API virtual FLinearColor GetMetadataTypeColor(EPCGMetadataTypes Type) const;
+
+	PCGEDITOR_API virtual const FSlateBrush* GetPropertyBagTypeSlateBrush(EPropertyBagContainerType ContainerType) const;
+	PCGEDITOR_API virtual FLinearColor GetPropertyBagTypeColor(const FPropertyBagPropertyDesc& Desc) const;
+
+protected:
+	PCGEDITOR_API virtual void GetGraphActions(FGraphActionMenuBuilder& ActionMenuBuilder, const FPCGActionsFilter& ActionFilter, bool bIsContextual) const;
+	PCGEDITOR_API virtual void GetNativeElementActions(FGraphActionMenuBuilder& ActionMenuBuilder, const FPCGActionsFilter& ActionFilter) const;
+	PCGEDITOR_API virtual void GetBlueprintElementActions(FGraphActionMenuBuilder& ActionMenuBuilder, const FPCGActionsFilter& ActionFilter) const;
+	PCGEDITOR_API virtual void GetSubgraphElementActions(FGraphActionMenuBuilder& ActionMenuBuilder, const FPCGActionsFilter& ActionFilter) const;
+	PCGEDITOR_API virtual void GetSettingsElementActions(FGraphActionMenuBuilder& ActionMenuBuilder, const FPCGActionsFilter& ActionFilter, bool bIsContextual) const;
+	PCGEDITOR_API virtual void GetExtraElementActions(FGraphActionMenuBuilder& ActionMenuBuilder, const FPCGActionsFilter& ActionFilter) const;
+	PCGEDITOR_API virtual void GetNamedRerouteUsageActions(FGraphActionMenuBuilder& ActionMenuBuilder, const FPCGActionsFilter& ActionFilter) const;
+	PCGEDITOR_API virtual void GetDataAssetActions(FGraphActionMenuBuilder& ActionMenuBuilder, const FPCGActionsFilter& ActionFilter) const;
+
+	PCGEDITOR_API virtual bool TryCreateConnectionInternal(UEdGraphPin* A, UEdGraphPin* B, bool bAddConversionNodeIfNeeded) const;
+};
+
+class FPCGEditorConnectionDrawingPolicy : public FConnectionDrawingPolicy
+{
+public:
+	FPCGEditorConnectionDrawingPolicy(int32 InBackLayerID, int32 InFrontLayerID, float InZoomFactor, const FSlateRect& InClippingRect, FSlateWindowElementList& InDrawElements, UEdGraph* InGraph);
+	virtual void DetermineWiringStyle(UEdGraphPin* OutputPin, UEdGraphPin* InputPin, /*inout*/ FConnectionParams& Params) override;
+
+protected:
+	bool UpdateParamsIfDebugging(UEdGraphPin* OutputPin, UEdGraphPin* InputPin, FConnectionParams& Params);
+
+	UPCGEditorGraph* Graph;
+};
diff --git a/Engine/Plugins/PCG/Source/PCGEditor/Public/Schema/PCGEditorGraphSchemaActions.h b/Engine/Plugins/PCG/Source/PCGEditor/Public/Schema/PCGEditorGraphSchemaActions.h
new file mode 100644
--- /dev/null
+++ b/Engine/Plugins/PCG/Source/PCGEditor/Public/Schema/PCGEditorGraphSchemaActions.h
@@ -0,0 +1,304 @@
+// Copyright Epic Games, Inc. All Rights Reserved.
+
+#pragma once
+
+#include "EdGraph/EdGraphSchema.h"
+
+#include "PCGGraph.h"
+#include "PCGSettings.h"
+#include "Templates/SubclassOf.h"
+
+#include "PCGEditorGraphSchemaActions.generated.h"
+
+class UEdGraph;
+class UEdGraphPin;
+class UPCGBlueprintElement;
+class UPCGNode;
+class UPCGEditorGraph;
+struct FSlateBrush;
+
+UENUM()
+enum class EPCGEditorNewSettingsBehavior : uint8
+{
+	Normal = 0,
+	ForceCopy,
+	ForceInstance
+};
+
+UENUM()
+enum class EPCGEditorNewPCGGraphBehavior : uint8
+{
+	Normal = 0,
+	SubgraphNode,
+	LoopNode
+};
+
+USTRUCT(MinimalAPI)
+struct FPCGEditorGraphSchemaAction_NewNativeElement : public FEdGraphSchemaAction
+{
+	GENERATED_BODY()
+
+	// Inherit the base class's constructors
+	using FEdGraphSchemaAction::FEdGraphSchemaAction;
+
+	// Simple type info
+	static FName StaticGetTypeId()
+	{
+		static FName Type("FPCGEditorGraphSchemaAction_NewNativeElement");
+		return Type;
+	}
+	
+	UPROPERTY()
+	TSubclassOf<UPCGSettings> SettingsClass;
+
+	UPROPERTY()
+	FPCGPreConfiguredSettingsInfo PreconfiguredInfo;
+
+	//~ Begin FEdGraphSchemaAction Interface
+	virtual FName GetTypeId() const override { return StaticGetTypeId(); }
+	using FEdGraphSchemaAction::PerformAction; // Prevent hiding of deprecated base class function with FVector2D
+	virtual UEdGraphNode* PerformAction(UEdGraph* ParentGraph, UEdGraphPin* FromPin, const FVector2f& Location, bool bSelectNewNode = true) override;
+	//~ End FEdGraphSchemaAction Interface
+
+	virtual FSlateBrush const* GetPaletteIcon() const override;
+
+protected:
+	virtual void PostCreation(UPCGNode* NewNode) {}
+
+private:
+	mutable TSharedPtr<FSlateBrush> PaletteIcon = nullptr;
+};
+
+USTRUCT()
+struct FPCGEditorGraphSchemaAction_NewGetParameterElement : public FPCGEditorGraphSchemaAction_NewNativeElement
+{
+	GENERATED_BODY()
+
+	// Inherit the base class's constructors
+	using FPCGEditorGraphSchemaAction_NewNativeElement::FPCGEditorGraphSchemaAction_NewNativeElement;
+
+	UPROPERTY()
+	FPropertyBagPropertyDesc PropertyDesc;
+
+	virtual FSlateBrush const* GetPaletteIcon() const override;
+
+protected:
+	virtual void PostCreation(UPCGNode* NewNode) override;
+
+private:
+	mutable TSharedPtr<FSlateBrush> PaletteIcon = nullptr;
+};
+
+USTRUCT()
+struct FPCGEditorGraphSchemaAction_NewSettingsElement : public FEdGraphSchemaAction
+{
+	GENERATED_BODY()
+
+	// Inherit the base class's constructors
+	using FEdGraphSchemaAction::FEdGraphSchemaAction;
+	
+	// Simple type info
+	static FName StaticGetTypeId()
+	{
+		static FName Type("FPCGEditorGraphSchemaAction_NewSettingsElement");
+		return Type;
+	}
+
+	UPROPERTY()
+	FSoftObjectPath SettingsObjectPath;
+
+	UPROPERTY()
+	EPCGEditorNewSettingsBehavior Behavior = EPCGEditorNewSettingsBehavior::Normal;
+
+	//~ Begin FEdGraphSchemaAction Interface
+	virtual FName GetTypeId() const override { return StaticGetTypeId(); }
+	using FEdGraphSchemaAction::PerformAction; // Prevent hiding of deprecated base class function with FVector2D
+	virtual UEdGraphNode* PerformAction(UEdGraph* ParentGraph, UEdGraphPin* FromPin, const FVector2f& Location, bool bSelectNewNode = true) override;
+	//~ End FEdGraphSchemaAction Interface
+
+	static void MakeSettingsNodesOrContextualMenu(const TSharedRef<class SWidget>& InPanel, const FVector2D& InScreenPosition, UEdGraph* InGraph, const TArray<FSoftObjectPath>& InSettingsPaths, const TArray<FVector2D>& InLocations, bool bInSelectNewNodes);
+	static void MakeSettingsNodes(UPCGEditorGraph* InEditorGraph, TArray<UPCGSettings*> InSettings, TArray<FVector2D> InSettingsLocations, bool bInSelectNewNodes, bool bInCreateInstances);
+	static UEdGraphNode* MakeSettingsNode(UPCGEditorGraph* InEditorGraph, UEdGraphPin* InFromPin, UPCGSettings* InSettings, FVector2D InLocation, bool bInSelectNewNode, bool bInCreateInstance);
+};
+
+USTRUCT()
+struct FPCGEditorGraphSchemaAction_NewBlueprintElement : public FEdGraphSchemaAction
+{
+	GENERATED_BODY()
+
+	// Inherit the base class's constructors
+	using FEdGraphSchemaAction::FEdGraphSchemaAction;
+	
+	// Simple type info
+	static FName StaticGetTypeId()
+	{
+		static FName Type("FPCGEditorGraphSchemaAction_NewBlueprintElement");
+		return Type;
+	}
+
+	UPROPERTY()
+	FSoftClassPath BlueprintClassPath;
+
+	UPROPERTY()
+	FPCGPreConfiguredSettingsInfo PreconfiguredInfo;
+
+	//~ Begin FEdGraphSchemaAction Interface
+	virtual FName GetTypeId() const override { return StaticGetTypeId(); }
+	using FEdGraphSchemaAction::PerformAction; // Prevent hiding of deprecated base class function with FVector2D
+	virtual UEdGraphNode* PerformAction(UEdGraph* ParentGraph, UEdGraphPin* FromPin, const FVector2f& Location, bool bSelectNewNode = true) override;
+	//~ End FEdGraphSchemaAction Interface
+};
+
+USTRUCT()
+struct FPCGEditorGraphSchemaAction_NewLoadAssetElement : public FEdGraphSchemaAction
+{
+	GENERATED_BODY()
+
+	// Inherit the base class's constructors
+	using FEdGraphSchemaAction::FEdGraphSchemaAction;
+
+	// Simple type info
+	static FName StaticGetTypeId()
+	{
+		static FName Type("FPCGEditorGraphSchemaAction_NewLoadAssetElement");
+		return Type;
+	}
+
+	UPROPERTY()
+	FAssetData Asset;
+
+	UPROPERTY()
+	TSubclassOf<UPCGSettings> SettingsClass;
+
+	//~ Begin FEdGraphSchemaAction Interface
+	virtual FName GetTypeId() const override { return StaticGetTypeId(); }
+	using FEdGraphSchemaAction::PerformAction; // Prevent hiding of deprecated base class function with FVector2D
+	virtual UEdGraphNode* PerformAction(UEdGraph* ParentGraph, UEdGraphPin* FromPin, const FVector2f& Location, bool bSelectNewNode = true) override;
+	//~ End FEdGraphSchemaAction Interface
+};
+
+/** Creates a new subgraph element and allows choosing whether it can be a subgraph node or a loop node. */
+USTRUCT()
+struct FPCGEditorGraphSchemaAction_NewSubgraphElement : public FEdGraphSchemaAction
+{
+	GENERATED_BODY()
+
+	// Inherit the base class's constructors
+	using FEdGraphSchemaAction::FEdGraphSchemaAction;
+
+	// Simple type info
+	static FName StaticGetTypeId()
+	{
+		static FName Type("FPCGEditorGraphSchemaAction_NewSubgraphElement");
+		return Type;
+	}
+
+	UPROPERTY()
+	FSoftObjectPath SubgraphObjectPath;
+
+	UPROPERTY()
+	EPCGEditorNewPCGGraphBehavior Behavior = EPCGEditorNewPCGGraphBehavior::Normal;
+
+	//~ Begin FEdGraphSchemaAction Interface
+	virtual FName GetTypeId() const override { return StaticGetTypeId(); }
+	using FEdGraphSchemaAction::PerformAction; // Prevent hiding of deprecated base class function with FVector2D
+	virtual UEdGraphNode* PerformAction(UEdGraph* ParentGraph, UEdGraphPin* FromPin, const FVector2f& Location, bool bSelectNewNode = true) override;
+	//~ End FEdGraphSchemaAction Interface
+
+	static void MakeGraphNodesOrContextualMenu(const TSharedRef<class SWidget>& InPanel, const FVector2D& InScreenPosition, UEdGraph* InGraph, const TArray<FSoftObjectPath>& InGraphPaths, const TArray<FVector2D>& InLocations, bool bInSelectNewNodes);
+	static void MakeGraphNodes(UPCGEditorGraph* InEditorGraph, TArray<UPCGGraphInterface*> InGraph, TArray<FVector2D> InGraphLocations, bool bInSelectNewNodes, bool bInCreateLoop);
+	static UEdGraphNode* MakeGraphNode(UPCGEditorGraph* InEditorGraph, UEdGraphPin* InFromPin, UPCGGraphInterface* InGraph, const FVector2D& InLocation, bool bInSelectNewNode, bool bInCreateLoop);
+};
+
+/** Action to add a 'comment' node to the graph */
+USTRUCT()
+struct FPCGEditorGraphSchemaAction_NewComment : public FEdGraphSchemaAction
+{
+	GENERATED_BODY()
+
+	// Inherit the base class's constructors
+	using FEdGraphSchemaAction::FEdGraphSchemaAction;
+
+	// Simple type info
+	static FName StaticGetTypeId()
+	{
+		static FName Type("FPCGEditorGraphSchemaAction_NewComment");
+		return Type;
+	}
+
+	// FEdGraphSchemaAction interface
+	virtual FName GetTypeId() const override { return StaticGetTypeId(); }
+	using FEdGraphSchemaAction::PerformAction;
+	virtual UEdGraphNode* PerformAction(UEdGraph* ParentGraph, UEdGraphPin* FromPin, const FVector2f& Location, bool bSelectNewNode = true) override;
+	// End of FEdGraphSchemaAction interface
+};
+
+USTRUCT()
+struct FPCGEditorGraphSchemaAction_NewReroute : public FEdGraphSchemaAction
+{
+	GENERATED_BODY()
+
+	// Inherit the base class's constructors
+	using FEdGraphSchemaAction::FEdGraphSchemaAction;
+
+	// Simple type info
+	static FName StaticGetTypeId()
+	{
+		static FName Type("FPCGEditorGraphSchemaAction_NewReroute");
+		return Type;
+	}
+
+	// FEdGraphSchemaAction interface
+	virtual FName GetTypeId() const override { return StaticGetTypeId(); }
+	using FEdGraphSchemaAction::PerformAction; // Prevent hiding of deprecated base class function with FVector2D
+	virtual UEdGraphNode* PerformAction(UEdGraph* ParentGraph, UEdGraphPin* FromPin, const FVector2f& Location, bool bSelectNewNode = true) override;
+	// End of FEdGraphSchemaAction interface
+};
+
+USTRUCT()
+struct FPCGEditorGraphSchemaAction_NewNamedRerouteUsage : public FEdGraphSchemaAction
+{
+	GENERATED_BODY()
+
+	// Inherit the base class's constructors
+	using FEdGraphSchemaAction::FEdGraphSchemaAction;
+
+	// Declaration that we want to add an usage of
+	UPROPERTY()
+	TObjectPtr<const UPCGNode> DeclarationNode = nullptr;
+
+	// Simple type info
+	static FName StaticGetTypeId()
+	{
+		static FName Type("FPCGEditorGraphSchemaAction_NewNamedRerouteUsage");
+		return Type;
+	}
+
+	// FEdGraphSchemaAction interface
+	virtual FName GetTypeId() const override { return StaticGetTypeId(); }
+	using FEdGraphSchemaAction::PerformAction; // Prevent hiding of deprecated base class function with FVector2D
+	virtual UEdGraphNode* PerformAction(UEdGraph* ParentGraph, UEdGraphPin* FromPin, const FVector2f& Location, bool bSelectNewNode = true) override;
+	// End of FEdGraphSchemaAction interface
+};
+
+USTRUCT()
+struct FPCGEditorGraphSchemaAction_NewNamedRerouteDeclaration : public FEdGraphSchemaAction
+{
+	GENERATED_BODY()
+
+	// Inherit the base class's constructors
+	using FEdGraphSchemaAction::FEdGraphSchemaAction;
+
+	// Simple type info
+	static FName StaticGetTypeId()
+	{
+		static FName Type("FPCGEditorGraphSchemaAction_NewNamedRerouteDeclaration");
+		return Type;
+	}
+
+	// FEdGraphSchemaAction interface
+	virtual FName GetTypeId() const override { return StaticGetTypeId(); }
+	using FEdGraphSchemaAction::PerformAction; // Prevent hiding of deprecated base class function with FVector2D
+	virtual UEdGraphNode* PerformAction(UEdGraph* ParentGraph, UEdGraphPin* FromPin, const FVector2f& Location, bool bSelectNewNode = true) override;
+	// End of FEdGraphSchemaAction interface
+};
