[PCG] Fix error introduced in previous commit - don't throw attribute errors if there are no data items.

#rnx

diff --git a/Engine/Plugins/PCG/Source/PCG/Private/Compute/BuiltInKernels/PCGCountUniqueAttributeValuesKernel.cpp b/Engine/Plugins/PCG/Source/PCG/Private/Compute/BuiltInKernels/PCGCountUniqueAttributeValuesKernel.cpp
--- a/Engine/Plugins/PCG/Source/PCG/Private/Compute/BuiltInKernels/PCGCountUniqueAttributeValuesKernel.cpp
+++ b/Engine/Plugins/PCG/Source/PCG/Private/Compute/BuiltInKernels/PCGCountUniqueAttributeValuesKernel.cpp
@@ -43,39 +43,43 @@ bool UPCGCountUniqueAttributeValuesKernel::IsKernelDataValid(FPCGContext* InCont
 			return false;
 		}
 
-		FPCGKernelAttributeDesc AttributeDesc;
-		bool bConflictingTypesFound = false;
-		bool bPresentOnAllData = false;
-		InputDataDesc->GetAttributeDesc(AttributeName, AttributeDesc, bConflictingTypesFound, bPresentOnAllData);
-		if (!bPresentOnAllData)
+		if (!InputDataDesc->DataDescs.IsEmpty())
 		{
-			if (!InputDataDesc->DataDescs.IsEmpty())
+			FPCGKernelAttributeDesc AttributeDesc;
+			bool bConflictingTypesFound = false;
+			bool bPresentOnAllData = false;
+			InputDataDesc->GetAttributeDesc(AttributeName, AttributeDesc, bConflictingTypesFound, bPresentOnAllData);
+
+			if (!bPresentOnAllData)
 			{
-				PCG_KERNEL_VALIDATION_ERR(InContext, GetSettings(), FText::Format(
-					LOCTEXT("AttributeMissing", "Count attribute '{0}' not found, this attribute must be present on all input data, and be of type String Key."),
-					FText::FromName(AttributeName)));
-			}
+				if (!InputDataDesc->DataDescs.IsEmpty())
+				{
+					PCG_KERNEL_VALIDATION_ERR(InContext, GetSettings(), FText::Format(
+						LOCTEXT("AttributeMissing", "Count attribute '{0}' not found, this attribute must be present on all input data, and be of type String Key."),
+						FText::FromName(AttributeName)));
+				}
 
-			return false;
-		}
+				return false;
+			}
 
-		if (bConflictingTypesFound)
-		{
-			PCG_KERNEL_VALIDATION_ERR(InContext, GetSettings(), FText::Format(
-				LOCTEXT("AttributeHasMultipleTypes", "Count attribute '{0}' found with multiple types in input data, all attributes must be of type String Key."),
-				FText::FromName(AttributeName)));
+			if (bConflictingTypesFound)
+			{
+				PCG_KERNEL_VALIDATION_ERR(InContext, GetSettings(), FText::Format(
+					LOCTEXT("AttributeHasMultipleTypes", "Count attribute '{0}' found with multiple types in input data, all attributes must be of type String Key."),
+					FText::FromName(AttributeName)));
 
-			return false;
-		}
+				return false;
+			}
 
-		if (AttributeDesc.AttributeKey.Type != EPCGKernelAttributeType::StringKey)
-		{
-			// Attribute value counting only currently supported for attributes of type StringKey.
-			PCG_KERNEL_VALIDATION_ERR(Context, GetSettings(), FText::Format(
-				LOCTEXT("AttributeTypeInvalid", "Cannot count values for attribute '{0}', only attributes of type String Key are currently supported."),
-				FText::FromName(AttributeName)));
+			if (AttributeDesc.AttributeKey.Type != EPCGKernelAttributeType::StringKey)
+			{
+				// Attribute value counting only currently supported for attributes of type StringKey.
+				PCG_KERNEL_VALIDATION_ERR(Context, GetSettings(), FText::Format(
+					LOCTEXT("AttributeTypeInvalid", "Cannot count values for attribute '{0}', only attributes of type String Key are currently supported."),
+					FText::FromName(AttributeName)));
 
-			return false;
+				return false;
+			}
 		}
 	}
 
