[PCG] Remove unused lambda (should've been removed in CL 41724792

#rnx

diff --git a/Engine/Plugins/PCG/Source/PCG/Private/Graph/PCGGraphExecutor.cpp b/Engine/Plugins/PCG/Source/PCG/Private/Graph/PCGGraphExecutor.cpp
--- a/Engine/Plugins/PCG/Source/PCG/Private/Graph/PCGGraphExecutor.cpp
+++ b/Engine/Plugins/PCG/Source/PCG/Private/Graph/PCGGraphExecutor.cpp
@@ -3197,19 +3197,6 @@ namespace PCGGraphExecutor
 			}
 
 			// At this point we could not get to the data, so we'll try executing the graph if we did not do that already.
-
-			auto WakeUpLambda = [ComponentWithData, ContextHandle=InContext->GetOrCreateHandle()](FPCGContext* InContext)
-			{
-				TSharedPtr<FPCGContextHandle> SharedHandle = ContextHandle.Pin();
-				if (FPCGContext* ContextPtr = SharedHandle ? SharedHandle->GetContext() : nullptr)
-				{
-					PCGGraphExecutionLogging::LogGridLinkageTaskExecuteRetrieveWakeUp(ContextPtr, ComponentWithData);
-					ContextPtr->bIsPaused = false;
-				}
-
-				return true;
-			};
-
 			// If we need data from a local component but the local component is still generating, then we'll wait for it. On the other hand
 			// if we need data from the original component we assume the generation has already happened because it is always scheduled before
 			// the local components.
