[PCG] Fixed crash in self-pruning in certain cases with complex (triangle meshes) collisions.

#rnx

diff --git a/Engine/Plugins/PCG/Source/PCG/Private/Elements/PCGSelfPruning.cpp b/Engine/Plugins/PCG/Source/PCG/Private/Elements/PCGSelfPruning.cpp
--- a/Engine/Plugins/PCG/Source/PCG/Private/Elements/PCGSelfPruning.cpp
+++ b/Engine/Plugins/PCG/Source/PCG/Private/Elements/PCGSelfPruning.cpp
@@ -210,6 +210,31 @@ namespace PCGSelfPruningElement
 					bHasComplexShapes = (CollisionQueryFlag == EPCGCollisionQueryFlag::Complex) ||
 						(CollisionQueryFlag == EPCGCollisionQueryFlag::ComplexFirst && bFirstChoice) ||
 						(CollisionQueryFlag == EPCGCollisionQueryFlag::SimpleFirst && !bFirstChoice);
+
+					// Special identification step: in some instances (when users mark their complex collision as simple...) it's possible that we get in a case here
+					// where simple collisions contain triangle meshes, which are going to break the tests below as the Physics API does not allow casting/testing for overlaps
+					// on triangle mesh - triangle mesh cases.
+					if (!bHasComplexShapes)
+					{
+						for (const FPhysicsShapeHandle& Shape : OutShapes)
+						{
+							FPhysicsGeometryCollection Collection = FPhysicsInterface::GetGeometryCollection(Shape);
+							const Chaos::FImplicitObject& ShapeGeom = Collection.GetGeometry();
+
+							if (ShapeGeom.IsUnderlyingMesh())
+							{
+								bHasComplexShapes = true;
+								break;
+							}
+						}
+					}
+
+					// Special verification step - if at this point we had asked *specifically* for simple shapes, we must discard triangle meshes if we had found some.
+					if (CollisionQueryFlag == EPCGCollisionQueryFlag::Simple && bHasComplexShapes)
+					{
+						OutInstance = nullptr;
+						bHasComplexShapes = false;
+					}
 				}
 			}
 
