[PCG] Fix Runtime Gen grid linkage broken when placing original volume into a Data Layer
- Runtime Gen components (local/original) grid descriptors should not have any DL/HLOD information

diff --git a/Engine/Plugins/PCG/Source/PCG/Private/Grid/PCGPartitionActor.cpp b/Engine/Plugins/PCG/Source/PCG/Private/Grid/PCGPartitionActor.cpp
--- a/Engine/Plugins/PCG/Source/PCG/Private/Grid/PCGPartitionActor.cpp
+++ b/Engine/Plugins/PCG/Source/PCG/Private/Grid/PCGPartitionActor.cpp
@@ -335,19 +335,24 @@ FPCGGridDescriptor APCGPartitionActor::GetGridDescriptor() const
 		.SetIs2DGrid(bUse2DGrid)
 		.SetIsRuntime(IsRuntimeGenerated());
 
-#if WITH_EDITORONLY_DATA
-	if (GetWorld() && GetWorld()->IsPlayInEditor())
+	// Runtime generated Partition Actors should never have any Data Layer Assets / HLOD Layer assigned as they are spawned at runtime
+	// but skip this code anyways.
+	if (!IsRuntimeGenerated())
 	{
-		GridDescriptor.SetRuntimeHash(RuntimeGridDescriptorHash);
-	}
-	else
-	{
-		GridDescriptor.SetDataLayerAssets(GetDataLayerAssets());
-		GridDescriptor.SetHLODLayer(GetHLODLayer());
-	}
+#if WITH_EDITORONLY_DATA
+		if (GetWorld() && GetWorld()->IsPlayInEditor())
+		{
+			GridDescriptor.SetRuntimeHash(RuntimeGridDescriptorHash);
+		}
+		else
+		{
+			GridDescriptor.SetDataLayerAssets(GetDataLayerAssets());
+			GridDescriptor.SetHLODLayer(GetHLODLayer());
+		}
 #else
-	GridDescriptor.SetRuntimeHash(RuntimeGridDescriptorHash);
+		GridDescriptor.SetRuntimeHash(RuntimeGridDescriptorHash);
 #endif
+	}
 
 	return GridDescriptor;
 }
@@ -653,7 +658,7 @@ void APCGPartitionActor::PostCreation(const FPCGGridDescriptor& GridDescriptor)
 	bUse2DGrid = GridDescriptor.Is2DGrid();
 
 #if WITH_EDITOR
-	if (!GetWorld()->IsGameWorld())
+	if (!GetWorld()->IsGameWorld() && !GridDescriptor.IsRuntime())
 	{
 		// Fetch Non External Assets and assign
 		TArray<TSoftObjectPtr<UDataLayerAsset>> DescriptorDataLayerAssets;
diff --git a/Engine/Plugins/PCG/Source/PCG/Private/PCGComponent.cpp b/Engine/Plugins/PCG/Source/PCG/Private/PCGComponent.cpp
--- a/Engine/Plugins/PCG/Source/PCG/Private/PCGComponent.cpp
+++ b/Engine/Plugins/PCG/Source/PCG/Private/PCGComponent.cpp
@@ -250,20 +250,25 @@ FPCGGridDescriptor UPCGComponent::GetGridDescriptorInternal(uint32 GridSize, boo
 		.SetIs2DGrid(Use2DGrid())
 		.SetIsRuntime(IsManagedByRuntimeGenSystem());
 
-#if WITH_EDITORONLY_DATA
-	// Only return the RuntimeGridDescriptorHash for PIE Worlds and not when we are updating the Runtime Hash
-	if (GetWorld() && GetWorld()->IsPlayInEditor() && !bRuntimeHashUpdate)
+	// Avoid setting this for runtime components. They do not support being assigned a DataLayers and/or HLOD Layer
+	// They also do not need a RuntimeGridDescriptorHash which exists only because of DataLayers/HLODLayers
+	if (!IsManagedByRuntimeGenSystem())
 	{
-		PCGGridDescriptor.SetRuntimeHash(RuntimeGridDescriptorHash);
-	}
-	else
-	{
-		PCGGridDescriptor.SetDataLayerAssets(GetOwner()->GetDataLayerAssets());
-		PCGGridDescriptor.SetHLODLayer(GetOwner()->GetHLODLayer());
-	}
+#if WITH_EDITORONLY_DATA
+		// Only return the RuntimeGridDescriptorHash for PIE Worlds and not when we are updating the Runtime Hash
+		if (GetWorld() && GetWorld()->IsPlayInEditor() && !bRuntimeHashUpdate)
+		{
+			PCGGridDescriptor.SetRuntimeHash(RuntimeGridDescriptorHash);
+		}
+		else
+		{
+			PCGGridDescriptor.SetDataLayerAssets(GetOwner()->GetDataLayerAssets());
+			PCGGridDescriptor.SetHLODLayer(GetOwner()->GetHLODLayer());
+		}
 #else
-	PCGGridDescriptor.SetRuntimeHash(RuntimeGridDescriptorHash);
+		PCGGridDescriptor.SetRuntimeHash(RuntimeGridDescriptorHash);
 #endif
+	}
 
 	return PCGGridDescriptor;
 }
