PCG Spline Attribute Channel Sampling (2) - Sample attribute channels in the spline sampler

diff --git a/Engine/Plugins/PCG/Source/PCG/Private/Data/PCGSplineData.cpp b/Engine/Plugins/PCG/Source/PCG/Private/Data/PCGSplineData.cpp
--- a/Engine/Plugins/PCG/Source/PCG/Private/Data/PCGSplineData.cpp
+++ b/Engine/Plugins/PCG/Source/PCG/Private/Data/PCGSplineData.cpp
@@ -135,6 +135,24 @@ bool UPCGSplineData::SetDomainFromDomainID(const FPCGMetadataDomainID& InDomainI
 	}
 }
 
+void UPCGSplineData::InitializeTargetMetadata(const FPCGInitializeFromDataParams& InParams, UPCGMetadata* MetadataToInitialize) const
+{
+	check(MetadataToInitialize);
+
+	Super::InitializeTargetMetadata(InParams, MetadataToInitialize);
+
+	FPCGMetadataDomain* MetadataDomain = MetadataToInitialize->GetMetadataDomain(PCGMetadataDomainID::Elements);
+	if (!MetadataDomain)
+	{
+		return;
+	}
+
+	for (const FName& ChannelName : SplineStruct.GetAttributeChannelNames())
+	{
+		MetadataDomain->FindOrCreateAttribute<float>(ChannelName, /*DefaultValue=*/0.f);
+	}
+}
+
 FTransform UPCGSplineData::GetTransform() const
 {
 	return SplineStruct.GetTransform();
@@ -276,6 +294,20 @@ void UPCGSplineData::WriteMetadataToPoint(float InputKey, FPCGPoint& OutPoint, U
 			OutMetadataDomain->ComputeWeightedAttribute(OutPoint.MetadataEntry, Coefficients, MetadataDomain);
 		}
 	}
+
+	if (const TArray<FName>& ChannelNames = SplineStruct.GetAttributeChannelNames(); ChannelNames.Num() > 0)
+	{
+		OutMetadataDomain->InitializeOnSet(OutPoint.MetadataEntry);
+
+		// Sample attribute channels,
+		for (const FName& ChannelName : ChannelNames)
+		{
+			if (FPCGMetadataAttribute<float>* UserAttribute = OutMetadataDomain->FindOrCreateAttribute<float>(ChannelName, /*DefaultValue=*/0.f))
+			{
+				UserAttribute->SetValue(OutPoint.MetadataEntry, SplineStruct.GetAttributeValue(InputKey, ChannelName));
+			}
+		}
+	}
 }
 
 TUniquePtr<IPCGAttributeAccessor> UPCGSplineData::CreateStaticAccessor(const FPCGAttributePropertySelector& InSelector, bool bQuiet)
diff --git a/Engine/Plugins/PCG/Source/PCG/Private/Data/PCGSplineStruct.cpp b/Engine/Plugins/PCG/Source/PCG/Private/Data/PCGSplineStruct.cpp
--- a/Engine/Plugins/PCG/Source/PCG/Private/Data/PCGSplineStruct.cpp
+++ b/Engine/Plugins/PCG/Source/PCG/Private/Data/PCGSplineStruct.cpp
@@ -198,6 +198,52 @@ void FPCGSplineStruct::ApplyTo(USplineComponent* InSplineComponent) const
 	InSplineComponent->UpdateBounds();
 }
 
+FPCGSplineStruct::FPCGSplineStruct(const FPCGSplineStruct& Other)
+{
+	*this = Other;
+}
+
+FPCGSplineStruct::FPCGSplineStruct(FPCGSplineStruct&& Other)
+{
+	*this = MoveTemp(Other);
+}
+
+FPCGSplineStruct& FPCGSplineStruct::operator=(const FPCGSplineStruct& Other)
+{
+	WarninglessSplineCurves() = Other.WarninglessSplineCurves();
+	Spline = Other.Spline;
+	
+	Transform = Other.Transform;
+	DefaultUpVector = Other.DefaultUpVector;
+	ReparamStepsPerSegment = Other.ReparamStepsPerSegment;
+	bClosedLoop = Other.bClosedLoop;
+
+	Bounds = Other.Bounds;
+	LocalBounds = Other.LocalBounds;
+
+	ControlPointsEntryKeys = Other.ControlPointsEntryKeys;
+	
+	return *this;
+}
+
+FPCGSplineStruct& FPCGSplineStruct::operator=(FPCGSplineStruct&& Other)
+{
+	WarninglessSplineCurves() = MoveTemp(Other.WarninglessSplineCurves());
+	Spline = MoveTemp(Other.Spline);
+	
+	Transform = MoveTemp(Other.Transform);
+	DefaultUpVector = MoveTemp(Other.DefaultUpVector);
+	ReparamStepsPerSegment = MoveTemp(Other.ReparamStepsPerSegment);
+	bClosedLoop = MoveTemp(Other.bClosedLoop);
+
+	Bounds = MoveTemp(Other.Bounds);
+	LocalBounds = MoveTemp(Other.LocalBounds);
+
+	ControlPointsEntryKeys = MoveTemp(Other.ControlPointsEntryKeys);
+	
+	return *this;
+}
+
 void FPCGSplineStruct::AddPoint(const FSplinePoint& InSplinePoint, bool bUpdateSpline)
 {
 	const TArray<FInterpCurvePoint<FVector>>& SourcePoints = ShouldUseSplineCurves()
@@ -339,6 +385,36 @@ EInterpCurveMode FPCGSplineStruct::GetSplinePointType(int32 Index) const
 		: Spline.GetSplinePointType(Index);
 }
 
+const TArray<FName>& FPCGSplineStruct::GetAttributeChannelNames() const
+{
+	if (ShouldUseSplineCurves())
+	{
+		if (!CachedSortedAttributeChannelNames.IsEmpty())
+		{
+			CachedSortedAttributeChannelNames.Empty();
+		}
+	}
+	else if (TArray<FName> Channels = Spline.GetFloatPropertyChannels(); Channels.Num() != CachedSortedAttributeChannelNames.Num())
+	{
+		Channels.Sort([](const FName& A, const FName& B)
+			{ return A.ToString() < B.ToString(); });
+		
+		if (UE::TScopeLock Lock(SpinLock); Channels.Num() != CachedSortedAttributeChannelNames.Num())
+		{
+			CachedSortedAttributeChannelNames = MoveTemp(Channels);
+		}
+	}
+
+	return CachedSortedAttributeChannelNames;
+}
+
+float FPCGSplineStruct::GetAttributeValue(float InputKey, FName ChannelName) const
+{
+	return ShouldUseSplineCurves()
+		? 0.f
+		: Spline.EvaluateAttribute<float>(InputKey, ChannelName);
+}
+
 void FPCGSplineStruct::UpdateSpline()
 {
 	const bool bLoopPositionOverride = false;
diff --git a/Engine/Plugins/PCG/Source/PCG/Public/Data/PCGSplineData.h b/Engine/Plugins/PCG/Source/PCG/Public/Data/PCGSplineData.h
--- a/Engine/Plugins/PCG/Source/PCG/Public/Data/PCGSplineData.h
+++ b/Engine/Plugins/PCG/Source/PCG/Public/Data/PCGSplineData.h
@@ -69,6 +69,10 @@ public:
 	UE_API virtual bool SetDomainFromDomainID(const FPCGMetadataDomainID& InDomainID, FPCGAttributePropertySelector& InOutSelector) const;
 	// ~End UPCGData interface
 
+	// ~UPCGSpatialData interface
+	UE_API virtual void InitializeTargetMetadata(const FPCGInitializeFromDataParams& InParams, UPCGMetadata* MetadataToInitialize) const override;
+	// ~End UPCGSpatialData interface
+	
 	//~Begin UPCGPolyLineData interface
 	UE_API virtual FTransform GetTransform() const override;
 	UE_API virtual int GetNumSegments() const override;
@@ -82,7 +86,7 @@ public:
 	UE_API virtual FVector GetLocationAtAlpha(float Alpha) const override;
 	UE_API virtual FTransform GetTransformAtAlpha(float Alpha) const override;
 	virtual bool IsClosed() const override { return SplineStruct.bClosedLoop; }
-	UE_API virtual void WriteMetadataToPoint(float InputKey, FPCGPoint& OutPoint, UPCGMetadata* OutMetadata) const;
+	UE_API virtual void WriteMetadataToPoint(float InputKey, FPCGPoint& OutPoint, UPCGMetadata* OutMetadata) const override;
 	virtual TConstArrayView<PCGMetadataEntryKey> GetConstVerticesEntryKeys() const override { return SplineStruct.GetConstControlPointsEntryKeys(); }
 	//~End UPCGPolyLineData interface
 	
diff --git a/Engine/Plugins/PCG/Source/PCG/Public/Data/PCGSplineStruct.h b/Engine/Plugins/PCG/Source/PCG/Public/Data/PCGSplineStruct.h
--- a/Engine/Plugins/PCG/Source/PCG/Public/Data/PCGSplineStruct.h
+++ b/Engine/Plugins/PCG/Source/PCG/Public/Data/PCGSplineStruct.h
@@ -5,6 +5,7 @@
 #include "Metadata/PCGMetadataCommon.h"
 
 #include "CoreMinimal.h"
+#include "Misc/SpinLock.h"
 #include "Components/SplineComponent.h"
 
 #include "PCGSplineStruct.generated.h"
@@ -55,6 +56,12 @@ struct FPCGSplineStruct
 	UE_API void Initialize(const TArray<FSplinePoint>& InSplinePoints, bool bIsClosedLoop, const FTransform& InTransform, TArray<PCGMetadataEntryKey> InOptionalEntryKeys = {});
 	UE_API void ApplyTo(USplineComponent* InSplineComponent) const;
 
+	FPCGSplineStruct() = default;
+	UE_API FPCGSplineStruct(const FPCGSplineStruct& Other);
+	UE_API FPCGSplineStruct(FPCGSplineStruct&& Other);
+	UE_API FPCGSplineStruct& operator=(const FPCGSplineStruct& Other);
+	UE_API FPCGSplineStruct& operator=(FPCGSplineStruct&& Other);
+	
 	FTransform GetTransform() const { return Transform; }
 
 	// Spline-related methods
@@ -79,6 +86,12 @@ struct FPCGSplineStruct
 	UE_API void SetSplinePointType(int32 Index, EInterpCurveMode Type);
 	UE_API EInterpCurveMode GetSplinePointType(int32 Index) const;
 	
+	// Gets a lexicographically sorted list of channel names.
+	UE_API const TArray<FName>& GetAttributeChannelNames() const;
+
+	// Evaluate the specified attribute channel at the specified spline input key.
+	UE_API float GetAttributeValue(float InputKey, FName ChannelName) const;
+	
 	UE_API void UpdateSpline();
 
 	UE_API int GetNumberOfSplineSegments() const;
@@ -157,6 +170,9 @@ private:
 	UPROPERTY()
 	FSpline Spline;
 
+	mutable UE::FSpinLock SpinLock;
+	mutable TArray<FName> CachedSortedAttributeChannelNames;
+
 	UE_API bool ShouldUseSplineCurves() const;
 	
 	// Remove when appending _DEPRECATED to SplineCurves.
