[PCG] Attribute partition uses analysis kernel instead of legacy element counters.

#rnx

diff --git a/Engine/Plugins/PCG/Shaders/Private/BuiltInKernels/PCGCountUniqueAttributeValues.usf b/Engine/Plugins/PCG/Shaders/Private/BuiltInKernels/PCGCountUniqueAttributeValues.usf
--- a/Engine/Plugins/PCG/Shaders/Private/BuiltInKernels/PCGCountUniqueAttributeValues.usf
+++ b/Engine/Plugins/PCG/Shaders/Private/BuiltInKernels/PCGCountUniqueAttributeValues.usf
@@ -47,15 +47,18 @@ void PCGCountUniqueAttributeValuesCS(uint3 GroupId : SV_GroupID, uint GroupIndex
 
 		const int AttributeValue = In_GetStringKey(In_DataIndex, In_ElementIndex, AttributeToCountId);
 
-		const uint Out_DataIndex = CountUniqueValues_GetEmitPerDataCounts() ? In_DataIndex : 0;
-		const uint Out_ElementIndex = CountUniqueValues_GetValueIndex(Out_DataIndex, AttributeValue);
-
-		if (Out_ElementIndex == (uint)-1)
+		if (AttributeValue < 0)
 		{
-			// Should not happen.
+			// Should not happen (0 is null/empty string, negative is invalid).
 			return;
 		}
 
+		// If not outputting per-data counts, then all counters are aggregated on a single data.
+		const uint Out_DataIndex = CountUniqueValues_GetEmitPerDataCounts() ? In_DataIndex : 0;
+
+		// Always one bucket per attribute value.
+		const uint Out_ElementIndex = AttributeValue;
+
 		const int OutputCountAttributeId = CountUniqueValues_GetOutputCountAttributeId();
 
 		// Increment count of this attribute.
diff --git a/Engine/Plugins/PCG/Shaders/Private/Elements/PCGMetadataPartition.usf b/Engine/Plugins/PCG/Shaders/Private/Elements/PCGMetadataPartition.usf
--- a/Engine/Plugins/PCG/Shaders/Private/Elements/PCGMetadataPartition.usf
+++ b/Engine/Plugins/PCG/Shaders/Private/Elements/PCGMetadataPartition.usf
@@ -22,22 +22,15 @@ void PCGMetadataPartitionCS(uint3 GroupId : SV_GroupID, uint GroupIndex : SV_Gro
 
 	if (In_IsPointRemoved(In_DataIndex, In_ElementIndex))
 	{
+		// Don't need to remove output point as we're appending points to the buffer below.
 		return;
 	}
 
 	const int PartitionAttributeId = MetadataPartition_GetPartitionAttributeId();
-	const int NumPartitions = MetadataPartition_GetNumPartitions();
-	
-	const int PartitionValueKey = In_GetStringKey(In_DataIndex, In_ElementIndex, PartitionAttributeId);
+	const int AttributeValue = In_GetStringKey(In_DataIndex, In_ElementIndex, PartitionAttributeId);
 
-	const int PartitionIndex = MetadataPartition_GetPartitionIndex(PartitionValueKey);
-
-	const uint Out_DataIndex = In_DataIndex * NumPartitions + PartitionIndex;
-	const int Out_ElementIndex = Out_AddToElementCounterInternal(Out_DataIndex, 1);
-	if (Out_ElementIndex == -1)
-	{
-		return;
-	}
+	uint Out_DataIndex, Out_ElementIndex;
+	MetadataPartition_GetOutputIndices(In_DataIndex, AttributeValue, Out_DataIndex, Out_ElementIndex);
 
 	// Write output point.
 	PCG_COPY_ALL_ATTRIBUTES_TO_OUTPUT(Out, In, Out_DataIndex, Out_ElementIndex, In_DataIndex, In_ElementIndex);
diff --git a/Engine/Plugins/PCG/Source/PCG/Private/Compute/BuiltInKernels/PCGCountUniqueAttributeValuesKernel.h b/Engine/Plugins/PCG/Source/PCG/Private/Compute/BuiltInKernels/PCGCountUniqueAttributeValuesKernel.h
--- a/Engine/Plugins/PCG/Source/PCG/Private/Compute/BuiltInKernels/PCGCountUniqueAttributeValuesKernel.h
+++ b/Engine/Plugins/PCG/Source/PCG/Private/Compute/BuiltInKernels/PCGCountUniqueAttributeValuesKernel.h
@@ -12,6 +12,10 @@ namespace PCGCountUniqueAttributeValuesConstants
 	const FName ValueCountAttributeName = TEXT("UniqueValueCount");
 }
 
+/**
+* Counts how many unique values of a string key attribute are present in an input data collection. Output attribute set is a table of unique string key
+* values and corresponding counts. Can output an attribute set per input data, or a single attribute set that counts across all data.
+*/
 UCLASS()
 class UPCGCountUniqueAttributeValuesKernel : public UPCGComputeKernel
 {
diff --git a/Engine/Plugins/PCG/Source/PCG/Private/Compute/DataInterfaces/BuiltInKernels/PCGCountUniqueAttributeValuesDataInterface.cpp b/Engine/Plugins/PCG/Source/PCG/Private/Compute/DataInterfaces/BuiltInKernels/PCGCountUniqueAttributeValuesDataInterface.cpp
--- a/Engine/Plugins/PCG/Source/PCG/Private/Compute/DataInterfaces/BuiltInKernels/PCGCountUniqueAttributeValuesDataInterface.cpp
+++ b/Engine/Plugins/PCG/Source/PCG/Private/Compute/DataInterfaces/BuiltInKernels/PCGCountUniqueAttributeValuesDataInterface.cpp
@@ -31,12 +31,6 @@ void UPCGCountUniqueAttributeValuesDataInterface::GetSupportedInputs(TArray<FSha
 		.SetName(TEXT("CountUniqueValues_GetOutputValueAttributeId"))
 		.AddReturnType(EShaderFundamentalType::Int);
 
-	OutFunctions.AddDefaulted_GetRef()
-		.SetName(TEXT("CountUniqueValues_GetValueIndex"))
-		.AddReturnType(EShaderFundamentalType::Uint)
-		.AddParam(EShaderFundamentalType::Uint) // InOutputDataIndex
-		.AddParam(EShaderFundamentalType::Uint); // InAttributeValue
-
 	OutFunctions.AddDefaulted_GetRef()
 		.SetName(TEXT("CountUniqueValues_GetEmitPerDataCounts"))
 		.AddReturnType(EShaderFundamentalType::Bool);
@@ -47,7 +41,6 @@ BEGIN_SHADER_PARAMETER_STRUCT(FPCGCountUniqueAttributeValuesDataInterfaceParamet
 	SHADER_PARAMETER(int32, OutputValueAttributeId)
 	SHADER_PARAMETER(int32, OutputCountAttributeId)
 	SHADER_PARAMETER(uint32, EmitPerDataCounts)
-	SHADER_PARAMETER(int32, MaxAttributeValueAcrossAllData)
 END_SHADER_PARAMETER_STRUCT()
 
 void UPCGCountUniqueAttributeValuesDataInterface::GetShaderParameters(TCHAR const* UID, FShaderParametersMetadataBuilder& InOutBuilder, FShaderParametersMetadataAllocations& InOutAllocations) const
@@ -67,7 +60,6 @@ void UPCGCountUniqueAttributeValuesDataInterface::GetHLSL(FString& OutHLSL, FStr
 		"int {DataInterfaceName}_OutputValueAttributeId;\n"
 		"int {DataInterfaceName}_OutputCountAttributeId;\n"
 		"uint {DataInterfaceName}_EmitPerDataCounts;\n"
-		"int {DataInterfaceName}_MaxAttributeValueAcrossAllData;\n"
 		"\n"
 		"int CountUniqueValues_GetAttributeToCountId_{DataInterfaceName}()\n"
 		"{\n"
@@ -84,11 +76,6 @@ void UPCGCountUniqueAttributeValuesDataInterface::GetHLSL(FString& OutHLSL, FStr
 		"	return {DataInterfaceName}_OutputCountAttributeId;\n"
 		"}\n"
 		"\n"
-		"uint CountUniqueValues_GetValueIndex_{DataInterfaceName}(uint InOutputDataIndex, int InAttributeValue)\n"
-		"{\n"
-		"	return (uint)( (int)InOutputDataIndex * {DataInterfaceName}_MaxAttributeValueAcrossAllData + InAttributeValue );\n"
-		"}\n"
-		"\n"
 		"bool CountUniqueValues_GetEmitPerDataCounts_{DataInterfaceName}()\n"
 		"{\n"
 		"	return {DataInterfaceName}_EmitPerDataCounts > 0;\n"
@@ -170,16 +157,6 @@ bool UPCGCountUniqueAttributeValuesDataProvider::PrepareForExecute_GameThread(UP
 		}
 	}
 
-	MaxAttributeValueAcrossAllData = INDEX_NONE;
-
-	for (const TArray<int32>& UniqueStringKeyValues : UniqueStringKeyValuesPerOutputData)
-	{
-		for (int32 Value : UniqueStringKeyValues)
-		{
-			MaxAttributeValueAcrossAllData = FMath::Max(MaxAttributeValueAcrossAllData, Value);
-		}
-	}
-
 	return true;
 }
 
@@ -187,7 +164,6 @@ FComputeDataProviderRenderProxy* UPCGCountUniqueAttributeValuesDataProvider::Get
 {
 	FPCGCountUniqueAttributeValuesProviderProxy::FCountUniqueValuesData_RenderThread Data;
 	Data.AttributeToCountId = AttributeToCountId;
-	Data.MaxAttributeValueAcrossAllData = MaxAttributeValueAcrossAllData;
 	Data.OutputValueAttributeId = OutputValueAttributeId;
 	Data.OutputCountAttributeId = OutputCountAttributeId;
 	Data.bEmitPerDataCounts = bEmitPerDataCounts;
@@ -203,7 +179,6 @@ void UPCGCountUniqueAttributeValuesDataProvider::Reset()
 	OutputValueAttributeId = INDEX_NONE;
 	OutputCountAttributeId = INDEX_NONE;
 	bEmitPerDataCounts = true;
-	MaxAttributeValueAcrossAllData = INDEX_NONE;
 
 	Super::Reset();
 }
@@ -224,6 +199,5 @@ void FPCGCountUniqueAttributeValuesProviderProxy::GatherDispatchData(FDispatchDa
 		Parameters.OutputValueAttributeId = Data.OutputValueAttributeId;
 		Parameters.OutputCountAttributeId = Data.OutputCountAttributeId;
 		Parameters.EmitPerDataCounts = Data.bEmitPerDataCounts ? 1u : 0u;
-		Parameters.MaxAttributeValueAcrossAllData = Data.MaxAttributeValueAcrossAllData + 1;
 	}
 }
diff --git a/Engine/Plugins/PCG/Source/PCG/Private/Compute/DataInterfaces/BuiltInKernels/PCGCountUniqueAttributeValuesDataInterface.h b/Engine/Plugins/PCG/Source/PCG/Private/Compute/DataInterfaces/BuiltInKernels/PCGCountUniqueAttributeValuesDataInterface.h
--- a/Engine/Plugins/PCG/Source/PCG/Private/Compute/DataInterfaces/BuiltInKernels/PCGCountUniqueAttributeValuesDataInterface.h
+++ b/Engine/Plugins/PCG/Source/PCG/Private/Compute/DataInterfaces/BuiltInKernels/PCGCountUniqueAttributeValuesDataInterface.h
@@ -61,8 +61,6 @@ public:
 	int32 OutputCountAttributeId = INDEX_NONE;
 
 	bool bEmitPerDataCounts = true;
-
-	int32 MaxAttributeValueAcrossAllData = INDEX_NONE;
 };
 
 class FPCGCountUniqueAttributeValuesProviderProxy : public FComputeDataProviderRenderProxy
@@ -74,7 +72,6 @@ public:
 		int32 OutputValueAttributeId = INDEX_NONE;
 		int32 OutputCountAttributeId = INDEX_NONE;
 		bool bEmitPerDataCounts = true;
-		int32 MaxAttributeValueAcrossAllData = INDEX_NONE;
 	};
 
 	FPCGCountUniqueAttributeValuesProviderProxy(FCountUniqueValuesData_RenderThread InData)
diff --git a/Engine/Plugins/PCG/Source/PCG/Private/Compute/DataInterfaces/Elements/PCGMetadataPartitionDataInterface.cpp b/Engine/Plugins/PCG/Source/PCG/Private/Compute/DataInterfaces/Elements/PCGMetadataPartitionDataInterface.cpp
--- a/Engine/Plugins/PCG/Source/PCG/Private/Compute/DataInterfaces/Elements/PCGMetadataPartitionDataInterface.cpp
+++ b/Engine/Plugins/PCG/Source/PCG/Private/Compute/DataInterfaces/Elements/PCGMetadataPartitionDataInterface.cpp
@@ -2,8 +2,11 @@
 
 #include "PCGMetadataPartitionDataInterface.h"
 
+#include "PCGParamData.h"
 #include "Compute/PCGComputeCommon.h"
 #include "Compute/PCGDataBinding.h"
+#include "Compute/BuiltInKernels/PCGCountUniqueAttributeValuesKernel.h"
+#include "Compute/Data/PCGProxyForGPUData.h"
 #include "Elements/Metadata/PCGMetadataPartition.h"
 #include "Elements/Metadata/PCGMetadataPartitionKernel.h"
 
@@ -23,26 +26,18 @@ void UPCGMetadataPartitionDataInterface::GetSupportedInputs(TArray<FShaderFuncti
 		.AddReturnType(EShaderFundamentalType::Int);
 
 	OutFunctions.AddDefaulted_GetRef()
-		.SetName(TEXT("MetadataPartition_GetNumPartitions"))
-		.AddReturnType(EShaderFundamentalType::Uint);
-
-	OutFunctions.AddDefaulted_GetRef()
-		.SetName(TEXT("MetadataPartition_GetPartitionIndex"))
-		.AddReturnType(EShaderFundamentalType::Uint)
-		.AddParam(EShaderFundamentalType::Uint); // InAttributeValue
-
-	OutFunctions.AddDefaulted_GetRef()
-		.SetName(TEXT("MetadataPartition_GetElementIndexToWriteTo"))
-		.AddReturnType(EShaderFundamentalType::Uint)
-		.AddParam(EShaderFundamentalType::Int) // InInputDataIndex
-		.AddParam(EShaderFundamentalType::Uint) // InPartitionIndex
-		.AddParam(EShaderFundamentalType::Uint); // InNumPartitions
+		.SetName(TEXT("MetadataPartition_GetOutputIndices"))
+		.AddParam(EShaderFundamentalType::Uint) // InInputDataIndex
+		.AddParam(EShaderFundamentalType::Int) // InAttributeValue
+		.AddParam(EShaderFundamentalType::Uint, 0, 0, EShaderParamModifier::Out) // OutDataIndex
+		.AddParam(EShaderFundamentalType::Uint, 0, 0, EShaderParamModifier::Out); // OutElementIndex
 }
 
 BEGIN_SHADER_PARAMETER_STRUCT(FPCGMetadataPartitionDataInterfaceParameters,)
 	SHADER_PARAMETER(int32, PartitionAttributeId)
-	SHADER_PARAMETER(uint32, NumPartitions)
-	SHADER_PARAMETER_RDG_BUFFER_SRV(StructuredBuffer<uint32>, AttributeValueToPartitionIndex)
+	SHADER_PARAMETER(int32, NumPartitions)
+	SHADER_PARAMETER_RDG_BUFFER_SRV(StructuredBuffer<int32>, AttributeValueToOutputDataIndex)
+	SHADER_PARAMETER_RDG_BUFFER_UAV(RWStructuredBuffer<uint32>, WriteCounters)
 END_SHADER_PARAMETER_STRUCT()
 
 void UPCGMetadataPartitionDataInterface::GetShaderParameters(TCHAR const* UID, FShaderParametersMetadataBuilder& InOutBuilder, FShaderParametersMetadataAllocations& InOutAllocations) const
@@ -58,25 +53,29 @@ void UPCGMetadataPartitionDataInterface::GetHLSL(FString& OutHLSL, FString const
 	};
 
 	OutHLSL += FString::Format(TEXT(
-		"StructuredBuffer<uint> {DataInterfaceName}_AttributeValueToPartitionIndex;\n"
+		"StructuredBuffer<int> {DataInterfaceName}_AttributeValueToOutputDataIndex;\n"
 		"int {DataInterfaceName}_PartitionAttributeId;\n"
-		"uint {DataInterfaceName}_NumPartitions;\n"
+		"int {DataInterfaceName}_NumPartitions;\n"
+		"RWStructuredBuffer<uint> {DataInterfaceName}_WriteCounters;\n"
 		"\n"
 		"int MetadataPartition_GetPartitionAttributeId_{DataInterfaceName}()\n"
 		"{\n"
 		"	return {DataInterfaceName}_PartitionAttributeId;\n"
 		"}\n"
 		"\n"
-		"uint MetadataPartition_GetNumPartitions_{DataInterfaceName}()\n"
+		"int MetadataPartition_GetOutputDataIndex_{DataInterfaceName}(uint InInputDataIndex, int InAttributeValue)\n"
 		"{\n"
-		"	return {DataInterfaceName}_NumPartitions;\n"
+		"	const int TableIndex = {DataInterfaceName}_NumPartitions * (int)InInputDataIndex + InAttributeValue;"
+		"	return {DataInterfaceName}_AttributeValueToOutputDataIndex[TableIndex];\n"
 		"}\n"
 		"\n"
-		"uint MetadataPartition_GetPartitionIndex_{DataInterfaceName}(int InAttributeValue)\n"
+		"void MetadataPartition_GetOutputIndices_{DataInterfaceName}(uint InInputDataIndex, int InAttributeValue, out uint OutDataIndex, out uint OutElementIndex)\n"
 		"{\n"
-		"	return {DataInterfaceName}_AttributeValueToPartitionIndex[InAttributeValue];\n"
+		"	const int CounterIndex = (int)InInputDataIndex * {DataInterfaceName}_NumPartitions + InAttributeValue;\n"
+		"	\n"
+		"	OutDataIndex = MetadataPartition_GetOutputDataIndex_{DataInterfaceName}(InInputDataIndex, InAttributeValue);"
+		"	InterlockedAdd({DataInterfaceName}_WriteCounters[CounterIndex], 1u, OutElementIndex);\n"
 		"}\n"
-		"\n"
 	), TemplateArgs);
 }
 
@@ -92,46 +91,144 @@ void UPCGMetaDataPartitionDataProvider::Initialize(const UComputeDataInterface*
 	Super::Initialize(InDataInterface, InBinding, InInputMask, InOutputMask);
 
 	const UPCGMetadataPartitionDataInterface* DataInterface = CastChecked<UPCGMetadataPartitionDataInterface>(InDataInterface);
-	Kernel = DataInterface->ProducerKernel;
+	ProducerKernel = DataInterface->ProducerKernel;
+}
+
+bool UPCGMetaDataPartitionDataProvider::PerformPreExecuteReadbacks_GameThread(UPCGDataBinding* InBinding)
+{
+	TRACE_CPUPROFILER_EVENT_SCOPE(UPCGMetaDataPartitionDataProvider::PerformPreExecuteReadbacks_GameThread);
+	check(InBinding);
+
+	if (!Super::PerformPreExecuteReadbacks_GameThread(InBinding))
+	{
+		return false;
+	}
+
+	if (AnalysisDataIndices.IsEmpty())
+	{
+		AnalysisDataIndices = InBinding->GetPinInputDataIndices(ProducerKernel, PCGMetadataPartitionConstants::ElementCountsPinLabel);
+
+		if (AnalysisDataIndices.IsEmpty())
+		{
+			UE_LOG(LogPCG, Verbose, TEXT("UPCGMetaDataPartitionDataProvider: No analysis data found from count kernel."));
+			return true;
+		}
+
+	}
+
+	bool bAllReadBack = true;
+
+	for (int32 AnalysisDataIndex : AnalysisDataIndices)
+	{
+		// Readback analysis data - poll until readback complete (return true).
+		if (!InBinding->ReadbackInputDataToCPU(AnalysisDataIndex))
+		{
+			bAllReadBack = false;
+		}
+	}
+
+	return bAllReadBack;
 }
 
 bool UPCGMetaDataPartitionDataProvider::PrepareForExecute_GameThread(UPCGDataBinding* InBinding)
 {
 	TRACE_CPUPROFILER_EVENT_SCOPE(UPCGMetadataPartitionDataInterface::PrepareForExecute_GameThread);
 	check(InBinding);
+	check(ProducerKernel);
 
 	if (!Super::PrepareForExecute_GameThread(InBinding))
 	{
 		return false;
 	}
 
-	check(Kernel);
-	const UPCGMetadataPartitionSettings* Settings = Cast<UPCGMetadataPartitionSettings>(Kernel->GetSettings());
-	check(Settings);
+	if (AnalysisDataIndices.IsEmpty())
+	{
+		UE_LOG(LogPCG, Verbose, TEXT("UPCGMetaDataPartitionDataProvider: No analysis data found from count kernel."));
+		return true;
+	}
 
-	if (ensure(Settings->PartitionAttributeSelectors.Num() == 1 && Settings->PartitionAttributeSelectors[0].IsBasicAttribute()))
+	const FPCGDataCollectionDesc* InputDesc = InBinding->GetCachedKernelPinDataDesc(ProducerKernel, PCGPinConstants::DefaultInputLabel, /*bIsInput=*/true);
+	if (!InputDesc)
 	{
-		const FPCGDataCollectionDesc* InputDataDesc = InBinding->GetCachedKernelPinDataDesc(Kernel, PCGPinConstants::DefaultInputLabel, /*bIsInput*/true);
+		UE_LOG(LogPCG, Error, TEXT("UPCGMetaDataPartitionDataProvider: Could not retrieve input data description."));
+		return true;
+	}
 
-		if (!ensure(InputDataDesc))
+	if (AnalysisDataIndices.Num() != InputDesc->DataDescs.Num())
+	{
+		UE_LOG(LogPCG, Error, TEXT("UPCGMetaDataPartitionDataProvider: Number of received analysis data items (%d) does not match input data count (%d)."),
+			AnalysisDataIndices.Num(), InputDesc->DataDescs.Num());
+		return true;
+	}
+
+	UniqueStringKeyValuesPerInputData.SetNum(AnalysisDataIndices.Num());
+	MaxAttributeValuePerInputData.SetNumZeroed(AnalysisDataIndices.Num());
+
+	for (int InputDataIndex = 0; InputDataIndex < AnalysisDataIndices.Num(); ++InputDataIndex)
+	{
+		const uint32 AnalysisDataIndex = AnalysisDataIndices[InputDataIndex];
+
+		if (!ensure(InBinding->InputDataCollection.TaggedData.IsValidIndex(AnalysisDataIndex)))
 		{
-			return true;
+			continue;
 		}
 
-		PartitionAttributeId = InBinding->GetAttributeId(Settings->PartitionAttributeSelectors[0].GetAttributeName(), EPCGKernelAttributeType::StringKey);
-		if (PartitionAttributeId != INDEX_NONE)
+		const UPCGParamData* AnalysisResultsData = Cast<UPCGParamData>(InBinding->InputDataCollection.TaggedData[AnalysisDataIndex].Data);
+		const UPCGMetadata* AnalysisMetadata = AnalysisResultsData ? AnalysisResultsData->ConstMetadata() : nullptr;
+
+		const FPCGMetadataAttributeBase* AnalysisCountAttributeBase = AnalysisMetadata ? AnalysisMetadata->GetConstAttribute(PCGCountUniqueAttributeValuesConstants::ValueCountAttributeName) : nullptr;
+
+		if (AnalysisCountAttributeBase && AnalysisCountAttributeBase->GetTypeId() == PCG::Private::MetadataTypes<int32>::Id)
+		{
+			const FPCGMetadataAttribute<int32>* CountAttribute = static_cast<const FPCGMetadataAttribute<int32>*>(AnalysisCountAttributeBase);
+
+			const int32 NumPartitionsInAnalysisData = AnalysisMetadata->GetItemCountForChild();
+
+			// TODO: Range based get would scale better.
+			for (int64 MetadataKey = 0; MetadataKey < NumPartitionsInAnalysisData; ++MetadataKey)
+			{
+				if (CountAttribute->GetValue(MetadataKey) > 0)
+				{
+					// The count kernel emits a count for each attribute value index, and the value attribute in the analysis data is just the identity map,
+					// so we don't bother reading the value attribute.
+					const int32 AttributeValue = static_cast<int32>(MetadataKey);
+
+					UniqueStringKeyValuesPerInputData[InputDataIndex].AddUnique(AttributeValue);
+
+					MaxAttributeValuePerInputData[InputDataIndex] = FMath::Max(MaxAttributeValuePerInputData[InputDataIndex], AttributeValue);
+				}
+			}
+		}
+		else
 		{
-			InputDataDesc->GetUniqueStringKeyValues(PartitionAttributeId, UniqueStringKeyValues);
+			if (InBinding->InputDataCollection.TaggedData[AnalysisDataIndex].Data->IsA<UPCGProxyForGPUData>())
+			{
+				UE_LOG(LogPCG, Error, TEXT("Data was not read back."));
+			}
+			else
+			{
+				UE_LOG(LogPCG, Error, TEXT("Analysis results attributes missing from read back data."));
+			}
+
+			return true;
 		}
+	}
+
+	const UPCGMetadataPartitionSettings* Settings = Cast<UPCGMetadataPartitionSettings>(ProducerKernel->GetSettings());
+	check(Settings);
 
-		MaxAttributeValue = INDEX_NONE;
-		for (int32 Key : UniqueStringKeyValues)
+	if (ensure(Settings->PartitionAttributeSelectors.Num() == 1 && Settings->PartitionAttributeSelectors[0].IsBasicAttribute()))
+	{
+		const FPCGDataCollectionDesc* InputDataDesc = InBinding->GetCachedKernelPinDataDesc(ProducerKernel, PCGPinConstants::DefaultInputLabel, /*bIsInput*/true);
+
+		if (!ensure(InputDataDesc))
 		{
-			MaxAttributeValue = FMath::Max(MaxAttributeValue, Key);
+			return true;
 		}
 
+		PartitionAttributeId = InBinding->GetAttributeId(Settings->PartitionAttributeSelectors[0].GetAttributeName(), EPCGKernelAttributeType::StringKey);
+
 		NumInputData = InputDataDesc->DataDescs.Num();
-		NumPartitions = UniqueStringKeyValues.Num();
 	}
 
 	return true;
@@ -139,54 +236,110 @@ bool UPCGMetaDataPartitionDataProvider::PrepareForExecute_GameThread(UPCGDataBin
 
 FComputeDataProviderRenderProxy* UPCGMetaDataPartitionDataProvider::GetRenderProxy()
 {
-	return new FPCGMetaDataPartitionProviderProxy(PartitionAttributeId, NumInputData, NumPartitions, MaxAttributeValue, UniqueStringKeyValues);
+	return new FPCGMetaDataPartitionProviderProxy(PartitionAttributeId, NumInputData/*, MaxAttributeValuePerInputData*/, UniqueStringKeyValuesPerInputData);
 }
 
 void UPCGMetaDataPartitionDataProvider::Reset()
 {
-	Kernel = nullptr;
 	PartitionAttributeId = INDEX_NONE;
-	NumInputData = 0;
-	NumPartitions = 0;
-	MaxAttributeValue = INDEX_NONE;
-	UniqueStringKeyValues.Empty();
+	NumInputData = INDEX_NONE;
+	MaxAttributeValuePerInputData.Empty(MaxAttributeValuePerInputData.Num());
+	UniqueStringKeyValuesPerInputData.Empty(UniqueStringKeyValuesPerInputData.Num());
+	AnalysisDataIndices.Empty(AnalysisDataIndices.Num());
 
 	Super::Reset();
 }
 
+FPCGMetaDataPartitionProviderProxy::FPCGMetaDataPartitionProviderProxy(int32 InPartitionAttributeId, int32 InNumInputData, TArray<TArray<int32>> InUniqueStringKeyValuesPerInputData)
+	: PartitionAttributeId(InPartitionAttributeId)
+	, NumInputData(InNumInputData)
+	, UniqueStringKeyValuesPerInputData(MoveTemp(InUniqueStringKeyValuesPerInputData))
+{
+	int32 MaxAttributeValueAcrossAllData = INDEX_NONE;
+
+	for (const TArray<int32>& UniqueStringKeyValues : UniqueStringKeyValuesPerInputData)
+	{
+		for (int32 StringKeyValue : UniqueStringKeyValues)
+		{
+			MaxAttributeValueAcrossAllData = FMath::Max(MaxAttributeValueAcrossAllData, StringKeyValue);
+		}
+	}
+
+	// We'll allocate space for all possible partitions across all input data.
+	NumPartitions = MaxAttributeValueAcrossAllData + 1;
+}
+
 bool FPCGMetaDataPartitionProviderProxy::IsValid(FValidationData const& InValidationData) const
 {
-	return InValidationData.ParameterStructSize == sizeof(FParameters)
-		&& !UniqueStringKeyValues.IsEmpty();
+	for (const TArray<int32>& UniqueStringKeyValues : UniqueStringKeyValuesPerInputData)
+	{
+		if (UniqueStringKeyValues.IsEmpty())
+		{
+			UE_LOG(LogPCG, Warning, TEXT("Metadata Partition will not execute, proxy invalid due to missing UniqueStringKeyValues."));
+			return false;
+		}
+	}
+
+	if (NumPartitions <= 0)
+	{
+		UE_LOG(LogPCG, Warning, TEXT("Metadata Partition will not execute, no partitions to generate."), NumPartitions);
+		return false;
+	}
+
+	if (PartitionAttributeId < 0)
+	{
+		UE_LOG(LogPCG, Warning, TEXT("Metadata Partition will not execute, partition attribute ID not set."));
+		return false;
+	}
+
+	return InValidationData.ParameterStructSize == sizeof(FParameters);
 }
 
 void FPCGMetaDataPartitionProviderProxy::AllocateResources(FRDGBuilder& GraphBuilder, FAllocationData const& InAllocationData)
 {
-	const int32 NumPartitionsToAllocate = FMath::Max(1u, NumInputData) * FMath::Max(1u, NumPartitions);
-	const int32 NumIndicesToAllocate = FMath::Max(1, MaxAttributeValue + 1);
+	// For each data, precompute a target partition index for each attribute value.
+	{
+		const int32 NumIndicesToAllocate = FMath::Max(1, NumInputData * NumPartitions);
 
-	TArray<uint32> Zeros;
-	Zeros.SetNumZeroed(FMath::Max(NumPartitionsToAllocate, NumIndicesToAllocate));
+		TArray<int32, TInlineAllocator<32>> Indices;
+		Indices.SetNumZeroed(NumIndicesToAllocate);
 
-	{
-		FRDGBufferDesc Desc = FRDGBufferDesc::CreateStructuredDesc(sizeof(uint32), NumIndicesToAllocate);
+		FRDGBufferDesc Desc = FRDGBufferDesc::CreateStructuredDesc(sizeof(int32), NumIndicesToAllocate);
 		Desc.Usage |= BUF_SourceCopy;
 
-		AttributeValueToPartitionIndex = GraphBuilder.CreateBuffer(Desc, TEXT("PCGAttributeValueToPartitionIndex"));
-		AttributeValueToPartitionIndexSRV = GraphBuilder.CreateSRV(AttributeValueToPartitionIndex);
+		FRDGBufferRef AttributeValueToOutputDataIndex = GraphBuilder.CreateBuffer(Desc, TEXT("PCGAttributeValueToOutputDataIndex"));
+		AttributeValueToOutputDataIndexSRV = GraphBuilder.CreateSRV(AttributeValueToOutputDataIndex);
+
+		int OutputDataIndex = 0;
 
-		int PartitionCounter = 0;
-		// Skipping default empty string.
-		for (int Index = 1; Index <= MaxAttributeValue; ++Index)
+		for (int InputDataIndex = 0; InputDataIndex < NumInputData; ++InputDataIndex)
 		{
-			const int PartitionIndex = UniqueStringKeyValues.IndexOfByKey(Index);
-			if (PartitionIndex != INDEX_NONE) // TODO - probably should make it uint and give it index_nones.
+			for (int ValueIndex = 0; ValueIndex < NumPartitions; ++ValueIndex)
 			{
-				Zeros[Index] = PartitionIndex;
+				const int PartitionIndex = UniqueStringKeyValuesPerInputData[InputDataIndex].IndexOfByKey(ValueIndex);
+				// Each possible attribute value that is present in the input data should go to one output data.
+				const int OutputIndex = (PartitionIndex == -1) ? -1 : (OutputDataIndex++);
+				Indices[InputDataIndex * NumPartitions + ValueIndex] = OutputIndex;
 			}
 		}
 
-		GraphBuilder.QueueBufferUpload(AttributeValueToPartitionIndex, Zeros.GetData(), NumIndicesToAllocate * Zeros.GetTypeSize());
+		GraphBuilder.QueueBufferUpload(AttributeValueToOutputDataIndex, MakeConstArrayView(Indices));
+	}
+
+	// For each data, set up and 0-initialize a write counter for each partition that GPU threads can atomically increment and use as write index.
+	{
+		const int32 NumWriteCountersToAllocate = FMath::Max(1, NumInputData * NumPartitions);
+
+		FRDGBufferDesc CountersDesc = FRDGBufferDesc::CreateStructuredDesc(sizeof(uint32), NumWriteCountersToAllocate);
+		CountersDesc.Usage |= BUF_SourceCopy;
+
+		TArray<uint32, TInlineAllocator<32>> Zeroes;
+		Zeroes.SetNumZeroed(NumWriteCountersToAllocate);
+
+		FRDGBufferRef WriteCounters = GraphBuilder.CreateBuffer(CountersDesc, TEXT("PCGDataCollection_Counters"));
+		GraphBuilder.QueueBufferUpload(WriteCounters, MakeConstArrayView(Zeroes));
+
+		WriteCountersUAV = GraphBuilder.CreateUAV(WriteCounters);
 	}
 }
 
@@ -199,6 +352,7 @@ void FPCGMetaDataPartitionProviderProxy::GatherDispatchData(FDispatchData const&
 
 		Parameters.PartitionAttributeId = PartitionAttributeId;
 		Parameters.NumPartitions = NumPartitions;
-		Parameters.AttributeValueToPartitionIndex = AttributeValueToPartitionIndexSRV;
+		Parameters.AttributeValueToOutputDataIndex = AttributeValueToOutputDataIndexSRV;
+		Parameters.WriteCounters = WriteCountersUAV;
 	}
 }
diff --git a/Engine/Plugins/PCG/Source/PCG/Private/Compute/DataInterfaces/Elements/PCGMetadataPartitionDataInterface.h b/Engine/Plugins/PCG/Source/PCG/Private/Compute/DataInterfaces/Elements/PCGMetadataPartitionDataInterface.h
--- a/Engine/Plugins/PCG/Source/PCG/Private/Compute/DataInterfaces/Elements/PCGMetadataPartitionDataInterface.h
+++ b/Engine/Plugins/PCG/Source/PCG/Private/Compute/DataInterfaces/Elements/PCGMetadataPartitionDataInterface.h
@@ -32,6 +32,7 @@ class UPCGMetaDataPartitionDataProvider : public UPCGComputeDataProvider
 
 public:
 	//~ Begin UPCGComputeDataProvider Interface
+	virtual bool PerformPreExecuteReadbacks_GameThread(UPCGDataBinding* InBinding) override;
 	virtual bool PrepareForExecute_GameThread(UPCGDataBinding* InBinding) override;
 	//~ End UPCGComputeDataProvider Interface
 
@@ -42,35 +43,21 @@ public:
 	//~ End UComputeDataProvider Interface
 
 public:
-	UPROPERTY()
-	TObjectPtr<const UPCGComputeKernel> Kernel;
-
-	UPROPERTY()
 	int32 PartitionAttributeId = INDEX_NONE;
 
-	UPROPERTY()
-	int32 NumInputData = 0;
+	int32 NumInputData = INDEX_NONE;
 
-	UPROPERTY()
-	uint32 NumPartitions = 0;
+	TArray<int32> MaxAttributeValuePerInputData;
 
-	UPROPERTY()
-	int32 MaxAttributeValue = INDEX_NONE;
+	TArray<TArray<int32>> UniqueStringKeyValuesPerInputData;
 
-	UPROPERTY()
-	TArray<int32> UniqueStringKeyValues;
+	TArray<int32> AnalysisDataIndices;
 };
 
 class FPCGMetaDataPartitionProviderProxy : public FComputeDataProviderRenderProxy
 {
 public:
-	FPCGMetaDataPartitionProviderProxy(int32 InPartitionAttributeId, int32 InNumInputData, uint32 InNumPartitions, int32 InMaxAttributeValue, TArray<int32> InUniqueStringKeyValues)
-		: PartitionAttributeId(InPartitionAttributeId)
-		, NumInputData(InNumInputData)
-		, NumPartitions(InNumPartitions)
-		, MaxAttributeValue(InMaxAttributeValue)
-		, UniqueStringKeyValues(MoveTemp(InUniqueStringKeyValues))
-	{}
+	FPCGMetaDataPartitionProviderProxy(int32 InPartitionAttributeId, int32 InNumInputData, TArray<TArray<int32>> InUniqueStringKeyValuesPerInputData);
 
 	//~ Begin FComputeDataProviderRenderProxy Interface
 	bool IsValid(FValidationData const& InValidationData) const override;
@@ -83,14 +70,13 @@ protected:
 
 	int32 PartitionAttributeId = INDEX_NONE;
 
-	uint32 NumInputData = 0;
+	int32 NumInputData = INDEX_NONE;
 
-	uint32 NumPartitions = 0;
+	int32 NumPartitions = INDEX_NONE;
 
-	int32 MaxAttributeValue = INDEX_NONE;
+	TArray<TArray<int32>> UniqueStringKeyValuesPerInputData;
 
-	TArray<int32> UniqueStringKeyValues;
+	FRDGBufferSRVRef AttributeValueToOutputDataIndexSRV = nullptr;
 
-	FRDGBufferRef AttributeValueToPartitionIndex = nullptr;
-	FRDGBufferSRVRef AttributeValueToPartitionIndexSRV = nullptr;
+	FRDGBufferUAVRef WriteCountersUAV = nullptr;
 };
diff --git a/Engine/Plugins/PCG/Source/PCG/Private/Compute/PCGDataBinding.cpp b/Engine/Plugins/PCG/Source/PCG/Private/Compute/PCGDataBinding.cpp
--- a/Engine/Plugins/PCG/Source/PCG/Private/Compute/PCGDataBinding.cpp
+++ b/Engine/Plugins/PCG/Source/PCG/Private/Compute/PCGDataBinding.cpp
@@ -281,6 +281,31 @@ int32 UPCGDataBinding::GetFirstInputDataIndex(const UPCGComputeKernel* InKernel,
 	return FirstDataIndex;
 }
 
+TArray<int32> UPCGDataBinding::GetPinInputDataIndices(const UPCGComputeKernel* InKernel, FName InPinLabel) const
+{
+	TArray<int32> PinDataIndices;
+
+	const int32 BindingIndex = Graph ? Graph->GetBindingIndex(FPCGKernelPin(InKernel->GetKernelIndex(), InPinLabel, /*bIsInput=*/true)) : INDEX_NONE;
+
+	if (ensure(BindingIndex != INDEX_NONE))
+	{
+		if (const FName* VirtualLabel = Graph->GetVirtualPinLabelFromBindingIndex(BindingIndex))
+		{
+			PinDataIndices.Reserve(InputDataCollection.TaggedData.Num());
+
+			for (int32 DataIndex = 0; DataIndex < InputDataCollection.TaggedData.Num(); ++DataIndex)
+			{
+				if (InputDataCollection.TaggedData[DataIndex].Pin == *VirtualLabel)
+				{
+					PinDataIndices.Add(DataIndex);
+				}
+			}
+		}
+	}
+
+	return PinDataIndices;
+}
+
 bool UPCGDataBinding::ReadbackInputDataToCPU(int32 InInputDataIndex)
 {
 	const UPCGProxyForGPUData* DataWithGPUSupport = InputDataCollection.TaggedData.IsValidIndex(InInputDataIndex) ? Cast<UPCGProxyForGPUData>(InputDataCollection.TaggedData[InInputDataIndex].Data) : nullptr;
diff --git a/Engine/Plugins/PCG/Source/PCG/Private/Compute/PCGDataBinding.h b/Engine/Plugins/PCG/Source/PCG/Private/Compute/PCGDataBinding.h
--- a/Engine/Plugins/PCG/Source/PCG/Private/Compute/PCGDataBinding.h
+++ b/Engine/Plugins/PCG/Source/PCG/Private/Compute/PCGDataBinding.h
@@ -126,6 +126,9 @@ public:
 	/** Helper to get index in input data collection of the first data item for the given kernel and input pin label. */
 	int32 GetFirstInputDataIndex(const UPCGComputeKernel* InKernel, FName InPinLabel) const;
 
+	/** Returns indices of data in the input data collection arriving on the given kernel pin. */
+	TArray<int32> GetPinInputDataIndices(const UPCGComputeKernel* InKernel, FName InPinLabel) const;
+
 	/** If there is a GPU proxy at the given index of the input data collection, triggers a readback and replaces the data item with CPU data if readback succeeds.
 	* Returns false while readback is in progress.
 	*/
diff --git a/Engine/Plugins/PCG/Source/PCG/Private/Elements/Metadata/PCGMetadataPartition.cpp b/Engine/Plugins/PCG/Source/PCG/Private/Elements/Metadata/PCGMetadataPartition.cpp
--- a/Engine/Plugins/PCG/Source/PCG/Private/Elements/Metadata/PCGMetadataPartition.cpp
+++ b/Engine/Plugins/PCG/Source/PCG/Private/Elements/Metadata/PCGMetadataPartition.cpp
@@ -5,6 +5,7 @@
 #include "PCGContext.h"
 #include "PCGCustomVersion.h"
 #include "Compute/PCGComputeCommon.h"
+#include "Compute/BuiltInKernels/PCGCountUniqueAttributeValuesKernel.h"
 #include "Elements/Metadata/PCGMetadataPartitionKernel.h"
 #include "Graph/PCGGPUGraphCompilationContext.h"
 #include "Helpers/PCGHelpers.h"
@@ -59,17 +60,32 @@ void UPCGMetadataPartitionSettings::PostLoad()
 #if WITH_EDITOR
 void UPCGMetadataPartitionSettings::CreateKernels(FPCGGPUCompilationContext& InOutContext, UObject* InObjectOuter, TArray<UPCGComputeKernel*>& OutKernels, TArray<FPCGKernelEdge>& OutEdges) const
 {
+	if (PartitionAttributeSelectors.IsEmpty())
+	{
+		ensure(false);
+		return;
+	}
+
 	FPCGComputeKernelParams KernelParams;
 	KernelParams.Settings = this;
 	KernelParams.bLogDescriptions = bDumpDataDescriptions;
 
-	UPCGMetadataPartitionKernel* Kernel = InOutContext.NewObject_AnyThread<UPCGMetadataPartitionKernel>(InObjectOuter);
-	Kernel->Initialize(KernelParams);
-	OutKernels.Add(Kernel);
+	UPCGMetadataPartitionKernel* PartitionKernel = InOutContext.NewObject_AnyThread<UPCGMetadataPartitionKernel>(InObjectOuter);
+	PartitionKernel->Initialize(KernelParams);
+	OutKernels.Add(PartitionKernel);
 
 	// Connect node pins to kernel pins
-	OutEdges.Emplace(FPCGPinReference(PCGPinConstants::DefaultInputLabel), FPCGPinReference(Kernel, PCGPinConstants::DefaultInputLabel));
-	OutEdges.Emplace(FPCGPinReference(Kernel, PCGPinConstants::DefaultOutputLabel), FPCGPinReference(PCGPinConstants::DefaultOutputLabel));
+	OutEdges.Emplace(FPCGPinReference(PCGPinConstants::DefaultInputLabel), FPCGPinReference(PartitionKernel, PCGPinConstants::DefaultInputLabel));
+	OutEdges.Emplace(FPCGPinReference(PartitionKernel, PCGPinConstants::DefaultOutputLabel), FPCGPinReference(PCGPinConstants::DefaultOutputLabel));
+
+	UPCGCountUniqueAttributeValuesKernel* CountKernel = InOutContext.NewObject_AnyThread<UPCGCountUniqueAttributeValuesKernel>(InObjectOuter);
+	CountKernel->Initialize(KernelParams);
+	CountKernel->SetAttributeName(PartitionAttributeSelectors[0].GetAttributeName());
+	CountKernel->SetEmitPerDataCounts(true);
+	OutKernels.Add(CountKernel);
+
+	OutEdges.Emplace(FPCGPinReference(PCGPinConstants::DefaultInputLabel), FPCGPinReference(CountKernel, PCGPinConstants::DefaultInputLabel));
+	OutEdges.Emplace(FPCGPinReference(CountKernel, PCGPinConstants::DefaultOutputLabel), FPCGPinReference(PartitionKernel, PCGMetadataPartitionConstants::ElementCountsPinLabel));
 }
 
 void UPCGMetadataPartitionSettings::ApplyDeprecation(UPCGNode* InOutNode)
diff --git a/Engine/Plugins/PCG/Source/PCG/Private/Elements/Metadata/PCGMetadataPartitionKernel.cpp b/Engine/Plugins/PCG/Source/PCG/Private/Elements/Metadata/PCGMetadataPartitionKernel.cpp
--- a/Engine/Plugins/PCG/Source/PCG/Private/Elements/Metadata/PCGMetadataPartitionKernel.cpp
+++ b/Engine/Plugins/PCG/Source/PCG/Private/Elements/Metadata/PCGMetadataPartitionKernel.cpp
@@ -4,9 +4,12 @@
 
 #include "PCGContext.h"
 #include "PCGModule.h"
+#include "PCGParamData.h"
 #include "Compute/PCGComputeGraph.h"
 #include "Compute/PCGDataBinding.h"
 #include "Compute/PCGPinPropertiesGPU.h"
+#include "Compute/BuiltInKernels/PCGCountUniqueAttributeValuesKernel.h"
+#include "Compute/Data/PCGProxyForGPUData.h"
 #include "Compute/DataInterfaces/Elements/PCGMetadataPartitionDataInterface.h"
 #include "Compute/Elements/PCGComputeGraphElement.h"
 #include "Elements/Metadata/PCGMetadataPartition.h"
@@ -44,14 +47,14 @@ bool UPCGMetadataPartitionKernel::IsKernelDataValid(FPCGContext* InContext) cons
 		FPCGKernelAttributeDesc AttributeDesc;
 		bool bConflictingTypesInData = false;
 		bool bPresentOnAllData = false;
-		const bool bAttributePresent = InputDesc->GetAttributeDesc(AttributeName, AttributeDesc, bConflictingTypesInData, bPresentOnAllData);
+		InputDesc->GetAttributeDesc(AttributeName, AttributeDesc, bConflictingTypesInData, bPresentOnAllData);
 
-		if (!bAttributePresent)
+		if (!bPresentOnAllData)
 		{
 			if (!InputDesc->DataDescs.IsEmpty())
 			{
 				PCG_KERNEL_VALIDATION_ERR(InContext, GetSettings(), FText::Format(
-					LOCTEXT("PartitionAttributeMissing", "Partition attribute was not present in GPU data."),
+					LOCTEXT("PartitionAttributeMissing", "Partition attribute was not present on all incoming data."),
 					FText::FromName(AttributeName)));
 			}
 
@@ -95,29 +98,96 @@ bool UPCGMetadataPartitionKernel::ComputeOutputBindingDataDesc(const UPCGCompute
 	const FPCGKernelPin InputKernelPin(KernelIndex, PCGPinConstants::DefaultInputLabel, /*bIsInput=*/true);
 	ensure(InBinding->ComputeKernelPinDataDesc(InputKernelPin, &InputPinDesc));
 
-	const int32 PartitionAttributeId = InBinding->GetAttributeId(MPSettings->PartitionAttributeSelectors[0].GetAttributeName(), EPCGKernelAttributeType::StringKey);
+	if (InputPinDesc.DataDescs.IsEmpty())
+	{
+		UE_LOG(LogPCG, Verbose, TEXT("UPCGMetadataPartitionKernel: No analysis data found from count kernel."));
+		return true;
+	}
+
+	const TArray<int32> AnalysisDataIndices = InBinding->GetPinInputDataIndices(this, PCGMetadataPartitionConstants::ElementCountsPinLabel);
+
+	if (AnalysisDataIndices.Num() != InputPinDesc.DataDescs.Num())
+	{
+		UE_LOG(LogPCG, Error, TEXT("UPCGMetadataPartitionKernel: Number of received analysis data items (%d) does not match input data count (%d)."),
+			AnalysisDataIndices.Num(), InputPinDesc.DataDescs.Num());
+		return true;
+	}
+
+	if (AnalysisDataIndices.IsEmpty())
+	{
+		return true;
+	}
+
+	bool bAllReadBack = true;
+
+	for (int32 AnalysisDataIndex : AnalysisDataIndices)
+	{
+		// Readback analysis data - poll until readback complete (return true).
+		if (!InBinding->ReadbackInputDataToCPU(AnalysisDataIndex))
+		{
+			bAllReadBack = false;
+		}
+	}
+
+	if (!bAllReadBack)
+	{
+		return false;
+	}
 
 	FPCGDataCollectionDesc PinDesc;
 
-	// Emit a point data for each potential partition.
-	if (PartitionAttributeId != INDEX_NONE)
+	for (int32 AnalysisDataIndex : AnalysisDataIndices)
 	{
-		TArray<int32> UniqueStringKeyValues;
-		InputPinDesc.GetUniqueStringKeyValues(PartitionAttributeId, UniqueStringKeyValues);
+		const UPCGParamData* AnalysisResultsData = Cast<UPCGParamData>(InBinding->InputDataCollection.TaggedData[AnalysisDataIndex].Data);
+		const UPCGMetadata* AnalysisMetadata = AnalysisResultsData ? AnalysisResultsData->ConstMetadata() : nullptr;
 
-		for (const FPCGDataDesc& InputData : InputPinDesc.DataDescs)
+		const FPCGMetadataAttributeBase* AnalysisCountAttributeBase = AnalysisMetadata ? AnalysisMetadata->GetConstAttribute(PCGCountUniqueAttributeValuesConstants::ValueCountAttributeName) : nullptr;
+
+		TArray<uint32> AttributeValueToCount;
+
+		if (AnalysisCountAttributeBase && AnalysisCountAttributeBase->GetTypeId() == PCG::Private::MetadataTypes<int32>::Id)
+		{
+			const FPCGMetadataAttribute<int32>* CountAttribute = static_cast<const FPCGMetadataAttribute<int32>*>(AnalysisCountAttributeBase);
+
+			const int32 NumElements = AnalysisMetadata->GetItemCountForChild();
+
+			AttributeValueToCount.Reserve(NumElements);
+
+			// TODO: Range based get would scale better.
+			for (int64 MetadataKey = 0; MetadataKey < NumElements; ++MetadataKey)
+			{
+				AttributeValueToCount.Add(CountAttribute->GetValue(MetadataKey));
+			}
+		}
+		else
 		{
-			// Before execution we don't know what string key value each attribute has, so we have to assume worst case - could be up to N points in M partitions.
-			for (int32 StringKeyValue : UniqueStringKeyValues)
+			UE_LOG(LogPCG, Warning, TEXT("No analysis data received by static mesh spawner kernel, worst case instance allocations will be made."));
+
+			if (InBinding->InputDataCollection.TaggedData[AnalysisDataIndex].Data->IsA<UPCGProxyForGPUData>())
 			{
-				FPCGDataDesc& Desc = PinDesc.DataDescs.Add_GetRef(InputData);
+				UE_LOG(LogPCG, Error, TEXT("Data was not read back."));
+			}
+
+			return true;
+		}
+
+		const int32 PartitionAttributeId = InBinding->GetAttributeId(MPSettings->PartitionAttributeSelectors[0].GetAttributeName(), EPCGKernelAttributeType::StringKey);
 
-				// Set the unique string keys to reflect the partition operation.
-				for (FPCGKernelAttributeDesc& AttrDesc : Desc.AttributeDescs)
+		for (int32 Value = 0; Value < AttributeValueToCount.Num(); ++Value)
+		{
+			if (AttributeValueToCount[Value] > 0)
+			{
+				FPCGDataDesc& DataDesc = PinDesc.DataDescs.Emplace_GetRef();
+				DataDesc.Type = EPCGDataType::Point;
+				DataDesc.ElementCount = AttributeValueToCount[Value];
+				DataDesc.AttributeDescs.Append(InputPinDesc.DataDescs[0].AttributeDescs);
+
+				for (FPCGKernelAttributeDesc& AttributeDesc : DataDesc.AttributeDescs)
 				{
-					if (AttrDesc.AttributeId == PartitionAttributeId)
+					if (AttributeDesc.AttributeId == PartitionAttributeId)
 					{
-						AttrDesc.SetStringKeys(MakeArrayView(&StringKeyValue, 1));
+						// By definition each data emitted by the partition node will have a single value for this attribute.
+						AttributeDesc.SetStringKeys(MakeConstArrayView(&Value, 1));
 						break;
 					}
 				}
@@ -132,13 +202,14 @@ bool UPCGMetadataPartitionKernel::ComputeOutputBindingDataDesc(const UPCGCompute
 	}
 
 	OutDataDesc = MoveTemp(PinDesc);
+
 	return true;
 }
 
 int UPCGMetadataPartitionKernel::ComputeThreadCount(const UPCGDataBinding* InBinding) const
 {
-	const FPCGDataCollectionDesc* InputPinDesc = InBinding->GetCachedKernelPinDataDesc(this, PCGPinConstants::DefaultInputLabel, /*bIsInput=*/true);
-	return ensure(InputPinDesc) ? InputPinDesc->ComputeDataElementCount(EPCGDataType::Point) : 0;
+	const FPCGDataCollectionDesc* OutputPinDesc = InBinding->GetCachedKernelPinDataDesc(this, PCGPinConstants::DefaultOutputLabel, /*bIsInput=*/false);
+	return ensure(OutputPinDesc) ? OutputPinDesc->ComputeDataElementCount(EPCGDataType::Point) : 0;
 }
 
 #if WITH_EDITOR
@@ -184,6 +255,7 @@ void UPCGMetadataPartitionKernel::GetKernelAttributeKeys(TArray<FPCGKernelAttrib
 void UPCGMetadataPartitionKernel::GetInputPins(TArray<FPCGPinProperties>& OutPins) const
 {
 	OutPins.Emplace(PCGPinConstants::DefaultInputLabel, EPCGDataType::Point);
+	OutPins.Emplace(PCGMetadataPartitionConstants::ElementCountsPinLabel, EPCGDataType::Param);
 }
 
 void UPCGMetadataPartitionKernel::GetOutputPins(TArray<FPCGPinPropertiesGPU>& OutPins) const
diff --git a/Engine/Plugins/PCG/Source/PCG/Public/Elements/Metadata/PCGMetadataPartitionKernel.h b/Engine/Plugins/PCG/Source/PCG/Public/Elements/Metadata/PCGMetadataPartitionKernel.h
--- a/Engine/Plugins/PCG/Source/PCG/Public/Elements/Metadata/PCGMetadataPartitionKernel.h
+++ b/Engine/Plugins/PCG/Source/PCG/Public/Elements/Metadata/PCGMetadataPartitionKernel.h
@@ -8,6 +8,11 @@
 
 class UPCGMetadataPartitionSettings;
 
+namespace PCGMetadataPartitionConstants
+{
+	const FName ElementCountsPinLabel = TEXT("ElementCounts");
+}
+
 UCLASS()
 class UPCGMetadataPartitionKernel : public UPCGComputeKernel
 {
@@ -25,7 +30,6 @@ public:
 	virtual bool SplitGraphAtOutput() const override { return true; }
 #endif
 	virtual void GetKernelAttributeKeys(TArray<FPCGKernelAttributeKey>& OutKeys) const override;
-	virtual bool DoesOutputPinRequireElementCounters(FName InOutputPinLabel) const override { return true; }
 	virtual void GetInputPins(TArray<FPCGPinProperties>& OutPins) const override;
 	virtual void GetOutputPins(TArray<FPCGPinPropertiesGPU>& OutPins) const override;
 
