[PCG] Remove ClearPerPinGeneratedOutput() from UPCGComponent::DirtyGenerated
- This was causing the PCGComponent to remain bGenerated but with no per pin output breaking Grid linkage in some specific cases (Landscape change while PCGActorAndComponentMapping::ShouldDiscardLandscapeRefresh returns true)
- Instead add ClearPerPinGeneratedOutput() in GenerateInternal and some other places where ClearGraphGeneratedOutput is also called

#rnx

diff --git a/Engine/Plugins/PCG/Source/PCG/Private/PCGComponent.cpp b/Engine/Plugins/PCG/Source/PCG/Private/PCGComponent.cpp
--- a/Engine/Plugins/PCG/Source/PCG/Private/PCGComponent.cpp
+++ b/Engine/Plugins/PCG/Source/PCG/Private/PCGComponent.cpp
@@ -610,6 +610,8 @@ FPCGTaskId UPCGComponent::GenerateInternal(bool bForce, EPCGHiGenGrid Grid, EPCG
 
 	if (CurrentGenerationTask != InvalidPCGTaskId)
 	{
+		ClearPerPinGeneratedOutput();
+
 #if WITH_EDITOR
 		// Notify Subsystem first
 		GetSubsystem()->OnPCGGraphStartGenerating(this);
@@ -918,6 +920,7 @@ void UPCGComponent::PostCleanupGraph(bool bRemoveComponents)
 	const bool bHadGeneratedGraphOutput = GeneratedGraphOutput.TaggedData.Num() > 0;
 
 	ClearGraphGeneratedOutput();
+	ClearPerPinGeneratedOutput();
 
 #if WITH_EDITOR
 	DynamicallyTrackedKeysToSettings.Reset();
@@ -1923,6 +1926,7 @@ void UPCGComponent::MarkSubObjectsAsGarbage()
 {
 	// First move out the generated output
 	ClearGraphGeneratedOutput();
+	ClearPerPinGeneratedOutput();
 
 	auto ModifyAndMarkForGarbage = [](UObject* SubObject)
 	{
@@ -2673,8 +2677,6 @@ void UPCGComponent::DirtyGenerated(EPCGComponentDirtyFlag DirtyFlag, const bool
 
 	bDirtyGenerated = true;
 
-	ClearPerPinGeneratedOutput();
-
 	// Dirty data as a waterfall from basic values
 	if (!!(DirtyFlag & EPCGComponentDirtyFlag::Actor))
 	{
