[PCG] Move NotifyGraphChanged to PCGModule so that PCGGraph doesn't need to discover the proper Subsystem to notify

#rnx

diff --git a/Engine/Plugins/PCG/Source/PCG/Private/Graph/PCGGraphExecutor.cpp b/Engine/Plugins/PCG/Source/PCG/Private/Graph/PCGGraphExecutor.cpp
--- a/Engine/Plugins/PCG/Source/PCG/Private/Graph/PCGGraphExecutor.cpp
+++ b/Engine/Plugins/PCG/Source/PCG/Private/Graph/PCGGraphExecutor.cpp
@@ -307,7 +307,7 @@ FPCGTaskId FPCGGraphExecutor::ScheduleGraph(const FPCGScheduleGraphParams& InPar
 		PCGEditorModule->OnScheduleGraph(*StackContextPtr);
 	}
 
-	if (UPCGSubsystem* Subsystem = UWorld::GetSubsystem<UPCGSubsystem>(InParams.ExecutionSource->GetExecutionState().GetWorld()))
+	if (IPCGBaseSubsystem* Subsystem = InParams.ExecutionSource->GetExecutionState().GetSubsystem())
 	{
 		Subsystem->OnScheduleGraph(*StackContextPtr);
 	}
@@ -696,6 +696,24 @@ int32 FPCGGraphExecutor::GetNonScheduledRemainingTaskCount() const
 	return NewReadyTaskCount + Tasks.Num() + ReadyTasks.Num() + ActiveTasks.Num() + ActiveTasksGameThreadOnly.Num() + PausedTasks.Num();
 }
 
+IPCGBaseSubsystem* FPCGGraphExecutor::GetSubsystem() const
+{
+	if (World)
+	{
+		return UPCGSubsystem::GetInstance(World);
+	}
+#if WITH_EDITOR
+	else if(IPCGEditorModule* EditorModule = IPCGEditorModule::Get())
+	{
+		return EditorModule->GetEditorSubsystem();
+	}
+#endif
+	else
+	{
+		return nullptr;
+	}
+}
+
 UPCGData* FPCGGraphExecutor::GetExecutionCacheData(FPCGTaskId InGraphExecutionTaskId, EExecutionCacheDataType InExecutionCacheDataType)
 {
 	if (InGraphExecutionTaskId != InvalidPCGTaskId)
@@ -2859,7 +2877,7 @@ void FPCGGraphExecutor::OnNotificationCancel()
 {
 	ReleaseGenerationNotification();
 
-	if (UPCGSubsystem* PCGSubsystem = UWorld::GetSubsystem<UPCGSubsystem>(World))
+	if (IPCGBaseSubsystem* PCGSubsystem = GetSubsystem())
 	{
 		PCGSubsystem->CancelAllGeneration();
 	}
diff --git a/Engine/Plugins/PCG/Source/PCG/Private/Graph/PCGGraphExecutor.h b/Engine/Plugins/PCG/Source/PCG/Private/Graph/PCGGraphExecutor.h
--- a/Engine/Plugins/PCG/Source/PCG/Private/Graph/PCGGraphExecutor.h
+++ b/Engine/Plugins/PCG/Source/PCG/Private/Graph/PCGGraphExecutor.h
@@ -158,6 +158,8 @@ private:
 	friend UPCGSubsystem;
 	friend IPCGBaseSubsystem;
 
+	IPCGBaseSubsystem* GetSubsystem() const;
+
 	enum class EExecutionCacheDataType : int8
 	{
 		PCGData = 0,
diff --git a/Engine/Plugins/PCG/Source/PCG/Private/PCGGraph.cpp b/Engine/Plugins/PCG/Source/PCG/Private/PCGGraph.cpp
--- a/Engine/Plugins/PCG/Source/PCG/Private/PCGGraph.cpp
+++ b/Engine/Plugins/PCG/Source/PCG/Private/PCGGraph.cpp
@@ -785,19 +785,7 @@ void UPCGGraph::BeginDestroy()
 		InputNode->OnNodeChangedDelegate.RemoveAll(this);
 	}
 
-	EPCGChangeType ChangeTypes = EPCGChangeType::Structural | EPCGChangeType::GenerationGrid;
-
-	if (IPCGEditorModule* PCGEditorModule = IPCGEditorModule::Get())
-	{
-		PCGEditorModule->NotifyGraphChanged(this, ChangeTypes);
-	}
-
-	// Notify the compiler to remove this graph from its cache
-	if (UPCGSubsystem* PCGSubsystem = UPCGSubsystem::GetActiveEditorInstance())
-	{
-		PCGSubsystem->NotifyGraphChanged(this, ChangeTypes);
-	}
-
+	FPCGModule::GetPCGModuleChecked().NotifyGraphChanged(this, EPCGChangeType::Structural | EPCGChangeType::GenerationGrid);
 #endif
 
 	Super::BeginDestroy();
@@ -1777,25 +1765,9 @@ void UPCGGraph::NotifyGraphChanged(EPCGChangeType ChangeType)
 
 	// Notify the subsystem/compiler cache before so it gets recompiled properly
 	const bool bNotifySubsystem = ((ChangeType & (EPCGChangeType::Structural | EPCGChangeType::Edge)) != EPCGChangeType::None);
-	if (bNotifySubsystem && GEditor)
+	if (bNotifySubsystem)
 	{
-		if (IPCGEditorModule* PCGEditorModule = IPCGEditorModule::Get())
-		{
-			PCGEditorModule->NotifyGraphChanged(this, ChangeType);
-		}
-
-		if (GEditor->PlayWorld)
-		{
-			if (UPCGSubsystem* PCGPIESubsystem = UPCGSubsystem::GetInstance(GEditor->PlayWorld.Get()))
-			{
-				PCGPIESubsystem->NotifyGraphChanged(this, ChangeType);
-			}
-		}
-
-		if (UPCGSubsystem* PCGSubsystem = UPCGSubsystem::GetInstance(GEditor->GetEditorWorldContext().World()))
-		{
-			PCGSubsystem->NotifyGraphChanged(this, ChangeType);
-		}
+		FPCGModule::GetPCGModuleChecked().NotifyGraphChanged(this, ChangeType);
 	}
 
 	if (bNonCosmeticChange)
diff --git a/Engine/Plugins/PCG/Source/PCG/Private/PCGModule.cpp b/Engine/Plugins/PCG/Source/PCG/Private/PCGModule.cpp
--- a/Engine/Plugins/PCG/Source/PCG/Private/PCGModule.cpp
+++ b/Engine/Plugins/PCG/Source/PCG/Private/PCGModule.cpp
@@ -25,6 +25,10 @@
 
 #define LOCTEXT_NAMESPACE "FPCGModule"
 
+#if WITH_EDITOR
+FPCGGraphChangedDelegate FPCGModule::OnGraphChangedDelegate;
+#endif
+
 namespace PCGModule
 {
 	FPCGModule* PCGModulePtr = nullptr;
diff --git a/Engine/Plugins/PCG/Source/PCG/Private/Subsystems/IPCGBaseSubsystem.cpp b/Engine/Plugins/PCG/Source/PCG/Private/Subsystems/IPCGBaseSubsystem.cpp
--- a/Engine/Plugins/PCG/Source/PCG/Private/Subsystems/IPCGBaseSubsystem.cpp
+++ b/Engine/Plugins/PCG/Source/PCG/Private/Subsystems/IPCGBaseSubsystem.cpp
@@ -5,6 +5,7 @@
 #include "PCGActorAndComponentMapping.h"
 #include "PCGComponent.h"
 #include "PCGGraph.h"
+#include "PCGModule.h"
 #include "Graph/PCGGraphCache.h"
 #include "Graph/PCGGraphCompiler.h"
 #include "Graph/PCGGraphExecutor.h"
@@ -24,10 +25,18 @@ void IPCGBaseSubsystem::InitializeBaseSubsystem()
 {
 	check(!GraphExecutor);
 	GraphExecutor = MakeShared<FPCGGraphExecutor>(GetSubsystemWorld());
+
+#if WITH_EDITOR
+	FPCGModule::GetPCGModuleChecked().OnGraphChanged().AddRaw(this, &IPCGBaseSubsystem::NotifyGraphChanged);
+#endif
 }
 
 void IPCGBaseSubsystem::DeinitializeBaseSubsystem()
 {
+#if WITH_EDITOR
+	FPCGModule::GetPCGModuleChecked().OnGraphChanged().RemoveAll(this);
+#endif
+
 	GraphExecutor = nullptr;
 }
 
diff --git a/Engine/Plugins/PCG/Source/PCG/Public/Editor/IPCGEditorModule.h b/Engine/Plugins/PCG/Source/PCG/Public/Editor/IPCGEditorModule.h
--- a/Engine/Plugins/PCG/Source/PCG/Public/Editor/IPCGEditorModule.h
+++ b/Engine/Plugins/PCG/Source/PCG/Public/Editor/IPCGEditorModule.h
@@ -37,6 +37,7 @@ protected:
 	friend class SPCGEditorGraphFind;
 	friend class UPCGGraph;
 	friend class UPCGSubsystem;
+	friend class UPCGEditorSubsystem;
 
 	virtual void OnScheduleGraph(const FPCGStackContext& StackContext) = 0;
 	virtual void OnGraphPreSave(UPCGGraph* Graph, FObjectPreSaveContext ObjectSaveContext) = 0;
diff --git a/Engine/Plugins/PCG/Source/PCG/Public/PCGModule.h b/Engine/Plugins/PCG/Source/PCG/Public/PCGModule.h
--- a/Engine/Plugins/PCG/Source/PCG/Public/PCGModule.h
+++ b/Engine/Plugins/PCG/Source/PCG/Public/PCGModule.h
@@ -8,6 +8,7 @@
 #include "Utils/PCGLogErrors.h"
 
 #include "Containers/Ticker.h"
+#include "Delegates/Delegate.h"
 #include "Modules/ModuleInterface.h"
 
 // Logs
@@ -26,6 +27,9 @@ DECLARE_STATS_GROUP(TEXT("PCG"), STATGROUP_PCG, STATCAT_Advanced);
 
 // CVars
 
+// Delegates
+DECLARE_MULTICAST_DELEGATE_TwoParams(FPCGGraphChangedDelegate, UPCGGraph* /*InGraph*/, EPCGChangeType /*ChangeType*/);
+
 class FPCGModule final : public IModuleInterface
 {
 public:
@@ -48,6 +52,10 @@ public:
 	
 	PCG_API void ExecuteNextTick(TFunction<void()> TickFunction);
 
+#if WITH_EDITOR
+	static FPCGGraphChangedDelegate& OnGraphChanged() { return OnGraphChangedDelegate; }
+#endif
+
 private:
 	bool Tick(float DeltaTime);
 	
@@ -63,14 +71,17 @@ private:
 private:
 	void RegisterNativeElementDeterminismTests();
 	void DeregisterNativeElementDeterminismTests();
-#endif
 
-#if WITH_EDITOR
 public:
 	static const FPCGDataVisualizationRegistry& GetConstPCGDataVisualizationRegistry() { return GetPCGModuleChecked().PCGDataVisualizationRegistry; }
 	static FPCGDataVisualizationRegistry& GetMutablePCGDataVisualizationRegistry() { return GetPCGModuleChecked().PCGDataVisualizationRegistry; }
 
 private:
 	FPCGDataVisualizationRegistry PCGDataVisualizationRegistry;
-#endif
-};
+
+	friend class UPCGGraph;
+	void NotifyGraphChanged(UPCGGraph* InGraph, EPCGChangeType ChangeType) { OnGraphChangedDelegate.Broadcast(InGraph, ChangeType); }
+
+	static FPCGGraphChangedDelegate OnGraphChangedDelegate;
+#endif // WITH_EDITOR
+};
\ No newline at end of file
diff --git a/Engine/Plugins/PCG/Source/PCG/Public/Subsystems/IPCGBaseSubsystem.h b/Engine/Plugins/PCG/Source/PCG/Public/Subsystems/IPCGBaseSubsystem.h
--- a/Engine/Plugins/PCG/Source/PCG/Public/Subsystems/IPCGBaseSubsystem.h
+++ b/Engine/Plugins/PCG/Source/PCG/Public/Subsystems/IPCGBaseSubsystem.h
@@ -98,7 +98,7 @@ public:
 #if WITH_EDITOR
 public:
 	/** Propagate to the graph compiler graph changes */
-	PCG_API void NotifyGraphChanged(UPCGGraph* InGraph, EPCGChangeType ChangeType);
+	PCG_API virtual void NotifyGraphChanged(UPCGGraph* InGraph, EPCGChangeType ChangeType);
 
 	/** Cleans up the graph cache on an element basis. InSettings is used for debugging and is optional. */
 	PCG_API void CleanFromCache(const IPCGElement* InElement, const UPCGSettings* InSettings = nullptr);
diff --git a/Engine/Plugins/PCG/Source/PCGEditor/Private/Subsystems/PCGEditorSubsystem.cpp b/Engine/Plugins/PCG/Source/PCGEditor/Private/Subsystems/PCGEditorSubsystem.cpp
--- a/Engine/Plugins/PCG/Source/PCGEditor/Private/Subsystems/PCGEditorSubsystem.cpp
+++ b/Engine/Plugins/PCG/Source/PCGEditor/Private/Subsystems/PCGEditorSubsystem.cpp
@@ -1,7 +1,8 @@
-﻿// Copyright Epic Games, Inc. All Rights Reserved.
+// Copyright Epic Games, Inc. All Rights Reserved.
 
 #include "Subsystems/PCGEditorSubsystem.h"
 
+#include "Editor/IPCGEditorModule.h"
 #include "Editor/PCGEditorDefaultExecutionSource.h"
 
 #if WITH_EDITOR
@@ -45,4 +46,14 @@ TStatId UPCGEditorSubsystem::GetStatId() const
 	RETURN_QUICK_DECLARE_CYCLE_STAT(FPCGEditorSubsystemTickFunction, STATGROUP_Tickables);
 }
 
+void UPCGEditorSubsystem::NotifyGraphChanged(UPCGGraph* InGraph, EPCGChangeType ChangeType)
+{
+	if (IPCGEditorModule* PCGEditorModule = IPCGEditorModule::Get())
+	{
+		PCGEditorModule->NotifyGraphChanged(InGraph, ChangeType);
+	}
+
+	IPCGBaseSubsystem::NotifyGraphChanged(InGraph, ChangeType);
+}
+
 #endif // WITH_EDITOR
diff --git a/Engine/Plugins/PCG/Source/PCGEditor/Public/Subsystems/PCGEditorSubsystem.h b/Engine/Plugins/PCG/Source/PCGEditor/Public/Subsystems/PCGEditorSubsystem.h
--- a/Engine/Plugins/PCG/Source/PCGEditor/Public/Subsystems/PCGEditorSubsystem.h
+++ b/Engine/Plugins/PCG/Source/PCGEditor/Public/Subsystems/PCGEditorSubsystem.h
@@ -1,4 +1,4 @@
-﻿// Copyright Epic Games, Inc. All Rights Reserved.
+// Copyright Epic Games, Inc. All Rights Reserved.
 
 #pragma once
 
@@ -23,6 +23,8 @@ public:
 
 	PCGEDITOR_API UPCGEditorDefaultExecutionSource* CreateEditorExecutionSource(const FPCGEditorDefaultExecutionSourceParams& InParams);
 
+	PCGEDITOR_API virtual void NotifyGraphChanged(UPCGGraph* InGraph, EPCGChangeType ChangeType) override;
+
 protected:
 	//** FTickableEditorObject Interface
 	PCGEDITOR_API virtual void Tick(float DeltaTime) override;
