[PCG] Stand up the Break Debugger debug feature for GPU nodes, and make the breakpoint conditional on inspection so that it is usable in partitioned/higen use cases.

diff --git a/Engine/Plugins/PCG/Source/PCG/Private/Compute/Elements/PCGComputeGraphElement.cpp b/Engine/Plugins/PCG/Source/PCG/Private/Compute/Elements/PCGComputeGraphElement.cpp
--- a/Engine/Plugins/PCG/Source/PCG/Private/Compute/Elements/PCGComputeGraphElement.cpp
+++ b/Engine/Plugins/PCG/Source/PCG/Private/Compute/Elements/PCGComputeGraphElement.cpp
@@ -27,11 +27,24 @@
 #define LOCTEXT_NAMESPACE "PCGComputeGraphElement"
 
 #if WITH_EDITOR
+#define PCG_COMPUTE_GRAPH_BREAKPOINT() \
+	if (Context->ComputeGraph && Context->ComputeGraph->bBreakDebugger && Context->IsExecutingGraphInspected()) \
+	{ \
+		if (const UEnum* EnumPtr = StaticEnum<EPCGComputeGraphExecutionPhase>()) \
+		{ \
+			UE_LOG(LogPCG, Log, TEXT("BREAKPOINT: FPCGComputeGraphElement: %s"), *EnumPtr->GetDisplayNameTextByValue(static_cast<int64>(Context->ExecutionSubPhase)).ToString()); \
+		} \
+		\
+		UE_DEBUG_BREAK(); \
+	}
+
 static TAutoConsoleVariable<bool> CVarReadbackTextureDataOnInspect(
 	TEXT("pcg.GPU.ReadbackTextureDataOnInspect"),
 	true,
 	TEXT("Reads texture data back to the CPU when inspected."));
-#endif
+#else // WITH_EDITOR
+#define PCG_COMPUTE_GRAPH_BREAKPOINT()
+#endif // WITH_EDITOR
 
 FPCGComputeGraphContext::~FPCGComputeGraphContext()
 {
@@ -105,14 +118,16 @@ bool FPCGComputeGraphElement::ExecuteInternal(FPCGContext* InContext) const
 	TRACE_CPUPROFILER_EVENT_SCOPE(FPCGComputeGraphElement::ExecuteInternal);
 	check(InContext);
 
+	FPCGComputeGraphContext* Context = static_cast<FPCGComputeGraphContext*>(InContext);
+
+	PCG_COMPUTE_GRAPH_BREAKPOINT();
+
 	if (!InContext->ExecutionSource.IsValid())
 	{
 		UE_LOG(LogPCG, Warning, TEXT("FPCGComputeGraphElement: Execution source lost, element execution halted."));
 		return true;
 	}
 
-	FPCGComputeGraphContext* Context = static_cast<FPCGComputeGraphContext*>(InContext);
-
 	auto SleepUntilNextFrame = [Context]()
 	{
 		Context->bIsPaused = true;
@@ -166,6 +181,8 @@ bool FPCGComputeGraphElement::ExecuteInternal(FPCGContext* InContext) const
 			return true;
 		}
 
+		PCG_COMPUTE_GRAPH_BREAKPOINT();
+
 		Context->ComputeGraphIndex = ComputeGraphIndex;
 
 		if (!Context->ComputeGraph->AreGraphSettingsValid(Context))
@@ -227,8 +244,7 @@ bool FPCGComputeGraphElement::ExecuteInternal(FPCGContext* InContext) const
 		}
 
 #if WITH_EDITOR
-		const bool bComponentInspected = Context->GetStack() && Graph && Graph->IsInspecting() && Graph->GetInspectedStack() == *Context->GetStack();
-		Context->ComputeGraphInstance->SetRenderCapturesEnabled(bComponentInspected);
+		Context->ComputeGraphInstance->SetRenderCapturesEnabled(Context->IsExecutingGraphInspected());
 #endif
 
 		Context->bComputeGraphInstanceInitialized = true;
@@ -590,16 +606,21 @@ bool FPCGComputeGraphElement::ExecuteInternal(FPCGContext* InContext) const
 			Context->OutputData = Context->DataBinding->OutputDataCollection;
 		}
 
-		Context->ExecutionSubPhase = EPCGComputeGraphExecutionPhase::DebugAndInspection;
-	}
-	case EPCGComputeGraphExecutionPhase::DebugAndInspection: // Fallthrough
-	{
 #if WITH_EDITOR
+		Context->ExecutionSubPhase = EPCGComputeGraphExecutionPhase::DebugAndInspection;
+
 		if (Context->AsyncState.ShouldStop())
 		{
 			return false;
 		}
-
+#else
+		Context->ExecutionSubPhase = EPCGComputeGraphExecutionPhase::Done;
+		return true;
+#endif
+	}
+	case EPCGComputeGraphExecutionPhase::DebugAndInspection: // Fallthrough
+	{
+#if WITH_EDITOR
 		CollectDebugDataPrepareActions(Context);
 
 		bool bAllReady = true;
@@ -625,6 +646,8 @@ bool FPCGComputeGraphElement::ExecuteInternal(FPCGContext* InContext) const
 
 		ExecuteDebugDraw(Context);
 		StoreDataForInspection(Context);
+
+		Context->ExecutionSubPhase = EPCGComputeGraphExecutionPhase::Done;
 #endif
 	}
 	}
@@ -897,4 +920,5 @@ FPCGElementPtr UPCGComputeGraphSettings::CreateElement() const
 	return MakeShared<FPCGComputeGraphElement>(ComputeGraphIndex);
 }
 
+#undef PCG_COMPUTE_GRAPH_BREAKPOINT
 #undef LOCTEXT_NAMESPACE
diff --git a/Engine/Plugins/PCG/Source/PCG/Private/Compute/Elements/PCGComputeGraphElement.h b/Engine/Plugins/PCG/Source/PCG/Private/Compute/Elements/PCGComputeGraphElement.h
--- a/Engine/Plugins/PCG/Source/PCG/Private/Compute/Elements/PCGComputeGraphElement.h
+++ b/Engine/Plugins/PCG/Source/PCG/Private/Compute/Elements/PCGComputeGraphElement.h
@@ -30,6 +30,7 @@ enum class EPCGComputeGraphExecutionPhase : uint8
 	WaitForExecutionComplete,
 	PostExecute,
 	DebugAndInspection,
+	Done,
 };
 
 struct FPCGComputeGraphContext : public FPCGContext
@@ -40,6 +41,10 @@ public:
 
 	bool HasPendingAsyncOperations() const;
 
+#if WITH_EDITOR
+	void DoDebugBreak() const;
+#endif
+
 protected:
 	virtual void AddExtraStructReferencedObjects(FReferenceCollector& Collector) override;
 
diff --git a/Engine/Plugins/PCG/Source/PCG/Private/Compute/PCGComputeGraph.h b/Engine/Plugins/PCG/Source/PCG/Private/Compute/PCGComputeGraph.h
--- a/Engine/Plugins/PCG/Source/PCG/Private/Compute/PCGComputeGraph.h
+++ b/Engine/Plugins/PCG/Source/PCG/Private/Compute/PCGComputeGraph.h
@@ -118,6 +118,11 @@ public:
 	UPROPERTY()
 	bool bLogDataDescriptions = false;
 
+#if WITH_EDITORONLY_DATA
+	UPROPERTY()
+	bool bBreakDebugger = false;
+#endif
+
 protected:
 	UPROPERTY()
 	TMap</*BindingIndex*/int32, /*PinLabel*/FName> KernelBindingToPinLabel;
diff --git a/Engine/Plugins/PCG/Source/PCG/Private/Graph/PCGGraphCompilerGPU.cpp b/Engine/Plugins/PCG/Source/PCG/Private/Graph/PCGGraphCompilerGPU.cpp
--- a/Engine/Plugins/PCG/Source/PCG/Private/Graph/PCGGraphCompilerGPU.cpp
+++ b/Engine/Plugins/PCG/Source/PCG/Private/Graph/PCGGraphCompilerGPU.cpp
@@ -1409,6 +1409,7 @@ void FPCGGraphCompilerGPU::CompileComputeGraph(
 		}
 
 		InOutComputeGraph->bLogDataDescriptions |= Settings->bDumpDataDescriptions;
+		InOutComputeGraph->bBreakDebugger |= Settings->bBreakDebugger;
 
 		// Now that all data interfaces added, create the (trivial) binding mapping. All map to primary binding, index 0.
 		InOutComputeGraph->DataInterfaceToBinding.SetNumZeroed(InOutComputeGraph->DataInterfaces.Num());
diff --git a/Engine/Plugins/PCG/Source/PCG/Private/PCGContext.cpp b/Engine/Plugins/PCG/Source/PCG/Private/PCGContext.cpp
--- a/Engine/Plugins/PCG/Source/PCG/Private/PCGContext.cpp
+++ b/Engine/Plugins/PCG/Source/PCG/Private/PCGContext.cpp
@@ -529,6 +529,14 @@ bool FPCGContext::HasVisualLogs() const
 
 	return false;
 }
+
+bool FPCGContext::IsExecutingGraphInspected() const
+{
+	const FPCGStack* StackPtr = GetStack();
+	const UPCGGraph* Graph = StackPtr ? StackPtr->GetGraphForCurrentFrame() : nullptr;
+
+	return Graph && Graph->IsInspecting() && Graph->GetInspectedStack() == *StackPtr;
+}
 #endif // WITH_EDITOR
 
 void FPCGContext::AddStructReferencedObjects(FReferenceCollector& Collector)
diff --git a/Engine/Plugins/PCG/Source/PCG/Private/PCGElement.cpp b/Engine/Plugins/PCG/Source/PCG/Private/PCGElement.cpp
--- a/Engine/Plugins/PCG/Source/PCG/Private/PCGElement.cpp
+++ b/Engine/Plugins/PCG/Source/PCG/Private/PCGElement.cpp
@@ -47,7 +47,7 @@ static TAutoConsoleVariable<bool> CVarPCGEnablePointArrayDataToPointDataConversi
 
 #if WITH_EDITOR
 #define PCG_ELEMENT_EXECUTION_BREAKPOINT() \
-	if (Context && Context->GetInputSettingsInterface() && Context->GetInputSettingsInterface()->bBreakDebugger) \
+	if (Context && Context->IsExecutingGraphInspected() && Context->GetInputSettingsInterface() && Context->GetInputSettingsInterface()->bBreakDebugger) \
 	{ \
 		UE_DEBUG_BREAK(); \
 	}
@@ -1178,4 +1178,5 @@ EPCGCachingStatus IPCGElement::RetrieveResultsFromCache(IPCGGraphCache* Cache, c
 	}
 }
 
+#undef PCG_ELEMENT_EXECUTION_BREAKPOINT
 #undef LOCTEXT_NAMESPACE
diff --git a/Engine/Plugins/PCG/Source/PCG/Public/PCGContext.h b/Engine/Plugins/PCG/Source/PCG/Public/PCGContext.h
--- a/Engine/Plugins/PCG/Source/PCG/Public/PCGContext.h
+++ b/Engine/Plugins/PCG/Source/PCG/Public/PCGContext.h
@@ -195,6 +195,9 @@ UE_API PRAGMA_ENABLE_DEPRECATION_WARNINGS
 
 	/** True if any issues were logged during last execution. */
 	UE_API bool HasVisualLogs() const;
+
+	/** Whether the graph for the current frame in the execution stack is selected for debugging in the graph editor. */
+	UE_API bool IsExecutingGraphInspected() const;
 #endif // WITH_EDITOR
 
 	/** Gathers references to objects to prevent them from being garbage collected. Extend resources to be collected with AddExtractStructReferencedObjects. */
diff --git a/Engine/Plugins/PCG/Source/PCG/Public/PCGSettings.h b/Engine/Plugins/PCG/Source/PCG/Public/PCGSettings.h
--- a/Engine/Plugins/PCG/Source/PCG/Public/PCGSettings.h
+++ b/Engine/Plugins/PCG/Source/PCG/Public/PCGSettings.h
@@ -165,7 +165,7 @@ public:
 	UPROPERTY(EditAnywhere, BlueprintReadWrite, Category = Debug, meta = (ShowOnlyInnerProperties))
 	FPCGDebugVisualizationSettings DebugSettings;
 
-	/** If a debugger is attached, triggers a breakpoint inside IPCGElement::Execute(). Editor only. Transient. */
+	/** If a debugger like Visual Studio is attached, a break point will be triggered in the execution code to enable debugging. Only applies when inspecting a debug object. Transient and Editor-only. */
 	UPROPERTY(Transient, DuplicateTransient, EditAnywhere, BlueprintReadWrite, Category = Debug, AdvancedDisplay, meta = (DisplayName = "Break In Debugger", EditCondition = bDisplayDebuggingProperties, EditConditionHides, HideEditConditionToggle))
 	bool bBreakDebugger = false;
 
