[PCG] GPU memory fuzzing.

Used this to make a repro to fix a bug in sample content due to bad memory access patterns in a kernel.

diff --git a/Engine/Plugins/PCG/Source/PCG/Private/Compute/DataInterfaces/PCGDataCollectionDataInterface.cpp b/Engine/Plugins/PCG/Source/PCG/Private/Compute/DataInterfaces/PCGDataCollectionDataInterface.cpp
--- a/Engine/Plugins/PCG/Source/PCG/Private/Compute/DataInterfaces/PCGDataCollectionDataInterface.cpp
+++ b/Engine/Plugins/PCG/Source/PCG/Private/Compute/DataInterfaces/PCGDataCollectionDataInterface.cpp
@@ -2,6 +2,7 @@
 
 #include "Compute/DataInterfaces/PCGDataCollectionDataInterface.h"
 
+#include "PCGCommon.h"
 #include "PCGComponent.h"
 #include "PCGContext.h"
 #include "PCGModule.h"
@@ -699,6 +700,20 @@ void FPCGDataCollectionDataProviderProxy::AllocateResources(FRDGBuilder& GraphBu
 		// Fill remaining part of full buffer with 0s.
 		PackedDataCollection.SetNumZeroed(SizeBytes / PackedDataCollection.GetTypeSize());
 	}
+#if !UE_BUILD_SHIPPING
+	else if (PCGSystemSwitches::CVarFuzzGPUMemory.GetValueOnAnyThread())
+	{
+		// Allocate rest of buffer and initialize with random numbers below.
+		const int NumBefore = PackedDataCollection.Num();
+		PackedDataCollection.SetNumUninitialized(SizeBytes / PackedDataCollection.GetTypeSize());
+
+		FRandomStream Rand(GFrameCounter);
+		for (int Index = NumBefore; Index < PackedDataCollection.Num(); ++Index)
+		{
+			PackedDataCollection[Index] = Rand.GetUnsignedInt();
+		}
+	}
+#endif // !UE_BUILD_SHIPPING
 
 	GraphBuilder.QueueBufferUpload(Buffer, PackedDataCollection.GetData(), PackedDataCollection.Num() * PackedDataCollection.GetTypeSize(), ERDGInitialDataFlags::None);
 
diff --git a/Engine/Plugins/PCG/Source/PCG/Private/Compute/PCGDataForGPU.cpp b/Engine/Plugins/PCG/Source/PCG/Private/Compute/PCGDataForGPU.cpp
--- a/Engine/Plugins/PCG/Source/PCG/Private/Compute/PCGDataForGPU.cpp
+++ b/Engine/Plugins/PCG/Source/PCG/Private/Compute/PCGDataForGPU.cpp
@@ -2,6 +2,7 @@
 
 #include "Compute/PCGDataForGPU.h"
 
+#include "PCGCommon.h"
 #include "PCGContext.h"
 #include "PCGData.h"
 #include "PCGEdge.h"
@@ -1087,6 +1088,17 @@ void FPCGDataCollectionDesc::PackDataCollection(const FPCGDataCollection& InData
 	// Uninitialized is fine, all data is initialized explicitly.
 	OutPackedDataCollection.SetNumUninitialized(PackedDataCollectionSizeBytes >> 2);
 
+#if !UE_BUILD_SHIPPING
+	if (PCGSystemSwitches::CVarFuzzGPUMemory.GetValueOnAnyThread())
+	{
+		FRandomStream Rand(GFrameCounter);
+		for (uint32& Value : OutPackedDataCollection)
+		{
+			Value = Rand.GetUnsignedInt();
+		}
+	}
+#endif // !UE_BUILD_SHIPPING
+
 	// Data addresses are written to the header and will be used during packing below.
 	WriteHeader(OutPackedDataCollection);
 
diff --git a/Engine/Plugins/PCG/Source/PCG/Private/PCGCommon.cpp b/Engine/Plugins/PCG/Source/PCG/Private/PCGCommon.cpp
--- a/Engine/Plugins/PCG/Source/PCG/Private/PCGCommon.cpp
+++ b/Engine/Plugins/PCG/Source/PCG/Private/PCGCommon.cpp
@@ -110,12 +110,19 @@ namespace PCGSystemSwitches
 		true,
 		TEXT("Aggressively release transient (esp. GPU) data when possible."));
 
-#if UE_ENABLE_DEBUG_DRAWING 
+#if UE_ENABLE_DEBUG_DRAWING
 	TAutoConsoleVariable<bool> CVarPCGDebugDrawGeneratedCells(
 		TEXT("pcg.GraphExecution.DebugDrawGeneratedCells"),
 		false,
 		TEXT("Draws debug boxes around any cell that is executing, colored by grid level, and in PIE/standalone labels grid size and coords."));
-#endif // UE_ENABLE_DEBUG_DRAWING 
+#endif // UE_ENABLE_DEBUG_DRAWING
+
+#if !UE_BUILD_SHIPPING
+	TAutoConsoleVariable<bool> CVarFuzzGPUMemory(
+		TEXT("pcg.GPU.FuzzMemory"),
+		false,
+		TEXT("Initializes GPU buffers with random numbers which helps to reproduce undefined behaviour from unitialized memory."));
+#endif //!UE_BUILD_SHIPPING
 }
 
 namespace PCGHiGenGrid
diff --git a/Engine/Plugins/PCG/Source/PCG/Public/PCGCommon.h b/Engine/Plugins/PCG/Source/PCG/Public/PCGCommon.h
--- a/Engine/Plugins/PCG/Source/PCG/Public/PCGCommon.h
+++ b/Engine/Plugins/PCG/Source/PCG/Public/PCGCommon.h
@@ -285,6 +285,10 @@ namespace PCGSystemSwitches
 
 	extern PCG_API TAutoConsoleVariable<bool> CVarReleaseTransientResourcesEarly;
 	extern PCG_API TAutoConsoleVariable<bool> CVarPCGDebugDrawGeneratedCells;
+
+#if !UE_BUILD_SHIPPING
+	extern PCG_API TAutoConsoleVariable<bool> CVarFuzzGPUMemory;
+#endif
 }
 
 /** Describes space referential for operations that create data */
