[PCG] Remove legacy element counters that were used for the old version of the GPU Attribute Partition node.

#rnx

diff --git a/Engine/Plugins/PCG/Shaders/Private/PCGDataCollectionDataInterface.ush b/Engine/Plugins/PCG/Shaders/Private/PCGDataCollectionDataInterface.ush
--- a/Engine/Plugins/PCG/Shaders/Private/PCGDataCollectionDataInterface.ush
+++ b/Engine/Plugins/PCG/Shaders/Private/PCGDataCollectionDataInterface.ush
@@ -21,9 +21,6 @@ RWByteAddressBuffer {DataInterfaceName}_DataCollectionBuffer;
 // Multiplier that can be configured on output pins.
 uint {DataInterfaceName}_ElementCountMultiplier;
 
-RWStructuredBuffer<uint> {DataInterfaceName}_ElementCountersBuffer;
-uint {DataInterfaceName}_ElementCountersPresent;
-
 // #################### HEADER READERS ##########################
 
 uint GetNumData_{DataInterfaceName}()
@@ -42,19 +39,6 @@ uint GetNumElements_{DataInterfaceName}(uint InDataIndex)
 	return {DataInterfaceName}_DataCollectionBuffer.Load(ReadAddress);
 }
 
-int AddToElementCounterInternal_{DataInterfaceName}(uint InDataIndex, uint InElementCount)
-{
-	if ({DataInterfaceName}_ElementCountersPresent == 0)
-	{
-		return -1;
-	}
-
-	uint ValueBefore;
-	InterlockedAdd({DataInterfaceName}_ElementCountersBuffer[InDataIndex], InElementCount, ValueBefore);
-
-	return (int)ValueBefore;
-}
-
 bool GetThreadData_{DataInterfaceName}(uint InThreadIndex, out uint OutDataIndex, out uint OutElementIndex)
 {
 	int ElementIndex = (int)InThreadIndex;
diff --git a/Engine/Plugins/PCG/Source/PCG/Internal/Compute/Data/PCGProxyForGPUData.h b/Engine/Plugins/PCG/Source/PCG/Internal/Compute/Data/PCGProxyForGPUData.h
--- a/Engine/Plugins/PCG/Source/PCG/Internal/Compute/Data/PCGProxyForGPUData.h
+++ b/Engine/Plugins/PCG/Source/PCG/Internal/Compute/Data/PCGProxyForGPUData.h
@@ -54,9 +54,6 @@ public:
 	/** Returns a description of this data. Does not trigger a GPU->CPU readback. */
 	UE_API bool GetDescription(FPCGDataDesc& OutDescription) const;
 
-	/** Apply the given per-data element counts to the data description. */
-	UE_API void UpdateElementCountsFromReadback(const TArray<uint32>& InElementCounts);
-
 	UE_API TSharedPtr<const FPCGProxyForGPUDataCollection> GetGPUInfo() const;
 
 protected:
diff --git a/Engine/Plugins/PCG/Source/PCG/Internal/Compute/PCGDataForGPU.h b/Engine/Plugins/PCG/Source/PCG/Internal/Compute/PCGDataForGPU.h
--- a/Engine/Plugins/PCG/Source/PCG/Internal/Compute/PCGDataForGPU.h
+++ b/Engine/Plugins/PCG/Source/PCG/Internal/Compute/PCGDataForGPU.h
@@ -246,8 +246,6 @@ public:
 
 	const FPCGDataCollectionDesc& GetDescription() const { return Description; }
 
-	void UpdateElementCountsFromReadback(const TArray<uint32>& InElementCounts);
-
 	const TArray<FString>& GetStringTable() const { return StringTable; }
 
 protected:
diff --git a/Engine/Plugins/PCG/Source/PCG/Private/Compute/Data/PCGProxyForGPUData.cpp b/Engine/Plugins/PCG/Source/PCG/Private/Compute/Data/PCGProxyForGPUData.cpp
--- a/Engine/Plugins/PCG/Source/PCG/Private/Compute/Data/PCGProxyForGPUData.cpp
+++ b/Engine/Plugins/PCG/Source/PCG/Private/Compute/Data/PCGProxyForGPUData.cpp
@@ -63,14 +63,6 @@ bool UPCGProxyForGPUData::GetDescription(FPCGDataDesc& OutDescription) const
 	return false;
 }
 
-void UPCGProxyForGPUData::UpdateElementCountsFromReadback(const TArray<uint32>& InElementCounts)
-{
-	if (TSharedPtr<FPCGProxyForGPUDataCollection> GPUInfo = GetGPUInfoMutable())
-	{
-		GPUInfo->UpdateElementCountsFromReadback(InElementCounts);
-	}
-}
-
 UPCGProxyForGPUData::FReadbackResult UPCGProxyForGPUData::GetCPUData(FPCGContext* InContext) const
 {
 	TSharedPtr<FPCGProxyForGPUDataCollection> DataOnGPU = GetGPUInfoMutable();
diff --git a/Engine/Plugins/PCG/Source/PCG/Private/Compute/DataInterfaces/PCGDataCollectionDataInterface.cpp b/Engine/Plugins/PCG/Source/PCG/Private/Compute/DataInterfaces/PCGDataCollectionDataInterface.cpp
--- a/Engine/Plugins/PCG/Source/PCG/Private/Compute/DataInterfaces/PCGDataCollectionDataInterface.cpp
+++ b/Engine/Plugins/PCG/Source/PCG/Private/Compute/DataInterfaces/PCGDataCollectionDataInterface.cpp
@@ -54,12 +54,6 @@ void UPCGDataCollectionDataInterface::GetSupportedInputs(TArray<FShaderFunctionD
 			.AddReturnType(EShaderFundamentalType::Uint)
 			.AddParam(EShaderFundamentalType::Uint); // InDataIndex
 
-		OutFunctions.AddDefaulted_GetRef()
-			.SetName(TEXT("AddToElementCounterInternal"))
-			.AddReturnType(EShaderFundamentalType::Int) // Value before increment. Returns -1 if counters disabled.
-			.AddParam(EShaderFundamentalType::Uint) // InDataIndex
-			.AddParam(EShaderFundamentalType::Uint); // InElementCount
-
 		OutFunctions.AddDefaulted_GetRef()
 			.SetName(TEXT("GetThreadData"))
 			.AddParam(EShaderFundamentalType::Uint) // InThreadIndex
@@ -415,8 +409,6 @@ void UPCGDataCollectionDataInterface::GetSupportedInputs(TArray<FShaderFunctionD
 BEGIN_SHADER_PARAMETER_STRUCT(FPCGDataCollectionDataInterfaceParameters, )
 	SHADER_PARAMETER_RDG_BUFFER_UAV(RWByteAddressBuffer, DataCollectionBuffer)
 	SHADER_PARAMETER(uint32, ElementCountMultiplier)
-	SHADER_PARAMETER_RDG_BUFFER_UAV(RWStructuredBuffer<uint32>, ElementCountersBuffer)
-	SHADER_PARAMETER(uint32, ElementCountersPresent)
 END_SHADER_PARAMETER_STRUCT()
 
 void UPCGDataCollectionDataInterface::GetShaderParameters(TCHAR const* UID, FShaderParametersMetadataBuilder& InOutBuilder, FShaderParametersMetadataAllocations& InOutAllocations) const
@@ -491,8 +483,6 @@ void UPCGDataCollectionDataProvider::Initialize(const UComputeDataInterface* InD
 	bRequiresZeroInitialization = DataInterface->bRequiresZeroInitialization;
 
 	ElementCountMultiplier = DataInterface->ElementCountMultiplier;
-
-	bExposeElementCounters = DataInterface->bExposeElementCounters;
 }
 
 bool UPCGDataCollectionDataProvider::PrepareForExecute_GameThread(UPCGDataBinding* InBinding)
@@ -519,66 +509,14 @@ FComputeDataProviderRenderProxy* UPCGDataCollectionDataProvider::GetRenderProxy(
 {
 	TRACE_CPUPROFILER_EVENT_SCOPE(UPCGDataCollectionDataProvider::GetRenderProxy);
 
-	FPCGDataCollectionDataProviderProxy* Proxy = new FPCGDataCollectionDataProviderProxy(this, PinDesc, ExportMode, bRequiresZeroInitialization, ElementCountMultiplier, bExposeElementCounters, OutputPinLabel, OutputPinLabelAlias);
-
-	if (bExposeElementCounters)
-	{
-		const uint64 GenerationCount = GenerationCounter;
-
-		Proxy->AsyncReadbackCallback_RenderThread = [ThisWeak=TWeakObjectPtr<UPCGDataCollectionDataProvider>(this), /*InitStack=InitStack, ProviderName,*/ GenerationCount](const void* InData, int InNumBytes)
-		{
-			TRACE_CPUPROFILER_EVENT_SCOPE(UPCGDataCollectionDataProvider::ProcessReadbackData_RenderThread);
-			TStrongObjectPtr<UPCGDataCollectionDataProvider> This = ThisWeak.Pin();
-			if (!This)
-			{
-				return;
-			}
-
-			if (This->GenerationCounter != GenerationCount)
-			{
-				return;
-			}
-
-			// We should never find ourselves stomping existing data.
-			check(This->ReadbackCounters.IsEmpty());
-
-			if (ensure(InData && InNumBytes > 0))
-			{
-				// Copy the data to temp storage for the game thread to pick up.
-				This->ReadbackCounters.SetNumUninitialized(InNumBytes / This->ReadbackCounters.GetTypeSize());
-				FMemory::Memcpy(This->ReadbackCounters.GetData(), InData, InNumBytes);
-			}
-
-			This->bReadbackComplete = true;
-
-			ExecuteOnGameThread(UE_SOURCE_LOCATION, [ThisWeak=ThisWeak, OutputPinLabelAlias=This->OutputPinLabelAlias, GenerationCount]()
-			{
-				TStrongObjectPtr<UPCGDataCollectionDataProvider> This = ThisWeak.Pin();
-				if (This && This->GenerationCounter == GenerationCount)
-				{
-					if (This->Binding.IsValid())
-					{
-						This->Binding->ReceiveElementCountsFromGPU_GameThread(OutputPinLabelAlias, This->ReadbackCounters);
-					}
-
-					This->OnCountersReadBack.Broadcast();
-				}
-			});
-		};
-	}
-
-	return Proxy;
+	return new FPCGDataCollectionDataProviderProxy(this, PinDesc, ExportMode, bRequiresZeroInitialization, ElementCountMultiplier, OutputPinLabel, OutputPinLabelAlias);
 }
 
 void UPCGDataCollectionDataProvider::Reset()
 {
 	Super::Reset();
 
-	OnCountersReadBack = {};
 	ElementCountMultiplier = 0;
-	bExposeElementCounters = false;
-	bReadbackComplete = false;
-	ReadbackCounters.Empty();
 	bRequiresZeroInitialization = false;
 }
 
@@ -588,13 +526,11 @@ FPCGDataCollectionDataProviderProxy::FPCGDataCollectionDataProviderProxy(
 	EPCGExportMode InExportMode,
 	bool bInZeroInitialize,
 	uint32 InElementCountMultiplier,
-	bool bInExposeElementCounters,
 	FName InOutputPinLabel,
 	FName InOutputPinLabelAlias)
 	: ExportMode(InExportMode)
 	, bZeroInitialize(bInZeroInitialize)
 	, ElementCountMultiplier(InElementCountMultiplier)
-	, bExposeElementCounters(bInExposeElementCounters)
 	, DataProviderWeakPtr(InDataProvider)
 	, OutputPinLabel(InOutputPinLabel)
 	, OutputPinLabelAlias(InOutputPinLabelAlias)
@@ -631,7 +567,6 @@ bool FPCGDataCollectionDataProviderProxy::IsValid(FValidationData const& InValid
 void FPCGDataCollectionDataProviderProxy::GatherDispatchData(FDispatchData const& InDispatchData)
 {
 	check(BufferUAV);
-	check(CountersUAV);
 
 	const TStridedView<FParameters> ParameterArray = MakeStridedParameterView<FParameters>(InDispatchData);
 	for (int32 InvocationIndex = 0; InvocationIndex < ParameterArray.Num(); ++InvocationIndex)
@@ -639,8 +574,6 @@ void FPCGDataCollectionDataProviderProxy::GatherDispatchData(FDispatchData const
 		FParameters& Parameters = ParameterArray[InvocationIndex];
 		Parameters.DataCollectionBuffer = BufferUAV;
 		Parameters.ElementCountMultiplier = ElementCountMultiplier;
-		Parameters.ElementCountersBuffer = CountersUAV;
-		Parameters.ElementCountersPresent = bExposeElementCounters ? 1 : 0;
 	}
 }
 
@@ -649,30 +582,6 @@ void FPCGDataCollectionDataProviderProxy::AllocateResources(FRDGBuilder& GraphBu
 	TRACE_CPUPROFILER_EVENT_SCOPE(FPCGDataCollectionDataProviderProxy::AllocateResources);
 	check(SizeBytes > 0);
 
-	{
-		// Ensure we always allocate at least 1 element so we have something to bind.
-		const uint32 NumCounters = bExposeElementCounters ? FMath::Max(PinDesc.DataDescs.Num(), 1) : 1;
-
-		FRDGBufferDesc CountersDesc = FRDGBufferDesc::CreateStructuredDesc(sizeof(uint32), NumCounters);
-		CountersDesc.Usage |= BUF_SourceCopy;
-
-		Counters = GraphBuilder.CreateBuffer(CountersDesc, TEXT("PCGDataCollection_Counters"));
-		CountersUAV = GraphBuilder.CreateUAV(Counters);
-
-		// A single counter, or a small number of counters, will be a very common case.
-		if (NumCounters <= 4)
-		{
-			GraphBuilder.QueueBufferUpload(Counters, &FUintVector4::ZeroValue, NumCounters * sizeof(uint32));
-		}
-		else
-		{
-			TArray<uint32> Zeroes;
-			Zeroes.SetNumZeroed(NumCounters);
-
-			GraphBuilder.QueueBufferUpload(Counters, Zeroes.GetData(), Zeroes.Num() * Zeroes.GetTypeSize());
-		}
-	}
-
 	FRDGBufferDesc Desc = FRDGBufferDesc::CreateByteAddressDesc(SizeBytes);
 	if (ExportMode != EPCGExportMode::NoExport)
 	{
@@ -760,17 +669,4 @@ void FPCGDataCollectionDataProviderProxy::AllocateResources(FRDGBuilder& GraphBu
 	}
 }
 
-void FPCGDataCollectionDataProviderProxy::GetReadbackData(TArray<FReadbackData>& OutReadbackData) const
-{
-	if (bExposeElementCounters)
-	{
-		FReadbackData Data;
-		Data.Buffer = Counters;
-		Data.NumBytes = PinDesc.DataDescs.Num() * sizeof(uint32);
-		Data.ReadbackCallback_RenderThread = &AsyncReadbackCallback_RenderThread;
-
-		OutReadbackData.Add(MoveTemp(Data));
-	}
-}
-
 #undef LOCTEXT_NAMESPACE
diff --git a/Engine/Plugins/PCG/Source/PCG/Private/Compute/DataInterfaces/PCGDataCollectionDataInterface.h b/Engine/Plugins/PCG/Source/PCG/Private/Compute/DataInterfaces/PCGDataCollectionDataInterface.h
--- a/Engine/Plugins/PCG/Source/PCG/Private/Compute/DataInterfaces/PCGDataCollectionDataInterface.h
+++ b/Engine/Plugins/PCG/Source/PCG/Private/Compute/DataInterfaces/PCGDataCollectionDataInterface.h
@@ -34,8 +34,6 @@ public:
 	void GetShaderHash(FString& InOutKey) const override;
 	void GetHLSL(FString& OutHLSL, FString const& InDataInterfaceName) const override;
 	void GetDefines(FComputeKernelDefinitionSet& OutDefinitionSet) const override;
-	// This data interface only uses the CF readback functionality for the element counters.
-	bool GetRequiresReadback() const override { return bExposeElementCounters; }
 	UComputeDataProvider* CreateDataProvider() const override;
 	//~ End UComputeDataInterface Interface
 
@@ -43,8 +41,6 @@ public:
 
 	void SetElementCountMultiplier(uint32 InElementCountMultiplier) { ElementCountMultiplier = InElementCountMultiplier; }
 
-	void SetExposeElementCounters(bool bInExposeElementCounters) { bExposeElementCounters = bInExposeElementCounters; }
-
 public:
 	/** Whether to perform full 0-initialization of the buffer. */
 	UPROPERTY()
@@ -53,9 +49,6 @@ public:
 	UPROPERTY()
 	uint32 ElementCountMultiplier = 0;
 
-	UPROPERTY()
-	bool bExposeElementCounters = false;
-
 protected:
 	static TCHAR const* TemplateFilePath;
 };
@@ -77,25 +70,12 @@ public:
 	virtual void Reset() override;
 	//~ End UComputeDataProvider Interface
 
-	DECLARE_EVENT(UPCGDataCollectionDataProvider, FOnCountersReadBack);
-	FOnCountersReadBack& OnCountersReadBack_GameThread() { return OnCountersReadBack; }
-
 public:
 	UPROPERTY()
 	uint32 ElementCountMultiplier = 0;
 
 	UPROPERTY()
 	bool bRequiresZeroInitialization = false;
-
-	UPROPERTY()
-	bool bExposeElementCounters = false;
-
-private:
-	FOnCountersReadBack OnCountersReadBack;
-
-	std::atomic<bool> bReadbackComplete = false;
-
-	TArray<uint32> ReadbackCounters;
 };
 
 class FPCGDataCollectionDataProviderProxy : public FComputeDataProviderRenderProxy
@@ -107,7 +87,6 @@ public:
 		EPCGExportMode InExportMode,
 		bool bInZeroInitialize,
 		uint32 InElementCountMultiplier,
-		bool bInExposeElementCounters,
 		FName InOutputPinLabel,
 		FName InOutputPinLabelAlias);
 
@@ -115,12 +94,9 @@ public:
 	bool IsValid(FValidationData const& InValidationData) const override;
 	void GatherDispatchData(FDispatchData const& InDispatchData) override;
 	void AllocateResources(FRDGBuilder& GraphBuilder, FAllocationData const& InAllocationData) override;
-	void GetReadbackData(TArray<FReadbackData>& OutReadbackData) const override;
 	//~ End FComputeDataProviderRenderProxy Interface
 
 public:
-	FReadbackCallback AsyncReadbackCallback_RenderThread;
-
 	/** Generation count of the data provider when the proxy was created. */
 	uint64 OriginatingGenerationCount = 0;
 
@@ -133,8 +109,6 @@ protected:
 
 	uint32 ElementCountMultiplier = 0;
 
-	bool bExposeElementCounters = false;
-
 	using FParameters = FPCGDataCollectionDataInterfaceParameters;
 
 	FPCGDataCollectionDesc PinDesc;
@@ -142,9 +116,6 @@ protected:
 	FRDGBufferRef Buffer = nullptr;
 	FRDGBufferUAVRef BufferUAV = nullptr;
 
-	FRDGBufferRef Counters = nullptr;
-	FRDGBufferUAVRef CountersUAV = nullptr;
-
 	// Weak pointer useful for passing back buffer handles. Do not access this directly from the render thread.
 	TWeakObjectPtr<UPCGDataCollectionDataProvider> DataProviderWeakPtr;
 
diff --git a/Engine/Plugins/PCG/Source/PCG/Private/Compute/DataInterfaces/PCGDataCollectionUploadDataInterface.cpp b/Engine/Plugins/PCG/Source/PCG/Private/Compute/DataInterfaces/PCGDataCollectionUploadDataInterface.cpp
--- a/Engine/Plugins/PCG/Source/PCG/Private/Compute/DataInterfaces/PCGDataCollectionUploadDataInterface.cpp
+++ b/Engine/Plugins/PCG/Source/PCG/Private/Compute/DataInterfaces/PCGDataCollectionUploadDataInterface.cpp
@@ -92,8 +92,7 @@ FComputeDataProviderRenderProxy* UPCGDataProviderDataCollectionUpload::GetRender
 		OutputPinLabelAlias,
 		ExportMode,
 		bRequiresZeroInitialization,
-		ElementCountMultiplier,
-		bExposeElementCounters);
+		ElementCountMultiplier);
 }
 
 void UPCGDataProviderDataCollectionUpload::Reset()
@@ -112,9 +111,8 @@ FPCGDataProviderDataCollectionUploadProxy::FPCGDataProviderDataCollectionUploadP
 	FName InOutputPinLabelAlias,
 	EPCGExportMode InExportMode,
 	bool bInZeroInitialize,
-	uint32 InElementCountMultiplier,
-	bool bInExposeElementCounters)
-	: FPCGDataCollectionDataProviderProxy(InDataProvider, InPinDesc, InExportMode, bInZeroInitialize, InElementCountMultiplier, bInExposeElementCounters, InOutputPinLabel, InOutputPinLabelAlias)
+	uint32 InElementCountMultiplier)
+	: FPCGDataCollectionDataProviderProxy(InDataProvider, InPinDesc, InExportMode, bInZeroInitialize, InElementCountMultiplier, InOutputPinLabel, InOutputPinLabelAlias)
 	, DataAdaptor(InDataAdaptor)
 {
 }
@@ -134,32 +132,6 @@ void FPCGDataProviderDataCollectionUploadProxy::AllocateResources(FRDGBuilder& G
 {
 	TRACE_CPUPROFILER_EVENT_SCOPE(FPCGDataProviderDataCollectionUploadProxy::AllocateResources);
 
-	{
-		Buffer = DataAdaptor->GetBuffer_RenderThread(GraphBuilder, ExportMode);
-		BufferUAV = GraphBuilder.CreateUAV(Buffer);
-	}
-
-	{
-		// Ensure we always allocate at least 1 element so we have something to bind.
-		const uint32 NumCounters = bExposeElementCounters ? FMath::Max(PinDesc.DataDescs.Num(), 1) : 1;
-
-		FRDGBufferDesc CountersDesc = FRDGBufferDesc::CreateStructuredDesc(sizeof(uint32), NumCounters);
-		CountersDesc.Usage |= BUF_SourceCopy;
-
-		Counters = GraphBuilder.CreateBuffer(CountersDesc, TEXT("PCGDataCollectionUpload_Counters"));
-		CountersUAV = GraphBuilder.CreateUAV(Counters);
-
-		// A single counter, or a small number of counters, will be a very common case.
-		if (NumCounters <= 4)
-		{
-			GraphBuilder.QueueBufferUpload(Counters, &FUintVector4::ZeroValue, NumCounters * sizeof(uint32));
-		}
-		else
-		{
-			TArray<uint32> Zeroes;
-			Zeroes.SetNumZeroed(NumCounters);
-
-			GraphBuilder.QueueBufferUpload(Counters, Zeroes.GetData(), Zeroes.Num() * Zeroes.GetTypeSize());
-		}
-	}
+	Buffer = DataAdaptor->GetBuffer_RenderThread(GraphBuilder, ExportMode);
+	BufferUAV = GraphBuilder.CreateUAV(Buffer);
 }
diff --git a/Engine/Plugins/PCG/Source/PCG/Private/Compute/DataInterfaces/PCGDataCollectionUploadDataInterface.h b/Engine/Plugins/PCG/Source/PCG/Private/Compute/DataInterfaces/PCGDataCollectionUploadDataInterface.h
--- a/Engine/Plugins/PCG/Source/PCG/Private/Compute/DataInterfaces/PCGDataCollectionUploadDataInterface.h
+++ b/Engine/Plugins/PCG/Source/PCG/Private/Compute/DataInterfaces/PCGDataCollectionUploadDataInterface.h
@@ -57,8 +57,7 @@ public:
 		FName InOutputPinLabelAlias,
 		EPCGExportMode InExportMode,
 		bool bInZeroInitialize,
-		uint32 InElementCountMultiplier,
-		bool bInExposeElementCounters);
+		uint32 InElementCountMultiplier);
 
 	//~ Begin FComputeDataProviderRenderProxy Interface
 	bool IsValid(FValidationData const& InValidationData) const override;
diff --git a/Engine/Plugins/PCG/Source/PCG/Private/Compute/Elements/PCGComputeGraphElement.cpp b/Engine/Plugins/PCG/Source/PCG/Private/Compute/Elements/PCGComputeGraphElement.cpp
--- a/Engine/Plugins/PCG/Source/PCG/Private/Compute/Elements/PCGComputeGraphElement.cpp
+++ b/Engine/Plugins/PCG/Source/PCG/Private/Compute/Elements/PCGComputeGraphElement.cpp
@@ -87,7 +87,7 @@ FPCGComputeGraphContext::~FPCGComputeGraphContext()
 
 bool FPCGComputeGraphContext::HasPendingAsyncOperations() const
 {
-	return !ProvidersWithBufferExports.IsEmpty() || !ProvidersWithReadbacks.IsEmpty();
+	return !ProvidersWithBufferExports.IsEmpty();
 }
 
 void FPCGComputeGraphContext::AddExtraStructReferencedObjects(FReferenceCollector& Collector)
@@ -290,36 +290,6 @@ bool FPCGComputeGraphElement::ExecuteInternal(FPCGContext* InContext) const
 					}
 				});
 			}
-
-			UPCGDataCollectionDataProvider* DataCollectionDataProvider = Cast<UPCGDataCollectionDataProvider>(ComputeDataProvider);
-
-			if (DataCollectionDataProvider && DataCollectionDataProvider->bExposeElementCounters)
-			{
-				Context->ProvidersWithReadbacks.Add(DataCollectionDataProvider);
-
-				TWeakObjectPtr<UPCGDataCollectionDataProvider> DataProviderWeak = DataCollectionDataProvider;
-
-				DataCollectionDataProvider->OnCountersReadBack_GameThread().AddLambda([ContextHandle = Context->GetOrCreateHandle(), DataProviderWeak]()
-				{
-					check(IsInGameThread());
-
-					TSharedPtr<FPCGContextHandle> SharedHandle = ContextHandle.Pin();
-					FPCGComputeGraphContext* ContextPtr = static_cast<FPCGComputeGraphContext*>(SharedHandle ? SharedHandle->GetContext() : nullptr);
-					UPCGDataCollectionDataProvider* DataProvider = DataProviderWeak.Get();
-					if (!DataProvider || !ContextPtr)
-					{
-						return;
-					}
-
-					ensure(ContextPtr->ProvidersWithReadbacks.Contains(DataProvider));
-					ContextPtr->ProvidersWithReadbacks.Remove(DataProvider);
-
-					if (!ContextPtr->HasPendingAsyncOperations())
-					{
-						ContextPtr->bIsPaused = false;
-					}
-				});
-			}
 		}
 
 		InitializeBindingTask.Wait();
@@ -562,9 +532,6 @@ bool FPCGComputeGraphElement::ExecuteInternal(FPCGContext* InContext) const
 
 		Context->bExecutionSuccess = true;
 
-		// Apply any read back element counts to the final data description.
-		Context->DataBinding->ApplyElementCounts_GameThread();
-
 		Context->DataProvidersPendingPostExecute.Reserve(Context->ComputeGraphInstance->GetNumDataProviders());
 		for (UComputeDataProvider* Provider : Context->ComputeGraphInstance->GetDataProviders())
 		{
@@ -694,7 +661,6 @@ void FPCGComputeGraphElement::ResetAsyncOperations(FPCGContext* InContext) const
 	{
 		FPCGComputeGraphContext* Context = static_cast<FPCGComputeGraphContext*>(InContext);
 		Context->ProvidersWithBufferExports.Reset();
-		Context->ProvidersWithReadbacks.Reset();
 	}
 }
 
diff --git a/Engine/Plugins/PCG/Source/PCG/Private/Compute/Elements/PCGComputeGraphElement.h b/Engine/Plugins/PCG/Source/PCG/Private/Compute/Elements/PCGComputeGraphElement.h
--- a/Engine/Plugins/PCG/Source/PCG/Private/Compute/Elements/PCGComputeGraphElement.h
+++ b/Engine/Plugins/PCG/Source/PCG/Private/Compute/Elements/PCGComputeGraphElement.h
@@ -41,10 +41,6 @@ public:
 
 	bool HasPendingAsyncOperations() const;
 
-#if WITH_EDITOR
-	void DoDebugBreak() const;
-#endif
-
 protected:
 	virtual void AddExtraStructReferencedObjects(FReferenceCollector& Collector) override;
 
@@ -82,8 +78,6 @@ public:
 	*/
 	TSet<TObjectPtr<UComputeDataProvider>> ProvidersWithBufferExports;
 
-	TSet<TObjectPtr<UComputeDataProvider>> ProvidersWithReadbacks;
-
 	EPCGComputeGraphExecutionPhase ExecutionSubPhase = EPCGComputeGraphExecutionPhase::GetComputeGraph;
 
 	std::atomic<bool> bGraphValid = false;
diff --git a/Engine/Plugins/PCG/Source/PCG/Private/Compute/PCGComputeKernel.h b/Engine/Plugins/PCG/Source/PCG/Private/Compute/PCGComputeKernel.h
--- a/Engine/Plugins/PCG/Source/PCG/Private/Compute/PCGComputeKernel.h
+++ b/Engine/Plugins/PCG/Source/PCG/Private/Compute/PCGComputeKernel.h
@@ -122,9 +122,6 @@ public:
 	/** Initialize all buffer data to 0. This is not free so should only be used when kernel execution requires 0-initialized data, such as for counters. */
 	virtual bool DoesOutputPinRequireZeroInitialization(FName InOutputPinLabel) const { return false; }
 
-	/** Whether the output pin should have a buffer of counters attached, useful for compaction and other counting requirements. */
-	virtual bool DoesOutputPinRequireElementCounters(FName InOutputPinLabel) const { return false; }
-
 	/** Add any strings emitted by this node that are known statically at compile time. */
 	virtual void AddStaticCreatedStrings(TArray<FString>& InOutStringTable) const {}
 
@@ -149,6 +146,7 @@ public:
 	virtual bool GetLogDataDescriptions() const { return bLogDataDescriptions; }
 
 	static void AddReferencedObjects(UObject* InThis, FReferenceCollector& Collector);
+
 protected:
 #if WITH_EDITOR
 	/** Implement to do any additional initialization specific to your kernel. */
diff --git a/Engine/Plugins/PCG/Source/PCG/Private/Compute/PCGDataBinding.cpp b/Engine/Plugins/PCG/Source/PCG/Private/Compute/PCGDataBinding.cpp
--- a/Engine/Plugins/PCG/Source/PCG/Private/Compute/PCGDataBinding.cpp
+++ b/Engine/Plugins/PCG/Source/PCG/Private/Compute/PCGDataBinding.cpp
@@ -211,43 +211,6 @@ void UPCGDataBinding::ReceiveDataFromGPU_GameThread(
 	}
 }
 
-void UPCGDataBinding::ReceiveElementCountsFromGPU_GameThread(FName InPinLabel, const TArray<uint32>& InCounters)
-{
-	PinToElementCounts.FindOrAdd(InPinLabel) = InCounters;
-}
-
-void UPCGDataBinding::ApplyElementCounts_GameThread()
-{
-	for (FPCGTaggedData& Output : OutputDataCollection.TaggedData)
-	{
-		if (const TArray<uint32>* FoundElementCounts = PinToElementCounts.Find(Output.Pin))
-		{
-			if (UPCGProxyForGPUData* Data = const_cast<UPCGProxyForGPUData*>(Cast<UPCGProxyForGPUData>(Output.Data)))
-			{
-				Data->UpdateElementCountsFromReadback(*FoundElementCounts);
-			}
-		}
-	}
-
-	auto UpdateCountsOnDebugData = [This=this](TArray<FPCGDataToDebug>& DebugData)
-	{
-		for (FPCGDataToDebug& Data : DebugData)
-		{
-			if (UPCGProxyForGPUData* Proxy = Cast<UPCGProxyForGPUData>(Data.Data))
-			{
-				const TArray<uint32>* FoundElementCounts = This->PinToElementCounts.Find(Data.PinLabelAlias);
-				if (FoundElementCounts)
-				{
-					Proxy->UpdateElementCountsFromReadback(*FoundElementCounts);
-				}
-			}
-		}
-	};
-
-	UpdateCountsOnDebugData(DataToInspect);
-	UpdateCountsOnDebugData(DataToDebug);
-}
-
 IPCGGraphExecutionSource* UPCGDataBinding::GetExecutionSource() const
 {
 	if (TSharedPtr<FPCGContextHandle> SharedHandle = ContextHandle.Pin())
diff --git a/Engine/Plugins/PCG/Source/PCG/Private/Compute/PCGDataBinding.h b/Engine/Plugins/PCG/Source/PCG/Private/Compute/PCGDataBinding.h
--- a/Engine/Plugins/PCG/Source/PCG/Private/Compute/PCGDataBinding.h
+++ b/Engine/Plugins/PCG/Source/PCG/Private/Compute/PCGDataBinding.h
@@ -113,10 +113,6 @@ public:
 	void ReceiveDataFromGPU_GameThread(UPCGData* InData, const UPCGSettings* InProducerSettings, EPCGExportMode InExportMode, FName InPinLabel, FName InPinLabelAlias);
 	void ReceiveDataFromGPU_GameThread(UPCGData* InData, const UPCGSettings* InProducerSettings, EPCGExportMode InExportMode, FName InPinLabel, FName InPinLabelAlias, const TSet<FString>& AdditionalTags);
 
-	void ReceiveElementCountsFromGPU_GameThread(FName InPinLabel, const TArray<uint32>& InCounters);
-
-	void ApplyElementCounts_GameThread();
-
 	using FSourceBufferAndAttributeIndex = TPair<TSharedPtr<const FPCGProxyForGPUDataCollection>, /*Source Metadata Attribute Index*/int32>;
 
 	const TMap<FSourceBufferAndAttributeIndex, /*Attribute Index*/int32>& GetAttributeIdRemap() const { return AttributeIdRemap; }
@@ -170,9 +166,6 @@ public:
 	TWeakPtr<FPCGContextHandle> ContextHandle;
 
 protected:
-	/** Final element counts read back from GPU. */
-	TMap<FName, TArray<uint32>> PinToElementCounts;
-
 	/** All attributes present in graph at execution time. Coherent across all graph branches. Seeded using statically known created attributes at compile time, then
 	 * augmented with incoming attributes from input data collection at runtime. */
 	UPROPERTY()
diff --git a/Engine/Plugins/PCG/Source/PCG/Private/Compute/PCGDataForGPU.cpp b/Engine/Plugins/PCG/Source/PCG/Private/Compute/PCGDataForGPU.cpp
--- a/Engine/Plugins/PCG/Source/PCG/Private/Compute/PCGDataForGPU.cpp
+++ b/Engine/Plugins/PCG/Source/PCG/Private/Compute/PCGDataForGPU.cpp
@@ -1895,19 +1895,4 @@ bool FPCGProxyForGPUDataCollection::GetCPUData(FPCGContext* InContext, int32 InD
 	return true;
 }
 
-void FPCGProxyForGPUDataCollection::UpdateElementCountsFromReadback(const TArray<uint32>& InElementCounts)
-{
-	uint32 DataCounter = 0;
-
-	if (!ensure(Description.DataDescs.Num() == InElementCounts.Num()))
-	{
-		return;
-	}
-
-	for (FPCGDataDesc& DataDesc : Description.DataDescs)
-	{
-		DataDesc.ElementCount = InElementCounts[DataCounter++];
-	}
-}
-
 #undef LOCTEXT_NAMESPACE
diff --git a/Engine/Plugins/PCG/Source/PCG/Private/Graph/PCGGraphCompilerGPU.cpp b/Engine/Plugins/PCG/Source/PCG/Private/Graph/PCGGraphCompilerGPU.cpp
--- a/Engine/Plugins/PCG/Source/PCG/Private/Graph/PCGGraphCompilerGPU.cpp
+++ b/Engine/Plugins/PCG/Source/PCG/Private/Graph/PCGGraphCompilerGPU.cpp
@@ -927,7 +927,6 @@ void FPCGGraphCompilerGPU::CreateDataInterfaces(
 			if (ProducerPin)
 			{
 				DataInterfacePCGData->SetElementCountMultiplier(ProducerKernel ? ProducerKernel->GetElementCountMultiplier(ProducerPin->Properties.Label) : 1);
-				DataInterfacePCGData->SetExposeElementCounters(ProducerKernel && ProducerKernel->DoesOutputPinRequireElementCounters(ProducerPin->Properties.Label));
 				DataInterfacePCGData->SetRequiresZeroInitialization(ProducerKernel && ProducerKernel->DoesOutputPinRequireZeroInitialization(ProducerPin->Properties.Label));
 			}
 
