[PCG] GPU - fix GPU data not being passed across Change Grid Size nodes.

[RN] Fixed performance bug where GPU-resident data is unnecessarily read back when passing through Change Grid Size nodes.

diff --git a/Engine/Plugins/PCG/Source/PCG/Private/PCGElement.cpp b/Engine/Plugins/PCG/Source/PCG/Private/PCGElement.cpp
--- a/Engine/Plugins/PCG/Source/PCG/Private/PCGElement.cpp
+++ b/Engine/Plugins/PCG/Source/PCG/Private/PCGElement.cpp
@@ -10,6 +10,7 @@
 #include "Data/PCGBasePointData.h"
 #include "Data/PCGPointArrayData.h"
 #include "Data/PCGPointData.h"
+#include "Elements/PCGHiGenGridSize.h"
 #include "Graph/PCGGraphCache.h"
 #include "Helpers/PCGActorHelpers.h"
 #include "Helpers/PCGHelpers.h"
@@ -397,7 +398,8 @@ bool IPCGElement::ConvertInputsIfNeeded(FPCGContext* Context) const
 		return true;
 	}
 
-	if (!SupportsGPUResidentData(Context))
+	// todo_pcg: Force-allow GPU resident data to tip-toe around ABI for hotfix. Add implementation of SupportsGPUResidentData to the hi-gen grid size node in 5.7.
+	if (!SupportsGPUResidentData(Context) && Context->GetInputSettings<UPCGHiGenGridSizeSettings>() == nullptr)
 	{
 		bool bHasPendingReadbacks = false;
 
