[PCG] Fix for passing GPU data through HiGen Grid Size nodes - implement the final version.

#rnx

diff --git a/Engine/Plugins/PCG/Source/PCG/Private/Elements/PCGHiGenGridSize.cpp b/Engine/Plugins/PCG/Source/PCG/Private/Elements/PCGHiGenGridSize.cpp
--- a/Engine/Plugins/PCG/Source/PCG/Private/Elements/PCGHiGenGridSize.cpp
+++ b/Engine/Plugins/PCG/Source/PCG/Private/Elements/PCGHiGenGridSize.cpp
@@ -2,6 +2,7 @@
 
 #include "Elements/PCGHiGenGridSize.h"
 
+#include "PCGCommon.h"
 #include "PCGComponent.h"
 #include "PCGContext.h"
 #include "PCGGraph.h"
@@ -259,4 +260,9 @@ void FPCGHiGenGridSizeElement::GetDependenciesCrc(const FPCGGetDependenciesCrcPa
 	OutCrc = Crc;
 }
 
+bool FPCGHiGenGridSizeElement::SupportsGPUResidentData(FPCGContext* InContext) const
+{
+	return PCGSystemSwitches::CVarPassGPUDataThroughGridLinks.GetValueOnAnyThread();
+}
+
 #undef LOCTEXT_NAMESPACE
diff --git a/Engine/Plugins/PCG/Source/PCG/Private/Graph/PCGGraphExecutor.cpp b/Engine/Plugins/PCG/Source/PCG/Private/Graph/PCGGraphExecutor.cpp
--- a/Engine/Plugins/PCG/Source/PCG/Private/Graph/PCGGraphExecutor.cpp
+++ b/Engine/Plugins/PCG/Source/PCG/Private/Graph/PCGGraphExecutor.cpp
@@ -91,11 +91,6 @@ namespace PCGGraphExecutor
 		true,
 		TEXT("Controls whether tasks are culled at execution time, for example in response to an deactivated dynamic branch pin"));
 
-	TAutoConsoleVariable<bool> CVarPassGPUDataThroughGridLinks(
-		TEXT("pcg.Graph.GPU.PassGPUDataThroughGridLinks"),
-		true,
-		TEXT("Whether proxies for GPU data are cached in per pin output data and passed through grid links. If false data is read back to CPU."));
-
 #if WITH_EDITOR
 	TAutoConsoleVariable<int> CVarCacheEditorCleanupBehavior(
 		TEXT("pcg.Cache.EditorCleanupLevel"),
@@ -3074,7 +3069,7 @@ namespace PCGGraphExecutor
 
 	bool FPCGGridLinkageElement::SupportsGPUResidentData(FPCGContext* InContext) const
 	{
-		return CVarPassGPUDataThroughGridLinks.GetValueOnAnyThread();
+		return PCGSystemSwitches::CVarPassGPUDataThroughGridLinks.GetValueOnAnyThread();
 	}
 
 	bool ExecuteGridLinkage(
diff --git a/Engine/Plugins/PCG/Source/PCG/Private/PCGCommon.cpp b/Engine/Plugins/PCG/Source/PCG/Private/PCGCommon.cpp
--- a/Engine/Plugins/PCG/Source/PCG/Private/PCGCommon.cpp
+++ b/Engine/Plugins/PCG/Source/PCG/Private/PCGCommon.cpp
@@ -123,6 +123,11 @@ namespace PCGSystemSwitches
 		false,
 		TEXT("Initializes GPU buffers with random numbers which helps to reproduce undefined behaviour from unitialized memory."));
 #endif //!UE_BUILD_SHIPPING
+
+	TAutoConsoleVariable<bool> CVarPassGPUDataThroughGridLinks(
+		TEXT("pcg.Graph.GPU.PassGPUDataThroughGridLinks"),
+		true,
+		TEXT("Whether proxies for GPU data are cached in per pin output data and passed through grid links. If false data is read back to CPU."));
 }
 
 namespace PCGHiGenGrid
diff --git a/Engine/Plugins/PCG/Source/PCG/Private/PCGElement.cpp b/Engine/Plugins/PCG/Source/PCG/Private/PCGElement.cpp
--- a/Engine/Plugins/PCG/Source/PCG/Private/PCGElement.cpp
+++ b/Engine/Plugins/PCG/Source/PCG/Private/PCGElement.cpp
@@ -10,7 +10,6 @@
 #include "Data/PCGBasePointData.h"
 #include "Data/PCGPointArrayData.h"
 #include "Data/PCGPointData.h"
-#include "Elements/PCGHiGenGridSize.h"
 #include "Graph/PCGGraphCache.h"
 #include "Helpers/PCGActorHelpers.h"
 #include "Helpers/PCGHelpers.h"
@@ -398,8 +397,7 @@ bool IPCGElement::ConvertInputsIfNeeded(FPCGContext* Context) const
 		return true;
 	}
 
-	// todo_pcg: Force-allow GPU resident data to tip-toe around ABI for hotfix. Add implementation of SupportsGPUResidentData to the hi-gen grid size node in 5.7.
-	if (!SupportsGPUResidentData(Context) && Context->GetInputSettings<UPCGHiGenGridSizeSettings>() == nullptr)
+	if (!SupportsGPUResidentData(Context))
 	{
 		bool bHasPendingReadbacks = false;
 
diff --git a/Engine/Plugins/PCG/Source/PCG/Public/Elements/PCGHiGenGridSize.h b/Engine/Plugins/PCG/Source/PCG/Public/Elements/PCGHiGenGridSize.h
--- a/Engine/Plugins/PCG/Source/PCG/Public/Elements/PCGHiGenGridSize.h
+++ b/Engine/Plugins/PCG/Source/PCG/Public/Elements/PCGHiGenGridSize.h
@@ -60,4 +60,5 @@ public:
 protected:
 	virtual bool ExecuteInternal(FPCGContext* Context) const override;
 	virtual bool SupportsBasePointDataInputs(FPCGContext* InContext) const override { return true; }
+	virtual bool SupportsGPUResidentData(FPCGContext* InContext) const;
 };
diff --git a/Engine/Plugins/PCG/Source/PCG/Public/PCGCommon.h b/Engine/Plugins/PCG/Source/PCG/Public/PCGCommon.h
--- a/Engine/Plugins/PCG/Source/PCG/Public/PCGCommon.h
+++ b/Engine/Plugins/PCG/Source/PCG/Public/PCGCommon.h
@@ -285,6 +285,7 @@ namespace PCGSystemSwitches
 
 	extern PCG_API TAutoConsoleVariable<bool> CVarReleaseTransientResourcesEarly;
 	extern PCG_API TAutoConsoleVariable<bool> CVarPCGDebugDrawGeneratedCells;
+	extern TAutoConsoleVariable<bool> CVarPassGPUDataThroughGridLinks;
 
 #if !UE_BUILD_SHIPPING
 	extern PCG_API TAutoConsoleVariable<bool> CVarFuzzGPUMemory;
