[PCG] Prefix texture data interface define to protect against collisions with engine shaders

#rnx

diff --git a/Engine/Plugins/PCG/Shaders/Private/PCGTextureDataInterface.ush b/Engine/Plugins/PCG/Shaders/Private/PCGTextureDataInterface.ush
--- a/Engine/Plugins/PCG/Shaders/Private/PCGTextureDataInterface.ush
+++ b/Engine/Plugins/PCG/Shaders/Private/PCGTextureDataInterface.ush
@@ -173,7 +173,7 @@ float4 Sample_{DataInterfaceName}(uint InDataIndex, float2 InTextureUV)
 	uint PointSample = 0;
 	GetTextureInfoInternal_{DataInterfaceName}(InDataIndex, BindingIndex, SliceIndex, TextureDimension, PointSample);
 
-#if ENABLE_MULTIPLE_TEXTURE_OBJECTS
+#if PCG_ENABLE_MULTIPLE_TEXTURE_OBJECTS
 	PCG_CONDITIONAL_TEXTURE_SAMPLE(0)
 	PCG_CONDITIONAL_TEXTURE_SAMPLE(1)
 	PCG_CONDITIONAL_TEXTURE_SAMPLE(2)
@@ -259,7 +259,7 @@ float4 SampleWorldPos_{DataInterfaceName}(uint InDataIndex, float2 WorldPos)
 	uint PointSample = 0;
 	GetTextureInfoInternal_{DataInterfaceName}(InDataIndex, BindingIndex, SliceIndex, TextureDimension, PointSample);
 
-#if ENABLE_MULTIPLE_TEXTURE_OBJECTS
+#if PCG_ENABLE_MULTIPLE_TEXTURE_OBJECTS
 	PCG_CONDITIONAL_TEXTURE_SAMPLE_WORLD_POS(0)
 	PCG_CONDITIONAL_TEXTURE_SAMPLE_WORLD_POS(1)
 	PCG_CONDITIONAL_TEXTURE_SAMPLE_WORLD_POS(2)
@@ -332,7 +332,7 @@ float4 Load_{DataInterfaceName}(uint InDataIndex, uint2 InElementIndex)
 	uint PointSample = 0;
 	GetTextureInfoInternal_{DataInterfaceName}(InDataIndex, BindingIndex, SliceIndex, TextureDimension, PointSample);
 
-#if ENABLE_MULTIPLE_TEXTURE_OBJECTS
+#if PCG_ENABLE_MULTIPLE_TEXTURE_OBJECTS
 	PCG_CONDITIONAL_TEXTURE_LOAD(0)
 	PCG_CONDITIONAL_TEXTURE_LOAD(1)
 	PCG_CONDITIONAL_TEXTURE_LOAD(2)
@@ -385,7 +385,7 @@ void Store_{DataInterfaceName}(uint InDataIndex, uint2 InElementIndex, float4 In
 		//return;
 	}
 
-#if ENABLE_MULTIPLE_TEXTURE_OBJECTS
+#if PCG_ENABLE_MULTIPLE_TEXTURE_OBJECTS
 	PCG_CONDITIONAL_TEXTURE_STORE(0)
 #else
 	{DataInterfaceName}_TextureUAV_0[InElementIndex] = InValue;
diff --git a/Engine/Plugins/PCG/Source/PCG/Private/Compute/DataInterfaces/PCGTextureDataInterface.cpp b/Engine/Plugins/PCG/Source/PCG/Private/Compute/DataInterfaces/PCGTextureDataInterface.cpp
--- a/Engine/Plugins/PCG/Source/PCG/Private/Compute/DataInterfaces/PCGTextureDataInterface.cpp
+++ b/Engine/Plugins/PCG/Source/PCG/Private/Compute/DataInterfaces/PCGTextureDataInterface.cpp
@@ -32,7 +32,7 @@
 
 namespace PCGTextureDataInterfaceConstants
 {
-	const TCHAR* EnableMultipleTextureObjectsPermutationName = TEXT("ENABLE_MULTIPLE_TEXTURE_OBJECTS");
+	const TCHAR* EnableMultipleTextureObjectsPermutationName = TEXT("PCG_ENABLE_MULTIPLE_TEXTURE_OBJECTS");
 }
 
 namespace PCGTextureDataInterfaceHelpers
