[PCG] Show CopyElementFrom helper for attribute set declarations as well as point declarations

#rnx

diff --git a/Engine/Plugins/PCG/Source/PCG/Private/Compute/Elements/PCGCustomHLSL.cpp b/Engine/Plugins/PCG/Source/PCG/Private/Compute/Elements/PCGCustomHLSL.cpp
--- a/Engine/Plugins/PCG/Source/PCG/Private/Compute/Elements/PCGCustomHLSL.cpp
+++ b/Engine/Plugins/PCG/Source/PCG/Private/Compute/Elements/PCGCustomHLSL.cpp
@@ -682,7 +682,10 @@ void UPCGCustomHLSLSettings::UpdateOutputDeclarations()
 			OutputDeclarations += TEXT("// Valid pins: ") + FString::Join(DataCollectionDataPins, TEXT(", ")) + TEXT("\n\n");
 		}
 
-		OutputDeclarations += FString::Format(TEXT("void {0}_GetElementCountMultiplier();\n\n"),
+		OutputDeclarations += FString::Format(TEXT(
+			"void {0}_GetElementCountMultiplier();\n"
+			"void {0}_CopyElementFrom_{input_pin}(uint TargetDataIndex, uint TargetElementIndex, uint SourceDataIndex, uint SourceElementIndex);\n"
+			"\n"),
 			{ bMultiPin ? PCGHLSLElement::PinDeclTemplateStr : DataCollectionDataPins[0] });
 
 		OutputDeclarations += FString::Format(TEXT(
@@ -709,7 +712,6 @@ void UPCGCustomHLSLSettings::UpdateOutputDeclarations()
 
 		OutputDeclarations += FString::Format(TEXT(
 			"void {0}_InitializePoint(uint DataIndex, uint ElementIndex);\n"
-			"void {0}_CopyElementFrom_{input pin}(uint TargetDataIndex, uint TargetElementIndex, uint SourceDataIndex, uint SourceElementIndex);\n"
 			"bool {0}_RemovePoint(uint DataIndex, uint ElementIndex);\n"
 			"\n"
 			"void {0}_SetPosition(uint DataIndex, uint ElementIndex, float3 Position);\n"
