[PCG] Fix processor kernels not having an input pin when setting kernel type from the node settings

#rnx

diff --git a/Engine/Plugins/PCG/Source/PCG/Private/Compute/Elements/PCGCustomHLSL.cpp b/Engine/Plugins/PCG/Source/PCG/Private/Compute/Elements/PCGCustomHLSL.cpp
--- a/Engine/Plugins/PCG/Source/PCG/Private/Compute/Elements/PCGCustomHLSL.cpp
+++ b/Engine/Plugins/PCG/Source/PCG/Private/Compute/Elements/PCGCustomHLSL.cpp
@@ -881,6 +881,26 @@ void UPCGCustomHLSLSettings::UpdateHelperDeclarations()
 
 void UPCGCustomHLSLSettings::UpdatePinSettings()
 {
+	if (IsProcessorKernel() && InputPins.IsEmpty())
+	{
+		if (IsPointKernel())
+		{
+			InputPins.Emplace(PCGPinConstants::DefaultInputLabel, EPCGDataType::Point);
+		}
+		else if (IsTextureKernel())
+		{
+			InputPins.Emplace(PCGPinConstants::DefaultInputLabel, EPCGDataType::BaseTexture);
+		}
+		else if (KernelType == EPCGKernelType::AttributeSetProcessor)
+		{
+			InputPins.Emplace(PCGPinConstants::DefaultInputLabel, EPCGDataType::Param);
+		}
+		else
+		{
+			ensure(false);
+		}
+	}
+
 	// Setup input pins.
 	for (int PinIndex = 0; PinIndex < InputPins.Num(); ++PinIndex)
 	{
