[PCG] Re-introduce cvar to disable landscape VT sampling

#rnx

diff --git a/Engine/Plugins/PCG/Source/PCG/Private/Compute/DataInterfaces/PCGLandscapeDataInterface.cpp b/Engine/Plugins/PCG/Source/PCG/Private/Compute/DataInterfaces/PCGLandscapeDataInterface.cpp
--- a/Engine/Plugins/PCG/Source/PCG/Private/Compute/DataInterfaces/PCGLandscapeDataInterface.cpp
+++ b/Engine/Plugins/PCG/Source/PCG/Private/Compute/DataInterfaces/PCGLandscapeDataInterface.cpp
@@ -28,6 +28,14 @@
 
 #define LOCTEXT_NAMESPACE "UPCGLandscapeDataInterface"
 
+namespace PCGLandscapeDataInterface
+{
+	static TAutoConsoleVariable<bool> CVarEnableLandscapeVTSampling(
+		TEXT("pcg.GPU.EnableLandscapeVTSampling"),
+		true,
+		TEXT("Allow the landscape data interface to use landscape virtual textures if available."));
+}
+
 /** RenderResource used to hold textures generated by this DI, pulled from the collision geometry of the terrain. */
 class FPCGLandscapeTextureResource : public FRenderResource
 {
@@ -288,7 +296,7 @@ void UPCGLandscapeDataProvider::Initialize(const UComputeDataInterface* InDataIn
 	const UPCGComponent* SourceComponent = Cast<UPCGComponent>(ExecutionSource);
 
 	// Sampling of virtual textures is not supported outside of runtime generation at this time, because the RuntimeGenScheduler is responsible for priming the virtual textures.
-	const bool bCanSampleVirtualTextures = SourceComponent ? SourceComponent->IsManagedByRuntimeGenSystem() : false;
+	const bool bCanSampleVirtualTextures = SourceComponent && SourceComponent->IsManagedByRuntimeGenSystem() && PCGLandscapeDataInterface::CVarEnableLandscapeVTSampling.GetValueOnAnyThread();
 
 	auto TestLandscape = [World, &ComponentBounds](const ALandscape* InLandscape)
 	{
