[PCG] Fix deadlock when cancelling GPU compilation

#rnx

diff --git a/Engine/Plugins/PCG/Source/PCG/Private/Compute/PCGComputeKernel.cpp b/Engine/Plugins/PCG/Source/PCG/Private/Compute/PCGComputeKernel.cpp
--- a/Engine/Plugins/PCG/Source/PCG/Private/Compute/PCGComputeKernel.cpp
+++ b/Engine/Plugins/PCG/Source/PCG/Private/Compute/PCGComputeKernel.cpp
@@ -99,51 +99,19 @@ bool UPCGComputeKernel::AreKernelSettingsValid(FPCGContext* InContext) const
 #if WITH_EDITOR
 void UPCGComputeKernel::GatherAdditionalSources(TArray<TObjectPtr<UComputeSource>>& OutAdditionalSources) const
 {
-	auto GatherSources = [this, &OutAdditionalSources, &DefaultAdditionalSourcePaths=PCGComputeKernel::DefaultAdditionalSourcePaths]()
+	for (const FSoftObjectPath& AdditionalSourcePath : PCGComputeKernel::DefaultAdditionalSourcePaths)
 	{
-		check(IsInGameThread());
-
-		for (const FSoftObjectPath& AdditionalSourcePath : DefaultAdditionalSourcePaths)
+		UPCGComputeSource* AdditionalSource = Cast<UPCGComputeSource>(AdditionalSourcePath.ResolveObject());
+		
+		if (ensure(AdditionalSource))
 		{
-			TSoftObjectPtr<UPCGComputeSource> AdditionalSourcePtr(AdditionalSourcePath);
-			UPCGComputeSource* AdditionalSource = Cast<UPCGComputeSource>(AdditionalSourcePtr.LoadSynchronous());
-			
-			if (ensure(AdditionalSource))
-			{
-				OutAdditionalSources.Add(AdditionalSource);
-			}
-			else
-			{
-				UE_LOG(LogPCG, Error, TEXT("Default additional compute source '%s' could not be loaded for kernel '%s'."), *AdditionalSourcePath.ToString(), *GetName());
-			}
-		}
-	};
-
-	if (!IsInGameThread())
-	{
-		// Loads must happen on the game thread, so we launch and wait for an async game thread task if necessary.
-		if (FEvent* LoadEvent = FPlatformProcess::GetSynchEventFromPool(/*bIsManualReset=*/true))
-		{
-			AsyncTask(ENamedThreads::GameThread, [&GatherSources, LoadEvent]()
-			{
-				GatherSources();
-				LoadEvent->Trigger();
-			});
-
-			LoadEvent->Wait();
-
-			FPlatformProcess::ReturnSynchEventToPool(LoadEvent);
+			OutAdditionalSources.Add(AdditionalSource);
 		}
 		else
 		{
-			UE_LOG(LogPCG, Error, TEXT("Default additional compute sources could not be loaded for kernel '%s'."), *GetName());
+			UE_LOG(LogPCG, Error, TEXT("Default additional compute source '%s' could not be resolved for kernel '%s'."), *AdditionalSourcePath.ToString(), *GetName());
 		}
 	}
-	else
-	{
-		// If we're already on the game thread, just load the compute sources.
-		GatherSources();
-	}
 }
 #endif
 
diff --git a/Engine/Plugins/PCG/Source/PCG/Private/PCGModule.cpp b/Engine/Plugins/PCG/Source/PCG/Private/PCGModule.cpp
--- a/Engine/Plugins/PCG/Source/PCG/Private/PCGModule.cpp
+++ b/Engine/Plugins/PCG/Source/PCG/Private/PCGModule.cpp
@@ -32,6 +32,19 @@ FPCGGraphChangedDelegate FPCGModule::OnGraphChangedDelegate;
 namespace PCGModule
 {
 	FPCGModule* PCGModulePtr = nullptr;
+
+#if WITH_EDITOR
+	// @todo_pcg: Temporarily copied from PCGComputeKernel.cpp to avoid changing the ABI for 5.6.1 hotfix. Should be removed in 5.7.
+	static const TStaticArray<FSoftObjectPath, 3> DefaultAdditionalSourcePaths =
+	{
+		FSoftObjectPath(TEXT("/Script/PCG.PCGComputeSource'/PCG/ComputeSources/PCGCS_ShaderUtils.PCGCS_ShaderUtils'")),
+		FSoftObjectPath(TEXT("/Script/PCG.PCGComputeSource'/PCG/ComputeSources/PCGCS_ShaderUtilsInternal.PCGCS_ShaderUtilsInternal'")),
+
+		// Note: PCGDataCollectionDataInterface.ush depends on the quaternion helpers, therefore all kernels also depend on the quaternion helpers.
+		// @todo_pcg: In the future quaternion compute source could be opt-in if the kernel does not manipulate point/attribute data.
+		FSoftObjectPath(TEXT("/Script/PCG.PCGComputeSource'/PCG/ComputeSources/PCGCS_Quaternion.PCGCS_Quaternion'"))
+	};
+#endif
 }
 
 FPCGModule& FPCGModule::GetPCGModuleChecked()
@@ -55,6 +68,17 @@ void FPCGModule::StartupModule()
 
 	const FString PluginShaderDir = FPaths::Combine(IPluginManager::Get().FindPlugin(TEXT("PCG"))->GetBaseDir(), TEXT("Shaders"));
 	AddShaderSourceDirectoryMapping(TEXT("/Plugin/PCG"), PluginShaderDir);
+
+	// @todo_pcg: In the future, this could be done in the compute kernel CDO constructor instead, to make kernels more self-contained.
+	for (const FSoftObjectPath& AdditionalSourcePath : PCGModule::DefaultAdditionalSourcePaths)
+	{
+		const UObject* LoadedAdditionalSource = AdditionalSourcePath.TryLoad();
+
+		if (!ensure(LoadedAdditionalSource))
+		{
+			UE_LOG(LogPCG, Error, TEXT("Failed to load compute source asset '%s'."), *AdditionalSourcePath.ToString());
+		}
+	}
 #endif // WITH_EDITOR
 
 	// Cache for fast access
