[PCG] GPU - remove unnecessary buffer copies when reusing data collection buffers from upstream graphs.

These copies could dominate overall cost in some scenarios.

[RN] Optimized reuse of upstream GPU data buffers by removing buffer copy passes.

diff --git a/Engine/Plugins/PCG/Shaders/Private/PCGDataCollectionAdaptor.usf b/Engine/Plugins/PCG/Shaders/Private/PCGDataCollectionAdaptor.usf
deleted file mode 100644
diff --git a/Engine/Plugins/PCG/Shaders/Private/PCGDataCollectionDataInterface.ush b/Engine/Plugins/PCG/Shaders/Private/PCGDataCollectionDataInterface.ush
--- a/Engine/Plugins/PCG/Shaders/Private/PCGDataCollectionDataInterface.ush
+++ b/Engine/Plugins/PCG/Shaders/Private/PCGDataCollectionDataInterface.ush
@@ -21,6 +21,10 @@ RWByteAddressBuffer {DataInterfaceName}_DataCollectionBuffer;
 // Multiplier that can be configured on output pins.
 uint {DataInterfaceName}_ElementCountMultiplier;
 
+StructuredBuffer<int> {DataInterfaceName}_GraphToBufferAttributeId;
+
+int {DataInterfaceName}_FirstRemappedAttributeId;
+
 // #################### HEADER READERS ##########################
 
 uint GetNumData_{DataInterfaceName}()
@@ -88,13 +92,28 @@ void StoreBufferInternal_{DataInterfaceName}(uint Address, uint Value)
 	{DataInterfaceName}_DataCollectionBuffer.Store(Address, Value);
 }
 
+int GetRemappedAttributeId_{DataInterfaceName}(int InGraphAttributeId)
+{
+	if (InGraphAttributeId >= {DataInterfaceName}_FirstRemappedAttributeId)
+	{
+		// Convert from current compute graph attribute ID present in buffer (required if buffer came from a different upstream compute graph).
+		return PCG_NUM_RESERVED_ATTRS + {DataInterfaceName}_GraphToBufferAttributeId[InGraphAttributeId - PCG_NUM_RESERVED_ATTRS];
+	}
+	else
+	{
+		return InGraphAttributeId;
+	}
+}
+
 uint GetFirstElementAddressInternal_{DataInterfaceName}(uint InDataIndex, int InAttributeId)
 {
+	const int AttributeId = GetRemappedAttributeId_{DataInterfaceName}(InAttributeId);
+
 	const uint ReadAddress =
 		PCG_DATA_COLLECTION_HEADER_SIZE_BYTES +
 		InDataIndex * PCG_DATA_HEADER_SIZE_BYTES +
 		/*TypeId*/4 + /*Attribute Count*/4 + /*Element Count*/4 +
-		InAttributeId * PCG_ATTRIBUTE_HEADER_SIZE_BYTES +
+		AttributeId * PCG_ATTRIBUTE_HEADER_SIZE_BYTES +
 		/*PackedAttributeIdAndStride*/4;
 
 	return {DataInterfaceName}_DataCollectionBuffer.Load(ReadAddress);
@@ -102,11 +121,13 @@ uint GetFirstElementAddressInternal_{DataInterfaceName}(uint InDataIndex, int In
 
 uint GetAttributeStrideInternal_{DataInterfaceName}(uint InDataIndex, int InAttributeId)
 {
+	const int AttributeId = GetRemappedAttributeId_{DataInterfaceName}(InAttributeId);
+
 	const uint ReadAddress =
 		PCG_DATA_COLLECTION_HEADER_SIZE_BYTES +
 		InDataIndex * PCG_DATA_HEADER_SIZE_BYTES +
 		/*TypeId*/4 + /*Attribute Count*/4 + /*Element Count*/4 +
-		InAttributeId * PCG_ATTRIBUTE_HEADER_SIZE_BYTES;
+		AttributeId * PCG_ATTRIBUTE_HEADER_SIZE_BYTES;
 
 	const uint PackedAttributeIdAndStride = {DataInterfaceName}_DataCollectionBuffer.Load(ReadAddress);
 	return PackedAttributeIdAndStride & (0xFFFFFFFF >> 24);
diff --git a/Engine/Plugins/PCG/Source/PCG/Internal/Compute/PCGDataForGPU.h b/Engine/Plugins/PCG/Source/PCG/Internal/Compute/PCGDataForGPU.h
--- a/Engine/Plugins/PCG/Source/PCG/Internal/Compute/PCGDataForGPU.h
+++ b/Engine/Plugins/PCG/Source/PCG/Internal/Compute/PCGDataForGPU.h
@@ -6,12 +6,12 @@
 #include "PCGData.h"
 #include "Compute/PCGComputeCommon.h"
 
+#include "RenderGraphResources.h"
 #include "RHIGPUReadback.h"
 #include "Misc/SpinLock.h"
 
 #include "PCGDataForGPU.generated.h"
 
-class FRDGPooledBuffer;
 class FReferenceCollector;
 class UPCGDataBinding;
 class UPCGMetadata;
diff --git a/Engine/Plugins/PCG/Source/PCG/Private/Compute/DataInterfaces/PCGDataCollectionDataInterface.cpp b/Engine/Plugins/PCG/Source/PCG/Private/Compute/DataInterfaces/PCGDataCollectionDataInterface.cpp
--- a/Engine/Plugins/PCG/Source/PCG/Private/Compute/DataInterfaces/PCGDataCollectionDataInterface.cpp
+++ b/Engine/Plugins/PCG/Source/PCG/Private/Compute/DataInterfaces/PCGDataCollectionDataInterface.cpp
@@ -18,6 +18,7 @@
 #include "RenderGraphResources.h"
 #include "ShaderCompilerCore.h"
 #include "ShaderParameterMetadataBuilder.h"
+#include "SystemTextures.h"
 #include "ComputeFramework/ComputeKernelPermutationSet.h"
 #include "ComputeFramework/ShaderParamTypeDefinition.h"
 
@@ -409,6 +410,8 @@ void UPCGDataCollectionDataInterface::GetSupportedInputs(TArray<FShaderFunctionD
 BEGIN_SHADER_PARAMETER_STRUCT(FPCGDataCollectionDataInterfaceParameters, )
 	SHADER_PARAMETER_RDG_BUFFER_UAV(RWByteAddressBuffer, DataCollectionBuffer)
 	SHADER_PARAMETER(uint32, ElementCountMultiplier)
+	SHADER_PARAMETER_RDG_BUFFER_SRV(StructuredBuffer<int32>, GraphToBufferAttributeId)
+	SHADER_PARAMETER(int32, FirstRemappedAttributeId)
 END_SHADER_PARAMETER_STRUCT()
 
 void UPCGDataCollectionDataInterface::GetShaderParameters(TCHAR const* UID, FShaderParametersMetadataBuilder& InOutBuilder, FShaderParametersMetadataAllocations& InOutAllocations) const
@@ -465,6 +468,7 @@ void UPCGDataCollectionDataInterface::GetDefines(FComputeKernelDefinitionSet& Ou
 	// Misc
 	OutDefinitionSet.Defines.Add(FComputeKernelDefinition(TEXT("PCG_KERNEL_EXECUTED_FLAG"), FString::FromInt(PCGComputeConstants::KernelExecutedFlag)));
 	OutDefinitionSet.Defines.Add(FComputeKernelDefinition(TEXT("PCG_INVALID_DENSITY"),      FString::Printf(TEXT("%f"), PCGComputeConstants::INVALID_DENSITY)));
+	OutDefinitionSet.Defines.Add(FComputeKernelDefinition(TEXT("PCG_NUM_RESERVED_ATTRS"),   FString::FromInt(PCGComputeConstants::NUM_RESERVED_ATTRS)));
 }
 
 UComputeDataProvider* UPCGDataCollectionDataInterface::CreateDataProvider() const
@@ -543,6 +547,14 @@ FPCGDataCollectionDataProviderProxy::FPCGDataCollectionDataProviderProxy(
 
 		// @todo_pcg: Avoid costly copy of data description.
 		PinDesc = *InPinDesc;
+
+		for (const FPCGDataDesc& DataDesc : PinDesc.DataDescs)
+		{
+			for (const FPCGKernelAttributeDesc& AttrDesc : DataDesc.AttributeDescs)
+			{
+				LargestAttributeIndex = FMath::Max(LargestAttributeIndex, PCGComputeHelpers::GetMetadataAttributeIndexFromAttributeId(AttrDesc.AttributeId));
+			}
+		}
 	}
 
 	SizeBytes = PinDesc.ComputePackedSizeBytes();
@@ -574,6 +586,8 @@ void FPCGDataCollectionDataProviderProxy::GatherDispatchData(FDispatchData const
 		FParameters& Parameters = ParameterArray[InvocationIndex];
 		Parameters.DataCollectionBuffer = BufferUAV;
 		Parameters.ElementCountMultiplier = ElementCountMultiplier;
+		Parameters.GraphToBufferAttributeId = AttributeIdRemapSRV;
+		Parameters.FirstRemappedAttributeId = FirstRemappedAttributeId;
 	}
 }
 
@@ -582,6 +596,8 @@ void FPCGDataCollectionDataProviderProxy::AllocateResources(FRDGBuilder& GraphBu
 	TRACE_CPUPROFILER_EVENT_SCOPE(FPCGDataCollectionDataProviderProxy::AllocateResources);
 	check(SizeBytes > 0);
 
+	AttributeIdRemapSRV = GetAttributeIdRemapBufferSRV(GraphBuilder, FirstRemappedAttributeId);
+
 	FRDGBufferDesc Desc = FRDGBufferDesc::CreateByteAddressDesc(SizeBytes);
 	if (ExportMode != EPCGExportMode::NoExport)
 	{
@@ -669,4 +685,12 @@ void FPCGDataCollectionDataProviderProxy::AllocateResources(FRDGBuilder& GraphBu
 	}
 }
 
+FRDGBufferSRVRef FPCGDataCollectionDataProviderProxy::GetAttributeIdRemapBufferSRV(FRDGBuilder& GraphBuilder, int32& OutFirstRemappedAttributeId) const
+{
+	// Default implementation. No remap required, disabled remapping.
+	OutFirstRemappedAttributeId = PCGComputeConstants::MAX_NUM_ATTRS;
+
+	return GraphBuilder.CreateSRV(FRDGBufferSRVDesc(GSystemTextures.GetDefaultStructuredBuffer(GraphBuilder, sizeof(int32))));
+}
+
 #undef LOCTEXT_NAMESPACE
diff --git a/Engine/Plugins/PCG/Source/PCG/Private/Compute/DataInterfaces/PCGDataCollectionDataInterface.h b/Engine/Plugins/PCG/Source/PCG/Private/Compute/DataInterfaces/PCGDataCollectionDataInterface.h
--- a/Engine/Plugins/PCG/Source/PCG/Private/Compute/DataInterfaces/PCGDataCollectionDataInterface.h
+++ b/Engine/Plugins/PCG/Source/PCG/Private/Compute/DataInterfaces/PCGDataCollectionDataInterface.h
@@ -96,6 +96,9 @@ public:
 	void AllocateResources(FRDGBuilder& GraphBuilder, FAllocationData const& InAllocationData) override;
 	//~ End FComputeDataProviderRenderProxy Interface
 
+protected:
+	virtual FRDGBufferSRVRef GetAttributeIdRemapBufferSRV(FRDGBuilder& GraphBuilder, int32& OutFirstRemappedAttributeId) const;
+
 public:
 	/** Generation count of the data provider when the proxy was created. */
 	uint64 OriginatingGenerationCount = 0;
@@ -121,4 +124,10 @@ protected:
 
 	FName OutputPinLabel;
 	FName OutputPinLabelAlias;
+
+	FRDGBufferSRVRef AttributeIdRemapSRV = nullptr;
+
+	int32 FirstRemappedAttributeId = INDEX_NONE;
+
+	int32 LargestAttributeIndex = INDEX_NONE;
 };
diff --git a/Engine/Plugins/PCG/Source/PCG/Private/Compute/DataInterfaces/PCGDataCollectionUploadDataInterface.cpp b/Engine/Plugins/PCG/Source/PCG/Private/Compute/DataInterfaces/PCGDataCollectionUploadDataInterface.cpp
--- a/Engine/Plugins/PCG/Source/PCG/Private/Compute/DataInterfaces/PCGDataCollectionUploadDataInterface.cpp
+++ b/Engine/Plugins/PCG/Source/PCG/Private/Compute/DataInterfaces/PCGDataCollectionUploadDataInterface.cpp
@@ -8,7 +8,6 @@
 #include "PCGNode.h"
 #include "PCGSettings.h"
 #include "PCGSubsystem.h"
-#include "PCGDataCollectionAdaptorCS.h"
 #include "Compute/PCGComputeCommon.h"
 #include "Compute/PCGComputeGraph.h"
 #include "Compute/PCGDataBinding.h"
@@ -134,4 +133,19 @@ void FPCGDataProviderDataCollectionUploadProxy::AllocateResources(FRDGBuilder& G
 
 	Buffer = DataAdaptor->GetBuffer_RenderThread(GraphBuilder, ExportMode);
 	BufferUAV = GraphBuilder.CreateUAV(Buffer);
+
+	AttributeIdRemapSRV = GetAttributeIdRemapBufferSRV(GraphBuilder, FirstRemappedAttributeId);
+}
+
+FRDGBufferSRVRef FPCGDataProviderDataCollectionUploadProxy::GetAttributeIdRemapBufferSRV(FRDGBuilder& GraphBuilder, int32& OutFirstRemappedAttributeId) const
+{
+	if (!DataAdaptor || !DataAdaptor->UsesAttributeIdRemap())
+	{
+		// No remap required.
+		return FPCGDataCollectionDataProviderProxy::GetAttributeIdRemapBufferSRV(GraphBuilder, OutFirstRemappedAttributeId);
+	}
+	else
+	{
+		return DataAdaptor->GetAttributeRemapBufferSRV(GraphBuilder, OutFirstRemappedAttributeId);
+	}
 }
diff --git a/Engine/Plugins/PCG/Source/PCG/Private/Compute/DataInterfaces/PCGDataCollectionUploadDataInterface.h b/Engine/Plugins/PCG/Source/PCG/Private/Compute/DataInterfaces/PCGDataCollectionUploadDataInterface.h
--- a/Engine/Plugins/PCG/Source/PCG/Private/Compute/DataInterfaces/PCGDataCollectionUploadDataInterface.h
+++ b/Engine/Plugins/PCG/Source/PCG/Private/Compute/DataInterfaces/PCGDataCollectionUploadDataInterface.h
@@ -64,6 +64,9 @@ public:
 	void AllocateResources(FRDGBuilder& GraphBuilder, FAllocationData const& InAllocationData) override;
 	//~ End FComputeDataProviderRenderProxy Interface
 
+protected:
+	virtual FRDGBufferSRVRef GetAttributeIdRemapBufferSRV(FRDGBuilder& GraphBuilder, int32& OutFirstRemappedAttributeId) const override;
+
 protected:
 	TSharedPtr<FPCGDataUploadAdaptor> DataAdaptor;
 
diff --git a/Engine/Plugins/PCG/Source/PCG/Private/Compute/PCGComputeGraph.cpp b/Engine/Plugins/PCG/Source/PCG/Private/Compute/PCGComputeGraph.cpp
--- a/Engine/Plugins/PCG/Source/PCG/Private/Compute/PCGComputeGraph.cpp
+++ b/Engine/Plugins/PCG/Source/PCG/Private/Compute/PCGComputeGraph.cpp
@@ -195,7 +195,7 @@ void UPCGComputeGraph::DebugLogDataDescriptions(const UPCGDataBinding* InBinding
 	const UEnum* PCGKernelAttributeTypeEnum = StaticEnum<EPCGKernelAttributeType>();
 	check(PCGDataTypeEnum && PCGKernelAttributeTypeEnum);
 
-	auto LogDataDescription = [this, PCGDataTypeEnum, PCGKernelAttributeTypeEnum](const FPCGDataCollectionDesc* InDataDesc)
+	auto LogDataDescription = [this, PCGDataTypeEnum, PCGKernelAttributeTypeEnum, InBinding](const FPCGDataCollectionDesc* InDataDesc)
 	{
 		if (!InDataDesc)
 		{
@@ -231,11 +231,41 @@ void UPCGComputeGraph::DebugLogDataDescriptions(const UPCGDataBinding* InBinding
 					UniqueStringKeys += FString::Format(TEXT("{0}"), { StringKey });
 				}
 
-				UE_LOG(LogPCG, Display, TEXT("\t\t\t\tID: %d\t\tName: %s\t\tType: %s\t\tUniqueStringKeys%s"),
-					AttributeDesc.AttributeId,
-					*AttributeDesc.AttributeKey.Identifier.ToString(),
-					*PCGKernelAttributeTypeEnum->GetDisplayNameTextByValue(static_cast<int64>(AttributeDesc.AttributeKey.Type)).ToString(),
-					UniqueStringKeys.IsEmpty() ? TEXT(": ") : *UniqueStringKeys);
+				// If attribute comes from an upstream GPU buffer, display the attribute ID in the upstream buffer.
+				int32 SourceBufferAttributeId = INDEX_NONE;
+
+				if (AttributeDesc.AttributeId >= PCGComputeConstants::NUM_RESERVED_ATTRS)
+				{
+					if (ensure(InBinding->InputDataCollection.TaggedData.IsValidIndex(DataIndex)))
+					{
+						const UPCGProxyForGPUData* Proxy = Cast<UPCGProxyForGPUData>(InBinding->InputDataCollection.TaggedData[DataIndex].Data);
+						if (TSharedPtr<const FPCGProxyForGPUDataCollection> DataCollectionInfo = Proxy ? Proxy->GetGPUInfo() : nullptr)
+						{
+							const int32 GraphAttributeIndex = PCGComputeHelpers::GetMetadataAttributeIndexFromAttributeId(AttributeDesc.AttributeId);
+							const int32* FoundIndex = InBinding->GetGraphAttributeToSourceBufferAttributeIndex().Find({ GraphAttributeIndex, DataCollectionInfo->GetBuffer() });
+
+							SourceBufferAttributeId = PCGComputeHelpers::GetAttributeIdFromMetadataAttributeIndex(FoundIndex ? *FoundIndex : GraphAttributeIndex);
+						}
+					}
+				}
+
+				if (SourceBufferAttributeId == INDEX_NONE)
+				{
+					UE_LOG(LogPCG, Display, TEXT("\t\t\t\tID: %d\t\tName: %s\t\tType: %s\t\tUniqueStringKeys%s"),
+						AttributeDesc.AttributeId,
+						*AttributeDesc.AttributeKey.Identifier.ToString(),
+						*PCGKernelAttributeTypeEnum->GetDisplayNameTextByValue(static_cast<int64>(AttributeDesc.AttributeKey.Type)).ToString(),
+						UniqueStringKeys.IsEmpty() ? TEXT(": ") : *UniqueStringKeys);
+				}
+				else
+				{
+					UE_LOG(LogPCG, Display, TEXT("\t\t\t\tID: %d\t\tName: %s\t\tType: %s\t\tUniqueStringKeys%s\t\tSource Buffer ID: %d"),
+						AttributeDesc.AttributeId,
+						*AttributeDesc.AttributeKey.Identifier.ToString(),
+						*PCGKernelAttributeTypeEnum->GetDisplayNameTextByValue(static_cast<int64>(AttributeDesc.AttributeKey.Type)).ToString(),
+						UniqueStringKeys.IsEmpty() ? TEXT(": ") : *UniqueStringKeys,
+						SourceBufferAttributeId);
+				}
 			}
 
 			UE_LOG(LogPCG, Display, TEXT("\t\t\tTags (%d)"), DataDesc.TagStringKeys.Num());
@@ -329,7 +359,7 @@ bool UPCGComputeGraph::ComputeTaggedDataPinDesc(const FPCGTaggedData& InTaggedDa
 					continue;
 				}
 
-				if (const int32* AttributeIndex = InBinding->GetAttributeIdRemap().Find(UPCGDataBinding::FSourceBufferAndAttributeIndex(DataGPU, Index)))
+				if (const int32* AttributeIndex = InBinding->GetSourceBufferAttributeToGraphAttributeIndex().Find(UPCGDataBinding::FSourceBufferAndAttributeIndex(DataGPU, Index)))
 				{
 					const int32 NewId = PCGComputeHelpers::GetAttributeIdFromMetadataAttributeIndex(*AttributeIndex);
 					//UE_LOG(LogPCG, Warning, TEXT("Remapped attribute '%s' ID: %d -> %d"), *AttrDesc.Name.ToString(), AttrDesc.Index, NewId);
diff --git a/Engine/Plugins/PCG/Source/PCG/Private/Compute/PCGDataBinding.cpp b/Engine/Plugins/PCG/Source/PCG/Private/Compute/PCGDataBinding.cpp
--- a/Engine/Plugins/PCG/Source/PCG/Private/Compute/PCGDataBinding.cpp
+++ b/Engine/Plugins/PCG/Source/PCG/Private/Compute/PCGDataBinding.cpp
@@ -61,7 +61,8 @@ void UPCGDataBinding::InitializeTables(FPCGContext* InContext)
 
 void UPCGDataBinding::ReleaseTransientResources()
 {
-	AttributeIdRemap.Reset();
+	SourceBufferAttributeToGraphAttributeIndex.Reset();
+	GraphAttributeToSourceBufferAttributeIndex.Reset();
 	OutputDataCollection.TaggedData.Reset();
 	InputDataCollection.TaggedData.Reset();
 	MeshSpawnersToPrimitives.Reset();
@@ -304,8 +305,10 @@ void UPCGDataBinding::AddInputDataAttributesToTable()
 	{
 		if (const UPCGProxyForGPUData* Proxy = Cast<UPCGProxyForGPUData>(Data.Data))
 		{
+			TSharedPtr<const FPCGProxyForGPUDataCollection> DataCollectionInfo = Proxy->GetGPUInfo();
+
 			FPCGDataDesc SourceDesc;
-			if (Proxy->GetDescription(SourceDesc))
+			if (DataCollectionInfo && Proxy->GetDescription(SourceDesc))
 			{
 				for (FPCGKernelAttributeDesc& SourceAttrDesc : SourceDesc.AttributeDescs)
 				{
@@ -317,11 +320,12 @@ void UPCGDataBinding::AddInputDataAttributesToTable()
 						continue;
 					}
 
-					const int32 ExistingAttributeIndex = AttributeTable.AddAttribute(SourceAttrDesc.AttributeKey.Identifier, SourceAttrDesc.AttributeKey.Type);
+					const int32 GraphAttributeIndex = AttributeTable.AddAttribute(SourceAttrDesc.AttributeKey.Identifier, SourceAttrDesc.AttributeKey.Type);
 
-					if (ExistingAttributeIndex != SourceAttributeTableIndex)
+					if (GraphAttributeIndex != SourceAttributeTableIndex)
 					{
-						AttributeIdRemap.Add(FSourceBufferAndAttributeIndex(Proxy->GetGPUInfo(), SourceAttributeTableIndex), ExistingAttributeIndex);
+						SourceBufferAttributeToGraphAttributeIndex.Add({ DataCollectionInfo, SourceAttributeTableIndex }, GraphAttributeIndex);
+						GraphAttributeToSourceBufferAttributeIndex.Add({ GraphAttributeIndex, DataCollectionInfo->GetBuffer() }, SourceAttributeTableIndex);
 					}
 				}
 			}
@@ -499,7 +503,6 @@ void UPCGDataBinding::DebugLogDataDescriptions() const
 			UE_LOG(LogPCG, Display, TEXT("\t%d: %s"), i, *StringTable[i]);
 		}
 
-		// @todo_pcg: We don't have explicit pin lists for kernels because we don't need them. Decide how to bring this back or pivot the feature.
 		UE_LOG(LogPCG, Display, TEXT("\n### INPUT PIN DATA DESCRIPTIONS ###"));
 		Graph->DebugLogDataDescriptions(this);
 	}
diff --git a/Engine/Plugins/PCG/Source/PCG/Private/Compute/PCGDataBinding.h b/Engine/Plugins/PCG/Source/PCG/Private/Compute/PCGDataBinding.h
--- a/Engine/Plugins/PCG/Source/PCG/Private/Compute/PCGDataBinding.h
+++ b/Engine/Plugins/PCG/Source/PCG/Private/Compute/PCGDataBinding.h
@@ -96,6 +96,8 @@ public:
 	int32 GetAttributeId(const FPCGKernelAttributeKey& InAttribute) const { ensure(bTablesInitialized); return AttributeTable.GetAttributeId(InAttribute); }
 	int32 GetAttributeId(FPCGAttributeIdentifier InIdentifier, EPCGKernelAttributeType InType) const { ensure(bTablesInitialized); return AttributeTable.GetAttributeId(InIdentifier, InType); }
 
+	int32 GetAttributeTableSize() const { return AttributeTable.Num(); }
+
 	/** Get the set of all unique strings entering the compute graph. */
 	const TArray<FString>& GetStringTable() const { ensure(bTablesInitialized); return StringTable; }
 
@@ -114,8 +116,13 @@ public:
 	void ReceiveDataFromGPU_GameThread(UPCGData* InData, const UPCGSettings* InProducerSettings, EPCGExportMode InExportMode, FName InPinLabel, FName InPinLabelAlias, const TSet<FString>& AdditionalTags);
 
 	using FSourceBufferAndAttributeIndex = TPair<TSharedPtr<const FPCGProxyForGPUDataCollection>, /*Source Metadata Attribute Index*/int32>;
+	using FGraphAttributeIndexAndSourceBuffer = TPair</*Graph Metadata Attribute Index*/int32, TRefCountPtr<FRDGPooledBuffer>>;
+
+	/** Get map from source buffer and source buffer attribute index to this-graph attribute index. */
+	const TMap<FSourceBufferAndAttributeIndex, /*Graph Attribute Index*/int32>& GetSourceBufferAttributeToGraphAttributeIndex() const { return SourceBufferAttributeToGraphAttributeIndex; }
 
-	const TMap<FSourceBufferAndAttributeIndex, /*Attribute Index*/int32>& GetAttributeIdRemap() const { return AttributeIdRemap; }
+	/** Get map from this-graph attribute index and upstream source buffer to upstream attribute index. */
+	const TMap<FGraphAttributeIndexAndSourceBuffer, /*Source Buffer Attribute Index*/int32>& GetGraphAttributeToSourceBufferAttributeIndex() const { return GraphAttributeToSourceBufferAttributeIndex; }
 
 	IPCGGraphExecutionSource* GetExecutionSource() const;
 
@@ -174,8 +181,11 @@ protected:
 	UPROPERTY()
 	TArray<FString> StringTable;
 
-	/** Remaps attribute IDs to avoid collisions with this compute graph's attribute IDs. */
-	TMap<FSourceBufferAndAttributeIndex, /*Attribute ID*/int32> AttributeIdRemap;
+	/** Map from source buffer and source buffer attribute index to this-graph attribute index. */
+	TMap<FSourceBufferAndAttributeIndex, /*Attribute ID*/int32> SourceBufferAttributeToGraphAttributeIndex;
+
+	/** Map from this-graph attribute index and upstream source buffer to upstream attribute index. */
+	TMap<FGraphAttributeIndexAndSourceBuffer, /*Source Buffer Attribute Index*/int32> GraphAttributeToSourceBufferAttributeIndex;
 
 	/** Cache of data descriptions to amortize cost of computing them at runtime. */
 	TMap</*Binding index*/int32, FPCGDataCollectionDesc> DataDescriptionCache;
diff --git a/Engine/Plugins/PCG/Source/PCG/Private/Compute/PCGDataUploadAdaptor.cpp b/Engine/Plugins/PCG/Source/PCG/Private/Compute/PCGDataUploadAdaptor.cpp
--- a/Engine/Plugins/PCG/Source/PCG/Private/Compute/PCGDataUploadAdaptor.cpp
+++ b/Engine/Plugins/PCG/Source/PCG/Private/Compute/PCGDataUploadAdaptor.cpp
@@ -2,7 +2,6 @@
 
 #include "Compute/PCGDataUploadAdaptor.h"
 
-#include "PCGDataCollectionAdaptorCS.h"
 #include "PCGModule.h"
 #include "Compute/PCGComputeCommon.h"
 #include "Compute/PCGComputeGraph.h"
@@ -13,17 +12,7 @@
 #include "RenderGraphBuilder.h"
 #include "RenderGraphResources.h"
 #include "RenderGraphUtils.h"
-
-namespace PCGDataUploadAdaptor
-{
-#if !UE_BUILD_SHIPPING
-	static int32 TriggerGPUCaptureDispatches = 0;
-	static FAutoConsoleVariableRef CVarTriggerGPUCapture(
-		TEXT("pcg.GPU.TriggerRenderCaptures.DataUploadAdaptor"),
-		TriggerGPUCaptureDispatches,
-		TEXT("Trigger GPU captures for this many of the subsequent data collection update dispatches."));
-#endif // !UE_BUILD_SHIPPING
-}
+#include "SystemTextures.h"
 
 FPCGDataUploadAdaptor::FPCGDataUploadAdaptor(UPCGDataBinding* InDataBinding, const FPCGDataCollectionDesc* InTargetDataCollectionDesc, FName InInputPinLabel)
 {
@@ -105,6 +94,8 @@ FPCGDataUploadAdaptor::FPCGDataUploadAdaptor(UPCGDataBinding* InDataBinding, con
 /** Do any preparation work such as data readbacks. Returns true when preparation is complete. */
 bool FPCGDataUploadAdaptor::PrepareData_GameThread()
 {
+	TRACE_CPUPROFILER_EVENT_SCOPE(FPCGDataUploadAdaptor::PrepareData_GameThread);
+
 	UPCGDataBinding* Binding = DataBinding.Get();
 	if (!ensure(Binding))
 	{
@@ -115,25 +106,16 @@ bool FPCGDataUploadAdaptor::PrepareData_GameThread()
 
 	if (ExternalBufferForReuse)
 	{
-		ensureAlways(SourceToTargetAttributeId.IsEmpty());
+		ensureAlways(GraphToSourceBufferAttributeIndex.IsEmpty());
+		const int32 NumGraphAttributes = Binding->GetAttributeTableSize();
+		GraphToSourceBufferAttributeIndex.SetNumUninitialized(NumGraphAttributes);
 
-		SourceToTargetAttributeId.Reserve(Binding->GetAttributeIdRemap().Num());
-
-		for (const TPair<UPCGDataBinding::FSourceBufferAndAttributeIndex, int32>& Entry : Binding->GetAttributeIdRemap())
+		// Now compute the mapping for any metadata attributes.
+		for (int32 GraphAttributeIndex = 0; GraphAttributeIndex < NumGraphAttributes; ++GraphAttributeIndex)
 		{
-			const TSharedPtr<const FPCGProxyForGPUDataCollection>& SourceBuffer = Entry.Key.Key;
-
-			if (SourceBuffer && SourceBuffer->GetBuffer() == ExternalBufferForReuse)
-			{
-				const int32 GraphAttributeId = PCGComputeHelpers::GetAttributeIdFromMetadataAttributeIndex(Entry.Value);
-
-				const int32 SourceBufferAttributeId = PCGComputeHelpers::GetAttributeIdFromMetadataAttributeIndex(Entry.Key.Value);
+			const int32* SourceBufferAttributeIndex = Binding->GetGraphAttributeToSourceBufferAttributeIndex().Find({ GraphAttributeIndex, ExternalBufferForReuse });
 
-				ensure(GraphAttributeId >= PCGComputeConstants::NUM_RESERVED_ATTRS);
-				ensure(SourceBufferAttributeId >= PCGComputeConstants::NUM_RESERVED_ATTRS);
-
-				SourceToTargetAttributeId.Add(FIntVector2(SourceBufferAttributeId, GraphAttributeId));
-			}
+			GraphToSourceBufferAttributeIndex[GraphAttributeIndex] = SourceBufferAttributeIndex ? *SourceBufferAttributeIndex : GraphAttributeIndex;
 		}
 	}
 	else
@@ -180,25 +162,29 @@ bool FPCGDataUploadAdaptor::PrepareData_GameThread()
 	return true;
 }
 
-FRDGBufferSRVRef FPCGDataUploadAdaptor::GetAttributeRemapBufferSRV(FRDGBuilder& InGraphBuilder)
+FRDGBufferSRVRef FPCGDataUploadAdaptor::GetAttributeRemapBufferSRV(FRDGBuilder& InGraphBuilder, int32& OutFirstRemappedAttributeId)
 {
-	const FRDGBufferDesc Desc = FRDGBufferDesc::CreateStructuredDesc(sizeof(FIntVector2), FMath::Max(SourceToTargetAttributeId.Num(), 1));
+	// Should only be called if UsesAttributeIdRemap() returns true and ExternalBufferForReuse is set.
+	check(ExternalBufferForReuse);
+
+	// Mapping created for all non-system attributes (system attributes always have fixed IDs).
+	OutFirstRemappedAttributeId = PCGComputeConstants::NUM_RESERVED_ATTRS;
+
+	FRDGBufferRef Buffer = nullptr;
 
-	if (SourceToTargetAttributeId.IsEmpty())
+	if (!GraphToSourceBufferAttributeIndex.IsEmpty())
 	{
-		FRDGBufferRef NullBuffer = InGraphBuilder.CreateBuffer(Desc, TEXT("EmptyRemap"));
-		InGraphBuilder.QueueBufferUpload(NullBuffer, &FUintVector2::ZeroValue, sizeof(FUintVector2::ZeroValue));
+		const FRDGBufferDesc BufferDesc = FRDGBufferDesc::CreateStructuredDesc(GraphToSourceBufferAttributeIndex.GetTypeSize(), GraphToSourceBufferAttributeIndex.Num());
+		Buffer = InGraphBuilder.CreateBuffer(BufferDesc, TEXT("PCGDataUploadAdaptor_GraphToBufferAttributeId"));
 
-		return InGraphBuilder.CreateSRV(NullBuffer);
+		InGraphBuilder.QueueBufferUpload(Buffer, MakeArrayView(GraphToSourceBufferAttributeIndex));
 	}
 	else
 	{
-		FRDGBufferRef Buffer = InGraphBuilder.CreateBuffer(Desc, TEXT("PCGDataCollectionUpload_AttributeRemap"));
-
-		InGraphBuilder.QueueBufferUpload(Buffer, MakeArrayView(SourceToTargetAttributeId));
-
-		return InGraphBuilder.CreateSRV(Buffer);
+		Buffer = GSystemTextures.GetDefaultStructuredBuffer(InGraphBuilder, sizeof(uint32));
 	}
+
+	return InGraphBuilder.CreateSRV(Buffer);
 }
 
 /** Create buffer with the element counts of each data. Stored in a buffer because we do not constrain the max data count. */
@@ -255,48 +241,10 @@ FRDGBufferRef FPCGDataUploadAdaptor::GetBuffer_RenderThread(FRDGBuilder& InGraph
 
 		return Buffer;
 	}
-
-	ensure(PackedDataCollection.IsEmpty());
-
-	// Copy the external input buffer to a transient buffer, then run a CS to adapt it:
-	// - Update data counts
-	// - Remap attribute IDs
-
-#if !UE_BUILD_SHIPPING
-	RenderCaptureInterface::FScopedCapture RenderCapture(PCGDataUploadAdaptor::TriggerGPUCaptureDispatches > 0, InGraphBuilder, TEXT("FPCGDataUploadAdaptor"));
-	PCGDataUploadAdaptor::TriggerGPUCaptureDispatches = FMath::Max(PCGDataUploadAdaptor::TriggerGPUCaptureDispatches - 1, 0);
-#endif // !UE_BUILD_SHIPPING
-
-	FRDGBufferRef InputBuffer = InGraphBuilder.RegisterExternalBuffer(ExternalBufferForReuse);
-	const uint32 InputBufferSizeBytes = ExternalBufferSizeBytes;
-	const uint32 InputBufferSizeUints = InputBufferSizeBytes / sizeof(uint32);
-
-	const uint32 OutputSizeBytes = InputBufferSizeBytes;
-	FRDGBufferDesc OutputDesc = FRDGBufferDesc::CreateByteAddressDesc(OutputSizeBytes);
-	if (InExportMode != EPCGExportMode::NoExport)
+	else
 	{
-		// We don't know for sure whether buffer will be read back or not, so need to flag the possibility if the buffer will be passed downstream.
-		OutputDesc.Usage |= BUF_SourceCopy;
-	}
-
-	FRDGBufferRef OutputBuffer = InGraphBuilder.CreateBuffer(OutputDesc, TEXT("PCG_UDC_AttributesRemapped"));
-
-	// Populate using copy. We can't do everything in one pass trivially due to race conditions between groups.
-	AddCopyBufferPass(InGraphBuilder, OutputBuffer, InputBuffer);
-
-	FPCGDataCollectionAdaptorCS::FParameters* Parameters = InGraphBuilder.AllocParameters<FPCGDataCollectionAdaptorCS::FParameters>();
-	Parameters->InSourceToTargetAttributeId = GetAttributeRemapBufferSRV(InGraphBuilder);
-	Parameters->InDataElementCounts = GetDataElementCountsBufferSRV(InGraphBuilder);
-	Parameters->InDataCollection = InGraphBuilder.CreateSRV(InputBuffer);
-	Parameters->OutDataCollection = InGraphBuilder.CreateUAV(OutputBuffer);
-	Parameters->InNumData = TargetDataCollectionDesc.DataDescs.Num();
-	Parameters->InNumRemappedAttributes = SourceToTargetAttributeId.Num();
+		ensure(PackedDataCollection.IsEmpty());
 
-	const uint32 NumUpdateTasks = Parameters->InNumData + Parameters->InNumRemappedAttributes * Parameters->InNumData;
-	const int GroupCount = FMath::DivideAndRoundUp(NumUpdateTasks, FPCGDataCollectionAdaptorCS::NumThreadsPerGroup);
-
-	TShaderMapRef<FPCGDataCollectionAdaptorCS> Shader(GetGlobalShaderMap(GMaxRHIFeatureLevel));
-	FComputeShaderUtils::AddPass(InGraphBuilder, RDG_EVENT_NAME("PCGDataCollectionAdaptor"), ERDGPassFlags::Compute | ERDGPassFlags::NeverCull /*TODO*/, Shader, Parameters, FIntVector(GroupCount, 1, 1));
-
-	return OutputBuffer;
+		return InGraphBuilder.RegisterExternalBuffer(ExternalBufferForReuse);
+	}
 }
diff --git a/Engine/Plugins/PCG/Source/PCG/Private/Compute/PCGDataUploadAdaptor.h b/Engine/Plugins/PCG/Source/PCG/Private/Compute/PCGDataUploadAdaptor.h
--- a/Engine/Plugins/PCG/Source/PCG/Private/Compute/PCGDataUploadAdaptor.h
+++ b/Engine/Plugins/PCG/Source/PCG/Private/Compute/PCGDataUploadAdaptor.h
@@ -20,7 +20,7 @@ public:
 	/** Do any preparation work such as data readbacks. Returns true when preparation is complete. */
 	bool PrepareData_GameThread();
 
-	FRDGBufferSRVRef GetAttributeRemapBufferSRV(FRDGBuilder& InGraphBuilder);
+	FRDGBufferSRVRef GetAttributeRemapBufferSRV(FRDGBuilder& InGraphBuilder, int32& OutFirstRemappedAttributeId);
 
 	/** Create buffer with the element counts of each data. Stored in a buffer because we do not constrain the max data count. */
 	FRDGBufferSRVRef GetDataElementCountsBufferSRV(FRDGBuilder& InGraphBuilder);
@@ -28,6 +28,8 @@ public:
 	/** Gets the buffer that can then be used to read the data in kernels. */
 	FRDGBufferRef GetBuffer_RenderThread(FRDGBuilder& InGraphBuilder, EPCGExportMode InExportMode);
 
+	bool UsesAttributeIdRemap() const { return !!ExternalBufferForReuse; }
+
 	bool IsUploadingFromCPU() const { return !ExternalBufferForReuse; }
 
 private:
@@ -40,7 +42,8 @@ private:
 	/** If data cannot be reused from GPU, will be packed into this buffer prior to upload. */
 	TArray<uint32> PackedDataCollection;
 
-	TArray<FIntVector2> SourceToTargetAttributeId;
+	/** Map from (this) compute graph attribute ID to the attribute ID used in the upstream compute graph that produced this buffer. */
+	TArray<int32> GraphToSourceBufferAttributeIndex;
 
 	FName InputPinLabel;
 
diff --git a/Engine/Plugins/PCG/Source/PCGCompute/Internal/PCGDataCollectionAdaptorCS.h b/Engine/Plugins/PCG/Source/PCGCompute/Internal/PCGDataCollectionAdaptorCS.h
deleted file mode 100644
diff --git a/Engine/Plugins/PCG/Source/PCGCompute/Private/PCGDataCollectionAdaptorCS.cpp b/Engine/Plugins/PCG/Source/PCGCompute/Private/PCGDataCollectionAdaptorCS.cpp
deleted file mode 100644
