[PCG] Tweaks to code that triggers GPU captures.

#rnx

diff --git a/Engine/Plugins/PCG/Source/PCG/Private/Compute/PCGComputeGraph.h b/Engine/Plugins/PCG/Source/PCG/Private/Compute/PCGComputeGraph.h
--- a/Engine/Plugins/PCG/Source/PCG/Private/Compute/PCGComputeGraph.h
+++ b/Engine/Plugins/PCG/Source/PCG/Private/Compute/PCGComputeGraph.h
@@ -119,7 +119,7 @@ public:
 	bool bLogDataDescriptions = false;
 
 #if WITH_EDITORONLY_DATA
-	UPROPERTY()
+	UPROPERTY(Transient)
 	bool bBreakDebugger = false;
 #endif
 
diff --git a/Engine/Plugins/PCG/Source/PCG/Public/PCGSettings.h b/Engine/Plugins/PCG/Source/PCG/Public/PCGSettings.h
--- a/Engine/Plugins/PCG/Source/PCG/Public/PCGSettings.h
+++ b/Engine/Plugins/PCG/Source/PCG/Public/PCGSettings.h
@@ -165,7 +165,7 @@ public:
 	UPROPERTY(EditAnywhere, BlueprintReadWrite, Category = Debug, meta = (ShowOnlyInnerProperties))
 	FPCGDebugVisualizationSettings DebugSettings;
 
-	/** If a debugger like Visual Studio is attached, a break point will be triggered in the execution code to enable debugging. Only applies when inspecting a debug object. Transient and Editor-only. */
+	/** If a debugger is attached, a breakpoint will be triggered in the execution code to enable debugging. Only applies when inspecting a debug object. Transient and Editor-only. */
 	UPROPERTY(Transient, DuplicateTransient, EditAnywhere, BlueprintReadWrite, Category = Debug, AdvancedDisplay, meta = (DisplayName = "Break In Debugger", EditCondition = bDisplayDebuggingProperties, EditConditionHides, HideEditConditionToggle))
 	bool bBreakDebugger = false;
 
