[PCG] Cleanup / refactor of loading deadlock fix.

#rnx

diff --git a/Engine/Plugins/PCG/Source/PCG/Private/Compute/PCGComputeKernel.cpp b/Engine/Plugins/PCG/Source/PCG/Private/Compute/PCGComputeKernel.cpp
--- a/Engine/Plugins/PCG/Source/PCG/Private/Compute/PCGComputeKernel.cpp
+++ b/Engine/Plugins/PCG/Source/PCG/Private/Compute/PCGComputeKernel.cpp
@@ -36,6 +36,11 @@ namespace PCGComputeKernel
 		// @todo_pcg: In the future quaternion compute source could be opt-in if the kernel does not manipulate point/attribute data.
 		FSoftObjectPath(TEXT("/Script/PCG.PCGComputeSource'/PCG/ComputeSources/PCGCS_Quaternion.PCGCS_Quaternion'"))
 	};
+
+	TConstArrayView<FSoftObjectPath> GetDefaultAdditionalSourcePaths()
+	{
+		return MakeArrayView(DefaultAdditionalSourcePaths);
+	}
 }
 
 #if WITH_EDITOR
diff --git a/Engine/Plugins/PCG/Source/PCG/Private/Compute/PCGComputeKernel.h b/Engine/Plugins/PCG/Source/PCG/Private/Compute/PCGComputeKernel.h
--- a/Engine/Plugins/PCG/Source/PCG/Private/Compute/PCGComputeKernel.h
+++ b/Engine/Plugins/PCG/Source/PCG/Private/Compute/PCGComputeKernel.h
@@ -190,3 +190,8 @@ protected:
 
 	mutable TObjectPtr<const UPCGSettings> ResolvedSettings = nullptr;
 };
+
+namespace PCGComputeKernel
+{
+	PCG_API TConstArrayView<FSoftObjectPath> GetDefaultAdditionalSourcePaths();
+}
diff --git a/Engine/Plugins/PCG/Source/PCG/Private/PCGModule.cpp b/Engine/Plugins/PCG/Source/PCG/Private/PCGModule.cpp
--- a/Engine/Plugins/PCG/Source/PCG/Private/PCGModule.cpp
+++ b/Engine/Plugins/PCG/Source/PCG/Private/PCGModule.cpp
@@ -4,6 +4,7 @@
 
 #include "PCGContext.h"
 #include "PCGElement.h"
+#include "Compute/PCGComputeKernel.h"
 #include "Data/PCGBasePointData.h"
 #include "Data/PCGPointData.h"
 #include "Data/PCGSplineData.h"
@@ -32,19 +33,6 @@ FPCGGraphChangedDelegate FPCGModule::OnGraphChangedDelegate;
 namespace PCGModule
 {
 	FPCGModule* PCGModulePtr = nullptr;
-
-#if WITH_EDITOR
-	// @todo_pcg: Temporarily copied from PCGComputeKernel.cpp to avoid changing the ABI for 5.6.1 hotfix. Should be removed in 5.7.
-	static const TStaticArray<FSoftObjectPath, 3> DefaultAdditionalSourcePaths =
-	{
-		FSoftObjectPath(TEXT("/Script/PCG.PCGComputeSource'/PCG/ComputeSources/PCGCS_ShaderUtils.PCGCS_ShaderUtils'")),
-		FSoftObjectPath(TEXT("/Script/PCG.PCGComputeSource'/PCG/ComputeSources/PCGCS_ShaderUtilsInternal.PCGCS_ShaderUtilsInternal'")),
-
-		// Note: PCGDataCollectionDataInterface.ush depends on the quaternion helpers, therefore all kernels also depend on the quaternion helpers.
-		// @todo_pcg: In the future quaternion compute source could be opt-in if the kernel does not manipulate point/attribute data.
-		FSoftObjectPath(TEXT("/Script/PCG.PCGComputeSource'/PCG/ComputeSources/PCGCS_Quaternion.PCGCS_Quaternion'"))
-	};
-#endif
 }
 
 FPCGModule& FPCGModule::GetPCGModuleChecked()
@@ -69,16 +57,7 @@ void FPCGModule::StartupModule()
 	const FString PluginShaderDir = FPaths::Combine(IPluginManager::Get().FindPlugin(TEXT("PCG"))->GetBaseDir(), TEXT("Shaders"));
 	AddShaderSourceDirectoryMapping(TEXT("/Plugin/PCG"), PluginShaderDir);
 
-	// @todo_pcg: In the future, this could be done in the compute kernel CDO constructor instead, to make kernels more self-contained.
-	for (const FSoftObjectPath& AdditionalSourcePath : PCGModule::DefaultAdditionalSourcePaths)
-	{
-		const UObject* LoadedAdditionalSource = AdditionalSourcePath.TryLoad();
-
-		if (!ensure(LoadedAdditionalSource))
-		{
-			UE_LOG(LogPCG, Error, TEXT("Failed to load compute source asset '%s'."), *AdditionalSourcePath.ToString());
-		}
-	}
+	LoadAdditionalKernelSources();
 #endif // WITH_EDITOR
 
 	// Cache for fast access
@@ -146,6 +125,24 @@ bool FPCGModule::Tick(float DeltaTime)
 }
 
 #if WITH_EDITOR
+// todo_pcg: Could move to editor module once PCGComputeKernel.h is publicly exposed.
+void FPCGModule::LoadAdditionalKernelSources()
+{
+	auto LoadSources = [](TConstArrayView<FSoftObjectPath> SourcePaths)
+	{
+		for (const FSoftObjectPath& SourcePath : SourcePaths)
+		{
+			if (!ensure(SourcePath.TryLoad()))
+			{
+				UE_LOG(LogPCG, Error, TEXT("Failed to load compute source asset '%s'."), *SourcePath.ToString());
+			}
+		}
+	};
+
+	// Load all additional sources needed by all kernels in PCG plugin here.
+	LoadSources(PCGComputeKernel::GetDefaultAdditionalSourcePaths());
+}
+
 void FPCGModule::RegisterNativeElementDeterminismTests()
 {
 	PCGDeterminismTests::FNativeTestRegistry::RegisterTestFunction(UPCGDifferenceSettings::StaticClass(), PCGDeterminismTests::DifferenceElement::RunTestSuite);
diff --git a/Engine/Plugins/PCG/Source/PCG/Public/PCGModule.h b/Engine/Plugins/PCG/Source/PCG/Public/PCGModule.h
--- a/Engine/Plugins/PCG/Source/PCG/Public/PCGModule.h
+++ b/Engine/Plugins/PCG/Source/PCG/Public/PCGModule.h
@@ -59,6 +59,11 @@ public:
 private:
 	bool Tick(float DeltaTime);
 	
+#if WITH_EDITOR
+	/** Load additional source assets used by all kernels in the PCG plugin, so that they are ready to be consumed during compilation. */
+	void LoadAdditionalKernelSources();
+#endif
+
 	FPCGGetDataFunctionRegistry GetDataFunctionRegistry;
 	FPCGAttributeAccessorFactory AttributeAccessorFactory;
 
