[PCG] Add preview scene menu to the data viewport
* Allows selecting a preview scene profile and toggling visibility on floor/environment/etc.

diff --git a/Engine/Plugins/PCG/Source/PCGEditor/Private/Widgets/AssetEditorViewport/SPCGEditorViewport.cpp b/Engine/Plugins/PCG/Source/PCGEditor/Private/Widgets/AssetEditorViewport/SPCGEditorViewport.cpp
--- a/Engine/Plugins/PCG/Source/PCGEditor/Private/Widgets/AssetEditorViewport/SPCGEditorViewport.cpp
+++ b/Engine/Plugins/PCG/Source/PCGEditor/Private/Widgets/AssetEditorViewport/SPCGEditorViewport.cpp
@@ -3,9 +3,12 @@
 #include "Widgets/AssetEditorViewport/SPCGEditorViewport.h"
 
 #include "AdvancedPreviewScene.h"
+#include "AdvancedPreviewSceneMenus.h"
 #include "AssetViewerSettings.h"
 #include "EditorViewportClient.h"
+#include "PreviewProfileController.h"
 #include "ToolMenus.h"
+#include "Engine/World.h"
 #include "ViewportToolbar/UnrealEdViewportToolbar.h"
 #include "ViewportToolbar/UnrealEdViewportToolbarContext.h"
 
@@ -31,17 +34,18 @@ public:
 	virtual void Tick(float DeltaSeconds) override;
 	//~ End FEditorViewportClient Interface
 
-	FAdvancedPreviewScene* GetPreviewScene() { return Scene.Get(); }
-	FAdvancedPreviewScene* ResetScene();
-
-private:
-	TSharedPtr<FAdvancedPreviewScene> Scene = nullptr;
+	void Reset();
 };
 
 FPCGEditorViewportClient::FPCGEditorViewportClient(const TSharedRef<SPCGEditorViewport>& InAssetEditorViewport)
-	: FEditorViewportClient(nullptr, nullptr, StaticCastSharedRef<SEditorViewport>(InAssetEditorViewport))
+	: FEditorViewportClient(/*InModeTools=*/nullptr, InAssetEditorViewport->GetAdvancedPreviewScene(), StaticCastSharedRef<SEditorViewport>(InAssetEditorViewport))
 {
-	ResetScene();
+	check(PreviewScene);
+
+	if (UWorld* World = PreviewScene->GetWorld())
+	{
+		World->ChangeFeatureLevel(GWorld->GetFeatureLevel());
+	}
 
 	bUsesDrawHelper = true;
 
@@ -58,34 +62,32 @@ FPCGEditorViewportClient::FPCGEditorViewportClient(const TSharedRef<SPCGEditorVi
 	EngineShowFlags.DisableAdvancedFeatures();
 	EngineShowFlags.SetSnap(0);
 	EngineShowFlags.SetSeparateTranslucency(true);
+	
+	Reset();
 }
 
 bool FPCGEditorViewportClient::InputKey(const FInputKeyEventArgs& Args)
 {
-	if (!Scene.IsValid())
-	{
-		return true;
-	}
+	FAdvancedPreviewScene* AdvancedPreviewScene = static_cast<FAdvancedPreviewScene*>(PreviewScene);
+	check(AdvancedPreviewScene);
 
 	bool bHandled = FEditorViewportClient::InputKey(Args);
 	bHandled |= InputTakeScreenshot(Args.Viewport, Args.Key, Args.Event);
-	bHandled |= Scene->HandleInputKey(Args);
+	bHandled |= AdvancedPreviewScene->HandleInputKey(Args);
 
 	return bHandled;
 }
 
 bool FPCGEditorViewportClient::InputAxis(const FInputKeyEventArgs& Args)
 {
-	if (!Scene.IsValid())
-	{
-		return true;
-	}
+	FAdvancedPreviewScene* AdvancedPreviewScene = static_cast<FAdvancedPreviewScene*>(PreviewScene);
+	check(AdvancedPreviewScene);
 
 	bool bResult = true;
 
 	if (!bDisableInput)
 	{
-		bResult = Scene->HandleViewportInput(Args.Viewport, Args.InputDevice, Args.Key, Args.AmountDepressed, Args.DeltaTime, Args.NumSamples, Args.IsGamepad());
+		bResult = AdvancedPreviewScene->HandleViewportInput(Args.Viewport, Args.InputDevice, Args.Key, Args.AmountDepressed, Args.DeltaTime, Args.NumSamples, Args.IsGamepad());
 
 		if (bResult)
 		{
@@ -102,45 +104,34 @@ bool FPCGEditorViewportClient::InputAxis(const FInputKeyEventArgs& Args)
 
 FLinearColor FPCGEditorViewportClient::GetBackgroundColor() const
 {
-	return Scene.IsValid() ? Scene->GetBackgroundColor() : FColor(64, 64, 64);
+	return ensure(PreviewScene) ? PreviewScene->GetBackgroundColor() : FColor(64, 64, 64);
 }
 
 void FPCGEditorViewportClient::Tick(float DeltaSeconds)
 {
 	FEditorViewportClient::Tick(DeltaSeconds);
 
-	if (!GIntraFrameDebuggingGameThread && Scene.IsValid() && Scene->GetWorld())
+	if (!GIntraFrameDebuggingGameThread && ensure(PreviewScene && PreviewScene->GetWorld()))
 	{
-		Scene->GetWorld()->Tick(LEVELTICK_All, DeltaSeconds);
+		PreviewScene->GetWorld()->Tick(LEVELTICK_All, DeltaSeconds);
 	}
 }
 
-FAdvancedPreviewScene* FPCGEditorViewportClient::ResetScene()
+void FPCGEditorViewportClient::Reset()
 {
-	TRACE_CPUPROFILER_EVENT_SCOPE(FPCGEditorViewportClient::ResetScene);
-
-	Scene = MakeShareable(new FAdvancedPreviewScene(FPreviewScene::ConstructionValues()));
-	PreviewScene = Scene.Get();
-
-	// Restore last used feature level.
-	if (UWorld* World = Scene->GetWorld())
-	{
-		World->ChangeFeatureLevel(GWorld->GetFeatureLevel());
-	}
+	TRACE_CPUPROFILER_EVENT_SCOPE(FPCGEditorViewportClient::Reset);
 
 	SetViewportType(ELevelViewportType::LVT_Perspective);
-	SetViewMode(EViewModeIndex::VMI_Lit);
 	SetViewLocation(PCGEditorViewportConstants::DefaultViewLocation);
 	SetViewRotation(PCGEditorViewportConstants::DefaultViewRotation);
 	SetViewLocationForOrbiting(PCGEditorViewportConstants::DefaultViewLocation, PCGEditorViewportConstants::DefaultOrbitDistance);
 
-	Scene->SetFloorOffset(0.0f);
-	Scene->SetFloorVisibility(true);
-	Scene->SetEnvironmentVisibility(true);
-
 	Invalidate();
+}
 
-	return Scene.Get();
+SPCGEditorViewport::SPCGEditorViewport()
+	: AdvancedPreviewScene(MakeShareable(new FAdvancedPreviewScene(FPreviewScene::ConstructionValues())))
+{
 }
 
 SPCGEditorViewport::~SPCGEditorViewport()
@@ -156,6 +147,8 @@ SPCGEditorViewport::~SPCGEditorViewport()
 void SPCGEditorViewport::Construct(const FArguments& InArgs)
 {
 	SEditorViewport::Construct(SEditorViewport::FArguments());
+
+	UE::AdvancedPreviewScene::BindDefaultOnSettingsChangedHandler(AdvancedPreviewScene, EditorViewportClient);
 }
 
 TSharedRef<class SEditorViewport> SPCGEditorViewport::GetViewportWidget()
@@ -172,21 +165,17 @@ void SPCGEditorViewport::SetupScene(const TArray<UObject*>& InResources, const F
 {
 	TRACE_CPUPROFILER_EVENT_SCOPE(SPCGEditorViewport::SetupScene);
 
-	ResetScene();
+	check(EditorViewportClient.IsValid());
+	check(AdvancedPreviewScene.IsValid());
 
-	if (!ensure(EditorViewportClient.IsValid()))
-	{
-		return;
-	}
-
-	FAdvancedPreviewScene* Scene = EditorViewportClient->GetPreviewScene();
+	ResetScene();
 
-	if (Scene && SetupFunc.IsSet())
+	if (SetupFunc.IsSet())
 	{
 		TRACE_CPUPROFILER_EVENT_SCOPE(SPCGEditorViewport::SetupSceneCallback);
 
 		FPCGSceneSetupParams SceneSetupParams;
-		SceneSetupParams.Scene = Scene;
+		SceneSetupParams.Scene = AdvancedPreviewScene.Get();
 		SceneSetupParams.EditorViewportClient = EditorViewportClient.Get();
 		SceneSetupParams.Resources = InResources;
 
@@ -200,14 +189,25 @@ void SPCGEditorViewport::SetupScene(const TArray<UObject*>& InResources, const F
 
 void SPCGEditorViewport::ResetScene()
 {
-	ReleaseManagedResources();
-	
-	if (!ensure(EditorViewportClient.IsValid()))
+	check(EditorViewportClient.IsValid());
+	check(AdvancedPreviewScene.IsValid());
+
+	EditorViewportClient->Reset();
+
+	if (FPreviewSceneProfile* Profile = AdvancedPreviewScene->GetCurrentProfile())
 	{
-		return;
+		Profile->SetShowFlags(EditorViewportClient->EngineShowFlags);
 	}
 
-	EditorViewportClient->ResetScene();
+	for (TObjectPtr<UObject> ManagedResource : ManagedResources)
+	{
+		if (UActorComponent* ActorComponent = Cast<UActorComponent>(ManagedResource))
+		{
+			AdvancedPreviewScene->RemoveComponent(ActorComponent);
+		}
+	}
+
+	ReleaseManagedResources();
 }
 
 TSharedRef<FEditorViewportClient> SPCGEditorViewport::MakeEditorViewportClient()
@@ -229,6 +229,8 @@ TSharedRef<FEditorViewportClient> SPCGEditorViewport::MakeEditorViewportClient()
 
 TSharedPtr<SWidget> SPCGEditorViewport::BuildViewportToolbar()
 {
+	check(AdvancedPreviewScene.IsValid());
+
 	const FName ToolbarName = "PCG.ViewportToolbar";
 	
 	if (!UToolMenus::Get()->IsMenuRegistered(ToolbarName))
@@ -242,16 +244,26 @@ TSharedPtr<SWidget> SPCGEditorViewport::BuildViewportToolbar()
 		RightSection.AddEntry(UE::UnrealEd::CreateCameraSubmenu(UE::UnrealEd::FViewportCameraMenuOptions().ShowAll()));
 		RightSection.AddEntry(UE::UnrealEd::CreateViewModesSubmenu());
 		RightSection.AddEntry(UE::UnrealEd::CreateDefaultShowSubmenu());
+		RightSection.AddEntry(UE::UnrealEd::CreateAssetViewerProfileSubmenu());
+		UE::AdvancedPreviewScene::Menus::ExtendAdvancedPreviewSceneSettings("PCG.ViewportToolbar.AssetViewerProfile");
 	}
-	
+
 	FToolMenuContext Context;
-	Context.AppendCommandList(GetCommandList());
-	Context.AddExtender(GetExtenders());
-	Context.AddObject(UE::UnrealEd::CreateViewportToolbarDefaultContext(GetViewportWidget()));
+	{
+		Context.AppendCommandList(AdvancedPreviewScene->GetCommandList());
+		Context.AppendCommandList(GetCommandList());
+		Context.AddExtender(GetExtenders());
+		Context.AddObject(UE::UnrealEd::CreateViewportToolbarDefaultContext(GetViewportWidget()));
+	}
 	
 	return UToolMenus::Get()->GenerateWidget(ToolbarName, Context);
 }
 
+TSharedPtr<IPreviewProfileController> SPCGEditorViewport::CreatePreviewProfileController()
+{
+	return MakeShared<FPreviewProfileController>();
+}
+
 void SPCGEditorViewport::ReleaseManagedResources()
 {
 	ManagedResources.Empty();
diff --git a/Engine/Plugins/PCG/Source/PCGEditor/Private/Widgets/AssetEditorViewport/SPCGEditorViewport.h b/Engine/Plugins/PCG/Source/PCGEditor/Private/Widgets/AssetEditorViewport/SPCGEditorViewport.h
--- a/Engine/Plugins/PCG/Source/PCGEditor/Private/Widgets/AssetEditorViewport/SPCGEditorViewport.h
+++ b/Engine/Plugins/PCG/Source/PCGEditor/Private/Widgets/AssetEditorViewport/SPCGEditorViewport.h
@@ -7,6 +7,7 @@
 #include "SAssetEditorViewport.h"
 #include "SCommonEditorViewportToolbarBase.h"
 
+class FAdvancedPreviewScene;
 class FPCGEditorViewportClient;
 
 class SPCGEditorViewport : public SAssetEditorViewport, public ICommonEditorViewportToolbarInfoProvider
@@ -15,6 +16,7 @@ public:
 	SLATE_BEGIN_ARGS(SPCGEditorViewport) {}
 	SLATE_END_ARGS()
 
+	SPCGEditorViewport();
 	virtual ~SPCGEditorViewport();
 
 	void Construct(const FArguments& InArgs);
@@ -28,19 +30,23 @@ public:
 	/** Rebuilds the scene using the provided resources and setup callback. */
 	void SetupScene(const TArray<UObject*>& InResources, const FPCGSetupSceneFunc& SetupFunc);
 	void ResetScene();
+	FAdvancedPreviewScene* GetAdvancedPreviewScene() { return AdvancedPreviewScene.Get(); }
 
-	/** Add UObject references for GC */
+	/** Add UObject references for GC. */
 	void AddReferencedObjects(FReferenceCollector& Collector);
+
 protected:
-	//~ Begin SAssetEditorViewport Interface
+	//~ Begin SEditorViewport Interface
 	virtual TSharedRef<FEditorViewportClient> MakeEditorViewportClient() override;
 	virtual TSharedPtr<SWidget> BuildViewportToolbar() override;
-	//~ End SAssetEditorViewport Interface
+	virtual TSharedPtr<IPreviewProfileController> CreatePreviewProfileController() override;
+	//~ End SEditorViewport Interface
 
 	void ReleaseManagedResources();
 
 private:
-	TSharedPtr<class FPCGEditorViewportClient> EditorViewportClient = nullptr;
+	TSharedPtr<FPCGEditorViewportClient> EditorViewportClient = nullptr;
+	TSharedPtr<FAdvancedPreviewScene> AdvancedPreviewScene = nullptr;
 
 	/** Objects used by the scene. Must be managed for GC. */
 	TArray<TObjectPtr<UObject>> ManagedResources;
