[ComputeFramework] Add AbortWork function to compute graph instance.
[PCG] Fix crash on shutdown when compute graph work is executing by adding call to AbortWork.

#rnx

diff --git a/Engine/Plugins/PCG/Source/PCG/Private/Compute/Elements/PCGComputeGraphElement.cpp b/Engine/Plugins/PCG/Source/PCG/Private/Compute/Elements/PCGComputeGraphElement.cpp
--- a/Engine/Plugins/PCG/Source/PCG/Private/Compute/Elements/PCGComputeGraphElement.cpp
+++ b/Engine/Plugins/PCG/Source/PCG/Private/Compute/Elements/PCGComputeGraphElement.cpp
@@ -497,7 +497,9 @@ bool FPCGComputeGraphElement::ExecuteInternal(FPCGContext* InContext) const
 				// Wake up so that we can terminate execution of the ComputeGraphElement.
 				Context->bIsPaused = false;
 				Context->bGraphSubmitFailed = true;
-			}));
+			}),
+			/*InOwnerPointer=*/Context->DataBinding // Use data binding as owner UObject as binding is created per-execution, and owner is needed if we abort the work.
+		);
 
 		if (!bGraphEnqueued)
 		{
@@ -650,6 +652,15 @@ void FPCGComputeGraphElement::PostExecuteInternal(FPCGContext* InContext) const
 
 void FPCGComputeGraphElement::AbortInternal(FPCGContext* InContext) const
 {
+	check(InContext);
+	FPCGComputeGraphContext* Context = static_cast<FPCGComputeGraphContext*>(InContext);
+
+	const IPCGGraphExecutionSource* Source = InContext->ExecutionSource.Get();
+	if (Context->ComputeGraphInstance && Context->DataBinding && Source && Source->GetExecutionState().GetWorld())
+	{
+		ComputeFramework::AbortWork(Source->GetExecutionState().GetWorld()->Scene, Context->DataBinding);
+	}
+
 	ResetAsyncOperations(InContext);
 }
 
