[PCG] Disable floor by default in data viewport

#rnx

diff --git a/Engine/Plugins/PCG/Source/PCGEditor/Private/Widgets/AssetEditorViewport/SPCGEditorViewport.cpp b/Engine/Plugins/PCG/Source/PCGEditor/Private/Widgets/AssetEditorViewport/SPCGEditorViewport.cpp
--- a/Engine/Plugins/PCG/Source/PCGEditor/Private/Widgets/AssetEditorViewport/SPCGEditorViewport.cpp
+++ b/Engine/Plugins/PCG/Source/PCGEditor/Private/Widgets/AssetEditorViewport/SPCGEditorViewport.cpp
@@ -62,6 +62,11 @@ FPCGEditorViewportClient::FPCGEditorViewportClient(const TSharedRef<SPCGEditorVi
 	EngineShowFlags.DisableAdvancedFeatures();
 	EngineShowFlags.SetSnap(0);
 	EngineShowFlags.SetSeparateTranslucency(true);
+
+	FAdvancedPreviewScene* AdvancedPreviewScene = static_cast<FAdvancedPreviewScene*>(PreviewScene);
+	check(AdvancedPreviewScene);
+
+	AdvancedPreviewScene->SetFloorVisibility(false);
 	
 	Reset();
 }
