[PCG] Fix crash from bad assumption that context is always not-null.

#rnx

diff --git a/Engine/Plugins/PCG/Source/PCG/Private/Compute/Elements/PCGComputeGraphElement.cpp b/Engine/Plugins/PCG/Source/PCG/Private/Compute/Elements/PCGComputeGraphElement.cpp
--- a/Engine/Plugins/PCG/Source/PCG/Private/Compute/Elements/PCGComputeGraphElement.cpp
+++ b/Engine/Plugins/PCG/Source/PCG/Private/Compute/Elements/PCGComputeGraphElement.cpp
@@ -652,13 +652,13 @@ void FPCGComputeGraphElement::PostExecuteInternal(FPCGContext* InContext) const
 
 void FPCGComputeGraphElement::AbortInternal(FPCGContext* InContext) const
 {
-	check(InContext);
-	FPCGComputeGraphContext* Context = static_cast<FPCGComputeGraphContext*>(InContext);
-
-	const IPCGGraphExecutionSource* Source = InContext->ExecutionSource.Get();
-	if (Context->ComputeGraphInstance && Context->DataBinding && Source && Source->GetExecutionState().GetWorld())
+	if (FPCGComputeGraphContext* Context = static_cast<FPCGComputeGraphContext*>(InContext))
 	{
-		ComputeFramework::AbortWork(Source->GetExecutionState().GetWorld()->Scene, Context->DataBinding);
+		const IPCGGraphExecutionSource* Source = InContext->ExecutionSource.Get();
+		if (Context->ComputeGraphInstance && Context->DataBinding && Source && Source->GetExecutionState().GetWorld())
+		{
+			ComputeFramework::AbortWork(Source->GetExecutionState().GetWorld()->Scene, Context->DataBinding);
+		}
 	}
 
 	ResetAsyncOperations(InContext);
