[PCG] Add debug feature that repeatedly dispatches compute graph to enable profiling, compiled out by default (PCG_GPU_KERNEL_PROFILING).

[RN] Added debug feature that repeatedly dispatches compute graph to enable profiling, compiled out by default (PCG_GPU_KERNEL_PROFILING).

diff --git a/Engine/Plugins/PCG/Source/PCG/PCG.Build.cs b/Engine/Plugins/PCG/Source/PCG/PCG.Build.cs
--- a/Engine/Plugins/PCG/Source/PCG/PCG.Build.cs
+++ b/Engine/Plugins/PCG/Source/PCG/PCG.Build.cs
@@ -45,6 +45,8 @@ namespace UnrealBuildTool.Rules
 					});
 			}
 
+			//PublicDefinitions.Add("PCG_DATA_USAGE_LOGGING"); // Generates log to help debug lifetimes of transient resources.
+			//PublicDefinitions.Add("PCG_GPU_KERNEL_PROFILING"); // Enables repeating kernel dispatches every frame to help profiling.
 		}
 	}
 }
diff --git a/Engine/Plugins/PCG/Source/PCG/Private/Compute/Elements/PCGComputeGraphElement.cpp b/Engine/Plugins/PCG/Source/PCG/Private/Compute/Elements/PCGComputeGraphElement.cpp
--- a/Engine/Plugins/PCG/Source/PCG/Private/Compute/Elements/PCGComputeGraphElement.cpp
+++ b/Engine/Plugins/PCG/Source/PCG/Private/Compute/Elements/PCGComputeGraphElement.cpp
@@ -508,7 +508,16 @@ bool FPCGComputeGraphElement::ExecuteInternal(FPCGContext* InContext) const
 			return true;
 		}
 
-		Context->ExecutionSubPhase = EPCGComputeGraphExecutionPhase::WaitForExecutionComplete;
+#ifdef PCG_GPU_KERNEL_PROFILING
+		if (Context->ComputeGraph->ShouldRepeatDispatch())
+		{
+			UE_LOG(LogPCG, Log, TEXT("Repeating dispatch (graph execution will not complete)."));
+		}
+		else
+#endif // PCG_GPU_KERNEL_PROFILING
+		{
+			Context->ExecutionSubPhase = EPCGComputeGraphExecutionPhase::WaitForExecutionComplete;
+		}
 
 		SleepUntilNextFrame();
 
diff --git a/Engine/Plugins/PCG/Source/PCG/Private/Compute/PCGComputeGraph.cpp b/Engine/Plugins/PCG/Source/PCG/Private/Compute/PCGComputeGraph.cpp
--- a/Engine/Plugins/PCG/Source/PCG/Private/Compute/PCGComputeGraph.cpp
+++ b/Engine/Plugins/PCG/Source/PCG/Private/Compute/PCGComputeGraph.cpp
@@ -319,6 +319,22 @@ void UPCGComputeGraph::DebugLogDataDescriptions(const UPCGDataBinding* InBinding
 #endif
 }
 
+#ifdef PCG_GPU_KERNEL_PROFILING
+bool UPCGComputeGraph::ShouldRepeatDispatch() const
+{
+	for (const UComputeKernel* Kernel : KernelInvocations)
+	{
+		const UPCGComputeKernel* PCGKernel = Cast<UPCGComputeKernel>(Kernel);
+		if (PCGKernel && PCGKernel->ShouldRepeatDispatch())
+		{
+			return true;
+		}
+	}
+
+	return false;
+}
+#endif // PCG_GPU_KERNEL_PROFILING
+
 void UPCGComputeGraph::ComputeExternalPinDesc(FName InVirtualLabel, UPCGDataBinding* InBinding, FPCGDataCollectionDesc& OutDataDesc) const
 {
 	check(InBinding);
diff --git a/Engine/Plugins/PCG/Source/PCG/Private/Compute/PCGComputeGraph.h b/Engine/Plugins/PCG/Source/PCG/Private/Compute/PCGComputeGraph.h
--- a/Engine/Plugins/PCG/Source/PCG/Private/Compute/PCGComputeGraph.h
+++ b/Engine/Plugins/PCG/Source/PCG/Private/Compute/PCGComputeGraph.h
@@ -97,6 +97,11 @@ public:
 	/** Logs detailed data descriptions for all the kernels in the compute graph. */
 	void DebugLogDataDescriptions(const UPCGDataBinding* InBinding) const;
 
+#ifdef PCG_GPU_KERNEL_PROFILING
+	/** Whether to repeatedly dispatch every frame. Helpful for profiling. */
+	bool ShouldRepeatDispatch() const;
+#endif
+
 private:
 	/** Compute the data description on a pin external to the compute graph. */
 	void ComputeExternalPinDesc(FName InVirtualLabel, UPCGDataBinding* InBinding, FPCGDataCollectionDesc& OutDataDesc) const;
diff --git a/Engine/Plugins/PCG/Source/PCG/Private/Compute/PCGComputeKernel.cpp b/Engine/Plugins/PCG/Source/PCG/Private/Compute/PCGComputeKernel.cpp
--- a/Engine/Plugins/PCG/Source/PCG/Private/Compute/PCGComputeKernel.cpp
+++ b/Engine/Plugins/PCG/Source/PCG/Private/Compute/PCGComputeKernel.cpp
@@ -48,12 +48,20 @@ void UPCGComputeKernel::Initialize(const FPCGComputeKernelParams& InParams)
 {
 	ResolvedSettings = InParams.Settings;
 	Settings = ResolvedSettings;
-	bLogDataDescriptions = InParams.bLogDescriptions;
+
+	if (InParams.bLogDescriptions)
+	{
+		PCGKernelFlags |= EPCGComputeKernelFlags::LogDataDescriptions;
+	}
+
 	bInitialized = true;
 
 	InitializeInternal();
 
-	bHasStaticValidationErrors = !PerformStaticValidation();
+	if (!PerformStaticValidation())
+	{
+		PCGKernelFlags |= EPCGComputeKernelFlags::HasStaticValidationErrors;
+	}
 }
 #endif
 
@@ -98,7 +106,7 @@ bool UPCGComputeKernel::AreKernelSettingsValid(FPCGContext* InContext) const
 	}
 #endif
 
-	return !bHasStaticValidationErrors;
+	return !HasStaticValidationErrors();
 }
 
 #if WITH_EDITOR
diff --git a/Engine/Plugins/PCG/Source/PCG/Private/Compute/PCGComputeKernel.h b/Engine/Plugins/PCG/Source/PCG/Private/Compute/PCGComputeKernel.h
--- a/Engine/Plugins/PCG/Source/PCG/Private/Compute/PCGComputeKernel.h
+++ b/Engine/Plugins/PCG/Source/PCG/Private/Compute/PCGComputeKernel.h
@@ -21,6 +21,19 @@ struct FPCGPinPropertiesGPU;
 
 class UComputeDataInterface;
 
+UENUM()
+enum class EPCGComputeKernelFlags : uint32
+{
+	None = 0,
+	/** Errors were detected with node setup at compute graph compilation time. */
+	HasStaticValidationErrors = 1 << 0,
+	/** Dump data descriptions to log. */
+	LogDataDescriptions = 1 << 1,
+	/** Dispatch compute graph containing this kernel every frame to allow profiling. Has no effect unless PCG_GPU_KERNEL_PROFILING defined. */
+	RepeatDispatchForProfiling = 1 << 2,
+};
+ENUM_CLASS_FLAGS(EPCGComputeKernelFlags);
+
 UENUM()
 enum class EPCGKernelLogVerbosity
 {
@@ -143,7 +156,16 @@ public:
 	virtual void GetInputPins(TArray<FPCGPinProperties>& OutPins) const {}
 	virtual void GetOutputPins(TArray<FPCGPinPropertiesGPU>& OutPins) const {}
 
-	virtual bool GetLogDataDescriptions() const { return bLogDataDescriptions; }
+	bool GetLogDataDescriptions() const { return !!(PCGKernelFlags & EPCGComputeKernelFlags::LogDataDescriptions); }
+
+	bool HasStaticValidationErrors() const { return !!(PCGKernelFlags & EPCGComputeKernelFlags::HasStaticValidationErrors); }
+
+#ifdef PCG_GPU_KERNEL_PROFILING
+	void EnableRepeatDispatch() { PCGKernelFlags |= EPCGComputeKernelFlags::RepeatDispatchForProfiling; }
+
+	/** Whether to repeatedly dispatch every frame to help profiling. Either configured on the GPU node or overridden here to avoid data cooks. */
+	virtual bool ShouldRepeatDispatch() const { return !!(PCGKernelFlags & EPCGComputeKernelFlags::RepeatDispatchForProfiling); }
+#endif
 
 	static void AddReferencedObjects(UObject* InThis, FReferenceCollector& Collector);
 
@@ -163,6 +185,9 @@ protected:
 #endif
 
 protected:
+	UPROPERTY()
+	EPCGComputeKernelFlags PCGKernelFlags = EPCGComputeKernelFlags::None;
+
 	/** Index into kernels array on owning compute graph. */
 	UPROPERTY()
 	int32 KernelIndex = INDEX_NONE;
@@ -170,15 +195,9 @@ protected:
 	UPROPERTY()
 	TSoftObjectPtr<const UPCGSettings> Settings;
 
-	UPROPERTY()
-	bool bLogDataDescriptions = false;
-
 	UPROPERTY()
 	bool bInitialized = false;
 
-	UPROPERTY()
-	bool bHasStaticValidationErrors = false;
-
 	// @todo_pcg: These are a bit of a hack to avoid supporting static/compilation logging in PCGNodeVisualLogs.h
 	/** Log entries created and cached on initialization. Are logged during AreKernelSettingsValid(), unless PCG_KERNEL_LOGGING_ENABLED is false. */
 	UPROPERTY()
diff --git a/Engine/Plugins/PCG/Source/PCG/Private/Graph/PCGGraphCompilerGPU.cpp b/Engine/Plugins/PCG/Source/PCG/Private/Graph/PCGGraphCompilerGPU.cpp
--- a/Engine/Plugins/PCG/Source/PCG/Private/Graph/PCGGraphCompilerGPU.cpp
+++ b/Engine/Plugins/PCG/Source/PCG/Private/Graph/PCGGraphCompilerGPU.cpp
@@ -278,6 +278,13 @@ void FPCGGraphCompilerGPU::ExpandGPUNodeKernelsToTasks(
 		NodeKernelEdges.Reset();
 		Settings->CreateKernels(InOutContext, GetTransientPackage(), NodeKernels, NodeKernelEdges);
 
+#ifdef PCG_GPU_KERNEL_PROFILING
+		if (NodeKernels.IsValidIndex(Settings->ProfileKernelIndex))
+		{
+			NodeKernels[Settings->ProfileKernelIndex]->EnableRepeatDispatch();
+		}
+#endif // PCG_GPU_KERNEL_PROFILING
+
 		ensureMsgf(NodeKernels.Remove(nullptr) == 0, TEXT("Settings '%s' returned one or more null kernels."), *Settings->GetName());
 
 		if (NodeKernels.IsEmpty())
diff --git a/Engine/Plugins/PCG/Source/PCG/Public/PCGSettings.h b/Engine/Plugins/PCG/Source/PCG/Public/PCGSettings.h
--- a/Engine/Plugins/PCG/Source/PCG/Public/PCGSettings.h
+++ b/Engine/Plugins/PCG/Source/PCG/Public/PCGSettings.h
@@ -564,6 +564,10 @@ public:
 	/** [EXPERIMENTAL] Note: This variable is subject to change without deprecation. */
 	UPROPERTY(EditAnywhere, Transient, Category = "GPU", AdvancedDisplay, meta = (Tooltip = "Will trigger a render capture when this node executes and a debug object is selected in the graph editor. Transient and Editor-only. Render captures must be enabled (e.g. -AttachRenderDoc or -AttachPIX).", EditCondition = "bExecuteOnGPU", EditConditionHides))
 	bool bTriggerRenderCapture = false;
+
+	/** [EXPERIMENTAL] Note: This variable is subject to change without deprecation. */
+	UPROPERTY(EditAnywhere, Transient, Category = "GPU", AdvancedDisplay, meta = (Tooltip = "Index of kernel emitted by this node to repeatedly dispatch every frame to enable profiling. Set to -1 to disable profiling. PCG_GPU_KERNEL_PROFILING must be defined in PCG.Build.cs.", EditCondition = "bExecuteOnGPU", EditConditionHides))
+	int ProfileKernelIndex = INDEX_NONE;
 #endif
 
 private:
