[PCG] Added String Contains and Matches nodes... and .Length extractor

+ Extractor: .Length - Length of the string
+ Contains - Whether string A contains substring B
+ Matches - Whether string A matches string B exactly

diff --git a/Engine/Plugins/PCG/Source/PCG/Private/Elements/Metadata/PCGMetadataStringOpElement.cpp b/Engine/Plugins/PCG/Source/PCG/Private/Elements/Metadata/PCGMetadataStringOpElement.cpp
--- a/Engine/Plugins/PCG/Source/PCG/Private/Elements/Metadata/PCGMetadataStringOpElement.cpp
+++ b/Engine/Plugins/PCG/Source/PCG/Private/Elements/Metadata/PCGMetadataStringOpElement.cpp
@@ -25,7 +25,9 @@ FName UPCGMetadataStringOpSettings::GetInputPinLabel(uint32 Index) const
 		return PCGPinConstants::DefaultInputLabel;
 	}
 
-	case EPCGMetadataStringOperation::Append:
+	case EPCGMetadataStringOperation::Append:    // fall-through
+	case EPCGMetadataStringOperation::Substring: // fall-through
+	case EPCGMetadataStringOperation::Matches:
 		switch (Index)
 		{
 		case 0: return PCGMetadataSettingsBaseConstants::DoubleInputFirstLabel;
@@ -58,7 +60,9 @@ uint32 UPCGMetadataStringOpSettings::GetOperandNum() const
 		case EPCGMetadataStringOperation::TrimEnd: // fall-through
 		case EPCGMetadataStringOperation::TrimStartAndEnd:
 			return 1;
-		case EPCGMetadataStringOperation::Append:
+		case EPCGMetadataStringOperation::Append:    // fall-through
+		case EPCGMetadataStringOperation::Substring: // fall-through
+		case EPCGMetadataStringOperation::Matches:
 			return 2;
 		case EPCGMetadataStringOperation::Replace:
 			return 3;
@@ -164,6 +168,14 @@ bool FPCGMetadataStringOpElement::DoOperation(PCGMetadataOps::FOperationData& Op
 	{
 		return DoBinaryOp<FString, FString>(OperationData, [](const FString& Value1, const FString& Value2) -> FString { return Value1 + Value2; });
 	}
+	else if (Settings->Operation == EPCGMetadataStringOperation::Substring)
+	{
+		return DoBinaryOp<FString, FString>(OperationData, [Settings](const FString& A, const FString& B) -> bool { return PCG::Private::MetadataTraits<FString>::Substring(A, B, Settings->SearchCase); });
+	}
+	else if (Settings->Operation == EPCGMetadataStringOperation::Matches)
+	{
+		return DoBinaryOp<FString, FString>(OperationData, [Settings](const FString& A, const FString& B) -> bool { return PCG::Private::MetadataTraits<FString>::Matches(A, B, Settings->SearchCase); });
+	}
 	else if (Settings->Operation == EPCGMetadataStringOperation::Replace)
 	{
 		return DoTernaryOp<FString, FString, FString>(OperationData, [](const FString& InValue, const FString& InSearch, const FString& InReplace)
diff --git a/Engine/Plugins/PCG/Source/PCG/Private/Metadata/Accessors/PCGAttributeAccessorHelpers.cpp b/Engine/Plugins/PCG/Source/PCG/Private/Metadata/Accessors/PCGAttributeAccessorHelpers.cpp
--- a/Engine/Plugins/PCG/Source/PCG/Private/Metadata/Accessors/PCGAttributeAccessorHelpers.cpp
+++ b/Engine/Plugins/PCG/Source/PCG/Private/Metadata/Accessors/PCGAttributeAccessorHelpers.cpp
@@ -38,19 +38,23 @@ namespace PCGAttributeAccessorHelpers
 
 			if constexpr (PCG::Private::IsOfTypes<AccessorType, FVector2D, FVector, FVector4>())
 			{
-				return PCGAttributeExtractor::CreateVectorExtractor<AccessorType>(std::move(Accessor), Name, bOutSuccess);
+				return PCGAttributeExtractor::CreateVectorExtractor<AccessorType>(MoveTemp(Accessor), Name, bOutSuccess);
 			}
 			else if constexpr (PCG::Private::IsOfTypes<AccessorType, FTransform>())
 			{
-				return PCGAttributeExtractor::CreateTransformExtractor(std::move(Accessor), Name, bOutSuccess);
+				return PCGAttributeExtractor::CreateTransformExtractor(MoveTemp(Accessor), Name, bOutSuccess);
 			}
 			else if constexpr (PCG::Private::IsOfTypes<AccessorType, FQuat>())
 			{
-				return PCGAttributeExtractor::CreateQuatExtractor(std::move(Accessor), Name, bOutSuccess);
+				return PCGAttributeExtractor::CreateQuatExtractor(MoveTemp(Accessor), Name, bOutSuccess);
 			}
 			else if constexpr (PCG::Private::IsOfTypes<AccessorType, FRotator>())
 			{
-				return PCGAttributeExtractor::CreateRotatorExtractor(std::move(Accessor), Name, bOutSuccess);
+				return PCGAttributeExtractor::CreateRotatorExtractor(MoveTemp(Accessor), Name, bOutSuccess);
+			}
+			else if constexpr (PCG::Private::IsOfTypes<AccessorType, FString, FName, FSoftObjectPath, FSoftClassPath>())
+			{
+				return PCGAttributeExtractor::CreateStringExtractor<AccessorType>(MoveTemp(Accessor), Name, bOutSuccess);
 			}
 			else
 			{
diff --git a/Engine/Plugins/PCG/Source/PCG/Private/Metadata/Accessors/PCGAttributeExtractor.cpp b/Engine/Plugins/PCG/Source/PCG/Private/Metadata/Accessors/PCGAttributeExtractor.cpp
--- a/Engine/Plugins/PCG/Source/PCG/Private/Metadata/Accessors/PCGAttributeExtractor.cpp
+++ b/Engine/Plugins/PCG/Source/PCG/Private/Metadata/Accessors/PCGAttributeExtractor.cpp
@@ -333,9 +333,36 @@ namespace PCGAttributeExtractor
 		return InAccessor;
 	}
 
+	template <typename StringType>
+	TUniquePtr<IPCGAttributeAccessor> CreateStringExtractor(TUniquePtr<IPCGAttributeAccessor> InAccessor, const FName Name, bool& bOutSuccess)
+	{
+		bOutSuccess = true;
+
+		if (Name == PCGAttributeExtractorConstants::StringLength)
+		{
+			if constexpr (std::is_same_v<StringType, FString>)
+			{
+				return MakeUnique<FPCGChainAccessor<int32, FString>>(MoveTemp(InAccessor),
+					[](const FString& Value) -> int32 { return Value.Len(); });
+			}
+			else
+			{
+				return MakeUnique<FPCGChainAccessor<int32, StringType>>(MoveTemp(InAccessor),
+					[](const StringType& Value) -> int32 { return Value.ToString().Len(); });
+			}
+		}
+
+		bOutSuccess = false;
+		return InAccessor;
+	}
+
 	// Template instantiation for all vectors types + quat
 	template TUniquePtr<IPCGAttributeAccessor> CreateVectorExtractor<FVector2D>(TUniquePtr<IPCGAttributeAccessor> InAccessor, FName Name, bool& bOutSuccess);
 	template TUniquePtr<IPCGAttributeAccessor> CreateVectorExtractor<FVector>(TUniquePtr<IPCGAttributeAccessor> InAccessor, FName Name, bool& bOutSuccess);
 	template TUniquePtr<IPCGAttributeAccessor> CreateVectorExtractor<FVector4>(TUniquePtr<IPCGAttributeAccessor> InAccessor, FName Name, bool& bOutSuccess);
 	template TUniquePtr<IPCGAttributeAccessor> CreateVectorExtractor<FQuat>(TUniquePtr<IPCGAttributeAccessor> InAccessor, FName Name, bool& bOutSuccess);
+	template TUniquePtr<IPCGAttributeAccessor> CreateStringExtractor<FString>(TUniquePtr<IPCGAttributeAccessor> InAccessor, FName Name, bool& bOutSuccess);
+	template TUniquePtr<IPCGAttributeAccessor> CreateStringExtractor<FName>(TUniquePtr<IPCGAttributeAccessor> InAccessor, FName Name, bool& bOutSuccess);
+	template TUniquePtr<IPCGAttributeAccessor> CreateStringExtractor<FSoftObjectPath>(TUniquePtr<IPCGAttributeAccessor> InAccessor, FName Name, bool& bOutSuccess);
+	template TUniquePtr<IPCGAttributeAccessor> CreateStringExtractor<FSoftClassPath>(TUniquePtr<IPCGAttributeAccessor> InAccessor, FName Name, bool& bOutSuccess);
 }
diff --git a/Engine/Plugins/PCG/Source/PCG/Private/Tests/Accessors/PCGAttributeExtractorTest.cpp b/Engine/Plugins/PCG/Source/PCG/Private/Tests/Accessors/PCGAttributeExtractorTest.cpp
--- a/Engine/Plugins/PCG/Source/PCG/Private/Tests/Accessors/PCGAttributeExtractorTest.cpp
+++ b/Engine/Plugins/PCG/Source/PCG/Private/Tests/Accessors/PCGAttributeExtractorTest.cpp
@@ -50,6 +50,7 @@ IMPLEMENT_CUSTOM_SIMPLE_AUTOMATION_TEST(FPCGAttributeSingleGetPropertyTest, FPCG
 IMPLEMENT_CUSTOM_SIMPLE_AUTOMATION_TEST(FPCGAttributeVectorPropertyExtractorTest, FPCGTestBaseClass, "Plugins.PCG.Accessor.Property.VectorExtractor", PCGTestsCommon::TestFlags)
 IMPLEMENT_CUSTOM_SIMPLE_AUTOMATION_TEST(FPCGAttributeRotatorPropertyExtractorTest, FPCGTestBaseClass, "Plugins.PCG.Accessor.Property.RotatorExtractor", PCGTestsCommon::TestFlags)
 IMPLEMENT_CUSTOM_SIMPLE_AUTOMATION_TEST(FPCGAttributeTransformPropertyExtractorTest, FPCGTestBaseClass, "Plugins.PCG.Accessor.Property.TransformExtractor", PCGTestsCommon::TestFlags)
+IMPLEMENT_CUSTOM_SIMPLE_AUTOMATION_TEST(FPCGAttributeStringPropertyExtractorTest, FPCGTestBaseClass, "Plugins.PCG.Accessor.Property.StringExtractor", PCGTestsCommon::TestFlags)
 IMPLEMENT_CUSTOM_SIMPLE_AUTOMATION_TEST(FPCGAttributeInvalidPropertyExtractorTest, FPCGTestBaseClass, "Plugins.PCG.Accessor.Property.InvalidExtractor", PCGTestsCommon::TestFlags)
 IMPLEMENT_CUSTOM_SIMPLE_AUTOMATION_TEST(FPCGAttributePropertyMultipleDepthTest, FPCGTestBaseClass, "Plugins.PCG.Accessor.Property.MultipleDepth", PCGTestsCommon::TestFlags)
 IMPLEMENT_CUSTOM_SIMPLE_AUTOMATION_TEST(FPCGAttributePropertyMultipleDepthRangeTest, FPCGTestBaseClass, "Plugins.PCG.Accessor.Property.MultipleDepthRange", PCGTestsCommon::TestFlags)
@@ -136,7 +137,7 @@ bool FPCGAttributeVectorPropertyExtractorTest::RunTest(const FString& Parameters
 		TUniquePtr<IPCGAttributeAccessor> ExtractorAccessor = PCGAttributeAccessorHelpers::CreateChainAccessor(std::move(Accessor), ExtractorName, bSuccess);
 		UTEST_TRUE(TEXT("Valid extractor accessor"), (bSuccess && ExtractorAccessor.IsValid()));
 		UTEST_TRUE(TEXT("Invalid accessor (ownership moved)"), !Accessor.IsValid());
-		UTEST_EQUAL(TEXT("Extractror accessor type"), PCG::Private::MetadataTypes<PropertyType>::Id, ExtractorAccessor->GetUnderlyingType());
+		UTEST_EQUAL(TEXT("Extractor accessor type"), PCG::Private::MetadataTypes<PropertyType>::Id, ExtractorAccessor->GetUnderlyingType());
 
 		FPCGAttributeAccessorKeysSingleObjectPtr ObjectKey(&Struct);
 
@@ -269,7 +270,7 @@ bool FPCGAttributeRotatorPropertyExtractorTest::RunTest(const FString & Paramete
 		TUniquePtr<IPCGAttributeAccessor> ExtractorAccessor = PCGAttributeAccessorHelpers::CreateChainAccessor(std::move(Accessor), ExtractorName, bSuccess);
 		UTEST_TRUE(TEXT("Valid extractor accessor"), (bSuccess && ExtractorAccessor.IsValid()));
 		UTEST_TRUE(TEXT("Invalid accessor (ownership moved)"), !Accessor.IsValid());
-		UTEST_EQUAL(TEXT("Extractror accessor type"), PCG::Private::MetadataTypes<ExpectedType>::Id, ExtractorAccessor->GetUnderlyingType());
+		UTEST_EQUAL(TEXT("Extractor accessor type"), PCG::Private::MetadataTypes<ExpectedType>::Id, ExtractorAccessor->GetUnderlyingType());
 
 		FPCGAttributeAccessorKeysSingleObjectPtr ObjectKey(&Struct);
 
@@ -349,6 +350,53 @@ bool FPCGAttributeTransformPropertyExtractorTest::RunTest(const FString & Parame
 	return true;
 }
 
+bool FPCGAttributeStringPropertyExtractorTest::RunTest(const FString & Parameters)
+{
+	UPCGAttributeExtractorTestObject* Object = NewObject<UPCGAttributeExtractorTestObject>();
+	FPCGMetadataTypesConstantStruct Struct = FPCGAttributeExtractorTestHelpers::CreateStruct(Object);
+
+	auto Verify = [this, &Struct]<typename T>(const T& ExpectedValue, const FName PropertyName, const TArray<FName>& ExtractorNames) -> bool
+	{
+		using PropertyType = std::decay_t<T>;
+
+		TUniquePtr<IPCGAttributeAccessor> Accessor = PCGAttributeAccessorHelpers::CreatePropertyAccessor(PropertyName, FPCGMetadataTypesConstantStruct::StaticStruct());
+		UTEST_TRUE(FString::Printf(TEXT("Valid accessor for property %s"), *PropertyName.ToString()), Accessor.IsValid());
+
+		bool bSuccess = false;
+		TUniquePtr<IPCGAttributeAccessor> ExtractorAccessor = std::move(Accessor);
+		for (const FName ExtractorName : ExtractorNames)
+		{
+			ExtractorAccessor = PCGAttributeAccessorHelpers::CreateChainAccessor(std::move(ExtractorAccessor), ExtractorName, bSuccess);
+			UTEST_TRUE(FString::Printf(TEXT("Valid extractor accessor %s"), *ExtractorName.ToString()), (bSuccess && ExtractorAccessor.IsValid()));
+		}
+
+		UTEST_EQUAL("Extractor accessor type", PCG::Private::MetadataTypes<PropertyType>::Id, ExtractorAccessor->GetUnderlyingType());
+
+		FPCGAttributeAccessorKeysSingleObjectPtr ObjectKey(&Struct);
+
+		PropertyType Value{};
+		FString ExtractorStr = FString::JoinBy(ExtractorNames, TEXT("."), [](const FName& Name) -> FString { return Name.ToString(); });
+		UTEST_TRUE(FString::Printf(TEXT("Getting a value for property %s and extractor %s"), *PropertyName.ToString(), *ExtractorStr), ExtractorAccessor->Get<PropertyType>(Value, ObjectKey));
+		UTEST_EQUAL(FString::Printf(TEXT("Value for property %s and extractor %s is equal to expected"), *PropertyName.ToString(), *ExtractorStr), Value, ExpectedValue);
+
+		return true;
+	};
+
+	bool bTestPassed = true;
+
+	// "Foo" = 3 characters
+	bTestPassed &= Verify(3, PCG_STRUCT_NAME_CHECK(StringValue), {PCGAttributeExtractorConstants::StringLength});
+	// "Bar" = 3 characters
+	bTestPassed &= Verify(3, PCG_STRUCT_NAME_CHECK(NameValue), {PCGAttributeExtractorConstants::StringLength});
+	auto String = FSoftObjectPath(Object).ToString();
+	// "/Engine/Transient.PCGAttributeExtractorTestObject_1" = 51 characters
+	bTestPassed &= Verify(51, PCG_STRUCT_NAME_CHECK(SoftObjectPathValue), {PCGAttributeExtractorConstants::StringLength});
+	// "/Script/PCG.PCGMetadata" = 23 characters
+	bTestPassed &= Verify(23, PCG_STRUCT_NAME_CHECK(SoftClassPathValue), {PCGAttributeExtractorConstants::StringLength});
+
+	return bTestPassed;
+}
+
 bool FPCGAttributeInvalidPropertyExtractorTest::RunTest(const FString& Parameters)
 {
 	const FName InvalidName = FName(TEXT("Invalid"));
diff --git a/Engine/Plugins/PCG/Source/PCG/Public/Elements/Metadata/PCGMetadataStringOpElement.h b/Engine/Plugins/PCG/Source/PCG/Public/Elements/Metadata/PCGMetadataStringOpElement.h
--- a/Engine/Plugins/PCG/Source/PCG/Public/Elements/Metadata/PCGMetadataStringOpElement.h
+++ b/Engine/Plugins/PCG/Source/PCG/Public/Elements/Metadata/PCGMetadataStringOpElement.h
@@ -11,6 +11,8 @@ enum class EPCGMetadataStringOperation : uint16
 {
 	Append UMETA(DisplayName="Append String"),
 	Replace UMETA(DisplayName="Replace String"),
+	Substring UMETA(DisplayName="Substring", Tooltip="True if string A contains substring B.", SearchHints="contains"),
+	Matches UMETA(DisplayName="Matches", Tooltip="True if string A matches substring B exactly.", SearchHints="equal"),
 	ToUpper UMETA(DisplayName="To Upper", Tooltip="Convert all characters to upper case."),
 	ToLower UMETA(DisplayName="To Lower", Tooltip="Convert all characters to lower case."),
 	TrimStart UMETA(DisplayName="Trim Start", Tooltip="Trim whitespace from the beginning of the string."),
@@ -63,6 +65,9 @@ public:
 	UPROPERTY(BlueprintReadWrite, EditAnywhere, Category = Settings)
 	EPCGMetadataStringOperation Operation = EPCGMetadataStringOperation::Append;
 
+	UPROPERTY(BlueprintReadWrite, EditAnywhere, Category = Settings, meta = (PCG_Overridable, EditCondition = "Operation == EPCGMetadataStringOperation::Substring || Operation == EPCGMetadataStringOperation::Matches", EditConditionHides))
+	TEnumAsByte<ESearchCase::Type> SearchCase = ESearchCase::CaseSensitive;
+
 	UPROPERTY(BlueprintReadWrite, EditAnywhere, Category = Input, meta = (PCG_Overridable))
 	FPCGAttributePropertyInputSelector InputSource1;
 
diff --git a/Engine/Plugins/PCG/Source/PCG/Public/Metadata/Accessors/PCGAttributeExtractor.h b/Engine/Plugins/PCG/Source/PCG/Public/Metadata/Accessors/PCGAttributeExtractor.h
--- a/Engine/Plugins/PCG/Source/PCG/Public/Metadata/Accessors/PCGAttributeExtractor.h
+++ b/Engine/Plugins/PCG/Source/PCG/Public/Metadata/Accessors/PCGAttributeExtractor.h
@@ -70,6 +70,9 @@ namespace PCGAttributeExtractorConstants
 		RotatorUp,
 		RotatorEulerAngles
 	};
+
+	// String
+	const FName StringLength = TEXT("Length");
 }
 
 namespace PCGAttributeExtractor
@@ -82,4 +85,7 @@ namespace PCGAttributeExtractor
 	TUniquePtr<IPCGAttributeAccessor> CreateQuatExtractor(TUniquePtr<IPCGAttributeAccessor> InAccessor, FName Name, bool& bOutSuccess);
 
 	TUniquePtr<IPCGAttributeAccessor> CreateTransformExtractor(TUniquePtr<IPCGAttributeAccessor> InAccessor, FName Name, bool& bOutSuccess);
+
+	template <typename StringType>
+	TUniquePtr<IPCGAttributeAccessor> CreateStringExtractor(TUniquePtr<IPCGAttributeAccessor> InAccessor, FName Name, bool& bOutSuccess);
 }
diff --git a/Engine/Plugins/PCG/Source/PCG/Public/Metadata/Accessors/PCGCustomAccessor.h b/Engine/Plugins/PCG/Source/PCG/Public/Metadata/Accessors/PCGCustomAccessor.h
--- a/Engine/Plugins/PCG/Source/PCG/Public/Metadata/Accessors/PCGCustomAccessor.h
+++ b/Engine/Plugins/PCG/Source/PCG/Public/Metadata/Accessors/PCGCustomAccessor.h
@@ -142,7 +142,15 @@ public:
 	bool GetRangeImpl(TArrayView<T> OutValues, int32 Index, const IPCGAttributeAccessorKeys& Keys) const
 	{
 		TArray<U> TempValues;
-		TempValues.SetNumUninitialized(OutValues.Num());
+		if constexpr (std::is_trivially_copyable_v<U>)
+		{
+			TempValues.SetNumUninitialized(Keys.GetNum());
+		}
+		else
+		{
+			TempValues.SetNum(Keys.GetNum());
+		}
+
 		if (!Accessor->GetRange<U>(TempValues, Index, Keys))
 		{
 			return false;
@@ -159,7 +167,15 @@ public:
 	bool SetRangeImpl(TArrayView<const T> InValues, int32 Index, IPCGAttributeAccessorKeys& Keys, EPCGAttributeAccessorFlags Flags)
 	{
 		TArray<U> TempValues;
-		TempValues.SetNumUninitialized(InValues.Num());
+		if constexpr (std::is_trivially_copyable_v<U>)
+		{
+			TempValues.SetNumUninitialized(Keys.GetNum());
+		}
+		else
+		{
+			TempValues.SetNum(Keys.GetNum());
+		}
+
 		if (!Accessor->GetRange<U>(TempValues, Index, Keys))
 		{
 			return false;
diff --git a/Engine/Plugins/PCG/Source/PCG/Public/Metadata/PCGMetadataAttributeTraits.h b/Engine/Plugins/PCG/Source/PCG/Public/Metadata/PCGMetadataAttributeTraits.h
--- a/Engine/Plugins/PCG/Source/PCG/Public/Metadata/PCGMetadataAttributeTraits.h
+++ b/Engine/Plugins/PCG/Source/PCG/Public/Metadata/PCGMetadataAttributeTraits.h
@@ -1135,14 +1135,14 @@ namespace PCG
 				return FString();
 			}
 
-			static bool Substring(const FString& A, const FString& B)
+			static bool Substring(const FString& A, const FString& B, const ESearchCase::Type SearchCase = ESearchCase::IgnoreCase)
 			{
-				return A.Contains(B);
+				return A.Contains(B, SearchCase);
 			}
 
-			static bool Matches(const FString& A, const FString& B)
+			static bool Matches(const FString& A, const FString& B, const ESearchCase::Type SearchCase = ESearchCase::IgnoreCase)
 			{
-				return A.MatchesWildcard(B);
+				return A.MatchesWildcard(B, SearchCase);
 			}
 		};
 
@@ -1190,14 +1190,14 @@ namespace PCG
 				return NAME_None;
 			}
 
-			static bool Substring(const FName& A, const FName& B)
+			static bool Substring(const FName& A, const FName& B, const ESearchCase::Type SearchCase = ESearchCase::IgnoreCase)
 			{
-				return A.ToString().Contains(B.ToString());
+				return A.ToString().Contains(B.ToString(), SearchCase);
 			}
 
-			static bool Matches(const FName& A, const FName& B)
+			static bool Matches(const FName& A, const FName& B, const ESearchCase::Type SearchCase = ESearchCase::IgnoreCase)
 			{
-				return A.ToString().MatchesWildcard(B.ToString());
+				return A.ToString().MatchesWildcard(B.ToString(), SearchCase);
 			}
 		};
 
@@ -1226,14 +1226,14 @@ namespace PCG
 				return T();
 			}
 
-			static bool Substring(const T& A, const T& B)
+			static bool Substring(const T& A, const T& B, const ESearchCase::Type SearchCase = ESearchCase::IgnoreCase)
 			{
-				return A.ToString().Contains(B.ToString());
+				return A.ToString().Contains(B.ToString(), SearchCase);
 			}
 
-			static bool Matches(const T& A, const T& B)
+			static bool Matches(const T& A, const T& B, const ESearchCase::Type SearchCase = ESearchCase::IgnoreCase)
 			{
-				return A.ToString().MatchesWildcard(B.ToString());
+				return A.ToString().MatchesWildcard(B.ToString(), SearchCase);
 			}
 
 			static bool Less(const T& A, const T& B)
diff --git a/Engine/Plugins/PCG/Source/PCGEditor/Private/Details/PCGAttributePropertySelectorDetails.cpp b/Engine/Plugins/PCG/Source/PCGEditor/Private/Details/PCGAttributePropertySelectorDetails.cpp
--- a/Engine/Plugins/PCG/Source/PCGEditor/Private/Details/PCGAttributePropertySelectorDetails.cpp
+++ b/Engine/Plugins/PCG/Source/PCGEditor/Private/Details/PCGAttributePropertySelectorDetails.cpp
@@ -121,7 +121,7 @@ namespace PCGAttributePropertySelectorDetails
 		}));
 	}
 
-	void BuildSubMenuSectionTransform(FMenuBuilder& InMenuBuilder, FPCGAttributePropertySelectorDetails* InDetailsObject, bool bIsInput)
+	void BuildSubMenuSectionTransform(FMenuBuilder& InMenuBuilder, FPCGAttributePropertySelectorDetails* InDetailsObject)
 	{
 		InMenuBuilder.AddSubMenu(
 			LOCTEXT("Transform", "Transform"),
@@ -135,6 +135,19 @@ namespace PCGAttributePropertySelectorDetails
 		}));
 	}
 
+	void BuildSubMenuSectionString(FMenuBuilder& InMenuBuilder, FPCGAttributePropertySelectorDetails* InDetailsObject)
+	{
+		InMenuBuilder.AddSubMenu(
+			LOCTEXT("String", "String"),
+			FText{},
+			FNewMenuDelegate::CreateLambda([InDetailsObject](FMenuBuilder& InSubMenuBuilder)
+			{
+				static const FText TooltipComponents = LOCTEXT("TooltipString", "Length of any string-based attribute.");
+
+				BuildSubMenuSectionForExtractor(InSubMenuBuilder, InDetailsObject, PCGAttributeExtractorConstants::StringLength, TooltipComponents);
+			}));
+	}
+
 	// Func signature: bool(FPCGAttributePropertySelector*, T)
 	template <typename T, typename Func>
 	void SetValue(FPCGAttributePropertySelector* InSelector, TSharedPtr<IPropertyHandle>& InPropertyHandle, T&& InValue, Func InCallback)
@@ -253,15 +266,17 @@ TSharedRef<SWidget> FPCGAttributePropertySelectorDetails::GenerateExtraMenu()
 	FStructProperty* StructProperty = CastFieldChecked<FStructProperty>(PropertyHandle->GetProperty());
 	const bool bIsInput = StructProperty->Struct == FPCGAttributePropertyInputSelector::StaticStruct();
 
+	// @todo_pcg: This submenu could be context friendly to the property type an only show when relevant.
 	MenuBuilder.AddSubMenu(
 		LOCTEXT("AddAnExtractor", "Add an extractor"),
 		LOCTEXT("AddAnExtractorTooltip", "List of all the extractor to add to the current attribute."),
 		FNewMenuDelegate::CreateLambda([this, bIsInput](FMenuBuilder& InSubMenuBuilder)
-	{
-		PCGAttributePropertySelectorDetails::BuildSubMenuSectionVector(InSubMenuBuilder, this, bIsInput);
-		PCGAttributePropertySelectorDetails::BuildSubMenuSectionRotator(InSubMenuBuilder, this, bIsInput);
-		PCGAttributePropertySelectorDetails::BuildSubMenuSectionTransform(InSubMenuBuilder, this, bIsInput);
-	}));
+		{
+			PCGAttributePropertySelectorDetails::BuildSubMenuSectionVector(InSubMenuBuilder, this, bIsInput);
+			PCGAttributePropertySelectorDetails::BuildSubMenuSectionRotator(InSubMenuBuilder, this, bIsInput);
+			PCGAttributePropertySelectorDetails::BuildSubMenuSectionTransform(InSubMenuBuilder, this);
+			PCGAttributePropertySelectorDetails::BuildSubMenuSectionString(InSubMenuBuilder, this);
+		}));
 
 	MenuBuilder.BeginSection("Attributes", LOCTEXT("AttributesHeader", "Attributes"));
 	{
