[PCG] Update the menu descriptions for asset copy and instances

TODO: Once new icons are added, they can be enabled optionally instead of the suffix

diff --git a/Engine/Plugins/PCG/Source/PCGEditor/Private/PCGEditorGraphSchema.cpp b/Engine/Plugins/PCG/Source/PCGEditor/Private/PCGEditorGraphSchema.cpp
--- a/Engine/Plugins/PCG/Source/PCGEditor/Private/PCGEditorGraphSchema.cpp
+++ b/Engine/Plugins/PCG/Source/PCGEditor/Private/PCGEditorGraphSchema.cpp
@@ -570,9 +570,9 @@ void UPCGEditorGraphSchema::GetSettingsElementActions(FGraphActionMenuBuilder& A
 		if (bExposeToLibrary)
 		{
 			const FText MenuDesc = FText::FromString(FName::NameToDisplayString(AssetData.AssetName.ToString(), false));
-			const FText Category = AssetData.GetTagValueRef<FText>(TEXT("Category"));
-			const FText MenuAndSubCategory = FText::Join(LOCTEXT("MenuDelimiter", "|"), Category, MenuDesc);
-			const FText Description = AssetData.GetTagValueRef<FText>(TEXT("Description"));
+			const FText AssetCategory = AssetData.GetTagValueRef<FText>(TEXT("Category"));
+			const FText Category = AssetCategory.IsEmptyOrWhitespace() ? LOCTEXT("AssetMenuDescriptionUncategorized", "Uncategorized Assets") : AssetCategory;
+			const FText Description = FText::Format(INVTEXT("{0}\n{1}"), FText::FromString(AssetData.GetObjectPathString()), AssetData.GetTagValueRef<FText>(TEXT("Description")));
 
 			if (!bIsContextual && ActionFilter.Accepts(Category))
 			{
@@ -580,14 +580,14 @@ void UPCGEditorGraphSchema::GetSettingsElementActions(FGraphActionMenuBuilder& A
 				NewSettingsAction->SettingsObjectPath = AssetData.GetSoftObjectPath();
 				ActionMenuBuilder.AddAction(NewSettingsAction);
 			}
-			else if(bIsContextual && ActionFilter.Accepts(MenuAndSubCategory))
+			else if(bIsContextual && ActionFilter.Accepts(Category))
 			{
-				TSharedPtr<FPCGEditorGraphSchemaAction_NewSettingsElement> NewSettingsActionCopy(new FPCGEditorGraphSchemaAction_NewSettingsElement(MenuAndSubCategory, LOCTEXT("ContextMenuCopySettings", "Copy"), Description, 0));
+				TSharedPtr<FPCGEditorGraphSchemaAction_NewSettingsElement> NewSettingsActionCopy(new FPCGEditorGraphSchemaAction_NewSettingsElement(Category, MenuDesc, Description, 0));
 				NewSettingsActionCopy->SettingsObjectPath = AssetData.GetSoftObjectPath();
 				NewSettingsActionCopy->Behavior = EPCGEditorNewSettingsBehavior::ForceCopy;
 				ActionMenuBuilder.AddAction(NewSettingsActionCopy);
 
-				TSharedPtr<FPCGEditorGraphSchemaAction_NewSettingsElement> NewSettingsActionInstance(new FPCGEditorGraphSchemaAction_NewSettingsElement(MenuAndSubCategory, LOCTEXT("ContextMenuInstanceSettings", "Instance"), Description, 0));
+				TSharedPtr<FPCGEditorGraphSchemaAction_NewSettingsElement> NewSettingsActionInstance(new FPCGEditorGraphSchemaAction_NewSettingsElement(Category, MenuDesc, Description, 0));
 				NewSettingsActionInstance->SettingsObjectPath = AssetData.GetSoftObjectPath();
 				NewSettingsActionInstance->Behavior = EPCGEditorNewSettingsBehavior::ForceInstance;
 				ActionMenuBuilder.AddAction(NewSettingsActionInstance);
diff --git a/Engine/Plugins/PCG/Source/PCGEditor/Private/PCGEditorGraphSchemaActions.cpp b/Engine/Plugins/PCG/Source/PCGEditor/Private/PCGEditorGraphSchemaActions.cpp
--- a/Engine/Plugins/PCG/Source/PCGEditor/Private/PCGEditorGraphSchemaActions.cpp
+++ b/Engine/Plugins/PCG/Source/PCGEditor/Private/PCGEditorGraphSchemaActions.cpp
@@ -29,6 +29,21 @@
 
 #define LOCTEXT_NAMESPACE "PCGEditorGraphSchemaActions"
 
+FText FPCGEditorGraphSchemaActionBase::GetMenuDescriptionOverride() const
+{
+	return GetMenuDescription();
+}
+
+FSlateBrush const* FPCGEditorGraphSchemaActionBase::GetPaletteIcon() const
+{
+	return PaletteIcon.Get();
+}
+
+void FPCGEditorGraphSchemaActionBase::SetPaletteIcon(TSharedPtr<FSlateBrush>& InPaletteIcon)
+{
+	PaletteIcon = InPaletteIcon;
+}
+
 UEdGraphNode* FPCGEditorGraphSchemaAction_NewNativeElement::PerformAction(UEdGraph* ParentGraph, UEdGraphPin* FromPin, const FVector2f& Location, bool bSelectNewNode)
 {
 	UPCGEditorGraph* EditorGraph = Cast<UPCGEditorGraph>(ParentGraph);
@@ -163,6 +178,21 @@ void FPCGEditorGraphSchemaAction_NewGetParameterElement::PostCreation(UPCGNode*
 	NewNode->SetSettingsInterface(Settings);
 }
 
+FText FPCGEditorGraphSchemaAction_NewSettingsElement::GetMenuDescriptionOverride() const
+{
+	switch (Behavior)
+	{
+		case EPCGEditorNewSettingsBehavior::ForceCopy:
+			return FText::Format(LOCTEXT("AssetMenuDescriptionCopyQualifierFormat", "{0} (copy)"), GetMenuDescription());
+		case EPCGEditorNewSettingsBehavior::ForceInstance:
+			return FText::Format(LOCTEXT("AssetMenuDescriptionInstanceQualifierFormat", "{0} (instance)"), GetMenuDescription());
+		case EPCGEditorNewSettingsBehavior::Normal:
+			break;
+	}
+
+	return Super::GetMenuDescriptionOverride();
+}
+
 UEdGraphNode* FPCGEditorGraphSchemaAction_NewSettingsElement::PerformAction(UEdGraph* ParentGraph, UEdGraphPin* FromPin, const FVector2f& Location, bool bSelectNewNode)
 {
 	UPCGEditorGraph* EditorGraph = Cast<UPCGEditorGraph>(ParentGraph);
@@ -207,6 +237,22 @@ UEdGraphNode* FPCGEditorGraphSchemaAction_NewSettingsElement::PerformAction(UEdG
 	return MakeSettingsNode(EditorGraph, FromPin, Settings, FDeprecateSlateVector2D(Location), bSelectNewNode, bCreateInstance);
 }
 
+FSlateBrush const* FPCGEditorGraphSchemaAction_NewSettingsElement::GetPaletteIcon() const
+{
+	// @todo_pcg: Enable when icons exist for these types.
+	// switch (Behavior)
+	// {
+	// 	case EPCGEditorNewSettingsBehavior::ForceCopy:
+	// 		return FAppStyle::GetBrush("");
+	// 	case EPCGEditorNewSettingsBehavior::ForceInstance:
+	// 		return FAppStyle::GetBrush("");
+	// 	case EPCGEditorNewSettingsBehavior::Normal:
+	// 		break;
+	// }
+
+	return nullptr;
+}
+
 void FPCGEditorGraphSchemaAction_NewSettingsElement::MakeSettingsNodesOrContextualMenu(const TSharedRef<class SWidget>& InPanel, const FVector2D& InScreenPosition, UEdGraph* InGraph, const TArray<FSoftObjectPath>& InSettingsPaths, const TArray<FVector2D>& InLocations, bool bInSelectNewNodes)
 {
 	UPCGEditorGraph* EditorGraph = Cast<UPCGEditorGraph>(InGraph);
@@ -448,6 +494,11 @@ UEdGraphNode* FPCGEditorGraphSchemaAction_NewBlueprintElement::PerformAction(UEd
 	return NewNode;
 }
 
+FSlateBrush const* FPCGEditorGraphSchemaAction_NewBlueprintElement::GetPaletteIcon() const
+{
+	return FAppStyle::GetBrush("Icons.Blueprint");
+}
+
 UEdGraphNode* FPCGEditorGraphSchemaAction_NewSubgraphElement::PerformAction(UEdGraph* ParentGraph, UEdGraphPin* FromPin, const FVector2f& Location, bool bSelectNewNode)
 {
 	UPCGEditorGraph* EditorGraph = Cast<UPCGEditorGraph>(ParentGraph);
diff --git a/Engine/Plugins/PCG/Source/PCGEditor/Private/Widgets/SPCGEditorGraphActionWidget.h b/Engine/Plugins/PCG/Source/PCGEditor/Private/Widgets/SPCGEditorGraphActionWidget.h
--- a/Engine/Plugins/PCG/Source/PCGEditor/Private/Widgets/SPCGEditorGraphActionWidget.h
+++ b/Engine/Plugins/PCG/Source/PCGEditor/Private/Widgets/SPCGEditorGraphActionWidget.h
@@ -2,8 +2,11 @@
 
 #pragma once
 
+#include "Schema/PCGEditorGraphSchemaActions.h"
+
 #include "SGraphActionMenu.h"
 #include "EdGraph/EdGraphSchema.h"
+#include "Templates/SharedPointer.h"
 #include "Widgets/Images/SLayeredImage.h"
 
 class SPCGGraphActionWidget : public SCompoundWidget
@@ -31,16 +34,17 @@ class SPCGGraphActionWidget : public SCompoundWidget
 			.Image(PrimaryIcon);
 		}
 
+		const TSharedPtr<FPCGEditorGraphSchemaActionBase> PCGAction = StaticCastSharedPtr<FPCGEditorGraphSchemaActionBase>(InCreateData->Action);
 		TSharedRef<SWidget> NameSlotWidget =
 			InArgs._NameWidget
 				? InArgs._NameWidget.ToSharedRef()
 				: SNew(STextBlock)
-				.Text(InCreateData->Action->GetMenuDescription());
+				.Text(PCGAction->GetMenuDescriptionOverride());
 
 		this->ChildSlot
 		[
 			SNew(SHorizontalBox)
-			.ToolTipText(InCreateData->Action->GetMenuDescription())
+			.ToolTipText(InCreateData->Action->GetTooltipDescription())
 			+ SHorizontalBox::Slot()
 			.AutoWidth()
 			[
diff --git a/Engine/Plugins/PCG/Source/PCGEditor/Public/Schema/PCGEditorGraphSchemaActions.h b/Engine/Plugins/PCG/Source/PCGEditor/Public/Schema/PCGEditorGraphSchemaActions.h
--- a/Engine/Plugins/PCG/Source/PCGEditor/Public/Schema/PCGEditorGraphSchemaActions.h
+++ b/Engine/Plugins/PCG/Source/PCGEditor/Public/Schema/PCGEditorGraphSchemaActions.h
@@ -34,13 +34,29 @@ enum class EPCGEditorNewPCGGraphBehavior : uint8
 };
 
 USTRUCT(MinimalAPI)
-struct FPCGEditorGraphSchemaAction_NewNativeElement : public FEdGraphSchemaAction
+struct FPCGEditorGraphSchemaActionBase : public FEdGraphSchemaAction
 {
 	GENERATED_BODY()
 
-	// Inherit the base class's constructors
 	using FEdGraphSchemaAction::FEdGraphSchemaAction;
 
+	virtual FText GetMenuDescriptionOverride() const;
+
+	virtual FSlateBrush const* GetPaletteIcon() const override;
+	void SetPaletteIcon(TSharedPtr<FSlateBrush>& InPaletteIcon);
+
+protected:
+	mutable TSharedPtr<FSlateBrush> PaletteIcon = nullptr;
+};
+
+USTRUCT(MinimalAPI)
+struct FPCGEditorGraphSchemaAction_NewNativeElement : public FPCGEditorGraphSchemaActionBase
+{
+	GENERATED_BODY()
+
+	// Inherit the base class's constructors
+	using FPCGEditorGraphSchemaActionBase::FPCGEditorGraphSchemaActionBase;
+
 	// Simple type info
 	static FName StaticGetTypeId()
 	{
@@ -58,15 +74,11 @@ struct FPCGEditorGraphSchemaAction_NewNativeElement : public FEdGraphSchemaActio
 	virtual FName GetTypeId() const override { return StaticGetTypeId(); }
 	using FEdGraphSchemaAction::PerformAction; // Prevent hiding of deprecated base class function with FVector2D
 	virtual UEdGraphNode* PerformAction(UEdGraph* ParentGraph, UEdGraphPin* FromPin, const FVector2f& Location, bool bSelectNewNode = true) override;
-	//~ End FEdGraphSchemaAction Interface
-
 	virtual FSlateBrush const* GetPaletteIcon() const override;
+	//~ End FEdGraphSchemaAction Interface
 
 protected:
 	virtual void PostCreation(UPCGNode* NewNode) {}
-
-private:
-	mutable TSharedPtr<FSlateBrush> PaletteIcon = nullptr;
 };
 
 USTRUCT()
@@ -84,19 +96,16 @@ struct FPCGEditorGraphSchemaAction_NewGetParameterElement : public FPCGEditorGra
 
 protected:
 	virtual void PostCreation(UPCGNode* NewNode) override;
-
-private:
-	mutable TSharedPtr<FSlateBrush> PaletteIcon = nullptr;
 };
 
 USTRUCT()
-struct FPCGEditorGraphSchemaAction_NewSettingsElement : public FEdGraphSchemaAction
+struct FPCGEditorGraphSchemaAction_NewSettingsElement : public FPCGEditorGraphSchemaActionBase
 {
 	GENERATED_BODY()
 
 	// Inherit the base class's constructors
-	using FEdGraphSchemaAction::FEdGraphSchemaAction;
-	
+	using FPCGEditorGraphSchemaActionBase::FPCGEditorGraphSchemaActionBase;
+
 	// Simple type info
 	static FName StaticGetTypeId()
 	{
@@ -110,10 +119,15 @@ struct FPCGEditorGraphSchemaAction_NewSettingsElement : public FEdGraphSchemaAct
 	UPROPERTY()
 	EPCGEditorNewSettingsBehavior Behavior = EPCGEditorNewSettingsBehavior::Normal;
 
+	//~ Begin FPCGEditorGraphSchemaActionBase Interface
+	virtual FText GetMenuDescriptionOverride() const override;
+	//~ End FPCGEditorGraphSchemaActionBase Interface
+
 	//~ Begin FEdGraphSchemaAction Interface
 	virtual FName GetTypeId() const override { return StaticGetTypeId(); }
 	using FEdGraphSchemaAction::PerformAction; // Prevent hiding of deprecated base class function with FVector2D
 	virtual UEdGraphNode* PerformAction(UEdGraph* ParentGraph, UEdGraphPin* FromPin, const FVector2f& Location, bool bSelectNewNode = true) override;
+	virtual FSlateBrush const* GetPaletteIcon() const override;
 	//~ End FEdGraphSchemaAction Interface
 
 	static void MakeSettingsNodesOrContextualMenu(const TSharedRef<class SWidget>& InPanel, const FVector2D& InScreenPosition, UEdGraph* InGraph, const TArray<FSoftObjectPath>& InSettingsPaths, const TArray<FVector2D>& InLocations, bool bInSelectNewNodes);
@@ -122,12 +136,12 @@ struct FPCGEditorGraphSchemaAction_NewSettingsElement : public FEdGraphSchemaAct
 };
 
 USTRUCT()
-struct FPCGEditorGraphSchemaAction_NewBlueprintElement : public FEdGraphSchemaAction
+struct FPCGEditorGraphSchemaAction_NewBlueprintElement : public FPCGEditorGraphSchemaActionBase
 {
 	GENERATED_BODY()
 
 	// Inherit the base class's constructors
-	using FEdGraphSchemaAction::FEdGraphSchemaAction;
+	using FPCGEditorGraphSchemaActionBase::FPCGEditorGraphSchemaActionBase;
 	
 	// Simple type info
 	static FName StaticGetTypeId()
@@ -146,16 +160,17 @@ struct FPCGEditorGraphSchemaAction_NewBlueprintElement : public FEdGraphSchemaAc
 	virtual FName GetTypeId() const override { return StaticGetTypeId(); }
 	using FEdGraphSchemaAction::PerformAction; // Prevent hiding of deprecated base class function with FVector2D
 	virtual UEdGraphNode* PerformAction(UEdGraph* ParentGraph, UEdGraphPin* FromPin, const FVector2f& Location, bool bSelectNewNode = true) override;
+	virtual FSlateBrush const* GetPaletteIcon() const override;
 	//~ End FEdGraphSchemaAction Interface
 };
 
 USTRUCT()
-struct FPCGEditorGraphSchemaAction_NewLoadAssetElement : public FEdGraphSchemaAction
+struct FPCGEditorGraphSchemaAction_NewLoadAssetElement : public FPCGEditorGraphSchemaActionBase
 {
 	GENERATED_BODY()
 
 	// Inherit the base class's constructors
-	using FEdGraphSchemaAction::FEdGraphSchemaAction;
+	using FPCGEditorGraphSchemaActionBase::FPCGEditorGraphSchemaActionBase;
 
 	// Simple type info
 	static FName StaticGetTypeId()
@@ -179,12 +194,12 @@ struct FPCGEditorGraphSchemaAction_NewLoadAssetElement : public FEdGraphSchemaAc
 
 /** Creates a new subgraph element and allows choosing whether it can be a subgraph node or a loop node. */
 USTRUCT()
-struct FPCGEditorGraphSchemaAction_NewSubgraphElement : public FEdGraphSchemaAction
+struct FPCGEditorGraphSchemaAction_NewSubgraphElement : public FPCGEditorGraphSchemaActionBase
 {
 	GENERATED_BODY()
 
 	// Inherit the base class's constructors
-	using FEdGraphSchemaAction::FEdGraphSchemaAction;
+	using FPCGEditorGraphSchemaActionBase::FPCGEditorGraphSchemaActionBase;
 
 	// Simple type info
 	static FName StaticGetTypeId()
@@ -212,12 +227,12 @@ struct FPCGEditorGraphSchemaAction_NewSubgraphElement : public FEdGraphSchemaAct
 
 /** Action to add a 'comment' node to the graph */
 USTRUCT()
-struct FPCGEditorGraphSchemaAction_NewComment : public FEdGraphSchemaAction
+struct FPCGEditorGraphSchemaAction_NewComment : public FPCGEditorGraphSchemaActionBase
 {
 	GENERATED_BODY()
 
 	// Inherit the base class's constructors
-	using FEdGraphSchemaAction::FEdGraphSchemaAction;
+	using FPCGEditorGraphSchemaActionBase::FPCGEditorGraphSchemaActionBase;
 
 	// Simple type info
 	static FName StaticGetTypeId()
@@ -234,12 +249,12 @@ struct FPCGEditorGraphSchemaAction_NewComment : public FEdGraphSchemaAction
 };
 
 USTRUCT()
-struct FPCGEditorGraphSchemaAction_NewReroute : public FEdGraphSchemaAction
+struct FPCGEditorGraphSchemaAction_NewReroute : public FPCGEditorGraphSchemaActionBase
 {
 	GENERATED_BODY()
 
 	// Inherit the base class's constructors
-	using FEdGraphSchemaAction::FEdGraphSchemaAction;
+	using FPCGEditorGraphSchemaActionBase::FPCGEditorGraphSchemaActionBase;
 
 	// Simple type info
 	static FName StaticGetTypeId()
@@ -256,12 +271,12 @@ struct FPCGEditorGraphSchemaAction_NewReroute : public FEdGraphSchemaAction
 };
 
 USTRUCT()
-struct FPCGEditorGraphSchemaAction_NewNamedRerouteUsage : public FEdGraphSchemaAction
+struct FPCGEditorGraphSchemaAction_NewNamedRerouteUsage : public FPCGEditorGraphSchemaActionBase
 {
 	GENERATED_BODY()
 
 	// Inherit the base class's constructors
-	using FEdGraphSchemaAction::FEdGraphSchemaAction;
+	using FPCGEditorGraphSchemaActionBase::FPCGEditorGraphSchemaActionBase;
 
 	// Declaration that we want to add an usage of
 	UPROPERTY()
@@ -282,12 +297,12 @@ struct FPCGEditorGraphSchemaAction_NewNamedRerouteUsage : public FEdGraphSchemaA
 };
 
 USTRUCT()
-struct FPCGEditorGraphSchemaAction_NewNamedRerouteDeclaration : public FEdGraphSchemaAction
+struct FPCGEditorGraphSchemaAction_NewNamedRerouteDeclaration : public FPCGEditorGraphSchemaActionBase
 {
 	GENERATED_BODY()
 
 	// Inherit the base class's constructors
-	using FEdGraphSchemaAction::FEdGraphSchemaAction;
+	using FPCGEditorGraphSchemaActionBase::FPCGEditorGraphSchemaActionBase;
 
 	// Simple type info
 	static FName StaticGetTypeId()
