[PCG] Runtime generation debug overlay

#rnx

diff --git a/Engine/Plugins/PCG/Source/PCG/Private/RuntimeGen/PCGRuntimeGenScheduler.cpp b/Engine/Plugins/PCG/Source/PCG/Private/RuntimeGen/PCGRuntimeGenScheduler.cpp
--- a/Engine/Plugins/PCG/Source/PCG/Private/RuntimeGen/PCGRuntimeGenScheduler.cpp
+++ b/Engine/Plugins/PCG/Source/PCG/Private/RuntimeGen/PCGRuntimeGenScheduler.cpp
@@ -21,6 +21,7 @@
 #include "DrawDebugHelpers.h"
 #include "EngineDefines.h" // For UE_ENABLE_DEBUG_DRAWING
 #include "Async/ParallelFor.h"
+#include "Engine/Engine.h"
 #include "Engine/LevelStreaming.h"
 #include "Components/RuntimeVirtualTextureComponent.h"
 #include "Streaming/LevelStreamingDelegates.h"
@@ -34,6 +35,9 @@
 #include "EditorViewportClient.h"
 #endif
 
+// Onscreen debug messages not supported in UE in shipping & test builds.
+#define PCG_RGS_ONSCREENDEBUGMESSAGES (!UE_BUILD_SHIPPING && !UE_BUILD_TEST)
+
 namespace PCGRuntimeGenSchedulerConstants
 {
 	const FString PooledPartitionActorName = TEXT("PCGRuntimePartitionGridActor_POOLED");
@@ -109,6 +113,60 @@ namespace PCGRuntimeGenSchedulerHelpers
 		TEXT("pcg.VirtualTexturePriming.DebugDrawTexturePrimingBounds"),
 		false,
 		TEXT("Draws debug boxes to indicate regions where PCG is requesting virtual texture priming."));
+
+#if PCG_RGS_ONSCREENDEBUGMESSAGES
+	static TAutoConsoleVariable<bool> CVarEnableDebugOverlay(
+		TEXT("pcg.RuntimeGeneration.EnableDebugOverlay"),
+		false,
+		TEXT("Display screen overlay with runtime generation statistics."));
+
+	struct FStatsOverlay
+	{
+		void BeginTick()
+		{
+			*this = {};
+
+			TickStartTime = FPlatformTime::Seconds();
+		}
+
+		void EndTick(const FPCGRuntimeGenScheduler& InRuntimeGenScheduler)
+		{
+			if (CVarEnableDebugOverlay.GetValueOnGameThread())
+			{
+				const double ElapsedMs = (FPlatformTime::Seconds() - TickStartTime) * 1000.0;
+
+				GEngine->AddOnScreenDebugMessage(
+					/*Key=*/-1,
+					/*TimeToDisplay=*/0.0f,
+					FColor::Yellow,
+					FString::Printf(
+						TEXT("PCG Runtime Generation\n")
+						TEXT("    Tick time: %.3fms\n")
+						TEXT("        Generate Calls: %d\n")
+						TEXT("        Cleanup Calls: %d\n")
+						TEXT("    Num Generating Components: %d / %d\n")
+						TEXT("    PA Pool: %d / %d available\n")
+						TEXT("    VT Preloads: %d\n"),
+						ElapsedMs,
+						GenerateCallCounter,
+						CleanupCallCounter,
+						NumGeneratingComponents, FMath::Max(1, CVarNumGeneratingComponentsAtSameTime.GetValueOnGameThread()),
+						InRuntimeGenScheduler.PartitionActorPool.Num(), InRuntimeGenScheduler.PartitionActorPoolSize,
+						VTPreloadCounter)
+				);
+			}
+		}
+
+		int GenerateCallCounter = 0;
+		int CleanupCallCounter = 0;
+		int NumGeneratingComponents = 0;
+		int VTPreloadCounter = 0;
+
+		double TickStartTime = 0.0;
+	};
+
+	static FStatsOverlay Stats;
+#endif // PCG_RGS_ONSCREENDEBUGMESSAGES
 }
 
 FPCGGridDescriptor FPCGRuntimeGenScheduler::FGridGenerationKey::GetGridDescriptor() const
@@ -152,6 +210,10 @@ void FPCGRuntimeGenScheduler::Tick(APCGWorldActor* InPCGWorldActor, double InEnd
 
 	TRACE_CPUPROFILER_EVENT_SCOPE(FPCGRuntimeGenScheduler::Tick);
 
+#if PCG_RGS_ONSCREENDEBUGMESSAGES
+	PCGRuntimeGenSchedulerHelpers::Stats.BeginTick();
+#endif
+
 	TickCVars(InPCGWorldActor);
 
 	GenSources.Empty(GenSources.Num());
@@ -230,6 +292,10 @@ void FPCGRuntimeGenScheduler::Tick(APCGWorldActor* InPCGWorldActor, double InEnd
 		// (and we already have CVarNumGeneratingComponentsAtSameTime to throttle new generations).
 		TickScheduleGeneration(ComponentsToGenerate);
 	}
+
+#if PCG_RGS_ONSCREENDEBUGMESSAGES
+	PCGRuntimeGenSchedulerHelpers::Stats.EndTick(*this);
+#endif
 }
 
 bool FPCGRuntimeGenScheduler::ShouldTick()
@@ -863,6 +929,10 @@ void FPCGRuntimeGenScheduler::TickScheduleGeneration(TMap<FGridGenerationKey, do
 					UE_LOG(LogPCG, Warning, TEXT("[RUNTIMEGEN] GENERATE: '%s' (priority %lf)"), *OriginalComponent->GetOwner()->GetActorNameOrLabel(), Priority);
 				}
 
+#if PCG_RGS_ONSCREENDEBUGMESSAGES
+				++PCGRuntimeGenSchedulerHelpers::Stats.GenerateCallCounter;
+#endif
+
 				// Force to refresh if the component is already generated.
 				OriginalComponent->GenerateLocal(EPCGComponentGenerationTrigger::GenerateAtRuntime, /*bForce=*/true, Grid);
 			}
@@ -958,6 +1028,10 @@ void FPCGRuntimeGenScheduler::TickScheduleGeneration(TMap<FGridGenerationKey, do
 					}
 				}
 
+#if PCG_RGS_ONSCREENDEBUGMESSAGES
+				++PCGRuntimeGenSchedulerHelpers::Stats.GenerateCallCounter;
+#endif
+
 				// Force to refresh if the component is already generated.
 				LocalComponent->GenerateLocal(EPCGComponentGenerationTrigger::GenerateAtRuntime, /*bForce=*/true, LocalComponent->GetGenerationGrid(), Dependencies);
 
@@ -968,6 +1042,10 @@ void FPCGRuntimeGenScheduler::TickScheduleGeneration(TMap<FGridGenerationKey, do
 		GeneratedComponents.Add(Key);
 		++NumGenerating;
 	}
+
+#if PCG_RGS_ONSCREENDEBUGMESSAGES
+	PCGRuntimeGenSchedulerHelpers::Stats.NumGeneratingComponents = NumGenerating;
+#endif
 }
 
 void FPCGRuntimeGenScheduler::TickRequestVirtualTexturePriming(const TSet<IPCGGenSourceBase*>& InGenSources)
@@ -1133,6 +1211,10 @@ void FPCGRuntimeGenScheduler::TickRequestVirtualTexturePriming(const TSet<IPCGGe
 
 						VirtualTextureComponent->RequestPreload(PrimingBounds, MipLevel);
 
+#if PCG_RGS_ONSCREENDEBUGMESSAGES
+						++PCGRuntimeGenSchedulerHelpers::Stats.VTPreloadCounter;
+#endif
+
 						if (PCGRuntimeGenSchedulerHelpers::CVarDebugDrawTexturePrimingBounds.GetValueOnAnyThread())
 						{
 							check(World);
@@ -1354,6 +1436,10 @@ void FPCGRuntimeGenScheduler::CleanupLocalComponent(APCGPartitionActor* Partitio
 		// This performs a CleanupLocalImmediate for us, no need to clean up ourselves.
 		PartitionActor->RemoveGraphInstance(OriginalComponent);
 
+#if PCG_RGS_ONSCREENDEBUGMESSAGES
+		++PCGRuntimeGenSchedulerHelpers::Stats.CleanupCallCounter;
+#endif
+
 		// Remove component mapping.
 		{
 			FWriteScopeLock WriteLock(ActorAndComponentMapping->ComponentToPartitionActorsMapLock);
@@ -1424,7 +1510,7 @@ void FPCGRuntimeGenScheduler::CleanupComponent(const FGridGenerationKey& Generat
 		// If the GeneratedComponent has been lost for some reason, get the PA directly from the ActorAndComponentMapping.
 		PartitionActor = ActorAndComponentMapping->GetPartitionActor(GridDescriptor, GridCoords);
 	}
-	else // If the GeneratedComponent does exist, we can clean it up!.
+	else // If the GeneratedComponent does exist, we can clean it up.
 	{
 		PartitionActor = Cast<APCGPartitionActor>(GeneratedComponent->GetOwner());
 		ensure(PartitionActor || Grid == EPCGHiGenGrid::Unbounded);
@@ -1436,6 +1522,10 @@ void FPCGRuntimeGenScheduler::CleanupComponent(const FGridGenerationKey& Generat
 				UE_LOG(LogPCG, Warning, TEXT("[RUNTIMEGEN] CLEANUP: '%s'"), *GeneratedComponent->GetOwner()->GetActorNameOrLabel());
 			}
 
+#if PCG_RGS_ONSCREENDEBUGMESSAGES
+			++PCGRuntimeGenSchedulerHelpers::Stats.CleanupCallCounter;
+#endif
+
 			GeneratedComponent->CleanupLocalImmediate(/*bRemoveComponents=*/true);
 		}
 	}
@@ -1786,3 +1876,5 @@ void FPCGRuntimeGenScheduler::AddReferencedObjects(FReferenceCollector& Collecto
 	// The level should be keeping the pooled PAs visible to GC. This is just a tentative fix for a crash in GetPartitionActorFromPool(), to understand if the crash is happening because of unreferenced GCed actors.
 	Collector.AddReferencedObjects(PartitionActorPool);
 }
+
+#undef PCG_RGS_ONSCREENDEBUGMESSAGES
diff --git a/Engine/Plugins/PCG/Source/PCG/Public/RuntimeGen/PCGRuntimeGenScheduler.h b/Engine/Plugins/PCG/Source/PCG/Public/RuntimeGen/PCGRuntimeGenScheduler.h
--- a/Engine/Plugins/PCG/Source/PCG/Public/RuntimeGen/PCGRuntimeGenScheduler.h
+++ b/Engine/Plugins/PCG/Source/PCG/Public/RuntimeGen/PCGRuntimeGenScheduler.h
@@ -21,6 +21,7 @@ class UPCGSubsystem;
 class UWorld;
 enum class ELevelStreamingState : uint8;
 enum class ELevelStreamingTargetState : uint8;
+namespace PCGRuntimeGenSchedulerHelpers { struct FStatsOverlay; }
 
 /** Used to inform what virtual textures to prime and on what grids they need to be present. */
 USTRUCT(BlueprintType)
@@ -239,4 +240,6 @@ private:
 	// Local to TickQueueComponentsForGeneration, cached here for efficiency.
 	using FStreamingCompleteQueryKey = TPair<FVector, float>;
 	TMap<FStreamingCompleteQueryKey, bool> CachedStreamingQueryResults;
+
+	friend struct PCGRuntimeGenSchedulerHelpers::FStatsOverlay;
 };
