[PCG] Expose CreateKernels for use.

#rnx

diff --git a/Engine/Plugins/PCG/Source/PCG/Public/PCGSettings.h b/Engine/Plugins/PCG/Source/PCG/Public/PCGSettings.h
--- a/Engine/Plugins/PCG/Source/PCG/Public/PCGSettings.h
+++ b/Engine/Plugins/PCG/Source/PCG/Public/PCGSettings.h
@@ -520,7 +520,7 @@ public:
 	/** [EXPERIMENTAL] Create the GPU kernels and kernel edges that should execute for this node.
 	* Note: This API function is subject to change without deprecation.
 	*/
-	virtual void CreateKernels(FPCGGPUCompilationContext& InOutContext, UObject* InObjectOuter, TArray<UPCGComputeKernel*>& OutKernels, TArray<FPCGKernelEdge>& OutEdges) const {};
+	UE_API virtual void CreateKernels(FPCGGPUCompilationContext& InOutContext, UObject* InObjectOuter, TArray<UPCGComputeKernel*>& OutKernels, TArray<FPCGKernelEdge>& OutEdges) const {};
 
 	/** [EXPERIMENTAL] Whether to display GPU execution option in node settings UI. 
 	* Note: This API function is subject to change without deprecation.
