[PCG] Fix cook error in the data collection DI USH.

#rnx

diff --git a/Engine/Plugins/PCG/Shaders/Private/PCGDataCollectionDataInterface.ush b/Engine/Plugins/PCG/Shaders/Private/PCGDataCollectionDataInterface.ush
--- a/Engine/Plugins/PCG/Shaders/Private/PCGDataCollectionDataInterface.ush
+++ b/Engine/Plugins/PCG/Shaders/Private/PCGDataCollectionDataInterface.ush
@@ -97,7 +97,8 @@ int GetRemappedAttributeId_{DataInterfaceName}(int InGraphAttributeId)
 	if (InGraphAttributeId >= {DataInterfaceName}_FirstRemappedAttributeId)
 	{
 		// Convert from current compute graph attribute ID present in buffer (required if buffer came from a different upstream compute graph).
-		return PCG_NUM_RESERVED_ATTRS + {DataInterfaceName}_GraphToBufferAttributeId[InGraphAttributeId - PCG_NUM_RESERVED_ATTRS];
+		// Note: InGraphAttributeId should always be >= PCG_NUM_RESERVED_ATTRS, but max() added to appease compilers.
+		return PCG_NUM_RESERVED_ATTRS + {DataInterfaceName}_GraphToBufferAttributeId[max(0, InGraphAttributeId - PCG_NUM_RESERVED_ATTRS)];
 	}
 	else
 	{
