[PCG] Made the physics initialization part of volumes & collision wrappers properly setup on the main thread by queuing a task before the end of the current task.

#rnx

diff --git a/Engine/Plugins/PCG/Source/PCG/Private/Data/PCGCollisionWrapperData.cpp b/Engine/Plugins/PCG/Source/PCG/Private/Data/PCGCollisionWrapperData.cpp
--- a/Engine/Plugins/PCG/Source/PCG/Private/Data/PCGCollisionWrapperData.cpp
+++ b/Engine/Plugins/PCG/Source/PCG/Private/Data/PCGCollisionWrapperData.cpp
@@ -266,7 +266,10 @@ FPCGCrc UPCGCollisionWrapperData::ComputeCrc(bool bFullDataCrc) const
 	FArchiveCrc32 Ar;
 
 	AddToCrc(Ar, bFullDataCrc);
-	GetPointData()->AddToCrc(Ar, bFullDataCrc);
+	if (GetPointData())
+	{
+		GetPointData()->AddToCrc(Ar, bFullDataCrc);
+	}
 
 	return FPCGCrc(Ar.GetCrc());
 }
@@ -397,7 +400,46 @@ bool UPCGCollisionWrapperData::SamplePoint(const FTransform& InTransform, const
 UPCGSpatialData* UPCGCollisionWrapperData::CopyInternal(FPCGContext* Context) const
 {
 	UPCGCollisionWrapperData* NewCollisionWrapperData = FPCGContext::NewObject_AnyThread<UPCGCollisionWrapperData>(Context);
-	NewCollisionWrapperData->Initialize(PointData, CollisionSelector, CollisionQueryFlag);
+
+	// Implementation note:
+	// Since the point data contributes to the CRC, it must be set prior to it's normal initialization, which would have been in the Initialize call.
+	NewCollisionWrapperData->PointData = PointData;
+#if WITH_EDITOR
+	NewCollisionWrapperData->RawPointData = RawPointData;
+#endif
+
+	// Initializing physics currently requires being on the GT.
+	// In this case, we'll queue a GT task to do the initialization, which will be done by the end of the caller task here.
+	// @todo_pcg: in future releases (5.7+) we might want to share the wrapper object, which would be simpler & more efficient
+	// assuming we can't change it in any way or fashion.
+	auto FinishWrapperDataInitialization = [WeakCollisionWrapperData = TWeakObjectPtr<UPCGCollisionWrapperData>(NewCollisionWrapperData), WeakPointData = TWeakObjectPtr<const UPCGBasePointData>(PointData), CollisionSelector = this->CollisionSelector, CollisionQueryFlag = this->CollisionQueryFlag](FPCGContext* InContext)
+	{
+		if (UPCGCollisionWrapperData* CollisionWrapperData = WeakCollisionWrapperData.Get())
+		{
+			CollisionWrapperData->Initialize(WeakPointData.Get(), CollisionSelector, CollisionQueryFlag);
+		}
+
+		return true;
+	};
+
+	if (IsInGameThread())
+	{
+		FinishWrapperDataInitialization(Context);
+	}
+	else
+	{
+		FPCGScheduleGenericParams Params(FinishWrapperDataInitialization, Context->ExecutionSource.Get());
+
+		// Physics setup must be run on the GT.
+		Params.bCanExecuteOnlyOnMainThread = true;
+
+		FPCGTaskId MainThreadTaskId = Context->ScheduleGeneric(Params);
+
+		if (MainThreadTaskId != InvalidPCGTaskId)
+		{
+			Context->DynamicDependencies.Add(MainThreadTaskId);
+		}
+	}
 
 	return NewCollisionWrapperData;
 }
diff --git a/Engine/Plugins/PCG/Source/PCG/Private/Data/PCGVolumeData.cpp b/Engine/Plugins/PCG/Source/PCG/Private/Data/PCGVolumeData.cpp
--- a/Engine/Plugins/PCG/Source/PCG/Private/Data/PCGVolumeData.cpp
+++ b/Engine/Plugins/PCG/Source/PCG/Private/Data/PCGVolumeData.cpp
@@ -217,12 +217,6 @@ void UPCGVolumeData::CopyBaseVolumeData(UPCGVolumeData* NewVolumeData) const
 	NewVolumeData->Volume = Volume;
 	NewVolumeData->Bounds = Bounds;
 	NewVolumeData->StrictBounds = StrictBounds;
-
-	// We can't copy the properties if the VolumeBodyInstance is setup for a body. Re-create it using the same method as in Initialize.
-	if (VolumeBodyInstance)
-	{
-		NewVolumeData->SetupVolumeBodyInstance();
-	}
 }
 
 UPCGSpatialData* UPCGVolumeData::CopyInternal(FPCGContext* Context) const
@@ -231,5 +225,41 @@ UPCGSpatialData* UPCGVolumeData::CopyInternal(FPCGContext* Context) const
 
 	CopyBaseVolumeData(NewVolumeData);
 
+	// Initializing physics currently requires being on the GT.
+	// In this case, we'll queue a GT task to do the initialization, which will be done by the end of the caller task here.
+	// @todo_pcg: in future releases (5.7+) we might want to share the body instance instead, which would be simpler & more efficient
+	// assuming we can't change the body instance in question.
+	if (VolumeBodyInstance)
+	{
+		auto FinishVolumeSetup = [WeakVolumeData = TWeakObjectPtr<UPCGVolumeData>(NewVolumeData)](FPCGContext* InContext)
+		{
+			if (UPCGVolumeData* VolumeData = WeakVolumeData.Get())
+			{
+				VolumeData->SetupVolumeBodyInstance();
+			}
+
+			return true;
+		};
+
+		if (IsInGameThread())
+		{
+			FinishVolumeSetup(Context);
+		}
+		else
+		{
+			FPCGScheduleGenericParams Params(FinishVolumeSetup, Context->ExecutionSource.Get());
+
+			// Physics setup must be run on the GT.
+			Params.bCanExecuteOnlyOnMainThread = true;
+
+			FPCGTaskId MainThreadTaskId = Context->ScheduleGeneric(Params);
+
+			if (MainThreadTaskId != InvalidPCGTaskId)
+			{
+				Context->DynamicDependencies.Add(MainThreadTaskId);
+			}
+		}
+	}
+
 	return NewVolumeData;
 }
diff --git a/Engine/Plugins/PCG/Source/PCG/Private/Graph/PCGGraphExecutor.cpp b/Engine/Plugins/PCG/Source/PCG/Private/Graph/PCGGraphExecutor.cpp
--- a/Engine/Plugins/PCG/Source/PCG/Private/Graph/PCGGraphExecutor.cpp
+++ b/Engine/Plugins/PCG/Source/PCG/Private/Graph/PCGGraphExecutor.cpp
@@ -1100,7 +1100,7 @@ void FPCGGraphExecutor::PostTaskExecute(TSharedPtr<FPCGGraphActiveTask> ActiveTa
 		ActiveTask.Context->bInputDataModified = false;
 	}
 
-	if (!bIsDone)
+	if (!bIsDone || !ActiveTask.Context->DynamicDependencies.IsEmpty())
 	{
 		TRACE_CPUPROFILER_EVENT_SCOPE(FPCGGraphExecutor::PostTaskExecute::NotDone);
 		PCGGraphExecutor::TScopeLock ScopeLock(LiveTasksLock);
@@ -1171,7 +1171,7 @@ void FPCGGraphExecutor::PostTaskExecute(TSharedPtr<FPCGGraphActiveTask> ActiveTa
 		}
 	}
 
-	const bool bTaskFullyExecuted = !ActiveTask.bWasCancelled && ActiveTask.Context->CurrentPhase == EPCGExecutionPhase::Done;
+	const bool bTaskFullyExecuted = !ActiveTask.bWasCancelled && ActiveTask.Context->CurrentPhase == EPCGExecutionPhase::Done && ActiveTask.Context->DynamicDependencies.IsEmpty();
 
 #if WITH_EDITOR
 	const bool bTaskWasBypassed = ActiveTask.bIsBypassed;
@@ -1733,7 +1733,8 @@ bool FPCGGraphExecutor::ExecuteScheduling(double EndTime, TSharedPtr<FPCGGraphAc
 					if (!ActiveTask->bIsBypassed)
 #endif
 					{
-						check(ActiveTask->Context->CurrentPhase != EPCGExecutionPhase::Done);
+						check(ActiveTask->Context->DynamicDependencies.IsEmpty());
+						check(ActiveTask->Context->CurrentPhase <= EPCGExecutionPhase::Done);
 						check(MaxNumThreads > 0);
 						ActiveTask->Context->AsyncState.NumAvailableTasks = FMath::Max(1, MaxPCGAsyncThreads);
 						ActiveTask->Context->AsyncState.EndTime = EndTime;
diff --git a/Engine/Plugins/PCG/Source/PCG/Private/PCGElement.cpp b/Engine/Plugins/PCG/Source/PCG/Private/PCGElement.cpp
--- a/Engine/Plugins/PCG/Source/PCG/Private/PCGElement.cpp
+++ b/Engine/Plugins/PCG/Source/PCG/Private/PCGElement.cpp
@@ -164,7 +164,8 @@ namespace PCGElementHelpers
 
 bool IPCGElement::Execute(FPCGContext* Context) const
 {
-	check(Context && Context->AsyncState.NumAvailableTasks != 0 && Context->CurrentPhase < EPCGExecutionPhase::Done);
+	// Note that this might be re-executed even if in the "Done" state because the task might have dynamic dependencies.
+	check(Context && Context->AsyncState.NumAvailableTasks != 0 && Context->CurrentPhase <= EPCGExecutionPhase::Done);
 	check(Context->AsyncState.bIsRunningOnMainThread || !CanExecuteOnlyOnMainThread(Context));
 
 	PCGUtils::FScopedCallOutputDevice OutputDevice;
