[PCG] Improvement for DebugDrawGeneratedCells debug feature - draw all radii instead of only largest.

A lot of the interesting (and pop-prone) generation will happen on smaller grids close to the viewer so it's important to capture all.

#rnx

diff --git a/Engine/Plugins/PCG/Source/PCG/Private/RuntimeGen/PCGRuntimeGenScheduler.cpp b/Engine/Plugins/PCG/Source/PCG/Private/RuntimeGen/PCGRuntimeGenScheduler.cpp
--- a/Engine/Plugins/PCG/Source/PCG/Private/RuntimeGen/PCGRuntimeGenScheduler.cpp
+++ b/Engine/Plugins/PCG/Source/PCG/Private/RuntimeGen/PCGRuntimeGenScheduler.cpp
@@ -512,22 +512,29 @@ void FPCGRuntimeGenScheduler::TickQueueComponentsForGeneration(
 					continue;
 				}
 
+				const FVector GenSourcePosition = GenSourcePositionOptional.GetValue();
+
 #if UE_ENABLE_DEBUG_DRAWING
 				if (bDebugDrawGenerationSources)
 				{
-					if (TOptional<FVector> Pos = GenSource->GetPosition())
+					if (const UObject* GenSourceObject = Cast<UObject>(GenSource))
 					{
-						if (const UObject* GenSourceObject = Cast<UObject>(GenSource))
-						{
-							DrawDebugString(World, *GenSourcePositionOptional, GenSourceObject->GetName(), /*TestBaseActor=*/nullptr, FColor::Red, /*Duration=*/0.0f);
-						}
+						TOptional<FVector> Direction = GenSource->GetDirection();
+
+						DrawDebugString(World, Direction ? (GenSourcePosition + 100.0f * Direction.GetValue()) : GenSourcePosition, GenSourceObject->GetName(), /*TestBaseActor=*/nullptr, FColor::Red, /*Duration=*/0.0f);
+					}
+
+					for (uint32 GridSize : GridSizes)
+					{
+						ensure(PCGHiGenGrid::IsValidGridSize(GridSize));
 
-						DrawDebugSphere(World, *GenSourcePositionOptional, MaxGenerationRadius, 64, FColor::Red, /*bPersistentLines=*/false, /*LifeTime=*/0.0f);
+						const double GenerationRadius = OriginalComponent->GetGenerationRadiusFromGrid(PCGHiGenGrid::GridSizeToGrid(GridSize));
+
+						DrawDebugSphere(World, GenSourcePosition, GenerationRadius, 64, FColor::Red, /*bPersistentLines=*/false, /*LifeTime=*/0.0f);
 					}
 				}
 #endif // UE_ENABLE_DEBUG_DRAWING
 
-				const FVector GenSourcePosition = GenSourcePositionOptional.GetValue();
 				const FBox OriginalComponentBounds = OriginalComponent->GetGridBounds();
 				const bool bIs2DGrid = OriginalComponent->Use2DGrid();
 
