[PCG] Refactor kernel settings which was generating bugs.

#rnx

diff --git a/Engine/Plugins/PCG/Source/PCG/Private/Compute/Elements/PCGCustomHLSLKernel.cpp b/Engine/Plugins/PCG/Source/PCG/Private/Compute/Elements/PCGCustomHLSLKernel.cpp
--- a/Engine/Plugins/PCG/Source/PCG/Private/Compute/Elements/PCGCustomHLSLKernel.cpp
+++ b/Engine/Plugins/PCG/Source/PCG/Private/Compute/Elements/PCGCustomHLSLKernel.cpp
@@ -681,6 +681,7 @@ void UPCGCustomHLSLKernel::GetOutputPins(TArray<FPCGPinPropertiesGPU>& OutPins)
 #if WITH_EDITOR
 void UPCGCustomHLSLKernel::InitEntryPoint()
 {
+	const UPCGSettings* Settings = GetSettings();
 	check(Settings);
 
 	// Record the node title - used to name the kernel and shows up in profiling.
diff --git a/Engine/Plugins/PCG/Source/PCG/Private/Compute/PCGComputeKernel.cpp b/Engine/Plugins/PCG/Source/PCG/Private/Compute/PCGComputeKernel.cpp
--- a/Engine/Plugins/PCG/Source/PCG/Private/Compute/PCGComputeKernel.cpp
+++ b/Engine/Plugins/PCG/Source/PCG/Private/Compute/PCGComputeKernel.cpp
@@ -47,7 +47,7 @@ namespace PCGComputeKernel
 void UPCGComputeKernel::Initialize(const FPCGComputeKernelParams& InParams)
 {
 	ResolvedSettings = InParams.Settings;
-	Settings = ResolvedSettings;
+	SettingsPath = ResolvedSettings;
 
 	if (InParams.bLogDescriptions)
 	{
@@ -70,7 +70,7 @@ const UPCGSettings* UPCGComputeKernel::GetSettings() const
 	if (!ResolvedSettings)
 	{
 		FGCScopeGuard Guard;
-		ResolvedSettings = Settings.Get();
+		ResolvedSettings = SettingsPath.Get();
 	}
 
 	return ResolvedSettings;
@@ -452,6 +452,7 @@ bool UPCGComputeKernel::PerformStaticValidation()
 
 bool UPCGComputeKernel::ValidatePCGNode(TArray<FPCGKernelLogEntry>& OutLogEntries) const
 {
+	const UPCGSettings* Settings = GetSettings();
 	if (!Settings)
 	{
 		return true;
diff --git a/Engine/Plugins/PCG/Source/PCG/Private/Compute/PCGComputeKernel.h b/Engine/Plugins/PCG/Source/PCG/Private/Compute/PCGComputeKernel.h
--- a/Engine/Plugins/PCG/Source/PCG/Private/Compute/PCGComputeKernel.h
+++ b/Engine/Plugins/PCG/Source/PCG/Private/Compute/PCGComputeKernel.h
@@ -193,7 +193,7 @@ protected:
 	int32 KernelIndex = INDEX_NONE;
 
 	UPROPERTY()
-	TSoftObjectPtr<const UPCGSettings> Settings;
+	TSoftObjectPtr<const UPCGSettings> SettingsPath;
 
 	UPROPERTY()
 	bool bInitialized = false;
