[PCG] Removed logged error that could create false positives in the GPU case. Will have a clearer mechanism in the next release.

#rnx

diff --git a/Engine/Plugins/PCG/Source/PCG/Private/Data/PCGLandscapeData.cpp b/Engine/Plugins/PCG/Source/PCG/Private/Data/PCGLandscapeData.cpp
--- a/Engine/Plugins/PCG/Source/PCG/Private/Data/PCGLandscapeData.cpp
+++ b/Engine/Plugins/PCG/Source/PCG/Private/Data/PCGLandscapeData.cpp
@@ -635,7 +635,9 @@ TArray<FPCGTaskId> UPCGLandscapeData::PrepareForSpatialQuery(FPCGContext* InCont
 
 	if (LandscapeCache->SerializationMode == EPCGLandscapeCacheSerializationMode::NeverSerialize && PCGHelpers::IsRuntimeOrPIE())
 	{
-		PCGLog::Landscape::LogLandscapeCacheNotAvailableError(InContext);
+		// Hotfix note: commented this error has it could do false positives, esp. in the GPU use case.
+		// @todo_pcg: use a warning here gated with a cvar instead.
+		//PCGLog::Landscape::LogLandscapeCacheNotAvailableError(InContext);
 		return {};
 	}
 
