[PCG] Fix crash when converting grid size nodes between preconfigured versions while connected to GPU pins

#rnx

diff --git a/Engine/Plugins/PCG/Source/PCG/Private/Elements/PCGHiGenGridSize.cpp b/Engine/Plugins/PCG/Source/PCG/Private/Elements/PCGHiGenGridSize.cpp
--- a/Engine/Plugins/PCG/Source/PCG/Private/Elements/PCGHiGenGridSize.cpp
+++ b/Engine/Plugins/PCG/Source/PCG/Private/Elements/PCGHiGenGridSize.cpp
@@ -177,6 +177,17 @@ bool UPCGHiGenGridSizeSettings::ConvertNode(const FPCGPreconfiguredInfo& Convers
 #endif
 	bShowInputPin = ConversionInfo.PreconfiguredIndex == PCGHiGenGridSizeConstants::ChangeGridSizePreconfiguredIndex;
 
+	if (UPCGNode* Node = Cast<UPCGNode>(GetOuter()))
+	{
+		// Grab the old output pin and break its edges so that they don't populate the new output pin.
+		UPCGPin* RemovedOutputPin = bShowInputPin ? Node->GetOutputPin(PCGPinConstants::DefaultExecutionDependencyLabel) : Node->GetOutputPin(PCGPinConstants::DefaultOutputLabel);
+
+		if (RemovedOutputPin)
+		{
+			RemovedOutputPin->BreakAllEdges();
+		}
+	}
+
 #if WITH_EDITOR
 	// Only broadcast an update if bShowInputPin actually changed.
 	if (bWasShowingInputPin != bShowInputPin)
