[PCG] Experimental: small PCG hash framework which allows hashing specific UObject types based on a registered policy and its external dependencies.  (using archive serialization + xxHash)
- Implement UPCGGraphHashContext / UPCGSettingsHashContext first

#rnx

diff --git a/Engine/Plugins/PCG/Source/PCG/Private/Hash/PCGGraphHashContext.cpp b/Engine/Plugins/PCG/Source/PCG/Private/Hash/PCGGraphHashContext.cpp
new file mode 100644
--- /dev/null
+++ b/Engine/Plugins/PCG/Source/PCG/Private/Hash/PCGGraphHashContext.cpp
@@ -0,0 +1,40 @@
+// Copyright Epic Games, Inc. All Rights Reserved.
+
+#include "Hash/PCGGraphHashContext.h"
+
+#if WITH_EDITOR
+
+#include "PCGEdge.h"
+#include "PCGGraph.h"
+#include "PCGNode.h"
+#include "PCGPin.h"
+
+FPCGGraphHashContext::FPCGGraphHashContext(UPCGGraphInterface* InGraph)
+	: FPCGObjectHashContext(InGraph)
+{
+	check(InGraph);
+	InGraph->OnGraphChangedDelegate.AddRaw(this, &FPCGGraphHashContext::OnGraphChanged);
+	AddHashPolicy(new FPCGObjectHashPolicyPCGNoHash());
+}
+
+FPCGGraphHashContext::~FPCGGraphHashContext()
+{
+	if (UPCGGraphInterface* Graph = GetObject<UPCGGraphInterface>())
+	{
+		Graph->OnGraphChangedDelegate.RemoveAll(this);
+	}
+}
+
+void FPCGGraphHashContext::OnGraphChanged(UPCGGraphInterface* InGraph, EPCGChangeType ChangeType)
+{
+	ensure(InGraph == GetObject<UPCGGraphInterface>());
+	// could probably filter out some changes
+	OnChanged().Execute();
+}
+
+FPCGObjectHashContext* FPCGGraphHashContext::MakeInstance(UObject* InObject)
+{
+	return new FPCGGraphHashContext(CastChecked<UPCGGraphInterface>(InObject));
+}
+
+#endif // WITH_EDITOR
\ No newline at end of file
diff --git a/Engine/Plugins/PCG/Source/PCG/Private/Hash/PCGObjectHash.cpp b/Engine/Plugins/PCG/Source/PCG/Private/Hash/PCGObjectHash.cpp
new file mode 100644
--- /dev/null
+++ b/Engine/Plugins/PCG/Source/PCG/Private/Hash/PCGObjectHash.cpp
@@ -0,0 +1,446 @@
+// Copyright Epic Games, Inc. All Rights Reserved.
+
+#include "Hash/PCGObjectHash.h"
+
+#if WITH_EDITOR
+
+#include "PCGModule.h"
+
+#include "Algo/Copy.h"
+#include "Serialization/ArchiveUObject.h"
+#include "UObject/UObjectAnnotation.h"
+
+namespace PCGObjectHash
+{
+	TAutoConsoleVariable<bool> CVarAlwaysHash(
+		TEXT("pcg.debug.hash.AlwaysHash"),
+		false,
+		TEXT("To debug hashing, set to true and the hash will always be recomputed even if the object and its dependencies didn't change"));
+
+	TAutoConsoleVariable<bool> CVarLogDependencies(
+		TEXT("pcg.debug.hash.LogDependencies"),
+		false,
+		TEXT("To debug hashing, set to true to log dependency hashes"));
+
+	TAutoConsoleVariable<bool> CVarLogSkippedProperties(
+		TEXT("pcg.debug.hash.LogSkippedProperties"),
+		false,
+		TEXT("To debug hashing, set to true to log properties not part of the hash"));
+
+	TAutoConsoleVariable<bool> CVarLogSkippedObjects(
+		TEXT("pcg.debug.hash.LogSkippedObjects"),
+		false,
+		TEXT("To debug hashing, set to true to log objects not part of the hash"));
+
+	// Debug command that outputs the Hash of supported objects and all it's dependency hashes also.
+	const FString GCalculatePCGObjectHashCommandName(TEXT("pcg.debug.hash"));
+	FAutoConsoleCommand CalculatePCGObjectHashCommand(
+		*GCalculatePCGObjectHashCommandName,
+		TEXT("Calculates and outputs a hash for specified object and all its dependencies"),
+		FConsoleCommandWithArgsDelegate::CreateLambda([](const TArray<FString>& InArgs)
+		{
+			if (InArgs.Num() < 1)
+			{
+				UE_LOG(LogPCG, Warning, TEXT("%s : Requires at least 1 argument. First argument should be the object path."), *GCalculatePCGObjectHashCommandName);
+				return;
+			}
+
+			TArray<FString> Args(InArgs);
+			FString ObjectPath;
+			Args.HeapPop(ObjectPath);
+			FSoftObjectPath SoftObjectPath(ObjectPath);
+
+			UObject* ResolvedObject = SoftObjectPath.TryLoad();
+			if (!ResolvedObject)
+			{
+				UE_LOG(LogPCG, Warning, TEXT("%s : Could not find object: '%s'"), *GCalculatePCGObjectHashCommandName, *SoftObjectPath.ToString());
+				return;
+			}
+
+			if (FPCGObjectHash* Hash = FPCGModule::GetPCGModuleChecked().GetConstObjectHashFactory().GetOrCreateObjectHash(ResolvedObject))
+			{
+				Hash->GetHash(/*bInVerbose=*/true);
+			}
+			else
+			{
+				UE_LOG(LogPCG, Warning, TEXT("%s : Unsupported object: '%s'"), *GCalculatePCGObjectHashCommandName, *SoftObjectPath.ToString());
+			}
+		}));
+
+	struct FPCGObjectHashAnnotation
+	{
+		bool IsDefault() const
+		{
+			return Ptr == nullptr;
+		}
+		TSharedPtr<FPCGObjectHash> Ptr;
+	};
+
+	// Global store for existing hash objects
+	FUObjectAnnotationSparse<FPCGObjectHashAnnotation, /*bAutoRemove=*/true> ObjectHashAnnotation;
+
+	// xxHash Property Hash + External dependency gathering archive (referenced objects living in a different UPackage)
+	class FPCGObjectHashArchive : public FArchiveUObject
+	{
+	public:
+		using Super = FArchiveUObject;
+
+		static FPCGObjectHash::FHash Hash(const FPCGObjectHashContext* InContext, TSet<UObject*>& OutObjectReferences, bool bInVerbose, int32 InLevel)
+		{
+			check(InContext)
+			if (!InContext->IsValid())
+			{
+				return FPCGObjectHash::FHash();
+			}
+
+			FPCGObjectHashArchive Ar(InContext, bInVerbose, InLevel);
+			Ar.Compute();
+			OutObjectReferences = Ar.GetObjectReferences();
+			return Ar.GetHash();
+		}
+
+		const TSet<UObject*> GetObjectReferences() const
+		{
+			return ObjectReferences;
+		}
+
+		FPCGObjectHash::FHash GetHash() const
+		{
+			return LocalHash;
+		}
+
+	protected:
+		FPCGObjectHashArchive(const FPCGObjectHashContext* InContext, bool bInVerbose, int32 InLevel)
+			: Context(InContext), bVerbose(bInVerbose), Level(InLevel)
+		{
+			check(Context && Context->IsValid());
+
+			SetIsPersistent(false);
+			SetIsSaving(true);
+			SetIsLoading(false);
+			SetFilterEditorOnly(false);
+
+			ArShouldSkipBulkData = true;
+		}
+
+		void Compute()
+		{
+			ObjectsToVisit.Reset();
+			VisitedObjects.Reset();
+			ObjectReferences.Reset();
+
+			ObjectsToVisit.Add(Context->GetObject());
+			while (ObjectsToVisit.Num() > 0)
+			{
+				UObject* CurObj = ObjectsToVisit.Pop(EAllowShrinking::No);
+				if (!VisitedObjects.Contains(CurObj))
+				{
+					VisitedObjects.Add(CurObj);
+
+					TGuardValue<UObject*> Guard(CurrentObject, CurObj);
+					CurObj->Serialize(*this);
+				}
+			}
+
+			LocalHash = Hasher.Finalize();
+		}
+
+		virtual bool ShouldSkipProperty(const FProperty* InProperty) const override
+		{
+			const bool bShouldHashProperty = Context->ShouldHashProperty(CurrentObject, InProperty);
+			
+			if (!bShouldHashProperty && bVerbose && PCGObjectHash::CVarLogSkippedProperties.GetValueOnAnyThread())
+			{
+				UE_LOG(LogPCG, Log, TEXT("[PCGHash] %s  Property skipped: '%s' (Type: '%s' - Object : '%s')"), FCString::Tab(Level), *InProperty->GetName(), *CurrentObject->GetClass()->GetName(), *CurrentObject->GetPathName());
+			}
+
+			return !bShouldHashProperty;
+		}
+
+		// For now we don't try and do a deep hashing of soft reference but we at least only serialize the path once and use an index for duplicated references
+		virtual FArchive& operator<<(FSoftObjectPath& Value) override 
+		{ 
+			if (int32* IndexPtr = SoftReferenceIndices.Find(Value))
+			{
+				*this << *IndexPtr;
+			}
+			else
+			{
+				int32 VisitingIndex = SoftReferenceIndices.Num();
+				SoftReferenceIndices.Add(Value, VisitingIndex);
+				
+				// Hash index of newly visited reference
+				*this << VisitingIndex;
+				
+				SerializeSoftObjectPath(*this, Value);
+			}
+
+			return *this;
+		}
+
+		virtual FArchive& operator<<(UObject*& InObjectReference) override
+		{
+			if (InObjectReference != nullptr)
+			{
+				int32* VisitedIndex = ObjectIndices.Find(InObjectReference);
+
+				// Haven't visited object yet
+				if (!VisitedIndex)
+				{
+					// Assign an index to this object reference and serialize it
+					int32 VisitingIndex = ObjectIndices.Num();
+					ObjectIndices.Add(InObjectReference, VisitingIndex);
+
+					// Hash index of newly visited reference
+					*this << VisitingIndex;
+
+					// Check if object supports hashing and add it as a reference (reference hashes are used in the final object hash but are stored independently to avoid full recomputations)
+					if (FPCGObjectHash* ObjectReferenceHash = FPCGModule::GetPCGModuleChecked().GetConstObjectHashFactory().GetOrCreateObjectHash(InObjectReference))
+					{
+						VisitedObjects.Add(InObjectReference);
+
+						ObjectReferences.Add(InObjectReference);
+					}
+					else if (InObjectReference->IsIn(Context->GetObject()) && Context->ShouldHashSubObject(InObjectReference))
+					{
+						// Will be visited later
+						ObjectsToVisit.Add(InObjectReference);
+					}
+					else
+					{
+						VisitedObjects.Add(InObjectReference);
+
+						// Serialize object path once, next time this object is encountered its visited index will get serialized
+						FString ObjectReferencePath = InObjectReference->GetPathName();
+						*this << ObjectReferencePath;
+
+						if (CVarLogSkippedObjects.GetValueOnAnyThread() && bVerbose)
+						{
+							UE_LOG(LogPCG, Log, TEXT("[PCGHash] %s  Object skipped: '%s' (Type: '%s')"), FCString::Tab(Level), *InObjectReference->GetPathName(), *InObjectReference->GetClass()->GetName());
+						}
+					}
+				}
+				else
+				{
+					// Hash visisted index
+					*this << *VisitedIndex;
+				}
+			}
+
+			return *this;
+		}
+
+		virtual FArchive& operator<<(FLazyObjectPtr& Value) override 
+		{ 
+			// unsupported ptr type
+			check(0);
+			return *this;
+		}
+
+		virtual FArchive& operator<<(FWeakObjectPtr& Value) override
+		{
+			// unsupported ptr type
+			check(0);
+			return *this;
+		}
+
+		virtual void Serialize(void* Data, int64 Num) override
+		{
+			Hasher.Update((uint8*)Data, Num);
+		}
+
+		using Super::operator<<;
+
+		virtual FArchive& operator<<(class FName& Value) override
+		{
+			FString NameAsString = Value.ToString();
+			*this << NameAsString;
+			return *this;
+		}
+
+		const FPCGObjectHashContext* Context = nullptr;
+
+		FPCGObjectHash::FHash LocalHash;
+		FXxHash64Builder Hasher;
+
+		TSet<UObject*> ObjectReferences;
+		TSet<UObject*> VisitedObjects;
+		TArray<UObject*> ObjectsToVisit;
+		UObject* CurrentObject = nullptr;
+
+		// Mappings to serialize indices when reference is encountered multiple times
+		TMap<FSoftObjectPath, int32> SoftReferenceIndices;
+		TMap<UObject*, int32> ObjectIndices;
+
+		bool bVerbose = false;
+		int32 Level = 0;
+	};
+}
+
+
+FPCGObjectHashContext::FPCGObjectHashContext(UObject* InObject)
+	: Object(InObject)
+{
+
+}
+
+FPCGObjectHashContext::~FPCGObjectHashContext()
+{
+	for (const IPCGObjectHashPolicy* Policy : HashPolicies)
+	{
+		delete Policy;
+	}
+}
+
+bool FPCGObjectHashContext::ShouldHashProperty(const UObject* InObject, const FProperty* InProperty) const
+{
+	if (!ShouldHashTransientProperties() && InProperty->HasAnyPropertyFlags(CPF_Transient))
+	{
+		return false;
+	}
+
+	for (const IPCGObjectHashPolicy* Policy : HashPolicies)
+	{
+		if (!Policy->ShouldHashProperty(InObject, InProperty))
+		{
+			return false;
+		}
+	}
+
+	return ShouldHashPropertyInternal(InObject, InProperty);
+}
+
+bool FPCGObjectHashContext::ShouldHashSubObject(const UObject* InObject) const
+{
+	return InObject && ShouldHashSubObjectInternal(InObject);
+}
+
+bool FPCGObjectHashPolicyPropertyMetaDataFilter::ShouldHashProperty(const UObject* InObject, const FProperty* InProperty) const
+{
+	return InProperty->HasMetaData(*MetaData) == bInclusionFilter;
+}
+
+FPCGObjectHash::FPCGObjectHash(FPCGObjectHashContext* InContext)
+	: Context(InContext)
+{
+	Context->OnChanged().BindRaw(this, &FPCGObjectHash::InvalidateHash);
+}
+
+FPCGObjectHash::~FPCGObjectHash()
+{
+	delete Context;
+}
+
+FPCGObjectHash::FHash FPCGObjectHash::GetHash(bool bInVerbose) const
+{
+	return GetHashInternal(bInVerbose, /*InLevel=*/0);
+}
+
+FPCGObjectHash::FHash FPCGObjectHash::GetHashInternal(bool bInVerbose, uint32 InLevel) const
+{
+	FXxHash64Builder FinalHasher;
+
+	const bool bShouldLog = bInVerbose && (InLevel == 0 || PCGObjectHash::CVarLogDependencies.GetValueOnAnyThread());
+	const double StartHashTime = bShouldLog ? FPlatformTime::Seconds() : 0.0;
+	if (bShouldLog)
+	{
+		UE_LOG(LogPCG, Log, TEXT("[PCGHash] %sBegin Hash: Object: '%s' (Type: '%s')"), FCString::Tab(InLevel), *Context->GetObject()->GetPathName(), *Context->GetObject()->GetClass()->GetName());
+	}
+
+	// Check that our local hash is set (it might have been invalidated or never computed)
+	if (!LocalHash.IsSet() || PCGObjectHash::CVarAlwaysHash.GetValueOnAnyThread())
+	{
+		TSet<UObject*> NewObjectReferences;
+		LocalHash = PCGObjectHash::FPCGObjectHashArchive::Hash(Context, NewObjectReferences, bInVerbose, InLevel);
+		
+		// Sort dependencies for deterministic hashing
+		ObjectReferences.Empty(NewObjectReferences.Num());
+		Algo::Transform(NewObjectReferences, ObjectReferences, [](UObject* InObject) { return TWeakObjectPtr(InObject); });
+		ObjectReferences.Sort([](const TWeakObjectPtr<UObject>& A, const TWeakObjectPtr<UObject>& B) { return A->GetFName().LexicalLess(B->GetFName()); });
+	}
+
+	// Update hasher with our local hash
+	if (ensure(LocalHash.IsSet()))
+	{
+		FPCGObjectHash::FHash& LocalHashValue = LocalHash.GetValue();
+		FinalHasher.Update(&LocalHashValue.Hash, sizeof(LocalHashValue.Hash));
+	}
+
+	// Update hasher with all hashable dependencies
+	for (TWeakObjectPtr<UObject> WeakObjectReference : ObjectReferences)
+	{
+		UObject* ObjectReference = WeakObjectReference.Get();
+		if (!ObjectReference)
+		{
+			continue;
+		}
+
+		if(FPCGObjectHash* ObjectReferenceHashPtr = FPCGModule::GetPCGModuleChecked().GetConstObjectHashFactory().GetOrCreateObjectHash(ObjectReference))
+		{
+			FPCGObjectHash::FHash ObjectReferenceHash = ObjectReferenceHashPtr->GetHashInternal(bInVerbose, InLevel + 1);
+			FinalHasher.Update(&ObjectReferenceHash.Hash, sizeof(ObjectReferenceHash.Hash));
+		}
+	}
+
+	FPCGObjectHash::FHash FinalHash = FinalHasher.Finalize();
+	
+	if (bShouldLog)
+	{
+		const double HashTime = FPlatformTime::Seconds() - StartHashTime;
+		UE_LOG(LogPCG, Log, TEXT("[PCGHash] %sEnd Hash: Object: '%s' -> Hash: '%s' took '%s' seconds"), FCString::Tab(InLevel), *Context->GetObject()->GetPathName(), *LexToString(FinalHash.Hash), *FText::AsNumber(HashTime).ToString());
+	}
+
+	return FinalHash;
+}
+
+void FPCGObjectHashFactory::RegisterObjectHashContextFactory(UClass* InClass, FPCGOnCreateObjectHashContext InCreatePCGObjectHashContext)
+{
+	ensure(!CreateObjectHashContextPerClass.Contains(InClass->GetFName()));
+	CreateObjectHashContextPerClass.FindOrAdd(InClass->GetFName(), InCreatePCGObjectHashContext);
+}
+
+FPCGObjectHash* FPCGObjectHashFactory::GetOrCreateObjectHash(UObject* InObject) const
+{
+	if (!InObject)
+	{
+		return nullptr;
+	}
+
+	// Object annocation maps a UObject* to some structure and handles the object lifespan, if object is destroyed, annotation gets removed.
+	// This is useful to share hash objects across different object hierarchies. Ex: 2 PCG Graphs that reference the same sub graph will use the same sub graph hash object.
+	PCGObjectHash::FPCGObjectHashAnnotation Annotation = PCGObjectHash::ObjectHashAnnotation.GetAnnotation(InObject);
+	if (!Annotation.IsDefault())
+	{
+		return Annotation.Ptr.Get();
+	}
+		
+	UClass* CurrentClass = InObject->GetClass();
+	check(CurrentClass);
+
+	// Find existing factory
+	while (CurrentClass)
+	{
+		if (const FPCGOnCreateObjectHashContext* Delegate = CreateObjectHashContextPerClass.Find(CurrentClass->GetFName()))
+		{
+			FPCGObjectHashContext* HashContext = Delegate->Execute(InObject);
+			if (HashContext)
+			{
+				TSharedPtr<FPCGObjectHash> ObjectHash = MakeShared<FPCGObjectHash>(HashContext);
+				FPCGObjectHash* ObjectHashPtr = ObjectHash.Get();
+
+				// Store annotation
+				PCGObjectHash::FPCGObjectHashAnnotation NewAnnotation;
+				NewAnnotation.Ptr = MoveTemp(ObjectHash);
+				PCGObjectHash::ObjectHashAnnotation.AddAnnotation(InObject, MoveTemp(NewAnnotation));
+
+				return ObjectHashPtr;
+			}
+		}
+
+		CurrentClass = CurrentClass->GetSuperClass();
+	}
+
+	return nullptr;
+}
+
+#endif // WITH_EDITOR
\ No newline at end of file
diff --git a/Engine/Plugins/PCG/Source/PCG/Private/Hash/PCGSettingsHashContext.cpp b/Engine/Plugins/PCG/Source/PCG/Private/Hash/PCGSettingsHashContext.cpp
new file mode 100644
--- /dev/null
+++ b/Engine/Plugins/PCG/Source/PCG/Private/Hash/PCGSettingsHashContext.cpp
@@ -0,0 +1,48 @@
+// Copyright Epic Games, Inc. All Rights Reserved.
+
+#include "Hash/PCGSettingsHashContext.h"
+
+#if WITH_EDITOR
+
+#include "PCGSettings.h"
+
+FPCGSettingsHashContext::FPCGSettingsHashContext(UPCGSettingsInterface* InSettings)
+	: FPCGObjectHashContext(InSettings)
+{
+	InSettings->OnSettingsChangedDelegate.AddRaw(this, &FPCGSettingsHashContext::OnSettingsChanged);
+	AddHashPolicy(new FPCGSettingsHashPolicy());
+}
+
+FPCGSettingsHashContext::~FPCGSettingsHashContext()
+{
+	if (UPCGSettingsInterface* Settings = GetObject<UPCGSettingsInterface>())
+	{
+		Settings->OnSettingsChangedDelegate.RemoveAll(this);
+	}
+}
+
+FPCGObjectHashContext* FPCGSettingsHashContext::MakeInstance(UObject* InObject)
+{
+	return new FPCGSettingsHashContext(CastChecked<UPCGSettings>(InObject));
+}
+
+void FPCGSettingsHashContext::OnSettingsChanged(UPCGSettings* InSettings, EPCGChangeType ChangeType)
+{
+	ensure(InSettings == GetObject<UPCGSettings>());
+	// could probably filter out some changes
+	OnChanged().Execute();
+}
+
+bool FPCGSettingsHashPolicy::ShouldHashProperty(const UObject* InObject, const FProperty* InProperty) const
+{
+	// Omit CRC'ing data collections here as it is very slow. Rely instead on UID/ObjectPath.
+	const FStructProperty* StructProperty = CastField<FStructProperty>(InProperty);
+	if (StructProperty && StructProperty->Struct && StructProperty->Struct->IsChildOf(FPCGDataCollection::StaticStruct()))
+	{
+		return false;
+	}
+
+	return Super::ShouldHashProperty(InObject, InProperty);
+}
+
+#endif // WITH_EDITOR
\ No newline at end of file
diff --git a/Engine/Plugins/PCG/Source/PCG/Private/PCGModule.cpp b/Engine/Plugins/PCG/Source/PCG/Private/PCGModule.cpp
--- a/Engine/Plugins/PCG/Source/PCG/Private/PCGModule.cpp
+++ b/Engine/Plugins/PCG/Source/PCG/Private/PCGModule.cpp
@@ -4,6 +4,7 @@
 
 #include "PCGContext.h"
 #include "PCGElement.h"
+#include "PCGGraph.h"
 #include "Compute/PCGComputeKernel.h"
 #include "Data/PCGBasePointData.h"
 #include "Data/PCGPointData.h"
@@ -15,6 +16,8 @@
 
 #if WITH_EDITOR
 #include "Elements/PCGDifferenceElement.h"
+#include "Hash/PCGGraphHashContext.h"
+#include "Hash/PCGSettingsHashContext.h"
 #include "Tests/Determinism/PCGDeterminismNativeTests.h"
 #include "Tests/Determinism/PCGDifferenceDeterminismTest.h"
 
@@ -71,6 +74,11 @@ void FPCGModule::StartupModule()
 	AttributeAccessorFactory.RegisterMethods<UPCGPointData>(UPCGPointData::GetPointAccessorMethods());
 	AttributeAccessorFactory.RegisterMethods<UPCGSplineData>(UPCGSplineData::GetSplineAccessorMethods());
 
+#if WITH_EDITOR
+	ObjectHashFactory.RegisterObjectHashContextFactory(UPCGGraphInterface::StaticClass(), FPCGOnCreateObjectHashContext::CreateStatic(&FPCGGraphHashContext::MakeInstance));
+	ObjectHashFactory.RegisterObjectHashContextFactory(UPCGSettingsInterface::StaticClass(), FPCGOnCreateObjectHashContext::CreateStatic(&FPCGSettingsHashContext::MakeInstance));
+#endif
+
 	// Register onto the PreExit, because we need the class to be still valid to remove them from the mapping
 	FCoreDelegates::OnPreExit.AddRaw(this, &FPCGModule::PreExit);
 
diff --git a/Engine/Plugins/PCG/Source/PCG/Private/PCGSettings.cpp b/Engine/Plugins/PCG/Source/PCG/Private/PCGSettings.cpp
--- a/Engine/Plugins/PCG/Source/PCG/Private/PCGSettings.cpp
+++ b/Engine/Plugins/PCG/Source/PCG/Private/PCGSettings.cpp
@@ -31,6 +31,7 @@
 
 #if WITH_EDITOR
 #include "Editor.h"
+#include "Hash/PCGSettingsHashContext.h"
 #endif
 
 #define LOCTEXT_NAMESPACE "PCGSettings"
@@ -63,7 +64,7 @@ public:
 		const bool bSkipUID = !OriginalSettings->GetPackage()->HasAnyPackageFlags(PKG_TransientFlags) && OriginalSettings->GetPackage() != GetTransientPackage();
 		const uint32 SettingsPathOrUIDCrc = bSkipUID ? FCrc::StrCrc32(*OriginalSettings->GetPathName()) : FCrc::TypeCrc32(OriginalSettings->GetStableUID());
 				
-		FPCGSettingsObjectCrc Ar;
+		FPCGSettingsObjectCrc Ar(PCGSettings);
 		uint32 ArCrc = Ar.Crc32(const_cast<UPCGSettings*>(PCGSettings), SettingsPathOrUIDCrc);
 		return FPCGCrc(ArCrc);
 	}
@@ -71,27 +72,18 @@ public:
 #if WITH_EDITOR
 	virtual bool ShouldSkipProperty(const FProperty* InProperty) const override
 	{
-		// Omit CRC'ing data collections here as it is very slow. Rely instead on UID/ObjectPath.
-		const FStructProperty* StructProperty = CastField<FStructProperty>(InProperty);
-		if (StructProperty && StructProperty->Struct && StructProperty->Struct->IsChildOf(FPCGDataCollection::StaticStruct()))
-		{
-			return true;
-		}
-
-		const bool bSkip = InProperty && (
-			InProperty->GetFName() == GET_MEMBER_NAME_CHECKED(UPCGSettings, UID)
-			|| InProperty->GetFName() == GET_MEMBER_NAME_CHECKED(UPCGSettings, DebugSettings)
-			|| InProperty->GetFName() == GET_MEMBER_NAME_CHECKED(UPCGSettings, DeterminismSettings)
-			|| InProperty->GetFName() == GET_MEMBER_NAME_CHECKED(UPCGSettings, bDebug)
-			|| InProperty->GetFName() == GET_MEMBER_NAME_CHECKED(UPCGSettings, Category)
-			|| InProperty->GetFName() == GET_MEMBER_NAME_CHECKED(UPCGSettings, Description)
-			|| InProperty->GetFName() == GET_MEMBER_NAME_CHECKED(UPCGSettings, bExposeToLibrary)
-			|| InProperty->GetFName() == GET_MEMBER_NAME_CHECKED(UPCGSettings, CachedOverridableParams)
-			);
-
-		return bSkip;
+		static FPCGSettingsHashPolicy Policy;
+		return !Policy.ShouldHashProperty(PCGSettings, InProperty);
 	}
 #endif // WITH_EDITOR
+
+private: 
+	FPCGSettingsObjectCrc(const UPCGSettings* InPCGSettings)
+		: PCGSettings(InPCGSettings)
+	{
+	}
+
+	const UPCGSettings* PCGSettings = nullptr;
 };
 
 FString FPCGSettingsOverridableParam::GetPropertyPath() const
diff --git a/Engine/Plugins/PCG/Source/PCG/Public/Hash/PCGGraphHashContext.h b/Engine/Plugins/PCG/Source/PCG/Public/Hash/PCGGraphHashContext.h
new file mode 100644
--- /dev/null
+++ b/Engine/Plugins/PCG/Source/PCG/Public/Hash/PCGGraphHashContext.h
@@ -0,0 +1,24 @@
+// Copyright Epic Games, Inc. All Rights Reserved.
+
+#pragma once
+
+#if WITH_EDITOR
+
+#include "Hash/PCGObjectHash.h"
+
+class UPCGGraphInterface;
+enum class EPCGChangeType : uint32;
+
+class FPCGGraphHashContext : public FPCGObjectHashContext
+{
+public:
+	explicit FPCGGraphHashContext(UPCGGraphInterface* InGraph);
+	virtual ~FPCGGraphHashContext();
+	
+	static FPCGObjectHashContext* MakeInstance(UObject* InObject);
+
+protected:
+	void OnGraphChanged(UPCGGraphInterface* InGraph, EPCGChangeType ChangeType);
+};
+
+#endif // WITH_EDITOR
\ No newline at end of file
diff --git a/Engine/Plugins/PCG/Source/PCG/Public/Hash/PCGObjectHash.h b/Engine/Plugins/PCG/Source/PCG/Public/Hash/PCGObjectHash.h
new file mode 100644
--- /dev/null
+++ b/Engine/Plugins/PCG/Source/PCG/Public/Hash/PCGObjectHash.h
@@ -0,0 +1,130 @@
+// Copyright Epic Games, Inc. All Rights Reserved.
+
+#pragma once
+
+#if WITH_EDITOR
+
+#include "Delegates/Delegate.h"
+#include "Hash/xxhash.h"
+#include "Misc/Optional.h"
+#include "UObject/WeakObjectPtr.h"
+
+DECLARE_DELEGATE(FPCGObjectHashContextChanged);
+
+/**
+* IPCGObjectHashPolicy is an interface that allows defining specific rules for property hashing.
+*/
+class IPCGObjectHashPolicy
+{
+public:
+	virtual ~IPCGObjectHashPolicy() = default;
+
+	virtual bool ShouldHashProperty(const UObject* InObject, const FProperty* InProperty) const = 0;
+};
+
+/**
+ * Base class to be specialized for specific UObject sub-classes that support PCG object hashing
+ * 
+ * Context is responsible for triggering event (OnChanged) that would invalidate the FPCGObjectHash and also provide logic for what properties need to be hashed.
+ */
+class FPCGObjectHashContext
+{
+public:
+	explicit FPCGObjectHashContext(UObject* InObject);
+	virtual ~FPCGObjectHashContext();
+
+	bool ShouldHashProperty(const UObject* InObject, const FProperty* InProperty) const;
+	bool ShouldHashSubObject(const UObject* InObject) const;
+
+	template <class T = UObject>
+	T* GetObject() const { return Cast<T>(Object.Get()); }
+
+	bool IsValid() const { return !!GetObject(); }
+
+	FPCGObjectHashContextChanged& OnChanged() { return Changed; }
+
+protected:
+	virtual bool ShouldHashPropertyInternal(const UObject* InObject, const FProperty* InProperty) const { return true; }
+	virtual bool ShouldHashSubObjectInternal(const UObject* InObject) const { return true; }
+	
+	void AddHashPolicy(const IPCGObjectHashPolicy* InPolicy) { HashPolicies.Add(InPolicy); }
+	virtual bool ShouldHashTransientProperties() const { return false; }
+
+private:
+	TWeakObjectPtr<UObject> Object;
+	TArray<const IPCGObjectHashPolicy*> HashPolicies;
+	FPCGObjectHashContextChanged Changed;
+};
+
+/**
+* Simple hash policy which includes or excludes properties based on a UPROPERTY metadata tag.
+*/
+class FPCGObjectHashPolicyPropertyMetaDataFilter : public IPCGObjectHashPolicy
+{
+public:
+	FPCGObjectHashPolicyPropertyMetaDataFilter(const FString& InMetaData, bool bInInclusionFilter)
+		: MetaData(InMetaData), bInclusionFilter(bInInclusionFilter)
+	{
+	}
+
+	virtual bool ShouldHashProperty(const UObject* InObject, const FProperty* InProperty) const override;
+
+private:
+	FString MetaData;
+	bool bInclusionFilter;
+};
+
+/** 
+* Simple hash policy which excludes every property with the 'PCGNoHash' metadata tag.
+*/
+class FPCGObjectHashPolicyPCGNoHash : public FPCGObjectHashPolicyPropertyMetaDataFilter
+{
+public:
+	FPCGObjectHashPolicyPCGNoHash()
+		: FPCGObjectHashPolicyPropertyMetaDataFilter(TEXT("PCGNoHash"), /*bInInclusionFilter=*/false)
+	{
+	}
+};
+
+/**
+* Class used with a FPCGObjectHashContext specialization to hash an Object and its External Dependencies.
+*/
+class FPCGObjectHash final
+{
+public:
+	using FHash = FXxHash64;
+
+	explicit FPCGObjectHash(FPCGObjectHashContext* InContext);
+	virtual ~FPCGObjectHash();
+
+	FHash GetHash(bool bVerbose = false) const;
+
+private:
+	FHash GetHashInternal(bool bVerbose, uint32 Level) const;
+	void InvalidateHash() { LocalHash.Reset(); }
+	
+	mutable TOptional<FHash> LocalHash;
+	mutable TArray<TWeakObjectPtr<UObject>> ObjectReferences;
+
+	FPCGObjectHashContext* Context = nullptr;
+};
+
+/**
+* PCG Object hash context factory delegate.
+*/ 
+DECLARE_DELEGATE_RetVal_OneParam(FPCGObjectHashContext*, FPCGOnCreateObjectHashContext, UObject*);
+
+/**
+* Factory accessed through FPCGModule to register new supported Hash context types.
+*/
+class FPCGObjectHashFactory final
+{
+public:
+	PCG_API void RegisterObjectHashContextFactory(UClass* InClass, FPCGOnCreateObjectHashContext InCreatePCGObjectHashContext);
+	PCG_API FPCGObjectHash* GetOrCreateObjectHash(UObject* InObject) const;
+
+private:
+	TMap<FName, FPCGOnCreateObjectHashContext> CreateObjectHashContextPerClass;
+};
+
+#endif // WITH_EDITOR
\ No newline at end of file
diff --git a/Engine/Plugins/PCG/Source/PCG/Public/Hash/PCGSettingsHashContext.h b/Engine/Plugins/PCG/Source/PCG/Public/Hash/PCGSettingsHashContext.h
new file mode 100644
--- /dev/null
+++ b/Engine/Plugins/PCG/Source/PCG/Public/Hash/PCGSettingsHashContext.h
@@ -0,0 +1,35 @@
+// Copyright Epic Games, Inc. All Rights Reserved.
+
+#pragma once
+
+#if WITH_EDITOR
+
+#include "Hash/PCGObjectHash.h"
+
+class UPCGSettings;
+class UPCGSettingsInterface;
+enum class EPCGChangeType : uint32;
+
+class FPCGSettingsHashContext : public FPCGObjectHashContext
+{
+public:
+	FPCGSettingsHashContext(UPCGSettingsInterface* InSettings);
+	virtual ~FPCGSettingsHashContext();
+
+	static FPCGObjectHashContext* MakeInstance(UObject* InObject);
+
+protected:
+	void OnSettingsChanged(UPCGSettings* InSettings, EPCGChangeType ChangeType);
+};
+
+class FPCGSettingsHashPolicy : public FPCGObjectHashPolicyPCGNoHash
+{
+public:
+	using Super = FPCGObjectHashPolicyPCGNoHash;
+
+	FPCGSettingsHashPolicy() = default;
+
+	virtual bool ShouldHashProperty(const UObject* InObject, const FProperty* InProperty) const override;
+};
+
+#endif // WITH_EDITOR
\ No newline at end of file
diff --git a/Engine/Plugins/PCG/Source/PCG/Public/PCGEdge.h b/Engine/Plugins/PCG/Source/PCG/Public/PCGEdge.h
--- a/Engine/Plugins/PCG/Source/PCG/Public/PCGEdge.h
+++ b/Engine/Plugins/PCG/Source/PCG/Public/PCGEdge.h
@@ -21,16 +21,16 @@ public:
 	UE_API virtual void PostLoad() override;
 	// ~End UObject interface
 
-	UPROPERTY()
+	UPROPERTY(meta=(PCGNoHash))
 	FName InboundLabel_DEPRECATED = NAME_None;
 
-	UPROPERTY()
+	UPROPERTY(meta = (PCGNoHash))
 	TObjectPtr<UPCGNode> InboundNode_DEPRECATED;
 
-	UPROPERTY()
+	UPROPERTY(meta = (PCGNoHash))
 	FName OutboundLabel_DEPRECATED = NAME_None;
 
-	UPROPERTY()
+	UPROPERTY(meta = (PCGNoHash))
 	TObjectPtr<UPCGNode> OutboundNode_DEPRECATED;
 
 	/** Pin at upstream end of edge. */
diff --git a/Engine/Plugins/PCG/Source/PCG/Public/PCGGraph.h b/Engine/Plugins/PCG/Source/PCG/Public/PCGGraph.h
--- a/Engine/Plugins/PCG/Source/PCG/Public/PCGGraph.h
+++ b/Engine/Plugins/PCG/Source/PCG/Public/PCGGraph.h
@@ -155,24 +155,24 @@ protected:
 public:
 
 #if WITH_EDITORONLY_DATA
-	UPROPERTY(BlueprintReadWrite, EditAnywhere, Category = AssetInfo, AssetRegistrySearchable, meta = (EditCondition = "IsExportToLibraryEnabled", EditConditionHides))
+	UPROPERTY(BlueprintReadWrite, EditAnywhere, Category = AssetInfo, AssetRegistrySearchable, meta = (EditCondition = "IsExportToLibraryEnabled", EditConditionHides, PCGNoHash))
 	bool bExposeToLibrary = false;
 
-	UPROPERTY(BlueprintReadWrite, EditAnywhere, Category = AssetInfo, AssetRegistrySearchable, meta = (EditCondition = "AreOverridesEnabled", EditConditionHides))
+	UPROPERTY(BlueprintReadWrite, EditAnywhere, Category = AssetInfo, AssetRegistrySearchable, meta = (EditCondition = "AreOverridesEnabled", EditConditionHides, PCGNoHash))
 	bool bOverrideTitle = false;
 
-	UPROPERTY(BlueprintReadWrite, EditAnywhere, Category = AssetInfo, AssetRegistrySearchable, meta = (EditCondition = "IsTemplatePropertyEnabled", EditConditionHides))
+	UPROPERTY(BlueprintReadWrite, EditAnywhere, Category = AssetInfo, AssetRegistrySearchable, meta = (EditCondition = "IsTemplatePropertyEnabled", EditConditionHides, PCGNoHash))
 	bool bIsTemplate = false;
 
 	/** Override of the title for the subgraph node for this graph. */
-	UPROPERTY(BlueprintReadWrite, EditAnywhere, Category = AssetInfo, AssetRegistrySearchable, meta = (EditCondition = "AreOverridesEnabled && bOverrideTitle", EditConditionHides))
+	UPROPERTY(BlueprintReadWrite, EditAnywhere, Category = AssetInfo, AssetRegistrySearchable, meta = (EditCondition = "AreOverridesEnabled && bOverrideTitle", EditConditionHides, PCGNoHash))
 	FText Title;
 
-	UPROPERTY(BlueprintReadWrite, EditAnywhere, Category = AssetInfo, meta = (EditCondition = "AreOverridesEnabled", EditConditionHides))
+	UPROPERTY(BlueprintReadWrite, EditAnywhere, Category = AssetInfo, meta = (EditCondition = "AreOverridesEnabled", EditConditionHides, PCGNoHash))
 	bool bOverrideColor = false;
 
 	/** Override of the color for the subgraph node for this graph. */
-	UPROPERTY(BlueprintReadWrite, EditAnywhere, Category = AssetInfo, meta = (EditCondition = "AreOverridesEnabled && bOverrideColor", EditConditionHides))
+	UPROPERTY(BlueprintReadWrite, EditAnywhere, Category = AssetInfo, meta = (EditCondition = "AreOverridesEnabled && bOverrideColor", EditConditionHides, PCGNoHash))
 	FLinearColor Color = FLinearColor::White;
 #endif
 
@@ -304,17 +304,17 @@ public:
 	UE_API void ReturnComputeGraphInstanceToPool(const FComputeGraphInstanceKey& InKey, TSharedPtr<FComputeGraphInstance> InInstance) const;
 
 #if WITH_EDITORONLY_DATA
-	UPROPERTY(BlueprintReadWrite, EditAnywhere, Category = AssetInfo, AssetRegistrySearchable)
+	UPROPERTY(BlueprintReadWrite, EditAnywhere, Category = AssetInfo, AssetRegistrySearchable, meta=(PCGNoHash))
 	FText Category;
 
-	UPROPERTY(BlueprintReadWrite, EditAnywhere, Category = AssetInfo, AssetRegistrySearchable)
+	UPROPERTY(BlueprintReadWrite, EditAnywhere, Category = AssetInfo, AssetRegistrySearchable, meta=(PCGNoHash))
 	FText Description;
 
 	/** Marks the graph to be not refreshed automatically when the landscape changes, even if it is used. */
-	UPROPERTY(BlueprintReadWrite, EditAnywhere, Category = "Settings|Advanced")
+	UPROPERTY(BlueprintReadWrite, EditAnywhere, Category = "Settings|Advanced", meta=(PCGNoHash))
 	bool bIgnoreLandscapeTracking = false;
 
-	UPROPERTY(EditAnywhere, Category = Customization)
+	UPROPERTY(EditAnywhere, Category = Customization, meta=(PCGNoHash))
 	FPCGGraphEditorCustomization GraphCustomization;
 #endif
 
@@ -486,19 +486,19 @@ protected:
 	TArray<TObjectPtr<UPCGNode>> Nodes;
 
 	// Add input/output nodes
-	UPROPERTY(BlueprintReadOnly, VisibleAnywhere, Category = Graph, meta = (NoResetToDefault))
+	UPROPERTY(BlueprintReadOnly, VisibleAnywhere, Category = Graph, meta = (NoResetToDefault, PCGNoHash))
 	TObjectPtr<UPCGNode> InputNode;
 
-	UPROPERTY(BlueprintReadOnly, VisibleAnywhere, Category = Graph, meta = (NoResetToDefault))
+	UPROPERTY(BlueprintReadOnly, VisibleAnywhere, Category = Graph, meta = (NoResetToDefault, PCGNoHash))
 	TObjectPtr<UPCGNode> OutputNode;
 
 #if WITH_EDITORONLY_DATA
 	// Extra data to hold information that is useful only in editor
-	UPROPERTY()
+	UPROPERTY(meta=(PCGNoHash))
 	TArray<TObjectPtr<UObject>> ExtraEditorNodes;
 
 	// Extra data to hold information for comments
-	UPROPERTY()
+	UPROPERTY(meta=(PCGNoHash))
 	TArray<FPCGGraphCommentNodeData> CommentNodes;
 
 	// Editor graph created from PCG Editor but owned by this, reference is collected using AddReferencedObjects
@@ -715,18 +715,18 @@ private:
 	TWeakObjectPtr<UPCGGraphInterface> PreGraphCache = nullptr;
 
 public:
-	UPROPERTY(BlueprintReadWrite, EditAnywhere, Category = AssetInfo, AssetRegistrySearchable, meta = (EditCondition = "AreOverridesEnabled", EditConditionHides, DisplayAfter = ColorOverride))
+	UPROPERTY(BlueprintReadWrite, EditAnywhere, Category = AssetInfo, AssetRegistrySearchable, meta = (EditCondition = "AreOverridesEnabled", EditConditionHides, DisplayAfter = ColorOverride, PCGNoHash))
 	bool bOverrideDescription = false;
 
 	/** Can override the description of this instance. */
-	UPROPERTY(BlueprintReadWrite, EditAnywhere, Category = AssetInfo, AssetRegistrySearchable, meta = (EditCondition = "AreOverridesEnabled && bOverrideDescription", EditConditionHides, DisplayAfter = ColorOverride))
+	UPROPERTY(BlueprintReadWrite, EditAnywhere, Category = AssetInfo, AssetRegistrySearchable, meta = (EditCondition = "AreOverridesEnabled && bOverrideDescription", EditConditionHides, DisplayAfter = ColorOverride, PCGNoHash))
 	FText Description;
 
-	UPROPERTY(BlueprintReadWrite, EditAnywhere, Category = AssetInfo, AssetRegistrySearchable, meta = (EditCondition = "AreOverridesEnabled", EditConditionHides, DisplayAfter = ColorOverride))
+	UPROPERTY(BlueprintReadWrite, EditAnywhere, Category = AssetInfo, AssetRegistrySearchable, meta = (EditCondition = "AreOverridesEnabled", EditConditionHides, DisplayAfter = ColorOverride, PCGNoHash))
 	bool bOverrideCategory = false;
 
 	/** Can override the category of this instance. */
-	UPROPERTY(BlueprintReadWrite, EditAnywhere, Category = AssetInfo, AssetRegistrySearchable, meta = (EditCondition = "AreOverridesEnabled && bOverrideCategory", EditConditionHides, DisplayAfter = ColorOverride))
+	UPROPERTY(BlueprintReadWrite, EditAnywhere, Category = AssetInfo, AssetRegistrySearchable, meta = (EditCondition = "AreOverridesEnabled && bOverrideCategory", EditConditionHides, DisplayAfter = ColorOverride, PCGNoHash))
 	FText Category;
 #endif // WITH_EDITORONLY_DATA
 };
diff --git a/Engine/Plugins/PCG/Source/PCG/Public/PCGModule.h b/Engine/Plugins/PCG/Source/PCG/Public/PCGModule.h
--- a/Engine/Plugins/PCG/Source/PCG/Public/PCGModule.h
+++ b/Engine/Plugins/PCG/Source/PCG/Public/PCGModule.h
@@ -4,6 +4,7 @@
 
 #include "PCGDataVisualizationRegistry.h"
 #include "Data/PCGGetDataFunctionRegistry.h"
+#include "Hash/PCGObjectHash.h"
 #include "Metadata/Accessors/PCGAttributeAccessorFactory.h"
 #include "Utils/PCGLogErrors.h"
 
@@ -48,6 +49,11 @@ public:
 	static const FPCGAttributeAccessorFactory& GetConstAttributeAccessorFactory() { return GetPCGModuleChecked().AttributeAccessorFactory; }
 	static FPCGAttributeAccessorFactory& GetMutableAttributeAccessorFactory() { return GetPCGModuleChecked().AttributeAccessorFactory; }
 	
+#if WITH_EDITOR
+	static const FPCGObjectHashFactory& GetConstObjectHashFactory() { return GetPCGModuleChecked().ObjectHashFactory; }
+	static FPCGObjectHashFactory& GetMutableObjectHashFactory() { return GetPCGModuleChecked().ObjectHashFactory; }
+#endif
+
 	PCG_API static bool IsPCGModuleLoaded();
 	
 	PCG_API void ExecuteNextTick(TFunction<void()> TickFunction);
@@ -67,6 +73,10 @@ private:
 	FPCGGetDataFunctionRegistry GetDataFunctionRegistry;
 	FPCGAttributeAccessorFactory AttributeAccessorFactory;
 
+#if WITH_EDITOR
+	FPCGObjectHashFactory ObjectHashFactory;
+#endif
+
 	FTSTicker::FDelegateHandle	TickDelegateHandle;
 
 	FCriticalSection ExecuteNextTickLock;
diff --git a/Engine/Plugins/PCG/Source/PCG/Public/PCGNode.h b/Engine/Plugins/PCG/Source/PCG/Public/PCGNode.h
--- a/Engine/Plugins/PCG/Source/PCG/Public/PCGNode.h
+++ b/Engine/Plugins/PCG/Source/PCG/Public/PCGNode.h
@@ -173,15 +173,15 @@ public:
 #endif
 
 #if WITH_EDITORONLY_DATA
-	UPROPERTY()
+	UPROPERTY(meta = (PCGNoHash))
 	TObjectPtr<UPCGSettings> DefaultSettings_DEPRECATED; 
 #endif
 
-	UPROPERTY(BlueprintReadWrite, EditAnywhere, Category = Node)
+	UPROPERTY(BlueprintReadWrite, EditAnywhere, Category = Node, meta = (PCGNoHash))
 	FName NodeTitle = NAME_None;
 
 #if WITH_EDITORONLY_DATA
-	UPROPERTY(BlueprintReadWrite, EditAnywhere, Category = Node)
+	UPROPERTY(BlueprintReadWrite, EditAnywhere, Category = Node, meta = (PCGNoHash))
 	FLinearColor NodeTitleColor = FLinearColor::White;
 #endif // WITH_EDITORONLY_DATA
 
@@ -190,19 +190,19 @@ public:
 #endif
 
 #if WITH_EDITORONLY_DATA
-	UPROPERTY()
+	UPROPERTY(meta = (PCGNoHash))
 	int32 PositionX;
 
-	UPROPERTY()
+	UPROPERTY(meta = (PCGNoHash))
 	int32 PositionY;
 
-	UPROPERTY()
+	UPROPERTY(meta = (PCGNoHash))
 	FString NodeComment;
 
-	UPROPERTY()
+	UPROPERTY(meta = (PCGNoHash))
 	uint8 bCommentBubblePinned : 1;
 
-	UPROPERTY()
+	UPROPERTY(meta = (PCGNoHash))
 	uint8 bCommentBubbleVisible : 1;
 #endif // WITH_EDITORONLY_DATA
 
@@ -229,13 +229,13 @@ protected:
 	UPROPERTY(BlueprintReadOnly, VisibleAnywhere, Category = Node)
 	TObjectPtr<UPCGSettingsInterface> SettingsInterface;
 
-	UPROPERTY()
+	UPROPERTY(meta = (PCGNoHash))
 	TArray<TObjectPtr<UPCGNode>> OutboundNodes_DEPRECATED;
 
-	UPROPERTY(TextExportTransient)
+	UPROPERTY(TextExportTransient, meta = (PCGNoHash))
 	TArray<TObjectPtr<UPCGEdge>> InboundEdges_DEPRECATED;
 
-	UPROPERTY(TextExportTransient)
+	UPROPERTY(TextExportTransient, meta = (PCGNoHash))
 	TArray<TObjectPtr<UPCGEdge>> OutboundEdges_DEPRECATED;
 
 	UPROPERTY(BlueprintReadOnly, VisibleAnywhere, Category = Node)
diff --git a/Engine/Plugins/PCG/Source/PCG/Public/PCGPin.h b/Engine/Plugins/PCG/Source/PCG/Public/PCGPin.h
--- a/Engine/Plugins/PCG/Source/PCG/Public/PCGPin.h
+++ b/Engine/Plugins/PCG/Source/PCG/Public/PCGPin.h
@@ -250,7 +250,7 @@ public:
 	UPROPERTY(BlueprintReadOnly, Category = Properties)
 	TObjectPtr<UPCGNode> Node = nullptr;
 
-	UPROPERTY()
+	UPROPERTY(meta=(PCGNoHash))
 	FName Label_DEPRECATED = NAME_None;
 
 	UPROPERTY(BlueprintReadOnly, TextExportTransient, Category = Properties)
diff --git a/Engine/Plugins/PCG/Source/PCG/Public/PCGSettings.h b/Engine/Plugins/PCG/Source/PCG/Public/PCGSettings.h
--- a/Engine/Plugins/PCG/Source/PCG/Public/PCGSettings.h
+++ b/Engine/Plugins/PCG/Source/PCG/Public/PCGSettings.h
@@ -162,7 +162,7 @@ public:
 	bool bDebug = false;
 
 #if WITH_EDITORONLY_DATA
-	UPROPERTY(EditAnywhere, BlueprintReadWrite, Category = Debug, meta = (ShowOnlyInnerProperties))
+	UPROPERTY(EditAnywhere, BlueprintReadWrite, Category = Debug, meta = (ShowOnlyInnerProperties, PCGNoHash))
 	FPCGDebugVisualizationSettings DebugSettings;
 
 	/** If a debugger is attached, a breakpoint will be triggered in the execution code to enable debugging. Only applies when inspecting a debug object. Transient and Editor-only. */
@@ -183,7 +183,7 @@ class UPCGSettings : public UPCGSettingsInterface
 {
 	GENERATED_BODY()
 
-	friend class FPCGSettingsObjectCrc;
+	friend class FPCGSettingsHashPolicy;
 	friend class UPCGSettingsInterface;
 	friend struct FPCGContext;
 
@@ -378,16 +378,16 @@ public:
 	UPROPERTY()
 	EPCGSettingsExecutionMode ExecutionMode_DEPRECATED = EPCGSettingsExecutionMode::Enabled;
 
-	UPROPERTY(EditAnywhere, BlueprintReadWrite, Category = Debug, AdvancedDisplay, meta = (DisplayName="Determinism", NoResetToDefault))
+	UPROPERTY(EditAnywhere, BlueprintReadWrite, Category = Debug, AdvancedDisplay, meta = (DisplayName="Determinism", NoResetToDefault, PCGNoHash))
 	FPCGDeterminismSettings DeterminismSettings;
 
-	UPROPERTY(BlueprintReadWrite, EditAnywhere, Category = AssetInfo, AssetRegistrySearchable)
+	UPROPERTY(BlueprintReadWrite, EditAnywhere, Category = AssetInfo, AssetRegistrySearchable, meta = (PCGNoHash))
 	bool bExposeToLibrary = true;
 
-	UPROPERTY(BlueprintReadWrite, EditAnywhere, Category = AssetInfo, AssetRegistrySearchable)
+	UPROPERTY(BlueprintReadWrite, EditAnywhere, Category = AssetInfo, AssetRegistrySearchable, meta = (PCGNoHash))
 	FText Category;
 
-	UPROPERTY(BlueprintReadWrite, EditAnywhere, Category = AssetInfo, AssetRegistrySearchable)
+	UPROPERTY(BlueprintReadWrite, EditAnywhere, Category = AssetInfo, AssetRegistrySearchable, meta = (PCGNoHash))
 	FText Description;
 #endif
 
@@ -498,7 +498,7 @@ protected:
 	UE_API virtual void FixingOverridableParamPropertyClass(FPCGSettingsOverridableParam& Param) const;
 
 	/** Needs to be serialized since property metadata (used to populate this array) is not available at runtime. */
-	UPROPERTY()
+	UPROPERTY(meta = (PCGNoHash))
 	TArray<FPCGSettingsOverridableParam> CachedOverridableParams;
 
 	// We need to make sure that if we have hard references that are overridable, and they are overriden by paths
