[PCG] Support basic overrides on GPU nodes
* Add a new property metadata specifier PCG_OverridableGPU to flag properties that support GPU override
* Enable CopyPoints node for GPU overridable parameters
* Overridable param getter functions are generated directly into kernel source

Some property overrides are still not supported:
* Anything that must be known at dispatch time. Requires graph splits and readback during dispatch.
* Anything that must be known at compilation time. Requires compiling kernel permutations and selecting them dynamically at runtime.
* Generic 'Overrides' pin. Requires a more intelligent approach to the generated kernel code for overridable param getters.

diff --git a/Engine/Plugins/PCG/Shaders/Private/Elements/PCGCopyPoints.usf b/Engine/Plugins/PCG/Shaders/Private/Elements/PCGCopyPoints.usf
--- a/Engine/Plugins/PCG/Shaders/Private/Elements/PCGCopyPoints.usf
+++ b/Engine/Plugins/PCG/Shaders/Private/Elements/PCGCopyPoints.usf
@@ -1,17 +1,5 @@
 // Copyright Epic Games, Inc. All Rights Reserved.
 
-// Mirrors EPCGCopyPointsInheritanceMode
-#define COPY_POINTS_INHERITANCE_MODE_Relative 0
-#define COPY_POINTS_INHERITANCE_MODE_Source 1
-#define COPY_POINTS_INHERITANCE_MODE_Target 2
-
-// Mirrors EPCGCopyPointsMetadataInheritanceMode
-#define COPY_POINTS_METADATA_INHERITANCE_MODE_SourceFirst 0
-#define COPY_POINTS_METADATA_INHERITANCE_MODE_TargetFirst 1
-#define COPY_POINTS_METADATA_INHERITANCE_MODE_SourceOnly 2
-#define COPY_POINTS_METADATA_INHERITANCE_MODE_TargetOnly 3
-#define COPY_POINTS_METADATA_INHERITANCE_MODE_None 4
-
 [numthreads(64, 1, 1)]
 void PCGCopyPointsCS(uint3 GroupId : SV_GroupID, uint GroupIndex : SV_GroupIndex)
 {
@@ -33,6 +21,8 @@ void PCGCopyPointsCS(uint3 GroupId : SV_GroupID, uint GroupIndex : SV_GroupIndex
 
 	const uint SourceNumData = Source_GetNumData();
 	const uint TargetNumData = Target_GetNumData();
+
+	// @todo_pcg: Must be known statically, so graph splits are required in order to support overriding this property.
 	const uint CopyEachSourceOnEveryTarget = CopyPoints_GetCopyEachSourceOnEveryTarget();
 
 	// Get source and target data indices for this output data index.
@@ -62,10 +52,12 @@ void PCGCopyPointsCS(uint3 GroupId : SV_GroupID, uint GroupIndex : SV_GroupIndex
 	const float4 TargetColor = Target_GetColor(TargetDataIndex, TargetPointIndexInData);
 	const uint TargetSeed = Target_GetSeed(TargetDataIndex, TargetPointIndexInData);
 
-	const uint RotationInheritance = CopyPoints_GetRotationInheritance();
-	const uint ApplyTargetRotationToPositions = CopyPoints_GetApplyTargetRotationToPositions();
-	const uint ScaleInheritance = CopyPoints_GetScaleInheritance();
-	const uint ApplyTargetScaleToPositions = CopyPoints_GetApplyTargetScaleToPositions();
+	const uint RotationInheritance = RotationInheritance_GetOverridableValue();
+	const uint ApplyTargetRotationToPositions = (RotationInheritance == PCG_COPY_POINTS_INHERITANCE_MODE_Source) ? bApplyTargetRotationToPositions_GetOverridableValue() : 1;
+	const uint ScaleInheritance = ScaleInheritance_GetOverridableValue();
+	const uint ApplyTargetScaleToPositions = (ScaleInheritance == PCG_COPY_POINTS_INHERITANCE_MODE_Source) ? bApplyTargetScaleToPositions_GetOverridableValue() : 1;
+
+	// @todo_pcg: Must be known statically, so graph splits are required in order to support overriding these properties.
 	const uint ColorInheritance = CopyPoints_GetColorInheritance();
 	const uint SeedInheritance = CopyPoints_GetSeedInheritance();
 	const uint AttributeInheritance = CopyPoints_GetAttributeInheritance();
@@ -79,54 +71,54 @@ void PCGCopyPointsCS(uint3 GroupId : SV_GroupID, uint GroupIndex : SV_GroupIndex
 	float4 OutColor = (float4)0;
 	uint OutSeed = 0;
 
-	if (RotationInheritance == COPY_POINTS_INHERITANCE_MODE_Relative)
+	if (RotationInheritance == PCG_COPY_POINTS_INHERITANCE_MODE_Relative)
 	{
 		OutRotation = QuatMultiply(TargetRotation, SourceRotation);
 	}
-	else if (RotationInheritance == COPY_POINTS_INHERITANCE_MODE_Source)
+	else if (RotationInheritance == PCG_COPY_POINTS_INHERITANCE_MODE_Source)
 	{
 		OutRotation = SourceRotation;
 	}
-	else if (RotationInheritance == COPY_POINTS_INHERITANCE_MODE_Target)
+	else if (RotationInheritance == PCG_COPY_POINTS_INHERITANCE_MODE_Target)
 	{
 		OutRotation = TargetRotation;
 	}
 
-	if (ScaleInheritance == COPY_POINTS_INHERITANCE_MODE_Relative)
+	if (ScaleInheritance == PCG_COPY_POINTS_INHERITANCE_MODE_Relative)
 	{
 		OutScale = SourceScale * TargetScale;
 	}
-	else if (ScaleInheritance == COPY_POINTS_INHERITANCE_MODE_Source)
+	else if (ScaleInheritance == PCG_COPY_POINTS_INHERITANCE_MODE_Source)
 	{
 		OutScale = SourceScale;
 	}
-	else if (ScaleInheritance == COPY_POINTS_INHERITANCE_MODE_Target)
+	else if (ScaleInheritance == PCG_COPY_POINTS_INHERITANCE_MODE_Target)
 	{
 		OutScale = TargetScale;
 	}
 
-	if (ColorInheritance == COPY_POINTS_INHERITANCE_MODE_Relative)
+	if (ColorInheritance == PCG_COPY_POINTS_INHERITANCE_MODE_Relative)
 	{
 		OutColor = SourceColor * TargetColor;
 	}
-	else if (ColorInheritance == COPY_POINTS_INHERITANCE_MODE_Source)
+	else if (ColorInheritance == PCG_COPY_POINTS_INHERITANCE_MODE_Source)
 	{
 		OutColor = SourceColor;
 	}
-	else if (ColorInheritance == COPY_POINTS_INHERITANCE_MODE_Target)
+	else if (ColorInheritance == PCG_COPY_POINTS_INHERITANCE_MODE_Target)
 	{
 		OutColor = TargetColor;
 	}
 
-	if (SeedInheritance == COPY_POINTS_INHERITANCE_MODE_Relative)
+	if (SeedInheritance == PCG_COPY_POINTS_INHERITANCE_MODE_Relative)
 	{
 		OutSeed = ComputeSeed(SourceSeed, TargetSeed);
 	}
-	else if (SeedInheritance == COPY_POINTS_INHERITANCE_MODE_Source)
+	else if (SeedInheritance == PCG_COPY_POINTS_INHERITANCE_MODE_Source)
 	{
 		OutSeed = SourceSeed;
 	}
-	else if (SeedInheritance == COPY_POINTS_INHERITANCE_MODE_Target)
+	else if (SeedInheritance == PCG_COPY_POINTS_INHERITANCE_MODE_Target)
 	{
 		OutSeed = TargetSeed;
 	}
@@ -145,22 +137,22 @@ void PCGCopyPointsCS(uint3 GroupId : SV_GroupID, uint GroupIndex : SV_GroupIndex
 	Out_SetSteepness(Out_DataIndex, Out_ElemIndex, Source_GetSteepness(SourceDataIndex, SourcePointIndexInData));
 
 	// Copy attributes
-	if (AttributeInheritance == COPY_POINTS_METADATA_INHERITANCE_MODE_SourceOnly)
+	if (AttributeInheritance == PCG_COPY_POINTS_METADATA_INHERITANCE_MODE_SourceOnly)
 	{
 		PCG_COPY_ATTRIBUTES_TO_OUTPUT(Out, Source, Out_DataIndex, Out_ElemIndex, SourceDataIndex, SourcePointIndexInData, /*bMetadataOnly=*/true, /*bInitNonCopiedAttributes=*/true);
 	}
-	else if (AttributeInheritance == COPY_POINTS_METADATA_INHERITANCE_MODE_TargetOnly)
+	else if (AttributeInheritance == PCG_COPY_POINTS_METADATA_INHERITANCE_MODE_TargetOnly)
 	{
 		PCG_COPY_ATTRIBUTES_TO_OUTPUT(Out, Target, Out_DataIndex, Out_ElemIndex, TargetDataIndex, TargetPointIndexInData, /*bMetadataOnly=*/true, /*bInitNonCopiedAttributes=*/true);
 	}
-	else if (AttributeInheritance == COPY_POINTS_METADATA_INHERITANCE_MODE_SourceFirst)
+	else if (AttributeInheritance == PCG_COPY_POINTS_METADATA_INHERITANCE_MODE_SourceFirst)
 	{
 		PCG_COPY_ATTRIBUTES_TO_OUTPUT(Out, Target, Out_DataIndex, Out_ElemIndex, TargetDataIndex, TargetPointIndexInData, /*bMetadataOnly=*/true, /*bInitNonCopiedAttributes=*/true);
 
 		// Don't allow Source to initialize non-copied attributes in the output, otherwise it will overwrite attributes from Target.
 		PCG_COPY_ATTRIBUTES_TO_OUTPUT(Out, Source, Out_DataIndex, Out_ElemIndex, SourceDataIndex, SourcePointIndexInData, /*bMetadataOnly=*/true, /*bInitNonCopiedAttributes=*/false);
 	}
-	else if (AttributeInheritance == COPY_POINTS_METADATA_INHERITANCE_MODE_TargetFirst)
+	else if (AttributeInheritance == PCG_COPY_POINTS_METADATA_INHERITANCE_MODE_TargetFirst)
 	{
 		PCG_COPY_ATTRIBUTES_TO_OUTPUT(Out, Source, Out_DataIndex, Out_ElemIndex, SourceDataIndex, SourcePointIndexInData, /*bMetadataOnly=*/true, /*bInitNonCopiedAttributes=*/true);
 
diff --git a/Engine/Plugins/PCG/Shaders/Private/PCGDataCollectionDataInterface.ush b/Engine/Plugins/PCG/Shaders/Private/PCGDataCollectionDataInterface.ush
--- a/Engine/Plugins/PCG/Shaders/Private/PCGDataCollectionDataInterface.ush
+++ b/Engine/Plugins/PCG/Shaders/Private/PCGDataCollectionDataInterface.ush
@@ -141,6 +141,21 @@ uint GetElementAddressInternal_{DataInterfaceName}(uint InDataIndex, uint InElem
 	return InElementIndex * AttributeStride + FirstElementAddress;
 }
 
+uint GetFirstAttributeAddressInternal_{DataInterfaceName}()
+{
+	for (int AttributeId = 0; AttributeId < PCG_MAX_NUM_ATTRS; ++AttributeId)
+	{
+		const uint AttributeAddress = GetFirstElementAddressInternal_{DataInterfaceName}(/*DataIndex=*/0, AttributeId);
+
+		if (AttributeAddress > 0)
+		{
+			return AttributeAddress;
+		}
+	}
+
+	return 0;
+}
+
 uint GetDataNumAttributesInternal_{DataInterfaceName}(uint InDataIndex)
 {
 	const uint ReadAddress =
@@ -357,6 +372,83 @@ void SetName_{DataInterfaceName}(uint DataIndex, uint ElementIndex, int Attribut
 	SetInt_{DataInterfaceName}(DataIndex, ElementIndex, AttributeId, (int)Value.x);
 }
 
+// #################### FIRST ATTRIBUTE GETTERS ##########################
+
+bool GetFirstAttributeAsBool_{DataInterfaceName}()
+{
+	const uint AttributeAddress = GetFirstAttributeAddressInternal_{DataInterfaceName}();
+	return AttributeAddress > 0 && {DataInterfaceName}_DataCollectionBuffer.Load(AttributeAddress) != 0;
+}
+
+int GetFirstAttributeAsInt_{DataInterfaceName}()
+{
+	const uint AttributeAddress = GetFirstAttributeAddressInternal_{DataInterfaceName}();
+	return AttributeAddress > 0 ? asint({DataInterfaceName}_DataCollectionBuffer.Load(AttributeAddress)) : 0;
+}
+
+float GetFirstAttributeAsFloat_{DataInterfaceName}()
+{
+	const uint AttributeAddress = GetFirstAttributeAddressInternal_{DataInterfaceName}();
+	return AttributeAddress > 0 ? asfloat({DataInterfaceName}_DataCollectionBuffer.Load(AttributeAddress)) : 0.0f;
+}
+
+float2 GetFirstAttributeAsFloat2_{DataInterfaceName}()
+{
+	const uint AttributeAddress = GetFirstAttributeAddressInternal_{DataInterfaceName}();
+	return AttributeAddress > 0 ? asfloat({DataInterfaceName}_DataCollectionBuffer.Load2(AttributeAddress)) : (float2)0;
+}
+
+float3 GetFirstAttributeAsFloat3_{DataInterfaceName}()
+{
+	const uint AttributeAddress = GetFirstAttributeAddressInternal_{DataInterfaceName}();
+	return AttributeAddress > 0 ? asfloat({DataInterfaceName}_DataCollectionBuffer.Load3(AttributeAddress)) : (float3)0;
+}
+
+float4 GetFirstAttributeAsFloat4_{DataInterfaceName}()
+{
+	const uint AttributeAddress = GetFirstAttributeAddressInternal_{DataInterfaceName}();
+	return AttributeAddress > 0 ? asfloat({DataInterfaceName}_DataCollectionBuffer.Load4(AttributeAddress)) : (float4)0;
+}
+
+float3 GetFirstAttributeAsRotator_{DataInterfaceName}()
+{
+	return GetFirstAttributeAsFloat3_{DataInterfaceName}();
+}
+
+float4 GetFirstAttributeAsQuat_{DataInterfaceName}()
+{
+	return GetFirstAttributeAsFloat4_{DataInterfaceName}();
+}
+
+float4x4 GetFirstAttributeAsTransform_{DataInterfaceName}()
+{
+	const uint AttributeAddress = GetFirstAttributeAddressInternal_{DataInterfaceName}();
+	if (AttributeAddress == 0)
+	{
+		return (float4x4)0;
+	}
+
+	return float4x4(
+		asfloat({DataInterfaceName}_DataCollectionBuffer.Load4(AttributeAddress)),
+		asfloat({DataInterfaceName}_DataCollectionBuffer.Load4(AttributeAddress + 4 * 4)),
+		asfloat({DataInterfaceName}_DataCollectionBuffer.Load4(AttributeAddress + 8 * 4)),
+		asfloat({DataInterfaceName}_DataCollectionBuffer.Load4(AttributeAddress + 12 * 4))
+	);
+}
+
+int GetFirstAttributeAsStringKey_{DataInterfaceName}()
+{
+	// String keys are represented as ints.
+	return GetFirstAttributeAsInt_{DataInterfaceName}();
+}
+
+// @todo_pcg: Left as uint2 which we locked into for 5.5.0. Future version should change this to int.
+uint2 GetFirstAttributeAsName_{DataInterfaceName}()
+{
+	// Names are represented as ints.
+	return uint2((uint)GetFirstAttributeAsInt_{DataInterfaceName}(), 0u);
+}
+
 // #################### ATOMICS ##########################
 
 // Returns value of attribute before incrementing.
diff --git a/Engine/Plugins/PCG/Source/PCG/Private/Compute/DataInterfaces/Elements/PCGCopyPointsDataInterface.cpp b/Engine/Plugins/PCG/Source/PCG/Private/Compute/DataInterfaces/Elements/PCGCopyPointsDataInterface.cpp
--- a/Engine/Plugins/PCG/Source/PCG/Private/Compute/DataInterfaces/Elements/PCGCopyPointsDataInterface.cpp
+++ b/Engine/Plugins/PCG/Source/PCG/Private/Compute/DataInterfaces/Elements/PCGCopyPointsDataInterface.cpp
@@ -18,6 +18,7 @@
 #include "ShaderCore.h"
 #include "ShaderParameterMetadataBuilder.h"
 #include "SystemTextures.h"
+#include "ComputeFramework/ComputeKernelPermutationSet.h"
 #include "ComputeFramework/ComputeKernelPermutationVector.h"
 #include "ComputeFramework/ShaderParamTypeDefinition.h"
 
@@ -25,21 +26,16 @@
 
 void UPCGCopyPointsDataInterface::GetSupportedInputs(TArray<FShaderFunctionDefinition>& OutFunctions) const
 {
-	OutFunctions.AddDefaulted_GetRef()
-		.SetName(TEXT("CopyPoints_GetRotationInheritance"))
-		.AddReturnType(FShaderValueType::Get(EShaderFundamentalType::Uint));
-
-	OutFunctions.AddDefaulted_GetRef()
-		.SetName(TEXT("CopyPoints_GetApplyTargetRotationToPositions"))
-		.AddReturnType(FShaderValueType::Get(EShaderFundamentalType::Uint));
+	const UPCGSettings* Settings = ProducerKernel->GetSettings();
+	check(Settings);
 
-	OutFunctions.AddDefaulted_GetRef()
-		.SetName(TEXT("CopyPoints_GetScaleInheritance"))
-		.AddReturnType(FShaderValueType::Get(EShaderFundamentalType::Uint));
-
-	OutFunctions.AddDefaulted_GetRef()
-		.SetName(TEXT("CopyPoints_GetApplyTargetScaleToPositions"))
-		.AddReturnType(FShaderValueType::Get(EShaderFundamentalType::Uint));
+	// @todo_pcg: Pull this into the base class or elsewhere to make supporting overrides easier.
+	for (const FPCGSettingsOverridableParam& OverridableParam : Settings->OverridableParamsForGPU())
+	{
+		OutFunctions.AddDefaulted_GetRef()
+			.SetName(FString::Format(TEXT("CopyPoints_Get{0}Internal"), { OverridableParam.Label.ToString()}))
+			.AddReturnType(FShaderValueType::Get(EShaderFundamentalType::Uint)); // @todo_pcg: Determine the correct type from the overridable param
+	}
 
 	OutFunctions.AddDefaulted_GetRef()
 		.SetName(TEXT("CopyPoints_GetColorInheritance"))
@@ -68,9 +64,9 @@ void UPCGCopyPointsDataInterface::GetSupportedInputs(TArray<FShaderFunctionDefin
 
 BEGIN_SHADER_PARAMETER_STRUCT(FPCGCopyPointsDataInterfaceParameters,)
 	SHADER_PARAMETER(uint32, RotationInheritance)
-	SHADER_PARAMETER(uint32, ApplyTargetRotationToPositions)
+	SHADER_PARAMETER(uint32, bApplyTargetRotationToPositions)
 	SHADER_PARAMETER(uint32, ScaleInheritance)
-	SHADER_PARAMETER(uint32, ApplyTargetScaleToPositions)
+	SHADER_PARAMETER(uint32, bApplyTargetScaleToPositions)
 	SHADER_PARAMETER(uint32, ColorInheritance)
 	SHADER_PARAMETER(uint32, SeedInheritance)
 	SHADER_PARAMETER(uint32, AttributeInheritance)
@@ -88,28 +84,47 @@ void UPCGCopyPointsDataInterface::GetPermutations(FComputeKernelPermutationVecto
 	OutPermutationVector.AddPermutation(TEXT("USE_INPUT_DATA_INDICES_BUFFER"), 2);
 }
 
+void UPCGCopyPointsDataInterface::GetDefines(FComputeKernelDefinitionSet& OutDefinitionSet) const
+{
+	Super::GetDefines(OutDefinitionSet);
+
+	OutDefinitionSet.Defines.Add(FComputeKernelDefinition(TEXT("PCG_COPY_POINTS_INHERITANCE_MODE_Relative"), FString::FromInt(static_cast<int32>(EPCGCopyPointsInheritanceMode::Relative))));
+	OutDefinitionSet.Defines.Add(FComputeKernelDefinition(TEXT("PCG_COPY_POINTS_INHERITANCE_MODE_Source"), FString::FromInt(static_cast<int32>(EPCGCopyPointsInheritanceMode::Source))));
+	OutDefinitionSet.Defines.Add(FComputeKernelDefinition(TEXT("PCG_COPY_POINTS_INHERITANCE_MODE_Target"), FString::FromInt(static_cast<int32>(EPCGCopyPointsInheritanceMode::Target))));
+	OutDefinitionSet.Defines.Add(FComputeKernelDefinition(TEXT("PCG_COPY_POINTS_METADATA_INHERITANCE_MODE_SourceFirst"), FString::FromInt(static_cast<int32>(EPCGCopyPointsMetadataInheritanceMode::SourceFirst))));
+	OutDefinitionSet.Defines.Add(FComputeKernelDefinition(TEXT("PCG_COPY_POINTS_METADATA_INHERITANCE_MODE_TargetFirst"), FString::FromInt(static_cast<int32>(EPCGCopyPointsMetadataInheritanceMode::TargetFirst))));
+	OutDefinitionSet.Defines.Add(FComputeKernelDefinition(TEXT("PCG_COPY_POINTS_METADATA_INHERITANCE_MODE_SourceOnly"), FString::FromInt(static_cast<int32>(EPCGCopyPointsMetadataInheritanceMode::SourceOnly))));
+	OutDefinitionSet.Defines.Add(FComputeKernelDefinition(TEXT("PCG_COPY_POINTS_METADATA_INHERITANCE_MODE_TargetOnly"), FString::FromInt(static_cast<int32>(EPCGCopyPointsMetadataInheritanceMode::TargetOnly))));
+	OutDefinitionSet.Defines.Add(FComputeKernelDefinition(TEXT("PCG_COPY_POINTS_METADATA_INHERITANCE_MODE_None"), FString::FromInt(static_cast<int32>(EPCGCopyPointsMetadataInheritanceMode::None))));
+}
+
 void UPCGCopyPointsDataInterface::GetHLSL(FString& OutHLSL, FString const& InDataInterfaceName) const
 {
 	TMap<FString, FStringFormatArg> TemplateArgs =
 	{
 		{ TEXT("DataInterfaceName"), InDataInterfaceName },
 	};
+	
+	const UPCGSettings* Settings = ProducerKernel->GetSettings();
+	check(Settings);
+
+	// @todo_pcg: Pull this into the base class or elsewhere to make supporting overrides easier.
+	for (const FPCGSettingsOverridableParam& OverridableParam : Settings->OverridableParamsForGPU())
+	{
+		// @todo_pcg: Determine the correct type from the overridable param
+		OutHLSL += FString::Format(TEXT(
+				"uint {0}_{1};\n"
+				"uint CopyPoints_Get{1}Internal_{0}() { return {0}_{1}; }\n"
+			), { InDataInterfaceName, OverridableParam.Label.ToString() });
+	}
 
 	OutHLSL += FString::Format(TEXT(
-		"uint {DataInterfaceName}_RotationInheritance;\n"
-		"uint {DataInterfaceName}_ApplyTargetRotationToPositions;\n"
-		"uint {DataInterfaceName}_ScaleInheritance;\n"
-		"uint {DataInterfaceName}_ApplyTargetScaleToPositions;\n"
 		"uint {DataInterfaceName}_ColorInheritance;\n"
 		"uint {DataInterfaceName}_SeedInheritance;\n"
 		"uint {DataInterfaceName}_AttributeInheritance;\n"
 		"uint {DataInterfaceName}_CopyEachSourceOnEveryTarget;\n"
 		"StructuredBuffer<uint2> {DataInterfaceName}_SourceAndTargetDataIndices;\n"
 		"\n"
-		"uint CopyPoints_GetRotationInheritance_{DataInterfaceName}() { return {DataInterfaceName}_RotationInheritance; }\n"
-		"uint CopyPoints_GetApplyTargetRotationToPositions_{DataInterfaceName}() { return {DataInterfaceName}_ApplyTargetRotationToPositions; }\n"
-		"uint CopyPoints_GetScaleInheritance_{DataInterfaceName}() { return {DataInterfaceName}_ScaleInheritance; }\n"
-		"uint CopyPoints_GetApplyTargetScaleToPositions_{DataInterfaceName}() { return {DataInterfaceName}_ApplyTargetScaleToPositions; }\n"
 		"uint CopyPoints_GetColorInheritance_{DataInterfaceName}() { return {DataInterfaceName}_ColorInheritance; }\n"
 		"uint CopyPoints_GetSeedInheritance_{DataInterfaceName}() { return {DataInterfaceName}_SeedInheritance; }\n"
 		"uint CopyPoints_GetAttributeInheritance_{DataInterfaceName}() { return {DataInterfaceName}_AttributeInheritance; }\n"
@@ -247,9 +262,9 @@ FComputeDataProviderRenderProxy* UPCGCopyPointsDataProvider::GetRenderProxy()
 
 	FPCGCopyPointsDataProviderProxy::FCopyPointsData_RenderThread ProxyData;
 	ProxyData.RotationInheritance = static_cast<uint32>(Settings->RotationInheritance);
-	ProxyData.ApplyTargetRotationToPositions = static_cast<uint32>((Settings->RotationInheritance != EPCGCopyPointsInheritanceMode::Source || Settings->bApplyTargetRotationToPositions) ? 1 : 0);
+	ProxyData.ApplyTargetRotationToPositions = Settings->bApplyTargetRotationToPositions ? 1u : 0u;
 	ProxyData.ScaleInheritance = static_cast<uint32>(Settings->ScaleInheritance);
-	ProxyData.ApplyTargetScaleToPositions = static_cast<uint32>((Settings->ScaleInheritance != EPCGCopyPointsInheritanceMode::Source || Settings->bApplyTargetScaleToPositions) ? 1 : 0);
+	ProxyData.ApplyTargetScaleToPositions = Settings->bApplyTargetScaleToPositions ? 1u : 0u;
 	ProxyData.ColorInheritance = static_cast<uint32>(Settings->ColorInheritance);
 	ProxyData.SeedInheritance = static_cast<uint32>(Settings->SeedInheritance);
 	ProxyData.AttributeInheritance = static_cast<uint32>(Settings->AttributeInheritance);
@@ -323,9 +338,9 @@ void FPCGCopyPointsDataProviderProxy::GatherDispatchData(FDispatchData const& In
 	{
 		FParameters& Parameters = ParameterArray[InvocationIndex];
 		Parameters.RotationInheritance = Data.RotationInheritance;
-		Parameters.ApplyTargetRotationToPositions = Data.ApplyTargetRotationToPositions;
+		Parameters.bApplyTargetRotationToPositions = Data.ApplyTargetRotationToPositions;
 		Parameters.ScaleInheritance = Data.ScaleInheritance;
-		Parameters.ApplyTargetScaleToPositions = Data.ApplyTargetScaleToPositions;
+		Parameters.bApplyTargetScaleToPositions = Data.ApplyTargetScaleToPositions;
 		Parameters.ColorInheritance = Data.ColorInheritance;
 		Parameters.SeedInheritance = Data.SeedInheritance;
 		Parameters.AttributeInheritance = Data.AttributeInheritance;
diff --git a/Engine/Plugins/PCG/Source/PCG/Private/Compute/DataInterfaces/Elements/PCGCopyPointsDataInterface.h b/Engine/Plugins/PCG/Source/PCG/Private/Compute/DataInterfaces/Elements/PCGCopyPointsDataInterface.h
--- a/Engine/Plugins/PCG/Source/PCG/Private/Compute/DataInterfaces/Elements/PCGCopyPointsDataInterface.h
+++ b/Engine/Plugins/PCG/Source/PCG/Private/Compute/DataInterfaces/Elements/PCGCopyPointsDataInterface.h
@@ -22,6 +22,7 @@ public:
 	void GetSupportedInputs(TArray<FShaderFunctionDefinition>& OutFunctions) const override;
 	void GetShaderParameters(TCHAR const* UID, FShaderParametersMetadataBuilder& InOutBuilder, FShaderParametersMetadataAllocations& InOutAllocations) const override;
 	void GetPermutations(FComputeKernelPermutationVector& OutPermutationVector) const override;
+	void GetDefines(FComputeKernelDefinitionSet& OutDefinitionSet) const override;
 	void GetHLSL(FString& OutHLSL, FString const& InDataInterfaceName) const override;
 	UComputeDataProvider* CreateDataProvider() const override;
 	//~ End UComputeDataInterface Interface
diff --git a/Engine/Plugins/PCG/Source/PCG/Private/Compute/DataInterfaces/PCGDataCollectionDataInterface.cpp b/Engine/Plugins/PCG/Source/PCG/Private/Compute/DataInterfaces/PCGDataCollectionDataInterface.cpp
--- a/Engine/Plugins/PCG/Source/PCG/Private/Compute/DataInterfaces/PCGDataCollectionDataInterface.cpp
+++ b/Engine/Plugins/PCG/Source/PCG/Private/Compute/DataInterfaces/PCGDataCollectionDataInterface.cpp
@@ -246,6 +246,53 @@ void UPCGDataCollectionDataInterface::GetSupportedInputs(TArray<FShaderFunctionD
 			.AddParam(EShaderFundamentalType::Uint, 2); // Name represented by 2 uints
 	}
 
+	// First Attribute Getters. Useful for overridable params. Currently not exposed to declarations and only called from compiler generated code.
+	{
+		OutFunctions.AddDefaulted_GetRef()
+			.SetName(TEXT("GetFirstAttributeAsBool"))
+			.AddReturnType(EShaderFundamentalType::Bool);
+
+		OutFunctions.AddDefaulted_GetRef()
+			.SetName(TEXT("GetFirstAttributeAsInt"))
+			.AddReturnType(EShaderFundamentalType::Int);
+
+		OutFunctions.AddDefaulted_GetRef()
+			.SetName(TEXT("GetFirstAttributeAsFloat"))
+			.AddReturnType(EShaderFundamentalType::Float);
+
+		OutFunctions.AddDefaulted_GetRef()
+			.SetName(TEXT("GetFirstAttributeAsFloat2"))
+			.AddReturnType(EShaderFundamentalType::Float, 2);
+
+		OutFunctions.AddDefaulted_GetRef()
+			.SetName(TEXT("GetFirstAttributeAsFloat3"))
+			.AddReturnType(EShaderFundamentalType::Float, 3);
+
+		OutFunctions.AddDefaulted_GetRef()
+			.SetName(TEXT("GetFirstAttributeAsFloat4"))
+			.AddReturnType(EShaderFundamentalType::Float, 4);
+
+		OutFunctions.AddDefaulted_GetRef()
+			.SetName(TEXT("GetFirstAttributeAsRotator"))
+			.AddReturnType(EShaderFundamentalType::Float, 3);
+
+		OutFunctions.AddDefaulted_GetRef()
+			.SetName(TEXT("GetFirstAttributeAsQuat"))
+			.AddReturnType(EShaderFundamentalType::Float, 4);
+
+		OutFunctions.AddDefaulted_GetRef()
+			.SetName(TEXT("GetFirstAttributeAsTransform"))
+			.AddReturnType(EShaderFundamentalType::Float, 4, 4);
+
+		OutFunctions.AddDefaulted_GetRef()
+			.SetName(TEXT("GetFirstAttributeAsStringKey"))
+			.AddReturnType(EShaderFundamentalType::Int); // String key represented by int
+
+		OutFunctions.AddDefaulted_GetRef()
+			.SetName(TEXT("GetFirstAttributeAsName"))
+			.AddReturnType(EShaderFundamentalType::Uint, 2); // Name represented by 2 uints
+	}
+
 	// Atomics
 	{
 		OutFunctions.AddDefaulted_GetRef()
diff --git a/Engine/Plugins/PCG/Source/PCG/Private/Compute/PCGComputeKernel.cpp b/Engine/Plugins/PCG/Source/PCG/Private/Compute/PCGComputeKernel.cpp
--- a/Engine/Plugins/PCG/Source/PCG/Private/Compute/PCGComputeKernel.cpp
+++ b/Engine/Plugins/PCG/Source/PCG/Private/Compute/PCGComputeKernel.cpp
@@ -6,12 +6,14 @@
 #include "PCGEdge.h"
 #include "PCGModule.h"
 #include "PCGNode.h"
+#include "PCGPin.h"
 #include "PCGSettings.h"
 #include "Compute/PCGComputeGraph.h"
 #include "Compute/PCGComputeSource.h"
 #include "Compute/PCGDataBinding.h"
 #include "Compute/PCGPinPropertiesGPU.h"
 #include "Graph/PCGGPUGraphCompilationContext.h"
+#include "Helpers/PCGSettingsHelpers.h"
 
 #include "ShaderCompilerCore.h"
 #include "Async/Async.h"
@@ -25,7 +27,7 @@ namespace PCGComputeKernel
 	static TAutoConsoleVariable<bool> CVarWarnOnOverridePinUsage(
 		TEXT("pcg.Graph.GPU.WarnOnOverridePinUsage"),
 		true,
-		TEXT("Enables warnings when parameters are overidden on GPU nodes."));
+		TEXT("Enables warnings when unsupported parameters are overidden on GPU nodes."));
 
 	static const TStaticArray<FSoftObjectPath, 3> DefaultAdditionalSourcePaths =
 	{
@@ -54,6 +56,11 @@ void UPCGComputeKernel::Initialize(const FPCGComputeKernelParams& InParams)
 		PCGKernelFlags |= EPCGComputeKernelFlags::LogDataDescriptions;
 	}
 
+	if (InParams.bRequiresOverridableParams)
+	{
+		PCGKernelFlags |= EPCGComputeKernelFlags::RequiresOverridableParams;
+	}
+
 	bInitialized = true;
 
 	InitializeInternal();
@@ -99,10 +106,14 @@ bool UPCGComputeKernel::AreKernelSettingsValid(FPCGContext* InContext) const
 		{
 			PCG_KERNEL_VALIDATION_ERR(InContext, GetSettings(), StaticLogEntry.Message);
 		}
-		else
+		else if (StaticLogEntry.Verbosity == EPCGKernelLogVerbosity::Warning)
 		{
 			PCG_KERNEL_VALIDATION_WARN(InContext, GetSettings(), StaticLogEntry.Message);
 		}
+		else
+		{
+			ensure(false);
+		}
 	}
 #endif
 
@@ -458,23 +469,44 @@ bool UPCGComputeKernel::ValidatePCGNode(TArray<FPCGKernelLogEntry>& OutLogEntrie
 		return true;
 	}
 
+#if PCG_KERNEL_LOGGING_ENABLED
 	if (PCGComputeKernel::CVarWarnOnOverridePinUsage.GetValueOnAnyThread())
 	{
-		for (const FPCGSettingsOverridableParam& Param : Settings->OverridableParams())
+		// @todo_pcg: There are three classes of properties that still do not support overrides:
+		// 1. Properties which must be known at dispatch time (e.g. UPCGCopyPointsSettings::bCopyEachSourceToEveryTarget), so to support them we will need injected graph splits when they are overridden.
+		// 2. Properties which must be known at compile time (e.g. UPCGCopyPointsSettings::bMatchBasedOnAttribute), which would require compiling two versions of the kernel graph and branching dynamically at execution time.
+		// 3. Generic 'Overrides' pin available on all nodes with override pins.
+		// Note, nothing currently validates that params of this nature are not overridden, so implement new overrides with caution.
+
+		const UPCGNode* Node = Cast<UPCGNode>(Settings->GetOuter());
+		const UPCGPin* Pin = Node ? Node->GetInputPin(PCGPinConstants::DefaultParamsLabel) : nullptr;
+
+		if (Pin && Pin->IsConnected())
 		{
-			if (ensure(!Param.PropertiesNames.IsEmpty()))
-			{
-				const FName PropertyName = Param.PropertiesNames[0];
+			OutLogEntries.Emplace(FText::Format(
+				LOCTEXT("InvalidOverridesPin", "Unsupported use of pin '{0}' on node '{1}'. Use a direct override pin instead."),
+				FText::FromName(PCGPinConstants::DefaultParamsLabel),
+				Node->GetNodeTitle(EPCGNodeTitleType::ListView)),
+				EPCGKernelLogVerbosity::Warning);
+		}
+
+		for (TFieldIterator<FProperty> PropIt(Settings->GetClass()); PropIt; ++PropIt)
+		{
+			const FName PropertyName = (*PropIt)->GetFName();
 
+			if (!Settings->OverridableParamsForGPU().ContainsByPredicate([PropertyName](const FPCGSettingsOverridableParam& OverridableParam) { return OverridableParam.Label == PropertyName; }))
+			{
 				if (Settings->IsPropertyOverriddenByPin(PropertyName))
 				{
-#if PCG_KERNEL_LOGGING_ENABLED
-					OutLogEntries.Emplace(FText::Format(LOCTEXT("ParamOverrideGPU", "Tried to override pin '{0}', but overrides are not supported on GPU nodes."), FText::FromName(PropertyName)), EPCGKernelLogVerbosity::Warning);
-#endif
+					OutLogEntries.Emplace(FText::Format(
+						LOCTEXT("UnsupportedOverrideError", "Tried to override parameter '{0}', but this parameter does not support overrides when executing on the GPU at this time."),
+						FText::FromName(PropertyName)),
+						EPCGKernelLogVerbosity::Warning);
 				}
 			}
 		}
 	}
+#endif
 
 	// Validate types of incident edges to make sure we catch invalid cases like Spatial -> Point.
 	return AreInputEdgesValid(OutLogEntries);
@@ -517,4 +549,19 @@ bool UPCGComputeKernel::AreInputEdgesValid(TArray<FPCGKernelLogEntry>& OutLogEnt
 }
 #endif
 
+void UPCGComputeKernel::FillOverridableParamsPins(TArray<FPCGPinProperties>& OutPins) const
+{
+	if (const UPCGSettings* Settings = GetSettings())
+	{
+		for (const FPCGSettingsOverridableParam& OverridableParam : Settings->OverridableParamsForGPU())
+		{
+			if (Settings->IsPropertyOverriddenByPin(OverridableParam.Label))
+			{
+				FPCGPinProperties& ParamPin = OutPins.Emplace_GetRef(OverridableParam.Label, EPCGDataType::Param, /*bInAllowMultipleConnections=*/false, /*bAllowMultipleData=*/false);
+				ParamPin.SetOverrideOrUserParamPin();
+			}
+		}
+	}
+}
+
 #undef LOCTEXT_NAMESPACE
diff --git a/Engine/Plugins/PCG/Source/PCG/Private/Compute/PCGComputeKernel.h b/Engine/Plugins/PCG/Source/PCG/Private/Compute/PCGComputeKernel.h
--- a/Engine/Plugins/PCG/Source/PCG/Private/Compute/PCGComputeKernel.h
+++ b/Engine/Plugins/PCG/Source/PCG/Private/Compute/PCGComputeKernel.h
@@ -31,6 +31,11 @@ enum class EPCGComputeKernelFlags : uint32
 	LogDataDescriptions = 1 << 1,
 	/** Dispatch compute graph containing this kernel every frame to allow profiling. Has no effect unless PCG_GPU_KERNEL_PROFILING defined. */
 	RepeatDispatchForProfiling = 1 << 2,
+	/**
+	 * Indicates that this kernel will use parameters which are marked as OverridableGPU on the settings.
+	 * This is necessary to correctly wire up the kernel and provide code generated functions for reading the overridden values.
+	 */
+	RequiresOverridableParams = 1 << 3,
 };
 ENUM_CLASS_FLAGS(EPCGComputeKernelFlags);
 
@@ -66,6 +71,7 @@ struct FPCGComputeKernelParams
 {
 	const UPCGSettings* Settings = nullptr;
 	bool bLogDescriptions = false;
+	bool bRequiresOverridableParams = false;
 };
 #endif
 
@@ -160,6 +166,8 @@ public:
 
 	bool HasStaticValidationErrors() const { return !!(PCGKernelFlags & EPCGComputeKernelFlags::HasStaticValidationErrors); }
 
+	bool RequiresOverridableParams() const { return !!(PCGKernelFlags & EPCGComputeKernelFlags::RequiresOverridableParams); }
+
 #ifdef PCG_GPU_KERNEL_PROFILING
 	void EnableRepeatDispatch() { PCGKernelFlags |= EPCGComputeKernelFlags::RepeatDispatchForProfiling; }
 
@@ -184,6 +192,9 @@ protected:
 	bool AreInputEdgesValid(TArray<FPCGKernelLogEntry>& OutLogEntries) const;
 #endif
 
+	/** Collects all param pins which are overridable on the GPU. */
+	void FillOverridableParamsPins(TArray<FPCGPinProperties>& OutPins) const;
+
 protected:
 	UPROPERTY()
 	EPCGComputeKernelFlags PCGKernelFlags = EPCGComputeKernelFlags::None;
diff --git a/Engine/Plugins/PCG/Source/PCG/Private/Elements/PCGCopyPoints.cpp b/Engine/Plugins/PCG/Source/PCG/Private/Elements/PCGCopyPoints.cpp
--- a/Engine/Plugins/PCG/Source/PCG/Private/Elements/PCGCopyPoints.cpp
+++ b/Engine/Plugins/PCG/Source/PCG/Private/Elements/PCGCopyPoints.cpp
@@ -28,6 +28,7 @@ void UPCGCopyPointsSettings::CreateKernels(FPCGGPUCompilationContext& InOutConte
 	FPCGComputeKernelParams KernelParams;
 	KernelParams.Settings = this;
 	KernelParams.bLogDescriptions = bDumpDataDescriptions;
+	KernelParams.bRequiresOverridableParams = true;
 
 	UPCGCopyPointsKernel* CopyKernel = InOutContext.NewObject_AnyThread<UPCGCopyPointsKernel>(InObjectOuter);
 	CopyKernel->Initialize(KernelParams);
@@ -39,6 +40,9 @@ void UPCGCopyPointsSettings::CreateKernels(FPCGGPUCompilationContext& InOutConte
 
 	if (bMatchBasedOnAttribute)
 	{
+		// Analysis kernel does not require any of the overridden properties used by the CopyPoints settings.
+		KernelParams.bRequiresOverridableParams = false;
+		
 		UPCGCopyPointsAnalysisKernel* AnalysisKernel = InOutContext.NewObject_AnyThread<UPCGCopyPointsAnalysisKernel>(InObjectOuter);
 		AnalysisKernel->Initialize(KernelParams);
 		OutKernels.Add(AnalysisKernel);
diff --git a/Engine/Plugins/PCG/Source/PCG/Private/Elements/PCGCopyPointsKernel.cpp b/Engine/Plugins/PCG/Source/PCG/Private/Elements/PCGCopyPointsKernel.cpp
--- a/Engine/Plugins/PCG/Source/PCG/Private/Elements/PCGCopyPointsKernel.cpp
+++ b/Engine/Plugins/PCG/Source/PCG/Private/Elements/PCGCopyPointsKernel.cpp
@@ -196,9 +196,39 @@ int UPCGCopyPointsKernel::ComputeThreadCount(const UPCGDataBinding* InBinding) c
 #if WITH_EDITOR
 FString UPCGCopyPointsKernel::GetCookedSource(FPCGGPUCompilationContext& InOutContext) const
 {
+	FString OverridableParamFunctions = "";
+
+	if (const UPCGSettings* Settings = GetSettings())
+	{
+		for (const FPCGSettingsOverridableParam& OverridableParam : Settings->OverridableParamsForGPU())
+		{
+			// @todo_pcg: Support overriding more than just Int/Uint. Should be able to get the type from the property.
+
+			if (Settings->IsPropertyOverriddenByPin(OverridableParam.Label))
+			{
+				OverridableParamFunctions += FString::Format(TEXT(
+					"uint {0}_GetOverridableValue()\n"
+					"{\n"
+					"	return (uint){0}_GetFirstAttributeAsInt();\n"
+					"}\n"
+					"\n"), { OverridableParam.Label.ToString() });
+			}
+			else
+			{
+				OverridableParamFunctions += FString::Format(TEXT(
+					"uint {0}_GetOverridableValue()\n"
+					"{\n"
+					"	return CopyPoints_Get{0}Internal();\n"
+					"}\n"
+					"\n"), { OverridableParam.Label.ToString() });
+			}
+		}
+	}
+
 	FString TemplateFile;
 	ensure(LoadShaderSourceFile(TEXT("/Plugin/PCG/Private/Elements/PCGCopyPoints.usf"), EShaderPlatform::SP_PCD3D_SM5, &TemplateFile, nullptr));
-	return TemplateFile;
+
+	return OverridableParamFunctions + TemplateFile;
 }
 
 void UPCGCopyPointsKernel::CreateAdditionalInputDataInterfaces(FPCGGPUCompilationContext& InOutContext, UObject* InObjectOuter, TArray<TObjectPtr<UComputeDataInterface>>& OutDataInterfaces) const
@@ -224,6 +254,8 @@ void UPCGCopyPointsKernel::GetInputPins(TArray<FPCGPinProperties>& OutPins) cons
 		// One bool flag per output data that signals if the output data should be computed or not.
 		OutPins.Emplace(PCGCopyPointsConstants::SelectedFlagsPinLabel, EPCGDataType::Param);
 	}
+
+	FillOverridableParamsPins(OutPins);
 }
 
 void UPCGCopyPointsKernel::GetOutputPins(TArray<FPCGPinPropertiesGPU>& OutPins) const
diff --git a/Engine/Plugins/PCG/Source/PCG/Private/Graph/PCGGraphCompilerGPU.cpp b/Engine/Plugins/PCG/Source/PCG/Private/Graph/PCGGraphCompilerGPU.cpp
--- a/Engine/Plugins/PCG/Source/PCG/Private/Graph/PCGGraphCompilerGPU.cpp
+++ b/Engine/Plugins/PCG/Source/PCG/Private/Graph/PCGGraphCompilerGPU.cpp
@@ -293,6 +293,21 @@ void FPCGGraphCompilerGPU::ExpandGPUNodeKernelsToTasks(
 			continue;
 		}
 
+		// Wire up overridden parameters for any kernels that require them.
+		for (UPCGComputeKernel* NodeKernel : NodeKernels)
+		{
+			if (NodeKernel->RequiresOverridableParams())
+			{
+				for (const FPCGSettingsOverridableParam& OverridableParam : Settings->OverridableParamsForGPU())
+				{
+					if (Settings->IsPropertyOverriddenByPin(OverridableParam.Label))
+					{
+						NodeKernelEdges.Emplace(FPCGPinReference(OverridableParam.Label), FPCGPinReference(NodeKernel, OverridableParam.Label));
+					}
+				}
+			}
+		}
+
 		if (NodeKernelEdges.IsEmpty())
 		{
 			continue;
@@ -420,7 +435,7 @@ void FPCGGraphCompilerGPU::ExpandGPUNodeKernelsToTasks(
 						return TaskInput.DownstreamPin && TaskInput.DownstreamPin->Label == NodePinLabel;
 					});
 
-					if (ensure(OriginalTaskInput))
+					if (OriginalTaskInput)
 					{
 						// Match - add a new input wire to the new task.
 						FPCGGraphTaskInput& NewTaskInput = NewKernelTask.Inputs.Add_GetRef(*OriginalTaskInput);
@@ -1359,7 +1374,7 @@ void FPCGGraphCompilerGPU::CompileComputeGraph(
 			InputDataInterfaces.Emplace(KernelDIIndex, /*Prefix=*/NAME_None, /*PinLabel=*/NAME_None);
 		}
 
-		// Additional DIs created by settings.
+		// Additional DIs created by kernels.
 		{
 			TArray<TObjectPtr<UComputeDataInterface>> AdditionalInputDIs, AdditionalOutputDIs;
 			KernelWithBindings.Kernel->CreateAdditionalInputDataInterfaces(InOutContext, InOutComputeGraph, AdditionalInputDIs);
diff --git a/Engine/Plugins/PCG/Source/PCG/Private/PCGSettings.cpp b/Engine/Plugins/PCG/Source/PCG/Private/PCGSettings.cpp
--- a/Engine/Plugins/PCG/Source/PCG/Private/PCGSettings.cpp
+++ b/Engine/Plugins/PCG/Source/PCG/Private/PCGSettings.cpp
@@ -898,6 +898,25 @@ bool UPCGSettings::IsPropertyOverriddenByPin(const TArrayView<const FName>& Prop
 	return false;
 }
 
+bool UPCGSettings::IsPinUsedByNodeExecution(const UPCGPin* InPin) const
+{
+#if WITH_EDITOR
+	if (InPin && ShouldExecuteOnGPU())
+	{
+		auto MatchesPinLabel = [PinLabel = InPin->Properties.Label](const FPCGSettingsOverridableParam& OverridableParam) { return OverridableParam.Label == PinLabel; };
+
+		// Pin is not in use if we are in GPU execution mode and the pin is an OverridableParam that doesn't support GPU overrides.
+		// Generic overrides pin is also not supported.
+		if ((OverridableParams().ContainsByPredicate(MatchesPinLabel) && !OverridableParamsForGPU().ContainsByPredicate(MatchesPinLabel)) || InPin->Properties.Label == PCGPinConstants::DefaultParamsLabel)
+		{
+			return false;
+		}
+	}
+#endif
+
+	return true;
+}
+
 bool UPCGSettings::DoesPinSupportPassThrough(UPCGPin* InPin) const
 {
 	return InPin && !InPin->Properties.IsAdvancedPin();
@@ -930,6 +949,15 @@ TArray<FPCGSettingsOverridableParam> UPCGSettings::GatherOverridableParams() con
 
 	return PCGSettingsHelpers::GetAllOverridableParams(GetClass(), Config);
 }
+
+TArray<FPCGSettingsOverridableParam> UPCGSettings::GatherOverridableParamsGPU() const
+{
+	PCGSettingsHelpers::FPCGGetAllOverridableParamsConfig Config;
+	Config.IncludeMetadataValues.Add(PCGObjectMetadata::OverridableGPU);
+	Config.ExcludeMetadataValues.Add(PCGObjectMetadata::NotOverridable);
+
+	return PCGSettingsHelpers::GetAllOverridableParams(GetClass(), Config);
+}
 #endif // WITH_EDITOR
 
 void UPCGSettings::InitializeCachedOverridableParams(bool bReset)
@@ -987,6 +1015,11 @@ void UPCGSettings::InitializeCachedOverridableParams(bool bReset)
 			}
 		}
 	}
+
+	if (CachedOverridableParamsGPU.IsEmpty() || bReset)
+	{
+		CachedOverridableParamsGPU = GatherOverridableParamsGPU();
+	}
 #endif // WITH_EDITOR
 
 	// Reset the value
diff --git a/Engine/Plugins/PCG/Source/PCG/Public/Elements/PCGCopyPoints.h b/Engine/Plugins/PCG/Source/PCG/Public/Elements/PCGCopyPoints.h
--- a/Engine/Plugins/PCG/Source/PCG/Public/Elements/PCGCopyPoints.h
+++ b/Engine/Plugins/PCG/Source/PCG/Public/Elements/PCGCopyPoints.h
@@ -64,19 +64,19 @@ protected:
 
 public:
 	/** The method used to determine output point rotation */
-	UPROPERTY(BlueprintReadWrite, EditAnywhere, Category = Settings, meta = (PCG_Overridable))
+	UPROPERTY(BlueprintReadWrite, EditAnywhere, Category = Settings, meta = (PCG_Overridable, PCG_OverridableGPU))
 	EPCGCopyPointsInheritanceMode RotationInheritance = EPCGCopyPointsInheritanceMode::Relative;
 
 	/** If this option is set, points will have their source position transformed using the target transform with rotation */
-	UPROPERTY(BlueprintReadWrite, EditAnywhere, Category = Settings, meta = (PCG_Overridable, EditCondition = "RotationInheritance==EPCGCopyPointsInheritanceMode::Source", EditConditionHides))
+	UPROPERTY(BlueprintReadWrite, EditAnywhere, Category = Settings, meta = (PCG_Overridable, PCG_OverridableGPU, EditCondition = "RotationInheritance==EPCGCopyPointsInheritanceMode::Source", EditConditionHides))
 	bool bApplyTargetRotationToPositions = true;
 
 	/** The method used to determine output point scale */
-	UPROPERTY(BlueprintReadWrite, EditAnywhere, Category = Settings, meta = (PCG_Overridable))
+	UPROPERTY(BlueprintReadWrite, EditAnywhere, Category = Settings, meta = (PCG_Overridable, PCG_OverridableGPU))
 	EPCGCopyPointsInheritanceMode ScaleInheritance = EPCGCopyPointsInheritanceMode::Relative;
 
 	/** If this option is set, points will have their source position transformed using the target transform with scale */
-	UPROPERTY(BlueprintReadWrite, EditAnywhere, Category = Settings, meta = (PCG_Overridable, EditCondition="ScaleInheritance==EPCGCopyPointsInheritanceMode::Source", EditConditionHides))
+	UPROPERTY(BlueprintReadWrite, EditAnywhere, Category = Settings, meta = (PCG_Overridable, PCG_OverridableGPU, EditCondition="ScaleInheritance==EPCGCopyPointsInheritanceMode::Source", EditConditionHides))
 	bool bApplyTargetScaleToPositions = true;
 
 	/** The method used to determine output point color */
diff --git a/Engine/Plugins/PCG/Source/PCG/Public/PCGCommon.h b/Engine/Plugins/PCG/Source/PCG/Public/PCGCommon.h
--- a/Engine/Plugins/PCG/Source/PCG/Public/PCGCommon.h
+++ b/Engine/Plugins/PCG/Source/PCG/Public/PCGCommon.h
@@ -229,6 +229,7 @@ namespace Icons
 namespace PCGObjectMetadata
 {
 	const FName Overridable = TEXT("PCG_Overridable");
+	const FName OverridableGPU = TEXT("PCG_OverridableGPU");
 	const FName NotOverridable = TEXT("PCG_NotOverridable");
 	const FName OverrideAliases = TEXT("PCG_OverrideAliases");
 	const FName DiscardPropertySelection = TEXT("PCG_DiscardPropertySelection");
@@ -241,6 +242,9 @@ namespace PCGObjectMetadata
 		/// [PropertyMetadata] Indicates that the property is overridable by params.
 		PCG_Overridable,
 
+		/// [PropertyMetadata] Indicates that the property is overridable by params on the GPU.
+		PCG_OverridableGPU,
+
 		/// [PropertyMetadata] Indicates that the property is not-overridable by params. Used in structs to hide some parameters
 		PCG_NotOverridable,
 
diff --git a/Engine/Plugins/PCG/Source/PCG/Public/PCGSettings.h b/Engine/Plugins/PCG/Source/PCG/Public/PCGSettings.h
--- a/Engine/Plugins/PCG/Source/PCG/Public/PCGSettings.h
+++ b/Engine/Plugins/PCG/Source/PCG/Public/PCGSettings.h
@@ -340,7 +340,7 @@ public:
 	virtual bool IsInputPinRequiredByExecution(const UPCGPin* InPin) const { return InPin && InPin->Properties.IsRequiredPin(); }
 
 	/** Returns true if InPin is in use by node (assuming node enabled). Can be used to communicate when a pin is not in use to user. */
-	virtual bool IsPinUsedByNodeExecution(const UPCGPin* InPin) const { return true; }
+	UE_API virtual bool IsPinUsedByNodeExecution(const UPCGPin* InPin) const;
 
 	/** True if we know prior to execution that the given pin will be active (not on inactive branch). */
 	virtual bool IsPinStaticallyActive(const FName& OutputPinLabel) const { return true; }
@@ -466,6 +466,9 @@ public:
 	/** List of all the overridable params available for these settings. */
 	virtual const TArray<FPCGSettingsOverridableParam>& OverridableParams() const { return CachedOverridableParams; }
 
+	/** List of all the overridable params available for these settings on the GPU. */
+	const TArray<FPCGSettingsOverridableParam>& OverridableParamsForGPU() const { return CachedOverridableParamsGPU; }
+
 	/** Check if we have some override. Can be overriden to force params pin for example */
 	virtual bool HasOverridableParams() const { return !CachedOverridableParams.IsEmpty(); }
 
@@ -482,6 +485,7 @@ protected:
 #if WITH_EDITOR
 	/** Can be overridden to add more custom params (like in BP). */
 	UE_API virtual TArray<FPCGSettingsOverridableParam> GatherOverridableParams() const;
+	TArray<FPCGSettingsOverridableParam> GatherOverridableParamsGPU() const;
 #endif // WITH_EDITOR
 
 	/** Called after intialization to construct the list of all overridable params. 
@@ -501,6 +505,10 @@ protected:
 	UPROPERTY(meta = (PCGNoHash))
 	TArray<FPCGSettingsOverridableParam> CachedOverridableParams;
 
+	/** Needs to be serialized since property metadata (used to populate this array) is not available at runtime. */
+	UPROPERTY()
+	TArray<FPCGSettingsOverridableParam> CachedOverridableParamsGPU;
+
 	// We need to make sure that if we have hard references that are overridable, and they are overriden by paths
 	// on objects that are not yet loaded, that we are loading it on the main thread.
 	bool bHasAnyOverridableHardReferences = false;
