[PCG] Fix build

#rnx

diff --git a/Engine/Plugins/PCG/Source/PCG/Public/Hash/PCGObjectHash.h b/Engine/Plugins/PCG/Source/PCG/Public/Hash/PCGObjectHash.h
--- a/Engine/Plugins/PCG/Source/PCG/Public/Hash/PCGObjectHash.h
+++ b/Engine/Plugins/PCG/Source/PCG/Public/Hash/PCGObjectHash.h
@@ -95,7 +95,7 @@ public:
 	using FHash = FXxHash64;
 
 	explicit FPCGObjectHash(FPCGObjectHashContext* InContext);
-	virtual ~FPCGObjectHash();
+	~FPCGObjectHash();
 
 	FHash GetHash(bool bVerbose = false) const;
 
