[PCG] Fix extra spacing for enum driven palette icons

diff --git a/Engine/Plugins/PCG/Source/PCG/Public/Utils/PCGPreconfiguration.h b/Engine/Plugins/PCG/Source/PCG/Public/Utils/PCGPreconfiguration.h
--- a/Engine/Plugins/PCG/Source/PCG/Public/Utils/PCGPreconfiguration.h
+++ b/Engine/Plugins/PCG/Source/PCG/Public/Utils/PCGPreconfiguration.h
@@ -203,18 +203,14 @@ struct FPCGPreConfiguredSettingsInfo : public FPCGPreconfiguredInfo
 						InOutInfo.ActionIconDescriptor = MakeShared<FPCGActionIconByMetadataDescriptor>(static_cast<EPCGMetadataTypes>(EnumPtr->GetValueByIndex(ValueIndex)));
 					}
 				}
-				else // Get the brush icon and color directly from the metadata
+				// Get the brush icon and color directly from the metadata
+				else if (EnumPtr->HasMetaData(ActionIconMetadataKey, ValueIndex))
 				{
 					FPCGActionIconByLabelDescriptor IconDescriptor;
-					if (EnumPtr->HasMetaData(ActionIconMetadataKey, ValueIndex))
-					{
-						IconDescriptor.BrushLabel = FName(EnumPtr->GetMetaData(ActionIconMetadataKey, ValueIndex));
-					}
+					IconDescriptor.BrushLabel = FName(EnumPtr->GetMetaData(ActionIconMetadataKey, ValueIndex));
 
-					if (EnumPtr->HasMetaData(ActionIconTintMetadataKey, ValueIndex))
-					{
-						// @todo_pcg: Parse the metadata into RGBA channels. For now just make it white.
-					}
+					// @todo_pcg: Parse the metadata into RGBA channels. For now just make it white.
+					// if (EnumPtr->HasMetaData(ActionIconTintMetadataKey, ValueIndex)) {}
 					IconDescriptor.Tint = FLinearColor::White;
 
 					InOutInfo.ActionIconDescriptor = MakeShared<FPCGActionIconByLabelDescriptor>(std::move(IconDescriptor));
diff --git a/Engine/Plugins/PCG/Source/PCGEditor/Private/PCGEditorGraphSchemaActions.cpp b/Engine/Plugins/PCG/Source/PCGEditor/Private/PCGEditorGraphSchemaActions.cpp
--- a/Engine/Plugins/PCG/Source/PCGEditor/Private/PCGEditorGraphSchemaActions.cpp
+++ b/Engine/Plugins/PCG/Source/PCGEditor/Private/PCGEditorGraphSchemaActions.cpp
@@ -137,7 +137,8 @@ FSlateBrush const* FPCGEditorGraphSchemaAction_NewNativeElement::GetPaletteIcon(
 				checkNoEntry();
 		}
 
-		if (BaseBrush)
+		// Brush is expected to be set and have a valid resource.
+		if (BaseBrush && ensure(BaseBrush->IsSet() && BaseBrush->GetResourceName() != NAME_None))
 		{
 			PaletteIcon = MakeShared<FSlateBrush>(*BaseBrush);
 			PaletteIcon->TintColor = TintColor;
