[PCG] Freshen up the texture DI function GetTexCoords to a new signature that will support any texture orientation.

#rnx

diff --git a/Engine/Plugins/PCG/Shaders/Private/PCGTextureDataInterface.ush b/Engine/Plugins/PCG/Shaders/Private/PCGTextureDataInterface.ush
--- a/Engine/Plugins/PCG/Shaders/Private/PCGTextureDataInterface.ush
+++ b/Engine/Plugins/PCG/Shaders/Private/PCGTextureDataInterface.ush
@@ -35,6 +35,7 @@ PCG_TEXTURE_UAV_DECLARATION(0);
 SamplerState {DataInterfaceName}_Sampler_Linear;
 SamplerState {DataInterfaceName}_Sampler_Point;
 
+// todo_pcg: Change to using texture transform matrix so we support different orientations.
 float4 {DataInterfaceName}_TextureBounds[PCG_MAX_NUM_SRV_BINDINGS];
 
 /** Layout of TextureInfos buffer:
@@ -119,18 +120,20 @@ bool GetThreadData_{DataInterfaceName}(uint InThreadIndex, out uint OutDataIndex
 
 // #################### INPUT FUNCTIONS ##########################
 
-float2 GetTexCoords_{DataInterfaceName}(float2 WorldPos, float2 Min, float2 Max)
+float2 GetTexCoords_{DataInterfaceName}(uint InDataIndex, float3 InWorldPos)
 {
-	// Inverse lerp
-	const float2 Pos = WorldPos - Min;
+	const float2 Min = {DataInterfaceName}_TextureBounds[InDataIndex].xy;
+	const float2 Max = {DataInterfaceName}_TextureBounds[InDataIndex].zw;
+
 	const float2 Extents = Max - Min;
 
-	if (!all(Extents))
+	if (!all(Extents > 0.0f))
 	{
 		return (float2)0;
 	}
 
-	return Pos / Extents;
+	// Inverse lerp
+	return (InWorldPos.xy - Min) / Extents;
 }
 
 #define PCG_CONDITIONAL_TEXTURE_SAMPLE(Index) \
@@ -408,6 +411,21 @@ float4 SampleWorldPos_{DataInterfaceName}(float2 WorldPos)
 	return SampleWorldPos_{DataInterfaceName}(/*InDataIndex=*/0, WorldPos);
 }
 
+// Deprecated (5.7). Does not take data index, and assumes axis-aligned Z projection.
+float2 GetTexCoords_{DataInterfaceName}(float2 WorldPos, float2 Min, float2 Max)
+{
+	// Inverse lerp
+	const float2 Pos = WorldPos - Min;
+	const float2 Extents = Max - Min;
+
+	if (!all(Extents))
+	{
+		return (float2)0;
+	}
+
+	return Pos / Extents;
+}
+
 #undef PCG_TextureDimension_Texture2D
 #undef PCG_TextureDimension_Texture2DArray
 //#undef PCG_TextureDimension_Texture3D
diff --git a/Engine/Plugins/PCG/Source/PCG/Private/Compute/DataInterfaces/PCGTextureDataInterface.cpp b/Engine/Plugins/PCG/Source/PCG/Private/Compute/DataInterfaces/PCGTextureDataInterface.cpp
--- a/Engine/Plugins/PCG/Source/PCG/Private/Compute/DataInterfaces/PCGTextureDataInterface.cpp
+++ b/Engine/Plugins/PCG/Source/PCG/Private/Compute/DataInterfaces/PCGTextureDataInterface.cpp
@@ -105,9 +105,8 @@ void UPCGTextureDataInterface::GetSupportedInputs(TArray<FShaderFunctionDefiniti
 	OutFunctions.AddDefaulted_GetRef()
 		.SetName(TEXT("GetTexCoords"))
 		.AddReturnType(EShaderFundamentalType::Float, 2)
-		.AddParam(EShaderFundamentalType::Float, 2) // WorldPos
-		.AddParam(EShaderFundamentalType::Float, 2) // Min
-		.AddParam(EShaderFundamentalType::Float, 2); // Max
+		.AddParam(EShaderFundamentalType::Uint) // InDataIndex
+		.AddParam(EShaderFundamentalType::Float, 3); // InWorldPos
 
 	OutFunctions.AddDefaulted_GetRef()
 		.SetName(TEXT("Sample"))
@@ -127,18 +126,25 @@ void UPCGTextureDataInterface::GetSupportedInputs(TArray<FShaderFunctionDefiniti
 		.AddParam(EShaderFundamentalType::Uint) // InDataIndex
 		.AddParam(EShaderFundamentalType::Uint, 2); // InElementIndex
 
-	{
-		// Pseudo-deprecated. Assumes data index is 0.
-		OutFunctions.AddDefaulted_GetRef()
-			.SetName(TEXT("Sample"))
-			.AddReturnType(EShaderFundamentalType::Float, 4)
-			.AddParam(EShaderFundamentalType::Float, 2); // InTextureUV
+	// Deprecated 5.6. Assumes data index is 0.
+	OutFunctions.AddDefaulted_GetRef()
+		.SetName(TEXT("Sample"))
+		.AddReturnType(EShaderFundamentalType::Float, 4)
+		.AddParam(EShaderFundamentalType::Float, 2); // InTextureUV
 
-		OutFunctions.AddDefaulted_GetRef()
-			.SetName(TEXT("SampleWorldPos"))
-			.AddReturnType(EShaderFundamentalType::Float, 4)
-			.AddParam(EShaderFundamentalType::Float, 2); // WorldPos
-	}
+	// Deprecated 5.6. Assumes data index is 0.
+	OutFunctions.AddDefaulted_GetRef()
+		.SetName(TEXT("SampleWorldPos"))
+		.AddReturnType(EShaderFundamentalType::Float, 4)
+		.AddParam(EShaderFundamentalType::Float, 2); // WorldPos
+
+	// Deprecated 5.7. Does not take data index, and assumes axis-aligned Z projection.
+	OutFunctions.AddDefaulted_GetRef()
+		.SetName(TEXT("GetTexCoords"))
+		.AddReturnType(EShaderFundamentalType::Float, 2)
+		.AddParam(EShaderFundamentalType::Float, 2) // WorldPos
+		.AddParam(EShaderFundamentalType::Float, 2) // Min
+		.AddParam(EShaderFundamentalType::Float, 2); // Max
 }
 
 void UPCGTextureDataInterface::GetSupportedOutputs(TArray<FShaderFunctionDefinition>& OutFunctions) const
diff --git a/Engine/Plugins/PCG/Source/PCG/Private/Compute/Elements/PCGCustomHLSL.cpp b/Engine/Plugins/PCG/Source/PCG/Private/Compute/Elements/PCGCustomHLSL.cpp
--- a/Engine/Plugins/PCG/Source/PCG/Private/Compute/Elements/PCGCustomHLSL.cpp
+++ b/Engine/Plugins/PCG/Source/PCG/Private/Compute/Elements/PCGCustomHLSL.cpp
@@ -549,8 +549,8 @@ void UPCGCustomHLSLSettings::UpdateInputDeclarations()
 		InputDeclarations += FString::Format(TEXT(
 			"uint {0}_GetNumData();\n"
 			"uint2 {0}_GetNumElements(uint DataIndex);\n"
-			"// Computes a texture coordinate as if WorldPos lies in a UV plane defined by WorldMin and WorldMax.\n"
-			"float2 {0}_GetTexCoords(float2 WorldPos, float2 WorldMin, float2 WorldMax);\n"
+			"// Computes an unclamped texture coordinate.\n"
+			"float2 {0}_GetTexCoords(uint DataIndex, float3 WorldPos);\n"
 			"float4 {0}_Sample(uint DataIndex, float2 TexCoords);\n"
 			"// Computes sample coordinates of the WorldPos relative to the texture data's bounds.\n"
 			"float4 {0}_SampleWorldPos(uint DataIndex, float2 WorldPos);\n"
