[PCG] Fix potential FindObject call while Garbage collecting

diff --git a/Engine/Plugins/PCG/Source/PCG/Private/Elements/PCGCreateSpline.cpp b/Engine/Plugins/PCG/Source/PCG/Private/Elements/PCGCreateSpline.cpp
--- a/Engine/Plugins/PCG/Source/PCG/Private/Elements/PCGCreateSpline.cpp
+++ b/Engine/Plugins/PCG/Source/PCG/Private/Elements/PCGCreateSpline.cpp
@@ -215,7 +215,16 @@ bool FPCGCreateSplineElement::ExecuteInternal(FPCGContext* Context) const
 		// Execute PostProcess Functions
 		if (SplineActor)
 		{
-			for (UFunction* Function : PCGHelpers::FindUserFunctions(SplineActor->GetClass(), Settings->PostProcessFunctionNames, { UPCGFunctionPrototypes::GetPrototypeWithNoParams() }, Context))
+			// @todo_pcg: for 5.7 move lock inside UPCGFunctionPrototypes::GetPrototypeWithNoParams().
+			// for 5.6.1 keeping fix in cpp to preserve abi.
+			// This can run off the GameThread and UPCGFunctionPrototypes::GetPrototypeWithNoParams() does a FindObject which can't run while GC is running
+			UFunction* FunctionPrototype = nullptr;
+			{
+				FGCScopeGuard Guard;
+				FunctionPrototype = UPCGFunctionPrototypes::GetPrototypeWithNoParams();
+			}
+
+			for (UFunction* Function : PCGHelpers::FindUserFunctions(SplineActor->GetClass(), Settings->PostProcessFunctionNames, { FunctionPrototype }, Context))
 			{
 				SplineActor->ProcessEvent(Function, nullptr);
 			}
