[PCG] Extend texture data interface update to SampleWorldPos as some configurations of the DI can fail due to calling the legacy GetTexCoords function.

#rnx

diff --git a/Engine/Plugins/PCG/Shaders/Private/PCGTextureDataInterface.ush b/Engine/Plugins/PCG/Shaders/Private/PCGTextureDataInterface.ush
--- a/Engine/Plugins/PCG/Shaders/Private/PCGTextureDataInterface.ush
+++ b/Engine/Plugins/PCG/Shaders/Private/PCGTextureDataInterface.ush
@@ -224,7 +224,7 @@ float4 Sample_{DataInterfaceName}(uint InDataIndex, float2 InTextureUV)
 #define PCG_CONDITIONAL_TEXTURE_SAMPLE_WORLD_POS(Index) \
 	if (BindingIndex == Index) \
 	{ \
-		const float2 TextureUV = GetTexCoords_{DataInterfaceName}(WorldPos, {DataInterfaceName}_TextureBounds[Index].xy, {DataInterfaceName}_TextureBounds[Index].zw); \
+		const float2 TextureUV = GetTexCoords_{DataInterfaceName}(Index, WorldPos); \
 		if (TextureDimension == PCG_TextureDimension_Texture2D) \
 		{ \
 			if (PointSample == 0) \
@@ -254,7 +254,7 @@ float4 Sample_{DataInterfaceName}(uint InDataIndex, float2 InTextureUV)
 	} \
 	else
 
-float4 SampleWorldPos_{DataInterfaceName}(uint InDataIndex, float2 WorldPos)
+float4 SampleWorldPos_{DataInterfaceName}(uint InDataIndex, float3 WorldPos)
 {
 	uint BindingIndex = 0;
 	uint SliceIndex = 0;
@@ -276,7 +276,7 @@ float4 SampleWorldPos_{DataInterfaceName}(uint InDataIndex, float2 WorldPos)
 		return (float4)0;
 	}
 #else
-	const float2 TextureUV = GetTexCoords_{DataInterfaceName}(WorldPos, {DataInterfaceName}_TextureBounds[0].xy, {DataInterfaceName}_TextureBounds[0].zw);
+	const float2 TextureUV = GetTexCoords_{DataInterfaceName}(0, WorldPos);
 
 	if (TextureDimension == PCG_TextureDimension_Texture2D)
 	{
@@ -408,7 +408,7 @@ float4 Sample_{DataInterfaceName}(float2 TexCoords)
 // Deprecated (5.6). Fallback to using first data.
 float4 SampleWorldPos_{DataInterfaceName}(float2 WorldPos)
 {
-	return SampleWorldPos_{DataInterfaceName}(/*InDataIndex=*/0, WorldPos);
+	return SampleWorldPos_{DataInterfaceName}(/*InDataIndex=*/0, float3(WorldPos, 0.0f));
 }
 
 // Deprecated (5.7). Does not take data index, and assumes axis-aligned Z projection.
@@ -426,6 +426,12 @@ float2 GetTexCoords_{DataInterfaceName}(float2 WorldPos, float2 Min, float2 Max)
 	return Pos / Extents;
 }
 
+// Deprecated (5.7). Switched to using 3D world pos to support any texture orientation.
+float4 SampleWorldPos_{DataInterfaceName}(uint InDataIndex, float2 InWorldPos)
+{
+	return SampleWorldPos_{DataInterfaceName}(InDataIndex, float3(InWorldPos, 0.0f));
+}
+
 #undef PCG_TextureDimension_Texture2D
 #undef PCG_TextureDimension_Texture2DArray
 //#undef PCG_TextureDimension_Texture3D
diff --git a/Engine/Plugins/PCG/Source/PCG/Private/Compute/DataInterfaces/PCGTextureDataInterface.cpp b/Engine/Plugins/PCG/Source/PCG/Private/Compute/DataInterfaces/PCGTextureDataInterface.cpp
--- a/Engine/Plugins/PCG/Source/PCG/Private/Compute/DataInterfaces/PCGTextureDataInterface.cpp
+++ b/Engine/Plugins/PCG/Source/PCG/Private/Compute/DataInterfaces/PCGTextureDataInterface.cpp
@@ -118,7 +118,7 @@ void UPCGTextureDataInterface::GetSupportedInputs(TArray<FShaderFunctionDefiniti
 		.SetName(TEXT("SampleWorldPos"))
 		.AddReturnType(EShaderFundamentalType::Float, 4)
 		.AddParam(EShaderFundamentalType::Uint) // InDataIndex
-		.AddParam(EShaderFundamentalType::Float, 2); // WorldPos
+		.AddParam(EShaderFundamentalType::Float, 3); // WorldPos
 
 	OutFunctions.AddDefaulted_GetRef()
 		.SetName(TEXT("Load"))
@@ -145,6 +145,13 @@ void UPCGTextureDataInterface::GetSupportedInputs(TArray<FShaderFunctionDefiniti
 		.AddParam(EShaderFundamentalType::Float, 2) // WorldPos
 		.AddParam(EShaderFundamentalType::Float, 2) // Min
 		.AddParam(EShaderFundamentalType::Float, 2); // Max
+
+	// Deprecated 5.7. Switch to using 3D world pos to support any texture orientation.
+	OutFunctions.AddDefaulted_GetRef()
+		.SetName(TEXT("SampleWorldPos"))
+		.AddReturnType(EShaderFundamentalType::Float, 4)
+		.AddParam(EShaderFundamentalType::Uint) // InDataIndex
+		.AddParam(EShaderFundamentalType::Float, 2); // WorldPos
 }
 
 void UPCGTextureDataInterface::GetSupportedOutputs(TArray<FShaderFunctionDefinition>& OutFunctions) const
diff --git a/Engine/Plugins/PCG/Source/PCG/Private/Compute/Elements/PCGCustomHLSL.cpp b/Engine/Plugins/PCG/Source/PCG/Private/Compute/Elements/PCGCustomHLSL.cpp
--- a/Engine/Plugins/PCG/Source/PCG/Private/Compute/Elements/PCGCustomHLSL.cpp
+++ b/Engine/Plugins/PCG/Source/PCG/Private/Compute/Elements/PCGCustomHLSL.cpp
@@ -553,7 +553,7 @@ void UPCGCustomHLSLSettings::UpdateInputDeclarations()
 			"float2 {0}_GetTexCoords(uint DataIndex, float3 WorldPos);\n"
 			"float4 {0}_Sample(uint DataIndex, float2 TexCoords);\n"
 			"// Computes sample coordinates of the WorldPos relative to the texture data's bounds.\n"
-			"float4 {0}_SampleWorldPos(uint DataIndex, float2 WorldPos);\n"
+			"float4 {0}_SampleWorldPos(uint DataIndex, float3 WorldPos);\n"
 			"float4 {0}_Load(uint DataIndex, uint2 ElementIndex);\n"),
 			{ bMultiPin ? PCGHLSLElement::PinDeclTemplateStr : TextureDataPins[0] });
 
