[PCG] Show in ALV if property is allocated

#rnx

diff --git a/Engine/Plugins/PCG/Source/PCGEditor/Private/DataVisualizations/PCGDataVisualizationHelpers.cpp b/Engine/Plugins/PCG/Source/PCGEditor/Private/DataVisualizations/PCGDataVisualizationHelpers.cpp
--- a/Engine/Plugins/PCG/Source/PCGEditor/Private/DataVisualizations/PCGDataVisualizationHelpers.cpp
+++ b/Engine/Plugins/PCG/Source/PCGEditor/Private/DataVisualizations/PCGDataVisualizationHelpers.cpp
@@ -15,6 +15,9 @@ namespace PCGDataVisualizationConstants
 
 	/** Labels of the columns. */
 	const FText TEXT_Index = LOCTEXT("Index", "$Index");
+
+	/** Text to flag an attribute as being full-width allocated, instead of representing the range of elements with a single value. */
+	const FText TEXT_ConstantValueCompressedAttribute = LOCTEXT("ConstantValueCompressedAttribute", "Utilizes constant value compression - a single value is shared by all entries.");
 }
 
 namespace PCGDataVisualizationHelpers
@@ -30,17 +33,28 @@ namespace PCGDataVisualizationHelpers
 		{
 			const uint16 Type = Overrides.TypeOverride == EPCGMetadataTypes::Unknown ? ColumnInfo.Accessor->GetUnderlyingType() : static_cast<uint16>(Overrides.TypeOverride);
 
-			if (Overrides.bAddTypeToTooltip && !Overrides.TooltipOverride.IsEmpty())
+			if (!Overrides.TooltipOverride.IsEmpty())
 			{
-				ColumnInfo.Tooltip = FText::Format(INVTEXT("{0} ({1})"), Overrides.TooltipOverride, PCG::Private::GetTypeNameText(Type));
+				ColumnInfo.Tooltip = Overrides.TooltipOverride;
 			}
-			else if (!Overrides.TooltipOverride.IsEmpty())
+
+			if (Overrides.bIsConstantValueCompressed)
 			{
-				ColumnInfo.Tooltip = Overrides.TooltipOverride;
+				ColumnInfo.Tooltip = FText::Format(INVTEXT("{0} {1}"), ColumnInfo.Tooltip, PCGDataVisualizationConstants::TEXT_ConstantValueCompressedAttribute);
 			}
-			else
+
+			if (Overrides.bAddTypeToTooltip)
 			{
-				ColumnInfo.Tooltip = PCG::Private::GetTypeNameText(Type);
+				const FText TypeNameText = PCG::Private::GetTypeNameText(Type);
+
+				if (!ColumnInfo.Tooltip.IsEmpty())
+				{
+					ColumnInfo.Tooltip = FText::Format(INVTEXT("{0} ({1})"), ColumnInfo.Tooltip, TypeNameText);
+				}
+				else
+				{
+					ColumnInfo.Tooltip = TypeNameText;
+				}
 			}
 
 			if (PCG::Private::IsOfTypes<FString>(Type))
diff --git a/Engine/Plugins/PCG/Source/PCGEditor/Private/DataVisualizations/PCGSpatialDataVisualization.cpp b/Engine/Plugins/PCG/Source/PCGEditor/Private/DataVisualizations/PCGSpatialDataVisualization.cpp
--- a/Engine/Plugins/PCG/Source/PCGEditor/Private/DataVisualizations/PCGSpatialDataVisualization.cpp
+++ b/Engine/Plugins/PCG/Source/PCGEditor/Private/DataVisualizations/PCGSpatialDataVisualization.cpp
@@ -275,15 +275,17 @@ FPCGTableVisualizerInfo IPCGSpatialDataVisualization::GetTableVisualizerInfoWith
 	AddColumnInfo(Info, PointData, FPCGAttributePropertySelector::CreateExtraPropertySelector(EPCGExtraProperties::Index));
 	Info.SortingColumn = Info.ColumnInfos.Last().Id;
 
-	AddPropertyEnumColumnInfo<FVector>(Info, PointData, EPCGPointProperties::Position);
-	AddPropertyEnumColumnInfo<FRotator>(Info, PointData, EPCGPointProperties::Rotation);
-	AddPropertyEnumColumnInfo<FVector>(Info, PointData, EPCGPointProperties::Scale);
-	AddPropertyEnumColumnInfo<FVector>(Info, PointData, EPCGPointProperties::BoundsMin);
-	AddPropertyEnumColumnInfo<FVector>(Info, PointData, EPCGPointProperties::BoundsMax);
-	AddPropertyEnumColumnInfo<FLinearColor>(Info, PointData, EPCGPointProperties::Color);
-	AddPropertyEnumColumnInfo<float>(Info, PointData, EPCGPointProperties::Density);
-	AddPropertyEnumColumnInfo<float>(Info, PointData, EPCGPointProperties::Steepness);
-	AddPropertyEnumColumnInfo<int32>(Info, PointData, EPCGPointProperties::Seed);
+	const EPCGPointNativeProperties AllocatedProperties = PointData->GetAllocatedProperties();
+
+	AddPropertyEnumColumnInfo<FVector>(Info, PointData, EPCGPointProperties::Position, { .bIsConstantValueCompressed = !(AllocatedProperties & EPCGPointNativeProperties::Transform) });
+	AddPropertyEnumColumnInfo<FRotator>(Info, PointData, EPCGPointProperties::Rotation, { .bIsConstantValueCompressed = !(AllocatedProperties & EPCGPointNativeProperties::Transform) });
+	AddPropertyEnumColumnInfo<FVector>(Info, PointData, EPCGPointProperties::Scale, { .bIsConstantValueCompressed = !(AllocatedProperties & EPCGPointNativeProperties::Transform) });
+	AddPropertyEnumColumnInfo<FVector>(Info, PointData, EPCGPointProperties::BoundsMin, { .bIsConstantValueCompressed = !(AllocatedProperties & EPCGPointNativeProperties::BoundsMin) });
+	AddPropertyEnumColumnInfo<FVector>(Info, PointData, EPCGPointProperties::BoundsMax, { .bIsConstantValueCompressed = !(AllocatedProperties & EPCGPointNativeProperties::BoundsMax) });
+	AddPropertyEnumColumnInfo<FLinearColor>(Info, PointData, EPCGPointProperties::Color, { .bIsConstantValueCompressed = !(AllocatedProperties & EPCGPointNativeProperties::Color) });
+	AddPropertyEnumColumnInfo<float>(Info, PointData, EPCGPointProperties::Density, { .bIsConstantValueCompressed = !(AllocatedProperties & EPCGPointNativeProperties::Density) });
+	AddPropertyEnumColumnInfo<float>(Info, PointData, EPCGPointProperties::Steepness, { .bIsConstantValueCompressed = !(AllocatedProperties & EPCGPointNativeProperties::Steepness) });
+	AddPropertyEnumColumnInfo<int32>(Info, PointData, EPCGPointProperties::Seed, { .bIsConstantValueCompressed = !(AllocatedProperties & EPCGPointNativeProperties::Seed) });
 
 	if (FPCGEditorCommon::CVarShowAdvancedAttributesFields.GetValueOnAnyThread())
 	{
@@ -295,6 +297,8 @@ FPCGTableVisualizerInfo IPCGSpatialDataVisualization::GetTableVisualizerInfoWith
 			MetadataEntrySelector.SetPropertyName(*StaticEnum<EPCGPointNativeProperties>()->GetNameStringByValue((int64)EPCGPointNativeProperties::MetadataEntry));
 			return TSharedPtr<const IPCGAttributeAccessor>(PCGAttributeAccessorHelpers::CreateConstAccessor(PointData, MetadataEntrySelector).Release());
 		};
+
+		Overrides.bIsConstantValueCompressed = !(AllocatedProperties & EPCGPointNativeProperties::MetadataEntry);
 		
 		AddTypedColumnInfo<int64>(Info, PointData, FPCGAttributePropertySelector{}, Overrides);
 
diff --git a/Engine/Plugins/PCG/Source/PCGEditor/Public/DataVisualizations/PCGDataVisualizationHelpers.h b/Engine/Plugins/PCG/Source/PCGEditor/Public/DataVisualizations/PCGDataVisualizationHelpers.h
--- a/Engine/Plugins/PCG/Source/PCGEditor/Public/DataVisualizations/PCGDataVisualizationHelpers.h
+++ b/Engine/Plugins/PCG/Source/PCGEditor/Public/DataVisualizations/PCGDataVisualizationHelpers.h
@@ -20,6 +20,7 @@ namespace PCGDataVisualizationHelpers
 		FText LabelOverride;
 		FText TooltipOverride;
 		bool bAddTypeToTooltip = true;
+		bool bIsConstantValueCompressed = false;
 		EPCGMetadataTypes TypeOverride = EPCGMetadataTypes::Unknown;
 		TFunction<TSharedPtr<const IPCGAttributeAccessor>()> CreateAccessorFuncOverride;
 		TFunction<TSharedPtr<const IPCGAttributeAccessorKeys>()> CreateAccessorKeysFuncOverride;
@@ -96,7 +97,7 @@ namespace PCGDataVisualizationHelpers
 	{
 		const UEnum* EnumClass = StaticEnum<EnumType>();
 		check(EnumClass);
-		AddPropertyEnumColumnInfo<T>(OutInfo, Data, EnumClass, static_cast<int64>(EnumValue));
+		AddPropertyEnumColumnInfo<T>(OutInfo, Data, EnumClass, static_cast<int64>(EnumValue), Overrides);
 	}
 
 	template <> PCGEDITOR_API void AddTypedColumnInfo_Impl<FVector2D>(FPCGTableVisualizerInfo& OutInfo, const UPCGData* Data, const FPCGAttributePropertySelector& InSelector, const FColumnInfoOverrides& Overrides);
