[PCG] Add PCG options menu to data viewport
* Added 'Auto Focus' toggle to allow control over whether the data viewport automatically focuses the inspected data on refresh.
* IPCGDataVisualizations no longer focus during setup. Now they should set a focus bounds/ortho zoom.
* Added OnFocusViewportToSelection() so that 'F' hotkey will focus the viewport to the focus bounds/ortho zoom.

diff --git a/Engine/Plugins/PCG/Source/PCG/Public/PCGDataVisualization.h b/Engine/Plugins/PCG/Source/PCG/Public/PCGDataVisualization.h
--- a/Engine/Plugins/PCG/Source/PCG/Public/PCGDataVisualization.h
+++ b/Engine/Plugins/PCG/Source/PCG/Public/PCGDataVisualization.h
@@ -7,6 +7,8 @@
 #include "Metadata/PCGMetadataCommon.h"
 #include "Metadata/Accessors/IPCGAttributeAccessor.h"
 
+#include "EngineDefines.h"
+
 #if WITH_EDITOR
 class AActor;
 class UPCGData;
@@ -27,6 +29,12 @@ struct FPCGSceneSetupParams
 
 	/** Any UObjects created during scene setup should be tracked here to be kept visible for GC. */
 	TArray<TObjectPtr<UObject>> ManagedResources;
+
+	/** Used when focusing the data viewport to the visualization. */
+	FBoxSphereBounds FocusBounds = FBoxSphereBounds(EForceInit::ForceInit);
+
+	/** Used when focusing the data viewport to the visualization for orthographic views. */
+	float FocusOrthoZoom = DEFAULT_ORTHOZOOM;
 };
 
 using FPCGSetupSceneFunc = TFunction<void(FPCGSceneSetupParams& InOutParams)>;
diff --git a/Engine/Plugins/PCG/Source/PCGEditor/Private/DataVisualizations/PCGBaseTextureDataVisualization.cpp b/Engine/Plugins/PCG/Source/PCGEditor/Private/DataVisualizations/PCGBaseTextureDataVisualization.cpp
--- a/Engine/Plugins/PCG/Source/PCGEditor/Private/DataVisualizations/PCGBaseTextureDataVisualization.cpp
+++ b/Engine/Plugins/PCG/Source/PCGEditor/Private/DataVisualizations/PCGBaseTextureDataVisualization.cpp
@@ -94,24 +94,23 @@ FPCGSetupSceneFunc FPCGBaseTextureDataVisualization::GetViewportSetupFunc(const
 
 		static const FTransform MeshTransform(FRotator(0.0, -90.0, 0.0), FVector::ZeroVector, FVector::OneVector);
 		InOutParams.Scene->AddComponent(MeshComponent, MeshTransform);
+
+		// @todo_pcg: These settings should all be exposed through InOutParams. Textures should probably have their own preview scene profile that gets selected automatically.
 		InOutParams.Scene->SetFloorVisibility(false);
 		InOutParams.Scene->SetEnvironmentVisibility(false);
 		InOutParams.EditorViewportClient->SetViewportType(ELevelViewportType::LVT_OrthoTop);
 		InOutParams.EditorViewportClient->SetViewMode(EViewModeIndex::VMI_Unlit);
-		InOutParams.EditorViewportClient->SetViewLocation(FVector::UpVector);
 
 		const FIntPoint ViewportSize = InOutParams.EditorViewportClient->Viewport->GetSizeXY();
 
 		if (ViewportSize.X > 0 && ViewportSize.Y > 0)
 		{
 			// Bounds will be updated already by SetStaticMesh() call.
-			const FBoxSphereBounds Bounds = MeshComponent->Bounds;
+			InOutParams.FocusBounds = MeshComponent->Bounds;
 
 			// Fit the orthographic zoom to the mesh. 0.8f is chosen arbitrarily to add some padding around the mesh.
-			const float MeshUnitsPerPixel = FMath::Max(Bounds.BoxExtent.X / ViewportSize.X, Bounds.BoxExtent.Y / ViewportSize.Y) * 2.0f;
-			const float OrthoZoom = FMath::Clamp(MeshUnitsPerPixel * DEFAULT_ORTHOZOOM * 0.8f, MIN_ORTHOZOOM, MAX_ORTHOZOOM);
-
-			InOutParams.EditorViewportClient->SetOrthoZoom(OrthoZoom);
+			const float MeshUnitsPerPixel = FMath::Max(InOutParams.FocusBounds.BoxExtent.X / ViewportSize.X, InOutParams.FocusBounds.BoxExtent.Y / ViewportSize.Y) * 2.0f;
+			InOutParams.FocusOrthoZoom = FMath::Clamp(MeshUnitsPerPixel * DEFAULT_ORTHOZOOM * 0.8f, MIN_ORTHOZOOM, MAX_ORTHOZOOM);
 		}
 	};
 }
diff --git a/Engine/Plugins/PCG/Source/PCGEditor/Private/DataVisualizations/PCGSpatialDataVisualization.cpp b/Engine/Plugins/PCG/Source/PCGEditor/Private/DataVisualizations/PCGSpatialDataVisualization.cpp
--- a/Engine/Plugins/PCG/Source/PCGEditor/Private/DataVisualizations/PCGSpatialDataVisualization.cpp
+++ b/Engine/Plugins/PCG/Source/PCGEditor/Private/DataVisualizations/PCGSpatialDataVisualization.cpp
@@ -457,14 +457,7 @@ FPCGSetupSceneFunc IPCGSpatialDataVisualization::GetViewportSetupFunc(const UPCG
 
 		if (bInitializedBounds)
 		{
-			const FBoxSphereBounds Bounds = FBoxSphereBounds(FBox(BoundsMin, BoundsMax));
-
-			InOutParams.EditorViewportClient->SetViewLocation(Bounds.Origin);
-
-			if (Bounds.SphereRadius > 0.0f)
-			{
-				InOutParams.EditorViewportClient->SetViewLocationForOrbiting(Bounds.Origin, Bounds.SphereRadius * 2.0f);
-			}
+			InOutParams.FocusBounds = FBoxSphereBounds(FBox(BoundsMin, BoundsMax));
 		}
 	};
 }
diff --git a/Engine/Plugins/PCG/Source/PCGEditor/Private/DataVisualizations/PCGSplineDataVisualization.cpp b/Engine/Plugins/PCG/Source/PCGEditor/Private/DataVisualizations/PCGSplineDataVisualization.cpp
--- a/Engine/Plugins/PCG/Source/PCGEditor/Private/DataVisualizations/PCGSplineDataVisualization.cpp
+++ b/Engine/Plugins/PCG/Source/PCGEditor/Private/DataVisualizations/PCGSplineDataVisualization.cpp
@@ -164,15 +164,7 @@ FPCGSetupSceneFunc IPCGSplineDataVisualization::GetViewportSetupFunc(const UPCGD
 
 		InOutParams.ManagedResources.Add(SplineComponent);
 		InOutParams.Scene->AddComponent(SplineComponent, SplineComponent->GetComponentToWorld());
-
-		FBoxSphereBounds Bounds = SplineComponent->CalcBounds(SplineComponent->GetComponentToWorld());
-
-		InOutParams.EditorViewportClient->SetViewLocation(Bounds.Origin);
-
-		if (Bounds.SphereRadius > 0.0f)
-		{
-			InOutParams.EditorViewportClient->SetViewLocationForOrbiting(Bounds.Origin, Bounds.SphereRadius * 2.0f);
-		}
+		InOutParams.FocusBounds = SplineComponent->CalcBounds(SplineComponent->GetComponentToWorld());
 	};
 }
 
diff --git a/Engine/Plugins/PCG/Source/PCGEditor/Private/DataVisualizations/PCGStaticMeshDataVisualization.cpp b/Engine/Plugins/PCG/Source/PCGEditor/Private/DataVisualizations/PCGStaticMeshDataVisualization.cpp
--- a/Engine/Plugins/PCG/Source/PCGEditor/Private/DataVisualizations/PCGStaticMeshDataVisualization.cpp
+++ b/Engine/Plugins/PCG/Source/PCGEditor/Private/DataVisualizations/PCGStaticMeshDataVisualization.cpp
@@ -57,13 +57,6 @@ FPCGSetupSceneFunc IPCGStaticMeshDataVisualization::GetViewportSetupFunc(const U
 		MeshComponent->SetStaticMesh(StaticMesh);
 
 		// Bounds will be updated already by SetStaticMesh() call.
-		const FBoxSphereBounds Bounds = MeshComponent->Bounds;
-
-		InOutParams.EditorViewportClient->SetViewLocation(Bounds.Origin);
-
-		if (Bounds.SphereRadius > 0.0f)
-		{
-			InOutParams.EditorViewportClient->SetViewLocationForOrbiting(Bounds.Origin, Bounds.SphereRadius * 2.0f);
-		}
+		InOutParams.FocusBounds = MeshComponent->Bounds;
 	};
 }
diff --git a/Engine/Plugins/PCG/Source/PCGEditor/Private/DataVisualizations/PCGVolumeDataVisualization.cpp b/Engine/Plugins/PCG/Source/PCGEditor/Private/DataVisualizations/PCGVolumeDataVisualization.cpp
--- a/Engine/Plugins/PCG/Source/PCGEditor/Private/DataVisualizations/PCGVolumeDataVisualization.cpp
+++ b/Engine/Plugins/PCG/Source/PCGEditor/Private/DataVisualizations/PCGVolumeDataVisualization.cpp
@@ -97,15 +97,7 @@ FPCGSetupSceneFunc FPCGPrimitiveDataVisualization::GetViewportSetupFunc(const UP
 			}
 		}
 
-		if (bHasBounds)
-		{
-			InOutParams.EditorViewportClient->SetViewLocation(Bounds.Origin);
-
-			if (Bounds.SphereRadius > 0.0f)
-			{
-				InOutParams.EditorViewportClient->SetViewLocationForOrbiting(Bounds.Origin, Bounds.SphereRadius * 2.0f);
-			}
-		}
+		InOutParams.FocusBounds = Bounds;
 	};
 }
 
diff --git a/Engine/Plugins/PCG/Source/PCGEditor/Private/PCGEditorCommands.cpp b/Engine/Plugins/PCG/Source/PCGEditor/Private/PCGEditorCommands.cpp
--- a/Engine/Plugins/PCG/Source/PCGEditor/Private/PCGEditorCommands.cpp
+++ b/Engine/Plugins/PCG/Source/PCGEditor/Private/PCGEditorCommands.cpp
@@ -44,6 +44,7 @@ void FPCGEditorCommands::RegisterCommands()
 	UI_COMMAND(SelectNamedRerouteUsages, "Select Named Reroute Usages", "Selects all usages of this Named Reroute Declaration", EUserInterfaceActionType::Button, FInputChord());
 	UI_COMMAND(SelectNamedRerouteDeclaration, "Select Named Reroute Declaration", "Selects the associated Named Reroute Declaration matching this Named Reroute Usage", EUserInterfaceActionType::Button, FInputChord());
 	UI_COMMAND(JumpToSource, "Jump to Source", "Jumps to the associated source file.", EUserInterfaceActionType::Button, FInputChord());
+	UI_COMMAND(AutoFocusViewport, "Auto Focus", "When inspection updates the viewport will re-focus to the data.", EUserInterfaceActionType::ToggleButton, FInputChord());
 
 	//~ PLATFORM SPECIFIC ~//
 #if PLATFORM_MAC
diff --git a/Engine/Plugins/PCG/Source/PCGEditor/Private/Widgets/AssetEditorViewport/PCGEditorViewportToolbarSections.cpp b/Engine/Plugins/PCG/Source/PCGEditor/Private/Widgets/AssetEditorViewport/PCGEditorViewportToolbarSections.cpp
new file mode 100644
--- /dev/null
+++ b/Engine/Plugins/PCG/Source/PCGEditor/Private/Widgets/AssetEditorViewport/PCGEditorViewportToolbarSections.cpp
@@ -0,0 +1,38 @@
+// Copyright Epic Games, Inc. All Rights Reserved.
+
+#include "Widgets/AssetEditorViewport/PCGEditorViewportToolbarSections.h"
+
+#include "PCGEditorCommands.h"
+#include "PCGEditorStyle.h"
+
+#include "ToolMenu.h"
+#include "ToolMenuSection.h"
+
+#define LOCTEXT_NAMESPACE "PCGEditorViewportToolbarSections"
+
+namespace UE::PCGEditor
+{
+	FToolMenuEntry CreatePCGViewportSubmenu()
+	{
+		return FToolMenuEntry::InitDynamicEntry(TEXT("PCGMenu"), FNewToolMenuSectionDelegate::CreateLambda([](FToolMenuSection& Section)
+		{
+			FToolMenuEntry& Entry = Section.AddSubMenu(
+				/*Name=*/TEXT("PCGSubmenu"),
+				/*Label=*/FText::GetEmpty(),
+				/*Tooltip=*/FText::GetEmpty(),
+				FNewToolMenuDelegate::CreateLambda([](UToolMenu* Submenu)
+				{
+					// Add action buttons to this section.
+					//FToolMenuSection& ActionsSection = Submenu->FindOrAddSection(TEXT("PCGActionsSection"), LOCTEXT("PCGActionsSectionLabel", "Actions"));
+
+					// Add preference toggles to this section.
+					FToolMenuSection& PreferencesSection = Submenu->FindOrAddSection(TEXT("PCGPreferencesSection"), LOCTEXT("PCGPreferencesSectionLabel", "Preferences"));
+					PreferencesSection.AddMenuEntry(FPCGEditorCommands::Get().AutoFocusViewport);
+				}),
+				/*bInOpenSubMenuOnClick=*/false,
+				FSlateIcon(FPCGEditorStyle::Get().GetStyleSetName(), "PCG.EditorIcon"));
+		}));
+	}
+}
+
+#undef LOCTEXT_NAMESPACE
diff --git a/Engine/Plugins/PCG/Source/PCGEditor/Private/Widgets/AssetEditorViewport/PCGEditorViewportToolbarSections.h b/Engine/Plugins/PCG/Source/PCGEditor/Private/Widgets/AssetEditorViewport/PCGEditorViewportToolbarSections.h
new file mode 100644
--- /dev/null
+++ b/Engine/Plugins/PCG/Source/PCGEditor/Private/Widgets/AssetEditorViewport/PCGEditorViewportToolbarSections.h
@@ -0,0 +1,10 @@
+// Copyright Epic Games, Inc. All Rights Reserved.
+
+#pragma once
+
+#include "ToolMenuEntry.h"
+
+namespace UE::PCGEditor
+{
+	FToolMenuEntry CreatePCGViewportSubmenu();
+}
diff --git a/Engine/Plugins/PCG/Source/PCGEditor/Private/Widgets/AssetEditorViewport/SPCGEditorViewport.cpp b/Engine/Plugins/PCG/Source/PCGEditor/Private/Widgets/AssetEditorViewport/SPCGEditorViewport.cpp
--- a/Engine/Plugins/PCG/Source/PCGEditor/Private/Widgets/AssetEditorViewport/SPCGEditorViewport.cpp
+++ b/Engine/Plugins/PCG/Source/PCGEditor/Private/Widgets/AssetEditorViewport/SPCGEditorViewport.cpp
@@ -2,6 +2,10 @@
 
 #include "Widgets/AssetEditorViewport/SPCGEditorViewport.h"
 
+#include "PCGEditorCommands.h"
+#include "PCGEditorSettings.h"
+#include "Widgets/AssetEditorViewport/PCGEditorViewportToolbarSections.h"
+
 #include "AdvancedPreviewScene.h"
 #include "AdvancedPreviewSceneMenus.h"
 #include "AssetViewerSettings.h"
@@ -127,9 +131,6 @@ void FPCGEditorViewportClient::Reset()
 	TRACE_CPUPROFILER_EVENT_SCOPE(FPCGEditorViewportClient::Reset);
 
 	SetViewportType(ELevelViewportType::LVT_Perspective);
-	SetViewLocation(PCGEditorViewportConstants::DefaultViewLocation);
-	SetViewRotation(PCGEditorViewportConstants::DefaultViewRotation);
-	SetViewLocationForOrbiting(PCGEditorViewportConstants::DefaultViewLocation, PCGEditorViewportConstants::DefaultOrbitDistance);
 
 	Invalidate();
 }
@@ -187,6 +188,13 @@ void SPCGEditorViewport::SetupScene(const TArray<UObject*>& InResources, const F
 		SetupFunc(SceneSetupParams);
 
 		ManagedResources = MoveTemp(SceneSetupParams.ManagedResources);
+		FocusBounds = MoveTemp(SceneSetupParams.FocusBounds);
+		FocusOrthoZoom = SceneSetupParams.FocusOrthoZoom;
+
+		if (IsAutoFocusViewportChecked())
+		{
+			OnFocusViewportToSelection();
+		}
 
 		EditorViewportClient->Invalidate();
 	}
@@ -213,6 +221,14 @@ void SPCGEditorViewport::ResetScene()
 	}
 
 	ReleaseManagedResources();
+
+	FocusBounds = FBoxSphereBounds(EForceInit::ForceInit);
+	FocusOrthoZoom = DEFAULT_ORTHOZOOM;
+
+	if (IsAutoFocusViewportChecked())
+	{
+		OnFocusViewportToSelection();
+	}
 }
 
 TSharedRef<FEditorViewportClient> SPCGEditorViewport::MakeEditorViewportClient()
@@ -250,7 +266,12 @@ TSharedPtr<SWidget> SPCGEditorViewport::BuildViewportToolbar()
 		RightSection.AddEntry(UE::UnrealEd::CreateViewModesSubmenu());
 		RightSection.AddEntry(UE::UnrealEd::CreateDefaultShowSubmenu());
 		RightSection.AddEntry(UE::UnrealEd::CreateAssetViewerProfileSubmenu());
-		UE::AdvancedPreviewScene::Menus::ExtendAdvancedPreviewSceneSettings("PCG.ViewportToolbar.AssetViewerProfile");
+		UE::AdvancedPreviewScene::Menus::ExtendAdvancedPreviewSceneSettings(
+			"PCG.ViewportToolbar.AssetViewerProfile",
+			UE::AdvancedPreviewScene::Menus::FSettingsOptions().ShowToggleGrid(false)
+		);
+
+		RightSection.AddEntry(UE::PCGEditor::CreatePCGViewportSubmenu());
 	}
 
 	FToolMenuContext Context;
@@ -269,6 +290,36 @@ TSharedPtr<IPreviewProfileController> SPCGEditorViewport::CreatePreviewProfileCo
 	return MakeShared<FPreviewProfileController>();
 }
 
+void SPCGEditorViewport::BindCommands()
+{
+	SAssetEditorViewport::BindCommands();
+
+	const FPCGEditorCommands& Commands = FPCGEditorCommands::Get();
+
+	CommandList->MapAction(
+		Commands.AutoFocusViewport,
+		FExecuteAction::CreateSP(this, &SPCGEditorViewport::ToggleAutoFocusViewport),
+		FCanExecuteAction(),
+		FIsActionChecked::CreateSP(this, &SPCGEditorViewport::IsAutoFocusViewportChecked));
+}
+
+void SPCGEditorViewport::OnFocusViewportToSelection()
+{
+	check(EditorViewportClient.IsValid());
+
+	if (FocusBounds != FBoxSphereBounds(EForceInit::ForceInit))
+	{
+		// In some cases the focus is a little bit too close, so expanding the bounds slightly is helpful.
+		const FBoxSphereBounds ExpandedFocusBounds = FocusBounds.ExpandBy(FocusBounds.SphereRadius * 0.1f);
+		EditorViewportClient->FocusViewportOnBox(ExpandedFocusBounds.GetBox());
+
+		if (EditorViewportClient->IsOrtho())
+		{
+			EditorViewportClient->SetOrthoZoom(FocusOrthoZoom);
+		}
+	}
+}
+
 void SPCGEditorViewport::ReleaseManagedResources()
 {
 	ManagedResources.Empty();
@@ -278,3 +329,14 @@ void SPCGEditorViewport::AddReferencedObjects(FReferenceCollector& Collector)
 {
 	Collector.AddReferencedObjects(ManagedResources);
 }
+
+void SPCGEditorViewport::ToggleAutoFocusViewport()
+{
+	UPCGEditorSettings* PCGEditorSettings = GetMutableDefault<UPCGEditorSettings>();
+	PCGEditorSettings->bAutoFocusViewport = !PCGEditorSettings->bAutoFocusViewport;
+}
+
+bool SPCGEditorViewport::IsAutoFocusViewportChecked() const
+{
+	return GetDefault<UPCGEditorSettings>()->bAutoFocusViewport;
+}
diff --git a/Engine/Plugins/PCG/Source/PCGEditor/Private/Widgets/AssetEditorViewport/SPCGEditorViewport.h b/Engine/Plugins/PCG/Source/PCGEditor/Private/Widgets/AssetEditorViewport/SPCGEditorViewport.h
--- a/Engine/Plugins/PCG/Source/PCGEditor/Private/Widgets/AssetEditorViewport/SPCGEditorViewport.h
+++ b/Engine/Plugins/PCG/Source/PCGEditor/Private/Widgets/AssetEditorViewport/SPCGEditorViewport.h
@@ -40,14 +40,25 @@ protected:
 	virtual TSharedRef<FEditorViewportClient> MakeEditorViewportClient() override;
 	virtual TSharedPtr<SWidget> BuildViewportToolbar() override;
 	virtual TSharedPtr<IPreviewProfileController> CreatePreviewProfileController() override;
+	virtual void BindCommands() override;
+	virtual void OnFocusViewportToSelection() override;
 	//~ End SEditorViewport Interface
 
 	void ReleaseManagedResources();
 
-private:
+	void ToggleAutoFocusViewport();
+	bool IsAutoFocusViewportChecked() const;
+
+protected:
 	TSharedPtr<FPCGEditorViewportClient> EditorViewportClient = nullptr;
 	TSharedPtr<FAdvancedPreviewScene> AdvancedPreviewScene = nullptr;
 
 	/** Objects used by the scene. Must be managed for GC. */
 	TArray<TObjectPtr<UObject>> ManagedResources;
+
+	/** Used when focusing the data viewport to the visualization. */
+	FBoxSphereBounds FocusBounds = FBoxSphereBounds(EForceInit::ForceInit);
+
+	/** Used when focusing the data viewport to the visualization for orthographic views. */
+	float FocusOrthoZoom = DEFAULT_ORTHOZOOM;
 };
diff --git a/Engine/Plugins/PCG/Source/PCGEditor/Public/PCGEditorCommands.h b/Engine/Plugins/PCG/Source/PCGEditor/Public/PCGEditorCommands.h
--- a/Engine/Plugins/PCG/Source/PCGEditor/Public/PCGEditorCommands.h
+++ b/Engine/Plugins/PCG/Source/PCGEditor/Public/PCGEditorCommands.h
@@ -42,6 +42,7 @@ public:
 	TSharedPtr<FUICommandInfo> SelectNamedRerouteUsages;
 	TSharedPtr<FUICommandInfo> SelectNamedRerouteDeclaration;
 	TSharedPtr<FUICommandInfo> JumpToSource;
+	TSharedPtr<FUICommandInfo> AutoFocusViewport;
 };
 
 struct FPCGSpawnNodeCommandInfo
diff --git a/Engine/Plugins/PCG/Source/PCGEditor/Public/PCGEditorSettings.h b/Engine/Plugins/PCG/Source/PCGEditor/Public/PCGEditorSettings.h
--- a/Engine/Plugins/PCG/Source/PCGEditor/Public/PCGEditorSettings.h
+++ b/Engine/Plugins/PCG/Source/PCGEditor/Public/PCGEditorSettings.h
@@ -282,6 +282,10 @@ public:
 	UPROPERTY(EditAnywhere, config, Category = "Graph UI")
 	bool bShowNodeDeterminismActionsRightClickContext = true;
 
+	/** When inspection updates the viewport will re-focus to the data. */
+	UPROPERTY(EditAnywhere, config, Category = "Viewport")
+	bool bAutoFocusViewport = true;
+
 	// Begin deprecated properties
 	UE_DEPRECATED(5.6, "bEnableNavigateToNativeNodes has been replaced by NodeDoubleClickAction.")
 	UPROPERTY()
diff --git a/Engine/Plugins/PCGInterops/PCGGeometryScriptInterop/Source/PCGGeometryScriptInterop/Private/Editor/PCGDynamicMeshDataVisualization.cpp b/Engine/Plugins/PCGInterops/PCGGeometryScriptInterop/Source/PCGGeometryScriptInterop/Private/Editor/PCGDynamicMeshDataVisualization.cpp
--- a/Engine/Plugins/PCGInterops/PCGGeometryScriptInterop/Source/PCGGeometryScriptInterop/Private/Editor/PCGDynamicMeshDataVisualization.cpp
+++ b/Engine/Plugins/PCGInterops/PCGGeometryScriptInterop/Source/PCGGeometryScriptInterop/Private/Editor/PCGDynamicMeshDataVisualization.cpp
@@ -61,15 +61,7 @@ FPCGSetupSceneFunc FPCGDynamicMeshDataVisualization::GetViewportSetupFunc(const
 		}
 
 		InOutParams.Scene->AddComponent(DynamicMeshComponent, FTransform::Identity);
-
-		const FBoxSphereBounds Bounds = DynamicMeshComponent->CalcLocalBounds();
-
-		InOutParams.EditorViewportClient->SetViewLocation(Bounds.Origin);
-
-		if (Bounds.SphereRadius > 0.0f)
-		{
-			InOutParams.EditorViewportClient->SetViewLocationForOrbiting(Bounds.Origin, Bounds.SphereRadius * 2.0f);
-		}
+		InOutParams.FocusBounds = DynamicMeshComponent->CalcLocalBounds();
 	};
 }
 
