[PCG] Moved common Json outside of export attribute node

On behalf of:
Review: 43625720

#rnx

diff --git a/Engine/Plugins/PCG/Source/PCG/Private/Elements/IO/PCGExportSelectedAttributes.cpp b/Engine/Plugins/PCG/Source/PCG/Private/Elements/IO/PCGExportSelectedAttributes.cpp
--- a/Engine/Plugins/PCG/Source/PCG/Private/Elements/IO/PCGExportSelectedAttributes.cpp
+++ b/Engine/Plugins/PCG/Source/PCG/Private/Elements/IO/PCGExportSelectedAttributes.cpp
@@ -4,6 +4,7 @@
 
 #include "Data/PCGBasePointData.h"
 #include "Data/PCGSplineData.h"
+#include "Elements/IO/PCGJsonHelpers.h"
 #include "Metadata/PCGMetadata.h"
 #include "Metadata/PCGMetadataAttributeTraits.h"
 #include "Metadata/Accessors/PCGAttributeAccessorHelpers.h"
@@ -19,8 +20,6 @@
 #include "IDesktopPlatform.h"
 #include "Data/PCGSplineStruct.h"
 #include "Misc/FileHelper.h"
-#include "Serialization/JsonSerializer.h"
-#include "Serialization/JsonWriter.h"
 
 #define LOCTEXT_NAMESPACE "PCGExportSelectedAttributesElement"
 
@@ -37,24 +36,6 @@ namespace PCGExportSelectedAttributes
 	namespace Constants
 	{
 		static constexpr int32 DefaultInlineSourceCount = 16;
-
-		struct FCustomExportVersion
-		{
-			enum Type
-			{
-				InitialVersion = 0,
-
-				// New versions can be added above this line
-				VersionPlusOne,
-				LatestVersion = VersionPlusOne - 1
-			};
-
-			static FName GetFriendlyName() { return FName(TEXT("Initial Version")); }
-			const static FGuid GUID;
-		};
-
-		const FGuid FCustomExportVersion::GUID = FGuid(0x04E74488, 0x4BAC8717, 0xBBB18694, 0x39F8F3CE);
-		FCustomVersionRegistration GRegisterPCGExportSelectedAttributesCustomVersion(FCustomExportVersion::GUID, FCustomExportVersion::LatestVersion, TEXT("PCGExportSelectedAttributes"));
 	}
 
 	namespace Helpers
@@ -135,78 +116,6 @@ namespace PCGExportSelectedAttributes
 				OutString = FString::SanitizeFloat(Value, Precision);
 			}
 		}
-
-		namespace Json
-		{
-			TSharedPtr<FJsonObject> CreateJsonHeaderObject(const bool bAddCustomDataVersion = false, const int32 CustomDataVersion = 0)
-			{
-				TSharedPtr<FJsonObject> JsonHeaderObject = MakeShared<FJsonObject>();
-				JsonHeaderObject->SetNumberField(LOCTEXT("JsonFieldExportVersionNumber", "Export Version").ToString(), Constants::FCustomExportVersion::LatestVersion);
-
-				if (bAddCustomDataVersion)
-				{
-					JsonHeaderObject->SetNumberField(LOCTEXT("JsonFieldCustomVersionNumber", "Custom Data Version").ToString(), CustomDataVersion);
-				}
-
-				return JsonHeaderObject;
-			}
-
-			TSharedPtr<FJsonValueArray> ToJsonArray(const TConstArrayView<double> ArrayView)
-			{
-				TArray<TSharedPtr<FJsonValue>> JsonArray;
-				for (double Value : ArrayView)
-				{
-					JsonArray.Emplace(MakeShared<FJsonValueNumber>(Value));
-				}
-
-				return MakeShared<FJsonValueArray>(JsonArray);
-			}
-
-			// The order matters here in the JsonObject construction.
-			template <typename T>
-			TSharedPtr<FJsonValue> ConvertFloatingPointTypeToJson(const T& Value)
-			{
-				using U = std::remove_cv_t<T>;
-				static_assert(PCG::Private::MetadataTraits<U>::IsFloatingPoint, "Helper is only designed to break floating point compound components.");
-
-				if constexpr (std::is_same_v<U, float> || std::is_same_v<U, double>)
-				{
-					return MakeShared<FJsonValueNumber>(Value);
-				}
-				else if constexpr (std::is_same_v<U, FVector2D>)
-				{
-					return ToJsonArray({Value.X, Value.Y});
-				}
-				else if constexpr (std::is_same_v<U, FVector>)
-				{
-					return ToJsonArray({Value.X, Value.Y, Value.Z});
-				}
-				else if constexpr (std::is_same_v<U, FVector4> || std::is_same_v<U, FQuat>)
-				{
-					return ToJsonArray({Value.X, Value.Y, Value.Z, Value.W});
-				}
-				else if constexpr (std::is_same_v<U, FTransform>)
-				{
-					TSharedPtr<FJsonObject> JsonTransformObject = MakeShared<FJsonObject>();
-					JsonTransformObject->SetArrayField(
-						LOCTEXT("JsonFieldTranslation", "Translation").ToString(),
-						ConvertFloatingPointTypeToJson(Value.GetTranslation())->AsArray());
-					JsonTransformObject->SetArrayField(
-						LOCTEXT("JsonFieldRotation", "Rotation").ToString(),
-						ConvertFloatingPointTypeToJson(Value.GetRotation())->AsArray());
-					JsonTransformObject->SetArrayField(
-						LOCTEXT("JsonFieldScale", "Scale").ToString(),
-						ConvertFloatingPointTypeToJson(Value.GetScale3D())->AsArray());
-
-					return MakeShared<FJsonValueObject>(JsonTransformObject);
-				}
-				else //if constexpr (std::is_same_v<U, FRotator>)
-				{
-					static_assert(std::is_same_v<U, FRotator>, "Unsupported type");
-					return ToJsonArray({Value.Pitch, Value.Yaw, Value.Roll});
-				}
-			}
-		}
 	}
 
 	using FSourceDataArray = TArray<FAccessorData, TInlineAllocator<Constants::DefaultInlineSourceCount>>;
@@ -305,7 +214,7 @@ bool FPCGExportSelectedAttributesElement::ExecuteInternal(FPCGContext* Context)
 		}
 
 		// Set up the custom export version
-		using namespace Constants;
+		using namespace PCG::Json::Constants;
 		Archive->SetCustomVersion(FCustomExportVersion::GUID, FCustomExportVersion::LatestVersion, FCustomExportVersion::GetFriendlyName());
 
 		// Set up the user's custom data version
@@ -320,7 +229,7 @@ bool FPCGExportSelectedAttributesElement::ExecuteInternal(FPCGContext* Context)
 		RootJsonObject = MakeShared<FJsonObject>();
 
 		// Set up the Json header to include custom versions
-		TSharedPtr<FJsonObject> JsonHeaderObject = Helpers::Json::CreateJsonHeaderObject(Settings->bAddCustomDataVersion, Settings->CustomVersion);
+		TSharedPtr<FJsonObject> JsonHeaderObject = PCG::Json::Helpers::CreateJsonHeaderObject(Settings->bAddCustomDataVersion, Settings->CustomVersion);
 		RootJsonObject->SetObjectField(LOCTEXT("JsonHeaderField", "Header").ToString(), std::move(JsonHeaderObject));
 	}
 
@@ -432,7 +341,7 @@ bool FPCGExportSelectedAttributesElement::ExecuteInternal(FPCGContext* Context)
 					SourceJsonValues.Reserve(Values.Num());
 					Algo::Transform(Values, SourceJsonValues, [&Data](const T& Value)
 					{
-						return Helpers::Json::ConvertFloatingPointTypeToJson(Value);
+						return PCG::Json::Helpers::ConvertFloatingPointTypeToJson(Value);
 					});
 				}
 				else // For everything else, automatically broadcast it to string
diff --git a/Engine/Plugins/PCG/Source/PCG/Private/Elements/IO/PCGJsonHelpers.cpp b/Engine/Plugins/PCG/Source/PCG/Private/Elements/IO/PCGJsonHelpers.cpp
new file mode 100644
--- /dev/null
+++ b/Engine/Plugins/PCG/Source/PCG/Private/Elements/IO/PCGJsonHelpers.cpp
@@ -0,0 +1,45 @@
+// Copyright Epic Games, Inc. All Rights Reserved.
+
+#include "Elements/IO/PCGJsonHelpers.h"
+
+#include "Serialization/CustomVersion.h"
+
+#define LOCTEXT_NAMESPACE "PCGJson"
+
+namespace PCG::Json
+{
+	namespace Constants
+	{
+		const FGuid FCustomExportVersion::GUID = FGuid(0x04E74488, 0x4BAC8717, 0xBBB18694, 0x39F8F3CE);
+		FCustomVersionRegistration GRegisterPCGExportSelectedAttributesCustomVersion(FCustomExportVersion::GUID, FCustomExportVersion::LatestVersion, TEXT("PCGExportSelectedAttributes"));
+	}
+
+	namespace Helpers
+	{
+		TSharedPtr<FJsonObject> CreateJsonHeaderObject(const bool bAddCustomDataVersion, const int32 CustomDataVersion)
+		{
+			TSharedPtr<FJsonObject> JsonHeaderObject = MakeShared<FJsonObject>();
+			JsonHeaderObject->SetNumberField(LOCTEXT("JsonFieldExportVersionNumber", "Export Version").ToString(), Constants::FCustomExportVersion::LatestVersion);
+
+			if (bAddCustomDataVersion)
+			{
+				JsonHeaderObject->SetNumberField(LOCTEXT("JsonFieldCustomVersionNumber", "Custom Data Version").ToString(), CustomDataVersion);
+			}
+
+			return JsonHeaderObject;
+		}
+
+		TSharedPtr<FJsonValueArray> ToJsonArray(const TConstArrayView<double> ArrayView)
+		{
+			TArray<TSharedPtr<FJsonValue>> JsonArray;
+			for (double Value : ArrayView)
+			{
+				JsonArray.Emplace(MakeShared<FJsonValueNumber>(Value));
+			}
+
+			return MakeShared<FJsonValueArray>(JsonArray);
+		}
+	}
+}
+
+#undef LOCTEXT_NAMESPACE
\ No newline at end of file
diff --git a/Engine/Plugins/PCG/Source/PCG/Private/Elements/IO/PCGJsonHelpers.h b/Engine/Plugins/PCG/Source/PCG/Private/Elements/IO/PCGJsonHelpers.h
new file mode 100644
--- /dev/null
+++ b/Engine/Plugins/PCG/Source/PCG/Private/Elements/IO/PCGJsonHelpers.h
@@ -0,0 +1,80 @@
+// Copyright Epic Games, Inc. All Rights Reserved.
+
+#pragma once
+
+#include "Metadata/PCGMetadataAttributeTraits.h"
+
+#include "Serialization/JsonSerializer.h"
+#include "Serialization/JsonWriter.h"
+
+namespace PCG::Json
+{
+	namespace Constants
+	{
+		struct FCustomExportVersion
+		{
+			enum Type
+			{
+				InitialVersion = 0,
+
+				// New versions can be added above this line
+				VersionPlusOne,
+				LatestVersion = VersionPlusOne - 1
+			};
+
+			static FName GetFriendlyName() { return FName(TEXT("Initial Version")); }
+			const static FGuid GUID;
+		};
+	}
+
+	namespace Helpers
+	{
+		TSharedPtr<FJsonObject> CreateJsonHeaderObject(const bool bAddCustomDataVersion = false, const int32 CustomDataVersion = 0);
+		TSharedPtr<FJsonValueArray> ToJsonArray(const TConstArrayView<double> ArrayView);
+
+		// The order matters here in the JsonObject construction.
+		template <typename T>
+		TSharedPtr<FJsonValue> ConvertFloatingPointTypeToJson(const T& Value)
+		{
+			using U = std::remove_cv_t<T>;
+			static_assert(PCG::Private::MetadataTraits<U>::IsFloatingPoint, "Helper is only designed to break floating point compound components.");
+
+			if constexpr (std::is_same_v<U, float> || std::is_same_v<U, double>)
+			{
+				return MakeShared<FJsonValueNumber>(Value);
+			}
+			else if constexpr (std::is_same_v<U, FVector2D>)
+			{
+				return ToJsonArray({ Value.X, Value.Y });
+			}
+			else if constexpr (std::is_same_v<U, FVector>)
+			{
+				return ToJsonArray({ Value.X, Value.Y, Value.Z });
+			}
+			else if constexpr (std::is_same_v<U, FVector4> || std::is_same_v<U, FQuat>)
+			{
+				return ToJsonArray({ Value.X, Value.Y, Value.Z, Value.W });
+			}
+			else if constexpr (std::is_same_v<U, FTransform>)
+			{
+				TSharedPtr<FJsonObject> JsonTransformObject = MakeShared<FJsonObject>();
+				JsonTransformObject->SetArrayField(
+					NSLOCTEXT("PCGJson", "JsonFieldTranslation", "Translation").ToString(),
+					ConvertFloatingPointTypeToJson(Value.GetTranslation())->AsArray());
+				JsonTransformObject->SetArrayField(
+					NSLOCTEXT("PCGJson", "JsonFieldRotation", "Rotation").ToString(),
+					ConvertFloatingPointTypeToJson(Value.GetRotation())->AsArray());
+				JsonTransformObject->SetArrayField(
+					NSLOCTEXT("PCGJson", "JsonFieldScale", "Scale").ToString(),
+					ConvertFloatingPointTypeToJson(Value.GetScale3D())->AsArray());
+
+				return MakeShared<FJsonValueObject>(JsonTransformObject);
+			}
+			else //if constexpr (std::is_same_v<U, FRotator>)
+			{
+				static_assert(std::is_same_v<U, FRotator>, "Unsupported type");
+				return ToJsonArray({ Value.Pitch, Value.Yaw, Value.Roll });
+			}
+		}
+	}
+}
\ No newline at end of file
