[PCG] Fixed point to attribute set conversion when not all points have a metadata entry

diff --git a/Engine/Plugins/PCG/Source/PCG/Private/Metadata/PCGMetadataDomain.cpp b/Engine/Plugins/PCG/Source/PCG/Private/Metadata/PCGMetadataDomain.cpp
--- a/Engine/Plugins/PCG/Source/PCG/Private/Metadata/PCGMetadataDomain.cpp
+++ b/Engine/Plugins/PCG/Source/PCG/Private/Metadata/PCGMetadataDomain.cpp
@@ -365,7 +365,7 @@ void FPCGMetadataDomain::InitializeAsCopy(const FPCGMetadataDomainInitializePara
 
 	// If we have a partial copy, it will flatten the metadata, so we don't need a parent.
 	// Otherwise, we keep the parent hierarchy.
-	const bool bPartialCopy = InParams.OptionalEntriesToCopy && InParams.OptionalEntriesToCopy->Num() <= InParams.Parent->GetItemCountForChild();
+	const bool bPartialCopy = InParams.OptionalEntriesToCopy.IsSet();
 	TArray<PCGMetadataEntryKey> NewEntryKeys;
 	TArray<PCGMetadataValueKey> NewValueKeys;
 	if (bPartialCopy)
