[PCG] Allow for any asset that is a subclass to UPCGBlueprintElement to be dragged and dropped in the graph

#rnx

diff --git a/Engine/Plugins/PCG/Source/PCGEditor/Private/PCGEditorUtils.cpp b/Engine/Plugins/PCG/Source/PCGEditor/Private/PCGEditorUtils.cpp
--- a/Engine/Plugins/PCG/Source/PCGEditor/Private/PCGEditorUtils.cpp
+++ b/Engine/Plugins/PCG/Source/PCGEditor/Private/PCGEditorUtils.cpp
@@ -36,10 +36,15 @@
 
 bool PCGEditorUtils::IsAssetPCGBlueprint(const FAssetData& InAssetData)
 {
-	FString InNativeParentClassName = InAssetData.GetTagValueRef<FString>(FBlueprintTags::NativeParentClassPath);
-	FString TargetNativeParentClassName = UPCGBlueprintElement::GetParentClassName();
+	if (InAssetData.AssetClassPath != UBlueprint::StaticClass()->GetClassPathName())
+	{
+		return false;
+	}
+
+	const FString InNativeParentClassName = InAssetData.GetTagValueRef<FString>(FBlueprintTags::NativeParentClassPath);
+	const UClass* NativeParentClass = !InNativeParentClassName.IsEmpty() ? FSoftClassPath(InNativeParentClassName).ResolveClass() : nullptr;
 
-	return InAssetData.AssetClassPath == UBlueprint::StaticClass()->GetClassPathName() && InNativeParentClassName == TargetNativeParentClassName;
+	return NativeParentClass && NativeParentClass->IsChildOf<UPCGBlueprintElement>();
 }
 
 void PCGEditorUtils::GetParentPackagePathAndUniqueName(const UObject* OriginalObject, const FString& NewAssetTentativeName, FString& OutPackagePath, FString& OutUniqueName)
