[PCG] Soften check to ensure and gracefully exit if preconfigured settings fails

#rnx

diff --git a/Engine/Plugins/PCG/Source/PCG/Private/Elements/PCGUserParameterGet.cpp b/Engine/Plugins/PCG/Source/PCG/Private/Elements/PCGUserParameterGet.cpp
--- a/Engine/Plugins/PCG/Source/PCG/Private/Elements/PCGUserParameterGet.cpp
+++ b/Engine/Plugins/PCG/Source/PCG/Private/Elements/PCGUserParameterGet.cpp
@@ -254,22 +254,28 @@ void UPCGUserParameterGetSettings::ApplyPreconfiguredSettings(const FPCGPreConfi
 
 	const UPCGNode* Node = CastChecked<UPCGNode>(GetOuter());
 	UPCGGraph* Graph = Node ? Node->GetGraph() : nullptr;
-	if (const FInstancedPropertyBag* PropertyBag = Graph ? Graph->GetUserParametersStruct() : nullptr)
+	if (!PCGGraphParameter::Helpers::GenerateUniqueName(Graph, PropertyName))
 	{
-		if (!PCGGraphParameter::Helpers::GenerateUniqueName(Graph, PropertyName))
-		{
-			PCGLog::LogErrorOnGraph(FText::Format(LOCTEXT("CouldNotCreateUniqueGraphParameterName", "Could not create a unique graph parameter with name '{0}'. Check if the graph is valid and the name has not reached its maximum."), FText::FromName(PropertyName)));
-			return;
-		}
-		check(PropertyBag->IsPropertyNameValid(PropertyName));
+		UE_LOG(LogPCG, Error, TEXT("Could not create a unique graph parameter with name '%s'. Check if the graph is valid and the name has not reached its maximum."), *PropertyName.ToString());
+		return;
+	}
+	check(FInstancedPropertyBag::IsPropertyNameValid(PropertyName));
 
-		FPropertyBagPropertyDesc PropertyDesc = PCGPropertyHelpers::CreatePropertyBagDescWithMetadataType(PropertyName, NewType);
-		Graph->AddUserParameters({PropertyDesc});
-		const FPropertyBagPropertyDesc* Desc = PropertyBag->FindPropertyDescByName(PropertyName);
-		check(Desc);
+	FPropertyBagPropertyDesc PropertyDesc = PCGPropertyHelpers::CreatePropertyBagDescWithMetadataType(PropertyName, NewType);
+	Graph->AddUserParameters({PropertyDesc});
 
-		PropertyGuid = Desc->ID;
-		PropertyName = Desc->Name;
+	if (FInstancedPropertyBag* PropertyBag = Graph ? Graph->GetMutableUserParametersStruct_Unsafe() : nullptr)
+	{
+		if (const FPropertyBagPropertyDesc* Desc = PropertyBag->FindPropertyDescByName(PropertyName))
+		{
+			PropertyGuid = Desc->ID;
+			PropertyName = Desc->Name;
+		}
+		else
+		{
+			ensure(false);
+			UE_LOG(LogPCG, Error, TEXT("Could not create a unique graph parameter with name '%s'."), *PropertyName.ToString());
+		}
 	}
 }
 
