[PCG] Fixed crash caused with Visualize Attribute and [REDACTED_HANDLE] level domain attribute

#rnx

diff --git a/Engine/Plugins/PCG/Source/PCG/Private/Elements/PCGVisualizeAttribute.cpp b/Engine/Plugins/PCG/Source/PCG/Private/Elements/PCGVisualizeAttribute.cpp
--- a/Engine/Plugins/PCG/Source/PCG/Private/Elements/PCGVisualizeAttribute.cpp
+++ b/Engine/Plugins/PCG/Source/PCG/Private/Elements/PCGVisualizeAttribute.cpp
@@ -74,7 +74,8 @@ bool FPCGVisualizeAttribute::ExecuteInternal(FPCGContext* Context) const
 		}
 
 		TArray<FString> ValuesToString;
-		ValuesToString.SetNum(Keys->GetNum());
+		const int32 NumValues = Keys->GetNum();
+		ValuesToString.SetNum(NumValues);
 
 		// Get all the attribute values at once in range, for efficiency
 		if (!Accessor->GetRange<FString>(ValuesToString, 0, *Keys, EPCGAttributeAccessorFlags::AllowBroadcast))
@@ -134,7 +135,7 @@ bool FPCGVisualizeAttribute::ExecuteInternal(FPCGContext* Context) const
 				CompoundedString += FString::Printf(TEXT("[%s]"), *InputSource.GetDisplayText().ToString());
 			}
 
-			CompoundedString = FString::Printf(TEXT("%s %s"), *CompoundedString, *ValuesToString[I]);
+			CompoundedString = FString::Printf(TEXT("%s %s"), *CompoundedString, *ValuesToString[(I + 1) % NumValues]);
 
 			DebugStrings.Emplace(CompoundedString.ToString(), TransformRanges[I].GetLocation() + Settings->LocalOffset, Settings->Color);
 		}
