[PCG] Fix Build Issue (Static Analysis) for Hotfix CL 43647825

#rnx

diff --git a/Engine/Plugins/PCG/Source/PCG/Private/Elements/PCGUserParameterGet.cpp b/Engine/Plugins/PCG/Source/PCG/Private/Elements/PCGUserParameterGet.cpp
--- a/Engine/Plugins/PCG/Source/PCG/Private/Elements/PCGUserParameterGet.cpp
+++ b/Engine/Plugins/PCG/Source/PCG/Private/Elements/PCGUserParameterGet.cpp
@@ -253,30 +253,32 @@ void UPCGUserParameterGetSettings::ApplyPreconfiguredSettings(const FPCGPreConfi
 	}
 
 	const UPCGNode* Node = CastChecked<UPCGNode>(GetOuter());
-	UPCGGraph* Graph = Node ? Node->GetGraph() : nullptr;
-	if (!PCGGraphParameter::Helpers::GenerateUniqueName(Graph, PropertyName))
-	{
-		UE_LOG(LogPCG, Error, TEXT("Could not create a unique graph parameter with name '%s'. Check if the graph is valid and the name has not reached its maximum."), *PropertyName.ToString());
-		return;
-	}
-	check(FInstancedPropertyBag::IsPropertyNameValid(PropertyName));
-
-	FPropertyBagPropertyDesc PropertyDesc = PCGPropertyHelpers::CreatePropertyBagDescWithMetadataType(PropertyName, NewType);
-	Graph->AddUserParameters({PropertyDesc});
-
-	if (FInstancedPropertyBag* PropertyBag = Graph ? Graph->GetMutableUserParametersStruct_Unsafe() : nullptr)
-	{
-		if (const FPropertyBagPropertyDesc* Desc = PropertyBag->FindPropertyDescByName(PropertyName))
-		{
-			PropertyGuid = Desc->ID;
-			PropertyName = Desc->Name;
-		}
-		else
-		{
-			ensure(false);
-			UE_LOG(LogPCG, Error, TEXT("Could not create a unique graph parameter with name '%s'."), *PropertyName.ToString());
-		}
-	}
+	if (UPCGGraph* Graph = Node ? Node->GetGraph() : nullptr)
+    {
+        if (!PCGGraphParameter::Helpers::GenerateUniqueName(Graph, PropertyName))
+        {
+	        UE_LOG(LogPCG, Error, TEXT("Could not create a unique graph parameter with name '%s'. Check if the graph is valid and the name has not reached its maximum."), *PropertyName.ToString());
+	        return;
+        }
+        check(FInstancedPropertyBag::IsPropertyNameValid(PropertyName));
+
+        FPropertyBagPropertyDesc PropertyDesc = PCGPropertyHelpers::CreatePropertyBagDescWithMetadataType(PropertyName, NewType);
+        Graph->AddUserParameters({PropertyDesc});
+
+        if (FInstancedPropertyBag* PropertyBag = Graph->GetMutableUserParametersStruct_Unsafe())
+        {
+	        if (const FPropertyBagPropertyDesc* Desc = PropertyBag->FindPropertyDescByName(PropertyName))
+	        {
+		        PropertyGuid = Desc->ID;
+		        PropertyName = Desc->Name;
+	        }
+	        else
+	        {
+		        ensure(false);
+		        UE_LOG(LogPCG, Error, TEXT("Could not create a unique graph parameter with name '%s'."), *PropertyName.ToString());
+	        }
+        }
+    }
 }
 
 bool UPCGUserParameterGetSettings::ConvertNode(const FPCGPreconfiguredInfo& ConversionInfo)
