[PCG] Added json description of graph structure

diff --git a/Engine/Plugins/PCG/Source/PCG/Private/Elements/IO/PCGExportSelectedAttributes.cpp b/Engine/Plugins/PCG/Source/PCG/Private/Elements/IO/PCGExportSelectedAttributes.cpp
--- a/Engine/Plugins/PCG/Source/PCG/Private/Elements/IO/PCGExportSelectedAttributes.cpp
+++ b/Engine/Plugins/PCG/Source/PCG/Private/Elements/IO/PCGExportSelectedAttributes.cpp
@@ -229,7 +229,7 @@ bool FPCGExportSelectedAttributesElement::ExecuteInternal(FPCGContext* Context)
 		RootJsonObject = MakeShared<FJsonObject>();
 
 		// Set up the Json header to include custom versions
-		TSharedPtr<FJsonObject> JsonHeaderObject = PCG::Json::Helpers::CreateJsonHeaderObject(Settings->bAddCustomDataVersion, Settings->CustomVersion);
+		TSharedPtr<FJsonObject> JsonHeaderObject = PCG::Json::Helpers::CreateJsonHeaderObject(PCG::Json::Constants::FCustomExportVersion::LatestVersion, Settings->bAddCustomDataVersion ? Settings->CustomVersion : -1);
 		RootJsonObject->SetObjectField(LOCTEXT("JsonHeaderField", "Header").ToString(), std::move(JsonHeaderObject));
 	}
 
diff --git a/Engine/Plugins/PCG/Source/PCG/Private/Elements/IO/PCGJsonHelpers.cpp b/Engine/Plugins/PCG/Source/PCG/Private/Elements/IO/PCGJsonHelpers.cpp
--- a/Engine/Plugins/PCG/Source/PCG/Private/Elements/IO/PCGJsonHelpers.cpp
+++ b/Engine/Plugins/PCG/Source/PCG/Private/Elements/IO/PCGJsonHelpers.cpp
@@ -16,14 +16,14 @@ namespace PCG::Json
 
 	namespace Helpers
 	{
-		TSharedPtr<FJsonObject> CreateJsonHeaderObject(const bool bAddCustomDataVersion, const int32 CustomDataVersion)
+		TSharedPtr<FJsonObject> CreateJsonHeaderObject(const int32 DataVersion, const int32 CustomDataVersion)
 		{
 			TSharedPtr<FJsonObject> JsonHeaderObject = MakeShared<FJsonObject>();
-			JsonHeaderObject->SetNumberField(LOCTEXT("JsonFieldExportVersionNumber", "Export Version").ToString(), Constants::FCustomExportVersion::LatestVersion);
+			JsonHeaderObject->SetNumberField(TEXT("Data Version"), DataVersion);
 
-			if (bAddCustomDataVersion)
+			if (CustomDataVersion >= 0)
 			{
-				JsonHeaderObject->SetNumberField(LOCTEXT("JsonFieldCustomVersionNumber", "Custom Data Version").ToString(), CustomDataVersion);
+				JsonHeaderObject->SetNumberField(TEXT("Custom Data Version"), CustomDataVersion);
 			}
 
 			return JsonHeaderObject;
diff --git a/Engine/Plugins/PCG/Source/PCG/Private/Elements/IO/PCGJsonHelpers.h b/Engine/Plugins/PCG/Source/PCG/Private/Elements/IO/PCGJsonHelpers.h
--- a/Engine/Plugins/PCG/Source/PCG/Private/Elements/IO/PCGJsonHelpers.h
+++ b/Engine/Plugins/PCG/Source/PCG/Private/Elements/IO/PCGJsonHelpers.h
@@ -29,33 +29,30 @@ namespace PCG::Json
 
 	namespace Helpers
 	{
-		TSharedPtr<FJsonObject> CreateJsonHeaderObject(const bool bAddCustomDataVersion = false, const int32 CustomDataVersion = 0);
-		TSharedPtr<FJsonValueArray> ToJsonArray(const TConstArrayView<double> ArrayView);
+		TSharedPtr<FJsonObject> CreateJsonHeaderObject(int32 DataVersion, int32 CustomDataVersion = -1);
+		TSharedPtr<FJsonValueArray> ToJsonArray(TConstArrayView<double> ArrayView);
 
-		// The order matters here in the JsonObject construction.
 		template <typename T>
 		TSharedPtr<FJsonValue> ConvertFloatingPointTypeToJson(const T& Value)
 		{
-			using U = std::remove_cv_t<T>;
-			static_assert(PCG::Private::MetadataTraits<U>::IsFloatingPoint, "Helper is only designed to break floating point compound components.");
-
-			if constexpr (std::is_same_v<U, float> || std::is_same_v<U, double>)
+			static_assert(PCG::Private::MetadataTraits<std::remove_cv_t<T>>::IsFloatingPoint, "Helper is only designed to break floating point compound components.");
+			if constexpr (std::is_floating_point_v<std::remove_cv_t<T>>)
 			{
-				return MakeShared<FJsonValueNumber>(Value);
+				return MakeShared<FJsonValueNumber>(static_cast<double>(Value));
 			}
-			else if constexpr (std::is_same_v<U, FVector2D>)
+			else if constexpr (std::is_same_v<T, FVector2D>)
 			{
-				return ToJsonArray({ Value.X, Value.Y });
+				return ToJsonArray({Value.X, Value.Y});
 			}
-			else if constexpr (std::is_same_v<U, FVector>)
+			else if constexpr (std::is_same_v<T, FVector>)
 			{
-				return ToJsonArray({ Value.X, Value.Y, Value.Z });
+				return ToJsonArray({Value.X, Value.Y, Value.Z});
 			}
-			else if constexpr (std::is_same_v<U, FVector4> || std::is_same_v<U, FQuat>)
+			else if constexpr (std::is_same_v<T, FVector4> || std::is_same_v<T, FQuat>)
 			{
-				return ToJsonArray({ Value.X, Value.Y, Value.Z, Value.W });
+				return ToJsonArray({Value.X, Value.Y, Value.Z, Value.W});
 			}
-			else if constexpr (std::is_same_v<U, FTransform>)
+			else if constexpr (std::is_same_v<T, FTransform>)
 			{
 				TSharedPtr<FJsonObject> JsonTransformObject = MakeShared<FJsonObject>();
 				JsonTransformObject->SetArrayField(
@@ -70,10 +67,36 @@ namespace PCG::Json
 
 				return MakeShared<FJsonValueObject>(JsonTransformObject);
 			}
-			else //if constexpr (std::is_same_v<U, FRotator>)
+			else //if constexpr (std::is_same_v<T, FRotator>)
+			{
+				static_assert(std::is_same_v<T, FRotator>, "Unsupported type");
+				return ToJsonArray({Value.Pitch, Value.Yaw, Value.Roll});
+			}
+		}
+
+		template <typename T>
+		void SetJsonValue(TSharedPtr<FJsonObject>& InOutJsonObject, FString&& ValueName, const T& Value)
+		{
+			check(InOutJsonObject);
+			if constexpr (std::is_same_v<T, bool>)
+			{
+				InOutJsonObject->SetBoolField(std::forward<FString>(ValueName), Value);
+			}
+			else if constexpr (std::is_unsigned_v<std::remove_cv_t<T>> || std::is_signed_v<std::remove_cv_t<T>>)
+			{
+				InOutJsonObject->SetNumberField(std::forward<FString>(ValueName), Value);
+			}
+			else if constexpr (Private::MetadataTraits<T>::IsFloatingPoint)
+			{
+				InOutJsonObject->SetField(std::forward<FString>(ValueName), ConvertFloatingPointTypeToJson(Value));
+			}
+			else if constexpr (std::is_same_v<T, FString>)
+			{
+				InOutJsonObject->SetStringField(std::forward<FString>(ValueName), Value);
+			}
+			else
 			{
-				static_assert(std::is_same_v<U, FRotator>, "Unsupported type");
-				return ToJsonArray({ Value.Pitch, Value.Yaw, Value.Roll });
+				InOutJsonObject->SetStringField(std::forward<FString>(ValueName), PCG::Private::MetadataTraits<T>::ToString(Value));
 			}
 		}
 	}
diff --git a/Engine/Plugins/PCG/Source/PCG/Private/PCGGraph.cpp b/Engine/Plugins/PCG/Source/PCG/Private/PCGGraph.cpp
--- a/Engine/Plugins/PCG/Source/PCG/Private/PCGGraph.cpp
+++ b/Engine/Plugins/PCG/Source/PCG/Private/PCGGraph.cpp
@@ -16,6 +16,7 @@
 #include "Elements/PCGUserParameterGet.h"
 #include "Elements/ControlFlow/PCGQualityBranch.h"
 #include "Elements/ControlFlow/PCGQualitySelect.h"
+#include "Elements/IO/PCGJsonHelpers.h"
 #include "Graph/PCGGraphCompilationData.h"
 #include "Graph/PCGGraphCompiler.h"
 #include "Graph/PCGGraphExecutor.h"
@@ -2222,6 +2223,94 @@ void UPCGGraph::AddUserParameters(const TArray<FPropertyBagPropertyDesc>& InDesc
 	OnGraphParametersChanged(EPCGGraphParameterEvent::MultiplePropertiesAdded, NAME_None);
 }
 
+#if WITH_EDITOR
+FString UPCGGraph::GetGraphDescriptionJson() const
+{
+	TSharedPtr<FJsonObject> RootJsonObject = MakeShared<FJsonObject>();
+
+	// Set up the Json header to include custom versions
+	TSharedPtr<FJsonObject> JsonHeaderObject = PCG::Json::Helpers::CreateJsonHeaderObject(PCG::Json::Constants::FCustomExportVersion::LatestVersion);
+	RootJsonObject->SetObjectField(LOCTEXT("JsonHeaderField", "Header").ToString(), std::move(JsonHeaderObject));
+
+	// 0. Graph information
+	// - Name
+	RootJsonObject->SetStringField("Name", GetFName().ToString());
+	// - Category (might identify family of graph
+	RootJsonObject->SetStringField("Category", Category.ToString());
+	// - Description
+	RootJsonObject->SetStringField("Description", Description.ToString());
+
+	auto CreatePinDescriptionJsonObject = [](const FPCGPinProperties& Pin)
+	{
+		TSharedPtr<FJsonObject> JsonObject = MakeShared<FJsonObject>();
+		JsonObject->SetStringField(TEXT("Type"), StaticEnum<EPCGDataType>()->GetValueOrBitfieldAsString(static_cast<int64>(Pin.AllowedTypes)));
+		JsonObject->SetBoolField(TEXT("Multiple"), Pin.bAllowMultipleData);
+		return JsonObject;
+	};
+
+	// 1. Inputs
+	TSharedPtr<FJsonObject> InputsJsonObject = MakeShared<FJsonObject>();
+	TArray<TSharedPtr<FJsonObject>> InputsJsonObjects;
+
+	TArray<FPCGPinProperties> InputPins = GetInputNode()->InputPinProperties();
+	for (const FPCGPinProperties& InputPin : InputPins)
+	{
+		InputsJsonObject->SetObjectField(InputPin.Label.ToString(), CreatePinDescriptionJsonObject(InputPin));
+	}
+
+	RootJsonObject->SetObjectField(TEXT("Inputs"), MoveTemp(InputsJsonObject));
+
+	// 2. Output
+	TSharedPtr<FJsonObject> OutputsJsonObject = MakeShared<FJsonObject>();
+
+	TArray<FPCGPinProperties> OutputPins = GetOutputNode()->OutputPinProperties();
+	for (const FPCGPinProperties& OutputPin : OutputPins)
+	{
+		OutputsJsonObject->SetObjectField(OutputPin.Label.ToString(), CreatePinDescriptionJsonObject(OutputPin));
+	}
+
+	RootJsonObject->SetObjectField(TEXT("Outputs"), MoveTemp(OutputsJsonObject));
+
+	// 3. Parameters
+	TSharedPtr<FJsonObject> ParametersJsonObject = MakeShared<FJsonObject>();
+
+	const FInstancedPropertyBag* Parameters = GetUserParametersStruct();
+	if (const UPropertyBag* PropertyBag = Parameters ? Parameters->GetPropertyBagStruct() : nullptr)
+	{
+		TConstArrayView<FPropertyBagPropertyDesc> PropertyDescs = PropertyBag->GetPropertyDescs();
+		for (const FPropertyBagPropertyDesc& PropertyDesc : PropertyDescs)
+		{
+			TSharedPtr<FJsonObject> JsonObject = MakeShared<FJsonObject>();
+			// - Type
+			JsonObject->SetStringField(TEXT("Type"), StaticEnum<EPropertyBagPropertyType>()->GetDisplayValueAsText(PropertyDesc.ValueType).ToString());
+
+			// - Value
+			auto SetJson = [&JsonObject]<typename T>(const T& Value) -> bool
+			{
+				PCG::Json::Helpers::SetJsonValue<T>(JsonObject, TEXT("Value"), Value);
+				return true;
+			};
+
+			PCGPropertyHelpers::GetPropertyValueWithCallback(Parameters->GetValue().GetMemory(), PropertyDesc.CachedProperty, SetJson);
+			ParametersJsonObject->SetObjectField(PropertyDesc.Name.ToString(), MoveTemp(JsonObject));
+		}
+	}
+
+	RootJsonObject->SetObjectField(TEXT("Parameters"), ParametersJsonObject);
+
+	FString JsonString;
+	const TSharedRef<TJsonWriter<>> JsonWriter = TJsonWriterFactory<>::Create(&JsonString);
+
+	if (!FJsonSerializer::Serialize(RootJsonObject.ToSharedRef(), JsonWriter))
+	{
+		UE_LOG(LogPCG, Error, TEXT("Serialization of graph '%s' Json failed."), *GetFName().ToString());
+		return {};
+	}
+
+	return JsonString;
+}
+#endif // WITH_EDITOR
+
 /****************************
 * UPCGGraphInstance
 ****************************/
diff --git a/Engine/Plugins/PCG/Source/PCG/Public/Elements/IO/PCGExportSelectedAttributes.h b/Engine/Plugins/PCG/Source/PCG/Public/Elements/IO/PCGExportSelectedAttributes.h
--- a/Engine/Plugins/PCG/Source/PCG/Public/Elements/IO/PCGExportSelectedAttributes.h
+++ b/Engine/Plugins/PCG/Source/PCG/Public/Elements/IO/PCGExportSelectedAttributes.h
@@ -83,7 +83,8 @@ public:
 	UPROPERTY(EditAnywhere, BlueprintReadWrite, Category = Advanced)
 	bool bAddCustomDataVersion = false;
 
-	UPROPERTY(EditAnywhere, BlueprintReadWrite, Category = Advanced, meta = (EditCondition = "bAddCustomDataVersion", EditConditionHides))
+	/** Extra user version for any special requirements. Must be >= 0.*/
+	UPROPERTY(EditAnywhere, BlueprintReadWrite, Category = Advanced, meta = (EditCondition = "bAddCustomDataVersion", EditConditionHides, ClampMin = 0))
 	int32 CustomVersion = 0;
 };
 
diff --git a/Engine/Plugins/PCG/Source/PCG/Public/PCGGraph.h b/Engine/Plugins/PCG/Source/PCG/Public/PCGGraph.h
--- a/Engine/Plugins/PCG/Source/PCG/Public/PCGGraph.h
+++ b/Engine/Plugins/PCG/Source/PCG/Public/PCGGraph.h
@@ -456,6 +456,12 @@ public:
 	TObjectPtr<UPCGGraphCompilationData> GetCookedCompilationData() { return CookedCompilationData; }
 	const TObjectPtr<UPCGGraphCompilationData> GetCookedCompilationData() const { return CookedCompilationData; }
 
+#if WITH_EDITOR
+	/** Outputs a Json representation of the graph inputs, outputs, and parameters. See the PCGExportSelectedAttributes node for more examples. */
+	UFUNCTION(BlueprintCallable, Category = "Graph|Advanced|Json")
+	PCG_API FString GetGraphDescriptionJson() const;
+#endif
+
 protected:
 #if WITH_EDITOR
 	virtual bool IsTemplatePropertyEnabled() const override { return IsAsset(); }
