[PCG] Add support to filter some overrides when allowing override on non-PCG related structs.

* Avoid to pollute with PCG references struct that are defined elsewhere in the engine
* Will be a property path filtering.

#rnx

diff --git a/Engine/Plugins/PCG/Source/PCG/Private/Helpers/PCGSettingsHelpers.cpp b/Engine/Plugins/PCG/Source/PCG/Private/Helpers/PCGSettingsHelpers.cpp
--- a/Engine/Plugins/PCG/Source/PCG/Private/Helpers/PCGSettingsHelpers.cpp
+++ b/Engine/Plugins/PCG/Source/PCG/Private/Helpers/PCGSettingsHelpers.cpp
@@ -379,6 +379,14 @@ namespace PCGSettingsHelpers
 				RecurseConfig.bUseSeed = true;
 #if WITH_EDITOR
 				RecurseConfig.IncludeMetadataValues.Remove(PCGObjectMetadata::Overridable);
+
+				// Also check if we have some filtering on the child properties
+				TArray<FString> ChildPropertiesFiltering;
+				if (Property->HasMetaData(PCGObjectMetadata::OverridableChildProperties))
+				{
+					const FString& ChildPropertiesMetadata = Property->GetMetaData(PCGObjectMetadata::OverridableChildProperties);
+					ChildPropertiesFiltering = PCGHelpers::GetStringArrayFromCommaSeparatedList(ChildPropertiesMetadata);
+				}
 #endif // WITH_EDITOR
 
 				if (RecurseConfig.MaxStructDepth > 0)
@@ -402,6 +410,12 @@ namespace PCGSettingsHelpers
 					}
 
 					LabelCache.AddUnique(Label);
+
+					// Also check if the current property path is filtered out
+					if (!ChildPropertiesFiltering.IsEmpty() && !ChildPropertiesFiltering.Contains(ChildParam.GetPropertyPath()))
+					{
+						continue;
+					}
 #endif // WITH_EDITOR
 
 					FPCGSettingsOverridableParam& Param = Res.Emplace_GetRef();
diff --git a/Engine/Plugins/PCG/Source/PCG/Public/PCGCommon.h b/Engine/Plugins/PCG/Source/PCG/Public/PCGCommon.h
--- a/Engine/Plugins/PCG/Source/PCG/Public/PCGCommon.h
+++ b/Engine/Plugins/PCG/Source/PCG/Public/PCGCommon.h
@@ -231,6 +231,7 @@ namespace PCGObjectMetadata
 	const FName Overridable = TEXT("PCG_Overridable");
 	const FName OverridableGPU = TEXT("PCG_OverridableGPU");
 	const FName NotOverridable = TEXT("PCG_NotOverridable");
+	const FName OverridableChildProperties = TEXT("PCG_OverridableChildProperties");
 	const FName OverrideAliases = TEXT("PCG_OverrideAliases");
 	const FName DiscardPropertySelection = TEXT("PCG_DiscardPropertySelection");
 	const FName DiscardExtraSelection = TEXT("PCG_DiscardExtraSelection");
@@ -248,6 +249,16 @@ namespace PCGObjectMetadata
 		/// [PropertyMetadata] Indicates that the property is not-overridable by params. Used in structs to hide some parameters
 		PCG_NotOverridable,
 
+		/// [PropertyMetadata] When a struct property is overridable, by default it will expose all the child properties (except the one marked by PCG_NotOverridable)
+		/// but for structs that are not defined in a PCG context, we don't want to pollute the metadata of those with PCG.
+		/// This extra field can be used to select property paths that should be overridable. The property path should be relative to the struct that is marked
+		/// PCG_Overridable, as a comma separated string.
+		/// 
+		/// Example if we have a struct FBar, that has a field Foo that is a struct with Prop1, Prop2 and Prop3, to expose Prop1 and Prop2:
+		/// UPROPERTY(meta=(PCG_Overridable, PCG_OverridableChildProperties = "Foo/Prop1,Foo/Prop2"))
+		/// FBar Bar;
+		PCG_OverridableChildProperties,
+
 		/// [PropertyMetadata] Extra names to match for a given property.
 		PCG_OverrideAliases,
 
