[PCG] Fix spline mesh descriptor collision preset on load

#rnx

diff --git a/Engine/Plugins/PCG/Source/PCG/Private/Elements/PCGSpawnSplineMesh.cpp b/Engine/Plugins/PCG/Source/PCG/Private/Elements/PCGSpawnSplineMesh.cpp
--- a/Engine/Plugins/PCG/Source/PCG/Private/Elements/PCGSpawnSplineMesh.cpp
+++ b/Engine/Plugins/PCG/Source/PCG/Private/Elements/PCGSpawnSplineMesh.cpp
@@ -19,6 +19,17 @@
 
 #define LOCTEXT_NAMESPACE "PCGCreateSplineMeshElement"
 
+void UPCGSpawnSplineMeshSettings::PostLoad()
+{
+	Super::PostLoad();
+
+#if WITH_EDITOR
+	// TODO: Remove if/when FBodyInstance is updated or replaced
+	// Necessary to update the collision Response Container from the Response Array
+	SplineMeshDescriptor.PostLoadFixup(this);
+#endif // WITH_EDITOR
+}
+
 #if WITH_EDITOR
 FText UPCGSpawnSplineMeshSettings::GetNodeTooltipText() const
 {
diff --git a/Engine/Plugins/PCG/Source/PCG/Public/Elements/PCGSpawnSplineMesh.h b/Engine/Plugins/PCG/Source/PCG/Public/Elements/PCGSpawnSplineMesh.h
--- a/Engine/Plugins/PCG/Source/PCG/Public/Elements/PCGSpawnSplineMesh.h
+++ b/Engine/Plugins/PCG/Source/PCG/Public/Elements/PCGSpawnSplineMesh.h
@@ -27,6 +27,10 @@ class UPCGSpawnSplineMeshSettings : public UPCGSettings
 	GENERATED_BODY()
 
 public:
+	//~Begin UObject interface
+	virtual void PostLoad() override;
+	//~End UObject interface
+
 	//~Begin UPCGSettings interface
 #if WITH_EDITOR
 	virtual FName GetDefaultNodeName() const override { return FName(TEXT("SpawnSplineMesh")); }
