[PCG] Fixed strings getting lost / scrambled when passed across graph splits.

Set up a string key remap table, similar to the attribute ID remap table.

#rnx

diff --git a/Engine/Plugins/PCG/Shaders/Private/PCGDataCollectionDataInterface.ush b/Engine/Plugins/PCG/Shaders/Private/PCGDataCollectionDataInterface.ush
--- a/Engine/Plugins/PCG/Shaders/Private/PCGDataCollectionDataInterface.ush
+++ b/Engine/Plugins/PCG/Shaders/Private/PCGDataCollectionDataInterface.ush
@@ -1,20 +1,20 @@
 // Copyright Epic Games, Inc. All Rights Reserved.
 
-// Buffer format (sizes in bytes):
+// Buffer format:
 //
 // -- START HEADER --
-// PackedExecutionFlagAndNumData (4)
+// PackedExecutionFlagAndNumData (4 bytes)
 // For each data:
-//   TypeId (4)
-//   Attribute Count (4)
-//   Element Count (4)
+//   TypeId (4 bytes)
+//   Attribute Count (4 bytes)
+//   Element Count (4 bytes)
 //   For attribute ID 0 to MAX_NUM_ATTR:
-//     PackedIdAndStride (4)
-//     AttributeElementStartAddress (4)
+//     PackedIdAndStride (4 bytes)
+//     AttributeElementStartAddress (4 bytes)
 // -- END HEADER --
 // For each data:
 //   For each present attribute:
-//     <Packed element data> (element count * stride)
+//     <Packed element data> (element count * stride bytes)
 
 RWByteAddressBuffer {DataInterfaceName}_DataCollectionBuffer;
 
@@ -25,6 +25,10 @@ StructuredBuffer<int> {DataInterfaceName}_GraphToBufferAttributeId;
 
 int {DataInterfaceName}_FirstRemappedAttributeId;
 
+StructuredBuffer<int> {DataInterfaceName}_BufferToGraphStringKey;
+
+int {DataInterfaceName}_NumRemappedStringKeys;
+
 // #################### HEADER READERS ##########################
 
 uint GetNumData_{DataInterfaceName}()
@@ -248,7 +252,8 @@ float4x4 GetTransform_{DataInterfaceName}(uint DataIndex, uint ElementIndex, int
 int GetStringKey_{DataInterfaceName}(uint DataIndex, uint ElementIndex, int AttributeId)
 {
 	// String keys are represented as ints.
-	return GetInt_{DataInterfaceName}(DataIndex, ElementIndex, AttributeId);
+	const int StringKey = GetInt_{DataInterfaceName}(DataIndex, ElementIndex, AttributeId);
+	return (StringKey < {DataInterfaceName}_NumRemappedStringKeys) ? {DataInterfaceName}_BufferToGraphStringKey[StringKey] : StringKey;
 }
 
 // TODO: left as uint2 which we locked into for 5.5.0. Future version should change this to int.
diff --git a/Engine/Plugins/PCG/Source/PCG/Private/Compute/DataInterfaces/PCGDataCollectionDataInterface.cpp b/Engine/Plugins/PCG/Source/PCG/Private/Compute/DataInterfaces/PCGDataCollectionDataInterface.cpp
--- a/Engine/Plugins/PCG/Source/PCG/Private/Compute/DataInterfaces/PCGDataCollectionDataInterface.cpp
+++ b/Engine/Plugins/PCG/Source/PCG/Private/Compute/DataInterfaces/PCGDataCollectionDataInterface.cpp
@@ -459,6 +459,8 @@ BEGIN_SHADER_PARAMETER_STRUCT(FPCGDataCollectionDataInterfaceParameters, )
 	SHADER_PARAMETER(uint32, ElementCountMultiplier)
 	SHADER_PARAMETER_RDG_BUFFER_SRV(StructuredBuffer<int32>, GraphToBufferAttributeId)
 	SHADER_PARAMETER(int32, FirstRemappedAttributeId)
+	SHADER_PARAMETER_RDG_BUFFER_SRV(StructuredBuffer<int32>, BufferToGraphStringKey)
+	SHADER_PARAMETER(int32, NumRemappedStringKeys)
 END_SHADER_PARAMETER_STRUCT()
 
 void UPCGDataCollectionDataInterface::GetShaderParameters(TCHAR const* UID, FShaderParametersMetadataBuilder& InOutBuilder, FShaderParametersMetadataAllocations& InOutAllocations) const
@@ -635,6 +637,8 @@ void FPCGDataCollectionDataProviderProxy::GatherDispatchData(FDispatchData const
 		Parameters.ElementCountMultiplier = ElementCountMultiplier;
 		Parameters.GraphToBufferAttributeId = AttributeIdRemapSRV;
 		Parameters.FirstRemappedAttributeId = FirstRemappedAttributeId;
+		Parameters.BufferToGraphStringKey = BufferToGraphStringKeySRV;
+		Parameters.NumRemappedStringKeys = NumRemappedStringKeys;
 	}
 }
 
@@ -645,6 +649,8 @@ void FPCGDataCollectionDataProviderProxy::AllocateResources(FRDGBuilder& GraphBu
 
 	AttributeIdRemapSRV = GetAttributeIdRemapBufferSRV(GraphBuilder, FirstRemappedAttributeId);
 
+	BufferToGraphStringKeySRV = GetBufferToGraphStringKeySRV(GraphBuilder, NumRemappedStringKeys);
+
 	FRDGBufferDesc Desc = FRDGBufferDesc::CreateByteAddressDesc(SizeBytes);
 	if (ExportMode != EPCGExportMode::NoExport)
 	{
@@ -740,4 +746,12 @@ FRDGBufferSRVRef FPCGDataCollectionDataProviderProxy::GetAttributeIdRemapBufferS
 	return GraphBuilder.CreateSRV(FRDGBufferSRVDesc(GSystemTextures.GetDefaultStructuredBuffer(GraphBuilder, sizeof(int32))));
 }
 
+FRDGBufferSRVRef FPCGDataCollectionDataProviderProxy::GetBufferToGraphStringKeySRV(FRDGBuilder& GraphBuilder, int32& OutNumRemappedStringKeys) const
+{
+	// Default implementation. Disable remapping.
+	OutNumRemappedStringKeys = 0;
+
+	return GraphBuilder.CreateSRV(FRDGBufferSRVDesc(GSystemTextures.GetDefaultStructuredBuffer(GraphBuilder, sizeof(int32))));
+}
+
 #undef LOCTEXT_NAMESPACE
diff --git a/Engine/Plugins/PCG/Source/PCG/Private/Compute/DataInterfaces/PCGDataCollectionDataInterface.h b/Engine/Plugins/PCG/Source/PCG/Private/Compute/DataInterfaces/PCGDataCollectionDataInterface.h
--- a/Engine/Plugins/PCG/Source/PCG/Private/Compute/DataInterfaces/PCGDataCollectionDataInterface.h
+++ b/Engine/Plugins/PCG/Source/PCG/Private/Compute/DataInterfaces/PCGDataCollectionDataInterface.h
@@ -99,6 +99,8 @@ public:
 protected:
 	virtual FRDGBufferSRVRef GetAttributeIdRemapBufferSRV(FRDGBuilder& GraphBuilder, int32& OutFirstRemappedAttributeId) const;
 
+	virtual FRDGBufferSRVRef GetBufferToGraphStringKeySRV(FRDGBuilder& GraphBuilder, int32& OutNumRemappedStringKeys) const;
+
 public:
 	/** Generation count of the data provider when the proxy was created. */
 	uint64 OriginatingGenerationCount = 0;
@@ -129,5 +131,9 @@ protected:
 
 	int32 FirstRemappedAttributeId = INDEX_NONE;
 
+	FRDGBufferSRVRef BufferToGraphStringKeySRV = nullptr;
+
+	int32 NumRemappedStringKeys = INDEX_NONE;
+
 	int32 LargestAttributeIndex = INDEX_NONE;
 };
diff --git a/Engine/Plugins/PCG/Source/PCG/Private/Compute/DataInterfaces/PCGDataCollectionUploadDataInterface.cpp b/Engine/Plugins/PCG/Source/PCG/Private/Compute/DataInterfaces/PCGDataCollectionUploadDataInterface.cpp
--- a/Engine/Plugins/PCG/Source/PCG/Private/Compute/DataInterfaces/PCGDataCollectionUploadDataInterface.cpp
+++ b/Engine/Plugins/PCG/Source/PCG/Private/Compute/DataInterfaces/PCGDataCollectionUploadDataInterface.cpp
@@ -135,6 +135,7 @@ void FPCGDataProviderDataCollectionUploadProxy::AllocateResources(FRDGBuilder& G
 	BufferUAV = GraphBuilder.CreateUAV(Buffer);
 
 	AttributeIdRemapSRV = GetAttributeIdRemapBufferSRV(GraphBuilder, FirstRemappedAttributeId);
+	BufferToGraphStringKeySRV = GetBufferToGraphStringKeySRV(GraphBuilder, NumRemappedStringKeys);
 }
 
 FRDGBufferSRVRef FPCGDataProviderDataCollectionUploadProxy::GetAttributeIdRemapBufferSRV(FRDGBuilder& GraphBuilder, int32& OutFirstRemappedAttributeId) const
@@ -149,3 +150,16 @@ FRDGBufferSRVRef FPCGDataProviderDataCollectionUploadProxy::GetAttributeIdRemapB
 		return DataAdaptor->GetAttributeRemapBufferSRV(GraphBuilder, OutFirstRemappedAttributeId);
 	}
 }
+
+FRDGBufferSRVRef FPCGDataProviderDataCollectionUploadProxy::GetBufferToGraphStringKeySRV(FRDGBuilder& GraphBuilder, int32& OutNumRemappedStringKeys) const
+{
+	if (!DataAdaptor || !DataAdaptor->UsesStringKeyRemap())
+	{
+		// No remap required.
+		return FPCGDataCollectionDataProviderProxy::GetBufferToGraphStringKeySRV(GraphBuilder, OutNumRemappedStringKeys);
+	}
+	else
+	{
+		return DataAdaptor->GetBufferToGraphStringKeySRV(GraphBuilder, OutNumRemappedStringKeys);
+	}
+}
diff --git a/Engine/Plugins/PCG/Source/PCG/Private/Compute/DataInterfaces/PCGDataCollectionUploadDataInterface.h b/Engine/Plugins/PCG/Source/PCG/Private/Compute/DataInterfaces/PCGDataCollectionUploadDataInterface.h
--- a/Engine/Plugins/PCG/Source/PCG/Private/Compute/DataInterfaces/PCGDataCollectionUploadDataInterface.h
+++ b/Engine/Plugins/PCG/Source/PCG/Private/Compute/DataInterfaces/PCGDataCollectionUploadDataInterface.h
@@ -67,6 +67,8 @@ public:
 protected:
 	virtual FRDGBufferSRVRef GetAttributeIdRemapBufferSRV(FRDGBuilder& GraphBuilder, int32& OutFirstRemappedAttributeId) const override;
 
+	virtual FRDGBufferSRVRef GetBufferToGraphStringKeySRV(FRDGBuilder& GraphBuilder, int32& OutNumRemappedStringKeys) const override;
+
 protected:
 	TSharedPtr<FPCGDataUploadAdaptor> DataAdaptor;
 
diff --git a/Engine/Plugins/PCG/Source/PCG/Private/Compute/PCGComputeGraph.cpp b/Engine/Plugins/PCG/Source/PCG/Private/Compute/PCGComputeGraph.cpp
--- a/Engine/Plugins/PCG/Source/PCG/Private/Compute/PCGComputeGraph.cpp
+++ b/Engine/Plugins/PCG/Source/PCG/Private/Compute/PCGComputeGraph.cpp
@@ -195,7 +195,7 @@ void UPCGComputeGraph::DebugLogDataDescriptions(const UPCGDataBinding* InBinding
 	const UEnum* PCGKernelAttributeTypeEnum = StaticEnum<EPCGKernelAttributeType>();
 	check(PCGDataTypeEnum && PCGKernelAttributeTypeEnum);
 
-	auto LogDataDescription = [this, PCGDataTypeEnum, PCGKernelAttributeTypeEnum, InBinding](const FPCGDataCollectionDesc* InDataDesc)
+	auto LogDataDescription = [this, PCGDataTypeEnum, PCGKernelAttributeTypeEnum, InBinding](const FPCGDataCollectionDesc* InDataDesc, const FPCGKernelPin& InKernelPin)
 	{
 		if (!InDataDesc)
 		{
@@ -220,33 +220,47 @@ void UPCGComputeGraph::DebugLogDataDescriptions(const UPCGDataBinding* InBinding
 
 			UE_LOG(LogPCG, Display, TEXT("\t\t\tAttributes (%d)"), DataDesc.AttributeDescs.Num());
 
+			TSharedPtr<const FPCGProxyForGPUDataCollection> UpstreamDataCollectionInfo;
+
+			if (const FName* VirtualLabel = GetVirtualPinLabelFromBindingIndex(GetBindingIndex(InKernelPin)))
+			{
+				TArray<FPCGTaggedData> PinInputs = InBinding->InputDataCollection.GetInputsByPin(*VirtualLabel);
+
+				if (ensure(PinInputs.IsValidIndex(DataIndex)))
+				{
+					if (const UPCGProxyForGPUData* Proxy = Cast<UPCGProxyForGPUData>(PinInputs[DataIndex].Data))
+					{
+						UpstreamDataCollectionInfo = Proxy->GetGPUInfo();
+					}
+				}
+			}
+
 			for (int32 AttributeIndex = 0; AttributeIndex < DataDesc.AttributeDescs.Num(); ++AttributeIndex)
 			{
 				const FPCGKernelAttributeDesc& AttributeDesc = DataDesc.AttributeDescs[AttributeIndex];
 
+				// If the data comes from an upstream / source buffer then look up its info.
 				FString UniqueStringKeys;
-				for (int32 StringKey : AttributeDesc.GetUniqueStringKeys())
+				for (int32 GraphStringKey : AttributeDesc.GetUniqueStringKeys())
 				{
 					UniqueStringKeys += TEXT(", ");
-					UniqueStringKeys += FString::Format(TEXT("{0}"), { StringKey });
+					UniqueStringKeys += FString::Format(TEXT("{0}"), { GraphStringKey });
+
+					if (const int32* UpstreamStringKey = UpstreamDataCollectionInfo ? InBinding->GetGraphToSourceBufferStringKey().Find({ UpstreamDataCollectionInfo->GetBuffer(), GraphStringKey }) : nullptr)
+					{
+						UniqueStringKeys += FString::Format(TEXT(" ({0})"), { *UpstreamStringKey });
+					}
 				}
 
 				// If attribute comes from an upstream GPU buffer, display the attribute ID in the upstream buffer.
 				int32 SourceBufferAttributeId = INDEX_NONE;
 
-				if (AttributeDesc.AttributeId >= PCGComputeConstants::NUM_RESERVED_ATTRS)
+				if (AttributeDesc.AttributeId >= PCGComputeConstants::NUM_RESERVED_ATTRS && UpstreamDataCollectionInfo)
 				{
-					if (ensure(InBinding->InputDataCollection.TaggedData.IsValidIndex(DataIndex)))
-					{
-						const UPCGProxyForGPUData* Proxy = Cast<UPCGProxyForGPUData>(InBinding->InputDataCollection.TaggedData[DataIndex].Data);
-						if (TSharedPtr<const FPCGProxyForGPUDataCollection> DataCollectionInfo = Proxy ? Proxy->GetGPUInfo() : nullptr)
-						{
-							const int32 GraphAttributeIndex = PCGComputeHelpers::GetMetadataAttributeIndexFromAttributeId(AttributeDesc.AttributeId);
-							const int32* FoundIndex = InBinding->GetGraphAttributeToSourceBufferAttributeIndex().Find({ GraphAttributeIndex, DataCollectionInfo->GetBuffer() });
+					const int32 GraphAttributeIndex = PCGComputeHelpers::GetMetadataAttributeIndexFromAttributeId(AttributeDesc.AttributeId);
+					const int32* FoundIndex = InBinding->GetGraphAttributeToSourceBufferAttributeIndex().Find({ GraphAttributeIndex, UpstreamDataCollectionInfo->GetBuffer() });
 
-							SourceBufferAttributeId = PCGComputeHelpers::GetAttributeIdFromMetadataAttributeIndex(FoundIndex ? *FoundIndex : GraphAttributeIndex);
-						}
-					}
+					SourceBufferAttributeId = PCGComputeHelpers::GetAttributeIdFromMetadataAttributeIndex(FoundIndex ? *FoundIndex : GraphAttributeIndex);
 				}
 
 				if (SourceBufferAttributeId == INDEX_NONE)
@@ -281,9 +295,11 @@ void UPCGComputeGraph::DebugLogDataDescriptions(const UPCGDataBinding* InBinding
 		}
 	};
 
-	for (TObjectPtr<const UComputeKernel> Kernel : KernelInvocations)
+	for (int32 KernelIndex = 0; KernelIndex < KernelInvocations.Num(); ++KernelIndex)
 	{
-		if (const UPCGComputeKernel* PCGKernel = Cast<UPCGComputeKernel>(Kernel); PCGKernel && PCGKernel->GetLogDataDescriptions())
+		const UPCGComputeKernel* PCGKernel = Cast<UPCGComputeKernel>(KernelInvocations[KernelIndex]);
+
+		if (PCGKernel && PCGKernel->GetLogDataDescriptions())
 		{
 			UE_LOG(LogPCG, Display, TEXT("Kernel: %s"), *PCGKernel->GetName());
 
@@ -296,7 +312,8 @@ void UPCGComputeGraph::DebugLogDataDescriptions(const UPCGDataBinding* InBinding
 					if (const FPCGDataCollectionDesc* DataDesc = InBinding->GetCachedKernelPinDataDesc(PCGKernel, Pin.Label, true))
 					{
 						UE_LOG(LogPCG, Display, TEXT("\tInput Pin: %s (%d data)"), *Pin.Label.ToString(), DataDesc->DataDescs.Num());
-						LogDataDescription(DataDesc);
+
+						LogDataDescription(DataDesc, FPCGKernelPin(KernelIndex, Pin.Label, true));
 					}
 				}
 			}
@@ -310,7 +327,8 @@ void UPCGComputeGraph::DebugLogDataDescriptions(const UPCGDataBinding* InBinding
 					if (const FPCGDataCollectionDesc* DataDesc = InBinding->GetCachedKernelPinDataDesc(PCGKernel, Pin.Label, false))
 					{
 						UE_LOG(LogPCG, Display, TEXT("\tOutput Pin: %s (%d data)"), *Pin.Label.ToString(), DataDesc->DataDescs.Num());
-						LogDataDescription(DataDesc);
+
+						LogDataDescription(DataDesc, FPCGKernelPin(KernelIndex, Pin.Label, false));
 					}
 				}
 			}
@@ -367,6 +385,7 @@ bool UPCGComputeGraph::ComputeTaggedDataPinDesc(const FPCGTaggedData& InTaggedDa
 
 		if (TSharedPtr<const FPCGProxyForGPUDataCollection> DataGPU = Proxy->GetGPUInfo())
 		{
+			// Remap attribute IDs and string keys.
 			for (FPCGKernelAttributeDesc& AttrDesc : OutDescription.AttributeDescs)
 			{
 				const int32 Index = PCGComputeHelpers::GetMetadataAttributeIndexFromAttributeId(AttrDesc.AttributeId);
@@ -375,12 +394,49 @@ bool UPCGComputeGraph::ComputeTaggedDataPinDesc(const FPCGTaggedData& InTaggedDa
 					continue;
 				}
 
+				// Map from the attribute IDs in the buffer / upstream compute graph to attribute IDs usable in the current graph.
 				if (const int32* AttributeIndex = InBinding->GetSourceBufferAttributeToGraphAttributeIndex().Find(UPCGDataBinding::FSourceBufferAndAttributeIndex(DataGPU, Index)))
 				{
 					const int32 NewId = PCGComputeHelpers::GetAttributeIdFromMetadataAttributeIndex(*AttributeIndex);
 					//UE_LOG(LogPCG, Warning, TEXT("Remapped attribute '%s' ID: %d -> %d"), *AttrDesc.Name.ToString(), AttrDesc.Index, NewId);
 					AttrDesc.AttributeId = NewId;
 				}
+
+				// Map from the attribute IDs in the buffer / upstream compute graph to attribute IDs usable in the current graph.
+				if (!InBinding->GetSourceBufferToGraphStringKey().IsEmpty())
+				{
+					TArray<int32> GraphStringKeys;
+					GraphStringKeys.Reserve(AttrDesc.GetUniqueStringKeys().Num());
+					bool bStringKeyRemapped = false;
+
+					for (int32 SourceStringKey : AttrDesc.GetUniqueStringKeys())
+					{
+						if (const int32* GraphStringKey = InBinding->GetSourceBufferToGraphStringKey().Find({ DataGPU->GetBuffer(), SourceStringKey }))
+						{
+							//UE_LOG(LogPCG, Warning, TEXT("Remapped string key: %d -> %d"), SourceStringKey, *GraphStringKey);
+							GraphStringKeys.AddUnique(*GraphStringKey);
+							bStringKeyRemapped = true;
+						}
+						else
+						{
+							GraphStringKeys.AddUnique(SourceStringKey);
+						}
+					}
+
+					if (bStringKeyRemapped)
+					{
+						AttrDesc.SetStringKeys(GraphStringKeys);
+					}
+				}
+			}
+
+			// Remap tag string keys.
+			for (int32& TagStringKey : OutDescription.TagStringKeys)
+			{
+				if (const int32* GraphStringKey = InBinding->GetSourceBufferToGraphStringKey().Find({ DataGPU->GetBuffer(), TagStringKey }))
+				{
+					TagStringKey = *GraphStringKey;
+				}
 			}
 		}
 	}
diff --git a/Engine/Plugins/PCG/Source/PCG/Private/Compute/PCGDataBinding.cpp b/Engine/Plugins/PCG/Source/PCG/Private/Compute/PCGDataBinding.cpp
--- a/Engine/Plugins/PCG/Source/PCG/Private/Compute/PCGDataBinding.cpp
+++ b/Engine/Plugins/PCG/Source/PCG/Private/Compute/PCGDataBinding.cpp
@@ -63,6 +63,8 @@ void UPCGDataBinding::ReleaseTransientResources()
 {
 	SourceBufferAttributeToGraphAttributeIndex.Reset();
 	GraphAttributeToSourceBufferAttributeIndex.Reset();
+	SourceBufferToGraphStringKey.Reset();
+	GraphToSourceBufferStringKey.Reset();
 	OutputDataCollection.TaggedData.Reset();
 	InputDataCollection.TaggedData.Reset();
 	MeshSpawnersToPrimitives.Reset();
@@ -392,24 +394,19 @@ void UPCGDataBinding::AddInputDataStringsToTable()
 
 				ProcessedBuffers.Add(GPUInfo);
 
-				if (GPUInfo->GetStringTable().IsEmpty())
+				// Add the string table from the upstream compute graph to the current graph string table (skipping 0 - null/empty string). There may
+				// be multiple upstream compute graphs with conflicting string key values, so we record a map from the current graph string key to the
+				// original source string key and use this map when reading string keys from source buffer.
+				for (int SourceBufferStringKey = 1; SourceBufferStringKey < GPUInfo->GetStringTable().Num(); ++SourceBufferStringKey)
 				{
-					continue;
-				}
+					const int GraphStringKey = StringTable.AddUnique(GPUInfo->GetStringTable()[SourceBufferStringKey]);
 
-				if (StringTable.IsEmpty())
-				{
-					StringTable = GPUInfo->GetStringTable();
-				}
-				else
-				{
-					for (const FString& String : GPUInfo->GetStringTable())
+					if (GraphStringKey != SourceBufferStringKey)
 					{
-						StringTable.AddUnique(String);
+						SourceBufferToGraphStringKey.Add({ GPUInfo->GetBuffer(), SourceBufferStringKey}, GraphStringKey);
+						GraphToSourceBufferStringKey.Add({ GPUInfo->GetBuffer(), GraphStringKey }, SourceBufferStringKey);
 					}
 				}
-
-				continue;
 			}
 		}
 		// Non-GPU-proxy: collect strings from metadata.
diff --git a/Engine/Plugins/PCG/Source/PCG/Private/Compute/PCGDataBinding.h b/Engine/Plugins/PCG/Source/PCG/Private/Compute/PCGDataBinding.h
--- a/Engine/Plugins/PCG/Source/PCG/Private/Compute/PCGDataBinding.h
+++ b/Engine/Plugins/PCG/Source/PCG/Private/Compute/PCGDataBinding.h
@@ -124,6 +124,14 @@ public:
 	/** Get map from this-graph attribute index and upstream source buffer to upstream attribute index. */
 	const TMap<FGraphAttributeIndexAndSourceBuffer, /*Source Buffer Attribute Index*/int32>& GetGraphAttributeToSourceBufferAttributeIndex() const { return GraphAttributeToSourceBufferAttributeIndex; }
 
+	using FSourceBufferAndStringKey = TPair<TRefCountPtr<FRDGPooledBuffer>, /*String Key*/int32>;
+
+	/** Map from string key value in graph to the key value for this string in the source buffer. */
+	const TMap<FSourceBufferAndStringKey, /*Graph String Key*/int32>& GetSourceBufferToGraphStringKey() const { return SourceBufferToGraphStringKey; }
+
+	/** Map from string key value in source buffer to the key value usable in this graph. */
+	const TMap<FSourceBufferAndStringKey, /*Source Buffer String Key*/int32>& GetGraphToSourceBufferStringKey() const { return GraphToSourceBufferStringKey; }
+
 	IPCGGraphExecutionSource* GetExecutionSource() const;
 
 	/** Helper to get index in input data collection of the first data item for the given kernel and input pin label. */
@@ -187,6 +195,12 @@ protected:
 	/** Map from this-graph attribute index and upstream source buffer to upstream attribute index. */
 	TMap<FGraphAttributeIndexAndSourceBuffer, /*Source Buffer Attribute Index*/int32> GraphAttributeToSourceBufferAttributeIndex;
 
+	/** Map from string key value in graph to the key value for this string in the source buffer. */
+	TMap<FSourceBufferAndStringKey, /*Graph String Key*/int32> SourceBufferToGraphStringKey;
+
+	/** Map from string key value in source buffer to the key value usable in this graph. */
+	TMap<FSourceBufferAndStringKey, /*Source Buffer String Key*/int32> GraphToSourceBufferStringKey;
+
 	/** Cache of data descriptions to amortize cost of computing them at runtime. */
 	TMap</*Binding index*/int32, FPCGDataCollectionDesc> DataDescriptionCache;
 	std::atomic<bool> bIsDataDescriptionCachePrimed = false;
diff --git a/Engine/Plugins/PCG/Source/PCG/Private/Compute/PCGDataUploadAdaptor.cpp b/Engine/Plugins/PCG/Source/PCG/Private/Compute/PCGDataUploadAdaptor.cpp
--- a/Engine/Plugins/PCG/Source/PCG/Private/Compute/PCGDataUploadAdaptor.cpp
+++ b/Engine/Plugins/PCG/Source/PCG/Private/Compute/PCGDataUploadAdaptor.cpp
@@ -82,13 +82,24 @@ FPCGDataUploadAdaptor::FPCGDataUploadAdaptor(UPCGDataBinding* InDataBinding, con
 		UniqueDataCollection.Reset();
 	}
 
-	if (!UniqueDataCollection)
+	if (UniqueDataCollection)
 	{
-		return;
-	}
+		ExternalBufferForReuse = UniqueDataCollection->GetBuffer();
+		ExternalBufferSizeBytes = UniqueDataCollection->GetBufferSizeBytes();
+
+		ensureAlways(ExternalBufferMaxStringKeyValue == INDEX_NONE);
 
-	ExternalBufferForReuse = UniqueDataCollection->GetBuffer();
-	ExternalBufferSizeBytes = UniqueDataCollection->GetBufferSizeBytes();
+		for (const FPCGDataDesc& DataDesc : UniqueDataCollection->GetDescription().DataDescs)
+		{
+			for (const FPCGKernelAttributeDesc& AttrDesc : DataDesc.AttributeDescs)
+			{
+				for (int32 StringKeyValue : AttrDesc.GetUniqueStringKeys())
+				{
+					ExternalBufferMaxStringKeyValue = FMath::Max(ExternalBufferMaxStringKeyValue, StringKeyValue);
+				}
+			}
+		}
+	}
 }
 
 /** Do any preparation work such as data readbacks. Returns true when preparation is complete. */
@@ -117,6 +128,23 @@ bool FPCGDataUploadAdaptor::PrepareData_GameThread()
 
 			GraphToSourceBufferAttributeIndex[GraphAttributeIndex] = SourceBufferAttributeIndex ? *SourceBufferAttributeIndex : GraphAttributeIndex;
 		}
+
+		ensureAlways(BufferToGraphStringKey.IsEmpty());
+
+		if (ExternalBufferMaxStringKeyValue != INDEX_NONE)
+		{
+			const int32 NumSourceStringKeys = ExternalBufferMaxStringKeyValue + 1;
+			BufferToGraphStringKey.SetNumUninitialized(NumSourceStringKeys);
+
+			// Now compute the mapping for any string keys.
+			BufferToGraphStringKey[0] = 0;
+			for (int32 SourceBufferStringKey = 1; SourceBufferStringKey < NumSourceStringKeys; ++SourceBufferStringKey)
+			{
+				const int32* GraphStringKeyIndex = Binding->GetSourceBufferToGraphStringKey().Find({ ExternalBufferForReuse, SourceBufferStringKey });
+
+				BufferToGraphStringKey[SourceBufferStringKey] = GraphStringKeyIndex ? *GraphStringKeyIndex : SourceBufferStringKey;
+			}
+		}
 	}
 	else
 	{
@@ -187,6 +215,30 @@ FRDGBufferSRVRef FPCGDataUploadAdaptor::GetAttributeRemapBufferSRV(FRDGBuilder&
 	return InGraphBuilder.CreateSRV(Buffer);
 }
 
+FRDGBufferSRVRef FPCGDataUploadAdaptor::GetBufferToGraphStringKeySRV(FRDGBuilder& InGraphBuilder, int32& OutNumRemappedStringKeys)
+{
+	// Should only be called if UsesStringKeyRemap() returns true and ExternalBufferForReuse is set.
+	check(ExternalBufferForReuse);
+
+	OutNumRemappedStringKeys = BufferToGraphStringKey.Num();
+
+	FRDGBufferRef Buffer = nullptr;
+
+	if (!BufferToGraphStringKey.IsEmpty())
+	{
+		const FRDGBufferDesc BufferDesc = FRDGBufferDesc::CreateStructuredDesc(BufferToGraphStringKey.GetTypeSize(), BufferToGraphStringKey.Num());
+		Buffer = InGraphBuilder.CreateBuffer(BufferDesc, TEXT("PCGDataUploadAdaptor_BufferToGraphStringKey"));
+
+		InGraphBuilder.QueueBufferUpload(Buffer, MakeArrayView(BufferToGraphStringKey));
+	}
+	else
+	{
+		Buffer = GSystemTextures.GetDefaultStructuredBuffer(InGraphBuilder, sizeof(uint32));
+	}
+
+	return InGraphBuilder.CreateSRV(Buffer);
+}
+
 /** Create buffer with the element counts of each data. Stored in a buffer because we do not constrain the max data count. */
 FRDGBufferSRVRef FPCGDataUploadAdaptor::GetDataElementCountsBufferSRV(FRDGBuilder& InGraphBuilder)
 {
diff --git a/Engine/Plugins/PCG/Source/PCG/Private/Compute/PCGDataUploadAdaptor.h b/Engine/Plugins/PCG/Source/PCG/Private/Compute/PCGDataUploadAdaptor.h
--- a/Engine/Plugins/PCG/Source/PCG/Private/Compute/PCGDataUploadAdaptor.h
+++ b/Engine/Plugins/PCG/Source/PCG/Private/Compute/PCGDataUploadAdaptor.h
@@ -22,6 +22,8 @@ public:
 
 	FRDGBufferSRVRef GetAttributeRemapBufferSRV(FRDGBuilder& InGraphBuilder, int32& OutFirstRemappedAttributeId);
 
+	FRDGBufferSRVRef GetBufferToGraphStringKeySRV(FRDGBuilder& InGraphBuilder, int32& OutNumRemappedStringKeys);
+
 	/** Create buffer with the element counts of each data. Stored in a buffer because we do not constrain the max data count. */
 	FRDGBufferSRVRef GetDataElementCountsBufferSRV(FRDGBuilder& InGraphBuilder);
 
@@ -30,11 +32,15 @@ public:
 
 	bool UsesAttributeIdRemap() const { return !!ExternalBufferForReuse; }
 
+	bool UsesStringKeyRemap() const { return ExternalBufferMaxStringKeyValue != INDEX_NONE; }
+
 	bool IsUploadingFromCPU() const { return !ExternalBufferForReuse; }
 
 private:
 	TRefCountPtr<FRDGPooledBuffer> ExternalBufferForReuse;
 
+	int32 ExternalBufferMaxStringKeyValue = INDEX_NONE;
+
 	TWeakObjectPtr<UPCGDataBinding> DataBinding;
 
 	FPCGDataCollectionDesc TargetDataCollectionDesc;
@@ -44,6 +50,9 @@ private:
 
 	/** Map from (this) compute graph attribute ID to the attribute ID used in the upstream compute graph that produced this buffer. */
 	TArray<int32> GraphToSourceBufferAttributeIndex;
+	
+	/** Map from string key values buffer from upstream compute graph to corresponding string key values in this graph. */
+	TArray<int32> BufferToGraphStringKey;
 
 	FName InputPinLabel;
 
