[PCG] Change `UPCGMetadata::GetAttributes` so it returns the attributes on the default domain only

* When domain were added, this function was iterating on all domains and getting their name and type
* But we lose information (and can have collision) on their original domain, which don't play nice with code that was
using the deprecated version of metadata (which was forwarding everything to the default domain)
* It crashed in a few places where we were getting all the attribute names, but they were on multiple domains, which was returning a null attribute on `Metadata->GetConstAttribute`
* This basically revert to old behavior for existing code.

#rnx

diff --git a/Engine/Plugins/PCG/Source/PCG/Private/Metadata/PCGMetadata.cpp b/Engine/Plugins/PCG/Source/PCG/Private/Metadata/PCGMetadata.cpp
--- a/Engine/Plugins/PCG/Source/PCG/Private/Metadata/PCGMetadata.cpp
+++ b/Engine/Plugins/PCG/Source/PCG/Private/Metadata/PCGMetadata.cpp
@@ -794,7 +794,10 @@ void UPCGMetadata::GetAttributes(TArray<FName>& AttributeNames, TArray<EPCGMetad
 	AttributeNames.Reset();
 	AttributeTypes.Reset();
 
-	ForEachValidUniqueConstMetadataDomain([&AttributeNames, &AttributeTypes](const FPCGMetadataDomainID& DomainID, const FPCGMetadataDomain& MetadataDomain) { MetadataDomain.GetAttributes(AttributeNames, AttributeTypes); });
+	if (const FPCGMetadataDomain* DefaultMetadataDomain = GetConstDefaultMetadataDomain())
+	{
+		DefaultMetadataDomain->GetAttributes(AttributeNames, AttributeTypes);
+	}
 }
 
 void UPCGMetadata::GetAllAttributes(TArray<FPCGAttributeIdentifier>& AttributeNames, TArray<EPCGMetadataTypes>& AttributeTypes) const
diff --git a/Engine/Plugins/PCG/Source/PCG/Public/Metadata/PCGMetadata.h b/Engine/Plugins/PCG/Source/PCG/Public/Metadata/PCGMetadata.h
--- a/Engine/Plugins/PCG/Source/PCG/Public/Metadata/PCGMetadata.h
+++ b/Engine/Plugins/PCG/Source/PCG/Public/Metadata/PCGMetadata.h
@@ -262,9 +262,11 @@ public:
 	template <typename T>
 	const FPCGMetadataAttribute<T>* GetConstTypedAttribute(FPCGAttributeIdentifier AttributeName) const;
 
+	/** Get all the attributes names and type on the default domain. If you need all the attributes on all domains, use GetAllAttributes. */
 	UFUNCTION(BlueprintCallable, Category = "PCG|Metadata")
 	UE_API void GetAttributes(TArray<FName>& AttributeNames, TArray<EPCGMetadataTypes>& AttributeTypes) const;
 
+	/** Get all the attributes identifiers and their type for all domains. */
 	UE_API void GetAllAttributes(TArray<FPCGAttributeIdentifier>& AttributeNames, TArray<EPCGMetadataTypes>& AttributeTypes) const;
 
 	/** Returns name of the most recently created attribute, or none if no attributes are present. */
