[PCG] Add support for attribute curve remap

* Added a new mode in the existing attribute remap node
* Add an alias for curve remap
* Also expose ExternalCurve for the spline sampler (with the new PCG_OverridableChildProperties)

#rnx

diff --git a/Engine/Plugins/PCG/Source/PCG/Private/Elements/PCGAttributeRemap.cpp b/Engine/Plugins/PCG/Source/PCG/Private/Elements/PCGAttributeRemap.cpp
--- a/Engine/Plugins/PCG/Source/PCG/Private/Elements/PCGAttributeRemap.cpp
+++ b/Engine/Plugins/PCG/Source/PCG/Private/Elements/PCGAttributeRemap.cpp
@@ -11,8 +11,13 @@
 
 namespace PCGAttributeRemapElement
 {
+	constexpr int32 DensityRemapPreconfiguredIndex = 0;
+	constexpr int32 CurveRemapPreconfiguredIndex = 1;
+	
 	struct FPCGAttributeRemapParams
 	{
+		EPCGAttributeRemapMode Mode = EPCGAttributeRemapMode::Ranges;
+
 		double InRangeStart; // Note: Start and End can be a positive or negative range (e.g. [0, 1] or [1, 0])
 		double InRangeEnd;
 		double InRangeMin; // Note: Min and Max are a true min/max used to reject samples that fall outside the range
@@ -21,6 +26,8 @@ namespace PCGAttributeRemapElement
 		double Intercept;
 		bool bClampToUnitRange;
 		bool bIgnoreValuesOutsideInputRange;
+		
+		const FRichCurve* Curve = nullptr;
 	};
 
 	// Fallback for all unsupported types, does nothing.
@@ -31,14 +38,21 @@ namespace PCGAttributeRemapElement
 	template <typename T> requires (std::is_arithmetic_v<T> && !std::is_same_v<T, bool>)
 	void Remap(T& InOutValue, const FPCGAttributeRemapParams& InParams)
 	{
-		if (!InParams.bIgnoreValuesOutsideInputRange || (InOutValue >= InParams.InRangeMin && InOutValue <= InParams.InRangeMax))
+		if (InParams.Mode == EPCGAttributeRemapMode::Ranges)
 		{
-			InOutValue = InParams.Slope * (InOutValue - InParams.InRangeStart) + InParams.Intercept;
-			if (InParams.bClampToUnitRange)
+			if (!InParams.bIgnoreValuesOutsideInputRange || (InOutValue >= InParams.InRangeMin && InOutValue <= InParams.InRangeMax))
 			{
-				InOutValue = FMath::Clamp<T>(InOutValue, T(0.0), T(1.0));
+				InOutValue = InParams.Slope * (InOutValue - InParams.InRangeStart) + InParams.Intercept;
+				if (InParams.bClampToUnitRange)
+				{
+					InOutValue = FMath::Clamp<T>(InOutValue, T(0.0), T(1.0));
+				}
 			}
 		}
+		else if (InParams.Mode == EPCGAttributeRemapMode::Curve && ensure(InParams.Curve))
+		{
+			InOutValue = static_cast<T>(InParams.Curve->Eval(static_cast<float>(InOutValue)));
+		}
 	}
 
 	template <>
@@ -93,25 +107,52 @@ FText UPCGAttributeRemapSettings::GetDefaultNodeTitle() const
 	return LOCTEXT("NodeTitle", "Attribute Remap");
 }
 
+FText UPCGAttributeRemapSettings::GetNodeTooltipText() const
+{
+	return LOCTEXT("NodeTooltip", "Remap an attribute using either a range or a curve.");
+}
+
 TArray<FPCGPreConfiguredSettingsInfo> UPCGAttributeRemapSettings::GetPreconfiguredInfo() const
 {
-	return { FPCGPreConfiguredSettingsInfo{0, LOCTEXT("DensityNodeTitle", "Density Remap")} };
+	return {
+		FPCGPreConfiguredSettingsInfo{PCGAttributeRemapElement::DensityRemapPreconfiguredIndex, LOCTEXT("DensityNodeTitle", "Density Remap")},
+		FPCGPreConfiguredSettingsInfo{PCGAttributeRemapElement::CurveRemapPreconfiguredIndex, LOCTEXT("CurveNodeTitle", "Attribute Curve Remap")},
+	};
 }
 #endif // WITH_EDITOR
 
 FString UPCGAttributeRemapSettings::GetAdditionalTitleInformation() const
 {
-	return FString::Printf(TEXT("%s -> %s"), *InputSource.ToString(), *OutputTarget.ToString());
+	TStringBuilder<256> Builder;
+	Builder += (Mode == EPCGAttributeRemapMode::Curve ? LOCTEXT("CurveSubtitle", "Curve") : LOCTEXT("CurveSubtitle", "Ranges")).ToString();
+
+	if (!IsPropertyOverriddenByPin(GET_MEMBER_NAME_CHECKED(UPCGAttributeRemapSettings, InputSource)) && !IsPropertyOverriddenByPin(GET_MEMBER_NAME_CHECKED(UPCGAttributeRemapSettings, OutputTarget)))
+	{
+		if (OutputTarget.GetAttributeName() != PCGMetadataAttributeConstants::SourceAttributeName)
+		{
+			Builder += FString::Printf(TEXT(" %s -> %s"), *InputSource.ToString(), *OutputTarget.ToString());
+		}
+		else
+		{
+			Builder += TEXT(" ");
+			Builder += InputSource.ToString();
+		}
+	}
+	
+	return Builder.ToString();
 }
 
 void UPCGAttributeRemapSettings::ApplyPreconfiguredSettings(const FPCGPreConfiguredSettingsInfo& PreconfigureInfo)
 {
 	switch (PreconfigureInfo.PreconfiguredIndex)
 	{
-	case 0:
+	case PCGAttributeRemapElement::DensityRemapPreconfiguredIndex:
 		InputSource.SetPointProperty(EPCGPointProperties::Density);
 		bClampToUnitRange = true;
 		break;
+	case PCGAttributeRemapElement::CurveRemapPreconfiguredIndex:
+		Mode = EPCGAttributeRemapMode::Curve;
+		break;
 	default:
 		break;
 	}
@@ -140,43 +181,62 @@ bool FPCGAttributeRemapElement::DoOperation(PCGMetadataOps::FOperationData& Oper
 	const UPCGAttributeRemapSettings* Settings = CastChecked<UPCGAttributeRemapSettings>(OperationData.Settings);
 
 	PCGAttributeRemapElement::FPCGAttributeRemapParams Params{};
-	Params.bClampToUnitRange = Settings->bClampToUnitRange;
-	Params.bIgnoreValuesOutsideInputRange = Settings->bIgnoreValuesOutsideInputRange;
+	Params.Mode = Settings->Mode;
 
-	Params.InRangeMin = FMath::Min(Settings->InRangeMin, Settings->InRangeMax);
-	Params.InRangeMax = FMath::Max(Settings->InRangeMin, Settings->InRangeMax);
+	if (Params.Mode == EPCGAttributeRemapMode::Ranges)
+	{
+		Params.bClampToUnitRange = Settings->bClampToUnitRange;
+		Params.bIgnoreValuesOutsideInputRange = Settings->bIgnoreValuesOutsideInputRange;
 
-	double OutRangeMin = 0.0;
-	double OutRangeMax = 0.0;
+		Params.InRangeMin = FMath::Min(Settings->InRangeMin, Settings->InRangeMax);
+		Params.InRangeMax = FMath::Max(Settings->InRangeMin, Settings->InRangeMax);
 
-	if (Settings->bAllowInverseRange)
-	{
-		Params.InRangeStart = Settings->InRangeMin;
-		Params.InRangeEnd = Settings->InRangeMax;
-		OutRangeMin = Settings->OutRangeMin;
-		OutRangeMax = Settings->OutRangeMax;
-	}
-	else
-	{
-		Params.InRangeStart = Params.InRangeMin;
-		Params.InRangeEnd = Params.InRangeMax;
-		OutRangeMin = FMath::Min(Settings->OutRangeMin, Settings->OutRangeMax);
-		OutRangeMax = FMath::Max(Settings->OutRangeMin, Settings->OutRangeMax);
-	}
+		double OutRangeMin = 0.0;
+		double OutRangeMax = 0.0;
+
+		if (Settings->bAllowInverseRange)
+		{
+			Params.InRangeStart = Settings->InRangeMin;
+			Params.InRangeEnd = Settings->InRangeMax;
+			OutRangeMin = Settings->OutRangeMin;
+			OutRangeMax = Settings->OutRangeMax;
+		}
+		else
+		{
+			Params.InRangeStart = Params.InRangeMin;
+			Params.InRangeEnd = Params.InRangeMax;
+			OutRangeMin = FMath::Min(Settings->OutRangeMin, Settings->OutRangeMax);
+			OutRangeMax = FMath::Max(Settings->OutRangeMin, Settings->OutRangeMax);
+		}
 
-	const double InRangeDifference = Params.InRangeEnd - Params.InRangeStart;
-	const double OutRangeDifference = OutRangeMax - OutRangeMin;
+		const double InRangeDifference = Params.InRangeEnd - Params.InRangeStart;
+		const double OutRangeDifference = OutRangeMax - OutRangeMin;
 
-	// When InRange is a point leave the Slope at 0 so that Density = Intercept
-	if (InRangeDifference == 0)
+		// When InRange is a point leave the Slope at 0 so that Density = Intercept
+		if (InRangeDifference == 0)
+		{
+			Params.Slope = 0;
+			Params.Intercept = (OutRangeMin + OutRangeMax) / 2.f;
+		}
+		else
+		{
+			Params.Slope = OutRangeDifference / InRangeDifference;
+			Params.Intercept = OutRangeMin;
+		}
+	}
+	else if (Params.Mode == EPCGAttributeRemapMode::Curve)
 	{
-		Params.Slope = 0;
-		Params.Intercept = (OutRangeMin + OutRangeMax) / 2.f;
+		Params.Curve = Settings->Curve.GetRichCurveConst();
+		if (!Params.Curve)
+		{
+			PCGLog::LogErrorOnGraph(LOCTEXT("InvalidCurve", "Curve is invalid"), OperationData.Context);
+			return true;
+		}
 	}
 	else
 	{
-		Params.Slope = OutRangeDifference / InRangeDifference;
-		Params.Intercept = OutRangeMin;
+		PCGLog::LogErrorOnGraph(LOCTEXT("InvalidMode", "Invalid mode"), OperationData.Context);
+		return true;
 	}
 
 	auto RemapFunc = [this, &Params, &OperationData]<typename AttributeType>(AttributeType) -> bool
diff --git a/Engine/Plugins/PCG/Source/PCG/Public/Elements/Metadata/PCGAttributeRemap.h b/Engine/Plugins/PCG/Source/PCG/Public/Elements/Metadata/PCGAttributeRemap.h
--- a/Engine/Plugins/PCG/Source/PCG/Public/Elements/Metadata/PCGAttributeRemap.h
+++ b/Engine/Plugins/PCG/Source/PCG/Public/Elements/Metadata/PCGAttributeRemap.h
@@ -8,10 +8,18 @@
 #include "Metadata/Accessors/IPCGAttributeAccessor.h"
 #include "Metadata/Accessors/PCGAttributeAccessorKeys.h"
 
+#include "Curves/CurveFloat.h"
 #include "Templates/UniquePtr.h"
 
 #include "PCGAttributeRemap.generated.h"
 
+UENUM()
+enum class EPCGAttributeRemapMode
+{
+	Ranges,
+	Curve
+};
+
 /**
 * Remap attribute values from one range to another.
 */
@@ -27,6 +35,7 @@ public:
 #if WITH_EDITOR
 	virtual FName GetDefaultNodeName() const override;
 	virtual FText GetDefaultNodeTitle() const override;
+	virtual FText GetNodeTooltipText() const override;
 	virtual TArray<FPCGPreConfiguredSettingsInfo> GetPreconfiguredInfo() const override;
 	virtual bool GroupPreconfiguredSettings() const override { return false; }
 	virtual bool HasFlippedTitleLines() const override { return false; }
@@ -50,34 +59,40 @@ public:
 	//~End IPCGSettingsDefaultValueProvider interface
 
 public:
+	UPROPERTY(BlueprintReadWrite, EditAnywhere, Category = Settings, meta = (PCG_Overridable))
+	EPCGAttributeRemapMode Mode = EPCGAttributeRemapMode::Ranges;
+	
 	UPROPERTY(BlueprintReadWrite, EditAnywhere, Category = Settings, meta = (PCG_Overridable))
 	FPCGAttributePropertyInputSelector InputSource;
 
 	/** If InRangeMin = InRangeMax, then that attribute value is mapped to the average of OutRangeMin and OutRangeMax */
-	UPROPERTY(BlueprintReadWrite, EditAnywhere, Category = Settings, meta = (PCG_Overridable))
+	UPROPERTY(BlueprintReadWrite, EditAnywhere, Category = Settings, meta = (PCG_Overridable, EditCondition = "Mode == EPCGAttributeRemapMode::Ranges", EditConditionHides))
 	double InRangeMin = 0.f;
 
 	/** If InRangeMin = InRangeMax, then that attribute value is mapped to the average of OutRangeMin and OutRangeMax */
-	UPROPERTY(BlueprintReadWrite, EditAnywhere, Category = Settings, meta = (PCG_Overridable))
+	UPROPERTY(BlueprintReadWrite, EditAnywhere, Category = Settings, meta = (PCG_Overridable, EditCondition = "Mode == EPCGAttributeRemapMode::Ranges", EditConditionHides))
 	double InRangeMax = 1.f;
 
-	UPROPERTY(BlueprintReadWrite, EditAnywhere, Category = Settings, meta = (PCG_Overridable))
+	UPROPERTY(BlueprintReadWrite, EditAnywhere, Category = Settings, meta = (PCG_Overridable, EditCondition = "Mode == EPCGAttributeRemapMode::Ranges", EditConditionHides))
 	double OutRangeMin = 0.f;
 
-	UPROPERTY(BlueprintReadWrite, EditAnywhere, Category = Settings, meta = (PCG_Overridable))
+	UPROPERTY(BlueprintReadWrite, EditAnywhere, Category = Settings, meta = (PCG_Overridable, EditCondition = "Mode == EPCGAttributeRemapMode::Ranges", EditConditionHides))
 	double OutRangeMax = 1.f;
 
 	/** If checked, outside values will be clamped between 0 and 1. */
-	UPROPERTY(BlueprintReadWrite, EditAnywhere, Category = Settings, meta = (PCG_Overridable))
+	UPROPERTY(BlueprintReadWrite, EditAnywhere, Category = Settings, meta = (PCG_Overridable, EditCondition = "Mode == EPCGAttributeRemapMode::Ranges", EditConditionHides))
 	bool bClampToUnitRange = false;
 
 	/** Attribute values outside of the input range will be unaffected by the remapping */
-	UPROPERTY(BlueprintReadWrite, EditAnywhere, Category = Settings, meta = (PCG_Overridable))
+	UPROPERTY(BlueprintReadWrite, EditAnywhere, Category = Settings, meta = (PCG_Overridable, EditCondition = "Mode == EPCGAttributeRemapMode::Ranges", EditConditionHides))
 	bool bIgnoreValuesOutsideInputRange = false;
 
 	/** Allow remapping when Min is larger than Max, e.g. from [0.0, 1.0] -> [1.0, 0.0]. */
-	UPROPERTY(BlueprintReadWrite, EditAnywhere, Category = Settings, meta = (PCG_Overridable))
+	UPROPERTY(BlueprintReadWrite, EditAnywhere, Category = Settings, meta = (PCG_Overridable, EditCondition = "Mode == EPCGAttributeRemapMode::Ranges", EditConditionHides))
 	bool bAllowInverseRange = false; // Note that this is no longer the default value for new nodes, it is now 'true'
+
+	UPROPERTY(BlueprintReadWrite, EditAnywhere, Category = Settings, meta = (PCG_Overridable, PCG_OverridableChildProperties="ExternalCurve", EditCondition = "Mode == EPCGAttributeRemapMode::Curve", EditConditionHides))
+	FRuntimeFloatCurve Curve;
 };
 
 class FPCGAttributeRemapElement : public FPCGMetadataElementBase
diff --git a/Engine/Plugins/PCG/Source/PCG/Public/Elements/PCGSplineSampler.h b/Engine/Plugins/PCG/Source/PCG/Public/Elements/PCGSplineSampler.h
--- a/Engine/Plugins/PCG/Source/PCG/Public/Elements/PCGSplineSampler.h
+++ b/Engine/Plugins/PCG/Source/PCG/Public/Elements/PCGSplineSampler.h
@@ -138,7 +138,7 @@ struct FPCGSplineSamplerParams
 
 	// TODO: DirtyCache for OnDependencyChanged when this float curve is an external asset
 	/** Defines the density for each sample based on its distance from the spline. X axis is normalized distance to boundary (0-1), Y axis is density value. */
-	UPROPERTY(BlueprintReadWrite, EditAnywhere, Category = Settings, meta = (EditCondition = "Dimension==EPCGSplineSamplingDimension::OnInterior", EditConditionHides))
+	UPROPERTY(BlueprintReadWrite, EditAnywhere, Category = Settings, meta = (PCG_Overridable, PCG_OverridableChildProperties="ExternalCurve", EditCondition = "Dimension==EPCGSplineSamplingDimension::OnInterior", EditConditionHides))
 	FRuntimeFloatCurve InteriorDensityFalloffCurve;
 
 	/** Compute the delta angle to the next point on the spline and write it to an attribute. */
