[PCG] Fixed crash in Launch in Standalone in PCG custom HLSL kernel by moving tokenizer helper out of editor module - it is needed when launching in standalone from editor which is an editor-like environment but does not have the PCG module.

#rnx

diff --git a/Engine/Plugins/PCG/Source/PCG/Internal/Compute/PCGHLSLSyntaxTokenizer.h b/Engine/Plugins/PCG/Source/PCG/Internal/Compute/PCGHLSLSyntaxTokenizer.h
new file mode 100644
--- /dev/null
+++ b/Engine/Plugins/PCG/Source/PCG/Internal/Compute/PCGHLSLSyntaxTokenizer.h
@@ -0,0 +1,20 @@
+// Copyright Epic Games, Inc. All Rights Reserved.
+
+#pragma once
+
+#if WITH_EDITOR
+
+#include "Editor/PCGSyntaxTokenizerParams.h"
+
+#include "Text/HLSLSyntaxTokenizer.h"
+
+/**
+ * Tokenize the text based on HLSL keywords and operators. Also considers PCG specific keywords as tokens if provided through AddAdditionalKeywords().
+ */
+class FPCGHLSLSyntaxTokenizer : public FHlslSyntaxTokenizer
+{
+public:
+	PCG_API FPCGHLSLSyntaxTokenizer(const FPCGSyntaxTokenizerParams& InParams = {});
+};
+
+#endif // WITH_EDITOR
diff --git a/Engine/Plugins/PCG/Source/PCG/Private/Compute/Elements/PCGCustomHLSLKernel.cpp b/Engine/Plugins/PCG/Source/PCG/Private/Compute/Elements/PCGCustomHLSLKernel.cpp
--- a/Engine/Plugins/PCG/Source/PCG/Private/Compute/Elements/PCGCustomHLSLKernel.cpp
+++ b/Engine/Plugins/PCG/Source/PCG/Private/Compute/Elements/PCGCustomHLSLKernel.cpp
@@ -14,7 +14,8 @@
 #include "Containers/StaticArray.h"
 
 #if WITH_EDITOR
-#include "Editor/IPCGEditorModule.h"
+#include "Compute/PCGHLSLSyntaxTokenizer.h"
+
 #include "Framework/Text/SyntaxTokenizer.h"
 #include "Internationalization/Regex.h"
 #endif
@@ -803,12 +804,10 @@ void UPCGCustomHLSLKernel::ParseShaderSource()
 	AdditionalKeywords.Append(CopyElementKeywords);
 	AdditionalKeywords.Append(InitializeKeywords);
 
-	IPCGEditorModule* PCGEditorModule = IPCGEditorModule::Get();
-	check(PCGEditorModule);
-
 	FPCGSyntaxTokenizerParams TokenizerParams;
 	TokenizerParams.AdditionalKeywords = MoveTemp(AdditionalKeywords);
-	TSharedPtr<ISyntaxTokenizer> Tokenizer = PCGEditorModule->CreateHLSLSyntaxTokenizer(TokenizerParams);
+
+	TSharedPtr<ISyntaxTokenizer> Tokenizer = MakeShared<FPCGHLSLSyntaxTokenizer>(TokenizerParams);
 	check(Tokenizer.IsValid());
 
 	auto ParseHelper = [&](FPCGCustomHLSLParsedSource& InOutParsedSource)
diff --git a/Engine/Plugins/PCG/Source/PCG/Private/Compute/PCGHLSLSyntaxTokenizer.cpp b/Engine/Plugins/PCG/Source/PCG/Private/Compute/PCGHLSLSyntaxTokenizer.cpp
new file mode 100644
--- /dev/null
+++ b/Engine/Plugins/PCG/Source/PCG/Private/Compute/PCGHLSLSyntaxTokenizer.cpp
@@ -0,0 +1,13 @@
+// Copyright Epic Games, Inc. All Rights Reserved.
+
+#include "Compute/PCGHLSLSyntaxTokenizer.h"
+
+#if WITH_EDITOR
+
+FPCGHLSLSyntaxTokenizer::FPCGHLSLSyntaxTokenizer(const FPCGSyntaxTokenizerParams& InParams)
+	: FHlslSyntaxTokenizer()
+{
+	Keywords.Append(InParams.AdditionalKeywords);
+}
+
+#endif // WITH_EDITOR
diff --git a/Engine/Plugins/PCG/Source/PCG/Public/Editor/IPCGEditorModule.h b/Engine/Plugins/PCG/Source/PCG/Public/Editor/IPCGEditorModule.h
--- a/Engine/Plugins/PCG/Source/PCG/Public/Editor/IPCGEditorModule.h
+++ b/Engine/Plugins/PCG/Source/PCG/Public/Editor/IPCGEditorModule.h
@@ -26,7 +26,6 @@ public:
 	virtual const FPCGNodeVisualLogs& GetNodeVisualLogs() const = 0;
 	virtual FPCGNodeVisualLogs& GetNodeVisualLogsMutable() = 0;
 	virtual bool CanSelectPartitionActors() const = 0;
-	virtual TSharedPtr<class ISyntaxTokenizer> CreateHLSLSyntaxTokenizer(const FPCGSyntaxTokenizerParams& InParams) const = 0;
 	virtual IPCGBaseSubsystem* GetEditorSubsystem() const = 0;
 #endif
 
diff --git a/Engine/Plugins/PCG/Source/PCGEditor/Private/PCGEditorModule.cpp b/Engine/Plugins/PCG/Source/PCGEditor/Private/PCGEditorModule.cpp
--- a/Engine/Plugins/PCG/Source/PCGEditor/Private/PCGEditorModule.cpp
+++ b/Engine/Plugins/PCG/Source/PCGEditor/Private/PCGEditorModule.cpp
@@ -28,7 +28,6 @@
 #include "PCGEditorSettings.h"
 #include "PCGEditorStyle.h"
 #include "PCGEditorUtils.h"
-#include "PCGHLSLSyntaxTokenizer.h"
 #include "DataVisualizations/PCGBaseTextureDataVisualization.h"
 #include "DataVisualizations/PCGLandscapeDataVisualization.h"
 #include "DataVisualizations/PCGParamDataVisualization.h"
@@ -264,11 +263,6 @@ bool FPCGEditorModule::CanSelectPartitionActors() const
 	return GetDefault<UPCGEditorSettings>()->bCanSelectPartitionActors;
 }
 
-TSharedPtr<ISyntaxTokenizer> FPCGEditorModule::CreateHLSLSyntaxTokenizer(const FPCGSyntaxTokenizerParams& InParams) const
-{
-	return MakeShared<FPCGHLSLSyntaxTokenizer>(InParams);
-}
-
 IPCGBaseSubsystem* FPCGEditorModule::GetEditorSubsystem() const
 {
 	return GEditor ? GEditor->GetEditorSubsystem<UPCGEditorSubsystem>() : nullptr;
diff --git a/Engine/Plugins/PCG/Source/PCGEditor/Private/PCGHLSLSyntaxHighlighter.cpp b/Engine/Plugins/PCG/Source/PCGEditor/Private/PCGHLSLSyntaxHighlighter.cpp
--- a/Engine/Plugins/PCG/Source/PCGEditor/Private/PCGHLSLSyntaxHighlighter.cpp
+++ b/Engine/Plugins/PCG/Source/PCGEditor/Private/PCGHLSLSyntaxHighlighter.cpp
@@ -1,7 +1,9 @@
 // Copyright Epic Games, Inc. All Rights Reserved.
 
 #include "PCGHLSLSyntaxHighlighter.h"
-#include "PCGHLSLSyntaxTokenizer.h"
+
+#include "Compute/PCGHLSLSyntaxTokenizer.h"
+
 #include "PCGEditorStyle.h"
 #include "Styling/SlateTypes.h"
 #include "Framework/Text/IRun.h"
diff --git a/Engine/Plugins/PCG/Source/PCGEditor/Private/PCGHLSLSyntaxTokenizer.cpp b/Engine/Plugins/PCG/Source/PCGEditor/Private/PCGHLSLSyntaxTokenizer.cpp
deleted file mode 100644
diff --git a/Engine/Plugins/PCG/Source/PCGEditor/Private/PCGHLSLSyntaxTokenizer.h b/Engine/Plugins/PCG/Source/PCGEditor/Private/PCGHLSLSyntaxTokenizer.h
deleted file mode 100644
diff --git a/Engine/Plugins/PCG/Source/PCGEditor/Public/PCGEditorModule.h b/Engine/Plugins/PCG/Source/PCGEditor/Public/PCGEditorModule.h
--- a/Engine/Plugins/PCG/Source/PCGEditor/Public/PCGEditorModule.h
+++ b/Engine/Plugins/PCG/Source/PCGEditor/Public/PCGEditorModule.h
@@ -37,7 +37,6 @@ public:
 	virtual const FPCGNodeVisualLogs& GetNodeVisualLogs() const override { return NodeVisualLogs; }
 	virtual FPCGNodeVisualLogs& GetNodeVisualLogsMutable() override { return NodeVisualLogs; }
 	virtual bool CanSelectPartitionActors() const override;
-	virtual TSharedPtr<class ISyntaxTokenizer> CreateHLSLSyntaxTokenizer(const FPCGSyntaxTokenizerParams& InParams) const override;
 	virtual IPCGBaseSubsystem* GetEditorSubsystem() const override;
 
 protected:
