[PCG] Improve tooltips in the spawn actor node

#rnx

diff --git a/Engine/Plugins/PCG/Source/PCG/Public/Elements/PCGSpawnActor.h b/Engine/Plugins/PCG/Source/PCG/Public/Elements/PCGSpawnActor.h
--- a/Engine/Plugins/PCG/Source/PCG/Public/Elements/PCGSpawnActor.h
+++ b/Engine/Plugins/PCG/Source/PCG/Public/Elements/PCGSpawnActor.h
@@ -16,18 +16,18 @@ class UPCGBasePointData;
 UENUM()
 enum class EPCGSpawnActorOption : uint8
 {
-	CollapseActors,
-	MergePCGOnly,
-	NoMerging
+	CollapseActors UMETA(Tooltip="Extracts the mesh(es) from the actor class and creates a single merged representation instead and merges graph execution if needed."),
+	MergePCGOnly UMETA(Tooltip="Spawns actors as-is but any PCG component on the actor will merge into a single graph execution."),
+	NoMerging UMETA(Tooltip="Spawns one actor per point.")
 };
 
 UENUM()
 enum class EPCGSpawnActorGenerationTrigger : uint8
 {
-	Default, // Generate if the component has "Generate On Load"
-	ForceGenerate, // Generate in all cases
-	DoNotGenerateInEditor, // Does not call generate in editor, but decays to Default otherwise
-	DoNotGenerate // Does not call generate
+	Default UMETA(Tooltip="Default; will trigger generation if a PCG Component exists and has the 'Generate On Load' trigger."),
+	ForceGenerate UMETA(Tooltip = "Will trigger generation if a PCG Component exists, regardless of generation trigger."),
+	DoNotGenerateInEditor UMETA(Tooltip = "Will not trigger generation in the editor but decays to Default otherwise."),
+	DoNotGenerate UMETA(Tooltip = "Will never call generation on any PCG components")
 };
 
 /*
@@ -52,8 +52,12 @@ public:
 	UPROPERTY(BlueprintReadWrite, EditAnywhere, Category = Settings, meta = (EditCondition = "Option!=EPCGSpawnActorOption::CollapseActors", EditConditionHides))
 	TArray<FName> PostSpawnFunctionNames;
 
+	/** 
+	* Controls how actors are spawned; collapsed in more efficient components or as-is with generation or not. 
+	* Note that new nodes now have the 'No Merging' option by default.
+	*/
 	UPROPERTY(BlueprintReadWrite, EditAnywhere, Category = Settings)
-	EPCGSpawnActorOption Option = EPCGSpawnActorOption::CollapseActors; // Note that this is no longer the default value for new nodes, it is now EPCGSpawnActorOption::NoMerging.
+	EPCGSpawnActorOption Option = EPCGSpawnActorOption::CollapseActors; 
 
 	UPROPERTY(BlueprintReadWrite, EditAnywhere, Category = Settings, meta = (EditCondition = "Option==EPCGSpawnActorOption::NoMerging", EditConditionHides))
 	bool bForceDisableActorParsing = true;
@@ -77,8 +81,12 @@ public:
 	UPROPERTY(meta = (PCG_Overridable))
 	TSoftObjectPtr<AActor> RootActor;
 
+	/** 
+	* Controls where spawned actors will appear in the Outliner. Note that attaching actors to an actor couples their streaming. 
+	* Note that new nodes now have the 'In Folder' option by default.
+	*/
 	UPROPERTY(BlueprintReadWrite, EditAnywhere, Category = Settings, meta = (EditCondition = "Option != EPCGSpawnActorOption::CollapseActors", EditConditionHides))
-	EPCGAttachOptions AttachOptions = EPCGAttachOptions::Attached; // Note that this is no longer the default value for new nodes, it is now EPCGAttachOptions::InFolder
+	EPCGAttachOptions AttachOptions = EPCGAttachOptions::Attached; 
 
 	UPROPERTY(BlueprintReadWrite, EditAnywhere, Category = Settings, meta = (InlineEditConditionToggle))
 	bool bSpawnByAttribute = false;
