[PCG] Fix Create Points Sphere crash from divide-by-zero when graph params are used to set degrees to 0

diff --git a/Engine/Plugins/PCG/Source/PCG/Private/Elements/PCGCreatePointsSphere.cpp b/Engine/Plugins/PCG/Source/PCG/Private/Elements/PCGCreatePointsSphere.cpp
--- a/Engine/Plugins/PCG/Source/PCG/Private/Elements/PCGCreatePointsSphere.cpp
+++ b/Engine/Plugins/PCG/Source/PCG/Private/Elements/PCGCreatePointsSphere.cpp
@@ -29,6 +29,7 @@ namespace PCGCreatePointsSphere
 			double InRadius,
 			const FVector& InOrigin,
 			int32 InSeed = 0,
+			int32 InPointLimit = INT_MAX,
 			FPCGContext* InContext = nullptr)
 		: LatitudinalStartAngle(FMath::Abs(InLatitudinalStartAngle) > 90 ? FMath::Fmod(InLatitudinalStartAngle, 90) : InLatitudinalStartAngle)
 		, LatitudinalEndAngle(FMath::Abs(InLatitudinalEndAngle) > 90 ? FMath::Fmod(InLatitudinalEndAngle, 90) : InLatitudinalEndAngle)
@@ -37,6 +38,7 @@ namespace PCGCreatePointsSphere
 		, Radius(InRadius)
 		, Origin(InOrigin)
 		, Seed(InSeed)
+		, PointLimit(InPointLimit)
 		, Context(InContext) {}
 
 		double LatitudinalStartAngle = 0;
@@ -46,9 +48,28 @@ namespace PCGCreatePointsSphere
 		double Radius = 0;
 		FVector Origin = FVector::ZeroVector;
 		int32 Seed = 0;
+		int32 PointLimit = INT_MAX;
 		FPCGContext* Context = nullptr;
 	};
 
+	bool GenerationWouldOverflow(const int32 PointCount, const int32 PointLimit, const FPCGContext* InContext = nullptr)
+	{
+		if (PointCount > PointLimit)
+		{
+			PCGLog::LogErrorOnGraph(LOCTEXT("PointOverflow", "The number of points to be produced exceeds the point limit. See 'Point Limit' under Advanced Settings."), InContext);
+			return true;
+		}
+		else if (PCGFeatureSwitches::CVarCheckSamplerMemory.GetValueOnAnyThread() && PCGFeatureSwitches::Helpers::GetAvailableMemoryForSamplers() < sizeof(FPCGPoint) * PointCount)
+		{
+			PCGLog::LogErrorOnGraph(LOCTEXT("MemoryOverflow", "The number of iterations produced is larger than available memory."), InContext);
+			return true;
+		}
+		else
+		{
+			return false;
+		}
+	}
+
 	FVector GeneratePointInWindow(const FCommonParams& Params, FRandomStream& RandomSource)
 	{
 		// Use the inverse cosine's range of [-1,1] to drive the angle for uniform sampling
@@ -190,7 +211,7 @@ namespace PCGCreatePointsSphere
 		OutLocations.Append(PointLocations);
 	}
 
-	void GenerateStandardSpherePointLocations(const FCommonParams& Params, double InTheta, double InPhi, double InJitter, TArray<FVector>& OutLocations)
+	void GenerateStandardSpherePointLocations(const FCommonParams& Params, const double InTheta, const double InPhi, const double InJitter, TArray<FVector>& OutLocations)
 	{
 		TRACE_CPUPROFILER_EVENT_SCOPE(FPCGCreatePointsSphereElement::StandardSphere);
 		OutLocations.Empty();
@@ -200,15 +221,18 @@ namespace PCGCreatePointsSphere
 		const double& LonStartAngle = Params.LongitudinalStartAngle;
 		const double& LonEndAngle = Params.LongitudinalEndAngle;
 
-		const int32 LatSegCount = 1 + FMath::FloorToInt32(FMath::Abs(LatEndAngle - LatStartAngle) / InTheta);
-		const int32 LonSegCount = FMath::FloorToInt32(FMath::Abs(LonEndAngle - LonStartAngle) / InPhi);
+		// Clamp to match the UI input and guard against divide-by-zero.
+		const double ThetaPerSegment = FMath::Clamp(InTheta, 0.001, 90.0);
+		const double PhiPerSegment = FMath::Clamp(InPhi, 0.001, 180.0);
+
+		const int32 LatSegCount = 1 + FMath::FloorToInt32(FMath::Abs(LatEndAngle - LatStartAngle) / ThetaPerSegment);
+		const int32 LonSegCount = FMath::FloorToInt32(FMath::Abs(LonEndAngle - LonStartAngle) / PhiPerSegment);
 
 		int64 PointCount = 0;
 		// To produce a single pole point, instead of a whole segment
 		int32 ReducedLatSegCount = LatSegCount;
 
 		// Account for the poles
-
 		bool bSouthPole = false;
 		if (FMath::IsNearlyEqual(LatEndAngle, -90.0))
 		{
@@ -232,9 +256,8 @@ namespace PCGCreatePointsSphere
 			return;
 		}
 
-		if (PCGFeatureSwitches::CVarCheckSamplerMemory.GetValueOnAnyThread() && FPlatformMemory::GetStats().AvailablePhysical < sizeof(FPCGPoint) * PointCount)
+		if (GenerationWouldOverflow(PointCount, Params.PointLimit, Params.Context))
 		{
-			PCGLog::LogErrorOnGraph(LOCTEXT("MemoryOverflow", "The number of iterations produced is larger than available memory."));
 			return;
 		}
 
@@ -262,19 +285,19 @@ namespace PCGCreatePointsSphere
 			}
 
 			// -90 to map to a -180 to 0 range for cosine
-			const double Theta = FMath::DegreesToRadians(LatStartAngle - 90 + InTheta * LatIndex);
+			const double Theta = FMath::DegreesToRadians(LatStartAngle - 90 + ThetaPerSegment * LatIndex);
 			// Pre-calculation of these for the internal loop is possible unless jittered
 			double SinTheta = bApplyJitter ? 0.0 : FMath::Sin(Theta);
 			double CosTheta = bApplyJitter ? 0.0 : FMath::Cos(Theta);
 
 			for (int LonIndex = 0; LonIndex < LonSegCount; ++LonIndex)
 			{
-				double Phi = FMath::DegreesToRadians(LonStartAngle + InPhi * LonIndex);
+				double Phi = FMath::DegreesToRadians(LonStartAngle + PhiPerSegment * LonIndex);
 
 				if (bApplyJitter)
 				{
-					Phi += FMath::DegreesToRadians(InPhi) * RandomSource.FRandRange(-Jitter, Jitter);
-					const double JitteredTheta = Theta + FMath::DegreesToRadians(InTheta) * RandomSource.FRandRange(-Jitter, Jitter);
+					Phi += FMath::DegreesToRadians(PhiPerSegment) * RandomSource.FRandRange(-Jitter, Jitter);
+					const double JitteredTheta = Theta + FMath::DegreesToRadians(ThetaPerSegment) * RandomSource.FRandRange(-Jitter, Jitter);
 					SinTheta = FMath::Sin(JitteredTheta);
 					CosTheta = FMath::Cos(JitteredTheta);
 				}
@@ -286,14 +309,13 @@ namespace PCGCreatePointsSphere
 		check(CurrentIndex == PointCount);
 	}
 
-	void GenerateRandomSamplingPointLocations(const FCommonParams& Params, int32 SampleCount, TArray<FVector>& OutLocations)
+	void GenerateRandomSamplingPointLocations(const FCommonParams& Params, const int32 SampleCount, TArray<FVector>& OutLocations)
 	{
 		TRACE_CPUPROFILER_EVENT_SCOPE(FPCGCreatePointsSphereElement::RandomSampling);
 		OutLocations.Empty();
 
-		if (PCGFeatureSwitches::CVarCheckSamplerMemory.GetValueOnAnyThread() && FPlatformMemory::GetStats().AvailablePhysical < sizeof(FPCGPoint) * SampleCount)
+		if (GenerationWouldOverflow(SampleCount, Params.PointLimit, Params.Context))
 		{
-			PCGLog::LogErrorOnGraph(LOCTEXT("MemoryOverflow", "The number of iterations produced is larger than available memory."));
 			return;
 		}
 
@@ -307,14 +329,13 @@ namespace PCGCreatePointsSphere
 		}
 	}
 
-	void GeneratePoissonSamplingPointLocations(const FCommonParams& Params, double Distance, int32 SampleCount, int32 MaxAttempts, TArray<FVector>& OutLocations)
+	void GeneratePoissonSamplingPointLocations(const FCommonParams& Params, const double Distance, const int32 SampleCount, const int32 MaxAttempts, TArray<FVector>& OutLocations)
 	{
 		TRACE_CPUPROFILER_EVENT_SCOPE(FPCGCreatePointsSphereElement::Poisson);
 		OutLocations.Empty();
 
-		if (PCGFeatureSwitches::CVarCheckSamplerMemory.GetValueOnAnyThread() && FPlatformMemory::GetStats().AvailablePhysical < sizeof(FPCGPoint) * SampleCount)
+		if (GenerationWouldOverflow(SampleCount, Params.PointLimit, Params.Context))
 		{
-			PCGLog::LogErrorOnGraph(LOCTEXT("MemoryOverflow", "The number of iterations produced is larger than available memory."));
 			return;
 		}
 
@@ -392,6 +413,7 @@ bool FPCGCreatePointsSphereElement::ExecuteInternal(FPCGContext* InContext) cons
 		Settings->Radius,
 		Settings->Origin,
 		InContext->GetSeed(),
+		Settings->PointLimit,
 		InContext);
 
 	TArray<FVector> OutLocalPointLocations;
diff --git a/Engine/Plugins/PCG/Source/PCG/Public/Elements/PCGCreatePointsSphere.h b/Engine/Plugins/PCG/Source/PCG/Public/Elements/PCGCreatePointsSphere.h
--- a/Engine/Plugins/PCG/Source/PCG/Public/Elements/PCGCreatePointsSphere.h
+++ b/Engine/Plugins/PCG/Source/PCG/Public/Elements/PCGCreatePointsSphere.h
@@ -138,6 +138,10 @@ public:
 	/** Points are removed if they are outside the volume. */
 	UPROPERTY(BlueprintReadWrite, EditAnywhere, Category = Settings, meta = (PCG_Overridable))
 	bool bCullPointsOutsideVolume = false;
+
+	/** Adds a graph error and prevents generation if the generated point count would exceed this amount. */
+	UPROPERTY(BlueprintReadWrite, EditAnywhere, Category = Settings, AdvancedDisplay, meta = (ClampMin = 1))
+	int32 PointLimit = 1000000;
 };
 
 class FPCGCreatePointsSphereElement : public IPCGElement
