[PCG] Fix concurrency issue in FPCGGraphExecutor::CheckState where ScheduleTasks was checked without lock but not later checked again with lock

#rnx

diff --git a/Engine/Plugins/PCG/Source/PCG/Private/Graph/PCGGraphExecutor.cpp b/Engine/Plugins/PCG/Source/PCG/Private/Graph/PCGGraphExecutor.cpp
--- a/Engine/Plugins/PCG/Source/PCG/Private/Graph/PCGGraphExecutor.cpp
+++ b/Engine/Plugins/PCG/Source/PCG/Private/Graph/PCGGraphExecutor.cpp
@@ -880,25 +880,51 @@ void FPCGGraphExecutor::ClearOutputData(FPCGTaskId TaskId)
 
 void FPCGGraphExecutor::CheckState(bool bIsInGameThread)
 {
-	if (bIsInGameThread)
+	// Only check state on game thread
+	if (!bIsInGameThread)
 	{
-		// Check every second to avoid locking too often for an error that shouldn't happen
-		double TimeInSeconds = FPlatformTime::Seconds();
-		if (TimeInSeconds - LastSchedulingErrorCheck > 1.0)
-		{
-			LastSchedulingErrorCheck = TimeInSeconds;
+		return;
+	}
 
-			PCGGraphExecutor::TScopeLock ScopeLock(TasksLock);
-			PCGGraphExecutor::TScopeLock ChildScopeLock(LiveTasksLock);
-			PCGGraphExecutor::TScopeLock PausedLock(PausedTaskSuccessorsLock);
-			PCGGraphExecutor::TScopeLock CachingResultsScopeLock(CachingResultsLock);
-			// This is a safeguard to check if we're in a stuck state
-			if (CachingResultsForGC.Num() == 0 && ReadyTasks.Num() == 0 && ActiveTasks.Num() == 0 && ActiveTasksGameThreadOnly.Num() == 0 && PausedTasks.Num() == 0 && Tasks.Num() > 0)
-			{
-				UE_LOG(LogPCG, Error, TEXT("PCG Graph executor error: tasks are in a deadlocked state. Will drop all tasks."));
-				ClearAllTasks();
-			}
-		}
+	// ScheduledTasks needs to be empty, check without a lock to early out if that is not the case
+	if (!ScheduledTasks.IsEmpty())
+	{
+		return;
+	}
+
+	// Check every second to avoid locking too often for an error that shouldn't happen
+	const double TimeInSeconds = FPlatformTime::Seconds();
+	const double ElapsedTimeSinceLastCheck = TimeInSeconds - LastSchedulingErrorCheck;
+	if (ElapsedTimeSinceLastCheck < 1.0)
+	{
+		return;
+	}
+
+	// Since this is for an error case avoid locking if lock is already taken, we can comeback later but if we haven't checked in more than 3 seconds actually lock
+	PCGGraphExecutor::TScopeLock ScheduleScopeLock(ScheduleLock, /*bTryLock=*/ ElapsedTimeSinceLastCheck < 3.0);
+	if (!ScheduleScopeLock.IsLocked())
+	{
+		return;
+	}
+	
+	LastSchedulingErrorCheck = TimeInSeconds;
+
+	// Check again now with lock
+	if (!ScheduledTasks.IsEmpty())
+	{
+		return;
+	}
+		
+	PCGGraphExecutor::TScopeLock ScopeLock(TasksLock);
+	PCGGraphExecutor::TScopeLock ChildScopeLock(LiveTasksLock);
+	PCGGraphExecutor::TScopeLock PausedLock(PausedTaskSuccessorsLock);
+	PCGGraphExecutor::TScopeLock CachingResultsScopeLock(CachingResultsLock);
+
+	// This is a safeguard to check if we're in a stuck state
+	if (CachingResultsForGC.Num() == 0 && ReadyTasks.Num() == 0 && ActiveTasks.Num() == 0 && ActiveTasksGameThreadOnly.Num() == 0 && PausedTasks.Num() == 0 && Tasks.Num() > 0)
+	{
+		UE_LOG(LogPCG, Error, TEXT("PCG Graph executor error: tasks are in a deadlocked state. Will drop all tasks."));
+		ClearAllTasks();
 	}
 }
 
@@ -909,13 +935,6 @@ bool FPCGGraphExecutor::ProcessScheduledTasks(double EndTime)
 	const bool bIsInGameThread = IsInGameThread();
 	TArray<FCachedResult*> CachedResults;
 
-	// Avoid locking if empty
-	if (ScheduledTasks.IsEmpty())
-	{
-		CheckState(bIsInGameThread);
-		return bProcessedNewTasks;
-	}
-
 	{
 		// Process any newly scheduled graphs to execute
 		PCGGraphExecutor::TScopeLock ScopeLock(ScheduleLock, !bIsInGameThread);
