[PCG] First pass at 3d spline intersection.

diff --git a/Engine/Plugins/PCG/Source/PCG/Private/Data/PCGSplineData.cpp b/Engine/Plugins/PCG/Source/PCG/Private/Data/PCGSplineData.cpp
--- a/Engine/Plugins/PCG/Source/PCG/Private/Data/PCGSplineData.cpp
+++ b/Engine/Plugins/PCG/Source/PCG/Private/Data/PCGSplineData.cpp
@@ -238,6 +238,11 @@ FTransform UPCGSplineData::GetTransformAtAlpha(float Alpha) const
 }
 
 void UPCGSplineData::WriteMetadataToPoint(float InputKey, FPCGPoint& OutPoint, UPCGMetadata* OutMetadata) const
+{
+	return WriteMetadataToEntry(InputKey, OutPoint.MetadataEntry, OutMetadata);
+}
+
+void UPCGSplineData::WriteMetadataToEntry(float InputKey, PCGMetadataEntryKey& OutEntryKey, UPCGMetadata* OutMetadata) const
 {
 	if (!OutMetadata || !ensure(Metadata))
 	{
@@ -283,7 +288,7 @@ void UPCGSplineData::WriteMetadataToPoint(float InputKey, FPCGPoint& OutPoint, U
 
 		check(EntryKeys.IsValidIndex(PreviousIndex) && EntryKeys.IsValidIndex(NextIndex));
 		
-		OutMetadataDomain->InitializeOnSet(OutPoint.MetadataEntry, EntryKeys[PreviousIndex], MetadataDomain);
+		OutMetadataDomain->InitializeOnSet(OutEntryKey, EntryKeys[PreviousIndex], MetadataDomain);
 		
 		if (PreviousIndex != NextIndex && !FMath::IsNearlyEqual(PreviousInputKey, NextInputKey))
 		{
@@ -291,20 +296,20 @@ void UPCGSplineData::WriteMetadataToPoint(float InputKey, FPCGPoint& OutPoint, U
 			TStaticArray<TPair<PCGMetadataEntryKey, float>, 2> Coefficients;
 			Coefficients[0] = {EntryKeys[PreviousIndex], 1.0f - Alpha};
 			Coefficients[1] = {EntryKeys[NextIndex], Alpha};
-			OutMetadataDomain->ComputeWeightedAttribute(OutPoint.MetadataEntry, Coefficients, MetadataDomain);
+			OutMetadataDomain->ComputeWeightedAttribute(OutEntryKey, Coefficients, MetadataDomain);
 		}
 	}
 
 	if (const TArray<FName>& ChannelNames = SplineStruct.GetAttributeChannelNames(); ChannelNames.Num() > 0)
 	{
-		OutMetadataDomain->InitializeOnSet(OutPoint.MetadataEntry);
+		OutMetadataDomain->InitializeOnSet(OutEntryKey);
 
 		// Sample attribute channels,
 		for (const FName& ChannelName : ChannelNames)
 		{
 			if (FPCGMetadataAttribute<float>* UserAttribute = OutMetadataDomain->FindOrCreateAttribute<float>(ChannelName, /*DefaultValue=*/0.f))
 			{
-				UserAttribute->SetValue(OutPoint.MetadataEntry, SplineStruct.GetAttributeValue(InputKey, ChannelName));
+				UserAttribute->SetValue(OutEntryKey, SplineStruct.GetAttributeValue(InputKey, ChannelName));
 			}
 		}
 	}
@@ -523,7 +528,7 @@ TArray<FSplinePoint> UPCGSplineData::GetSplinePoints() const
 	TArray<FSplinePoint> ControlPoints;
 	ControlPoints.Reserve(ControlPointsPosition.Points.Num());
 
-	for (int i = 0; i < ControlPointsPosition.Points.Num() - 1; ++i)
+	for (int i = 0; i < ControlPointsPosition.Points.Num(); ++i)
 	{
 		ControlPoints.Emplace(static_cast<float>(ControlPoints.Num()),
 			ControlPointsPosition.Points[i].OutVal,
diff --git a/Engine/Plugins/PCG/Source/PCG/Private/Data/PCGSplineStruct.cpp b/Engine/Plugins/PCG/Source/PCG/Private/Data/PCGSplineStruct.cpp
--- a/Engine/Plugins/PCG/Source/PCG/Private/Data/PCGSplineStruct.cpp
+++ b/Engine/Plugins/PCG/Source/PCG/Private/Data/PCGSplineStruct.cpp
@@ -157,7 +157,8 @@ void FPCGSplineStruct::Initialize(const TArray<FSplinePoint>& InSplinePoints, bo
 	ReparamStepsPerSegment = 10; // default value in USplineComponent
 
 	bClosedLoop = bIsClosedLoop;
-	AddPoints(InSplinePoints, true);
+	const bool bShouldUseOptionalEntryKeys = !InOptionalEntryKeys.IsEmpty() && InSplinePoints.Num() == InOptionalEntryKeys.Num();
+	AddPoints(InSplinePoints, true, bShouldUseOptionalEntryKeys ? &InOptionalEntryKeys : nullptr);
 
 	if (ShouldUseSplineCurves())
 	{
@@ -170,11 +171,7 @@ void FPCGSplineStruct::Initialize(const TArray<FSplinePoint>& InSplinePoints, bo
 		LocalBounds = PCGSplineStruct::CalcBounds(Spline, bClosedLoop, FTransform::Identity);
 	}
 
-	if (!InOptionalEntryKeys.IsEmpty() && InSplinePoints.Num() == InOptionalEntryKeys.Num())
-	{
-		ControlPointsEntryKeys = std::move(InOptionalEntryKeys);
-	}
-	else
+	if (!bShouldUseOptionalEntryKeys)
 	{
 		// If we have a mismatch, we can't set the entry keys, so reset it
 		ControlPointsEntryKeys.Empty();
@@ -253,7 +250,7 @@ FPCGSplineStruct& FPCGSplineStruct::operator=(FPCGSplineStruct&& Other)
 	return *this;
 }
 
-void FPCGSplineStruct::AddPoint(const FSplinePoint& InSplinePoint, bool bUpdateSpline)
+void FPCGSplineStruct::AddPoint(const FSplinePoint& InSplinePoint, bool bUpdateSpline, const int64* InOptionalControlPointEntryKey)
 {
 	Spline.AddPoint(InSplinePoint);
  
@@ -285,7 +282,7 @@ void FPCGSplineStruct::AddPoint(const FSplinePoint& InSplinePoint, bool bUpdateS
 
 	if (!ControlPointsEntryKeys.IsEmpty())
 	{
-		ControlPointsEntryKeys.Insert(PCGInvalidEntryKey, Index);
+		ControlPointsEntryKeys.Insert(InOptionalControlPointEntryKey ? *InOptionalControlPointEntryKey : PCGInvalidEntryKey, Index);
 	}
 
 	if (bUpdateSpline)
@@ -294,16 +291,18 @@ void FPCGSplineStruct::AddPoint(const FSplinePoint& InSplinePoint, bool bUpdateS
 	}
 }
 
-void FPCGSplineStruct::AddPoints(const TArray<FSplinePoint>& InSplinePoints, bool bUpdateSpline)
+void FPCGSplineStruct::AddPoints(const TArray<FSplinePoint>& InSplinePoints, bool bUpdateSpline, const TArray<int64>* InOptionalControlPointsEntryKeys)
 {
 	if (!ControlPointsEntryKeys.IsEmpty())
 	{
 		ControlPointsEntryKeys.Reserve(ControlPointsEntryKeys.Num() + InSplinePoints.Num());
 	}
 
-	for (const auto& SplinePoint : InSplinePoints)
+	check(!InOptionalControlPointsEntryKeys || InSplinePoints.Num() == InOptionalControlPointsEntryKeys->Num());
+
+	for (int SplinePointIndex = 0; SplinePointIndex < InSplinePoints.Num(); ++SplinePointIndex)
 	{
-		AddPoint(SplinePoint, false);
+		AddPoint(InSplinePoints[SplinePointIndex], false, InOptionalControlPointsEntryKeys ? &((*InOptionalControlPointsEntryKeys)[SplinePointIndex]) : nullptr);
 	}
 
 	if (bUpdateSpline)
@@ -454,43 +453,47 @@ FVector::FReal FPCGSplineStruct::GetSplineLength() const
 
 FBox FPCGSplineStruct::GetBounds() const
 {
-	// See USplineComponent::CalcBounds
-	const int32 NumPoints = GetNumberOfPoints();
-	const int32 NumSegments = GetNumberOfSplineSegments();
-	
-	const FInterpCurveVector& SplinePointsPosition = ShouldUseSplineCurves()
-		? WarninglessSplineCurves().Position
-		: Spline.GetSplinePointsPosition();
+	FBoxSphereBounds BoundsFromStruct;
 
-	FVector Min(WORLD_MAX);
-	FVector Max(-WORLD_MAX);
-	if (NumSegments > 0)
+	if (ShouldUseSplineCurves())
 	{
-		for (int32 Index = 0; Index < NumSegments; Index++)
-		{
-			const bool bLoopSegment = (Index == NumPoints - 1);
-			const int32 NextIndex = bLoopSegment ? 0 : Index + 1;
-			const FInterpCurvePoint<FVector>& ThisInterpPoint = SplinePointsPosition.Points[Index];
-			FInterpCurvePoint<FVector> NextInterpPoint = SplinePointsPosition.Points[NextIndex];
-			if (bLoopSegment)
-			{
-				NextInterpPoint.InVal = ThisInterpPoint.InVal /* + SplinePointsPosition.LoopKeyOffset*/;
-			}
-
-			CurveVectorFindIntervalBounds(ThisInterpPoint, NextInterpPoint, Min, Max);
-		}
+		BoundsFromStruct = PCGSplineStruct::CalcBounds(WarninglessSplineCurves(), bClosedLoop, FTransform::Identity);
 	}
-	else if (NumPoints == 1)
+	else
 	{
-		Min = Max = SplinePointsPosition.Points[0].OutVal;
+		BoundsFromStruct = PCGSplineStruct::CalcBounds(Spline, bClosedLoop, FTransform::Identity);
 	}
-	else
+
+	return BoundsFromStruct.GetBox();
+}
+
+FBox FPCGSplineStruct::GetSegmentBounds(int32 SegmentIndex) const
+{
+	if (SegmentIndex < 0 || SegmentIndex >= GetNumberOfSplineSegments())
 	{
-		Min = FVector::ZeroVector;
-		Max = FVector::ZeroVector;
+		return FBox(EForceInit::ForceInit);
 	}
 
-	return FBox(Min, Max);
+	const FInterpCurveVector& SplinePointsPosition = ShouldUseSplineCurves()
+		? WarninglessSplineCurves().Position
+		: Spline.GetSplinePointsPosition();
+
+	const int32 NumPoints = GetNumberOfPoints();
+	const bool bLoopSegment = (SegmentIndex == NumPoints - 1);
+	const int32 NextIndex = bLoopSegment ? 0 : SegmentIndex + 1;
+	const FInterpCurvePoint<FVector>& ThisInterpPoint = SplinePointsPosition.Points[SegmentIndex];
+	FInterpCurvePoint<FVector> NextInterpPoint = SplinePointsPosition.Points[NextIndex];
+	if (bLoopSegment)
+	{
+		NextInterpPoint.InVal = ThisInterpPoint.InVal + SplinePointsPosition.LoopKeyOffset;
+	}
+
+	FVector CurrentMin(WORLD_MAX);
+	FVector CurrentMax(-WORLD_MAX);
+
+	CurveVectorFindIntervalBounds(ThisInterpPoint, NextInterpPoint, CurrentMin, CurrentMax);
+
+	return FBox(CurrentMin, CurrentMax);
 }
 
 const FInterpCurveVector& FPCGSplineStruct::GetSplinePointsScale() const
@@ -614,6 +617,15 @@ float FPCGSplineStruct::FindInputKeyClosestToWorldLocation(const FVector& WorldL
 		: Spline.FindNearest(LocalLocation, Dummy);
 }
 
+float FPCGSplineStruct::FindInputKeyOnSegmentClosestToWorldLocation(const FVector& WorldLocation, int32 Index) const
+{
+	const FVector LocalLocation = Transform.InverseTransformPosition(WorldLocation);
+	float Dummy;
+	return ShouldUseSplineCurves()
+		? WarninglessSplineCurves().Position.FindNearestOnSegment(LocalLocation, Index, Dummy)
+		: Spline.FindNearestOnSegment(LocalLocation, Index, Dummy);
+}
+
 TTuple<int, float> FPCGSplineStruct::GetSegmentStartIndexAndKeyAtInputKey(float InKey) const
 {
 	const int32 Index = ShouldUseSplineCurves()
@@ -703,6 +715,20 @@ FVector FPCGSplineStruct::GetScaleAtSplineInputKey(float InKey) const
 		: Spline.GetSplinePointsScale().Eval(InKey, FVector(1.0f));
 }
 
+FVector FPCGSplineStruct::GetTangentAtSplineInputKey(float InKey, ESplineCoordinateSpace::Type CoordinateSpace) const
+{
+	FVector Tangent = ShouldUseSplineCurves()
+		? WarninglessSplineCurves().Position.EvalDerivative(InKey, FVector::ZeroVector)
+		: Spline.EvaluateDerivative(InKey);
+
+	if (CoordinateSpace == ESplineCoordinateSpace::World)
+	{
+		Tangent = Transform.TransformVector(Tangent);
+	}
+
+	return Tangent;
+}
+
 /** Taken from USplineComponent::ConvertSplineSegmentToPolyLine. */
 bool FPCGSplineStruct::ConvertSplineSegmentToPolyLine(int32 SplinePointStartIndex, ESplineCoordinateSpace::Type CoordinateSpace, const float MaxSquareDistanceFromSpline, TArray<FVector>& OutPoints) const
 {
diff --git a/Engine/Plugins/PCG/Source/PCG/Private/Elements/PCGSplineIntersection.cpp b/Engine/Plugins/PCG/Source/PCG/Private/Elements/PCGSplineIntersection.cpp
new file mode 100644
--- /dev/null
+++ b/Engine/Plugins/PCG/Source/PCG/Private/Elements/PCGSplineIntersection.cpp
@@ -0,0 +1,532 @@
+// Copyright Epic Games, Inc. All Rights Reserved.
+
+#include "Elements/PCGSplineIntersection.h"
+
+#include "PCGContext.h"
+#include "Data/PCGBasePointData.h"
+#include "Data/PCGSplineData.h"
+#include "Elements/Metadata/PCGMetadataElementCommon.h"
+#include "Helpers/PCGHelpers.h"
+#include "Metadata/Accessors/PCGAttributeAccessorHelpers.h"
+
+#define LOCTEXT_NAMESPACE "PCGSplineIntersectionElement"
+
+namespace PCGSplineIntersection
+{
+	namespace Constants
+	{
+		const FName OriginatingSplineIndexAttributeName = TEXT("OriginatingSplineIndex");
+		const FName IntersectingSplineIndexAttributeName = TEXT("IntersectingSplineIndex");
+	}
+
+	struct InputData
+	{
+		void InitializeBounds(const FVector& HalfRadiusTolerance)
+		{
+			// Implementation note: we add half the distance tolerance to the bounds so we never need to update them after, 
+			// as we'll use them for bounds-bounds tests only.
+			check(Spline);
+			Bounds = Spline->GetBounds().ExpandBy(HalfRadiusTolerance, HalfRadiusTolerance);
+
+			SegmentsNum = Spline->GetNumSegments();
+			SegmentBounds.SetNumUninitialized(SegmentsNum);
+			for (int32 SegmentIndex = 0; SegmentIndex < SegmentBounds.Num(); ++SegmentIndex)
+			{
+				SegmentBounds[SegmentIndex] = Spline->SplineStruct.GetSegmentBounds(SegmentIndex).TransformBy(Spline->SplineStruct.Transform).ExpandBy(HalfRadiusTolerance, HalfRadiusTolerance);
+			}
+		}
+
+		const TArray<FVector>& GetRoughSamples(int SegmentIndex, double MinDistanceThreshold, double& OutRoughSampleStep) const
+		{
+			check(Spline && SegmentIndex >= 0 && SegmentIndex < SegmentsNum);
+			TArray<FVector>& Samples = RoughSampling.FindOrAdd(SegmentIndex);
+
+			const FVector::FReal SegmentLength = Spline->GetSegmentLength(SegmentIndex);
+			FVector::FReal RoughSampleSize = FMath::Max(SegmentLength / 100.0, MinDistanceThreshold);
+
+			if (Samples.IsEmpty())
+			{
+				// Rough sampling should be in the 10-100 range.
+				int NumSamples = FMath::Floor(SegmentLength / RoughSampleSize) + 1;
+
+				Samples.SetNumUninitialized(NumSamples+1);
+
+				const FVector::FReal DistanceAtSegmentStart = Spline->SplineStruct.GetDistanceAlongSplineAtSplinePoint(SegmentIndex);
+				for (int Sample = 0; Sample <= NumSamples; ++Sample)
+				{
+					Samples[Sample] = Spline->SplineStruct.GetLocationAtDistanceAlongSpline(DistanceAtSegmentStart + SegmentLength * (Sample / (double)NumSamples), ESplineCoordinateSpace::World);
+				}
+			}
+
+			// Note: to address scaling, we'll take the distance between two points
+			OutRoughSampleStep = FMath::Max(RoughSampleSize, (Samples[1] - Samples[0]).Length());
+			return Samples;
+		}
+
+		const UPCGSplineData* Spline = nullptr;
+		int InputIndex = INDEX_NONE;
+		FBox Bounds = FBox(EForceInit::ForceInit);
+		TArray<FBox> SegmentBounds;
+		mutable TMap<int, TArray<FVector>> RoughSampling;
+		int SegmentsNum;
+	};
+
+	struct IntersectionPoint
+	{
+		const UPCGSplineData* A = nullptr;
+		const UPCGSplineData* B = nullptr;
+
+		FVector AveragePosition;
+		FVector PositionA;
+		FVector PositionB;
+		float InputKeyA = 0.0f;
+		float InputKeyB = 0.0f;
+
+		bool IsOnSpline(const UPCGSplineData* In) const { return In == A || In == B; }
+		float GetKey(const UPCGSplineData* In) const { check(In == A || In == B); return In == A ? InputKeyA : InputKeyB; }
+		const UPCGSplineData* GetOtherSpline(const UPCGSplineData* In) const { check(In == A || In == B); return In == A ? B : A; }
+		const FVector& GetPosition(const UPCGSplineData* In) const { check(In == A || In == B); return In == A ? PositionA : PositionB; }
+		FVector& GetPosition(const UPCGSplineData* In) { check(In == A || In == B); return In == A ? PositionA : PositionB; }
+	};
+}
+
+UPCGSplineIntersectionSettings::UPCGSplineIntersectionSettings()
+{
+	//OriginatingSplineIndexAttribute.SetAttributeName(PCGSplineIntersection::Constants::OriginatingSplineIndexAttributeName);
+	//IntersectingSplineIndexAttribute.SetAttributeName(PCGSplineIntersection::Constants::IntersectingSplineIndexAttributeName);
+}
+
+#if WITH_EDITOR
+FName UPCGSplineIntersectionSettings::GetDefaultNodeName() const
+{
+	return FName(TEXT("SplineIntersection"));
+}
+
+FText UPCGSplineIntersectionSettings::GetDefaultNodeTitle() const
+{
+	return LOCTEXT("NodeTitle", "Spline Intersection");
+}
+
+EPCGChangeType UPCGSplineIntersectionSettings::GetChangeTypeForProperty(const FName& InPropertyName) const
+{
+	EPCGChangeType ChangeType = Super::GetChangeTypeForProperty(InPropertyName);
+
+	if (InPropertyName == GET_MEMBER_NAME_CHECKED(UPCGSplineIntersectionSettings, Output))
+	{
+		// This can change the output pin types, so we need a structural change
+		ChangeType |= EPCGChangeType::Structural;
+	}
+
+	return ChangeType;
+}
+#endif // WITH_EDITOR
+
+TArray<FPCGPinProperties> UPCGSplineIntersectionSettings::InputPinProperties() const
+{
+	TArray<FPCGPinProperties> PinProperties;
+	PinProperties.Emplace(PCGPinConstants::DefaultInputLabel, EPCGDataType::Spline);
+
+	return PinProperties;
+}
+
+TArray<FPCGPinProperties> UPCGSplineIntersectionSettings::OutputPinProperties() const
+{
+	TArray<FPCGPinProperties> PinProperties;
+	if (Output == EPCGSplineIntersectionOutput::IntersectionPointsOnly)
+	{
+		PinProperties.Emplace(PCGPinConstants::DefaultOutputLabel, EPCGDataType::Point);
+	}
+	else
+	{
+		PinProperties.Emplace(PCGPinConstants::DefaultOutputLabel, EPCGDataType::Spline);
+	}
+
+	return PinProperties;
+}
+
+FPCGElementPtr UPCGSplineIntersectionSettings::CreateElement() const
+{
+	return MakeShared<FPCGSplineIntersectionElement>();
+}
+
+bool FPCGSplineIntersectionElement::ExecuteInternal(FPCGContext* Context) const
+{
+	TRACE_CPUPROFILER_EVENT_SCOPE(FPCGSplineIntersectionElement::Execute);
+
+	const UPCGSplineIntersectionSettings* Settings = Context->GetInputSettings<UPCGSplineIntersectionSettings>();
+	check(Settings);
+
+	TArray<FPCGTaggedData> Inputs = Context->InputData.GetInputsByPin(PCGPinConstants::DefaultInputLabel);
+
+	TArray<PCGSplineIntersection::InputData> InputSplines;
+	const FVector HalfDistanceRadius(0.5 * Settings->DistanceThreshold);
+	const double SqrDistanceThreshold = FMath::Square(Settings->DistanceThreshold);
+
+	for (int InputIndex = 0; InputIndex < Inputs.Num(); ++InputIndex)
+	{
+		const FPCGTaggedData& Input = Inputs[InputIndex];
+		if (const UPCGSplineData* SplineData = Cast<UPCGSplineData>(Input.Data))
+		{
+			PCGSplineIntersection::InputData& SplineInputData = InputSplines.Add_GetRef(PCGSplineIntersection::InputData{ .Spline = SplineData, .InputIndex = InputIndex });
+			SplineInputData.InitializeBounds(HalfDistanceRadius);
+		}
+	}
+	
+	// Early out if we have no work to do.
+	if (InputSplines.IsEmpty() || Settings->DistanceThreshold <= 0)
+	{
+		return true;
+	}
+
+	// Compute intersections...
+	TArray<PCGSplineIntersection::IntersectionPoint> Intersections;
+
+	TArray<int> SplinesToTest;
+	for (int SplineIndex = 0; SplineIndex < InputSplines.Num(); ++SplineIndex)
+	{
+		SplinesToTest.Reset();
+
+		if (Settings->Type == EPCGSplineIntersectionType::Self)
+		{
+			SplinesToTest.Add(SplineIndex);
+		}
+		else
+		{
+			for (int OtherSplineIndex = SplineIndex + 1; OtherSplineIndex < InputSplines.Num(); ++OtherSplineIndex)
+			{
+				SplinesToTest.Add(OtherSplineIndex);
+			}
+		}
+
+		const PCGSplineIntersection::InputData& CurrentSplineData = InputSplines[SplineIndex];
+		const FBox& CurrentBounds = CurrentSplineData.Bounds;
+		
+		for (int SplineToTest : SplinesToTest)
+		{
+			const PCGSplineIntersection::InputData& OtherSplineData = InputSplines[SplineToTest];
+			const FBox& OtherBounds = OtherSplineData.Bounds;
+			const bool bIsSameData = (OtherSplineData.Spline == CurrentSplineData.Spline);
+
+			// Early test
+			if (!CurrentBounds.Intersect(OtherBounds))
+			{
+				continue;
+			}
+
+			// Test on a segment-segment basis
+			for (int32 SegmentIndex = 0; SegmentIndex < CurrentSplineData.SegmentBounds.Num(); ++SegmentIndex)
+			{
+				const FBox& CurrentSegmentBounds = CurrentSplineData.SegmentBounds[SegmentIndex];
+
+				// First early out if the segment does not intersect with the other spline bounds
+				if (!CurrentSegmentBounds.Intersect(OtherBounds))
+				{
+					continue;
+				}
+
+				double CurrentRoughStep = 0.0;
+				const TArray<FVector>& CurrentSamples = CurrentSplineData.GetRoughSamples(SegmentIndex, Settings->DistanceThreshold, CurrentRoughStep);
+
+				// Then iterate through potential segments
+				for (int32 OtherSegmentIndex = 0; OtherSegmentIndex < OtherSplineData.SegmentBounds.Num(); ++OtherSegmentIndex)
+				{
+					// Test against bounds first
+					if (!CurrentSegmentBounds.Intersect(OtherSplineData.SegmentBounds[OtherSegmentIndex]))
+					{
+						continue;
+					}
+
+					double OtherRoughStep = 0.0;
+					const TArray<FVector>& OtherSamples = OtherSplineData.GetRoughSamples(OtherSegmentIndex, Settings->DistanceThreshold, OtherRoughStep);
+					
+					//const double RoughToleranceSqr = 0.25 * (FMath::Square(CurrentRoughStep) + FMath::Square(OtherRoughStep));
+
+					// At this point, we are in intersection-possible territory.
+					// We don't have an analytical solution here, but there are a few things to consider:
+					// It's possible that two segments intersect in more than one place (up to 4 assumedly),
+					// but it's also possible that segments are collinear.
+					for (int CurrentSampleIndex = 0; CurrentSampleIndex < CurrentSamples.Num() - 1; ++CurrentSampleIndex)
+					{
+						const FVector& A = CurrentSamples[CurrentSampleIndex];
+						const FVector& B = CurrentSamples[CurrentSampleIndex + 1];
+
+						const int OtherStartSampleIndex = bIsSameData ? CurrentSampleIndex + 1 : 0;
+						for (int OtherSampleIndex = OtherStartSampleIndex; OtherSampleIndex < OtherSamples.Num() - 1; ++OtherSampleIndex)
+						{
+							const FVector& C = OtherSamples[OtherSampleIndex];
+							const FVector& D = OtherSamples[OtherSampleIndex + 1];
+
+							// This is an approximation because we linearized the spline segment.
+							// If this isn't enough, there are a few solutions here:
+							// - Make the discretization size configurable
+							// - Change the threshold here so we incorporate the approximate difference
+							// .. and eventually resplit said segment
+							FVector X, Y;
+							// @todo_pcg: might want to use the safe version if we have colocated control points
+							FMath::SegmentDistToSegment(A, B, C, D, X, Y);
+
+							if ((X - Y).SquaredLength() > SqrDistanceThreshold)
+							{
+								continue;
+							}
+
+							const float CurrentKey = CurrentSplineData.Spline->SplineStruct.FindInputKeyOnSegmentClosestToWorldLocation(X, SegmentIndex);
+							const FVector ClosestPointOnCurrentSpline = CurrentSplineData.Spline->SplineStruct.GetLocationAtSplineInputKey(CurrentKey, ESplineCoordinateSpace::World);
+
+							const float OtherKey = OtherSplineData.Spline->SplineStruct.FindInputKeyOnSegmentClosestToWorldLocation(Y, OtherSegmentIndex);
+							const FVector ClosestPointOnOtherSpline = OtherSplineData.Spline->SplineStruct.GetLocationAtSplineInputKey(OtherKey, ESplineCoordinateSpace::World);
+
+							// @todo_pcg: Validate distance after?
+							const FVector MidPoint = (ClosestPointOnCurrentSpline + ClosestPointOnOtherSpline) * 0.5;
+							PCGSplineIntersection::IntersectionPoint& Intersection = Intersections.Emplace_GetRef();
+							Intersection.A = CurrentSplineData.Spline;
+							Intersection.B = OtherSplineData.Spline;
+							Intersection.PositionA = ClosestPointOnCurrentSpline;
+							Intersection.PositionB = ClosestPointOnOtherSpline;
+							Intersection.AveragePosition = MidPoint;
+							Intersection.InputKeyA = CurrentKey;
+							Intersection.InputKeyB = OtherKey;
+						}
+					}
+				}
+			}
+		}
+	}
+
+	TArray<FPCGTaggedData>& Outputs = Context->OutputData.TaggedData;
+
+	// Finally, build output data.
+	if (Settings->Output == EPCGSplineIntersectionOutput::IntersectionPointsOnly)
+	{
+		// Filter out points that are too close to each other.
+		TArray<TPair<FVector, double>> WeightedPoints;
+		WeightedPoints.Reserve(Intersections.Num());
+
+		for (const PCGSplineIntersection::IntersectionPoint& Intersection : Intersections)
+		{
+			const FVector& Position = Intersection.AveragePosition;
+
+			int MatchingPoint = WeightedPoints.IndexOfByPredicate([&Position, &SqrDistanceThreshold](const TPair<FVector, double>& WP)
+			{
+				return (WP.Key - Position).SquaredLength() <= SqrDistanceThreshold;
+			});
+
+			if (MatchingPoint == INDEX_NONE)
+			{
+				WeightedPoints.Emplace(Position, 1.0);
+			}
+			else
+			{
+				TPair<FVector, double>& Match = WeightedPoints[MatchingPoint];
+				Match.Key = (Match.Value * Match.Key + Position) / (Match.Value + 1);
+				Match.Value += 1.0;
+			}
+		}
+
+		FPCGTaggedData& Output = Outputs.Emplace_GetRef();
+		
+		UPCGBasePointData* PointData = FPCGContext::NewPointData_AnyThread(Context);
+		Output.Data = PointData;
+
+		PointData->SetNumPoints(WeightedPoints.Num());
+		PointData->AllocateProperties(EPCGPointNativeProperties::Transform);
+
+		FPCGPointTransform::ValueRange Transforms = PointData->GetTransformValueRange();
+		for (int Index = 0; Index < WeightedPoints.Num(); ++Index)
+		{
+			Transforms[Index] = FTransform(WeightedPoints[Index].Key);
+		}
+	}
+	else if (Settings->Output == EPCGSplineIntersectionOutput::OriginalSplinesWithIntersections)
+	{
+		for (const PCGSplineIntersection::InputData& InputData : InputSplines)
+		{
+			const UPCGSplineData* InputSpline = InputData.Spline;
+
+			TArray<TPair<float, PCGSplineIntersection::IntersectionPoint>> FilteredPoints;
+			for (const PCGSplineIntersection::IntersectionPoint& Intersection : Intersections)
+			{
+				if (Intersection.IsOnSpline(InputSpline))
+				{
+					FilteredPoints.Emplace(Intersection.GetKey(InputSpline), Intersection);
+				}
+			}
+
+			// Sort by key
+			FilteredPoints.Sort([](const auto& A, const auto& B) { return A.Key < B.Key; });
+
+			// Collapse points that would be too close
+			TArray<bool> PointCollapsed;
+			PointCollapsed.SetNumZeroed(FilteredPoints.Num()); // 0 == false
+
+			TArray<TPair<float, PCGSplineIntersection::IntersectionPoint>> CollapsedPoints;
+			CollapsedPoints.Reserve(FilteredPoints.Num());
+
+			for (int FP = 0; FP < FilteredPoints.Num(); ++FP)
+			{
+				if (PointCollapsed[FP])
+				{
+					continue;
+				}
+
+				TPair<float, PCGSplineIntersection::IntersectionPoint>& PointAndKey = CollapsedPoints.Add_GetRef(FilteredPoints[FP]);
+				PCGSplineIntersection::IntersectionPoint& Point = PointAndKey.Value;
+				FVector& PointPosition = Point.GetPosition(InputSpline);
+
+				double Weight = 1.0;
+				for (int OFP = FP + 1; OFP < FilteredPoints.Num(); ++OFP)
+				{
+					if (PointCollapsed[OFP] || Point.GetOtherSpline(InputSpline) != FilteredPoints[OFP].Value.GetOtherSpline(InputSpline))
+					{
+						continue;
+					}
+
+					const PCGSplineIntersection::IntersectionPoint& OtherPoint = FilteredPoints[OFP].Value;
+					const FVector& OtherPointPosition = OtherPoint.GetPosition(InputSpline);
+
+					if ((PointPosition - OtherPointPosition).SquaredLength() > SqrDistanceThreshold)
+					{
+						break;
+					}
+
+					// Merge points
+					PointCollapsed[OFP] = true;
+
+					PointPosition = (PointPosition * Weight + OtherPointPosition) / (Weight + 1.0);
+					Weight += 1.0;
+				}
+			}
+
+			// @todo_pcg : we should not add points that are already on existing control points,
+			// but we should mark them as intersections
+
+			// At this point, we've successfully filtered/collapsed points if it was needed,
+			// Or we had no new points - in which case we have nothing to do.
+			FPCGTaggedData& Output = Outputs.Add_GetRef(Inputs[InputData.InputIndex]);
+			if (!CollapsedPoints.IsEmpty())
+			{
+				UPCGSplineData* NewSplineData = FPCGContext::NewObject_AnyThread<UPCGSplineData>(Context);
+
+				// Implementation note: with the old spline representation, it's not possible to add control points in between other points, 
+				// because there are some places in the code where we rely on the index being equal to the key (esp. around the ReparamTable).
+				// Not respecting this makes some operations fail, so we need to conform to that expectation at this point in time.
+				// However, this means we basically need to have custom tangents everywhere.
+				TArray<FSplinePoint> SplinePoints = InputData.Spline->GetSplinePoints();
+				TArray<int64> SplinePointsEntryKeys = InputData.Spline->GetMetadataEntryKeysForSplinePoints();
+				check(SplinePoints.Num() == SplinePointsEntryKeys.Num() || SplinePointsEntryKeys.IsEmpty());
+				const bool bHasMetadata = !SplinePointsEntryKeys.IsEmpty();
+
+				const float MinKeyDifferenceThreshold = 1.0e-4f;
+				
+				for (const TPair<float, PCGSplineIntersection::IntersectionPoint>& PointToAdd : CollapsedPoints)
+				{
+					const float& Key = PointToAdd.Key;
+					const FVector& Position = InputData.Spline->SplineStruct.GetTransform().InverseTransformPosition(PointToAdd.Value.GetPosition(InputSpline));
+
+					// Find insertion index
+					// @todo_pcg: use bisection
+					int InsertionIndex = SplinePoints.Num();
+					for (int SPIndex = 0; SPIndex < SplinePoints.Num(); ++SPIndex)
+					{
+						if (SplinePoints[SPIndex].InputKey > Key)
+						{
+							InsertionIndex = SPIndex;
+							break;
+						}
+					}
+
+					bool bIsOnLastSegment = (InsertionIndex == SplinePoints.Num());
+
+					// Skip points that are on control points - need to check against previous and next.
+					FSplinePoint& PreviousPoint = SplinePoints[InsertionIndex - 1];
+					FSplinePoint& NextPoint = SplinePoints[bIsOnLastSegment ? 0 : InsertionIndex];
+
+					const float NextPointVirtualKey = (bIsOnLastSegment ? FMath::CeilToInt(PreviousPoint.InputKey + 1.0e-4f) : NextPoint.InputKey);
+
+					if((FMath::Abs(PreviousPoint.InputKey - Key) <= MinKeyDifferenceThreshold && (PreviousPoint.Position - Position).SquaredLength() <= SqrDistanceThreshold) ||
+						((FMath::Abs(NextPointVirtualKey - Key) <= MinKeyDifferenceThreshold && (NextPoint.Position - Position).SquaredLength() <= SqrDistanceThreshold)))
+					{
+						continue;
+					}
+
+					FSplinePoint SplinePoint;
+					SplinePoint.InputKey = Key;
+					SplinePoint.Position = Position;
+					// @todo_pcg: we don't need to add points that are on the control point, however, we might want to do metadata for those
+					SplinePoint.ArriveTangent = InputData.Spline->SplineStruct.GetTangentAtSplineInputKey(Key, ESplineCoordinateSpace::Local);
+					SplinePoint.LeaveTangent = SplinePoint.ArriveTangent;
+					SplinePoint.Rotation = NewSplineData->SplineStruct.GetQuaternionAtSplineInputKey(Key, ESplineCoordinateSpace::Local).Rotator();
+					SplinePoint.Scale = NewSplineData->SplineStruct.GetScaleAtSplineInputKey(Key);
+					SplinePoint.Type = ESplinePointType::CurveCustomTangent;
+
+					// Adjust tangents
+					const double KeyRange = FMath::Max(NextPointVirtualKey - PreviousPoint.InputKey, 1.0e-4);
+					const double PreviousRatio = (Key - PreviousPoint.InputKey) / KeyRange;
+					const double NextRatio = (NextPointVirtualKey - Key) / KeyRange;
+
+					PreviousPoint.LeaveTangent *= PreviousRatio;
+					NextPoint.ArriveTangent *= NextRatio;
+
+					// The spline point tangent is "always" taken from the original spline's perspective,
+					// so we will assume that the range is always one, so we need to scale accordingly.
+					// @todo_pcg: get key at the beginning of the segment
+					SplinePoint.ArriveTangent *= (Key - PreviousPoint.InputKey);
+					SplinePoint.LeaveTangent *= (NextPointVirtualKey - Key);
+
+					SplinePoints.Insert(SplinePoint, InsertionIndex);
+					if (bHasMetadata)
+					{
+						SplinePointsEntryKeys.Insert(PCGInvalidEntryKey, InsertionIndex);
+					}
+				}
+
+				// Mark all tangents as custom, and reset the keys
+				TArray<float> OriginalKeys;
+				OriginalKeys.SetNumUninitialized(SplinePoints.Num());
+
+				for(int SPIndex = 0; SPIndex < SplinePoints.Num(); ++SPIndex)
+				{
+					FSplinePoint& SplinePoint = SplinePoints[SPIndex];
+
+					OriginalKeys[SPIndex] = SplinePoint.InputKey;
+					SplinePoint.InputKey = static_cast<float>(SPIndex);
+					if (SplinePoint.Type == ESplinePointType::Curve || SplinePoint.Type == ESplinePointType::CurveClamped)
+					{
+						SplinePoints[SPIndex].Type = ESplinePointType::CurveCustomTangent;
+					}
+				}
+
+				// Initialize data now
+				NewSplineData->Initialize(SplinePoints,
+					InputData.Spline->SplineStruct.IsClosedLoop(),
+					InputData.Spline->SplineStruct.GetTransform(),
+					SplinePointsEntryKeys);
+				NewSplineData->InitializeFromData(InputData.Spline);
+
+				// Compute metadata for added points
+				if (bHasMetadata)
+				{
+					TArrayView<PCGMetadataEntryKey> SplineMEK = NewSplineData->SplineStruct.GetMutableControlPointsEntryKeys();
+					for (int SPIndex = 0; SPIndex < SplinePoints.Num(); ++SPIndex)
+					{
+						FSplinePoint& SplinePoint = SplinePoints[SPIndex];
+						if (SplineMEK[SPIndex] == PCGInvalidEntryKey)
+						{
+							InputData.Spline->WriteMetadataToEntry(OriginalKeys[SPIndex], SplineMEK[SPIndex], NewSplineData->MutableMetadata());
+						}
+					}
+				}
+
+				Output.Data = NewSplineData;
+			}
+		}
+	}
+	else if(Settings->Output == EPCGSplineIntersectionOutput::SplitSplines)
+	{
+		// @todo_pcg
+	}
+
+	return true;
+}
+
+#undef LOCTEXT_NAMESPACE
\ No newline at end of file
diff --git a/Engine/Plugins/PCG/Source/PCG/Public/Data/PCGSplineData.h b/Engine/Plugins/PCG/Source/PCG/Public/Data/PCGSplineData.h
--- a/Engine/Plugins/PCG/Source/PCG/Public/Data/PCGSplineData.h
+++ b/Engine/Plugins/PCG/Source/PCG/Public/Data/PCGSplineData.h
@@ -96,7 +96,6 @@ public:
 	/** Helper to create an accessor keys for data that doesn't yet exist. */
 	UE_API TUniquePtr<IPCGAttributeAccessorKeys> CreateAccessorKeys(const FPCGAttributePropertySelector& InSelector, bool bQuiet = false);
 	UE_API TUniquePtr<const IPCGAttributeAccessorKeys> CreateConstAccessorKeys(const FPCGAttributePropertySelector& InSelector, bool bQuiet = false) const;
-
 	
 	// Get the functions to the accessor factory
 	static UE_API FPCGAttributeAccessorMethods GetSplineAccessorMethods();
@@ -118,13 +117,15 @@ protected:
 
 	UE_API const UPCGBasePointData* CreateBasePointData(FPCGContext* Context, TSubclassOf<UPCGBasePointData> PointDataClass) const;
 
+public:
 	UFUNCTION(BlueprintCallable, Category = "SplineData")
 	UE_API TArray<FSplinePoint> GetSplinePoints() const;
 
 	UFUNCTION(BlueprintCallable, Category = "SplineData")
 	UE_API TArray<int64> GetMetadataEntryKeysForSplinePoints() const;
 
-public:
+	UE_API void WriteMetadataToEntry(float InputKey, PCGMetadataEntryKey& OutEntryKey, UPCGMetadata* OutMetadata) const;
+
 	// Minimal data needed to replicate the behavior from USplineComponent
 	UPROPERTY()
 	FPCGSplineStruct SplineStruct;
diff --git a/Engine/Plugins/PCG/Source/PCG/Public/Data/PCGSplineStruct.h b/Engine/Plugins/PCG/Source/PCG/Public/Data/PCGSplineStruct.h
--- a/Engine/Plugins/PCG/Source/PCG/Public/Data/PCGSplineStruct.h
+++ b/Engine/Plugins/PCG/Source/PCG/Public/Data/PCGSplineStruct.h
@@ -65,8 +65,8 @@ struct FPCGSplineStruct
 	FTransform GetTransform() const { return Transform; }
 
 	// Spline-related methods
-	UE_API void AddPoint(const FSplinePoint& InSplinePoint, bool bUpdateSpline); 
-	UE_API void AddPoints(const TArray<FSplinePoint>& InSplinePoints, bool bUpdateSpline);
+	UE_API void AddPoint(const FSplinePoint& InSplinePoint, bool bUpdateSpline, const int64* InOptionalEntryKey = nullptr); 
+	UE_API void AddPoints(const TArray<FSplinePoint>& InSplinePoints, bool bUpdateSpline, const TArray<int64>* InOptionalEntryKeys = nullptr);
 
 	UE_API void SetLocation(int32 Index, const FVector& InLocation);
 	UE_API FVector GetLocation(const int32 Index) const;
@@ -99,6 +99,7 @@ struct FPCGSplineStruct
 	bool IsClosedLoop() const { return bClosedLoop; }
 	UE_API FVector::FReal GetSplineLength() const;
 	UE_API FBox GetBounds() const;
+	UE_API FBox GetSegmentBounds(int32 Index) const;
 
 	UE_API const FInterpCurveVector& GetSplinePointsScale() const;
 	UE_API const FInterpCurveQuat& GetSplinePointsRotation() const;
@@ -109,6 +110,7 @@ struct FPCGSplineStruct
 	UE_API FVector GetLocationAtSplineInputKey(float InKey, ESplineCoordinateSpace::Type CoordinateSpace) const;
 	UE_API FQuat GetQuaternionAtSplineInputKey(float InKey, ESplineCoordinateSpace::Type CoordinateSpace) const;
 	UE_API FVector GetScaleAtSplineInputKey(float InKey) const;
+	UE_API FVector GetTangentAtSplineInputKey(float InKey, ESplineCoordinateSpace::Type CoordinateSpace) const;
 
 	UE_API FVector::FReal GetDistanceAlongSplineAtSplinePoint(int32 PointIndex) const;
 	UE_API FVector GetLocationAtDistanceAlongSpline(FVector::FReal Distance, ESplineCoordinateSpace::Type CoordinateSpace) const;
@@ -122,6 +124,7 @@ struct FPCGSplineStruct
 	UE_API bool ConvertSplineToPolyLine(ESplineCoordinateSpace::Type CoordinateSpace, const float MaxSquareDistanceFromSpline, TArray<FVector>& OutPoints) const;
 	
 	UE_API float FindInputKeyClosestToWorldLocation(const FVector& WorldLocation) const;
+	UE_API float FindInputKeyOnSegmentClosestToWorldLocation(const FVector& WorldLocation, int32 Index) const;
 
 	/** Return the index of the segment for a given InputKey, and the input key at the segment start.
 	 * It's for future-proof-ness as even if there is no guarantee a segment start is at an integer input key, right now it is not well supported.
diff --git a/Engine/Plugins/PCG/Source/PCG/Public/Elements/PCGSplineIntersection.h b/Engine/Plugins/PCG/Source/PCG/Public/Elements/PCGSplineIntersection.h
new file mode 100644
--- /dev/null
+++ b/Engine/Plugins/PCG/Source/PCG/Public/Elements/PCGSplineIntersection.h
@@ -0,0 +1,88 @@
+// Copyright Epic Games, Inc. All Rights Reserved.
+
+#pragma once
+
+#include "PCGSettings.h"
+
+#include "PCGSplineIntersection.generated.h"
+
+UENUM()
+enum class EPCGSplineIntersectionType : uint8
+{
+	/** Performs self-intersection only on a per-spline basis. */
+	Self UMETA(Hidden),
+
+	/** Performs intersection test for each spline against each other spline. */
+	AgainstOtherSplines
+};
+
+UENUM()
+enum class EPCGSplineIntersectionOutput : uint8
+{
+	/** Returns only the intersection positions. */
+	IntersectionPointsOnly,
+
+	/** Adds control points at the intersections on the input splines. */
+	OriginalSplinesWithIntersections,
+
+	/** Splits splines at their intersections. */
+	SplitSplines UMETA(Hidden)
+};
+
+/** 
+* Intersects splines against other splines (or themselves) and returns varied results
+* based on user need.
+*/
+UCLASS(MinimalAPI, BlueprintType, ClassGroup=(Procedural))
+class UPCGSplineIntersectionSettings : public UPCGSettings
+{
+	GENERATED_BODY()
+
+public:
+	UPCGSplineIntersectionSettings();
+	
+	//~Begin UPCGSettings interface
+#if WITH_EDITOR
+	virtual FName GetDefaultNodeName() const override;
+	virtual FText GetDefaultNodeTitle() const override;
+	virtual EPCGSettingsType GetType() const override { return EPCGSettingsType::Spatial; }
+	virtual EPCGChangeType GetChangeTypeForProperty(const FName& InPropertyName) const override;
+#endif
+
+protected:
+	virtual TArray<FPCGPinProperties> InputPinProperties() const override;
+	virtual TArray<FPCGPinProperties> OutputPinProperties() const override;
+	virtual FPCGElementPtr CreateElement() const override;
+	//~End UPCGSettings interface
+
+public:
+	/** Controls type of intersection test that will be done. */
+	UPROPERTY(BlueprintReadWrite, EditAnywhere, Category = Settings)
+	EPCGSplineIntersectionType Type = EPCGSplineIntersectionType::AgainstOtherSplines;
+
+	/** Controls type of output result. */
+	UPROPERTY(BlueprintReadWrite, EditAnywhere, Category = Settings)
+	EPCGSplineIntersectionOutput Output = EPCGSplineIntersectionOutput::OriginalSplinesWithIntersections;
+
+	/** Controls whether the output will contain the index of the intersecting splines */
+	//UPROPERTY(BlueprintReadWrite, EditAnywhere, Category = Settings, meta=(EditCondition="Type!=EPCGSplineIntersectionType::Self"))
+	//bool bOutputSplineIndices = false;
+
+	/** Attribute that will contain the first spline index (only in the output case of intersection points.) */
+	//UPROPERTY(BlueprintReadWrite, EditAnywhere, Category = Settings, meta=(PCG_Overridable, EditCondition="bOutputSplineIndices && Type!=EPCGSplineIntersetionType::Self && Output==EPCGSplineIntersectionOutput::IntersectionPointsOnly"))
+	//FPCGAttributePropertyOutputSelector OriginatingSplineIndexAttribute;
+
+	/** Attribute that will contain the first intersecting spline index, or -1 if not an intersection control point. */
+	//UPROPERTY(BlueprintReadWrite, EditAnywhere, Category = Settings, meta=(PCG_Overridable, EditCondition="bOutputSplineIndices"))
+	//FPCGAttributePropertyOutputSelector IntersectingSplineIndexAttribute;
+
+	/** Maximum distance at which we report an intersection on the splines. */
+	UPROPERTY(BlueprintReadWrite, EditAnywhere, Category = Settings, meta=(ClampMin="0.0001", PCG_Overridable))
+	double DistanceThreshold = 10.0;
+};
+
+class FPCGSplineIntersectionElement : public IPCGElement
+{
+public:
+	virtual bool ExecuteInternal(FPCGContext* InContext) const override;
+};
\ No newline at end of file
