[PCG] Fix unwanted creation of actor packages in PIE with PCG Preview components

diff --git a/Engine/Plugins/PCG/Source/PCG/Private/Helpers/PCGActorHelpers.cpp b/Engine/Plugins/PCG/Source/PCG/Private/Helpers/PCGActorHelpers.cpp
--- a/Engine/Plugins/PCG/Source/PCG/Private/Helpers/PCGActorHelpers.cpp
+++ b/Engine/Plugins/PCG/Source/PCG/Private/Helpers/PCGActorHelpers.cpp
@@ -651,7 +651,7 @@ AActor* UPCGActorHelpers::SpawnDefaultActor(const FSpawnDefaultActorParams& Para
 
 	// If we are spawning a preview actor in a world partition world, we need to assign it an external package as Data Layers are not supported on non OFPA actors.
 	// Note: this could at somepoint become an option in the FActorSpawnParameters to allow RF_Transient actors to create their own package in UWorld::SpawnActor.
-	if (Params.bIsPreviewActor && SpawnParams.OverrideLevel->GetWorldPartition())
+	if (Params.bIsPreviewActor && SpawnParams.OverrideLevel->GetWorldPartition() && !Params.World->IsGameWorld())
 	{
 		// Remove trailing _C from BP generated class
 		FString ActorClassName = Params.ActorClass->GetName();
