[PCG] Fix rare occurence where LandscapeCache can be nullptr

diff --git a/Engine/Plugins/PCG/Source/PCG/Private/Data/PCGLandscapeData.cpp b/Engine/Plugins/PCG/Source/PCG/Private/Data/PCGLandscapeData.cpp
--- a/Engine/Plugins/PCG/Source/PCG/Private/Data/PCGLandscapeData.cpp
+++ b/Engine/Plugins/PCG/Source/PCG/Private/Data/PCGLandscapeData.cpp
@@ -633,6 +633,11 @@ TArray<FPCGTaskId> UPCGLandscapeData::PrepareForSpatialQuery(FPCGContext* InCont
 		return {};
 	}
 
+	if (!LandscapeCache)
+	{
+		return {};
+	}
+
 	if (LandscapeCache->SerializationMode == EPCGLandscapeCacheSerializationMode::NeverSerialize && PCGHelpers::IsRuntimeOrPIE())
 	{
 		// Hotfix note: commented this error has it could do false positives, esp. in the GPU use case.
