[PCG] Fixed bug in previous hotfix, causing the index to be off in Visualize Attribute

#rnx

diff --git a/Engine/Plugins/PCG/Source/PCG/Private/Elements/PCGVisualizeAttribute.cpp b/Engine/Plugins/PCG/Source/PCG/Private/Elements/PCGVisualizeAttribute.cpp
--- a/Engine/Plugins/PCG/Source/PCG/Private/Elements/PCGVisualizeAttribute.cpp
+++ b/Engine/Plugins/PCG/Source/PCG/Private/Elements/PCGVisualizeAttribute.cpp
@@ -75,6 +75,11 @@ bool FPCGVisualizeAttribute::ExecuteInternal(FPCGContext* Context) const
 
 		TArray<FString> ValuesToString;
 		const int32 NumValues = Keys->GetNum();
+		if (NumValues == 0)
+		{
+			continue;
+		}
+
 		ValuesToString.SetNum(NumValues);
 
 		// Get all the attribute values at once in range, for efficiency
@@ -135,7 +140,7 @@ bool FPCGVisualizeAttribute::ExecuteInternal(FPCGContext* Context) const
 				CompoundedString += FString::Printf(TEXT("[%s]"), *InputSource.GetDisplayText().ToString());
 			}
 
-			CompoundedString = FString::Printf(TEXT("%s %s"), *CompoundedString, *ValuesToString[(I + 1) % NumValues]);
+			CompoundedString = FString::Printf(TEXT("%s %s"), *CompoundedString, *ValuesToString[I % NumValues]);
 
 			DebugStrings.Emplace(CompoundedString.ToString(), TransformRanges[I].GetLocation() + Settings->LocalOffset, Settings->Color);
 		}
