[PCG] Fix Debug Object Tree not getting populated when editing a graph that is dynamically invoked.

[RN] Fixed bug where Debug Object Tree did not populate when editing a dynamic subgraph.

diff --git a/Engine/Plugins/PCG/Source/PCGEditor/Private/PCGEditor.cpp b/Engine/Plugins/PCG/Source/PCGEditor/Private/PCGEditor.cpp
--- a/Engine/Plugins/PCG/Source/PCGEditor/Private/PCGEditor.cpp
+++ b/Engine/Plugins/PCG/Source/PCGEditor/Private/PCGEditor.cpp
@@ -3103,7 +3103,7 @@ void FPCGEditor::OnSourceGenerationDone(IPCGBaseSubsystem* Subsystem, IPCGGraphE
 
 	// Additionally, if we are not inspecting but the component that's done executing contains this graph, then we should also update.
 	IPCGEditorModule* PCGEditorModule = IPCGEditorModule::Get();
-	if (!bShouldRefresh && PCGSourceBeingInspected.IsValid() && PCGGraphBeingEdited && PCGEditorModule)
+	if (!bShouldRefresh && !PCGSourceBeingInspected.IsValid() && PCGGraphBeingEdited && PCGEditorModule)
 	{
 		TArray<FPCGStackSharedPtr> ExecutedStacksContainingThisGraph = PCGEditorModule->GetExecutedStacksPtrs(Source, PCGGraphBeingEdited);
 		bShouldRefresh |= !ExecutedStacksContainingThisGraph.IsEmpty();
