[PCG] [Experimental] PCGWorldPartitionBuilder : will now revert unchanged actors using resource/output crc. (greatly reduce actor checkins)
- There is still some determinism issue that can happen with some graphs but still in those case we see a > 90% reduction in changes from consecutive builds with no or minimal user changes.
- PCGBuilder Settings now has 2 new parameters 'Revert unchanged actors' and 'Outputs to ignore for revert' to enable this feature.
- New CVar must also be turned on: pcg.FullDataCrcInCommandlets which allows consistency of crcs between builder runs
- Fix input pin hashing n GraphExecutor which was using a fname memory hash (was changing between runs)
- DispatchToRegisteredLocalComponents will now sort the PartitionActors and chain them together so that they do not run at the same time for determinism reasons when new builder settings option is enabled

diff --git a/Engine/Plugins/PCG/Source/PCG/Private/Graph/PCGGraphExecutor.cpp b/Engine/Plugins/PCG/Source/PCG/Private/Graph/PCGGraphExecutor.cpp
--- a/Engine/Plugins/PCG/Source/PCG/Private/Graph/PCGGraphExecutor.cpp
+++ b/Engine/Plugins/PCG/Source/PCG/Private/Graph/PCGGraphExecutor.cpp
@@ -24,6 +24,7 @@
 #include "Algo/ForEach.h"
 #include "Engine/Engine.h"
 #include "GameFramework/Actor.h"
+#include "Misc/Crc.h"
 #include "Misc/ScopeExit.h"
 #include "Tasks/Task.h"
 
@@ -2244,7 +2245,7 @@ void FPCGGraphExecutor::BuildTaskInput(FPCGGraphTask& Task, TArray<FPCGTaskId>&
 			}
 
 			const bool bAllowMultipleData = Input.DownstreamPin.IsSet() ? Input.DownstreamPin.GetValue().bAllowMultipleData : true;
-			const uint32 InputPinLabelCrc = Input.DownstreamPin.IsSet() ? GetTypeHash(Input.DownstreamPin.GetValue().Label) : DefaultHashForNoOutputPin;
+			const uint32 InputPinLabelCrc = Input.DownstreamPin.IsSet() ? FCrc::StrCrc32(*Input.DownstreamPin.GetValue().Label.ToString()) : DefaultHashForNoOutputPin;
 
 			// Enforce single data - if already have input for this pin, don't add more. Early check before other side effects below.
 			if (Input.DownstreamPin.IsSet() && !bAllowMultipleData && Task.TaskInput.GetInputCountByPin(Input.DownstreamPin.GetValue().Label) > 0)
diff --git a/Engine/Plugins/PCG/Source/PCG/Private/PCGActorAndComponentMapping.cpp b/Engine/Plugins/PCG/Source/PCG/Private/PCGActorAndComponentMapping.cpp
--- a/Engine/Plugins/PCG/Source/PCG/Private/PCGActorAndComponentMapping.cpp
+++ b/Engine/Plugins/PCG/Source/PCG/Private/PCGActorAndComponentMapping.cpp
@@ -234,7 +234,7 @@ UWorld* FPCGActorAndComponentMapping::GetWorld() const
 	return PCGSubsystem->GetWorld();
 }
 
-TArray<FPCGTaskId> FPCGActorAndComponentMapping::DispatchToRegisteredLocalComponents(UPCGComponent* OriginalComponent, const TFunctionRef<FPCGTaskId(UPCGComponent*)>& InFunc) const
+TArray<FPCGTaskId> FPCGActorAndComponentMapping::DispatchToRegisteredLocalComponents(UPCGComponent* OriginalComponent, const TFunctionRef<FPCGTaskId(UPCGComponent*, const TArray<FPCGTaskId>&)>& InFunc) const
 {
 	TRACE_CPUPROFILER_EVENT_SCOPE(FPCGActorAndComponentMapping::DispatchToRegisteredLocalComponents);
 	if (!ensure(OriginalComponent))
@@ -256,10 +256,11 @@ TArray<FPCGTaskId> FPCGActorAndComponentMapping::DispatchToRegisteredLocalCompon
 	return DispatchToLocalComponents(OriginalComponent, *PartitionActorsPtr, InFunc);
 }
 
-TArray<FPCGTaskId> FPCGActorAndComponentMapping::DispatchToLocalComponents(UPCGComponent* OriginalComponent, const TSet<TObjectPtr<APCGPartitionActor>>& PartitionActors, const TFunctionRef<FPCGTaskId(UPCGComponent*)>& InFunc) const
+TArray<FPCGTaskId> FPCGActorAndComponentMapping::DispatchToLocalComponents(UPCGComponent* OriginalComponent, const TSet<TObjectPtr<APCGPartitionActor>>& PartitionActors, const TFunctionRef<FPCGTaskId(UPCGComponent*, const TArray<FPCGTaskId>&)>& InFunc) const
 {
 	TArray<FPCGTaskId> TaskIds;
-	for (APCGPartitionActor* PartitionActor : PartitionActors)
+
+	auto DispatchLambda = [this, OriginalComponent, &InFunc, &TaskIds](const TObjectPtr<APCGPartitionActor>& PartitionActor)
 	{
 		if (PartitionActor)
 		{
@@ -268,7 +269,17 @@ TArray<FPCGTaskId> FPCGActorAndComponentMapping::DispatchToLocalComponents(UPCGC
 				// Add check to avoid infinite loop
 				if (ensure(!LocalComponent->IsPartitioned()))
 				{
-					FPCGTaskId LocalTask = InFunc(LocalComponent);
+					FPCGTaskId LocalTask = InvalidPCGTaskId;
+#if WITH_EDITOR
+					if (bChainedDispatchToLocalComponents && TaskIds.Num() > 0)
+					{
+						LocalTask = InFunc(LocalComponent, { TaskIds.Last() });
+					}
+					else
+#endif // WITH_EDITOR
+					{
+						LocalTask = InFunc(LocalComponent, { });
+					}
 
 					if (LocalTask != InvalidPCGTaskId)
 					{
@@ -277,6 +288,23 @@ TArray<FPCGTaskId> FPCGActorAndComponentMapping::DispatchToLocalComponents(UPCGC
 				}
 			}
 		}
+	};
+
+#if WITH_EDITOR
+	if (bChainedDispatchToLocalComponents)
+	{
+		TArray<TObjectPtr<APCGPartitionActor>> SortedPartitionActors = PartitionActors.Array();
+		SortedPartitionActors.Remove(nullptr);
+		SortedPartitionActors.Sort([](const APCGPartitionActor& ActorA, const APCGPartitionActor& ActorB)
+		{
+			return ActorA.GetPathName() < ActorB.GetPathName();
+		});
+		Algo::ForEach(SortedPartitionActors, DispatchLambda);
+	}
+	else
+#endif // WITH_EDITOR
+	{
+		Algo::ForEach(PartitionActors, DispatchLambda);
 	}
 
 	return TaskIds;
diff --git a/Engine/Plugins/PCG/Source/PCG/Private/PCGActorAndComponentMapping.h b/Engine/Plugins/PCG/Source/PCG/Private/PCGActorAndComponentMapping.h
--- a/Engine/Plugins/PCG/Source/PCG/Private/PCGActorAndComponentMapping.h
+++ b/Engine/Plugins/PCG/Source/PCG/Private/PCGActorAndComponentMapping.h
@@ -144,10 +144,10 @@ private:
 	void UnregisterNonPartitionedPCGComponent(UPCGComponent* InComponent);
 
 	/* Call the InFunc function to all local component registered to the original component. Return the list of all the tasks scheduled. Thread safe*/
-	TArray<FPCGTaskId> DispatchToRegisteredLocalComponents(UPCGComponent* OriginalComponent, const TFunctionRef<FPCGTaskId(UPCGComponent*)>& InFunc) const;
+	TArray<FPCGTaskId> DispatchToRegisteredLocalComponents(UPCGComponent* OriginalComponent, const TFunctionRef<FPCGTaskId(UPCGComponent*, const TArray<FPCGTaskId>&)>& InFunc) const;
 
 	/* Call the InFunc function to all local component from the set of partition actors. Return the list of all the tasks scheduled. */
-	TArray<FPCGTaskId> DispatchToLocalComponents(UPCGComponent* OriginalComponent, const TSet<TObjectPtr<APCGPartitionActor>>& PartitionActors, const TFunctionRef<FPCGTaskId(UPCGComponent*)>& InFunc) const;
+	TArray<FPCGTaskId> DispatchToLocalComponents(UPCGComponent* OriginalComponent, const TSet<TObjectPtr<APCGPartitionActor>>& PartitionActors, const TFunctionRef<FPCGTaskId(UPCGComponent*, const TArray<FPCGTaskId>&)>& InFunc) const;
 
 	/** Call the InFunc function to all partitioned components which bounds intersect 'InBounds'. */
 	void ForAllIntersectingPartitionedComponents(const FBoxCenterAndExtent& InBounds, TFunctionRef<void(UPCGComponent*)> InFunc) const;
@@ -229,6 +229,9 @@ private:
 
 	/** Build initial records for existing partition actors */
 	void BuildPartitionActorRecords();
+
+	/** Used by the PCG Builder to force dispatched local components to be dependency chained together so that they do not run at the same time for determinism */
+	void SetChainedDispatchToLocalComponents(bool bInChainedDispatch) { bChainedDispatchToLocalComponents = bInChainedDispatch; }
 #endif // WITH_EDITOR
 
 private:
@@ -291,5 +294,8 @@ private:
 	TMap<FPCGGridCellDescriptor, FGuid> PartitionActorRecords;
 	// Previously generated PCG Partition Actors to ignore (will prevent them from getting registered and mark them for deletion)
 	TSet<FGuid> InvalidPartitionActors;
+
+	// If dispatch needs to run the local comopnents one after the other through dependency chain for determinism
+	bool bChainedDispatchToLocalComponents = false;
 #endif // WITH_EDITOR
 };
\ No newline at end of file
diff --git a/Engine/Plugins/PCG/Source/PCG/Private/PCGData.cpp b/Engine/Plugins/PCG/Source/PCG/Private/PCGData.cpp
--- a/Engine/Plugins/PCG/Source/PCG/Private/PCGData.cpp
+++ b/Engine/Plugins/PCG/Source/PCG/Private/PCGData.cpp
@@ -20,6 +20,11 @@
 
 //#define PCG_DATA_USAGE_LOGGING
 
+static TAutoConsoleVariable<bool> CVarPCGFullDataCrcInCommandlets(
+	TEXT("pcg.FullDataCrcInCommandlets"),
+	false,
+	TEXT("Always compute full data Crcs in Commandlets (to improve reusability)."));
+
 UPCGData::UPCGData(const FObjectInitializer& ObjectInitializer)
 	: Super(ObjectInitializer)
 {
@@ -31,8 +36,12 @@ UPCGData::UPCGData(const FObjectInitializer& ObjectInitializer)
 
 FPCGCrc UPCGData::GetOrComputeCrc(bool bFullDataCrc) const
 {
+#if WITH_EDITOR
+	const bool bComputeFullDataCrc = (bFullDataCrc || (CVarPCGFullDataCrcInCommandlets.GetValueOnAnyThread() && IsRunningCommandlet())) && SupportsFullDataCrc();
+#else
 	const bool bComputeFullDataCrc = bFullDataCrc && SupportsFullDataCrc();
-	
+#endif
+
 	// If we already valid Crc and it is a full data crc or we aren't being requested one, return the current Crc
 	if (Crc.IsValid() && (bIsFullDataCrc || !bComputeFullDataCrc))
 	{
diff --git a/Engine/Plugins/PCG/Source/PCG/Private/Subsystems/PCGSubsystem.cpp b/Engine/Plugins/PCG/Source/PCG/Private/Subsystems/PCGSubsystem.cpp
--- a/Engine/Plugins/PCG/Source/PCG/Private/Subsystems/PCGSubsystem.cpp
+++ b/Engine/Plugins/PCG/Source/PCG/Private/Subsystems/PCGSubsystem.cpp
@@ -657,7 +657,7 @@ FPCGTaskId UPCGSubsystem::ScheduleComponent(UPCGComponent* PCGComponent, EPCGHiG
 			Dependencies.Add(OriginalComponentTask);
 		}
 
-		auto LocalGenerateTask = [OriginalComponent = PCGComponent, Grid, &Dependencies, bForce, &GridSizes](UPCGComponent* LocalComponent)
+		auto LocalGenerateTask = [OriginalComponent = PCGComponent, Grid, &Dependencies, bForce, &GridSizes](UPCGComponent* LocalComponent, const TArray<FPCGTaskId>& ExtraDependencies)
 		{
 			if (!GridSizes.Contains(LocalComponent->GetGenerationGridSize()))
 			{
@@ -690,11 +690,21 @@ FPCGTaskId UPCGSubsystem::ScheduleComponent(UPCGComponent* PCGComponent, EPCGHiG
 
 			TArray<FPCGTaskId> AdditionalDependencies;
 			const TArray<FPCGTaskId>* AllDependencies = &Dependencies;
-			if (LocalCleanupTaskId != InvalidPCGTaskId)
+			if (LocalCleanupTaskId != InvalidPCGTaskId || !ExtraDependencies.IsEmpty())
 			{
-				AdditionalDependencies.Reserve(Dependencies.Num() + 1);
+				AdditionalDependencies.Reserve(Dependencies.Num() + ExtraDependencies.Num() + 1);
 				AdditionalDependencies.Append(Dependencies);
-				AdditionalDependencies.Add(LocalCleanupTaskId);
+
+				if (LocalCleanupTaskId != InvalidPCGTaskId)
+				{
+					AdditionalDependencies.Add(LocalCleanupTaskId);
+				}
+
+				if (!ExtraDependencies.IsEmpty())
+				{
+					AdditionalDependencies.Append(ExtraDependencies);
+				}
+
 				AllDependencies = &AdditionalDependencies;
 			}
 
@@ -762,7 +772,7 @@ FPCGTaskId UPCGSubsystem::ScheduleCleanup(UPCGComponent* PCGComponent, bool bRem
 	// If the component is partitioned, we will forward the calls to its registered PCG Partition actors
 	if (PCGComponent->IsPartitioned())
 	{
-		auto LocalCleanupTask = [bRemoveComponents, &Dependencies](UPCGComponent* LocalComponent)
+		auto LocalCleanupTask = [bRemoveComponents, &Dependencies](UPCGComponent* LocalComponent, const TArray<FPCGTaskId>& ExtraDependencies)
 		{
 			// If the local component is currently cleaning up, it's probably because it was requested by a refresh.
 			// Wait after this one instead
@@ -771,8 +781,18 @@ FPCGTaskId UPCGSubsystem::ScheduleCleanup(UPCGComponent* PCGComponent, bool bRem
 				return LocalComponent->CurrentCleanupTask;
 			}
 
+			TArray<FPCGTaskId> AdditionalDependencies;
+			const TArray<FPCGTaskId>* AllDependencies = &Dependencies;
+			if (!ExtraDependencies.IsEmpty())
+			{
+				AdditionalDependencies.Reserve(Dependencies.Num() + ExtraDependencies.Num());
+				AdditionalDependencies.Append(Dependencies);
+				AdditionalDependencies.Append(ExtraDependencies);
+				AllDependencies = &AdditionalDependencies;
+			}
+
 			// Always executes regardless of local component grid size - clean up as much as possible.
-			return LocalComponent->CleanupLocal(bRemoveComponents, Dependencies);
+			return LocalComponent->CleanupLocal(bRemoveComponents, *AllDependencies);
 		};
 
 		AllTasks.Append(ActorAndComponentMapping->DispatchToRegisteredLocalComponents(PCGComponent, LocalCleanupTask));
@@ -889,8 +909,9 @@ void UPCGSubsystem::CancelGenerationInternal(IPCGGraphExecutionSource* Source, b
 
 	if (Component->IsPartitioned())
 	{
-		auto LocalCancel = [this, bCleanupUnusedResources](UPCGComponent* LocalComponent)
+		auto LocalCancel = [this, bCleanupUnusedResources](UPCGComponent* LocalComponent, const TArray<FPCGTaskId>& ExtraDependencies)
 		{
+			ensure(ExtraDependencies.IsEmpty());
 			if (LocalComponent->IsGenerating())
 			{
 				CancelGeneration(LocalComponent, bCleanupUnusedResources);
@@ -944,8 +965,9 @@ void UPCGSubsystem::RefreshAllComponentsFiltered(const TFunction<bool(UPCGCompon
 
 void UPCGSubsystem::ForAllRegisteredLocalComponents(UPCGComponent* InOriginalComponent, const TFunctionRef<void(UPCGComponent*)>& InFunc) const
 {
-	auto WrapperFunc = [&InFunc](UPCGComponent* Component) -> FPCGTaskId
+	auto WrapperFunc = [&InFunc](UPCGComponent* Component, const TArray<FPCGTaskId>& ExtraDependencies) -> FPCGTaskId
 	{
+		ensure(ExtraDependencies.IsEmpty());
 		InFunc(Component);
 		return InvalidPCGTaskId;
 	};
@@ -1100,8 +1122,9 @@ void UPCGSubsystem::CleanupLocalComponentsImmediate(UPCGComponent* InOriginalCom
 		return;
 	}
 
-	auto LocalCleanupTask = [bRemoveComponents](UPCGComponent* LocalComponent)
+	auto LocalCleanupTask = [bRemoveComponents](UPCGComponent* LocalComponent, const TArray<FPCGTaskId>& ExtraDependencies)
 	{
+		ensure(ExtraDependencies.IsEmpty());
 		if (ensure(LocalComponent) && !LocalComponent->IsCleaningUp())
 		{
 			LocalComponent->CleanupLocalImmediate(bRemoveComponents);
@@ -1338,8 +1361,9 @@ void UPCGSubsystem::DirtyGraph(UPCGComponent* Component, const FBox& InBounds, E
 	check(Component);
 
 	// Immediate operation
-	auto DirtyTask = [DirtyFlag](UPCGComponent* LocalComponent)
+	auto DirtyTask = [DirtyFlag](UPCGComponent* LocalComponent, const TArray<FPCGTaskId>& ExtraDependencies)
 	{
+		ensure(ExtraDependencies.IsEmpty());
 		LocalComponent->DirtyGenerated(DirtyFlag);
 		return InvalidPCGTaskId;
 	};
@@ -1860,6 +1884,11 @@ void UPCGSubsystem::UpdateMappingPCGComponentPartitionActor(UPCGComponent* InCom
 	ActorAndComponentMapping->UpdateMappingPCGComponentPartitionActor(InComponent); 
 }
 
+void UPCGSubsystem::SetChainedDispatchToLocalComponents(bool bInChainedDispatch)
+{
+	ActorAndComponentMapping->SetChainedDispatchToLocalComponents(bInChainedDispatch);
+}
+
 void UPCGSubsystem::OnPCGGraphCancelled(UPCGComponent* InComponent)
 {
 	ActorAndComponentMapping->OnPCGGraphCancelled(InComponent);
diff --git a/Engine/Plugins/PCG/Source/PCG/Public/Subsystems/PCGSubsystem.h b/Engine/Plugins/PCG/Source/PCG/Public/Subsystems/PCGSubsystem.h
--- a/Engine/Plugins/PCG/Source/PCG/Public/Subsystems/PCGSubsystem.h
+++ b/Engine/Plugins/PCG/Source/PCG/Public/Subsystems/PCGSubsystem.h
@@ -360,6 +360,8 @@ private:
 
 	UE_API void CreatePartitionActorsWithinBounds(UPCGComponent* InComponent, const FBox& InBounds, const PCGHiGenGrid::FSizeArray& InGridSizes);
 	UE_API void UpdateMappingPCGComponentPartitionActor(UPCGComponent* InComponent);
+
+	UE_API void SetChainedDispatchToLocalComponents(bool bInChainedDispatch);
 #endif // WITH_EDITOR
 	
 private:
diff --git a/Engine/Plugins/PCG/Source/PCGEditor/Private/WorldPartitionBuilder/PCGBuilderSettings.h b/Engine/Plugins/PCG/Source/PCGEditor/Private/WorldPartitionBuilder/PCGBuilderSettings.h
--- a/Engine/Plugins/PCG/Source/PCGEditor/Private/WorldPartitionBuilder/PCGBuilderSettings.h
+++ b/Engine/Plugins/PCG/Source/PCGEditor/Private/WorldPartitionBuilder/PCGBuilderSettings.h
@@ -61,4 +61,12 @@ public:
 	/** Call generate on each component and wait until completion and any async processes before generating the next. */
 	UPROPERTY(EditAnywhere, Category = Advanced, meta = (DisplayName = "One Component at a Time"))
 	bool bOneComponentAtATime = false;
+
+	/** Experimental feature that will try and revert PCG Component Actors which haven't changed from their previous generation. */
+	UPROPERTY(EditAnywhere, Category = Experimental)
+	bool bRevertUnchangedActors = false;
+
+	/** Outputs that can change on every generation but do not have an impact on final result should be listed here so they can be ignored when comparing results for revert. */
+	UPROPERTY(EditAnywhere, Category = Experimental, meta= (EditCondition=bRevertUnchangedActors, EditConditionHides))
+	TArray<FString> OutputsToIgnoreForRevert;
 };
\ No newline at end of file
diff --git a/Engine/Plugins/PCG/Source/PCGEditor/Private/WorldPartitionBuilder/PCGWorldPartitionBuilder.cpp b/Engine/Plugins/PCG/Source/PCGEditor/Private/WorldPartitionBuilder/PCGWorldPartitionBuilder.cpp
--- a/Engine/Plugins/PCG/Source/PCGEditor/Private/WorldPartitionBuilder/PCGWorldPartitionBuilder.cpp
+++ b/Engine/Plugins/PCG/Source/PCGEditor/Private/WorldPartitionBuilder/PCGWorldPartitionBuilder.cpp
@@ -6,6 +6,7 @@
 #include "PCGComponent.h"
 #include "PCGEditorSettings.h"
 #include "PCGGraph.h"
+#include "PCGManagedResource.h"
 #include "PCGSubsystem.h"
 #include "SPCGBuilderDialog.h"
 
@@ -26,6 +27,7 @@
 #include "Interfaces/IMainFrameModule.h"
 #include "Misc/OutputDevice.h"
 #include "Modules/ModuleManager.h"
+#include "Serialization/ArchiveCrc32.h"
 #include "UObject/Linker.h"
 #include "UObject/UObjectIterator.h"
 #include "WorldPartition/IWorldPartitionEditorModule.h"
@@ -35,7 +37,7 @@
 #include "WorldPartition/DataLayer/DataLayerManager.h"
 #include "WorldPartition/LoaderAdapter/LoaderAdapterShape.h"
 
-DEFINE_LOG_CATEGORY_STATIC(LogPCGWorldPartitionBuilder, All, All);
+DEFINE_LOG_CATEGORY_STATIC(LogPCGWorldPartitionBuilder, Verbose, All);
 
 #define LOCTEXT_NAMESPACE "PCGWorldPartitionBuilder"
 
@@ -95,6 +97,17 @@ struct FPCGWorldPartitionBuilder
 	/** Allows to temporarly disable the clearing of results */
 	static void SetDisableClearResults(UPCGSubsystem* InSubsystem, bool bInDisableClearResults);
 
+	static void GetRevertableActorsFromComponents(TArray<TWeakObjectPtr<UPCGComponent>> InComponentsToGenerate, const FPCGWorldPartitionBuilderArgs& InArgs, TMap<TSoftObjectPtr<AActor>, FPCGCrc>& OutRevertableActors);
+	static void GetRevertablePartitionActors(UWorld* InWorld, const FPCGWorldPartitionBuilderArgs& InArgs, TMap<TSoftObjectPtr<AActor>, FPCGCrc>& OutRevertableActors);
+	static void GetRevertableActorsFromDirtyPackages(TArray<UPackage*> InDirtyPackages, const FPCGWorldPartitionBuilderArgs& InArgs, TMap<TSoftObjectPtr<AActor>, FPCGCrc>& OutRevertableActors);
+	static void RevertUnchangedActors(TArray<UPackage*>& InOutDirtyPackages, const FPCGWorldPartitionBuilderArgs& InArgs, const TMap<TSoftObjectPtr<AActor>, FPCGCrc>& InRevertableActors, TSet<TSoftObjectPtr<AActor>>& InOutRevertedActors);
+	static bool IsActorRevertable(AActor* InActor, const FPCGWorldPartitionBuilderArgs& InArgs, FPCGCrc& OutCrc);
+	static void LogRevertedActors(const TSet<TSoftObjectPtr<AActor>>& InRevertedActors);
+
+	static void GetModifiedAndNewActorsFromDirtyPackages(const TArray<UPackage*>& InDirtyPackages, TArray<AActor*>& OutModifiedActors, TArray<AActor*>& OutNewActors);
+	static void LogModifiedAndNewActorsFromDirtyPackages(const TArray<UPackage*>& InDirtyPackages);
+
+	static void SetChainedDispatchToLocalComponents(UWorld* InWorld, bool bInChainedDispatch);
 private:
 	static void CollectSubGridComponentsToGenerateInternal(UPCGSubsystem* InPCGSubsystem, int32 InGridSizeIndex, const PCGHiGenGrid::FSizeArray& InGridSizes, UPCGComponent* InOriginalComponent, APCGPartitionActor* InPartitionActor, TArray<UPCGComponent*>& OutAdditionalComponentsToGenerate, TMap<UPCGComponent*, UPCGComponent*>& OutDependencies);
 
@@ -269,7 +282,7 @@ bool FPCGWorldPartitionBuilderArgs::InitializeFrom(const FPCGWorldPartitionComma
 			BuilderArgs.bGenerateEditingModeNormalComponents = BuilderSettings->EditingModes.Contains(EPCGEditorDirtyMode::Normal);
 			BuilderArgs.bGenerateEditingModePreviewComponents = BuilderSettings->EditingModes.Contains(EPCGEditorDirtyMode::Preview);
 			BuilderArgs.bIgnoreGenerationErrors = BuilderSettings->bIgnoreGenerationErrors;
-			BuilderArgs.bOneComponentAtATime = BuilderSettings->bOneComponentAtATime;
+			BuilderArgs.bOneComponentAtATime = BuilderSettings->bOneComponentAtATime || BuilderSettings->bRevertUnchangedActors;
 			BuilderArgs.IncludeActorIDs = BuilderSettings->FilterByActorNames;
 			BuilderArgs.IncludeGraphs = BuilderSettings->Graphs;
 			BuilderArgs.bIterativeCellLoading = BuilderSettings->bIterativeCellLoading;
@@ -278,6 +291,8 @@ bool FPCGWorldPartitionBuilderArgs::InitializeFrom(const FPCGWorldPartitionComma
 			BuilderArgs.bLoadActivatedRuntimeDataLayers = BuilderSettings->bLoadActivatedRuntimeDataLayers;
 			BuilderArgs.IncludedDataLayers = BuilderSettings->IncludedDataLayers;
 			BuilderArgs.ExcludedDataLayers = BuilderSettings->ExcludedDataLayers;
+			BuilderArgs.bRevertUnchangedActors = BuilderSettings->bRevertUnchangedActors;
+			BuilderArgs.OutputsToIgnoreForRevert = BuilderSettings->OutputsToIgnoreForRevert;
 
 			UE_LOG(LogPCGWorldPartitionBuilder, Display, TEXT("Loaded PCGBuilderSettings '%s'"), *SettingsPath.ToString());
 			bUsingBuilderSettings = true;
@@ -422,6 +437,15 @@ bool FPCGWorldPartitionBuilderArgs::InitializeFrom(const FPCGWorldPartitionComma
 		DataLayersToString.IsEmpty() ? TEXT("None") : *FString::Join(DataLayersToString, TEXT(",")),
 		FromDefault);
 
+	// Log Revert parameters
+	UE_LOG(LogPCGWorldPartitionBuilder, Display, TEXT("bRevertUnchangedActors '%d' (%s)"),
+		BuilderArgs.bRevertUnchangedActors ? 1 : 0,
+		FromDefault);
+
+	UE_LOG(LogPCGWorldPartitionBuilder, Display, TEXT("OutputsToIgnoreForRevert '%s' (%s)"),
+		BuilderArgs.OutputsToIgnoreForRevert.IsEmpty() ? TEXT("None") : *FString::Join(BuilderArgs.OutputsToIgnoreForRevert, TEXT(",")),
+		FromDefault);
+
 	OutArgs = BuilderArgs;
 	return true;
 }
@@ -536,6 +560,11 @@ bool UPCGWorldPartitionBuilder::PreRun(UWorld* World, FPackageSourceControlHelpe
 		return false;
 	}
 
+	FPCGWorldPartitionBuilder::SetChainedDispatchToLocalComponents(World, Args.bRevertUnchangedActors);
+
+	RevertableActors.Empty();
+	RevertedActors.Empty();
+
 	// Reset Allow list
 	AllowedAssetPackages.Empty();
 
@@ -603,7 +632,12 @@ bool UPCGWorldPartitionBuilder::RunEntireWorld(UWorld* World, FPackageSourceCont
 
 	// 2. Clear the dirty flag on any packages that dirtied themselves while loading. We're not interested in saving these.
 	FPCGWorldPartitionBuilder::ResetPackageDirtyFlags();
-		
+
+	// Before we run any generation, gather the revertable actors from the components we want to generate and the loaded pcg partition actors.
+	// This will give us the pre generation Actorr to CRC mapping to compare against when generation completes.
+	FPCGWorldPartitionBuilder::GetRevertableActorsFromComponents(ComponentsToGenerate, Args, RevertableActors);
+	FPCGWorldPartitionBuilder::GetRevertablePartitionActors(World, Args, RevertableActors);
+
 	// Create missing Partition Actors
 	FPCGWorldPartitionBuilder::SetDisablePartitionActorCreationForWorld(World, /*bDisable=*/false);
 	FPCGWorldPartitionBuilder::CreatePartitionedActors(World, Args, PartitionedComponents);
@@ -633,6 +667,9 @@ bool UPCGWorldPartitionBuilder::RunEntireWorld(UWorld* World, FPackageSourceCont
 		
 		DirtyPackages = FPCGWorldPartitionBuilder::GetDirtyPackages(AllowedAssetPackages);
 	}
+	
+	// 5. Revert unchanged actors
+	FPCGWorldPartitionBuilder::RevertUnchangedActors(DirtyPackages, Args, RevertableActors, RevertedActors);
 
 	if (!bGeneratedAnyComponent)
 	{
@@ -690,6 +727,15 @@ bool UPCGWorldPartitionBuilder::SaveEntireWorldPackages(UWorld* World, FPackageS
 			else
 			{
 				PackagesToSave.Add(DirtyPackage);
+
+				if (DirtyPackage->HasAnyPackageFlags(PKG_NewlyCreated))
+				{
+					++NumNewPackages;
+				}
+				else
+				{
+					++NumModifiedPackages;
+				}
 			}
 
 			// Remove from here if package was still part of the loaded DirtyPackages
@@ -697,8 +743,11 @@ bool UPCGWorldPartitionBuilder::SaveEntireWorldPackages(UWorld* World, FPackageS
 		}
 	}
 
-	// Log final changes after combining packages (there may have been duplicates in the two deleted package lists).
-	UE_LOG(LogPCGWorldPartitionBuilder, Display, TEXT("PostRun: Final package changes: %d modified, %d deleted."), DirtyPackages.Num(), DeletedActorPackages.Num() + PackagesToDelete.Num());
+	// Log New/Modified Actors
+	FPCGWorldPartitionBuilder::LogModifiedAndNewActorsFromDirtyPackages(DirtyPackages);
+
+	// Set final counts
+	NumDeletedPackages = DeletedActorPackages.Num() + PackagesToDelete.Num();
 
 	if (!SavePackages(PackagesToSave, PackageHelper))
 	{
@@ -746,6 +795,9 @@ bool UPCGWorldPartitionBuilder::RunPerCell(UWorld* InWorld, const FCellInfo& InC
 	// Reset Dirty flags so we can save only packages that will get modified after this line
 	FPCGWorldPartitionBuilder::ResetPackageDirtyFlags();
 
+	// Gather already loaded revertable Partition actors (before any modification occurs)
+	FPCGWorldPartitionBuilder::GetRevertablePartitionActors(InWorld, Args, RevertableActors);
+
 	// Create missing PAs and Update mappings
 	FPCGWorldPartitionBuilder::SetDisablePartitionActorCreationForWorld(InWorld, /*bDisable=*/false);
 
@@ -822,9 +874,12 @@ bool UPCGWorldPartitionBuilder::RunPerCell(UWorld* InWorld, const FCellInfo& InC
 		return true;
 	};
 
-	// Collect Components to Generate a first time, this will allow us to call FPCGWorldPartitionBuilder::UpdatePartitionedActors on ParttionActors with all their overlapp
+	// Collect Components to Generate a first time, this will allow us to call FPCGWorldPartitionBuilder::UpdatePartitionedActors on ParttionActors with all their overlaps
 	FPCGWorldPartitionBuilder::CollectComponentsToGenerate(InWorld, Args, GenerateComponentFilter, ComponentsToGenerate);
-		
+
+	// Gather revertable actors from the components to generate (this should be non-partitioned components as we already called GetRevertablePartitionActors at the beginning of this function)
+	FPCGWorldPartitionBuilder::GetRevertableActorsFromComponents(ComponentsToGenerate, Args, RevertableActors);
+
 	// Make sure Partitioned (Unbounded) components are Generated first
 	ComponentsToGenerate.Sort([](const TWeakObjectPtr<UPCGComponent>& WeakComponentA, const TWeakObjectPtr<UPCGComponent>& WeakComponentB)
 	{
@@ -845,6 +900,9 @@ bool UPCGWorldPartitionBuilder::RunPerCell(UWorld* InWorld, const FCellInfo& InC
 	// Do a Bounds expansion to make sure we are generating PCG Components on areas covering their owning actor bounds (which might be a bigger area then the currently loaded cell)
 	TUniquePtr<IWorldPartitionActorLoaderInterface::ILoaderAdapter> LoaderAdapter = FPCGWorldPartitionBuilder::ExpandLoadedBoundsIfNeeded(InWorld, InCellInfo.Bounds, ExpandedBounds, TEXT("Expanded Bounds"));
 	
+	// Gather newly loaded Partition actors before they get updated
+	FPCGWorldPartitionBuilder::GetRevertablePartitionActors(InWorld, Args, RevertableActors);
+
 	if (LoaderAdapter)
 	{
 		// Update Mappings for new loaded PAs from Expanded Bounds
@@ -901,8 +959,10 @@ bool UPCGWorldPartitionBuilder::RunPerCell(UWorld* InWorld, const FCellInfo& InC
 
 bool UPCGWorldPartitionBuilder::PostRun(UWorld* World, FPackageSourceControlHelper& PackageHelper, const bool bInRunSuccess)
 {
+	FPCGWorldPartitionBuilder::LogRevertedActors(RevertedActors);
+
 	// Log final changes after combining packages (there may have been duplicates in the two deleted package lists).
-	UE_LOG(LogPCGWorldPartitionBuilder, Display, TEXT("PostRun: Final package changes: %d new, %d modified, %d deleted."), NewPackageNames.Num(), ModifiedPackageNames.Num(), DeletedPackageNames.Num());
+	UE_LOG(LogPCGWorldPartitionBuilder, Display, TEXT("PostRun: Final package changes: %d new, %d modified, %d deleted, %d unchanged."), NumNewPackages, NumModifiedPackages , NumDeletedPackages, RevertedActors.Num());
 
 	FPCGWorldPartitionBuilder::SetDisablePartitionActorCreationForWorld(World, /*bDisable=*/false);
 
@@ -991,6 +1051,9 @@ bool UPCGWorldPartitionBuilder::SavePerCellPackages(UWorld* InWorld, FPackageSou
 	uint32 CellModifiedPackageCount = 0;
 	uint32 CellDeletedPackageCount = 0;
 
+	// Revert unchanged actors
+	FPCGWorldPartitionBuilder::RevertUnchangedActors(DirtyPackages, Args, RevertableActors, RevertedActors);
+
 	for (int i = DirtyPackages.Num() - 1; i >= 0; --i)
 	{
 		if (UPackage* DirtyPackage = DirtyPackages[i])
@@ -1023,6 +1086,14 @@ bool UPCGWorldPartitionBuilder::SavePerCellPackages(UWorld* InWorld, FPackageSou
 		}
 	}
 
+	// Update stats
+	NumNewPackages += CellNewPackageCount;
+	NumModifiedPackages += CellModifiedPackageCount;
+	NumDeletedPackages += CellDeletedPackageCount;
+
+	// Log New/Modified Actors
+	FPCGWorldPartitionBuilder::LogModifiedAndNewActorsFromDirtyPackages(DirtyPackages);
+
 	// Delete loaded packages
 	if (PackagesToDelete.Num() > 0 && !PackageHelper.Delete(PackagesToDelete))
 	{
@@ -1041,7 +1112,7 @@ bool UPCGWorldPartitionBuilder::SavePerCellPackages(UWorld* InWorld, FPackageSou
 		return false;
 	}
 
-	UE_LOG(LogPCGWorldPartitionBuilder, Display, TEXT("Cell package changes: %d new, %d modified, %d deleted."), CellNewPackageCount, CellModifiedPackageCount, CellDeletedPackageCount);
+	UE_LOG(LogPCGWorldPartitionBuilder, Display, TEXT("Cell package changes: %d modified, %d new, %d deleted."), CellNewPackageCount, CellModifiedPackageCount, CellDeletedPackageCount);
 	return true;
 }
 
@@ -1090,15 +1161,24 @@ bool FPCGWorldPartitionBuilder::GenerateComponents(
 		// Important: the list of component created sorted by dependency so that we can rely on the order of scheduling here
 		const FPCGDisableClearResults DisableClearResults(InWorld);
 
-		for (TWeakObjectPtr<UPCGComponent> ComponentWeakPtr : Components)
-		{
-			UPCGComponent* Component = ComponentWeakPtr.Get();
-			if (!Component)
+		TArray<UPCGComponent*> ValidComponents;
+		Algo::TransformIf(Components, ValidComponents,
+			[](const TWeakObjectPtr<UPCGComponent>& InWeakComponent)
 			{
-				UE_LOG(LogPCGWorldPartitionBuilder, Warning, TEXT("Lost a PCG component weak pointer, component will not be generated."));
-				continue;
-			}
+				return InWeakComponent.IsValid();
+			},
+			[](const TWeakObjectPtr<UPCGComponent>& InWeakComponent)
+			{
+				return InWeakComponent.Get();
+			});
+
+		ValidComponents.Sort([](const UPCGComponent& ComponentA, const UPCGComponent& ComponentB)
+		{
+			return ComponentA.GetPathName() < ComponentB.GetPathName();
+		});
 
+		for (UPCGComponent* Component : ValidComponents)
+		{
 			FPCGTaskId TaskIdDependency = InvalidPCGTaskId;
 			if (const UPCGComponent* const* TaskDependency = InTaskDependencies.Find(Component))
 			{
@@ -1352,6 +1432,220 @@ void FPCGWorldPartitionBuilder::SetDisableClearResults(UPCGSubsystem* InSubsyste
 	}
 }
 
+bool FPCGWorldPartitionBuilder::IsActorRevertable(AActor* InActor, const FPCGWorldPartitionBuilderArgs& InArgs, FPCGCrc& OutCrc)
+{
+	if (!InArgs.bRevertUnchangedActors)
+	{
+		return false;
+	}
+
+	if (!InActor->IsPackageExternal())
+	{
+		return false;
+	}
+
+	UE_LOG(LogPCGWorldPartitionBuilder, VeryVerbose, TEXT("IsActorRevertable : '%s' : begin"), *InActor->GetActorLabel());
+
+	const bool bLocalComponents = InActor->IsA<APCGPartitionActor>();
+	bool bCanRevertActor = true;
+	TArray<const UPCGComponent*> PCGComponents;
+	InActor->ForEachComponent(/*bIncludeFromChildActors=*/false, [bLocalComponents, &bCanRevertActor, &PCGComponents](const UActorComponent* InComponent)
+	{
+		const UPCGComponent* PCGComponent = Cast<UPCGComponent>(InComponent);
+		if (PCGComponent && (PCGComponent->IsLocalComponent() == bLocalComponents))
+		{
+			PCGComponents.Add(PCGComponent);
+		}
+		else if (!InComponent->IsDefaultSubobject())
+		{
+			bCanRevertActor = false;
+			UE_LOG(LogPCGWorldPartitionBuilder, VeryVerbose, TEXT("IsActorRevertable : '%s' : Can't revert actor because of non default component '%s'"), *InComponent->GetOwner()->GetActorLabel(), *InComponent->GetName());
+		}
+	});
+
+	if (!bCanRevertActor)
+	{
+		UE_LOG(LogPCGWorldPartitionBuilder, VeryVerbose, TEXT("IsActorRevertable : '%s' : end"), *InActor->GetActorLabel());
+		return false;
+	}
+
+	// Start from a random prime.
+	FPCGCrc Crc(1000003);
+
+	for (const UPCGComponent* PCGComponent : PCGComponents)
+	{
+		FArchiveCrc32 Ar;
+		bool bComponentGenerated = PCGComponent->bGenerated;
+		Ar << bComponentGenerated;
+		if (bComponentGenerated)
+		{
+			FBox ComponentLastGeneratedBounds = PCGComponent->LastGeneratedBounds;
+			Ar << ComponentLastGeneratedBounds;
+		}
+
+		Crc.Combine(Ar.GetCrc());
+
+		PCGComponent->ForEachConstManagedResource([&Crc, &bCanRevertActor, InActor](const UPCGManagedResource* ManagedResource)
+		{
+			FPCGCrc ResourceCrc = ManagedResource->GetCrc();
+			if (!ResourceCrc.IsValid())
+			{
+				bCanRevertActor = false;
+				UE_LOG(LogPCGWorldPartitionBuilder, VeryVerbose, TEXT("IsActorRevertable : '%s' Can't revert actor because of invalid crc for resource : '%s'"), *InActor->GetActorLabel(), *ManagedResource->GetName());
+			}
+			else
+			{
+				Crc.Combine(ManagedResource->GetCrc());
+				UE_LOG(LogPCGWorldPartitionBuilder, VeryVerbose, TEXT("IsActorRevertable : '%s' : resource '%s' : crc '%d'"), *InActor->GetActorLabel(), *ManagedResource->GetName(), ResourceCrc.GetValue());
+			}
+		});
+
+		for (const FPCGTaggedData& Output : PCGComponent->GetGeneratedGraphOutput().TaggedData)
+		{
+			if (InArgs.OutputsToIgnoreForRevert.Contains(Output.Pin.ToString()))
+			{
+				UE_LOG(LogPCGWorldPartitionBuilder, VeryVerbose, TEXT("IsActorRevertable : '%s' : output pin ignored: '%s'"), *InActor->GetActorLabel(), *Output.Pin.ToString());
+				continue;
+			}
+
+			FPCGCrc OutputCrc = Output.ComputeCrc(/*bFullDataCrc=*/true);
+			UE_LOG(LogPCGWorldPartitionBuilder, VeryVerbose, TEXT("IsActorRevertable : '%s' : output pin: '%s' : crc '%d'"), *InActor->GetActorLabel() , *Output.Pin.ToString(), OutputCrc.GetValue());
+			Crc.Combine(OutputCrc);
+		}
+	}
+		
+	if (!bCanRevertActor)
+	{
+		UE_LOG(LogPCGWorldPartitionBuilder, VeryVerbose, TEXT("IsActorRevertable : '%s' : end"), *InActor->GetActorLabel());
+		return false;
+	}
+
+	OutCrc = Crc;
+	UE_LOG(LogPCGWorldPartitionBuilder, VeryVerbose, TEXT("IsActorRevertable : '%s' : crc : '%d' : end"), *InActor->GetActorLabel(), OutCrc.GetValue());
+	return true;
+}
+
+void FPCGWorldPartitionBuilder::GetRevertableActorsFromComponents(TArray<TWeakObjectPtr<UPCGComponent>> InComponentsToGenerate, const FPCGWorldPartitionBuilderArgs& InArgs, TMap<TSoftObjectPtr<AActor>, FPCGCrc>& OutRevertableActors)
+{
+	for (TWeakObjectPtr<UPCGComponent> Component : InComponentsToGenerate)
+	{
+		AActor* Owner = Component.IsValid() ? Component->GetOwner() : nullptr;
+		if (Owner)
+		{
+			FPCGCrc Crc;
+			if (!OutRevertableActors.Contains(Owner) && FPCGWorldPartitionBuilder::IsActorRevertable(Owner, InArgs, Crc))
+			{
+				OutRevertableActors.Add(Owner, Crc);
+			}
+		}
+	}
+}
+
+void FPCGWorldPartitionBuilder::GetRevertablePartitionActors(UWorld* InWorld, const FPCGWorldPartitionBuilderArgs& InArgs, TMap<TSoftObjectPtr<AActor>, FPCGCrc>& OutRevertableActors)
+{
+	for (TActorIterator<APCGPartitionActor> It(InWorld); It; ++It)
+	{
+		FPCGCrc Crc;
+		if (!OutRevertableActors.Contains(*It) && FPCGWorldPartitionBuilder::IsActorRevertable(*It, InArgs, Crc))
+		{
+			OutRevertableActors.Add(*It, Crc);
+		}
+	}
+}
+
+void FPCGWorldPartitionBuilder::GetRevertableActorsFromDirtyPackages(TArray<UPackage*> InDirtyPackages, const FPCGWorldPartitionBuilderArgs& InArgs, TMap<TSoftObjectPtr<AActor>, FPCGCrc>& OutRevertableActors)
+{
+	for (UPackage* DirtyPackage : InDirtyPackages)
+	{
+		if (AActor* Actor = AActor::FindActorInPackage(DirtyPackage))
+		{
+			FPCGCrc Crc;
+			if (!OutRevertableActors.Contains(Actor) && FPCGWorldPartitionBuilder::IsActorRevertable(Actor, InArgs, Crc))
+			{
+				OutRevertableActors.Add(Actor, Crc);
+			}
+		}
+	}
+}
+
+void FPCGWorldPartitionBuilder::RevertUnchangedActors(TArray<UPackage*>& InOutDirtyPackages, const FPCGWorldPartitionBuilderArgs& InArgs, const TMap<TSoftObjectPtr<AActor>, FPCGCrc>& InRevertableActors, TSet<TSoftObjectPtr<AActor>>& InOutRevertedActors)
+{
+	TMap<TSoftObjectPtr<AActor>, FPCGCrc> DirtyRevertableActors;
+	FPCGWorldPartitionBuilder::GetRevertableActorsFromDirtyPackages(InOutDirtyPackages, InArgs, DirtyRevertableActors);
+
+	for (const auto& KVP : DirtyRevertableActors)
+	{
+		// Actor found in revertable actors and crc matches
+		if (const FPCGCrc* PreCrc = InRevertableActors.Find(KVP.Key); PreCrc && *PreCrc == KVP.Value)
+		{
+			if (InOutDirtyPackages.Remove(KVP.Key->GetPackage()))
+			{
+				InOutRevertedActors.Add(KVP.Key);
+			}
+		} 
+		else
+		{
+			// actor not found and or crc doesn't match. This can happen if actor is processed more than once in a per cell situation, one cell might not change the actor but the next one does.
+			InOutRevertedActors.Remove(KVP.Key);
+		}
+	}
+}
+
+void FPCGWorldPartitionBuilder::GetModifiedAndNewActorsFromDirtyPackages(const TArray<UPackage*>& InDirtyPackages, TArray<AActor*>& OutModifiedActors, TArray<AActor*>& OutNewActors)
+{
+	for (UPackage* DirtyPackage : InDirtyPackages)
+	{
+		if (AActor* Actor = AActor::FindActorInPackage(DirtyPackage))
+		{
+			if (DirtyPackage->HasAnyPackageFlags(PKG_NewlyCreated))
+			{
+				OutNewActors.Add(Actor);
+			}
+			else
+			{
+				OutModifiedActors.Add(Actor);
+			}
+		}
+	}
+}
+
+void FPCGWorldPartitionBuilder::LogRevertedActors(const TSet<TSoftObjectPtr<AActor>>& InRevertedActors)
+{
+	// Log actor changes
+	for (TSoftObjectPtr<AActor> RevertedActorPtr : InRevertedActors)
+	{
+		if (AActor* RevertedActor = RevertedActorPtr.Get())
+		{
+			UE_LOG(LogPCGWorldPartitionBuilder, Display, TEXT("Unchanged actor '%s' (name '%s' - package '%s')"), *RevertedActor->GetActorLabel(), *RevertedActor->GetName(), *RevertedActor->GetPackage()->GetName());
+		}
+	}
+}
+
+void FPCGWorldPartitionBuilder::LogModifiedAndNewActorsFromDirtyPackages(const TArray<UPackage*>& InDirtyPackages)
+{
+	TArray<AActor*> ModifiedActors;
+	TArray<AActor*> NewActors;
+	FPCGWorldPartitionBuilder::GetModifiedAndNewActorsFromDirtyPackages(InDirtyPackages, ModifiedActors, NewActors);
+
+	// Log actor changes
+	for (AActor* ModifiedActor : ModifiedActors)
+	{
+		UE_LOG(LogPCGWorldPartitionBuilder, Display, TEXT("Modified actor '%s' (name '%s' - package '%s')"), *ModifiedActor->GetActorLabel(), *ModifiedActor->GetName(), *ModifiedActor->GetPackage()->GetName());
+	}
+
+	for (AActor* NewActor : NewActors)
+	{
+		UE_LOG(LogPCGWorldPartitionBuilder, Display, TEXT("Added actor '%s' (name '%s' - package '%s')"), *NewActor->GetActorLabel(), *NewActor->GetName(), *NewActor->GetPackage()->GetName());
+	}
+}
+
+void FPCGWorldPartitionBuilder::SetChainedDispatchToLocalComponents(UWorld* InWorld, bool bInChainedDispatch)
+{
+	UPCGSubsystem* PCGSubsystem = UWorld::GetSubsystem<UPCGSubsystem>(InWorld);
+	check(PCGSubsystem);
+	PCGSubsystem->SetChainedDispatchToLocalComponents(bInChainedDispatch);
+}
+
 TArray<UPackage*> FPCGWorldPartitionBuilder::GetDirtyPackages(const TSet<FString>& AllowedAssetPackages)
 {
 	TArray<UPackage*> DirtyPackages;
diff --git a/Engine/Plugins/PCG/Source/PCGEditor/Private/WorldPartitionBuilder/PCGWorldPartitionBuilder.h b/Engine/Plugins/PCG/Source/PCGEditor/Private/WorldPartitionBuilder/PCGWorldPartitionBuilder.h
--- a/Engine/Plugins/PCG/Source/PCGEditor/Private/WorldPartitionBuilder/PCGWorldPartitionBuilder.h
+++ b/Engine/Plugins/PCG/Source/PCGEditor/Private/WorldPartitionBuilder/PCGWorldPartitionBuilder.h
@@ -2,6 +2,8 @@
 
 #pragma once
 
+#include "PCGCrc.h"
+
 #include "WorldPartition/WorldPartitionBuilder.h"
 #include "WorldPartition/WorldPartitionHandle.h"
 
@@ -92,6 +94,12 @@ struct FPCGWorldPartitionBuilderArgs
 	/** Data layers that need to be unloaded when running builder. */
 	TArray<TObjectPtr<UDataLayerAsset>> ExcludedDataLayers;
 
+	/** Experimental feature that will try and revert PCG Component Actors which haven't changed from their previous generation. */
+	bool bRevertUnchangedActors = false;
+
+	/** Outputs that can change on every generation but do not have an impact on final result should be listed here so they can be ignored when comparing results for revert. */
+	TArray<FString> OutputsToIgnoreForRevert;
+
 	static bool InitializeFrom(const FPCGWorldPartitionCommandlineArgs& CommandlineArgs, FPCGWorldPartitionBuilderArgs& OutArgs);
 };
 
@@ -179,6 +187,18 @@ private:
 	/** Non Actor assets that are allowed to be submitted by the builder */
 	TSet<FString> AllowedAssetPackages;
 
+	/** Revertable Actors */
+	TMap<TSoftObjectPtr<AActor>, FPCGCrc> RevertableActors;
+
+	/** Reverted Actors : when using the builder in per cell iteration, it is possible that one cell reverts an actor because it doesn't modify it but the next one will so keep track for stats */
+	TSet<TSoftObjectPtr<AActor>> RevertedActors;
+
+	/** Final stats */
+	int32 NumNewPackages = 0;
+	int32 NumModifiedPackages = 0;
+	int32 NumDeletedPackages = 0;
+	int32 NumRevertedPackages = 0;
+
 	/** Flag to register if any error occurred during generation. */
 	bool bErrorOccurredWhileGenerating = false;
 };
