[PCG] PR #13489: Fixes a null pointer exception in FPCGComponentDetails::OnGenerateClicked()

#rnx

diff --git a/Engine/Plugins/PCG/Source/PCGEditor/Private/Details/PCGComponentDetails.cpp b/Engine/Plugins/PCG/Source/PCGEditor/Private/Details/PCGComponentDetails.cpp
--- a/Engine/Plugins/PCG/Source/PCGEditor/Private/Details/PCGComponentDetails.cpp
+++ b/Engine/Plugins/PCG/Source/PCGEditor/Private/Details/PCGComponentDetails.cpp
@@ -278,7 +278,11 @@ FReply FPCGComponentDetails::OnGenerateClicked()
 		{
 			if (Component.IsValid())
 			{
-				Component->GetSubsystem()->FlushCache();
+				if (UPCGSubsystem* PCGSubsystem = Component->GetSubsystem())
+				{
+					PCGSubsystem->FlushCache();
+				}
+
 				break;
 			}
 		}
