PCG: Temporary crash fix for PCG component clearing transaction buffer out of PostLoad- relies on the fact that GUndo will always be null in PostLoad. Will be revisited by PCG team soon.

diff --git a/Engine/Plugins/PCG/Source/PCG/Private/PCGComponent.cpp b/Engine/Plugins/PCG/Source/PCG/Private/PCGComponent.cpp
--- a/Engine/Plugins/PCG/Source/PCG/Private/PCGComponent.cpp
+++ b/Engine/Plugins/PCG/Source/PCG/Private/PCGComponent.cpp
@@ -3945,7 +3945,11 @@ void UPCGComponent::ChangeTransientState(EPCGEditorDirtyMode NewEditingMode)
 
 	// Changing the transient state can and will play with packages and is not meant to be undoable.
 	// However, don't kill the transaction buffer unless it's actually needed, and FIXME: especially not during a construction script
-	if (bShouldMarkDirty && GEditor && GEditor->Trans)
+	if (bShouldMarkDirty && GEditor && GEditor->Trans 
+		// TODO: This is a temporary fix- we check GUndo here in part to guard against being called out of PostLoad, 
+		//  where clearing the transaction buffer will cause a crash, as upstream code sets GUndo to null and then 
+		//  later resets it back to something we deleted.
+		&& GUndo)
 	{
 		UWorld* World = GetWorld();
 
