[PCG] Limit Partition Grid Size to the min/max values of HiGen

* This was crashing if the value was too big

#rnx

diff --git a/Engine/Plugins/PCG/Source/PCG/Private/PCGWorldActor.cpp b/Engine/Plugins/PCG/Source/PCG/Private/PCGWorldActor.cpp
--- a/Engine/Plugins/PCG/Source/PCG/Private/PCGWorldActor.cpp
+++ b/Engine/Plugins/PCG/Source/PCG/Private/PCGWorldActor.cpp
@@ -249,6 +249,9 @@ void APCGWorldActor::OnPartitionGridSizeChanged()
 {
 	TRACE_CPUPROFILER_EVENT_SCOPE(APCGWorldActor::OnPartitionGridSizeChanged);
 
+	// Clamp within expected limits, defined by the min/max sizes from HiGen, multiply by 100 since HiGen sizes are in m, while the partition grid size is in cm.
+	PartitionGridSize = FMath::Clamp(PartitionGridSize, static_cast<uint32>(EPCGHiGenGrid::GridMin) * 100u, static_cast<uint32>(EPCGHiGenGrid::GridMax) * 100u);
+
 	UWorld* World = GetWorld();
 	if (!World)
 	{
diff --git a/Engine/Plugins/PCG/Source/PCG/Public/PCGWorldActor.h b/Engine/Plugins/PCG/Source/PCG/Public/PCGWorldActor.h
--- a/Engine/Plugins/PCG/Source/PCG/Public/PCGWorldActor.h
+++ b/Engine/Plugins/PCG/Source/PCG/Public/PCGWorldActor.h
@@ -57,8 +57,8 @@ public:
 
 	static inline constexpr uint32 DefaultPartitionGridSize = 25600; // 256m
 
-	/** Size of the PCG partition actor grid for non-hierarchical-generation graphs. */
-	UPROPERTY(config, EditAnywhere, Category = GenerationSettings)
+	/** Size of the PCG partition actor grid for non-hierarchical-generation graphs. Min value: 400, MaxValue: 419430400 */
+	UPROPERTY(config, EditAnywhere, Category = GenerationSettings, meta = (ClampMin=400, ClampMax=419430400))
 	uint32 PartitionGridSize;
 
 	/** Contains all the PCG data required to query the landscape complete. Serialized in cooked builds only */
