Save 0.38MB in a large server binary by adding a callsite forceinline on a function where common subexpression optimization would have a large effect.

diff --git a/Engine/Plugins/PCG/Source/PCG/Public/Metadata/PCGMetadataAttributeTpl.h b/Engine/Plugins/PCG/Source/PCG/Public/Metadata/PCGMetadataAttributeTpl.h
--- a/Engine/Plugins/PCG/Source/PCG/Public/Metadata/PCGMetadataAttributeTpl.h
+++ b/Engine/Plugins/PCG/Source/PCG/Public/Metadata/PCGMetadataAttributeTpl.h
@@ -1123,14 +1123,14 @@ namespace PCGMetadataAttribute
 	}
 
 	template <typename Func, typename... Args>
-	inline decltype(auto) CallbackWithRightType(uint16 TypeId, Func Callback, Args&& ...InArgs)
+	FORCEINLINE decltype(auto) CallbackWithRightType(uint16 TypeId, Func Callback, Args&& ...InArgs)
 	{
 		using ReturnType = decltype(Callback(double{}, std::forward<Args>(InArgs)...));
 
 		switch (TypeId)
 		{
 
-#define PCG_CALLBACKWITHRIGHTTYPE_DECL(T) case (uint16)(PCG::Private::MetadataTypes<T>::Id): return Callback(T{}, std::forward<Args>(InArgs)...);
+#define PCG_CALLBACKWITHRIGHTTYPE_DECL(T) case (uint16)(PCG::Private::MetadataTypes<T>::Id): CALLSITE_FORCEINLINE return Callback(T{}, std::forward<Args>(InArgs)...);
 		PCG_FOREACH_SUPPORTEDTYPES(PCG_CALLBACKWITHRIGHTTYPE_DECL)
 #undef PCG_CALLBACKWITHRIGHTTYPE_DECL
 
diff --git a/Engine/Source/Runtime/Core/Public/Clang/ClangPlatform.h b/Engine/Source/Runtime/Core/Public/Clang/ClangPlatform.h
--- a/Engine/Source/Runtime/Core/Public/Clang/ClangPlatform.h
+++ b/Engine/Source/Runtime/Core/Public/Clang/ClangPlatform.h
@@ -59,4 +59,6 @@
 #	define UE_NO_PROFILE_ATTRIBUTE
 #endif
 
+#define CALLSITE_FORCEINLINE [[clang::always_inline]]			/* Force code to be inlined at the callsite */
+
 #endif//defined(__clang__)
diff --git a/Engine/Source/Runtime/Core/Public/HAL/Platform.h b/Engine/Source/Runtime/Core/Public/HAL/Platform.h
--- a/Engine/Source/Runtime/Core/Public/HAL/Platform.h
+++ b/Engine/Source/Runtime/Core/Public/HAL/Platform.h
@@ -694,6 +694,8 @@
 #ifndef RESTRICT
 	#define RESTRICT __restrict						/* no alias hint */
 #endif
+#ifndef CALLSITE_FORCEINLINE
+	#define CALLSITE_FORCEINLINE 					/* Force code to be inline at the callsite */
 
 #ifndef UE_DEFINE_FORCEINLINE_HINT_TO_INLINE
 	#define UE_DEFINE_FORCEINLINE_HINT_TO_INLINE 0
