[PCG] Add `IsBuilder` to Get Execution Context node

* Needed to know which component the builder was scheduling
* Added that information to the PCG Editor Module, which can be queried through IPCGEditorModule
* Add helper function to get the result from the IPCGEditorModule

#rnx

diff --git a/Engine/Plugins/PCG/Source/PCG/Private/Elements/PCGGetExecutionContext.cpp b/Engine/Plugins/PCG/Source/PCG/Private/Elements/PCGGetExecutionContext.cpp
--- a/Engine/Plugins/PCG/Source/PCG/Private/Elements/PCGGetExecutionContext.cpp
+++ b/Engine/Plugins/PCG/Source/PCG/Private/Elements/PCGGetExecutionContext.cpp
@@ -104,6 +104,12 @@ bool FPCGGetExecutionContextElement::ExecuteInternal(FPCGContext* Context) const
 	{
 		Value = ExecutionSource && ExecutionSource->GetExecutionState().HasAuthority();
 	}
+	else if (Mode == EPCGGetExecutionContextMode::IsBuilder)
+	{
+#if WITH_EDITOR
+		Value = PCGHelpers::IsInWorldPartitionBuilderCommandletContext();
+#endif
+	}
 
 	FPCGMetadataAttribute<bool>* Attribute = ParamData->Metadata->CreateAttribute<bool>(PCGGetExecutionContextConstants::AttributeName, Value, /*bAllowInterpolation=*/false, /*bOverrideParent=*/false);
 	ParamData->Metadata->AddEntry();
diff --git a/Engine/Plugins/PCG/Source/PCG/Private/Helpers/PCGHelpers.cpp b/Engine/Plugins/PCG/Source/PCG/Private/Helpers/PCGHelpers.cpp
--- a/Engine/Plugins/PCG/Source/PCG/Private/Helpers/PCGHelpers.cpp
+++ b/Engine/Plugins/PCG/Source/PCG/Private/Helpers/PCGHelpers.cpp
@@ -516,6 +516,14 @@ namespace PCGHelpers
 			}
 		}
 	}
+
+	bool IsInWorldPartitionBuilderCommandletContext()
+	{
+		static const UClass* WorldPartitionBuilderCommandletClass = FindObject<UClass>(nullptr, TEXT("/Script/UnrealEd.WorldPartitionBuilderCommandlet"), EFindObjectFlags::ExactClass);
+		check(WorldPartitionBuilderCommandletClass);
+		const UClass* RunningCommandletClass = GetRunningCommandletClass();
+		return RunningCommandletClass && RunningCommandletClass->IsChildOf(WorldPartitionBuilderCommandletClass);
+	}
 #endif
 
 	bool IsNewObjectAndNotDefault(const UObject* InObject, bool bCheckHierarchy)
diff --git a/Engine/Plugins/PCG/Source/PCG/Public/Elements/PCGGetExecutionContext.h b/Engine/Plugins/PCG/Source/PCG/Public/Elements/PCGGetExecutionContext.h
--- a/Engine/Plugins/PCG/Source/PCG/Public/Elements/PCGGetExecutionContext.h
+++ b/Engine/Plugins/PCG/Source/PCG/Public/Elements/PCGGetExecutionContext.h
@@ -17,7 +17,8 @@ enum class EPCGGetExecutionContextMode : uint8
 	IsPartitioned UMETA(Tooltip = "Returns whether the executing context is the original context and that it is partitioned."),
 	IsRuntimeGeneration UMETA(Tooltip="Returns whether the executing context is doing runtime gneration."),
 	IsDedicatedServer UMETA(Tooltip="Returns whether the executing context is running on a dedicated server."),
-	HasAuthority UMETA(Tooltip="Returns whether the execution context is running with network authority.")
+	HasAuthority UMETA(Tooltip="Returns whether the execution context is running with network authority."),
+	IsBuilder UMETA(Tooltip="Returns whether the execution source is generated by the PCG World Partition Builder. Always false in non-editor."),
 };
 
 /** Returns some context-specific information in the form of an attribute set. */
diff --git a/Engine/Plugins/PCG/Source/PCG/Public/Helpers/PCGHelpers.h b/Engine/Plugins/PCG/Source/PCG/Public/Helpers/PCGHelpers.h
--- a/Engine/Plugins/PCG/Source/PCG/Public/Helpers/PCGHelpers.h
+++ b/Engine/Plugins/PCG/Source/PCG/Public/Helpers/PCGHelpers.h
@@ -69,6 +69,9 @@ namespace PCGHelpers
 #if WITH_EDITOR
 	PCG_API void GatherDependencies(UObject* Object, TSet<TObjectPtr<UObject>>& OutDependencies, int32 MaxDepth = -1, const TArray<UClass*>& InExcludedClasses = {});
 	PCG_API void GatherDependencies(FProperty* Property, const void* InContainer, TSet<TObjectPtr<UObject>>& OutDependencies, int32 MaxDepth, const TArray<UClass*>& InExcludedClasses = {});
+
+	/** Returns true if it's running a World Partition builder commandlet. */
+	PCG_API bool IsInWorldPartitionBuilderCommandletContext();
 #endif
 
 	/** 
