[PCG] Fix spline control points metadata entry keys initialization

* Will fix CleanSpline unit tests

#rnx

diff --git a/Engine/Plugins/PCG/Source/PCG/Private/Data/PCGSplineStruct.cpp b/Engine/Plugins/PCG/Source/PCG/Private/Data/PCGSplineStruct.cpp
--- a/Engine/Plugins/PCG/Source/PCG/Private/Data/PCGSplineStruct.cpp
+++ b/Engine/Plugins/PCG/Source/PCG/Private/Data/PCGSplineStruct.cpp
@@ -284,6 +284,16 @@ void FPCGSplineStruct::AddPoint(const FSplinePoint& InSplinePoint, bool bUpdateS
 	{
 		ControlPointsEntryKeys.Insert(InOptionalControlPointEntryKey ? *InOptionalControlPointEntryKey : PCGInvalidEntryKey, Index);
 	}
+	else if (InOptionalControlPointEntryKey)
+	{
+		// We have an entry key but the array of control points entry keys is empty. Fill it with invalid entry keys values and set this point entry
+		// key.
+		ControlPointsEntryKeys.SetNumUninitialized(GetNumberOfPoints());
+		for (int32 i = 0; i < ControlPointsEntryKeys.Num(); ++i)
+		{
+			ControlPointsEntryKeys[i] = i == Index ? *InOptionalControlPointEntryKey : PCGInvalidEntryKey;
+		}
+	}
 
 	if (bUpdateSpline)
 	{
