[PCG] Fix loc

#rnx

diff --git a/Engine/Plugins/PCG/Source/PCG/Private/Elements/PCGAttributeRemap.cpp b/Engine/Plugins/PCG/Source/PCG/Private/Elements/PCGAttributeRemap.cpp
--- a/Engine/Plugins/PCG/Source/PCG/Private/Elements/PCGAttributeRemap.cpp
+++ b/Engine/Plugins/PCG/Source/PCG/Private/Elements/PCGAttributeRemap.cpp
@@ -126,7 +126,7 @@ TArray<FPCGPreConfiguredSettingsInfo> UPCGAttributeRemapSettings::GetPreconfigur
 FString UPCGAttributeRemapSettings::GetAdditionalTitleInformation() const
 {
 	TStringBuilder<256> Builder;
-	Builder += (Mode == EPCGAttributeRemapMode::Curve ? LOCTEXT("CurveSubtitle", "Curve") : LOCTEXT("CurveSubtitle", "Ranges")).ToString();
+	Builder += (Mode == EPCGAttributeRemapMode::Curve ? LOCTEXT("CurveSubtitle", "Curve") : LOCTEXT("RangesSubtitle", "Ranges")).ToString();
 
 	if (!IsPropertyOverriddenByPin(GET_MEMBER_NAME_CHECKED(UPCGAttributeRemapSettings, InputSource)) && !IsPropertyOverriddenByPin(GET_MEMBER_NAME_CHECKED(UPCGAttributeRemapSettings, OutputTarget)))
 	{
