[PCG] Add new functions on the graph to hide Input/Output nodes.

* Private bHidden flag in UPCGNode
* They will still exist on the backend, but won't be constructed on the frontend

#rnx

diff --git a/Engine/Plugins/PCG/Source/PCG/Private/PCGGraph.cpp b/Engine/Plugins/PCG/Source/PCG/Private/PCGGraph.cpp
--- a/Engine/Plugins/PCG/Source/PCG/Private/PCGGraph.cpp
+++ b/Engine/Plugins/PCG/Source/PCG/Private/PCGGraph.cpp
@@ -2309,6 +2309,24 @@ FString UPCGGraph::GetGraphDescriptionJson() const
 
 	return JsonString;
 }
+
+void UPCGGraph::SetHiddenFlagInputNode(bool bHidden)
+{
+	if (InputNode && InputNode->bHidden != bHidden)
+	{
+		InputNode->Modify();
+		InputNode->bHidden = bHidden;
+	}
+}
+
+void UPCGGraph::SetHiddenFlagOutputNode(bool bHidden)
+{
+	if (OutputNode && OutputNode->bHidden != bHidden)
+	{
+		OutputNode->Modify();
+		OutputNode->bHidden = bHidden;
+	}
+}
 #endif // WITH_EDITOR
 
 /****************************
diff --git a/Engine/Plugins/PCG/Source/PCG/Public/PCGGraph.h b/Engine/Plugins/PCG/Source/PCG/Public/PCGGraph.h
--- a/Engine/Plugins/PCG/Source/PCG/Public/PCGGraph.h
+++ b/Engine/Plugins/PCG/Source/PCG/Public/PCGGraph.h
@@ -410,6 +410,9 @@ public:
 	UE_API double GetGridCleanupRadiusFromGrid(EPCGHiGenGrid Grid) const;
 
 #if WITH_EDITOR
+	/** Can be overriden by child class to disable debug globally on all settings. */
+	virtual bool ShouldDisplayDebuggingProperties() const { return true; };
+	
 	UE_API void DisableNotificationsForEditor();
 	UE_API void EnableNotificationsForEditor();
 	UE_API void ToggleUserPausedNotificationsForEditor();
@@ -465,6 +468,12 @@ public:
 protected:
 #if WITH_EDITOR
 	virtual bool IsTemplatePropertyEnabled() const override { return IsAsset(); }
+
+	/** Mark the input node hidden/not hidden. */
+	PCG_API void SetHiddenFlagInputNode(bool bHidden);
+	
+	/** Mark the output node hidden/not hidden. */
+	PCG_API void SetHiddenFlagOutputNode(bool bHidden); 
 #endif
 
 	/** Internal function to react to add/remove nodes. bNotify can be set to false to not notify the world. */
diff --git a/Engine/Plugins/PCG/Source/PCG/Public/PCGNode.h b/Engine/Plugins/PCG/Source/PCG/Public/PCGNode.h
--- a/Engine/Plugins/PCG/Source/PCG/Public/PCGNode.h
+++ b/Engine/Plugins/PCG/Source/PCG/Public/PCGNode.h
@@ -204,8 +204,18 @@ public:
 
 	UPROPERTY(meta = (PCGNoHash))
 	uint8 bCommentBubbleVisible : 1;
+
+private:
+	/** Will be hidden in the Editor, but will still exist in the backend. Not exposed and mainly used to hide the input/output node. */
+	UPROPERTY(meta = (PCGNoHash))
+	uint8 bHidden : 1;
 #endif // WITH_EDITORONLY_DATA
 
+public:
+#if WITH_EDITOR
+	bool IsHidden() const { return bHidden;}
+#endif // WITH_EDITOR
+	
 protected:
 	/** Updates pins based on node settings. Attempts to migrate pins via matching. Broadcasts node change events for affected nodes. */
 	UE_API EPCGChangeType UpdatePins();
diff --git a/Engine/Plugins/PCG/Source/PCGEditor/Private/PCGEditorGraph.cpp b/Engine/Plugins/PCG/Source/PCGEditor/Private/PCGEditorGraph.cpp
--- a/Engine/Plugins/PCG/Source/PCGEditor/Private/PCGEditorGraph.cpp
+++ b/Engine/Plugins/PCG/Source/PCGEditor/Private/PCGEditorGraph.cpp
@@ -83,30 +83,26 @@ void UPCGEditorGraph::ReconstructGraph()
 	TMap<UPCGNode*, UPCGEditorGraphNodeBase*> NodeLookup;
 	constexpr bool bSelectNewNode = false;
 
-	// Create input and output nodes directly.
+	auto ConstructNode = [this, &NodeLookup]<typename NodeType>(UPCGNode* PCGNode, TSubclassOf<NodeType> NodeClass = NodeType::StaticClass())
 	{
-		UPCGNode* InputNode = PCGGraph->GetInputNode();
-		check(InputNode);
-		FGraphNodeCreator<UPCGEditorGraphNodeInput> InputNodeCreator(*this);
-		UPCGEditorGraphNodeInput* InputGraphNode = InputNodeCreator.CreateNode(bSelectNewNode);
-		InputGraphNode->Construct(InputNode);
-		InputNodeCreator.Finalize();
-		NodeLookup.Add(InputNode, InputGraphNode);
-	}
+		check(PCGNode && NodeClass);
+		if (!PCGNode->IsHidden())
+		{
+			FGraphNodeCreator<NodeType> NodeCreator(*this);
+			NodeType* GraphNode = NodeCreator.CreateNode(bSelectNewNode, NodeClass);
+			GraphNode->Construct(PCGNode);
+			NodeCreator.Finalize();
+			NodeLookup.Add(PCGNode, GraphNode);
+		}
+	};
 
-	{
-		UPCGNode* OutputNode = PCGGraph->GetOutputNode();
-		check(OutputNode);
-		FGraphNodeCreator<UPCGEditorGraphNodeOutput> OutputNodeCreator(*this);
-		UPCGEditorGraphNodeOutput* OutputGraphNode = OutputNodeCreator.CreateNode(bSelectNewNode);
-		OutputGraphNode->Construct(OutputNode);
-		OutputNodeCreator.Finalize();
-		NodeLookup.Add(OutputNode, OutputGraphNode);
-	}
+	// Create input and output nodes directly.
+	ConstructNode.operator()<UPCGEditorGraphNodeInput>(PCGGraph->GetInputNode());
+	ConstructNode.operator()<UPCGEditorGraphNodeOutput>(PCGGraph->GetOutputNode());
 
 	for (UPCGNode* PCGNode : PCGGraph->GetNodes())
 	{
-		if (!IsValid(PCGNode))
+		if (!IsValid(PCGNode) || PCGNode->IsHidden())
 		{
 			continue;
 		}
@@ -118,13 +114,7 @@ void UPCGEditorGraph::ReconstructGraph()
 				? GetGraphNodeClassFromPCGSettings(PCGSettings)
 				: TSubclassOf<UPCGEditorGraphNodeBase>(UPCGEditorGraphNode::StaticClass());
 
-		FGraphNodeCreator<UPCGEditorGraphNodeBase> NodeCreator(*this);
-		check(PCGNode && PCGGraphNodeClass);
-		// No need to select, since its reconstruction.
-		UPCGEditorGraphNodeBase* GraphNode = NodeCreator.CreateNode(bSelectNewNode, PCGGraphNodeClass);
-		GraphNode->Construct(PCGNode);
-		NodeCreator.Finalize();
-		NodeLookup.Add(PCGNode, GraphNode);
+		ConstructNode(PCGNode, PCGGraphNodeClass);
 	}
 
 	for (const auto& NodeLookupIt : NodeLookup)
