Ran update script to removed null initializers when creating TSubclassOf<T> since it will use a code path that is the same as default initializer except that it checks so nullptr is a valid type (which it always is). And that check requires full knowledge of T (so type can't be forward declared)
Same for function parameter default values but instead replaced = nullptr with = {} (which uses default initializer)

diff --git a/Engine/Plugins/Animation/LiveLinkHub/Source/LiveLinkHub/Private/Recording/Implementations/LiveLinkRecordingDataContainer.h b/Engine/Plugins/Animation/LiveLinkHub/Source/LiveLinkHub/Private/Recording/Implementations/LiveLinkRecordingDataContainer.h
--- a/Engine/Plugins/Animation/LiveLinkHub/Source/LiveLinkHub/Private/Recording/Implementations/LiveLinkRecordingDataContainer.h
+++ b/Engine/Plugins/Animation/LiveLinkHub/Source/LiveLinkHub/Private/Recording/Implementations/LiveLinkRecordingDataContainer.h
@@ -103,7 +103,7 @@ struct FLiveLinkRecordingStaticDataContainer : public FLiveLinkRecordingBaseData
 
 	/** The role of the static data being recorded. */
 	UPROPERTY()
-	TSubclassOf<ULiveLinkRole> Role = nullptr;
+	TSubclassOf<ULiveLinkRole> Role;
 };
 
 USTRUCT()
diff --git a/Engine/Plugins/PCG/Source/PCG/Public/Elements/IO/PCGSaveAssetElement.h b/Engine/Plugins/PCG/Source/PCG/Public/Elements/IO/PCGSaveAssetElement.h
--- a/Engine/Plugins/PCG/Source/PCG/Public/Elements/IO/PCGSaveAssetElement.h
+++ b/Engine/Plugins/PCG/Source/PCG/Public/Elements/IO/PCGSaveAssetElement.h
@@ -67,7 +67,7 @@ public:
 	TArray<FPCGPinProperties> Pins;
 
 	UPROPERTY(BlueprintReadWrite, EditAnywhere, Category = Data)
-	TSubclassOf<UPCGDataCollectionExporter> CustomDataCollectionExporterClass = nullptr;
+	TSubclassOf<UPCGDataCollectionExporter> CustomDataCollectionExporterClass;
 
 	UPROPERTY(BlueprintReadWrite, EditAnywhere, Category = Data, meta = (ShowOnlyInnerProperties, PCG_Overridable))
 	FPCGAssetExporterParameters Params;
diff --git a/Engine/Plugins/PCG/Source/PCG/Public/Elements/PCGActorSelector.h b/Engine/Plugins/PCG/Source/PCG/Public/Elements/PCGActorSelector.h
--- a/Engine/Plugins/PCG/Source/PCG/Public/Elements/PCGActorSelector.h
+++ b/Engine/Plugins/PCG/Source/PCG/Public/Elements/PCGActorSelector.h
@@ -97,7 +97,7 @@ struct FPCGSelectionKey
 	FName Tag = NAME_None;
 
 	UPROPERTY()
-	TSubclassOf<UObject> SelectionClass = nullptr;
+	TSubclassOf<UObject> SelectionClass;
 
 	// If the Selection is ByPath, contain the path to select.
 	UPROPERTY()
diff --git a/Engine/Plugins/PCG/Source/PCG/Public/Elements/PCGAddComponent.h b/Engine/Plugins/PCG/Source/PCG/Public/Elements/PCGAddComponent.h
--- a/Engine/Plugins/PCG/Source/PCG/Public/Elements/PCGAddComponent.h
+++ b/Engine/Plugins/PCG/Source/PCG/Public/Elements/PCGAddComponent.h
@@ -81,7 +81,7 @@ public:
 
 protected:
 	UPROPERTY(BlueprintReadOnly, EditAnywhere, Category = Settings, meta = (EditCondition = "!bUseClassAttribute", EditConditionHides))
-	TSubclassOf<UActorComponent> TemplateComponentClass = nullptr;
+	TSubclassOf<UActorComponent> TemplateComponentClass;
 
 	UPROPERTY(BlueprintReadWrite, EditAnywhere, Category = Settings, meta = (EditCondition = "!bUseClassAttribute", EditConditionHides))
 	bool bAllowTemplateComponentEditing = false;
@@ -130,4 +130,4 @@ private:
 		const IPCGAttributeAccessorKeys* TargetActorKeys,
 		IPCGAttributeAccessor* OutputComponentAccessor,
 		IPCGAttributeAccessorKeys* OutputComponentKeys) const;
-};
\ No newline at end of file
+};
diff --git a/Engine/Plugins/PCG/Source/PCG/Public/Elements/PCGCreateTargetActor.h b/Engine/Plugins/PCG/Source/PCG/Public/Elements/PCGCreateTargetActor.h
--- a/Engine/Plugins/PCG/Source/PCG/Public/Elements/PCGCreateTargetActor.h
+++ b/Engine/Plugins/PCG/Source/PCG/Public/Elements/PCGCreateTargetActor.h
@@ -99,7 +99,7 @@ public:
 
 protected:
 	UPROPERTY(BlueprintReadOnly, EditAnywhere, Category = Settings, meta = (OnlyPlaceable, DisallowCreateNew))
-	TSubclassOf<AActor> TemplateActorClass = nullptr;
+	TSubclassOf<AActor> TemplateActorClass;
 
 	// TODO: make this InlineEditConditionToggle, not done because property changed event does not propagate correctly so we can't track accurately the need to create the target actor
 	UPROPERTY(EditAnywhere, BlueprintReadWrite, Category = Settings)
@@ -115,4 +115,4 @@ protected:
 	virtual bool CanExecuteOnlyOnMainThread(FPCGContext* Context) const override { return true; }
 	virtual bool IsCacheable(const UPCGSettings* Settings) const override { return false; }
 	virtual bool ExecuteInternal(FPCGContext* Context) const override;
-};
\ No newline at end of file
+};
diff --git a/Engine/Plugins/PCG/Source/PCG/Public/Elements/PCGSpawnActor.h b/Engine/Plugins/PCG/Source/PCG/Public/Elements/PCGSpawnActor.h
--- a/Engine/Plugins/PCG/Source/PCG/Public/Elements/PCGSpawnActor.h
+++ b/Engine/Plugins/PCG/Source/PCG/Public/Elements/PCGSpawnActor.h
@@ -109,7 +109,7 @@ public:
 
 protected:
 	UPROPERTY(BlueprintReadOnly, EditAnywhere, Category = Settings, meta = (OnlyPlaceable, DisallowCreateNew))
-	TSubclassOf<AActor> TemplateActorClass = nullptr;
+	TSubclassOf<AActor> TemplateActorClass;
 
 	UPROPERTY(BlueprintReadWrite, EditAnywhere, Category = Settings, meta = (EditCondition = "Option != EPCGSpawnActorOption::CollapseActors"))
 	bool bAllowTemplateActorEditing = false;
diff --git a/Engine/Plugins/PCG/Source/PCG/Public/Elements/PCGSurfaceSampler.h b/Engine/Plugins/PCG/Source/PCG/Public/Elements/PCGSurfaceSampler.h
--- a/Engine/Plugins/PCG/Source/PCG/Public/Elements/PCGSurfaceSampler.h
+++ b/Engine/Plugins/PCG/Source/PCG/Public/Elements/PCGSurfaceSampler.h
@@ -90,7 +90,7 @@ namespace PCGSurfaceSampler
 	};
 
 	/** Sample a surface and returns the resulting point data. Can't be timesliced. */
-	PCG_API UPCGBasePointData* SampleSurface(FPCGContext* Context, const FSurfaceSamplerParams& ExecutionParams, const UPCGSurfaceData* InSurface, const UPCGSpatialData* InBoundingShape, const FBox& EffectiveBounds, TSubclassOf<UPCGBasePointData> PointDataClass = nullptr);
+	PCG_API UPCGBasePointData* SampleSurface(FPCGContext* Context, const FSurfaceSamplerParams& ExecutionParams, const UPCGSurfaceData* InSurface, const UPCGSpatialData* InBoundingShape, const FBox& EffectiveBounds, TSubclassOf<UPCGBasePointData> PointDataClass = {});
 		
 	PCG_API UPCGPointData* SampleSurface(FPCGContext* Context, const UPCGSurfaceData* InSurface, const UPCGSpatialData* InBoundingShape, const FBox& EffectiveBounds, const FSurfaceSamplerParams& ExecutionParams);
 
diff --git a/Engine/Plugins/PCG/Source/PCG/Public/PCGComponent.h b/Engine/Plugins/PCG/Source/PCG/Public/PCGComponent.h
--- a/Engine/Plugins/PCG/Source/PCG/Public/PCGComponent.h
+++ b/Engine/Plugins/PCG/Source/PCG/Public/PCGComponent.h
@@ -338,7 +338,7 @@ public:
 
 	/** A Scheduling Policy dictates the order in which instances of this component will be scheduled. */
 	UPROPERTY(BlueprintReadOnly, EditAnywhere, NoClear, Category = RuntimeGeneration, meta = (EditCondition = "GenerationTrigger == EPCGComponentGenerationTrigger::GenerateAtRuntime", EditConditionHides))
-	TSubclassOf<UPCGSchedulingPolicyBase> SchedulingPolicyClass = nullptr;
+	TSubclassOf<UPCGSchedulingPolicyBase> SchedulingPolicyClass;
 
 	/** This is the instanced UPCGSchedulingPolicy object which holds scheduling parameters and calculates priorities. */
 	UPROPERTY(BlueprintReadOnly, VisibleAnywhere, Instanced, Category = RuntimeGeneration, meta = (EditCondition = "GenerationTrigger == EPCGComponentGenerationTrigger::GenerateAtRuntime", EditConditionHides))
diff --git a/Engine/Plugins/PCG/Source/PCG/Public/PCGData.h b/Engine/Plugins/PCG/Source/PCG/Public/PCGData.h
--- a/Engine/Plugins/PCG/Source/PCG/Public/PCGData.h
+++ b/Engine/Plugins/PCG/Source/PCG/Public/PCGData.h
@@ -395,19 +395,19 @@ class UPCGDataFunctionLibrary : public UBlueprintFunctionLibrary
 public:
 	/** Gets all inputs of the given class type, returning matching tagged data in the OutTaggedData value too */
 	UFUNCTION(BlueprintCallable, Category = Data, meta = (ScriptMethod, DeterminesOutputType = "InDataTypeClass"))
-	static UE_API TArray<UPCGData*> GetTypedInputs(const FPCGDataCollection& InCollection, TArray<FPCGTaggedData>& OutTaggedData, TSubclassOf<UPCGData> InDataTypeClass = nullptr);
+	static UE_API TArray<UPCGData*> GetTypedInputs(const FPCGDataCollection& InCollection, TArray<FPCGTaggedData>& OutTaggedData, TSubclassOf<UPCGData> InDataTypeClass = {});
 
 	/** Gets all inputs of the given class type and on the given pin, returning matching tagged data in the OutTaggedData value too */
 	UFUNCTION(BlueprintCallable, Category = Data, meta = (ScriptMethod, DeterminesOutputType = "InDataTypeClass"))
-	static UE_API TArray<UPCGData*> GetTypedInputsByPin(const FPCGDataCollection& InCollection, const FPCGPinProperties& InPin, TArray<FPCGTaggedData>& OutTaggedData, TSubclassOf<UPCGData> InDataTypeClass = nullptr);
+	static UE_API TArray<UPCGData*> GetTypedInputsByPin(const FPCGDataCollection& InCollection, const FPCGPinProperties& InPin, TArray<FPCGTaggedData>& OutTaggedData, TSubclassOf<UPCGData> InDataTypeClass = {});
 
 	/** Gets all inputs of the given class type and on the given pin label, returning matching tagged data in the OutTaggedData value too */
 	UFUNCTION(BlueprintCallable, Category = Data, meta = (ScriptMethod, DeterminesOutputType = "InDataTypeClass"))
-	static UE_API TArray<UPCGData*> GetTypedInputsByPinLabel(const FPCGDataCollection& InCollection, FName InPinLabel, TArray<FPCGTaggedData>& OutTaggedData, TSubclassOf<UPCGData> InDataTypeClass = nullptr);
+	static UE_API TArray<UPCGData*> GetTypedInputsByPinLabel(const FPCGDataCollection& InCollection, FName InPinLabel, TArray<FPCGTaggedData>& OutTaggedData, TSubclassOf<UPCGData> InDataTypeClass = {});
 
 	/** Gets all inputs of the given class type and having the provided tag, returning matching tagged data in the OutTaggedData value too */
 	UFUNCTION(BlueprintCallable, Category = Data, meta = (ScriptMethod, DeterminesOutputType = "InDataTypeClass"))
-	static UE_API TArray<UPCGData*> GetTypedInputsByTag(const FPCGDataCollection& InCollection, const FString& InTag, TArray<FPCGTaggedData>& OutTaggedData, TSubclassOf<UPCGData> InDataTypeClass = nullptr);
+	static UE_API TArray<UPCGData*> GetTypedInputsByTag(const FPCGDataCollection& InCollection, const FString& InTag, TArray<FPCGTaggedData>& OutTaggedData, TSubclassOf<UPCGData> InDataTypeClass = {});
 
 	/** Adds a data object to a given collection, simpler usage than making a PCGTaggedData object. InTags can be empty. */
 	UFUNCTION(BlueprintCallable, Category = Data, meta = (ScriptMethod, AutoCreateRefTerm = "InTags"))
diff --git a/Engine/Plugins/PCG/Source/PCG/Public/PCGGraph.h b/Engine/Plugins/PCG/Source/PCG/Public/PCGGraph.h
--- a/Engine/Plugins/PCG/Source/PCG/Public/PCGGraph.h
+++ b/Engine/Plugins/PCG/Source/PCG/Public/PCGGraph.h
@@ -379,7 +379,7 @@ public:
 	UE_API UPCGNode* FindNodeWithSettings(const UPCGSettingsInterface* InSettings, bool bRecursive = false) const;
 
 	/** Returns the first node that matches the given name in the graph, if any.*/
-	UE_API UPCGNode* FindNodeByTitleName(FName NodeTitle, bool bRecursive = false, TSubclassOf<const UPCGSettings> OptionalClass = nullptr) const;
+	UE_API UPCGNode* FindNodeByTitleName(FName NodeTitle, bool bRecursive = false, TSubclassOf<const UPCGSettings> OptionalClass = {}) const;
 
 	const TArray<UPCGNode*>& GetNodes() const { return Nodes; }
 	UE_API void AddNode(UPCGNode* InNode);
diff --git a/Engine/Plugins/Runtime/MassGameplay/Source/MassEQS/Public/Generators/MassEnvQueryGenerator_MassEntityHandles.h b/Engine/Plugins/Runtime/MassGameplay/Source/MassEQS/Public/Generators/MassEnvQueryGenerator_MassEntityHandles.h
--- a/Engine/Plugins/Runtime/MassGameplay/Source/MassEQS/Public/Generators/MassEnvQueryGenerator_MassEntityHandles.h
+++ b/Engine/Plugins/Runtime/MassGameplay/Source/MassEQS/Public/Generators/MassEnvQueryGenerator_MassEntityHandles.h
@@ -34,7 +34,7 @@ protected:
 
 	/** Context of query */
 	UPROPERTY(EditAnywhere, Category=Generator)
-	TSubclassOf<UEnvQueryContext> SearchCenter = nullptr;
+	TSubclassOf<UEnvQueryContext> SearchCenter;
 };
 
 /** Data required to be sent to Mass for processing this Generator Request */
@@ -58,4 +58,4 @@ struct FMassEnvQueryResultData_MassEntityHandles: public FMassEQSRequestData
 	}
 
 	TArray<FMassEnvQueryEntityInfo> GeneratedEntityInfo;
-};
\ No newline at end of file
+};
diff --git a/Engine/Plugins/Runtime/MassGameplay/Source/MassEQS/Public/Processors/MassEnvQueryProcessorBase.h b/Engine/Plugins/Runtime/MassGameplay/Source/MassEQS/Public/Processors/MassEnvQueryProcessorBase.h
--- a/Engine/Plugins/Runtime/MassGameplay/Source/MassEQS/Public/Processors/MassEnvQueryProcessorBase.h
+++ b/Engine/Plugins/Runtime/MassGameplay/Source/MassEQS/Public/Processors/MassEnvQueryProcessorBase.h
@@ -16,6 +16,6 @@ class UMassEnvQueryProcessorBase : public UMassProcessor
 protected:
 	virtual void InitializeInternal(UObject& Owner, const TSharedRef<FMassEntityManager>& EntityManager);
 	
-	TSubclassOf<UEnvQueryNode> CorrespondingRequestClass = nullptr;
+	TSubclassOf<UEnvQueryNode> CorrespondingRequestClass;
 	int32 CachedRequestQueryIndex = -1;
-};
\ No newline at end of file
+};
diff --git a/Engine/Plugins/Runtime/ScriptableToolsFramework/Source/ScriptableToolsFramework/Public/ScriptableInteractiveTool.h b/Engine/Plugins/Runtime/ScriptableToolsFramework/Source/ScriptableToolsFramework/Public/ScriptableInteractiveTool.h
--- a/Engine/Plugins/Runtime/ScriptableToolsFramework/Source/ScriptableToolsFramework/Public/ScriptableInteractiveTool.h
+++ b/Engine/Plugins/Runtime/ScriptableToolsFramework/Source/ScriptableToolsFramework/Public/ScriptableInteractiveTool.h
@@ -366,13 +366,13 @@ public:
 		                                                                                       EditCondition = "ToolStartupRequirements==EScriptableToolStartupRequirements::Custom",
 																							   EditConditionHides,
 																					           MustImplement = "/Script/ScriptableToolsFramework.CustomScriptableToolBuilderBaseInterface"))
-	TSubclassOf<UCustomScriptableToolBuilder> CustomToolBuilderClass = nullptr;
+	TSubclassOf<UCustomScriptableToolBuilder> CustomToolBuilderClass;
 
 	UPROPERTY(BlueprintReadWrite, EditAnywhere, Category = "Scriptable Tool Settings", meta = (DisplayName = "Tool Builder Class",
 		                                                                                       EditCondition = "ToolStartupRequirements==EScriptableToolStartupRequirements::ToolTarget",
 		                                                                                       EditConditionHides,
 		                                                                                       MustImplement="/Script/ScriptableToolsFramework.CustomScriptableToolBuilderBaseInterface"))
-	TSubclassOf<UToolTargetScriptableToolBuilder> ToolTargetToolBuilderClass = nullptr;
+	TSubclassOf<UToolTargetScriptableToolBuilder> ToolTargetToolBuilderClass;
 
 
 
