[PCG] Debug tree: filter out Inactive/Preview world pcg components
- Might need to be revisited for the preview window if at some point it runs its own execution source

#rnx

diff --git a/Engine/Plugins/PCG/Source/PCGEditor/Private/Widgets/SPCGEditorGraphDebugObjectTree.cpp b/Engine/Plugins/PCG/Source/PCGEditor/Private/Widgets/SPCGEditorGraphDebugObjectTree.cpp
--- a/Engine/Plugins/PCG/Source/PCGEditor/Private/Widgets/SPCGEditorGraphDebugObjectTree.cpp
+++ b/Engine/Plugins/PCG/Source/PCGEditor/Private/Widgets/SPCGEditorGraphDebugObjectTree.cpp
@@ -1192,6 +1192,13 @@ void SPCGEditorGraphDebugObjectTree::RefreshTree()
 				continue;
 			}
 
+			const UWorld* ExecutionStateWorld = PCGSource->GetExecutionState().GetWorld();
+			// Ignore Inactive and Preview worlds
+			if (ExecutionStateWorld && (ExecutionStateWorld->WorldType == EWorldType::Inactive || ExecutionStateWorld->WorldType == EWorldType::EditorPreview))
+			{
+				continue;
+			}
+
 			const bool bIsPartitioned = PCGComponent && PCGComponent->IsPartitioned();
 
 			// Process static stacks that can be read from the compiled graph.
