Add NanitePixelProgrammableDistance to ISM and PCG component descriptors

diff --git a/Engine/Plugins/PCG/Source/PCG/Private/Components/PCGProceduralISMComponent.cpp b/Engine/Plugins/PCG/Source/PCG/Private/Components/PCGProceduralISMComponent.cpp
--- a/Engine/Plugins/PCG/Source/PCG/Private/Components/PCGProceduralISMComponent.cpp
+++ b/Engine/Plugins/PCG/Source/PCG/Private/Components/PCGProceduralISMComponent.cpp
@@ -151,6 +151,7 @@ FPCGProceduralISMComponentDescriptor& FPCGProceduralISMComponentDescriptor::oper
 	bEvaluateWorldPositionOffset = Other.bEvaluateWorldPositionOffset;
 	bReverseCulling = Other.bReverseCulling;
 	WorldPositionOffsetDisableDistance = Other.WorldPositionOffsetDisableDistance;
+	NanitePixelProgrammableDistance = Other.NanitePixelProgrammableDistance;
 	ShadowCacheInvalidationBehavior = Other.ShadowCacheInvalidationBehavior;
 	DetailMode = Other.DetailMode;
 	bVisibleInRayTracing = Other.bVisibleInRayTracing;
@@ -202,6 +203,7 @@ void FPCGProceduralISMComponentDescriptor::InitComponent(UPCGProceduralISMCompon
 	InProceduralISMComponent->bEvaluateWorldPositionOffset = bEvaluateWorldPositionOffset;
 	InProceduralISMComponent->bReverseCulling = bReverseCulling;
 	InProceduralISMComponent->WorldPositionOffsetDisableDistance = WorldPositionOffsetDisableDistance;
+	InProceduralISMComponent->NanitePixelProgrammableDistance = NanitePixelProgrammableDistance;
 	InProceduralISMComponent->ShadowCacheInvalidationBehavior = ShadowCacheInvalidationBehavior;
 	InProceduralISMComponent->DetailMode = DetailMode;
 	InProceduralISMComponent->bVisibleInRayTracing = bVisibleInRayTracing;
@@ -241,6 +243,7 @@ void FPCGProceduralISMComponentDescriptor::InitFrom(UPCGProceduralISMComponent*
 	bEvaluateWorldPositionOffset = InProceduralISMComponent->bEvaluateWorldPositionOffset;
 	bReverseCulling = InProceduralISMComponent->bReverseCulling;
 	WorldPositionOffsetDisableDistance = InProceduralISMComponent->WorldPositionOffsetDisableDistance;
+	NanitePixelProgrammableDistance = InProceduralISMComponent->NanitePixelProgrammableDistance;
 	ShadowCacheInvalidationBehavior = InProceduralISMComponent->ShadowCacheInvalidationBehavior;
 	DetailMode = InProceduralISMComponent->DetailMode;
 	bVisibleInRayTracing = InProceduralISMComponent->bVisibleInRayTracing;
@@ -297,6 +300,7 @@ bool FPCGProceduralISMComponentDescriptor::operator==(const FPCGProceduralISMCom
 		&& bEvaluateWorldPositionOffset == Other.bEvaluateWorldPositionOffset
 		&& bReverseCulling == Other.bReverseCulling
 		&& WorldPositionOffsetDisableDistance == Other.WorldPositionOffsetDisableDistance
+		&& NanitePixelProgrammableDistance == Other.NanitePixelProgrammableDistance
 		&& ShadowCacheInvalidationBehavior == Other.ShadowCacheInvalidationBehavior
 		&& DetailMode == Other.DetailMode
 		&& bVisibleInRayTracing == Other.bVisibleInRayTracing
diff --git a/Engine/Plugins/PCG/Source/PCG/Private/Components/PCGProceduralISMComponent.h b/Engine/Plugins/PCG/Source/PCG/Private/Components/PCGProceduralISMComponent.h
--- a/Engine/Plugins/PCG/Source/PCG/Private/Components/PCGProceduralISMComponent.h
+++ b/Engine/Plugins/PCG/Source/PCG/Private/Components/PCGProceduralISMComponent.h
@@ -328,6 +328,9 @@ public:
 	UPROPERTY()
 	int32 WorldPositionOffsetDisableDistance;
 
+	UPROPERTY()
+	int32 NanitePixelProgrammableDistance;
+
 	UPROPERTY()
 	EShadowCacheInvalidationBehavior ShadowCacheInvalidationBehavior;
 
diff --git a/Engine/Source/Runtime/Engine/Private/ISMPartition/ISMComponentDescriptor.cpp b/Engine/Source/Runtime/Engine/Private/ISMPartition/ISMComponentDescriptor.cpp
--- a/Engine/Source/Runtime/Engine/Private/ISMPartition/ISMComponentDescriptor.cpp
+++ b/Engine/Source/Runtime/Engine/Private/ISMPartition/ISMComponentDescriptor.cpp
@@ -133,6 +133,7 @@ void FISMComponentDescriptorBase::InitFrom(const UStaticMeshComponent* Template,
 	bRenderInDepthPass = Template->bRenderInDepthPass;	
 	bEvaluateWorldPositionOffset = Template->bEvaluateWorldPositionOffset;
 	WorldPositionOffsetDisableDistance = Template->WorldPositionOffsetDisableDistance;
+	NanitePixelProgrammableDistance = Template->NanitePixelProgrammableDistance;
 	ShadowCacheInvalidationBehavior = Template->ShadowCacheInvalidationBehavior;
 	DetailMode = Template->DetailMode;
 	// Determine if this instance must render with reversed culling based on both scale and the component property
@@ -294,6 +295,7 @@ bool FISMComponentDescriptorBase::operator==(const FISMComponentDescriptorBase&
 	bFillCollisionUnderneathForNavmesh == Other.bFillCollisionUnderneathForNavmesh &&
 	bRasterizeAsFilledConvexVolume == Other.bRasterizeAsFilledConvexVolume &&
 	WorldPositionOffsetDisableDistance == Other.WorldPositionOffsetDisableDistance &&
+	NanitePixelProgrammableDistance == Other.NanitePixelProgrammableDistance &&
 	ShadowCacheInvalidationBehavior == Other.ShadowCacheInvalidationBehavior &&
 	DetailMode == Other.DetailMode &&
 #if WITH_EDITORONLY_DATA
@@ -434,6 +436,7 @@ void FISMComponentDescriptorBase::InitComponent(UInstancedStaticMeshComponent* I
 	ISMComponent->bFillCollisionUnderneathForNavmesh = bFillCollisionUnderneathForNavmesh;
 	ISMComponent->bRasterizeAsFilledConvexVolume = bRasterizeAsFilledConvexVolume;
 	ISMComponent->WorldPositionOffsetDisableDistance = WorldPositionOffsetDisableDistance;
+	ISMComponent->NanitePixelProgrammableDistance = NanitePixelProgrammableDistance;
 	ISMComponent->ShadowCacheInvalidationBehavior = ShadowCacheInvalidationBehavior;
 	ISMComponent->DetailMode = DetailMode;
 	
diff --git a/Engine/Source/Runtime/Engine/Public/ISMPartition/ISMComponentDescriptor.h b/Engine/Source/Runtime/Engine/Public/ISMPartition/ISMComponentDescriptor.h
--- a/Engine/Source/Runtime/Engine/Public/ISMPartition/ISMComponentDescriptor.h
+++ b/Engine/Source/Runtime/Engine/Public/ISMPartition/ISMComponentDescriptor.h
@@ -242,6 +242,9 @@ public:
 	UPROPERTY(EditAnywhere, Category = "Component Settings")
 	int32 WorldPositionOffsetDisableDistance;
 
+	UPROPERTY(EditAnywhere, Category = "Component Settings")
+	int32 NanitePixelProgrammableDistance;
+
 	UPROPERTY(EditAnywhere, Category = "Component Settings")
 	EShadowCacheInvalidationBehavior ShadowCacheInvalidationBehavior;
 
