[PCG] Guard subsystem existence in some function in the PCG Component

* Most Subsystem calls are guarded if the subsystem is null, but there was a few places where it was not.

#rnx

diff --git a/Engine/Plugins/PCG/Source/PCG/Private/PCGComponent.cpp b/Engine/Plugins/PCG/Source/PCG/Private/PCGComponent.cpp
--- a/Engine/Plugins/PCG/Source/PCG/Private/PCGComponent.cpp
+++ b/Engine/Plugins/PCG/Source/PCG/Private/PCGComponent.cpp
@@ -689,6 +689,8 @@ void UPCGComponent::PostProcessGraph(const FBox& InNewBounds, bool bInGenerated,
 	ResetIgnoredChangeOrigins(/*bLogIfAnyPresent=*/true);
 #endif
 
+	UPCGSubsystem* Subsystem = GetSubsystem();
+
 	if (bInGenerated)
 	{
 		bGenerated = true;
@@ -811,8 +813,11 @@ void UPCGComponent::PostProcessGraph(const FBox& InNewBounds, bool bInGenerated,
 
 		bDirtyGenerated = false;
 
-		// Notify Subsystem first
-		GetSubsystem()->OnPCGGraphGenerated(this);
+		// Notify Subsystem first, if it exists
+		if (Subsystem)
+		{
+			Subsystem->OnPCGGraphGenerated(this);
+		}
 #endif
 
 		// Notify Delegate next
@@ -848,8 +853,10 @@ void UPCGComponent::PostProcessGraph(const FBox& InNewBounds, bool bInGenerated,
 
 	StopGenerationInProgress();
 
-	UPCGSubsystem* Subsystem = GetSubsystem();
-	Subsystem->OnPCGSourceGenerationDone.Broadcast(Subsystem, IsValidChecked(this) ? this : nullptr, EPCGGenerationStatus::Completed);
+	if (Subsystem)
+	{
+		Subsystem->OnPCGSourceGenerationDone.Broadcast(Subsystem, IsValidChecked(this) ? this : nullptr, EPCGGenerationStatus::Completed);
+	}
 #endif
 }
 
@@ -979,14 +986,20 @@ void UPCGComponent::OnProcessGraphAborted(bool bQuiet, bool bCleanupUnusedResour
 
 	// Notify Subsystem first
 	UPCGSubsystem* Subsystem = GetSubsystem();
-	Subsystem->OnPCGGraphCancelled(this);
+	if (Subsystem)
+	{
+		Subsystem->OnPCGGraphCancelled(this);
+	}
 #endif
 
 	// Notify Delegate next
 	OnPCGGraphCancelledDelegate.Broadcast(this);
 
 #if WITH_EDITOR
-	Subsystem->OnPCGSourceGenerationDone.Broadcast(Subsystem, this, EPCGGenerationStatus::Aborted);
+	if (Subsystem)
+	{
+		Subsystem->OnPCGSourceGenerationDone.Broadcast(Subsystem, this, EPCGGenerationStatus::Aborted);
+	}
 #endif
 
 	PCGComponent::BroadcastDynamicDelegate(OnPCGGraphCancelledExternal, this);
@@ -1125,13 +1138,15 @@ void UPCGComponent::CleanupLocal(bool bRemoveComponents)
 
 FPCGTaskId UPCGComponent::CleanupLocal(bool bRemoveComponents, const TArray<FPCGTaskId>& Dependencies)
 {
-	if (!GetSubsystem() || IsCleaningUp())
+	UPCGSubsystem* Subsystem = GetSubsystem();
+
+	if (!Subsystem || IsCleaningUp())
 	{
 		return InvalidPCGTaskId;
 	}
 
 	const bool bNeedsLocalCleanup = bGenerated || IsGenerating() || !GeneratedResources.IsEmpty();
-	if (!bNeedsLocalCleanup && GetSubsystem()->GetPCGComponentPartitionActorMappings(this).IsEmpty())
+	if (!bNeedsLocalCleanup && Subsystem->GetPCGComponentPartitionActorMappings(this).IsEmpty())
 	{
 		return InvalidPCGTaskId;
 	}
@@ -1144,7 +1159,7 @@ FPCGTaskId UPCGComponent::CleanupLocal(bool bRemoveComponents, const TArray<FPCG
 	ExtraCapture.ResetCapturedMessages();
 #endif
 
-	CurrentCleanupTask = GetSubsystem()->ScheduleCleanup(this, bRemoveComponents, Dependencies);
+	CurrentCleanupTask = Subsystem->ScheduleCleanup(this, bRemoveComponents, Dependencies);
 	return CurrentCleanupTask;
 }
 
