[PCG] Robustness pass to adding a new graph parameter.

+ Helper to issue warnings for common property bag errors related to naming properties
+ Updated CreateConstant and UserParameterGet support handling errors
* Weakened check to ensure on failure

#rnx

diff --git a/Engine/Plugins/PCG/Source/PCG/Private/Elements/PCGCreateAttribute.cpp b/Engine/Plugins/PCG/Source/PCG/Private/Elements/PCGCreateAttribute.cpp
--- a/Engine/Plugins/PCG/Source/PCG/Private/Elements/PCGCreateAttribute.cpp
+++ b/Engine/Plugins/PCG/Source/PCG/Private/Elements/PCGCreateAttribute.cpp
@@ -352,16 +352,28 @@ bool UPCGCreateAttributeSetSettings::ConvertNode(const FPCGPreconfiguredInfo& Co
 				check(UserParametersStruct->IsPropertyNameValid(PropertyName));
 
 				FPropertyBagPropertyDesc PropertyDesc = PCGPropertyHelpers::CreatePropertyBagDescWithMetadataType(PropertyName, AttributeTypes.Type);
-				Graph->AddUserParameters({PropertyDesc});
-				const FPropertyBagPropertyDesc* Desc = UserParametersStruct->FindPropertyDescByName(PropertyName);
-				check(Desc);
 
-				// Update the generated node with the new user parameter
-				UPCGUserParameterGetSettings* Settings = CastChecked<UPCGUserParameterGetSettings>(NodeConverter.GetGeneratedSettings());
-				Settings->PropertyName = Desc->Name;
-				Settings->PropertyGuid = Desc->ID;
+				const EPropertyBagAlterationResult Result = Graph->AddUserParameters({PropertyDesc});
+				if (Result != EPropertyBagAlterationResult::Success)
+				{
+					PCGGraphParameter::Helpers::LogGraphParamNamingErrors(PropertyName, Result);
+					return false;
+				}
 
-				AttributeTypes.Dispatcher([PropertyName, Graph]<typename T>(T&& Value) { Graph->SetGraphParameter(PropertyName, std::forward<T>(Value)); });
+				if (const FPropertyBagPropertyDesc* Desc = UserParametersStruct->FindPropertyDescByName(PropertyName))
+				{
+					// Update the generated node with the new user parameter
+					UPCGUserParameterGetSettings* Settings = CastChecked<UPCGUserParameterGetSettings>(NodeConverter.GetGeneratedSettings());
+					Settings->PropertyName = Desc->Name;
+					Settings->PropertyGuid = Desc->ID;
+					AttributeTypes.Dispatcher([PropertyName, Graph]<typename T>(T&& Value) { Graph->SetGraphParameter(PropertyName, std::forward<T>(Value)); });
+				}
+				else
+				{
+					// This should never trigger, but keep it a soft assert just in case.
+					UE_LOG(LogPCG, Error, TEXT("Could not create a unique graph parameter with name '%s'."), *PropertyName.ToString());
+					return ensure(false);
+				}
 			}
 
 			NodeConverter.ApplyStructural();
diff --git a/Engine/Plugins/PCG/Source/PCG/Private/Elements/PCGUserParameterGet.cpp b/Engine/Plugins/PCG/Source/PCG/Private/Elements/PCGUserParameterGet.cpp
--- a/Engine/Plugins/PCG/Source/PCG/Private/Elements/PCGUserParameterGet.cpp
+++ b/Engine/Plugins/PCG/Source/PCG/Private/Elements/PCGUserParameterGet.cpp
@@ -256,31 +256,37 @@ void UPCGUserParameterGetSettings::ApplyPreconfiguredSettings(const FPCGPreConfi
 
 	const UPCGNode* Node = CastChecked<UPCGNode>(GetOuter());
 	if (UPCGGraph* Graph = Node ? Node->GetGraph() : nullptr)
-    {
-        if (!PCGGraphParameter::Helpers::GenerateUniqueName(Graph, PropertyName))
-        {
-	        UE_LOG(LogPCG, Error, TEXT("Could not create a unique graph parameter with name '%s'. Check if the graph is valid and the name has not reached its maximum."), *PropertyName.ToString());
-	        return;
-        }
-        check(FInstancedPropertyBag::IsPropertyNameValid(PropertyName));
-
-        FPropertyBagPropertyDesc PropertyDesc = PCGPropertyHelpers::CreatePropertyBagDescWithMetadataType(PropertyName, NewType);
-        Graph->AddUserParameters({PropertyDesc});
-
-        if (FInstancedPropertyBag* PropertyBag = Graph->GetMutableUserParametersStruct_Unsafe())
-        {
-	        if (const FPropertyBagPropertyDesc* Desc = PropertyBag->FindPropertyDescByName(PropertyName))
-	        {
-		        PropertyGuid = Desc->ID;
-		        PropertyName = Desc->Name;
-	        }
-	        else
-	        {
-		        ensure(false);
-		        UE_LOG(LogPCG, Error, TEXT("Could not create a unique graph parameter with name '%s'."), *PropertyName.ToString());
-	        }
-        }
-    }
+	{
+		if (!PCGGraphParameter::Helpers::GenerateUniqueName(Graph, PropertyName))
+		{
+			UE_LOG(LogPCG, Error, TEXT("Could not create a unique graph parameter with name '%s'. Check if the graph is valid and the name has not reached its maximum."), *PropertyName.ToString());
+			return;
+		}
+		check(FInstancedPropertyBag::IsPropertyNameValid(PropertyName));
+
+		FPropertyBagPropertyDesc PropertyDesc = PCGPropertyHelpers::CreatePropertyBagDescWithMetadataType(PropertyName, NewType);
+
+		const EPropertyBagAlterationResult Result = Graph->AddUserParameters({PropertyDesc});
+		if (Result != EPropertyBagAlterationResult::Success)
+		{
+			PCGGraphParameter::Helpers::LogGraphParamNamingErrors(PropertyName, Result);
+			return;
+		}
+
+		if (FInstancedPropertyBag* PropertyBag = Graph->GetMutableUserParametersStruct_Unsafe())
+		{
+			if (const FPropertyBagPropertyDesc* Desc = PropertyBag->FindPropertyDescByName(PropertyName))
+			{
+				PropertyGuid = Desc->ID;
+				PropertyName = Desc->Name;
+			}
+			else
+			{
+				UE_LOG(LogPCG, Error, TEXT("Could not create a unique graph parameter with name '%s'."), *PropertyName.ToString());
+				ensure(false);
+			}
+		}
+	}
 }
 
 bool UPCGUserParameterGetSettings::ConvertNode(const FPCGPreconfiguredInfo& ConversionInfo)
diff --git a/Engine/Plugins/PCG/Source/PCG/Private/Helpers/PCGGraphParametersHelpers.cpp b/Engine/Plugins/PCG/Source/PCG/Private/Helpers/PCGGraphParametersHelpers.cpp
--- a/Engine/Plugins/PCG/Source/PCG/Private/Helpers/PCGGraphParametersHelpers.cpp
+++ b/Engine/Plugins/PCG/Source/PCG/Private/Helpers/PCGGraphParametersHelpers.cpp
@@ -150,6 +150,28 @@ namespace PCGGraphParameter::Helpers
 
 		return true;
 	}
+
+	void LogGraphParamNamingErrors(const FName IntendedName, const EPropertyBagAlterationResult Result, const FPCGContext* InContext)
+	{
+		switch (Result)
+		{
+			case EPropertyBagAlterationResult::PropertyNameEmpty:
+				PCGLog::LogErrorOnGraph(LOCTEXT("EmptyPropertyName", "Could not add graph parameter with empty name."), InContext);
+				return;
+			case EPropertyBagAlterationResult::PropertyNameInvalidCharacters:
+				PCGLog::LogErrorOnGraph(FText::Format(LOCTEXT("InvalidCharacters", "Could not add graph parameter '{0}'. It contains invalid characters."), FText::FromName(IntendedName)), InContext);
+				return;
+			case EPropertyBagAlterationResult::TargetPropertyAlreadyExists:
+				PCGLog::LogErrorOnGraph(FText::Format(LOCTEXT("AlreadyExists", "Could not add graph parameter '{0}'. It already exists."), FText::FromName(IntendedName)), InContext);
+				return;
+			case EPropertyBagAlterationResult::InternalError:
+				PCGLog::LogErrorOnGraph(FText::Format(LOCTEXT("InternalError", "Could not add graph parameter '{0}' due to internal problems with the property bag."), FText::FromName(IntendedName)), InContext);
+				return;
+			default: // fall-through
+			case EPropertyBagAlterationResult::Success:
+				return;
+		}
+	}
 } // PCGGraphParameter::Helpers
 
 //////////////////////////
diff --git a/Engine/Plugins/PCG/Source/PCG/Private/PCGGraph.cpp b/Engine/Plugins/PCG/Source/PCG/Private/PCGGraph.cpp
--- a/Engine/Plugins/PCG/Source/PCG/Private/PCGGraph.cpp
+++ b/Engine/Plugins/PCG/Source/PCG/Private/PCGGraph.cpp
@@ -2209,9 +2209,16 @@ PCGHiGenGrid::FSizeArray UPCGGraph::CalculateNodeGridSizesRecursiveNoLock(const
 	return AllGridSizes;
 }
 
-void UPCGGraph::AddUserParameters(const TArray<FPropertyBagPropertyDesc>& InDescs, const UPCGGraph* InOptionalOriginalGraph)
+EPropertyBagAlterationResult UPCGGraph::AddUserParameters(const TArray<FPropertyBagPropertyDesc>& InDescs, const UPCGGraph* InOptionalOriginalGraph)
 {
-	UserParameters.AddProperties(InDescs);
+	switch (const EPropertyBagAlterationResult Result = UserParameters.AddProperties(InDescs))
+	{
+		case EPropertyBagAlterationResult::Success:
+			break;
+		default: // Any error, return the result.
+			return Result;
+	}
+
 	if (InOptionalOriginalGraph)
 	{
 		if (const FInstancedPropertyBag* OriginalPropertyBag = InOptionalOriginalGraph->GetUserParametersStruct())
@@ -2221,6 +2228,8 @@ void UPCGGraph::AddUserParameters(const TArray<FPropertyBagPropertyDesc>& InDesc
 	}
 
 	OnGraphParametersChanged(EPCGGraphParameterEvent::MultiplePropertiesAdded, NAME_None);
+
+	return EPropertyBagAlterationResult::Success;
 }
 
 #if WITH_EDITOR
diff --git a/Engine/Plugins/PCG/Source/PCG/Public/Helpers/PCGGraphParametersHelpers.h b/Engine/Plugins/PCG/Source/PCG/Public/Helpers/PCGGraphParametersHelpers.h
--- a/Engine/Plugins/PCG/Source/PCG/Public/Helpers/PCGGraphParametersHelpers.h
+++ b/Engine/Plugins/PCG/Source/PCG/Public/Helpers/PCGGraphParametersHelpers.h
@@ -2,17 +2,22 @@
 
 #pragma once
 
+#include "PCGContext.h"
 #include "Kismet/BlueprintFunctionLibrary.h"
 
 #include "PCGGraphParametersHelpers.generated.h"
 
 #define UE_API PCG_API
 
+enum class EPropertyBagAlterationResult : uint8;
+
 class UPCGGraphInterface;
 
 namespace PCGGraphParameter::Helpers
 {
 	PCG_API bool GenerateUniqueName(const UPCGGraphInterface* InGraph, FName& InOutName);
+
+	PCG_API void LogGraphParamNamingErrors(FName IntendedName, const EPropertyBagAlterationResult Result, const FPCGContext* InContext = nullptr);
 }
 
 /**
diff --git a/Engine/Plugins/PCG/Source/PCG/Public/PCGGraph.h b/Engine/Plugins/PCG/Source/PCG/Public/PCGGraph.h
--- a/Engine/Plugins/PCG/Source/PCG/Public/PCGGraph.h
+++ b/Engine/Plugins/PCG/Source/PCG/Public/PCGGraph.h
@@ -592,7 +592,7 @@ public:
 	// Original Graph needs to have the properties.
 	// Be careful if there is any overlap between existing parameters, that also exists in the original graph, they will be overridden by the original.
 	// Best used on a brand new PCG Graph.
-	UE_API void AddUserParameters(const TArray<FPropertyBagPropertyDesc>& InDescs, const UPCGGraph* InOptionalOriginalGraph = nullptr);
+	UE_API EPropertyBagAlterationResult AddUserParameters(const TArray<FPropertyBagPropertyDesc>& InDescs, const UPCGGraph* InOptionalOriginalGraph = nullptr);
 
 	UPROPERTY(BlueprintReadWrite, EditAnywhere, Category = "Runtime Generation")
 	FPCGRuntimeGenerationRadii GenerationRadii;
