[PCG] Added Text Overflow user preferences to ALV for strings, names, and hyperlink (soft object/class)

+ Some Minor Housekeeping

#rnx

diff --git a/Engine/Plugins/PCG/Source/PCGEditor/Private/Widgets/SPCGEditorGraphAttributeListView.cpp b/Engine/Plugins/PCG/Source/PCGEditor/Private/Widgets/SPCGEditorGraphAttributeListView.cpp
--- a/Engine/Plugins/PCG/Source/PCGEditor/Private/Widgets/SPCGEditorGraphAttributeListView.cpp
+++ b/Engine/Plugins/PCG/Source/PCGEditor/Private/Widgets/SPCGEditorGraphAttributeListView.cpp
@@ -20,6 +20,7 @@
 
 #include "PCGEditor.h"
 #include "PCGEditorGraph.h"
+#include "PCGEditorSettings.h"
 #include "PCGEditorStyle.h"
 #include "PCGEditorUtils.h"
 #include "DataVisualizations/PCGDataVisualizationHelpers.h"
@@ -140,13 +141,15 @@ TSharedRef<SWidget> SPCGListViewItemRow::GenerateWidgetForColumn(const FName& Co
 	// Make the text clickable for soft object paths
 	bool bIsClickable = false;
 
+	ETextOverflowPolicy OverflowPolicy = GetDefault<UPCGEditorSettings>()->AttributeListViewSettings.GeneralOverflowPolicy;
+
 	if (FPCGColumnData* PCGColumnData = SharedAttributeListView->PCGColumnData.Find(ColumnId))
 	{
 		// We have to make sure the data is still valid otherwise this will crash.
 		if (ensure(SharedAttributeListView->DataPtr) && PCGColumnData->DataAccessor.IsValid() && PCGColumnData->DataKeys.IsValid())
 		{
 			const int32 Index = InternalItem->Index;
-			auto Callback = [&PCGColumnData, Index, &RowText, &bShouldComputeCellWidth, &bIsClickable, ColumnId] (auto Dummy)
+			auto Callback = [&PCGColumnData, Index, &RowText, &bShouldComputeCellWidth, &bIsClickable, &OverflowPolicy, ColumnId] (auto Dummy)
 			{
 				using ValueType = decltype(Dummy);
 				ValueType Value = PCG::Private::MetadataTraits<ValueType>::ZeroValue();
@@ -168,27 +171,24 @@ TSharedRef<SWidget> SPCGListViewItemRow::GenerateWidgetForColumn(const FName& Co
 					{
 						RowText = FText::FromString(Value);
 						bShouldComputeCellWidth = true;
+						OverflowPolicy = GetDefault<UPCGEditorSettings>()->AttributeListViewSettings.StringOverflowPolicy;
 					}
 					else if constexpr (PCG::Private::IsOfTypes<ValueType, FName>())
 					{
 						RowText = FText::FromName(Value);
 						bShouldComputeCellWidth = true;
+						OverflowPolicy = GetDefault<UPCGEditorSettings>()->AttributeListViewSettings.NameOverflowPolicy;
 					}
 					else if constexpr (FTextAsNumberIsValid<ValueType>::value)
 					{
 						RowText = FText::AsNumber(Value);
 					}
-					else if constexpr (PCG::Private::IsOfTypes<ValueType, FSoftObjectPath>())
-					{
-						RowText = FText::FromString(Value.ToString());
-						bShouldComputeCellWidth = true;
-						bIsClickable = true;
-					}
-					else if constexpr (PCG::Private::IsOfTypes<ValueType, FSoftClassPath>())
+					else if constexpr (PCG::Private::IsOfTypes<ValueType, FSoftObjectPath, FSoftClassPath>())
 					{
 						RowText = FText::FromString(Value.ToString());
 						bShouldComputeCellWidth = true;
 						bIsClickable = true;
+						OverflowPolicy = GetDefault<UPCGEditorSettings>()->AttributeListViewSettings.HyperLinkOverflowPolicy;
 					}
 					else
 					{
@@ -213,7 +213,6 @@ TSharedRef<SWidget> SPCGListViewItemRow::GenerateWidgetForColumn(const FName& Co
 	}
 
 	const FMargin Margin = FMargin(2.0f, 0.0f);
-	const ETextOverflowPolicy OverflowPolicy = ETextOverflowPolicy::Ellipsis;
 
 	// TextStyles are pointers so they need to stay alive. Create 2 static copies for each of our style: Normal and Italic.
 	static const FTextBlockStyle NormalTextStyle = FCoreStyle::Get().GetWidgetStyle<FTextBlockStyle>("NormalText");
@@ -225,13 +224,14 @@ TSharedRef<SWidget> SPCGListViewItemRow::GenerateWidgetForColumn(const FName& Co
 
 	if (bIsClickable)
 	{
-		// @todo_pcg: We're missing Margin and Overflow policy with this widget, maybe something we can re-add in the future.
 		TSharedPtr<SHyperlink> Hyperlink;
 		SAssignNew(Hyperlink, SHyperlink)
 			.Text(RowText)
 			.ToolTipText(LOCTEXT("Hyperlink", "Ctrl + Click to jump to the asset."))
 			.Style(FAppStyle::Get(), "HoverOnlyHyperlink")
 			.TextStyle(TextStyle)
+			.Padding(Margin)
+			.OverflowPolicy(OverflowPolicy)
 			.OnNavigate(FSimpleDelegate::CreateLambda([Text = std::move(RowText)](){ SPCGListViewItemRow::OnSoftObjectPathHyperlinkClicked(Text); }));
 
 		return Hyperlink.ToSharedRef();
@@ -911,13 +911,12 @@ void SPCGEditorGraphAttributeListView::RefreshAttributeList()
 	if (const IPCGDataVisualization* DataVisualization = DataVisRegistry.GetDataVisualization(PCGData->GetClass()))
 	{
 		const FPCGTableVisualizerInfo TableVisualizerInfo = DataVisualization->GetTableVisualizerInfoWithDomain(PCGData, MetadataDomainID);
-		DataToVisualize = TableVisualizerInfo.Data;
 
 		int32 NumEntries = 0;
 
 		for (const FPCGTableVisualizerColumnInfo& ColumnInfo : TableVisualizerInfo.ColumnInfos)
 		{
-			AddColumn(DataToVisualize, ColumnInfo);
+			AddColumn(ColumnInfo);
 			NumEntries = FMath::Max(NumEntries, ColumnInfo.AccessorKeys.IsValid() ? ColumnInfo.AccessorKeys->GetNum() : 0);
 		}
 
@@ -1678,7 +1677,7 @@ FReply SPCGEditorGraphAttributeListView::OnListViewKeyDown(const FGeometry& /*In
 	return FReply::Unhandled();
 }
 
-void SPCGEditorGraphAttributeListView::AddColumn(const UPCGData* InPCGData, const FPCGTableVisualizerColumnInfo& InColumnInfo)
+void SPCGEditorGraphAttributeListView::AddColumn(const FPCGTableVisualizerColumnInfo& InColumnInfo)
 {
 	FPCGColumnData& ColumnData = PCGColumnData.Add(InColumnInfo.Id);
 	ColumnData.DataAccessor = InColumnInfo.Accessor;
@@ -1696,7 +1695,7 @@ void SPCGEditorGraphAttributeListView::AddColumn(const UPCGData* InPCGData, cons
 	Arguments.HAlignCell((EHorizontalAlignment)InColumnInfo.CellAlignment);
 	Arguments.SortMode(this, &SPCGEditorGraphAttributeListView::GetColumnSortMode, InColumnInfo.Id);
 	Arguments.OnSort(this, &SPCGEditorGraphAttributeListView::OnColumnSortModeChanged);
-	Arguments.OverflowPolicy(ETextOverflowPolicy::Ellipsis);
+	Arguments.OverflowPolicy(GetDefault<UPCGEditorSettings>()->AttributeListViewSettings.ColumnNameOverflowPolicy);
 	Arguments.HeaderComboVisibility(EHeaderComboVisibility::Never);
 	Arguments.MenuContent()
 	[
diff --git a/Engine/Plugins/PCG/Source/PCGEditor/Private/Widgets/SPCGEditorGraphAttributeListView.h b/Engine/Plugins/PCG/Source/PCGEditor/Private/Widgets/SPCGEditorGraphAttributeListView.h
--- a/Engine/Plugins/PCG/Source/PCGEditor/Private/Widgets/SPCGEditorGraphAttributeListView.h
+++ b/Engine/Plugins/PCG/Source/PCGEditor/Private/Widgets/SPCGEditorGraphAttributeListView.h
@@ -203,7 +203,6 @@ private:
 
 	const FPCGDataCollection* GetInspectionData(const TSharedPtr<FPinComboBoxItem>& EditorPin) const;
 	const FPCGDataCollection* GetInspectionData() const;
-	const UPCGData* GetInspectedData(bool& bOutNeedsRefresh) const;
 
 	FText OnGenerateSelectedPinText() const;
 	void OnSelectionChangedPin(TSharedPtr<FPinComboBoxItem> InItem, ESelectInfo::Type InSelectInfo);
@@ -244,9 +243,7 @@ private:
 
 	FReply OnListViewKeyDown(const FGeometry& InGeometry, const FKeyEvent& InKeyEvent) const;
 
-	void AddColumn(const UPCGData* InPCGData, const FPCGTableVisualizerColumnInfo& InColumnInfo);
-	void AddPointDataColumns(const UPCGPointData* InPCGPointData);
-	void AddMetadataColumn(const UPCGData* InPCGData, const FName& InColumnId, EPCGMetadataTypes InMetadataType, const TCHAR* PostFix = nullptr);
+	void AddColumn(const FPCGTableVisualizerColumnInfo& InColumnInfo);
 
 	void CopySelectionToClipboard() const;
 	bool CanCopySelectionToClipboard() const;
diff --git a/Engine/Plugins/PCG/Source/PCGEditor/Public/PCGEditorSettings.h b/Engine/Plugins/PCG/Source/PCGEditor/Public/PCGEditorSettings.h
--- a/Engine/Plugins/PCG/Source/PCGEditor/Public/PCGEditorSettings.h
+++ b/Engine/Plugins/PCG/Source/PCGEditor/Public/PCGEditorSettings.h
@@ -3,12 +3,11 @@
 #pragma once
 
 #include "Engine/DeveloperSettings.h"
+#include "Styling/SlateTypes.h"
 #include "Templates/SubclassOf.h"
 
 #include "PCGEditorSettings.generated.h"
 
-#define UE_API PCGEDITOR_API
-
 class UPCGSettings;
 class UPCGBuilderSettings;
 
@@ -29,8 +28,8 @@ class UPCGEditorProjectSettings : public UDeveloperSettings
 	GENERATED_BODY()
 
 public:
-	UE_API UPCGEditorProjectSettings(const FObjectInitializer& ObjectInitializer);
-		
+	PCGEDITOR_API UPCGEditorProjectSettings(const FObjectInitializer& ObjectInitializer);
+
 	/** Default Builder Settings to use when running the PCGWorldPartitionBuilder */
 	UPROPERTY(EditAnywhere, config, Category = Builder)
 	TSoftObjectPtr<UPCGBuilderSettings> DefaultBuilderSetting;
@@ -43,16 +42,42 @@ public:
 	bool bAutoRefreshGPUStaticMeshSpawners = true;
 };
 
+USTRUCT()
+struct FPCGEditorAttributeListViewSettings
+{
+	GENERATED_BODY()
+
+	/** Overflow policy for column headers in the Attribute List View. */
+	UPROPERTY(EditAnywhere, Config, Category = "Attribute List View")
+	ETextOverflowPolicy ColumnNameOverflowPolicy = ETextOverflowPolicy::Ellipsis;
+
+	/** Overflow policy for non-string types in the Attribute List View. */
+	UPROPERTY(EditAnywhere, Config, Category = "Attribute List View")
+	ETextOverflowPolicy GeneralOverflowPolicy = ETextOverflowPolicy::Ellipsis;
+
+	/** Overflow policy for 'FString' in the Attribute List View. */
+	UPROPERTY(EditAnywhere, Config, Category = "Attribute List View")
+	ETextOverflowPolicy StringOverflowPolicy = ETextOverflowPolicy::Ellipsis;
+
+	/** Overflow policy for 'FName' in the Attribute List View. */
+	UPROPERTY(EditAnywhere, Config, Category = "Attribute List View")
+	ETextOverflowPolicy NameOverflowPolicy = ETextOverflowPolicy::Ellipsis;
+
+	/** Overflow policy for 'FSoftObjectPath' and 'FSoftClassPath', which are HyperLinked in the Attribute List View. */
+	UPROPERTY(EditAnywhere, Config, Category = "Attribute List View")
+	ETextOverflowPolicy HyperLinkOverflowPolicy = ETextOverflowPolicy::MiddleEllipsis;
+};
+
 UCLASS(MinimalAPI, config=EditorPerProjectUserSettings)
 class UPCGEditorSettings : public UObject
 {
 	GENERATED_BODY()
 
 public:
-	UE_API UPCGEditorSettings(const FObjectInitializer& ObjectInitializer);
+	PCGEDITOR_API UPCGEditorSettings(const FObjectInitializer& ObjectInitializer);
 
-	UE_API FLinearColor GetColor(UPCGSettings* InSettings) const;
-	UE_API FLinearColor GetPinColor(const FEdGraphPinType& InPinType) const;
+	PCGEDITOR_API FLinearColor GetColor(UPCGSettings* InSettings) const;
+	PCGEDITOR_API FLinearColor GetPinColor(const FEdGraphPinType& InPinType) const;
 
 	/** Default node color */
 	UPROPERTY(EditAnywhere, config, Category = Node, meta = (HideAlphaChannel))
@@ -282,6 +307,10 @@ public:
 	UPROPERTY(EditAnywhere, config, Category = "Graph UI")
 	bool bShowNodeDeterminismActionsRightClickContext = true;
 
+	/** Attribute List View user configurations. */
+	UPROPERTY(EditAnywhere, Config, Category = "Attribute List View", meta = (ShowOnlyInnerProperties))
+	FPCGEditorAttributeListViewSettings AttributeListViewSettings;
+
 	/** When inspection updates the viewport will re-focus to the data. */
 	UPROPERTY(EditAnywhere, config, Category = "Viewport")
 	bool bAutoFocusViewport = true;
@@ -291,5 +320,3 @@ public:
 	UPROPERTY()
 	bool bEnableNavigateToNativeNodes = true;
 };
-
-#undef UE_API
