[PCG] Properly support undo/redo of CurrentEditingMode
- UPCGComponent::ChangeTransientState: suppress transaction with GUndo nulling
- Handle undo/redo of CurrentEditingMode through UPCGComponent::PreEditUndo/PostEditUndo
- Capture instance data properly for now undoable CurrentEditingMode and do the same for the other non serialized members

diff --git a/Engine/Plugins/PCG/Source/PCG/Private/PCGComponent.cpp b/Engine/Plugins/PCG/Source/PCG/Private/PCGComponent.cpp
--- a/Engine/Plugins/PCG/Source/PCG/Private/PCGComponent.cpp
+++ b/Engine/Plugins/PCG/Source/PCG/Private/PCGComponent.cpp
@@ -2570,6 +2570,10 @@ void UPCGComponent::PreEditUndo()
 	// Here we will keep a copy of flags that we require to keep through the undo
 	// so we can have a consistent state
 	LastGeneratedBoundsPriorToUndo = LastGeneratedBounds;
+	
+	// Capture CurrentEditingMode before it potentially changes due to undo/redo so we can compare against in PostEditUndo.
+	// This only applies to non-BP PCGComponents because PreEditUndo/PostEditUndo will not be called in BP paths.
+	LastEditingModePriorToUndo = CurrentEditingMode;
 
 	if (IsManagedByRuntimeGenSystem())
 	{
@@ -2593,6 +2597,13 @@ void UPCGComponent::PostEditUndo()
 {
 	LastGeneratedBounds = LastGeneratedBoundsPriorToUndo;
 
+	// Handle undo/redo of CurrentEditingMode
+	if (LastEditingModePriorToUndo != CurrentEditingMode)
+	{
+		ChangeTransientState(CurrentEditingMode);
+		LastEditingModePriorToUndo = CurrentEditingMode;
+	}
+
 	UpdateTrackingCache();
 	DirtyGenerated(EPCGComponentDirtyFlag::All);
 
@@ -3885,6 +3896,10 @@ void UPCGComponent::MarkResourcesAsTransientOnLoad()
 
 void UPCGComponent::ChangeTransientState(EPCGEditorDirtyMode NewEditingMode)
 {
+	// Changing the transient state can and will play with packages and those changes are not meant to be tracked by the transaction buffer.
+	// Instead, UPCGComponent will react to CurrentEditingMode changes between PreEditUndo/PostEditUndo and call ChangeTransientState accordingly to reflect the change to CurrentEditingMode.
+	TGuardValue<ITransaction*> SuppressTransaction(GUndo, nullptr);
+
 	bool bShouldMarkDirty = false;
 
 	// Affect all resources
@@ -3957,22 +3972,6 @@ void UPCGComponent::ChangeTransientState(EPCGEditorDirtyMode NewEditingMode)
 			Subsystem->PropagateEditingModeToLocalComponents(this, NewEditingMode);
 		}
 	}
-
-	// Changing the transient state can and will play with packages and is not meant to be undoable.
-	// However, don't kill the transaction buffer unless it's actually needed, and FIXME: especially not during a construction script
-	if (bShouldMarkDirty && GEditor && GEditor->Trans 
-		// TODO: This is a temporary fix- we check GUndo here in part to guard against being called out of PostLoad, 
-		//  where clearing the transaction buffer will cause a crash, as upstream code sets GUndo to null and then 
-		//  later resets it back to something we deleted.
-		&& GUndo)
-	{
-		UWorld* World = GetWorld();
-
-		if (!World || !World->bIsRunningConstructionScript)
-		{
-			GEditor->Trans->Reset(LOCTEXT("ChangeEditingMode", "Changing Editing Mode"));
-		}
-	}
 }
 
 bool UPCGComponent::GetStackContext(FPCGStackContext& OutStackContext) const
@@ -4011,7 +4010,22 @@ TArray<TSoftObjectPtr<AActor>> UPCGComponent::GetManagedActorPaths(AActor* InAct
 FPCGComponentInstanceData::FPCGComponentInstanceData(const UPCGComponent* InSourceComponent)
 	: FActorComponentInstanceData(InSourceComponent)
 	, SourceComponent(InSourceComponent)
+#if WITH_EDITOR
+	, CurrentEditingMode(InSourceComponent->CurrentEditingMode)
+	, PreviousEditingMode(InSourceComponent->PreviousEditingMode)
+	, bForceGenerateOnBPAddedToWorld(InSourceComponent->bForceGenerateOnBPAddedToWorld)
+#endif
 {
+#if WITH_EDITOR
+	// Convert map to array for serialization
+	for (const auto& KeyValuePair : InSourceComponent->DynamicallyTrackedKeysToSettings)
+	{
+		for (const auto& SettingsAndCullingPair : KeyValuePair.Value)
+		{
+			DynamicallyTrackedKeysToSettings.Add({ KeyValuePair.Key, SettingsAndCullingPair.Key, SettingsAndCullingPair.Value });
+		}
+	}
+#endif
 }
 
 bool FPCGComponentInstanceData::ContainsData() const
@@ -4043,26 +4057,23 @@ void FPCGComponentInstanceData::ApplyToComponent(UActorComponent* Component, con
 				ConstructionSourceComponent = FoundConstructionSourceComponent;
 			}
 		}
-#endif
-		// IMPORTANT NOTE:
-		// Any non-visible (i.e. UPROPERTY() with no specifiers) are NOT copied over when re-running the construction script
-		// This means that some properties need to be reapplied here manually unless we make them visible
-		if (SourceComponent)
-		{
-			// No need to copy PCGComponent->GeneratedGraphOutput anymore as it is now a visible property
-
-#if WITH_EDITOR
-			// While this is serialized, it is not properly copied because it is not visible. This is needed otherwise a refresh can retoggle from not generated to generated
-			PCGComponent->bForceGenerateOnBPAddedToWorld = SourceComponent->bForceGenerateOnBPAddedToWorld;
 
-			PCGComponent->CurrentEditingMode = SourceComponent->CurrentEditingMode;
-			PCGComponent->PreviousEditingMode = SourceComponent->PreviousEditingMode;
-			PCGComponent->DynamicallyTrackedKeysToSettings = SourceComponent->DynamicallyTrackedKeysToSettings;
-#endif // WITH_EDITOR
+		// Any members:
+		// 
+		// - UPROPERTY() with no specifiers
+		// - Transient UPROPERTY
+		// - non UPROPERTY
+		// 
+		// are NOT copied over when re-running the construction script and need to be handled by the instance data so we reapply them here.
+		PCGComponent->bForceGenerateOnBPAddedToWorld = bForceGenerateOnBPAddedToWorld;
+		PCGComponent->PreviousEditingMode = PreviousEditingMode;	
+		PCGComponent->CurrentEditingMode = CurrentEditingMode;
 
-			// Non-critical but should be done: transient data, tracked actors cache, landscape tracking
-			// TODO Validate usefulness + move accordingly
+		for (const FDynamicTrackedKeyInstanceData& DynamicTrackedKeyInstanceData : DynamicallyTrackedKeysToSettings)
+		{
+			PCGComponent->DynamicallyTrackedKeysToSettings.FindOrAdd(DynamicTrackedKeyInstanceData.SelectionKey).Add({ DynamicTrackedKeyInstanceData.Settings, DynamicTrackedKeyInstanceData.bValue });
 		}
+#endif
 
 		if (ConstructionSourceComponent)
 		{
@@ -4162,6 +4173,9 @@ void FPCGComponentInstanceData::ApplyToComponent(UActorComponent* Component, con
 			const_cast<UPCGComponent*>(ConstructionSourceComponent)->MarkSubObjectsAsGarbage();
 		}
 
+		// Make sure to apply latest editing mode
+		PCGComponent->ChangeTransientState(PCGComponent->CurrentEditingMode);
+
 		// Finally, start a delayed refresh task (if there is not one already), in editor only
 		// It is important to be delayed, because we cannot spawn Partition Actors within this scope,
 		// because we are in a construction script.
diff --git a/Engine/Plugins/PCG/Source/PCG/Public/PCGComponent.h b/Engine/Plugins/PCG/Source/PCG/Public/PCGComponent.h
--- a/Engine/Plugins/PCG/Source/PCG/Public/PCGComponent.h
+++ b/Engine/Plugins/PCG/Source/PCG/Public/PCGComponent.h
@@ -612,6 +612,9 @@ protected:
 	/** Used to store the CurrentEditingMode when it is forcefully changed by another system, such as runtime generation. */
 	EPCGEditorDirtyMode PreviousEditingMode = EPCGEditorDirtyMode::Normal;
 
+	/** Used to compare pre/post undo to know if we need to call ChangeTransientState. */
+	EPCGEditorDirtyMode LastEditingModePriorToUndo = EPCGEditorDirtyMode::Normal;
+
 public:
 	UPROPERTY(BlueprintReadWrite, EditAnywhere, Category = "Input Node Settings (Deprecated)", meta = (DisplayPriority = 800))
 	EPCGComponentInput InputType = EPCGComponentInput::Actor;
@@ -814,6 +817,21 @@ private:
 	FPCGComponentExecutionState ExecutionState;
 };
 
+USTRUCT()
+struct FDynamicTrackedKeyInstanceData
+{
+	GENERATED_BODY();
+
+	UPROPERTY()
+	FPCGSelectionKey SelectionKey;
+
+	UPROPERTY()
+	TSoftObjectPtr<const UPCGSettings> Settings;
+
+	UPROPERTY()
+	bool bValue;
+};
+
 /** Used to store generated resources data during RerunConstructionScripts */
 USTRUCT()
 struct FPCGComponentInstanceData : public FActorComponentInstanceData
@@ -830,6 +848,29 @@ protected:
 
 	UPROPERTY()
 	TObjectPtr<const UPCGComponent> SourceComponent;
+
+#if WITH_EDITORONLY_DATA
+	// We could rely on the SourceComponent to re-apply the non editable members but to simplify future code we capture
+	// the values in the FPCGComponentInstanceData constructor so we do not depend on the state of the SourceComponent anymore.
+	// Editable properties like: CurrentEditingMode need to be captures because in the case of user edits, the SourceComponent will actually be modified and we would lose the original value.
+
+	// UPCGComponent::CurrentEditingMode is a Transient property so we need to capture/apply it ourselves
+	UPROPERTY()
+	EPCGEditorDirtyMode CurrentEditingMode;
+	
+	// UPCGComponent::PreviousEditingMode is not a property so we need to capture/apply it ourselves
+	UPROPERTY()
+	EPCGEditorDirtyMode PreviousEditingMode;
+
+	// UPCGComponent::bForceGenerateOnBPAddedToWorld is not a visible property so we need to capture/apply it ourselves
+	UPROPERTY()
+	bool bForceGenerateOnBPAddedToWorld = false;
+
+	// UPCGComponent::DynamicallyTrackedKeysToSettings is not a property so we need to capture/apply it ourselves.
+	// To avoid having to write custom serialization code the map is converted into an array for serialization.
+	UPROPERTY()
+	TArray<FDynamicTrackedKeyInstanceData> DynamicallyTrackedKeysToSettings;
+#endif
 };
 
 #undef UE_API
