[PCG] Support single value ranges in GPU buffers

#virtualized

diff --git a/Engine/Plugins/PCG/Shaders/Private/PCGDataCollectionDataInterface.ush b/Engine/Plugins/PCG/Shaders/Private/PCGDataCollectionDataInterface.ush
--- a/Engine/Plugins/PCG/Shaders/Private/PCGDataCollectionDataInterface.ush
+++ b/Engine/Plugins/PCG/Shaders/Private/PCGDataCollectionDataInterface.ush
@@ -3,19 +3,23 @@
 // Buffer format:
 //
 // -- START HEADER --
-// PackedExecutionFlagAndNumData (4 bytes)
-// For each data:
-//   TypeId (4 bytes)
-//   Attribute Count (4 bytes)
-//   Element Count (4 bytes)
-//   For attribute ID 0 to MAX_NUM_ATTR:
-//     PackedIdAndStride (4 bytes)
-//     AttributeElementStartAddress (4 bytes)
+// PackedExecutionFlagAndNumData (4B)
+// For each data (1036B when MAX_NUM_ATTR=128):
+//   TypeId (4B)
+//   Attribute Count (4B)
+//   Element Count (4B)
+//   For attribute ID in [0, MAX_NUM_ATTR-1]:
+//     PackedAttributeInfo (4B)
+//         Stride (8b)
+//         Id (23b)
+//         Allocated (1b - MSB)
+//     AttributeElementStartAddress (4B)
 // -- END HEADER --
 // For each data:
 //   For each present attribute:
-//     <Packed element data> (element count * stride bytes)
+//     <Packed element data> (element count * stride bytes) - Note: element count may be 1 if the attribute is not 'Allocated' (that is, uses a single value for the entire range).
 
+// @todo_pcg: Support having a non-RW byte address buffer for input data.
 RWByteAddressBuffer {DataInterfaceName}_DataCollectionBuffer;
 
 // Multiplier that can be configured on output pins.
@@ -29,6 +33,26 @@ StructuredBuffer<int> {DataInterfaceName}_BufferToGraphStringKey;
 
 int {DataInterfaceName}_NumRemappedStringKeys;
 
+struct FPCGAttributeHeader_{DataInterfaceName}
+{
+	/** Absolute address of this attribute's first data element in the buffer. */
+	uint FirstElementAddress;
+
+	/** Stride (in bytes) of this attribute. */
+	uint Stride;
+
+	/** Whether or not this attribute is full-width allocated. */
+	bool bAllocated;
+};
+
+// #################### SHADER PARAM GETTERS ##########################
+
+// @todo_pcg: This parameter is broader than this data interface. Needs to be shared in some way (e.g. with TextureDataInterface).
+uint GetElementCountMultiplier_{DataInterfaceName}()
+{
+	return {DataInterfaceName}_ElementCountMultiplier;
+}
+
 // #################### HEADER READERS ##########################
 
 uint GetNumData_{DataInterfaceName}()
@@ -86,17 +110,17 @@ uint GetNumElements_{DataInterfaceName}()
 
 // #################### INTERNAL HELPERS ##########################
 
-uint LoadBufferInternal_{DataInterfaceName}(uint Address)
+uint LoadBufferInternal_{DataInterfaceName}(uint InAddress)
 {
-	return {DataInterfaceName}_DataCollectionBuffer.Load(Address);
+	return {DataInterfaceName}_DataCollectionBuffer.Load(InAddress);
 }
 
-void StoreBufferInternal_{DataInterfaceName}(uint Address, uint Value)
+void StoreBufferInternal_{DataInterfaceName}(uint InAddress, uint InValue)
 {
-	{DataInterfaceName}_DataCollectionBuffer.Store(Address, Value);
+	{DataInterfaceName}_DataCollectionBuffer.Store(InAddress, InValue);
 }
 
-int GetRemappedAttributeId_{DataInterfaceName}(int InGraphAttributeId)
+int GetRemappedAttributeIdInternal_{DataInterfaceName}(int InGraphAttributeId)
 {
 	if (InGraphAttributeId >= {DataInterfaceName}_FirstRemappedAttributeId)
 	{
@@ -110,50 +134,48 @@ int GetRemappedAttributeId_{DataInterfaceName}(int InGraphAttributeId)
 	}
 }
 
-uint GetFirstElementAddressInternal_{DataInterfaceName}(uint InDataIndex, int InAttributeId)
+FPCGAttributeHeader_{DataInterfaceName} GetAttributeHeaderInternal_{DataInterfaceName}(uint InDataIndex, int InAttributeId)
 {
-	const int AttributeId = GetRemappedAttributeId_{DataInterfaceName}(InAttributeId);
+	const int AttributeId = GetRemappedAttributeIdInternal_{DataInterfaceName}(InAttributeId);
 
 	const uint ReadAddress =
 		PCG_DATA_COLLECTION_HEADER_SIZE_BYTES +
 		InDataIndex * PCG_DATA_HEADER_SIZE_BYTES +
 		/*TypeId*/4 + /*Attribute Count*/4 + /*Element Count*/4 +
-		AttributeId * PCG_ATTRIBUTE_HEADER_SIZE_BYTES +
-		/*PackedAttributeIdAndStride*/4;
+		AttributeId * PCG_ATTRIBUTE_HEADER_SIZE_BYTES;
 
-	return {DataInterfaceName}_DataCollectionBuffer.Load(ReadAddress);
+	const uint2 PackedAttributeHeader = {DataInterfaceName}_DataCollectionBuffer.Load2(ReadAddress);
+
+	FPCGAttributeHeader_{DataInterfaceName} AttributeHeader;
+	AttributeHeader.FirstElementAddress = PackedAttributeHeader.y;
+	AttributeHeader.Stride = PackedAttributeHeader.x & PCG_ATTRIBUTE_STRIDE_MASK;
+	AttributeHeader.bAllocated = (PackedAttributeHeader.x & PCG_ATTRIBUTE_ALLOCATED_MASK) != 0;
+
+	return AttributeHeader;
 }
 
-uint GetAttributeStrideInternal_{DataInterfaceName}(uint InDataIndex, int InAttributeId)
+uint2 GetElementAddressAndStrideInternal_{DataInterfaceName}(uint InDataIndex, uint InElementIndex, int InAttributeId)
 {
-	const int AttributeId = GetRemappedAttributeId_{DataInterfaceName}(InAttributeId);
+	const FPCGAttributeHeader_{DataInterfaceName} AttributeHeader = GetAttributeHeaderInternal_{DataInterfaceName}(InDataIndex, InAttributeId);
+	const uint ElementIndex = AttributeHeader.bAllocated ? InElementIndex : 0u;
 
-	const uint ReadAddress =
-		PCG_DATA_COLLECTION_HEADER_SIZE_BYTES +
-		InDataIndex * PCG_DATA_HEADER_SIZE_BYTES +
-		/*TypeId*/4 + /*Attribute Count*/4 + /*Element Count*/4 +
-		AttributeId * PCG_ATTRIBUTE_HEADER_SIZE_BYTES;
-
-	const uint PackedAttributeIdAndStride = {DataInterfaceName}_DataCollectionBuffer.Load(ReadAddress);
-	return PackedAttributeIdAndStride & (0xFFFFFFFF >> 24);
+	return uint2(ElementIndex * AttributeHeader.Stride + AttributeHeader.FirstElementAddress, AttributeHeader.Stride);
 }
 
 uint GetElementAddressInternal_{DataInterfaceName}(uint InDataIndex, uint InElementIndex, int InAttributeId)
 {
-	const uint FirstElementAddress = GetFirstElementAddressInternal_{DataInterfaceName}(InDataIndex, InAttributeId);
-	const uint AttributeStride = GetAttributeStrideInternal_{DataInterfaceName}(InDataIndex, InAttributeId);
-	return InElementIndex * AttributeStride + FirstElementAddress;
+	return GetElementAddressAndStrideInternal_{DataInterfaceName}(InDataIndex, InElementIndex, InAttributeId).x;
 }
 
-uint GetFirstAttributeAddressInternal_{DataInterfaceName}()
+uint GetFirstValidAttributeElementInternal_{DataInterfaceName}()
 {
 	for (int AttributeId = 0; AttributeId < PCG_MAX_NUM_ATTRS; ++AttributeId)
 	{
-		const uint AttributeAddress = GetFirstElementAddressInternal_{DataInterfaceName}(/*DataIndex=*/0, AttributeId);
+		const uint ElementAddress = GetElementAddressInternal_{DataInterfaceName}(/*DataIndex=*/0, /*ElementIndex=*/0, AttributeId);
 
-		if (AttributeAddress > 0)
+		if (ElementAddress > 0)
 		{
-			return AttributeAddress;
+			return ElementAddress;
 		}
 	}
 
@@ -179,239 +201,243 @@ void SetAsExecutedInternal_{DataInterfaceName}()
 
 // #################### ATTRIBUTE GETTERS ##########################
 
-bool GetBool_{DataInterfaceName}(uint DataIndex, uint ElementIndex, int AttributeId)
+bool GetBool_{DataInterfaceName}(uint InDataIndex, uint InElementIndex, int InAttributeId)
 {
-	const uint AttributeAddress = GetFirstElementAddressInternal_{DataInterfaceName}(DataIndex, AttributeId);
-	return AttributeAddress > 0 && {DataInterfaceName}_DataCollectionBuffer.Load(AttributeAddress + ElementIndex * 1 * 4) != 0;
+	const uint ElementAddress = GetElementAddressInternal_{DataInterfaceName}(InDataIndex, InElementIndex, InAttributeId);
+	return ElementAddress > 0 && {DataInterfaceName}_DataCollectionBuffer.Load(ElementAddress) != 0;
 }
 
-int GetInt_{DataInterfaceName}(uint DataIndex, uint ElementIndex, int AttributeId)
+int GetInt_{DataInterfaceName}(uint InDataIndex, uint InElementIndex, int InAttributeId)
 {
-	const uint AttributeAddress = GetFirstElementAddressInternal_{DataInterfaceName}(DataIndex, AttributeId);
-	return AttributeAddress > 0 ? asint({DataInterfaceName}_DataCollectionBuffer.Load(AttributeAddress + ElementIndex * 1 * 4)) : 0;
+	const uint ElementAddress = GetElementAddressInternal_{DataInterfaceName}(InDataIndex, InElementIndex, InAttributeId);
+	return ElementAddress > 0 ? asint({DataInterfaceName}_DataCollectionBuffer.Load(ElementAddress)) : 0;
 }
 
-int GetUint_{DataInterfaceName}(uint DataIndex, uint ElementIndex, int AttributeId)
+int GetUint_{DataInterfaceName}(uint InDataIndex, uint InElementIndex, int InAttributeId)
 {
-	const uint AttributeAddress = GetFirstElementAddressInternal_{DataInterfaceName}(DataIndex, AttributeId);
-	return AttributeAddress > 0 ? {DataInterfaceName}_DataCollectionBuffer.Load(AttributeAddress + ElementIndex * 1 * 4) : 0;
+	const uint ElementAddress = GetElementAddressInternal_{DataInterfaceName}(InDataIndex, InElementIndex, InAttributeId);
+	return ElementAddress > 0 ? {DataInterfaceName}_DataCollectionBuffer.Load(ElementAddress) : 0;
 }
 
-float GetFloat_{DataInterfaceName}(uint DataIndex, uint ElementIndex, int AttributeId)
+float GetFloat_{DataInterfaceName}(uint InDataIndex, uint InElementIndex, int InAttributeId)
 {
-	const uint AttributeAddress = GetFirstElementAddressInternal_{DataInterfaceName}(DataIndex, AttributeId);
-	return AttributeAddress > 0 ? asfloat({DataInterfaceName}_DataCollectionBuffer.Load(AttributeAddress + ElementIndex * 1 * 4)) : 0.0f;
+	const uint ElementAddress = GetElementAddressInternal_{DataInterfaceName}(InDataIndex, InElementIndex, InAttributeId);
+	return ElementAddress > 0 ? asfloat({DataInterfaceName}_DataCollectionBuffer.Load(ElementAddress)) : 0.0f;
 }
 
-float2 GetFloat2_{DataInterfaceName}(uint DataIndex, uint ElementIndex, int AttributeId)
+float2 GetFloat2_{DataInterfaceName}(uint InDataIndex, uint InElementIndex, int InAttributeId)
 {
-	const uint AttributeAddress = GetFirstElementAddressInternal_{DataInterfaceName}(DataIndex, AttributeId);
-	return AttributeAddress > 0 ? asfloat({DataInterfaceName}_DataCollectionBuffer.Load2(AttributeAddress + ElementIndex * 2 * 4)) : (float2)0;
+	const uint ElementAddress = GetElementAddressInternal_{DataInterfaceName}(InDataIndex, InElementIndex, InAttributeId);
+	return ElementAddress > 0 ? asfloat({DataInterfaceName}_DataCollectionBuffer.Load2(ElementAddress)) : (float2)0;
 }
 
-float3 GetFloat3_{DataInterfaceName}(uint DataIndex, uint ElementIndex, int AttributeId)
+float3 GetFloat3_{DataInterfaceName}(uint InDataIndex, uint InElementIndex, int InAttributeId)
 {
-	const uint AttributeAddress = GetFirstElementAddressInternal_{DataInterfaceName}(DataIndex, AttributeId);
-	return AttributeAddress > 0 ? asfloat({DataInterfaceName}_DataCollectionBuffer.Load3(AttributeAddress + ElementIndex * 3 * 4)) : (float3)0;
+	const uint ElementAddress = GetElementAddressInternal_{DataInterfaceName}(InDataIndex, InElementIndex, InAttributeId);
+	return ElementAddress > 0 ? asfloat({DataInterfaceName}_DataCollectionBuffer.Load3(ElementAddress)) : (float3)0;
 }
 
-float4 GetFloat4_{DataInterfaceName}(uint DataIndex, uint ElementIndex, int AttributeId)
+float4 GetFloat4_{DataInterfaceName}(uint InDataIndex, uint InElementIndex, int InAttributeId)
 {
-	const uint AttributeAddress = GetFirstElementAddressInternal_{DataInterfaceName}(DataIndex, AttributeId);
-	return AttributeAddress > 0 ? asfloat({DataInterfaceName}_DataCollectionBuffer.Load4(AttributeAddress + ElementIndex * 4 * 4)) : (float4)0;
+	const uint ElementAddress = GetElementAddressInternal_{DataInterfaceName}(InDataIndex, InElementIndex, InAttributeId);
+	return ElementAddress > 0 ? asfloat({DataInterfaceName}_DataCollectionBuffer.Load4(ElementAddress)) : (float4)0;
 }
 
-float3 GetRotator_{DataInterfaceName}(uint DataIndex, uint ElementIndex, int AttributeId)
+float3 GetRotator_{DataInterfaceName}(uint InDataIndex, uint InElementIndex, int InAttributeId)
 {
-	return GetFloat3_{DataInterfaceName}(DataIndex, ElementIndex, AttributeId);
+	return GetFloat3_{DataInterfaceName}(InDataIndex, InElementIndex, InAttributeId);
 }
 
-float4 GetQuat_{DataInterfaceName}(uint DataIndex, uint ElementIndex, int AttributeId)
+float4 GetQuat_{DataInterfaceName}(uint InDataIndex, uint InElementIndex, int InAttributeId)
 {
-	return GetFloat4_{DataInterfaceName}(DataIndex, ElementIndex, AttributeId);
+	return GetFloat4_{DataInterfaceName}(InDataIndex, InElementIndex, InAttributeId);
 }
 
-float4x4 GetTransform_{DataInterfaceName}(uint DataIndex, uint ElementIndex, int AttributeId)
+float4x4 GetTransform_{DataInterfaceName}(uint InDataIndex, uint InElementIndex, int InAttributeId)
 {
-	const uint AttributeAddress = GetFirstElementAddressInternal_{DataInterfaceName}(DataIndex, AttributeId);
-	if (AttributeAddress == 0)
+	const uint ElementAddress = GetElementAddressInternal_{DataInterfaceName}(InDataIndex, InElementIndex, InAttributeId);
+
+	if (ElementAddress > 0)
+	{
+		return float4x4(
+			asfloat({DataInterfaceName}_DataCollectionBuffer.Load4(ElementAddress)),
+			asfloat({DataInterfaceName}_DataCollectionBuffer.Load4(ElementAddress + 4 * 4)),
+			asfloat({DataInterfaceName}_DataCollectionBuffer.Load4(ElementAddress + 8 * 4)),
+			asfloat({DataInterfaceName}_DataCollectionBuffer.Load4(ElementAddress + 12 * 4)));
+	}
+	else
 	{
 		return (float4x4)0;
 	}
-
-	const uint StartAddress = ElementIndex * 16 * 4 + AttributeAddress;
-
-	return float4x4(
-		asfloat({DataInterfaceName}_DataCollectionBuffer.Load4(StartAddress)),
-		asfloat({DataInterfaceName}_DataCollectionBuffer.Load4(StartAddress + 4 * 4)),
-		asfloat({DataInterfaceName}_DataCollectionBuffer.Load4(StartAddress + 8 * 4)),
-		asfloat({DataInterfaceName}_DataCollectionBuffer.Load4(StartAddress + 12 * 4))
-	);
 }
 
-int GetStringKey_{DataInterfaceName}(uint DataIndex, uint ElementIndex, int AttributeId)
+int GetStringKey_{DataInterfaceName}(uint InDataIndex, uint InElementIndex, int InAttributeId)
 {
 	// String keys are represented as ints.
-	const int StringKey = GetInt_{DataInterfaceName}(DataIndex, ElementIndex, AttributeId);
+	const int StringKey = GetInt_{DataInterfaceName}(InDataIndex, InElementIndex, InAttributeId);
 	return (StringKey < {DataInterfaceName}_NumRemappedStringKeys) ? {DataInterfaceName}_BufferToGraphStringKey[StringKey] : StringKey;
 }
 
 // TODO: left as uint2 which we locked into for 5.5.0. Future version should change this to int.
-uint2 GetName_{DataInterfaceName}(uint DataIndex, uint ElementIndex, int AttributeId)
+uint2 GetName_{DataInterfaceName}(uint InDataIndex, uint InElementIndex, int InAttributeId)
 {
 	// Names are represented as ints.
-	return uint2((uint)GetInt_{DataInterfaceName}(DataIndex, ElementIndex, AttributeId), 0u);
-}
-
-uint GetElementCountMultiplier_{DataInterfaceName}()
-{
-	return {DataInterfaceName}_ElementCountMultiplier;
+	return uint2((uint)GetInt_{DataInterfaceName}(InDataIndex, InElementIndex, InAttributeId), 0u);
 }
 
 // #################### ATTRIBUTE SETTERS ##########################
 
-void SetBool_{DataInterfaceName}(uint DataIndex, uint ElementIndex, int AttributeId, bool Value)
+void SetBool_{DataInterfaceName}(uint InDataIndex, uint InElementIndex, int InAttributeId, bool InValue)
 {
-	const uint AttributeAddress = GetFirstElementAddressInternal_{DataInterfaceName}(DataIndex, AttributeId);
-	if (AttributeAddress > 0)
+	const uint ElementAddress = GetElementAddressInternal_{DataInterfaceName}(InDataIndex, InElementIndex, InAttributeId);
+
+	if (ElementAddress > 0)
 	{
-		{DataInterfaceName}_DataCollectionBuffer.Store(AttributeAddress + ElementIndex * 1 * 4, Value ? 1u : 0u);
+		{DataInterfaceName}_DataCollectionBuffer.Store(ElementAddress, InValue ? 1u : 0u);
 	}
 }
 
-void SetInt_{DataInterfaceName}(uint DataIndex, uint ElementIndex, int AttributeId, int Value)
+void SetInt_{DataInterfaceName}(uint InDataIndex, uint InElementIndex, int InAttributeId, int InValue)
 {
-	const uint AttributeAddress = GetFirstElementAddressInternal_{DataInterfaceName}(DataIndex, AttributeId);
-	if (AttributeAddress > 0)
+	const uint ElementAddress = GetElementAddressInternal_{DataInterfaceName}(InDataIndex, InElementIndex, InAttributeId);
+
+	if (ElementAddress > 0)
 	{
-		{DataInterfaceName}_DataCollectionBuffer.Store(AttributeAddress + ElementIndex * 1 * 4, asuint(Value));
+		{DataInterfaceName}_DataCollectionBuffer.Store(ElementAddress, asuint(InValue));
 	}
 }
 
-void SetUint_{DataInterfaceName}(uint DataIndex, uint ElementIndex, int AttributeId, uint Value)
+void SetUint_{DataInterfaceName}(uint InDataIndex, uint InElementIndex, int InAttributeId, uint InValue)
 {
-	const uint AttributeAddress = GetFirstElementAddressInternal_{DataInterfaceName}(DataIndex, AttributeId);
-	if (AttributeAddress > 0)
+	const uint ElementAddress = GetElementAddressInternal_{DataInterfaceName}(InDataIndex, InElementIndex, InAttributeId);
+
+	if (ElementAddress > 0)
 	{
-		{DataInterfaceName}_DataCollectionBuffer.Store(AttributeAddress + ElementIndex * 1 * 4, Value);
+		{DataInterfaceName}_DataCollectionBuffer.Store(ElementAddress, InValue);
 	}
 }
 
-void SetFloat_{DataInterfaceName}(uint DataIndex, uint ElementIndex, int AttributeId, float Value)
+void SetFloat_{DataInterfaceName}(uint InDataIndex, uint InElementIndex, int InAttributeId, float InValue)
 {
-	const uint AttributeAddress = GetFirstElementAddressInternal_{DataInterfaceName}(DataIndex, AttributeId);
-	if (AttributeAddress > 0)
+	const uint ElementAddress = GetElementAddressInternal_{DataInterfaceName}(InDataIndex, InElementIndex, InAttributeId);
+
+	if (ElementAddress > 0)
 	{
-		{DataInterfaceName}_DataCollectionBuffer.Store(AttributeAddress + ElementIndex * 1 * 4, asuint(Value));
+		{DataInterfaceName}_DataCollectionBuffer.Store(ElementAddress, asuint(InValue));
 	}
 }
 
-void SetFloat2_{DataInterfaceName}(uint DataIndex, uint ElementIndex, int AttributeId, float2 Value)
+void SetFloat2_{DataInterfaceName}(uint InDataIndex, uint InElementIndex, int InAttributeId, float2 InValue)
 {
-	const uint AttributeAddress = GetFirstElementAddressInternal_{DataInterfaceName}(DataIndex, AttributeId);
-	if (AttributeAddress > 0)
+	const uint ElementAddress = GetElementAddressInternal_{DataInterfaceName}(InDataIndex, InElementIndex, InAttributeId);
+
+	if (ElementAddress > 0)
 	{
-		{DataInterfaceName}_DataCollectionBuffer.Store2(AttributeAddress + ElementIndex * 2 * 4, asuint(Value));
+		{DataInterfaceName}_DataCollectionBuffer.Store2(ElementAddress, asuint(InValue));
 	}
 }
 
-void SetFloat3_{DataInterfaceName}(uint DataIndex, uint ElementIndex, int AttributeId, float3 Value)
+void SetFloat3_{DataInterfaceName}(uint InDataIndex, uint InElementIndex, int InAttributeId, float3 InValue)
 {
-	const uint AttributeAddress = GetFirstElementAddressInternal_{DataInterfaceName}(DataIndex, AttributeId);
-	if (AttributeAddress > 0)
+	const uint ElementAddress = GetElementAddressInternal_{DataInterfaceName}(InDataIndex, InElementIndex, InAttributeId);
+
+	if (ElementAddress > 0)
 	{
-		{DataInterfaceName}_DataCollectionBuffer.Store3(AttributeAddress + ElementIndex * 3 * 4, asuint(Value));
+		{DataInterfaceName}_DataCollectionBuffer.Store3(ElementAddress, asuint(InValue));
 	}
 }
 
-void SetFloat4_{DataInterfaceName}(uint DataIndex, uint ElementIndex, int AttributeId, float4 Value)
+void SetFloat4_{DataInterfaceName}(uint InDataIndex, uint InElementIndex, int InAttributeId, float4 InValue)
 {
-	const uint AttributeAddress = GetFirstElementAddressInternal_{DataInterfaceName}(DataIndex, AttributeId);
-	if (AttributeAddress > 0)
+	const uint ElementAddress = GetElementAddressInternal_{DataInterfaceName}(InDataIndex, InElementIndex, InAttributeId);
+
+	if (ElementAddress > 0)
 	{
-		{DataInterfaceName}_DataCollectionBuffer.Store4(AttributeAddress + ElementIndex * 4 * 4, asuint(Value));
+		{DataInterfaceName}_DataCollectionBuffer.Store4(ElementAddress, asuint(InValue));
 	}
 }
 
-void SetRotator_{DataInterfaceName}(uint DataIndex, uint ElementIndex, int AttributeId, float3 Value)
+void SetRotator_{DataInterfaceName}(uint InDataIndex, uint InElementIndex, int InAttributeId, float3 InValue)
 {
-	const uint AttributeAddress = GetFirstElementAddressInternal_{DataInterfaceName}(DataIndex, AttributeId);
-	if (AttributeAddress > 0)
+	const uint ElementAddress = GetElementAddressInternal_{DataInterfaceName}(InDataIndex, InElementIndex, InAttributeId);
+
+	if (ElementAddress > 0)
 	{
-		{DataInterfaceName}_DataCollectionBuffer.Store3(AttributeAddress + ElementIndex * 3 * 4, asuint(Value));
+		{DataInterfaceName}_DataCollectionBuffer.Store3(ElementAddress, asuint(InValue));
 	}
 }
 
-void SetQuat_{DataInterfaceName}(uint DataIndex, uint ElementIndex, int AttributeId, float4 Value)
+void SetQuat_{DataInterfaceName}(uint InDataIndex, uint InElementIndex, int InAttributeId, float4 InValue)
 {
-	const uint AttributeAddress = GetFirstElementAddressInternal_{DataInterfaceName}(DataIndex, AttributeId);
-	if (AttributeAddress > 0)
+	const uint ElementAddress = GetElementAddressInternal_{DataInterfaceName}(InDataIndex, InElementIndex, InAttributeId);
+
+	if (ElementAddress > 0)
 	{
-		{DataInterfaceName}_DataCollectionBuffer.Store4(AttributeAddress + ElementIndex * 4 * 4, asuint(Value));
+		{DataInterfaceName}_DataCollectionBuffer.Store4(ElementAddress, asuint(InValue));
 	}
 }
 
-void SetTransform_{DataInterfaceName}(uint DataIndex, uint ElementIndex, int AttributeId, float4x4 Value)
+void SetTransform_{DataInterfaceName}(uint InDataIndex, uint InElementIndex, int InAttributeId, float4x4 InValue)
 {
-	const uint AttributeAddress = GetFirstElementAddressInternal_{DataInterfaceName}(DataIndex, AttributeId);
-	if (AttributeAddress > 0)
+	const uint ElementAddress = GetElementAddressInternal_{DataInterfaceName}(InDataIndex, InElementIndex, InAttributeId);
+
+	if (ElementAddress > 0)
 	{
-		const uint StartAddress = ElementIndex * 16 * 4 + AttributeAddress;
-		{DataInterfaceName}_DataCollectionBuffer.Store4(StartAddress, asuint(Value[0]));
-		{DataInterfaceName}_DataCollectionBuffer.Store4(StartAddress + 4 * 4, asuint(Value[1]));
-		{DataInterfaceName}_DataCollectionBuffer.Store4(StartAddress + 8 * 4, asuint(Value[2]));
-		{DataInterfaceName}_DataCollectionBuffer.Store4(StartAddress + 12 * 4, asuint(Value[3]));
+		{DataInterfaceName}_DataCollectionBuffer.Store4(ElementAddress, asuint(InValue[0]));
+		{DataInterfaceName}_DataCollectionBuffer.Store4(ElementAddress + 4 * 4, asuint(InValue[1]));
+		{DataInterfaceName}_DataCollectionBuffer.Store4(ElementAddress + 8 * 4, asuint(InValue[2]));
+		{DataInterfaceName}_DataCollectionBuffer.Store4(ElementAddress + 12 * 4, asuint(InValue[3]));
 	}
 }
 
-void SetStringKey_{DataInterfaceName}(uint DataIndex, uint ElementIndex, int AttributeId, int Value)
+void SetStringKey_{DataInterfaceName}(uint InDataIndex, uint InElementIndex, int InAttributeId, int InValue)
 {
 	// String keys are represented as ints.
-	SetInt_{DataInterfaceName}(DataIndex, ElementIndex, AttributeId, Value);
+	SetInt_{DataInterfaceName}(InDataIndex, InElementIndex, InAttributeId, InValue);
 }
 
-// TODO: Value left as uint2 which we locked into for 5.5.0. Future version should change this to int.
-void SetName_{DataInterfaceName}(uint DataIndex, uint ElementIndex, int AttributeId, uint2 Value)
+// TODO: InValue left as uint2 which we locked into for 5.5.0. Future version should change this to int.
+void SetName_{DataInterfaceName}(uint InDataIndex, uint InElementIndex, int InAttributeId, uint2 InValue)
 {
 	// Names are represented as ints.
-	SetInt_{DataInterfaceName}(DataIndex, ElementIndex, AttributeId, (int)Value.x);
+	SetInt_{DataInterfaceName}(InDataIndex, InElementIndex, InAttributeId, (int)InValue.x);
 }
 
 // #################### FIRST ATTRIBUTE GETTERS ##########################
 
 bool GetFirstAttributeAsBool_{DataInterfaceName}()
 {
-	const uint AttributeAddress = GetFirstAttributeAddressInternal_{DataInterfaceName}();
+	const uint AttributeAddress = GetFirstValidAttributeElementInternal_{DataInterfaceName}();
 	return AttributeAddress > 0 && {DataInterfaceName}_DataCollectionBuffer.Load(AttributeAddress) != 0;
 }
 
 int GetFirstAttributeAsInt_{DataInterfaceName}()
 {
-	const uint AttributeAddress = GetFirstAttributeAddressInternal_{DataInterfaceName}();
+	const uint AttributeAddress = GetFirstValidAttributeElementInternal_{DataInterfaceName}();
 	return AttributeAddress > 0 ? asint({DataInterfaceName}_DataCollectionBuffer.Load(AttributeAddress)) : 0;
 }
 
 float GetFirstAttributeAsFloat_{DataInterfaceName}()
 {
-	const uint AttributeAddress = GetFirstAttributeAddressInternal_{DataInterfaceName}();
+	const uint AttributeAddress = GetFirstValidAttributeElementInternal_{DataInterfaceName}();
 	return AttributeAddress > 0 ? asfloat({DataInterfaceName}_DataCollectionBuffer.Load(AttributeAddress)) : 0.0f;
 }
 
 float2 GetFirstAttributeAsFloat2_{DataInterfaceName}()
 {
-	const uint AttributeAddress = GetFirstAttributeAddressInternal_{DataInterfaceName}();
+	const uint AttributeAddress = GetFirstValidAttributeElementInternal_{DataInterfaceName}();
 	return AttributeAddress > 0 ? asfloat({DataInterfaceName}_DataCollectionBuffer.Load2(AttributeAddress)) : (float2)0;
 }
 
 float3 GetFirstAttributeAsFloat3_{DataInterfaceName}()
 {
-	const uint AttributeAddress = GetFirstAttributeAddressInternal_{DataInterfaceName}();
+	const uint AttributeAddress = GetFirstValidAttributeElementInternal_{DataInterfaceName}();
 	return AttributeAddress > 0 ? asfloat({DataInterfaceName}_DataCollectionBuffer.Load3(AttributeAddress)) : (float3)0;
 }
 
 float4 GetFirstAttributeAsFloat4_{DataInterfaceName}()
 {
-	const uint AttributeAddress = GetFirstAttributeAddressInternal_{DataInterfaceName}();
+	const uint AttributeAddress = GetFirstValidAttributeElementInternal_{DataInterfaceName}();
 	return AttributeAddress > 0 ? asfloat({DataInterfaceName}_DataCollectionBuffer.Load4(AttributeAddress)) : (float4)0;
 }
 
@@ -427,7 +453,7 @@ float4 GetFirstAttributeAsQuat_{DataInterfaceName}()
 
 float4x4 GetFirstAttributeAsTransform_{DataInterfaceName}()
 {
-	const uint AttributeAddress = GetFirstAttributeAddressInternal_{DataInterfaceName}();
+	const uint AttributeAddress = GetFirstValidAttributeElementInternal_{DataInterfaceName}();
 	if (AttributeAddress == 0)
 	{
 		return (float4x4)0;
@@ -457,14 +483,14 @@ uint2 GetFirstAttributeAsName_{DataInterfaceName}()
 // #################### ATOMICS ##########################
 
 // Returns value of attribute before incrementing.
-int AtomicAddInt_{DataInterfaceName}(uint DataIndex, uint ElementIndex, int AttributeId, int ValueToAdd)
+int AtomicAddInt_{DataInterfaceName}(uint InDataIndex, uint InElementIndex, int InAttributeId, int InValueToAdd)
 {
-	const uint AttributeAddress = GetFirstElementAddressInternal_{DataInterfaceName}(DataIndex, AttributeId);
+	const uint ElementAddress = GetElementAddressInternal_{DataInterfaceName}(InDataIndex, InElementIndex, InAttributeId);
 
-	if (AttributeAddress > 0)
+	if (ElementAddress > 0)
 	{
 		uint OriginalValue;
-		{DataInterfaceName}_DataCollectionBuffer.InterlockedAdd(AttributeAddress + ElementIndex * 1 * 4, (uint)ValueToAdd, OriginalValue);
+		{DataInterfaceName}_DataCollectionBuffer.InterlockedAdd(ElementAddress, (uint)InValueToAdd, OriginalValue);
 
 		return (int)OriginalValue;
 	}
@@ -476,49 +502,49 @@ int AtomicAddInt_{DataInterfaceName}(uint DataIndex, uint ElementIndex, int Attr
 
 // #################### POINT ATTRIBUTE GETTERS ##########################
 
-float3 GetPosition_{DataInterfaceName}(uint DataIndex, uint ElementIndex)
+float3 GetPosition_{DataInterfaceName}(uint InDataIndex, uint InElementIndex)
 {
-	return GetFloat3_{DataInterfaceName}(DataIndex, ElementIndex, PCG_POINT_POSITION_ID);
+	return GetFloat3_{DataInterfaceName}(InDataIndex, InElementIndex, PCG_POINT_POSITION_ID);
 }
 
-float4 GetRotation_{DataInterfaceName}(uint DataIndex, uint ElementIndex)
+float4 GetRotation_{DataInterfaceName}(uint InDataIndex, uint InElementIndex)
 {
-	return GetFloat4_{DataInterfaceName}(DataIndex, ElementIndex, PCG_POINT_ROTATION_ID);
+	return GetFloat4_{DataInterfaceName}(InDataIndex, InElementIndex, PCG_POINT_ROTATION_ID);
 }
 
-float3 GetScale_{DataInterfaceName}(uint DataIndex, uint ElementIndex)
+float3 GetScale_{DataInterfaceName}(uint InDataIndex, uint InElementIndex)
 {
-	return GetFloat3_{DataInterfaceName}(DataIndex, ElementIndex, PCG_POINT_SCALE_ID);
+	return GetFloat3_{DataInterfaceName}(InDataIndex, InElementIndex, PCG_POINT_SCALE_ID);
 }
 
-float3 GetBoundsMin_{DataInterfaceName}(uint DataIndex, uint ElementIndex)
+float3 GetBoundsMin_{DataInterfaceName}(uint InDataIndex, uint InElementIndex)
 {
-	return GetFloat3_{DataInterfaceName}(DataIndex, ElementIndex, PCG_POINT_BOUNDS_MIN_ID);
+	return GetFloat3_{DataInterfaceName}(InDataIndex, InElementIndex, PCG_POINT_BOUNDS_MIN_ID);
 }
 
-float3 GetBoundsMax_{DataInterfaceName}(uint DataIndex, uint ElementIndex)
+float3 GetBoundsMax_{DataInterfaceName}(uint InDataIndex, uint InElementIndex)
 {
-	return GetFloat3_{DataInterfaceName}(DataIndex, ElementIndex, PCG_POINT_BOUNDS_MAX_ID);
+	return GetFloat3_{DataInterfaceName}(InDataIndex, InElementIndex, PCG_POINT_BOUNDS_MAX_ID);
 }
 
-float4 GetColor_{DataInterfaceName}(uint DataIndex, uint ElementIndex)
+float4 GetColor_{DataInterfaceName}(uint InDataIndex, uint InElementIndex)
 {
-	return GetFloat4_{DataInterfaceName}(DataIndex, ElementIndex, PCG_POINT_COLOR_ID);
+	return GetFloat4_{DataInterfaceName}(InDataIndex, InElementIndex, PCG_POINT_COLOR_ID);
 }
 
-float GetDensity_{DataInterfaceName}(uint DataIndex, uint ElementIndex)
+float GetDensity_{DataInterfaceName}(uint InDataIndex, uint InElementIndex)
 {
-	return GetFloat_{DataInterfaceName}(DataIndex, ElementIndex, PCG_POINT_DENSITY_ID);
+	return GetFloat_{DataInterfaceName}(InDataIndex, InElementIndex, PCG_POINT_DENSITY_ID);
 }
 
-float GetSteepness_{DataInterfaceName}(uint DataIndex, uint ElementIndex)
+float GetSteepness_{DataInterfaceName}(uint InDataIndex, uint InElementIndex)
 {
-	return GetFloat_{DataInterfaceName}(DataIndex, ElementIndex, PCG_POINT_STEEPNESS_ID);
+	return GetFloat_{DataInterfaceName}(InDataIndex, InElementIndex, PCG_POINT_STEEPNESS_ID);
 }
 
-uint GetSeed_{DataInterfaceName}(uint DataIndex, uint ElementIndex)
+uint GetSeed_{DataInterfaceName}(uint InDataIndex, uint InElementIndex)
 {
-	return GetUint_{DataInterfaceName}(DataIndex, ElementIndex, PCG_POINT_SEED_ID);
+	return GetUint_{DataInterfaceName}(InDataIndex, InElementIndex, PCG_POINT_SEED_ID);
 }
 
 float4x4 GetPointTransform_{DataInterfaceName}(uint InDataIndex, uint InElementIndex)
@@ -546,49 +572,49 @@ bool IsPointRemoved_{DataInterfaceName}(uint InDataIndex, uint InElementIndex)
 
 // #################### POINT ATTRIBUTE SETTERS ##########################
 
-void SetPosition_{DataInterfaceName}(uint DataIndex, uint ElementIndex, float3 Position)
+void SetPosition_{DataInterfaceName}(uint InDataIndex, uint InElementIndex, float3 Position)
 {
-	SetFloat3_{DataInterfaceName}(DataIndex, ElementIndex, PCG_POINT_POSITION_ID, Position);
+	SetFloat3_{DataInterfaceName}(InDataIndex, InElementIndex, PCG_POINT_POSITION_ID, Position);
 }
 
-void SetRotation_{DataInterfaceName}(uint DataIndex, uint ElementIndex, float4 Rotation)
+void SetRotation_{DataInterfaceName}(uint InDataIndex, uint InElementIndex, float4 Rotation)
 {
-	SetFloat4_{DataInterfaceName}(DataIndex, ElementIndex, PCG_POINT_ROTATION_ID, Rotation);
+	SetFloat4_{DataInterfaceName}(InDataIndex, InElementIndex, PCG_POINT_ROTATION_ID, Rotation);
 }
 
-void SetScale_{DataInterfaceName}(uint DataIndex, uint ElementIndex, float3 Scale)
+void SetScale_{DataInterfaceName}(uint InDataIndex, uint InElementIndex, float3 Scale)
 {
-	SetFloat3_{DataInterfaceName}(DataIndex, ElementIndex, PCG_POINT_SCALE_ID, Scale);
+	SetFloat3_{DataInterfaceName}(InDataIndex, InElementIndex, PCG_POINT_SCALE_ID, Scale);
 }
 
-void SetBoundsMin_{DataInterfaceName}(uint DataIndex, uint ElementIndex, float3 BoundsMin)
+void SetBoundsMin_{DataInterfaceName}(uint InDataIndex, uint InElementIndex, float3 BoundsMin)
 {
-	SetFloat3_{DataInterfaceName}(DataIndex, ElementIndex, PCG_POINT_BOUNDS_MIN_ID, BoundsMin);
+	SetFloat3_{DataInterfaceName}(InDataIndex, InElementIndex, PCG_POINT_BOUNDS_MIN_ID, BoundsMin);
 }
 
-void SetBoundsMax_{DataInterfaceName}(uint DataIndex, uint ElementIndex, float3 BoundsMax)
+void SetBoundsMax_{DataInterfaceName}(uint InDataIndex, uint InElementIndex, float3 BoundsMax)
 {
-	SetFloat3_{DataInterfaceName}(DataIndex, ElementIndex, PCG_POINT_BOUNDS_MAX_ID, BoundsMax);
+	SetFloat3_{DataInterfaceName}(InDataIndex, InElementIndex, PCG_POINT_BOUNDS_MAX_ID, BoundsMax);
 }
 
-void SetColor_{DataInterfaceName}(uint DataIndex, uint ElementIndex, float4 Color)
+void SetColor_{DataInterfaceName}(uint InDataIndex, uint InElementIndex, float4 Color)
 {
-	SetFloat4_{DataInterfaceName}(DataIndex, ElementIndex, PCG_POINT_COLOR_ID, Color);
+	SetFloat4_{DataInterfaceName}(InDataIndex, InElementIndex, PCG_POINT_COLOR_ID, Color);
 }
 
-void SetDensity_{DataInterfaceName}(uint DataIndex, uint ElementIndex, float Density)
+void SetDensity_{DataInterfaceName}(uint InDataIndex, uint InElementIndex, float Density)
 {
-	SetFloat_{DataInterfaceName}(DataIndex, ElementIndex, PCG_POINT_DENSITY_ID, Density);
+	SetFloat_{DataInterfaceName}(InDataIndex, InElementIndex, PCG_POINT_DENSITY_ID, Density);
 }
 
-void SetSteepness_{DataInterfaceName}(uint DataIndex, uint ElementIndex, float Steepness)
+void SetSteepness_{DataInterfaceName}(uint InDataIndex, uint InElementIndex, float Steepness)
 {
-	SetFloat_{DataInterfaceName}(DataIndex, ElementIndex, PCG_POINT_STEEPNESS_ID, Steepness);
+	SetFloat_{DataInterfaceName}(InDataIndex, InElementIndex, PCG_POINT_STEEPNESS_ID, Steepness);
 }
 
-void SetSeed_{DataInterfaceName}(uint DataIndex, uint ElementIndex, uint Seed)
+void SetSeed_{DataInterfaceName}(uint InDataIndex, uint InElementIndex, uint Seed)
 {
-	SetUint_{DataInterfaceName}(DataIndex, ElementIndex, PCG_POINT_SEED_ID, Seed);
+	SetUint_{DataInterfaceName}(InDataIndex, InElementIndex, PCG_POINT_SEED_ID, Seed);
 }
 
 void SetPointTransform_{DataInterfaceName}(uint InDataIndex, uint InElementIndex, float4x4 Transform)
@@ -631,7 +657,9 @@ void InitializePoint_{DataInterfaceName}(uint InDataIndex, uint InElementIndex)
 
 	for (int AttributeId = PCG_NUM_RESERVED_ATTRS; AttributeId < PCG_MAX_NUM_ATTRS; ++AttributeId)
 	{
-		const uint Stride = GetAttributeStrideInternal_{DataInterfaceName}(InDataIndex, AttributeId);
+		const uint2 ElementAddressAndStride = GetElementAddressAndStrideInternal_{DataInterfaceName}(InDataIndex, InElementIndex, AttributeId);
+		const uint ElementAddress = ElementAddressAndStride.x;
+		const uint Stride = ElementAddressAndStride.y;
 
 		if (Stride == 0)
 		{
@@ -640,9 +668,6 @@ void InitializePoint_{DataInterfaceName}(uint InDataIndex, uint InElementIndex)
 		}
 
 		// We don't bother checking element address is nonzero as we already check Stride above.
-		const uint FirstElementAddress = GetFirstElementAddressInternal_{DataInterfaceName}(InDataIndex, AttributeId);
-		const uint ElementAddress = InElementIndex * Stride + FirstElementAddress;
-
 		for (uint Offset = 0; Offset < Stride; Offset += 4)
 		{
 			// TODO: In future could upload actual default values for attributes rather than 0-initializing.
diff --git a/Engine/Plugins/PCG/Source/PCG/Internal/Compute/PCGComputeCommon.h b/Engine/Plugins/PCG/Source/PCG/Internal/Compute/PCGComputeCommon.h
--- a/Engine/Plugins/PCG/Source/PCG/Internal/Compute/PCGComputeCommon.h
+++ b/Engine/Plugins/PCG/Source/PCG/Internal/Compute/PCGComputeCommon.h
@@ -70,6 +70,10 @@ namespace PCGComputeConstants
 
 	constexpr uint32 KernelExecutedFlag = 1 << 31;
 
+	// PackedAttributeInfo bitmasks
+	constexpr uint32 AttributeStrideMask = 0xFF;
+	constexpr uint32 AttributeAllocatedMask = 1 << 31;
+
 	constexpr uint32 MAX_PRIMITIVE_COMPONENTS_PER_SPAWNER = 256;
 
 	constexpr uint32 THREAD_GROUP_SIZE = 64;
diff --git a/Engine/Plugins/PCG/Source/PCG/Internal/Compute/PCGDataForGPU.h b/Engine/Plugins/PCG/Source/PCG/Internal/Compute/PCGDataForGPU.h
--- a/Engine/Plugins/PCG/Source/PCG/Internal/Compute/PCGDataForGPU.h
+++ b/Engine/Plugins/PCG/Source/PCG/Internal/Compute/PCGDataForGPU.h
@@ -19,6 +19,7 @@ class UPCGPin;
 class UPCGSettings;
 struct FPCGContext;
 enum class EPCGMetadataTypes : uint8;
+enum class EPCGPointNativeProperties : uint32;
 
 enum class EPCGUnpackDataCollectionResult
 {
@@ -175,6 +176,15 @@ public:
 
 	bool IsDomain2D() const { return !!(Type & EPCGDataType::BaseTexture); }
 
+	void AllocateProperties(EPCGPointNativeProperties InProperties) { AllocatedPointProperties |= InProperties; }
+	EPCGPointNativeProperties GetAllocatedProperties() const { return AllocatedPointProperties; }
+
+	bool IsAttributeAllocated(const FPCGKernelAttributeDesc& InAttributeDesc) const { return IsAttributeAllocated(InAttributeDesc.AttributeId); }
+	bool IsAttributeAllocated(int32 InAttributeId) const;
+
+private:
+	void InitializeAttributeDescs(const UPCGData* InData, const UPCGDataBinding* InBinding);
+
 public:
 	EPCGDataType Type = EPCGDataType::Point;
 	TArray<FPCGKernelAttributeDesc> AttributeDescs;
@@ -184,7 +194,8 @@ public:
 	TArray<int32, TInlineAllocator<4>> TagStringKeys;
 
 private:
-	void InitializeAttributeDescs(const UPCGData* InData, const UPCGDataBinding* InBinding);
+	/** Properties that should be fully allocated for this data. Only applies to EPCGDataType::Point currently. */
+	EPCGPointNativeProperties AllocatedPointProperties = EPCGPointNativeProperties::None;
 };
 
 /** Data description for a data collection (FPCGDataCollection). */
@@ -223,6 +234,8 @@ struct FPCGDataCollectionDesc
 	/** Makes attribute present on all data. If data has existing attribute with same name then the given type will be applied. */
 	void AddAttributeToAllData(FPCGKernelAttributeKey InAttribute, const UPCGDataBinding* InBinding, const TArray<int32>* InOptionalUniqueStringKeys = nullptr);
 
+	void AllocatePropertiesForAllData(EPCGPointNativeProperties InProperties);
+
 	void GetUniqueStringKeyValues(int32 InAttributeId, TArray<int32>& OutUniqueStringKeys) const;
 
 	int GetNumStringKeyValues(int32 InAttributeId) const;
diff --git a/Engine/Plugins/PCG/Source/PCG/Private/Compute/DataInterfaces/PCGDataCollectionDataInterface.cpp b/Engine/Plugins/PCG/Source/PCG/Private/Compute/DataInterfaces/PCGDataCollectionDataInterface.cpp
--- a/Engine/Plugins/PCG/Source/PCG/Private/Compute/DataInterfaces/PCGDataCollectionDataInterface.cpp
+++ b/Engine/Plugins/PCG/Source/PCG/Private/Compute/DataInterfaces/PCGDataCollectionDataInterface.cpp
@@ -73,13 +73,14 @@ void UPCGDataCollectionDataInterface::GetSupportedInputs(TArray<FShaderFunctionD
 			.AddReturnType(EShaderFundamentalType::Uint)
 			.AddParam(EShaderFundamentalType::Uint) // InDataIndex
 			.AddParam(EShaderFundamentalType::Uint) // InElementIndex
-			.AddParam(EShaderFundamentalType::Uint); // InAttributeId
+			.AddParam(EShaderFundamentalType::Int); // InAttributeId
 
 		OutFunctions.AddDefaulted_GetRef()
-			.SetName(TEXT("GetAttributeStrideInternal"))
-			.AddReturnType(EShaderFundamentalType::Uint)
+			.SetName(TEXT("GetElementAddressAndStrideInternal"))
+			.AddReturnType(EShaderFundamentalType::Uint, 2)
 			.AddParam(EShaderFundamentalType::Uint) // InDataIndex
-			.AddParam(EShaderFundamentalType::Uint); // InAttributeId
+			.AddParam(EShaderFundamentalType::Uint) // InElementIndex
+			.AddParam(EShaderFundamentalType::Int); // InAttributeId
 	}
 
 	// Header Writers
@@ -520,6 +521,9 @@ void UPCGDataCollectionDataInterface::GetDefines(FComputeKernelDefinitionSet& Ou
 	OutDefinitionSet.Defines.Add(FComputeKernelDefinition(TEXT("PCG_KERNEL_EXECUTED_FLAG"), FString::FromInt(PCGComputeConstants::KernelExecutedFlag)));
 	OutDefinitionSet.Defines.Add(FComputeKernelDefinition(TEXT("PCG_INVALID_DENSITY"),      FString::Printf(TEXT("%f"), PCGComputeConstants::INVALID_DENSITY)));
 	OutDefinitionSet.Defines.Add(FComputeKernelDefinition(TEXT("PCG_NUM_RESERVED_ATTRS"),   FString::FromInt(PCGComputeConstants::NUM_RESERVED_ATTRS)));
+
+	OutDefinitionSet.Defines.Add(FComputeKernelDefinition(TEXT("PCG_ATTRIBUTE_STRIDE_MASK"), FString::FromInt(PCGComputeConstants::AttributeStrideMask)));
+	OutDefinitionSet.Defines.Add(FComputeKernelDefinition(TEXT("PCG_ATTRIBUTE_ALLOCATED_MASK"), FString::FromInt(PCGComputeConstants::AttributeAllocatedMask)));
 }
 
 UComputeDataProvider* UPCGDataCollectionDataInterface::CreateDataProvider() const
diff --git a/Engine/Plugins/PCG/Source/PCG/Private/Compute/Elements/PCGComputeGraphElement.h b/Engine/Plugins/PCG/Source/PCG/Private/Compute/Elements/PCGComputeGraphElement.h
--- a/Engine/Plugins/PCG/Source/PCG/Private/Compute/Elements/PCGComputeGraphElement.h
+++ b/Engine/Plugins/PCG/Source/PCG/Private/Compute/Elements/PCGComputeGraphElement.h
@@ -114,6 +114,7 @@ protected:
 	virtual void PostExecuteInternal(FPCGContext* InContext) const override;
 	virtual void AbortInternal(FPCGContext* InContext) const override;;
 	virtual bool SupportsGPUResidentData(FPCGContext* InContext) const override { return true; }
+	virtual bool SupportsBasePointDataInputs(FPCGContext* InContext) const override { return true; }
 
 	// The calls to initialize the compute graph are not thread safe.
 	virtual bool CanExecuteOnlyOnMainThread(FPCGContext* Context) const override { return true; }
diff --git a/Engine/Plugins/PCG/Source/PCG/Private/Compute/Elements/PCGCustomHLSLKernel.cpp b/Engine/Plugins/PCG/Source/PCG/Private/Compute/Elements/PCGCustomHLSLKernel.cpp
--- a/Engine/Plugins/PCG/Source/PCG/Private/Compute/Elements/PCGCustomHLSLKernel.cpp
+++ b/Engine/Plugins/PCG/Source/PCG/Private/Compute/Elements/PCGCustomHLSLKernel.cpp
@@ -5,6 +5,7 @@
 #include "PCGComputeGraphElement.h"
 #include "PCGContext.h"
 #include "PCGModule.h"
+#include "PCGPointPropertiesTraits.h"
 #include "Compute/PCGComputeGraph.h"
 #include "Compute/PCGDataBinding.h"
 #include "Graph/PCGGPUGraphCompilationContext.h"
@@ -31,6 +32,16 @@ namespace PCGCustomHLSLKernel
 	constexpr TCHAR CopyElementFunctionKeyword[] = { TEXT("CopyElementFrom") };
 	constexpr TCHAR InitializeFunctionKeyword[] = { TEXT("Initialize") };
 	constexpr TCHAR StoreFunctionKeyword[] = { TEXT("Store") };
+	constexpr TCHAR SetPositionKeyword[] = { TEXT("SetPosition") };
+	constexpr TCHAR SetRotationKeyword[] = { TEXT("SetRotation") };
+	constexpr TCHAR SetScaleKeyword[] = { TEXT("SetScale") };
+	constexpr TCHAR SetBoundsMinKeyword[] = { TEXT("SetBoundsMin") };
+	constexpr TCHAR SetBoundsMaxKeyword[] = { TEXT("SetBoundsMax") };
+	constexpr TCHAR SetColorKeyword[] = { TEXT("SetColor") };
+	constexpr TCHAR SetDensityKeyword[] = { TEXT("SetDensity") };
+	constexpr TCHAR SetSteepnessKeyword[] = { TEXT("SetSteepness") };
+	constexpr TCHAR SetSeedKeyword[] = { TEXT("SetSeed") };
+	constexpr TCHAR SetPointTransformKeyword[] = { TEXT("SetPointTransform") };
 
 #if WITH_EDITOR
 	enum class EParseState : uint8
@@ -375,6 +386,12 @@ bool UPCGCustomHLSLKernel::ComputeOutputBindingDataDesc(const UPCGComputeGraph*
 		}
 	}
 
+	// Allocate any properties that were modified by the Custom HLSL source. If this data was initialized from another, it will have inherited its allocated properties already too.
+	if (const EPCGPointNativeProperties* AllocatedProperties = PinAllocatedProperties.Find(OutputPinProperties->Label))
+	{
+		OutputDataDesc.AllocatePropertiesForAllData(*AllocatedProperties);
+	}
+
 	OutDataDesc = MoveTemp(OutputDataDesc);
 	return true;
 }
@@ -759,6 +776,7 @@ void UPCGCustomHLSLKernel::ParseShaderSource()
 	TArray<FString> AttributeKeywords;
 	TArray<FString> CopyElementKeywords;
 	TArray<FString> InitializeKeywords;
+	TArray<FString> PointPropertySetterKeywords;
 
 	const UPCGCustomHLSLSettings* CustomHLSLSettings = CastChecked<UPCGCustomHLSLSettings>(GetSettings());
 
@@ -791,18 +809,33 @@ void UPCGCustomHLSLKernel::ParseShaderSource()
 					CopyElementKeywords.Add(PinStr + "_" + PCGCustomHLSLKernel::CopyElementFunctionKeyword + "_" + InputPinProps.Label.ToString());
 				}
 			}
-			
+
 			InitializeKeywords.Add(PinStr + "_" + PCGCustomHLSLKernel::InitializeFunctionKeyword);
 		}
 		else if (!!(PinProps.AllowedTypes & EPCGDataType::BaseTexture))
 		{
 			InitializeKeywords.Add(PinStr + "_" + PCGCustomHLSLKernel::StoreFunctionKeyword);
 		}
+
+		if (!!(PinProps.AllowedTypes & EPCGDataType::Point))
+		{
+			PointPropertySetterKeywords.Add(PinStr + "_" + PCGCustomHLSLKernel::SetPositionKeyword);
+			PointPropertySetterKeywords.Add(PinStr + "_" + PCGCustomHLSLKernel::SetRotationKeyword);
+			PointPropertySetterKeywords.Add(PinStr + "_" + PCGCustomHLSLKernel::SetScaleKeyword);
+			PointPropertySetterKeywords.Add(PinStr + "_" + PCGCustomHLSLKernel::SetBoundsMinKeyword);
+			PointPropertySetterKeywords.Add(PinStr + "_" + PCGCustomHLSLKernel::SetBoundsMaxKeyword);
+			PointPropertySetterKeywords.Add(PinStr + "_" + PCGCustomHLSLKernel::SetColorKeyword);
+			PointPropertySetterKeywords.Add(PinStr + "_" + PCGCustomHLSLKernel::SetDensityKeyword);
+			PointPropertySetterKeywords.Add(PinStr + "_" + PCGCustomHLSLKernel::SetSteepnessKeyword);
+			PointPropertySetterKeywords.Add(PinStr + "_" + PCGCustomHLSLKernel::SetSeedKeyword);
+			PointPropertySetterKeywords.Add(PinStr + "_" + PCGCustomHLSLKernel::SetPointTransformKeyword);
+		}
 	}
 
 	AdditionalKeywords.Append(AttributeKeywords);
 	AdditionalKeywords.Append(CopyElementKeywords);
 	AdditionalKeywords.Append(InitializeKeywords);
+	AdditionalKeywords.Append(PointPropertySetterKeywords);
 
 	FPCGSyntaxTokenizerParams TokenizerParams;
 	TokenizerParams.AdditionalKeywords = MoveTemp(AdditionalKeywords);
@@ -930,6 +963,52 @@ void UPCGCustomHLSLKernel::ParseShaderSource()
 
 					InOutParsedSource.InitializedOutputPins.Add(MoveTemp(PinStr));
 				}
+				else if (PointPropertySetterKeywords.Contains(TokenString))
+				{
+					const int32 DelimiterIndex = TokenString.Find("_");
+					const FString PinStr = TokenString.Left(DelimiterIndex);
+					const FString PropertyStr = TokenString.RightChop(DelimiterIndex + 1);
+					
+					// If the kernel writes to an attribute, it must be fully allocated if it isn't already.
+					// @todo_pcg: Could have an API for setting the default/constant value in a constant value range as well, to avoid expanding to full-width allocation always.
+					EPCGPointNativeProperties& AllocatedProperties = PinAllocatedProperties.FindOrAdd(FName(PinStr));
+
+					if (PropertyStr == PCGCustomHLSLKernel::SetPositionKeyword
+						|| PropertyStr == PCGCustomHLSLKernel::SetRotationKeyword
+						|| PropertyStr == PCGCustomHLSLKernel::SetScaleKeyword
+						|| PropertyStr == PCGCustomHLSLKernel::SetPointTransformKeyword)
+					{
+						AllocatedProperties |= EPCGPointNativeProperties::Transform;
+					}
+					else if (PropertyStr == PCGCustomHLSLKernel::SetBoundsMinKeyword)
+					{
+						AllocatedProperties |= EPCGPointNativeProperties::BoundsMin;
+					}
+					else if (PropertyStr == PCGCustomHLSLKernel::SetBoundsMaxKeyword)
+					{
+						AllocatedProperties |= EPCGPointNativeProperties::BoundsMax;
+					}
+					else if (PropertyStr == PCGCustomHLSLKernel::SetColorKeyword)
+					{
+						AllocatedProperties |= EPCGPointNativeProperties::Color;
+					}
+					else if (PropertyStr == PCGCustomHLSLKernel::SetDensityKeyword)
+					{
+						AllocatedProperties |= EPCGPointNativeProperties::Density;
+					}
+					else if (PropertyStr == PCGCustomHLSLKernel::SetSteepnessKeyword)
+					{
+						AllocatedProperties |= EPCGPointNativeProperties::Steepness;
+					}
+					else if (PropertyStr == PCGCustomHLSLKernel::SetSeedKeyword)
+					{
+						AllocatedProperties |= EPCGPointNativeProperties::Seed;
+					}
+					else
+					{
+						ensure(false);
+					}
+				}
 			}
 			else if (Token.Type == FPCGCustomHLSLParsedSource::ETokenType::Normal)
 			{
diff --git a/Engine/Plugins/PCG/Source/PCG/Private/Compute/Elements/PCGCustomHLSLKernel.h b/Engine/Plugins/PCG/Source/PCG/Private/Compute/Elements/PCGCustomHLSLKernel.h
--- a/Engine/Plugins/PCG/Source/PCG/Private/Compute/Elements/PCGCustomHLSLKernel.h
+++ b/Engine/Plugins/PCG/Source/PCG/Private/Compute/Elements/PCGCustomHLSLKernel.h
@@ -9,6 +9,8 @@
 
 #include "PCGCustomHLSLKernel.generated.h"
 
+enum class EPCGPointNativeProperties : uint32;
+
 USTRUCT()
 struct FPCGParsedAttributeFunction
 {
@@ -172,4 +174,8 @@ protected:
 
 	UPROPERTY()
 	FPCGPinDataLabels PinDataLabels;
+
+	/** Properties that need to be fully allocated per pin. */
+	UPROPERTY()
+	TMap<FName, EPCGPointNativeProperties> PinAllocatedProperties;
 };
diff --git a/Engine/Plugins/PCG/Source/PCG/Private/Compute/Elements/PCGNormalToDensityKernel.cpp b/Engine/Plugins/PCG/Source/PCG/Private/Compute/Elements/PCGNormalToDensityKernel.cpp
--- a/Engine/Plugins/PCG/Source/PCG/Private/Compute/Elements/PCGNormalToDensityKernel.cpp
+++ b/Engine/Plugins/PCG/Source/PCG/Private/Compute/Elements/PCGNormalToDensityKernel.cpp
@@ -28,6 +28,15 @@ bool UPCGNormalToDensityKernel::ComputeOutputBindingDataDesc(const UPCGComputeGr
 
 	ensure(InBinding->ComputeKernelPinDataDesc(SourceKernelPin, &OutDataDesc));
 
+	// Allocate density property if the transform is also allocated. Otherwise the densities will all be the same value and full-width allocation is unnecessary.
+	for (FPCGDataDesc& DataDesc : OutDataDesc.DataDescs)
+	{
+		if (!!(DataDesc.GetAllocatedProperties() & EPCGPointNativeProperties::Transform))
+		{
+			DataDesc.AllocateProperties(EPCGPointNativeProperties::Density);
+		}
+	}
+
 	return true;
 }
 
diff --git a/Engine/Plugins/PCG/Source/PCG/Private/Compute/PCGComputeCommon.cpp b/Engine/Plugins/PCG/Source/PCG/Private/Compute/PCGComputeCommon.cpp
--- a/Engine/Plugins/PCG/Source/PCG/Private/Compute/PCGComputeCommon.cpp
+++ b/Engine/Plugins/PCG/Source/PCG/Private/Compute/PCGComputeCommon.cpp
@@ -9,7 +9,7 @@
 #include "Compute/PCGDataBinding.h"
 #include "Compute/PCGPinPropertiesGPU.h"
 #include "Compute/Data/PCGProxyForGPUData.h"
-#include "Data/PCGPointData.h"
+#include "Data/PCGBasePointData.h"
 #include "Data/PCGTextureData.h"
 #include "Editor/IPCGEditorModule.h"
 
@@ -47,9 +47,9 @@ namespace PCGComputeHelpers
 
 	int GetElementCount(const UPCGData* InData)
 	{
-		if (const UPCGPointData* PointData = Cast<UPCGPointData>(InData))
+		if (const UPCGBasePointData* PointData = Cast<UPCGBasePointData>(InData))
 		{
-			return PointData->GetPoints().Num();
+			return PointData->GetNumPoints();
 		}
 		else if (const UPCGParamData* ParamData = Cast<UPCGParamData>(InData))
 		{
@@ -95,7 +95,7 @@ namespace PCGComputeHelpers
 	{
 		// We only read and expose attributes to compute from the following types. Other types are supported but we don't
 		// register/upload their metadata attributes automatically.
-		return InData && (InData->IsA<UPCGParamData>() || InData->IsA<UPCGPointData>());
+		return InData && (InData->IsA<UPCGParamData>() || InData->IsA<UPCGBasePointData>());
 	}
 
 #if PCG_KERNEL_LOGGING_ENABLED
diff --git a/Engine/Plugins/PCG/Source/PCG/Private/Compute/PCGComputeGraph.cpp b/Engine/Plugins/PCG/Source/PCG/Private/Compute/PCGComputeGraph.cpp
--- a/Engine/Plugins/PCG/Source/PCG/Private/Compute/PCGComputeGraph.cpp
+++ b/Engine/Plugins/PCG/Source/PCG/Private/Compute/PCGComputeGraph.cpp
@@ -238,6 +238,7 @@ void UPCGComputeGraph::DebugLogDataDescriptions(const UPCGDataBinding* InBinding
 			for (int32 AttributeIndex = 0; AttributeIndex < DataDesc.AttributeDescs.Num(); ++AttributeIndex)
 			{
 				const FPCGKernelAttributeDesc& AttributeDesc = DataDesc.AttributeDescs[AttributeIndex];
+				const bool bAllocated = DataDesc.IsAttributeAllocated(AttributeDesc);
 
 				// If the data comes from an upstream / source buffer then look up its info.
 				FString UniqueStringKeys;
@@ -265,18 +266,20 @@ void UPCGComputeGraph::DebugLogDataDescriptions(const UPCGDataBinding* InBinding
 
 				if (SourceBufferAttributeId == INDEX_NONE)
 				{
-					UE_LOG(LogPCG, Display, TEXT("\t\t\t\tID: %d\t\tName: %s\t\tType: %s\t\tUniqueStringKeys%s"),
+					UE_LOG(LogPCG, Display, TEXT("\t\t\t\tID: %d\t\tName: %s\t\tType: %s\t\tAllocated: %s\t\tUniqueStringKeys%s"),
 						AttributeDesc.AttributeId,
 						*AttributeDesc.AttributeKey.Identifier.ToString(),
 						*PCGKernelAttributeTypeEnum->GetDisplayNameTextByValue(static_cast<int64>(AttributeDesc.AttributeKey.Type)).ToString(),
+						bAllocated ? TEXT("true") : TEXT("false"),
 						UniqueStringKeys.IsEmpty() ? TEXT(": ") : *UniqueStringKeys);
 				}
 				else
 				{
-					UE_LOG(LogPCG, Display, TEXT("\t\t\t\tID: %d\t\tName: %s\t\tType: %s\t\tUniqueStringKeys%s\t\tSource Buffer ID: %d"),
+					UE_LOG(LogPCG, Display, TEXT("\t\t\t\tID: %d\t\tName: %s\t\tType: %s\t\tAllocated: %s\t\tUniqueStringKeys%s\t\tSource Buffer ID: %d"),
 						AttributeDesc.AttributeId,
 						*AttributeDesc.AttributeKey.Identifier.ToString(),
 						*PCGKernelAttributeTypeEnum->GetDisplayNameTextByValue(static_cast<int64>(AttributeDesc.AttributeKey.Type)).ToString(),
+						bAllocated ? TEXT("true") : TEXT("false"),
 						UniqueStringKeys.IsEmpty() ? TEXT(": ") : *UniqueStringKeys,
 						SourceBufferAttributeId);
 				}
diff --git a/Engine/Plugins/PCG/Source/PCG/Private/Compute/PCGComputeKernel.cpp b/Engine/Plugins/PCG/Source/PCG/Private/Compute/PCGComputeKernel.cpp
--- a/Engine/Plugins/PCG/Source/PCG/Private/Compute/PCGComputeKernel.cpp
+++ b/Engine/Plugins/PCG/Source/PCG/Private/Compute/PCGComputeKernel.cpp
@@ -264,6 +264,8 @@ void UPCGComputeKernel::ComputeDataDescFromPinProperties(const FPCGPinProperties
 					{
 						DataDesc.TagStringKeys.AddUnique(TagStringKey);
 					}
+
+					DataDesc.AllocateProperties(InputDataDesc.GetAllocatedProperties());
 				}
 
 				DataDesc.ElementCount = TotalNumElements;
@@ -297,6 +299,7 @@ void UPCGComputeKernel::ComputeDataDescFromPinProperties(const FPCGPinProperties
 					}
 
 					DataDesc.TagStringKeys = InputDataDesc.TagStringKeys;
+					DataDesc.AllocateProperties(InputDataDesc.GetAllocatedProperties());
 				}
 			}
 			// Take pairs of datas, where the pairs are given by each data of each pin to each data of every other pin.
@@ -402,6 +405,8 @@ void UPCGComputeKernel::ComputeDataDescFromPinProperties(const FPCGPinProperties
 								{
 									DataDesc.TagStringKeys.AddUnique(TagStringKey);
 								}
+
+								DataDesc.AllocateProperties(InputDataDesc.GetAllocatedProperties() | OtherInputDataDesc.GetAllocatedProperties());
 							}
 						}
 					}
diff --git a/Engine/Plugins/PCG/Source/PCG/Private/Compute/PCGDataForGPU.cpp b/Engine/Plugins/PCG/Source/PCG/Private/Compute/PCGDataForGPU.cpp
--- a/Engine/Plugins/PCG/Source/PCG/Private/Compute/PCGDataForGPU.cpp
+++ b/Engine/Plugins/PCG/Source/PCG/Private/Compute/PCGDataForGPU.cpp
@@ -9,11 +9,13 @@
 #include "PCGNode.h"
 #include "PCGParamData.h"
 #include "PCGPoint.h"
+#include "PCGPointPropertiesTraits.h"
 #include "PCGSettings.h"
 #include "Compute/PCGComputeCommon.h"
 #include "Compute/PCGComputeGraph.h"
 #include "Compute/PCGDataBinding.h"
 #include "Compute/Data/PCGProxyForGPUData.h"
+#include "Data/PCGBasePointData.h"
 #include "Data/PCGPointData.h"
 #include "Elements/Metadata/PCGMetadataElementCommon.h"
 #include "Helpers/PCGAsync.h"
@@ -847,8 +849,50 @@ void FPCGDataDesc::AddAttribute(FPCGKernelAttributeKey InAttribute, const UPCGDa
 
 int32 FPCGDataDesc::GetElementCountForAttribute(const FPCGKernelAttributeDesc& AttributeDesc) const
 {
-	// Data domain only have a single element.
-	return AttributeDesc.AttributeKey.Identifier.MetadataDomain == PCGMetadataDomainID::Data ? 1 : ElementCount;
+	// A buffer that is not fully allocated uses only a single value to represent all elements.
+	// Data domain only has a single element.
+	if (!IsAttributeAllocated(AttributeDesc) || AttributeDesc.AttributeKey.Identifier.MetadataDomain == PCGMetadataDomainID::Data)
+	{
+		return 1;
+	}
+	else
+	{
+		return ElementCount;
+	}
+}
+
+bool FPCGDataDesc::IsAttributeAllocated(int32 InAttributeId) const
+{
+	// @todo_pcg: Only point properties are supported for single value ranges. All other attributes are fully allocated.
+	if (Type == EPCGDataType::Point && InAttributeId < PCGComputeConstants::NUM_RESERVED_ATTRS)
+	{
+		switch (InAttributeId)
+		{
+		case PCGComputeConstants::POINT_POSITION_ATTRIBUTE_ID:
+			return !!(AllocatedPointProperties & EPCGPointNativeProperties::Transform);
+		case PCGComputeConstants::POINT_ROTATION_ATTRIBUTE_ID:
+			return !!(AllocatedPointProperties & EPCGPointNativeProperties::Transform);
+		case PCGComputeConstants::POINT_SCALE_ATTRIBUTE_ID:
+			return !!(AllocatedPointProperties & EPCGPointNativeProperties::Transform);
+		case PCGComputeConstants::POINT_BOUNDS_MIN_ATTRIBUTE_ID:
+			return !!(AllocatedPointProperties & EPCGPointNativeProperties::BoundsMin);
+		case PCGComputeConstants::POINT_BOUNDS_MAX_ATTRIBUTE_ID:
+			return !!(AllocatedPointProperties & EPCGPointNativeProperties::BoundsMax);
+		case PCGComputeConstants::POINT_COLOR_ATTRIBUTE_ID:
+			return !!(AllocatedPointProperties & EPCGPointNativeProperties::Color);
+		case PCGComputeConstants::POINT_SEED_ATTRIBUTE_ID:
+			return !!(AllocatedPointProperties & EPCGPointNativeProperties::Seed);
+		case PCGComputeConstants::POINT_STEEPNESS_ATTRIBUTE_ID:
+			return !!(AllocatedPointProperties & EPCGPointNativeProperties::Steepness);
+		// Note: Because density is used to represent invalid points, it must always be allocated.
+		case PCGComputeConstants::POINT_DENSITY_ATTRIBUTE_ID:
+			return true;
+		default:
+			return true;
+		}
+	}
+
+	return true;
 }
 
 void FPCGDataDesc::InitializeAttributeDescs(const UPCGData* InData, const UPCGDataBinding* InBinding)
@@ -858,6 +902,11 @@ void FPCGDataDesc::InitializeAttributeDescs(const UPCGData* InData, const UPCGDa
 	if (Type == EPCGDataType::Point)
 	{
 		AttributeDescs.Append(PCGDataForGPUConstants::PointPropertyDescs, NUM_POINT_PROPERTIES);
+
+		if (const UPCGBasePointData* PointData = Cast<UPCGBasePointData>(InData))
+		{
+			AllocatedPointProperties = PointData->GetAllocatedProperties();
+		}
 	}
 	else { /* TODO: More types! */ }
 
@@ -1050,10 +1099,9 @@ void FPCGDataCollectionDesc::WriteHeader(TArray<uint32>& OutPackedDataCollection
 			// Scatter from attributes that are present into header which has slots for all possible attributes.
 			WriteAddressUints = (AttributeDesc.AttributeId * ATTRIBUTE_HEADER_SIZE_BYTES + DataAttributesHeaderStartAddressBytes) >> 2;
 
-			// Data i element j: packed ID and stride
-			const uint32 AttributeId = AttributeDesc.AttributeId;
-			const uint32 PackedIdAndStride = (AttributeId << 8) + AttributeStride;
-			OutPackedDataCollectionHeader[WriteAddressUints++] = PackedIdAndStride;
+			// Data i element j: packed attribute info and data address.
+			const uint32 PackedAttributeInfo = (DataDesc.IsAttributeAllocated(AttributeDesc) ? PCGComputeConstants::AttributeAllocatedMask : 0) + (AttributeDesc.AttributeId << 8) + AttributeStride;
+			OutPackedDataCollectionHeader[WriteAddressUints++] = PackedAttributeInfo;
 			OutPackedDataCollectionHeader[WriteAddressUints++] = DataStartAddressBytes;
 
 			// Move the DataStartAddress to the beginning of the next attribute by skipping all the bytes of the current attribute.
@@ -1126,16 +1174,15 @@ void FPCGDataCollectionDesc::PackDataCollection(const FPCGDataCollection& InData
 			continue;
 		}
 
-		if (const UPCGPointData* PointData = Cast<UPCGPointData>(InputData[DataIndex].Data))
+		if (const UPCGBasePointData* PointData = Cast<UPCGBasePointData>(InputData[DataIndex].Data))
 		{
-			const TArray<FPCGPoint>& Points = PointData->GetPoints();
-			if (Points.IsEmpty())
+			const uint32 NumElements = PointData->GetNumPoints();
+
+			if (NumElements == 0)
 			{
 				continue;
 			}
 
-			const uint32 NumElements = Points.Num();
-
 			for (const FPCGKernelAttributeDesc& AttributeDesc : DataDesc.AttributeDescs)
 			{
 				const uint32 AttributeId = AttributeDesc.AttributeId;
@@ -1147,14 +1194,18 @@ void FPCGDataCollectionDesc::PackDataCollection(const FPCGDataCollection& InData
 
 				if (AttributeId < NUM_RESERVED_ATTRS)
 				{
+					const uint32 NumElementsForAttribute = DataDesc.IsAttributeAllocated(AttributeDesc) ? NumElements : 1u;
+
 					// Point property.
 					switch (AttributeId)
 					{
 					case POINT_POSITION_ATTRIBUTE_ID:
 					{
-						for (uint32 ElementIndex = 0; ElementIndex < NumElements; ++ElementIndex)
+						const TConstPCGValueRange<FTransform> ReadValueRange = PointData->GetConstTransformValueRange();
+
+						for (uint32 ElementIndex = 0; ElementIndex < NumElementsForAttribute; ++ElementIndex)
 						{
-							const FVector Position = Points[ElementIndex].Transform.GetLocation();
+							const FVector Position = ReadValueRange[ElementIndex].GetLocation();
 							OutPackedDataCollection[AddressUints++] = FMath::AsUInt(static_cast<float>(Position.X));
 							OutPackedDataCollection[AddressUints++] = FMath::AsUInt(static_cast<float>(Position.Y));
 							OutPackedDataCollection[AddressUints++] = FMath::AsUInt(static_cast<float>(Position.Z));
@@ -1163,9 +1214,11 @@ void FPCGDataCollectionDesc::PackDataCollection(const FPCGDataCollection& InData
 					}
 					case POINT_ROTATION_ATTRIBUTE_ID:
 					{
-						for (uint32 ElementIndex = 0; ElementIndex < NumElements; ++ElementIndex)
+						const TConstPCGValueRange<FTransform> ReadValueRange = PointData->GetConstTransformValueRange();
+
+						for (uint32 ElementIndex = 0; ElementIndex < NumElementsForAttribute; ++ElementIndex)
 						{
-							const FQuat Rotation = Points[ElementIndex].Transform.GetRotation();
+							const FQuat Rotation = ReadValueRange[ElementIndex].GetRotation();
 							OutPackedDataCollection[AddressUints++] = FMath::AsUInt(static_cast<float>(Rotation.X));
 							OutPackedDataCollection[AddressUints++] = FMath::AsUInt(static_cast<float>(Rotation.Y));
 							OutPackedDataCollection[AddressUints++] = FMath::AsUInt(static_cast<float>(Rotation.Z));
@@ -1175,9 +1228,11 @@ void FPCGDataCollectionDesc::PackDataCollection(const FPCGDataCollection& InData
 					}
 					case POINT_SCALE_ATTRIBUTE_ID:
 					{
-						for (uint32 ElementIndex = 0; ElementIndex < NumElements; ++ElementIndex)
+						const TConstPCGValueRange<FTransform> ReadValueRange = PointData->GetConstTransformValueRange();
+
+						for (uint32 ElementIndex = 0; ElementIndex < NumElementsForAttribute; ++ElementIndex)
 						{
-							const FVector Scale = Points[ElementIndex].Transform.GetScale3D();
+							const FVector Scale = ReadValueRange[ElementIndex].GetScale3D();
 							OutPackedDataCollection[AddressUints++] = FMath::AsUInt(static_cast<float>(Scale.X));
 							OutPackedDataCollection[AddressUints++] = FMath::AsUInt(static_cast<float>(Scale.Y));
 							OutPackedDataCollection[AddressUints++] = FMath::AsUInt(static_cast<float>(Scale.Z));
@@ -1186,9 +1241,11 @@ void FPCGDataCollectionDesc::PackDataCollection(const FPCGDataCollection& InData
 					}
 					case POINT_BOUNDS_MIN_ATTRIBUTE_ID:
 					{
-						for (uint32 ElementIndex = 0; ElementIndex < NumElements; ++ElementIndex)
+						const TConstPCGValueRange<FVector> ReadValueRange = PointData->GetConstBoundsMinValueRange();
+
+						for (uint32 ElementIndex = 0; ElementIndex < NumElementsForAttribute; ++ElementIndex)
 						{
-							const FVector& BoundsMin = Points[ElementIndex].BoundsMin;
+							const FVector& BoundsMin = ReadValueRange[ElementIndex];
 							OutPackedDataCollection[AddressUints++] = FMath::AsUInt(static_cast<float>(BoundsMin.X));
 							OutPackedDataCollection[AddressUints++] = FMath::AsUInt(static_cast<float>(BoundsMin.Y));
 							OutPackedDataCollection[AddressUints++] = FMath::AsUInt(static_cast<float>(BoundsMin.Z));
@@ -1197,9 +1254,11 @@ void FPCGDataCollectionDesc::PackDataCollection(const FPCGDataCollection& InData
 					}
 					case POINT_BOUNDS_MAX_ATTRIBUTE_ID:
 					{
-						for (uint32 ElementIndex = 0; ElementIndex < NumElements; ++ElementIndex)
+						const TConstPCGValueRange<FVector> ReadValueRange = PointData->GetConstBoundsMaxValueRange();
+
+						for (uint32 ElementIndex = 0; ElementIndex < NumElementsForAttribute; ++ElementIndex)
 						{
-							const FVector& BoundsMax = Points[ElementIndex].BoundsMax;
+							const FVector& BoundsMax = ReadValueRange[ElementIndex];
 							OutPackedDataCollection[AddressUints++] = FMath::AsUInt(static_cast<float>(BoundsMax.X));
 							OutPackedDataCollection[AddressUints++] = FMath::AsUInt(static_cast<float>(BoundsMax.Y));
 							OutPackedDataCollection[AddressUints++] = FMath::AsUInt(static_cast<float>(BoundsMax.Z));
@@ -1208,9 +1267,11 @@ void FPCGDataCollectionDesc::PackDataCollection(const FPCGDataCollection& InData
 					}
 					case POINT_COLOR_ATTRIBUTE_ID:
 					{
-						for (uint32 ElementIndex = 0; ElementIndex < NumElements; ++ElementIndex)
+						const TConstPCGValueRange<FVector4> ReadValueRange = PointData->GetConstColorValueRange();
+
+						for (uint32 ElementIndex = 0; ElementIndex < NumElementsForAttribute; ++ElementIndex)
 						{
-							const FVector4& Color = Points[ElementIndex].Color;
+							const FVector4& Color = ReadValueRange[ElementIndex];
 							OutPackedDataCollection[AddressUints++] = FMath::AsUInt(static_cast<float>(Color.X));
 							OutPackedDataCollection[AddressUints++] = FMath::AsUInt(static_cast<float>(Color.Y));
 							OutPackedDataCollection[AddressUints++] = FMath::AsUInt(static_cast<float>(Color.Z));
@@ -1220,27 +1281,33 @@ void FPCGDataCollectionDesc::PackDataCollection(const FPCGDataCollection& InData
 					}
 					case POINT_DENSITY_ATTRIBUTE_ID:
 					{
-						for (uint32 ElementIndex = 0; ElementIndex < NumElements; ++ElementIndex)
+						const TConstPCGValueRange<float> ReadValueRange = PointData->GetConstDensityValueRange();
+
+						for (uint32 ElementIndex = 0; ElementIndex < NumElementsForAttribute; ++ElementIndex)
 						{
-							const float Density = Points[ElementIndex].Density;
+							const float Density = ReadValueRange[ElementIndex];
 							OutPackedDataCollection[AddressUints++] = FMath::AsUInt(Density);
 						}
 						break;
 					}
 					case POINT_SEED_ATTRIBUTE_ID:
 					{
-						for (uint32 ElementIndex = 0; ElementIndex < NumElements; ++ElementIndex)
+						const TConstPCGValueRange<int32> ReadValueRange = PointData->GetConstSeedValueRange();
+
+						for (uint32 ElementIndex = 0; ElementIndex < NumElementsForAttribute; ++ElementIndex)
 						{
-							const int Seed = Points[ElementIndex].Seed;
+							const int Seed = ReadValueRange[ElementIndex];
 							OutPackedDataCollection[AddressUints++] = Seed;
 						}
 						break;
 					}
 					case POINT_STEEPNESS_ATTRIBUTE_ID:
 					{
-						for (uint32 ElementIndex = 0; ElementIndex < NumElements; ++ElementIndex)
+						const TConstPCGValueRange<float> ReadValueRange = PointData->GetConstSteepnessValueRange();
+
+						for (uint32 ElementIndex = 0; ElementIndex < NumElementsForAttribute; ++ElementIndex)
 						{
-							const float Steepness = Points[ElementIndex].Steepness;
+							const float Steepness = ReadValueRange[ElementIndex];
 							OutPackedDataCollection[AddressUints++] = FMath::AsUInt(Steepness);
 						}
 						break;
@@ -1259,10 +1326,12 @@ void FPCGDataCollectionDesc::PackDataCollection(const FPCGDataCollection& InData
 					}
 					else
 					{
-						ensure(PCGDataForGPUHelpers::PackAttributeHelper(AttributeBase, AttributeDesc, Points[0].MetadataEntry, InDataBinding->GetStringTable(), OutPackedDataCollection, AddressUints));
+						const TConstPCGValueRange<int64> ReadValueRange = PointData->GetConstMetadataEntryValueRange();
+
+						ensure(PCGDataForGPUHelpers::PackAttributeHelper(AttributeBase, AttributeDesc, ReadValueRange[0], InDataBinding->GetStringTable(), OutPackedDataCollection, AddressUints));
 						for (uint32 ElementIndex = 1; ElementIndex < NumElements; ++ElementIndex)
 						{
-							PCGDataForGPUHelpers::PackAttributeHelper(AttributeBase, AttributeDesc, Points[ElementIndex].MetadataEntry, InDataBinding->GetStringTable(), OutPackedDataCollection, AddressUints);
+							PCGDataForGPUHelpers::PackAttributeHelper(AttributeBase, AttributeDesc, ReadValueRange[ElementIndex], InDataBinding->GetStringTable(), OutPackedDataCollection, AddressUints);
 						}
 					}
 				}
@@ -1409,6 +1478,9 @@ EPCGUnpackDataCollectionResult FPCGDataCollectionDesc::UnpackDataCollection(FPCG
 				continue;
 			}
 
+			FTransform DefaultTransform = FTransform::Identity;
+			bool bTransformIsAllocated = false;
+
 			// Loop over attributes.
 			for (const FPCGKernelAttributeDesc& AttributeDesc : AttributeDescs)
 			{
@@ -1420,122 +1492,221 @@ EPCGUnpackDataCollectionResult FPCGDataCollectionDesc::UnpackDataCollection(FPCG
 				{
 					TRACE_CPUPROFILER_EVENT_SCOPE(UnpackPointProperty);
 
+					const bool bAllocated = DataDesc.IsAttributeAllocated(AttributeDesc);
+
 					// We tried hoisting this decision to a lambda but it didn't appear to help.
 					switch (AttributeId)
 					{
 					case POINT_POSITION_ATTRIBUTE_ID:
 					{
-						TPCGValueRange<FTransform> TransformRange = OutPointData->GetTransformValueRange(/*bAlocate=*/true);
-						ParallelFor(NumElements, [DataAsFloat, AddressUints, &TransformRange](int32 ElementIndex)
+						bTransformIsAllocated = bAllocated;
+
+						if (bAllocated)
+						{
+							TPCGValueRange<FTransform> TransformRange = OutPointData->GetTransformValueRange(/*bAllocate=*/true);
+							ParallelFor(NumElements, [DataAsFloat, AddressUints, &TransformRange](int32 ElementIndex)
+							{
+								const FVector Location = FVector(
+									DataAsFloat[AddressUints + ElementIndex * 3 + 0],
+									DataAsFloat[AddressUints + ElementIndex * 3 + 1],
+									DataAsFloat[AddressUints + ElementIndex * 3 + 2]);
+
+								TransformRange[ElementIndex].SetLocation(Location);
+							});
+						}
+						else
 						{
 							const FVector Location = FVector(
-								DataAsFloat[AddressUints + ElementIndex * 3 + 0],
-								DataAsFloat[AddressUints + ElementIndex * 3 + 1],
-								DataAsFloat[AddressUints + ElementIndex * 3 + 2]);
+								DataAsFloat[AddressUints + /*ElementIndex=*/0 * 3 + 0],
+								DataAsFloat[AddressUints + /*ElementIndex=*/0 * 3 + 1],
+								DataAsFloat[AddressUints + /*ElementIndex=*/0 * 3 + 2]);
+
+							DefaultTransform.SetLocation(Location);
+						}
 
-							TransformRange[ElementIndex].SetLocation(Location);
-						});
 						break;
 					}
 					case POINT_ROTATION_ATTRIBUTE_ID:
 					{
-						TPCGValueRange<FTransform> TransformRange = OutPointData->GetTransformValueRange(/*bAlocate=*/true);
-						ParallelFor(NumElements, [DataAsFloat, AddressUints, &TransformRange](int32 ElementIndex)
+						if (bAllocated)
+						{
+							TPCGValueRange<FTransform> TransformRange = OutPointData->GetTransformValueRange(/*bAllocate=*/true);
+							ParallelFor(NumElements, [DataAsFloat, AddressUints, &TransformRange](int32 ElementIndex)
+							{
+								const FQuat Rotation = FQuat(
+									DataAsFloat[AddressUints + ElementIndex * 4 + 0],
+									DataAsFloat[AddressUints + ElementIndex * 4 + 1],
+									DataAsFloat[AddressUints + ElementIndex * 4 + 2],
+									DataAsFloat[AddressUints + ElementIndex * 4 + 3]);
+
+								// Normalize here with default tolerance (zero quat will return identity).
+								TransformRange[ElementIndex].SetRotation(Rotation.GetNormalized());
+							});
+						}
+						else
 						{
 							const FQuat Rotation = FQuat(
-								DataAsFloat[AddressUints + ElementIndex * 4 + 0],
-								DataAsFloat[AddressUints + ElementIndex * 4 + 1],
-								DataAsFloat[AddressUints + ElementIndex * 4 + 2],
-								DataAsFloat[AddressUints + ElementIndex * 4 + 3]);
-
-							// Normalize here with default tolerance (zero quat will return identity).
-							TransformRange[ElementIndex].SetRotation(Rotation.GetNormalized());
-						});
+								DataAsFloat[AddressUints + /*ElementIndex=*/0 * 4 + 0],
+								DataAsFloat[AddressUints + /*ElementIndex=*/0 * 4 + 1],
+								DataAsFloat[AddressUints + /*ElementIndex=*/0 * 4 + 2],
+								DataAsFloat[AddressUints + /*ElementIndex=*/0 * 4 + 3]);
+
+							DefaultTransform.SetRotation(Rotation);
+						}
 						break;
 					}
 					case POINT_SCALE_ATTRIBUTE_ID:
 					{
-						TPCGValueRange<FTransform> TransformRange = OutPointData->GetTransformValueRange(/*bAlocate=*/true);
-						ParallelFor(NumElements, [DataAsFloat, AddressUints, &TransformRange](int32 ElementIndex)
+						if (bAllocated)
+						{
+							TPCGValueRange<FTransform> TransformRange = OutPointData->GetTransformValueRange(/*bAllocate=*/true);
+							ParallelFor(NumElements, [DataAsFloat, AddressUints, &TransformRange](int32 ElementIndex)
+							{
+								const FVector Scale = FVector(
+									DataAsFloat[AddressUints + ElementIndex * 3 + 0],
+									DataAsFloat[AddressUints + ElementIndex * 3 + 1],
+									DataAsFloat[AddressUints + ElementIndex * 3 + 2]);
+
+								TransformRange[ElementIndex].SetScale3D(Scale);
+							});
+						}
+						else
 						{
-							const FVector Scale = FVector
-							(
-								DataAsFloat[AddressUints + ElementIndex * 3 + 0],
-								DataAsFloat[AddressUints + ElementIndex * 3 + 1],
-								DataAsFloat[AddressUints + ElementIndex * 3 + 2]);
-
-							TransformRange[ElementIndex].SetScale3D(Scale);
-						});
+							const FVector Scale = FVector(
+								DataAsFloat[AddressUints + /*ElementIndex=*/0 * 3 + 0],
+								DataAsFloat[AddressUints + /*ElementIndex=*/0 * 3 + 1],
+								DataAsFloat[AddressUints + /*ElementIndex=*/0 * 3 + 2]);
+
+							DefaultTransform.SetScale3D(Scale);
+						}
 						break;
 					}
 					case POINT_BOUNDS_MIN_ATTRIBUTE_ID:
 					{
-						TPCGValueRange<FVector> BoundsMinRange = OutPointData->GetBoundsMinValueRange(/*bAllocate=*/true);
-						ParallelFor(NumElements, [DataAsFloat, AddressUints, &BoundsMinRange](int32 ElementIndex)
+						if (bAllocated)
+						{
+							TPCGValueRange<FVector> BoundsMinRange = OutPointData->GetBoundsMinValueRange(/*bAllocate=*/true);
+							ParallelFor(NumElements, [DataAsFloat, AddressUints, &BoundsMinRange](int32 ElementIndex)
+							{
+								const FVector BoundsMin = FVector(
+									DataAsFloat[AddressUints + ElementIndex * 3 + 0],
+									DataAsFloat[AddressUints + ElementIndex * 3 + 1],
+									DataAsFloat[AddressUints + ElementIndex * 3 + 2]);
+
+								BoundsMinRange[ElementIndex] = BoundsMin;
+							});
+						}
+						else
 						{
 							const FVector BoundsMin = FVector(
-								DataAsFloat[AddressUints + ElementIndex * 3 + 0],
-								DataAsFloat[AddressUints + ElementIndex * 3 + 1],
-								DataAsFloat[AddressUints + ElementIndex * 3 + 2]);
+								DataAsFloat[AddressUints + /*ElementIndex=*/0 * 3 + 0],
+								DataAsFloat[AddressUints + /*ElementIndex=*/0 * 3 + 1],
+								DataAsFloat[AddressUints + /*ElementIndex=*/0 * 3 + 2]);
 
-							BoundsMinRange[ElementIndex] = BoundsMin;
-						});
+							OutPointData->SetBoundsMin(BoundsMin);
+						}
 						break;
 					}
 					case POINT_BOUNDS_MAX_ATTRIBUTE_ID:
 					{
-						TPCGValueRange<FVector> BoundsMaxRange = OutPointData->GetBoundsMaxValueRange(/*bAllocate=*/true);
-						ParallelFor(NumElements, [DataAsFloat, AddressUints, &BoundsMaxRange](int32 ElementIndex)
+						if (bAllocated)
+						{
+							TPCGValueRange<FVector> BoundsMaxRange = OutPointData->GetBoundsMaxValueRange(/*bAllocate=*/true);
+							ParallelFor(NumElements, [DataAsFloat, AddressUints, &BoundsMaxRange](int32 ElementIndex)
+							{
+								const FVector BoundsMax = FVector(
+									DataAsFloat[AddressUints + ElementIndex * 3 + 0],
+									DataAsFloat[AddressUints + ElementIndex * 3 + 1],
+									DataAsFloat[AddressUints + ElementIndex * 3 + 2]);
+
+								BoundsMaxRange[ElementIndex] = BoundsMax;
+							});
+						}
+						else
 						{
 							const FVector BoundsMax = FVector(
-								DataAsFloat[AddressUints + ElementIndex * 3 + 0],
-								DataAsFloat[AddressUints + ElementIndex * 3 + 1],
-								DataAsFloat[AddressUints + ElementIndex * 3 + 2]);
+								DataAsFloat[AddressUints + /*ElementIndex=*/0 * 3 + 0],
+								DataAsFloat[AddressUints + /*ElementIndex=*/0 * 3 + 1],
+								DataAsFloat[AddressUints + /*ElementIndex=*/0 * 3 + 2]);
 
-							BoundsMaxRange[ElementIndex] = BoundsMax;
-						});
+							OutPointData->SetBoundsMax(BoundsMax);
+						}
 						break;
 					}
 					case POINT_COLOR_ATTRIBUTE_ID:
 					{
-						TPCGValueRange<FVector4> ColorRange = OutPointData->GetColorValueRange(/*bAllocate=*/true);
-						ParallelFor(NumElements, [DataAsFloat, AddressUints, &ColorRange](int32 ElementIndex)
+						if (bAllocated)
+						{
+							TPCGValueRange<FVector4> ColorRange = OutPointData->GetColorValueRange(/*bAllocate=*/true);
+							ParallelFor(NumElements, [DataAsFloat, AddressUints, &ColorRange](int32 ElementIndex)
+							{
+								const FVector4 Color = FVector4(
+									DataAsFloat[AddressUints + ElementIndex * 4 + 0],
+									DataAsFloat[AddressUints + ElementIndex * 4 + 1],
+									DataAsFloat[AddressUints + ElementIndex * 4 + 2],
+									DataAsFloat[AddressUints + ElementIndex * 4 + 3]);
+
+								ColorRange[ElementIndex] = Color;
+							});
+						}
+						else
 						{
 							const FVector4 Color = FVector4(
-								DataAsFloat[AddressUints + ElementIndex * 4 + 0],
-								DataAsFloat[AddressUints + ElementIndex * 4 + 1],
-								DataAsFloat[AddressUints + ElementIndex * 4 + 2],
-								DataAsFloat[AddressUints + ElementIndex * 4 + 3]);
+								DataAsFloat[AddressUints + /*ElementIndex=*/0 * 4 + 0],
+								DataAsFloat[AddressUints + /*ElementIndex=*/0 * 4 + 1],
+								DataAsFloat[AddressUints + /*ElementIndex=*/0 * 4 + 2],
+								DataAsFloat[AddressUints + /*ElementIndex=*/0 * 4 + 3]);
 
-							ColorRange[ElementIndex] = Color;
-						});
+							OutPointData->SetColor(Color);
+						}
 						break;
 					}
 					case POINT_DENSITY_ATTRIBUTE_ID:
 					{
-						TPCGValueRange<float> DensityRange = OutPointData->GetDensityValueRange(/*bAllocate=*/true);
-						ParallelFor(NumElements, [DataAsFloat, AddressUints, &DensityRange](int32 ElementIndex)
+						if (bAllocated)
 						{
-							DensityRange[ElementIndex] = DataAsFloat[AddressUints + ElementIndex];
-						});
+							TPCGValueRange<float> DensityRange = OutPointData->GetDensityValueRange(/*bAllocate=*/true);
+							ParallelFor(NumElements, [DataAsFloat, AddressUints, &DensityRange](int32 ElementIndex)
+							{
+								DensityRange[ElementIndex] = DataAsFloat[AddressUints + ElementIndex];
+							});
+						}
+						else
+						{
+							OutPointData->SetDensity(DataAsFloat[AddressUints]);
+						}
 						break;
 					}
 					case POINT_SEED_ATTRIBUTE_ID:
 					{
-						TPCGValueRange<int32> SeedRange = OutPointData->GetSeedValueRange(/*bAllocate=*/true);
-						ParallelFor(NumElements, [DataAsInt, AddressUints, &SeedRange](int32 ElementIndex)
+						if (bAllocated)
+						{
+							TPCGValueRange<int32> SeedRange = OutPointData->GetSeedValueRange(/*bAllocate=*/true);
+							ParallelFor(NumElements, [DataAsInt, AddressUints, &SeedRange](int32 ElementIndex)
+							{
+								SeedRange[ElementIndex] = DataAsInt[AddressUints + ElementIndex];
+							});
+						}
+						else
 						{
-							SeedRange[ElementIndex] = DataAsInt[AddressUints + ElementIndex];
-						});
+							OutPointData->SetSeed(DataAsInt[AddressUints]);
+						}
 						break;
 					}
 					case POINT_STEEPNESS_ATTRIBUTE_ID:
 					{
-						TPCGValueRange<float> SteepnessRange = OutPointData->GetSteepnessValueRange(/*bAllocate=*/true);
-						ParallelFor(NumElements, [DataAsFloat, AddressUints, &SteepnessRange](int32 ElementIndex)
+						if (bAllocated)
+						{
+							TPCGValueRange<float> SteepnessRange = OutPointData->GetSteepnessValueRange(/*bAllocate=*/true);
+							ParallelFor(NumElements, [DataAsFloat, AddressUints, &SteepnessRange](int32 ElementIndex)
+							{
+								SteepnessRange[ElementIndex] = DataAsFloat[AddressUints + ElementIndex];
+							});
+						}
+						else
 						{
-							SteepnessRange[ElementIndex] = DataAsFloat[AddressUints + ElementIndex];
-						});
+							OutPointData->SetSteepness(DataAsFloat[AddressUints]);
+						}
 						break;
 					}
 					default:
@@ -1554,6 +1725,11 @@ EPCGUnpackDataCollectionResult FPCGDataCollectionDesc::UnpackDataCollection(FPCG
 				}
 			}
 
+			if (!bTransformIsAllocated)
+			{
+				OutPointData->SetTransform(DefaultTransform);
+			}
+
 			// TODO: It may be more efficient to create a mapping from input point index to final output point index and do everything in one pass.
 			auto ProcessRangeFunc = [OutPointData](int32 StartReadIndex, int32 StartWriteIndex, int32 Count)
 			{
@@ -1745,6 +1921,14 @@ void FPCGDataCollectionDesc::AddAttributeToAllData(FPCGKernelAttributeKey InAttr
 	}
 }
 
+void FPCGDataCollectionDesc::AllocatePropertiesForAllData(EPCGPointNativeProperties InProperties)
+{
+	for (FPCGDataDesc& DataDesc : DataDescs)
+	{
+		DataDesc.AllocateProperties(InProperties);
+	}
+}
+
 void FPCGDataCollectionDesc::GetUniqueStringKeyValues(int32 InAttributeId, TArray<int32>& OutUniqueStringKeys) const
 {
 	for (const FPCGDataDesc& DataDesc : DataDescs)
diff --git a/Engine/Plugins/PCG/Source/PCG/Private/Elements/PCGCopyPointsKernel.cpp b/Engine/Plugins/PCG/Source/PCG/Private/Elements/PCGCopyPointsKernel.cpp
--- a/Engine/Plugins/PCG/Source/PCG/Private/Elements/PCGCopyPointsKernel.cpp
+++ b/Engine/Plugins/PCG/Source/PCG/Private/Elements/PCGCopyPointsKernel.cpp
@@ -110,6 +110,9 @@ bool UPCGCopyPointsKernel::ComputeOutputBindingDataDesc(const UPCGComputeGraph*
 	const int32 NumTargets = TargetPinDesc.DataDescs.Num();
 	const int32 NumIterations = CPSettings->bCopyEachSourceOnEveryTarget ? NumSources * NumTargets : FMath::Max(NumSources, NumTargets);
 
+	// Properties that might need to be allocated
+	const EPCGPointNativeProperties OptionalProperties = EPCGPointNativeProperties::Seed | EPCGPointNativeProperties::Color;
+
 	if (NumSources > 0 && NumTargets > 0 && (CPSettings->bCopyEachSourceOnEveryTarget || NumSources == NumTargets || NumSources == 1 || NumTargets == 1))
 	{
 		for (int32 I = 0; I < NumIterations; ++I)
@@ -169,6 +172,34 @@ bool UPCGCopyPointsKernel::ComputeOutputBindingDataDesc(const UPCGComputeGraph*
 					ResultDataDesc.TagStringKeys.AddUnique(TargetTagStringKey);
 				}
 			}
+
+			const EPCGPointNativeProperties SourceAllocatedProperties = SourceDesc.GetAllocatedProperties();
+			const EPCGPointNativeProperties TargetAllocatedProperties = TargetDesc.GetAllocatedProperties();
+
+			// Allocate all properties from source except properties that might not need to be.
+			EPCGPointNativeProperties PropertiesToAllocate = SourceAllocatedProperties & ~OptionalProperties;
+
+			auto GetPropertyToAllocate = [&PropertiesToAllocate, SourceAllocatedProperties, TargetAllocatedProperties](EPCGCopyPointsInheritanceMode Inheritance, EPCGPointNativeProperties Property)
+			{
+				if (Inheritance == EPCGCopyPointsInheritanceMode::Relative)
+				{
+					return ((SourceAllocatedProperties | TargetAllocatedProperties) & Property);
+				}
+				else if (Inheritance == EPCGCopyPointsInheritanceMode::Source)
+				{
+					return (SourceAllocatedProperties & Property);
+				}
+				else // if (Inheritance == EPCGCopyPointsInheritanceMode::Target)
+				{
+					return (TargetAllocatedProperties & Property);
+				}
+			};
+
+			PropertiesToAllocate |= EPCGPointNativeProperties::Transform;
+			PropertiesToAllocate |= GetPropertyToAllocate(CPSettings->SeedInheritance, EPCGPointNativeProperties::Seed);
+			PropertiesToAllocate |= GetPropertyToAllocate(CPSettings->ColorInheritance, EPCGPointNativeProperties::Color);
+
+			ResultDataDesc.AllocateProperties(PropertiesToAllocate);
 		}
 	}
 
diff --git a/Engine/Plugins/PCG/Source/PCG/Private/Elements/PCGSkinnedMeshSpawnerKernel.cpp b/Engine/Plugins/PCG/Source/PCG/Private/Elements/PCGSkinnedMeshSpawnerKernel.cpp
--- a/Engine/Plugins/PCG/Source/PCG/Private/Elements/PCGSkinnedMeshSpawnerKernel.cpp
+++ b/Engine/Plugins/PCG/Source/PCG/Private/Elements/PCGSkinnedMeshSpawnerKernel.cpp
@@ -47,6 +47,12 @@ bool UPCGSkinnedMeshSpawnerKernel::ComputeOutputBindingDataDesc(const UPCGComput
 	const FPCGKernelPin InputKernelPin(KernelIndex, PCGPinConstants::DefaultInputLabel, /*bIsInput=*/true);
 	ensure(InBinding->ComputeKernelPinDataDesc(InputKernelPin, &PinDesc));
 
+	const UPCGSkinnedMeshSpawnerSettings* SMSettings = CastChecked<UPCGSkinnedMeshSpawnerSettings>(GetSettings());
+	if (SMSettings->bApplyMeshBoundsToPoints)
+	{
+		PinDesc.AllocatePropertiesForAllData(EPCGPointNativeProperties::BoundsMin | EPCGPointNativeProperties::BoundsMax);
+	}
+
 	OutDataDesc = MoveTemp(PinDesc);
 	return true;
 }
diff --git a/Engine/Plugins/PCG/Source/PCG/Private/Elements/PCGStaticMeshSpawnerKernel.cpp b/Engine/Plugins/PCG/Source/PCG/Private/Elements/PCGStaticMeshSpawnerKernel.cpp
--- a/Engine/Plugins/PCG/Source/PCG/Private/Elements/PCGStaticMeshSpawnerKernel.cpp
+++ b/Engine/Plugins/PCG/Source/PCG/Private/Elements/PCGStaticMeshSpawnerKernel.cpp
@@ -52,11 +52,11 @@ bool UPCGStaticMeshSpawnerKernel::ComputeOutputBindingDataDesc(const UPCGCompute
 	{
 		const UPCGStaticMeshSpawnerSettings* SMSettings = CastChecked<UPCGStaticMeshSpawnerSettings>(GetSettings());
 
+		TArray<int32> UniqueStringKeys;
+
 		// Create unique value keys for the output string.
 		if (UPCGMeshSelectorWeighted* SelectorWeighted = Cast<UPCGMeshSelectorWeighted>(SMSettings->MeshSelectorParameters))
 		{
-			TArray<int32> UniqueStringKeys;
-
 			// Weighted selection - add explicit strings from settings.
 			for (const FPCGMeshSelectorWeightedEntry& Entry : SelectorWeighted->MeshEntries)
 			{
@@ -90,6 +90,11 @@ bool UPCGStaticMeshSpawnerKernel::ComputeOutputBindingDataDesc(const UPCGCompute
 						}
 					}
 
+					for (int32 StringKey : StringKeys)
+					{
+						UniqueStringKeys.AddUnique(StringKey);
+					}
+
 					DataDesc.AddAttribute(FPCGKernelAttributeKey(SMSettings->OutAttributeName, EPCGKernelAttributeType::StringKey), InBinding, &StringKeys);
 				}
 			}
@@ -98,6 +103,12 @@ bool UPCGStaticMeshSpawnerKernel::ComputeOutputBindingDataDesc(const UPCGCompute
 		{
 			UE_LOG(LogPCG, Error, TEXT("Mesh selector not supported by GPU Static Mesh Spawner: %s"), *SMSettings->MeshSelectorParameters->GetName());
 		}
+
+		// Allocate bounds if the spawner will write more than one unique value for the mesh bounds.
+		if (SMSettings->bApplyMeshBoundsToPoints && UniqueStringKeys.Num() > 1)
+		{
+			PinDesc.AllocatePropertiesForAllData(EPCGPointNativeProperties::BoundsMin | EPCGPointNativeProperties::BoundsMax);
+		}
 	}
 
 	OutDataDesc = MoveTemp(PinDesc);
