[PCG] Runtime Gen Sources: deprecate usage of WP streaming sources
- Instead properly gather sources from world's Player controller list (optional enabled by default)
- Add new CVar:s pcg.RuntimeGeneration.Server.Enable & pcg.RuntimeGeneration.ListenServer.Enable (enabled by default)

diff --git a/Engine/Plugins/PCG/Source/PCG/Private/Elements/PCGGetExecutionContext.cpp b/Engine/Plugins/PCG/Source/PCG/Private/Elements/PCGGetExecutionContext.cpp
--- a/Engine/Plugins/PCG/Source/PCG/Private/Elements/PCGGetExecutionContext.cpp
+++ b/Engine/Plugins/PCG/Source/PCG/Private/Elements/PCGGetExecutionContext.cpp
@@ -98,7 +98,11 @@ bool FPCGGetExecutionContextElement::ExecuteInternal(FPCGContext* Context) const
 	}
 	else if (Mode == EPCGGetExecutionContextMode::IsDedicatedServer)
 	{
-		Value = SupportingWorld && SupportingWorld->IsNetMode(NM_DedicatedServer);
+		Value = PCGHelpers::IsDedicatedServer(SupportingWorld);
+	}
+	else if (Mode == EPCGGetExecutionContextMode::IsListenServer)
+	{
+		Value = PCGHelpers::IsListenServer(SupportingWorld);
 	}
 	else if (Mode == EPCGGetExecutionContextMode::HasAuthority)
 	{
diff --git a/Engine/Plugins/PCG/Source/PCG/Private/Helpers/PCGHelpers.cpp b/Engine/Plugins/PCG/Source/PCG/Private/Helpers/PCGHelpers.cpp
--- a/Engine/Plugins/PCG/Source/PCG/Private/Helpers/PCGHelpers.cpp
+++ b/Engine/Plugins/PCG/Source/PCG/Private/Helpers/PCGHelpers.cpp
@@ -755,6 +755,37 @@ namespace PCGHelpers
 		return RandomIndices;
 	}
 
+	bool IsServer(const UWorld* InWorld)
+	{
+		return IsDedicatedServer(InWorld) || IsListenServer(InWorld);
+	}
+	
+	bool IsDedicatedServer(const UWorld* InWorld)
+	{
+		if (InWorld)
+		{
+			const ENetMode NetMode = InWorld->GetNetMode();
+			return NetMode == NM_DedicatedServer;
+		}
+		else
+		{
+			return false;
+		}
+	}
+
+	bool IsListenServer(const UWorld* InWorld)
+	{
+		if (InWorld)
+		{
+			const ENetMode NetMode = InWorld->GetNetMode();
+			return NetMode == NM_ListenServer;
+		}
+		else
+		{
+			return false;
+		}
+	}
+
 #if UE_ENABLE_DEBUG_DRAWING
 	void DebugDrawGenerationVolume(FPCGContext* InContext, const FColor* InOverrideColor)
 	{
diff --git a/Engine/Plugins/PCG/Source/PCG/Private/PCGWorldActor.cpp b/Engine/Plugins/PCG/Source/PCG/Private/PCGWorldActor.cpp
--- a/Engine/Plugins/PCG/Source/PCG/Private/PCGWorldActor.cpp
+++ b/Engine/Plugins/PCG/Source/PCG/Private/PCGWorldActor.cpp
@@ -319,3 +319,18 @@ void APCGWorldActor::BeginDestroy()
 	UnregisterFromSubsystem();
 	Super::BeginDestroy();
 }
+
+void APCGWorldActor::Serialize(FArchive& Archive)
+{
+	Archive.UsingCustomVersion(FFortniteMainBranchObjectVersion::GUID);
+
+	if (Archive.IsLoading() && Archive.CustomVer(FFortniteMainBranchObjectVersion::GUID) < FFortniteMainBranchObjectVersion::PCGDeprecateWorldPartitionGenerationSources)
+	{
+PRAGMA_DISABLE_DEPRECATION_WARNINGS
+		// Default value is now 'false' but when loading older data set default back to 'true' (previous default)
+		bEnableWorldPartitionGenerationSources = true;
+PRAGMA_ENABLE_DEPRECATION_WARNINGS
+	}
+
+	Super::Serialize(Archive);
+}
diff --git a/Engine/Plugins/PCG/Source/PCG/Private/RuntimeGen/PCGGenSourceManager.cpp b/Engine/Plugins/PCG/Source/PCG/Private/RuntimeGen/PCGGenSourceManager.cpp
--- a/Engine/Plugins/PCG/Source/PCG/Private/RuntimeGen/PCGGenSourceManager.cpp
+++ b/Engine/Plugins/PCG/Source/PCG/Private/RuntimeGen/PCGGenSourceManager.cpp
@@ -4,6 +4,7 @@
 
 #include "PCGModule.h"
 #include "PCGWorldActor.h"
+#include "Helpers/PCGHelpers.h"
 #include "RuntimeGen/GenSources/PCGGenSourceBase.h"
 #include "RuntimeGen/GenSources/PCGGenSourceComponent.h"
 #include "RuntimeGen/GenSources/PCGGenSourceEditorCamera.h"
@@ -19,15 +20,24 @@
 #include "EditorViewportClient.h"
 #endif
 
+namespace PCGGenSourceManager
+{
+	static TAutoConsoleVariable<bool> CVarDedicatedServerRuntimeGenerationEnable(
+		TEXT("pcg.RuntimeGeneration.Server.Enable"),
+		true,
+		TEXT("Enable the RuntimeGeneration system on dedicated servers."));
+
+	static TAutoConsoleVariable<bool> CVarListenServerRuntimeGenerationEnable(
+		TEXT("pcg.RuntimeGeneration.ListenServer.Enable"),
+		true,
+		TEXT("Enable the RuntimeGeneration system on listen servers."));
+}
+
 FPCGGenSourceManager::FPCGGenSourceManager(const UWorld* InWorld)
 {
 	// We capture the World so we can differentiate between editor-world and PIE-world Generation Sources.
 	World = InWorld;
 
-	// Collect PlayerControllers.
-	FGameModeEvents::GameModePostLoginEvent.AddRaw(this, &FPCGGenSourceManager::OnGameModePostLogin);
-	FGameModeEvents::GameModeLogoutEvent.AddRaw(this, &FPCGGenSourceManager::OnGameModePostLogout);
-
 #if WITH_EDITOR
 	EditorCameraGenSource = NewObject<UPCGGenSourceEditorCamera>();
 #endif
@@ -62,27 +72,8 @@ TSet<IPCGGenSourceBase*> FPCGGenSourceManager::GetAllGenSources(const APCGWorldA
 		AllGenSources.Add(NamedGenSource.Value.GetInterface());
 	}
 
-	// Acquire a generation source for each active streaming source that doesn't have a generation source already (based on name collision test).
-	for (TObjectPtr<UPCGGenSourceWPStreamingSource> WPGenSource : WorldPartitionGenSources)
-	{
-		bool bAlreadyTracked = false;
-
-		if (TScriptInterface<IPCGGenSourceBase>* FoundGenSource = (WPGenSource && WPGenSource->StreamingSource) ? RegisteredNamedGenSources.Find(WPGenSource->StreamingSource->Name) : nullptr)
-		{
-			// We've seen during cinematics there are two generation sources with the same name but with different locations in the world, so match based on position too.
-			TOptional<FVector> WPGenSourcePosition = WPGenSource->GetPosition();
-			TOptional<FVector> FoundGenSourcePosition = (*FoundGenSource)->GetPosition();
-			if (WPGenSourcePosition && FoundGenSourcePosition && (*FoundGenSourcePosition - *WPGenSourcePosition).SquaredLength() < UE_KINDA_SMALL_NUMBER)
-			{
-				bAlreadyTracked = true;
-			}
-		}
-
-		if (!bAlreadyTracked)
-		{
-			AllGenSources.Add(WPGenSource);
-		}
-	}
+	// Support deprecated behavior
+	GatherWorldPartitionGenSources(AllGenSources);
 
 #if WITH_EDITOR
 	// Acquire a generation source for the active editor viewport if one exists.
@@ -159,48 +150,111 @@ bool FPCGGenSourceManager::UnregisterGenSource(FName InGenSourceName)
 	return RegisteredNamedGenSources.Remove(InGenSourceName) > 0;
 }
 
-void FPCGGenSourceManager::OnGameModePostLogin(AGameModeBase* InGameMode, APlayerController* InPlayerController)
+void FPCGGenSourceManager::UpdatePerTickGenSources(const APCGWorldActor* InPCGWorldActor)
 {
-	if (InPlayerController == nullptr || InPlayerController->GetWorld() != World)
+	// Gather Player controller sources
+	TSet<APlayerController*> CurrentControllers;
+
+	const bool bIsDedicatedServer = PCGHelpers::IsDedicatedServer(World);
+	const bool bIsListenServer = PCGHelpers::IsListenServer(World);
+	
+	if (InPCGWorldActor->bTreatPlayerControllersAsGenerationSources &&
+		((!bIsDedicatedServer && !bIsListenServer) 
+		|| (bIsDedicatedServer && PCGGenSourceManager::CVarDedicatedServerRuntimeGenerationEnable.GetValueOnAnyThread())
+		|| (bIsListenServer && PCGGenSourceManager::CVarListenServerRuntimeGenerationEnable.GetValueOnAnyThread())))
 	{
-		return;
+		// Update existing sources or add new ones
+		for (auto It = World->GetPlayerControllerIterator(); It; ++It)
+		{
+			APlayerController* PlayerController = It->Get();
+			if (!PlayerController)
+			{
+				continue;
+			}
+
+			CurrentControllers.Add(PlayerController);
+
+			if (TScriptInterface<IPCGGenSourceBase>* FoundSource = RegisteredNamedGenSources.Find(PlayerController->GetFName()))
+			{
+				// Only update if the registered source is a player controller
+				if (UPCGGenSourcePlayer* SourcePlayer = Cast<UPCGGenSourcePlayer>(FoundSource->GetInterface()))
+				{
+					SourcePlayer->SetPlayerController(PlayerController);
+				}
+			}
+			else
+			{
+				UPCGGenSourcePlayer* GenSource = NewObject<UPCGGenSourcePlayer>();
+				GenSource->SetPlayerController(PlayerController);
+				RegisterGenSource(GenSource, PlayerController->GetFName());
+			}
+		}
 	}
 
-	ensure(IsInGameThread());
+	// Remove obsolete sources
+	for (TMap<FName, TScriptInterface<IPCGGenSourceBase>>::TIterator It = RegisteredNamedGenSources.CreateIterator(); It; ++It)
+	{
+		if (UPCGGenSourcePlayer* SourcePlayer = Cast<UPCGGenSourcePlayer>(It.Value().GetInterface()); !CurrentControllers.Contains(SourcePlayer->GetPlayerController().Get()))
+		{
+			It.RemoveCurrent();
+		}
+	}
+
+	// Support deprecated behavior
+	UpdateWorldPartitionGenSources(InPCGWorldActor);
 
-	UPCGGenSourcePlayer* GenSource = NewObject<UPCGGenSourcePlayer>();
-	GenSource->SetPlayerController(InPlayerController);
+#if WITH_EDITOR
+	EditorCameraGenSource->EditorViewportClient = nullptr;
 
-	RegisterGenSource(GenSource, InPlayerController->GetFName());
+	// Update the active editor viewport client for the EditorCameraGenSource, only if requested by the world actor, in-editor, and the viewport is visible.
+	if (InPCGWorldActor->bTreatEditorViewportAsGenerationSource && !World->IsGameWorld() && GEditor)
+	{
+		if (FViewport* Viewport = GEditor->GetActiveViewport())
+		{
+			if (FEditorViewportClient* ViewportClient = static_cast<FEditorViewportClient*>(Viewport->GetClient()))
+			{
+				if (ViewportClient->IsVisible())
+				{
+					EditorCameraGenSource->EditorViewportClient = ViewportClient;
+				}
+			}
+		}
+	}
+#endif
 }
 
-void FPCGGenSourceManager::OnGameModePostLogout(AGameModeBase* InGameMode, AController* InController)
+void FPCGGenSourceManager::GatherWorldPartitionGenSources(TSet<IPCGGenSourceBase*>& OutGenSources) const
 {
-	if (!InController || InController->GetWorld() != World)
+PRAGMA_DISABLE_DEPRECATION_WARNINGS
+	// Acquire a generation source for each active streaming source that doesn't have a generation source already (based on name collision test).
+	for (TObjectPtr<UPCGGenSourceWPStreamingSource> WPGenSource : WorldPartitionGenSources)
 	{
-		return;
-	}
+		bool bAlreadyTracked = false;
 
-	ensure(IsInGameThread());
+		if (const TScriptInterface<IPCGGenSourceBase>* FoundGenSource = (WPGenSource && WPGenSource->StreamingSource) ? RegisteredNamedGenSources.Find(WPGenSource->StreamingSource->Name) : nullptr)
+		{
+		// We've seen during cinematics there are two generation sources with the same name but with different locations in the world, so match based on position too.
+			TOptional<FVector> WPGenSourcePosition = WPGenSource->GetPosition();
+			TOptional<FVector> FoundGenSourcePosition = (*FoundGenSource)->GetPosition();
+			if (WPGenSourcePosition && FoundGenSourcePosition && (*FoundGenSourcePosition - *WPGenSourcePosition).SquaredLength() < UE_KINDA_SMALL_NUMBER)
+			{
+				bAlreadyTracked = true;
+			}
+		}
 
-	TScriptInterface<IPCGGenSourceBase> GenSource = nullptr;
-	if (RegisteredNamedGenSources.RemoveAndCopyValue(InController->GetFName(), GenSource))
-	{
-		check(GenSource);
-		UPCGGenSourcePlayer* GenSourcePlayer = Cast<UPCGGenSourcePlayer>(GenSource.GetInterface());
-		if (!GenSourcePlayer)
+		if (!bAlreadyTracked)
 		{
-			// The gen source we removed isn't a player controller gen source, so let's put it back.
-			RegisteredNamedGenSources.Add(InController->GetFName(), GenSource);
+			OutGenSources.Add(WPGenSource);
 		}
 	}
+PRAGMA_ENABLE_DEPRECATION_WARNINGS
 }
 
-void FPCGGenSourceManager::UpdatePerTickGenSources(const APCGWorldActor* InPCGWorldActor)
+void FPCGGenSourceManager::UpdateWorldPartitionGenSources(const APCGWorldActor* InPCGWorldActor)
 {
+PRAGMA_DISABLE_DEPRECATION_WARNINGS
+	// Deprecated : UPCGGenSourceComponent should be used instead
 	// Start by updating, adding and removing WP generation sources if it's needed.
-	// TODO: Grab StreamingSourceProviders instead of StreamingSources?
-	// TODO: Is it possible to avoid adding a StreamingSource for the Player, which we already capture in OnGameModePostLogin?
 	// Note: GetStreamingSources only works in GameWorld, so StreamingSources do not act as generation sources in editor.
 	const TArray<FWorldPartitionStreamingSource>* WPStreamingSources = nullptr;
 
@@ -208,6 +262,13 @@ void FPCGGenSourceManager::UpdatePerTickGenSources(const APCGWorldActor* InPCGWo
 		InPCGWorldActor->GetWorld() &&
 		InPCGWorldActor->GetWorld()->GetWorldPartition())
 	{
+		// Log deprecation warning once
+		if (!bLoggedWorldPartitionGenSourceWarning && InPCGWorldActor->GetWorld()->IsGameWorld())
+		{
+			UE_LOG(LogPCG, Warning, TEXT("[RUNTIMEGEN] APCGWorldActor::bEnableWorldPartitionGenerationSources is deprecated and will be removed. You can disable this property to remove the warning. By default, player controllers act as generation sources and you can use UPCGGenSourceComponent to provide custom generation sources."));
+			bLoggedWorldPartitionGenSourceWarning = true;
+		}
+
 		WPStreamingSources = &InPCGWorldActor->GetWorld()->GetWorldPartition()->GetStreamingSources();
 	}
 
@@ -217,12 +278,12 @@ void FPCGGenSourceManager::UpdatePerTickGenSources(const APCGWorldActor* InPCGWo
 	const int NumGenerationSourcesToRemove = FMath::Max(0, NumWPGenerationSources - NumWPStreamingSources);
 	const int NumGenerationSourcesToAdd = FMath::Max(0, NumWPStreamingSources - NumWPGenerationSources);
 	check((NumGenerationSourcesToRemove == 0 && NumGenerationSourcesToAdd == 0) || ((NumGenerationSourcesToRemove > 0) != (NumGenerationSourcesToAdd > 0)));
-	
+
 	for (int I = 0; I < NumGenerationSourcesToAdd; ++I)
 	{
 		UPCGGenSourceWPStreamingSource* GenSource = WorldPartitionGenSources.Add_GetRef(NewObject<UPCGGenSourceWPStreamingSource>());
 	}
-	
+
 	check((NumGenerationSourcesToUpdate + NumGenerationSourcesToAdd == 0) || WPStreamingSources);
 	// Implementation note: not really needed but we need to appease the static analysis gods
 	if (WPStreamingSources)
@@ -234,25 +295,7 @@ void FPCGGenSourceManager::UpdatePerTickGenSources(const APCGWorldActor* InPCGWo
 	}
 
 	WorldPartitionGenSources.SetNum(NumGenerationSourcesToUpdate + NumGenerationSourcesToAdd);
-
-#if WITH_EDITOR
-	EditorCameraGenSource->EditorViewportClient = nullptr;
-
-	// Update the active editor viewport client for the EditorCameraGenSource, only if requested by the world actor, in-editor, and the viewport is visible.
-	if (InPCGWorldActor->bTreatEditorViewportAsGenerationSource && !World->IsGameWorld() && GEditor)
-	{
-		if (FViewport* Viewport = GEditor->GetActiveViewport())
-		{
-			if (FEditorViewportClient* ViewportClient = static_cast<FEditorViewportClient*>(Viewport->GetClient()))
-			{
-				if (ViewportClient->IsVisible())
-				{
-					EditorCameraGenSource->EditorViewportClient = ViewportClient;
-				}
-			}
-		}
-	}
-#endif
+PRAGMA_ENABLE_DEPRECATION_WARNINGS
 }
 
 void FPCGGenSourceManager::AddReferencedObjects(FReferenceCollector& Collector)
@@ -260,7 +303,10 @@ void FPCGGenSourceManager::AddReferencedObjects(FReferenceCollector& Collector)
 #if WITH_EDITOR
 	Collector.AddReferencedObject(EditorCameraGenSource);
 #endif
+
+PRAGMA_DISABLE_DEPRECATION_WARNINGS
 	Collector.AddReferencedObjects(WorldPartitionGenSources);
+PRAGMA_ENABLE_DEPRECATION_WARNINGS
 
 	for (TScriptInterface<IPCGGenSourceBase>& GenSource : RegisteredGenSources)
 	{
diff --git a/Engine/Plugins/PCG/Source/PCG/Public/Elements/PCGGetExecutionContext.h b/Engine/Plugins/PCG/Source/PCG/Public/Elements/PCGGetExecutionContext.h
--- a/Engine/Plugins/PCG/Source/PCG/Public/Elements/PCGGetExecutionContext.h
+++ b/Engine/Plugins/PCG/Source/PCG/Public/Elements/PCGGetExecutionContext.h
@@ -17,6 +17,7 @@ enum class EPCGGetExecutionContextMode : uint8
 	IsPartitioned UMETA(Tooltip = "Returns whether the executing context is the original context and that it is partitioned."),
 	IsRuntimeGeneration UMETA(Tooltip="Returns whether the executing context is doing runtime gneration."),
 	IsDedicatedServer UMETA(Tooltip="Returns whether the executing context is running on a dedicated server."),
+	IsListenServer UMETA(Tooltip="Returns whether the executing context is running on a listen server."),
 	HasAuthority UMETA(Tooltip="Returns whether the execution context is running with network authority."),
 	IsBuilder UMETA(Tooltip="Returns whether the execution source is generated by the PCG World Partition Builder. Always false in non-editor."),
 };
diff --git a/Engine/Plugins/PCG/Source/PCG/Public/Helpers/PCGHelpers.h b/Engine/Plugins/PCG/Source/PCG/Public/Helpers/PCGHelpers.h
--- a/Engine/Plugins/PCG/Source/PCG/Public/Helpers/PCGHelpers.h
+++ b/Engine/Plugins/PCG/Source/PCG/Public/Helpers/PCGHelpers.h
@@ -109,6 +109,10 @@ namespace PCGHelpers
 	/** Get an array of a randomized, uniformly distributed indices to a provided array view. */
 	PCG_API TArray<int32> GetRandomIndices(FRandomStream& RandomStream, const int32 ArraySize, const int32 NumSelections);
 
+	PCG_API bool IsServer(const UWorld* InWorld);
+	PCG_API bool IsDedicatedServer(const UWorld* InWorld);
+	PCG_API bool IsListenServer(const UWorld* InWorld);
+
 	/** Shuffles the elements of an array randomly and uniformly. */
 	template <typename T>
 	void ShuffleArray(FRandomStream& RandomStream, TArray<T>& Array)
diff --git a/Engine/Plugins/PCG/Source/PCG/Public/PCGWorldActor.h b/Engine/Plugins/PCG/Source/PCG/Public/PCGWorldActor.h
--- a/Engine/Plugins/PCG/Source/PCG/Public/PCGWorldActor.h
+++ b/Engine/Plugins/PCG/Source/PCG/Public/PCGWorldActor.h
@@ -47,6 +47,7 @@ public:
 	virtual void PostEditChangeProperty(FPropertyChangedEvent& PropertyChangedEvent) override;
 #endif
 	virtual void BeginDestroy() override;
+	virtual void Serialize(FArchive& Archive) override;
 	//~ End UObject Interface.
 
 	void MergeFrom(APCGWorldActor* OtherWorldActor);
@@ -75,10 +76,15 @@ public:
 	UPROPERTY(BlueprintReadWrite, EditAnywhere, Category = RuntimeGeneration)
 	bool bTreatEditorViewportAsGenerationSource = false;
 #endif
+	
+	/** Allows any player controller to be considered as a Runtime Generation Source. Disabling this will require registering generation sources using UPCGGenSourceComponent or IPCGGenSourceBase sub-classing. */
+	UPROPERTY(EditAnywhere, Category = RuntimeGeneration)
+	bool bTreatPlayerControllersAsGenerationSources = true;
 
 	/** Allows World Partition streaming sources to act as Runtime Generation Sources. */
-	UPROPERTY(BlueprintReadWrite, EditAnywhere, Category = RuntimeGeneration)
-	bool bEnableWorldPartitionGenerationSources = true;
+	UE_DEPRECATED(5.7, "No longer required, by default, player controllers are considered generation sources and additional UPCGGenSourceComponent can be used for more complex use cases")
+	UPROPERTY(BlueprintReadWrite, EditAnywhere, Category = RuntimeGeneration, meta = (EditCondition=bEnableWorldPartitionGenerationSources, EditConditionHides))
+	bool bEnableWorldPartitionGenerationSources = false;
 
 private:
 	friend class FPCGActorAndComponentMapping;
diff --git a/Engine/Plugins/PCG/Source/PCG/Public/RuntimeGen/GenSources/PCGGenSourceWPStreamingSource.h b/Engine/Plugins/PCG/Source/PCG/Public/RuntimeGen/GenSources/PCGGenSourceWPStreamingSource.h
--- a/Engine/Plugins/PCG/Source/PCG/Public/RuntimeGen/GenSources/PCGGenSourceWPStreamingSource.h
+++ b/Engine/Plugins/PCG/Source/PCG/Public/RuntimeGen/GenSources/PCGGenSourceWPStreamingSource.h
@@ -10,13 +10,8 @@
 
 struct FWorldPartitionStreamingSource;
 
-/**
- * This GenerationSource captures WorldPartitionStreamingSources for RuntimeGeneration.
- * 
- * UPCGGenSourceWPStreamingSources are created per tick and live only until the end of the tick.
- */
 UCLASS(MinimalAPI, BlueprintType, ClassGroup = (Procedural))
-class UPCGGenSourceWPStreamingSource : public UObject, public IPCGGenSourceBase
+class UE_DEPRECATED(5.7, "No longer supported.") UPCGGenSourceWPStreamingSource : public UObject, public IPCGGenSourceBase
 {
 	GENERATED_BODY()
 
diff --git a/Engine/Plugins/PCG/Source/PCG/Public/RuntimeGen/PCGGenSourceManager.h b/Engine/Plugins/PCG/Source/PCG/Public/RuntimeGen/PCGGenSourceManager.h
--- a/Engine/Plugins/PCG/Source/PCG/Public/RuntimeGen/PCGGenSourceManager.h
+++ b/Engine/Plugins/PCG/Source/PCG/Public/RuntimeGen/PCGGenSourceManager.h
@@ -42,12 +42,13 @@ public:
 	void AddReferencedObjects(FReferenceCollector& Collector);
 
 protected:
-	void OnGameModePostLogin(AGameModeBase* InGameMode, APlayerController* InPlayerController);
-	void OnGameModePostLogout(AGameModeBase* InGameMode, AController* InController);
 
 	/** Updates tracked generation sources that should be refreshed per tick. */
 	void UpdatePerTickGenSources(const APCGWorldActor* InPCGWorldActor);
 
+	void UpdateWorldPartitionGenSources(const APCGWorldActor* InPCGWorldActor);
+	void GatherWorldPartitionGenSources(TSet<IPCGGenSourceBase*>& OutGenSources) const;
+
 protected:
 	/** Tracks named registered generation sources, such as UPCGGenSourcePlayer. */
 	TMap<FName, TScriptInterface<IPCGGenSourceBase>> RegisteredNamedGenSources;
@@ -60,8 +61,11 @@ protected:
 	TObjectPtr<UPCGGenSourceEditorCamera> EditorCameraGenSource = nullptr;
 #endif
 
+PRAGMA_DISABLE_DEPRECATION_WARNINGS
 	/** Pool of GenSources dedicated to WorldPartition StreamingSources. This list grows as necessary and refreshes every tick. */
 	TArray<TObjectPtr<UPCGGenSourceWPStreamingSource>> WorldPartitionGenSources;
+	bool bLoggedWorldPartitionGenSourceWarning = false;
+PRAGMA_ENABLE_DEPRECATION_WARNINGS
 
 	const UWorld* World = nullptr;
 	bool bDirty = false;
diff --git a/Engine/Source/Runtime/Core/Public/UObject/FortniteMainBranchObjectVersions.inl b/Engine/Source/Runtime/Core/Public/UObject/FortniteMainBranchObjectVersions.inl
--- a/Engine/Source/Runtime/Core/Public/UObject/FortniteMainBranchObjectVersions.inl
+++ b/Engine/Source/Runtime/Core/Public/UObject/FortniteMainBranchObjectVersions.inl
@@ -661,5 +661,8 @@ FORTNITE_MAIN_VERSION(PerParticleIterationCountMovedToDynamicMisc, 215)
 // Added missing custom serialization for some properties in the ParticleDynamicMisc structure used by the Chaos Visual Debugger
 FORTNITE_MAIN_VERSION(AddedMissingSerializationForPropertiesInDynamicMisc, 216)
 
+// Change default value for deprecated bEnableWorldPartitionGenerationSources
+FORTNITE_MAIN_VERSION(PCGDeprecateWorldPartitionGenerationSources, 217)
+
 // -----<new versions can be added above this line>-------------------------------------------------
 #undef FORTNITE_MAIN_VERSION
\ No newline at end of file
