[PCG] PCGWorldActor: Fix selection loss when changing grid size
- Also ignore interactive changes so that value can be scrubbed and applied only on release

diff --git a/Engine/Plugins/PCG/Source/PCG/Private/PCGWorldActor.cpp b/Engine/Plugins/PCG/Source/PCG/Private/PCGWorldActor.cpp
--- a/Engine/Plugins/PCG/Source/PCG/Private/PCGWorldActor.cpp
+++ b/Engine/Plugins/PCG/Source/PCG/Private/PCGWorldActor.cpp
@@ -290,9 +290,18 @@ void APCGWorldActor::OnPartitionGridSizeChanged()
 		return;
 	}
 
+	// Backup selected state
+	const bool bWasSelected = IsSelected();
+
 	// Then delete all PCGPartitionActors
 	PCGSubsystem->DeleteSerializedPartitionActors(/*bDeleteOnlyUnused=*/false);
 
+	// Restore selection if needed
+	if (bWasSelected && !IsSelected())
+	{
+		GEditor->SelectActor(this, /*bSelected=*/ true, /*bNotify=*/ true);
+	}
+
 	// And finally, regenerate all components that are partitioned (registered to the PCGSubsystem)
 	// to let them recreate the needed PCG Partition Actors.
 	for (UPCGComponent* PCGComponent : PCGSubsystem->GetAllRegisteredPartitionedComponents())
@@ -307,7 +316,7 @@ void APCGWorldActor::PostEditChangeProperty(FPropertyChangedEvent& PropertyChang
 {
 	Super::PostEditChangeProperty(PropertyChangedEvent);
 	
-	if (PropertyChangedEvent.GetPropertyName() == GET_MEMBER_NAME_CHECKED(APCGWorldActor, PartitionGridSize))
+	if (PropertyChangedEvent.GetPropertyName() == GET_MEMBER_NAME_CHECKED(APCGWorldActor, PartitionGridSize) && !(PropertyChangedEvent.ChangeType & EPropertyChangeType::Interactive))
 	{
 		OnPartitionGridSizeChanged();
 	}
