[PCG] Disambiguate debug logging for nodes that produce multiple kernels

#rnx

diff --git a/Engine/Plugins/PCG/Source/PCG/Private/Compute/DataInterfaces/PCGDebugDataInterface.cpp b/Engine/Plugins/PCG/Source/PCG/Private/Compute/DataInterfaces/PCGDebugDataInterface.cpp
--- a/Engine/Plugins/PCG/Source/PCG/Private/Compute/DataInterfaces/PCGDebugDataInterface.cpp
+++ b/Engine/Plugins/PCG/Source/PCG/Private/Compute/DataInterfaces/PCGDebugDataInterface.cpp
@@ -3,6 +3,8 @@
 #include "Compute/DataInterfaces/PCGDebugDataInterface.h"
 
 #include "PCGModule.h"
+#include "PCGSettings.h"
+#include "Compute/PCGComputeKernel.h"
 
 #include "RenderGraphBuilder.h"
 #include "RenderGraphResources.h"
@@ -79,7 +81,11 @@ void UPCGDebugDataProvider::Initialize(const UComputeDataInterface* InDataInterf
 
 FComputeDataProviderRenderProxy* UPCGDebugDataProvider::GetRenderProxy()
 {
-	auto ProcessReadbackData_RenderThread = [](const void* InData, int InNumBytes)
+	const UPCGSettings* Settings = GetProducerSettings();
+	const FString SettingsName = Settings ? Settings->GetName() : TEXT("MISSINGSETTINGS");
+	const FString KernelName = ProducerKernel ? ProducerKernel->GetName() : TEXT("MISSINGKERNEL");
+
+	auto ProcessReadbackData_RenderThread = [SettingsName=SettingsName, KernelName=KernelName](const void* InData, int InNumBytes)
 	{
 		if (!InData)
 		{
@@ -87,6 +93,9 @@ FComputeDataProviderRenderProxy* UPCGDebugDataProvider::GetRenderProxy()
 		}
 
 		const float* DataAsFloat = static_cast<const float*>(InData);
+
+		UE_LOG(LogPCG, Warning, TEXT("%s - %s"), *SettingsName, *KernelName);
+
 		for (int Index = 0; Index < InNumBytes / sizeof(float); ++Index)
 		{
 			UE_LOG(LogPCG, Warning, TEXT("\t%d:\t%f"), Index, DataAsFloat[Index]);
diff --git a/Engine/Plugins/PCG/Source/PCG/Private/Compute/PCGComputeGraph.cpp b/Engine/Plugins/PCG/Source/PCG/Private/Compute/PCGComputeGraph.cpp
--- a/Engine/Plugins/PCG/Source/PCG/Private/Compute/PCGComputeGraph.cpp
+++ b/Engine/Plugins/PCG/Source/PCG/Private/Compute/PCGComputeGraph.cpp
@@ -304,6 +304,9 @@ void UPCGComputeGraph::DebugLogDataDescriptions(const UPCGDataBinding* InBinding
 
 		if (PCGKernel && PCGKernel->GetLogDataDescriptions())
 		{
+			const UPCGSettings* Settings = PCGKernel->GetSettings();
+
+			UE_LOG(LogPCG, Display, TEXT("Settings: %s"), Settings ? *Settings->GetName() : TEXT("MISSINGSETTINGS"));
 			UE_LOG(LogPCG, Display, TEXT("Kernel: %s"), *PCGKernel->GetName());
 
 			TArray<FPCGPinProperties> InputPins;
diff --git a/Engine/Plugins/PCG/Source/PCG/Private/Graph/PCGGraphCompilerGPU.cpp b/Engine/Plugins/PCG/Source/PCG/Private/Graph/PCGGraphCompilerGPU.cpp
--- a/Engine/Plugins/PCG/Source/PCG/Private/Graph/PCGGraphCompilerGPU.cpp
+++ b/Engine/Plugins/PCG/Source/PCG/Private/Graph/PCGGraphCompilerGPU.cpp
@@ -1421,6 +1421,8 @@ void FPCGGraphCompilerGPU::CompileComputeGraph(
 		{
 			UPCGDebugDataInterface* DebugDI = InOutContext.NewObject_AnyThread<UPCGDebugDataInterface>(InOutComputeGraph);
 			DebugDI->SetDebugBufferSize(Settings->DebugBufferSize);
+			DebugDI->SetProducerSettings(Settings);
+			DebugDI->ProducerKernel = KernelWithBindings.Kernel;
 
 			const int32 DebugDIIndex = InOutComputeGraph->DataInterfaces.Num();
 			InOutComputeGraph->DataInterfaces.Add(DebugDI);
@@ -1497,7 +1499,7 @@ void FPCGGraphCompilerGPU::CompileComputeGraph(
 
 			if (Settings->bDumpCookedHLSL)
 			{
-				UE_LOG(LogPCG, Log, TEXT("Cooked HLSL:\n%s\n"), *KernelSource->GetSource());
+				UE_LOG(LogPCG, Log, TEXT("Cooked HLSL (%s - %s):\n%s\n"), *Settings->GetName(), *KernelWithBindings.Kernel->GetName(), *KernelSource->GetSource());
 			}
 
 #if PCG_KERNEL_LOGGING_ENABLED
