[PCG] Fixes after running GPU Reshape.

Main change is to avoid ternary's wherever it is unsafe to execute both sides. GPU evaluates both sides and selects a result.

#rnx

diff --git a/Engine/Plugins/PCG/Shaders/Private/Elements/PCGCopyPoints.usf b/Engine/Plugins/PCG/Shaders/Private/Elements/PCGCopyPoints.usf
--- a/Engine/Plugins/PCG/Shaders/Private/Elements/PCGCopyPoints.usf
+++ b/Engine/Plugins/PCG/Shaders/Private/Elements/PCGCopyPoints.usf
@@ -53,9 +53,18 @@ void PCGCopyPointsCS(uint3 GroupId : SV_GroupID, uint GroupIndex : SV_GroupIndex
 	const uint TargetSeed = Target_GetSeed(TargetDataIndex, TargetPointIndexInData);
 
 	const uint RotationInheritance = RotationInheritance_GetOverridableValue();
-	const uint ApplyTargetRotationToPositions = (RotationInheritance == PCG_COPY_POINTS_INHERITANCE_MODE_Source) ? bApplyTargetRotationToPositions_GetOverridableValue() : 1;
+	uint ApplyTargetRotationToPositions = 1;
+	if (RotationInheritance == PCG_COPY_POINTS_INHERITANCE_MODE_Source)
+	{
+		ApplyTargetRotationToPositions = bApplyTargetRotationToPositions_GetOverridableValue();
+	}
+
 	const uint ScaleInheritance = ScaleInheritance_GetOverridableValue();
-	const uint ApplyTargetScaleToPositions = (ScaleInheritance == PCG_COPY_POINTS_INHERITANCE_MODE_Source) ? bApplyTargetScaleToPositions_GetOverridableValue() : 1;
+	uint ApplyTargetScaleToPositions = 1;
+	if (ScaleInheritance == PCG_COPY_POINTS_INHERITANCE_MODE_Source)
+	{
+		ApplyTargetScaleToPositions = bApplyTargetScaleToPositions_GetOverridableValue();
+	}
 
 	// @todo_pcg: Must be known statically, so graph splits are required in order to support overriding these properties.
 	const uint ColorInheritance = CopyPoints_GetColorInheritance();
@@ -64,7 +73,15 @@ void PCGCopyPointsCS(uint3 GroupId : SV_GroupID, uint GroupIndex : SV_GroupIndex
 
 	const float3 ScaleForPosition = ApplyTargetScaleToPositions > 0 ? TargetScale : float3(1, 1, 1);
 
-	const float3 OutPosition = TargetPosition + (ApplyTargetRotationToPositions > 0 ? QuatRotateVector(TargetRotation, (SourcePosition * ScaleForPosition)) : (SourcePosition * ScaleForPosition));
+	float3 OutPosition = TargetPosition;
+	if (ApplyTargetRotationToPositions > 0)
+	{
+		OutPosition += QuatRotateVector(TargetRotation, (SourcePosition * ScaleForPosition));
+	}
+	else
+	{
+		OutPosition += SourcePosition * ScaleForPosition;
+	}
 
 	float4 OutRotation = float4(0, 0, 0, 1);
 	float3 OutScale = (float3)0;
diff --git a/Engine/Plugins/PCG/Shaders/Private/Elements/PCGStaticMeshSpawner.usf b/Engine/Plugins/PCG/Shaders/Private/Elements/PCGStaticMeshSpawner.usf
--- a/Engine/Plugins/PCG/Shaders/Private/Elements/PCGStaticMeshSpawner.usf
+++ b/Engine/Plugins/PCG/Shaders/Private/Elements/PCGStaticMeshSpawner.usf
@@ -76,7 +76,11 @@ void PCGStaticMeshSpawnerCS(uint3 GroupId : SV_GroupID, uint GroupIndex : SV_Gro
 	}
 
 	// Get instance index to write to. Will return a negative value if we have saturated instances for this primitive, in which case we will not push the data.
-	const int InstanceIndexInt = (PrimitiveIndex != (uint)-1) ? InstanceData_GetIndexToWriteTo(PrimitiveIndex) : -1;
+	int InstanceIndexInt = -1;
+	if (PrimitiveIndex != (uint)-1)
+	{
+		InstanceIndexInt = InstanceData_GetIndexToWriteTo(PrimitiveIndex);
+	}
 
 	// Write instance transform, only if not culled by primitive instance index
 	if (InstanceIndexInt >= 0)
@@ -131,7 +135,12 @@ void PCGStaticMeshSpawnerCS(uint3 GroupId : SV_GroupID, uint GroupIndex : SV_Gro
 
 	if (SMSpawner_GetSelectedMeshAttributeId() != (uint)-1)
 	{
-		const int OutputMeshStringKey = (PrimitiveIndex != (uint)-1) ? SMSpawner_GetPrimitiveStringKey(PrimitiveIndex) : -1;
+		int OutputMeshStringKey = -1;
+		if (PrimitiveIndex != (uint)-1)
+		{
+			OutputMeshStringKey = SMSpawner_GetPrimitiveStringKey(PrimitiveIndex);
+		}
+
 		Out_SetStringKey(Out_DataIndex, Out_ElementIndex, SMSpawner_GetSelectedMeshAttributeId(), OutputMeshStringKey);
 	}
 }
diff --git a/Engine/Plugins/PCG/Shaders/Private/PCGDataCollectionDataInterface.ush b/Engine/Plugins/PCG/Shaders/Private/PCGDataCollectionDataInterface.ush
--- a/Engine/Plugins/PCG/Shaders/Private/PCGDataCollectionDataInterface.ush
+++ b/Engine/Plugins/PCG/Shaders/Private/PCGDataCollectionDataInterface.ush
@@ -201,46 +201,75 @@ void SetAsExecutedInternal_{DataInterfaceName}()
 
 // #################### ATTRIBUTE GETTERS ##########################
 
-bool GetBool_{DataInterfaceName}(uint InDataIndex, uint InElementIndex, int InAttributeId)
+uint GetUint_{DataInterfaceName}(uint InDataIndex, uint InElementIndex, int InAttributeId)
 {
 	const uint ElementAddress = GetElementAddressInternal_{DataInterfaceName}(InDataIndex, InElementIndex, InAttributeId);
-	return ElementAddress > 0 && {DataInterfaceName}_DataCollectionBuffer.Load(ElementAddress) != 0;
+
+	if (ElementAddress > 0)
+	{
+		return {DataInterfaceName}_DataCollectionBuffer.Load(ElementAddress);
+	}
+	else
+	{
+		return 0u;
+	}
 }
 
-int GetInt_{DataInterfaceName}(uint InDataIndex, uint InElementIndex, int InAttributeId)
+bool GetBool_{DataInterfaceName}(uint InDataIndex, uint InElementIndex, int InAttributeId)
 {
-	const uint ElementAddress = GetElementAddressInternal_{DataInterfaceName}(InDataIndex, InElementIndex, InAttributeId);
-	return ElementAddress > 0 ? asint({DataInterfaceName}_DataCollectionBuffer.Load(ElementAddress)) : 0;
+	return GetUint_{DataInterfaceName}(InDataIndex, InElementIndex, InAttributeId) != 0;
 }
 
-int GetUint_{DataInterfaceName}(uint InDataIndex, uint InElementIndex, int InAttributeId)
+int GetInt_{DataInterfaceName}(uint InDataIndex, uint InElementIndex, int InAttributeId)
 {
-	const uint ElementAddress = GetElementAddressInternal_{DataInterfaceName}(InDataIndex, InElementIndex, InAttributeId);
-	return ElementAddress > 0 ? {DataInterfaceName}_DataCollectionBuffer.Load(ElementAddress) : 0;
+	return asint(GetUint_{DataInterfaceName}(InDataIndex, InElementIndex, InAttributeId));
 }
 
 float GetFloat_{DataInterfaceName}(uint InDataIndex, uint InElementIndex, int InAttributeId)
 {
-	const uint ElementAddress = GetElementAddressInternal_{DataInterfaceName}(InDataIndex, InElementIndex, InAttributeId);
-	return ElementAddress > 0 ? asfloat({DataInterfaceName}_DataCollectionBuffer.Load(ElementAddress)) : 0.0f;
+	return asfloat(GetUint_{DataInterfaceName}(InDataIndex, InElementIndex, InAttributeId));
 }
 
 float2 GetFloat2_{DataInterfaceName}(uint InDataIndex, uint InElementIndex, int InAttributeId)
 {
 	const uint ElementAddress = GetElementAddressInternal_{DataInterfaceName}(InDataIndex, InElementIndex, InAttributeId);
-	return ElementAddress > 0 ? asfloat({DataInterfaceName}_DataCollectionBuffer.Load2(ElementAddress)) : (float2)0;
+
+	if (ElementAddress > 0)
+	{
+		return asfloat({DataInterfaceName}_DataCollectionBuffer.Load2(ElementAddress));
+	}
+	else
+	{
+		return (float2)0;
+	}
 }
 
 float3 GetFloat3_{DataInterfaceName}(uint InDataIndex, uint InElementIndex, int InAttributeId)
 {
 	const uint ElementAddress = GetElementAddressInternal_{DataInterfaceName}(InDataIndex, InElementIndex, InAttributeId);
-	return ElementAddress > 0 ? asfloat({DataInterfaceName}_DataCollectionBuffer.Load3(ElementAddress)) : (float3)0;
+
+	if (ElementAddress > 0)
+	{
+		return asfloat({DataInterfaceName}_DataCollectionBuffer.Load3(ElementAddress));
+	}
+	else
+	{
+		return (float3)0;
+	}
 }
 
 float4 GetFloat4_{DataInterfaceName}(uint InDataIndex, uint InElementIndex, int InAttributeId)
 {
 	const uint ElementAddress = GetElementAddressInternal_{DataInterfaceName}(InDataIndex, InElementIndex, InAttributeId);
-	return ElementAddress > 0 ? asfloat({DataInterfaceName}_DataCollectionBuffer.Load4(ElementAddress)) : (float4)0;
+
+	if (ElementAddress > 0)
+	{
+		return asfloat({DataInterfaceName}_DataCollectionBuffer.Load4(ElementAddress));
+	}
+	else
+	{
+		return (float4)0;
+	}
 }
 
 float3 GetRotator_{DataInterfaceName}(uint InDataIndex, uint InElementIndex, int InAttributeId)
@@ -275,7 +304,14 @@ int GetStringKey_{DataInterfaceName}(uint InDataIndex, uint InElementIndex, int
 {
 	// String keys are represented as ints.
 	const int StringKey = GetInt_{DataInterfaceName}(InDataIndex, InElementIndex, InAttributeId);
-	return (StringKey < {DataInterfaceName}_NumRemappedStringKeys) ? {DataInterfaceName}_BufferToGraphStringKey[StringKey] : StringKey;
+	if (StringKey < {DataInterfaceName}_NumRemappedStringKeys)
+	{
+		return {DataInterfaceName}_BufferToGraphStringKey[StringKey];
+	}
+	else
+	{
+		return StringKey;
+	}
 }
 
 // TODO: left as uint2 which we locked into for 5.5.0. Future version should change this to int.
@@ -405,40 +441,75 @@ void SetName_{DataInterfaceName}(uint InDataIndex, uint InElementIndex, int InAt
 
 // #################### FIRST ATTRIBUTE GETTERS ##########################
 
-bool GetFirstAttributeAsBool_{DataInterfaceName}()
+uint GetFirstAttributeAsUint_{DataInterfaceName}()
 {
 	const uint AttributeAddress = GetFirstValidAttributeElementInternal_{DataInterfaceName}();
-	return AttributeAddress > 0 && {DataInterfaceName}_DataCollectionBuffer.Load(AttributeAddress) != 0;
+
+	if (AttributeAddress > 0)
+	{
+		return {DataInterfaceName}_DataCollectionBuffer.Load(AttributeAddress);
+	}
+	else
+	{
+		return 0;
+	}
+}
+
+bool GetFirstAttributeAsBool_{DataInterfaceName}()
+{
+	return GetFirstAttributeAsUint_{DataInterfaceName}() != 0;
 }
 
 int GetFirstAttributeAsInt_{DataInterfaceName}()
 {
-	const uint AttributeAddress = GetFirstValidAttributeElementInternal_{DataInterfaceName}();
-	return AttributeAddress > 0 ? asint({DataInterfaceName}_DataCollectionBuffer.Load(AttributeAddress)) : 0;
+	return asint(GetFirstAttributeAsUint_{DataInterfaceName}());
 }
 
 float GetFirstAttributeAsFloat_{DataInterfaceName}()
 {
-	const uint AttributeAddress = GetFirstValidAttributeElementInternal_{DataInterfaceName}();
-	return AttributeAddress > 0 ? asfloat({DataInterfaceName}_DataCollectionBuffer.Load(AttributeAddress)) : 0.0f;
+	return asfloat(GetFirstAttributeAsUint_{DataInterfaceName}());
 }
 
 float2 GetFirstAttributeAsFloat2_{DataInterfaceName}()
 {
 	const uint AttributeAddress = GetFirstValidAttributeElementInternal_{DataInterfaceName}();
-	return AttributeAddress > 0 ? asfloat({DataInterfaceName}_DataCollectionBuffer.Load2(AttributeAddress)) : (float2)0;
+
+	if (AttributeAddress > 0)
+	{
+		return asfloat({DataInterfaceName}_DataCollectionBuffer.Load2(AttributeAddress));
+	}
+	else
+	{
+		return (float2)0;
+	}
 }
 
 float3 GetFirstAttributeAsFloat3_{DataInterfaceName}()
 {
 	const uint AttributeAddress = GetFirstValidAttributeElementInternal_{DataInterfaceName}();
-	return AttributeAddress > 0 ? asfloat({DataInterfaceName}_DataCollectionBuffer.Load3(AttributeAddress)) : (float3)0;
+
+	if (AttributeAddress > 0)
+	{
+		return asfloat({DataInterfaceName}_DataCollectionBuffer.Load3(AttributeAddress));
+	}
+	else
+	{
+		return (float3)0;
+	}
 }
 
 float4 GetFirstAttributeAsFloat4_{DataInterfaceName}()
 {
 	const uint AttributeAddress = GetFirstValidAttributeElementInternal_{DataInterfaceName}();
-	return AttributeAddress > 0 ? asfloat({DataInterfaceName}_DataCollectionBuffer.Load4(AttributeAddress)) : (float4)0;
+
+	if (AttributeAddress > 0)
+	{
+		return asfloat({DataInterfaceName}_DataCollectionBuffer.Load4(AttributeAddress));
+	}
+	else
+	{
+		return (float4)0;
+	}
 }
 
 float3 GetFirstAttributeAsRotator_{DataInterfaceName}()
diff --git a/Engine/Plugins/PCG/Shaders/Private/PCGSceneWriter.usf b/Engine/Plugins/PCG/Shaders/Private/PCGSceneWriter.usf
--- a/Engine/Plugins/PCG/Shaders/Private/PCGSceneWriter.usf
+++ b/Engine/Plugins/PCG/Shaders/Private/PCGSceneWriter.usf
@@ -51,6 +51,7 @@ void SceneWriter_CS(const uint ThreadId : SV_DispatchThreadID)
 	LocalToWorld[3] = float4(0.0f, 0.0f, 0.0f, 1.0f);
 	LocalToWorld = transpose(LocalToWorld);
 
+	// todo_pcg: We should write local space transforms to the instance data buffer from the SM Spawner and use InitializeInstanceSceneData (non-WS) here instead.
 	InitializeInstanceSceneDataWS(
 		InstanceId,
 		InPrimitiveId,
@@ -61,10 +62,6 @@ void SceneWriter_CS(const uint ThreadId : SV_DispatchThreadID)
 		DFFromTileOffset(LWCPromote(LocalToWorld))
 	);
 	
-	// todo_pcg: We should write local space transforms to the instance data buffer from the SM Spawner and use WriteInstanceLocalToRelativeWorld here instead.
-	const FDFMatrix Transform = MakeDFMatrix4x3((float3)0.0f, LocalToWorld);
-	WriteInstanceLocalToWorld(InstanceId, InPrimitiveId, Transform, false);
-
 	FInstancePayloadDataOffsets PayloadOffsets = GetInstancePayloadDataOffsets(InPrimitiveId, InPayloadDataFlags, LocalInstanceId);
 
 	const uint NumCustomFloat4s = (InCustomDataCount + 3) / 4;
diff --git a/Engine/Plugins/PCG/Shaders/Private/PCGStaticMeshDataInterface.ush b/Engine/Plugins/PCG/Shaders/Private/PCGStaticMeshDataInterface.ush
--- a/Engine/Plugins/PCG/Shaders/Private/PCGStaticMeshDataInterface.ush
+++ b/Engine/Plugins/PCG/Shaders/Private/PCGStaticMeshDataInterface.ush
@@ -68,14 +68,23 @@ void GetVertex_{DataInterfaceName}(int DataIndex, int VertexIndex, out float3 Ou
 float4 GetVertexColor_{DataInterfaceName}(int DataIndex, int VertexIndex)
 {
 	VertexIndex = clamp(VertexIndex, 0, {DataInterfaceName}_NumVertices - 1);
-	return {DataInterfaceName}_HasColors ? {DataInterfaceName}_ColorBuffer[VertexIndex] FMANUALFETCH_COLOR_COMPONENT_SWIZZLE : float4(1, 1, 1, 1);
+
+	if ({DataInterfaceName}_HasColors > 0)
+	{
+		return {DataInterfaceName}_ColorBuffer[VertexIndex] FMANUALFETCH_COLOR_COMPONENT_SWIZZLE;
+	}
+	else
+	{
+		return float4(1.0f, 1.0f, 1.0f, 1.0f);
+	}
 }
 
 float2 GetVertexUV_{DataInterfaceName}(int DataIndex, int VertexIndex, int UVSet)
 {
 	VertexIndex = clamp(VertexIndex, 0, {DataInterfaceName}_NumVertices - 1);
 	UVSet = clamp(UVSet, 0, {DataInterfaceName}_NumUVs - 1);
-	return UVSet >= 0 && UVSet < {DataInterfaceName}_NumUVs ? {DataInterfaceName}_UVBuffer[VertexIndex * {DataInterfaceName}_NumUVs + UVSet] : float2(0, 0);
+
+	return {DataInterfaceName}_UVBuffer[VertexIndex * {DataInterfaceName}_NumUVs + UVSet];
 }
 
 // #################### TRIANGLE BASED SAMPLING ##########################
diff --git a/Engine/Plugins/PCG/Shaders/Private/PCGVirtualTextureDataInterface.ush b/Engine/Plugins/PCG/Shaders/Private/PCGVirtualTextureDataInterface.ush
--- a/Engine/Plugins/PCG/Shaders/Private/PCGVirtualTextureDataInterface.ush
+++ b/Engine/Plugins/PCG/Shaders/Private/PCGVirtualTextureDataInterface.ush
@@ -73,8 +73,9 @@ float4 SampleLayer_{DataInterfaceName}(
 #define PCG_CONDITIONAL_VIRTUAL_TEXTURE_SAMPLE(DataIndex) \
 	if (InDataIndex == DataIndex) \
 	{ \
-		LayerSample[0] = (ValidLayerMask & 0x1) != 0 \
-			? SampleLayer_{DataInterfaceName}( \
+		if (ValidLayerMask & 0x1) \
+		{ \
+			LayerSample[0] = SampleLayer_{DataInterfaceName}( \
 				InDataIndex, \
 				SampleUV, \
 				{DataInterfaceName}_VirtualTexture0_##DataIndex, \
@@ -82,10 +83,15 @@ float4 SampleLayer_{DataInterfaceName}(
 				{DataInterfaceName}_PageTableIndirection0_##DataIndex, \
 				Adaptive0, \
 				Uniforms0, \
-				PageTableUniforms) \
-			: 0; \
-		LayerSample[1] = (ValidLayerMask & 0x2) != 0 \
-			? SampleLayer_{DataInterfaceName}( \
+				PageTableUniforms); \
+		} \
+		else \
+		{ \
+			LayerSample[0] = (float4)0; \
+		} \
+		if (ValidLayerMask & 0x2) \
+		{ \
+			LayerSample[1] = SampleLayer_{DataInterfaceName}( \
 				InDataIndex, \
 				SampleUV, \
 				{DataInterfaceName}_VirtualTexture1_##DataIndex, \
@@ -93,10 +99,15 @@ float4 SampleLayer_{DataInterfaceName}(
 				{DataInterfaceName}_PageTableIndirection1_##DataIndex, \
 				Adaptive1, \
 				Uniforms1, \
-				PageTableUniforms) \
-			: 0; \
-		LayerSample[2] = (ValidLayerMask & 0x4) != 0 \
-			? SampleLayer_{DataInterfaceName}( \
+				PageTableUniforms); \
+		} \
+		else \
+		{ \
+			LayerSample[1] = (float4)0; \
+		} \
+		if (ValidLayerMask & 0x4) \
+		{ \
+			LayerSample[2] = SampleLayer_{DataInterfaceName}( \
 				InDataIndex, \
 				SampleUV, \
 				{DataInterfaceName}_VirtualTexture2_##DataIndex, \
@@ -104,8 +115,12 @@ float4 SampleLayer_{DataInterfaceName}(
 				{DataInterfaceName}_PageTableIndirection2_##DataIndex, \
 				Adaptive2, \
 				Uniforms2, \
-				PageTableUniforms) \
-			: 0; \
+				PageTableUniforms); \
+		} \
+		else \
+		{ \
+			LayerSample[2] = (float4)0; \
+		} \
 	}
 
 void SampleVirtualTexture_{DataInterfaceName}(
