[PCG] Fix pairwise data multiplicity computation for data descriptions
* Add validation to enforce pairwise input pins having N:N, 1:N, or N:1 data

#rnx

diff --git a/Engine/Plugins/PCG/Source/PCG/Private/Compute/Elements/PCGCustomHLSL.cpp b/Engine/Plugins/PCG/Source/PCG/Private/Compute/Elements/PCGCustomHLSL.cpp
--- a/Engine/Plugins/PCG/Source/PCG/Private/Compute/Elements/PCGCustomHLSL.cpp
+++ b/Engine/Plugins/PCG/Source/PCG/Private/Compute/Elements/PCGCustomHLSL.cpp
@@ -111,12 +111,12 @@ void UPCGCustomHLSLSettings::PostEditChangeProperty(FPropertyChangedEvent& Prope
 		&& PropertyChangedEvent.ChangeType == EPropertyChangeType::ArrayAdd)
 	{
 		// Whenever a new output pin is created, we should default initialize 'PinsToInitializeFrom' with the first input pin label (if it exists).
-		if (const UPCGPin* FirstInputPin = GetFirstInputPin())
+		if (const FPCGPinProperties* FirstInputPin = GetFirstInputPinProperties())
 		{
 			check(!OutputPins.IsEmpty());
 
 			FPCGPinPropertiesGPU& PinProps = OutputPins.Last();
-			PinProps.PropertiesGPU.PinsToInititalizeFrom.Add(FirstInputPin->Properties.Label);
+			PinProps.PropertiesGPU.PinsToInititalizeFrom.Add(FirstInputPin->Label);
 		}
 	}
 	else if (MemberProperty == GET_MEMBER_NAME_CHECKED(UPCGCustomHLSLSettings, InputPins)
@@ -201,12 +201,12 @@ void UPCGCustomHLSLSettings::ApplyPreconfiguredSettings(const FPCGPreConfiguredS
 #endif
 
 			// Default to initializing the first output pin's from the first input pin's data.
-			if (const UPCGPin* FirstInputPin = GetFirstInputPin())
+			if (const FPCGPinProperties* FirstInputPin = GetFirstInputPinProperties())
 			{
 				if (!OutputPins.IsEmpty())
 				{
 					FPCGPinPropertiesGPU& PinProps = OutputPins.Last();
-					PinProps.PropertiesGPU.PinsToInititalizeFrom.Add(FirstInputPin->Properties.Label);
+					PinProps.PropertiesGPU.PinsToInititalizeFrom.Add(FirstInputPin->Label);
 				}
 			}
 
@@ -227,36 +227,14 @@ FString UPCGCustomHLSLSettings::GetAdditionalTitleInformation() const
 	return FString();
 }
 
-const UPCGPin* UPCGCustomHLSLSettings::GetInputPin(FName Label) const
+const FPCGPinProperties* UPCGCustomHLSLSettings::GetFirstInputPinProperties() const
 {
-	const UPCGNode* Node = Cast<UPCGNode>(GetOuter());
-	return Node ? Node->GetInputPin(Label) : nullptr;
+	return !InputPins.IsEmpty() ? &InputPins[0] : nullptr;
 }
 
-const UPCGPin* UPCGCustomHLSLSettings::GetOutputPin(FName Label) const
+const FPCGPinPropertiesGPU* UPCGCustomHLSLSettings::GetFirstOutputPinProperties() const
 {
-	const UPCGNode* Node = Cast<UPCGNode>(GetOuter());
-	return Node ? Node->GetOutputPin(Label) : nullptr;
-}
-
-const UPCGPin* UPCGCustomHLSLSettings::GetFirstInputPin() const
-{
-	if (const UPCGNode* Node = Cast<UPCGNode>(GetOuter()))
-	{
-		return !Node->GetInputPins().IsEmpty() ? Node->GetInputPins()[0] : nullptr;
-	}
-
-	return nullptr;
-}
-
-const UPCGPin* UPCGCustomHLSLSettings::GetFirstOutputPin() const
-{
-	if (const UPCGNode* Node = Cast<UPCGNode>(GetOuter()))
-	{
-		return !Node->GetOutputPins().IsEmpty() ? Node->GetOutputPins()[0] : nullptr;
-	}
-
-	return nullptr;
+	return !OutputPins.IsEmpty() ? &OutputPins[0] : nullptr;
 }
 
 #if WITH_EDITOR
@@ -387,24 +365,24 @@ void UPCGCustomHLSLSettings::UpdateInputDeclarations()
 
 		if (IsProcessorKernel())
 		{
-			const UPCGPin* InputPin = GetFirstInputPin();
-			const UPCGPin* OutputPin = GetFirstOutputPin();
+			const FPCGPinProperties* InputPin = GetFirstInputPinProperties();
+			const FPCGPinPropertiesGPU* OutputPin = GetFirstOutputPinProperties();
 
 			if (InputPin && OutputPin)
 			{
 				InputDeclarations += FString::Format(TEXT(
 					"const uint {0}_DataIndex;\n"
 					"const uint {1}_DataIndex;\n"),
-					{ InputPin->Properties.Label.ToString(),  OutputPin->Properties.Label.ToString() });
+					{ InputPin->Label.ToString(),  OutputPin->Label.ToString() });
 			}
 		}
 		else if (IsGeneratorKernel())
 		{
-			if (const UPCGPin* PointProcessingOutputPin = GetFirstOutputPin())
+			if (const FPCGPinPropertiesGPU* PointProcessingOutputPin = GetFirstOutputPinProperties())
 			{
 				InputDeclarations += FString::Format(
 					TEXT("const uint {0}_DataIndex;\n"),
-					{ PointProcessingOutputPin->Properties.Label.ToString() });
+					{ PointProcessingOutputPin->Label.ToString() });
 			}
 		}
 
diff --git a/Engine/Plugins/PCG/Source/PCG/Private/Compute/Elements/PCGCustomHLSL.h b/Engine/Plugins/PCG/Source/PCG/Private/Compute/Elements/PCGCustomHLSL.h
--- a/Engine/Plugins/PCG/Source/PCG/Private/Compute/Elements/PCGCustomHLSL.h
+++ b/Engine/Plugins/PCG/Source/PCG/Private/Compute/Elements/PCGCustomHLSL.h
@@ -125,10 +125,8 @@ protected:
 	TArray<FName> GetInputPinNamesAndNone() const;
 #endif
 
-	const UPCGPin* GetInputPin(FName Label) const;
-	const UPCGPin* GetOutputPin(FName Label) const;
-	const UPCGPin* GetFirstInputPin() const;
-	const UPCGPin* GetFirstOutputPin() const;
+	const FPCGPinProperties* GetFirstInputPinProperties() const;
+	const FPCGPinPropertiesGPU* GetFirstOutputPinProperties() const;
 
 	/** Will the ThreadCountMultiplier value be applied when calculating the dispatch thread count. */
 	bool IsThreadCountMultiplierInUse() const { return KernelType == EPCGKernelType::Custom && DispatchThreadCount != EPCGDispatchThreadCount::Fixed; }
diff --git a/Engine/Plugins/PCG/Source/PCG/Private/Compute/PCGComputeKernel.cpp b/Engine/Plugins/PCG/Source/PCG/Private/Compute/PCGComputeKernel.cpp
--- a/Engine/Plugins/PCG/Source/PCG/Private/Compute/PCGComputeKernel.cpp
+++ b/Engine/Plugins/PCG/Source/PCG/Private/Compute/PCGComputeKernel.cpp
@@ -12,6 +12,7 @@
 #include "Compute/PCGComputeSource.h"
 #include "Compute/PCGDataBinding.h"
 #include "Compute/PCGPinPropertiesGPU.h"
+#include "Compute/Elements/PCGComputeGraphElement.h"
 #include "Graph/PCGGPUGraphCompilationContext.h"
 #include "Helpers/PCGSettingsHelpers.h"
 
@@ -120,6 +121,82 @@ bool UPCGComputeKernel::AreKernelSettingsValid(FPCGContext* InContext) const
 	return !HasStaticValidationErrors();
 }
 
+bool UPCGComputeKernel::IsKernelDataValid(FPCGContext* InContext) const
+{
+	TRACE_CPUPROFILER_EVENT_SCOPE(UPCGComputeKernel::IsKernelDataValid);
+
+	// The context can either be a compute graph element context (if the compute graph was successfully created), otherwise
+	// it will be the original CPU node context. We need the former to run the following validation.
+	if (!InContext || !InContext->IsComputeContext())
+	{
+		return true;
+	}
+
+	const UPCGSettings* Settings = CastChecked<UPCGSettings>(GetSettings());
+	const UPCGNode* Node = Cast<UPCGNode>(Settings->GetOuter());
+	const FPCGComputeGraphContext* Context = static_cast<FPCGComputeGraphContext*>(InContext);
+	const UPCGDataBinding* DataBinding = Context->DataBinding.Get();
+
+	if (DataBinding && Node)
+	{
+		TArray<FPCGPinPropertiesGPU> OutputPins;
+		GetOutputPins(OutputPins);
+
+		for (const FPCGPinPropertiesGPU& Pin : OutputPins)
+		{
+			// Pairwise requires all data to be N:N, 1:N, or N:1.
+			if (Pin.PropertiesGPU.InitializationMode == EPCGPinInitMode::FromInputPins
+				&& Pin.PropertiesGPU.DataCountMode == EPCGDataCountMode::FromInputData
+				&& Pin.PropertiesGPU.DataMultiplicity == EPCGDataMultiplicity::Pairwise)
+			{
+				uint32 NumData = -1;
+
+				for (FName InitPin : Pin.PropertiesGPU.PinsToInititalizeFrom)
+				{
+					// This is a defensive condition for legacy data to avoid a long chain of ensures:
+					// It is possible for the InitPins to be out of date with which pins actually exist on the node. This could cause a cache miss
+					// when looking up the pin data description, so only consider InitPins which actually exist.
+					if (Node->GetInputPin(InitPin))
+					{
+						const FPCGDataCollectionDesc* InputDesc = DataBinding->GetCachedKernelPinDataDesc(this, InitPin, /*bIsInputPin=*/true);
+
+						if (ensure(InputDesc))
+						{
+							const uint32 InputNumData = InputDesc->DataDescs.Num();
+
+							// Only test pins that have a non-trivial number of data.
+							if (InputNumData > 1)
+							{
+								if (NumData == -1)
+								{
+									NumData = InputNumData;
+								}
+								else if (InputNumData != NumData)
+								{
+									PCG_KERNEL_VALIDATION_ERR(InContext, GetSettings(), FText::Format(
+										LOCTEXT("NumDataMismatch", "Number of data found on input pin '{0}' ({1}), mismatches with number of data found on another input pin ({2}). Pairwise data multiplicity only supports N:N, 1:N and N:1 operation."),
+										FText::FromName(InitPin),
+										InputNumData,
+										NumData));
+
+									return false;
+								}
+							}
+							else if (InputNumData == 0)
+							{
+								// No validation is necessary if any pin has 0 data. There are no pairs to validate.
+								break;
+							}
+						}
+					}
+				}
+			}
+		}
+	}
+
+	return true;
+}
+
 #if WITH_EDITOR
 void UPCGComputeKernel::GatherAdditionalSources(TArray<TObjectPtr<UComputeSource>>& OutAdditionalSources) const
 {
@@ -195,6 +272,13 @@ void UPCGComputeKernel::ComputeDataDescFromPinProperties(const FPCGPinProperties
 				int TotalNumElements = 0; // Total number of elements computed for this data index.
 				FIntPoint TotalNumElements2D = FIntPoint::ZeroValue;
 
+				// If the element count is a product, we have to start at 1 instead of 0 or we'll end up with 0!
+				if (Props.ElementMultiplicity == EPCGElementMultiplicity::Product && NumInitPins > 0)
+				{
+					TotalNumElements = 1;
+					TotalNumElements2D = FIntPoint(1, 1);
+				}
+
 				// For each data index, loop over all the pins and create the uber-data.
 				for (int InputPinIndex = 0; InputPinIndex < NumInitPins; ++InputPinIndex)
 				{
@@ -202,10 +286,10 @@ void UPCGComputeKernel::ComputeDataDescFromPinProperties(const FPCGPinProperties
 
 					int ClampedDataIndex = DataIndex;
 
-					// If this pin does not have the same number of data, clamp it to the first data.
-					if (DataIndex != InputDesc.DataDescs.Num())
+					// If this pin does not have the same number of data, clamp it to the last data.
+					if (DataIndex >= InputDesc.DataDescs.Num())
 					{
-						ClampedDataIndex = 0;
+						ClampedDataIndex = InputDesc.DataDescs.Num() - 1;
 					}
 
 					const FPCGDataDesc& InputDataDesc = InputDesc.DataDescs.IsValidIndex(ClampedDataIndex) ? InputDesc.DataDescs[ClampedDataIndex] : FPCGDataDesc(EPCGDataType::Any, 0);
@@ -217,15 +301,11 @@ void UPCGComputeKernel::ComputeDataDescFromPinProperties(const FPCGPinProperties
 						{
 							if (bDomainIs2D)
 							{
-								TotalNumElements2D = TotalNumElements2D.ComponentMax(FIntPoint(1, 1)) * (bSourceDomainIs2D
-									? InputDataDesc.ElementCount2D.X * InputDataDesc.ElementCount2D.Y
-									: InputDataDesc.ElementCount);
+								TotalNumElements2D *= bSourceDomainIs2D ? (InputDataDesc.ElementCount2D.X * InputDataDesc.ElementCount2D.Y) : InputDataDesc.ElementCount;
 							}
 							else
 							{
-								TotalNumElements = FMath::Max(TotalNumElements, 1) * (bSourceDomainIs2D
-									? InputDataDesc.ElementCount2D.X * InputDataDesc.ElementCount2D.Y
-									: InputDataDesc.ElementCount);
+								TotalNumElements *= bSourceDomainIs2D ? (InputDataDesc.ElementCount2D.X * InputDataDesc.ElementCount2D.Y) : InputDataDesc.ElementCount;
 							}
 						}
 						else if (Props.ElementMultiplicity == EPCGElementMultiplicity::Sum)
@@ -236,7 +316,7 @@ void UPCGComputeKernel::ComputeDataDescFromPinProperties(const FPCGPinProperties
 							}
 							else
 							{
-								TotalNumElements += bSourceDomainIs2D ? InputDataDesc.ElementCount2D.X * InputDataDesc.ElementCount2D.Y : InputDataDesc.ElementCount;
+								TotalNumElements += bSourceDomainIs2D ? (InputDataDesc.ElementCount2D.X * InputDataDesc.ElementCount2D.Y) : InputDataDesc.ElementCount;
 							}
 						}
 						else
@@ -420,7 +500,16 @@ void UPCGComputeKernel::ComputeDataDescFromPinProperties(const FPCGPinProperties
 				// Find the maximum number of data among the init pins. Note, they should all be the same number of data, or only one data.
 				for (int I = 0; I < NumInitPins; ++I)
 				{
-					MaxDataCount = FMath::Max(MaxDataCount, InputDescs[I].DataDescs.Num());
+					const int DataCount = InputDescs[I].DataDescs.Num();
+
+					// If any init pins have no data, then the pairing fails and no data are produced.
+					if (DataCount == 0)
+					{
+						MaxDataCount = 0;
+						break;
+					}
+
+					MaxDataCount = FMath::Max(MaxDataCount, DataCount);
 				}
 
 				AddDataPairwise(MaxDataCount);
diff --git a/Engine/Plugins/PCG/Source/PCG/Private/Compute/PCGComputeKernel.h b/Engine/Plugins/PCG/Source/PCG/Private/Compute/PCGComputeKernel.h
--- a/Engine/Plugins/PCG/Source/PCG/Private/Compute/PCGComputeKernel.h
+++ b/Engine/Plugins/PCG/Source/PCG/Private/Compute/PCGComputeKernel.h
@@ -97,7 +97,7 @@ public:
 	virtual bool AreKernelSettingsValid(FPCGContext* InContext) const;
 
 	/** Performs data validation and returns true if this node is suitable for deployment to the GPU. */
-	virtual bool IsKernelDataValid(FPCGContext* InContext) const { return true; }
+	virtual bool IsKernelDataValid(FPCGContext* InContext) const;
 
 #if WITH_EDITOR
 	/** Produces the node specific portion of kernel shader source text, including the main entry point. */
diff --git a/Engine/Plugins/PCG/Source/PCG/Private/Compute/PCGPinPropertiesGPU.h b/Engine/Plugins/PCG/Source/PCG/Private/Compute/PCGPinPropertiesGPU.h
--- a/Engine/Plugins/PCG/Source/PCG/Private/Compute/PCGPinPropertiesGPU.h
+++ b/Engine/Plugins/PCG/Source/PCG/Private/Compute/PCGPinPropertiesGPU.h
@@ -37,7 +37,7 @@ enum class EPCGElementCountMode : uint8
 UENUM()
 enum class EPCGDataMultiplicity : uint8
 {
-	Pairwise UMETA(Tooltip = "A data item will be produced for each pair/tuple/etc of input data items across the input pins. Requires all pins to have the same number of data."),
+	Pairwise UMETA(Tooltip = "A data item will be produced for each pair/tuple/etc of input data items across the input pins. Requires all pins to have 1 or N data items."),
 	CartesianProduct UMETA(Tooltip = "If there are two input pins with N and M data items respectively, the output will have NxM data items."),
 };
 
diff --git a/Engine/Plugins/PCG/Source/PCG/Private/Elements/PCGCopyPointsKernelShared.cpp b/Engine/Plugins/PCG/Source/PCG/Private/Elements/PCGCopyPointsKernelShared.cpp
--- a/Engine/Plugins/PCG/Source/PCG/Private/Elements/PCGCopyPointsKernelShared.cpp
+++ b/Engine/Plugins/PCG/Source/PCG/Private/Elements/PCGCopyPointsKernelShared.cpp
@@ -3,6 +3,7 @@
 #include "Elements/PCGCopyPointsKernelShared.h"
 
 #include "PCGContext.h"
+#include "Compute/PCGComputeKernel.h"
 #include "Compute/PCGDataBinding.h"
 #include "Compute/Elements/PCGComputeGraphElement.h"
 #include "Elements/PCGCopyPoints.h"
@@ -14,20 +15,39 @@ bool PCGCopyPointsKernel::IsKernelDataValid(const UPCGComputeKernel* InKernel, c
 {
 	check(InCopyPointSettings);
 	check(InContext);
+	check(InKernel);
 
-	if (InCopyPointSettings->bMatchBasedOnAttribute)
+	if (const UPCGDataBinding* DataBinding = InContext->DataBinding.Get(); ensure(DataBinding))
 	{
-		if (const UPCGDataBinding* DataBinding = InContext->DataBinding.Get(); ensure(DataBinding))
+		const FPCGDataCollectionDesc* SourcePinDesc = DataBinding->GetCachedKernelPinDataDesc(InKernel, PCGCopyPointsConstants::SourcePointsLabel, /*bIsInputPin=*/true);
+
+		if (!ensure(SourcePinDesc))
 		{
-			auto ValidateAttributeExists = [InContext, DataBinding, InKernel, InCopyPointSettings](FName InputPin) -> bool
-			{
-				const FPCGDataCollectionDesc* PinDataDesc = DataBinding->GetCachedKernelPinDataDesc(InKernel, InputPin, /*bIsInputPin=*/true);
+			return false;
+		}
 
-				if (!ensure(PinDataDesc))
-				{
-					return false;
-				}
+		const FPCGDataCollectionDesc* TargetPinDesc = DataBinding->GetCachedKernelPinDataDesc(InKernel, PCGCopyPointsConstants::TargetPointsLabel, /*bIsInputPin=*/true);
+
+		if (!ensure(TargetPinDesc))
+		{
+			return false;
+		}
+
+		const int32 NumSources = SourcePinDesc->DataDescs.Num();
+		const int32 NumTargets = TargetPinDesc->DataDescs.Num();
 
+		if (!InCopyPointSettings->bCopyEachSourceOnEveryTarget && (NumSources != NumTargets && NumSources != 1 && NumTargets != 1 && NumSources != 0 && NumTargets != 0))
+		{
+			PCG_KERNEL_VALIDATION_ERR(InContext, InCopyPointSettings, FText::Format(
+				LOCTEXT("NumDataMismatch", "Num Sources ({0}), mismatches with Num Targets ({1}). Only supports N:N, 1:N and N:1 operation."),
+				NumSources,
+				NumTargets));
+		}
+
+		if (InCopyPointSettings->bMatchBasedOnAttribute)
+		{
+			auto ValidateAttributeExists = [InContext, InCopyPointSettings](const FPCGDataCollectionDesc* PinDataDesc) -> bool
+			{
 				const FName MatchAttributeName = InCopyPointSettings->MatchAttribute;
 
 				// todo_pcg: Can generalize this to any type?
@@ -49,7 +69,7 @@ bool PCGCopyPointsKernel::IsKernelDataValid(const UPCGComputeKernel* InKernel, c
 				return !PinDataDesc->DataDescs.IsEmpty();
 			};
 
-			if (!ValidateAttributeExists(PCGCopyPointsConstants::SourcePointsLabel) || !ValidateAttributeExists(PCGCopyPointsConstants::TargetPointsLabel))
+			if (!ValidateAttributeExists(SourcePinDesc) || !ValidateAttributeExists(TargetPinDesc))
 			{
 				return false;
 			}
