[PCG] Fix uninitialized data warning when data is initialized with InitializePoint

#rnx

diff --git a/Engine/Plugins/PCG/Source/PCG/Private/Compute/Elements/PCGCustomHLSLKernel.cpp b/Engine/Plugins/PCG/Source/PCG/Private/Compute/Elements/PCGCustomHLSLKernel.cpp
--- a/Engine/Plugins/PCG/Source/PCG/Private/Compute/Elements/PCGCustomHLSLKernel.cpp
+++ b/Engine/Plugins/PCG/Source/PCG/Private/Compute/Elements/PCGCustomHLSLKernel.cpp
@@ -30,7 +30,7 @@ namespace PCGCustomHLSLKernel
 	constexpr TCHAR AttributeFunctionGetKeyword[] = { TEXT("Get") };
 	constexpr TCHAR AttributeFunctionSetKeyword[] = { TEXT("Set") };
 	constexpr TCHAR CopyElementFunctionKeyword[] = { TEXT("CopyElementFrom") };
-	constexpr TCHAR InitializeFunctionKeyword[] = { TEXT("Initialize") };
+	constexpr TCHAR InitializeFunctionKeyword[] = { TEXT("InitializePoint") };
 	constexpr TCHAR StoreFunctionKeyword[] = { TEXT("Store") };
 	constexpr TCHAR SetPositionKeyword[] = { TEXT("SetPosition") };
 	constexpr TCHAR SetRotationKeyword[] = { TEXT("SetRotation") };
