[PCG] Improved filtering cleanliness and added first metadata attributes. Also fixed small error wrt scale/rotation on intersection points.

diff --git a/Engine/Plugins/PCG/Source/PCG/Private/Elements/PCGSplineIntersection.cpp b/Engine/Plugins/PCG/Source/PCG/Private/Elements/PCGSplineIntersection.cpp
--- a/Engine/Plugins/PCG/Source/PCG/Private/Elements/PCGSplineIntersection.cpp
+++ b/Engine/Plugins/PCG/Source/PCG/Private/Elements/PCGSplineIntersection.cpp
@@ -76,7 +76,9 @@ namespace PCGSplineIntersection
 	struct IntersectionPoint
 	{
 		const UPCGSplineData* A = nullptr;
+		const InputData* InputA = nullptr;
 		const UPCGSplineData* B = nullptr;
+		const InputData* InputB = nullptr;
 
 		FVector AveragePosition;
 		FVector PositionA;
@@ -87,15 +89,57 @@ namespace PCGSplineIntersection
 		bool IsOnSpline(const UPCGSplineData* In) const { return In == A || In == B; }
 		float GetKey(const UPCGSplineData* In) const { check(In == A || In == B); return In == A ? InputKeyA : InputKeyB; }
 		const UPCGSplineData* GetOtherSpline(const UPCGSplineData* In) const { check(In == A || In == B); return In == A ? B : A; }
+		const InputData* GetOtherInput(const InputData* In) const { check(In == InputA || In == InputB); return In == InputA ? InputB : InputA; }
 		const FVector& GetPosition(const UPCGSplineData* In) const { check(In == A || In == B); return In == A ? PositionA : PositionB; }
 		FVector& GetPosition(const UPCGSplineData* In) { check(In == A || In == B); return In == A ? PositionA : PositionB; }
 	};
+
+	TArray<IntersectionPoint> CollapseColocatedIntersections(const TArray<IntersectionPoint>& InPoints, const double SqrDistanceThreshold, const UPCGSplineData* SplineFilter)
+	{
+		TArray<IntersectionPoint> OutPoints;
+		OutPoints.Reserve(InPoints.Num());
+
+		TArray<int> PointsWeight;
+		PointsWeight.Reserve(InPoints.Num());
+
+		for (const IntersectionPoint& Point : InPoints)
+		{
+			if (SplineFilter && !Point.IsOnSpline(SplineFilter))
+			{
+				continue;
+			}
+
+			const FVector& Position = SplineFilter ? Point.GetPosition(SplineFilter) : Point.AveragePosition;
+
+			int MatchingPointIndex = OutPoints.IndexOfByPredicate([&Position, &SqrDistanceThreshold](const IntersectionPoint& OutPt)
+			{
+				return (OutPt.AveragePosition - Position).SquaredLength() <= SqrDistanceThreshold;
+			});
+
+			if (MatchingPointIndex == INDEX_NONE)
+			{
+				IntersectionPoint& NewPoint = OutPoints.Emplace_GetRef(Point);
+				NewPoint.AveragePosition = Position; // meaning of average position will change when we are filtering on a specific spline
+				PointsWeight.Add(1);
+			}
+			else
+			{
+				const FVector NewAveragePosition = ((PointsWeight[MatchingPointIndex] * OutPoints[MatchingPointIndex].AveragePosition) + Point.AveragePosition) / static_cast<double>(PointsWeight[MatchingPointIndex] + 1);
+				OutPoints[MatchingPointIndex].AveragePosition = NewAveragePosition;
+				++PointsWeight[MatchingPointIndex];
+
+				// @todo_pcg: Add additional information for more intersection info
+			}
+		}
+
+		return OutPoints;
+	}
 }
 
 UPCGSplineIntersectionSettings::UPCGSplineIntersectionSettings()
 {
-	//OriginatingSplineIndexAttribute.SetAttributeName(PCGSplineIntersection::Constants::OriginatingSplineIndexAttributeName);
-	//IntersectingSplineIndexAttribute.SetAttributeName(PCGSplineIntersection::Constants::IntersectingSplineIndexAttributeName);
+	OriginatingSplineIndexAttribute.SetAttributeName(PCGSplineIntersection::Constants::OriginatingSplineIndexAttributeName);
+	IntersectingSplineIndexAttribute.SetAttributeName(PCGSplineIntersection::Constants::IntersectingSplineIndexAttributeName);
 }
 
 #if WITH_EDITOR
@@ -282,7 +326,9 @@ bool FPCGSplineIntersectionElement::ExecuteInternal(FPCGContext* Context) const
 							const FVector MidPoint = (ClosestPointOnCurrentSpline + ClosestPointOnOtherSpline) * 0.5;
 							PCGSplineIntersection::IntersectionPoint& Intersection = Intersections.Emplace_GetRef();
 							Intersection.A = CurrentSplineData.Spline;
+							Intersection.InputA = &CurrentSplineData;
 							Intersection.B = OtherSplineData.Spline;
+							Intersection.InputB = &OtherSplineData;
 							Intersection.PositionA = ClosestPointOnCurrentSpline;
 							Intersection.PositionB = ClosestPointOnOtherSpline;
 							Intersection.AveragePosition = MidPoint;
@@ -297,107 +343,82 @@ bool FPCGSplineIntersectionElement::ExecuteInternal(FPCGContext* Context) const
 
 	TArray<FPCGTaggedData>& Outputs = Context->OutputData.TaggedData;
 
+	auto LogErrorOnIndexAttributeCreation = [Context](const FName IndexName)
+	{
+		PCGLog::LogErrorOnGraph(FText::Format(LOCTEXT("FailedCreateIndexAttribute", "Failed to create the index attribute '{0}'."), FText::FromName(IndexName)), Context);
+	};
+
 	// Finally, build output data.
 	if (Settings->Output == EPCGSplineIntersectionOutput::IntersectionPointsOnly)
 	{
-		// Filter out points that are too close to each other.
-		TArray<TPair<FVector, double>> WeightedPoints;
-		WeightedPoints.Reserve(Intersections.Num());
-
-		for (const PCGSplineIntersection::IntersectionPoint& Intersection : Intersections)
-		{
-			const FVector& Position = Intersection.AveragePosition;
-
-			int MatchingPoint = WeightedPoints.IndexOfByPredicate([&Position, &SqrDistanceThreshold](const TPair<FVector, double>& WP)
-			{
-				return (WP.Key - Position).SquaredLength() <= SqrDistanceThreshold;
-			});
-
-			if (MatchingPoint == INDEX_NONE)
-			{
-				WeightedPoints.Emplace(Position, 1.0);
-			}
-			else
-			{
-				TPair<FVector, double>& Match = WeightedPoints[MatchingPoint];
-				Match.Key = (Match.Value * Match.Key + Position) / (Match.Value + 1);
-				Match.Value += 1.0;
-			}
-		}
+		// Start by filtering out points that are too close to each other.
+		TArray<PCGSplineIntersection::IntersectionPoint> FilteredIntersections = PCGSplineIntersection::CollapseColocatedIntersections(Intersections, SqrDistanceThreshold, nullptr);
 
 		FPCGTaggedData& Output = Outputs.Emplace_GetRef();
 		
 		UPCGBasePointData* PointData = FPCGContext::NewPointData_AnyThread(Context);
 		Output.Data = PointData;
 
-		PointData->SetNumPoints(WeightedPoints.Num());
-		PointData->AllocateProperties(EPCGPointNativeProperties::Transform);
+		PointData->SetNumPoints(FilteredIntersections.Num());
 
-		FPCGPointTransform::ValueRange Transforms = PointData->GetTransformValueRange();
-		for (int Index = 0; Index < WeightedPoints.Num(); ++Index)
-		{
-			Transforms[Index] = FTransform(WeightedPoints[Index].Key);
-		}
-	}
-	else if (Settings->Output == EPCGSplineIntersectionOutput::OriginalSplinesWithIntersections)
-	{
-		for (const PCGSplineIntersection::InputData& InputData : InputSplines)
-		{
-			const UPCGSplineData* InputSpline = InputData.Spline;
+		const EPCGPointNativeProperties PropertiesToAllocate = EPCGPointNativeProperties::Transform | (Settings->bOutputSplineIndices ? EPCGPointNativeProperties::MetadataEntry : EPCGPointNativeProperties::None);
+		PointData->AllocateProperties(PropertiesToAllocate);
 
-			TArray<TPair<float, PCGSplineIntersection::IntersectionPoint>> FilteredPoints;
-			for (const PCGSplineIntersection::IntersectionPoint& Intersection : Intersections)
+		FPCGMetadataAttribute<int32>* FirstIndexAttribute = nullptr;
+		FPCGMetadataAttribute<int32>* SecondIndexAttribute = nullptr;
+
+		if (Settings->bOutputSplineIndices)
+		{
+			auto CreateIndexAttribute = [PointData, &LogErrorOnIndexAttributeCreation](const FName IndexName) -> FPCGMetadataAttribute<int32>*
 			{
-				if (Intersection.IsOnSpline(InputSpline))
+				FPCGMetadataAttribute<int32>* IndexAttribute = PointData->MutableMetadata()->FindOrCreateAttribute<int32>(IndexName, -1, /*bAllowInterpolation=*/false);
+				if (!IndexAttribute)
 				{
-					FilteredPoints.Emplace(Intersection.GetKey(InputSpline), Intersection);
+					LogErrorOnIndexAttributeCreation(IndexName);
 				}
-			}
 
-			// Sort by key
-			FilteredPoints.Sort([](const auto& A, const auto& B) { return A.Key < B.Key; });
+				return IndexAttribute;
+			};
+
+			FirstIndexAttribute = CreateIndexAttribute(Settings->OriginatingSplineIndexAttribute.GetAttributeName());
+			SecondIndexAttribute = CreateIndexAttribute(Settings->IntersectingSplineIndexAttribute.GetAttributeName());
+		}
 
-			// Collapse points that would be too close
-			TArray<bool> PointCollapsed;
-			PointCollapsed.SetNumZeroed(FilteredPoints.Num()); // 0 == false
+		FPCGPointValueRanges OutRanges(PointData, /*bAllocate=*/false);
 
-			TArray<TPair<float, PCGSplineIntersection::IntersectionPoint>> CollapsedPoints;
-			CollapsedPoints.Reserve(FilteredPoints.Num());
+		for (int Index = 0; Index < FilteredIntersections.Num(); ++Index)
+		{
+			const PCGSplineIntersection::IntersectionPoint& Point = FilteredIntersections[Index];
+			OutRanges.TransformRange[Index] = FTransform(Point.AveragePosition);
 
-			for (int FP = 0; FP < FilteredPoints.Num(); ++FP)
+			if (FirstIndexAttribute || SecondIndexAttribute)
 			{
-				if (PointCollapsed[FP])
+				PCGMetadataEntryKey& CurrentEntryKey = OutRanges.MetadataEntryRange[Index];
+				PointData->Metadata->InitializeOnSet(CurrentEntryKey);
+
+				if (FirstIndexAttribute)
 				{
-					continue;
+					FirstIndexAttribute->SetValue(CurrentEntryKey, Point.InputA->InputIndex);
 				}
 
-				TPair<float, PCGSplineIntersection::IntersectionPoint>& PointAndKey = CollapsedPoints.Add_GetRef(FilteredPoints[FP]);
-				PCGSplineIntersection::IntersectionPoint& Point = PointAndKey.Value;
-				FVector& PointPosition = Point.GetPosition(InputSpline);
-
-				double Weight = 1.0;
-				for (int OFP = FP + 1; OFP < FilteredPoints.Num(); ++OFP)
+				if (SecondIndexAttribute)
 				{
-					if (PointCollapsed[OFP] || Point.GetOtherSpline(InputSpline) != FilteredPoints[OFP].Value.GetOtherSpline(InputSpline))
-					{
-						continue;
-					}
-
-					const PCGSplineIntersection::IntersectionPoint& OtherPoint = FilteredPoints[OFP].Value;
-					const FVector& OtherPointPosition = OtherPoint.GetPosition(InputSpline);
-
-					if ((PointPosition - OtherPointPosition).SquaredLength() > SqrDistanceThreshold)
-					{
-						break;
-					}
-
-					// Merge points
-					PointCollapsed[OFP] = true;
-
-					PointPosition = (PointPosition * Weight + OtherPointPosition) / (Weight + 1.0);
-					Weight += 1.0;
+					SecondIndexAttribute->SetValue(CurrentEntryKey, Point.InputB->InputIndex);
 				}
 			}
+		}
+	}
+	else if (Settings->Output == EPCGSplineIntersectionOutput::OriginalSplinesWithIntersections)
+	{
+		for (const PCGSplineIntersection::InputData& InputData : InputSplines)
+		{
+			const UPCGSplineData* InputSpline = InputData.Spline;
+
+			// Filter points for the current spline and collapse points that would be too close.
+			// Note: the average position in this case will be when points are collapsed on the spline itself, not against the other splines
+			TArray<PCGSplineIntersection::IntersectionPoint> CollapsedPoints = PCGSplineIntersection::CollapseColocatedIntersections(Intersections, SqrDistanceThreshold, InputSpline);
+			// Sort by key
+			CollapsedPoints.Sort([InputSpline](const auto& A, const auto& B) { return A.GetKey(InputSpline) < B.GetKey(InputSpline); });
 
 			// @todo_pcg : we should not add points that are already on existing control points,
 			// but we should mark them as intersections
@@ -408,6 +429,21 @@ bool FPCGSplineIntersectionElement::ExecuteInternal(FPCGContext* Context) const
 			if (!CollapsedPoints.IsEmpty())
 			{
 				UPCGSplineData* NewSplineData = FPCGContext::NewObject_AnyThread<UPCGSplineData>(Context);
+				NewSplineData->InitializeFromData(InputData.Spline);
+
+				// implementation note: this does not support typecasting, will just overwrite the attribute.
+				FPCGMetadataAttribute<int32>* IndexAttribute = nullptr;
+
+				if (Settings->bOutputSplineIndices)
+				{
+					FName IndexAttributeName = Settings->IntersectingSplineIndexAttribute.GetAttributeName();
+					IndexAttribute = NewSplineData->MutableMetadata()->FindOrCreateAttribute(IndexAttributeName, -1, /*bAllowInterpolation=*/false);
+
+					if (!IndexAttribute)
+					{
+						LogErrorOnIndexAttributeCreation(IndexAttributeName);
+					}
+				}
 
 				// Implementation note: with the old spline representation, it's not possible to add control points in between other points, 
 				// because there are some places in the code where we rely on the index being equal to the key (esp. around the ReparamTable).
@@ -417,16 +453,39 @@ bool FPCGSplineIntersectionElement::ExecuteInternal(FPCGContext* Context) const
 				TArray<int64> SplinePointsEntryKeys = InputData.Spline->GetMetadataEntryKeysForSplinePoints();
 				check(SplinePoints.Num() == SplinePointsEntryKeys.Num() || SplinePointsEntryKeys.IsEmpty());
 				const bool bHasMetadata = !SplinePointsEntryKeys.IsEmpty();
+				const bool bNeedsMetadata = (IndexAttribute != nullptr);
+
+				// Prepare metadata entry keys if we're going to use them later
+				if (bNeedsMetadata && !bHasMetadata)
+				{
+					SplinePointsEntryKeys.SetNumUninitialized(SplinePoints.Num());
+					for (int64& SplinePointEntryKey : SplinePointsEntryKeys)
+					{
+						SplinePointEntryKey = PCGInvalidEntryKey;
+					}
+				}
+
+				auto OptionalWriteIntersectionMetadata = [bNeedsMetadata, NewSplineData, IndexAttribute, &InputData](int64& EntryKey, const PCGSplineIntersection::IntersectionPoint& IntersectionPoint)
+				{
+					if (bNeedsMetadata)
+					{
+						NewSplineData->Metadata->InitializeOnSet(EntryKey);
 
+						if (IndexAttribute)
+						{
+							IndexAttribute->SetValue(EntryKey, IntersectionPoint.GetOtherInput(&InputData)->InputIndex);
+						}
+					}
+				};
+				
 				const float MinKeyDifferenceThreshold = 1.0e-4f;
 				
-				for (const TPair<float, PCGSplineIntersection::IntersectionPoint>& PointToAdd : CollapsedPoints)
+				for (const PCGSplineIntersection::IntersectionPoint& PointToAdd : CollapsedPoints)
 				{
-					const float& Key = PointToAdd.Key;
-					const FVector& Position = InputData.Spline->SplineStruct.GetTransform().InverseTransformPosition(PointToAdd.Value.GetPosition(InputSpline));
+					const float Key = PointToAdd.GetKey(InputSpline);
+					const FVector& Position = InputData.Spline->SplineStruct.GetTransform().InverseTransformPosition(PointToAdd.AveragePosition);
 
-					// Find insertion index
-					// @todo_pcg: use bisection
+					// Find insertion index @todo_pcg: use bisection
 					int InsertionIndex = SplinePoints.Num();
 					for (int SPIndex = 0; SPIndex < SplinePoints.Num(); ++SPIndex)
 					{
@@ -445,9 +504,16 @@ bool FPCGSplineIntersectionElement::ExecuteInternal(FPCGContext* Context) const
 
 					const float NextPointVirtualKey = (bIsOnLastSegment ? FMath::CeilToInt(PreviousPoint.InputKey + 1.0e-4f) : NextPoint.InputKey);
 
-					if((FMath::Abs(PreviousPoint.InputKey - Key) <= MinKeyDifferenceThreshold && (PreviousPoint.Position - Position).SquaredLength() <= SqrDistanceThreshold) ||
-						((FMath::Abs(NextPointVirtualKey - Key) <= MinKeyDifferenceThreshold && (NextPoint.Position - Position).SquaredLength() <= SqrDistanceThreshold)))
+					if ((FMath::Abs(PreviousPoint.InputKey - Key) <= MinKeyDifferenceThreshold && (PreviousPoint.Position - Position).SquaredLength() <= SqrDistanceThreshold))
+					{
+						// Mark previous point as an intersection
+						OptionalWriteIntersectionMetadata(SplinePointsEntryKeys[InsertionIndex - 1], PointToAdd);
+						continue;
+					}
+					else if ((FMath::Abs(NextPointVirtualKey - Key) <= MinKeyDifferenceThreshold && (NextPoint.Position - Position).SquaredLength() <= SqrDistanceThreshold))
 					{
+						// Mark next point as an intersection
+						OptionalWriteIntersectionMetadata(SplinePointsEntryKeys[bIsOnLastSegment ? 0 : InsertionIndex], PointToAdd);
 						continue;
 					}
 
@@ -457,8 +523,8 @@ bool FPCGSplineIntersectionElement::ExecuteInternal(FPCGContext* Context) const
 					// @todo_pcg: we don't need to add points that are on the control point, however, we might want to do metadata for those
 					SplinePoint.ArriveTangent = InputData.Spline->SplineStruct.GetTangentAtSplineInputKey(Key, ESplineCoordinateSpace::Local);
 					SplinePoint.LeaveTangent = SplinePoint.ArriveTangent;
-					SplinePoint.Rotation = NewSplineData->SplineStruct.GetQuaternionAtSplineInputKey(Key, ESplineCoordinateSpace::Local).Rotator();
-					SplinePoint.Scale = NewSplineData->SplineStruct.GetScaleAtSplineInputKey(Key);
+					SplinePoint.Rotation = InputData.Spline->SplineStruct.GetQuaternionAtSplineInputKey(Key, ESplineCoordinateSpace::Local).Rotator();
+					SplinePoint.Scale = InputData.Spline->SplineStruct.GetScaleAtSplineInputKey(Key);
 					SplinePoint.Type = ESplinePointType::CurveCustomTangent;
 
 					// Adjust tangents
@@ -476,21 +542,24 @@ bool FPCGSplineIntersectionElement::ExecuteInternal(FPCGContext* Context) const
 					SplinePoint.LeaveTangent *= (NextPointVirtualKey - Key);
 
 					SplinePoints.Insert(SplinePoint, InsertionIndex);
-					if (bHasMetadata)
+					if (bHasMetadata || bNeedsMetadata)
 					{
-						SplinePointsEntryKeys.Insert(PCGInvalidEntryKey, InsertionIndex);
+						PCGMetadataEntryKey EntryKey = PCGInvalidEntryKey;
+
+						if (bHasMetadata)
+						{
+							InputData.Spline->WriteMetadataToEntry(Key, EntryKey, NewSplineData->MutableMetadata());
+						}
+
+						OptionalWriteIntersectionMetadata(EntryKey, PointToAdd);
+						SplinePointsEntryKeys.Insert(EntryKey, InsertionIndex);
 					}
 				}
 
 				// Mark all tangents as custom, and reset the keys
-				TArray<float> OriginalKeys;
-				OriginalKeys.SetNumUninitialized(SplinePoints.Num());
-
 				for(int SPIndex = 0; SPIndex < SplinePoints.Num(); ++SPIndex)
 				{
 					FSplinePoint& SplinePoint = SplinePoints[SPIndex];
-
-					OriginalKeys[SPIndex] = SplinePoint.InputKey;
 					SplinePoint.InputKey = static_cast<float>(SPIndex);
 					if (SplinePoint.Type == ESplinePointType::Curve || SplinePoint.Type == ESplinePointType::CurveClamped)
 					{
@@ -503,30 +572,11 @@ bool FPCGSplineIntersectionElement::ExecuteInternal(FPCGContext* Context) const
 					InputData.Spline->SplineStruct.IsClosedLoop(),
 					InputData.Spline->SplineStruct.GetTransform(),
 					SplinePointsEntryKeys);
-				NewSplineData->InitializeFromData(InputData.Spline);
-
-				// Compute metadata for added points
-				if (bHasMetadata)
-				{
-					TArrayView<PCGMetadataEntryKey> SplineMEK = NewSplineData->SplineStruct.GetMutableControlPointsEntryKeys();
-					for (int SPIndex = 0; SPIndex < SplinePoints.Num(); ++SPIndex)
-					{
-						FSplinePoint& SplinePoint = SplinePoints[SPIndex];
-						if (SplineMEK[SPIndex] == PCGInvalidEntryKey)
-						{
-							InputData.Spline->WriteMetadataToEntry(OriginalKeys[SPIndex], SplineMEK[SPIndex], NewSplineData->MutableMetadata());
-						}
-					}
-				}
 
 				Output.Data = NewSplineData;
 			}
 		}
 	}
-	else if(Settings->Output == EPCGSplineIntersectionOutput::SplitSplines)
-	{
-		// @todo_pcg
-	}
 
 	return true;
 }
diff --git a/Engine/Plugins/PCG/Source/PCG/Public/Elements/PCGSplineIntersection.h b/Engine/Plugins/PCG/Source/PCG/Public/Elements/PCGSplineIntersection.h
--- a/Engine/Plugins/PCG/Source/PCG/Public/Elements/PCGSplineIntersection.h
+++ b/Engine/Plugins/PCG/Source/PCG/Public/Elements/PCGSplineIntersection.h
@@ -23,10 +23,7 @@ enum class EPCGSplineIntersectionOutput : uint8
 	IntersectionPointsOnly,
 
 	/** Adds control points at the intersections on the input splines. */
-	OriginalSplinesWithIntersections,
-
-	/** Splits splines at their intersections. */
-	SplitSplines UMETA(Hidden)
+	OriginalSplinesWithIntersections
 };
 
 /** 
@@ -65,16 +62,16 @@ public:
 	EPCGSplineIntersectionOutput Output = EPCGSplineIntersectionOutput::OriginalSplinesWithIntersections;
 
 	/** Controls whether the output will contain the index of the intersecting splines */
-	//UPROPERTY(BlueprintReadWrite, EditAnywhere, Category = Settings, meta=(EditCondition="Type!=EPCGSplineIntersectionType::Self"))
-	//bool bOutputSplineIndices = false;
+	UPROPERTY(BlueprintReadWrite, EditAnywhere, Category = Settings, meta=(EditCondition="Type!=EPCGSplineIntersectionType::Self"))
+	bool bOutputSplineIndices = false;
 
 	/** Attribute that will contain the first spline index (only in the output case of intersection points.) */
-	//UPROPERTY(BlueprintReadWrite, EditAnywhere, Category = Settings, meta=(PCG_Overridable, EditCondition="bOutputSplineIndices && Type!=EPCGSplineIntersetionType::Self && Output==EPCGSplineIntersectionOutput::IntersectionPointsOnly"))
-	//FPCGAttributePropertyOutputSelector OriginatingSplineIndexAttribute;
+	UPROPERTY(BlueprintReadWrite, EditAnywhere, Category = Settings, meta=(PCG_Overridable, PCG_DiscardPropertySelection, PCG_DiscardExtraSelection, EditCondition="bOutputSplineIndices && Type!=EPCGSplineIntersectionType::Self && Output==EPCGSplineIntersectionOutput::IntersectionPointsOnly"))
+	FPCGAttributePropertyOutputSelector OriginatingSplineIndexAttribute;
 
 	/** Attribute that will contain the first intersecting spline index, or -1 if not an intersection control point. */
-	//UPROPERTY(BlueprintReadWrite, EditAnywhere, Category = Settings, meta=(PCG_Overridable, EditCondition="bOutputSplineIndices"))
-	//FPCGAttributePropertyOutputSelector IntersectingSplineIndexAttribute;
+	UPROPERTY(BlueprintReadWrite, EditAnywhere, Category = Settings, meta=(PCG_Overridable, PCG_DiscardPropertySelection, PCG_DiscardExtraSelection, EditCondition="bOutputSplineIndices"))
+	FPCGAttributePropertyOutputSelector IntersectingSplineIndexAttribute;
 
 	/** Maximum distance at which we report an intersection on the splines. */
 	UPROPERTY(BlueprintReadWrite, EditAnywhere, Category = Settings, meta=(ClampMin="0.0001", PCG_Overridable))
