[PCG] Fix outstanding issue in package build where the Switch and MultiSelect nodes can fails if the DisplayName is different from the real name for an EnumValue

* GetDisplayName won't return the display name in package build but the real name
* Since we use the Display name for naming the pins, we need to keep a mapping.

#rnx

diff --git a/Engine/Plugins/PCG/Source/PCG/Private/Elements/ControlFlow/PCGMultiSelect.cpp b/Engine/Plugins/PCG/Source/PCG/Private/Elements/ControlFlow/PCGMultiSelect.cpp
--- a/Engine/Plugins/PCG/Source/PCG/Private/Elements/ControlFlow/PCGMultiSelect.cpp
+++ b/Engine/Plugins/PCG/Source/PCG/Private/Elements/ControlFlow/PCGMultiSelect.cpp
@@ -19,8 +19,6 @@ void UPCGMultiSelectSettings::PostLoad()
 {
 	Super::PostLoad();
 
-	CachePinLabels();
-
 #if WITH_EDITOR
 	// @todo_pcg To be behind a version bump in 5.5. Cannot do that in an hotfix
 	// Make sure we rename the pin properties that were serialized with a localized text. Since we can't exactly match the text
@@ -52,8 +50,21 @@ void UPCGMultiSelectSettings::PostLoad()
 				}
 			}
 		}
+
+		if (EnumValueToPinLabelMapping.IsEmpty())
+		{
+			CacheEnumValueToPinLabelsMapping();
+		}
 	}
 #endif // WITH_EDITOR
+
+	CachePinLabels();
+}
+
+void UPCGMultiSelectSettings::UpdateSettings()
+{
+	CacheEnumValueToPinLabelsMapping();
+	CachePinLabels();
 }
 
 void UPCGMultiSelectSettings::OnOverrideSettingsDuplicatedInternal(bool bSkippedPostLoad)
@@ -62,6 +73,7 @@ void UPCGMultiSelectSettings::OnOverrideSettingsDuplicatedInternal(bool bSkipped
 	if (bSkippedPostLoad)
 	{
 		CachePinLabels();
+		// We don't want to update settings here, because we don't want to erase the enum value to pin label mapping.
 	}
 }
 
@@ -71,11 +83,12 @@ void UPCGMultiSelectSettings::PostEditChangeProperty(FPropertyChangedEvent& Prop
 	Super::PostEditChangeProperty(PropertyChangedEvent);
 
 	// Only need to change the pin labels if the options have changed
-	if (PropertyChangedEvent.GetMemberPropertyName() == GET_MEMBER_NAME_CHECKED(UPCGMultiSelectSettings, IntOptions) ||
+	if (PropertyChangedEvent.GetMemberPropertyName() == GET_MEMBER_NAME_CHECKED(UPCGMultiSelectSettings, SelectionMode) ||
+		PropertyChangedEvent.GetMemberPropertyName() == GET_MEMBER_NAME_CHECKED(UPCGMultiSelectSettings, IntOptions) ||
 		PropertyChangedEvent.GetMemberPropertyName() == GET_MEMBER_NAME_CHECKED(UPCGMultiSelectSettings, StringOptions) ||
 		PropertyChangedEvent.GetMemberPropertyName() == GET_MEMBER_NAME_CHECKED(UPCGMultiSelectSettings, EnumSelection))
 	{
-		CachePinLabels();
+		UpdateSettings();
 	}
 }
 
@@ -135,7 +148,10 @@ FString UPCGMultiSelectSettings::GetAdditionalTitleInformation() const
 				FString Subtitle = EnumSelection.Class->GetName();
 				if (!IsPropertyOverriddenByPin({GET_MEMBER_NAME_CHECKED(UPCGMultiSelectSettings, EnumSelection), GET_MEMBER_NAME_CHECKED(FEnumSelector, Value)}))
 				{
-					Subtitle += FString::Format(TEXT(": {0}"), { EnumSelection.GetCultureInvariantDisplayName() });
+					if (FName* EnumValueName = EnumValueToPinLabelMapping.Find(EnumSelection.Value))
+					{
+						Subtitle += FString::Format(TEXT(": {0}"), { EnumValueName->ToString() });
+					}
 				}
 
 				return Subtitle;
@@ -183,19 +199,14 @@ TArray<FPCGPinProperties> UPCGMultiSelectSettings::InputPinProperties() const
 			}
 			break;
 		case EPCGControlFlowSelectionMode::Enum:
-			// -1 to bypass the MAX value
-			for (int32 Index = 0; EnumSelection.Class && Index < EnumSelection.Class->NumEnums() - 1; ++Index)
+			if (EnumValueToPinLabelMapping.IsEmpty())
 			{
-				bool bHidden = false;
-#if WITH_EDITOR
-				bHidden = EnumSelection.Class->HasMetaData(TEXT("Hidden"), Index) || EnumSelection.Class->HasMetaData(TEXT("Spacer"), Index);
-#endif // WITH_EDITOR
+				CacheEnumValueToPinLabelsMapping();
+			}
 
-				if (!bHidden)
-				{
-					const FString EnumDisplayName = EnumSelection.Class->GetDisplayNameTextByIndex(Index).BuildSourceString();
-					PinProperties.Emplace(FName(EnumDisplayName));
-				}
+			for (const auto& [Value, PinName] : EnumValueToPinLabelMapping)
+			{
+				PinProperties.Emplace(PinName);
 			}
 			break;
 		default:
@@ -263,30 +274,14 @@ bool UPCGMultiSelectSettings::GetSelectedPinLabel(FName& OutSelectedPinLabel) co
 	}
 	else if (SelectionMode == EPCGControlFlowSelectionMode::Enum && IsValuePresent(EnumSelection.Value))
 	{
-		// To account for hidden enums missing from the pin properties, find the pin by label instead of index
-		const FName PinLabel(EnumSelection.GetCultureInvariantDisplayName());
-		for (int i = 0; i < CachedPinLabels.Num(); ++i)
-		{
-			if (CachedPinLabels[i] == PinLabel)
-			{
-				Index = i;
-				break;
-			}
-		}
-	}
-	else
-	{
-		OutSelectedPinLabel = PCGControlFlowConstants::DefaultPathPinLabel;
+		const FName* PinLabel = EnumValueToPinLabelMapping.Find(EnumSelection.Value);
 
-		return true;
+		// Return index if found, otherwise fallback to the index after the options which will be "Default" pin.
+		Index = PinLabel ? CachedPinLabels.IndexOfByKey(*PinLabel) : INDEX_NONE;
 	}
 
-	if (Index < 0 || Index >= CachedPinLabels.Num())
-	{
-		return false;
-	}
-
-	OutSelectedPinLabel = CachedPinLabels[Index];
+	// If value is not found, use Default pin.
+	OutSelectedPinLabel = CachedPinLabels.IsValidIndex(Index) ? CachedPinLabels[Index] : PCGControlFlowConstants::DefaultPathPinLabel;
 
 	return true;
 }
@@ -300,6 +295,32 @@ void UPCGMultiSelectSettings::CachePinLabels()
 	});
 }
 
+void UPCGMultiSelectSettings::CacheEnumValueToPinLabelsMapping() const
+{
+	EnumValueToPinLabelMapping.Empty();
+
+	if (SelectionMode != EPCGControlFlowSelectionMode::Enum)
+	{
+		return;
+	}
+
+	// -1 to bypass the MAX value
+	for (int32 Index = 0; EnumSelection.Class && Index < EnumSelection.Class->NumEnums() - 1; ++Index)
+	{
+		bool bHidden = false;
+#if WITH_EDITOR
+		// HasMetaData is editor only, so there will be extra pins at runtime, but that should be okay
+		bHidden = EnumSelection.Class->HasMetaData(TEXT("Hidden"), Index) || EnumSelection.Class->HasMetaData(TEXT("Spacer"), Index);
+#endif // WITH_EDITOR
+
+		if (!bHidden)
+		{
+			const FString EnumDisplayName = EnumSelection.Class->GetDisplayNameTextByIndex(Index).BuildSourceString();
+			EnumValueToPinLabelMapping.Emplace( EnumSelection.Class->GetValueByIndex(Index), FName(EnumDisplayName));
+		}
+	}
+}
+
 bool FPCGMultiSelectElement::ExecuteInternal(FPCGContext* Context) const
 {
 	TRACE_CPUPROFILER_EVENT_SCOPE(FPCGMultiSelectElement::ExecuteInternal);
diff --git a/Engine/Plugins/PCG/Source/PCG/Private/Elements/ControlFlow/PCGSwitch.cpp b/Engine/Plugins/PCG/Source/PCG/Private/Elements/ControlFlow/PCGSwitch.cpp
--- a/Engine/Plugins/PCG/Source/PCG/Private/Elements/ControlFlow/PCGSwitch.cpp
+++ b/Engine/Plugins/PCG/Source/PCG/Private/Elements/ControlFlow/PCGSwitch.cpp
@@ -19,8 +19,6 @@ void UPCGSwitchSettings::PostLoad()
 {
 	Super::PostLoad();
 
-	CachePinLabels();
-
 #if WITH_EDITOR
 	// @todo_pcg To be behind a version bump in 5.5. Cannot do that in an hotfix
 	// Make sure we rename the pin properties that were serialized with a localized text. Since we can't exactly match the text
@@ -44,8 +42,21 @@ void UPCGSwitchSettings::PostLoad()
 				}
 			}
 		}
+
+		if (EnumValueToPinLabelMapping.IsEmpty())
+		{
+			CacheEnumValueToPinLabelsMapping();
+		}
 	}
 #endif // WITH_EDITOR
+
+	CachePinLabels();
+}
+
+void UPCGSwitchSettings::UpdateSettings()
+{
+	CacheEnumValueToPinLabelsMapping();
+	CachePinLabels();
 }
 
 void UPCGSwitchSettings::OnOverrideSettingsDuplicatedInternal(bool bSkippedPostLoad)
@@ -54,6 +65,7 @@ void UPCGSwitchSettings::OnOverrideSettingsDuplicatedInternal(bool bSkippedPostL
 	if (bSkippedPostLoad)
 	{
 		CachePinLabels();
+		// We don't want to update settings here, because we don't want to erase the enum value to pin label mapping.
 	}
 }
 
@@ -63,11 +75,12 @@ void UPCGSwitchSettings::PostEditChangeProperty(FPropertyChangedEvent& PropertyC
 	Super::PostEditChangeProperty(PropertyChangedEvent);
 
 	// Only need to change the pin labels if the options have changed
-	if (PropertyChangedEvent.GetMemberPropertyName() == GET_MEMBER_NAME_CHECKED(UPCGSwitchSettings, IntOptions) ||
+	if (PropertyChangedEvent.GetMemberPropertyName() == GET_MEMBER_NAME_CHECKED(UPCGSwitchSettings, SelectionMode) ||
+		PropertyChangedEvent.GetMemberPropertyName() == GET_MEMBER_NAME_CHECKED(UPCGSwitchSettings, IntOptions) ||
 		PropertyChangedEvent.GetMemberPropertyName() == GET_MEMBER_NAME_CHECKED(UPCGSwitchSettings, StringOptions) ||
 		PropertyChangedEvent.GetMemberPropertyName() == GET_MEMBER_NAME_CHECKED(UPCGSwitchSettings, EnumSelection))
 	{
-		CachePinLabels();
+		UpdateSettings();
 	}
 }
 
@@ -127,7 +140,10 @@ FString UPCGSwitchSettings::GetAdditionalTitleInformation() const
 				FString Subtitle = EnumSelection.Class->GetName();
 				if (!IsPropertyOverriddenByPin({GET_MEMBER_NAME_CHECKED(UPCGSwitchSettings, EnumSelection), GET_MEMBER_NAME_CHECKED(FEnumSelector, Value)}))
 				{
-					Subtitle += FString::Format(TEXT(": {0}"), { EnumSelection.GetCultureInvariantDisplayName() });
+					if (FName* EnumValueName = EnumValueToPinLabelMapping.Find(EnumSelection.Value))
+					{
+						Subtitle += FString::Format(TEXT(": {0}"), { EnumValueName->ToString() });
+					}
 				}
 
 				return Subtitle;
@@ -210,20 +226,14 @@ TArray<FPCGPinProperties> UPCGSwitchSettings::OutputPinProperties() const
 			break;
 
 		case EPCGControlFlowSelectionMode::Enum:
-			// -1 to bypass the MAX value
-			for (int32 Index = 0; EnumSelection.Class && Index < EnumSelection.Class->NumEnums() - 1; ++Index)
+			if (EnumValueToPinLabelMapping.IsEmpty())
 			{
-				bool bHidden = false;
-#if WITH_EDITOR
-				// HasMetaData is editor only, so there will be extra pins at runtime, but that should be okay
-				bHidden = EnumSelection.Class->HasMetaData(TEXT("Hidden"), Index) || EnumSelection.Class->HasMetaData(TEXT("Spacer"), Index);
-#endif // WITH_EDITOR
-
-				if (!bHidden)
-				{
-					const FString EnumDisplayName = EnumSelection.Class->GetDisplayNameTextByIndex(Index).BuildSourceString();
-					PinProperties.Emplace(FName(EnumDisplayName));
-				}
+				CacheEnumValueToPinLabelsMapping();
+			}
+			
+			for (const auto& [Value, PinName] : EnumValueToPinLabelMapping)
+			{
+				PinProperties.Emplace(PinName);
 			}
 			break;
 
@@ -289,12 +299,11 @@ int UPCGSwitchSettings::GetSelectedOutputPinIndex() const
 	}
 	else if (SelectionMode == EPCGControlFlowSelectionMode::Enum)
 	{
-		// A "hidden" value could be selected that wasn't cached, so do a name-wise comparison
-		const FName PinLabel(EnumSelection.GetCultureInvariantDisplayName());
+		const FName* PinLabel = EnumValueToPinLabelMapping.Find(EnumSelection.Value);
 
 		// Return index if found, otherwise fallback to the index after the options which will be "Default" pin.
-		const int FoundIndex = CachedPinLabels.IndexOfByKey(PinLabel);
-		return FoundIndex != INDEX_NONE ? FoundIndex : StringOptions.Num();
+		const int FoundIndex = PinLabel ? CachedPinLabels.IndexOfByKey(*PinLabel) : INDEX_NONE;
+		return FoundIndex != INDEX_NONE ? FoundIndex : EnumValueToPinLabelMapping.Num();
 	}
 
 	return INDEX_NONE;
@@ -306,43 +315,25 @@ bool UPCGSwitchSettings::GetSelectedPinLabel(FName& OutSelectedPinLabel) const
 	{
 		return false;
 	}
-
-	int32 Index = INDEX_NONE;
-	if (SelectionMode == EPCGControlFlowSelectionMode::Integer && IsValuePresent(IntegerSelection))
-	{
-		Index = IntOptions.IndexOfByKey(IntegerSelection);
-	}
-	else if (SelectionMode == EPCGControlFlowSelectionMode::String && IsValuePresent(StringSelection))
-	{
-		Index = StringOptions.IndexOfByKey(StringSelection);
-	}
-	else if (SelectionMode == EPCGControlFlowSelectionMode::Enum && IsValuePresent(EnumSelection.Value))
+	
+	if ((SelectionMode == EPCGControlFlowSelectionMode::Integer && IsValuePresent(IntegerSelection))
+		|| (SelectionMode == EPCGControlFlowSelectionMode::String && IsValuePresent(StringSelection))
+		|| (SelectionMode == EPCGControlFlowSelectionMode::Enum && IsValuePresent(EnumSelection.Value)))
 	{
-		// A "hidden" value could be selected that wasn't cached, so do a name-wise comparison
-		const FName PinLabel(EnumSelection.GetCultureInvariantDisplayName());
-		for (int i = 0; i < CachedPinLabels.Num(); ++i)
+		const int32 Index = GetSelectedOutputPinIndex();
+
+		if (!CachedPinLabels.IsValidIndex(Index))
 		{
-			if (CachedPinLabels[i] == PinLabel)
-			{
-				Index = i;
-				break;
-			}
+			return false;
 		}
+
+		OutSelectedPinLabel = CachedPinLabels[Index];
 	}
 	else
 	{
 		OutSelectedPinLabel = PCGControlFlowConstants::DefaultPathPinLabel;
-
-		return true;
 	}
 
-	if (Index < 0 || Index >= CachedPinLabels.Num())
-	{
-		return false;
-	}
-
-	OutSelectedPinLabel = CachedPinLabels[Index];
-
 	return true;
 }
 
@@ -355,6 +346,32 @@ void UPCGSwitchSettings::CachePinLabels()
 	});
 }
 
+void UPCGSwitchSettings::CacheEnumValueToPinLabelsMapping() const
+{
+	EnumValueToPinLabelMapping.Empty();
+
+	if (SelectionMode != EPCGControlFlowSelectionMode::Enum)
+	{
+		return;
+	}
+	
+	// -1 to bypass the MAX value
+	for (int32 Index = 0; EnumSelection.Class && Index < EnumSelection.Class->NumEnums() - 1; ++Index)
+	{
+		bool bHidden = false;
+#if WITH_EDITOR
+		// HasMetaData is editor only, so there will be extra pins at runtime, but that should be okay
+		bHidden = EnumSelection.Class->HasMetaData(TEXT("Hidden"), Index) || EnumSelection.Class->HasMetaData(TEXT("Spacer"), Index);
+#endif // WITH_EDITOR
+
+		if (!bHidden)
+		{
+			const FString EnumDisplayName = EnumSelection.Class->GetDisplayNameTextByIndex(Index).BuildSourceString();
+			EnumValueToPinLabelMapping.Emplace( EnumSelection.Class->GetValueByIndex(Index), FName(EnumDisplayName));
+		}
+	}
+}
+
 bool UPCGSwitchSettings::IsPinUsedByNodeExecution(const UPCGPin* InPin) const
 {
 	if (!InPin->IsOutputPin())
@@ -375,7 +392,7 @@ bool UPCGSwitchSettings::IsPinUsedByNodeExecution(const UPCGPin* InPin) const
 	}
 
 	const int SelectedPinIndex = GetSelectedOutputPinIndex();
-	if (SelectedPinIndex == INDEX_NONE || !CachedPinLabels.IsValidIndex(SelectedPinIndex))
+	if (!CachedPinLabels.IsValidIndex(SelectedPinIndex))
 	{
 		ensure(false);
 		return true;
diff --git a/Engine/Plugins/PCG/Source/PCG/Private/Tests/Elements/PCGMultiSelectTest.cpp b/Engine/Plugins/PCG/Source/PCG/Private/Tests/Elements/PCGMultiSelectTest.cpp
--- a/Engine/Plugins/PCG/Source/PCG/Private/Tests/Elements/PCGMultiSelectTest.cpp
+++ b/Engine/Plugins/PCG/Source/PCG/Private/Tests/Elements/PCGMultiSelectTest.cpp
@@ -28,7 +28,7 @@ protected:
 		SelectSettings->StringOptions = InSettings->StringOptions;
 		SelectSettings->StringSelection = InSettings->StringSelection;
 		SelectSettings->EnumSelection = InSettings->EnumSelection;
-		SelectSettings->CachePinLabels();
+		SelectSettings->UpdateSettings();
 	}
 
 	bool ExecuteTest(const UPCGMultiSelectSettings* InSettings, const int32 InputNum, TArray<FPCGTaggedData>& OutputData)
@@ -215,9 +215,6 @@ bool FPCGMultiSelectTest_Enum::RunTest(const FString& Parameters)
 
 	const int32 NumEnums = EnumClass->NumEnums();
 
-	// Two for the hidden enum values
-	AddExpectedError("Selected value is not a valid option.", EAutomationExpectedMessageFlags::Contains, 2);
-
 	// Positive Validation (also accounting for hidden metadata)
 	for (uint8 i = 0; i < NumEnums - 1; ++i)
 	{
@@ -229,7 +226,7 @@ bool FPCGMultiSelectTest_Enum::RunTest(const FString& Parameters)
 
 		if (TestSettings->EnumSelection.Class->HasMetaData(TEXT("Hidden"), i) || TestSettings->EnumSelection.Class->HasMetaData(TEXT("Spacer"), i))
 		{
-			UTEST_TRUE("Hidden enum selection output is empty", OutputData.IsEmpty());
+			UTEST_TRUE("Hidden or spacer enum selection defaults output", ValidateDefaultPathOutput(OutputData));
 		}
 		else
 		{
diff --git a/Engine/Plugins/PCG/Source/PCG/Private/Tests/Elements/PCGSwitchTest.cpp b/Engine/Plugins/PCG/Source/PCG/Private/Tests/Elements/PCGSwitchTest.cpp
--- a/Engine/Plugins/PCG/Source/PCG/Private/Tests/Elements/PCGSwitchTest.cpp
+++ b/Engine/Plugins/PCG/Source/PCG/Private/Tests/Elements/PCGSwitchTest.cpp
@@ -22,7 +22,7 @@ protected:
 		SelectSettings->StringOptions = InSettings->StringOptions;
 		SelectSettings->StringSelection = InSettings->StringSelection;
 		SelectSettings->EnumSelection = InSettings->EnumSelection;
-		SelectSettings->CachePinLabels();
+		SelectSettings->UpdateSettings();
 	}
 
 	bool ExecuteTest(const UPCGSwitchSettings* InSettings, FPCGDataCollection& OutputData)
@@ -151,9 +151,6 @@ bool FPCGSwitchTest_Enum::RunTest(const FString& Parameters)
 
 	const int32 NumEnums = EnumClass->NumEnums();
 
-	// Two for the hidden enum values
-	AddExpectedError("Selected value is not a valid option.", EAutomationExpectedMessageFlags::Contains, 2);
-
 	// Positive Validation (also accounting for hidden metadata)
 	for (uint8 i = 0; i < NumEnums - 1; ++i)
 	{
@@ -165,7 +162,7 @@ bool FPCGSwitchTest_Enum::RunTest(const FString& Parameters)
 
 		if (TestSettings->EnumSelection.Class->HasMetaData(TEXT("Hidden"), i) || TestSettings->EnumSelection.Class->HasMetaData(TEXT("Spacer"), i))
 		{
-			UTEST_TRUE("Hidden enum selection output is empty", OutputData.GetAllSpatialInputs().IsEmpty());
+			UTEST_TRUE("Hidden or spacer enum selection defaults output", ValidateOutput(OutputData, DefaultPathPinLabel));
 		}
 		else
 		{
diff --git a/Engine/Plugins/PCG/Source/PCG/Public/Elements/ControlFlow/PCGMultiSelect.h b/Engine/Plugins/PCG/Source/PCG/Public/Elements/ControlFlow/PCGMultiSelect.h
--- a/Engine/Plugins/PCG/Source/PCG/Public/Elements/ControlFlow/PCGMultiSelect.h
+++ b/Engine/Plugins/PCG/Source/PCG/Public/Elements/ControlFlow/PCGMultiSelect.h
@@ -6,6 +6,8 @@
 
 #include "PCGMultiSelect.generated.h"
 
+class FPCGMultiSelectElement;
+
 /**
  * Selects data from any number of input pins, based on a static selection criteria (Int/String/Enum) 
  */
@@ -14,6 +16,8 @@ class UPCGMultiSelectSettings : public UPCGControlFlowSettings
 {
 	GENERATED_BODY()
 
+	friend FPCGMultiSelectElement;
+
 public:
 	//~Begin UObject interface
 	virtual void PostLoad() override;
@@ -66,10 +70,18 @@ public:
 	UPROPERTY(BlueprintReadWrite, EditAnywhere, Category=Settings, meta=(PCG_Overridable, EditConditionHides, EditCondition="SelectionMode == EPCGControlFlowSelectionMode::Enum"))
 	FEnumSelector EnumSelection;
 
+	/** Since we use the Display name as the output pins for enums, this can be different in non-editor build. We need to keep a mapping between the enum value and its associated pin name. */
+	UPROPERTY()
+	mutable TMap<int64, FName> EnumValueToPinLabelMapping;
+
 	/** Cached pin labels for use during the selection process. */
 	UPROPERTY(Transient)
 	TArray<FName> CachedPinLabels;
 
+	/** To be called after the settings are modified directly in native code. */
+	PCG_API void UpdateSettings(); 
+
+protected:
 	/** Returns true if the integer value exists in the user defined array. */
 	bool IsValuePresent(int32 Value) const;
 
@@ -83,6 +95,8 @@ public:
 	bool GetSelectedPinLabel(FName& OutSelectedPinLabel) const;
 
 	void CachePinLabels();
+
+	void CacheEnumValueToPinLabelsMapping() const;
 };
 
 class FPCGMultiSelectElement : public IPCGElement
diff --git a/Engine/Plugins/PCG/Source/PCG/Public/Elements/ControlFlow/PCGSwitch.h b/Engine/Plugins/PCG/Source/PCG/Public/Elements/ControlFlow/PCGSwitch.h
--- a/Engine/Plugins/PCG/Source/PCG/Public/Elements/ControlFlow/PCGSwitch.h
+++ b/Engine/Plugins/PCG/Source/PCG/Public/Elements/ControlFlow/PCGSwitch.h
@@ -6,6 +6,8 @@
 
 #include "PCGSwitch.generated.h"
 
+class FPCGSwitchElement;
+
 /**
  * Routes data from the input pin, to a specific output pin based on a selection criteria (Int/String/Enum) 
  */
@@ -14,6 +16,8 @@ class UPCGSwitchSettings : public UPCGControlFlowSettings
 {
 	GENERATED_BODY()
 
+	friend FPCGSwitchElement;
+
 public:
 	//~Begin UObject interface
 	virtual void PostLoad() override;
@@ -74,9 +78,18 @@ public:
 	UPROPERTY(BlueprintReadWrite, EditAnywhere, Category=Settings, meta=(PCG_Overridable, EditConditionHides, EditCondition="SelectionMode == EPCGControlFlowSelectionMode::Enum"))
 	FEnumSelector EnumSelection;
 
+	/** Since we use the Display name as the output pins for enums, this can be different in non-editor build. We need to keep a mapping between the enum value and its associated pin name. */
+	UPROPERTY()
+	mutable TMap<int64, FName> EnumValueToPinLabelMapping;
+
 	/** Cached pin labels for use during the selection process. */
+	UPROPERTY(Transient)
 	TArray<FName> CachedPinLabels;
 
+	/** To be called after the settings are modified directly in native code. */
+	PCG_API void UpdateSettings(); 
+
+protected:
 	/** Returns true if the integer value exists in the user defined options. */
 	bool IsValuePresent(int32 Value) const;
 
@@ -93,6 +106,8 @@ public:
 	bool GetSelectedPinLabel(FName& OutSelectedPinLabel) const;
 
 	void CachePinLabels();
+	
+	void CacheEnumValueToPinLabelsMapping() const;
 };
 
 class FPCGSwitchElement : public IPCGElement
