[PCG] Improvements to Spline Mesh spawning

* Component class is now respected when spawning the spline mesh
* Added override on the component itself, if users have special fields they want to fill.

#rnx

diff --git a/Engine/Plugins/PCG/Source/PCG/Private/Elements/PCGSpawnSplineMesh.cpp b/Engine/Plugins/PCG/Source/PCG/Private/Elements/PCGSpawnSplineMesh.cpp
--- a/Engine/Plugins/PCG/Source/PCG/Private/Elements/PCGSpawnSplineMesh.cpp
+++ b/Engine/Plugins/PCG/Source/PCG/Private/Elements/PCGSpawnSplineMesh.cpp
@@ -150,6 +150,12 @@ bool FPCGSpawnSplineMeshElement::PrepareDataInternal(FPCGContext* InContext) con
 			// Also initialize the override for the params
 			OutState.ParamsOverrides.Initialize(Settings->SplineMeshParamsOverride, &OutState.SMCBuilderParams.SplineMeshParams, OutState.SplineData, Context);
 
+			// And initialize the override with the CDO component for now, it will be updated with the right component when it'll be spawned.
+			if (const TSubclassOf<USplineMeshComponent> SplineMeshComponentClass = Settings->SplineMeshDescriptor.ComponentClass; SplineMeshComponentClass)
+			{
+				OutState.ComponentOverrides.Initialize(Settings->SplineMeshComponentOverride, SplineMeshComponentClass->GetDefaultObject<USplineMeshComponent>(), OutState.SplineData, Context);
+			}
+
 			return EPCGTimeSliceInitResult::Success;
 		});
 
@@ -302,6 +308,13 @@ bool FPCGSpawnSplineMeshElement::ExecuteInternal(FPCGContext* InContext) const
 
 			USplineMeshComponent* SplineMeshComponent = UPCGActorHelpers::GetOrCreateSplineMeshComponent(ExecState.TargetActor, SourceComponent, IterState.SMCBuilderParams, Context);
 
+			// Apply the overrides on the spawned component
+			if (IterState.ComponentOverrides.IsValid())
+			{
+				IterState.ComponentOverrides.UpdateTemplateObject(SplineMeshComponent);
+				IterState.ComponentOverrides.Apply(IterState.ElementIndex);
+			}
+
 			// TODO: Write out the geometry to a dynamic mesh type.
 			//SplineMeshComponent->BodySetup->TriMeshGeometries
 
diff --git a/Engine/Plugins/PCG/Source/PCG/Private/Helpers/PCGActorHelpers.cpp b/Engine/Plugins/PCG/Source/PCG/Private/Helpers/PCGActorHelpers.cpp
--- a/Engine/Plugins/PCG/Source/PCG/Private/Helpers/PCGActorHelpers.cpp
+++ b/Engine/Plugins/PCG/Source/PCG/Private/Helpers/PCGActorHelpers.cpp
@@ -467,12 +467,10 @@ UPCGManagedSplineMeshComponent* UPCGActorHelpers::GetOrCreateManagedSplineMeshCo
 
 	FString ComponentName = TEXT("PCGSplineMeshComponent_") + StaticMesh->GetName();
 	const EObjectFlags ObjectFlags = (InSourceComponent->IsInPreviewMode() ? RF_Transient | RF_NonPIEDuplicateTransient : RF_NoFlags);
-	USplineMeshComponent* SplineMeshComponent = NewObject<USplineMeshComponent>(InTargetActor, MakeUniqueObjectName(InTargetActor, USplineMeshComponent::StaticClass(), FName(ComponentName)), ObjectFlags);
+	USplineMeshComponent* SplineMeshComponent = InParams.Descriptor.CreateComponent(InTargetActor, MakeUniqueObjectName(InTargetActor, USplineMeshComponent::StaticClass(), FName(ComponentName)), ObjectFlags);
 
 	// Init Component
 	{
-		InParams.Descriptor.InitComponent(SplineMeshComponent);
-
 		const FPCGSplineMeshParams& SplineMeshParams = InParams.SplineMeshParams;
 		SplineMeshComponent->SetStartAndEnd(SplineMeshParams.StartPosition, SplineMeshParams.StartTangent, SplineMeshParams.EndPosition, SplineMeshParams.EndTangent);
 		SplineMeshComponent->SetStartRollDegrees(SplineMeshParams.StartRollDegrees);
diff --git a/Engine/Plugins/PCG/Source/PCG/Public/Elements/PCGSpawnSplineMesh.h b/Engine/Plugins/PCG/Source/PCG/Public/Elements/PCGSpawnSplineMesh.h
--- a/Engine/Plugins/PCG/Source/PCG/Public/Elements/PCGSpawnSplineMesh.h
+++ b/Engine/Plugins/PCG/Source/PCG/Public/Elements/PCGSpawnSplineMesh.h
@@ -12,6 +12,7 @@
 #include "Helpers/PCGActorHelpers.h"
 #include "Metadata/PCGObjectPropertyOverride.h"
 
+#include "Components/SplineMeshComponent.h"
 #include "Engine/SplineMeshComponentDescriptor.h"
 
 #include "PCGSpawnSplineMesh.generated.h"
@@ -70,6 +71,10 @@ public:
 	/** Overrides for spline mesh params. */
 	UPROPERTY(EditAnywhere, BlueprintReadWrite, Category = Settings)
 	TArray<FPCGObjectPropertyOverrideDescription> SplineMeshParamsOverride;
+
+	/** Overrides for the spline mesh component. */
+	UPROPERTY(EditAnywhere, BlueprintReadWrite, Category = Settings)
+	TArray<FPCGObjectPropertyOverrideDescription> SplineMeshComponentOverride;
 };
 
 struct FPCGSpawnSplineMeshPerExecutionState
@@ -82,6 +87,7 @@ struct FPCGSpawnSplineMeshPerIterationState
 	FPCGSplineMeshComponentBuilderParameters SMCBuilderParams;
 	FPCGObjectOverrides<FSplineMeshComponentDescriptor> DescriptionOverrides = FPCGObjectOverrides(&SMCBuilderParams.Descriptor);
 	FPCGObjectOverrides<FPCGSplineMeshParams> ParamsOverrides = FPCGObjectOverrides(&SMCBuilderParams.SplineMeshParams);
+	FPCGObjectOverrides<USplineMeshComponent> ComponentOverrides{nullptr};
 	const UPCGPolyLineData* SplineData = nullptr;
 	const UPCGLandscapeSplineData* LandscapeSplineData = nullptr;
 	int32 ElementIndex = 0;
