[PCG] Remove ClampMin/Max for Partition grid size as it makes the slider unusable

* Rely only on Tooltip + clamping when the value changes to respect the valid range

#rnx

diff --git a/Engine/Plugins/PCG/Source/PCG/Public/PCGWorldActor.h b/Engine/Plugins/PCG/Source/PCG/Public/PCGWorldActor.h
--- a/Engine/Plugins/PCG/Source/PCG/Public/PCGWorldActor.h
+++ b/Engine/Plugins/PCG/Source/PCG/Public/PCGWorldActor.h
@@ -59,7 +59,7 @@ public:
 	static inline constexpr uint32 DefaultPartitionGridSize = 25600; // 256m
 
 	/** Size of the PCG partition actor grid for non-hierarchical-generation graphs. Min value: 400, MaxValue: 419430400 */
-	UPROPERTY(config, EditAnywhere, Category = GenerationSettings, meta = (ClampMin=400, ClampMax=419430400))
+	UPROPERTY(config, EditAnywhere, Category = GenerationSettings)
 	uint32 PartitionGridSize;
 
 	/** Contains all the PCG data required to query the landscape complete. Serialized in cooked builds only */
