Fix V1008: Consider inspecting the 'for' operator. No more than one iteration of the loop will be performed.

#rnx

diff --git a/Engine/Plugins/PCG/Source/PCG/Public/Elements/Grammar/PCGSubdivisionBase.h b/Engine/Plugins/PCG/Source/PCG/Public/Elements/Grammar/PCGSubdivisionBase.h
--- a/Engine/Plugins/PCG/Source/PCG/Public/Elements/Grammar/PCGSubdivisionBase.h
+++ b/Engine/Plugins/PCG/Source/PCG/Public/Elements/Grammar/PCGSubdivisionBase.h
@@ -210,17 +210,17 @@ namespace PCGSubdivisionBase
 			}
 
 			int NumConcreteModules = Submodule.GetNumRepeat();
-			int NumRepeatModules = 0;
+			bool bHasRepeatModule = false;
 
 			if (NumConcreteModules == PCGGrammar::InfiniteRepetition)
 			{
 				NumConcreteModules = 0;
-				NumRepeatModules = 1;
+				bHasRepeatModule = true;
 			}
 			else if (NumConcreteModules == PCGGrammar::AtLeastOneRepetition)
 			{
 				NumConcreteModules = 1;
-				NumRepeatModules = 1;
+				bHasRepeatModule = true;
 			}
 
 			for (int C = 0; C < NumConcreteModules; ++C)
@@ -229,7 +229,7 @@ namespace PCGSubdivisionBase
 				ExpandedSubmodule.bSkipExpansion = false;
 			}
 
-			for (int R = 0; R < NumRepeatModules; ++R)
+			if (bHasRepeatModule)
 			{
 				TModuleInstance<T>& ExpandedSubmodule = ExpandedModules.Emplace_GetRef(&Submodule);
 				ExpandedSubmodule.bSkipExpansion = true;
diff --git a/Engine/Source/Runtime/MovieScene/Private/Channels/MovieSceneInterpolation.cpp b/Engine/Source/Runtime/MovieScene/Private/Channels/MovieSceneInterpolation.cpp
--- a/Engine/Source/Runtime/MovieScene/Private/Channels/MovieSceneInterpolation.cpp
+++ b/Engine/Source/Runtime/MovieScene/Private/Channels/MovieSceneInterpolation.cpp
@@ -165,9 +165,9 @@ FInterpolationExtents FCachedInterpolation::ComputeExtents(FFrameTime From, FFra
 		FFrameTime Solutions[1];
 		const int32 NumSolutions = Derivative.SolveWithin(From, To, 0.0, Solutions);
 		check(NumSolutions <= 1);
-		for (int32 Index = 0; Index < NumSolutions; ++Index)
+		if (NumSolutions > 0)
 		{
-			Extents.AddPoint(Quadratic->Evaluate(Solutions[Index]), Solutions[Index]);
+			Extents.AddPoint(Quadratic->Evaluate(Solutions[0]), Solutions[0]);
 		}
 	}
 	else if (const FCubicInterpolation* Cubic = Data.TryGet<FCubicInterpolation>())
