[PCG] Throttle scanning of all grid cells in runtime generation scheduler based on changes in environment (cvars, gen sources, etc).

First pass - has hardcode position delta of 50cm (which could potentially be a percentage of smallest grid size or such), and only deals with partitioned components (where vast majority of cost typically is).

Next step will be better handling of cells pending streaming. Currently pending streaming will result in rescans every frame.

[RN] Added change detection to runtime generation state to throttle scans over grid cells.

diff --git a/Engine/Plugins/PCG/Source/PCG/Internal/RuntimeGen/PCGRuntimeGenChangeDetection.h b/Engine/Plugins/PCG/Source/PCG/Internal/RuntimeGen/PCGRuntimeGenChangeDetection.h
new file mode 100644
--- /dev/null
+++ b/Engine/Plugins/PCG/Source/PCG/Internal/RuntimeGen/PCGRuntimeGenChangeDetection.h
@@ -0,0 +1,55 @@
+// Copyright Epic Games, Inc. All Rights Reserved.
+
+#pragma once
+
+#include "PCGCommon.h"
+#include "PCGComponent.h"
+
+#include "UObject/ObjectKey.h"
+
+class IPCGGenSourceBase;
+class UPCGComponent;
+
+namespace PCGRuntimeGenChangeDetection
+{
+	/** The key used for the state cache. */
+	struct FDetectionInputs : TTuple<TObjectKey<UPCGComponent>, const IPCGGenSourceBase*, /*RadiusMultiplier*/float, /*UseGenSourceDirection*/bool>
+	{
+		FDetectionInputs(const UPCGComponent* InOriginalComponent, const IPCGGenSourceBase* InGenSource, float InRadiusMultiplier, bool bInUseGenSourceDirection)
+			: TTuple<TObjectKey<UPCGComponent>, const IPCGGenSourceBase*, float, bool>(InOriginalComponent, InGenSource, InRadiusMultiplier, bInUseGenSourceDirection)
+		{}
+
+		const UPCGComponent* GetOriginalComponent() const { return Get<0>().ResolveObjectPtr(); }
+		const IPCGGenSourceBase* GetGenSource() const { return Get<1>(); }
+		float GetGenerationRadiusMultiplier() const { return Get<2>(); }
+		bool UseGenSourceDirection() const { return Get<3>(); }
+	};
+
+	/** Used to detect when world state (generation sources, cvars etc) have changed enough to warrant rescanning generation cells. */
+	struct FDetector
+	{
+		/** The value in the state cache. Stores transient values so they can later used to detect change. */
+		struct FCachedState
+		{
+			void UpdateFromInputs(const FDetectionInputs& InInputs);
+
+			bool ShouldScanCells(const FDetectionInputs& InInputs) const;
+
+			TOptional<FVector> QuantizedGenSourcePosition;
+			TOptional<FVector> QuantizedGenSourceDirection;
+			float GenerationRadiusMultiplier = 1.0f;
+		};
+
+		/** Compare current state against cached state to determine if a rescan of cells is required. */
+		bool IsCellScanRequired(const FDetectionInputs& InInputs);
+
+		/** Called after all cells have been scanned to update the cached state. */
+		void OnCellsScanned(const FDetectionInputs& InInputs);
+
+		/** Flush any recorded state, resulting in a rescan of the component on next tick. */
+		void FlushCachedState(const UPCGComponent* InOriginalComponent);
+
+	private:
+		TMap<FDetectionInputs, FCachedState> StateCache;
+	};
+}
diff --git a/Engine/Plugins/PCG/Source/PCG/Private/PCGComponent.cpp b/Engine/Plugins/PCG/Source/PCG/Private/PCGComponent.cpp
--- a/Engine/Plugins/PCG/Source/PCG/Private/PCGComponent.cpp
+++ b/Engine/Plugins/PCG/Source/PCG/Private/PCGComponent.cpp
@@ -22,6 +22,7 @@
 #include "Helpers/PCGActorHelpers.h"
 #include "Helpers/PCGBlueprintHelpers.h"
 #include "Helpers/PCGHelpers.h"
+#include "RuntimeGen/PCGRuntimeGenScheduler.h"
 #include "RuntimeGen/GenSources/PCGGenSourceBase.h"
 #include "RuntimeGen/SchedulingPolicies/PCGSchedulingPolicyBase.h"
 #include "RuntimeGen/SchedulingPolicies/PCGSchedulingPolicyDistanceAndDirection.h"
@@ -86,11 +87,6 @@ namespace PCGComponent
 		false,
 		TEXT("At the end of execution, generated output data is flattened in place if it can be stolen, but it can be disabled if it has unwanted side effects."));
 
-	static TAutoConsoleVariable<float> CVarRuntimeGenerationRadiusMultiplier(
-		TEXT("pcg.RuntimeGeneration.GlobalRadiusMultiplier"),
-		1.0f,
-		TEXT("Global multiplier for generation radius of all runtime gen components."));
-
 	static TAutoConsoleVariable<bool> CVarConvertToPointDataInPostProcessGraph(
 		TEXT("pcg.Advanced.ConvertToPointDataInPostProcessGraph"),
 		false,
@@ -1341,7 +1337,7 @@ void UPCGComponent::SetSchedulingPolicyClass(TSubclassOf<UPCGSchedulingPolicyBas
 
 double UPCGComponent::GetGenerationRadiusFromGrid(EPCGHiGenGrid Grid) const
 {
-	const double Multiplier = FMath::Max(PCGComponent::CVarRuntimeGenerationRadiusMultiplier.GetValueOnAnyThread(), 0.0f);
+	const double Multiplier = FMath::Max(PCGRuntimeGenSchedulerHelpers::CVarRuntimeGenerationRadiusMultiplier.GetValueOnAnyThread(), 0.0f);
 
 	if (bOverrideGenerationRadii)
 	{
@@ -1359,7 +1355,7 @@ double UPCGComponent::GetGenerationRadiusFromGrid(EPCGHiGenGrid Grid) const
 
 double UPCGComponent::GetCleanupRadiusFromGrid(EPCGHiGenGrid Grid) const
 {
-	const double Multiplier = FMath::Max(PCGComponent::CVarRuntimeGenerationRadiusMultiplier.GetValueOnAnyThread(), 0.0f);
+	const double Multiplier = FMath::Max(PCGRuntimeGenSchedulerHelpers::CVarRuntimeGenerationRadiusMultiplier.GetValueOnAnyThread(), 0.0f);
 
 	if (bOverrideGenerationRadii)
 	{
diff --git a/Engine/Plugins/PCG/Source/PCG/Private/RuntimeGen/PCGRuntimeGenChangeDetection.cpp b/Engine/Plugins/PCG/Source/PCG/Private/RuntimeGen/PCGRuntimeGenChangeDetection.cpp
new file mode 100644
--- /dev/null
+++ b/Engine/Plugins/PCG/Source/PCG/Private/RuntimeGen/PCGRuntimeGenChangeDetection.cpp
@@ -0,0 +1,128 @@
+// Copyright Epic Games, Inc. All Rights Reserved.
+
+#include "RuntimeGen/PCGRuntimeGenChangeDetection.h"
+
+#include "PCGComponent.h"
+#include "RuntimeGen/GenSources/PCGGenSourceBase.h"
+
+namespace PCGRuntimeGenChangeDetection
+{
+	static FVector QuantizeVector(const FVector& InVector, float InQuanta)
+	{
+		return FVector(
+			FMath::FloorToFloat(InVector.X / InQuanta) * InQuanta,
+			FMath::FloorToFloat(InVector.Y / InQuanta) * InQuanta,
+			FMath::FloorToFloat(InVector.Z / InQuanta) * InQuanta
+		);
+	}
+
+	const static float DirectionQuantaSize = 0.025f;
+
+	void FDetector::FCachedState::UpdateFromInputs(const FDetectionInputs& InInputs)
+	{
+		QuantizedGenSourcePosition = InInputs.GetGenSource()->GetPosition();
+
+		if (QuantizedGenSourcePosition)
+		{
+			// todo_pcg: possibly proportional to grid size, but need to ensure we generate parents when a child generates.
+			const float QuantaSize = 50.0f;
+
+			QuantizedGenSourcePosition = PCGRuntimeGenChangeDetection::QuantizeVector(*QuantizedGenSourcePosition, QuantaSize);
+		}
+
+		if (InInputs.UseGenSourceDirection())
+		{
+			QuantizedGenSourceDirection = InInputs.GetGenSource()->GetDirection();
+
+			if (QuantizedGenSourceDirection)
+			{
+				// Quantize direction vector. Could be expanded to heading angle or such in the future.
+				QuantizedGenSourceDirection = PCGRuntimeGenChangeDetection::QuantizeVector(*QuantizedGenSourceDirection, DirectionQuantaSize);
+			}
+		}
+		else
+		{
+			QuantizedGenSourceDirection.Reset();
+		}
+
+		GenerationRadiusMultiplier = InInputs.GetGenerationRadiusMultiplier();
+	}
+
+	bool FDetector::FCachedState::ShouldScanCells(const FDetectionInputs& InInputs) const
+	{
+		if (TOptional<FVector> NewPosition = InInputs.GetGenSource()->GetPosition())
+		{
+			if (!QuantizedGenSourcePosition)
+			{
+				// Position has just become available.
+				return true;
+			}
+
+			// Rescan if position has changed significantly.
+			// todo_pcg: Possibly proportional to (smallest?) grid size, but need to ensure we generate parents when a child generates.
+			const float QuantaSize = 50.0f;
+
+			if (PCGRuntimeGenChangeDetection::QuantizeVector(*NewPosition, QuantaSize) != *QuantizedGenSourcePosition)
+			{
+				return true;
+			}
+		}
+
+		if (InInputs.UseGenSourceDirection())
+		{
+			if (TOptional<FVector> NewDirection = InInputs.GetGenSource()->GetDirection())
+			{
+				if (!QuantizedGenSourceDirection)
+				{
+					// Direction has just become available.
+					return true;
+				}
+
+				// Rescan if direction has changed significantly.
+				// todo_pcg: Possibly only take heading angle if 2D grid.
+				if (PCGRuntimeGenChangeDetection::QuantizeVector(*NewDirection, DirectionQuantaSize) != *QuantizedGenSourceDirection)
+				{
+					return true;
+				}
+			}
+			else if (QuantizedGenSourceDirection)
+			{
+				// Gen source direction has stopped coming through. Safest to trigger rescan - for e.g. this might disable frustum culling.
+				return true;
+			}
+		}
+
+		return false;
+	}
+
+	bool FDetector::IsCellScanRequired(const FDetectionInputs& InInputs)
+	{
+		// If no cached state was present it will be added here and ShouldScanCells will return true in this case.
+		return StateCache.FindOrAdd(InInputs).ShouldScanCells(InInputs);
+	}
+
+	void FDetector::OnCellsScanned(const FDetectionInputs& InInputs)
+	{
+		// Add or update the cached state.
+		StateCache.FindOrAdd(InInputs).UpdateFromInputs(InInputs);
+	}
+
+	void FDetector::FlushCachedState(const UPCGComponent* InOriginalComponent)
+	{
+		TArray<FDetectionInputs, TInlineAllocator<32>> KeysToRemove;
+
+		for (const TPair<FDetectionInputs, FCachedState>& Entry : StateCache)
+		{
+			const UPCGComponent* Component = Entry.Key.GetOriginalComponent();
+			if (!Component || Component == InOriginalComponent)
+			{
+				KeysToRemove.Add(Entry.Key);
+			}
+		}
+
+		for (const FDetectionInputs& Key : KeysToRemove)
+		{
+			StateCache.Remove(Key);
+		}
+	}
+}
diff --git a/Engine/Plugins/PCG/Source/PCG/Private/RuntimeGen/PCGRuntimeGenScheduler.cpp b/Engine/Plugins/PCG/Source/PCG/Private/RuntimeGen/PCGRuntimeGenScheduler.cpp
--- a/Engine/Plugins/PCG/Source/PCG/Private/RuntimeGen/PCGRuntimeGenScheduler.cpp
+++ b/Engine/Plugins/PCG/Source/PCG/Private/RuntimeGen/PCGRuntimeGenScheduler.cpp
@@ -5,7 +5,6 @@
 
 #include "PCGActorAndComponentMapping.h"
 #include "PCGCommon.h"
-#include "PCGComponent.h"
 #include "PCGGraph.h"
 #include "PCGModule.h"
 #include "PCGSubsystem.h"
@@ -35,9 +34,9 @@
 #include "EditorViewportClient.h"
 #endif
 
-// Onscreen debug messages not supported in UE in shipping & test builds.
-
 #include UE_INLINE_GENERATED_CPP_BY_NAME(PCGRuntimeGenScheduler)
+
+// Onscreen debug messages not supported in UE in shipping & test builds.
 #define PCG_RGS_ONSCREENDEBUGMESSAGES (!UE_BUILD_SHIPPING && !UE_BUILD_TEST)
 
 namespace PCGRuntimeGenSchedulerConstants
@@ -53,11 +52,21 @@ namespace PCGRuntimeGenSchedulerHelpers
 		true,
 		TEXT("Enable the RuntimeGeneration system."));
 
+	static TAutoConsoleVariable<bool> CVarRuntimeGenerationEnableChangeDetection(
+		TEXT("pcg.RuntimeGeneration.EnableChangeDetection"),
+		false,
+		TEXT("Skips execution of scheduling loop if generation sources and other world state have not changed."));
+
 	static TAutoConsoleVariable<int32> CVarNumGeneratingComponentsAtSameTime(
 		TEXT("pcg.RuntimeGeneration.NumGeneratingComponents"),
 		16,
 		TEXT("Defines the maximum number of runtime components that can generate at the same time."));
 
+	TAutoConsoleVariable<float> CVarRuntimeGenerationRadiusMultiplier(
+		TEXT("pcg.RuntimeGeneration.GlobalRadiusMultiplier"),
+		1.0f,
+		TEXT("Global multiplier for generation radius of all runtime gen components."));
+
 	static TAutoConsoleVariable<bool> CVarRuntimeGenerationEnableDebugging(
 		TEXT("pcg.RuntimeGeneration.EnableDebugging"),
 		false,
@@ -146,12 +155,14 @@ namespace PCGRuntimeGenSchedulerHelpers
 						TEXT("    Tick time: %.3fms\n")
 						TEXT("        Generate Calls: %d\n")
 						TEXT("        Cleanup Calls: %d\n")
+						TEXT("        Grid Cell Scans: %d\n")
 						TEXT("    Num Generating Components: %d / %d\n")
 						TEXT("    PA Pool: %d / %d available\n")
 						TEXT("    VT Preloads: %d\n"),
 						ElapsedMs,
 						GenerateCallCounter,
 						CleanupCallCounter,
+						GridCellScanCounter,
 						NumGeneratingComponents, FMath::Max(1, CVarNumGeneratingComponentsAtSameTime.GetValueOnGameThread()),
 						InRuntimeGenScheduler.PartitionActorPool.Num(), InRuntimeGenScheduler.PartitionActorPoolSize,
 						VTPreloadCounter)
@@ -161,6 +172,7 @@ namespace PCGRuntimeGenSchedulerHelpers
 
 		int GenerateCallCounter = 0;
 		int CleanupCallCounter = 0;
+		int GridCellScanCounter = 0;
 		int NumGeneratingComponents = 0;
 		int VTPreloadCounter = 0;
 
@@ -441,6 +453,7 @@ void FPCGRuntimeGenScheduler::TickQueueComponentsForGeneration(
 		return bIsLoaded;
 	};
 
+	const bool bDoChangeDetection = PCGRuntimeGenSchedulerHelpers::CVarRuntimeGenerationEnableChangeDetection.GetValueOnGameThread();
 #if UE_ENABLE_DEBUG_DRAWING
 	const bool bDebugDrawGenerationSources = PCGSystemSwitches::CVarPCGDebugDrawGeneratedCells.GetValueOnGameThread();
 #endif
@@ -457,7 +470,8 @@ void FPCGRuntimeGenScheduler::TickQueueComponentsForGeneration(
 
 		const UPCGSchedulingPolicyBase* Policy = OriginalComponent->GetRuntimeGenSchedulingPolicy();
 
-		// TODO: For each execution domain (for now only GenAtRuntime/dynamic), assuming we run Preview through this scheduler, which it seems like we will.
+		// todo_pcg: For each execution domain (for now only GenAtRuntime/dynamic), assuming we run Preview through this scheduler, which it seems like we will.
+		// todo_pcg: All this stuff can be hoisted - whether valid, has graph, active, managed by runtime gen..
 		if (OriginalComponent->IsManagedByRuntimeGenSystem())
 		{
 			ensure(Policy);
@@ -508,6 +522,23 @@ void FPCGRuntimeGenScheduler::TickQueueComponentsForGeneration(
 
 			for (const IPCGGenSourceBase* GenSource : *Inputs.GenSources)
 			{
+				const PCGRuntimeGenChangeDetection::FDetectionInputs ChangeDetectionInputs =
+				{
+					OriginalComponent,
+					GenSource,
+					PCGRuntimeGenSchedulerHelpers::CVarRuntimeGenerationRadiusMultiplier.GetValueOnGameThread(),
+					Policy->CullsBasedOnDirection()
+				};
+
+				if (bDoChangeDetection && !ChangeDetector.IsCellScanRequired(ChangeDetectionInputs))
+				{
+					continue;
+				}
+
+#if PCG_RGS_ONSCREENDEBUGMESSAGES
+				++PCGRuntimeGenSchedulerHelpers::Stats.GridCellScanCounter;
+#endif
+
 				const TOptional<FVector> GenSourcePositionOptional = GenSource->GetPosition();
 				if (!GenSourcePositionOptional.IsSet())
 				{
@@ -548,9 +579,14 @@ void FPCGRuntimeGenScheduler::TickQueueComponentsForGeneration(
 
 				const double DistanceSquared = OriginalComponentBounds.ComputeSquaredDistanceToPoint(ModifiedGenSourcePosition);
 
+				// If GenSource is not within range of the component then skip it.
 				if (DistanceSquared > MaxGenerationRadius * MaxGenerationRadius)
 				{
-					// GenSource is not within range of the component, skip!
+					if (bDoChangeDetection)
+					{
+						ChangeDetector.OnCellsScanned(ChangeDetectionInputs);
+					}
+
 					continue;
 				}
 
@@ -568,6 +604,8 @@ void FPCGRuntimeGenScheduler::TickQueueComponentsForGeneration(
 					}
 				}
 
+				bool bAnyCellGenerationBlockedByStreaming = false;
+
 				// TODO: once one of the larger grid sizes is out of range, we can forego checking any smaller grid sizes. they can't possibly be closer!
 				// This assumes we enforce generation radii to increase monotonically.
 				for (int GridIndex = 0; GridIndex < GridSizes.Num(); ++GridIndex)
@@ -658,11 +696,23 @@ void FPCGRuntimeGenScheduler::TickQueueComponentsForGeneration(
 									{
 										AddComponentToGenerate(Key, GenSource, Policy, IntersectedBounds, Key.GetGridDescriptor().Is2DGrid());
 									}
+									else
+									{
+										bAnyCellGenerationBlockedByStreaming = true;
+									}
 								}
 							}
 						}
 					}
 				}
+
+				// Require repeated scanning if we are waiting for level to stream.
+				// todo_pcg: We can definitely do better than this. We should probably have a separate list of keys for components that are awaiting streaming so
+				// that we can keep track of these specifically rather than requiring a full rescan.
+				if (bDoChangeDetection && !bAnyCellGenerationBlockedByStreaming)
+				{
+					ChangeDetector.OnCellsScanned(ChangeDetectionInputs);
+				}
 			}
 		}
 	}
@@ -677,6 +727,8 @@ void FPCGRuntimeGenScheduler::TickQueueComponentsForGeneration(
 			continue;
 		}
 
+		// todo_pcg: Worth adding change detection here? Decide on quantization - perhaps a proportion of minimum volume extent or such?
+		
 		// The generation key for a non-partitioned component should always have unbounded grid size and 0,0,0 cell coord.
 		if (Inputs.GeneratedComponents.Contains({ PCGHiGenGrid::UnboundedGridSize(), FIntVector(0), OriginalComponent }))
 		{
@@ -1617,6 +1669,9 @@ void FPCGRuntimeGenScheduler::RefreshComponent(UPCGComponent* InComponent, bool
 		return;
 	}
 
+	// Trigger a rescan of generation cells for this component.
+	ChangeDetector.FlushCachedState(OriginalComponent);
+
 	const bool bLoggingEnabled = PCGRuntimeGenSchedulerHelpers::CVarRuntimeGenerationEnableDebugging.GetValueOnAnyThread();
 
 	// Useful because we can run into generation order issues if components are left to continue generating.
diff --git a/Engine/Plugins/PCG/Source/PCG/Public/RuntimeGen/PCGRuntimeGenScheduler.h b/Engine/Plugins/PCG/Source/PCG/Public/RuntimeGen/PCGRuntimeGenScheduler.h
--- a/Engine/Plugins/PCG/Source/PCG/Public/RuntimeGen/PCGRuntimeGenScheduler.h
+++ b/Engine/Plugins/PCG/Source/PCG/Public/RuntimeGen/PCGRuntimeGenScheduler.h
@@ -5,6 +5,8 @@
 #include "PCGCommon.h"
 #include "PCGComponent.h"
 #include "Grid/PCGGridDescriptor.h"
+#include "RuntimeGen/PCGRuntimeGenChangeDetection.h"
+
 #include "UObject/WeakObjectPtr.h"
 
 #include "PCGRuntimeGenScheduler.generated.h"
@@ -23,6 +25,11 @@ enum class ELevelStreamingState : uint8;
 enum class ELevelStreamingTargetState : uint8;
 namespace PCGRuntimeGenSchedulerHelpers { struct FStatsOverlay; }
 
+namespace PCGRuntimeGenSchedulerHelpers
+{
+	extern PCG_API TAutoConsoleVariable<float> CVarRuntimeGenerationRadiusMultiplier;
+}
+
 /** Used to inform what virtual textures to prime and on what grids they need to be present. */
 USTRUCT(BlueprintType)
 struct FPCGVirtualTexturePrimingInfo
@@ -241,5 +248,8 @@ private:
 	using FStreamingCompleteQueryKey = TPair<FVector, float>;
 	TMap<FStreamingCompleteQueryKey, bool> CachedStreamingQueryResults;
 
+	/** Used to detect when world state (generation sources, cvars etc) have changed enough to warrant rescanning generation cells. */
+	PCGRuntimeGenChangeDetection::FDetector ChangeDetector;
+
 	friend struct PCGRuntimeGenSchedulerHelpers::FStatsOverlay;
 };
diff --git a/Engine/Plugins/PCG/Source/PCG/Public/RuntimeGen/SchedulingPolicies/PCGSchedulingPolicyBase.h b/Engine/Plugins/PCG/Source/PCG/Public/RuntimeGen/SchedulingPolicies/PCGSchedulingPolicyBase.h
--- a/Engine/Plugins/PCG/Source/PCG/Public/RuntimeGen/SchedulingPolicies/PCGSchedulingPolicyBase.h
+++ b/Engine/Plugins/PCG/Source/PCG/Public/RuntimeGen/SchedulingPolicies/PCGSchedulingPolicyBase.h
@@ -33,6 +33,9 @@ public:
 	/** True if the generation source would cull the given bounds. Only applies to bounds within the cleanup generation radius. */
 	virtual bool ShouldCull(const IPCGGenSourceBase* InGenSource, const FBox& GenerationBounds, bool bUse2DGrid) const { return false; }
 
+	/** Should return true if this policy select cells for generation based on generation source direction (for e.g. frustum culling), and false otherwise. */
+	virtual bool CullsBasedOnDirection() const { return true; }
+
 	/** A SchedulingPolicy is equivalent to another SchedulingPolicy if they are the same (same ptr), or if they have the same type and parameter values. */
 	virtual bool IsEquivalent(const UPCGSchedulingPolicyBase* OtherSchedulingPolicy) const PURE_VIRTUAL(UPCGSchedulingPolicyBase::IsEquivalent, return false;);
 
diff --git a/Engine/Plugins/PCG/Source/PCG/Public/RuntimeGen/SchedulingPolicies/PCGSchedulingPolicyDistanceAndDirection.h b/Engine/Plugins/PCG/Source/PCG/Public/RuntimeGen/SchedulingPolicies/PCGSchedulingPolicyDistanceAndDirection.h
--- a/Engine/Plugins/PCG/Source/PCG/Public/RuntimeGen/SchedulingPolicies/PCGSchedulingPolicyDistanceAndDirection.h
+++ b/Engine/Plugins/PCG/Source/PCG/Public/RuntimeGen/SchedulingPolicies/PCGSchedulingPolicyDistanceAndDirection.h
@@ -20,17 +20,13 @@ class UPCGSchedulingPolicyDistanceAndDirection : public UPCGSchedulingPolicyBase
 	GENERATED_BODY()
 
 public:
-	/** Calculate the runtime scheduling priority with respect to a Generation Source. Should return a value in the range [0, 1], where higher values will be scheduled sooner. */
+	//~ Begin UPCGSchedulingPolicyBase interface
 	virtual double CalculatePriority(const IPCGGenSourceBase* InGenSource, const FBox& GenerationBounds, bool bUse2DGrid) const override;
-
-	/** True if the generation source would consider the given bounds for generation. */
 	virtual bool ShouldGenerate(const IPCGGenSourceBase* InGenSource, const FBox& GenerationBounds, bool bUse2DGrid) const override;
-
-	/** True if the generation source would cull the given bounds. Only applies to bounds within the cleanup generation radius. */
 	virtual bool ShouldCull(const IPCGGenSourceBase* InGenSource, const FBox& GenerationBounds, bool bUse2DGrid) const override;
-
-	/** A SchedulingPolicy is equivalent to another SchedulingPolicy if they are the same (same ptr), or if they have the same type and parameter values. */
-	virtual bool IsEquivalent(const UPCGSchedulingPolicyBase* OtherSchedulingPolicy) const;
+	virtual bool CullsBasedOnDirection() const override { return bUseFrustumCulling; }
+	virtual bool IsEquivalent(const UPCGSchedulingPolicyBase* OtherSchedulingPolicy) const override;
+	//~ End UPCGSchedulingPolicyBase interface
 
 public:
 	/** Toggle whether or not distance is used to calculate the scheduling priority. */
