[PCG] Fix initialized pin parsing looking at commented code to find point property setters

#rnx

diff --git a/Engine/Plugins/PCG/Source/PCG/Private/Compute/Elements/PCGCustomHLSLKernel.cpp b/Engine/Plugins/PCG/Source/PCG/Private/Compute/Elements/PCGCustomHLSLKernel.cpp
--- a/Engine/Plugins/PCG/Source/PCG/Private/Compute/Elements/PCGCustomHLSLKernel.cpp
+++ b/Engine/Plugins/PCG/Source/PCG/Private/Compute/Elements/PCGCustomHLSLKernel.cpp
@@ -1008,6 +1008,8 @@ void UPCGCustomHLSLKernel::ParseShaderSource()
 					{
 						ensure(false);
 					}
+
+					InOutParsedSource.InitializedOutputPins.Add(PinStr);
 				}
 			}
 			else if (Token.Type == FPCGCustomHLSLParsedSource::ETokenType::Normal)
@@ -1053,23 +1055,6 @@ void UPCGCustomHLSLKernel::ParseShaderSource()
 
 		// @todo_pcg: Maybe this should also be parsed instead of regex'd, but it's not as trivial to create tokens to detect the data labels pattern.
 		CollectDataLabels(InOutParsedSource);
-
-		// Capture any pin inititialization that won't have been caught by parsing.
-		// @todo_pcg: We should just add the point property setters as keywords to the parsing.
-		for (const FPCGPinPropertiesGPU& Pin : CustomHLSLSettings->OutputPins)
-		{
-			if (Pin.Label == NAME_None)
-			{
-				continue;
-			}
-
-			const FString PinStr = Pin.Label.ToString();
-
-			if (InOutParsedSource.Source.Contains(PinStr + TEXT("_Set")))
-			{
-				InOutParsedSource.InitializedOutputPins.Add(PinStr);
-			}
-		}
 	};
 
 	for (FPCGCustomHLSLParsedSource& ParsedSource : ParsedSources)
