[PCG] GraphExecutor: Fix bIsPaused race condition

diff --git a/Engine/Plugins/PCG/Source/PCG/Private/Graph/PCGGraphExecutor.cpp b/Engine/Plugins/PCG/Source/PCG/Private/Graph/PCGGraphExecutor.cpp
--- a/Engine/Plugins/PCG/Source/PCG/Private/Graph/PCGGraphExecutor.cpp
+++ b/Engine/Plugins/PCG/Source/PCG/Private/Graph/PCGGraphExecutor.cpp
@@ -1135,10 +1135,17 @@ void FPCGGraphExecutor::PostTaskExecute(TSharedPtr<FPCGGraphActiveTask> ActiveTa
 			check(TaskRemoveIndex != INDEX_NONE && ActiveTaskArrayToRemove[TaskRemoveIndex] == ActiveTaskPtr);
 
 			const bool bHasPausedDependencies = !ActiveTask.Context->DynamicDependencies.IsEmpty();
-			if (ActiveTask.Context->bIsPaused || bHasPausedDependencies)
+
+			// We have dynamic dependencies bIsPaused flag might not have been set by the element
+			// It is safe to set bIsPaused here because tasks with DynamicDependencies can't be unpaused by another task (race-condition)
+			if (bHasPausedDependencies)
 			{
 				ActiveTask.Context->bIsPaused = true;
-								
+			}
+
+			// Still paused (because of multithreading it is possible that we were unpaused before processing the paused state)
+			if (ActiveTask.Context->bIsPaused)
+			{
 				bool bMoveFromActiveToPaused = !bHasPausedDependencies;
 
 				// When using DynamicDependencies the task will be un-paused  (bIsPaused will toggle back to false) after completion of those dependency tasks.
@@ -1174,7 +1181,7 @@ void FPCGGraphExecutor::PostTaskExecute(TSharedPtr<FPCGGraphActiveTask> ActiveTa
 				{
 					// All the paused dependencies have completed
 					ActiveTask.Context->bIsPaused = false;
-				}							
+				}
 			}
 			else if (ActiveTask.bIsGameThreadOnly != ActiveTask.CanExecuteOnlyOnMainThread())
 			{
