[PCG] Added node to split splines based on a few criteria (key, distance, alpha and a predicate on control points).

diff --git a/Engine/Plugins/PCG/Source/PCG/Private/Elements/PCGSplitSplines.cpp b/Engine/Plugins/PCG/Source/PCG/Private/Elements/PCGSplitSplines.cpp
new file mode 100644
--- /dev/null
+++ b/Engine/Plugins/PCG/Source/PCG/Private/Elements/PCGSplitSplines.cpp
@@ -0,0 +1,407 @@
+// Copyright Epic Games, Inc. All Rights Reserved.
+
+#include "Elements/PCGSplitSplines.h"
+
+#include "PCGContext.h"
+#include "Data/PCGSplineData.h"
+#include "Elements/PCGSplineIntersection.h"
+#include "Elements/Metadata/PCGMetadataElementCommon.h"
+#include "Helpers/PCGHelpers.h"
+#include "Metadata/Accessors/PCGAttributeAccessorHelpers.h"
+
+#include UE_INLINE_GENERATED_CPP_BY_NAME(PCGSplitSplines)
+
+#define LOCTEXT_NAMESPACE "PCGSplitSplinesElement"
+
+namespace PCGSplitSplineConstants
+{
+	const FName SplitInfoLabel = TEXT("SplitInfo");
+}
+
+UPCGSplitSplinesSettings::UPCGSplitSplinesSettings()
+{
+	OutputOriginatingSplineIndex = FPCGAttributePropertySelector::CreateAttributeSelector<FPCGAttributePropertyOutputSelector>(TEXT("OriginatingSplineIndex"), TEXT("Data"));
+}
+
+#if WITH_EDITOR
+FName UPCGSplitSplinesSettings::GetDefaultNodeName() const
+{
+	return FName(TEXT("SplitSpline"));
+}
+
+FText UPCGSplitSplinesSettings::GetDefaultNodeTitle() const
+{
+	return LOCTEXT("NodeTitle", "Split Spline");
+}
+
+EPCGChangeType UPCGSplitSplinesSettings::GetChangeTypeForProperty(const FName& InPropertyName) const
+{
+	EPCGChangeType ChangeType = Super::GetChangeTypeForProperty(InPropertyName);
+
+	if (InPropertyName == GET_MEMBER_NAME_CHECKED(UPCGSplitSplinesSettings, Mode) ||
+		InPropertyName == GET_MEMBER_NAME_CHECKED(UPCGSplitSplinesSettings, bUseConstant))
+	{
+		// This can change the output pin types, so we need a structural change
+		ChangeType |= EPCGChangeType::Structural;
+	}
+
+	return ChangeType;
+}
+#endif // WITH_EDITOR
+
+TArray<FPCGPinProperties> UPCGSplitSplinesSettings::InputPinProperties() const
+{
+	TArray<FPCGPinProperties> PinProperties;
+	PinProperties.Emplace_GetRef(PCGPinConstants::DefaultInputLabel, EPCGDataType::Spline).SetRequiredPin();
+
+	if(!bUseConstant && Mode != EPCGSplitSplineMode::ByPredicateOnControlPoints)
+	{
+		PinProperties.Emplace_GetRef(PCGSplitSplineConstants::SplitInfoLabel, EPCGDataType::Param).SetRequiredPin();
+	}
+
+	return PinProperties;
+}
+
+TArray<FPCGPinProperties> UPCGSplitSplinesSettings::OutputPinProperties() const
+{
+	TArray<FPCGPinProperties> PinProperties;
+	PinProperties.Emplace(PCGPinConstants::DefaultOutputLabel, EPCGDataType::Spline);
+
+	return PinProperties;
+}
+
+FPCGElementPtr UPCGSplitSplinesSettings::CreateElement() const
+{
+	return MakeShared<FPCGSplitSplineElement>();
+}
+
+bool FPCGSplitSplineElement::ExecuteInternal(FPCGContext* Context) const
+{
+	TRACE_CPUPROFILER_EVENT_SCOPE(FPCGSplitSplineElement::Execute);
+
+	const UPCGSplitSplinesSettings* Settings = Context->GetInputSettings<UPCGSplitSplinesSettings>();
+	check(Settings);
+
+	TArray<FPCGTaggedData> Inputs = Context->InputData.GetInputsByPin(PCGPinConstants::DefaultInputLabel);
+	TArray<FPCGTaggedData> SplitInfo = Context->InputData.GetInputsByPin(PCGSplitSplineConstants::SplitInfoLabel);
+	const EPCGSplitSplineMode Mode = Settings->Mode;
+
+	TArray<FPCGTaggedData>& Outputs = Context->OutputData.TaggedData;
+
+	// Validate inputs:
+	// If we are in the ByKey / ByDistance, we need to validate cardinality here
+	if(!Settings->bUseConstant && Mode != EPCGSplitSplineMode::ByPredicateOnControlPoints)
+	{
+		if (SplitInfo.Num() != 1 && SplitInfo.Num() != Inputs.Num())
+		{
+			PCGLog::InputOutput::LogInvalidCardinalityError(PCGPinConstants::DefaultInputLabel, PCGSplitSplineConstants::SplitInfoLabel, Context);
+			return true;
+		}
+	}
+
+	// Goal: Cut the spline at specific key locations where these depend on the input data (fixed key; fixed distance; specified keys, specified distances, boolean predicate on control points)
+	for (int InputIndex = 0; InputIndex < Inputs.Num(); ++InputIndex)
+	{
+		const FPCGTaggedData& Input = Inputs[InputIndex];
+		const UPCGSplineData* SplineData = Cast<UPCGSplineData>(Input.Data);
+
+		// Input data isn't supported, just discard the original input
+		if (!SplineData)
+		{
+			continue;
+		}
+
+		TArray<FSplinePoint> SplinePoints = SplineData->GetSplinePoints();
+		TArray<int64> SplinePointsEntryKeys = SplineData->GetMetadataEntryKeysForSplinePoints();
+		check(SplinePoints.Num() == SplinePointsEntryKeys.Num() || SplinePointsEntryKeys.IsEmpty());
+
+		// Build keys list
+		TArray<double> SplitKeys;
+		if (Mode == EPCGSplitSplineMode::ByPredicateOnControlPoints)
+		{
+			TArray<bool> Predicates;
+			if (!PCGAttributeAccessorHelpers::ExtractAllValues(SplineData, Settings->Attribute, Predicates, Context))
+			{
+				Outputs.Add(Input);
+				continue;
+			}
+
+			if (Predicates.Num() != SplinePoints.Num())
+			{
+				FPCGAttributePropertyInputSelector ResolvedAttribute = Settings->Attribute.CopyAndFixLast(SplineData);
+
+				// Mismatch in predicate vs. the expected number of values
+				PCGLog::LogErrorOnGraph(FText::Format(LOCTEXT("MismatchInPredicateCardinality", "Property {0} is not a proper predicate for splitting the spline."), FText::FromName(ResolvedAttribute.GetAttributeName())), Context);
+				Outputs.Add(Input);
+				continue;
+			}
+
+			// @todo_pcg: Spline implementation assumption, i.e. that control points are on integer keys.
+			for (int ControlPointIndex = 0; ControlPointIndex < Predicates.Num(); ++ControlPointIndex)
+			{
+				if (Predicates[ControlPointIndex])
+				{
+					SplitKeys.Add(static_cast<double>(ControlPointIndex));
+				}
+			}
+		}
+		else
+		{
+			TArray<double> Values;
+			if (Settings->bUseConstant)
+			{
+				Values.Add(Settings->Constant);
+			}
+			else
+			{
+				int SplitInfoIndex = InputIndex % SplitInfo.Num();
+				const UPCGData* SplitInfoData = SplitInfo[SplitInfoIndex].Data;
+
+				if (!PCGAttributeAccessorHelpers::ExtractAllValues(SplitInfoData, Settings->Attribute, Values, Context))
+				{
+					Outputs.Add(Input);
+					continue;
+				}
+			}
+
+			if (Mode == EPCGSplitSplineMode::ByKey)
+			{
+				SplitKeys = MoveTemp(Values);
+			}
+			else if (Mode == EPCGSplitSplineMode::ByDistance)
+			{
+				Algo::Transform(Values, SplitKeys, [SplineData](const double& InValue) { return SplineData->SplineStruct.GetInputKeyAtDistanceAlongSpline(InValue); });
+			}
+			else if (Mode == EPCGSplitSplineMode::ByAlpha)
+			{
+				Algo::Transform(Values, SplitKeys, [SplineData](const double& InValue) { return SplineData->GetInputKeyAtAlpha(InValue); });
+			}
+			else
+			{
+				check(0); // Invalid mode
+			}
+		}
+
+		const bool bIsClosed = SplineData->IsClosed();
+
+		// Remove values that are outside of the spline
+		if (bIsClosed)
+		{
+			// Closed loops should properly cut at key = 0 and up to the last key + 1.
+			for (int SplitIndex = SplitKeys.Num() - 1; SplitIndex >= 0; --SplitIndex)
+			{
+				if (SplitKeys[SplitIndex] < SplinePoints[0].InputKey || SplitKeys[SplitIndex] > 1.0 + SplinePoints.Last().InputKey)
+				{
+					SplitKeys.RemoveAtSwap(SplitIndex);
+				}
+			}
+		}
+		else
+		{
+			for (int SplitIndex = SplitKeys.Num() - 1; SplitIndex >= 0; --SplitIndex)
+			{
+				if (SplitKeys[SplitIndex] <= SplinePoints[0].InputKey || SplitKeys[SplitIndex] >= SplinePoints.Last().InputKey)
+				{
+					SplitKeys.RemoveAtSwap(SplitIndex);
+				}
+			}
+		}
+
+		// If there are no valid values, then there is nothing to do.
+		if (SplitKeys.IsEmpty())
+		{
+			Outputs.Add(Input);
+			continue;
+		}
+
+		// Sort values
+		SplitKeys.Sort([](const double& A, const double& B) { return A < B; });
+
+		const double KeyEpsilon = 1.0e-6f;
+
+		// Remove duplicates/values too close to each other
+		for (int SplitIndex = SplitKeys.Num() - 1; SplitIndex > 0; --SplitIndex)
+		{
+			if ((SplitKeys[SplitIndex] - SplitKeys[SplitIndex - 1] < KeyEpsilon) ||
+				(bIsClosed && SplitIndex == SplitKeys.Num() - 1 && (SplitKeys[0] + 1.0 + SplinePoints.Last().InputKey) - SplitKeys[SplitIndex] < KeyEpsilon))
+			{
+				SplitKeys.RemoveAt(SplitIndex);
+			}
+		}
+
+		// Once we've removed split keys that are outside their normal range, we'll add keys so we can process segments by pairs of segment keys.
+		// For closed splines: we'll add a duplicate of the first key to the end of the range to do the full loop.
+		// Note that it will mean that the segments will always start from a split point. [split key -> ... -> split key]
+		// For open splines: we'll add a starting key (assumedly 0 in most cases) and a finishing key so we have proper segments [0 -> split key] ... [split key -> end]
+		if (bIsClosed)
+		{
+			const double FirstKey = SplitKeys[0];
+			SplitKeys.Add(FirstKey);
+		}
+		else
+		{
+			SplitKeys.Insert(SplinePoints[0].InputKey, 0);
+			SplitKeys.Add(SplinePoints.Last().InputKey);
+		}
+
+		for (int CurrentSplitPair = 0; CurrentSplitPair < SplitKeys.Num() - 1; ++CurrentSplitPair)
+		{
+			const double StartKey = SplitKeys[CurrentSplitPair];
+			const double EndKey = SplitKeys[CurrentSplitPair+1];
+
+			UPCGSplineData* NewSplineData = FPCGContext::NewObject_AnyThread<UPCGSplineData>(Context);
+			NewSplineData->InitializeFromData(SplineData);
+
+			TArray<FSplinePoint> NewSplinePoints;
+			TArray<int64> NewSplinePointsEntryKeys;
+
+			auto AddExistingControlPoint = [&NewSplinePoints, &NewSplinePointsEntryKeys, &SplinePoints, &SplinePointsEntryKeys](const int ControlPointIndex)
+			{
+				NewSplinePoints.Add(SplinePoints[ControlPointIndex]);
+				if (!SplinePointsEntryKeys.IsEmpty())
+				{
+					NewSplinePointsEntryKeys.Add(SplinePointsEntryKeys[ControlPointIndex]);
+				}
+			};
+
+			auto AddSplitPoint = [&SplineData, &NewSplineData, &NewSplinePoints, &NewSplinePointsEntryKeys, &SplinePoints, &SplinePointsEntryKeys](const double Key)
+			{
+					FSplinePoint SplitPoint;
+					SplitPoint.InputKey = Key;
+					SplitPoint.Position = SplineData->SplineStruct.GetLocationAtSplineInputKey(Key, ESplineCoordinateSpace::Local);
+					SplitPoint.ArriveTangent = SplineData->SplineStruct.GetTangentAtSplineInputKey(Key, ESplineCoordinateSpace::Local);
+					SplitPoint.LeaveTangent = SplitPoint.ArriveTangent;
+					SplitPoint.Rotation = SplineData->SplineStruct.GetQuaternionAtSplineInputKey(Key, ESplineCoordinateSpace::Local).Rotator();
+					SplitPoint.Scale = SplineData->SplineStruct.GetScaleAtSplineInputKey(Key);
+					SplitPoint.Type = ESplinePointType::CurveCustomTangent;
+
+					// The tangent is always computed from the original spline's perspective so we will assume the key range is always one.
+					SplitPoint.ArriveTangent *= FMath::Frac(Key);
+					SplitPoint.LeaveTangent *= 1.0 - FMath::Frac(Key);
+
+					NewSplinePoints.Add(SplitPoint);
+
+					// Add metadata if needed
+					if (!SplinePointsEntryKeys.IsEmpty())
+					{
+						PCGMetadataEntryKey EntryKey = PCGInvalidEntryKey;
+						SplineData->WriteMetadataToEntry(Key, EntryKey, NewSplineData->MutableMetadata());
+					}
+			};
+
+			int FirstControlPointInSegment = FMath::CeilToInt(StartKey);
+			int LastControlPointInSegment = FMath::FloorToInt(EndKey);
+
+			if (bIsClosed && EndKey <= StartKey) // allow for loop
+			{
+				LastControlPointInSegment += SplinePoints.Num();
+			}
+
+			bool bNeedsStartTangentAdjustment = false;
+			bool bNeedsEndTangentAdjustment = false;
+
+			// Add left split point, prioritize existing control points when possible
+			const double FracStartKey = FMath::Frac(StartKey);
+			if (FracStartKey < KeyEpsilon)
+			{
+				if (FracStartKey > 0) // otherwise point is already added in the control point loop
+				{
+					AddExistingControlPoint(FMath::FloorToInt(StartKey));
+				}
+			}
+			else if (1.0 - FracStartKey < KeyEpsilon)
+			{
+				// will be added by the control point loop
+			}
+			else
+			{
+				AddSplitPoint(StartKey);
+				bNeedsStartTangentAdjustment = true;
+			}
+			
+			// Add control points that are taken as-is from the original spline points
+			for (int ControlPointIndex = FirstControlPointInSegment; ControlPointIndex <= LastControlPointInSegment; ++ControlPointIndex)
+			{
+				AddExistingControlPoint((ControlPointIndex % SplinePoints.Num()));
+			}
+
+			// Add right split point
+			const double FracEndKey = FMath::Frac(EndKey);
+			if (FracEndKey < KeyEpsilon)
+			{
+				// will be already added by the control point loop
+			}
+			else if (1.0 - FracEndKey < KeyEpsilon)
+			{
+				AddExistingControlPoint(FMath::CeilToInt(EndKey));
+			}
+			else
+			{
+				AddSplitPoint(EndKey);
+				bNeedsEndTangentAdjustment = true;
+			}
+
+			// Adjust tangents if required
+			if (bNeedsStartTangentAdjustment)
+			{
+				// assumes key range is 1.
+				const double Ratio = (1.0 - FMath::Frac(NewSplinePoints[0].InputKey));
+				NewSplinePoints[1].ArriveTangent *= Ratio;
+			}
+
+			if (bNeedsEndTangentAdjustment)
+			{
+				// assumes key range is 1.
+				const double Ratio = FMath::Frac(NewSplinePoints.Last().InputKey);
+				NewSplinePoints[NewSplinePoints.Num() - 2].LeaveTangent *= Ratio;
+			}
+
+			// Finally, mark all tangents as custom and reset the keys.
+			for (int SPIndex = 0; SPIndex < NewSplinePoints.Num(); ++SPIndex)
+			{
+				FSplinePoint& SplinePoint = NewSplinePoints[SPIndex];
+				SplinePoint.InputKey = static_cast<float>(SPIndex);
+				if (SplinePoint.Type == ESplinePointType::Curve || SplinePoint.Type == ESplinePointType::CurveClamped)
+				{
+					NewSplinePoints[SPIndex].Type = ESplinePointType::CurveCustomTangent;
+				}
+			}
+
+			// Initialize data now
+			NewSplineData->Initialize(NewSplinePoints,
+				/*bIsClosedLoop*/false,
+				SplineData->SplineStruct.GetTransform(),
+				NewSplinePointsEntryKeys);
+
+			// Write originating spline index if required
+			if (Settings->bShouldOutputOriginatingSplineIndex)
+			{
+				const FName IndexAttributeName = Settings->OutputOriginatingSplineIndex.GetName();
+				FPCGMetadataAttribute<int32>* OriginatingIndexAttribute = nullptr;
+
+				if (FPCGMetadataDomain* OutputMetadataDomain = NewSplineData->MutableMetadata()->GetMetadataDomainFromSelector(Settings->OutputOriginatingSplineIndex))
+				{
+					OriginatingIndexAttribute = OutputMetadataDomain->FindOrCreateAttribute(IndexAttributeName, InputIndex, /*bAllowInterpolation=*/false);
+				}
+
+				if (!OriginatingIndexAttribute)
+				{
+					PCGLog::LogErrorOnGraph(FText::Format(LOCTEXT("FailedCreateIndexAttribute", "Failed to create the index attribute '{0}'."), FText::FromName(IndexAttributeName)), Context);
+				}
+			}
+			
+			Outputs.Add_GetRef(Input).Data = NewSplineData;
+		}
+	}
+
+	return true;
+}
+
+EPCGElementExecutionLoopMode FPCGSplitSplineElement::ExecutionLoopMode(const UPCGSettings* InSettings) const
+{
+	const UPCGSplitSplinesSettings* Settings = Cast<UPCGSplitSplinesSettings>(InSettings);
+	// @todo_pcg: Currently we don't have the concept of N:N or N:1 primary pins in the execution loop modes
+	return (Settings && (Settings->bUseConstant || Settings->Mode == EPCGSplitSplineMode::ByPredicateOnControlPoints)) ? EPCGElementExecutionLoopMode::SinglePrimaryPin : EPCGElementExecutionLoopMode::NotALoop;
+}
+
+#undef LOCTEXT_NAMESPACE
\ No newline at end of file
diff --git a/Engine/Plugins/PCG/Source/PCG/Public/Elements/PCGSplitSplines.h b/Engine/Plugins/PCG/Source/PCG/Public/Elements/PCGSplitSplines.h
new file mode 100644
--- /dev/null
+++ b/Engine/Plugins/PCG/Source/PCG/Public/Elements/PCGSplitSplines.h
@@ -0,0 +1,72 @@
+// Copyright Epic Games, Inc. All Rights Reserved.
+
+#pragma once
+
+#include "PCGSettings.h"
+
+#include "PCGSplitSplines.generated.h"
+
+UENUM()
+enum class EPCGSplitSplineMode
+{
+	ByKey,
+	ByDistance,
+	ByAlpha,
+	ByPredicateOnControlPoints,
+};
+
+/** Splits spline at a specific distance(s), key(s) or at certain values. */
+UCLASS(MinimalAPI, BlueprintType, ClassGroup=(Procedural))
+class UPCGSplitSplinesSettings : public UPCGSettings
+{
+	GENERATED_BODY()
+
+public:
+	UPCGSplitSplinesSettings();
+
+	//~Begin UPCGSettings interface
+#if WITH_EDITOR
+	virtual FName GetDefaultNodeName() const override;
+	virtual FText GetDefaultNodeTitle() const override;
+	virtual EPCGChangeType GetChangeTypeForProperty(const FName& InPropertyName) const override;
+	virtual EPCGSettingsType GetType() const override { return EPCGSettingsType::Spatial; }
+#endif
+
+protected:
+	virtual TArray<FPCGPinProperties> InputPinProperties() const override;
+	virtual TArray<FPCGPinProperties> OutputPinProperties() const override;
+	virtual FPCGElementPtr CreateElement() const override;
+	//~End UPCGSettings interface
+
+public:
+	/** Split criteria. */
+	UPROPERTY(BlueprintReadWrite, EditAnywhere, Category = Settings)
+	EPCGSplitSplineMode Mode = EPCGSplitSplineMode::ByAlpha;
+
+	/** Controls whether the input splines will be cut using a single constant or values driven either by an additional input or based on a predicate on the spline control points. */
+	UPROPERTY(BlueprintReadWrite, EditAnywhere, Category = Settings, meta = (EditCondition ="Mode != EPCGSplitSplineMode::ByPredicateOnControlPoints"))
+	bool bUseConstant = true;
+
+	/** Constant (either Key, Distance or Alpha) that will be used to split the splines. */
+	UPROPERTY(BlueprintReadWrite, EditAnywhere, Category = Settings, meta = (EditCondition ="bUseConstant && Mode != EPCGSplitSplineMode::ByPredicateOnControlPoints", ClampMin = "0", PCG_Overridable))
+	double Constant = 0.5;
+
+	/** Attribute identifying either the provenance of the split constants (key, distance or alpha) or the predicate on the control points. */
+	UPROPERTY(BlueprintReadWrite, EditAnywhere, Category = Settings, meta = (EditCondition="!bUseConstant || Mode == EPCGSplitSplineMode::ByPredicateOnControlPoints", EditConditionHides, PCG_Overridable))
+	FPCGAttributePropertyInputSelector Attribute;
+
+	/** Controls whether the output spline will have an attribute containing the index of the originating spline. */
+	UPROPERTY(BlueprintReadWrite, EditAnywhere, Category = Settings, meta = (InlineEditConditionToggle, PCG_Overridable))
+	bool bShouldOutputOriginatingSplineIndex = true;
+
+	/** Attribute to write the originating spline index to. */
+	UPROPERTY(BlueprintReadWrite, EditAnywhere, Category = Settings, meta = (EditCondition="bShouldOutputOriginatingSplineIndex", PCG_DiscardPropertySelection, PCG_DiscardExtraSelection, PCG_Overridable))
+	FPCGAttributePropertyOutputSelector OutputOriginatingSplineIndex;
+};
+
+class FPCGSplitSplineElement : public IPCGElement
+{
+public:
+	virtual bool ExecuteInternal(FPCGContext* InContext) const override;
+	virtual EPCGElementExecutionLoopMode ExecutionLoopMode(const UPCGSettings* Settings) const override;
+};
\ No newline at end of file
