[PCG] GPU: Fix Debug flag on upstream CPU nodes soft locking compute graph element execution.

System was waiting for a non-sensical readback of CPU data.

diff --git a/Engine/Plugins/PCG/Source/PCG/Private/Compute/Elements/PCGComputeGraphElement.cpp b/Engine/Plugins/PCG/Source/PCG/Private/Compute/Elements/PCGComputeGraphElement.cpp
--- a/Engine/Plugins/PCG/Source/PCG/Private/Compute/Elements/PCGComputeGraphElement.cpp
+++ b/Engine/Plugins/PCG/Source/PCG/Private/Compute/Elements/PCGComputeGraphElement.cpp
@@ -258,7 +258,8 @@ bool FPCGComputeGraphElement::ExecuteInternal(FPCGContext* InContext) const
 				continue;
 			}
 
-			const bool bDataProviderDebugged = DataProvider->GetProducerSettings() && DataProvider->GetProducerSettings()->bDebug;
+			const UPCGSettings* ProducerSettings = DataProvider->GetProducerSettings();
+			const bool bDataProviderDebugged = ProducerSettings && ProducerSettings->ShouldExecuteOnGPU() && ProducerSettings->bDebug;
 
 			if (DataProvider->IsExportRequired() || bDataProviderDebugged)
 			{
