[PCG] Integrate data collection upload data interface into the base data interface.

The upload DI had diverged away from its hierarchy which this fixes, and upload should have always been a mode of the DCDI rather than a distinct data interface.

#rnx

diff --git a/Engine/Plugins/PCG/Source/PCG/Private/Compute/DataInterfaces/PCGComputeDataInterface.cpp b/Engine/Plugins/PCG/Source/PCG/Private/Compute/DataInterfaces/PCGComputeDataInterface.cpp
--- a/Engine/Plugins/PCG/Source/PCG/Private/Compute/DataInterfaces/PCGComputeDataInterface.cpp
+++ b/Engine/Plugins/PCG/Source/PCG/Private/Compute/DataInterfaces/PCGComputeDataInterface.cpp
@@ -60,25 +60,28 @@ void UPCGComputeDataInterface::AddReferencedObjects(UObject* InThis, FReferenceC
 
 void UPCGComputeDataProvider::Initialize(const UComputeDataInterface* InDataInterface, UObject* InBinding, uint64 InInputMask, uint64 InOutputMask)
 {
+	Super::Initialize(InDataInterface, InBinding, InInputMask, InOutputMask);
+
 	// Bump counter so any async callbacks from previous usages of this provider will be invalidated.
 	++GenerationCounter;
 
-	Super::Initialize(InDataInterface, InBinding, InInputMask, InOutputMask);
-
 	const UPCGComputeDataInterface* DataInterface = CastChecked<UPCGComputeDataInterface>(InDataInterface);
-	ProducerKernel = DataInterface->ProducerKernel;
 	ResolvedProducerSettings = DataInterface->GetProducerSettings();
 	ProducerSettings = ResolvedProducerSettings;
-
+	bRequiresUpload = DataInterface->bRequiresUpload;
 	GraphBindingIndex = DataInterface->GraphBindingIndex;
+	ProducerKernel = DataInterface->ProducerKernel;
+	DownstreamInputPinLabelAliases = DataInterface->DownstreamInputPinLabelAliases;
 
-	// The original label is needed to store inspection data.
-	OutputPinLabel = DataInterface->OutputPinLabel;
-
-	// Use the aliased label for normal data output as this is the output from the compute graph.
-	OutputPinLabelAlias = DataInterface->OutputPinLabelAlias;
+	// Set data that is only relevant for inter-kernel data interfaces, and compute graph outputs. Don't set them for inputs (CPU->GPU uploads).
+	if (!bRequiresUpload)
+	{
+		// The original label is needed to store inspection data.
+		OutputPinLabel = DataInterface->OutputPinLabel;
 
-	DownstreamInputPinLabelAliases = DataInterface->DownstreamInputPinLabelAliases;
+		// Use the aliased label for CPU data output as this is the output from the compute graph.
+		OutputPinLabelAlias = DataInterface->OutputPinLabelAlias;
+	}
 }
 
 void UPCGComputeDataProvider::Reset()
@@ -95,6 +98,7 @@ void UPCGComputeDataProvider::Reset()
 	OutputPinLabel = NAME_None;
 	OutputPinLabelAlias = NAME_None;
 	DownstreamInputPinLabelAliases.Empty();
+	bRequiresUpload = false;
 }
 
 const UPCGSettings* UPCGComputeDataProvider::GetProducerSettings() const
@@ -156,31 +160,35 @@ void UPCGComputeDataProvider::NotifyProducerUploadedData(UPCGDataBinding* InBind
 void UPCGExportableDataProvider::Initialize(const UComputeDataInterface* InDataInterface, UObject* InBinding, uint64 InInputMask, uint64 InOutputMask)
 {
 	Super::Initialize(InDataInterface, InBinding, InInputMask, InOutputMask);
-	Binding = CastChecked<UPCGDataBinding>(InBinding);
 
 	const UPCGExportableDataInterface* DataInterface = CastChecked<UPCGExportableDataInterface>(InDataInterface);
-	ExportMode = DataInterface->GetRequiresExport() ? EPCGExportMode::ComputeGraphOutput : EPCGExportMode::NoExport;
-
-#if WITH_EDITOR
-	// Some exportable data providers don't support inspect/debug because their producer settings are not available (e.g. GridLinkage).
-	const UPCGSettings* LocalProducerSettings = GetProducerSettings();
-	const bool bCanInspectOrDebug = LocalProducerSettings && (!ProducerKernel || !ProducerKernel->IsPinInternal(OutputPinLabel) || PCGComputeDataInterfaceHelpers::CVarEnableInternalKernelPinInspection.GetValueOnAnyThread());
+	Binding = CastChecked<UPCGDataBinding>(InBinding);
 
-	if (bCanInspectOrDebug)
+	if (!bRequiresUpload)
 	{
-		const IPCGGraphExecutionSource* ExecutionSource = Binding->GetExecutionSource();
+		ExportMode = DataInterface->GetRequiresExport() ? EPCGExportMode::ComputeGraphOutput : EPCGExportMode::NoExport;
 
-		if (LocalProducerSettings->bIsInspecting && ExecutionSource && ExecutionSource->GetExecutionState().GetInspection().IsInspecting())
-		{
-			ExportMode |= EPCGExportMode::Inspection;
-		}
+#if WITH_EDITOR
+		// Some exportable data providers don't support inspect/debug because their producer settings are not available (e.g. GridLinkage).
+		const UPCGSettings* LocalProducerSettings = GetProducerSettings();
+		const bool bCanInspectOrDebug = LocalProducerSettings && (!ProducerKernel || !ProducerKernel->IsPinInternal(OutputPinLabel) || PCGComputeDataInterfaceHelpers::CVarEnableInternalKernelPinInspection.GetValueOnAnyThread());
 
-		if (LocalProducerSettings->bDebug)
+		if (bCanInspectOrDebug)
 		{
-			ExportMode |= EPCGExportMode::DebugVisualization;
+			const IPCGGraphExecutionSource* ExecutionSource = Binding->GetExecutionSource();
+
+			if (LocalProducerSettings->bIsInspecting && ExecutionSource && ExecutionSource->GetExecutionState().GetInspection().IsInspecting())
+			{
+				ExportMode |= EPCGExportMode::Inspection;
+			}
+
+			if (LocalProducerSettings->bDebug)
+			{
+				ExportMode |= EPCGExportMode::DebugVisualization;
+			}
 		}
-	}
 #endif
+	}
 }
 
 void UPCGExportableDataProvider::Reset()
@@ -202,7 +210,10 @@ bool UPCGExportableDataProvider::PrepareForExecute_GameThread(UPCGDataBinding* I
 		return false;
 	}
 
-	PinDesc = InBinding->GetCachedKernelPinDataDesc(GraphBindingIndex);
+	if (!PinDesc)
+	{
+		PinDesc = InBinding->GetCachedKernelPinDataDesc(GraphBindingIndex);
+	}
 
 	return true;
 }
diff --git a/Engine/Plugins/PCG/Source/PCG/Private/Compute/DataInterfaces/PCGComputeDataInterface.h b/Engine/Plugins/PCG/Source/PCG/Private/Compute/DataInterfaces/PCGComputeDataInterface.h
--- a/Engine/Plugins/PCG/Source/PCG/Private/Compute/DataInterfaces/PCGComputeDataInterface.h
+++ b/Engine/Plugins/PCG/Source/PCG/Private/Compute/DataInterfaces/PCGComputeDataInterface.h
@@ -29,6 +29,9 @@ public:
 
 	static void AddReferencedObjects(UObject* InThis, FReferenceCollector& Collector);
 
+	/** Set whether data originates from the CPU and should be uploaded to the GPU. */
+	void SetRequiresUpload(bool bInRequiresUpload) { bRequiresUpload = bInRequiresUpload; }
+
 public:
 	UPROPERTY()
 	int32 GraphBindingIndex = INDEX_NONE;
@@ -47,6 +50,9 @@ public:
 	UPROPERTY()
 	TArray<FName> DownstreamInputPinLabelAliases;
 
+	UPROPERTY()
+	bool bRequiresUpload = false;
+
 protected:
 	UPROPERTY()
 	TSoftObjectPtr<const UPCGSettings> ProducerSettings;
@@ -90,6 +96,7 @@ public:
 	FName OutputPinLabel;
 	FName OutputPinLabelAlias;
 	TArray<FName> DownstreamInputPinLabelAliases;
+	bool bRequiresUpload = false;
 
 	/** Bumped each time the data provider is initialized or reset, so that async callbacks can detect if they originated from a previous usage of the data provider and no-op. */
 	std::atomic<uint64> GenerationCounter = 0;
@@ -106,6 +113,7 @@ private:
 	mutable TObjectPtr<const UPCGSettings> ResolvedProducerSettings = nullptr;
 };
 
+/** Interface for data that can be exported/downloaded from the GPU to the CPU. */
 UCLASS(Abstract, ClassGroup = (Procedural))
 class UPCGExportableDataInterface : public UPCGComputeDataInterface
 {
diff --git a/Engine/Plugins/PCG/Source/PCG/Private/Compute/DataInterfaces/PCGDataCollectionDataInterface.cpp b/Engine/Plugins/PCG/Source/PCG/Private/Compute/DataInterfaces/PCGDataCollectionDataInterface.cpp
--- a/Engine/Plugins/PCG/Source/PCG/Private/Compute/DataInterfaces/PCGDataCollectionDataInterface.cpp
+++ b/Engine/Plugins/PCG/Source/PCG/Private/Compute/DataInterfaces/PCGDataCollectionDataInterface.cpp
@@ -10,6 +10,7 @@
 #include "PCGSubsystem.h"
 #include "Compute/PCGComputeCommon.h"
 #include "Compute/PCGDataBinding.h"
+#include "Compute/PCGDataUploadAdaptor.h"
 #include "Compute/Data/PCGProxyForGPUData.h"
 #include "Compute/Elements/PCGComputeGraphElement.h"
 #include "Helpers/PCGHelpers.h"
@@ -537,11 +538,15 @@ void UPCGDataCollectionDataProvider::Initialize(const UComputeDataInterface* InD
 
 	Super::Initialize(InDataInterface, InBinding, InInputMask, InOutputMask);
 
-	const UPCGDataCollectionDataInterface* DataInterface = CastChecked<UPCGDataCollectionDataInterface>(InDataInterface);
-
-	bRequiresZeroInitialization = DataInterface->bRequiresZeroInitialization;
+	// Initialize settings relevant for inter-kernel and graph output DIs.
+	if (!bRequiresUpload)
+	{
+		const UPCGDataCollectionDataInterface* DataInterface = CastChecked<UPCGDataCollectionDataInterface>(InDataInterface);
 
-	ElementCountMultiplier = DataInterface->ElementCountMultiplier;
+		// Only applicable for outputs from GPU nodes (not uploads from CPU).
+		bRequiresZeroInitialization = DataInterface->bRequiresZeroInitialization;
+		ElementCountMultiplier = DataInterface->ElementCountMultiplier;
+	}
 }
 
 bool UPCGDataCollectionDataProvider::PrepareForExecute_GameThread(UPCGDataBinding* InBinding)
@@ -556,19 +561,73 @@ bool UPCGDataCollectionDataProvider::PrepareForExecute_GameThread(UPCGDataBindin
 
 	check(PinDesc);
 
-	if (PCGComputeHelpers::IsBufferSizeTooLarge(PinDesc->ComputePackedSizeBytes()))
+	if (!bBufferSizeValidated)
 	{
-		PinDesc = nullptr;
+		if (PCGComputeHelpers::IsBufferSizeTooLarge(PinDesc->ComputePackedSizeBytes()))
+		{
+			PinDesc = nullptr;
+			return true;
+		}
+
+		bBufferSizeValidated = true;
 	}
 
-	return true;
+	if (bRequiresUpload)
+	{
+		if (!DataAdaptor)
+		{
+			// Requires virtual input pin labels to pick the data items from the compute graph element's input data collection.
+			check(!DownstreamInputPinLabelAliases.IsEmpty());
+			DataAdaptor = MakeShared<FPCGDataUploadAdaptor>(InBinding, PinDesc, DownstreamInputPinLabelAliases[0]);
+		}
+
+		const bool bPreparationDone = !DataAdaptor || DataAdaptor->PrepareData_GameThread();
+
+#if WITH_EDITOR
+		if (bPreparationDone && DataAdaptor && DataAdaptor->IsUploadingFromCPU())
+		{
+			NotifyProducerUploadedData(InBinding);
+		}
+#endif
+
+		return bPreparationDone;
+	}
+	else
+	{
+		return true;
+	}
+}
+
+void UPCGDataCollectionDataProvider::ReleaseTransientResources(const TCHAR* InReason)
+{
+	if (DataAdaptor)
+	{
+#ifdef PCG_DATA_USAGE_LOGGING
+		UE_LOG(LogPCG, Warning, TEXT("%s: Releasing resources due to %s"), *GetName(), InReason ? InReason : TEXT("NOREASON"));
+#endif
+
+		DataAdaptor.Reset();
+	}
 }
 
 FComputeDataProviderRenderProxy* UPCGDataCollectionDataProvider::GetRenderProxy()
 {
 	TRACE_CPUPROFILER_EVENT_SCOPE(UPCGDataCollectionDataProvider::GetRenderProxy);
 
-	return new FPCGDataCollectionDataProviderProxy(this, PinDesc, ExportMode, bRequiresZeroInitialization, ElementCountMultiplier, OutputPinLabel, OutputPinLabelAlias);
+	FPCGDataCollectionDataProviderProxy::FSetupParams ProxyParams =
+	{
+		.DataProvider = this,
+		.PinDesc = PinDesc,
+		.ExportMode = ExportMode,
+		.bZeroInitialize = bRequiresZeroInitialization,
+		.ElementCountMultiplier = ElementCountMultiplier,
+		.OutputPinLabel = OutputPinLabel,
+		.OutputPinLabelAlias = OutputPinLabelAlias,
+		.bRequiresUpload = bRequiresUpload,
+		.DataAdaptor = DataAdaptor,
+	};
+
+	return new FPCGDataCollectionDataProviderProxy(ProxyParams);
 }
 
 void UPCGDataCollectionDataProvider::Reset()
@@ -577,31 +636,28 @@ void UPCGDataCollectionDataProvider::Reset()
 
 	ElementCountMultiplier = 0;
 	bRequiresZeroInitialization = false;
+	DataAdaptor.Reset();
+	bBufferSizeValidated = false;
 }
 
-FPCGDataCollectionDataProviderProxy::FPCGDataCollectionDataProviderProxy(
-	TWeakObjectPtr<UPCGDataCollectionDataProvider> InDataProvider,
-	const FPCGDataCollectionDesc* InPinDesc,
-	EPCGExportMode InExportMode,
-	bool bInZeroInitialize,
-	uint32 InElementCountMultiplier,
-	FName InOutputPinLabel,
-	FName InOutputPinLabelAlias)
-	: ExportMode(InExportMode)
-	, bZeroInitialize(bInZeroInitialize)
-	, ElementCountMultiplier(InElementCountMultiplier)
-	, DataProviderWeakPtr(InDataProvider)
-	, OutputPinLabel(InOutputPinLabel)
-	, OutputPinLabelAlias(InOutputPinLabelAlias)
+FPCGDataCollectionDataProviderProxy::FPCGDataCollectionDataProviderProxy(const FSetupParams& InParams)
+	: ExportMode(InParams.ExportMode)
+	, bZeroInitialize(InParams.bZeroInitialize)
+	, ElementCountMultiplier(InParams.ElementCountMultiplier)
+	, DataProviderWeakPtr(InParams.DataProvider)
+	, OutputPinLabel(InParams.OutputPinLabel)
+	, OutputPinLabelAlias(InParams.OutputPinLabelAlias)
+	, bRequiresUpload(InParams.bRequiresUpload)
+	, DataAdaptor(InParams.DataAdaptor)
 {
-	OriginatingGenerationCount = InDataProvider->GenerationCounter;
+	OriginatingGenerationCount = DataProviderWeakPtr->GenerationCounter;
 
-	if (InPinDesc)
+	if (InParams.PinDesc)
 	{
 		TRACE_CPUPROFILER_EVENT_SCOPE(CopyDataDescription);
 
 		// @todo_pcg: Avoid costly copy of data description.
-		PinDesc = *InPinDesc;
+		PinDesc = *InParams.PinDesc;
 
 		for (const FPCGDataDesc& DataDesc : PinDesc.DataDescs)
 		{
@@ -628,6 +684,12 @@ bool FPCGDataCollectionDataProviderProxy::IsValid(FValidationData const& InValid
 		return false;
 	}
 
+	if (bRequiresUpload && !DataAdaptor.IsValid())
+	{
+		UE_LOG(LogPCG, Warning, TEXT("FPCGDataProviderDataCollectionProxy invalid due to null DataAdaptor (and bRequiresUpload=true)."));
+		return false;
+	}
+
 	return true;
 }
 
@@ -651,113 +713,142 @@ void FPCGDataCollectionDataProviderProxy::GatherDispatchData(FDispatchData const
 void FPCGDataCollectionDataProviderProxy::AllocateResources(FRDGBuilder& GraphBuilder, FAllocationData const& InAllocationData)
 {
 	TRACE_CPUPROFILER_EVENT_SCOPE(FPCGDataCollectionDataProviderProxy::AllocateResources);
-	check(SizeBytes > 0);
-
-	AttributeIdRemapSRV = GetAttributeIdRemapBufferSRV(GraphBuilder, FirstRemappedAttributeId);
 
-	BufferToGraphStringKeySRV = GetBufferToGraphStringKeySRV(GraphBuilder, NumRemappedStringKeys);
-
-	FRDGBufferDesc Desc = FRDGBufferDesc::CreateByteAddressDesc(SizeBytes);
-	if (ExportMode != EPCGExportMode::NoExport)
+	if (!bRequiresUpload)
 	{
-		// We don't know for sure whether buffer will be read back or not, so need to flag the possibility if the buffer will be passed downstream.
-		Desc.Usage |= BUF_SourceCopy;
-	}
+		// Initialize buffer for output from a GPU kernel.
 
-	Buffer = GraphBuilder.CreateBuffer(Desc, TEXT("PCGDataCollection"));
-	BufferUAV = GraphBuilder.CreateUAV(Buffer);
+		check(SizeBytes > 0);
 
-	// Initialize with an empty data collection. The kernel may not run, for example if indirect dispatch args end up being 0. This ensures
-	// there is something meaningful to readback.
-	TArray<uint32> PackedDataCollection;
-	PackedDataCollection.Reserve(bZeroInitialize ? (SizeBytes / PackedDataCollection.GetTypeSize()) : PinDesc.ComputePackedHeaderSizeBytes());
+		AttributeIdRemapSRV = GetAttributeIdRemapBufferSRV(GraphBuilder, FirstRemappedAttributeId);
 
-	{
-		TRACE_CPUPROFILER_EVENT_SCOPE(FPCGDataCollectionDataProviderProxy::AllocateResources::WriteHeader);
-		PinDesc.WriteHeader(PackedDataCollection);
-	}
+		BufferToGraphStringKeySRV = GetBufferToGraphStringKeySRV(GraphBuilder, NumRemappedStringKeys);
 
-	if (bZeroInitialize)
-	{
-		TRACE_CPUPROFILER_EVENT_SCOPE(FPCGDataCollectionDataProviderProxy::AllocateResources::ZeroInitialize);
+		FRDGBufferDesc Desc = FRDGBufferDesc::CreateByteAddressDesc(SizeBytes);
+		if (ExportMode != EPCGExportMode::NoExport)
+		{
+			// We don't know for sure whether buffer will be read back or not, so need to flag the possibility if the buffer will be passed downstream.
+			Desc.Usage |= BUF_SourceCopy;
+		}
 
-		// Fill remaining part of full buffer with 0s.
-		PackedDataCollection.SetNumZeroed(SizeBytes / PackedDataCollection.GetTypeSize());
-	}
-#if !UE_BUILD_SHIPPING
-	else if (PCGSystemSwitches::CVarFuzzGPUMemory.GetValueOnAnyThread())
-	{
-		// Allocate rest of buffer and initialize with random numbers below.
-		const int NumBefore = PackedDataCollection.Num();
-		PackedDataCollection.SetNumUninitialized(SizeBytes / PackedDataCollection.GetTypeSize());
+		Buffer = GraphBuilder.CreateBuffer(Desc, TEXT("PCGDataCollection"));
+		BufferUAV = GraphBuilder.CreateUAV(Buffer);
+
+		// Initialize with an empty data collection. The kernel may not run, for example if indirect dispatch args end up being 0. This ensures
+		// there is something meaningful to readback.
+		TArray<uint32> PackedDataCollection;
+		PackedDataCollection.Reserve(bZeroInitialize ? (SizeBytes / PackedDataCollection.GetTypeSize()) : PinDesc.ComputePackedHeaderSizeBytes());
 
-		FRandomStream Rand(GFrameCounter);
-		for (int Index = NumBefore; Index < PackedDataCollection.Num(); ++Index)
 		{
-			PackedDataCollection[Index] = Rand.GetUnsignedInt();
+			TRACE_CPUPROFILER_EVENT_SCOPE(FPCGDataCollectionDataProviderProxy::AllocateResources::WriteHeader);
+			PinDesc.WriteHeader(PackedDataCollection);
 		}
-	}
-#endif // !UE_BUILD_SHIPPING
-
-	GraphBuilder.QueueBufferUpload(Buffer, PackedDataCollection.GetData(), PackedDataCollection.Num() * PackedDataCollection.GetTypeSize(), ERDGInitialDataFlags::None);
 
-	// If buffer needs to be exported, can do it now. Pass back a reference to the main thread where it can be picked up by the compute graph element.
-	if (ExportMode != EPCGExportMode::NoExport)
-	{
-		const TRefCountPtr<FRDGPooledBuffer> ExportedPooledBuffer = GraphBuilder.ConvertToExternalBuffer(Buffer);
-		GraphBuilder.SetBufferAccessFinal(Buffer, ERHIAccess::CopySrc);
+		if (bZeroInitialize)
+		{
+			TRACE_CPUPROFILER_EVENT_SCOPE(FPCGDataCollectionDataProviderProxy::AllocateResources::ZeroInitialize);
 
-		ExecuteOnGameThread(UE_SOURCE_LOCATION, [ExportedPooledBuffer, DataProviderWeakPtr=DataProviderWeakPtr, PinDesc=PinDesc, OutputPinLabel=OutputPinLabel, SizeBytes=SizeBytes, ExportMode=ExportMode, OutputPinLabelAlias=OutputPinLabelAlias, GenerationCount=OriginatingGenerationCount]()
+			// Fill remaining part of full buffer with 0s.
+			PackedDataCollection.SetNumZeroed(SizeBytes / PackedDataCollection.GetTypeSize());
+		}
+#if !UE_BUILD_SHIPPING
+		else if (PCGSystemSwitches::CVarFuzzGPUMemory.GetValueOnAnyThread())
 		{
-			// Obtain objects. No ensures added because a graph cancellation could feasibly destroy some or all of these.
-			UPCGDataCollectionDataProvider* DataProvider = DataProviderWeakPtr.Get();
-			if (!DataProvider)
-			{
-				UE_LOG(LogPCG, Error, TEXT("Could not resolve UPCGDataCollectionDataProvider object to pass back buffer handle."));
-				return;
-			}
+			// Allocate rest of buffer and initialize with random numbers below.
+			const int NumBefore = PackedDataCollection.Num();
+			PackedDataCollection.SetNumUninitialized(SizeBytes / PackedDataCollection.GetTypeSize());
 
-			if (DataProvider->GenerationCounter != GenerationCount)
+			FRandomStream Rand(GFrameCounter);
+			for (int Index = NumBefore; Index < PackedDataCollection.Num(); ++Index)
 			{
-				return;
+				PackedDataCollection[Index] = Rand.GetUnsignedInt();
 			}
+		}
+#endif // !UE_BUILD_SHIPPING
+
+		GraphBuilder.QueueBufferUpload(Buffer, PackedDataCollection.GetData(), PackedDataCollection.Num() * PackedDataCollection.GetTypeSize(), ERDGInitialDataFlags::None);
 
-			UPCGDataBinding* Binding = DataProvider ? DataProvider->GetDataBinding() : nullptr;
-			UPCGSubsystem* Subsystem = UPCGSubsystem::GetSubsystemForCurrentWorld();
+		// If buffer needs to be exported, can do it now. Pass back a reference to the main thread where it can be picked up by the compute graph element.
+		if (ExportMode != EPCGExportMode::NoExport)
+		{
+			const TRefCountPtr<FRDGPooledBuffer> ExportedPooledBuffer = GraphBuilder.ConvertToExternalBuffer(Buffer);
+			GraphBuilder.SetBufferAccessFinal(Buffer, ERHIAccess::CopySrc);
 
-			if (Binding && Subsystem)
+			ExecuteOnGameThread(UE_SOURCE_LOCATION, [ExportedPooledBuffer, DataProviderWeakPtr = DataProviderWeakPtr, PinDesc = PinDesc, OutputPinLabel = OutputPinLabel, SizeBytes = SizeBytes, ExportMode = ExportMode, OutputPinLabelAlias = OutputPinLabelAlias, GenerationCount = OriginatingGenerationCount]()
 			{
-				TSharedPtr<FPCGProxyForGPUDataCollection> DataCollectionOnGPU = MakeShared<FPCGProxyForGPUDataCollection>(ExportedPooledBuffer, SizeBytes, PinDesc, Binding->GetStringTable());
+				// Obtain objects. No ensures added because a graph cancellation could feasibly destroy some or all of these.
+				UPCGDataCollectionDataProvider* DataProvider = DataProviderWeakPtr.Get();
+				if (!DataProvider)
+				{
+					UE_LOG(LogPCG, Error, TEXT("Could not resolve UPCGDataCollectionDataProvider object to pass back buffer handle."));
+					return;
+				}
+
+				if (DataProvider->GenerationCounter != GenerationCount)
+				{
+					return;
+				}
+
+				UPCGDataBinding* Binding = DataProvider ? DataProvider->GetDataBinding() : nullptr;
+				UPCGSubsystem* Subsystem = UPCGSubsystem::GetSubsystemForCurrentWorld();
 
-				for (int DataIndex = 0; DataIndex < PinDesc.DataDescs.Num(); ++DataIndex)
+				if (Binding && Subsystem)
 				{
-					UPCGProxyForGPUData* Proxy = NewObject<UPCGProxyForGPUData>();
-					Proxy->Initialize(DataCollectionOnGPU, DataIndex);
+					TSharedPtr<FPCGProxyForGPUDataCollection> DataCollectionOnGPU = MakeShared<FPCGProxyForGPUDataCollection>(ExportedPooledBuffer, SizeBytes, PinDesc, Binding->GetStringTable());
 
-					// TODO - binding is doing a lot of work. Could store a context handle in the data provider instead?
-					Binding->ReceiveDataFromGPU_GameThread(Proxy, DataProvider->GetProducerSettings(), ExportMode, OutputPinLabel, OutputPinLabelAlias);
+					for (int DataIndex = 0; DataIndex < PinDesc.DataDescs.Num(); ++DataIndex)
+					{
+						UPCGProxyForGPUData* Proxy = NewObject<UPCGProxyForGPUData>();
+						Proxy->Initialize(DataCollectionOnGPU, DataIndex);
+
+						// TODO - binding is doing a lot of work. Could store a context handle in the data provider instead?
+						Binding->ReceiveDataFromGPU_GameThread(Proxy, DataProvider->GetProducerSettings(), ExportMode, OutputPinLabel, OutputPinLabelAlias);
+					}
+
+					DataProvider->OnDataExported_GameThread().Broadcast();
 				}
+			});
+		}
+	}
+	else
+	{
+		// For the upload case, obtain the buffers from the data adaptor.
+		Buffer = DataAdaptor->GetBuffer_RenderThread(GraphBuilder, ExportMode);
+		BufferUAV = GraphBuilder.CreateUAV(Buffer);
 
-				DataProvider->OnDataExported_GameThread().Broadcast();
-			}
-		});
+		AttributeIdRemapSRV = GetAttributeIdRemapBufferSRV(GraphBuilder, FirstRemappedAttributeId);
+		BufferToGraphStringKeySRV = GetBufferToGraphStringKeySRV(GraphBuilder, NumRemappedStringKeys);
 	}
 }
 
 FRDGBufferSRVRef FPCGDataCollectionDataProviderProxy::GetAttributeIdRemapBufferSRV(FRDGBuilder& GraphBuilder, int32& OutFirstRemappedAttributeId) const
 {
-	// Default implementation. No remap required, disabled remapping.
-	OutFirstRemappedAttributeId = PCGComputeConstants::MAX_NUM_ATTRS;
+	if (!bRequiresUpload || !DataAdaptor || !DataAdaptor->UsesAttributeIdRemap())
+	{
+		// Default implementation. No remap required, disabled remapping.
+		OutFirstRemappedAttributeId = PCGComputeConstants::MAX_NUM_ATTRS;
 
-	return GraphBuilder.CreateSRV(FRDGBufferSRVDesc(GSystemTextures.GetDefaultStructuredBuffer(GraphBuilder, sizeof(int32))));
+		return GraphBuilder.CreateSRV(FRDGBufferSRVDesc(GSystemTextures.GetDefaultStructuredBuffer(GraphBuilder, sizeof(int32))));
+	}
+	else
+	{
+		return DataAdaptor->GetAttributeRemapBufferSRV(GraphBuilder, OutFirstRemappedAttributeId);
+	}
 }
 
 FRDGBufferSRVRef FPCGDataCollectionDataProviderProxy::GetBufferToGraphStringKeySRV(FRDGBuilder& GraphBuilder, int32& OutNumRemappedStringKeys) const
 {
-	// Default implementation. Disable remapping.
-	OutNumRemappedStringKeys = 0;
+	if (!bRequiresUpload || !DataAdaptor || !DataAdaptor->UsesStringKeyRemap())
+	{
+		// Default implementation. Disable remapping.
+		OutNumRemappedStringKeys = 0;
 
-	return GraphBuilder.CreateSRV(FRDGBufferSRVDesc(GSystemTextures.GetDefaultStructuredBuffer(GraphBuilder, sizeof(int32))));
+		return GraphBuilder.CreateSRV(FRDGBufferSRVDesc(GSystemTextures.GetDefaultStructuredBuffer(GraphBuilder, sizeof(int32))));
+	}
+	else
+	{
+		return DataAdaptor->GetBufferToGraphStringKeySRV(GraphBuilder, OutNumRemappedStringKeys);
+	}
 }
 
 #undef LOCTEXT_NAMESPACE
diff --git a/Engine/Plugins/PCG/Source/PCG/Private/Compute/DataInterfaces/PCGDataCollectionDataInterface.h b/Engine/Plugins/PCG/Source/PCG/Private/Compute/DataInterfaces/PCGDataCollectionDataInterface.h
--- a/Engine/Plugins/PCG/Source/PCG/Private/Compute/DataInterfaces/PCGDataCollectionDataInterface.h
+++ b/Engine/Plugins/PCG/Source/PCG/Private/Compute/DataInterfaces/PCGDataCollectionDataInterface.h
@@ -12,6 +12,7 @@
 
 struct FPCGComputeGraphContext;
 class FPCGDataCollectionDataInterfaceParameters;
+class FPCGDataUploadAdaptor;
 class UPCGDataBinding;
 
 /** Compute Framework Data Interface for reading PCG data. */
@@ -37,8 +38,10 @@ public:
 	UComputeDataProvider* CreateDataProvider() const override;
 	//~ End UComputeDataInterface Interface
 
+	/** Set whether entire buffer needs to be initialized to 0 on the GPU, if not only the header is initialized for efficiency. */
 	void SetRequiresZeroInitialization(bool bInZeroInit) { bRequiresZeroInitialization = bInZeroInit; }
 
+	/** Apply multiplier to all element counts (only used for kernel output data). */
 	void SetElementCountMultiplier(uint32 InElementCountMultiplier) { ElementCountMultiplier = InElementCountMultiplier; }
 
 public:
@@ -46,6 +49,7 @@ public:
 	UPROPERTY()
 	bool bRequiresZeroInitialization = false;
 
+	/** Multiplier applied to all element counts (only used for kernel output data). */
 	UPROPERTY()
 	uint32 ElementCountMultiplier = 0;
 
@@ -62,6 +66,7 @@ class UPCGDataCollectionDataProvider : public UPCGExportableDataProvider
 public:
 	//~ Begin UPCGComputeDataProvider Interface
 	virtual bool PrepareForExecute_GameThread(UPCGDataBinding* InBinding) override;
+	virtual void ReleaseTransientResources(const TCHAR* InReason) override;
 	//~ End UPCGComputeDataProvider Interface
 
 	//~ Begin UComputeDataProvider Interface
@@ -71,24 +76,34 @@ public:
 	//~ End UComputeDataProvider Interface
 
 public:
-	UPROPERTY()
 	uint32 ElementCountMultiplier = 0;
 
-	UPROPERTY()
 	bool bRequiresZeroInitialization = false;
+
+protected:
+	bool bBufferSizeValidated = false;
+
+	TSharedPtr<FPCGDataUploadAdaptor> DataAdaptor;
 };
 
 class FPCGDataCollectionDataProviderProxy : public FComputeDataProviderRenderProxy
 {
 public:
-	FPCGDataCollectionDataProviderProxy(
-		TWeakObjectPtr<UPCGDataCollectionDataProvider> InDataProvider,
-		const FPCGDataCollectionDesc* InPinDesc,
-		EPCGExportMode InExportMode,
-		bool bInZeroInitialize,
-		uint32 InElementCountMultiplier,
-		FName InOutputPinLabel,
-		FName InOutputPinLabelAlias);
+	struct FSetupParams
+	{
+		// Weak pointer useful for passing back buffer handles. Do not access this directly from the render thread.
+		TWeakObjectPtr<UPCGDataCollectionDataProvider> DataProvider;
+		const FPCGDataCollectionDesc* PinDesc = nullptr;
+		EPCGExportMode ExportMode = EPCGExportMode::NoExport;
+		bool bZeroInitialize = false;
+		uint32 ElementCountMultiplier = 0;
+		FName OutputPinLabel;
+		FName OutputPinLabelAlias;
+		bool bRequiresUpload = false;
+		TSharedPtr<FPCGDataUploadAdaptor> DataAdaptor;
+	};
+
+	FPCGDataCollectionDataProviderProxy(const FSetupParams& InParams);
 
 	//~ Begin FComputeDataProviderRenderProxy Interface
 	bool IsValid(FValidationData const& InValidationData) const override;
@@ -97,9 +112,9 @@ public:
 	//~ End FComputeDataProviderRenderProxy Interface
 
 protected:
-	virtual FRDGBufferSRVRef GetAttributeIdRemapBufferSRV(FRDGBuilder& GraphBuilder, int32& OutFirstRemappedAttributeId) const;
+	FRDGBufferSRVRef GetAttributeIdRemapBufferSRV(FRDGBuilder& GraphBuilder, int32& OutFirstRemappedAttributeId) const;
 
-	virtual FRDGBufferSRVRef GetBufferToGraphStringKeySRV(FRDGBuilder& GraphBuilder, int32& OutNumRemappedStringKeys) const;
+	FRDGBufferSRVRef GetBufferToGraphStringKeySRV(FRDGBuilder& GraphBuilder, int32& OutNumRemappedStringKeys) const;
 
 public:
 	/** Generation count of the data provider when the proxy was created. */
@@ -136,4 +151,8 @@ protected:
 	int32 NumRemappedStringKeys = INDEX_NONE;
 
 	int32 LargestAttributeIndex = INDEX_NONE;
+
+	bool bRequiresUpload = false;
+
+	TSharedPtr<FPCGDataUploadAdaptor> DataAdaptor;
 };
diff --git a/Engine/Plugins/PCG/Source/PCG/Private/Compute/DataInterfaces/PCGDataCollectionUploadDataInterface.cpp b/Engine/Plugins/PCG/Source/PCG/Private/Compute/DataInterfaces/PCGDataCollectionUploadDataInterface.cpp
deleted file mode 100644
diff --git a/Engine/Plugins/PCG/Source/PCG/Private/Compute/DataInterfaces/PCGDataCollectionUploadDataInterface.h b/Engine/Plugins/PCG/Source/PCG/Private/Compute/DataInterfaces/PCGDataCollectionUploadDataInterface.h
deleted file mode 100644
diff --git a/Engine/Plugins/PCG/Source/PCG/Private/Graph/PCGGraphCompilerGPU.cpp b/Engine/Plugins/PCG/Source/PCG/Private/Graph/PCGGraphCompilerGPU.cpp
--- a/Engine/Plugins/PCG/Source/PCG/Private/Graph/PCGGraphCompilerGPU.cpp
+++ b/Engine/Plugins/PCG/Source/PCG/Private/Graph/PCGGraphCompilerGPU.cpp
@@ -12,7 +12,6 @@
 #include "Compute/DataInterfaces/PCGComputeDataInterface.h"
 #include "Compute/DataInterfaces/PCGCustomKernelDataInterface.h"
 #include "Compute/DataInterfaces/PCGDataCollectionDataInterface.h"
-#include "Compute/DataInterfaces/PCGDataCollectionUploadDataInterface.h"
 #include "Compute/DataInterfaces/PCGDataLabelResolverDataInterface.h"
 #include "Compute/DataInterfaces/PCGDebugDataInterface.h"
 #include "Compute/DataInterfaces/PCGLandscapeDataInterface.h"
@@ -930,20 +929,13 @@ void FPCGGraphCompilerGPU::CreateDataInterfaces(
 		case EPCGDataType::Point:
 		case EPCGDataType::Param:
 		{
-			UPCGDataCollectionDataInterface* DataInterfacePCGData = nullptr;
+			UPCGDataCollectionDataInterface* DataInterfacePCGData = InOutContext.NewObject_AnyThread<UPCGDataCollectionDataInterface>(InOutComputeGraph);
 
-			if (bUpstreamIsGPUTask)
-			{
-				// Provides data for GPU -> GPU and GPU -> CPU edges.
-				DataInterfacePCGData = InOutContext.NewObject_AnyThread<UPCGDataCollectionDataInterface>(InOutComputeGraph);
-			}
-			else
-			{
-				// Provides data for CPU -> GPU edge.
-				DataInterfacePCGData = InOutContext.NewObject_AnyThread<UPCGDataCollectionUploadDataInterface>(InOutComputeGraph);
-			}
+			// If data comes from a CPU task, upload it to the GPU.
+			const bool bRequiresUpload = !bUpstreamIsGPUTask;
+			ensureMsgf(!bRequiresUpload || !bRequiresExport, TEXT("A data interface represents a kernel input or output and should not both upload and download to/from the GPU simultaneously."));
 
-			check(DataInterfacePCGData);
+			DataInterfacePCGData->SetRequiresUpload(bRequiresUpload);
 			DataInterfacePCGData->SetRequiresExport(bRequiresExport);
 
 			if (ProducerPin)
