[PCG] RGS change detection - move away from quantization to actual deltas, and when one or more grids are present, use a proportion of the smallest grid size.

I avoid using different thresholds per grid for now but we might explore this in the future (as long as we ensure whenever a child is generated the parent is generated if needed too).

#rnx

diff --git a/Engine/Plugins/PCG/Source/PCG/Internal/RuntimeGen/PCGRuntimeGenChangeDetection.h b/Engine/Plugins/PCG/Source/PCG/Internal/RuntimeGen/PCGRuntimeGenChangeDetection.h
--- a/Engine/Plugins/PCG/Source/PCG/Internal/RuntimeGen/PCGRuntimeGenChangeDetection.h
+++ b/Engine/Plugins/PCG/Source/PCG/Internal/RuntimeGen/PCGRuntimeGenChangeDetection.h
@@ -12,17 +12,20 @@ class UPCGComponent;
 
 namespace PCGRuntimeGenChangeDetection
 {
+	using FDetectionInputsBase = TTuple<TObjectKey<UPCGComponent>, const IPCGGenSourceBase*, /*RadiusMultiplier*/float, /*UseGenSourceDirection*/bool, /*MinGridSize*/uint32>;
+
 	/** The key used for the state cache. */
-	struct FDetectionInputs : TTuple<TObjectKey<UPCGComponent>, const IPCGGenSourceBase*, /*RadiusMultiplier*/float, /*UseGenSourceDirection*/bool>
+	struct FDetectionInputs : FDetectionInputsBase
 	{
-		FDetectionInputs(const UPCGComponent* InOriginalComponent, const IPCGGenSourceBase* InGenSource, float InRadiusMultiplier, bool bInUseGenSourceDirection)
-			: TTuple<TObjectKey<UPCGComponent>, const IPCGGenSourceBase*, float, bool>(InOriginalComponent, InGenSource, InRadiusMultiplier, bInUseGenSourceDirection)
+		FDetectionInputs(const UPCGComponent* InOriginalComponent, const IPCGGenSourceBase* InGenSource, float InRadiusMultiplier, bool bInUseGenSourceDirection, uint32 InMinGridSize)
+			: FDetectionInputsBase(InOriginalComponent, InGenSource, InRadiusMultiplier, bInUseGenSourceDirection, InMinGridSize)
 		{}
 
 		const UPCGComponent* GetOriginalComponent() const { return Get<0>().ResolveObjectPtr(); }
 		const IPCGGenSourceBase* GetGenSource() const { return Get<1>(); }
 		float GetGenerationRadiusMultiplier() const { return Get<2>(); }
 		bool UseGenSourceDirection() const { return Get<3>(); }
+		uint32 GetMinGridSize() const { return Get<4>(); }
 	};
 
 	/** Used to detect when world state (generation sources, cvars etc) have changed enough to warrant rescanning generation cells. */
@@ -35,8 +38,8 @@ namespace PCGRuntimeGenChangeDetection
 
 			bool ShouldScanCells(const FDetectionInputs& InInputs) const;
 
-			TOptional<FVector> QuantizedGenSourcePosition;
-			TOptional<FVector> QuantizedGenSourceDirection;
+			TOptional<FVector> GenSourcePosition;
+			TOptional<FVector> GenSourceDirection;
 			float GenerationRadiusMultiplier = 1.0f;
 		};
 
diff --git a/Engine/Plugins/PCG/Source/PCG/Private/RuntimeGen/PCGRuntimeGenChangeDetection.cpp b/Engine/Plugins/PCG/Source/PCG/Private/RuntimeGen/PCGRuntimeGenChangeDetection.cpp
--- a/Engine/Plugins/PCG/Source/PCG/Private/RuntimeGen/PCGRuntimeGenChangeDetection.cpp
+++ b/Engine/Plugins/PCG/Source/PCG/Private/RuntimeGen/PCGRuntimeGenChangeDetection.cpp
@@ -7,42 +7,27 @@
 
 namespace PCGRuntimeGenChangeDetection
 {
-	static FVector QuantizeVector(const FVector& InVector, float InQuanta)
-	{
-		return FVector(
-			FMath::FloorToFloat(InVector.X / InQuanta) * InQuanta,
-			FMath::FloorToFloat(InVector.Y / InQuanta) * InQuanta,
-			FMath::FloorToFloat(InVector.Z / InQuanta) * InQuanta
-		);
-	}
+	static TAutoConsoleVariable<float> CVarMaxPositionDeltaGridProportion(
+		TEXT("pcg.RuntimeGeneration.ChangeDetection.MaxPositionDeltaGridProportion"),
+		0.1f,
+		TEXT("For partitioned components that have one or more grid, use this proportion of the smallest grid as the maximum generation source position delta before triggering a rescan of cells."));
 
-	const static float DirectionQuantaSize = 0.025f;
+	static TAutoConsoleVariable<float> CVarMaxPositionDeltaUnbounded(
+		TEXT("pcg.RuntimeGeneration.ChangeDetection.MaxPositionDeltaUnbounded"),
+		100.0f,
+		TEXT("The maximum generation source position delta for unbounded generation (cm)."));
 
 	void FDetector::FCachedState::UpdateFromInputs(const FDetectionInputs& InInputs)
 	{
-		QuantizedGenSourcePosition = InInputs.GetGenSource()->GetPosition();
-
-		if (QuantizedGenSourcePosition)
-		{
-			// todo_pcg: possibly proportional to grid size, but need to ensure we generate parents when a child generates.
-			const float QuantaSize = 50.0f;
-
-			QuantizedGenSourcePosition = PCGRuntimeGenChangeDetection::QuantizeVector(*QuantizedGenSourcePosition, QuantaSize);
-		}
+		GenSourcePosition = InInputs.GetGenSource()->GetPosition();
 
 		if (InInputs.UseGenSourceDirection())
 		{
-			QuantizedGenSourceDirection = InInputs.GetGenSource()->GetDirection();
-
-			if (QuantizedGenSourceDirection)
-			{
-				// Quantize direction vector. Could be expanded to heading angle or such in the future.
-				QuantizedGenSourceDirection = PCGRuntimeGenChangeDetection::QuantizeVector(*QuantizedGenSourceDirection, DirectionQuantaSize);
-			}
+			GenSourceDirection = InInputs.GetGenSource()->GetDirection();
 		}
 		else
 		{
-			QuantizedGenSourceDirection.Reset();
+			GenSourceDirection.Reset();
 		}
 
 		GenerationRadiusMultiplier = InInputs.GetGenerationRadiusMultiplier();
@@ -52,17 +37,18 @@ namespace PCGRuntimeGenChangeDetection
 	{
 		if (TOptional<FVector> NewPosition = InInputs.GetGenSource()->GetPosition())
 		{
-			if (!QuantizedGenSourcePosition)
+			if (!GenSourcePosition)
 			{
-				// Position has just become available.
+				// Position has just become available, perform scan.
 				return true;
 			}
 
 			// Rescan if position has changed significantly.
-			// todo_pcg: Possibly proportional to (smallest?) grid size, but need to ensure we generate parents when a child generates.
-			const float QuantaSize = 50.0f;
+			const float MaxPositionDelta = (InInputs.GetMinGridSize() == PCGHiGenGrid::UnboundedGridSize())
+				? CVarMaxPositionDeltaUnbounded.GetValueOnGameThread()
+				: (CVarMaxPositionDeltaGridProportion.GetValueOnGameThread() * static_cast<float>(InInputs.GetMinGridSize()));
 
-			if (PCGRuntimeGenChangeDetection::QuantizeVector(*NewPosition, QuantaSize) != *QuantizedGenSourcePosition)
+			if (FVector::DistSquared(*GenSourcePosition, *NewPosition) > MaxPositionDelta * MaxPositionDelta)
 			{
 				return true;
 			}
@@ -72,7 +58,7 @@ namespace PCGRuntimeGenChangeDetection
 		{
 			if (TOptional<FVector> NewDirection = InInputs.GetGenSource()->GetDirection())
 			{
-				if (!QuantizedGenSourceDirection)
+				if (!GenSourceDirection)
 				{
 					// Direction has just become available.
 					return true;
@@ -80,12 +66,13 @@ namespace PCGRuntimeGenChangeDetection
 
 				// Rescan if direction has changed significantly.
 				// todo_pcg: Possibly only take heading angle if 2D grid.
-				if (PCGRuntimeGenChangeDetection::QuantizeVector(*NewDirection, DirectionQuantaSize) != *QuantizedGenSourceDirection)
+				const double MinDotProduct = 0.997; // Value tweaked by hand, could be driven by cvar if needed.
+				if (FVector::DotProduct(*GenSourceDirection, *NewDirection) < MinDotProduct)
 				{
 					return true;
 				}
 			}
-			else if (QuantizedGenSourceDirection)
+			else if (GenSourceDirection)
 			{
 				// Gen source direction has stopped coming through. Safest to trigger rescan - for e.g. this might disable frustum culling.
 				return true;
diff --git a/Engine/Plugins/PCG/Source/PCG/Private/RuntimeGen/PCGRuntimeGenScheduler.cpp b/Engine/Plugins/PCG/Source/PCG/Private/RuntimeGen/PCGRuntimeGenScheduler.cpp
--- a/Engine/Plugins/PCG/Source/PCG/Private/RuntimeGen/PCGRuntimeGenScheduler.cpp
+++ b/Engine/Plugins/PCG/Source/PCG/Private/RuntimeGen/PCGRuntimeGenScheduler.cpp
@@ -527,7 +527,8 @@ void FPCGRuntimeGenScheduler::TickQueueComponentsForGeneration(
 					OriginalComponent,
 					GenSource,
 					PCGRuntimeGenSchedulerHelpers::CVarRuntimeGenerationRadiusMultiplier.GetValueOnGameThread(),
-					Policy->CullsBasedOnDirection()
+					Policy->CullsBasedOnDirection(),
+					GridSizes.IsEmpty() ? PCGHiGenGrid::UnboundedGridSize() : GridSizes.Last(),
 				};
 
 				if (bDoChangeDetection && !ChangeDetector.IsCellScanRequired(ChangeDetectionInputs))
