[PCG] Fix struct member initialization errors on editor startup.

#rnx

diff --git a/Engine/Plugins/PCG/Source/PCG/Public/PCGComponent.h b/Engine/Plugins/PCG/Source/PCG/Public/PCGComponent.h
--- a/Engine/Plugins/PCG/Source/PCG/Public/PCGComponent.h
+++ b/Engine/Plugins/PCG/Source/PCG/Public/PCGComponent.h
@@ -856,11 +856,11 @@ protected:
 
 	// UPCGComponent::CurrentEditingMode is a Transient property so we need to capture/apply it ourselves
 	UPROPERTY()
-	EPCGEditorDirtyMode CurrentEditingMode;
+	EPCGEditorDirtyMode CurrentEditingMode = EPCGEditorDirtyMode::Normal;
 	
 	// UPCGComponent::PreviousEditingMode is not a property so we need to capture/apply it ourselves
 	UPROPERTY()
-	EPCGEditorDirtyMode PreviousEditingMode;
+	EPCGEditorDirtyMode PreviousEditingMode = EPCGEditorDirtyMode::Normal;
 
 	// UPCGComponent::bForceGenerateOnBPAddedToWorld is not a visible property so we need to capture/apply it ourselves
 	UPROPERTY()
