[PCG] Fixed another struct member initialization error.

#rnx

diff --git a/Engine/Plugins/PCG/Source/PCG/Public/PCGComponent.h b/Engine/Plugins/PCG/Source/PCG/Public/PCGComponent.h
--- a/Engine/Plugins/PCG/Source/PCG/Public/PCGComponent.h
+++ b/Engine/Plugins/PCG/Source/PCG/Public/PCGComponent.h
@@ -829,7 +829,7 @@ struct FDynamicTrackedKeyInstanceData
 	TSoftObjectPtr<const UPCGSettings> Settings;
 
 	UPROPERTY()
-	bool bValue;
+	bool bValue = false;
 };
 
 /** Used to store generated resources data during RerunConstructionScripts */
