[PCG] GPU: Small additional tweaks to last refactor.

Rebrand bRequiresUpload to bProducedByCPU but leave it at the base DI class as it seems like it will be useful for future DIs.

#rnx

diff --git a/Engine/Plugins/PCG/Source/PCG/Private/Compute/DataInterfaces/PCGComputeDataInterface.cpp b/Engine/Plugins/PCG/Source/PCG/Private/Compute/DataInterfaces/PCGComputeDataInterface.cpp
--- a/Engine/Plugins/PCG/Source/PCG/Private/Compute/DataInterfaces/PCGComputeDataInterface.cpp
+++ b/Engine/Plugins/PCG/Source/PCG/Private/Compute/DataInterfaces/PCGComputeDataInterface.cpp
@@ -68,13 +68,13 @@ void UPCGComputeDataProvider::Initialize(const UComputeDataInterface* InDataInte
 	const UPCGComputeDataInterface* DataInterface = CastChecked<UPCGComputeDataInterface>(InDataInterface);
 	ResolvedProducerSettings = DataInterface->GetProducerSettings();
 	ProducerSettings = ResolvedProducerSettings;
-	bRequiresUpload = DataInterface->bRequiresUpload;
+	bProducedByCPU = DataInterface->bProducedByCPU;
 	GraphBindingIndex = DataInterface->GraphBindingIndex;
 	ProducerKernel = DataInterface->ProducerKernel;
 	DownstreamInputPinLabelAliases = DataInterface->DownstreamInputPinLabelAliases;
 
-	// Set data that is only relevant for inter-kernel data interfaces, and compute graph outputs. Don't set them for inputs (CPU->GPU uploads).
-	if (!bRequiresUpload)
+	// Set data that is only relevant for data produced on GPU, mainly to help prevent misuse.
+	if (!bProducedByCPU)
 	{
 		// The original label is needed to store inspection data.
 		OutputPinLabel = DataInterface->OutputPinLabel;
@@ -98,7 +98,7 @@ void UPCGComputeDataProvider::Reset()
 	OutputPinLabel = NAME_None;
 	OutputPinLabelAlias = NAME_None;
 	DownstreamInputPinLabelAliases.Empty();
-	bRequiresUpload = false;
+	bProducedByCPU = false;
 }
 
 const UPCGSettings* UPCGComputeDataProvider::GetProducerSettings() const
@@ -164,10 +164,11 @@ void UPCGExportableDataProvider::Initialize(const UComputeDataInterface* InDataI
 	const UPCGExportableDataInterface* DataInterface = CastChecked<UPCGExportableDataInterface>(InDataInterface);
 	Binding = CastChecked<UPCGDataBinding>(InBinding);
 
-	if (!bRequiresUpload)
-	{
-		ExportMode = DataInterface->GetRequiresExport() ? EPCGExportMode::ComputeGraphOutput : EPCGExportMode::NoExport;
+	ExportMode = DataInterface->GetRequiresExport() ? EPCGExportMode::ComputeGraphOutput : EPCGExportMode::NoExport;
 
+	// Debug/inspection functionality for GPU-produced data.
+	if (!bProducedByCPU)
+	{
 #if WITH_EDITOR
 		// Some exportable data providers don't support inspect/debug because their producer settings are not available (e.g. GridLinkage).
 		const UPCGSettings* LocalProducerSettings = GetProducerSettings();
diff --git a/Engine/Plugins/PCG/Source/PCG/Private/Compute/DataInterfaces/PCGComputeDataInterface.h b/Engine/Plugins/PCG/Source/PCG/Private/Compute/DataInterfaces/PCGComputeDataInterface.h
--- a/Engine/Plugins/PCG/Source/PCG/Private/Compute/DataInterfaces/PCGComputeDataInterface.h
+++ b/Engine/Plugins/PCG/Source/PCG/Private/Compute/DataInterfaces/PCGComputeDataInterface.h
@@ -30,7 +30,7 @@ public:
 	static void AddReferencedObjects(UObject* InThis, FReferenceCollector& Collector);
 
 	/** Set whether data originates from the CPU and should be uploaded to the GPU. */
-	void SetRequiresUpload(bool bInRequiresUpload) { bRequiresUpload = bInRequiresUpload; }
+	void SetProducedByCPU(bool bInProducedByCPU) { bProducedByCPU = bInProducedByCPU; }
 
 public:
 	UPROPERTY()
@@ -51,7 +51,7 @@ public:
 	TArray<FName> DownstreamInputPinLabelAliases;
 
 	UPROPERTY()
-	bool bRequiresUpload = false;
+	bool bProducedByCPU = false;
 
 protected:
 	UPROPERTY()
@@ -96,7 +96,7 @@ public:
 	FName OutputPinLabel;
 	FName OutputPinLabelAlias;
 	TArray<FName> DownstreamInputPinLabelAliases;
-	bool bRequiresUpload = false;
+	bool bProducedByCPU = false;
 
 	/** Bumped each time the data provider is initialized or reset, so that async callbacks can detect if they originated from a previous usage of the data provider and no-op. */
 	std::atomic<uint64> GenerationCounter = 0;
diff --git a/Engine/Plugins/PCG/Source/PCG/Private/Compute/DataInterfaces/PCGDataCollectionDataInterface.cpp b/Engine/Plugins/PCG/Source/PCG/Private/Compute/DataInterfaces/PCGDataCollectionDataInterface.cpp
--- a/Engine/Plugins/PCG/Source/PCG/Private/Compute/DataInterfaces/PCGDataCollectionDataInterface.cpp
+++ b/Engine/Plugins/PCG/Source/PCG/Private/Compute/DataInterfaces/PCGDataCollectionDataInterface.cpp
@@ -538,8 +538,8 @@ void UPCGDataCollectionDataProvider::Initialize(const UComputeDataInterface* InD
 
 	Super::Initialize(InDataInterface, InBinding, InInputMask, InOutputMask);
 
-	// Initialize settings relevant for inter-kernel and graph output DIs.
-	if (!bRequiresUpload)
+	// Initialize settings relevant for data produced on GPU.
+	if (!bProducedByCPU)
 	{
 		const UPCGDataCollectionDataInterface* DataInterface = CastChecked<UPCGDataCollectionDataInterface>(InDataInterface);
 
@@ -572,7 +572,7 @@ bool UPCGDataCollectionDataProvider::PrepareForExecute_GameThread(UPCGDataBindin
 		bBufferSizeValidated = true;
 	}
 
-	if (bRequiresUpload)
+	if (bProducedByCPU)
 	{
 		if (!DataAdaptor)
 		{
@@ -623,7 +623,7 @@ FComputeDataProviderRenderProxy* UPCGDataCollectionDataProvider::GetRenderProxy(
 		.ElementCountMultiplier = ElementCountMultiplier,
 		.OutputPinLabel = OutputPinLabel,
 		.OutputPinLabelAlias = OutputPinLabelAlias,
-		.bRequiresUpload = bRequiresUpload,
+		.bProducedByCPU = bProducedByCPU,
 		.DataAdaptor = DataAdaptor,
 	};
 
@@ -647,7 +647,7 @@ FPCGDataCollectionDataProviderProxy::FPCGDataCollectionDataProviderProxy(const F
 	, DataProviderWeakPtr(InParams.DataProvider)
 	, OutputPinLabel(InParams.OutputPinLabel)
 	, OutputPinLabelAlias(InParams.OutputPinLabelAlias)
-	, bRequiresUpload(InParams.bRequiresUpload)
+	, bProducedByCPU(InParams.bProducedByCPU)
 	, DataAdaptor(InParams.DataAdaptor)
 {
 	OriginatingGenerationCount = DataProviderWeakPtr->GenerationCounter;
@@ -684,7 +684,7 @@ bool FPCGDataCollectionDataProviderProxy::IsValid(FValidationData const& InValid
 		return false;
 	}
 
-	if (bRequiresUpload && !DataAdaptor.IsValid())
+	if (bProducedByCPU && !DataAdaptor.IsValid())
 	{
 		UE_LOG(LogPCG, Warning, TEXT("FPCGDataProviderDataCollectionProxy invalid due to null DataAdaptor (and bRequiresUpload=true)."));
 		return false;
@@ -714,7 +714,7 @@ void FPCGDataCollectionDataProviderProxy::AllocateResources(FRDGBuilder& GraphBu
 {
 	TRACE_CPUPROFILER_EVENT_SCOPE(FPCGDataCollectionDataProviderProxy::AllocateResources);
 
-	if (!bRequiresUpload)
+	if (!bProducedByCPU)
 	{
 		// Initialize buffer for output from a GPU kernel.
 
@@ -812,7 +812,7 @@ void FPCGDataCollectionDataProviderProxy::AllocateResources(FRDGBuilder& GraphBu
 	}
 	else
 	{
-		// For the upload case, obtain the buffers from the data adaptor.
+		// For the CPU->GPU upload case, obtain the buffers from the data adaptor.
 		Buffer = DataAdaptor->GetBuffer_RenderThread(GraphBuilder, ExportMode);
 		BufferUAV = GraphBuilder.CreateUAV(Buffer);
 
@@ -823,7 +823,7 @@ void FPCGDataCollectionDataProviderProxy::AllocateResources(FRDGBuilder& GraphBu
 
 FRDGBufferSRVRef FPCGDataCollectionDataProviderProxy::GetAttributeIdRemapBufferSRV(FRDGBuilder& GraphBuilder, int32& OutFirstRemappedAttributeId) const
 {
-	if (!bRequiresUpload || !DataAdaptor || !DataAdaptor->UsesAttributeIdRemap())
+	if (!bProducedByCPU || !DataAdaptor || !DataAdaptor->UsesAttributeIdRemap())
 	{
 		// Default implementation. No remap required, disabled remapping.
 		OutFirstRemappedAttributeId = PCGComputeConstants::MAX_NUM_ATTRS;
@@ -838,7 +838,7 @@ FRDGBufferSRVRef FPCGDataCollectionDataProviderProxy::GetAttributeIdRemapBufferS
 
 FRDGBufferSRVRef FPCGDataCollectionDataProviderProxy::GetBufferToGraphStringKeySRV(FRDGBuilder& GraphBuilder, int32& OutNumRemappedStringKeys) const
 {
-	if (!bRequiresUpload || !DataAdaptor || !DataAdaptor->UsesStringKeyRemap())
+	if (!bProducedByCPU || !DataAdaptor || !DataAdaptor->UsesStringKeyRemap())
 	{
 		// Default implementation. Disable remapping.
 		OutNumRemappedStringKeys = 0;
diff --git a/Engine/Plugins/PCG/Source/PCG/Private/Compute/DataInterfaces/PCGDataCollectionDataInterface.h b/Engine/Plugins/PCG/Source/PCG/Private/Compute/DataInterfaces/PCGDataCollectionDataInterface.h
--- a/Engine/Plugins/PCG/Source/PCG/Private/Compute/DataInterfaces/PCGDataCollectionDataInterface.h
+++ b/Engine/Plugins/PCG/Source/PCG/Private/Compute/DataInterfaces/PCGDataCollectionDataInterface.h
@@ -99,7 +99,7 @@ public:
 		uint32 ElementCountMultiplier = 0;
 		FName OutputPinLabel;
 		FName OutputPinLabelAlias;
-		bool bRequiresUpload = false;
+		bool bProducedByCPU = false;
 		TSharedPtr<FPCGDataUploadAdaptor> DataAdaptor;
 	};
 
@@ -152,7 +152,7 @@ protected:
 
 	int32 LargestAttributeIndex = INDEX_NONE;
 
-	bool bRequiresUpload = false;
+	bool bProducedByCPU = false;
 
 	TSharedPtr<FPCGDataUploadAdaptor> DataAdaptor;
 };
diff --git a/Engine/Plugins/PCG/Source/PCG/Private/Graph/PCGGraphCompilerGPU.cpp b/Engine/Plugins/PCG/Source/PCG/Private/Graph/PCGGraphCompilerGPU.cpp
--- a/Engine/Plugins/PCG/Source/PCG/Private/Graph/PCGGraphCompilerGPU.cpp
+++ b/Engine/Plugins/PCG/Source/PCG/Private/Graph/PCGGraphCompilerGPU.cpp
@@ -932,10 +932,8 @@ void FPCGGraphCompilerGPU::CreateDataInterfaces(
 			UPCGDataCollectionDataInterface* DataInterfacePCGData = InOutContext.NewObject_AnyThread<UPCGDataCollectionDataInterface>(InOutComputeGraph);
 
 			// If data comes from a CPU task, upload it to the GPU.
-			const bool bRequiresUpload = !bUpstreamIsGPUTask;
-			ensureMsgf(!bRequiresUpload || !bRequiresExport, TEXT("A data interface represents a kernel input or output and should not both upload and download to/from the GPU simultaneously."));
+			ensureMsgf(!bRequiresExport || bUpstreamIsGPUTask, TEXT("Download from GPU only relevant for data produced on GPU."));
 
-			DataInterfacePCGData->SetRequiresUpload(bRequiresUpload);
 			DataInterfacePCGData->SetRequiresExport(bRequiresExport);
 
 			if (ProducerPin)
@@ -990,6 +988,7 @@ void FPCGGraphCompilerGPU::CreateDataInterfaces(
 			DataInterface->ProducerKernel = ProducerKernel;
 			DataInterface->SetProducerSettings(ProducerSettings);
 			DataInterface->SetOutputPin(InOutputPinProperties.Label);
+			DataInterface->SetProducedByCPU(!bUpstreamIsGPUTask);
 		}
 
 		return DataInterface;
