[PCG] GPU implementation of Transform Points

* Also adds 2 functions for Quaternions: Quat_FromNormal and Quat_FromRotator
* Renamed all functions for Quaternions from Quat... to Quat_... with the aliases to the old naming for existing ones.

#rnx
#virtualized

diff --git a/Engine/Plugins/PCG/Shaders/Private/Elements/PCGTransformPoints.usf b/Engine/Plugins/PCG/Shaders/Private/Elements/PCGTransformPoints.usf
new file mode 100644
--- /dev/null
+++ b/Engine/Plugins/PCG/Shaders/Private/Elements/PCGTransformPoints.usf
@@ -0,0 +1,157 @@
+// Copyright Epic Games, Inc. All Rights Reserved.
+
+[numthreads(64, 1, 1)]
+void PCGTransformPointsCS(uint3 GroupId : SV_GroupID, uint GroupIndex : SV_GroupIndex)
+{
+	// Mark the kernel as having executed. Must run before we early out via thread index, because the kernel is still 'executed' even if the number of
+	// threads to iterate on is zero. Even if GetNumThreads() returns 0, the kernel will still have been dispatched on a single thread to set this flag.
+	if (all(GroupId == 0) && GroupIndex == 0)
+	{
+		Out_SetAsExecutedInternal();
+	}
+
+	const uint ThreadIndex = GetUnWrappedDispatchThreadId(GroupId, GroupIndex, 64);
+	if (ThreadIndex >= GetNumThreads().x) return;
+
+	uint InDataIndex, InElemIndex;
+	if (!In_GetThreadData(ThreadIndex, InDataIndex, InElemIndex))
+	{
+		return;
+	}
+
+	uint OutDataIndex, OutElemIndex;
+	if (!Out_GetThreadData(ThreadIndex, OutDataIndex, OutElemIndex))
+	{
+		return;
+	}
+
+	// Propagate the 'Removed' status to the output point. Otherwise, this point will not be culled.
+	if (In_IsPointRemoved(InDataIndex, InElemIndex))
+	{
+		Out_RemovePoint(OutDataIndex, OutElemIndex);
+		return;
+	}
+
+	// Transform points code
+	const float3 OffsetMin = TransformPoints_GetOffsetMin();
+	const float3 OffsetMax = TransformPoints_GetOffsetMax();
+	const bool bAbsoluteOffset = TransformPoints_GetAbsoluteOffset() != 0;
+	const float3 RotationMin = TransformPoints_GetRotationMin();
+	const float3 RotationMax = TransformPoints_GetRotationMax();
+	const bool bAbsoluteRotation = TransformPoints_GetAbsoluteRotation() != 0;
+	const float3 ScaleMin = TransformPoints_GetScaleMin();
+	const float3 ScaleMax = TransformPoints_GetScaleMax();
+	const bool bAbsoluteScale = TransformPoints_GetAbsoluteScale() != 0;
+	const bool bUniformScale = TransformPoints_GetUniformScale() != 0;
+	const bool bRecomputeSeed = TransformPoints_GetRecomputeSeed() != 0;
+	const int AttributeId = TransformPoints_GetAttributeId();
+	const bool bApplyToAttribute = (TransformPoints_GetApplyToAttribute() != 0) && (AttributeId != -1);
+	
+	float3 InPosition;
+	FQuat InRotation;
+	float3 InScale;
+
+	if (bApplyToAttribute)
+	{
+		// TODO: If the transform is transposed in pack, remove this transpose.
+		float4x4 InTransform = transpose(In_GetTransform(InDataIndex, InElemIndex, AttributeId));
+
+		InScale = float3(length(InTransform._m00_m10_m20), length(InTransform._m01_m11_m21), length(InTransform._m02_m12_m22));
+
+		// Assumes Transform axes have the correct handedness. We could do more work to fix this up too if required.
+		float3x3 RotationMatix = (float3x3)transpose(InTransform);
+		// Required unfortunately, Quat_FromMatrix expects orthonormalized.
+		RotationMatix[0] /= InScale.x;
+		RotationMatix[1] /= InScale.y;
+		RotationMatix[2] /= InScale.z;
+		InRotation = Quat_FromMatrix(RotationMatix);
+
+		InPosition = InTransform._m03_m13_m23;
+	}
+	else
+	{
+		InPosition = In_GetPosition(InDataIndex, InElemIndex);
+		InRotation = In_GetRotation(InDataIndex, InElemIndex);
+		InScale = In_GetScale(InDataIndex, InElemIndex);
+	}
+	
+	const uint InSeed = In_GetSeed(InDataIndex, InElemIndex);
+	uint Seed = ComputeSeed(GetSeed(), InSeed);
+
+	// Position
+	float3 OutPosition = bAbsoluteOffset ? float3(0,0,0) : InPosition;
+	OutPosition.x += lerp(OffsetMin.x, OffsetMax.x, FRand(Seed));
+	OutPosition.y += lerp(OffsetMin.y, OffsetMax.y, FRand(Seed));
+	OutPosition.z += lerp(OffsetMin.z, OffsetMax.z, FRand(Seed));
+
+	// Rotation
+	FQuat OutRotation = bAbsoluteRotation ? FQuat(0,0,0,1) : InRotation;
+	float3 RandomRotation = float3(0,0,0);
+	RandomRotation.x = lerp(RotationMin.x, RotationMax.x, FRand(Seed));
+	RandomRotation.y = lerp(RotationMin.y, RotationMax.y, FRand(Seed));
+	RandomRotation.z = lerp(RotationMin.z, RotationMax.z, FRand(Seed));
+
+	OutRotation = Quat_Multiply(OutRotation, Quat_FromRotator(RandomRotation));
+
+	// Scale
+	float3 OutScale = bAbsoluteScale ? float3(1, 1, 1) : InScale;
+	if (bUniformScale)
+	{
+		OutScale *= lerp(ScaleMin.x, ScaleMax.x, FRand(Seed));
+	}
+	else
+	{
+		OutScale.x *= lerp(ScaleMin.x, ScaleMax.x, FRand(Seed));
+		OutScale.y *= lerp(ScaleMin.y, ScaleMax.y, FRand(Seed));
+		OutScale.z *= lerp(ScaleMin.z, ScaleMax.z, FRand(Seed));
+	}
+	
+	// Copy metadata attributes first.
+	PCG_COPY_ATTRIBUTES_TO_OUTPUT(Out, In, OutDataIndex, OutElemIndex, InDataIndex, InElemIndex, /*bMetadataOnly=*/true, /*bInitNonCopiedAttributes=*/true);
+
+	// Write output.
+	if (bApplyToAttribute)
+	{
+		const float3x3 RotationMatrix = Quat_ToMatrix(OutRotation);
+
+		const float3 Axis0 = OutScale.x * RotationMatrix[0];
+		const float3 Axis1 = OutScale.y * RotationMatrix[1];
+		const float3 Axis2 = OutScale.z * RotationMatrix[2];
+		
+		float4x4 OutTransform = float4x4(
+			Axis0.x, Axis1.x, Axis2.x, OutPosition.x,
+			Axis0.y, Axis1.y, Axis2.y, OutPosition.y,
+			Axis0.z, Axis1.z, Axis2.z, OutPosition.z,
+			(float3)0.0, 1.0);
+
+		// TODO: If the transform is transposed in unpack, remove this transpose.
+		Out_SetTransform(OutDataIndex, OutElemIndex, AttributeId, transpose(OutTransform));
+
+		// Copy original
+		Out_SetPosition(OutDataIndex, OutElemIndex, In_GetPosition(InDataIndex, InElemIndex));
+		Out_SetRotation(OutDataIndex, OutElemIndex, In_GetRotation(InDataIndex, InElemIndex));
+		Out_SetScale(OutDataIndex, OutElemIndex, In_GetScale(InDataIndex, InElemIndex));
+	}
+	else
+	{
+		Out_SetPosition(OutDataIndex, OutElemIndex, OutPosition);
+		Out_SetRotation(OutDataIndex, OutElemIndex, OutRotation);
+		Out_SetScale(OutDataIndex, OutElemIndex, OutScale);
+	}
+
+	if (bApplyToAttribute || !bRecomputeSeed)
+	{
+		Out_SetSeed(OutDataIndex, OutElemIndex, InSeed);
+	}
+	else
+	{
+		Out_SetSeed(OutDataIndex, OutElemIndex, ComputeSeedFromPosition(OutPosition));
+	}
+
+	// We didn't modify these properties, so simply copy from the input data.
+	Out_SetColor(OutDataIndex, OutElemIndex, In_GetColor(InDataIndex, InElemIndex));
+	Out_SetBoundsMin(OutDataIndex, OutElemIndex, In_GetBoundsMin(InDataIndex, InElemIndex));
+	Out_SetBoundsMax(OutDataIndex, OutElemIndex, In_GetBoundsMax(InDataIndex, InElemIndex));
+	Out_SetDensity(OutDataIndex, OutElemIndex, In_GetDensity(InDataIndex, InElemIndex));
+	Out_SetSteepness(OutDataIndex, OutElemIndex, In_GetSteepness(InDataIndex, InElemIndex));
+}
diff --git a/Engine/Plugins/PCG/Source/PCG/Private/Compute/DataInterfaces/Elements/PCGTransformPointsDataInterface.cpp b/Engine/Plugins/PCG/Source/PCG/Private/Compute/DataInterfaces/Elements/PCGTransformPointsDataInterface.cpp
new file mode 100644
--- /dev/null
+++ b/Engine/Plugins/PCG/Source/PCG/Private/Compute/DataInterfaces/Elements/PCGTransformPointsDataInterface.cpp
@@ -0,0 +1,260 @@
+// Copyright Epic Games, Inc. All Rights Reserved.
+
+#include "PCGTransformPointsDataInterface.h"
+
+#include "PCGContext.h"
+#include "Compute/PCGComputeGraph.h"
+#include "Compute/PCGDataBinding.h"
+#include "Elements/PCGTransformPoints.h"
+#include "Compute/Elements/PCGTransformPointsKernel.h"
+
+#include "GlobalRenderResources.h"
+#include "RenderGraphBuilder.h"
+#include "RenderGraphResources.h"
+#include "RHIResources.h"
+#include "ShaderCompilerCore.h"
+#include "ShaderCore.h"
+#include "ShaderParameterMetadataBuilder.h"
+#include "Algo/Find.h"
+#include "ComputeFramework/ShaderParamTypeDefinition.h"
+
+#define LOCTEXT_NAMESPACE "PCGTransformPointsDataInterface"
+
+void UPCGTransformPointsDataInterface::GetSupportedInputs(TArray<FShaderFunctionDefinition>& OutFunctions) const
+{
+	OutFunctions.AddDefaulted_GetRef()
+		.SetName(TEXT("TransformPoints_GetOffsetMin"))
+		.AddReturnType(FShaderValueType::Get(EShaderFundamentalType::Float, 3));
+
+	OutFunctions.AddDefaulted_GetRef()
+		.SetName(TEXT("TransformPoints_GetOffsetMax"))
+		.AddReturnType(FShaderValueType::Get(EShaderFundamentalType::Float, 3));
+
+	OutFunctions.AddDefaulted_GetRef()
+		.SetName(TEXT("TransformPoints_GetAbsoluteOffset"))
+		.AddReturnType(FShaderValueType::Get(EShaderFundamentalType::Uint, 1));
+
+	OutFunctions.AddDefaulted_GetRef()
+		.SetName(TEXT("TransformPoints_GetRotationMin"))
+		.AddReturnType(FShaderValueType::Get(EShaderFundamentalType::Float, 3));
+
+	OutFunctions.AddDefaulted_GetRef()
+		.SetName(TEXT("TransformPoints_GetRotationMax"))
+		.AddReturnType(FShaderValueType::Get(EShaderFundamentalType::Float, 3));
+
+	OutFunctions.AddDefaulted_GetRef()
+		.SetName(TEXT("TransformPoints_GetAbsoluteRotation"))
+	.AddReturnType(FShaderValueType::Get(EShaderFundamentalType::Uint, 1));
+
+	OutFunctions.AddDefaulted_GetRef()
+		.SetName(TEXT("TransformPoints_GetScaleMin"))
+		.AddReturnType(FShaderValueType::Get(EShaderFundamentalType::Float, 3));
+
+	OutFunctions.AddDefaulted_GetRef()
+		.SetName(TEXT("TransformPoints_GetScaleMax"))
+		.AddReturnType(FShaderValueType::Get(EShaderFundamentalType::Float, 3));
+
+	OutFunctions.AddDefaulted_GetRef()
+		.SetName(TEXT("TransformPoints_GetAbsoluteScale"))
+		.AddReturnType(FShaderValueType::Get(EShaderFundamentalType::Uint, 1));
+
+	OutFunctions.AddDefaulted_GetRef()
+		.SetName(TEXT("TransformPoints_GetUniformScale"))
+		.AddReturnType(FShaderValueType::Get(EShaderFundamentalType::Uint, 1));
+
+	OutFunctions.AddDefaulted_GetRef()
+		.SetName(TEXT("TransformPoints_GetRecomputeSeed"))
+		.AddReturnType(FShaderValueType::Get(EShaderFundamentalType::Uint, 1));
+
+	OutFunctions.AddDefaulted_GetRef()
+		.SetName(TEXT("TransformPoints_GetApplyToAttribute"))
+		.AddReturnType(FShaderValueType::Get(EShaderFundamentalType::Uint, 1));
+
+	OutFunctions.AddDefaulted_GetRef()
+		.SetName(TEXT("TransformPoints_GetAttributeId"))
+		.AddReturnType(FShaderValueType::Get(EShaderFundamentalType::Int, 1));
+}
+
+BEGIN_SHADER_PARAMETER_STRUCT(FPCGTransformPointsDataInterfaceParameters,)
+	SHADER_PARAMETER(FVector3f, OffsetMin)
+	SHADER_PARAMETER(FVector3f, OffsetMax)
+	SHADER_PARAMETER(uint32, bAbsoluteOffset)
+	SHADER_PARAMETER(FVector3f, RotationMin)
+	SHADER_PARAMETER(FVector3f, RotationMax)
+	SHADER_PARAMETER(uint32, bAbsoluteRotation)
+	SHADER_PARAMETER(FVector3f, ScaleMin)
+	SHADER_PARAMETER(FVector3f, ScaleMax)
+	SHADER_PARAMETER(uint32, bAbsoluteScale)
+	SHADER_PARAMETER(uint32, bUniformScale)
+	SHADER_PARAMETER(uint32, bRecomputeSeed)
+	SHADER_PARAMETER(uint32, bApplyToAttribute)
+	SHADER_PARAMETER(int32, AttributeId)
+END_SHADER_PARAMETER_STRUCT()
+
+void UPCGTransformPointsDataInterface::GetShaderParameters(TCHAR const* UID, FShaderParametersMetadataBuilder& InOutBuilder, FShaderParametersMetadataAllocations& InOutAllocations) const
+{
+	InOutBuilder.AddNestedStruct<FPCGTransformPointsDataInterfaceParameters>(UID);
+}
+
+void UPCGTransformPointsDataInterface::GetHLSL(FString& OutHLSL, FString const& InDataInterfaceName) const
+{
+	TMap<FString, FStringFormatArg> TemplateArgs =
+	{
+		{ TEXT("DataInterfaceName"), InDataInterfaceName },
+	};
+
+	OutHLSL += FString::Format(TEXT(
+		"float3 {DataInterfaceName}_OffsetMin;\n"
+		"float3 {DataInterfaceName}_OffsetMax;\n"
+		"uint {DataInterfaceName}_bAbsoluteOffset;\n"
+		"float3 {DataInterfaceName}_RotationMin;\n"
+		"float3 {DataInterfaceName}_RotationMax;\n"
+		"uint {DataInterfaceName}_bAbsoluteRotation;\n"
+		"float3 {DataInterfaceName}_ScaleMin;\n"
+		"float3 {DataInterfaceName}_ScaleMax;\n"
+		"uint {DataInterfaceName}_bAbsoluteScale;\n"
+		"uint {DataInterfaceName}_bUniformScale;\n"
+		"uint {DataInterfaceName}_bRecomputeSeed;\n"
+		"uint {DataInterfaceName}_bApplyToAttribute;\n"
+		"int {DataInterfaceName}_AttributeId;\n"
+		"\n"
+		"float3 TransformPoints_GetOffsetMin_{DataInterfaceName}() { return {DataInterfaceName}_OffsetMin; }\n"
+		"float3 TransformPoints_GetOffsetMax_{DataInterfaceName}() { return {DataInterfaceName}_OffsetMax; }\n"
+		"uint TransformPoints_GetAbsoluteOffset_{DataInterfaceName}() { return {DataInterfaceName}_bAbsoluteOffset; }\n"
+		"float3 TransformPoints_GetRotationMin_{DataInterfaceName}() { return {DataInterfaceName}_RotationMin; }\n"
+		"float3 TransformPoints_GetRotationMax_{DataInterfaceName}() { return {DataInterfaceName}_RotationMax; }\n"
+		"uint TransformPoints_GetAbsoluteRotation_{DataInterfaceName}() { return {DataInterfaceName}_bAbsoluteRotation; }\n"
+		"float3 TransformPoints_GetScaleMin_{DataInterfaceName}() { return {DataInterfaceName}_ScaleMin; }\n"
+		"float3 TransformPoints_GetScaleMax_{DataInterfaceName}() { return {DataInterfaceName}_ScaleMax; }\n"
+		"uint TransformPoints_GetAbsoluteScale_{DataInterfaceName}() { return {DataInterfaceName}_bAbsoluteScale; }\n"
+		"uint TransformPoints_GetUniformScale_{DataInterfaceName}() { return {DataInterfaceName}_bUniformScale; }\n"
+		"uint TransformPoints_GetRecomputeSeed_{DataInterfaceName}() { return {DataInterfaceName}_bRecomputeSeed; }\n"
+		"uint TransformPoints_GetApplyToAttribute_{DataInterfaceName}() { return {DataInterfaceName}_bApplyToAttribute; }\n"
+		"int TransformPoints_GetAttributeId_{DataInterfaceName}() { return {DataInterfaceName}_AttributeId; }\n"
+		), TemplateArgs);
+}
+
+UComputeDataProvider* UPCGTransformPointsDataInterface::CreateDataProvider() const
+{
+	return NewObject<UPCGTransformPointsDataProvider>();
+}
+
+FComputeDataProviderRenderProxy* UPCGTransformPointsDataProvider::GetRenderProxy()
+{
+	const UPCGTransformPointsSettings* Settings = CastChecked<UPCGTransformPointsSettings>(ProducerKernel->GetSettings());
+
+	FPCGTransformPointsDataProviderProxy::FData ProxyData;
+
+	ProxyData.OffsetMin = Settings->OffsetMin;
+	ProxyData.OffsetMax = Settings->OffsetMax;
+	ProxyData.bAbsoluteOffset = Settings->bAbsoluteOffset;
+	ProxyData.RotationMin = Settings->RotationMin.Euler();
+	ProxyData.RotationMax = Settings->RotationMax.Euler();
+	ProxyData.bAbsoluteRotation = Settings->bAbsoluteRotation;
+	ProxyData.ScaleMin = Settings->ScaleMin;
+	ProxyData.ScaleMax = Settings->ScaleMax;
+	ProxyData.bAbsoluteScale = Settings->bAbsoluteScale;
+	ProxyData.bUniformScale = Settings->bUniformScale;
+	ProxyData.bRecomputeSeed = Settings->bRecomputeSeed;
+	ProxyData.bApplyToAttribute = Settings->bApplyToAttribute;
+	ProxyData.AttributeId = AttributeId;
+
+	return new FPCGTransformPointsDataProviderProxy(MoveTemp(ProxyData));
+}
+
+void UPCGTransformPointsDataProvider::Reset()
+{
+	Super::Reset();
+
+	AttributeId = INDEX_NONE;
+}
+
+bool UPCGTransformPointsDataProvider::PrepareForExecute_GameThread(UPCGDataBinding* InBinding)
+{
+	if (!Super::PrepareForExecute_GameThread(InBinding))
+	{
+		return false;
+	}
+
+	TSharedPtr<FPCGContextHandle> ContextHandle = InBinding->ContextHandle.Pin();
+	FPCGContext* Context = ContextHandle ? ContextHandle->GetContext() : nullptr;
+	if (!ensure(Context))
+	{
+		return true;
+	}
+
+	const UPCGTransformPointsSettings* Settings = CastChecked<UPCGTransformPointsSettings>(ProducerKernel->GetSettings());
+
+	if (Settings->bApplyToAttribute && AttributeId == INDEX_NONE)
+	{
+		const FPCGDataCollectionDesc* InputDataDesc = InBinding->GetCachedKernelPinDataDesc(ProducerKernel, PCGPinConstants::DefaultInputLabel, /*bIsInput=*/true);
+
+		if (!ensure(InputDataDesc))
+		{
+			return true;
+		}
+
+		bool bAnyPointsPresent = false;
+
+		for (const FPCGDataDesc& Desc : InputDataDesc->DataDescs)
+		{
+			if (Desc.ElementCount <= 0)
+			{
+				continue;
+			}
+
+			bAnyPointsPresent = true;
+			
+			const FPCGKernelAttributeDesc* It = Algo::FindByPredicate(Desc.AttributeDescs, [Settings](const FPCGKernelAttributeDesc& AttributeDesc)
+			{
+				return AttributeDesc.AttributeKey.Identifier.Name == Settings->AttributeName && AttributeDesc.AttributeKey.Type == EPCGKernelAttributeType::Transform;
+			});
+			
+			if (It)
+			{
+				AttributeId = It->AttributeId;
+				break;
+			}
+		}
+
+		// Mute this error if the point data is empty.
+		if (AttributeId == INDEX_NONE && !InputDataDesc->DataDescs.IsEmpty() && bAnyPointsPresent)
+		{
+			PCG_KERNEL_VALIDATION_ERR(Context, Settings, FText::Format(
+				LOCTEXT("TransformPointsAttributeNotFound", "Transfrom points attribute '{0}' not found."),
+				FText::FromName(Settings->AttributeName)));
+		}
+	}
+
+	return true;
+}
+
+bool FPCGTransformPointsDataProviderProxy::IsValid(FValidationData const& InValidationData) const
+{
+	return InValidationData.ParameterStructSize == sizeof(FParameters);
+}
+
+void FPCGTransformPointsDataProviderProxy::GatherDispatchData(FDispatchData const& InDispatchData)
+{
+	const TStridedView<FParameters> ParameterArray = MakeStridedParameterView<FParameters>(InDispatchData);
+
+	for (int32 InvocationIndex = 0; InvocationIndex < ParameterArray.Num(); ++InvocationIndex)
+	{
+		FParameters& Parameters = ParameterArray[InvocationIndex];
+		Parameters.OffsetMin = static_cast<FVector3f>(Data.OffsetMin);
+		Parameters.OffsetMax = static_cast<FVector3f>(Data.OffsetMax);
+		Parameters.bAbsoluteOffset = Data.bAbsoluteOffset;
+		Parameters.RotationMin = static_cast<FVector3f>(Data.RotationMin);
+		Parameters.RotationMax = static_cast<FVector3f>(Data.RotationMax);
+		Parameters.bAbsoluteRotation = Data.bAbsoluteRotation;
+		Parameters.ScaleMin = static_cast<FVector3f>(Data.ScaleMin);
+		Parameters.ScaleMax = static_cast<FVector3f>(Data.ScaleMax);
+		Parameters.bAbsoluteScale = Data.bAbsoluteScale;
+		Parameters.bUniformScale = Data.bUniformScale;
+		Parameters.bRecomputeSeed = Data.bRecomputeSeed;
+		Parameters.bApplyToAttribute = Data.bApplyToAttribute;
+		Parameters.AttributeId = Data.AttributeId;
+	}
+}
+
+#undef LOCTEXT_NAMESPACE
\ No newline at end of file
diff --git a/Engine/Plugins/PCG/Source/PCG/Private/Compute/DataInterfaces/Elements/PCGTransformPointsDataInterface.h b/Engine/Plugins/PCG/Source/PCG/Private/Compute/DataInterfaces/Elements/PCGTransformPointsDataInterface.h
new file mode 100644
--- /dev/null
+++ b/Engine/Plugins/PCG/Source/PCG/Private/Compute/DataInterfaces/Elements/PCGTransformPointsDataInterface.h
@@ -0,0 +1,77 @@
+// Copyright Epic Games, Inc. All Rights Reserved.
+
+#pragma once
+
+#include "Compute/DataInterfaces/PCGComputeDataInterface.h"
+
+#include "PCGTransformPointsDataInterface.generated.h"
+
+class FPCGTransformPointsDataInterfaceParameters;
+
+/** Data Interface to marshal transform points settings to the GPU. */
+UCLASS(ClassGroup = (Procedural))
+class UPCGTransformPointsDataInterface : public UPCGComputeDataInterface
+{
+	GENERATED_BODY()
+
+public:
+	//~ Begin UComputeDataInterface Interface
+	TCHAR const* GetClassName() const override { return TEXT("PCGTransformPointsDataInterface"); }
+	void GetSupportedInputs(TArray<FShaderFunctionDefinition>& OutFunctions) const override;
+	void GetShaderParameters(TCHAR const* UID, FShaderParametersMetadataBuilder& InOutBuilder, FShaderParametersMetadataAllocations& InOutAllocations) const override;
+	void GetHLSL(FString& OutHLSL, FString const& InDataInterfaceName) const override;
+	UComputeDataProvider* CreateDataProvider() const override;
+	//~ End UComputeDataInterface Interface
+};
+
+UCLASS()
+class UPCGTransformPointsDataProvider : public UPCGComputeDataProvider
+{
+	GENERATED_BODY()
+
+public:
+	//~ Begin UComputeDataProvider Interface
+	FComputeDataProviderRenderProxy* GetRenderProxy() override;
+	virtual void Reset() override;
+	virtual bool PrepareForExecute_GameThread(UPCGDataBinding* InBinding) override;
+	//~ End UComputeDataProvider Interface
+
+	/** Output attribute Id. */
+	UPROPERTY()
+	int32 AttributeId = INDEX_NONE;
+};
+
+class FPCGTransformPointsDataProviderProxy : public FComputeDataProviderRenderProxy
+{
+public:
+	struct FData
+	{
+		FVector OffsetMin = FVector::ZeroVector;
+		FVector OffsetMax = FVector::ZeroVector;
+		uint32 bAbsoluteOffset = 0;
+		FVector RotationMin = FVector::ZeroVector;
+		FVector RotationMax = FVector::ZeroVector;
+		uint32 bAbsoluteRotation = 0;
+		FVector ScaleMin = FVector::OneVector;
+		FVector ScaleMax = FVector::OneVector;
+		uint32 bAbsoluteScale = 0;
+		uint32 bUniformScale = 1;
+		uint32 bRecomputeSeed = 0;
+		uint32 bApplyToAttribute = 0;
+		int32 AttributeId = INDEX_NONE;
+	};
+
+	FPCGTransformPointsDataProviderProxy(FData InData)
+		: Data(MoveTemp(InData))
+	{}
+
+	//~ Begin FComputeDataProviderRenderProxy Interface
+	bool IsValid(FValidationData const& InValidationData) const override;
+	void GatherDispatchData(FDispatchData const& InDispatchData) override;
+	//~ End FComputeDataProviderRenderProxy Interface
+
+protected:
+	using FParameters = FPCGTransformPointsDataInterfaceParameters;
+
+	FData Data;
+};
diff --git a/Engine/Plugins/PCG/Source/PCG/Private/Compute/Elements/PCGTransformPointsKernel.cpp b/Engine/Plugins/PCG/Source/PCG/Private/Compute/Elements/PCGTransformPointsKernel.cpp
new file mode 100644
--- /dev/null
+++ b/Engine/Plugins/PCG/Source/PCG/Private/Compute/Elements/PCGTransformPointsKernel.cpp
@@ -0,0 +1,90 @@
+// Copyright Epic Games, Inc. All Rights Reserved.
+
+#include "Compute/Elements/PCGTransformPointsKernel.h"
+
+#include "PCGContext.h"
+#include "Elements/PCGTransformPoints.h"
+#include "Graph/PCGGPUGraphCompilationContext.h"
+#include "Compute/Elements/PCGComputeGraphElement.h"
+#include "Compute/DataInterfaces/Elements/PCGTransformPointsDataInterface.h"
+
+#include "ShaderCompilerCore.h"
+
+#define LOCTEXT_NAMESPACE "PCGTransformPointsKernel"
+
+bool UPCGTransformPointsKernel::ComputeOutputBindingDataDesc(const UPCGComputeGraph* InGraph, FName InOutputPinLabel, UPCGDataBinding* InBinding, FPCGDataCollectionDesc& OutDataDesc) const
+{
+	check(InGraph);
+	check(InBinding);
+
+	// Code assumes single output pin.
+	if (!ensure(InOutputPinLabel == PCGPinConstants::DefaultOutputLabel))
+	{
+		return false;
+	}
+
+	const FPCGKernelPin SourceKernelPin(KernelIndex, PCGPinConstants::DefaultInputLabel, /*bIsInput=*/true);
+	ensure(InBinding->ComputeKernelPinDataDesc(SourceKernelPin, &OutDataDesc));
+	
+	const UPCGTransformPointsSettings* Settings = CastChecked<UPCGTransformPointsSettings>(GetSettings());
+	if (!ensure(Settings))
+	{
+		return false;
+	}
+
+	// Allocate properties if needed
+	for (FPCGDataDesc& DataDesc : OutDataDesc.DataDescs)
+	{
+		DataDesc.AllocateProperties(Settings->GetPropertiesToAllocate());
+	}
+
+	return true;
+}
+
+int UPCGTransformPointsKernel::ComputeThreadCount(const UPCGDataBinding* InBinding) const
+{
+	int ThreadCount = 0;
+
+	const FPCGDataCollectionDesc* OutputPinDesc = InBinding->GetCachedKernelPinDataDesc(this, PCGPinConstants::DefaultOutputLabel, /*bIsInput=*/false);
+
+	if (ensure(OutputPinDesc))
+	{
+		for (const FPCGDataDesc& DataDesc : OutputPinDesc->DataDescs)
+		{
+			ThreadCount += DataDesc.ElementCount;
+		}
+	}
+
+	return ThreadCount;
+}
+
+#if WITH_EDITOR
+FString UPCGTransformPointsKernel::GetCookedSource(FPCGGPUCompilationContext& InOutContext) const
+{
+	FString TemplateFile;
+	ensure(LoadShaderSourceFile(TEXT("/Plugin/PCG/Private/Elements/PCGTransformPoints.usf"), EShaderPlatform::SP_PCD3D_SM5, &TemplateFile, nullptr));
+	return TemplateFile;
+}
+
+void UPCGTransformPointsKernel::CreateAdditionalInputDataInterfaces(FPCGGPUCompilationContext& InOutContext, UObject* InObjectOuter, TArray<TObjectPtr<UComputeDataInterface>>& OutDataInterfaces) const
+{
+	Super::CreateAdditionalInputDataInterfaces(InOutContext, InObjectOuter, OutDataInterfaces);
+
+	TObjectPtr<UPCGTransformPointsDataInterface> NodeDI = InOutContext.NewObject_AnyThread<UPCGTransformPointsDataInterface>(InObjectOuter);
+	NodeDI->ProducerKernel = this;
+
+	OutDataInterfaces.Add(NodeDI);
+}
+#endif
+
+void UPCGTransformPointsKernel::GetInputPins(TArray<FPCGPinProperties>& OutPins) const
+{
+	OutPins.Emplace(PCGPinConstants::DefaultInputLabel, EPCGDataType::Point);
+}
+
+void UPCGTransformPointsKernel::GetOutputPins(TArray<FPCGPinPropertiesGPU>& OutPins) const
+{
+	OutPins.Emplace(PCGPinConstants::DefaultOutputLabel, EPCGDataType::Point);
+}
+
+#undef LOCTEXT_NAMESPACE
diff --git a/Engine/Plugins/PCG/Source/PCG/Private/Compute/Elements/PCGTransformPointsKernel.h b/Engine/Plugins/PCG/Source/PCG/Private/Compute/Elements/PCGTransformPointsKernel.h
new file mode 100644
--- /dev/null
+++ b/Engine/Plugins/PCG/Source/PCG/Private/Compute/Elements/PCGTransformPointsKernel.h
@@ -0,0 +1,28 @@
+// Copyright Epic Games, Inc. All Rights Reserved.
+
+#pragma once
+
+#include "Compute/PCGComputeKernel.h"
+
+#include "Compute/PCGPinPropertiesGPU.h"
+
+#include "PCGTransformPointsKernel.generated.h"
+
+UCLASS()
+class UPCGTransformPointsKernel : public UPCGComputeKernel
+{
+	GENERATED_BODY()
+
+public:
+	//~Begin UPCGComputeKernel interface
+	virtual bool ComputeOutputBindingDataDesc(const UPCGComputeGraph* InGraph, FName InOutputPinLabel, UPCGDataBinding* InBinding, FPCGDataCollectionDesc& OutDataDesc) const override;
+	virtual int ComputeThreadCount(const UPCGDataBinding* Binding) const override;
+#if WITH_EDITOR
+	virtual FString GetCookedSource(FPCGGPUCompilationContext& InOutContext) const override;
+	virtual FString GetEntryPoint() const override { return TEXT("PCGTransformPointsCS"); }
+	void CreateAdditionalInputDataInterfaces(FPCGGPUCompilationContext& InOutContext, UObject* InObjectOuter, TArray<TObjectPtr<UComputeDataInterface>>& OutDataInterfaces) const override;
+#endif
+	virtual void GetInputPins(TArray<FPCGPinProperties>& OutPins) const override;
+	virtual void GetOutputPins(TArray<FPCGPinPropertiesGPU>& OutPins) const override;
+	//~End UPCGComputeKernel interface
+};
diff --git a/Engine/Plugins/PCG/Source/PCG/Private/Elements/PCGTransformPoints.cpp b/Engine/Plugins/PCG/Source/PCG/Private/Elements/PCGTransformPoints.cpp
--- a/Engine/Plugins/PCG/Source/PCG/Private/Elements/PCGTransformPoints.cpp
+++ b/Engine/Plugins/PCG/Source/PCG/Private/Elements/PCGTransformPoints.cpp
@@ -9,17 +9,51 @@
 #include "Helpers/PCGAsync.h"
 #include "Helpers/PCGHelpers.h"
 
+#if WITH_EDITOR
+#include "Compute/PCGComputeKernel.h"
+#include "Compute/Elements/PCGTransformPointsKernel.h"
+#include "Graph/PCGGPUGraphCompilationContext.h"
+#endif // WITH_EDITOR
+
 #include "Math/RandomStream.h"
 
 #include UE_INLINE_GENERATED_CPP_BY_NAME(PCGTransformPoints)
 
 #define LOCTEXT_NAMESPACE "PCGTransformPointsElement"
 
+class UPCGTransformPointsKernel;
+
+#if WITH_EDITOR
+/**
+ * TODO: This function might be interesting to factorized in the base PCGSettings class, or even into a UInterface so we have those boilerplate functions
+ * already defined since most of those nodes would work with the default implementation.
+ */
+void UPCGTransformPointsSettings::CreateKernels(FPCGGPUCompilationContext& InOutContext, UObject* InObjectOuter,
+                                                TArray<UPCGComputeKernel*>& OutKernels, TArray<FPCGKernelEdge>& OutEdges) const
+{
+	FPCGComputeKernelParams KernelParams;
+	KernelParams.Settings = this;
+	KernelParams.bLogDescriptions = bDumpDataDescriptions;
+
+	UPCGTransformPointsKernel* Kernel = InOutContext.NewObject_AnyThread<UPCGTransformPointsKernel>(InObjectOuter);
+	Kernel->Initialize(KernelParams);
+	OutKernels.Add(Kernel);
+	
+	OutEdges.Emplace(FPCGPinReference(PCGPinConstants::DefaultInputLabel), FPCGPinReference(Kernel, PCGPinConstants::DefaultInputLabel));
+	OutEdges.Emplace(FPCGPinReference(Kernel, PCGPinConstants::DefaultOutputLabel), FPCGPinReference(PCGPinConstants::DefaultOutputLabel));
+}
+#endif // WITH_EDITOR
+
 FPCGElementPtr UPCGTransformPointsSettings::CreateElement() const
 {
 	return MakeShared<FPCGTransformPointsElement>();
 }
 
+EPCGPointNativeProperties UPCGTransformPointsSettings::GetPropertiesToAllocate() const
+{
+	return bApplyToAttribute ? EPCGPointNativeProperties::MetadataEntry : (bRecomputeSeed ? (EPCGPointNativeProperties::Transform | EPCGPointNativeProperties::Seed) : EPCGPointNativeProperties::Transform);
+}
+
 bool FPCGTransformPointsElement::ExecuteInternal(FPCGContext* Context) const
 {
 	TRACE_CPUPROFILER_EVENT_SCOPE(FPCGTransformPointsElement::Execute);
@@ -99,7 +133,7 @@ bool FPCGTransformPointsElement::ExecuteInternal(FPCGContext* Context) const
 		OutputData->InitializeFromData(PointData);
 		OutputData->SetNumPoints(PointData->GetNumPoints(), /*bInitializeValues=*/false);
 
-		const EPCGPointNativeProperties PropertiesToAllocate = bApplyToAttribute ? EPCGPointNativeProperties::MetadataEntry : (bRecomputeSeed ? (EPCGPointNativeProperties::Transform | EPCGPointNativeProperties::Seed) : EPCGPointNativeProperties::Transform);
+		const EPCGPointNativeProperties PropertiesToAllocate = Settings->GetPropertiesToAllocate();
 
 		if (OutputData->HasSpatialDataParent())
 		{
diff --git a/Engine/Plugins/PCG/Source/PCG/Public/Elements/PCGTransformPoints.h b/Engine/Plugins/PCG/Source/PCG/Public/Elements/PCGTransformPoints.h
--- a/Engine/Plugins/PCG/Source/PCG/Public/Elements/PCGTransformPoints.h
+++ b/Engine/Plugins/PCG/Source/PCG/Public/Elements/PCGTransformPoints.h
@@ -6,9 +6,15 @@
 
 #include "PCGTransformPoints.generated.h"
 
+class FPCGTransformPointsElement;
+class UPCGTransformPointsKernel;
+
 UCLASS(MinimalAPI, BlueprintType, ClassGroup = (Procedural))
 class UPCGTransformPointsSettings : public UPCGSettings
 {
+	friend UPCGTransformPointsKernel;
+	friend FPCGTransformPointsElement;
+	
 	GENERATED_BODY()
 
 public:
@@ -17,6 +23,8 @@ public:
 	virtual FName GetDefaultNodeName() const override { return FName(TEXT("TransformPoints")); }
 	virtual FText GetDefaultNodeTitle() const override { return NSLOCTEXT("PCGTransformPointsSettings", "NodeTitle", "Transform Points"); }
 	virtual EPCGSettingsType GetType() const override { return EPCGSettingsType::PointOps; }
+	virtual bool DisplayExecuteOnGPUSetting() const override { return true; }
+	virtual void CreateKernels(FPCGGPUCompilationContext& InOutContext, UObject* InObjectOuter, TArray<UPCGComputeKernel*>& OutKernels, TArray<FPCGKernelEdge>& OutEdges) const override;
 #endif
 	virtual bool UseSeed() const override { return true; }
 
@@ -25,12 +33,15 @@ protected:
 	virtual TArray<FPCGPinProperties> OutputPinProperties() const override { return Super::DefaultPointOutputPinProperties(); }
 	virtual FPCGElementPtr CreateElement() const override;
 	//~End UPCGSettings interface
+	
+	EPCGPointNativeProperties GetPropertiesToAllocate() const;
 
 public:
+	/** If the operation should be applied to a Transform attribute instead of the point transform. */
 	UPROPERTY(BlueprintReadWrite, EditAnywhere, Category = Settings, meta = (PCG_Overridable))
 	bool bApplyToAttribute = false;
 
-	UPROPERTY(BlueprintReadWrite, EditAnywhere, Category = Settings, meta = (EditCondition="bApplyToAttribute", PCG_Overridable))
+	UPROPERTY(BlueprintReadWrite, EditAnywhere, Category = Settings, meta = (EditCondition="bApplyToAttribute", EditConditionHides, PCG_Overridable))
 	FName AttributeName;
 
 	UPROPERTY(BlueprintReadWrite, EditAnywhere, Category = Settings, meta = (PCG_Overridable))
