[PCG] Attribute Rename support for Selectors and Data domain

* Add deprecation path for previous behavior
* Add unit tests

#rnx

diff --git a/Engine/Plugins/PCG/Source/PCG/Private/Elements/Metadata/PCGMetadataRenameElement.cpp b/Engine/Plugins/PCG/Source/PCG/Private/Elements/Metadata/PCGMetadataRenameElement.cpp
--- a/Engine/Plugins/PCG/Source/PCG/Private/Elements/Metadata/PCGMetadataRenameElement.cpp
+++ b/Engine/Plugins/PCG/Source/PCG/Private/Elements/Metadata/PCGMetadataRenameElement.cpp
@@ -3,6 +3,7 @@
 #include "Elements/Metadata/PCGMetadataRenameElement.h"
 
 #include "PCGContext.h"
+#include "PCGCustomVersion.h"
 #include "PCGPin.h"
 #include "PCGParamData.h"
 #include "Data/PCGSpatialData.h"
@@ -12,6 +13,24 @@
 
 #define LOCTEXT_NAMESPACE "PCGMetadataRenameElement"
 
+UPCGMetadataRenameSettings::UPCGMetadataRenameSettings(const FObjectInitializer& ObjectInitializer)
+{
+	AttributeToRename.SetAttributeName(NAME_None);
+}
+
+void UPCGMetadataRenameSettings::PostLoad()
+{
+	Super::PostLoad();
+
+#if WITH_EDITOR
+	if (DataVersion < FPCGCustomVersion::AttributeRenameSupportToSelectors && AttributeToRename.GetAttributeName() == NAME_None)
+	{
+		// Previous behavior was that None -> Last created.
+		AttributeToRename.SetAttributeName(PCGMetadataAttributeConstants::LastCreatedAttributeName);
+	}
+#endif // WITH_EDITOR
+}
+
 FString UPCGMetadataRenameSettings::GetAdditionalTitleInformation() const
 {
 #if WITH_EDITOR
@@ -49,52 +68,48 @@ bool FPCGMetadataRenameElement::ExecuteInternal(FPCGContext* Context) const
 	TArray<FPCGTaggedData> Inputs = Context->InputData.GetInputsByPin(PCGPinConstants::DefaultInputLabel);
 	TArray<FPCGTaggedData>& Outputs = Context->OutputData.TaggedData;
 	
-	const FName AttributeToRename = Settings->AttributeToRename;
 	const FName NewAttributeName = Settings->NewAttributeName;
 
-	if (NewAttributeName == NAME_None)
-	{
-		PCGE_LOG(Warning, GraphAndLog, FText::Format(LOCTEXT("InvalidTargetAttributeName", "Invalid target attribute name '{0}'"), FText::FromName(NewAttributeName)));
-		// Bypass
-		Context->OutputData = Context->InputData;
-		return true;
-	}
-
 	for (const FPCGTaggedData& Input : Inputs)
 	{
 		FPCGTaggedData& Output = Outputs.Add_GetRef(Input);
-
-		const UPCGMetadata* Metadata = nullptr;
-
-		if (const UPCGSpatialData* SpatialInput = Cast<const UPCGSpatialData>(Input.Data))
+		
+		const UPCGMetadata* Metadata = Input.Data ? Input.Data->ConstMetadata() : nullptr;
+		
+		if (!Metadata)
 		{
-			Metadata = SpatialInput->Metadata;
-			Output.Data = SpatialInput;
+			PCGLog::Metadata::LogInvalidMetadata(Context);
+			continue;
 		}
-		else if (const UPCGParamData* ParamInput = Cast<UPCGParamData>(Input.Data))
+		
+		const FPCGAttributePropertyInputSelector Selector = Settings->AttributeToRename.CopyAndFixLast(Input.Data);
+
+		if (!Selector.IsBasicAttribute())
 		{
-			Metadata = ParamInput->Metadata;
-			Output.Data = ParamInput;
+			PCGE_LOG(Warning, GraphAndLog, FText::Format(LOCTEXT("InvalidInputName", "Attribute to rename '{0}' is not an attribute."), Selector.GetDisplayText()));
+			continue;
 		}
-
-		if (!Metadata)
+		
+		const FPCGMetadataDomain* MetadataDomain = Metadata->GetConstMetadataDomainFromSelector(Selector);
+		
+		if (!MetadataDomain)
 		{
-			PCGE_LOG(Warning, GraphAndLog, LOCTEXT("InvalidInputType", "Input data type is not supported, only supports Spatial and Attribute Set"));
-			Output.Data = nullptr;
+			PCGLog::Metadata::LogInvalidMetadataDomain(Selector, Context);
 			continue;
 		}
 
-		const FName LocalAttributeToRename = ((AttributeToRename != NAME_None) ? AttributeToRename : Metadata->GetLatestAttributeNameOrNone());
+		const FName LocalAttributeToRename = Selector.GetAttributeName();
 
-		if (!Metadata->HasAttribute(LocalAttributeToRename))
+		if (!MetadataDomain->HasAttribute(LocalAttributeToRename))
 		{
 			continue;
 		}
 
 		UPCGMetadata* NewMetadata = nullptr;
 		PCGMetadataElementCommon::DuplicateTaggedData(Context, Input, Output, NewMetadata);
+		FPCGMetadataDomain* NewMetadataDomain = NewMetadata ? NewMetadata->GetMetadataDomainFromSelector(Selector) : nullptr;
 
-		if (!NewMetadata || !NewMetadata->RenameAttribute(LocalAttributeToRename, NewAttributeName))
+		if (!NewMetadataDomain || !NewMetadataDomain->RenameAttribute(LocalAttributeToRename, NewAttributeName))
 		{
 			PCGE_LOG(Warning, GraphAndLog, FText::Format(LOCTEXT("AttributeRenamedFailed", "Failed to rename attribute from '{0}' to '{1}'"),
 				FText::FromName(LocalAttributeToRename), FText::FromName(NewAttributeName)));
diff --git a/Engine/Plugins/PCG/Source/PCG/Private/Tests/Elements/Metadata/PCGMetadataRenameTests.cpp b/Engine/Plugins/PCG/Source/PCG/Private/Tests/Elements/Metadata/PCGMetadataRenameTests.cpp
new file mode 100644
--- /dev/null
+++ b/Engine/Plugins/PCG/Source/PCG/Private/Tests/Elements/Metadata/PCGMetadataRenameTests.cpp
@@ -0,0 +1,234 @@
+// Copyright Epic Games, Inc. All Rights Reserved.
+
+#include "Tests/PCGTestsCommon.h"
+
+#include "PCGComponent.h"
+#include "PCGContext.h"
+#include "PCGCustomVersion.h"
+#include "PCGParamData.h"
+#include "Algo/Count.h"
+
+#include "Elements/Metadata/PCGMetadataRenameElement.h"
+
+#if WITH_EDITOR
+
+class FPCGMetadataRenameTest : public FPCGTestBaseClass
+{
+	using FPCGTestBaseClass::FPCGTestBaseClass;
+
+protected:
+	bool Run()
+	{
+		PCGTestsCommon::FTestData TestData;
+		PCGTestsCommon::GenerateSettings<UPCGMetadataRenameSettings>(TestData);
+		UPCGMetadataRenameSettings* Settings = CastChecked<UPCGMetadataRenameSettings>(TestData.Settings);
+
+		Settings->AttributeToRename = InSelector;
+		Settings->NewAttributeName = OutName;
+
+		if (bDeprecation)
+		{
+			// Apply the deprecation on the settings, by forcing a previous version and calling postload
+			Settings->DataVersion = FPCGCustomVersion::AttributePropertySelectorDeprecatePointProperties;
+			Settings->PostLoad();
+		}
+		
+		FPCGElementPtr MetadataBooleanElement = TestData.Settings->GetElement();
+
+		UPCGParamData* ParamData = NewObject<UPCGParamData>();
+		check(ParamData);
+		UPCGMetadata* Metadata = ParamData->MutableMetadata();
+		FPCGMetadataDomain* MetadataDomain = Metadata ? Metadata->GetMetadataDomain(ExpectedDomain) : nullptr;
+		check(MetadataDomain || ExpectedDomain == PCGMetadataDomainID::Invalid);
+
+		if (ExpectedDomain != PCGMetadataDomainID::Invalid)
+		{
+			for (const FName Attribute : ExistingAttributes)
+			{
+				MetadataDomain->CreateAttribute<float>(Attribute, 0.0f, /*bAllowInterpolation=*/true, /*bOverrideParent=*/true);
+			}
+		}
+
+		TUniquePtr<FPCGContext> Context = TestData.InitializeTestContext();
+		check(Context.IsValid());
+		
+		FPCGTaggedData& InputData = Context->InputData.TaggedData.Emplace_GetRef();
+		InputData.Data = ParamData;
+		InputData.Pin = PCGPinConstants::DefaultInputLabel;
+
+		while (!MetadataBooleanElement->Execute(Context.Get())){}
+
+		UTEST_EQUAL("Output data is present", Context->OutputData.TaggedData.Num(), 1);
+		const UPCGParamData* OutputParamData = Cast<UPCGParamData>(Context->OutputData.TaggedData[0].Data);
+		
+		UTEST_NOT_NULL("Output data is a param data", OutputParamData);
+		check(OutputParamData);
+
+		const UPCGMetadata* OutputMetadata = OutputParamData->ConstMetadata();
+		const FPCGMetadataDomain* OutputMetadataDomain = OutputMetadata ? OutputMetadata->GetConstMetadataDomain(ExpectedDomain) : nullptr;
+
+		if (!OutputMetadataDomain && ExpectedDomain == PCGMetadataDomainID::Invalid)
+		{
+			return true;
+		}
+		
+		check(OutputMetadataDomain);
+		
+		TArray<FName> AttributeNames;
+		TArray<EPCGMetadataTypes> AttributeTypes;
+		OutputMetadataDomain->GetAttributes(AttributeNames, AttributeTypes);
+		
+		UTEST_EQUAL(*FString::Printf(TEXT("Metadata has %d attributes"), ExistingAttributes.Num()), AttributeNames.Num(), ExistingAttributes.Num())
+		const int32 NumAttributesMatchExisting = Algo::CountIf(AttributeNames, [this](const FName AttributeName) { return ExistingAttributes.Contains(AttributeName); });
+
+		if (bValid)
+		{
+			UTEST_TRUE("Attribute was renamed correctly", OutputMetadataDomain->HasAttribute(OutName) && !OutputMetadataDomain->HasAttribute(ExpectedAttributeToBeRenamed));
+			UTEST_EQUAL("Other attributes are untouched", NumAttributesMatchExisting, ExistingAttributes.Num() - 1);
+		}
+		else
+		{
+			UTEST_EQUAL("Attributes are untouched", NumAttributesMatchExisting, ExistingAttributes.Num());
+		}
+
+		return true;
+	}
+
+	TArray<FName> ExistingAttributes = {TEXT(""), TEXT("Attr1"), TEXT("Attr2")};
+	FPCGAttributePropertyInputSelector InSelector;
+	FName ExpectedAttributeToBeRenamed;
+	FPCGMetadataDomainID ExpectedDomain = PCGMetadataDomainID::Default;
+	FName OutName = TEXT("NewAttr");
+	bool bDeprecation = false;
+	bool bValid = true;
+};
+
+IMPLEMENT_CUSTOM_SIMPLE_AUTOMATION_TEST(FPCGMetadataRenameNoneTest, FPCGMetadataRenameTest, "Plugins.PCG.Metadata.Rename.None", PCGTestsCommon::TestFlags)
+IMPLEMENT_CUSTOM_SIMPLE_AUTOMATION_TEST(FPCGMetadataRenameAtLastTest, FPCGMetadataRenameTest, "Plugins.PCG.Metadata.Rename.AtLast", PCGTestsCommon::TestFlags)
+IMPLEMENT_CUSTOM_SIMPLE_AUTOMATION_TEST(FPCGMetadataRenameNormalTest, FPCGMetadataRenameTest, "Plugins.PCG.Metadata.Rename.Normal", PCGTestsCommon::TestFlags)
+IMPLEMENT_CUSTOM_SIMPLE_AUTOMATION_TEST(FPCGMetadataRenameInvalidTest, FPCGMetadataRenameTest, "Plugins.PCG.Metadata.Rename.Invalid", PCGTestsCommon::TestFlags)
+IMPLEMENT_CUSTOM_SIMPLE_AUTOMATION_TEST(FPCGMetadataRenameDataDomainTest, FPCGMetadataRenameTest, "Plugins.PCG.Metadata.Rename.DataDomain", PCGTestsCommon::TestFlags)
+IMPLEMENT_CUSTOM_SIMPLE_AUTOMATION_TEST(FPCGMetadataRenameInvalidDataDomainTest, FPCGMetadataRenameTest, "Plugins.PCG.Metadata.Rename.InvalidDataDomain", PCGTestsCommon::TestFlags)
+IMPLEMENT_CUSTOM_SIMPLE_AUTOMATION_TEST(FPCGMetadataRenameDeprecationTest, FPCGMetadataRenameTest, "Plugins.PCG.Metadata.Rename.Deprecation", PCGTestsCommon::TestFlags)
+
+bool FPCGMetadataRenameNoneTest::RunTest(const FString& Parameters)
+{
+	InSelector.SetAttributeName(NAME_None);
+	ExpectedAttributeToBeRenamed = NAME_None;
+	
+	return Run();
+}
+
+bool FPCGMetadataRenameAtLastTest::RunTest(const FString& Parameters)
+{
+	InSelector.SetAttributeName(TEXT("@Last"));
+	ExpectedAttributeToBeRenamed = TEXT("Attr2");
+	
+	return Run();
+}
+
+bool FPCGMetadataRenameNormalTest::RunTest(const FString& Parameters)
+{
+	InSelector.SetAttributeName(TEXT("Attr1"));
+	ExpectedAttributeToBeRenamed = TEXT("Attr1");
+	
+	return Run();
+}
+
+bool FPCGMetadataRenameInvalidTest::RunTest(const FString& Parameters)
+{
+	// Property are not supported
+	AddInfo(TEXT("Testing property in selector is not supported."));
+	InSelector.SetPropertyName("Position");
+	bValid = false;
+
+	AddExpectedError(TEXT("Attribute to rename '$Position' is not an attribute."), EAutomationExpectedMessageFlags::Contains, 1, /*bIsRegex=*/false);
+	
+	if (!Run())
+	{
+		return false;
+	}
+
+	// Extra names are not supported
+	AddInfo(TEXT("Testing extra names in selector is not supported."));
+	InSelector.SetAttributeName("Attr1");
+	InSelector.GetExtraNamesMutable().Emplace(TEXT("Foo"));
+
+	AddExpectedError(TEXT("Attribute to rename 'Attr1.Foo' is not an attribute."), EAutomationExpectedMessageFlags::Contains, 1, /*bIsRegex=*/false);
+	
+	if (!Run())
+	{
+		return false;
+	}
+
+	// If the attribute doesn't exist, nothing is done.
+	AddInfo(TEXT("Testing non-existing attribute is ignored."));
+	InSelector.SetAttributeName("Foo");
+	if (!Run())
+	{
+		return false;
+	}
+
+	// If the output attribute name is invalid, it fails
+	AddInfo(TEXT("Testing non-valid attribute is not supported."));
+	
+	AddExpectedError(TEXT("Failed to rename attribute from 'Attr1' to 'Foo#'"), EAutomationExpectedMessageFlags::Contains, 1, /*bIsRegex=*/false);
+	AddExpectedError(TEXT("New attribute name Foo# is not valid"), EAutomationExpectedMessageFlags::Contains, 1, /*bIsRegex=*/false);
+	InSelector.SetAttributeName("Attr1");
+	OutName = TEXT("Foo#");
+	return Run();
+}
+
+bool FPCGMetadataRenameDataDomainTest::RunTest(const FString& Parameters)
+{
+	InSelector.Update(TEXT("@Data.Attr1"));
+	ExpectedAttributeToBeRenamed = TEXT("Attr1");
+	ExpectedDomain = PCGMetadataDomainID::Data;
+	
+	return Run();
+}
+
+bool FPCGMetadataRenameInvalidDataDomainTest::RunTest(const FString& Parameters)
+{
+	InSelector.Update(TEXT("@Foo.Attr1"));
+	ExpectedAttributeToBeRenamed = TEXT("Attr1");
+	ExpectedDomain = PCGMetadataDomainID::Invalid;
+	bValid = false;
+	
+	AddExpectedError(TEXT("Metadata domain Foo is invalid for this data."), EAutomationExpectedMessageFlags::Contains, 1, /*bIsRegex=*/false);
+	
+	return Run();
+}
+
+bool FPCGMetadataRenameDeprecationTest::RunTest(const FString& Parameters)
+{
+	// If the attribute to rename is not None, it should work as-is
+	AddInfo(TEXT("Testing deprecation with multiple attribute and not None as input."));
+	InSelector.SetAttributeName(TEXT("Attr1"));
+	ExpectedAttributeToBeRenamed = TEXT("Attr1");
+	bDeprecation = true;
+
+	if (!Run())
+	{
+		return false;
+	}
+	
+	// If there are multiple attributes, InSelector is None and deprecation is on, it acts like @LastCreated
+	InSelector.SetAttributeName(NAME_None);
+	ExpectedAttributeToBeRenamed = TEXT("Attr2");
+
+	AddInfo(TEXT("Testing deprecation with multiple attributes and None as input."));
+	if (!Run())
+	{
+		return false;
+	}
+
+	// If there is just None, it should be None
+	ExistingAttributes = {NAME_None};
+	ExpectedAttributeToBeRenamed = NAME_None;
+	
+	AddInfo(TEXT("Testing deprecation with single attribute and None as input."));
+	return Run();
+}
+
+#endif // WITH_EDITOR
\ No newline at end of file
diff --git a/Engine/Plugins/PCG/Source/PCG/Public/Elements/Metadata/PCGMetadataRenameElement.h b/Engine/Plugins/PCG/Source/PCG/Public/Elements/Metadata/PCGMetadataRenameElement.h
--- a/Engine/Plugins/PCG/Source/PCG/Public/Elements/Metadata/PCGMetadataRenameElement.h
+++ b/Engine/Plugins/PCG/Source/PCG/Public/Elements/Metadata/PCGMetadataRenameElement.h
@@ -3,6 +3,7 @@
 #pragma once
 
 #include "PCGSettings.h"
+#include "Metadata/PCGAttributePropertySelector.h"
 
 #include "PCGMetadataRenameElement.generated.h"
 
@@ -12,6 +13,12 @@ class UPCGMetadataRenameSettings : public UPCGSettings
 	GENERATED_BODY()
 
 public:
+	UPCGMetadataRenameSettings(const FObjectInitializer& ObjectInitializer);
+	
+	//~Begin UObject interface
+	void PostLoad() override;
+	//~End UObject interface
+	
 	//~Begin UPCGSettings interface
 #if WITH_EDITOR
 	virtual FName GetDefaultNodeName() const override { return FName(TEXT("AttributeRename")); }
@@ -28,9 +35,11 @@ protected:
 	//~End UPCGSettings interface
 
 public:
-	UPROPERTY(BlueprintReadWrite, EditAnywhere, Category = Settings, meta = (PCG_Overridable))
-	FName AttributeToRename = NAME_None;
+	/** Selector to the attribute to rename. Can also specify a metadata domain. Can't be a property or contains extractors. */
+	UPROPERTY(BlueprintReadWrite, EditAnywhere, Category = Settings, meta = (PCG_Overridable, PCG_DiscardPropertySelection, PCG_DiscardExtraSelection))
+	FPCGAttributePropertyInputSelector AttributeToRename;
 
+	/** Name of the new attribute, rename is in place. It will be on the same domain specified by the AttributeToRename selector. */
 	UPROPERTY(BlueprintReadWrite, EditAnywhere, Category = Settings, meta = (PCG_Overridable))
 	FName NewAttributeName = NAME_None;
 };
diff --git a/Engine/Plugins/PCG/Source/PCG/Public/PCGCustomVersion.h b/Engine/Plugins/PCG/Source/PCG/Public/PCGCustomVersion.h
--- a/Engine/Plugins/PCG/Source/PCG/Public/PCGCustomVersion.h
+++ b/Engine/Plugins/PCG/Source/PCG/Public/PCGCustomVersion.h
@@ -86,6 +86,9 @@ struct FPCGCustomVersion
 		// Refactor of the Attribute Property selector to deprecate point properties and support any property.
 		AttributePropertySelectorDeprecatePointProperties = 25,
 
+		// Changed Attribute rename to support selectors
+		AttributeRenameSupportToSelectors = 26,
+
 		// -----<new versions can be added above this line>-------------------------------------------------
 		VersionPlusOne,
 		LatestVersion = VersionPlusOne - 1
