[PCG] Added Graph Param Tooltip as description for PCG Graph's Get Graph Description (JSON)

#rnx

diff --git a/Engine/Plugins/PCG/Source/PCG/Private/PCGGraph.cpp b/Engine/Plugins/PCG/Source/PCG/Private/PCGGraph.cpp
--- a/Engine/Plugins/PCG/Source/PCG/Private/PCGGraph.cpp
+++ b/Engine/Plugins/PCG/Source/PCG/Private/PCGGraph.cpp
@@ -2293,6 +2293,12 @@ FString UPCGGraph::GetGraphDescriptionJson() const
 			// - Type
 			JsonObject->SetStringField(TEXT("Type"), StaticEnum<EPropertyBagPropertyType>()->GetDisplayValueAsText(PropertyDesc.ValueType).ToString());
 
+			// - Description
+			if (PropertyDesc.CachedProperty->HasMetaData(TEXT("ToolTip")))
+			{
+				JsonObject->SetStringField(TEXT("Description"), PropertyDesc.CachedProperty->GetMetaData(TEXT("ToolTip")));
+			}
+
 			// - Value
 			auto SetJson = [&JsonObject]<typename T>(const T& Value) -> bool
 			{
