[PCG] FPCGSelectionKey extensibility
- FPCGSelectionKey can now be extended to narrow down tracking selection with FPCGSelectionKey::CustomKey
- Custom Keys are child USTRUCTS of FPCGCustomSelectionKey
- New event PCGSubsystem::NotifySelectionKeyChanged allows external systems to notify when custom keys have changed so that PCG can refresh the invalidated keys / components
- New event supports providing an array of changed bounds so that Culled keys do not need to rely only on Actor bounds
- Deprecate existing IsMatching signatures for the new FPCGSelectionKey signature

diff --git a/Engine/Plugins/PCG/Source/PCG/Private/Elements/PCGActorSelector.cpp b/Engine/Plugins/PCG/Source/PCG/Private/Elements/PCGActorSelector.cpp
--- a/Engine/Plugins/PCG/Source/PCG/Private/Elements/PCGActorSelector.cpp
+++ b/Engine/Plugins/PCG/Source/PCG/Private/Elements/PCGActorSelector.cpp
@@ -355,6 +355,14 @@ FPCGSelectionKey FPCGSelectionKey::CreateFromPath(FSoftObjectPath&& InObjectPath
 	return Key;
 }
 
+FPCGSelectionKey FPCGSelectionKey::CreateFromObjectPtr(const UObject* InObject)
+{
+	FPCGSelectionKey Key = FPCGSelectionKey::CreateFromPath(FSoftObjectPath(InObject));
+	Key.CachedObject = InObject;
+
+	return Key;
+}
+
 void FPCGSelectionKey::SetExtraDependency(const UClass* InExtraDependency)
 {
 	OptionalExtraDependency = InExtraDependency;
@@ -363,12 +371,33 @@ void FPCGSelectionKey::SetExtraDependency(const UClass* InExtraDependency)
 // DEPRECATED 5.6
 bool FPCGSelectionKey::IsMatching(const UObject* InObject, const UPCGComponent* InComponent) const
 {
-	return IsMatching(InObject, FSoftObjectPath(InObject), InComponent);
+	// This signature of IsMatching does not support custom keys
+	if (CustomKey.IsValid())
+	{
+		return false;
+	}
+
+	return IsMatching(FPCGSelectionKey::CreateFromObjectPtr(InObject), InComponent);
 }
 
+// DEPRECATED 5.7
 bool FPCGSelectionKey::IsMatching(const UObject* InObject, const FSoftObjectPath& InObjectPath, const UPCGComponent* InComponent) const
 {
-	if (!InObject)
+	// This signature of IsMatching does not support custom keys
+	if (CustomKey.IsValid())
+	{
+		return false;
+	}
+
+	const FPCGSelectionKey SelectionKey = FPCGSelectionKey::CreateFromPath(InObjectPath);
+	SelectionKey.CachedObject = InObject;
+	return IsMatching(SelectionKey, InComponent);
+}
+
+bool FPCGSelectionKey::IsMatching(const FPCGSelectionKey& InSelectionKey, const UPCGComponent* InComponent) const
+{
+	const UObject* SelectionObject = InSelectionKey.GetObjectFromPath();
+	if (!SelectionObject)
 	{
 		return false;
 	}
@@ -377,7 +406,7 @@ bool FPCGSelectionKey::IsMatching(const UObject* InObject, const FSoftObjectPath
 	// Re-use the same mechanism than Get Actor Data, which should be cheap since we don't look for all actors in the world.
 	if (ActorFilter != EPCGActorFilter::AllWorldActors)
 	{
-		const AActor* InActor = Cast<const AActor>(InObject);
+		const AActor* InActor = Cast<const AActor>(SelectionObject);
 
 		if (!InActor)
 		{
@@ -397,13 +426,13 @@ bool FPCGSelectionKey::IsMatching(const UObject* InObject, const FSoftObjectPath
 	{
 	case EPCGActorSelection::ByTag:
 	{
-		const AActor* InActor = Cast<const AActor>(InObject);
+		const AActor* InActor = Cast<const AActor>(SelectionObject);
 		return InActor && PCGActorSelector::TagListMatchesTag(InActor->Tags, Tag, bTagContainsWildcard, CachedTagString);
 	}
 	case EPCGActorSelection::ByClass:
-		return InObject && InObject->GetClass()->IsChildOf(SelectionClass);
+		return SelectionObject && SelectionObject->GetClass()->IsChildOf(SelectionClass);
 	case EPCGActorSelection::ByPath:
-		return InObjectPath == ObjectPath;
+		return ObjectPath == InSelectionKey.ObjectPath;
 	default:
 		return false;
 	}
@@ -412,15 +441,51 @@ bool FPCGSelectionKey::IsMatching(const UObject* InObject, const FSoftObjectPath
 // DEPRECATED 5.6
 bool FPCGSelectionKey::IsMatching(const UObject* InObject, const TSet<FName>& InRemovedTags, const TSet<UPCGComponent*>& InComponents, TSet<UPCGComponent*>* OptionalMatchedComponents) const
 {
-	return IsMatching(InObject, FSoftObjectPath(InObject), InRemovedTags, InComponents, OptionalMatchedComponents);
+	// This signature of IsMatching does not support custom keys
+	if (CustomKey.IsValid())
+	{
+		return false;
+	}
+
+	return IsMatching(FPCGSelectionKey::CreateFromObjectPtr(InObject), InRemovedTags, InComponents, OptionalMatchedComponents);
 }
 
+// DEPRECATED 5.7
 bool FPCGSelectionKey::IsMatching(const UObject* InObject, const FSoftObjectPath& InObjectPath, const TSet<FName>& InRemovedTags, const TSet<UPCGComponent*>& InComponents, TSet<UPCGComponent*>* OptionalMatchedComponents) const
 {
-	if (!InObject)
+	// This signature of IsMatching does not support custom keys
+	if (CustomKey.IsValid())
+	{
+		return false;
+	}
+
+	FPCGSelectionKey SelectionKey = FPCGSelectionKey::CreateFromPath(InObjectPath);
+	SelectionKey.CachedObject = InObject;
+
+	return IsMatching(SelectionKey, InRemovedTags, InComponents, OptionalMatchedComponents);
+}
+
+bool FPCGSelectionKey::IsMatching(const FPCGSelectionKey& InSelectionKey, const TSet<FName>& InRemovedTags, const TSet<UPCGComponent*>& InComponents, TSet<UPCGComponent*>* OptionalMatchedComponents) const
+{
+	const UObject* SelectionObject = InSelectionKey.GetObjectFromPath();
+
+	if (!SelectionObject)
+	{
+		return false;
+	}
+
+	// Check custom key first
+	if (CustomKey.IsValid() != InSelectionKey.CustomKey.IsValid())
 	{
 		return false;
 	}
+	else if (CustomKey.IsValid())
+	{
+		if (!CustomKey.Get().IsMatching(InSelectionKey.CustomKey.Get()))
+		{
+			return false;
+		}
+	}
 
 	// If we filter something else than all world actors, matching depends on the component.
 	// Since we can have a lot of components in InComponents, we go the other way around (Actor to component)
@@ -428,7 +493,7 @@ bool FPCGSelectionKey::IsMatching(const UObject* InObject, const FSoftObjectPath
 	{
 		bool bFoundMatch = false;
 
-		const AActor* InActor = Cast<const AActor>(InObject);
+		const AActor* InActor = Cast<const AActor>(SelectionObject);
 
 		if (!InActor)
 		{
@@ -482,16 +547,16 @@ bool FPCGSelectionKey::IsMatching(const UObject* InObject, const FSoftObjectPath
 	{
 	case EPCGActorSelection::ByTag:
 	{
-		const AActor* InActor = Cast<const AActor>(InObject);
+		const AActor* InActor = Cast<const AActor>(SelectionObject);
 		bIsMatched = PCGActorSelector::TagListMatchesTag(InRemovedTags, Tag, bTagContainsWildcard, CachedTagString) ||
 			(InActor && PCGActorSelector::TagListMatchesTag(InActor->Tags, Tag, bTagContainsWildcard, CachedTagString));
 		break;
 	}
 	case EPCGActorSelection::ByClass:
-		bIsMatched = InObject->IsA(SelectionClass);
+		bIsMatched = SelectionObject->IsA(SelectionClass);
 		break;
 	case EPCGActorSelection::ByPath:
-		bIsMatched = InObjectPath == ObjectPath;
+		bIsMatched = InSelectionKey.ObjectPath == ObjectPath;
 		break;
 	default:
 		bIsMatched = false;
@@ -513,6 +578,18 @@ bool FPCGSelectionKey::operator==(const FPCGSelectionKey& InOther) const
 		return false;
 	}
 
+	if (CustomKey.IsValid() != InOther.CustomKey.IsValid())
+	{
+		return false;
+	}
+	else if (CustomKey.IsValid())
+	{
+		if (!CustomKey.Get().Equals(InOther.CustomKey.Get()))
+		{
+			return false;
+		}
+	}
+
 	switch (Selection)
 	{
 	case EPCGActorSelection::ByTag:
@@ -537,6 +614,21 @@ void FPCGSelectionKey::UpdateAfterTagChange()
 	PCGActorSelector::PrepareTag(Selection == EPCGActorSelection::ByTag, Tag, bTagContainsWildcard, CachedTagString);
 }
 
+const UObject* FPCGSelectionKey::GetObjectFromPath() const
+{
+	if (CachedObject)
+	{
+		return CachedObject;
+	}
+
+	if (ObjectPath.IsValid())
+	{
+		CachedObject = ObjectPath.ResolveObject();
+	}
+
+	return CachedObject;
+}
+
 FArchive& operator<<(FArchive& Ar, FPCGSelectionKey& Key)
 {
 	// Serialize the normal UPROPERTY data
@@ -556,6 +648,11 @@ FArchive& operator<<(FArchive& Ar, FPCGSelectionKey& Key)
 	return Ar;
 }
 
+uint32 GetTypeHash(const FPCGCustomSelectionKey& In)
+{
+	return In.GetTypeHashInternal();
+}
+
 uint32 GetTypeHash(const FPCGSelectionKey& In)
 {
 	uint32 HashResult = HashCombine(GetTypeHash(In.ActorFilter), GetTypeHash(In.Selection));
@@ -564,6 +661,11 @@ uint32 GetTypeHash(const FPCGSelectionKey& In)
 	HashResult = HashCombine(HashResult, GetTypeHash(In.OptionalExtraDependency));
 	HashResult = HashCombine(HashResult, GetTypeHash(In.ObjectPath));
 
+	if (In.CustomKey.IsValid())
+	{
+		HashResult = HashCombine(HashResult, GetTypeHash(In.CustomKey.Get()));
+	}
+
 	return HashResult;
 }
 
diff --git a/Engine/Plugins/PCG/Source/PCG/Private/Elements/PCGDataFromActor.cpp b/Engine/Plugins/PCG/Source/PCG/Private/Elements/PCGDataFromActor.cpp
--- a/Engine/Plugins/PCG/Source/PCG/Private/Elements/PCGDataFromActor.cpp
+++ b/Engine/Plugins/PCG/Source/PCG/Private/Elements/PCGDataFromActor.cpp
@@ -430,7 +430,7 @@ bool FPCGDataFromActorElement::ExecuteInternal(FPCGContext* InContext) const
 					{
 						if (AActor* Actor = Component->GetOwner())
 						{
-							if (Key.IsMatching(Actor, FSoftObjectPath(Actor), Component))
+							if (Key.IsMatching(FPCGSelectionKey::CreateFromObjectPtr(Actor), Component))
 							{
 								Context->FoundActors.AddUnique(Actor);
 							}
@@ -443,7 +443,7 @@ bool FPCGDataFromActorElement::ExecuteInternal(FPCGContext* InContext) const
 					{
 						if (AActor* Actor = Component->GetOwner())
 						{
-							if (Key.IsMatching(Actor, FSoftObjectPath(Actor), Component))
+							if (Key.IsMatching(FPCGSelectionKey::CreateFromObjectPtr(Actor), Component))
 							{
 								Context->FoundActors.AddUnique(Actor);
 							}
diff --git a/Engine/Plugins/PCG/Source/PCG/Private/PCGActorAndComponentMapping.cpp b/Engine/Plugins/PCG/Source/PCG/Private/PCGActorAndComponentMapping.cpp
--- a/Engine/Plugins/PCG/Source/PCG/Private/PCGActorAndComponentMapping.cpp
+++ b/Engine/Plugins/PCG/Source/PCG/Private/PCGActorAndComponentMapping.cpp
@@ -1234,10 +1234,10 @@ bool FPCGActorAndComponentMapping::IsActorTracked(const AActor* InActor) const
 
 	TSet<FName> EmptySet;
 
-	FSoftObjectPath ActorPath(InActor);
-	auto Matching = [InActor, &ActorPath, &EmptySet](const TPair<FPCGSelectionKey, TSet<UPCGComponent*>>& It) -> bool
+	const FPCGSelectionKey ActorKey = FPCGSelectionKey::CreateFromObjectPtr(InActor);
+	auto Matching = [&ActorKey, &EmptySet](const TPair<FPCGSelectionKey, TSet<UPCGComponent*>>& It) -> bool
 	{
-		return It.Key.IsMatching(InActor, ActorPath, EmptySet, It.Value, nullptr);
+		return It.Key.IsMatching(ActorKey, EmptySet, It.Value, nullptr);
 	};
 
 	return Algo::AnyOf(CulledTrackedKeysToComponentsMap, Matching) || Algo::AnyOf(AlwaysTrackedKeysToComponentsMap, Matching);
@@ -1822,24 +1822,59 @@ void FPCGActorAndComponentMapping::OnObjectChanged(UObject* InObject, const ECha
 
 	ensure(!Actor || Actor->GetWorld() == GetWorld());
 
-	/** First gather all the components that are tracking this actor */
-	TSet<UPCGComponent*> CulledTrackedComponents;
-	TSet<UPCGComponent*> AlwaysTrackedComponents;
-	TArray<FPCGSelectionKey> MatchedKeys;
-
 	TSet<FName> RemovedTags;
 	if (Actor && InPreviousData && !InPreviousData->Get<1>().IsEmpty())
 	{
 		RemovedTags = InPreviousData->Get<1>().Difference(TSet<FName>(Actor->Tags));
 	}
 
-	FSoftObjectPath ObjectPath(InObject);
-	FSoftObjectPath InOriginatingChangeObjectPath(InOriginatingChangeObject);
-	auto Gather = [InObject, &ObjectPath, &RemovedTags, &MatchedKeys, InOriginatingChangeObject, &InOriginatingChangeObjectPath](const TMap<FPCGSelectionKey, TSet<UPCGComponent*>>& InMap, TSet<UPCGComponent*>& OutSet)
+	FPCGSelectionKey SelectionKey = FPCGSelectionKey::CreateFromObjectPtr(InObject);
+	
+	TArray<FBox, TInlineAllocator<2>> Bounds;
+	const FBox ActorBounds = Actor ? PCGHelpers::GetActorBounds(Actor, /*bIgnorePCGCreatedComponents=*/false) : FBox{ EForceInit::ForceInit };
+	Bounds.Add(ActorBounds);
+	// Also provide previous bounds if we do have previous data
+	if (InPreviousData)
+	{
+		const FBox& PreviousBounds = InPreviousData->Get<0>();
+		if (PreviousBounds.IsValid && !PreviousBounds.Equals(ActorBounds))
+		{
+			Bounds.Add(PreviousBounds);
+		}
+	}
+
+	OnSelectionKeyChanged(SelectionKey, InOriginatingChangeObject, RemovedTags, bNoRefreshOwner, Bounds);
+}
+
+void FPCGActorAndComponentMapping::OnSelectionKeyChanged(const FPCGSelectionKey& InSelectionKey, const UObject* InOriginatingChangeObject, TArrayView<const FBox> InBounds)
+{
+	const UObject* ChangedObject = InSelectionKey.GetObjectFromPath();
+	ensure(ChangedObject);
+
+	const AActor* Actor = Cast<AActor>(ChangedObject);
+
+	OnSelectionKeyChanged(InSelectionKey, InOriginatingChangeObject, {}, /*bInNoRefreshOwner=*/false, InBounds);
+}
+
+void FPCGActorAndComponentMapping::OnSelectionKeyChanged(const FPCGSelectionKey& InSelectionKey, const UObject* InOriginatingChangeObject, const TSet<FName>& InRemovedTags, bool bInNoRefreshOwner, TArrayView<const FBox> InBounds)
+{
+	TRACE_CPUPROFILER_EVENT_SCOPE(FPCGActorAndComponentMapping::OnSelectionKeyChanged);
+		
+	/** First gather all the components that are tracking this actor */
+	TSet<UPCGComponent*> CulledTrackedComponents;
+	TSet<UPCGComponent*> AlwaysTrackedComponents;
+	TArray<FPCGSelectionKey> MatchedKeys;
+	
+	FPCGSelectionKey InOriginatingSelectionKey = FPCGSelectionKey::CreateFromObjectPtr(InOriginatingChangeObject);
+	
+	const UObject* ChangedObject = InSelectionKey.GetObjectFromPath();
+	ensure(ChangedObject);
+
+	auto Gather = [InSelectionKey, InOriginatingSelectionKey, &InRemovedTags, &MatchedKeys](const TMap<FPCGSelectionKey, TSet<UPCGComponent*>>& InMap, TSet<UPCGComponent*>& OutSet)
 	{
 		for (const auto& It : InMap)
 		{
-			if (It.Key.IsMatching(InObject, ObjectPath, RemovedTags, It.Value, &OutSet) || It.Key.IsMatching(InOriginatingChangeObject, InOriginatingChangeObjectPath, RemovedTags, It.Value, &OutSet))
+			if (It.Key.IsMatching(InSelectionKey, InRemovedTags, It.Value, &OutSet) || It.Key.IsMatching(InOriginatingSelectionKey, InRemovedTags, It.Value, &OutSet))
 			{
 				MatchedKeys.Add(It.Key);
 			}
@@ -1855,7 +1890,7 @@ void FPCGActorAndComponentMapping::OnObjectChanged(UObject* InObject, const ECha
 		return;
 	}
 
-	TArray<const UObject*, TInlineAllocator<2>> ChangedObjects = {InObject};
+	TArray<const UObject*, TInlineAllocator<2>> ChangedObjects = { ChangedObject };
 	ChangedObjects.AddUnique(InOriginatingChangeObject);
 
 	// We discard the component if:
@@ -1864,11 +1899,11 @@ void FPCGActorAndComponentMapping::OnObjectChanged(UObject* InObject, const ECha
 	//   * We should not refresh the owner and its owner is one of the objects that changed.
 	//   * It is generating and one of the objects that changed is managed by it (safe-guard, would probably induce an infinite loop).
 	//   * One of the objects that changed is ignored by its original component.
-	auto ShouldDiscardComponent = [&ChangedObjects, bNoRefreshOwner](const UPCGComponent* InComponent)
+	auto ShouldDiscardComponent = [&ChangedObjects, bInNoRefreshOwner](const UPCGComponent* InComponent)
 	{
 		if (!InComponent
 			|| ChangedObjects.Contains(InComponent)
-			|| (bNoRefreshOwner && ChangedObjects.Contains(InComponent->GetOwner()))
+			|| (bInNoRefreshOwner && ChangedObjects.Contains(InComponent->GetOwner()))
 			|| (InComponent->IsGenerating() && InComponent->IsAnyObjectManagedByResource(ChangedObjects)))
 		{
 			return true;
@@ -1889,14 +1924,11 @@ void FPCGActorAndComponentMapping::OnObjectChanged(UObject* InObject, const ECha
 	TSet<UPCGComponent*> DirtyComponents;
 
 	EPCGComponentDirtyFlag DirtyFlag = EPCGComponentDirtyFlag::Actor;
-	if (InObject->IsA<ALandscapeProxy>())
+	if (ChangedObject->IsA<ALandscapeProxy>())
 	{
 		DirtyFlag = DirtyFlag | EPCGComponentDirtyFlag::Landscape;
 	}
 
-	// We need the actor bounds to know if we are intersecting, for the culled settings to be dirtied.
-	const FBox ActorBounds = Actor ? PCGHelpers::GetActorBounds(Actor, /*bIgnorePCGCreatedComponents=*/false) : FBox{EForceInit::ForceInit};
-
 	{
 		TRACE_CPUPROFILER_EVENT_SCOPE(FPCGActorAndComponentMapping::OnObjectChanged::AlwaysTrackedUpdate);
 
@@ -1909,8 +1941,8 @@ void FPCGActorAndComponentMapping::OnObjectChanged(UObject* InObject, const ECha
 
 			// Don't mark "Owner changed" if the change originate from a PCG Component. It will be delegated to the ClearCacheForActor.
 			// It is necessary to avoid infine loops when there are multiple PCG components on one actor, and one component was generated.
-			const bool bOwnerChanged = (PCGComponent->GetOwner() == InObject) && (!InOriginatingChangeObject || !InOriginatingChangeObject->IsA<UPCGComponent>());
-			bool bShouldDirty = bOwnerChanged || (PCGComponent->ShouldTrackLandscape() && InObject->IsA<ALandscapeProxy>());
+			const bool bOwnerChanged = (PCGComponent->GetOwner() == ChangedObject) && (!InOriginatingChangeObject || !InOriginatingChangeObject->IsA<UPCGComponent>());
+			bool bShouldDirty = bOwnerChanged || (PCGComponent->ShouldTrackLandscape() && ChangedObject->IsA<ALandscapeProxy>());
 
 			// Since when we clear the cache for a settings, we clear it all, it's only necessary to do it once on the original component. And we will dirty all the local components (with dirty dispatch)
 			if (!DirtyComponents.Contains(PCGComponent))
@@ -1930,13 +1962,10 @@ void FPCGActorAndComponentMapping::OnObjectChanged(UObject* InObject, const ECha
 
 	if (!CulledTrackedComponents.IsEmpty())
 	{
-		// Keep a pointer to the current bounds we are looking at for the partition case. Set it to Actorbounds first, then will be set to old bounds if necessary.
-		const FBox* CurrentActorBoundsPtr = &ActorBounds;
-
 		// Then do an octree find to get all components that intersect with this actor.
 		// If the actor has moved, we also need to find components that intersected with it before
 		// We first do it for non-partitioned, then we do it for partitioned
-		auto UpdateNonPartitioned = [this, &DirtyComponents, InObject, CulledTrackedComponents, &RemovedTags, DirtyFlag, InOriginatingChangeObject, bNoRefreshOwner, &MatchedKeys, &ShouldDiscardComponent](const FPCGComponentRef& ComponentRef) -> void
+		auto UpdateNonPartitioned = [this, &DirtyComponents, ChangedObject, CulledTrackedComponents, &InRemovedTags, DirtyFlag, InOriginatingChangeObject, bInNoRefreshOwner, &MatchedKeys, &ShouldDiscardComponent](const FPCGComponentRef& ComponentRef, const FBox& InBounds) -> void
 		{
 			TRACE_CPUPROFILER_EVENT_SCOPE(FPCGActorAndComponentMapping::OnObjectChanged::UpdateNonPartitioned);
 
@@ -1948,7 +1977,7 @@ void FPCGActorAndComponentMapping::OnObjectChanged(UObject* InObject, const ECha
 				return;
 			}
 
-			bool bShouldDirty = ComponentRef.Component->ShouldTrackLandscape() && InObject->IsA<ALandscapeProxy>();
+			bool bShouldDirty = ComponentRef.Component->ShouldTrackLandscape() && ChangedObject->IsA<ALandscapeProxy>();
 			bShouldDirty |= ClearCacheForKeys(MatchedKeys, ComponentRef.Component, /*bIntersect=*/true, InOriginatingChangeObject);
 			if (bShouldDirty)
 			{
@@ -1957,11 +1986,9 @@ void FPCGActorAndComponentMapping::OnObjectChanged(UObject* InObject, const ECha
 			}
 		};
 
-		NonPartitionedOctree.FindElementsWithBoundsTest(ActorBounds, UpdateNonPartitioned);
-
 		// For partitioned, we first need check if the original component intersect with the bounds, then forward the dirty call only to locals that intersect with the bounds.
 		// Note: CurrentActorBoundsPtr is passed by reference because it will be modified between lambda calls (cf comment above).
-		auto UpdatePartitioned = [this, &DirtyComponents, InObject, CulledTrackedComponents, &CurrentActorBoundsPtr, &RemovedTags, DirtyFlag, InOriginatingChangeObject, bNoRefreshOwner, &MatchedKeys, &ShouldDiscardComponent](const FPCGComponentRef& ComponentRef)  -> void
+		auto UpdatePartitioned = [this, &DirtyComponents, ChangedObject, CulledTrackedComponents, &InRemovedTags, DirtyFlag, InOriginatingChangeObject, bInNoRefreshOwner, &MatchedKeys, &ShouldDiscardComponent](const FPCGComponentRef& ComponentRef, const FBox& InBounds)  -> void
 		{
 			TRACE_CPUPROFILER_EVENT_SCOPE(FPCGActorAndComponentMapping::OnObjectChanged::UpdatePartitioned);
 
@@ -1974,23 +2001,22 @@ void FPCGActorAndComponentMapping::OnObjectChanged(UObject* InObject, const ECha
 				return;
 			}
 
-			check(CurrentActorBoundsPtr);
 			const FBox ComponentBounds = ComponentRef.Bounds.GetBox();
-			const bool bIntersect = CurrentActorBoundsPtr->Intersect(ComponentBounds);
+			const bool bIntersect = InBounds.Intersect(ComponentBounds);
 			if (!bIntersect)
 			{
 				return;
 			}
 
-			const FBox Overlap = CurrentActorBoundsPtr->Overlap(ComponentBounds);
-			bool bShouldDirty = ComponentRef.Component->ShouldTrackLandscape() && InObject->IsA<ALandscapeProxy>();
+			const FBox Overlap = InBounds.Overlap(ComponentBounds);
+			bool bShouldDirty = ComponentRef.Component->ShouldTrackLandscape() && ChangedObject->IsA<ALandscapeProxy>();
 			bShouldDirty |= ClearCacheForKeys(MatchedKeys, ComponentRef.Component, /*bIntersect=*/true, InOriginatingChangeObject);
 			bool bWasDirtied = false;
 
 			// Since when we clear the cache for a settings, we clear it all, it's only necessary to do it once on the original component, then only dirty the local that intersects.
 			if (bShouldDirty)
 			{
-				ForAllIntersectingPartitionActors(Overlap, [InObject, Component = ComponentRef.Component, &bWasDirtied, DirtyFlag, InOriginatingChangeObject](APCGPartitionActor* InPartitionActor) -> void
+				ForAllIntersectingPartitionActors(Overlap, [ChangedObject, Component = ComponentRef.Component, &bWasDirtied, DirtyFlag, InOriginatingChangeObject](APCGPartitionActor* InPartitionActor) -> void
 				{
 					if (UPCGComponent* LocalComponent = InPartitionActor->GetLocalComponent(Component))
 					{
@@ -2008,21 +2034,11 @@ void FPCGActorAndComponentMapping::OnObjectChanged(UObject* InObject, const ECha
 			}
 		};
 
-		PartitionedOctree.FindElementsWithBoundsTest(ActorBounds, UpdatePartitioned);
-
-		// If it has moved, redo it with the old bounds.
-		if (InPreviousData)
+		// Gather for all bounds
+		for (const FBox& Bounds : InBounds)
 		{
-			const FBox& PreviousBounds = InPreviousData->Get<0>();
-			if (PreviousBounds.IsValid && !PreviousBounds.Equals(ActorBounds))
-			{
-				TRACE_CPUPROFILER_EVENT_SCOPE(FPCGActorAndComponentMapping::OnObjectChanged::SecondUpdateHasMoved);
-
-				// Set the actor bounds with the old one, to have the right Overlap in the Partition case.
-				CurrentActorBoundsPtr = &PreviousBounds;
-				NonPartitionedOctree.FindElementsWithBoundsTest(PreviousBounds, UpdateNonPartitioned);
-				PartitionedOctree.FindElementsWithBoundsTest(PreviousBounds, UpdatePartitioned);
-			}
+			NonPartitionedOctree.FindElementsWithBoundsTest(Bounds, [&Bounds, &UpdateNonPartitioned](const FPCGComponentRef& ComponentRef) { UpdateNonPartitioned(ComponentRef, Bounds); });
+			PartitionedOctree.FindElementsWithBoundsTest(Bounds, [&Bounds, &UpdatePartitioned](const FPCGComponentRef& ComponentRef) { UpdatePartitioned(ComponentRef, Bounds); });
 		}
 	}
 
@@ -2030,7 +2046,7 @@ void FPCGActorAndComponentMapping::OnObjectChanged(UObject* InObject, const ECha
 
 	// If it is a landscape and we should discard the refresh, early out.
 	bool bIsInEditingMode = false;
-	ALandscapeProxy* Landscape = Cast<ALandscapeProxy>(InObject);
+	const ALandscapeProxy* Landscape = Cast<ALandscapeProxy>(ChangedObject);
 	if (PCGActorAndComponentMapping::ShouldDiscardLandscapeRefresh(Landscape, bIsInEditingMode, bIsCurrentlyExitingLandscapeEditMode))
 	{
 		// If we are in editing, keep track of all the dirtied landscape to refresh them when we exit.
@@ -2071,10 +2087,12 @@ void FPCGActorAndComponentMapping::OnObjectChanged(UObject* InObject, const ECha
 			}
 		}
 
-		const bool bOwnerHasChanged = Component->GetOwner() == InObject;
+		const bool bOwnerHasChanged = Component->GetOwner() == ChangedObject;
 
-		if ((!bNoRefreshOwner || !bOwnerHasChanged) && (!Component->bOnlyTrackItself || bOwnerHasChanged))
+		if ((!bInNoRefreshOwner || !bOwnerHasChanged) && (!Component->bOnlyTrackItself || bOwnerHasChanged))
 		{
+			const AActor* Actor = Cast<AActor>(ChangedObject);
+
 			// When an object changes, we need to make sure that we don't trigger a refresh on PCG components that are "higher" in the
 			// level hierarchy, otherwise we will end up generating in the Level Instance level, which is wrong.
 			// Note that in some instances (e.g. when something happens at the Level Instance level) we need to make sure that the
@@ -2180,7 +2198,7 @@ void FPCGActorAndComponentMapping::OnPCGGraphStartsGenerating(UPCGComponent* InC
 	auto IsKeyTrackingPCGData = [](const FPCGSelectionKey& InKey) { return InKey.OptionalExtraDependency && InKey.OptionalExtraDependency->IsChildOf(UPCGComponent::StaticClass()); };
 
 	AActor* Actor = InComponent->GetOwner();
-	FSoftObjectPath ActorPath(Actor);
+	const FPCGSelectionKey ActorKey = FPCGSelectionKey::CreateFromObjectPtr(Actor);
 
 	for (TPair<FPCGSelectionKey, TSet<UPCGComponent*>>& It : CulledTrackedKeysToComponentsMap)
 	{
@@ -2190,7 +2208,7 @@ void FPCGActorAndComponentMapping::OnPCGGraphStartsGenerating(UPCGComponent* InC
 		}
 
 		TSet<UPCGComponent*> TempTrackedComponents;
-		It.Key.IsMatching(Actor, ActorPath, RemovedTags, It.Value, &TempTrackedComponents);
+		It.Key.IsMatching(ActorKey, RemovedTags, It.Value, &TempTrackedComponents);
 
 		// Removing all components that aren't intersecting with it, since it won't contribute to refresh
 		for (UPCGComponent* TrackedComponent : TempTrackedComponents)
@@ -2209,7 +2227,7 @@ void FPCGActorAndComponentMapping::OnPCGGraphStartsGenerating(UPCGComponent* InC
 			continue;
 		}
 
-		It.Key.IsMatching(Actor, ActorPath, RemovedTags, It.Value, &TrackedComponents);
+		It.Key.IsMatching(ActorKey, RemovedTags, It.Value, &TrackedComponents);
 	}
 
 	for (UPCGComponent* Component : TrackedComponents)
diff --git a/Engine/Plugins/PCG/Source/PCG/Private/PCGActorAndComponentMapping.h b/Engine/Plugins/PCG/Source/PCG/Private/PCGActorAndComponentMapping.h
--- a/Engine/Plugins/PCG/Source/PCG/Private/PCGActorAndComponentMapping.h
+++ b/Engine/Plugins/PCG/Source/PCG/Private/PCGActorAndComponentMapping.h
@@ -224,6 +224,9 @@ private:
 	*/
 	void OnObjectChanged(UObject* InObject, const EChangedObjectReason InChangeReason, const FActorPreviousData* InPreviousData = nullptr, const UObject* InOriginatingChangeObject = nullptr);
 
+	void OnSelectionKeyChanged(const FPCGSelectionKey& InSelectionKey, const UObject* InOriginatingChangeObject, TArrayView<const FBox> InBounds);
+	void OnSelectionKeyChanged(const FPCGSelectionKey& InSelectionKey, const UObject* InOriginatingChangeObject, const TSet<FName>& InRemovedTags, bool bInNoRefreshOwner, TArrayView<const FBox> InBounds);
+
 	/** Gather all settings from a given component that track the key, and clear the cache for them. Returns true if we should dirty afterwards (aka at least one settings was cleared and/or landscape changed). */
 	bool ClearCacheForKeys(const TArray<FPCGSelectionKey>& InKeys, const UPCGComponent* InComponent, const bool bIntersect, const UObject* InOriginatingChange) const;
 
diff --git a/Engine/Plugins/PCG/Source/PCG/Private/PCGComponent.cpp b/Engine/Plugins/PCG/Source/PCG/Private/PCGComponent.cpp
--- a/Engine/Plugins/PCG/Source/PCG/Private/PCGComponent.cpp
+++ b/Engine/Plugins/PCG/Source/PCG/Private/PCGComponent.cpp
@@ -2978,12 +2978,13 @@ bool UPCGComponent::IsObjectTracked(const UObject* InObject, bool& bOutIsCulled)
 		return true;
 	}
 
-	FSoftObjectPath ObjectPath(InObject);
-	auto CheckMap = [this, &InObject, &ObjectPath, &bOutIsCulled](const FPCGSelectionKeyToSettingsMap& InMap) -> bool
+	const FPCGSelectionKey ObjectKey = FPCGSelectionKey::CreateFromObjectPtr(InObject);
+
+	auto CheckMap = [this, &ObjectKey, &bOutIsCulled](const FPCGSelectionKeyToSettingsMap& InMap) -> bool
 	{
 		for (const auto& It : InMap)
 		{
-			bool bFound = It.Key.IsMatching(InObject, ObjectPath, this);
+			bool bFound = It.Key.IsMatching(ObjectKey, this);
 
 			if (bFound)
 			{
diff --git a/Engine/Plugins/PCG/Source/PCG/Private/PCGContext.cpp b/Engine/Plugins/PCG/Source/PCG/Private/PCGContext.cpp
--- a/Engine/Plugins/PCG/Source/PCG/Private/PCGContext.cpp
+++ b/Engine/Plugins/PCG/Source/PCG/Private/PCGContext.cpp
@@ -495,9 +495,24 @@ void FPCGContext::OverrideSettings()
 	SettingsWithOverride->CacheCrc();
 }
 
-bool FPCGContext::IsValueOverriden(const FName PropertyName)
+bool FPCGContext::IsValueOverriden(const TArrayView<const FName>& PropertyNameChain) const
 {
-	const FPCGSettingsOverridableParam** OverriddenParam = OverriddenParams.FindByPredicate([PropertyName](const FPCGSettingsOverridableParam* ParamToCheck)
+	if (PropertyNameChain.IsEmpty())
+	{
+		return false;
+	}
+
+	const FPCGSettingsOverridableParam* const* OverrideParam = OverriddenParams.FindByPredicate([&PropertyNameChain](const FPCGSettingsOverridableParam* ParamToCheck)
+	{
+		return ParamToCheck && ParamToCheck->PropertiesNames == PropertyNameChain;
+	});
+
+	return OverrideParam != nullptr;
+}
+
+bool FPCGContext::IsValueOverriden(const FName PropertyName) const
+{
+	const FPCGSettingsOverridableParam* const* OverriddenParam = OverriddenParams.FindByPredicate([PropertyName](const FPCGSettingsOverridableParam* ParamToCheck)
 	{
 		return ParamToCheck && !ParamToCheck->PropertiesNames.IsEmpty() && ParamToCheck->PropertiesNames[0] == PropertyName;
 	});
diff --git a/Engine/Plugins/PCG/Source/PCG/Private/Subsystems/PCGSubsystem.cpp b/Engine/Plugins/PCG/Source/PCG/Private/Subsystems/PCGSubsystem.cpp
--- a/Engine/Plugins/PCG/Source/PCG/Private/Subsystems/PCGSubsystem.cpp
+++ b/Engine/Plugins/PCG/Source/PCG/Private/Subsystems/PCGSubsystem.cpp
@@ -1725,6 +1725,11 @@ void UPCGSubsystem::CleanupAllPCGComponents(bool bPurge) const
 	}
 }
 
+void UPCGSubsystem::NotifySelectionKeyChanged(const FPCGSelectionKey& InSelectionKey, const UObject* InOriginatingChangeObject, TArrayView<const FBox> InChangedBounds)
+{
+	ActorAndComponentMapping->OnSelectionKeyChanged(InSelectionKey, InOriginatingChangeObject, InChangedBounds);
+}
+
 const FPCGNodeVisualLogs& UPCGSubsystem::GetNodeVisualLogs() const
 {
 	if (IPCGEditorModule* PCGEditorModule = IPCGEditorModule::Get())
diff --git a/Engine/Plugins/PCG/Source/PCG/Public/Elements/PCGActorSelector.h b/Engine/Plugins/PCG/Source/PCG/Public/Elements/PCGActorSelector.h
--- a/Engine/Plugins/PCG/Source/PCG/Public/Elements/PCGActorSelector.h
+++ b/Engine/Plugins/PCG/Source/PCG/Public/Elements/PCGActorSelector.h
@@ -4,6 +4,8 @@
 
 #include "Metadata/PCGAttributePropertySelector.h"
 
+#include "Misc/Guid.h"
+#include "StructUtils/InstancedStruct.h"
 #include "Templates/SubclassOf.h"
 #include "UObject/NameTypes.h"
 
@@ -53,6 +55,69 @@ enum class EPCGActorFilter : uint8
 	FromInput
 };
 
+/**
+ * Base struct to extend FPCGSelectionKey to allow more fine grained and custom selection
+ */
+USTRUCT()
+struct FPCGCustomSelectionKey
+{
+public:
+	GENERATED_BODY()
+
+	FPCGCustomSelectionKey() = default;
+	virtual ~FPCGCustomSelectionKey() = default;
+		
+	bool IsMatching(const FPCGCustomSelectionKey& InSelectionKey) const
+	{
+		if (GetScriptStruct() != InSelectionKey.GetScriptStruct())
+		{
+			return false;
+		}
+
+		return IsMatchingInternal(InSelectionKey);
+	}
+
+	bool Equals(const FPCGCustomSelectionKey& InSelectionKey) const
+	{
+		if (GetScriptStruct() != InSelectionKey.GetScriptStruct())
+		{
+			return false;
+		}
+
+		return EqualsInternal(InSelectionKey);
+	}
+
+	template <typename T>
+	T* CastTo()
+	{
+		return GetScriptStruct() == T::StaticStruct() ? StaticCast<T*>(this) : nullptr;
+	}
+	
+	template <typename T>
+	const T* CastTo() const
+	{
+		return GetScriptStruct() == T::StaticStruct() ? StaticCast<const T*>(this) : nullptr;
+	}
+
+	PCG_API friend uint32 GetTypeHash(const FPCGCustomSelectionKey& In);
+	
+	virtual UScriptStruct* GetScriptStruct() const { return nullptr; }
+
+protected:
+	virtual bool IsMatchingInternal(const FPCGCustomSelectionKey&) const PURE_VIRTUAL(FPCGCustomSelectionKey::IsMatchingInternal, return false;)
+	virtual bool EqualsInternal(const FPCGCustomSelectionKey&) const PURE_VIRTUAL(FPCGCustomSelectionKey::EqualsInternal, return false;)
+	virtual uint32 GetTypeHashInternal() const PURE_VIRTUAL(FPCGCustomSelectionKey::GetTypeHashInternal, return 0;)
+};
+
+template<>
+struct TStructOpsTypeTraits<FPCGCustomSelectionKey> : public TStructOpsTypeTraitsBase2<FPCGCustomSelectionKey>
+{
+	enum
+	{
+		WithPureVirtual = true,
+	};
+};
+
 /**
 * Structure to specify a selection criteria for an object/actor
 * Object can be selected using the EPCGActorSelection::ByClass or EPCGActorSelection::ByPath
@@ -61,6 +126,8 @@ enum class EPCGActorFilter : uint8
 USTRUCT()
 struct FPCGSelectionKey
 {
+	friend class FPCGActorAndComponentMapping;
+
 	GENERATED_BODY()
 
 	FPCGSelectionKey() = default;
@@ -75,18 +142,21 @@ struct FPCGSelectionKey
 
 	static UE_API FPCGSelectionKey CreateFromPath(const FSoftObjectPath& InObjectPath);
 	static UE_API FPCGSelectionKey CreateFromPath(FSoftObjectPath&& InObjectPath);
+	static UE_API FPCGSelectionKey CreateFromObjectPtr(const UObject* InObject);
 
 	UE_API bool operator==(const FPCGSelectionKey& InOther) const;
 
 	// Friend functions need to be explicitly exported.
 	PCG_API friend uint32 GetTypeHash(const FPCGSelectionKey& In);
 
-	UE_API bool IsMatching(const UObject* InObject, const FSoftObjectPath& InObjectPath, const UPCGComponent* InComponent) const;
-	UE_API bool IsMatching(const UObject* InObject, const FSoftObjectPath& InObjectPath, const TSet<FName>& InRemovedTags, const TSet<UPCGComponent*>& InComponents, TSet<UPCGComponent*>* OptionalMatchedComponents = nullptr) const;
+	UE_API bool IsMatching(const FPCGSelectionKey& InSelectionKey, const UPCGComponent* InComponent) const;
+	UE_API bool IsMatching(const FPCGSelectionKey& InSelectionKey, const TSet<FName>& InRemovedTags, const TSet<UPCGComponent*>& InComponents, TSet<UPCGComponent*>* OptionalMatchedComponents = nullptr) const;
 
 	UE_API void SetExtraDependency(const UClass* InExtraDependency);
 	UE_API void UpdateAfterTagChange();
 
+	UE_API const UObject* GetObjectFromPath() const;
+
 	UPROPERTY()
 	EPCGActorFilter ActorFilter = EPCGActorFilter::AllWorldActors;
 
@@ -107,7 +177,15 @@ struct FPCGSelectionKey
 	UPROPERTY()
 	TObjectPtr<const UClass> OptionalExtraDependency = nullptr;
 
+	UPROPERTY()
+	TInstancedStruct<FPCGCustomSelectionKey> CustomKey;
+
 	// Deprecated methods
+	UE_DEPRECATED(5.7, "Replaced by the version with a FPCGSelectionKey.")
+	UE_API bool IsMatching(const UObject* InObject, const FSoftObjectPath& InObjectPath, const UPCGComponent* InComponent) const;
+	UE_DEPRECATED(5.7, "Replaced by the version with a FPCGSelectionKey.")
+	UE_API bool IsMatching(const UObject* InObject, const FSoftObjectPath& InObjectPath, const TSet<FName>& InRemovedTags, const TSet<UPCGComponent*>& InComponents, TSet<UPCGComponent*>* OptionalMatchedComponents = nullptr) const;
+
 	UE_DEPRECATED(5.6, "Replaced by the version with a soft object path.")
 	UE_API bool IsMatching(const UObject* InObject, const UPCGComponent* InComponent) const;
 	UE_DEPRECATED(5.6, "Replaced by the version with a soft object path.")
@@ -116,6 +194,7 @@ struct FPCGSelectionKey
 protected:
 	FString CachedTagString;
 	bool bTagContainsWildcard = false;
+	mutable const UObject* CachedObject;
 };
 
 PCG_API FArchive& operator<<(FArchive& Ar, FPCGSelectionKey& Key);
diff --git a/Engine/Plugins/PCG/Source/PCG/Public/PCGContext.h b/Engine/Plugins/PCG/Source/PCG/Public/PCGContext.h
--- a/Engine/Plugins/PCG/Source/PCG/Public/PCGContext.h
+++ b/Engine/Plugins/PCG/Source/PCG/Public/PCGContext.h
@@ -142,7 +142,8 @@ UE_API PRAGMA_ENABLE_DEPRECATION_WARNINGS
 	UE_API void OverrideSettings();
 
 	// Returns true if the given property has been overriden
-	UE_API bool IsValueOverriden(const FName PropertyName);
+	UE_API bool IsValueOverriden(const FName PropertyName) const;
+	UE_API bool IsValueOverriden(const TArrayView<const FName>& PropertyNameChain) const;
 
 	// Return the seed, possibly overriden by params, and combined with the source component (if any).
 	UE_API int GetSeed() const;
diff --git a/Engine/Plugins/PCG/Source/PCG/Public/Subsystems/PCGSubsystem.h b/Engine/Plugins/PCG/Source/PCG/Public/Subsystems/PCGSubsystem.h
--- a/Engine/Plugins/PCG/Source/PCG/Public/Subsystems/PCGSubsystem.h
+++ b/Engine/Plugins/PCG/Source/PCG/Public/Subsystems/PCGSubsystem.h
@@ -309,6 +309,9 @@ public:
 	/** Will gather all the components registered, and ask for cleanup. */
 	UE_API void CleanupAllPCGComponents(bool bPurge) const;
 
+	/** Notification that a Selection key needs refreshing */
+	UE_API void NotifySelectionKeyChanged(const FPCGSelectionKey& InSelectionKey, const UObject* InOriginatingChangeObject, TArrayView<const FBox> InBounds);
+
 	/** Get graph warnings and errors for all nodes. */
 	UE_DEPRECATED(5.6, "Use IPCGEditorModule::GetNodeVisualLogs instead")
 	UE_API const FPCGNodeVisualLogs& GetNodeVisualLogs() const;
