[PCG] FPCGDynamicTrackingHelper::AddSingleDynamicTrackingKey now checks UPCGSettings::CanDynamicallyTrackKeys

#rnx

diff --git a/Engine/Plugins/PCG/Source/PCG/Private/Helpers/PCGDynamicTrackingHelpers.cpp b/Engine/Plugins/PCG/Source/PCG/Private/Helpers/PCGDynamicTrackingHelpers.cpp
--- a/Engine/Plugins/PCG/Source/PCG/Private/Helpers/PCGDynamicTrackingHelpers.cpp
+++ b/Engine/Plugins/PCG/Source/PCG/Private/Helpers/PCGDynamicTrackingHelpers.cpp
@@ -47,8 +47,13 @@ void FPCGDynamicTrackingHelper::AddSingleDynamicTrackingKey(FPCGContext* InConte
 	{
 		if (IPCGGraphExecutionSource* ExecutionSource = InContext->ExecutionSource.Get())
 		{
-			TPair<FPCGSelectionKey, bool> NewPair(std::forward<FPCGSelectionKey>(InKey), bIsCulled);
-			ExecutionSource->GetExecutionState().RegisterDynamicTracking(InContext->GetOriginalSettings<UPCGSettings>(), MakeArrayView(&NewPair, 1));
+			const UPCGSettings* Settings = InContext->GetOriginalSettings<UPCGSettings>();
+			
+			if(Settings && Settings->CanDynamicallyTrackKeys())
+			{
+				TPair<FPCGSelectionKey, bool> NewPair(std::forward<FPCGSelectionKey>(InKey), bIsCulled);
+				ExecutionSource->GetExecutionState().RegisterDynamicTracking(Settings, MakeArrayView(&NewPair, 1));
+			}
 		}
 	}
 }
