[PCG] Hide 'Clear PCG Link' button on runtime generated PCG Components
- Since visibility of buttons is visible as soon as one PCG Component satisfies the visibility condition, make sure to check the condition in the clicked event

diff --git a/Engine/Plugins/PCG/Source/PCGEditor/Private/Details/PCGComponentDetails.cpp b/Engine/Plugins/PCG/Source/PCGEditor/Private/Details/PCGComponentDetails.cpp
--- a/Engine/Plugins/PCG/Source/PCGEditor/Private/Details/PCGComponentDetails.cpp
+++ b/Engine/Plugins/PCG/Source/PCGEditor/Private/Details/PCGComponentDetails.cpp
@@ -180,6 +180,7 @@ void FPCGComponentDetails::CustomizeDetails(IDetailLayoutBuilder& DetailBuilder)
 				SNew(SButton)
 				.VAlign(VAlign_Center)
 				.OnClicked(this, &FPCGComponentDetails::OnClearPCGLinkClicked)
+				.Visibility(this, &FPCGComponentDetails::ClearPCGLinkButtonVisible)
 				[
 					SNew(STextBlock)
 					.Font(IDetailLayoutBuilder::GetDetailFont())
@@ -267,32 +268,44 @@ EVisibility FPCGComponentDetails::RefreshButtonVisible() const
 	return EVisibility::Collapsed;
 }
 
+EVisibility FPCGComponentDetails::ClearPCGLinkButtonVisible() const
+{
+	for (const TWeakObjectPtr<UPCGComponent>& Component : SelectedComponents)
+	{
+		if (Component.IsValid() && !Component->IsManagedByRuntimeGenSystem())
+		{
+			return EVisibility::Visible;
+		}
+	}
+
+	return EVisibility::Collapsed;
+}
+
+
 FReply FPCGComponentDetails::OnGenerateClicked()
 {
 	const FModifierKeysState ModifierKeys = FSlateApplication::Get().GetModifierKeys();
 	const bool bIsControlDown = ModifierKeys.IsControlDown();
 
-	if (bIsControlDown)
+	bool bFlushedCache = false;
+
+	for (TWeakObjectPtr<UPCGComponent>& SelectedComponentPtr : SelectedComponents)
 	{
-		for (TWeakObjectPtr<UPCGComponent>& Component : SelectedComponents)
+		UPCGComponent* PCGComponent = SelectedComponentPtr.Get();
+
+		if (IsValid(PCGComponent) && !PCGComponent->IsGenerating() && !PCGComponent->IsManagedByRuntimeGenSystem())
 		{
-			if (Component.IsValid())
+			// Flush cache if needed
+			if (bIsControlDown && !bFlushedCache)
 			{
-				if (UPCGSubsystem* PCGSubsystem = Component->GetSubsystem())
+				if (UPCGSubsystem* PCGSubsystem = PCGComponent->GetSubsystem())
 				{
 					PCGSubsystem->FlushCache();
+					bFlushedCache = true;
 				}
-
-				break;
 			}
-		}
-	}
 
-	for (TWeakObjectPtr<UPCGComponent>& Component : SelectedComponents)
-	{
-		if (Component.IsValid())
-		{
-			Component.Get()->Generate(/*bForce=*/bIsControlDown);
+			PCGComponent->Generate(/*bForce=*/bIsControlDown);
 		}
 	}
 
@@ -303,7 +316,7 @@ FReply FPCGComponentDetails::OnCancelClicked()
 {
 	for (TWeakObjectPtr<UPCGComponent>& Component : SelectedComponents)
 	{
-		if (Component.IsValid())
+		if (Component.IsValid() && Component->IsGenerating() && !Component->IsManagedByRuntimeGenSystem())
 		{
 			Component.Get()->CancelGeneration();
 		}
@@ -316,7 +329,7 @@ FReply FPCGComponentDetails::OnRefreshClicked()
 {
 	for (TWeakObjectPtr<UPCGComponent>& Component : SelectedComponents)
 	{
-		if (Component.IsValid())
+		if (Component.IsValid() && Component->IsManagedByRuntimeGenSystem())
 		{
 			// Trigger the deepest refresh - re-initialize the PAs.
 			Component.Get()->Refresh(EPCGChangeType::Structural | EPCGChangeType::GenerationGrid, /*bCancelExistingRefresh=*/true);
@@ -330,7 +343,7 @@ FReply FPCGComponentDetails::OnClearPCGLinkClicked()
 {
 	for (TWeakObjectPtr<UPCGComponent>& Component : SelectedComponents)
 	{
-		if (Component.IsValid())
+		if (Component.IsValid() && !Component->IsManagedByRuntimeGenSystem())
 		{
 			Component.Get()->ClearPCGLink();
 		}
@@ -358,7 +371,7 @@ FReply FPCGComponentDetails::OnCleanupClicked()
 {
 	for (TWeakObjectPtr<UPCGComponent>& Component : SelectedComponents)
 	{
-		if (Component.IsValid())
+		if (Component.IsValid() && !Component->IsManagedByRuntimeGenSystem())
 		{
 			FModifierKeysState ModifierKeys = FSlateApplication::Get().GetModifierKeys();
 			if (ModifierKeys.IsControlDown())
diff --git a/Engine/Plugins/PCG/Source/PCGEditor/Public/Details/PCGComponentDetails.h b/Engine/Plugins/PCG/Source/PCGEditor/Public/Details/PCGComponentDetails.h
--- a/Engine/Plugins/PCG/Source/PCGEditor/Public/Details/PCGComponentDetails.h
+++ b/Engine/Plugins/PCG/Source/PCGEditor/Public/Details/PCGComponentDetails.h
@@ -29,6 +29,7 @@ private:
 	EVisibility CancelButtonVisible() const;
 	EVisibility CleanupButtonVisible() const;
 	EVisibility RefreshButtonVisible() const;
+	EVisibility ClearPCGLinkButtonVisible() const;
 	FReply OnGenerateClicked();
 	FReply OnCancelClicked();
 	FReply OnCleanupClicked();
