[PCG] Fix crash when debug display encounters instances with invalid transforms

#rnx

diff --git a/Engine/Plugins/PCG/Source/PCGEditor/Private/DataVisualizations/PCGSpatialDataVisualization.cpp b/Engine/Plugins/PCG/Source/PCGEditor/Private/DataVisualizations/PCGSpatialDataVisualization.cpp
--- a/Engine/Plugins/PCG/Source/PCGEditor/Private/DataVisualizations/PCGSpatialDataVisualization.cpp
+++ b/Engine/Plugins/PCG/Source/PCGEditor/Private/DataVisualizations/PCGSpatialDataVisualization.cpp
@@ -112,6 +112,7 @@ void IPCGSpatialDataVisualization::ExecuteDebugDisplayHelper(
 	}
 
 	const FConstPCGPointValueRanges ValueRanges(PointData);
+	bool bFoundNonNormalizedInstances = false;
 
 	for (int i = 0; i < NumDesiredInstances; ++i)
 	{
@@ -135,6 +136,18 @@ void IPCGSpatialDataVisualization::ExecuteDebugDisplayHelper(
 		{
 			InstanceTransform.SetScale3D(FVector(PointScale));
 		}
+
+		// If any instances have non-normalized rotations, normalize them to avoid crashing, but emit a warning.
+		if (!InstanceTransform.IsRotationNormalized())
+		{
+			InstanceTransform.NormalizeRotation();
+			bFoundNonNormalizedInstances = true;
+		}
+	}
+
+	if (bFoundNonNormalizedInstances)
+	{
+		PCGE_LOG_C(Warning, GraphAndLog, Context, LOCTEXT("UnnormalizedRotation", "PCGSpatialDataVisualization: Encountered one or more transforms with unnormalized rotation. Rotations will be normalized for visualization."));
 	}
 
 	FPCGISMComponentBuilderParams Params[2];
