[PCG] Update PCG Editor Settings option for show pause button

#rnx

diff --git a/Engine/Plugins/PCG/Source/PCGEditor/Public/PCGEditorSettings.h b/Engine/Plugins/PCG/Source/PCGEditor/Public/PCGEditorSettings.h
--- a/Engine/Plugins/PCG/Source/PCGEditor/Public/PCGEditorSettings.h
+++ b/Engine/Plugins/PCG/Source/PCGEditor/Public/PCGEditorSettings.h
@@ -264,7 +264,7 @@ public:
 	bool bDisableCPUThrottlingDuringGraphExecution = false;
 
 	/** Controls whether the "Pause PCG" button will be display in the editor menu */
-	UPROPERTY(EditAnywhere, config, Category = "Editor Performance", meta = (DisplayName = "Show PCG pause button in the editor (requires restart)"))
+	UPROPERTY(EditAnywhere, config, Category = "Editor Performance", meta = (DisplayName = "Show PCG pause button in the editor", ConfigRestartRequired = true))
 	bool bShowPauseButton = false;
 
 	/** Controls whether PCG tasks will be cancelled when unpausing PCG, which might result in dirty/stale content, but will not require significant time. */
