[PCG] GPU implementation for CullPointsOutsideActorBounds

diff --git a/Engine/Plugins/PCG/Shaders/Private/Elements/PCGCullPointsOutsideActorBounds.usf b/Engine/Plugins/PCG/Shaders/Private/Elements/PCGCullPointsOutsideActorBounds.usf
new file mode 100644
--- /dev/null
+++ b/Engine/Plugins/PCG/Shaders/Private/Elements/PCGCullPointsOutsideActorBounds.usf
@@ -0,0 +1,48 @@
+// Copyright Epic Games, Inc. All Rights Reserved.
+
+[numthreads(64, 1, 1)]
+void PCGCullPointsOutsideActorBoundsCS(uint3 GroupId : SV_GroupID, uint GroupIndex : SV_GroupIndex)
+{
+	// Mark the kernel as having executed. Must run before we early out via thread index, because the kernel is still 'executed' even if the number of
+	// threads to iterate on is zero. Even if GetNumThreads() returns 0, the kernel will still have been dispatched on a single thread to set this flag.
+	if (all(GroupId == 0) && GroupIndex == 0)
+	{
+		Out_SetAsExecutedInternal();
+	}
+
+	const uint ThreadIndex = GetUnWrappedDispatchThreadId(GroupId, GroupIndex, 64);
+	if (ThreadIndex >= GetNumThreads().x) return;
+
+	uint InDataIndex, InElemIndex;
+	if (!In_GetThreadData(ThreadIndex, InDataIndex, InElemIndex))
+	{
+		return;
+	}
+
+	uint OutDataIndex, OutElemIndex;
+	if (!Out_GetThreadData(ThreadIndex, OutDataIndex, OutElemIndex))
+	{
+		return;
+	}
+
+	// Propagate the 'Removed' status to the output point. Otherwise, this point will not be culled.
+	if (In_IsPointRemoved(InDataIndex, InElemIndex))
+	{
+		Out_RemovePoint(OutDataIndex, OutElemIndex);
+		return;
+	}
+
+	const float BoundsExpansion = CullPointsOutsideActorBounds_GetBoundsExpansion();
+	const float3 InPosition = In_GetPosition(InDataIndex, InElemIndex);
+	const float3 ComponentBoundsMin = GetComponentBoundsMin() - BoundsExpansion;
+	const float3 ComponentBoundsMax = GetComponentBoundsMax() + BoundsExpansion;
+	const bool bInsideBounds = all(InPosition >= ComponentBoundsMin) && all(InPosition < ComponentBoundsMax);
+
+	if (!bInsideBounds)
+	{
+		Out_RemovePoint(OutDataIndex, OutElemIndex);
+		return;
+	}
+
+	PCG_COPY_ATTRIBUTES_TO_OUTPUT(Out, In, OutDataIndex, OutElemIndex, InDataIndex, InElemIndex, /*bMetadataOnly=*/false, /*bInitNonCopiedAttributes=*/true);
+}
diff --git a/Engine/Plugins/PCG/Source/PCG/Private/Compute/DataInterfaces/Elements/PCGCullPointsOutsideActorBoundsDataInterface.cpp b/Engine/Plugins/PCG/Source/PCG/Private/Compute/DataInterfaces/Elements/PCGCullPointsOutsideActorBoundsDataInterface.cpp
new file mode 100644
--- /dev/null
+++ b/Engine/Plugins/PCG/Source/PCG/Private/Compute/DataInterfaces/Elements/PCGCullPointsOutsideActorBoundsDataInterface.cpp
@@ -0,0 +1,81 @@
+// Copyright Epic Games, Inc. All Rights Reserved.
+
+#include "PCGCullPointsOutsideActorBoundsDataInterface.h"
+
+#include "PCGContext.h"
+#include "Compute/PCGComputeGraph.h"
+#include "Compute/PCGDataBinding.h"
+#include "Elements/PCGCullPointsOutsideActorBounds.h"
+#include "Compute/Elements/PCGCullPointsOutsideActorBoundsKernel.h"
+
+#include "GlobalRenderResources.h"
+#include "RenderGraphBuilder.h"
+#include "RenderGraphResources.h"
+#include "RHIResources.h"
+#include "ShaderCompilerCore.h"
+#include "ShaderCore.h"
+#include "ShaderParameterMetadataBuilder.h"
+#include "ComputeFramework/ShaderParamTypeDefinition.h"
+
+#include UE_INLINE_GENERATED_CPP_BY_NAME(PCGCullPointsOutsideActorBoundsDataInterface)
+
+void UPCGCullPointsOutsideActorBoundsDataInterface::GetSupportedInputs(TArray<FShaderFunctionDefinition>& OutFunctions) const
+{
+	OutFunctions.AddDefaulted_GetRef()
+		.SetName(TEXT("CullPointsOutsideActorBounds_GetBoundsExpansion"))
+		.AddReturnType(FShaderValueType::Get(EShaderFundamentalType::Float));
+}
+
+BEGIN_SHADER_PARAMETER_STRUCT(FPCGCullPointsOutsideActorBoundsDataInterfaceParameters,)
+	SHADER_PARAMETER(float, BoundsExpansion)
+END_SHADER_PARAMETER_STRUCT()
+
+void UPCGCullPointsOutsideActorBoundsDataInterface::GetShaderParameters(TCHAR const* UID, FShaderParametersMetadataBuilder& InOutBuilder, FShaderParametersMetadataAllocations& InOutAllocations) const
+{
+	InOutBuilder.AddNestedStruct<FPCGCullPointsOutsideActorBoundsDataInterfaceParameters>(UID);
+}
+
+void UPCGCullPointsOutsideActorBoundsDataInterface::GetHLSL(FString& OutHLSL, FString const& InDataInterfaceName) const
+{
+	TMap<FString, FStringFormatArg> CullPointsOutsideActorBoundsArgs =
+	{
+		{ TEXT("DataInterfaceName"), InDataInterfaceName },
+	};
+
+	OutHLSL += FString::Format(TEXT(
+		"float {DataInterfaceName}_BoundsExpansion;\n"
+		"\n"
+		"float CullPointsOutsideActorBounds_GetBoundsExpansion_{DataInterfaceName}() { return {DataInterfaceName}_BoundsExpansion; }\n"
+		), CullPointsOutsideActorBoundsArgs);
+}
+
+UComputeDataProvider* UPCGCullPointsOutsideActorBoundsDataInterface::CreateDataProvider() const
+{
+	return NewObject<UPCGCullPointsOutsideActorBoundsDataProvider>();
+}
+
+FComputeDataProviderRenderProxy* UPCGCullPointsOutsideActorBoundsDataProvider::GetRenderProxy()
+{
+	const UPCGCullPointsOutsideActorBoundsSettings* Settings = CastChecked<UPCGCullPointsOutsideActorBoundsSettings>(ProducerKernel->GetSettings());
+
+	FPCGCullPointsOutsideActorBoundsDataProviderProxy::FCullPointsOutsideActorBoundsData_RenderThread ProxyData;
+	ProxyData.BoundsExpansion = Settings->BoundsExpansion;
+
+	return new FPCGCullPointsOutsideActorBoundsDataProviderProxy(MoveTemp(ProxyData));
+}
+
+bool FPCGCullPointsOutsideActorBoundsDataProviderProxy::IsValid(FValidationData const& InValidationData) const
+{
+	return InValidationData.ParameterStructSize == sizeof(FParameters);
+}
+
+void FPCGCullPointsOutsideActorBoundsDataProviderProxy::GatherDispatchData(FDispatchData const& InDispatchData)
+{
+	const TStridedView<FParameters> ParameterArray = MakeStridedParameterView<FParameters>(InDispatchData);
+
+	for (int32 InvocationIndex = 0; InvocationIndex < ParameterArray.Num(); ++InvocationIndex)
+	{
+		FParameters& Parameters = ParameterArray[InvocationIndex];
+		Parameters.BoundsExpansion = Data.BoundsExpansion;
+	}
+}
diff --git a/Engine/Plugins/PCG/Source/PCG/Private/Compute/DataInterfaces/Elements/PCGCullPointsOutsideActorBoundsDataInterface.h b/Engine/Plugins/PCG/Source/PCG/Private/Compute/DataInterfaces/Elements/PCGCullPointsOutsideActorBoundsDataInterface.h
new file mode 100644
--- /dev/null
+++ b/Engine/Plugins/PCG/Source/PCG/Private/Compute/DataInterfaces/Elements/PCGCullPointsOutsideActorBoundsDataInterface.h
@@ -0,0 +1,59 @@
+// Copyright Epic Games, Inc. All Rights Reserved.
+
+#pragma once
+
+#include "Compute/DataInterfaces/PCGComputeDataInterface.h"
+
+#include "PCGCullPointsOutsideActorBoundsDataInterface.generated.h"
+
+class FPCGCullPointsOutsideActorBoundsDataInterfaceParameters;
+
+/** Data Interface to marshal CullPointsOutsideActorBounds settings to the GPU. */
+UCLASS(ClassGroup = (Procedural))
+class UPCGCullPointsOutsideActorBoundsDataInterface : public UPCGComputeDataInterface
+{
+	GENERATED_BODY()
+
+public:
+	//~ Begin UComputeDataInterface Interface
+	TCHAR const* GetClassName() const override { return TEXT("PCGCullPointsOutsideActorBounds"); }
+	void GetSupportedInputs(TArray<FShaderFunctionDefinition>& OutFunctions) const override;
+	void GetShaderParameters(TCHAR const* UID, FShaderParametersMetadataBuilder& InOutBuilder, FShaderParametersMetadataAllocations& InOutAllocations) const override;
+	void GetHLSL(FString& OutHLSL, FString const& InDataInterfaceName) const override;
+	UComputeDataProvider* CreateDataProvider() const override;
+	//~ End UComputeDataInterface Interface
+};
+
+UCLASS()
+class UPCGCullPointsOutsideActorBoundsDataProvider : public UPCGComputeDataProvider
+{
+	GENERATED_BODY()
+
+public:
+	//~ Begin UComputeDataProvider Interface
+	FComputeDataProviderRenderProxy* GetRenderProxy() override;
+	//~ End UComputeDataProvider Interface
+};
+
+class FPCGCullPointsOutsideActorBoundsDataProviderProxy : public FComputeDataProviderRenderProxy
+{
+public:
+	struct FCullPointsOutsideActorBoundsData_RenderThread
+	{
+		float BoundsExpansion = 0.0f;
+	};
+
+	FPCGCullPointsOutsideActorBoundsDataProviderProxy(FCullPointsOutsideActorBoundsData_RenderThread InData)
+		: Data(MoveTemp(InData))
+	{}
+
+	//~ Begin FComputeDataProviderRenderProxy Interface
+	bool IsValid(FValidationData const& InValidationData) const override;
+	void GatherDispatchData(FDispatchData const& InDispatchData) override;
+	//~ End FComputeDataProviderRenderProxy Interface
+
+protected:
+	using FParameters = FPCGCullPointsOutsideActorBoundsDataInterfaceParameters;
+
+	FCullPointsOutsideActorBoundsData_RenderThread Data;
+};
diff --git a/Engine/Plugins/PCG/Source/PCG/Private/Compute/Elements/PCGCullPointsOutsideActorBoundsKernel.cpp b/Engine/Plugins/PCG/Source/PCG/Private/Compute/Elements/PCGCullPointsOutsideActorBoundsKernel.cpp
new file mode 100644
--- /dev/null
+++ b/Engine/Plugins/PCG/Source/PCG/Private/Compute/Elements/PCGCullPointsOutsideActorBoundsKernel.cpp
@@ -0,0 +1,80 @@
+// Copyright Epic Games, Inc. All Rights Reserved.
+
+#include "Compute/Elements/PCGCullPointsOutsideActorBoundsKernel.h"
+
+#include "PCGContext.h"
+#include "Compute/Elements/PCGComputeGraphElement.h"
+#include "Compute/DataInterfaces/Elements/PCGCullPointsOutsideActorBoundsDataInterface.h"
+#include "Elements/PCGCullPointsOutsideActorBounds.h"
+#include "Graph/PCGGPUGraphCompilationContext.h"
+
+#include "ShaderCompilerCore.h"
+
+#include UE_INLINE_GENERATED_CPP_BY_NAME(PCGCullPointsOutsideActorBoundsKernel)
+
+#define LOCTEXT_NAMESPACE "PCGCullPointsOutsideActorBoundsKernel"
+
+bool UPCGCullPointsOutsideActorBoundsKernel::ComputeOutputBindingDataDesc(const UPCGComputeGraph* InGraph, FName InOutputPinLabel, UPCGDataBinding* InBinding, FPCGDataCollectionDesc& OutDataDesc) const
+{
+	check(InBinding);
+
+	// Code assumes single output pin.
+	if (!ensure(InOutputPinLabel == PCGPinConstants::DefaultOutputLabel))
+	{
+		return false;
+	}
+
+	const FPCGKernelPin SourceKernelPin(KernelIndex, PCGPinConstants::DefaultInputLabel, /*bIsInput=*/true);
+
+	ensure(InBinding->ComputeKernelPinDataDesc(SourceKernelPin, &OutDataDesc));
+
+	return true;
+}
+
+int UPCGCullPointsOutsideActorBoundsKernel::ComputeThreadCount(const UPCGDataBinding* InBinding) const
+{
+	int ThreadCount = 0;
+
+	const FPCGDataCollectionDesc* OutputPinDesc = InBinding->GetCachedKernelPinDataDesc(this, PCGPinConstants::DefaultOutputLabel, /*bIsInput=*/false);
+
+	if (ensure(OutputPinDesc))
+	{
+		for (const FPCGDataDesc& DataDesc : OutputPinDesc->DataDescs)
+		{
+			ThreadCount += DataDesc.ElementCount;
+		}
+	}
+
+	return ThreadCount;
+}
+
+#if WITH_EDITOR
+FString UPCGCullPointsOutsideActorBoundsKernel::GetCookedSource(FPCGGPUCompilationContext& InOutContext) const
+{
+	FString TemplateFile;
+	ensure(LoadShaderSourceFile(TEXT("/Plugin/PCG/Private/Elements/PCGCullPointsOutsideActorBounds.usf"), EShaderPlatform::SP_PCD3D_SM5, &TemplateFile, nullptr));
+	return TemplateFile;
+}
+
+void UPCGCullPointsOutsideActorBoundsKernel::CreateAdditionalInputDataInterfaces(FPCGGPUCompilationContext& InOutContext, UObject* InObjectOuter, TArray<TObjectPtr<UComputeDataInterface>>& OutDataInterfaces) const
+{
+	Super::CreateAdditionalInputDataInterfaces(InOutContext, InObjectOuter, OutDataInterfaces);
+
+	TObjectPtr<UPCGCullPointsOutsideActorBoundsDataInterface> NodeDI = InOutContext.NewObject_AnyThread<UPCGCullPointsOutsideActorBoundsDataInterface>(InObjectOuter);
+	NodeDI->ProducerKernel = this;
+
+	OutDataInterfaces.Add(NodeDI);
+}
+#endif
+
+void UPCGCullPointsOutsideActorBoundsKernel::GetInputPins(TArray<FPCGPinProperties>& OutPins) const
+{
+	OutPins.Emplace(PCGPinConstants::DefaultInputLabel, EPCGDataType::Point);
+}
+
+void UPCGCullPointsOutsideActorBoundsKernel::GetOutputPins(TArray<FPCGPinPropertiesGPU>& OutPins) const
+{
+	OutPins.Emplace(PCGPinConstants::DefaultOutputLabel, EPCGDataType::Point);
+}
+
+#undef LOCTEXT_NAMESPACE
diff --git a/Engine/Plugins/PCG/Source/PCG/Private/Compute/Elements/PCGCullPointsOutsideActorBoundsKernel.h b/Engine/Plugins/PCG/Source/PCG/Private/Compute/Elements/PCGCullPointsOutsideActorBoundsKernel.h
new file mode 100644
--- /dev/null
+++ b/Engine/Plugins/PCG/Source/PCG/Private/Compute/Elements/PCGCullPointsOutsideActorBoundsKernel.h
@@ -0,0 +1,28 @@
+// Copyright Epic Games, Inc. All Rights Reserved.
+
+#pragma once
+
+#include "Compute/PCGComputeKernel.h"
+
+#include "Compute/PCGPinPropertiesGPU.h"
+
+#include "PCGCullPointsOutsideActorBoundsKernel.generated.h"
+
+UCLASS()
+class UPCGCullPointsOutsideActorBoundsKernel : public UPCGComputeKernel
+{
+	GENERATED_BODY()
+
+public:
+	//~Begin UPCGComputeKernel interface
+	virtual bool ComputeOutputBindingDataDesc(const UPCGComputeGraph* InGraph, FName InOutputPinLabel, UPCGDataBinding* InBinding, FPCGDataCollectionDesc& OutDataDesc) const override;
+	virtual int ComputeThreadCount(const UPCGDataBinding* Binding) const override;
+#if WITH_EDITOR
+	virtual FString GetCookedSource(FPCGGPUCompilationContext& InOutContext) const override;
+	virtual FString GetEntryPoint() const override { return TEXT("PCGCullPointsOutsideActorBoundsCS"); }
+	void CreateAdditionalInputDataInterfaces(FPCGGPUCompilationContext& InOutContext, UObject* InObjectOuter, TArray<TObjectPtr<UComputeDataInterface>>& OutDataInterfaces) const override;
+#endif
+	virtual void GetInputPins(TArray<FPCGPinProperties>& OutPins) const override;
+	virtual void GetOutputPins(TArray<FPCGPinPropertiesGPU>& OutPins) const override;
+	//~End UPCGComputeKernel interface
+};
diff --git a/Engine/Plugins/PCG/Source/PCG/Private/Elements/PCGCullPointsOutsideActorBounds.cpp b/Engine/Plugins/PCG/Source/PCG/Private/Elements/PCGCullPointsOutsideActorBounds.cpp
--- a/Engine/Plugins/PCG/Source/PCG/Private/Elements/PCGCullPointsOutsideActorBounds.cpp
+++ b/Engine/Plugins/PCG/Source/PCG/Private/Elements/PCGCullPointsOutsideActorBounds.cpp
@@ -5,7 +5,10 @@
 #include "PCGComponent.h"
 #include "PCGContext.h"
 #include "PCGPin.h"
+#include "Compute/PCGComputeKernel.h"
+#include "Compute/Elements/PCGCullPointsOutsideActorBoundsKernel.h"
 #include "Data/PCGBasePointData.h"
+#include "Graph/PCGGPUGraphCompilationContext.h"
 #include "Helpers/PCGAsync.h"
 #include "Helpers/PCGHelpers.h"
 
@@ -23,6 +26,22 @@ FText UPCGCullPointsOutsideActorBoundsSettings::GetNodeTooltipText() const
 {
 	return LOCTEXT("NodeTooltip", "Culls points that lie outside the current actor bounds.");
 }
+
+void UPCGCullPointsOutsideActorBoundsSettings::CreateKernels(FPCGGPUCompilationContext& InOutContext, UObject* InObjectOuter, TArray<UPCGComputeKernel*>& OutKernels, TArray<FPCGKernelEdge>& OutEdges) const
+{
+	FPCGComputeKernelParams KernelParams;
+	KernelParams.Settings = this;
+	KernelParams.bLogDescriptions = bDumpDataDescriptions;
+
+	UPCGCullPointsOutsideActorBoundsKernel* Kernel = InOutContext.NewObject_AnyThread<UPCGCullPointsOutsideActorBoundsKernel>(InObjectOuter);
+	Kernel->Initialize(KernelParams);
+	OutKernels.Add(Kernel);
+
+	OutEdges.Emplace(FPCGPinReference(PCGPinConstants::DefaultInputLabel), FPCGPinReference(Kernel, PCGPinConstants::DefaultInputLabel));
+	OutEdges.Emplace(FPCGPinReference(Kernel, PCGPinConstants::DefaultOutputLabel), FPCGPinReference(PCGPinConstants::DefaultOutputLabel));
+
+	// @todo_pcg: Inject a compaction kernel to automatically compact after culling. Perhaps with a threshold to only compact if 10%+ of points were culled?
+}
 #endif
 
 TArray<FPCGPinProperties> UPCGCullPointsOutsideActorBoundsSettings::InputPinProperties() const
diff --git a/Engine/Plugins/PCG/Source/PCG/Public/Elements/PCGCullPointsOutsideActorBounds.h b/Engine/Plugins/PCG/Source/PCG/Public/Elements/PCGCullPointsOutsideActorBounds.h
--- a/Engine/Plugins/PCG/Source/PCG/Public/Elements/PCGCullPointsOutsideActorBounds.h
+++ b/Engine/Plugins/PCG/Source/PCG/Public/Elements/PCGCullPointsOutsideActorBounds.h
@@ -28,8 +28,9 @@ public:
 	virtual FText GetDefaultNodeTitle() const override;
 	virtual FText GetNodeTooltipText() const override;
 	virtual EPCGSettingsType GetType() const override { return EPCGSettingsType::Spatial; }
+	virtual bool DisplayExecuteOnGPUSetting() const override { return true; }
+	virtual void CreateKernels(FPCGGPUCompilationContext& InOutContext, UObject* InObjectOuter, TArray<UPCGComputeKernel*>& OutKernels, TArray<FPCGKernelEdge>& OutEdges) const override;
 #endif
-	
 
 protected:
 	virtual TArray<FPCGPinProperties> InputPinProperties() const override;
