[PCG] Fix regression from CL 43575500
- UID used to get skipped with hardcoded condition, now that we use the same policy as FPCGSettingsHashContext the change didn't take into account that UID doesn't have the PCGNoHash tag as it is transient and FPCGSettingsHashContext skips transient properties...but FPCGSettingsObjectCrc did not.
- Left a todo to convert FPCGSettingsObjectCrc to use FPCGSettingsHashContext

#rnx

diff --git a/Engine/Plugins/PCG/Source/PCG/Private/PCGSettings.cpp b/Engine/Plugins/PCG/Source/PCG/Private/PCGSettings.cpp
--- a/Engine/Plugins/PCG/Source/PCG/Private/PCGSettings.cpp
+++ b/Engine/Plugins/PCG/Source/PCG/Private/PCGSettings.cpp
@@ -72,6 +72,12 @@ public:
 #if WITH_EDITOR
 	virtual bool ShouldSkipProperty(const FProperty* InProperty) const override
 	{
+		// @todo_pcg: convert FPCGSettingsObjectCrc to use the new hashing framework (FPCGSettingsHashContext) which skips transient properties by default.
+		if (InProperty && InProperty->HasAnyPropertyFlags(CPF_Transient))
+		{
+			return true;
+		}
+
 		static FPCGSettingsHashPolicy Policy;
 		return !Policy.ShouldHashProperty(PCGSettings, InProperty);
 	}
