[PCG] Fix HLSL source not detecting changes with case sensitivity

#rnx

diff --git a/Engine/Plugins/PCG/Source/PCG/Private/Compute/PCGComputeSource.cpp b/Engine/Plugins/PCG/Source/PCG/Private/Compute/PCGComputeSource.cpp
--- a/Engine/Plugins/PCG/Source/PCG/Private/Compute/PCGComputeSource.cpp
+++ b/Engine/Plugins/PCG/Source/PCG/Private/Compute/PCGComputeSource.cpp
@@ -49,7 +49,7 @@ FString UPCGComputeSource::GetVirtualPath() const
 #if WITH_EDITOR
 void UPCGComputeSource::SetShaderText(const FString& NewText)
 {
-	if (Source != NewText)
+	if (!Source.Equals(NewText, ESearchCase::CaseSensitive))
 	{
 		Modify();
 
diff --git a/Engine/Plugins/PCG/Source/PCGEditor/Private/Details/PCGComputeSourceDetails.cpp b/Engine/Plugins/PCG/Source/PCGEditor/Private/Details/PCGComputeSourceDetails.cpp
--- a/Engine/Plugins/PCG/Source/PCGEditor/Private/Details/PCGComputeSourceDetails.cpp
+++ b/Engine/Plugins/PCG/Source/PCGEditor/Private/Details/PCGComputeSourceDetails.cpp
@@ -140,7 +140,7 @@ void FPCGComputeSourceDetails::SetComputeSourceText() const
 	{
 		const FString SourceString = SourceText.ToString();
 
-		if (SourceString != ComputeSource->GetShaderText())
+		if (!SourceString.Equals(ComputeSource->GetShaderText(), ESearchCase::CaseSensitive))
 		{
 			ComputeSource->SetShaderText(SourceString);
 		}
diff --git a/Engine/Plugins/PCG/Source/PCGEditor/Private/Widgets/SPCGEditorNodeSource.cpp b/Engine/Plugins/PCG/Source/PCGEditor/Private/Widgets/SPCGEditorNodeSource.cpp
--- a/Engine/Plugins/PCG/Source/PCGEditor/Private/Widgets/SPCGEditorNodeSource.cpp
+++ b/Engine/Plugins/PCG/Source/PCGEditor/Private/Widgets/SPCGEditorNodeSource.cpp
@@ -225,7 +225,7 @@ void SPCGEditorNodeSource::SetShaderFunctionsText() const
 	{
 		const FString ShaderFunctionsString = ShaderFunctionsText.ToString();
 
-		if (ShaderFunctionsString != Provider->GetShaderFunctionsText())
+		if (!ShaderFunctionsString.Equals(Provider->GetShaderFunctionsText(), ESearchCase::CaseSensitive))
 		{
 			Provider->SetShaderFunctionsText(ShaderFunctionsString);
 		}
@@ -238,7 +238,7 @@ void SPCGEditorNodeSource::SetShaderSourceText() const
 	{
 		const FString ShaderSourceString = ShaderSourceText.ToString();
 
-		if (ShaderSourceString != Provider->GetShaderText())
+		if (!ShaderSourceString.Equals(Provider->GetShaderText(), ESearchCase::CaseSensitive))
 		{
 			Provider->SetShaderText(ShaderSourceString);
 		}
