[PCG] Flush runtime gen change detector state when bTreatEditorViewportAsGenerationSource toggled so that we rescan grid cells.

#rnx

diff --git a/Engine/Plugins/PCG/Source/PCG/Internal/RuntimeGen/PCGRuntimeGenChangeDetection.h b/Engine/Plugins/PCG/Source/PCG/Internal/RuntimeGen/PCGRuntimeGenChangeDetection.h
--- a/Engine/Plugins/PCG/Source/PCG/Internal/RuntimeGen/PCGRuntimeGenChangeDetection.h
+++ b/Engine/Plugins/PCG/Source/PCG/Internal/RuntimeGen/PCGRuntimeGenChangeDetection.h
@@ -49,8 +49,8 @@ namespace PCGRuntimeGenChangeDetection
 		/** Called after all cells have been scanned to update the cached state. */
 		void OnCellsScanned(const FDetectionInputs& InInputs);
 
-		/** Flush any recorded state, resulting in a rescan of the component on next tick. */
-		void FlushCachedState(const UPCGComponent* InOriginalComponent);
+		/** Flush any recorded state, resulting in a rescan of the component on next tick. If nullptr provided all components will rescan. */
+		void FlushCachedState(const UPCGComponent* InOptionalOriginalComponent = nullptr);
 
 	private:
 		TMap<FDetectionInputs, FCachedState> StateCache;
diff --git a/Engine/Plugins/PCG/Source/PCG/Private/RuntimeGen/PCGRuntimeGenChangeDetection.cpp b/Engine/Plugins/PCG/Source/PCG/Private/RuntimeGen/PCGRuntimeGenChangeDetection.cpp
--- a/Engine/Plugins/PCG/Source/PCG/Private/RuntimeGen/PCGRuntimeGenChangeDetection.cpp
+++ b/Engine/Plugins/PCG/Source/PCG/Private/RuntimeGen/PCGRuntimeGenChangeDetection.cpp
@@ -94,22 +94,29 @@ namespace PCGRuntimeGenChangeDetection
 		StateCache.FindOrAdd(InInputs).UpdateFromInputs(InInputs);
 	}
 
-	void FDetector::FlushCachedState(const UPCGComponent* InOriginalComponent)
+	void FDetector::FlushCachedState(const UPCGComponent* InOptionalOriginalComponent)
 	{
-		TArray<FDetectionInputs, TInlineAllocator<32>> KeysToRemove;
-
-		for (const TPair<FDetectionInputs, FCachedState>& Entry : StateCache)
+		if (InOptionalOriginalComponent)
 		{
-			const UPCGComponent* Component = Entry.Key.GetOriginalComponent();
-			if (!Component || Component == InOriginalComponent)
+			TArray<FDetectionInputs, TInlineAllocator<32>> KeysToRemove;
+
+			for (const TPair<FDetectionInputs, FCachedState>& Entry : StateCache)
 			{
-				KeysToRemove.Add(Entry.Key);
+				const UPCGComponent* Component = Entry.Key.GetOriginalComponent();
+				if (!Component || Component == InOptionalOriginalComponent)
+				{
+					KeysToRemove.Add(Entry.Key);
+				}
 			}
-		}
 
-		for (const FDetectionInputs& Key : KeysToRemove)
+			for (const FDetectionInputs& Key : KeysToRemove)
+			{
+				StateCache.Remove(Key);
+			}
+		}
+		else
 		{
-			StateCache.Remove(Key);
+			StateCache.Empty();
 		}
 	}
 }
diff --git a/Engine/Plugins/PCG/Source/PCG/Private/RuntimeGen/PCGRuntimeGenScheduler.cpp b/Engine/Plugins/PCG/Source/PCG/Private/RuntimeGen/PCGRuntimeGenScheduler.cpp
--- a/Engine/Plugins/PCG/Source/PCG/Private/RuntimeGen/PCGRuntimeGenScheduler.cpp
+++ b/Engine/Plugins/PCG/Source/PCG/Private/RuntimeGen/PCGRuntimeGenScheduler.cpp
@@ -1332,6 +1332,18 @@ void FPCGRuntimeGenScheduler::TickCVars(const APCGWorldActor* InPCGWorldActor)
 			ResetPartitionActorPoolToSize(BasePoolSize);
 		}
 	}
+
+#if WITH_EDITOR
+	const bool bTreatEditorViewportAsGenSource = InPCGWorldActor->bTreatEditorViewportAsGenerationSource;
+	
+	if (bTreatEditorViewportAsGenSource != bTreatEditorViewportAsGenSourcePreviousFrame && World && World->IsEditorWorld())
+	{
+		UE_LOG(LogPCG, Verbose, TEXT("Flushed change detection status after bTreatEditorViewportAsGenerationSource value change."));
+		ChangeDetector.FlushCachedState();
+	}
+
+	bTreatEditorViewportAsGenSourcePreviousFrame = bTreatEditorViewportAsGenSource;
+#endif // WITH_EDITOR
 }
 
 void FPCGRuntimeGenScheduler::OnOriginalComponentRegistered(UPCGComponent* InOriginalComponent)
diff --git a/Engine/Plugins/PCG/Source/PCG/Public/RuntimeGen/PCGRuntimeGenScheduler.h b/Engine/Plugins/PCG/Source/PCG/Public/RuntimeGen/PCGRuntimeGenScheduler.h
--- a/Engine/Plugins/PCG/Source/PCG/Public/RuntimeGen/PCGRuntimeGenScheduler.h
+++ b/Engine/Plugins/PCG/Source/PCG/Public/RuntimeGen/PCGRuntimeGenScheduler.h
@@ -229,6 +229,10 @@ private:
 	bool bPoolingWasEnabledLastFrame = true;
 	uint32 BasePoolSizeLastFrame = 0;
 
+#if WITH_EDITOR
+	bool bTreatEditorViewportAsGenSourcePreviousFrame = false;
+#endif
+
 	/** Requests to flush all actors are deferred so they can be handled at a known time during tick. */
 	bool bActorFlushRequested = false;
 
