[PCG] Add basic 2d polygon data type. Metadata support will be done in a later CL.

diff --git a/Engine/Plugins/PCG/PCG.uplugin b/Engine/Plugins/PCG/PCG.uplugin
--- a/Engine/Plugins/PCG/PCG.uplugin
+++ b/Engine/Plugins/PCG/PCG.uplugin
@@ -39,6 +39,10 @@
 		{
 			"Name": "ComputeFramework",
 			"Enabled": true
+		},
+		{
+			"Name": "GeometryProcessing",
+			"Enabled": true
 		}
 	]
 }
diff --git a/Engine/Plugins/PCG/Source/PCG/PCG.Build.cs b/Engine/Plugins/PCG/Source/PCG/PCG.Build.cs
--- a/Engine/Plugins/PCG/Source/PCG/PCG.Build.cs
+++ b/Engine/Plugins/PCG/Source/PCG/PCG.Build.cs
@@ -16,8 +16,10 @@ namespace UnrealBuildTool.Rules
 					"CoreUObject",
 					"DeveloperSettings",
 					"Engine",
-					"Landscape",
 					"Foliage",
+					"GeometryCore",
+					"GeometryAlgorithms",
+					"Landscape",
 					"PhysicsCore",
 					"Projects",
 					"RenderCore",
diff --git a/Engine/Plugins/PCG/Source/PCG/Private/Data/PCGPolygon2DData.cpp b/Engine/Plugins/PCG/Source/PCG/Private/Data/PCGPolygon2DData.cpp
new file mode 100644
--- /dev/null
+++ b/Engine/Plugins/PCG/Source/PCG/Private/Data/PCGPolygon2DData.cpp
@@ -0,0 +1,356 @@
+// Copyright Epic Games, Inc. All Rights Reserved.
+
+#include "Data/PCGPolygon2DData.h"
+
+#include "PCGContext.h"
+#include "PCGCustomVersion.h"
+#include "Data/PCGPointData.h"
+#include "Data/PCGPointArrayData.h"
+
+#include "SegmentTypes.h"
+#include "Serialization/ArchiveCrc32.h"
+#include "UObject/Package.h"
+
+namespace PCGPolygon2DData
+{
+	// @todo_pcg : could suggest to move this to the polygon implementation
+	void SerializePolygon(UE::Geometry::FGeneralPolygon2d& Polygon, FArchive& Archive)
+	{
+		if (Archive.IsLoading())
+		{
+			Polygon.ClearHoles();
+
+			TArray<FVector2d> OuterVertices;
+			Archive << OuterVertices;
+			Polygon.SetOuter(UE::Geometry::FPolygon2d(MoveTemp(OuterVertices)));
+
+			int32 NumHoles = 0;
+			Archive << NumHoles;
+
+			for (int HoleIndex = 0; HoleIndex < NumHoles; ++HoleIndex)
+			{
+				TArray<FVector2d> HoleVertices;
+				Archive << HoleVertices;
+				Polygon.AddHole(UE::Geometry::FPolygon2d(MoveTemp(HoleVertices)), /*bCheckContainment=*/false, /*bCheckOrientation=*/false);
+			}
+		}
+		else
+		{
+			const UE::Geometry::FPolygon2d& Outer = Polygon.GetOuter();
+			const TArray<FVector2d>& OuterVertices = Outer.GetVertices();
+			Archive << const_cast<TArray<FVector2d>&>(OuterVertices);
+
+			const TArray<UE::Geometry::FPolygon2d>& Holes = Polygon.GetHoles();
+			int32 NumHoles = Holes.Num();
+			Archive << NumHoles;
+
+			for (const UE::Geometry::FPolygon2d& Hole : Holes)
+			{
+				const TArray<FVector2d>& HoleVertices = Hole.GetVertices();
+				Archive << const_cast<TArray<FVector2d>&>(HoleVertices);
+			}
+		}
+	}
+}
+
+void UPCGPolygon2DData::Serialize(FArchive& InArchive)
+{
+	Super::Serialize(InArchive);
+
+	InArchive.UsingCustomVersion(FPCGCustomVersion::GUID);
+
+	PCGPolygon2DData::SerializePolygon(Polygon, InArchive);
+	InArchive << Transform;
+}
+
+void UPCGPolygon2DData::AddToCrc(FArchiveCrc32& Ar, bool bFullDataCrc) const
+{
+	Super::AddToCrc(Ar, bFullDataCrc);
+
+	// The code below has non-trivial cost...
+	if (!bFullDataCrc)
+	{
+		// Fallback to UID
+		AddUIDToCrc(Ar);
+		return;
+	}
+
+	TRACE_CPUPROFILER_EVENT_SCOPE(UPCGPolygon2DData::AddToCrc);
+
+	FString ClassName = StaticClass()->GetPathName();
+	Ar << ClassName;
+	PCGPolygon2DData::SerializePolygon(const_cast<UE::Geometry::FGeneralPolygon2d&>(Polygon), Ar);
+	Ar << const_cast<FTransform&>(Transform);
+
+	if (Metadata)
+	{
+		Metadata->AddToCrc(Ar, bFullDataCrc);
+	}
+}
+
+bool UPCGPolygon2DData::SamplePoint(const FTransform& InTransform, const FBox& InBounds, FPCGPoint& OutPoint, UPCGMetadata* OutMetadata) const
+{
+	// Project point down to the 2d plane.
+	const FTransform LocalSpaceTransform = InTransform.GetRelativeTransform(Transform);
+	const FBox BoundsInLocalSpace = InBounds.TransformBy(LocalSpaceTransform);
+
+	if(BoundsInLocalSpace.Min.Z > 0 || BoundsInLocalSpace.Max.Z < 0)
+	{
+		return false;
+	}
+
+	// @todo_pcg: Check if the bounds (as a polygon) overlap with the polygon.
+	// Note that the polygon class has an overlap test, but not at the general polygon level
+	// We'll test just the center until then.
+	if (!Polygon.Contains(FVector2D(BoundsInLocalSpace.GetCenter())))
+	{
+		return false;
+	}
+
+	new(&OutPoint) FPCGPoint(InTransform, /*Density=*/1.0f, /*Seed=*/0);
+	OutPoint.SetLocalBounds(InBounds);
+	return true;
+}
+
+bool UPCGPolygon2DData::ProjectPoint(const FTransform& InTransform, const FBox& InBounds, const FPCGProjectionParams& InParams, FPCGPoint& OutPoint, UPCGMetadata* OutMetadata) const
+{
+	// Project point down to the 2d plane.
+	const FTransform LocalSpaceTransform = InTransform.GetRelativeTransform(Transform);
+	const FBox BoundsInLocalSpace = InBounds.TransformBy(LocalSpaceTransform);
+
+	// @todo_pcg: Check if the bounds (as a polygon) overlap with the polygon.
+	// Note that the polygon class has an overlap test, but not at the general polygon level
+	// We'll test just the center until then.
+	const FVector Center = BoundsInLocalSpace.GetCenter();
+	if (!Polygon.Contains(FVector2D(Center)))
+	{
+		return false;
+	}
+
+	new(&OutPoint) FPCGPoint(InTransform, /*Density=*/1.0f, /*Seed=*/0);
+
+	if (InParams.bProjectPositions)
+	{
+		OutPoint.Transform.SetLocation(Transform.TransformPosition(FVector(Center.X, Center.Y, 0.0)));
+	}
+
+	if (InParams.bProjectRotations)
+	{
+		// Want to project the rotation so that the up vector aligns to the surface's up vector
+		FQuat UpVectorRotation = FQuat::FindBetweenNormals(InTransform.GetRotation().GetUpVector(), Transform.GetRotation().GetUpVector());
+		OutPoint.Transform.SetRotation(UpVectorRotation * InTransform.GetRotation());
+	}
+
+	// Implementation note - the polygon transform has always a unit scale, so there's nothing to project against.
+	return true;
+}
+
+UPCGSpatialData* UPCGPolygon2DData::CopyInternal(FPCGContext* Context) const
+{
+	UPCGPolygon2DData* NewPolygonData = FPCGContext::NewObject_AnyThread<UPCGPolygon2DData>(Context);
+	NewPolygonData->Polygon = Polygon;
+	NewPolygonData->Transform = Transform;
+	return NewPolygonData;
+}
+
+FBox UPCGPolygon2DData::GetBounds() const
+{
+	const UE::Geometry::TAxisAlignedBox2<double> Bounds2D = Polygon.Bounds();
+
+	constexpr double FixedHalfHeight = 50.0;
+	const FBox Bounds(FVector(Bounds2D.Min, -FixedHalfHeight), FVector(Bounds2D.Max, +FixedHalfHeight));
+	return Bounds.TransformBy(Transform);
+}
+
+const UPCGPointData* UPCGPolygon2DData::CreatePointData(FPCGContext* Context) const
+{
+	return CastChecked<UPCGPointData>(CreateBasePointData(Context, UPCGPointData::StaticClass()), ECastCheckedType::NullAllowed);
+}
+
+const UPCGPointArrayData* UPCGPolygon2DData::CreatePointArrayData(FPCGContext* Context, const FBox& InBounds) const
+{
+	return CastChecked<UPCGPointArrayData>(CreateBasePointData(Context, UPCGPointArrayData::StaticClass()), ECastCheckedType::NullAllowed);
+}
+
+const UPCGBasePointData* UPCGPolygon2DData::CreateBasePointData(FPCGContext* Context, TSubclassOf<UPCGBasePointData> PointDataClass) const
+{
+	UPCGBasePointData* PointData = FPCGContext::NewObject_AnyThread<UPCGBasePointData>(Context, GetTransientPackage(), PointDataClass);
+
+	FPCGInitializeFromDataParams InitializeFromDataParams(this);
+	InitializeFromDataParams.bInheritSpatialData = false;
+
+	PointData->InitializeFromDataWithParams(InitializeFromDataParams);
+
+	// Allocate properties & points
+	// @todo_pcg: add metadata entries if the vertices had entries
+	const EPCGPointNativeProperties PropertiesToAllocate = EPCGPointNativeProperties::Transform;
+
+	const int NumPoints = GetNumSegments();
+	PointData->SetNumPoints(NumPoints);
+	PointData->AllocateProperties(PropertiesToAllocate);
+
+	FPCGPointValueRanges OutRanges(PointData, /*bAllocate=*/false);
+
+	for (int PointIndex = 0; PointIndex < NumPoints; ++PointIndex)
+	{
+		OutRanges.TransformRange[PointIndex] = GetTransformAtDistance(PointIndex, 0.0, /*bWorldSpace=*/true);
+	}
+
+	return PointData;
+}
+
+int UPCGPolygon2DData::GetNumSegments() const
+{
+	return SegmentCount;
+}
+
+double UPCGPolygon2DData::GetSegmentLength(int SegmentIndex) const
+{
+	if (SegmentIndex < 0 || SegmentIndex >= SegmentCount)
+	{
+		return 0.0;
+	}
+
+	// Implementation note - when we parse the polygon we build the transform such that the scale is always 1.
+	const TPair<int, int> SegmentMapping = SegmentIndexToSegmentAndHoleIndices[SegmentIndex];
+	return Polygon.Segment(SegmentMapping.Key, SegmentMapping.Value).Length();
+}
+
+FVector UPCGPolygon2DData::GetLocationAtAlpha(float Alpha) const
+{
+	return GetTransformAtAlpha(Alpha).GetLocation();
+}
+
+FTransform UPCGPolygon2DData::GetTransformAtAlpha(float InAlpha) const
+{
+	const double Alpha = FMath::Clamp(InAlpha, 0.0, 1.0 - UE_SMALL_NUMBER);
+
+	// Note: alpha in splines is a value from 0 to 1, with each segment having the same relative value.
+	const double SegmentKey = Alpha * SegmentCount;
+	const int SegmentIndex = FMath::FloorToInt(SegmentKey);
+	const double Ratio = FMath::Frac(SegmentKey);
+	const TPair<int, int> SegmentMapping = SegmentIndexToSegmentAndHoleIndices[SegmentIndex];
+
+	const FVector2D PointOnSegment = Polygon.Segment(SegmentMapping.Key, SegmentMapping.Value).PointBetween(Ratio);
+	return Transform * FTransform(FVector(PointOnSegment, 0.0));
+}
+
+FTransform UPCGPolygon2DData::GetTransformAtDistance(int SegmentIndex, double Distance, bool bWorldSpace, FBox* OutBounds) const
+{
+	if (OutBounds)
+	{
+		*OutBounds = FBox::BuildAABB(FVector::ZeroVector, FVector::OneVector);
+	}
+
+	FTransform OutTransform;
+
+	if (SegmentIndex < 0 || SegmentIndex >= SegmentCount)
+	{
+		return FTransform::Identity;
+	}
+
+	const TPair<int, int> SegmentMapping = SegmentIndexToSegmentAndHoleIndices[SegmentIndex];
+	const UE::Geometry::FSegment2d Segment = Polygon.Segment(SegmentMapping.Key, SegmentMapping.Value);
+	
+	// @todo_pcg: should we normalize the distance by scale if the user is asking for world space coordinates?
+	const FVector2D PointOnSegment = Segment.PointAt(Distance - Segment.Extent);
+
+	OutTransform.SetLocation(FVector(PointOnSegment, 0.0));
+	// @todo_pcg: check if Z sign has to be consequent with the polygon winding (likely)
+	OutTransform.SetRotation(FRotationMatrix::MakeFromXZ(FVector(Segment.Direction, 0.0), FVector::UpVector).ToQuat()); // unclear z sign here, should depend on winding(?)
+	OutTransform.SetScale3D(FVector::OneVector);
+
+	if (bWorldSpace)
+	{
+		OutTransform = OutTransform * Transform;
+	}
+
+	return OutTransform;
+}
+
+float UPCGPolygon2DData::GetInputKeyAtDistance(int SegmentIndex, double Distance) const
+{
+	if (SegmentIndex < 0)
+	{
+		return 0;
+	}
+	else if (SegmentIndex >= SegmentCount)
+	{
+		return 1.0;
+	}
+	else
+	{
+		const TPair<int, int> SegmentMapping = SegmentIndexToSegmentAndHoleIndices[SegmentIndex];
+		const UE::Geometry::FSegment2d Segment = Polygon.Segment(SegmentMapping.Key, SegmentMapping.Value);
+		return (static_cast<double>(SegmentIndex) / SegmentCount) + (0.5 * Distance / Segment.Extent);
+	}
+}
+
+double UPCGPolygon2DData::GetDistanceAtSegmentStart(int SegmentIndex) const
+{
+	SegmentIndex = FMath::Clamp(SegmentIndex, 0, SegmentCount);
+
+	double Distance = 0;
+	int CurrentSegmentIndex = 0;
+	while (CurrentSegmentIndex < SegmentIndex && CurrentSegmentIndex < SegmentCount)
+	{
+		TPair<int, int> SegmentMapping = SegmentIndexToSegmentAndHoleIndices[CurrentSegmentIndex];
+		UE::Geometry::FSegment2d Segment = Polygon.Segment(SegmentMapping.Key, SegmentMapping.Value);
+		Distance += 2.0 * Segment.Extent;
+		++CurrentSegmentIndex;
+	}
+
+	return Distance;
+}
+
+void UPCGPolygon2DData::SetPolygon(const UE::Geometry::FGeneralPolygon2d& InPolygon)
+{
+	Polygon = InPolygon;
+	UpdateMappings();
+}
+
+void UPCGPolygon2DData::SetPolygon(UE::Geometry::FGeneralPolygon2d&& InPolygon)
+{
+	Polygon = MoveTemp(InPolygon);
+	UpdateMappings();
+}
+
+void UPCGPolygon2DData::SetTransform(const FTransform& InTransform, bool bCheckWinding)
+{
+	bool bNeedsToReversePolygon = false;
+
+	if (bCheckWinding)
+	{
+		bNeedsToReversePolygon = (FMath::Sign(InTransform.GetRotation().GetUpVector() | Transform.GetRotation().GetUpVector()) < 0);
+	}
+
+	Transform = InTransform;
+
+	if (bNeedsToReversePolygon)
+	{
+		// @todo_pcg - check how that affects the polygon operations
+		Polygon.Reverse();
+	}
+}
+
+void UPCGPolygon2DData::UpdateMappings()
+{
+	SegmentIndexToSegmentAndHoleIndices.Reset();
+	SegmentCount = 0;
+
+	const UE::Geometry::FPolygon2d& Outer = Polygon.GetOuter();
+	for (int VertexIndex = 0; VertexIndex < Outer.VertexCount(); ++VertexIndex)
+	{
+		SegmentIndexToSegmentAndHoleIndices.Add(SegmentCount++, TPair<int, int>(VertexIndex, -1));
+	}
+
+	const TArray<UE::Geometry::FPolygon2d>& Holes = Polygon.GetHoles();
+	for (int HoleIndex = 0; HoleIndex < Holes.Num(); ++HoleIndex)
+	{
+		const UE::Geometry::FPolygon2d& Hole = Holes[HoleIndex];
+		for (int VertexIndex = 0; VertexIndex < Hole.VertexCount(); ++VertexIndex)
+		{
+			SegmentIndexToSegmentAndHoleIndices.Add(SegmentCount++, TPair<int, int>(VertexIndex, HoleIndex));
+		}
+	}
+}
\ No newline at end of file
diff --git a/Engine/Plugins/PCG/Source/PCG/Private/Data/PCGPolygon2DInteriorData.cpp b/Engine/Plugins/PCG/Source/PCG/Private/Data/PCGPolygon2DInteriorData.cpp
new file mode 100644
--- /dev/null
+++ b/Engine/Plugins/PCG/Source/PCG/Private/Data/PCGPolygon2DInteriorData.cpp
@@ -0,0 +1,84 @@
+// Copyright Epic Games, Inc. All Rights Reserved.
+
+#include "Data/PCGPolygon2DInteriorData.h"
+
+#include "PCGContext.h"
+#include "Data/PCGPointArrayData.h"
+#include "Data/PCGPointData.h"
+#include "Elements/PCGSurfaceSampler.h"
+
+#include "Serialization/ArchiveCrc32.h"
+
+void UPCGPolygon2DInteriorSurfaceData::Initialize(FPCGContext* Context, const UPCGPolygon2DData* InPolygonData)
+{
+	check(InPolygonData);
+	Polygon = InPolygonData;
+}
+
+void UPCGPolygon2DInteriorSurfaceData::AddToCrc(FArchiveCrc32& Ar, bool bFullDataCrc) const
+{
+	Super::AddToCrc(Ar, bFullDataCrc);
+
+	if (Metadata)
+	{
+		Metadata->AddToCrc(Ar, bFullDataCrc);
+	}
+
+	FString ClassName = StaticClass()->GetPathName();
+	Ar << ClassName;
+
+	// @todo_pcg either move the contents of the polygon in this class or use ComputeCrc to defer to the polygon for crc.
+	if (Polygon)
+	{
+		Polygon->AddToCrc(Ar, bFullDataCrc);
+	}
+}
+
+FBox UPCGPolygon2DInteriorSurfaceData::GetBounds() const
+{
+	check(Polygon);
+	return Polygon->GetBounds();
+}
+
+bool UPCGPolygon2DInteriorSurfaceData::SamplePoint(const FTransform& InTransform, const FBox& Bounds, FPCGPoint& OutPoint, UPCGMetadata* OutMetadata) const
+{
+	check(Polygon);
+	return Polygon->SamplePoint(InTransform, Bounds, OutPoint, OutMetadata);
+}
+
+bool UPCGPolygon2DInteriorSurfaceData::ProjectPoint(const FTransform& InTransform, const FBox& Bounds, const FPCGProjectionParams& Params, FPCGPoint& OutPoint, UPCGMetadata* OutMetadata) const
+{
+	check(Polygon);
+	return Polygon->ProjectPoint(InTransform, Bounds, Params, OutPoint, OutMetadata);
+}
+
+UPCGSpatialData* UPCGPolygon2DInteriorSurfaceData::CopyInternal(FPCGContext* Context) const
+{
+	UPCGPolygon2DInteriorSurfaceData* NewData = FPCGContext::NewObject_AnyThread<UPCGPolygon2DInteriorSurfaceData>(Context);
+
+	CopyBaseSurfaceData(NewData);
+	NewData->Polygon = Polygon;
+
+	return NewData;
+}
+
+const UPCGPointData* UPCGPolygon2DInteriorSurfaceData::CreatePointData(FPCGContext* Context) const
+{
+	return CastChecked<UPCGPointData>(CreateBasePointData(Context, UPCGPointData::StaticClass()));
+}
+
+const UPCGPointArrayData* UPCGPolygon2DInteriorSurfaceData::CreatePointArrayData(FPCGContext* Context, const FBox& InBounds) const
+{
+	return CastChecked<UPCGPointArrayData>(CreateBasePointData(Context, UPCGPointArrayData::StaticClass()));
+}
+
+const UPCGBasePointData* UPCGPolygon2DInteriorSurfaceData::CreateBasePointData(FPCGContext* Context, TSubclassOf<UPCGBasePointData> PointDataClass) const
+{
+	TRACE_CPUPROFILER_EVENT_SCOPE(UPCGPolygon2DInteriorSurfaceData::CreatePointData);
+
+	// Default surface sampling
+	const PCGSurfaceSampler::FSurfaceSamplerParams SamplerParams;
+	UPCGBasePointData* Data = PCGSurfaceSampler::SampleSurface(Context, SamplerParams, this, /*InBoundingShape=*/nullptr, GetBounds(), PointDataClass);
+
+	return Data;
+}
diff --git a/Engine/Plugins/PCG/Source/PCG/Private/Elements/Polygon/PCGCreatePolygon2D.cpp b/Engine/Plugins/PCG/Source/PCG/Private/Elements/Polygon/PCGCreatePolygon2D.cpp
new file mode 100644
--- /dev/null
+++ b/Engine/Plugins/PCG/Source/PCG/Private/Elements/Polygon/PCGCreatePolygon2D.cpp
@@ -0,0 +1,170 @@
+// Copyright Epic Games, Inc. All Rights Reserved.
+
+#include "Elements/Polygon/PCGCreatePolygon2D.h"
+
+#include "PCGContext.h"
+#include "Data/PCGBasePointData.h"
+#include "Data/PCGSplineData.h"
+#include "Data/PCGPolygon2DData.h"
+#include "Utils/PCGLogErrors.h"
+
+#include "CompGeom/PolygonTriangulation.h"
+#include "Curve/PolygonOffsetUtils.h"
+
+TArray<FPCGPinProperties> UPCGCreatePolygon2DSettings::InputPinProperties() const
+{
+	TArray<FPCGPinProperties> PinProperties;
+	// @todo_pcg: add landscape spline support?
+	PinProperties.Emplace(PCGPinConstants::DefaultInputLabel, EPCGDataType::Point | EPCGDataType::Spline);
+
+	return PinProperties;
+}
+
+bool FPCGCreatePolygonElement::ExecuteInternal(FPCGContext* Context) const
+{
+	TRACE_CPUPROFILER_EVENT_SCOPE(FPCGSplitSplineElement::Execute);
+
+	const UPCGCreatePolygon2DSettings* Settings = Context->GetInputSettings<UPCGCreatePolygon2DSettings>();
+	check(Settings);
+
+	TArray<FPCGTaggedData> Inputs = Context->InputData.GetInputsByPin(PCGPinConstants::DefaultInputLabel);
+
+	TArray<FVector> Positions;
+
+	for (const FPCGTaggedData& Input : Inputs)
+	{
+		const UPCGSpatialData* SpatialData = nullptr;
+
+		bool bInputIsOpen = false;
+		bool bInputFromClosedSpline = false;
+
+		if (const UPCGSplineData* SplineData = Cast<UPCGSplineData>(Input.Data))
+		{
+			SpatialData = SplineData;
+			bInputFromClosedSpline = SplineData->IsClosed();
+			bInputIsOpen = !bInputFromClosedSpline;
+
+			// @todo_pcg: this method doesn't allow proper metadata inheritance, we will need to update it.
+			const double ErrorTolerance = FMath::Max(Settings->SplineMaxDiscretizationError, 0.001);
+			SplineData->SplineStruct.ConvertSplineToPolyLine(ESplineCoordinateSpace::World, ErrorTolerance, Positions);
+		}
+		else if (const UPCGBasePointData* InputPointData = Cast<UPCGBasePointData>(Input.Data))
+		{
+			SpatialData = InputPointData;
+
+			const TConstPCGValueRange<FTransform> TransformRange = InputPointData->GetConstTransformValueRange();
+
+			if (TransformRange.Num() > 0)
+			{
+				Positions.Reset();
+				Positions.Reserve(TransformRange.Num());
+				for (const FTransform& Transform : TransformRange)
+				{
+					Positions.Emplace(Transform.GetLocation());
+				}
+			}
+		}
+		else
+		{
+			// Unsupported type
+			PCGLog::InputOutput::LogInvalidInputDataError(Context);
+			continue;
+		}
+
+		// @todo_pcg: Eventually, add metadata.
+		auto CreatePolygonData = [Context, SpatialData, &Input]()
+		{
+			UPCGPolygon2DData* Polygon = FPCGContext::NewObject_AnyThread<UPCGPolygon2DData>(Context);
+			Polygon->InitializeFromData(SpatialData);
+			Context->OutputData.TaggedData.Add_GetRef(Input).Data = Polygon;
+			return Polygon;
+		};
+
+		// Project the positions so they lie on a plane. Note that we'll assume at this point that the polygon is always closed, regardless of whether that's true or not, without loss of generality.
+		FVector PlaneNormal = FVector::UpVector;
+		FVector PlaneCentroid = FVector::ZeroVector;
+		PolygonTriangulation::ComputePolygonPlane(Positions, PlaneNormal, PlaneCentroid);
+
+		// If the polygon is a point/line or otherwise has a null area, assume normal
+		if (PlaneNormal.SquaredLength() < UE_KINDA_SMALL_NUMBER)
+		{
+			PlaneNormal = FVector::UpVector;
+		}
+
+		FTransform PlaneTransform;
+		PlaneTransform.SetRotation(FRotationMatrix::MakeFromZ(PlaneNormal).ToQuat());
+		PlaneTransform.SetLocation(PlaneCentroid);
+
+		TArray<FVector2D> Positions2D;
+		Positions2D.Reserve(Positions.Num());
+		for (const FVector& Position : Positions)
+		{
+			FVector LocalPosition = PlaneTransform.InverseTransformPosition(Position);
+			Positions2D.Emplace(LocalPosition.X, LocalPosition.Y);
+		}
+
+		const bool bCreateOpenPolygon = (Settings->InputType == EPCGCreatePolygonInputType::ForceOpen || (bInputIsOpen && Settings->InputType == EPCGCreatePolygonInputType::Automatic));
+
+		if (bCreateOpenPolygon)
+		{
+			// Special case: if we had a closed polygon but forced this here to be considered as an opened polygon, we'll still respect the loop.
+			if (bInputFromClosedSpline && !Positions2D.IsEmpty())
+			{
+				const FVector2D FirstPosition2D = Positions2D[0];
+				Positions2D.Add(FirstPosition2D);
+			}
+
+			//Note: same values as FGeometryScriptOpenPathOffsetOptions, but could be parametrized in advanced settings
+			UE::Geometry::FOffsetPolygon2d PolygonOffset;
+			PolygonOffset.EndType = UE::Geometry::EPolygonOffsetEndType::Square;
+			PolygonOffset.JoinType = UE::Geometry::EPolygonOffsetJoinType::Square;
+			PolygonOffset.MiterLimit = 2.0;
+			PolygonOffset.MaxStepsPerRadian = 10.0;
+			PolygonOffset.DefaultStepsPerRadianScale = 1.0;
+			PolygonOffset.Offset = Settings->OpenPolygonWidth;
+			PolygonOffset.Polygons.Add(MoveTemp(Positions2D));
+
+			bool bOperationSuccessful = PolygonOffset.ComputeResult();
+			if (bOperationSuccessful)
+			{
+				for (int PolyIndex = 0; PolyIndex < PolygonOffset.Result.Num(); ++PolyIndex)
+				{
+					UPCGPolygon2DData* Polygon = CreatePolygonData();
+
+					if (PolygonOffset.Result[PolyIndex].OuterIsClockwise())
+					{
+						PolygonOffset.Result[PolyIndex].Reverse();
+					}
+
+					Polygon->SetPolygon(MoveTemp(PolygonOffset.Result[PolyIndex]));
+					Polygon->SetTransform(PlaneTransform);
+				}
+			}
+			else
+			{
+				// The offset failed, not something we can do a lot about at this point.
+				PCGLog::LogWarningOnGraph(NSLOCTEXT("PCGCreatePolygon2DElement", "ErrorInPolygonOffset", "Unable to create extruded polygon from open polygon."), Context);
+				continue;
+			}
+		}
+		else
+		{
+			// Create the outer polygon
+			UE::Geometry::FPolygon2d OuterPolygon(Positions2D);
+			UE::Geometry::FGeneralPolygon2d NewPolygon(OuterPolygon);
+
+			if (NewPolygon.OuterIsClockwise())
+			{
+				NewPolygon.Reverse();
+			}
+
+			UPCGPolygon2DData* Polygon = CreatePolygonData();
+			Polygon->SetPolygon(MoveTemp(NewPolygon));
+			Polygon->SetTransform(PlaneTransform);
+
+			// @todo_pcg: Create hole polygons / multi polygons...
+		}
+	}
+
+	return true;
+}
diff --git a/Engine/Plugins/PCG/Source/PCG/Private/Elements/Polygon/PCGPolygon2DUtils.cpp b/Engine/Plugins/PCG/Source/PCG/Private/Elements/Polygon/PCGPolygon2DUtils.cpp
new file mode 100644
--- /dev/null
+++ b/Engine/Plugins/PCG/Source/PCG/Private/Elements/Polygon/PCGPolygon2DUtils.cpp
@@ -0,0 +1,24 @@
+// Copyright Epic Games, Inc. All Rights Reserved.
+
+#include "Elements/Polygon/PCGPolygon2DUtils.h"
+
+namespace PCGPolygon2DUtils
+{
+
+TArray<FPCGPinProperties> DefaultPolygonInputPinProperties()
+{
+	TArray<FPCGPinProperties> PinProperties;
+	PinProperties.Emplace_GetRef(PCGPinConstants::DefaultInputLabel, EPCGDataType::Polygon2D).SetRequiredPin();
+
+	return PinProperties;
+}
+
+TArray<FPCGPinProperties> DefaultPolygonOutputPinProperties()
+{
+	TArray<FPCGPinProperties> PinProperties;
+	PinProperties.Emplace(PCGPinConstants::DefaultOutputLabel, EPCGDataType::Polygon2D);
+
+	return PinProperties;
+}
+
+} // namespace PCGPolygon2DUtils
diff --git a/Engine/Plugins/PCG/Source/PCG/Private/Elements/Polygon/PCGSurfaceFromPolygon2D.cpp b/Engine/Plugins/PCG/Source/PCG/Private/Elements/Polygon/PCGSurfaceFromPolygon2D.cpp
new file mode 100644
--- /dev/null
+++ b/Engine/Plugins/PCG/Source/PCG/Private/Elements/Polygon/PCGSurfaceFromPolygon2D.cpp
@@ -0,0 +1,42 @@
+// Copyright Epic Games, Inc. All Rights Reserved.
+
+#include "Elements/Polygon/PCGSurfaceFromPolygon2D.h"
+
+#include "PCGContext.h"
+#include "Data/PCGPolygon2DData.h"
+#include "Data/PCGPolygon2DInteriorData.h"
+
+TArray<FPCGPinProperties> UPCGCreateSurfaceFromPolygon2DSettings::OutputPinProperties() const
+{
+	TArray<FPCGPinProperties> PinProperties;
+	PinProperties.Emplace(PCGPinConstants::DefaultOutputLabel, EPCGDataType::Surface);
+	return PinProperties;
+}
+
+bool FPCGCreateSurfaceFromPolygon2DElement::ExecuteInternal(FPCGContext* Context) const
+{
+	TRACE_CPUPROFILER_EVENT_SCOPE(FPCGCreateSurfaceFromPolygon2DElement::Execute);
+
+	const UPCGCreateSurfaceFromPolygon2DSettings* Settings = Context->GetInputSettings<UPCGCreateSurfaceFromPolygon2DSettings>();
+	check(Settings);
+
+	const TArray<FPCGTaggedData> Inputs = Context->InputData.GetInputsByPin(PCGPinConstants::DefaultInputLabel);
+	TArray<FPCGTaggedData>& Outputs = Context->OutputData.TaggedData;
+
+	for (const FPCGTaggedData& Input : Inputs)
+	{
+		const UPCGPolygon2DData* PolygonData = Cast<UPCGPolygon2DData>(Input.Data);
+
+		if (!PolygonData)
+		{
+			continue;
+		}
+
+		UPCGPolygon2DInteriorSurfaceData* SurfaceData = FPCGContext::NewObject_AnyThread<UPCGPolygon2DInteriorSurfaceData>(Context);
+		SurfaceData->Initialize(Context, PolygonData);
+
+		Outputs.Add_GetRef(Input).Data = SurfaceData;
+	}
+
+	return true;
+}
\ No newline at end of file
diff --git a/Engine/Plugins/PCG/Source/PCG/Public/Data/PCGPolygon2DData.h b/Engine/Plugins/PCG/Source/PCG/Public/Data/PCGPolygon2DData.h
new file mode 100644
--- /dev/null
+++ b/Engine/Plugins/PCG/Source/PCG/Public/Data/PCGPolygon2DData.h
@@ -0,0 +1,83 @@
+// Copyright Epic Games, Inc. All Rights Reserved.
+
+#pragma once
+
+#include "Data/PCGPolyLineData.h"
+
+#include "PCGPointData.h"
+
+#include "Curve/GeneralPolygon2.h"
+
+#include "PCGPolygon2DData.generated.h"
+
+/** Data representing a single 2D polygon with a 3D transform (for spatial operations). */
+UCLASS(MinimalAPI, BlueprintType, ClassGroup = (Procedural))
+class UPCGPolygon2DData : public UPCGPolyLineData
+{
+	GENERATED_BODY()
+
+public:
+	// ~Begin UObject interface
+	PCG_API virtual void Serialize(FArchive& InArchive) override;
+	// ~End UObject interface
+
+	// ~Begin UPCGData interface
+	virtual EPCGDataType GetDataType() const override { return EPCGDataType::Polygon2D; }
+	// @todo_pcg: supported domain ids.. elements?
+	PCG_API virtual void AddToCrc(FArchiveCrc32& Ar, bool bFullDataCrc) const override;
+	// ~End UPCGData interface
+
+	// ~Begin UPCGSpatialData interface
+	virtual int GetDimension() const override { return 2; }
+	PCG_API virtual FBox GetBounds() const override;
+	PCG_API virtual bool SamplePoint(const FTransform& Transform, const FBox& Bounds, FPCGPoint& OutPoint, UPCGMetadata* OutMetadata) const override;
+	PCG_API virtual bool ProjectPoint(const FTransform& Transform, const FBox& Bounds, const FPCGProjectionParams& Params, FPCGPoint& OutPoint, UPCGMetadata* OutMetadata) const override;
+
+protected:
+	virtual UPCGSpatialData* CopyInternal(FPCGContext* Context) const override;
+	// @todo_pcg: initialize target metadata?
+	// ~End UPCGSpatialData interface
+
+public:
+	// ~Begin UPCGSpatialDataWithPointCache interface
+	PCG_API virtual const UPCGPointData* CreatePointData(FPCGContext* Context) const override;
+	PCG_API virtual const UPCGPointArrayData* CreatePointArrayData(FPCGContext* Context, const FBox& InBounds) const override;
+	// ~End UPCGSpatialDataWithPointCache interface
+
+	// ~Begin UPCGPolyLineData interface - all these functions will return world space information (e.g. including scale, ...) unless specified
+	PCG_API virtual FTransform GetTransform() const override { return Transform; }
+	PCG_API virtual int GetNumSegments() const override;
+	PCG_API virtual double GetSegmentLength(int SegmentIndex) const override;
+	PCG_API virtual FVector GetLocationAtAlpha(float Alpha) const override;
+	PCG_API virtual FTransform GetTransformAtAlpha(float Alpha) const override;
+	PCG_API virtual FTransform GetTransformAtDistance(int SegmentIndex, double Distance, bool bWorldSpace = true, FBox* OutBounds = nullptr) const override;
+	PCG_API virtual float GetInputKeyAtDistance(int SegmentIndex, double Distance) const override;
+	PCG_API virtual double GetDistanceAtSegmentStart(int SegmentIndex) const override;
+	// @todo_pcg: write metadata to point + get const vertices entry keys
+	// ~End UPCGPolyLineData interface
+
+	// @todo_pcg maybe static create attribute accessor methods if we support elements
+public:
+	const UE::Geometry::FGeneralPolygon2d& GetPolygon() const { return Polygon; }
+
+	void SetPolygon(const UE::Geometry::FGeneralPolygon2d& InPolygon);
+	void SetPolygon(UE::Geometry::FGeneralPolygon2d&& InPolygon);
+	void SetTransform(const FTransform& InTransform, bool bCheckWinding = false);
+
+protected:
+	const UPCGBasePointData* CreateBasePointData(FPCGContext* Context, TSubclassOf<UPCGBasePointData> PointDataClass) const;
+
+private:
+	void UpdateMappings();
+
+protected:
+	// Note: can't be a property
+	UE::Geometry::FGeneralPolygon2d Polygon;
+
+	UPROPERTY()
+	FTransform Transform;
+
+private:
+	TMap<int, TPair<int, int>> SegmentIndexToSegmentAndHoleIndices;
+	int SegmentCount = 0;
+};
\ No newline at end of file
diff --git a/Engine/Plugins/PCG/Source/PCG/Public/Data/PCGPolygon2DInteriorData.h b/Engine/Plugins/PCG/Source/PCG/Public/Data/PCGPolygon2DInteriorData.h
new file mode 100644
--- /dev/null
+++ b/Engine/Plugins/PCG/Source/PCG/Public/Data/PCGPolygon2DInteriorData.h
@@ -0,0 +1,45 @@
+// Copyright Epic Games, Inc. All Rights Reserved.
+
+#pragma once
+
+#include "Data/PCGSurfaceData.h"
+#include "Data/PCGPolygon2DData.h"
+
+#include "PCGPolygon2DInteriorData.generated.h"
+
+UCLASS(MinimalAPI, BlueprintType, ClassGroup = (Procedural))
+class UPCGPolygon2DInteriorSurfaceData : public UPCGSurfaceData
+{
+	GENERATED_BODY()
+
+public:
+	PCG_API void Initialize(FPCGContext* Context, const UPCGPolygon2DData* InPolygonData);
+
+	// ~Begin UPCGData interface
+	virtual void AddToCrc(FArchiveCrc32& Ar, bool bFullDataCrc) const override;
+	// ~End UPCGData interface
+
+	// ~Begin UPCGSpatialData interface
+	PCG_API virtual FBox GetBounds() const override;
+	PCG_API virtual bool SamplePoint(const FTransform& Transform, const FBox& Bounds, FPCGPoint& OutPoint, UPCGMetadata* OutMetadata) const override;
+	PCG_API virtual bool ProjectPoint(const FTransform& Transform, const FBox& Bounds, const FPCGProjectionParams& Params, FPCGPoint& OutPoint, UPCGMetadata* OutMetadata) const override;
+	virtual bool HasNonTrivialTransform() const override { return true; }
+protected:
+	virtual UPCGSpatialData* CopyInternal(FPCGContext* Context) const override;
+	// ~End UPCGSpatialData interface
+
+public:
+	// ~Begin UPCGSpatialDataWithPointCache interface
+	virtual const UPCGPointData* CreatePointData(FPCGContext* Context) const override;
+	virtual const UPCGPointArrayData* CreatePointArrayData(FPCGContext* Context, const FBox& InBounds) const override;
+	// ~End UPCGConcreteDataWithPointCache interface
+
+	const UPCGPolygon2DData* GetPolygonData() const { return Polygon; }
+
+protected:
+	const UPCGBasePointData* CreateBasePointData(FPCGContext* Context, TSubclassOf<UPCGBasePointData> PointDataClass) const;
+
+protected:
+	UPROPERTY()
+	TObjectPtr<const UPCGPolygon2DData> Polygon;
+};
\ No newline at end of file
diff --git a/Engine/Plugins/PCG/Source/PCG/Public/Elements/Polygon/PCGCreatePolygon2D.h b/Engine/Plugins/PCG/Source/PCG/Public/Elements/Polygon/PCGCreatePolygon2D.h
new file mode 100644
--- /dev/null
+++ b/Engine/Plugins/PCG/Source/PCG/Public/Elements/Polygon/PCGCreatePolygon2D.h
@@ -0,0 +1,65 @@
+// Copyright Epic Games, Inc. All Rights Reserved.
+
+#pragma once
+
+#include "PCGSettings.h"
+#include "Elements/Polygon/PCGPolygon2DUtils.h"
+
+#include "PCGCreatePolygon2D.generated.h"
+
+UENUM()
+enum class EPCGCreatePolygonInputType : uint8
+{
+	Automatic UMETA(Tooltip="Assumes that point data represent a closed polygon, but for splines will respect spline type."),
+	ForceOpen UMETA(Tooltip="Always considers inputs as open polygons."),
+	ForceClosed UMETA(Tooltip="Always considers inputs as closed polygons.")
+};
+
+class FPCGCreatePolygonElement : public IPCGElement
+{
+public:
+	virtual bool ExecuteInternal(FPCGContext* InContext) const override;
+	virtual EPCGElementExecutionLoopMode ExecutionLoopMode(const UPCGSettings* Settings) const override { return EPCGElementExecutionLoopMode::SinglePrimaryPin; }
+
+protected:
+	virtual bool SupportsBasePointDataInputs(FPCGContext* InContext) const { return true; }
+};
+
+/**
+* Creates polygon(s) from the specified point data.
+*/
+UCLASS(MinimalAPI, BlueprintType, ClassGroup=(Procedural))
+class UPCGCreatePolygon2DSettings : public UPCGSettings
+{
+	GENERATED_BODY()
+
+public:
+	//~Begin UPCGSettings interface
+#if WITH_EDITOR
+	virtual FName GetDefaultNodeName() const override { return FName(TEXT("CreatePolygon2D")); }
+	virtual FText GetDefaultNodeTitle() const override { return NSLOCTEXT("PCGCreatePolygon2DElement", "NodeTitle", "Create Polygon 2D"); }
+	virtual FText GetNodeTooltipText() const override { return NSLOCTEXT("PCGCreatePolygon2DElement", "NodeTooltip", "Creates a (closed) 2d polygon per input."); }
+	virtual EPCGSettingsType GetType() const override { return EPCGSettingsType::Spatial; }
+#endif
+
+protected:
+	virtual TArray<FPCGPinProperties> InputPinProperties() const override;
+	virtual TArray<FPCGPinProperties> OutputPinProperties() const override { return PCGPolygon2DUtils::DefaultPolygonOutputPinProperties(); }
+	virtual FPCGElementPtr CreateElement() const override { return MakeShared<FPCGCreatePolygonElement>(); }
+	//~End UPCGSettings interface
+
+public:
+	//@todo_pcg: add hole support & multi polygon
+
+	/** Controls the way input data will be considered - either as closed or open paths. */
+	UPROPERTY(EditAnywhere, BlueprintReadWrite, Category = Settings, meta = (PCG_Overridable))
+	EPCGCreatePolygonInputType InputType = EPCGCreatePolygonInputType::Automatic;
+
+	/** Target width of the polygon generated from an open path. */
+	UPROPERTY(EditAnywhere, BlueprintReadWrite, Category = Settings, meta = (EditCondition = "InputType != EPCGCreatePolygonInputType::ForceClosed", PCG_Overridable))
+	double OpenPolygonWidth = 100.0;
+
+	/* Maximum squared distance before we need to subdivide a segment again as part of the spline discretization to a polygon. */
+	UPROPERTY(EditAnywhere, BlueprintReadWrite, Category = Settings, AdvancedDisplay, meta = (ClampMin = 0.001, PCG_Overridable))
+	double SplineMaxDiscretizationError = 1.0;
+};
\ No newline at end of file
diff --git a/Engine/Plugins/PCG/Source/PCG/Public/Elements/Polygon/PCGPolygon2DUtils.h b/Engine/Plugins/PCG/Source/PCG/Public/Elements/Polygon/PCGPolygon2DUtils.h
new file mode 100644
--- /dev/null
+++ b/Engine/Plugins/PCG/Source/PCG/Public/Elements/Polygon/PCGPolygon2DUtils.h
@@ -0,0 +1,11 @@
+// Copyright Epic Games, Inc. All Rights Reserved.
+
+#pragma once
+
+#include "PCGSettings.h"
+
+namespace PCGPolygon2DUtils
+{
+	TArray<FPCGPinProperties> DefaultPolygonInputPinProperties();
+	TArray<FPCGPinProperties> DefaultPolygonOutputPinProperties();
+}
\ No newline at end of file
diff --git a/Engine/Plugins/PCG/Source/PCG/Public/Elements/Polygon/PCGSurfaceFromPolygon2D.h b/Engine/Plugins/PCG/Source/PCG/Public/Elements/Polygon/PCGSurfaceFromPolygon2D.h
new file mode 100644
--- /dev/null
+++ b/Engine/Plugins/PCG/Source/PCG/Public/Elements/Polygon/PCGSurfaceFromPolygon2D.h
@@ -0,0 +1,38 @@
+// Copyright Epic Games, Inc. All Rights Reserved.
+
+#include "PCGSettings.h"
+#include "Elements/Polygon/PCGPolygon2DUtils.h"
+
+#include "PCGSurfaceFromPolygon2D.generated.h"
+
+class FPCGCreateSurfaceFromPolygon2DElement : public IPCGElement
+{
+protected:
+	virtual bool ExecuteInternal(FPCGContext* Context) const override;
+	virtual EPCGElementExecutionLoopMode ExecutionLoopMode(const UPCGSettings* Settings) const override { return EPCGElementExecutionLoopMode::SinglePrimaryPin; }
+};
+
+/** 
+* Creates a surface representation from a polygon data. 
+*/
+UCLASS(MinimalAPI, BlueprintType, ClassGroup = (Procedural))
+class UPCGCreateSurfaceFromPolygon2DSettings : public UPCGSettings
+{
+	GENERATED_BODY()
+
+public:
+#if WITH_EDITOR
+	virtual FName GetDefaultNodeName() const override { return FName(TEXT("CreateSurfaceFromPolygon2D")); }
+	virtual FText GetDefaultNodeTitle() const override { return NSLOCTEXT("PCGCreateSurfaceFromPolygon2DElement", "NodeTitle", "Create Surface From Polygon2D"); }
+	virtual FText GetNodeTooltipText() const override { return NSLOCTEXT("PCGCreateSurfaceFromPolygon2DElement", "NodeTooltip", "Creates a surface data from a Polygon 2D data."); }
+	virtual EPCGSettingsType GetType() const override { return EPCGSettingsType::Spatial; }
+	virtual bool ShouldDrawNodeCompact() const override { return false; }
+#endif
+	virtual bool HasExecutionDependencyPin() const override { return false; }
+
+protected:
+	virtual TArray<FPCGPinProperties> InputPinProperties() const override { return PCGPolygon2DUtils::DefaultPolygonInputPinProperties(); }
+	virtual TArray<FPCGPinProperties> OutputPinProperties() const override;
+	virtual FPCGElementPtr CreateElement() const override { return MakeShared<FPCGCreateSurfaceFromPolygon2DElement>(); }
+	//~End UPCGSettings interface
+};
\ No newline at end of file
diff --git a/Engine/Plugins/PCG/Source/PCG/Public/PCGCommon.h b/Engine/Plugins/PCG/Source/PCG/Public/PCGCommon.h
--- a/Engine/Plugins/PCG/Source/PCG/Public/PCGCommon.h
+++ b/Engine/Plugins/PCG/Source/PCG/Public/PCGCommon.h
@@ -90,7 +90,7 @@ ENUM_CLASS_FLAGS(EPCGChangeType);
 // This is why we have a parallel enum just below that must match on a name basis 1:1 to allow the make/break functions to work properly
 // in blueprint.
 // WARNING: Please be mindful that combination of flags that are not explicitly defined there won't be serialized correctly, inducing data loss.
-UENUM(meta = (Bitflags))
+UENUM(Flags, meta = (Bitflags))
 enum class EPCGDataType : uint32
 {
 	None = 0 UMETA(Hidden),
@@ -98,7 +98,8 @@ enum class EPCGDataType : uint32
 
 	Spline = 1 << 2,
 	LandscapeSpline = 1 << 3,
-	PolyLine = Spline | LandscapeSpline UMETA(DisplayName = "Curve"),
+	Polygon2D = 1 << 13,
+	PolyLine = Spline | LandscapeSpline | Polygon2D UMETA(DisplayName = "Curve"),
 
 	Landscape = 1 << 4,
 	Texture = 1 << 5,
@@ -131,7 +132,7 @@ enum class EPCGDataType : uint32
 
 	Param = 1 << 27 UMETA(DisplayName = "Attribute Set"),
 
-	// Combination of Param and Point, necessary for named-based serialization of enums.
+	// Combination of Param and Point, necessary for name-based serialization of enums.
 	PointOrParam = Point | Param UMETA(DisplayName = "Point or Attribute Set"),
 
 	Settings = 1 << 28 UMETA(Hidden),
@@ -167,7 +168,8 @@ enum class EPCGExclusiveDataType : uint8
 	PointOrParam UMETA(DisplayName = "Point or Attribute Set"),
 	DynamicMesh,
 	StaticMeshResource,
-	Resource UMETA(Hidden)
+	Resource UMETA(Hidden),
+	Polygon2D
 };
 
 UENUM(meta=(DisplayName="PCG Container Type"))
diff --git a/Engine/Plugins/PCG/Source/PCGEditor/Private/DataVisualizations/PCGPolygon2DDataVisualization.cpp b/Engine/Plugins/PCG/Source/PCGEditor/Private/DataVisualizations/PCGPolygon2DDataVisualization.cpp
new file mode 100644
--- /dev/null
+++ b/Engine/Plugins/PCG/Source/PCGEditor/Private/DataVisualizations/PCGPolygon2DDataVisualization.cpp
@@ -0,0 +1,85 @@
+// Copyright Epic Games, Inc. All Rights Reserved.
+
+#include "DataVisualizations/PCGPolygon2DDataVisualization.h"
+
+#include "Data/PCGPolygon2DData.h"
+#include "Data/PCGPolygon2DInteriorData.h"
+
+#include "AdvancedPreviewScene.h"
+#include "EditorViewportClient.h"
+#include "Components/SplineComponent.h"
+
+FPCGSetupSceneFunc FPCGPolygon2DDataVisualization::GetViewportSetupFunc(const UPCGData* Data) const
+{
+	return[this, Data](FPCGSceneSetupParams& InOutParams)
+	{
+		check(InOutParams.Scene);
+		check(InOutParams.EditorViewportClient);
+
+		const UPCGPolygon2DData* PolygonData = Cast<UPCGPolygon2DData>(Data);
+
+		if (!PolygonData)
+		{
+			if (const UPCGPolygon2DInteriorSurfaceData* PolygonSurfaceData = Cast<UPCGPolygon2DInteriorSurfaceData>(Data))
+			{
+				PolygonData = PolygonSurfaceData->GetPolygonData();
+			}
+		}
+
+		if (!PolygonData)
+		{
+			return;
+		}
+
+		const FTransform& Transform = PolygonData->GetTransform();
+
+		auto CreateSplineComponent = [PolygonData, &InOutParams, &Transform](const UE::Geometry::FPolygon2d& Polygon)
+		{
+			TObjectPtr<USplineComponent> SplineComponent = NewObject<USplineComponent>(GetTransientPackage(), NAME_None, RF_Transient);
+			
+			// Setup spline - done similarly to the Create Spline node.
+			{
+				// Create points
+				TArray<FSplinePoint> SplinePoints;
+				const TArray<FVector2d>& PolygonVertices = Polygon.GetVertices();
+
+				SplinePoints.Reserve(PolygonVertices.Num());
+
+				for(int VertexIndex = 0; VertexIndex < PolygonVertices.Num(); ++VertexIndex)
+				{
+					SplinePoints.Emplace(static_cast<float>(VertexIndex),
+						FVector(PolygonVertices[VertexIndex], 0.0),
+						ESplinePointType::Linear);
+				}
+
+				// Set points to component
+				SplineComponent->SetWorldTransform(Transform);
+				SplineComponent->DefaultUpVector = Transform.GetRotation().GetUpVector();
+				SplineComponent->ClearSplinePoints(/*bUpdateSpline=*/false);
+				SplineComponent->AddPoints(SplinePoints, /*bUpdateSpline=*/false);
+				SplineComponent->SetClosedLoop(true);
+				SplineComponent->UpdateSpline();
+				SplineComponent->UpdateBounds();
+			}
+
+			if (GEditor->PreviewPlatform.GetEffectivePreviewFeatureLevel() <= ERHIFeatureLevel::ES3_1)
+			{
+				SplineComponent->SetMobility(EComponentMobility::Static);
+			}
+
+			InOutParams.ManagedResources.Add(SplineComponent);
+			InOutParams.Scene->AddComponent(SplineComponent, Transform);
+
+			return SplineComponent;
+		};
+
+		TObjectPtr<USplineComponent> OuterPolygon = CreateSplineComponent(PolygonData->GetPolygon().GetOuter());
+
+		for (const UE::Geometry::FPolygon2d& Hole : PolygonData->GetPolygon().GetHoles())
+		{
+			CreateSplineComponent(Hole);
+		}
+
+		InOutParams.FocusBounds = OuterPolygon->CalcBounds(Transform);
+	};
+}
\ No newline at end of file
diff --git a/Engine/Plugins/PCG/Source/PCGEditor/Private/PCGEditorModule.cpp b/Engine/Plugins/PCG/Source/PCGEditor/Private/PCGEditorModule.cpp
--- a/Engine/Plugins/PCG/Source/PCGEditor/Private/PCGEditorModule.cpp
+++ b/Engine/Plugins/PCG/Source/PCGEditor/Private/PCGEditorModule.cpp
@@ -12,6 +12,8 @@
 #include "Data/PCGCollisionShapeData.h"
 #include "Data/PCGCollisionWrapperData.h"
 #include "Data/PCGLandscapeData.h"
+#include "Data/PCGPolygon2DData.h"
+#include "Data/PCGPolygon2DInteriorData.h"
 #include "Data/PCGPrimitiveData.h"
 #include "Data/PCGSpatialData.h"
 #include "Data/PCGSplineData.h"
@@ -31,6 +33,7 @@
 #include "DataVisualizations/PCGBaseTextureDataVisualization.h"
 #include "DataVisualizations/PCGLandscapeDataVisualization.h"
 #include "DataVisualizations/PCGParamDataVisualization.h"
+#include "DataVisualizations/PCGPolygon2DDataVisualization.h"
 #include "DataVisualizations/PCGSpatialDataVisualization.h"
 #include "DataVisualizations/PCGSplineDataVisualization.h"
 #include "DataVisualizations/PCGStaticMeshDataVisualization.h"
@@ -808,6 +811,8 @@ void FPCGEditorModule::RegisterPCGDataVisualizations()
 	DataVisRegistry.InternalRegistry.Add(UPCGLandscapeData::StaticClass(), MakeUnique<const IPCGLandscapeDataVisualization>());
 	DataVisRegistry.InternalRegistry.Add(UPCGStaticMeshResourceData::StaticClass(), MakeUnique<const IPCGStaticMeshDataVisualization>());
 	DataVisRegistry.InternalRegistry.Add(UPCGBaseTextureData::StaticClass(), MakeUnique<const FPCGBaseTextureDataVisualization>());
+	DataVisRegistry.InternalRegistry.Add(UPCGPolygon2DData::StaticClass(), MakeUnique<const FPCGPolygon2DDataVisualization>());
+	DataVisRegistry.InternalRegistry.Add(UPCGPolygon2DInteriorSurfaceData::StaticClass(), MakeUnique<const FPCGPolygon2DDataVisualization>());
 }
 
 void FPCGEditorModule::UnregisterPCGDataVisualizations()
diff --git a/Engine/Plugins/PCG/Source/PCGEditor/Public/DataVisualizations/PCGPolygon2DDataVisualization.h b/Engine/Plugins/PCG/Source/PCGEditor/Public/DataVisualizations/PCGPolygon2DDataVisualization.h
new file mode 100644
--- /dev/null
+++ b/Engine/Plugins/PCG/Source/PCGEditor/Public/DataVisualizations/PCGPolygon2DDataVisualization.h
@@ -0,0 +1,16 @@
+// Copyright Epic Games, Inc. All Rights Reserved.
+
+#pragma once
+
+#include "PCGSpatialDataVisualization.h"
+#include "PCGEditorModule.h"
+
+class UPCGData;
+
+class FPCGPolygon2DDataVisualization : public IPCGSpatialDataVisualization
+{
+public:
+	// ~Begin IPCGDataVisualization interface
+	PCGEDITOR_API virtual FPCGSetupSceneFunc GetViewportSetupFunc(const UPCGData* Data) const override;
+	// ~End IPCGDataVisualization interface
+};
\ No newline at end of file
