[PCG] Add polygon offset functionality

diff --git a/Engine/Plugins/PCG/Source/PCG/Private/Elements/Polygon/PCGOffsetPolygon2D.cpp b/Engine/Plugins/PCG/Source/PCG/Private/Elements/Polygon/PCGOffsetPolygon2D.cpp
new file mode 100644
--- /dev/null
+++ b/Engine/Plugins/PCG/Source/PCG/Private/Elements/Polygon/PCGOffsetPolygon2D.cpp
@@ -0,0 +1,86 @@
+// Copyright Epic Games, Inc. All Rights Reserved.
+
+#include "Elements/Polygon/PCGOffsetPolygon2D.h"
+
+#include "PCGContext.h"
+#include "Data/PCGPolygon2DData.h"
+
+#include "Curve/PolygonOffsetUtils.h"
+
+bool FPCGOffsetPolygonElement::ExecuteInternal(FPCGContext* Context) const
+{
+	TRACE_CPUPROFILER_EVENT_SCOPE(FPCGOffsetPolygonElement::Execute);
+
+	const UPCGOffsetPolygon2DSettings* Settings = Context->GetInputSettings<UPCGOffsetPolygon2DSettings>();
+	check(Settings);
+
+	TArray<FPCGTaggedData> Inputs = Context->InputData.GetInputsByPin(PCGPinConstants::DefaultInputLabel);
+	TArray<FPCGTaggedData>& Outputs = Context->OutputData.TaggedData;
+
+	for (const FPCGTaggedData& Input : Inputs)
+	{
+		const UPCGPolygon2DData* PolygonData = Cast<UPCGPolygon2DData>(Input.Data);
+
+		if (!PolygonData)
+		{
+			PCGLog::InputOutput::LogInvalidInputDataError(Context);
+			continue;
+		}
+
+		const UE::Geometry::FGeneralPolygon2d& Polygon = PolygonData->GetPolygon();
+		TArrayView<const UE::Geometry::FGeneralPolygon2d> PolyView = MakeArrayView(&Polygon, 1);
+		TArray<UE::Geometry::FGeneralPolygon2d> OutPolys;
+
+		bool bOperationSuccess = false;
+
+		const double MiterLimit = 2.0;
+		const UE::Geometry::EPolygonOffsetJoinType JoinType = UE::Geometry::EPolygonOffsetJoinType::Square;
+		const UE::Geometry::EPolygonOffsetEndType EndType = UE::Geometry::EPolygonOffsetEndType::Polygon;
+
+		if (Settings->Operation == EPCGPolygonOffsetOperation::Offset)
+		{
+			bOperationSuccess = UE::Geometry::PolygonsOffset(
+				Settings->Offset,
+				PolyView,
+				OutPolys,
+				/*bCopyOnFailure=*/false,
+				MiterLimit,
+				JoinType,
+				EndType);
+		}
+		else
+		{
+			const double FirstOffset = (Settings->Operation == EPCGPolygonOffsetOperation::Open) ? -Settings->Offset : +Settings->Offset;
+			const double SecondOffset = -FirstOffset;
+
+			bOperationSuccess = UE::Geometry::PolygonsOffsets(
+				FirstOffset,
+				SecondOffset,
+				PolyView,
+				OutPolys,
+				/*bCopyOnFailure=*/false,
+				MiterLimit,
+				JoinType,
+				EndType);
+		}
+
+		if (bOperationSuccess)
+		{
+			for (int OutPolyIndex = 0; OutPolyIndex < OutPolys.Num(); ++OutPolyIndex)
+			{
+				UPCGPolygon2DData* OutPolyData = FPCGContext::NewObject_AnyThread<UPCGPolygon2DData>(Context);
+
+				OutPolyData->InitializeFromData(PolygonData);
+				Context->OutputData.TaggedData.Add_GetRef(Input).Data = OutPolyData;
+				OutPolyData->SetPolygon(OutPolys[OutPolyIndex]);
+				OutPolyData->SetTransform(PolygonData->GetTransform());
+			}
+		}
+		else
+		{
+			PCGLog::LogWarningOnGraph(NSLOCTEXT("PCGOffsetPolygon2D", "OperationFailed", "Polygon offset operation failed."), Context);
+		}
+	}
+
+	return true;
+}
\ No newline at end of file
diff --git a/Engine/Plugins/PCG/Source/PCG/Public/Elements/Polygon/PCGOffsetPolygon2D.h b/Engine/Plugins/PCG/Source/PCG/Public/Elements/Polygon/PCGOffsetPolygon2D.h
new file mode 100644
--- /dev/null
+++ b/Engine/Plugins/PCG/Source/PCG/Public/Elements/Polygon/PCGOffsetPolygon2D.h
@@ -0,0 +1,57 @@
+// Copyright Epic Games, Inc. All Rights Reserved.
+
+#pragma once
+
+#include "PCGSettings.h"
+#include "Elements/Polygon/PCGPolygon2DUtils.h"
+
+#include "PCGOffsetPolygon2D.generated.h"
+
+UENUM()
+enum class EPCGPolygonOffsetOperation
+{
+	Offset UMETA(Tooltip = "Expands (positive) or contracts (negative) the polygon"),
+	Open UMETA(Tooltip = "Morphological open on the polygon based on the offset distance (contracts then expands)."),
+	Close UMETA(Tooltip = "Morphological close on the polygon based on the offset distance (expands then contracts).")
+};
+
+class FPCGOffsetPolygonElement : public IPCGElement
+{
+public:
+	virtual bool ExecuteInternal(FPCGContext* InContext) const override;
+	virtual EPCGElementExecutionLoopMode ExecutionLoopMode(const UPCGSettings* Settings) const override { return EPCGElementExecutionLoopMode::SinglePrimaryPin; }
+};
+
+/**
+* Offsets polygons, depending on the offset quantity and join type.
+*/
+UCLASS(MinimalAPI, BlueprintType, ClassGroup = (Procedural))
+class UPCGOffsetPolygon2DSettings : public UPCGSettings
+{
+	GENERATED_BODY()
+
+public:
+	//~Begin UPCGSettings interface
+#if WITH_EDITOR
+	virtual FName GetDefaultNodeName() const override { return FName(TEXT("OffsetPolygon2D")); }
+	virtual FText GetDefaultNodeTitle() const override { return NSLOCTEXT("PCGOffsetPolygon2D", "NodeTitle", "Offset Polygon"); }
+	virtual FText GetNodeTooltipText() const override { return NSLOCTEXT("PCGOffsetPolygon2D", "NodeTooltip", "Offsets polygon to either make it larger or smaller, or open/close holes based on the offset quantity."); }
+	virtual EPCGSettingsType GetType() const override { return EPCGSettingsType::Spatial; }
+	// @todo_pcg : presets for offset/open/close
+#endif
+
+protected:
+	virtual TArray<FPCGPinProperties> InputPinProperties() const override { return PCGPolygon2DUtils::DefaultPolygonInputPinProperties(); }
+	virtual TArray<FPCGPinProperties> OutputPinProperties() const override { return PCGPolygon2DUtils::DefaultPolygonOutputPinProperties(); }
+	virtual FPCGElementPtr CreateElement() const override { return MakeShared<FPCGOffsetPolygonElement>(); }
+	//~End UPCGSettings interface
+
+public:
+	UPROPERTY(BlueprintReadWrite, EditAnywhere, Category = Settings)
+	EPCGPolygonOffsetOperation Operation = EPCGPolygonOffsetOperation::Offset;
+
+	UPROPERTY(BlueprintReadWrite, EditAnywhere, Category = Settings, meta = (PCG_Overridable))
+	double Offset = 100.0;
+	
+	// Note: uses UE::Geometry::PolygonsOffset, and geometry script does the same thing.
+};
\ No newline at end of file
