[PCG] Fix ensure when executing GenerateGrassMaps on a region that is too large

#rnx

diff --git a/Engine/Plugins/PCG/Source/PCG/Private/Elements/PCGGenerateGrassMaps.cpp b/Engine/Plugins/PCG/Source/PCG/Private/Elements/PCGGenerateGrassMaps.cpp
--- a/Engine/Plugins/PCG/Source/PCG/Private/Elements/PCGGenerateGrassMaps.cpp
+++ b/Engine/Plugins/PCG/Source/PCG/Private/Elements/PCGGenerateGrassMaps.cpp
@@ -501,6 +501,15 @@ bool FPCGGenerateGrassMapsElement::ExecuteInternal(FPCGContext* InContext) const
 			/*bExport=*/false,
 			/*bShouldReadback=*/false);
 
+		const FIntPoint& GrassMapTextureTargetSize = Context->LandscapeGrassWeightExporter->GetTargetSize();
+		const int32 Max2DTextureDimension = GetMax2DTextureDimension();
+
+		if (GrassMapTextureTargetSize.X < 1 || GrassMapTextureTargetSize.X > Max2DTextureDimension || GrassMapTextureTargetSize.Y < 1 || GrassMapTextureTargetSize.Y > Max2DTextureDimension)
+		{
+			PCGLog::LogErrorOnGraph(FText::Format(LOCTEXT("InvalidTextureSize", "Invalid grass map texture size ({0}, {1}), cannot exceed ({2}, {2}). Consider partitioning the component onto a smaller grid size."), GrassMapTextureTargetSize.X, GrassMapTextureTargetSize.Y, Max2DTextureDimension), InContext);
+			return true;
+		}
+
 		UE::RenderCommandPipe::FSyncScope SyncScope;
 
 		RenderCaptureInterface::FScopedCapture RenderCapture(PCGGenerateGrassMaps::GTriggerGPUCaptureDispatches > 0, TEXT("PCGLandscapeGrassmapCapture"));
