[PCG] Allow multiple landscape materials for grass map generation as long as they have the same grass types

#rnx

diff --git a/Engine/Plugins/PCG/Source/PCG/Private/Elements/PCGGenerateGrassMaps.cpp b/Engine/Plugins/PCG/Source/PCG/Private/Elements/PCGGenerateGrassMaps.cpp
--- a/Engine/Plugins/PCG/Source/PCG/Private/Elements/PCGGenerateGrassMaps.cpp
+++ b/Engine/Plugins/PCG/Source/PCG/Private/Elements/PCGGenerateGrassMaps.cpp
@@ -233,12 +233,13 @@ bool FPCGGenerateGrassMapsElement::ExecuteInternal(FPCGContext* InContext) const
 		}
 
 		const FBox GenerationBounds = ExecutionSource->GetExecutionState().GetBounds();
+		TArray<TObjectPtr<ULandscapeGrassType>> GrassTypes;
 
 		for (TSoftObjectPtr<ALandscapeProxy> LandscapeProxyPtr : LandscapeData->Landscapes)
 		{
 			ALandscapeProxy* LandscapeProxy = LandscapeProxyPtr.Get();
 
-			if (!LandscapeProxy || !LandscapeProxy->LandscapeMaterial)
+			if (!LandscapeProxy || !LandscapeProxy->GetLandscapeMaterial())
 			{
 				continue;
 			}
@@ -247,12 +248,6 @@ bool FPCGGenerateGrassMapsElement::ExecuteInternal(FPCGContext* InContext) const
 			{
 				Context->LandscapeProxy = LandscapeProxy;
 			}
-			// TODO: In order to support multiple-material we need to do a dispatch of one landscape grass exporter per LS proxy.
-			else if (Context->LandscapeProxy->LandscapeMaterial != LandscapeProxy->LandscapeMaterial)
-			{
-				PCGLog::LogErrorOnGraph(LOCTEXT("LandscapeMaterialMismatch", "Encountered multiple landscape materials on the same landscape. Only one landscape material per landscape is supported at this time."), Context);
-				return true;
-			}
 
 			for (ULandscapeComponent* LandscapeComponent : LandscapeProxy->LandscapeComponents)
 			{
@@ -271,28 +266,13 @@ bool FPCGGenerateGrassMapsElement::ExecuteInternal(FPCGContext* InContext) const
 						break;
 					}
 
-					if (LandscapeComponent->OverrideMaterial && Context->LandscapeProxy->LandscapeMaterial->GetBaseMaterial() != LandscapeComponent->OverrideMaterial->GetBaseMaterial())
-					{
-#if !(UE_BUILD_SHIPPING || UE_BUILD_TEST) || USE_LOGGING_IN_SHIPPING
-						PCGLog::LogErrorOnGraph(
-							FText::Format(LOCTEXT("LandscapeOverrideMaterialMismatch",
-								"Encountered landscape component {0} with override material {1} that did not match the landscape proxy material {2}. Only one landscape material per landscape is supported at this time."),
-								FText::FromString(LandscapeComponent->GetName()),
-								FText::FromString(LandscapeComponent->OverrideMaterial->GetName()),
-								FText::FromString(Context->LandscapeProxy->LandscapeMaterial->GetName())),
-							Context);
-#endif
-
-						return true;
-					}
-
 					LandscapeComponent->UpdateGrassTypes();
 
 					Context->LandscapeComponents.Add(LandscapeComponent);
 
 					if (Context->LandscapeComponents.Num() == 1)
 					{
-						const TArray<TObjectPtr<ULandscapeGrassType>>& GrassTypes = LandscapeComponent->GetGrassTypes();
+						GrassTypes = LandscapeComponent->GetGrassTypes();
 						Context->NumGrassTypes = GrassTypes.Num();
 
 						FString GrassTypeName;
@@ -331,10 +311,27 @@ bool FPCGGenerateGrassMapsElement::ExecuteInternal(FPCGContext* InContext) const
 					{
 						Context->GrassMapBounds += LandscapeComponentBounds;
 
-						// We expect all landscape components for a single landscape to have the same grass types. Do weak validation here.
-						const uint32 ExpectedNumGrassTypes = Context->NumGrassTypes;
-						const uint32 ActualNumGrassTypes = LandscapeComponent->GetGrassTypes().Num();
-						ensure(ExpectedNumGrassTypes == ActualNumGrassTypes);
+						const UMaterialInterface* LandscapeProxyMaterial = Context->LandscapeProxy->GetLandscapeMaterial();
+						const UMaterialInterface* LandscapeMaterial = LandscapeComponent->GetLandscapeMaterial();
+
+						// We expect all landscape components for a single landscape to have the same grass types.
+						// TODO: In order to support distinct grass types we need to do a dispatch of one landscape grass exporter per LS proxy.
+						if (LandscapeMaterial != LandscapeProxyMaterial && GrassTypes != LandscapeComponent->GetGrassTypes())
+						{
+#if !(UE_BUILD_SHIPPING || UE_BUILD_TEST) || USE_LOGGING_IN_SHIPPING
+							const FString LandscapeProxyMaterialName = LandscapeProxyMaterial ? LandscapeProxyMaterial->GetName() : TEXT("MISSINGMATERIAL");
+							const FString LandscapeMaterialName = LandscapeMaterial ? LandscapeMaterial->GetName() : TEXT("MISSINGMATERIAL");
+							
+							PCGLog::LogErrorOnGraph(
+								FText::Format(LOCTEXT("LandscapeComponentGrassTypeMismatch",
+									"Encountered landscape component '{0}' with material '{1}' that does not match the grass types on material '{2}'. All landscape materials must have the same grass types at this time."),
+									FText::FromString(LandscapeComponent->GetName()),
+									FText::FromString(LandscapeMaterialName),
+									FText::FromString(LandscapeProxyMaterialName)),
+								Context);
+#endif
+							return true;
+						}
 
 						// Currently assuming all landscape have the same extents.
 						ensure(FMath::IsNearlyEqual(Context->LandscapeComponentExtent, LandscapeComponent->Bounds.BoxExtent.X * 2.0));
