[PCG] Add polygon operations (intersection, difference, union, inner intersection, exclusive or)

diff --git a/Engine/Plugins/PCG/Source/PCG/Private/Elements/Polygon/PCGPolygon2DOperation.cpp b/Engine/Plugins/PCG/Source/PCG/Private/Elements/Polygon/PCGPolygon2DOperation.cpp
new file mode 100644
--- /dev/null
+++ b/Engine/Plugins/PCG/Source/PCG/Private/Elements/Polygon/PCGPolygon2DOperation.cpp
@@ -0,0 +1,298 @@
+// Copyright Epic Games, Inc. All Rights Reserved.
+
+#include "Elements/Polygon/PCGPolygon2DOperation.h"
+
+#include "PCGContext.h"
+#include "Data/PCGPolygon2DData.h"
+
+#include "Curve/PolygonIntersectionUtils.h"
+
+#define LOCTEXT_NAMESPACE "PCGPolygon2DOperation"
+
+namespace PCGPolygon2DOperationConstants
+{
+	const FName ClipPolysLabel = TEXT("Clip Polygons");
+}
+
+#if WITH_EDITOR
+FText UPCGPolygon2DOperationSettings::GetDefaultNodeTitle() const
+{
+	return LOCTEXT("NodeTitle", "Polygon Operation");
+}
+
+FText UPCGPolygon2DOperationSettings::GetNodeTooltipText() const
+{
+	return LOCTEXT("NodeTooltip", "Performs polygon operations between the inputs.");
+}
+
+EPCGChangeType UPCGPolygon2DOperationSettings::GetChangeTypeForProperty(FPropertyChangedEvent& PropertyChangedEvent) const
+{
+	EPCGChangeType ChangeType = Super::GetChangeTypeForProperty(PropertyChangedEvent);
+
+	if (PropertyChangedEvent.GetPropertyName() == GET_MEMBER_NAME_CHECKED(UPCGPolygon2DOperationSettings, Operation))
+	{
+		ChangeType |= EPCGChangeType::Structural;
+	}
+
+	return ChangeType;
+}
+#endif // WITH_EDITOR
+
+FString UPCGPolygon2DOperationSettings::GetAdditionalTitleInformation() const
+{
+	if (const UEnum* EnumPtr = StaticEnum<EPCGPolygonOperation>())
+	{
+		return EnumPtr->GetDisplayNameTextByValue(static_cast<int64>(Operation)).ToString();
+	}
+	else
+	{
+		return FString();
+	}
+}
+
+TArray<FPCGPinProperties> UPCGPolygon2DOperationSettings::InputPinProperties() const
+{
+	TArray<FPCGPinProperties> PinProperties = PCGPolygon2DUtils::DefaultPolygonInputPinProperties();
+
+	if (Operation != EPCGPolygonOperation::Union && Operation != EPCGPolygonOperation::InnerIntersection)
+	{
+		PinProperties.Emplace_GetRef(PCGPolygon2DOperationConstants::ClipPolysLabel, EPCGDataType::Polygon2D).SetRequiredPin();
+	}
+
+	return PinProperties;
+}
+
+EPCGElementExecutionLoopMode FPCGPolygon2DOperationElement::ExecutionLoopMode(const UPCGSettings* InSettings) const
+{
+	if (const UPCGPolygon2DOperationSettings* Settings = Cast<UPCGPolygon2DOperationSettings>(InSettings))
+	{
+		return (Settings->Operation == EPCGPolygonOperation::Union || Settings->Operation == EPCGPolygonOperation::InnerIntersection) ? EPCGElementExecutionLoopMode::NotALoop : EPCGElementExecutionLoopMode::SinglePrimaryPin;
+	}
+	else
+	{
+		return EPCGElementExecutionLoopMode::NotALoop;
+	}
+}
+
+bool FPCGPolygon2DOperationElement::ExecuteInternal(FPCGContext* Context) const
+{
+	TRACE_CPUPROFILER_EVENT_SCOPE(FPCGPolygon2DOperationElement::Execute);
+
+	const UPCGPolygon2DOperationSettings* Settings = Context->GetInputSettings<UPCGPolygon2DOperationSettings>();
+	check(Settings);
+
+	TArray<FPCGTaggedData> Inputs = Context->InputData.GetInputsByPin(PCGPinConstants::DefaultInputLabel);
+	TArray<FPCGTaggedData> ClipInputs = Context->InputData.GetInputsByPin(PCGPolygon2DOperationConstants::ClipPolysLabel);
+
+	bool bHasTransform = false;
+	FTransform Transform = FTransform::Identity;
+	TArray<UE::Geometry::FGeneralPolygon2d> Polygons;
+	TArray<int> PolyIndexToInputIndex;
+	TArray<UE::Geometry::FGeneralPolygon2d> ClipPolygons;
+
+	auto AddPolygon = [&bHasTransform, &Transform](const UE::Geometry::FGeneralPolygon2d& InPoly, const FTransform& InTransform, TArray<UE::Geometry::FGeneralPolygon2d>& OutPolyList)
+	{
+		if (!bHasTransform)
+		{
+			UE::Geometry::FGeneralPolygon2d& OutPoly = OutPolyList.Add_GetRef(InPoly);
+			// @todo_pcg: having clockwise polygons should be supported because it allows more flexibility
+			ensure(!OutPoly.OuterIsClockwise());
+
+			Transform = InTransform;
+			bHasTransform = true;
+		}
+		else
+		{
+			// @todo_pcg: we could validate that the normal are parallel.
+			UE::Geometry::FGeneralPolygon2d& Poly = OutPolyList.Emplace_GetRef();
+			// Important note: since this can technically flip the normal, that would in turn change the polygon orientation.
+			// Unfortunately, this means that the winding of the polygon might change.
+			// However, the polygon transform isn't aware of this, so we need to rebuild the polygon from scratch here instead.
+			const FTransform RelativeTransform = InTransform.GetRelativeTransform(Transform);
+
+			auto TransformPoly = [&RelativeTransform](UE::Geometry::FPolygon2d& Poly, bool bIsHole)
+			{
+				Poly.Transform([&RelativeTransform](const FVector2d& InVertex)
+				{
+					const FVector TransformedVertex = RelativeTransform.TransformPosition(FVector(InVertex, 0.0));
+					return FVector2d(TransformedVertex.X, TransformedVertex.Y);
+				});
+
+				if (Poly.IsClockwise() != bIsHole)
+				{
+					Poly.Reverse();
+				}
+			};
+
+			UE::Geometry::FPolygon2d Outer = InPoly.GetOuter();
+			TransformPoly(Outer, false);
+
+			Poly.SetOuter(MoveTemp(Outer));
+
+			for (const UE::Geometry::FPolygon2d& InHole : InPoly.GetHoles())
+			{
+				UE::Geometry::FPolygon2d Hole = InHole;
+				TransformPoly(Hole, true);
+
+				Poly.AddHole(MoveTemp(Hole));
+			}
+		}
+	};
+
+	for(int InputIndex = 0; InputIndex < Inputs.Num(); ++InputIndex)
+	{
+		if (const UPCGPolygon2DData* Polygon = Cast<UPCGPolygon2DData>(Inputs[InputIndex].Data))
+		{
+			PolyIndexToInputIndex.Add(InputIndex);
+			AddPolygon(Polygon->GetPolygon(), Polygon->GetTransform(), Polygons);
+		}
+	}
+
+	for (int ClipIndex = 0; ClipIndex < ClipInputs.Num(); ++ClipIndex)
+	{
+		if (const UPCGPolygon2DData* Polygon = Cast<UPCGPolygon2DData>(ClipInputs[ClipIndex].Data))
+		{
+			AddPolygon(Polygon->GetPolygon(), Polygon->GetTransform(), ClipPolygons);
+		}
+	}
+
+	TArray<UE::Geometry::FGeneralPolygon2d> OutPolys;
+
+	auto AddToOutput = [Context, &Transform](TArray<UE::Geometry::FGeneralPolygon2d>& OutPolys, const FPCGTaggedData* InputPtr)
+	{
+		for (int OutPolyIndex = 0; OutPolyIndex < OutPolys.Num(); ++OutPolyIndex)
+		{
+			UPCGPolygon2DData* OutPolyData = FPCGContext::NewObject_AnyThread<UPCGPolygon2DData>(Context);
+
+			if (InputPtr)
+			{
+				OutPolyData->InitializeFromData(Cast<UPCGSpatialData>(InputPtr->Data));
+				Context->OutputData.TaggedData.Add_GetRef(*InputPtr).Data = OutPolyData;
+			}
+			else
+			{
+				Context->OutputData.TaggedData.Emplace_GetRef().Data = OutPolyData;
+			}
+
+			OutPolyData->SetPolygon(MoveTemp(OutPolys[OutPolyIndex]));
+			OutPolyData->SetTransform(Transform);
+		}
+
+		OutPolys.Reset();
+	};
+
+	if (Settings->Operation == EPCGPolygonOperation::Union)
+	{
+		if (Polygons.Num() <= 1)
+		{
+			// Nothing to do.
+			Context->OutputData.TaggedData.Append(Inputs);
+		}
+		else
+		{
+			if (UE::Geometry::PolygonsUnion(Polygons, OutPolys, /*bCopyInputOnFailure=*/false))
+			{
+				AddToOutput(OutPolys, nullptr);
+			}
+			else if (!Settings->bQuiet)
+			{
+				PCGLog::LogWarningOnGraph(LOCTEXT("UnionFailed", "Union operation failed."), Context);
+			}
+		}
+	}
+	else if (Settings->Operation == EPCGPolygonOperation::InnerIntersection)
+	{
+		if (Polygons.Num() <= 1)
+		{
+			// Nothing to do
+			Context->OutputData.TaggedData.Append(Inputs);
+		}
+		else
+		{
+			TArray<UE::Geometry::FGeneralPolygon2d> CurrentPolygons;
+			CurrentPolygons.Add(Polygons[0]);
+
+			for (int PolygonIndex = 1; PolygonIndex < Polygons.Num(); ++PolygonIndex)
+			{
+				TArrayView<UE::Geometry::FGeneralPolygon2d> ClipPolygon = MakeArrayView(&Polygons[PolygonIndex], 1);
+
+				if (UE::Geometry::PolygonsIntersection(CurrentPolygons, ClipPolygon, OutPolys))
+				{
+					CurrentPolygons = MoveTemp(OutPolys);
+					OutPolys.Reset();
+				}
+				else
+				{
+					if (!Settings->bQuiet)
+					{
+						PCGLog::LogWarningOnGraph(FText::Format(LOCTEXT("InnerIntersectionFailed", "Inner intersection operation failed at index '{0}'."), PolygonIndex), Context);
+					}
+
+					CurrentPolygons.Reset();
+					break;
+				}
+			}
+
+			OutPolys = MoveTemp(CurrentPolygons);
+			AddToOutput(OutPolys, nullptr);
+		}
+	}
+	else if (ClipPolygons.IsEmpty())
+	{
+		// Nothing to do.
+		Context->OutputData.TaggedData.Append(Inputs);
+	}
+	else
+	{
+		for(int PolygonIndex = 0; PolygonIndex < Polygons.Num(); ++PolygonIndex)
+		{
+			const FPCGTaggedData& CurrentInput = Inputs[PolyIndexToInputIndex[PolygonIndex]];
+
+			UE::Geometry::FGeneralPolygon2d& Poly = Polygons[PolygonIndex];
+			OutPolys.Reset();
+			bool bOperationOk = false;
+
+			TArrayView<UE::Geometry::FGeneralPolygon2d> SubjPolygons = MakeArrayView(&Poly, 1);
+
+			if (Settings->Operation == EPCGPolygonOperation::Difference ||
+				Settings->Operation == EPCGPolygonOperation::Intersection ||
+				Settings->Operation == EPCGPolygonOperation::ExclusiveOr)
+			{
+				if ((Settings->Operation == EPCGPolygonOperation::Difference && UE::Geometry::PolygonsDifference(SubjPolygons, ClipPolygons, OutPolys)) ||
+					(Settings->Operation == EPCGPolygonOperation::Intersection && UE::Geometry::PolygonsIntersection(SubjPolygons, ClipPolygons, OutPolys)) ||
+					(Settings->Operation == EPCGPolygonOperation::ExclusiveOr && UE::Geometry::PolygonsExclusiveOr(SubjPolygons, ClipPolygons, OutPolys)))
+				{
+					AddToOutput(OutPolys, &CurrentInput);
+				}
+				else if(!Settings->bQuiet)
+				{
+					PCGLog::LogWarningOnGraph(LOCTEXT("OperationFailed", "Polygon operation failed."), Context);
+				}
+			}
+			else if (Settings->Operation == EPCGPolygonOperation::PairwiseIntersection)
+			{
+				for (int ClipIndex = 0; ClipIndex < ClipPolygons.Num(); ++ClipIndex)
+				{
+					OutPolys.Reset();
+					TArrayView<UE::Geometry::FGeneralPolygon2d> ClipSubsetPolygons = MakeArrayView(&ClipPolygons[ClipIndex], 1);
+					if (UE::Geometry::PolygonsIntersection(SubjPolygons, ClipSubsetPolygons, OutPolys))
+					{
+						AddToOutput(OutPolys, &CurrentInput);
+					}
+					else if (!Settings->bQuiet)
+					{
+						PCGLog::LogWarningOnGraph(FText::Format(LOCTEXT("PairwiseIntersectionFailed", "Pairwise polygon intersection failed at polygon '{0}' and clip polygon '{1}'."), PolygonIndex, ClipIndex), Context);
+					}
+				}
+			}
+			else
+			{
+				PCGLog::LogErrorOnGraph(LOCTEXT("InvalidOperation", "Invalid polygon operation."), Context);
+			}
+		}
+	}
+
+	return true;
+}
+
+#undef LOCTEXT_NAMESPACE
\ No newline at end of file
diff --git a/Engine/Plugins/PCG/Source/PCG/Public/Elements/Polygon/PCGPolygon2DOperation.h b/Engine/Plugins/PCG/Source/PCG/Public/Elements/Polygon/PCGPolygon2DOperation.h
new file mode 100644
--- /dev/null
+++ b/Engine/Plugins/PCG/Source/PCG/Public/Elements/Polygon/PCGPolygon2DOperation.h
@@ -0,0 +1,60 @@
+// Copyright Epic Games, Inc. All Rights Reserved.
+
+#pragma once
+
+#include "PCGSettings.h"
+#include "Elements/Polygon/PCGPolygon2DUtils.h"
+
+#include "PCGPolygon2DOperation.generated.h"
+
+UENUM()
+enum class EPCGPolygonOperation : uint8
+{
+	Union,
+	Difference,
+	Intersection,
+	PairwiseIntersection,
+	InnerIntersection,
+	ExclusiveOr,
+};
+
+class FPCGPolygon2DOperationElement : public IPCGElement
+{
+public:
+	virtual bool ExecuteInternal(FPCGContext* InContext) const override;
+	virtual EPCGElementExecutionLoopMode ExecutionLoopMode(const UPCGSettings* Settings) const override;
+};
+
+/**
+* Applies polygon operations between polygons
+*/
+UCLASS(MinimalAPI, BlueprintType, ClassGroup = (Procedural))
+class UPCGPolygon2DOperationSettings : public UPCGSettings
+{
+	GENERATED_BODY()
+
+public:
+	//~Begin UPCGSettings interface
+#if WITH_EDITOR
+	virtual FName GetDefaultNodeName() const override { return FName(TEXT("Polygon2DOperation")); }
+	virtual FText GetDefaultNodeTitle() const override;
+	virtual FText GetNodeTooltipText() const override;
+	virtual EPCGSettingsType GetType() const override { return EPCGSettingsType::Spatial; }
+	// @todo_pcg : presets per operation
+	virtual EPCGChangeType GetChangeTypeForProperty(FPropertyChangedEvent& PropertyChangedEvent) const override;
+#endif
+	virtual FString GetAdditionalTitleInformation() const override;
+
+protected:
+	virtual TArray<FPCGPinProperties> InputPinProperties() const override;
+	virtual TArray<FPCGPinProperties> OutputPinProperties() const override { return PCGPolygon2DUtils::DefaultPolygonOutputPinProperties(); }
+	virtual FPCGElementPtr CreateElement() const override { return MakeShared<FPCGPolygon2DOperationElement>(); }
+	//~End UPCGSettings interface
+
+public:
+	UPROPERTY(BlueprintReadWrite, EditAnywhere, Category = Settings)
+	EPCGPolygonOperation Operation = EPCGPolygonOperation::Intersection;
+
+	UPROPERTY(BlueprintReadWrite, EditAnywhere, Category = Settings)
+	bool bQuiet = false;
+};
\ No newline at end of file
