[PCG] GPU: Add helper for constructing and wiring compute kernels.

#rnx

diff --git a/Engine/Plugins/PCG/Source/PCG/Private/Compute/Elements/PCGCustomHLSL.cpp b/Engine/Plugins/PCG/Source/PCG/Private/Compute/Elements/PCGCustomHLSL.cpp
--- a/Engine/Plugins/PCG/Source/PCG/Private/Compute/Elements/PCGCustomHLSL.cpp
+++ b/Engine/Plugins/PCG/Source/PCG/Private/Compute/Elements/PCGCustomHLSL.cpp
@@ -9,6 +9,7 @@
 #include "PCGPoint.h"
 #include "PCGSubsystem.h"
 #include "Compute/PCGComputeSource.h"
+#include "Compute/PCGKernelHelpers.h"
 #include "Compute/Elements/PCGCustomHLSLKernel.h"
 #include "Data/PCGPointData.h"
 #include "Graph/PCGGPUGraphCompilationContext.h"
@@ -296,24 +297,15 @@ bool UPCGCustomHLSLSettings::IsShaderTextReadOnly() const
 
 void UPCGCustomHLSLSettings::CreateKernels(FPCGGPUCompilationContext& InOutContext, UObject* InObjectOuter, TArray<UPCGComputeKernel*>& OutKernels, TArray<FPCGKernelEdge>& OutEdges) const
 {
-	FPCGComputeKernelParams KernelParams;
-	KernelParams.Settings = this;
-	KernelParams.bLogDescriptions = bDumpDataDescriptions;
+	PCGKernelHelpers::FCreateKernelParams CreateParams(InObjectOuter, this);
 
-	UPCGCustomHLSLKernel* Kernel = InOutContext.NewObject_AnyThread<UPCGCustomHLSLKernel>(InObjectOuter);
-	Kernel->Initialize(KernelParams);
-	OutKernels.Add(Kernel);
+	CreateParams.NodeInputPinsToWire.Empty(InputPins.Num());
+	Algo::Transform(InputPins, CreateParams.NodeInputPinsToWire, [](const FPCGPinProperties& InProps) { return InProps.Label; });
 
-	// Connect node pins to kernel pins
-	for (const FPCGPinProperties& Pin : InputPins)
-	{
-		OutEdges.Emplace(FPCGPinReference(Pin.Label), FPCGPinReference(Kernel, Pin.Label));
-	}
+	CreateParams.NodeOutputPinsToWire.Empty(OutputPins.Num());
+	Algo::Transform(OutputPins, CreateParams.NodeOutputPinsToWire, [](const FPCGPinProperties& InProps) { return InProps.Label; });
 
-	for (const FPCGPinPropertiesGPU& Pin : OutputPins)
-	{
-		OutEdges.Emplace(FPCGPinReference(Kernel, Pin.Label), FPCGPinReference(Pin.Label));
-	}
+	PCGKernelHelpers::CreateKernel<UPCGCustomHLSLKernel>(InOutContext, CreateParams, OutKernels, OutEdges);
 }
 #endif // WITH_EDITOR
 
diff --git a/Engine/Plugins/PCG/Source/PCG/Private/Compute/PCGKernelHelpers.h b/Engine/Plugins/PCG/Source/PCG/Private/Compute/PCGKernelHelpers.h
new file mode 100644
--- /dev/null
+++ b/Engine/Plugins/PCG/Source/PCG/Private/Compute/PCGKernelHelpers.h
@@ -0,0 +1,66 @@
+// Copyright Epic Games, Inc. All Rights Reserved.
+
+#pragma once
+
+#include "PCGCommon.h"
+#include "PCGSettings.h"
+#include "Compute/PCGComputeCommon.h"
+#include "Compute/PCGComputeKernel.h"
+#include "Graph/PCGGPUGraphCompilationContext.h"
+
+class UPCGSettings;
+
+namespace PCGKernelHelpers
+{
+#if WITH_EDITOR
+	struct FCreateKernelParams
+	{
+		FCreateKernelParams(UObject* InObjectOuter, const UPCGSettings* InSettings = nullptr)
+			: ObjectOuter(InObjectOuter)
+			, OwnerSettings(InSettings)
+			, bDumpDataDescriptions(InSettings && InSettings->bDumpDataDescriptions)
+		{
+		}
+
+		UObject* ObjectOuter = nullptr;
+
+		bool bRequiresOverridableParams = false;
+
+		const UPCGSettings* OwnerSettings = nullptr;
+
+		bool bDumpDataDescriptions = false;
+
+		/** Create edges from node input pins to kernel inputs, assuming label is same on both. */
+		TArray<FName> NodeInputPinsToWire = { PCGPinConstants::DefaultInputLabel };
+
+		/** Create edges from node output pins to kernel outputs, assuming label is same on both. */
+		TArray<FName> NodeOutputPinsToWire = { PCGPinConstants::DefaultOutputLabel };
+	};
+
+	template<class KernelType>
+	KernelType* CreateKernel(FPCGGPUCompilationContext& InCompilationContext, const FCreateKernelParams& InParams, TArray<UPCGComputeKernel*>& OutKernels, TArray<FPCGKernelEdge>& OutKernelEdges)
+	{
+		FPCGComputeKernelParams KernelParams;
+		KernelParams.Settings = InParams.OwnerSettings;
+		KernelParams.bLogDescriptions = InParams.bDumpDataDescriptions;
+
+		KernelType* Kernel = InCompilationContext.NewObject_AnyThread<KernelType>(InParams.ObjectOuter);
+
+		Kernel->Initialize(KernelParams);
+		OutKernels.Add(Kernel);
+
+		// Connect node pins to kernel pins
+		for (FName PinLabel : InParams.NodeInputPinsToWire)
+		{
+			OutKernelEdges.Emplace(FPCGPinReference(PinLabel), FPCGPinReference(Kernel, PinLabel));
+		}
+
+		for (FName PinLabel : InParams.NodeOutputPinsToWire)
+		{
+			OutKernelEdges.Emplace(FPCGPinReference(Kernel, PinLabel), FPCGPinReference(PinLabel));
+		}
+
+		return CastChecked<KernelType>(Kernel);
+	}
+#endif // WITH_EDITOR
+}
diff --git a/Engine/Plugins/PCG/Source/PCG/Private/Elements/Metadata/PCGMetadataPartition.cpp b/Engine/Plugins/PCG/Source/PCG/Private/Elements/Metadata/PCGMetadataPartition.cpp
--- a/Engine/Plugins/PCG/Source/PCG/Private/Elements/Metadata/PCGMetadataPartition.cpp
+++ b/Engine/Plugins/PCG/Source/PCG/Private/Elements/Metadata/PCGMetadataPartition.cpp
@@ -5,6 +5,7 @@
 #include "PCGContext.h"
 #include "PCGCustomVersion.h"
 #include "Compute/PCGComputeCommon.h"
+#include "Compute/PCGKernelHelpers.h"
 #include "Compute/BuiltInKernels/PCGCountUniqueAttributeValuesKernel.h"
 #include "Elements/Metadata/PCGMetadataPartitionKernel.h"
 #include "Graph/PCGGPUGraphCompilationContext.h"
@@ -66,25 +67,18 @@ void UPCGMetadataPartitionSettings::CreateKernels(FPCGGPUCompilationContext& InO
 		return;
 	}
 
-	FPCGComputeKernelParams KernelParams;
-	KernelParams.Settings = this;
-	KernelParams.bLogDescriptions = bDumpDataDescriptions;
+	PCGKernelHelpers::FCreateKernelParams CreateParams(InObjectOuter, this);
 
-	UPCGMetadataPartitionKernel* PartitionKernel = InOutContext.NewObject_AnyThread<UPCGMetadataPartitionKernel>(InObjectOuter);
-	PartitionKernel->Initialize(KernelParams);
-	OutKernels.Add(PartitionKernel);
+	UPCGMetadataPartitionKernel* PartitionKernel = PCGKernelHelpers::CreateKernel<UPCGMetadataPartitionKernel>(InOutContext, CreateParams, OutKernels, OutEdges);
 
-	// Connect node pins to kernel pins
-	OutEdges.Emplace(FPCGPinReference(PCGPinConstants::DefaultInputLabel), FPCGPinReference(PartitionKernel, PCGPinConstants::DefaultInputLabel));
-	OutEdges.Emplace(FPCGPinReference(PartitionKernel, PCGPinConstants::DefaultOutputLabel), FPCGPinReference(PCGPinConstants::DefaultOutputLabel));
+	// Don't wire analysis kernel to node output pin, wire manually to the copy kernel below.
+	CreateParams.NodeOutputPinsToWire.Empty();
 
-	UPCGCountUniqueAttributeValuesKernel* CountKernel = InOutContext.NewObject_AnyThread<UPCGCountUniqueAttributeValuesKernel>(InObjectOuter);
-	CountKernel->Initialize(KernelParams);
+	UPCGCountUniqueAttributeValuesKernel* CountKernel = PCGKernelHelpers::CreateKernel<UPCGCountUniqueAttributeValuesKernel>(InOutContext, CreateParams, OutKernels, OutEdges);
 	CountKernel->SetAttributeName(PartitionAttributeSelectors[0].GetAttributeName());
 	CountKernel->SetEmitPerDataCounts(true);
-	OutKernels.Add(CountKernel);
 
-	OutEdges.Emplace(FPCGPinReference(PCGPinConstants::DefaultInputLabel), FPCGPinReference(CountKernel, PCGPinConstants::DefaultInputLabel));
+	// Wire count kernel to copy partition kernel.
 	OutEdges.Emplace(FPCGPinReference(CountKernel, PCGPinConstants::DefaultOutputLabel), FPCGPinReference(PartitionKernel, PCGMetadataPartitionConstants::ElementCountsPinLabel));
 }
 
diff --git a/Engine/Plugins/PCG/Source/PCG/Private/Elements/PCGCopyPoints.cpp b/Engine/Plugins/PCG/Source/PCG/Private/Elements/PCGCopyPoints.cpp
--- a/Engine/Plugins/PCG/Source/PCG/Private/Elements/PCGCopyPoints.cpp
+++ b/Engine/Plugins/PCG/Source/PCG/Private/Elements/PCGCopyPoints.cpp
@@ -3,6 +3,7 @@
 #include "Elements/PCGCopyPoints.h"
 
 #include "PCGContext.h"
+#include "Compute/PCGKernelHelpers.h"
 #include "Data/PCGBasePointData.h"
 #include "Data/PCGSpatialData.h"
 #include "Elements/PCGCopyPointsAnalysisKernel.h"
@@ -25,31 +26,24 @@ FText UPCGCopyPointsSettings::GetNodeTooltipText() const
 
 void UPCGCopyPointsSettings::CreateKernels(FPCGGPUCompilationContext& InOutContext, UObject* InObjectOuter, TArray<UPCGComputeKernel*>& OutKernels, TArray<FPCGKernelEdge>& OutEdges) const
 {
-	FPCGComputeKernelParams KernelParams;
-	KernelParams.Settings = this;
-	KernelParams.bLogDescriptions = bDumpDataDescriptions;
-	KernelParams.bRequiresOverridableParams = true;
+	PCGKernelHelpers::FCreateKernelParams CreateParams(InObjectOuter, this);
 
-	UPCGCopyPointsKernel* CopyKernel = InOutContext.NewObject_AnyThread<UPCGCopyPointsKernel>(InObjectOuter);
-	CopyKernel->Initialize(KernelParams);
-	OutKernels.Add(CopyKernel);
+	CreateParams.bRequiresOverridableParams = true;
+	CreateParams.NodeInputPinsToWire = { PCGCopyPointsConstants::SourcePointsLabel, PCGCopyPointsConstants::TargetPointsLabel };
 
-	OutEdges.Emplace(FPCGPinReference(PCGCopyPointsConstants::SourcePointsLabel), FPCGPinReference(CopyKernel, PCGCopyPointsConstants::SourcePointsLabel));
-	OutEdges.Emplace(FPCGPinReference(PCGCopyPointsConstants::TargetPointsLabel), FPCGPinReference(CopyKernel, PCGCopyPointsConstants::TargetPointsLabel));
-	OutEdges.Emplace(FPCGPinReference(CopyKernel, PCGPinConstants::DefaultOutputLabel), FPCGPinReference(PCGPinConstants::DefaultOutputLabel));
+	UPCGComputeKernel* CopyKernel = PCGKernelHelpers::CreateKernel<UPCGCopyPointsKernel>(InOutContext, CreateParams, OutKernels, OutEdges);
 
 	if (bMatchBasedOnAttribute)
 	{
 		// Analysis kernel does not require any of the overridden properties used by the CopyPoints settings.
-		KernelParams.bRequiresOverridableParams = false;
+		CreateParams.bRequiresOverridableParams = false;
 		
-		UPCGCopyPointsAnalysisKernel* AnalysisKernel = InOutContext.NewObject_AnyThread<UPCGCopyPointsAnalysisKernel>(InObjectOuter);
-		AnalysisKernel->Initialize(KernelParams);
-		OutKernels.Add(AnalysisKernel);
+		// Don't wire analysis kernel to node output pin, wire manually to the copy kernel below.
+		CreateParams.NodeOutputPinsToWire.Empty();
 
-		// Connect to node pins and the copy kernel.
-		OutEdges.Emplace(FPCGPinReference(PCGCopyPointsConstants::SourcePointsLabel), FPCGPinReference(AnalysisKernel, PCGCopyPointsConstants::SourcePointsLabel));
-		OutEdges.Emplace(FPCGPinReference(PCGCopyPointsConstants::TargetPointsLabel), FPCGPinReference(AnalysisKernel, PCGCopyPointsConstants::TargetPointsLabel));
+		UPCGComputeKernel* AnalysisKernel = PCGKernelHelpers::CreateKernel<UPCGCopyPointsAnalysisKernel>(InOutContext, CreateParams, OutKernels, OutEdges);
+
+		// Wire analysis kernel to copy kernel.
 		OutEdges.Emplace(FPCGPinReference(AnalysisKernel, PCGPinConstants::DefaultOutputLabel), FPCGPinReference(CopyKernel, PCGCopyPointsConstants::SelectedFlagsPinLabel));
 	}
 }
diff --git a/Engine/Plugins/PCG/Source/PCG/Private/Elements/PCGCullPointsOutsideActorBounds.cpp b/Engine/Plugins/PCG/Source/PCG/Private/Elements/PCGCullPointsOutsideActorBounds.cpp
--- a/Engine/Plugins/PCG/Source/PCG/Private/Elements/PCGCullPointsOutsideActorBounds.cpp
+++ b/Engine/Plugins/PCG/Source/PCG/Private/Elements/PCGCullPointsOutsideActorBounds.cpp
@@ -6,6 +6,7 @@
 #include "PCGContext.h"
 #include "PCGPin.h"
 #include "Compute/PCGComputeKernel.h"
+#include "Compute/PCGKernelHelpers.h"
 #include "Compute/Elements/PCGCullPointsOutsideActorBoundsKernel.h"
 #include "Data/PCGBasePointData.h"
 #include "Graph/PCGGPUGraphCompilationContext.h"
@@ -29,16 +30,8 @@ FText UPCGCullPointsOutsideActorBoundsSettings::GetNodeTooltipText() const
 
 void UPCGCullPointsOutsideActorBoundsSettings::CreateKernels(FPCGGPUCompilationContext& InOutContext, UObject* InObjectOuter, TArray<UPCGComputeKernel*>& OutKernels, TArray<FPCGKernelEdge>& OutEdges) const
 {
-	FPCGComputeKernelParams KernelParams;
-	KernelParams.Settings = this;
-	KernelParams.bLogDescriptions = bDumpDataDescriptions;
-
-	UPCGCullPointsOutsideActorBoundsKernel* Kernel = InOutContext.NewObject_AnyThread<UPCGCullPointsOutsideActorBoundsKernel>(InObjectOuter);
-	Kernel->Initialize(KernelParams);
-	OutKernels.Add(Kernel);
-
-	OutEdges.Emplace(FPCGPinReference(PCGPinConstants::DefaultInputLabel), FPCGPinReference(Kernel, PCGPinConstants::DefaultInputLabel));
-	OutEdges.Emplace(FPCGPinReference(Kernel, PCGPinConstants::DefaultOutputLabel), FPCGPinReference(PCGPinConstants::DefaultOutputLabel));
+	PCGKernelHelpers::FCreateKernelParams CreateParams(InObjectOuter, this);
+	PCGKernelHelpers::CreateKernel<UPCGCullPointsOutsideActorBoundsKernel>(InOutContext, CreateParams, OutKernels, OutEdges);
 
 	// @todo_pcg: Inject a compaction kernel to automatically compact after culling. Perhaps with a threshold to only compact if 10%+ of points were culled?
 }
diff --git a/Engine/Plugins/PCG/Source/PCG/Private/Elements/PCGDataNum.cpp b/Engine/Plugins/PCG/Source/PCG/Private/Elements/PCGDataNum.cpp
--- a/Engine/Plugins/PCG/Source/PCG/Private/Elements/PCGDataNum.cpp
+++ b/Engine/Plugins/PCG/Source/PCG/Private/Elements/PCGDataNum.cpp
@@ -5,6 +5,7 @@
 #include "PCGContext.h"
 #include "PCGParamData.h"
 #include "PCGPin.h"
+#include "Compute/PCGKernelHelpers.h"
 #include "Compute/BuiltInKernels/PCGDataNumKernel.h"
 #include "Graph/PCGGPUGraphCompilationContext.h"
 
@@ -30,17 +31,8 @@ FText UPCGDataNumSettings::GetNodeTooltipText() const
 
 void UPCGDataNumSettings::CreateKernels(FPCGGPUCompilationContext& InOutContext, UObject* InObjectOuter, TArray<UPCGComputeKernel*>& OutKernels, TArray<FPCGKernelEdge>& OutEdges) const
 {
-	FPCGComputeKernelParams KernelParams;
-	KernelParams.Settings = this;
-	KernelParams.bLogDescriptions = bDumpDataDescriptions;
-
-	UPCGDataNumKernel* Kernel = InOutContext.NewObject_AnyThread<UPCGDataNumKernel>(InObjectOuter);
-	Kernel->SetOutputAttribute(OutputAttributeName);
-	Kernel->Initialize(KernelParams);
-	OutKernels.Add(Kernel);
-
-	OutEdges.Emplace(FPCGPinReference(PCGPinConstants::DefaultInputLabel), FPCGPinReference(Kernel, PCGPinConstants::DefaultInputLabel));
-	OutEdges.Emplace(FPCGPinReference(Kernel, PCGPinConstants::DefaultOutputLabel), FPCGPinReference(PCGPinConstants::DefaultOutputLabel));
+	PCGKernelHelpers::FCreateKernelParams CreateParams(InObjectOuter, this);
+	PCGKernelHelpers::CreateKernel<UPCGDataNumKernel>(InOutContext, CreateParams, OutKernels, OutEdges);
 }
 #endif // WITH_EDITOR
 
diff --git a/Engine/Plugins/PCG/Source/PCG/Private/Elements/PCGNormalToDensity.cpp b/Engine/Plugins/PCG/Source/PCG/Private/Elements/PCGNormalToDensity.cpp
--- a/Engine/Plugins/PCG/Source/PCG/Private/Elements/PCGNormalToDensity.cpp
+++ b/Engine/Plugins/PCG/Source/PCG/Private/Elements/PCGNormalToDensity.cpp
@@ -3,6 +3,7 @@
 #include "Elements/PCGNormalToDensity.h"
 
 #include "PCGContext.h"
+#include "Compute/PCGKernelHelpers.h"
 #include "Compute/Elements/PCGNormalToDensityKernel.h"
 #include "Data/PCGPointData.h"
 #include "Graph/PCGGPUGraphCompilationContext.h"
@@ -15,17 +16,8 @@
 #if WITH_EDITOR
 void UPCGNormalToDensitySettings::CreateKernels(FPCGGPUCompilationContext& InOutContext, UObject* InObjectOuter, TArray<UPCGComputeKernel*>& OutKernels, TArray<FPCGKernelEdge>& OutEdges) const
 {
-	FPCGComputeKernelParams KernelParams;
-	KernelParams.Settings = this;
-	KernelParams.bLogDescriptions = bDumpDataDescriptions;
-
-	UPCGNormalToDensityKernel* Kernel = InOutContext.NewObject_AnyThread<UPCGNormalToDensityKernel>(InObjectOuter);
-	Kernel->Initialize(KernelParams);
-	OutKernels.Add(Kernel);
-
-	// Connect node pins to kernel pins
-	OutEdges.Emplace(FPCGPinReference(PCGPinConstants::DefaultInputLabel), FPCGPinReference(Kernel, PCGPinConstants::DefaultInputLabel));
-	OutEdges.Emplace(FPCGPinReference(Kernel, PCGPinConstants::DefaultOutputLabel), FPCGPinReference(PCGPinConstants::DefaultOutputLabel));
+	PCGKernelHelpers::FCreateKernelParams CreateParams(InObjectOuter, this);
+	PCGKernelHelpers::CreateKernel<UPCGNormalToDensityKernel>(InOutContext, CreateParams, OutKernels, OutEdges);
 }
 #endif // WITH_EDITOR
 
diff --git a/Engine/Plugins/PCG/Source/PCG/Private/Elements/PCGSkinnedMeshSpawner.cpp b/Engine/Plugins/PCG/Source/PCG/Private/Elements/PCGSkinnedMeshSpawner.cpp
--- a/Engine/Plugins/PCG/Source/PCG/Private/Elements/PCGSkinnedMeshSpawner.cpp
+++ b/Engine/Plugins/PCG/Source/PCG/Private/Elements/PCGSkinnedMeshSpawner.cpp
@@ -5,6 +5,7 @@
 #include "PCGComponent.h"
 #include "PCGCustomVersion.h"
 #include "PCGManagedResource.h"
+#include "Compute/PCGKernelHelpers.h"
 #include "Compute/BuiltInKernels/PCGCountUniqueAttributeValuesKernel.h"
 #include "Data/PCGPointData.h"
 #include "Data/PCGSpatialData.h"
@@ -39,33 +40,22 @@ UPCGSkinnedMeshSpawnerSettings::UPCGSkinnedMeshSpawnerSettings(const FObjectInit
 #if WITH_EDITOR
 void UPCGSkinnedMeshSpawnerSettings::CreateKernels(FPCGGPUCompilationContext& InOutContext, UObject* InObjectOuter, TArray<UPCGComputeKernel*>& OutKernels, TArray<FPCGKernelEdge>& OutEdges) const
 {
-	FPCGComputeKernelParams KernelParams;
-	KernelParams.Settings = this;
-	KernelParams.bLogDescriptions = bDumpDataDescriptions;
+	PCGKernelHelpers::FCreateKernelParams CreateParams(InObjectOuter, this);
 
-	UPCGSkinnedMeshSpawnerKernel* Kernel = InOutContext.NewObject_AnyThread<UPCGSkinnedMeshSpawnerKernel>(InObjectOuter);
-	Kernel->Initialize(KernelParams);
+	UPCGSkinnedMeshSpawnerKernel* SpawnerKernel = PCGKernelHelpers::CreateKernel<UPCGSkinnedMeshSpawnerKernel>(InOutContext, CreateParams, OutKernels, OutEdges);
 
-	OutKernels.Add(Kernel);
-
-	// Connect node pins to kernel pins
-	OutEdges.Emplace(FPCGPinReference(PCGPinConstants::DefaultInputLabel), FPCGPinReference(Kernel, PCGPinConstants::DefaultInputLabel));
-	OutEdges.Emplace(FPCGPinReference(Kernel, PCGPinConstants::DefaultOutputLabel), FPCGPinReference(PCGPinConstants::DefaultOutputLabel));
-
-	// Add analysis kernel that will count how many instances of each mesh are present.
-	const UPCGSkinnedMeshSelector* Selector = Cast<UPCGSkinnedMeshSelector>(MeshSelectorParameters);
-	if (ensure(Selector))
+	// If doing by-attribute selection, add analysis kernel that will count how many instances of each mesh are present.
+	if (const UPCGSkinnedMeshSelector* Selector = Cast<UPCGSkinnedMeshSelector>(MeshSelectorParameters); ensure(Selector))
 	{
-		UPCGCountUniqueAttributeValuesKernel* CountKernel = InOutContext.NewObject_AnyThread<UPCGCountUniqueAttributeValuesKernel>(InObjectOuter);
-		CountKernel->Initialize(KernelParams);
+		// Don't wire count kernel to node output pin, wire manually to the spawner kernel below.
+		CreateParams.NodeOutputPinsToWire.Empty();
+
+		UPCGCountUniqueAttributeValuesKernel* CountKernel = PCGKernelHelpers::CreateKernel<UPCGCountUniqueAttributeValuesKernel>(InOutContext, CreateParams, OutKernels, OutEdges);
 		CountKernel->SetAttributeName(MeshSelectorParameters->MeshAttribute.GetName());
 		// We operate across all input data rather than spawning for each input data separately.
 		CountKernel->SetEmitPerDataCounts(false);
 
-		OutKernels.Add(CountKernel);
-
-		OutEdges.Emplace(FPCGPinReference(PCGPinConstants::DefaultInputLabel), FPCGPinReference(CountKernel, PCGPinConstants::DefaultInputLabel));
-		OutEdges.Emplace(FPCGPinReference(CountKernel, PCGPinConstants::DefaultOutputLabel), FPCGPinReference(Kernel, PCGSkinnedMeshSpawnerConstants::InstanceCountsPinLabel));
+		OutEdges.Emplace(FPCGPinReference(CountKernel, PCGPinConstants::DefaultOutputLabel), FPCGPinReference(SpawnerKernel, PCGSkinnedMeshSpawnerConstants::InstanceCountsPinLabel));
 	}
 }
 
diff --git a/Engine/Plugins/PCG/Source/PCG/Private/Elements/PCGStaticMeshSpawner.cpp b/Engine/Plugins/PCG/Source/PCG/Private/Elements/PCGStaticMeshSpawner.cpp
--- a/Engine/Plugins/PCG/Source/PCG/Private/Elements/PCGStaticMeshSpawner.cpp
+++ b/Engine/Plugins/PCG/Source/PCG/Private/Elements/PCGStaticMeshSpawner.cpp
@@ -6,6 +6,7 @@
 #include "PCGCustomVersion.h"
 #include "PCGManagedResource.h"
 #include "Compute/PCGComputeCommon.h"
+#include "Compute/PCGKernelHelpers.h"
 #include "Compute/BuiltInKernels/PCGCountUniqueAttributeValuesKernel.h"
 #include "Data/PCGBasePointData.h"
 #include "Data/PCGSpatialData.h"
@@ -49,31 +50,22 @@ UPCGStaticMeshSpawnerSettings::UPCGStaticMeshSpawnerSettings(const FObjectInitia
 #if WITH_EDITOR
 void UPCGStaticMeshSpawnerSettings::CreateKernels(FPCGGPUCompilationContext& InOutContext, UObject* InObjectOuter, TArray<UPCGComputeKernel*>& OutKernels, TArray<FPCGKernelEdge>& OutEdges) const
 {
-	FPCGComputeKernelParams KernelParams;
-	KernelParams.Settings = this;
-	KernelParams.bLogDescriptions = bDumpDataDescriptions;
+	PCGKernelHelpers::FCreateKernelParams CreateParams(InObjectOuter, this);
 
-	UPCGStaticMeshSpawnerKernel* Kernel = InOutContext.NewObject_AnyThread<UPCGStaticMeshSpawnerKernel>(InObjectOuter);
-	Kernel->Initialize(KernelParams);
-	OutKernels.Add(Kernel);
-
-	// Connect node pins to kernel pins
-	OutEdges.Emplace(FPCGPinReference(PCGPinConstants::DefaultInputLabel), FPCGPinReference(Kernel, PCGPinConstants::DefaultInputLabel));
-	OutEdges.Emplace(FPCGPinReference(Kernel, PCGPinConstants::DefaultOutputLabel), FPCGPinReference(PCGPinConstants::DefaultOutputLabel));
+	UPCGStaticMeshSpawnerKernel* SpawnerKernel = PCGKernelHelpers::CreateKernel<UPCGStaticMeshSpawnerKernel>(InOutContext, CreateParams, OutKernels, OutEdges);
 
 	// If doing by-attribute selection, add analysis kernel that will count how many instances of each mesh are present.
 	if (const UPCGMeshSelectorByAttribute* Selector = Cast<UPCGMeshSelectorByAttribute>(MeshSelectorParameters))
 	{
-		UPCGCountUniqueAttributeValuesKernel* CountKernel = InOutContext.NewObject_AnyThread<UPCGCountUniqueAttributeValuesKernel>(InObjectOuter);
-		CountKernel->Initialize(KernelParams);
+		// Don't wire count kernel to node output pin, wire manually to the spawner kernel below.
+		CreateParams.NodeOutputPinsToWire.Empty();
+
+		UPCGCountUniqueAttributeValuesKernel* CountKernel = PCGKernelHelpers::CreateKernel<UPCGCountUniqueAttributeValuesKernel>(InOutContext, CreateParams, OutKernels, OutEdges);
 		CountKernel->SetAttributeName(Selector->AttributeName);
 		// We operate across all input data rather than spawning for each input data separately.
 		CountKernel->SetEmitPerDataCounts(false);
 
-		OutKernels.Add(CountKernel);
-
-		OutEdges.Emplace(FPCGPinReference(PCGPinConstants::DefaultInputLabel), FPCGPinReference(CountKernel, PCGPinConstants::DefaultInputLabel));
-		OutEdges.Emplace(FPCGPinReference(CountKernel, PCGPinConstants::DefaultOutputLabel), FPCGPinReference(Kernel, PCGStaticMeshSpawnerConstants::InstanceCountsPinLabel));
+		OutEdges.Emplace(FPCGPinReference(CountKernel, PCGPinConstants::DefaultOutputLabel), FPCGPinReference(SpawnerKernel, PCGStaticMeshSpawnerConstants::InstanceCountsPinLabel));
 	}
 }
 
diff --git a/Engine/Plugins/PCG/Source/PCG/Private/Elements/PCGTransformPoints.cpp b/Engine/Plugins/PCG/Source/PCG/Private/Elements/PCGTransformPoints.cpp
--- a/Engine/Plugins/PCG/Source/PCG/Private/Elements/PCGTransformPoints.cpp
+++ b/Engine/Plugins/PCG/Source/PCG/Private/Elements/PCGTransformPoints.cpp
@@ -4,6 +4,7 @@
 
 #include "PCGContext.h"
 #include "PCGPin.h"
+#include "Compute/PCGKernelHelpers.h"
 #include "Data/PCGSpatialData.h"
 #include "Data/PCGBasePointData.h"
 #include "Helpers/PCGAsync.h"
@@ -24,23 +25,11 @@
 class UPCGTransformPointsKernel;
 
 #if WITH_EDITOR
-/**
- * TODO: This function might be interesting to factorized in the base PCGSettings class, or even into a UInterface so we have those boilerplate functions
- * already defined since most of those nodes would work with the default implementation.
- */
 void UPCGTransformPointsSettings::CreateKernels(FPCGGPUCompilationContext& InOutContext, UObject* InObjectOuter,
                                                 TArray<UPCGComputeKernel*>& OutKernels, TArray<FPCGKernelEdge>& OutEdges) const
 {
-	FPCGComputeKernelParams KernelParams;
-	KernelParams.Settings = this;
-	KernelParams.bLogDescriptions = bDumpDataDescriptions;
-
-	UPCGTransformPointsKernel* Kernel = InOutContext.NewObject_AnyThread<UPCGTransformPointsKernel>(InObjectOuter);
-	Kernel->Initialize(KernelParams);
-	OutKernels.Add(Kernel);
-	
-	OutEdges.Emplace(FPCGPinReference(PCGPinConstants::DefaultInputLabel), FPCGPinReference(Kernel, PCGPinConstants::DefaultInputLabel));
-	OutEdges.Emplace(FPCGPinReference(Kernel, PCGPinConstants::DefaultOutputLabel), FPCGPinReference(PCGPinConstants::DefaultOutputLabel));
+	PCGKernelHelpers::FCreateKernelParams CreateParams(InObjectOuter, this);
+	PCGKernelHelpers::CreateKernel<UPCGTransformPointsKernel>(InOutContext, CreateParams, OutKernels, OutEdges);
 }
 #endif // WITH_EDITOR
 
