[PCG] GPU: Fix debugging multi-kernel node soft locking execution.

#rnx

diff --git a/Engine/Plugins/PCG/Source/PCG/Private/Compute/Elements/PCGComputeGraphElement.cpp b/Engine/Plugins/PCG/Source/PCG/Private/Compute/Elements/PCGComputeGraphElement.cpp
--- a/Engine/Plugins/PCG/Source/PCG/Private/Compute/Elements/PCGComputeGraphElement.cpp
+++ b/Engine/Plugins/PCG/Source/PCG/Private/Compute/Elements/PCGComputeGraphElement.cpp
@@ -258,10 +258,7 @@ bool FPCGComputeGraphElement::ExecuteInternal(FPCGContext* InContext) const
 				continue;
 			}
 
-			const UPCGSettings* ProducerSettings = DataProvider->GetProducerSettings();
-			const bool bDataProviderDebugged = ProducerSettings && ProducerSettings->ShouldExecuteOnGPU() && ProducerSettings->bDebug;
-
-			if (DataProvider->IsExportRequired() || bDataProviderDebugged)
+			if (DataProvider->IsExportRequired())
 			{
 				Context->ProvidersWithBufferExports.Add(DataProvider);
 
