[PCG] GPU: Fix attribute partition node not outputting data - point array data support missed from this node.

#rnx

diff --git a/Engine/Plugins/PCG/Source/PCG/Private/Elements/Metadata/PCGMetadataPartitionKernel.cpp b/Engine/Plugins/PCG/Source/PCG/Private/Elements/Metadata/PCGMetadataPartitionKernel.cpp
--- a/Engine/Plugins/PCG/Source/PCG/Private/Elements/Metadata/PCGMetadataPartitionKernel.cpp
+++ b/Engine/Plugins/PCG/Source/PCG/Private/Elements/Metadata/PCGMetadataPartitionKernel.cpp
@@ -118,6 +118,12 @@ bool UPCGMetadataPartitionKernel::ComputeOutputBindingDataDesc(const UPCGCompute
 		return true;
 	}
 
+	// Expect analysis data produced for each input data item (indexing below depends on it).
+	if (!ensure(AnalysisDataIndices.Num() == InputPinDesc.DataDescs.Num()))
+	{
+		return true;
+	}
+
 	bool bAllReadBack = true;
 
 	for (int32 AnalysisDataIndex : AnalysisDataIndices)
@@ -136,8 +142,12 @@ bool UPCGMetadataPartitionKernel::ComputeOutputBindingDataDesc(const UPCGCompute
 
 	FPCGDataCollectionDesc PinDesc;
 
-	for (int32 AnalysisDataIndex : AnalysisDataIndices)
+	for (int32 DataIndex = 0; DataIndex < InBinding->InputDataCollection.TaggedData.Num(); ++DataIndex)
 	{
+		const int32 AnalysisDataIndex = AnalysisDataIndices[DataIndex];
+		
+		check(InBinding->InputDataCollection.TaggedData.IsValidIndex(AnalysisDataIndex));
+
 		const UPCGParamData* AnalysisResultsData = Cast<UPCGParamData>(InBinding->InputDataCollection.TaggedData[AnalysisDataIndex].Data);
 		const UPCGMetadata* AnalysisMetadata = AnalysisResultsData ? AnalysisResultsData->ConstMetadata() : nullptr;
 
@@ -180,7 +190,8 @@ bool UPCGMetadataPartitionKernel::ComputeOutputBindingDataDesc(const UPCGCompute
 				FPCGDataDesc& DataDesc = PinDesc.DataDescs.Emplace_GetRef();
 				DataDesc.Type = EPCGDataType::Point;
 				DataDesc.ElementCount = AttributeValueToCount[Value];
-				DataDesc.AttributeDescs.Append(InputPinDesc.DataDescs[0].AttributeDescs);
+				DataDesc.AttributeDescs.Append(InputPinDesc.DataDescs[DataIndex].AttributeDescs);
+				DataDesc.AllocateProperties(InputPinDesc.DataDescs[DataIndex].GetAllocatedProperties());
 
 				for (FPCGKernelAttributeDesc& AttributeDesc : DataDesc.AttributeDescs)
 				{
