[PCG] GPU: Move name attribute type from legacy uint2 type to int to match string keys.

#rnx

diff --git a/Engine/Plugins/PCG/Shaders/Private/PCGDataCollectionDataInterface.ush b/Engine/Plugins/PCG/Shaders/Private/PCGDataCollectionDataInterface.ush
--- a/Engine/Plugins/PCG/Shaders/Private/PCGDataCollectionDataInterface.ush
+++ b/Engine/Plugins/PCG/Shaders/Private/PCGDataCollectionDataInterface.ush
@@ -314,11 +314,10 @@ int GetStringKey_{DataInterfaceName}(uint InDataIndex, uint InElementIndex, int
 	}
 }
 
-// TODO: left as uint2 which we locked into for 5.5.0. Future version should change this to int.
-uint2 GetName_{DataInterfaceName}(uint InDataIndex, uint InElementIndex, int InAttributeId)
+int GetName_{DataInterfaceName}(uint InDataIndex, uint InElementIndex, int InAttributeId)
 {
-	// Names are represented as ints.
-	return uint2((uint)GetInt_{DataInterfaceName}(InDataIndex, InElementIndex, InAttributeId), 0u);
+	// Names are currently put in the string table.
+	return GetStringKey_{DataInterfaceName}(InDataIndex, InElementIndex, InAttributeId);
 }
 
 // #################### ATTRIBUTE SETTERS ##########################
@@ -432,11 +431,10 @@ void SetStringKey_{DataInterfaceName}(uint InDataIndex, uint InElementIndex, int
 	SetInt_{DataInterfaceName}(InDataIndex, InElementIndex, InAttributeId, InValue);
 }
 
-// TODO: InValue left as uint2 which we locked into for 5.5.0. Future version should change this to int.
-void SetName_{DataInterfaceName}(uint InDataIndex, uint InElementIndex, int InAttributeId, uint2 InValue)
+void SetName_{DataInterfaceName}(uint InDataIndex, uint InElementIndex, int InAttributeId, int InValue)
 {
 	// Names are represented as ints.
-	SetInt_{DataInterfaceName}(InDataIndex, InElementIndex, InAttributeId, (int)InValue.x);
+	SetInt_{DataInterfaceName}(InDataIndex, InElementIndex, InAttributeId, InValue);
 }
 
 // #################### FIRST ATTRIBUTE GETTERS ##########################
@@ -544,11 +542,10 @@ int GetFirstAttributeAsStringKey_{DataInterfaceName}()
 	return GetFirstAttributeAsInt_{DataInterfaceName}();
 }
 
-// @todo_pcg: Left as uint2 which we locked into for 5.5.0. Future version should change this to int.
-uint2 GetFirstAttributeAsName_{DataInterfaceName}()
+int GetFirstAttributeAsName_{DataInterfaceName}()
 {
 	// Names are represented as ints.
-	return uint2((uint)GetFirstAttributeAsInt_{DataInterfaceName}(), 0u);
+	return GetFirstAttributeAsInt_{DataInterfaceName}();
 }
 
 // #################### ATOMICS ##########################
diff --git a/Engine/Plugins/PCG/Source/PCG/Private/Compute/DataInterfaces/PCGDataCollectionDataInterface.cpp b/Engine/Plugins/PCG/Source/PCG/Private/Compute/DataInterfaces/PCGDataCollectionDataInterface.cpp
--- a/Engine/Plugins/PCG/Source/PCG/Private/Compute/DataInterfaces/PCGDataCollectionDataInterface.cpp
+++ b/Engine/Plugins/PCG/Source/PCG/Private/Compute/DataInterfaces/PCGDataCollectionDataInterface.cpp
@@ -164,7 +164,7 @@ void UPCGDataCollectionDataInterface::GetSupportedInputs(TArray<FShaderFunctionD
 
 		OutFunctions.AddDefaulted_GetRef()
 			.SetName(TEXT("GetName"))
-			.AddReturnType(EShaderFundamentalType::Uint, 2) // Name represented by 2 uints
+			.AddReturnType(EShaderFundamentalType::Int) // Name represented by int
 			.AddParam(EShaderFundamentalType::Uint) // DataIndex
 			.AddParam(EShaderFundamentalType::Uint) // ElementIndex
 			.AddParam(EShaderFundamentalType::Uint); // AttributeId
@@ -247,7 +247,7 @@ void UPCGDataCollectionDataInterface::GetSupportedInputs(TArray<FShaderFunctionD
 			.AddParam(EShaderFundamentalType::Uint) // DataIndex
 			.AddParam(EShaderFundamentalType::Uint) // ElementIndex
 			.AddParam(EShaderFundamentalType::Uint) // AttributeId
-			.AddParam(EShaderFundamentalType::Uint, 2); // Name represented by 2 uints
+			.AddParam(EShaderFundamentalType::Int); // Name represented by int
 	}
 
 	// First Attribute Getters. Useful for overridable params. Currently not exposed to declarations and only called from compiler generated code.
@@ -294,7 +294,7 @@ void UPCGDataCollectionDataInterface::GetSupportedInputs(TArray<FShaderFunctionD
 
 		OutFunctions.AddDefaulted_GetRef()
 			.SetName(TEXT("GetFirstAttributeAsName"))
-			.AddReturnType(EShaderFundamentalType::Uint, 2); // Name represented by 2 uints
+			.AddReturnType(EShaderFundamentalType::Int); // Name represented by int
 	}
 
 	// Atomics
diff --git a/Engine/Plugins/PCG/Source/PCG/Private/Compute/Elements/PCGCustomHLSL.cpp b/Engine/Plugins/PCG/Source/PCG/Private/Compute/Elements/PCGCustomHLSL.cpp
--- a/Engine/Plugins/PCG/Source/PCG/Private/Compute/Elements/PCGCustomHLSL.cpp
+++ b/Engine/Plugins/PCG/Source/PCG/Private/Compute/Elements/PCGCustomHLSL.cpp
@@ -447,7 +447,7 @@ void UPCGCustomHLSLSettings::UpdateInputDeclarations()
 			"uint {0}_GetNumElements();\n"
 			"uint {0}_GetNumElements(uint DataIndex);\n"
 			"\n"
-			"// Valid types: bool, int, float, float2, float3, float4, Rotator (float3), Quat (float4), Transform (float4x4), StringKey (int), Name (uint2)\n"
+			"// Valid types: bool, int, float, float2, float3, float4, Rotator (float3), Quat (float4), Transform (float4x4), StringKey (int), Name (int)\n"
 			"\n"
 			"{type} {0}_Get{type}(uint DataIndex, uint ElementIndex, 'AttributeName');\n"
 			"\n"
