[PCG] Moving PCGDynamicMeshData into the main plugin

* GeometryCore and Framework are now part of the engine.

#rnx

diff --git a/Engine/Plugins/PCG/Config/BasePCG.ini b/Engine/Plugins/PCG/Config/BasePCG.ini
--- a/Engine/Plugins/PCG/Config/BasePCG.ini
+++ b/Engine/Plugins/PCG/Config/BasePCG.ini
@@ -97,6 +97,7 @@
 +ClassRedirects=(OldName="PCGSortPointsSettings",NewName="/Script/PCG.PCGSortAttributesSettings")
 +ClassRedirects=(OldName="PCGMetadataOperationSettings",NewName="/Script/PCG.PCGCopyAttributesSettings")
 +ClassRedirects=(OldName="PCGCreateSplineMeshSettings",NewName="/Script/PCG.PCGSpawnSplineMeshSettings")
++ClassRedirects=(OldName="/Script/PCGGeometryScriptInterop.PCGDynamicMeshData",NewName="/Script/PCG.PCGDynamicMeshData")
 
 +StructRedirects=(OldName="PCGPointFilterThresholdSettings",NewName="/Script/PCG.PCGAttributeFilterThresholdSettings")
 +StructRedirects=(OldName="PCGActorPropertyOverrideDescription",NewName="/Script/PCG.PCGObjectPropertyOverrideDescription")
diff --git a/Engine/Plugins/PCG/Source/PCG/PCG.Build.cs b/Engine/Plugins/PCG/Source/PCG/PCG.Build.cs
--- a/Engine/Plugins/PCG/Source/PCG/PCG.Build.cs
+++ b/Engine/Plugins/PCG/Source/PCG/PCG.Build.cs
@@ -17,8 +17,9 @@ namespace UnrealBuildTool.Rules
 					"DeveloperSettings",
 					"Engine",
 					"Foliage",
-					"GeometryCore",
 					"GeometryAlgorithms",
+					"GeometryCore",
+					"GeometryFramework",
 					"Landscape",
 					"PhysicsCore",
 					"Projects",
diff --git a/Engine/Plugins/PCG/Source/PCG/Private/Data/PCGDynamicMeshData.cpp b/Engine/Plugins/PCG/Source/PCG/Private/Data/PCGDynamicMeshData.cpp
new file mode 100644
--- /dev/null
+++ b/Engine/Plugins/PCG/Source/PCG/Private/Data/PCGDynamicMeshData.cpp
@@ -0,0 +1,249 @@
+// Copyright Epic Games, Inc. All Rights Reserved.
+
+#include "Data/PCGDynamicMeshData.h"
+
+#include "PCGContext.h"
+#include "PCGModule.h"
+#include "Data/PCGPointArrayData.h"
+#include "Data/PCGPointData.h"
+#include "Elements/PCGVolumeSampler.h"
+
+#include "UDynamicMesh.h"
+#include "Components/DynamicMeshComponent.h"
+#include "Helpers/PCGHelpers.h"
+#include "Serialization/ArchiveCrc32.h"
+#include "UObject/Package.h"
+
+#include UE_INLINE_GENERATED_CPP_BY_NAME(PCGDynamicMeshData)
+
+UPCGDynamicMeshData::UPCGDynamicMeshData(const FObjectInitializer& ObjectInitializer)
+	: UPCGSpatialData(ObjectInitializer)
+{
+	DynamicMesh = ObjectInitializer.CreateDefaultSubobject<UDynamicMesh>(this, TEXT("DynamicMesh"));
+}
+
+void UPCGDynamicMeshData::Initialize(UDynamicMesh* InMesh, bool bCanTakeOwnership, const TArray<UMaterialInterface*>& InMaterials)
+{
+	check(DynamicMesh);
+
+	if (!InMesh)
+	{
+		UE_LOG(LogPCG, Error, TEXT("Dynamic Mesh initialization failed because InMesh is null."));
+		return;
+	}
+
+	if (bCanTakeOwnership)
+	{
+		Initialize(std::move(InMesh->GetMeshRef()), InMaterials);
+	}
+	else
+	{
+		// Make a copy
+		Initialize(UE::Geometry::FDynamicMesh3(InMesh->GetMeshRef()), InMaterials);
+	}
+}
+
+void UPCGDynamicMeshData::Initialize(UE::Geometry::FDynamicMesh3&& InMesh, const TArray<UMaterialInterface*>& InMaterials)
+{
+	check(DynamicMesh);
+	DynamicMesh->SetMesh(std::move(InMesh));
+	ResetBounds();
+	
+	bDynamicMeshOctreeIsDirty = true;
+	
+	SetMaterials(InMaterials);
+}
+
+void UPCGDynamicMeshData::ResetBounds() const
+{
+	check(DynamicMesh);
+	CachedBounds = static_cast<FBox>(DynamicMesh->GetMeshRef().GetBounds());
+	bDynamicMeshBoundsAreDirty = false;
+}
+
+FBox UPCGDynamicMeshData::GetBounds() const
+{
+	if (bDynamicMeshBoundsAreDirty)
+	{
+		DynamicMeshBoundsLock.Lock();
+		if (bDynamicMeshBoundsAreDirty)
+		{
+			ResetBounds();
+		}
+		DynamicMeshBoundsLock.Unlock();
+	}
+
+	return CachedBounds;
+}
+
+void UPCGDynamicMeshData::AddToCrc(FArchiveCrc32& Ar, bool bFullDataCrc) const
+{
+	Super::AddToCrc(Ar, bFullDataCrc);
+
+	if (bFullDataCrc)
+	{
+		// Implementation note: metadata not supported at this point.
+
+		FString ClassName = StaticClass()->GetPathName();
+		Ar << ClassName;
+
+		Ar << DynamicMesh->GetMeshRef();
+	}
+	else
+	{
+		AddUIDToCrc(Ar);
+	}
+}
+
+bool UPCGDynamicMeshData::SamplePoint(const FTransform& InTransform, const FBox& InBounds, FPCGPoint& OutPoint, UPCGMetadata* OutMetadata) const
+{
+	if (!DynamicMesh)
+	{
+		return false;
+	}
+	
+	const FVector BoxCenter = InTransform.TransformPosition(InBounds.GetCenter());
+	const FVector HalfScaledExtents(InBounds.GetExtent() * InTransform.GetScale3D());
+	const UE::Geometry::FAxisAlignedBox3d ScaledBounds(BoxCenter - HalfScaledExtents, BoxCenter + HalfScaledExtents);
+	
+	if (!ScaledBounds.Intersects(GetBounds()))
+	{
+		return false;
+	}
+
+	const UE::Geometry::FDynamicMeshOctree3& Octree = GetDynamicMeshOctree();
+	const UE::Geometry::FDynamicMesh3& MeshRef = DynamicMesh->GetMeshRef();
+	
+
+	// TODO: Awful way to do this, need to implement a triangle AABB box intersection if we want more precision. But not even sure that is the right approach using the Octree
+	auto TriangleIntersectionTest = [this, &ScaledBounds, MeshRef](int32 TriangleID) -> bool
+	{
+		return ScaledBounds.Intersects(MeshRef.GetTriBounds(TriangleID));
+	};
+
+	const bool bIntersect = Octree.ParallelOverlapAnyQuery(ScaledBounds, TriangleIntersectionTest) != INDEX_NONE;
+
+	if (bIntersect)
+	{
+		OutPoint.Transform = InTransform;
+		OutPoint.SetLocalBounds(InBounds);
+		OutPoint.Density = 1.0f;
+		return true;
+	}
+	else
+	{
+		return false;
+	}
+}
+
+const UE::Geometry::FDynamicMeshOctree3& UPCGDynamicMeshData::GetDynamicMeshOctree() const
+{
+	if (DynamicMesh && bDynamicMeshOctreeIsDirty)
+	{
+		FScopeLock Lock(&DynamicMeshOctreeLock);
+
+		if (bDynamicMeshOctreeIsDirty)
+		{
+			if (DynamicMeshOctree.Mesh != DynamicMesh->GetMeshPtr())
+			{
+				// Was never initialized, or with another mesh, reset it
+				if (DynamicMeshOctree.Mesh)
+				{
+					DynamicMeshOctree = UE::Geometry::FDynamicMeshOctree3{};
+				}
+				
+				DynamicMeshOctree.Initialize(&DynamicMesh->GetMeshRef());
+			}
+			else
+			{
+				// We can re-insert all triangles to avoid recomputing the whole octree
+				DynamicMeshOctree.ReinsertTriangles(DynamicMesh->GetMeshRef().TriangleIndicesItr());
+			}
+			
+			bDynamicMeshOctreeIsDirty = false;
+		}
+	}
+
+	return DynamicMeshOctree;
+}
+
+const UPCGBasePointData* UPCGDynamicMeshData::ToBasePointData(FPCGContext* Context, const FBox& InBounds, TSubclassOf<UPCGBasePointData> PointDataClass) const
+{
+	TRACE_CPUPROFILER_EVENT_SCOPE(UPCGDynamicMeshData::CreatePointData);
+
+	check(DynamicMesh);
+
+	TArray<FVector> Vertices;
+
+	DynamicMesh->ProcessMesh([&Vertices](const FDynamicMesh3& ReadMesh)
+		{
+			Vertices.Init(FVector::ZeroVector, ReadMesh.MaxVertexID());
+			for (int32 vid : ReadMesh.VertexIndicesItr())
+			{
+				Vertices[vid] = static_cast<FVector>(ReadMesh.GetVertex(vid));
+			}
+		});
+
+	UPCGBasePointData* Data = FPCGContext::NewObject_AnyThread<UPCGBasePointData>(Context, GetTransientPackage(), PointDataClass);
+	Data->SetNumPoints(Vertices.Num());
+	Data->AllocateProperties(EPCGPointNativeProperties::Transform | EPCGPointNativeProperties::Seed);
+
+	int32 NumWritten = 0;
+	FPCGPointValueRanges OutRanges(Data, /*bAllocate=*/false);
+
+	for (const FVector& Vertex : Vertices)
+	{
+		if (InBounds.IsValid && !InBounds.IsInside(Vertex))
+		{
+			continue;
+		}
+		
+		OutRanges.TransformRange[NumWritten].SetLocation(Vertex);
+		OutRanges.SeedRange[NumWritten] = PCGHelpers::ComputeSeedFromPosition(Vertex);
+
+		++NumWritten;
+	}
+
+	Data->SetNumPoints(NumWritten);
+
+	UE_LOG(LogPCG, Verbose, TEXT("Dynamic Mesh extracted %d points"), NumWritten);
+
+	return Data;
+}
+
+const UPCGPointArrayData* UPCGDynamicMeshData::ToPointArrayData(FPCGContext* Context, const FBox& InBounds) const
+{
+	return CastChecked<UPCGPointArrayData>(ToBasePointData(Context, InBounds, UPCGPointArrayData::StaticClass()));
+}
+
+const UPCGPointData* UPCGDynamicMeshData::ToPointData(FPCGContext* Context, const FBox& InBounds) const
+{
+	return CastChecked<UPCGPointData>(ToBasePointData(Context, InBounds, UPCGPointData::StaticClass()));
+}
+
+void UPCGDynamicMeshData::SetMaterials(const TArray<UMaterialInterface*>& InMaterials)
+{
+	Materials.Reset(InMaterials.Num());
+	Materials.Append(InMaterials);
+}
+
+void UPCGDynamicMeshData::InitializeDynamicMeshComponentFromData(UDynamicMeshComponent* InComponent) const
+{
+	check(InComponent);
+
+	InComponent->SetMesh(UE::Geometry::FDynamicMesh3(DynamicMesh->GetMeshRef()));
+	for (int32 i = 0; i < Materials.Num(); ++i)
+	{
+		InComponent->SetMaterial(i, Materials[i]);
+	}
+}
+
+UPCGSpatialData* UPCGDynamicMeshData::CopyInternal(FPCGContext* Context) const
+{
+	UPCGDynamicMeshData* NewDynamicMeshData = FPCGContext::NewObject_AnyThread<UPCGDynamicMeshData>(Context);
+	
+	NewDynamicMeshData->Initialize(DynamicMesh, /*bCanTakeOwnership=*/false);
+	NewDynamicMeshData->Materials = Materials;
+
+	return NewDynamicMeshData;
+}
diff --git a/Engine/Plugins/PCG/Source/PCG/Public/Data/PCGDynamicMeshData.h b/Engine/Plugins/PCG/Source/PCG/Public/Data/PCGDynamicMeshData.h
new file mode 100644
--- /dev/null
+++ b/Engine/Plugins/PCG/Source/PCG/Public/Data/PCGDynamicMeshData.h
@@ -0,0 +1,99 @@
+// Copyright Epic Games, Inc. All Rights Reserved.
+
+#pragma once
+
+#include "Data/PCGSpatialData.h"
+
+#include "DynamicMesh/DynamicMeshOctree3.h"
+#include "DynamicMesh/MeshIndexMappings.h"
+
+#include "Misc/SpinLock.h"
+
+#include "PCGDynamicMeshData.generated.h"
+
+struct FPCGContext;
+class UDynamicMesh;
+class UDynamicMeshComponent;
+class UMaterialInterface;
+
+UCLASS(MinimalAPI, BlueprintType, ClassGroup = (Procedural))
+class UPCGDynamicMeshData : public UPCGSpatialData
+{
+	GENERATED_BODY()
+
+public:
+	PCG_API UPCGDynamicMeshData(const FObjectInitializer& ObjectInitializer);
+	
+	PCG_API void Initialize(UDynamicMesh* InMesh, bool bCanTakeOwnership = false, const TArray<UMaterialInterface*>& InOptionalMaterials = {});
+	PCG_API void Initialize(UE::Geometry::FDynamicMesh3&& InMesh, const TArray<UMaterialInterface*>& InOptionalMaterials = {});
+
+	/**
+	 * Initialize the dynamic mesh data from an input dynamic mesh object.
+	 * If the input dynamic mesh is not meant to be re-used after this initialization, you can set Can Take Ownership to true. Be careful as it
+	 * will put the previous object in an invalid state.
+	 * You can also pass an array of materials that correspond to the referenced materials in the dynamic mesh.
+	 */
+	UFUNCTION(BlueprintCallable, Category="DynamicMesh", meta = (DisplayName = "Initialize", AutoCreateRefTerm = "InMaterials"))
+	void K2_Initialize(UDynamicMesh* InMesh, const TArray<UMaterialInterface*>& InMaterials, bool bCanTakeOwnership = false) { Initialize(InMesh, bCanTakeOwnership, InMaterials); }
+	
+	// ~Begin UPCGData interface
+	virtual EPCGDataType GetDataType() const override { return EPCGDataType::DynamicMesh; }
+	PCG_API virtual void AddToCrc(FArchiveCrc32& Ar, bool bFullDataCrc) const override;
+	// ~End UPCGData interface
+
+	// ~Begin UPCGSpatialData interface
+	virtual int GetDimension() const override { return 3; }
+	PCG_API virtual FBox GetBounds() const override;
+	PCG_API virtual bool SamplePoint(const FTransform& Transform, const FBox& Bounds, FPCGPoint& OutPoint, UPCGMetadata* OutMetadata) const override;
+	// TODO needs an implementation to support projection
+	//virtual bool ProjectPoint(const FTransform& InTransform, const FBox& InBounds, const FPCGProjectionParams& InParams, FPCGPoint& OutPoint, UPCGMetadata* OutMetadata) const;
+	//~End UPCGSpatialData interface
+
+	PCG_API const UE::Geometry::FDynamicMeshOctree3& GetDynamicMeshOctree() const;
+	
+	UDynamicMesh* GetMutableDynamicMesh() { bDynamicMeshBoundsAreDirty = true; bDynamicMeshOctreeIsDirty = true; return DynamicMesh; }
+	const UDynamicMesh* GetDynamicMesh() const { return DynamicMesh; }
+
+	UFUNCTION(BlueprintCallable, Category="DynamicMesh")
+	PCG_API void SetMaterials(const TArray<UMaterialInterface*>& InMaterials);
+	
+	TArray<TObjectPtr<UMaterialInterface>>& GetMutableMaterials() { return Materials; }
+	const TArray<TObjectPtr<UMaterialInterface>>& GetMaterials() const { return Materials; }
+
+	// Copy the mesh of the data into the component and set the materials.
+	PCG_API void InitializeDynamicMeshComponentFromData(UDynamicMeshComponent* InComponent) const;
+
+protected:
+	// ~Begin UPCGData interface
+	virtual bool SupportsFullDataCrc() const override { return true; }
+	// ~End UPCGData interface
+
+	// ~Begin UPCGSpatialData interface
+	PCG_API virtual UPCGSpatialData* CopyInternal(FPCGContext* Context) const override;
+public:
+	PCG_API virtual const UPCGPointData* ToPointData(FPCGContext* Context, const FBox& InBounds) const override;
+	PCG_API virtual const UPCGPointArrayData* ToPointArrayData(FPCGContext* Context, const FBox& InBounds) const override;
+	//~End UPCGSpatialData interface
+
+private:
+	// const but will set the mutable CachedBounds
+	void ResetBounds() const;
+	const UPCGBasePointData* ToBasePointData(FPCGContext* Context, const FBox& InBounds, TSubclassOf<UPCGBasePointData> PointDataClass) const;
+
+protected:
+	UPROPERTY(BlueprintReadOnly, VisibleAnywhere, Instanced, Category="DynamicMesh")
+	TObjectPtr<UDynamicMesh> DynamicMesh;
+	
+	UPROPERTY(BlueprintReadOnly, VisibleAnywhere, Instanced, Category="DynamicMesh")
+	TArray<TObjectPtr<UMaterialInterface>> Materials;
+
+	mutable UE::Geometry::FDynamicMeshOctree3 DynamicMeshOctree;
+	mutable bool bDynamicMeshOctreeIsDirty = true;
+	mutable FCriticalSection DynamicMeshOctreeLock;
+	
+	mutable FBox CachedBounds = FBox(EForceInit::ForceInit);
+	mutable bool bDynamicMeshBoundsAreDirty = true;
+	mutable UE::FSpinLock DynamicMeshBoundsLock;
+};
+
+#undef UE_API
diff --git a/Engine/Plugins/PCGInterops/PCGGeometryScriptInterop/Source/PCGGeometryScriptInterop/Private/Data/PCGDynamicMeshData.cpp b/Engine/Plugins/PCGInterops/PCGGeometryScriptInterop/Source/PCGGeometryScriptInterop/Private/Data/PCGDynamicMeshData.cpp
deleted file mode 100644
diff --git a/Engine/Plugins/PCGInterops/PCGGeometryScriptInterop/Source/PCGGeometryScriptInterop/Public/Data/PCGDynamicMeshData.h b/Engine/Plugins/PCGInterops/PCGGeometryScriptInterop/Source/PCGGeometryScriptInterop/Public/Data/PCGDynamicMeshData.h
deleted file mode 100644
