[PCG] Simplify GPU node implementation
* GetCookedSource now has default implementation that calls GetSourceFilePath, instead of requiring all implementers to read the file themselves
* Simplified implementation of ComputeThreadCount for passthrough kernels

#rnx

diff --git a/Engine/Plugins/PCG/Source/PCG/Internal/Compute/PCGDataForGPU.h b/Engine/Plugins/PCG/Source/PCG/Internal/Compute/PCGDataForGPU.h
--- a/Engine/Plugins/PCG/Source/PCG/Internal/Compute/PCGDataForGPU.h
+++ b/Engine/Plugins/PCG/Source/PCG/Internal/Compute/PCGDataForGPU.h
@@ -219,6 +219,9 @@ struct FPCGDataCollectionDesc
 	/** Unpack a buffer of 8-bit uints to a data collection. */
 	EPCGUnpackDataCollectionResult UnpackDataCollection(FPCGContext* InContext, const TArray<uint8>& InPackedData, FName InPin, const TArray<FString>& InStringTable, FPCGDataCollection& OutDataCollection) const;
 
+	/** Compute total number of processing elements. */
+	uint32 ComputeDataElementCount() const;
+
 	/** Compute total number of processing elements of the given type. */
 	uint32 ComputeDataElementCount(EPCGDataType InDataType) const;
 
diff --git a/Engine/Plugins/PCG/Source/PCG/Private/Compute/BuiltInKernels/PCGCountUniqueAttributeValuesKernel.cpp b/Engine/Plugins/PCG/Source/PCG/Private/Compute/BuiltInKernels/PCGCountUniqueAttributeValuesKernel.cpp
--- a/Engine/Plugins/PCG/Source/PCG/Private/Compute/BuiltInKernels/PCGCountUniqueAttributeValuesKernel.cpp
+++ b/Engine/Plugins/PCG/Source/PCG/Private/Compute/BuiltInKernels/PCGCountUniqueAttributeValuesKernel.cpp
@@ -176,13 +176,6 @@ bool UPCGCountUniqueAttributeValuesKernel::DoesOutputPinRequireZeroInitializatio
 }
 
 #if WITH_EDITOR
-FString UPCGCountUniqueAttributeValuesKernel::GetCookedSource(FPCGGPUCompilationContext& InOutContext) const
-{
-	FString TemplateFile;
-	ensure(LoadShaderSourceFile(TEXT("/Plugin/PCG/Private/BuiltInKernels/PCGCountUniqueAttributeValues.usf"), EShaderPlatform::SP_PCD3D_SM5, &TemplateFile, nullptr));
-	return TemplateFile;
-}
-
 void UPCGCountUniqueAttributeValuesKernel::CreateAdditionalInputDataInterfaces(FPCGGPUCompilationContext& InOutContext, UObject* InObjectOuter, TArray<TObjectPtr<UComputeDataInterface>>& OutDataInterfaces) const
 {
 	Super::CreateAdditionalInputDataInterfaces(InOutContext, InObjectOuter, OutDataInterfaces);
diff --git a/Engine/Plugins/PCG/Source/PCG/Private/Compute/BuiltInKernels/PCGCountUniqueAttributeValuesKernel.h b/Engine/Plugins/PCG/Source/PCG/Private/Compute/BuiltInKernels/PCGCountUniqueAttributeValuesKernel.h
--- a/Engine/Plugins/PCG/Source/PCG/Private/Compute/BuiltInKernels/PCGCountUniqueAttributeValuesKernel.h
+++ b/Engine/Plugins/PCG/Source/PCG/Private/Compute/BuiltInKernels/PCGCountUniqueAttributeValuesKernel.h
@@ -32,7 +32,7 @@ public:
 	virtual int ComputeThreadCount(const UPCGDataBinding* Binding) const override;
 	virtual bool DoesOutputPinRequireZeroInitialization(FName InOutputPinLabel) const override;
 #if WITH_EDITOR
-	virtual FString GetCookedSource(FPCGGPUCompilationContext& InOutContext) const override;
+	virtual const TCHAR* GetSourceFilePath() const override { return TEXT("/Plugin/PCG/Private/BuiltInKernels/PCGCountUniqueAttributeValues.usf"); }
 	virtual FString GetEntryPoint() const override { return TEXT("PCGCountUniqueAttributeValuesCS"); }
 	void CreateAdditionalInputDataInterfaces(FPCGGPUCompilationContext& InOutContext, UObject* InObjectOuter, TArray<TObjectPtr<UComputeDataInterface>>& OutDataInterfaces) const override;
 	// Split graph to read back analysis results.
diff --git a/Engine/Plugins/PCG/Source/PCG/Private/Compute/BuiltInKernels/PCGDataNumKernel.cpp b/Engine/Plugins/PCG/Source/PCG/Private/Compute/BuiltInKernels/PCGDataNumKernel.cpp
--- a/Engine/Plugins/PCG/Source/PCG/Private/Compute/BuiltInKernels/PCGDataNumKernel.cpp
+++ b/Engine/Plugins/PCG/Source/PCG/Private/Compute/BuiltInKernels/PCGDataNumKernel.cpp
@@ -86,13 +86,6 @@ bool UPCGDataNumKernel::IsKernelDataValid(FPCGContext* InContext) const
 }
 
 #if WITH_EDITOR
-FString UPCGDataNumKernel::GetCookedSource(FPCGGPUCompilationContext& InOutContext) const
-{
-	FString DataNumFile;
-	ensure(LoadShaderSourceFile(TEXT("/Plugin/PCG/Private/BuiltInKernels/PCGDataNum.usf"), EShaderPlatform::SP_PCD3D_SM5, &DataNumFile, nullptr));
-	return DataNumFile;
-}
-
 void UPCGDataNumKernel::CreateAdditionalInputDataInterfaces(FPCGGPUCompilationContext& InOutContext, UObject* InObjectOuter, TArray<TObjectPtr<UComputeDataInterface>>& OutDataInterfaces) const
 {
 	Super::CreateAdditionalInputDataInterfaces(InOutContext, InObjectOuter, OutDataInterfaces);
diff --git a/Engine/Plugins/PCG/Source/PCG/Private/Compute/BuiltInKernels/PCGDataNumKernel.h b/Engine/Plugins/PCG/Source/PCG/Private/Compute/BuiltInKernels/PCGDataNumKernel.h
--- a/Engine/Plugins/PCG/Source/PCG/Private/Compute/BuiltInKernels/PCGDataNumKernel.h
+++ b/Engine/Plugins/PCG/Source/PCG/Private/Compute/BuiltInKernels/PCGDataNumKernel.h
@@ -19,7 +19,7 @@ public:
 	virtual bool IsKernelDataValid(FPCGContext* InContext) const override;
 
 #if WITH_EDITOR
-	virtual FString GetCookedSource(FPCGGPUCompilationContext& InOutContext) const override;
+	virtual const TCHAR* GetSourceFilePath() const override { return TEXT("/Plugin/PCG/Private/BuiltInKernels/PCGDataNum.usf"); }
 	virtual FString GetEntryPoint() const override { return TEXT("PCGDataNumCS"); }
 	virtual void CreateAdditionalInputDataInterfaces(FPCGGPUCompilationContext& InOutContext, UObject* InObjectOuter, TArray<TObjectPtr<UComputeDataInterface>>& OutDataInterfaces) const override;
 #endif
diff --git a/Engine/Plugins/PCG/Source/PCG/Private/Compute/Elements/PCGCullPointsOutsideActorBoundsKernel.cpp b/Engine/Plugins/PCG/Source/PCG/Private/Compute/Elements/PCGCullPointsOutsideActorBoundsKernel.cpp
--- a/Engine/Plugins/PCG/Source/PCG/Private/Compute/Elements/PCGCullPointsOutsideActorBoundsKernel.cpp
+++ b/Engine/Plugins/PCG/Source/PCG/Private/Compute/Elements/PCGCullPointsOutsideActorBoundsKernel.cpp
@@ -33,29 +33,11 @@ bool UPCGCullPointsOutsideActorBoundsKernel::ComputeOutputBindingDataDesc(const
 
 int UPCGCullPointsOutsideActorBoundsKernel::ComputeThreadCount(const UPCGDataBinding* InBinding) const
 {
-	int ThreadCount = 0;
-
 	const FPCGDataCollectionDesc* OutputPinDesc = InBinding->GetCachedKernelPinDataDesc(this, PCGPinConstants::DefaultOutputLabel, /*bIsInput=*/false);
-
-	if (ensure(OutputPinDesc))
-	{
-		for (const FPCGDataDesc& DataDesc : OutputPinDesc->DataDescs)
-		{
-			ThreadCount += DataDesc.ElementCount;
-		}
-	}
-
-	return ThreadCount;
+	return ensure(OutputPinDesc) ? OutputPinDesc->ComputeDataElementCount() : 0;
 }
 
 #if WITH_EDITOR
-FString UPCGCullPointsOutsideActorBoundsKernel::GetCookedSource(FPCGGPUCompilationContext& InOutContext) const
-{
-	FString TemplateFile;
-	ensure(LoadShaderSourceFile(TEXT("/Plugin/PCG/Private/Elements/PCGCullPointsOutsideActorBounds.usf"), EShaderPlatform::SP_PCD3D_SM5, &TemplateFile, nullptr));
-	return TemplateFile;
-}
-
 void UPCGCullPointsOutsideActorBoundsKernel::CreateAdditionalInputDataInterfaces(FPCGGPUCompilationContext& InOutContext, UObject* InObjectOuter, TArray<TObjectPtr<UComputeDataInterface>>& OutDataInterfaces) const
 {
 	Super::CreateAdditionalInputDataInterfaces(InOutContext, InObjectOuter, OutDataInterfaces);
diff --git a/Engine/Plugins/PCG/Source/PCG/Private/Compute/Elements/PCGCullPointsOutsideActorBoundsKernel.h b/Engine/Plugins/PCG/Source/PCG/Private/Compute/Elements/PCGCullPointsOutsideActorBoundsKernel.h
--- a/Engine/Plugins/PCG/Source/PCG/Private/Compute/Elements/PCGCullPointsOutsideActorBoundsKernel.h
+++ b/Engine/Plugins/PCG/Source/PCG/Private/Compute/Elements/PCGCullPointsOutsideActorBoundsKernel.h
@@ -18,7 +18,7 @@ public:
 	virtual bool ComputeOutputBindingDataDesc(const UPCGComputeGraph* InGraph, FName InOutputPinLabel, UPCGDataBinding* InBinding, FPCGDataCollectionDesc& OutDataDesc) const override;
 	virtual int ComputeThreadCount(const UPCGDataBinding* Binding) const override;
 #if WITH_EDITOR
-	virtual FString GetCookedSource(FPCGGPUCompilationContext& InOutContext) const override;
+	virtual const TCHAR* GetSourceFilePath() const override { return TEXT("/Plugin/PCG/Private/Elements/PCGCullPointsOutsideActorBounds.usf"); }
 	virtual FString GetEntryPoint() const override { return TEXT("PCGCullPointsOutsideActorBoundsCS"); }
 	void CreateAdditionalInputDataInterfaces(FPCGGPUCompilationContext& InOutContext, UObject* InObjectOuter, TArray<TObjectPtr<UComputeDataInterface>>& OutDataInterfaces) const override;
 #endif
diff --git a/Engine/Plugins/PCG/Source/PCG/Private/Compute/Elements/PCGNormalToDensityKernel.cpp b/Engine/Plugins/PCG/Source/PCG/Private/Compute/Elements/PCGNormalToDensityKernel.cpp
--- a/Engine/Plugins/PCG/Source/PCG/Private/Compute/Elements/PCGNormalToDensityKernel.cpp
+++ b/Engine/Plugins/PCG/Source/PCG/Private/Compute/Elements/PCGNormalToDensityKernel.cpp
@@ -42,29 +42,11 @@ bool UPCGNormalToDensityKernel::ComputeOutputBindingDataDesc(const UPCGComputeGr
 
 int UPCGNormalToDensityKernel::ComputeThreadCount(const UPCGDataBinding* InBinding) const
 {
-	int ThreadCount = 0;
-
 	const FPCGDataCollectionDesc* OutputPinDesc = InBinding->GetCachedKernelPinDataDesc(this, PCGPinConstants::DefaultOutputLabel, /*bIsInput=*/false);
-
-	if (ensure(OutputPinDesc))
-	{
-		for (const FPCGDataDesc& DataDesc : OutputPinDesc->DataDescs)
-		{
-			ThreadCount += DataDesc.ElementCount;
-		}
-	}
-
-	return ThreadCount;
+	return ensure(OutputPinDesc) ? OutputPinDesc->ComputeDataElementCount() : 0;
 }
 
 #if WITH_EDITOR
-FString UPCGNormalToDensityKernel::GetCookedSource(FPCGGPUCompilationContext& InOutContext) const
-{
-	FString TemplateFile;
-	ensure(LoadShaderSourceFile(TEXT("/Plugin/PCG/Private/Elements/PCGNormalToDensity.usf"), EShaderPlatform::SP_PCD3D_SM5, &TemplateFile, nullptr));
-	return TemplateFile;
-}
-
 void UPCGNormalToDensityKernel::CreateAdditionalInputDataInterfaces(FPCGGPUCompilationContext& InOutContext, UObject* InObjectOuter, TArray<TObjectPtr<UComputeDataInterface>>& OutDataInterfaces) const
 {
 	Super::CreateAdditionalInputDataInterfaces(InOutContext, InObjectOuter, OutDataInterfaces);
diff --git a/Engine/Plugins/PCG/Source/PCG/Private/Compute/Elements/PCGNormalToDensityKernel.h b/Engine/Plugins/PCG/Source/PCG/Private/Compute/Elements/PCGNormalToDensityKernel.h
--- a/Engine/Plugins/PCG/Source/PCG/Private/Compute/Elements/PCGNormalToDensityKernel.h
+++ b/Engine/Plugins/PCG/Source/PCG/Private/Compute/Elements/PCGNormalToDensityKernel.h
@@ -19,7 +19,7 @@ public:
 	virtual int ComputeThreadCount(const UPCGDataBinding* Binding) const override;
 
 #if WITH_EDITOR
-	virtual FString GetCookedSource(FPCGGPUCompilationContext& InOutContext) const override;
+	virtual const TCHAR* GetSourceFilePath() const override { return TEXT("/Plugin/PCG/Private/Elements/PCGNormalToDensity.usf"); }
 	virtual FString GetEntryPoint() const override { return TEXT("PCGNormalToDensityCS"); }
 	void CreateAdditionalInputDataInterfaces(FPCGGPUCompilationContext& InOutContext, UObject* InObjectOuter, TArray<TObjectPtr<UComputeDataInterface>>& OutDataInterfaces) const override;
 #endif
diff --git a/Engine/Plugins/PCG/Source/PCG/Private/Compute/Elements/PCGTransformPointsKernel.cpp b/Engine/Plugins/PCG/Source/PCG/Private/Compute/Elements/PCGTransformPointsKernel.cpp
--- a/Engine/Plugins/PCG/Source/PCG/Private/Compute/Elements/PCGTransformPointsKernel.cpp
+++ b/Engine/Plugins/PCG/Source/PCG/Private/Compute/Elements/PCGTransformPointsKernel.cpp
@@ -43,29 +43,11 @@ bool UPCGTransformPointsKernel::ComputeOutputBindingDataDesc(const UPCGComputeGr
 
 int UPCGTransformPointsKernel::ComputeThreadCount(const UPCGDataBinding* InBinding) const
 {
-	int ThreadCount = 0;
-
 	const FPCGDataCollectionDesc* OutputPinDesc = InBinding->GetCachedKernelPinDataDesc(this, PCGPinConstants::DefaultOutputLabel, /*bIsInput=*/false);
-
-	if (ensure(OutputPinDesc))
-	{
-		for (const FPCGDataDesc& DataDesc : OutputPinDesc->DataDescs)
-		{
-			ThreadCount += DataDesc.ElementCount;
-		}
-	}
-
-	return ThreadCount;
+	return ensure(OutputPinDesc) ? OutputPinDesc->ComputeDataElementCount() : 0;
 }
 
 #if WITH_EDITOR
-FString UPCGTransformPointsKernel::GetCookedSource(FPCGGPUCompilationContext& InOutContext) const
-{
-	FString TemplateFile;
-	ensure(LoadShaderSourceFile(TEXT("/Plugin/PCG/Private/Elements/PCGTransformPoints.usf"), EShaderPlatform::SP_PCD3D_SM5, &TemplateFile, nullptr));
-	return TemplateFile;
-}
-
 void UPCGTransformPointsKernel::CreateAdditionalInputDataInterfaces(FPCGGPUCompilationContext& InOutContext, UObject* InObjectOuter, TArray<TObjectPtr<UComputeDataInterface>>& OutDataInterfaces) const
 {
 	Super::CreateAdditionalInputDataInterfaces(InOutContext, InObjectOuter, OutDataInterfaces);
diff --git a/Engine/Plugins/PCG/Source/PCG/Private/Compute/Elements/PCGTransformPointsKernel.h b/Engine/Plugins/PCG/Source/PCG/Private/Compute/Elements/PCGTransformPointsKernel.h
--- a/Engine/Plugins/PCG/Source/PCG/Private/Compute/Elements/PCGTransformPointsKernel.h
+++ b/Engine/Plugins/PCG/Source/PCG/Private/Compute/Elements/PCGTransformPointsKernel.h
@@ -18,7 +18,7 @@ public:
 	virtual bool ComputeOutputBindingDataDesc(const UPCGComputeGraph* InGraph, FName InOutputPinLabel, UPCGDataBinding* InBinding, FPCGDataCollectionDesc& OutDataDesc) const override;
 	virtual int ComputeThreadCount(const UPCGDataBinding* Binding) const override;
 #if WITH_EDITOR
-	virtual FString GetCookedSource(FPCGGPUCompilationContext& InOutContext) const override;
+	virtual const TCHAR* GetSourceFilePath() const override { return TEXT("/Plugin/PCG/Private/Elements/PCGTransformPoints.usf"); }
 	virtual FString GetEntryPoint() const override { return TEXT("PCGTransformPointsCS"); }
 	void CreateAdditionalInputDataInterfaces(FPCGGPUCompilationContext& InOutContext, UObject* InObjectOuter, TArray<TObjectPtr<UComputeDataInterface>>& OutDataInterfaces) const override;
 #endif
diff --git a/Engine/Plugins/PCG/Source/PCG/Private/Compute/PCGComputeKernel.cpp b/Engine/Plugins/PCG/Source/PCG/Private/Compute/PCGComputeKernel.cpp
--- a/Engine/Plugins/PCG/Source/PCG/Private/Compute/PCGComputeKernel.cpp
+++ b/Engine/Plugins/PCG/Source/PCG/Private/Compute/PCGComputeKernel.cpp
@@ -198,6 +198,13 @@ bool UPCGComputeKernel::IsKernelDataValid(FPCGContext* InContext) const
 }
 
 #if WITH_EDITOR
+FString UPCGComputeKernel::GetCookedSource(FPCGGPUCompilationContext& InOutContext) const
+{
+	FString SourceFile;
+	ensure(LoadShaderSourceFile(GetSourceFilePath(), EShaderPlatform::SP_PCD3D_SM5, &SourceFile, nullptr));
+	return SourceFile;
+}
+
 void UPCGComputeKernel::GatherAdditionalSources(TArray<TObjectPtr<UComputeSource>>& OutAdditionalSources) const
 {
 	for (const FSoftObjectPath& AdditionalSourcePath : PCGComputeKernel::DefaultAdditionalSourcePaths)
diff --git a/Engine/Plugins/PCG/Source/PCG/Private/Compute/PCGComputeKernel.h b/Engine/Plugins/PCG/Source/PCG/Private/Compute/PCGComputeKernel.h
--- a/Engine/Plugins/PCG/Source/PCG/Private/Compute/PCGComputeKernel.h
+++ b/Engine/Plugins/PCG/Source/PCG/Private/Compute/PCGComputeKernel.h
@@ -101,7 +101,10 @@ public:
 
 #if WITH_EDITOR
 	/** Produces the node specific portion of kernel shader source text, including the main entry point. */
-	virtual FString GetCookedSource(FPCGGPUCompilationContext& InOutContext) const PURE_VIRTUAL(UPCGComputeKernel::GetCookedSource, return {};);
+	virtual FString GetCookedSource(FPCGGPUCompilationContext& InOutContext) const;
+
+	/** Get the virtual file path of the shader source file for this kernel. */
+	virtual const TCHAR* GetSourceFilePath() const PURE_VIRTUAL(UPCGComputeKernel::GetSourceFilePath, return nullptr;);
 
 	/** Get the name of the main kernel function in the source. This name is also displayed in GPU profile scopes. */
 	virtual FString GetEntryPoint() const PURE_VIRTUAL(UPCGComputeKernel::GetEntryPoint, return TEXT("Main"););
diff --git a/Engine/Plugins/PCG/Source/PCG/Private/Compute/PCGDataForGPU.cpp b/Engine/Plugins/PCG/Source/PCG/Private/Compute/PCGDataForGPU.cpp
--- a/Engine/Plugins/PCG/Source/PCG/Private/Compute/PCGDataForGPU.cpp
+++ b/Engine/Plugins/PCG/Source/PCG/Private/Compute/PCGDataForGPU.cpp
@@ -1830,6 +1830,11 @@ EPCGUnpackDataCollectionResult FPCGDataCollectionDesc::UnpackDataCollection(FPCG
 	return EPCGUnpackDataCollectionResult::Success;
 }
 
+uint32 FPCGDataCollectionDesc::ComputeDataElementCount() const
+{
+	return ComputeDataElementCount(EPCGDataType::Any);
+}
+
 uint32 FPCGDataCollectionDesc::ComputeDataElementCount(EPCGDataType InDataType) const
 {
 	uint32 ElementCount = 0;
@@ -1838,14 +1843,7 @@ uint32 FPCGDataCollectionDesc::ComputeDataElementCount(EPCGDataType InDataType)
 	{
 		if (!!(DataDesc.Type & InDataType))
 		{
-			if (!!(DataDesc.Type & EPCGDataType::BaseTexture))
-			{
-				ElementCount += DataDesc.ElementCount2D.X * DataDesc.ElementCount2D.Y;
-			}
-			else
-			{
-				ElementCount += DataDesc.ElementCount;
-			}
+			ElementCount += DataDesc.IsDomain2D() ? (DataDesc.ElementCount2D.X * DataDesc.ElementCount2D.Y) : DataDesc.ElementCount;
 		}
 	}
 
diff --git a/Engine/Plugins/PCG/Source/PCG/Private/Elements/Metadata/PCGMetadataPartitionKernel.cpp b/Engine/Plugins/PCG/Source/PCG/Private/Elements/Metadata/PCGMetadataPartitionKernel.cpp
--- a/Engine/Plugins/PCG/Source/PCG/Private/Elements/Metadata/PCGMetadataPartitionKernel.cpp
+++ b/Engine/Plugins/PCG/Source/PCG/Private/Elements/Metadata/PCGMetadataPartitionKernel.cpp
@@ -224,13 +224,6 @@ int UPCGMetadataPartitionKernel::ComputeThreadCount(const UPCGDataBinding* InBin
 }
 
 #if WITH_EDITOR
-FString UPCGMetadataPartitionKernel::GetCookedSource(FPCGGPUCompilationContext& InOutContext) const
-{
-	FString TemplateFile;
-	ensure(LoadShaderSourceFile(TEXT("/Plugin/PCG/Private/Elements/PCGMetadataPartition.usf"), EShaderPlatform::SP_PCD3D_SM5, &TemplateFile, nullptr));
-	return TemplateFile;
-}
-
 void UPCGMetadataPartitionKernel::CreateAdditionalInputDataInterfaces(FPCGGPUCompilationContext& InOutContext, UObject* InObjectOuter, TArray<TObjectPtr<UComputeDataInterface>>& OutDataInterfaces) const
 {
 	Super::CreateAdditionalInputDataInterfaces(InOutContext, InObjectOuter, OutDataInterfaces);
diff --git a/Engine/Plugins/PCG/Source/PCG/Private/Elements/PCGCopyPointsAnalysisKernel.cpp b/Engine/Plugins/PCG/Source/PCG/Private/Elements/PCGCopyPointsAnalysisKernel.cpp
--- a/Engine/Plugins/PCG/Source/PCG/Private/Elements/PCGCopyPointsAnalysisKernel.cpp
+++ b/Engine/Plugins/PCG/Source/PCG/Private/Elements/PCGCopyPointsAnalysisKernel.cpp
@@ -88,29 +88,11 @@ bool UPCGCopyPointsAnalysisKernel::ComputeOutputBindingDataDesc(const UPCGComput
 
 int UPCGCopyPointsAnalysisKernel::ComputeThreadCount(const UPCGDataBinding* InBinding) const
 {
-	int ThreadCount = 0;
-
 	const FPCGDataCollectionDesc* OutputPinDesc = InBinding->GetCachedKernelPinDataDesc(this, PCGPinConstants::DefaultOutputLabel, /*bIsInput=*/false);
-
-	if (ensure(OutputPinDesc))
-	{
-		for (const FPCGDataDesc& DataDesc : OutputPinDesc->DataDescs)
-		{
-			ThreadCount += DataDesc.ElementCount;
-		}
-	}
-
-	return ThreadCount;
+	return ensure(OutputPinDesc) ? OutputPinDesc->ComputeDataElementCount(EPCGDataType::Point) : 0;
 }
 
 #if WITH_EDITOR
-FString UPCGCopyPointsAnalysisKernel::GetCookedSource(FPCGGPUCompilationContext& InOutContext) const
-{
-	FString TemplateFile;
-	ensure(LoadShaderSourceFile(TEXT("/Plugin/PCG/Private/Elements/PCGCopyPointsAnalysis.usf"), EShaderPlatform::SP_PCD3D_SM5, &TemplateFile, nullptr));
-	return TemplateFile;
-}
-
 void UPCGCopyPointsAnalysisKernel::CreateAdditionalInputDataInterfaces(FPCGGPUCompilationContext& InOutContext, UObject* InObjectOuter, TArray<TObjectPtr<UComputeDataInterface>>& OutDataInterfaces) const
 {
 	Super::CreateAdditionalInputDataInterfaces(InOutContext, InObjectOuter, OutDataInterfaces);
diff --git a/Engine/Plugins/PCG/Source/PCG/Private/Elements/PCGCopyPointsKernel.cpp b/Engine/Plugins/PCG/Source/PCG/Private/Elements/PCGCopyPointsKernel.cpp
--- a/Engine/Plugins/PCG/Source/PCG/Private/Elements/PCGCopyPointsKernel.cpp
+++ b/Engine/Plugins/PCG/Source/PCG/Private/Elements/PCGCopyPointsKernel.cpp
@@ -209,24 +209,15 @@ bool UPCGCopyPointsKernel::ComputeOutputBindingDataDesc(const UPCGComputeGraph*
 
 int UPCGCopyPointsKernel::ComputeThreadCount(const UPCGDataBinding* InBinding) const
 {
-	int ThreadCount = 0;
-
 	const FPCGDataCollectionDesc* OutputPinDesc = InBinding->GetCachedKernelPinDataDesc(this, PCGPinConstants::DefaultOutputLabel, /*bIsInput=*/false);
-
-	if (ensure(OutputPinDesc))
-	{
-		for (const FPCGDataDesc& DataDesc : OutputPinDesc->DataDescs)
-		{
-			ThreadCount += DataDesc.ElementCount;
-		}
-	}
-
-	return ThreadCount;
+	return ensure(OutputPinDesc) ? OutputPinDesc->ComputeDataElementCount(EPCGDataType::Point) : 0;
 }
 
 #if WITH_EDITOR
 FString UPCGCopyPointsKernel::GetCookedSource(FPCGGPUCompilationContext& InOutContext) const
 {
+	const FString SourceFile = Super::GetCookedSource(InOutContext);
+
 	FString OverridableParamFunctions = "";
 
 	if (const UPCGSettings* Settings = GetSettings())
@@ -256,10 +247,7 @@ FString UPCGCopyPointsKernel::GetCookedSource(FPCGGPUCompilationContext& InOutCo
 		}
 	}
 
-	FString TemplateFile;
-	ensure(LoadShaderSourceFile(TEXT("/Plugin/PCG/Private/Elements/PCGCopyPoints.usf"), EShaderPlatform::SP_PCD3D_SM5, &TemplateFile, nullptr));
-
-	return OverridableParamFunctions + TemplateFile;
+	return OverridableParamFunctions + SourceFile;
 }
 
 void UPCGCopyPointsKernel::CreateAdditionalInputDataInterfaces(FPCGGPUCompilationContext& InOutContext, UObject* InObjectOuter, TArray<TObjectPtr<UComputeDataInterface>>& OutDataInterfaces) const
diff --git a/Engine/Plugins/PCG/Source/PCG/Private/Elements/PCGSkinnedMeshSpawnerKernel.cpp b/Engine/Plugins/PCG/Source/PCG/Private/Elements/PCGSkinnedMeshSpawnerKernel.cpp
--- a/Engine/Plugins/PCG/Source/PCG/Private/Elements/PCGSkinnedMeshSpawnerKernel.cpp
+++ b/Engine/Plugins/PCG/Source/PCG/Private/Elements/PCGSkinnedMeshSpawnerKernel.cpp
@@ -64,13 +64,6 @@ int UPCGSkinnedMeshSpawnerKernel::ComputeThreadCount(const UPCGDataBinding* InBi
 }
 
 #if WITH_EDITOR
-FString UPCGSkinnedMeshSpawnerKernel::GetCookedSource(FPCGGPUCompilationContext& InOutContext) const
-{
-	FString TemplateFile;
-	ensure(LoadShaderSourceFile(TEXT("/Plugin/PCG/Private/Elements/PCGStaticMeshSpawner.usf"), EShaderPlatform::SP_PCD3D_SM5, &TemplateFile, nullptr));
-	return TemplateFile;
-}
-
 void UPCGSkinnedMeshSpawnerKernel::CreateAdditionalInputDataInterfaces(FPCGGPUCompilationContext& InOutContext, UObject* InObjectOuter, TArray<TObjectPtr<UComputeDataInterface>>& OutDataInterfaces) const
 {
 	Super::CreateAdditionalInputDataInterfaces(InOutContext, InObjectOuter, OutDataInterfaces);
diff --git a/Engine/Plugins/PCG/Source/PCG/Private/Elements/PCGStaticMeshSpawnerKernel.cpp b/Engine/Plugins/PCG/Source/PCG/Private/Elements/PCGStaticMeshSpawnerKernel.cpp
--- a/Engine/Plugins/PCG/Source/PCG/Private/Elements/PCGStaticMeshSpawnerKernel.cpp
+++ b/Engine/Plugins/PCG/Source/PCG/Private/Elements/PCGStaticMeshSpawnerKernel.cpp
@@ -122,13 +122,6 @@ int UPCGStaticMeshSpawnerKernel::ComputeThreadCount(const UPCGDataBinding* InBin
 }
 
 #if WITH_EDITOR
-FString UPCGStaticMeshSpawnerKernel::GetCookedSource(FPCGGPUCompilationContext& InOutContext) const
-{
-	FString TemplateFile;
-	ensure(LoadShaderSourceFile(TEXT("/Plugin/PCG/Private/Elements/PCGStaticMeshSpawner.usf"), EShaderPlatform::SP_PCD3D_SM5, &TemplateFile, nullptr));
-	return TemplateFile;
-}
-
 void UPCGStaticMeshSpawnerKernel::CreateAdditionalInputDataInterfaces(FPCGGPUCompilationContext& InOutContext, UObject* InObjectOuter, TArray<TObjectPtr<UComputeDataInterface>>& OutDataInterfaces) const
 {
 	Super::CreateAdditionalInputDataInterfaces(InOutContext, InObjectOuter, OutDataInterfaces);
diff --git a/Engine/Plugins/PCG/Source/PCG/Public/Elements/Metadata/PCGMetadataPartitionKernel.h b/Engine/Plugins/PCG/Source/PCG/Public/Elements/Metadata/PCGMetadataPartitionKernel.h
--- a/Engine/Plugins/PCG/Source/PCG/Public/Elements/Metadata/PCGMetadataPartitionKernel.h
+++ b/Engine/Plugins/PCG/Source/PCG/Public/Elements/Metadata/PCGMetadataPartitionKernel.h
@@ -23,7 +23,7 @@ public:
 	virtual bool ComputeOutputBindingDataDesc(const UPCGComputeGraph* InGraph, FName InOutputPinLabel, UPCGDataBinding* InBinding, FPCGDataCollectionDesc& OutDataDesc) const override;
 	virtual int ComputeThreadCount(const UPCGDataBinding* Binding) const override;
 #if WITH_EDITOR
-	virtual FString GetCookedSource(FPCGGPUCompilationContext& InOutContext) const override;
+	virtual const TCHAR* GetSourceFilePath() const override { return TEXT("/Plugin/PCG/Private/Elements/PCGMetadataPartition.usf"); }
 	virtual FString GetEntryPoint() const override { return TEXT("PCGMetadataPartitionCS"); }
 	void CreateAdditionalInputDataInterfaces(FPCGGPUCompilationContext& InOutContext, UObject* InObjectOuter, TArray<TObjectPtr<UComputeDataInterface>>& OutDataInterfaces) const override;
 	// Split graph to update element counts on CPU before proceeding. Compacts data for downstream kernels.
diff --git a/Engine/Plugins/PCG/Source/PCG/Public/Elements/PCGCopyPointsAnalysisKernel.h b/Engine/Plugins/PCG/Source/PCG/Public/Elements/PCGCopyPointsAnalysisKernel.h
--- a/Engine/Plugins/PCG/Source/PCG/Public/Elements/PCGCopyPointsAnalysisKernel.h
+++ b/Engine/Plugins/PCG/Source/PCG/Public/Elements/PCGCopyPointsAnalysisKernel.h
@@ -18,7 +18,7 @@ public:
 	virtual bool ComputeOutputBindingDataDesc(const UPCGComputeGraph* InGraph, FName InOutputPinLabel, UPCGDataBinding* InBinding, FPCGDataCollectionDesc& OutDataDesc) const override;
 	virtual int ComputeThreadCount(const UPCGDataBinding* Binding) const override;
 #if WITH_EDITOR
-	virtual FString GetCookedSource(FPCGGPUCompilationContext& InOutContext) const override;
+	virtual const TCHAR* GetSourceFilePath() const override { return TEXT("/Plugin/PCG/Private/Elements/PCGCopyPointsAnalysis.usf"); }
 	virtual FString GetEntryPoint() const override { return TEXT("PCGCopyPointsAnalysisCS"); }
 	void CreateAdditionalInputDataInterfaces(FPCGGPUCompilationContext& InOutContext, UObject* InObjectOuter, TArray<TObjectPtr<UComputeDataInterface>>& OutDataInterfaces) const override;
 	// Split graph to read back analysis results.
diff --git a/Engine/Plugins/PCG/Source/PCG/Public/Elements/PCGCopyPointsKernel.h b/Engine/Plugins/PCG/Source/PCG/Public/Elements/PCGCopyPointsKernel.h
--- a/Engine/Plugins/PCG/Source/PCG/Public/Elements/PCGCopyPointsKernel.h
+++ b/Engine/Plugins/PCG/Source/PCG/Public/Elements/PCGCopyPointsKernel.h
@@ -21,6 +21,7 @@ public:
 	virtual int ComputeThreadCount(const UPCGDataBinding* Binding) const override;
 #if WITH_EDITOR
 	virtual FString GetCookedSource(FPCGGPUCompilationContext& InOutContext) const override;
+	virtual const TCHAR* GetSourceFilePath() const override { return TEXT("/Plugin/PCG/Private/Elements/PCGCopyPoints.usf"); }
 	virtual FString GetEntryPoint() const override { return TEXT("PCGCopyPointsCS"); }
 	void CreateAdditionalInputDataInterfaces(FPCGGPUCompilationContext& InOutContext, UObject* InObjectOuter, TArray<TObjectPtr<UComputeDataInterface>>& OutDataInterfaces) const override;
 #endif
diff --git a/Engine/Plugins/PCG/Source/PCG/Public/Elements/PCGSkinnedMeshSpawnerKernel.h b/Engine/Plugins/PCG/Source/PCG/Public/Elements/PCGSkinnedMeshSpawnerKernel.h
--- a/Engine/Plugins/PCG/Source/PCG/Public/Elements/PCGSkinnedMeshSpawnerKernel.h
+++ b/Engine/Plugins/PCG/Source/PCG/Public/Elements/PCGSkinnedMeshSpawnerKernel.h
@@ -22,7 +22,7 @@ public:
 	virtual bool ComputeOutputBindingDataDesc(const UPCGComputeGraph* InGraph, FName InOutputPinLabel, UPCGDataBinding* InBinding, FPCGDataCollectionDesc& OutDataDesc) const override;
 	virtual int ComputeThreadCount(const UPCGDataBinding* Binding) const override;
 #if WITH_EDITOR
-	virtual FString GetCookedSource(FPCGGPUCompilationContext& InOutContext) const override;
+	virtual const TCHAR* GetSourceFilePath() const override { return TEXT("/Plugin/PCG/Private/Elements/PCGStaticMeshSpawner.usf"); }
 	virtual FString GetEntryPoint() const override { return TEXT("PCGStaticMeshSpawnerCS"); }
 	void CreateAdditionalInputDataInterfaces(FPCGGPUCompilationContext& InOutContext, UObject* InObjectOuter, TArray<TObjectPtr<UComputeDataInterface>>& OutDataInterfaces) const override;
 	void CreateAdditionalOutputDataInterfaces(FPCGGPUCompilationContext& InOutContext, UObject* InObjectOuter, TArray<TObjectPtr<UComputeDataInterface>>& OutDataInterfaces) const override;
diff --git a/Engine/Plugins/PCG/Source/PCG/Public/Elements/PCGStaticMeshSpawnerKernel.h b/Engine/Plugins/PCG/Source/PCG/Public/Elements/PCGStaticMeshSpawnerKernel.h
--- a/Engine/Plugins/PCG/Source/PCG/Public/Elements/PCGStaticMeshSpawnerKernel.h
+++ b/Engine/Plugins/PCG/Source/PCG/Public/Elements/PCGStaticMeshSpawnerKernel.h
@@ -22,7 +22,7 @@ public:
 	virtual bool ComputeOutputBindingDataDesc(const UPCGComputeGraph* InGraph, FName InOutputPinLabel, UPCGDataBinding* InBinding, FPCGDataCollectionDesc& OutDataDesc) const override;
 	virtual int ComputeThreadCount(const UPCGDataBinding* Binding) const override;
 #if WITH_EDITOR
-	virtual FString GetCookedSource(FPCGGPUCompilationContext& InOutContext) const override;
+	virtual const TCHAR* GetSourceFilePath() const override { return TEXT("/Plugin/PCG/Private/Elements/PCGStaticMeshSpawner.usf"); }
 	virtual FString GetEntryPoint() const override { return TEXT("PCGStaticMeshSpawnerCS"); }
 	void CreateAdditionalInputDataInterfaces(FPCGGPUCompilationContext& InOutContext, UObject* InObjectOuter, TArray<TObjectPtr<UComputeDataInterface>>& OutDataInterfaces) const override;
 	void CreateAdditionalOutputDataInterfaces(FPCGGPUCompilationContext& InOutContext, UObject* InObjectOuter, TArray<TObjectPtr<UComputeDataInterface>>& OutDataInterfaces) const override;
