[PCG] Fix build health issue in PCGComputeKernel

#rnx

diff --git a/Engine/Plugins/PCG/Source/PCG/Private/Compute/PCGComputeKernel.h b/Engine/Plugins/PCG/Source/PCG/Private/Compute/PCGComputeKernel.h
--- a/Engine/Plugins/PCG/Source/PCG/Private/Compute/PCGComputeKernel.h
+++ b/Engine/Plugins/PCG/Source/PCG/Private/Compute/PCGComputeKernel.h
@@ -104,7 +104,7 @@ public:
 	virtual FString GetCookedSource(FPCGGPUCompilationContext& InOutContext) const;
 
 	/** Get the virtual file path of the shader source file for this kernel. */
-	virtual const TCHAR* GetSourceFilePath() const PURE_VIRTUAL(UPCGComputeKernel::GetSourceFilePath, return nullptr;);
+	virtual const TCHAR* GetSourceFilePath() const { return nullptr; }
 
 	/** Get the name of the main kernel function in the source. This name is also displayed in GPU profile scopes. */
 	virtual FString GetEntryPoint() const PURE_VIRTUAL(UPCGComputeKernel::GetEntryPoint, return TEXT("Main"););
