[PCG] Cleared up the flags specifier as it was causing typing issue, and added combination flags

#rnx

diff --git a/Engine/Plugins/PCG/Source/PCG/Public/PCGCommon.h b/Engine/Plugins/PCG/Source/PCG/Public/PCGCommon.h
--- a/Engine/Plugins/PCG/Source/PCG/Public/PCGCommon.h
+++ b/Engine/Plugins/PCG/Source/PCG/Public/PCGCommon.h
@@ -90,7 +90,7 @@ ENUM_CLASS_FLAGS(EPCGChangeType);
 // This is why we have a parallel enum just below that must match on a name basis 1:1 to allow the make/break functions to work properly
 // in blueprint.
 // WARNING: Please be mindful that combination of flags that are not explicitly defined there won't be serialized correctly, inducing data loss.
-UENUM(Flags, meta = (Bitflags))
+UENUM(meta = (Bitflags))
 enum class EPCGDataType : uint32
 {
 	None = 0 UMETA(Hidden),
@@ -135,6 +135,9 @@ enum class EPCGDataType : uint32
 	// Combination of Param and Point, necessary for name-based serialization of enums.
 	PointOrParam = Point | Param UMETA(DisplayName = "Point or Attribute Set"),
 
+	VolumeOrPrimitiveOrDynamicMesh = Volume | Primitive | DynamicMesh UMETA(Hidden),
+	PointOrSpline = Point | Spline UMETA(Hidden),
+
 	Settings = 1 << 28 UMETA(Hidden),
 	Other = 1 << 29,
 	Any = (1 << 30) - 1
