[PCG] RuntimeGen: Change server global enable flag to a mode (0 = disabled, 1 = local gen sources, 2 = all gen sources)
- Add a IsLocal() virtual to UPCGGenSourceBase (Implement for UPCGGenSourcePlayer to check IsLocalController on the PlayerController)
- Add a NetworkMode (Client, Server, All) to the UPCGSchedulingPolicyDistanceAndDirection  in combination with UPCGGenSourceBase::IsLocal can filter out some gen sources per policy

diff --git a/Engine/Plugins/PCG/Source/PCG/Private/RuntimeGen/GenSources/PCGGenSourcePlayer.cpp b/Engine/Plugins/PCG/Source/PCG/Private/RuntimeGen/GenSources/PCGGenSourcePlayer.cpp
--- a/Engine/Plugins/PCG/Source/PCG/Private/RuntimeGen/GenSources/PCGGenSourcePlayer.cpp
+++ b/Engine/Plugins/PCG/Source/PCG/Private/RuntimeGen/GenSources/PCGGenSourcePlayer.cpp
@@ -60,3 +60,15 @@ void UPCGGenSourcePlayer::SetPlayerController(const APlayerController* InPlayerC
 {
 	PlayerController = InPlayerController;
 }
+
+bool UPCGGenSourcePlayer::IsLocal() const
+{
+	if (const APlayerController* PlayerControllerPtr = PlayerController.Get())
+	{
+		return PlayerControllerPtr->IsLocalController();
+	}
+	else
+	{
+		return false;
+	}
+}
\ No newline at end of file
diff --git a/Engine/Plugins/PCG/Source/PCG/Private/RuntimeGen/PCGGenSourceManager.cpp b/Engine/Plugins/PCG/Source/PCG/Private/RuntimeGen/PCGGenSourceManager.cpp
--- a/Engine/Plugins/PCG/Source/PCG/Private/RuntimeGen/PCGGenSourceManager.cpp
+++ b/Engine/Plugins/PCG/Source/PCG/Private/RuntimeGen/PCGGenSourceManager.cpp
@@ -22,15 +22,17 @@
 
 namespace PCGGenSourceManager
 {
-	static TAutoConsoleVariable<bool> CVarDedicatedServerRuntimeGenerationEnable(
-		TEXT("pcg.RuntimeGeneration.Server.Enable"),
-		true,
-		TEXT("Enable the RuntimeGeneration system on dedicated servers."));
-
-	static TAutoConsoleVariable<bool> CVarListenServerRuntimeGenerationEnable(
-		TEXT("pcg.RuntimeGeneration.ListenServer.Enable"),
-		true,
-		TEXT("Enable the RuntimeGeneration system on listen servers."));
+	enum class EPCGServerMode : int32
+	{
+		Disabled = 0,
+		Local,
+		All
+	};
+
+	static TAutoConsoleVariable<int32> CVarServerRuntimeGenerationMode(
+		TEXT("pcg.RuntimeGeneration.Server.Mode"),
+		2,
+		TEXT("Enable the RuntimeGeneration system on servers. (0 disabled, 1 local generation sources only (listen server), 2 all generation sources"));
 }
 
 FPCGGenSourceManager::FPCGGenSourceManager(const UWorld* InWorld)
@@ -155,13 +157,10 @@ void FPCGGenSourceManager::UpdatePerTickGenSources(const APCGWorldActor* InPCGWo
 	// Gather Player controller sources
 	TSet<APlayerController*> CurrentControllers;
 
-	const bool bIsDedicatedServer = PCGHelpers::IsDedicatedServer(World);
-	const bool bIsListenServer = PCGHelpers::IsListenServer(World);
-	
-	if (InPCGWorldActor->bTreatPlayerControllersAsGenerationSources &&
-		((!bIsDedicatedServer && !bIsListenServer) 
-		|| (bIsDedicatedServer && PCGGenSourceManager::CVarDedicatedServerRuntimeGenerationEnable.GetValueOnAnyThread())
-		|| (bIsListenServer && PCGGenSourceManager::CVarListenServerRuntimeGenerationEnable.GetValueOnAnyThread())))
+	const bool bIsServer = PCGHelpers::IsListenServer(World) || PCGHelpers::IsDedicatedServer(World);
+	const int32 ServerMode = PCGGenSourceManager::CVarServerRuntimeGenerationMode.GetValueOnAnyThread();
+
+	if (InPCGWorldActor->bTreatPlayerControllersAsGenerationSources && (!bIsServer || ServerMode > 0))
 	{
 		// Update existing sources or add new ones
 		for (auto It = World->GetPlayerControllerIterator(); It; ++It)
@@ -172,6 +171,11 @@ void FPCGGenSourceManager::UpdatePerTickGenSources(const APCGWorldActor* InPCGWo
 				continue;
 			}
 
+			if (bIsServer && ServerMode == static_cast<std::underlying_type_t<PCGGenSourceManager::EPCGServerMode>>(PCGGenSourceManager::EPCGServerMode::Local) && !PlayerController->IsLocalController())
+			{
+				continue;
+			}
+
 			CurrentControllers.Add(PlayerController);
 
 			if (TScriptInterface<IPCGGenSourceBase>* FoundSource = RegisteredNamedGenSources.Find(PlayerController->GetFName()))
diff --git a/Engine/Plugins/PCG/Source/PCG/Private/RuntimeGen/SchedulingPolicies/PCGSchedulingPolicyDistanceAndDirection.cpp b/Engine/Plugins/PCG/Source/PCG/Private/RuntimeGen/SchedulingPolicies/PCGSchedulingPolicyDistanceAndDirection.cpp
--- a/Engine/Plugins/PCG/Source/PCG/Private/RuntimeGen/SchedulingPolicies/PCGSchedulingPolicyDistanceAndDirection.cpp
+++ b/Engine/Plugins/PCG/Source/PCG/Private/RuntimeGen/SchedulingPolicies/PCGSchedulingPolicyDistanceAndDirection.cpp
@@ -58,6 +58,13 @@ double UPCGSchedulingPolicyDistanceAndDirection::CalculatePriority(const IPCGGen
 
 bool UPCGSchedulingPolicyDistanceAndDirection::ShouldGenerate(const IPCGGenSourceBase* InGenSource, const FBox& GenerationBounds, bool bUse2DGrid) const
 {
+	// Filter out Gen sources based on network mode
+	const bool bIsLocal = InGenSource->IsLocal();
+	if ((NetworkMode == EPCGSchedulingPolicyNetworkMode::Client && !bIsLocal) || (NetworkMode == EPCGSchedulingPolicyNetworkMode::Server && bIsLocal))
+	{
+		return false;
+	}
+
 	if (!bUseFrustumCulling)
 	{
 		return true;
@@ -81,6 +88,13 @@ bool UPCGSchedulingPolicyDistanceAndDirection::ShouldGenerate(const IPCGGenSourc
 
 bool UPCGSchedulingPolicyDistanceAndDirection::ShouldCull(const IPCGGenSourceBase* InGenSource, const FBox& GenerationBounds, bool bUse2DGrid) const
 {
+	// Filter out Gen sources based on network mode
+	const bool bIsLocal = InGenSource->IsLocal();
+	if ((NetworkMode == EPCGSchedulingPolicyNetworkMode::Client && !bIsLocal) || (NetworkMode == EPCGSchedulingPolicyNetworkMode::Server && bIsLocal))
+	{
+		return true;
+	}
+
 	if (!bUseFrustumCulling)
 	{
 		// If frustum culling is not enabled, don't try cleaning up volumes within the cleanup radius.
@@ -124,7 +138,8 @@ bool UPCGSchedulingPolicyDistanceAndDirection::IsEquivalent(const UPCGScheduling
 			&& DirectionWeight == Other->DirectionWeight
 			&& bUseFrustumCulling == Other->bUseFrustumCulling
 			&& GenerateBoundsModifier == Other->GenerateBoundsModifier
-			&& CleanupBoundsModifier == Other->CleanupBoundsModifier;
+			&& CleanupBoundsModifier == Other->CleanupBoundsModifier
+			&& NetworkMode == Other->NetworkMode;
 	}
 	else
 	{
diff --git a/Engine/Plugins/PCG/Source/PCG/Public/RuntimeGen/GenSources/PCGGenSourceBase.h b/Engine/Plugins/PCG/Source/PCG/Public/RuntimeGen/GenSources/PCGGenSourceBase.h
--- a/Engine/Plugins/PCG/Source/PCG/Public/RuntimeGen/GenSources/PCGGenSourceBase.h
+++ b/Engine/Plugins/PCG/Source/PCG/Public/RuntimeGen/GenSources/PCGGenSourceBase.h
@@ -35,6 +35,9 @@ public:
 
 	/** Returns the view frustum of this gen source. */
 	virtual TOptional<FConvexVolume> GetViewFrustum(bool bIs2DGrid) const { return TOptional<FConvexVolume>(); }
+
+	/** Returns true if this is a local gen source. (client-server specific) */
+	virtual bool IsLocal() const { return true; }
 };
 
 #undef UE_API
diff --git a/Engine/Plugins/PCG/Source/PCG/Public/RuntimeGen/GenSources/PCGGenSourcePlayer.h b/Engine/Plugins/PCG/Source/PCG/Public/RuntimeGen/GenSources/PCGGenSourcePlayer.h
--- a/Engine/Plugins/PCG/Source/PCG/Public/RuntimeGen/GenSources/PCGGenSourcePlayer.h
+++ b/Engine/Plugins/PCG/Source/PCG/Public/RuntimeGen/GenSources/PCGGenSourcePlayer.h
@@ -31,6 +31,9 @@ public:
 	/** Returns the view frustum of this gen source. */
 	virtual TOptional<FConvexVolume> GetViewFrustum(bool bIs2DGrid) const override { return ViewFrustum; }
 
+	/** Returns true if this is a local gen source. (network) */
+	virtual bool IsLocal() const override;
+
 	TWeakObjectPtr<const APlayerController> GetPlayerController() const { return PlayerController; }
 	UE_API void SetPlayerController(const APlayerController* InPlayerController);
 
diff --git a/Engine/Plugins/PCG/Source/PCG/Public/RuntimeGen/SchedulingPolicies/PCGSchedulingPolicyDistanceAndDirection.h b/Engine/Plugins/PCG/Source/PCG/Public/RuntimeGen/SchedulingPolicies/PCGSchedulingPolicyDistanceAndDirection.h
--- a/Engine/Plugins/PCG/Source/PCG/Public/RuntimeGen/SchedulingPolicies/PCGSchedulingPolicyDistanceAndDirection.h
+++ b/Engine/Plugins/PCG/Source/PCG/Public/RuntimeGen/SchedulingPolicies/PCGSchedulingPolicyDistanceAndDirection.h
@@ -8,6 +8,14 @@
 
 class IPCGGenSourceBase;
 
+UENUM()
+enum class EPCGSchedulingPolicyNetworkMode : uint8
+{
+	Client UMETA(ToolTip = "Considers generation sources that return true to IsLocal(). Can be used for standalone or listen server clients."),
+	Server UMETA(ToolTip = "Considers genration sources that return false to IsLocal(). Can be used for dedicated server only generation."),
+	All UMETA(ToolTip = "Considers all generation sources.")
+};
+
 /**
  * SchedulingPolicyDistanceAndDirection uses distance from the generating volume 
  * and alignment with view direction to choose the most important volumes to generate.
@@ -56,4 +64,8 @@ public:
 	/** Multiplier to scale bounds by when comparing against the view frustum for clean up. */
 	UPROPERTY(BlueprintReadWrite, EditAnywhere, Category = "RuntimeGeneration|Scheduling Policy Parameters", meta = (ClampMin = 1.0f, EditCondition = "bShouldDisplayProperties && bUseFrustumCulling", EditConditionHides))
 	float CleanupBoundsModifier = 1.2f;
+
+	/** Client/Server policy mode will determine which generation sources to consider. */
+	UPROPERTY(BlueprintReadWrite, EditAnywhere, Category = "RuntimeGeneration|Scheduling Policy Parameters")
+	EPCGSchedulingPolicyNetworkMode NetworkMode = EPCGSchedulingPolicyNetworkMode::All;
 };
