[PCG] Add IPCGGraphExecutionState::GetDynamicTrackingPriority()
- implemented by the FPCGComponentExecutionState for now and uses a mix of default priority (object path hash) (lease significant bytes) + optional manual priority (most significant bytes)
- allows breaking key tracking feedback loops

diff --git a/Engine/Plugins/PCG/Source/PCG/Private/PCGActorAndComponentMapping.cpp b/Engine/Plugins/PCG/Source/PCG/Private/PCGActorAndComponentMapping.cpp
--- a/Engine/Plugins/PCG/Source/PCG/Private/PCGActorAndComponentMapping.cpp
+++ b/Engine/Plugins/PCG/Source/PCG/Private/PCGActorAndComponentMapping.cpp
@@ -1893,18 +1893,36 @@ void FPCGActorAndComponentMapping::OnSelectionKeyChanged(const FPCGSelectionKey&
 	TArray<const UObject*, TInlineAllocator<2>> ChangedObjects = { ChangedObject };
 	ChangedObjects.AddUnique(InOriginatingChangeObject);
 
+	TOptional<FPCGDynamicTrackingPriority> ChangedObjectDynamicTrackingPriority;
+	// @todo_pcg: Find a better way to get this info as this doesn't work for managed objects that were spawned under a different outer
+	if (AActor* OuterActor = InOriginatingChangeObject ? InOriginatingChangeObject->GetTypedOuter<AActor>() : nullptr)
+	{
+		TArray<UPCGComponent*> OriginatingComponents;
+		OuterActor->GetComponents<UPCGComponent>(OriginatingComponents);
+		for (UPCGComponent* OriginatingComponent : OriginatingComponents)
+		{
+			if (OriginatingComponent->IsAnyObjectManagedByResource(ChangedObjects))
+			{
+				ChangedObjectDynamicTrackingPriority = OriginatingComponent->GetExecutionState().GetDynamicTrackingPriority();
+				break;
+			}
+		}
+	}
+
 	// We discard the component if:
 	//   * It is null.
 	//   * It is one of the objects that changed.
 	//   * We should not refresh the owner and its owner is one of the objects that changed.
-	//   * It is generating and one of the objects that changed is managed by it (safe-guard, would probably induce an infinite loop).
+	//   * The objects that changed are managed by it (safe-guard, would probably induce an infinite loop).
+	//   * The objects that changed are managed by a lower tracking priority component.
 	//   * One of the objects that changed is ignored by its original component.
-	auto ShouldDiscardComponent = [&ChangedObjects, bInNoRefreshOwner](const UPCGComponent* InComponent)
+	auto ShouldDiscardComponent = [&ChangedObjects, &ChangedObjectDynamicTrackingPriority, bInNoRefreshOwner](const UPCGComponent* InComponent)
 	{
 		if (!InComponent
 			|| ChangedObjects.Contains(InComponent)
 			|| (bInNoRefreshOwner && ChangedObjects.Contains(InComponent->GetOwner()))
-			|| (InComponent->IsGenerating() && InComponent->IsAnyObjectManagedByResource(ChangedObjects)))
+			|| InComponent->IsAnyObjectManagedByResource(ChangedObjects)
+			|| (ChangedObjectDynamicTrackingPriority.IsSet() && InComponent->GetExecutionState().GetDynamicTrackingPriority() > ChangedObjectDynamicTrackingPriority.GetValue()))
 		{
 			return true;
 		}
diff --git a/Engine/Plugins/PCG/Source/PCG/Private/PCGComponent.cpp b/Engine/Plugins/PCG/Source/PCG/Private/PCGComponent.cpp
--- a/Engine/Plugins/PCG/Source/PCG/Private/PCGComponent.cpp
+++ b/Engine/Plugins/PCG/Source/PCG/Private/PCGComponent.cpp
@@ -4306,6 +4306,15 @@ bool FPCGComponentExecutionState::IsRefreshInProgress() const
 	check(Component);
 	return Component->IsRefreshInProgress();
 }
+
+FPCGDynamicTrackingPriority FPCGComponentExecutionState::GetDynamicTrackingPriority() const
+{
+	check(Component);
+
+	// Give user control over priority by making UPCGComponent::Priority the most significant part of the resulting priority
+	return static_cast<uint64>((static_cast<uint64>(Component->TrackingPriority) << 32) | GetTypeHash(Component->GetPathName()));
+}
+
 #endif // WITH_EDITOR
 
 #undef LOCTEXT_NAMESPACE
diff --git a/Engine/Plugins/PCG/Source/PCG/Private/PCGGraphExecutionStateInterface.cpp b/Engine/Plugins/PCG/Source/PCG/Private/PCGGraphExecutionStateInterface.cpp
--- a/Engine/Plugins/PCG/Source/PCG/Private/PCGGraphExecutionStateInterface.cpp
+++ b/Engine/Plugins/PCG/Source/PCG/Private/PCGGraphExecutionStateInterface.cpp
@@ -25,6 +25,16 @@ IPCGBaseSubsystem* IPCGGraphExecutionState::GetSubsystem() const
 	return Subsystem;
 }
 
+#if WITH_EDITOR
+
+FPCGDynamicTrackingPriority IPCGGraphExecutionState::GetDynamicTrackingPriority() const
+{
+	ensureMsgf(false, TEXT("Please implement IPCGGraphExecutionState::GetDynamicTrackingPriority"));
+	return FPCGDynamicTrackingPriority();
+}
+
+#endif
+
 //////////////////////////////////////////////////////////////////////
 
 UPCGGraphExecutionSource::UPCGGraphExecutionSource(const FObjectInitializer& ObjectInitializer)
diff --git a/Engine/Plugins/PCG/Source/PCG/Public/Editor/PCGEditorDefaultExecutionSource.h b/Engine/Plugins/PCG/Source/PCG/Public/Editor/PCGEditorDefaultExecutionSource.h
--- a/Engine/Plugins/PCG/Source/PCG/Public/Editor/PCGEditorDefaultExecutionSource.h
+++ b/Engine/Plugins/PCG/Source/PCG/Public/Editor/PCGEditorDefaultExecutionSource.h
@@ -1,4 +1,4 @@
-﻿// Copyright Epic Games, Inc. All Rights Reserved.
+// Copyright Epic Games, Inc. All Rights Reserved.
 
 #pragma once
 
@@ -42,6 +42,8 @@ public:
 	virtual void RegisterDynamicTracking(const FPCGSelectionKeyToSettingsMap& InKeysToSettings) override {};
 	
 	virtual bool IsRefreshInProgress() const override {return IsGenerating();}
+
+	virtual FPCGDynamicTrackingPriority GetDynamicTrackingPriority() const override { return FPCGDynamicTrackingPriority(); }
 #endif // WITH_EDITOR
 
 private:
diff --git a/Engine/Plugins/PCG/Source/PCG/Public/PCGComponent.h b/Engine/Plugins/PCG/Source/PCG/Public/PCGComponent.h
--- a/Engine/Plugins/PCG/Source/PCG/Public/PCGComponent.h
+++ b/Engine/Plugins/PCG/Source/PCG/Public/PCGComponent.h
@@ -120,6 +120,7 @@ public:
 	virtual void RegisterDynamicTracking(const FPCGSelectionKeyToSettingsMap& InKeysToSettings) override;
 	
 	virtual bool IsRefreshInProgress() const override;
+	virtual FPCGDynamicTrackingPriority GetDynamicTrackingPriority() const override;
 #endif // WITH_EDITOR
 
 private:
@@ -356,6 +357,10 @@ public:
 	UPROPERTY(BlueprintReadWrite, EditAnywhere, Category = "Editing Settings", meta = (DisplayPriority = 451))
 	bool bIgnoreLandscapeTracking = false;
 
+	/** Tracking priority used to solve tracking dependencies between PCG Components. By default object path is used to create a default priority (order) but overriding this value gives control to the user. */
+	UPROPERTY(BlueprintReadWrite, EditAnywhere, Category= "Editing Settings", meta = (DisplayPriority = 452, ClampMin=0, UIMin=0))
+	int32 TrackingPriority = 0; 
+
 	UPROPERTY(BlueprintReadOnly, VisibleAnywhere, Transient, Category = Debug, meta = (NoResetToDefault))
 	bool bDirtyGenerated = false;
 
diff --git a/Engine/Plugins/PCG/Source/PCG/Public/PCGGraphExecutionStateInterface.h b/Engine/Plugins/PCG/Source/PCG/Public/PCGGraphExecutionStateInterface.h
--- a/Engine/Plugins/PCG/Source/PCG/Public/PCGGraphExecutionStateInterface.h
+++ b/Engine/Plugins/PCG/Source/PCG/Public/PCGGraphExecutionStateInterface.h
@@ -23,6 +23,8 @@ namespace PCGUtils
 	class FExtraCapture;
 }
 
+using FPCGDynamicTrackingPriority = uint64;
+
 /**
 * Interface returned by a IPCGGraphExecutionSource that is queried / updated during execution of a PCG Graph.
 */
@@ -82,6 +84,9 @@ public:
 
 	virtual FPCGGraphExecutionInspection& GetInspection() = 0;
 
+	/** Get an execution priority */
+	PCG_API virtual FPCGDynamicTrackingPriority GetDynamicTrackingPriority() const;
+
 	/** Register tracking dependencies, so ExecutionState can be updated when they change */
 	virtual void RegisterDynamicTracking(const UPCGSettings* InSettings, const TArrayView<TPair<FPCGSelectionKey, bool>>& InDynamicKeysAndCulling) = 0;
 
