[PCG] Fix missing pragma once in PCGSurfaceFromPolygon2D.h

#rnx

diff --git a/Engine/Plugins/PCG/Source/PCG/Public/Elements/Polygon/PCGSurfaceFromPolygon2D.h b/Engine/Plugins/PCG/Source/PCG/Public/Elements/Polygon/PCGSurfaceFromPolygon2D.h
--- a/Engine/Plugins/PCG/Source/PCG/Public/Elements/Polygon/PCGSurfaceFromPolygon2D.h
+++ b/Engine/Plugins/PCG/Source/PCG/Public/Elements/Polygon/PCGSurfaceFromPolygon2D.h
@@ -1,5 +1,7 @@
 // Copyright Epic Games, Inc. All Rights Reserved.
 
+#pragma once
+
 #include "PCGSettings.h"
 #include "Elements/Polygon/PCGPolygon2DUtils.h"
 
