[PCG] Fix compact node bug with two or more pins

diff --git a/Engine/Plugins/PCG/Source/PCGEditor/Private/Nodes/SPCGEditorGraphNodeCompact.cpp b/Engine/Plugins/PCG/Source/PCGEditor/Private/Nodes/SPCGEditorGraphNodeCompact.cpp
--- a/Engine/Plugins/PCG/Source/PCGEditor/Private/Nodes/SPCGEditorGraphNodeCompact.cpp
+++ b/Engine/Plugins/PCG/Source/PCGEditor/Private/Nodes/SPCGEditorGraphNodeCompact.cpp
@@ -219,10 +219,13 @@ void SPCGEditorGraphNodeCompact::UpdateGraphNode()
 	const TSharedRef<SVerticalBox> InnerVerticalBox =
 		SNew(SVerticalBox)
 		+ SVerticalBox::Slot()
-		.MinHeight(Constants::CompactNodeSize)
-		.MaxHeight(Constants::CompactNodeSize)
+		.AutoHeight()
 		[
-			NodeContentOverlay
+			SNew(SBox)
+			.MinDesiredHeight(Constants::CompactNodeSize)
+			[
+				NodeContentOverlay
+			]
 		];
 
 	// Enabled state bar
