[PCG] Improvements on Boolean and Bitwise op

* Boolean op now supports arithmetic types as input, still produce a bool as output
* Add operator NAND, NOR, XNOR, IMPLY, NIMPLY
* Add ShiftLeft/Right for Bitwise
* Unit test for booleans

#rnx

diff --git a/Engine/Plugins/PCG/Source/PCG/Private/Elements/Metadata/PCGMetadataBitwiseOpElement.cpp b/Engine/Plugins/PCG/Source/PCG/Private/Elements/Metadata/PCGMetadataBitwiseOpElement.cpp
--- a/Engine/Plugins/PCG/Source/PCG/Private/Elements/Metadata/PCGMetadataBitwiseOpElement.cpp
+++ b/Engine/Plugins/PCG/Source/PCG/Private/Elements/Metadata/PCGMetadataBitwiseOpElement.cpp
@@ -24,6 +24,10 @@ namespace PCGMetadataBitwiseSettings
 			return (Value1 | Value2);
 		case EPCGMetadataBitwiseOperation::Xor:
 			return (Value1 ^ Value2);
+		case EPCGMetadataBitwiseOperation::ShiftLeft:
+			return Value1 << Value2;
+		case EPCGMetadataBitwiseOperation::ShiftRight:
+			return Value1 >> Value2;
 		default:
 			return 0;
 		}
@@ -51,12 +55,15 @@ void UPCGMetadataBitwiseSettings::PostLoad()
 
 FName UPCGMetadataBitwiseSettings::GetInputPinLabel(uint32 Index) const
 {
+	const bool bIsShiftOp = Operation == EPCGMetadataBitwiseOperation::ShiftLeft ||
+			Operation == EPCGMetadataBitwiseOperation::ShiftRight;
+	
 	switch (Index)
 	{
 	case 0:
-		return (Operation == EPCGMetadataBitwiseOperation::Not) ? PCGPinConstants::DefaultInputLabel : PCGMetadataSettingsBaseConstants::DoubleInputFirstLabel;
+		return (Operation == EPCGMetadataBitwiseOperation::Not || bIsShiftOp) ? PCGPinConstants::DefaultInputLabel : PCGMetadataSettingsBaseConstants::DoubleInputFirstLabel;
 	case 1:
-		return PCGMetadataSettingsBaseConstants::DoubleInputSecondLabel;
+		return bIsShiftOp ? PCGMetadataBitwiseConstants::ShiftInputLabel : PCGMetadataSettingsBaseConstants::DoubleInputSecondLabel;
 	default:
 		return NAME_None;
 	}
diff --git a/Engine/Plugins/PCG/Source/PCG/Private/Elements/Metadata/PCGMetadataBooleanOpElement.cpp b/Engine/Plugins/PCG/Source/PCG/Private/Elements/Metadata/PCGMetadataBooleanOpElement.cpp
--- a/Engine/Plugins/PCG/Source/PCG/Private/Elements/Metadata/PCGMetadataBooleanOpElement.cpp
+++ b/Engine/Plugins/PCG/Source/PCG/Private/Elements/Metadata/PCGMetadataBooleanOpElement.cpp
@@ -24,6 +24,16 @@ namespace PCGMetadataBooleanSettings
 			return (Value1 || Value2);
 		case EPCGMetadataBooleanOperation::Xor:
 			return (Value1 != Value2);
+		case EPCGMetadataBooleanOperation::Nand:
+			return !(Value1 && Value2);
+		case EPCGMetadataBooleanOperation::Nor:
+			return !(Value1 || Value2);
+		case EPCGMetadataBooleanOperation::Xnor:
+			return (Value1 == Value2);
+		case EPCGMetadataBooleanOperation::Imply:
+			return !Value1 || Value2;
+		case EPCGMetadataBooleanOperation::Nimply:
+			return Value1 && !Value2;
 		default:
 			return false;
 		}
@@ -70,7 +80,7 @@ uint32 UPCGMetadataBooleanSettings::GetOperandNum() const
 bool UPCGMetadataBooleanSettings::IsSupportedInputType(uint16 TypeId, uint32 InputIndex, bool& bHasSpecialRequirement) const
 {
 	bHasSpecialRequirement = false;
-	return PCG::Private::IsOfTypes<bool>(TypeId);
+	return PCG::Private::IsBroadcastableOrConstructible(TypeId, PCG::Private::MetadataTypes<bool>::Id);
 }
 
 FPCGAttributePropertyInputSelector UPCGMetadataBooleanSettings::GetInputSource(uint32 Index) const
diff --git a/Engine/Plugins/PCG/Source/PCG/Private/Tests/Elements/Metadata/PCGMetadataBooleanOpTest.cpp b/Engine/Plugins/PCG/Source/PCG/Private/Tests/Elements/Metadata/PCGMetadataBooleanOpTest.cpp
--- a/Engine/Plugins/PCG/Source/PCG/Private/Tests/Elements/Metadata/PCGMetadataBooleanOpTest.cpp
+++ b/Engine/Plugins/PCG/Source/PCG/Private/Tests/Elements/Metadata/PCGMetadataBooleanOpTest.cpp
@@ -10,51 +10,73 @@
 
 #if WITH_EDITOR
 
-IMPLEMENT_CUSTOM_SIMPLE_AUTOMATION_TEST(FPCGMetadataBooleanOpTest, FPCGTestBaseClass, "Plugins.PCG.Metadata.BooleanOp", PCGTestsCommon::TestFlags)
-
-bool FPCGMetadataBooleanOpTest::RunTest(const FString& Parameters)
+class FPCGMetadataBooleanOpTest : public FPCGTestBaseClass
 {
-	PCGTestsCommon::FTestData TestData;
-	PCGTestsCommon::GenerateSettings<UPCGMetadataBooleanSettings>(TestData);
-	UPCGMetadataBooleanSettings* Settings = CastChecked<UPCGMetadataBooleanSettings>(TestData.Settings);
-	FPCGElementPtr MetadataBooleanElement = TestData.Settings->GetElement();
-
-	FPCGTaggedData& ParamTaggedData1 = TestData.InputData.TaggedData.Emplace_GetRef(FPCGTaggedData());
-	TObjectPtr<UPCGParamData> ParamData1 = PCGTestsCommon::CreateEmptyParamData();
-	ParamTaggedData1.Data = ParamData1;
-
-	FPCGTaggedData& ParamTaggedData2 = TestData.InputData.TaggedData.Emplace_GetRef(FPCGTaggedData());
-	TObjectPtr<UPCGParamData> ParamData2 = PCGTestsCommon::CreateEmptyParamData();
-	ParamTaggedData2.Data = ParamData2;
-
+public:
 	const FName TrueAttribute = TEXT("True");
 	const FName FalseAttribute = TEXT("False");
+	const FName Int0Attribute = TEXT("Int0");
+	const FName Int1Attribute = TEXT("Int1");
 	const FName InvalidAttribute = TEXT("Invalid");
 	const FName OutputAttributeName = TEXT("Output");
 
-	Settings->OutputTarget.SetAttributeName(OutputAttributeName);
-	Settings->ForceOutputConnections[0] = true;
+	using FPCGTestBaseClass::FPCGTestBaseClass;
 
-	const bool bAllowInterpolation = false;
-	const bool bOverrideParent = false;
+protected:
+	void Init()
+	{
+		ParamData1 = PCGTestsCommon::CreateEmptyParamData();
+		ParamData2 = PCGTestsCommon::CreateEmptyParamData();
+
+		constexpr bool bAllowInterpolation = false;
+		constexpr bool bOverrideParent = false;
+
+		ParamData1->Metadata->CreateAttribute<bool>(TrueAttribute, /*DefaultValue=*/ true, bAllowInterpolation, bOverrideParent);
+		ParamData1->Metadata->CreateAttribute<bool>(FalseAttribute, /*DefaultValue=*/ false, bAllowInterpolation, bOverrideParent);
+		ParamData1->Metadata->CreateAttribute<int64>(Int0Attribute, /*DefaultValue=*/ 0, bAllowInterpolation, bOverrideParent);
+		ParamData1->Metadata->CreateAttribute<int64>(Int1Attribute, /*DefaultValue=*/ 1, bAllowInterpolation, bOverrideParent);
+		ParamData1->Metadata->CreateAttribute<FVector>(InvalidAttribute, /*DefaultValue=*/ FVector::ZeroVector, bAllowInterpolation, bOverrideParent);
+		ParamData1->Metadata->AddEntry();
+
+		ParamData2->Metadata->CreateAttribute<bool>(TrueAttribute, /*DefaultValue=*/ true, bAllowInterpolation, bOverrideParent);
+		ParamData2->Metadata->CreateAttribute<bool>(FalseAttribute, /*DefaultValue=*/ false, bAllowInterpolation, bOverrideParent);
+		ParamData2->Metadata->CreateAttribute<int64>(Int0Attribute, /*DefaultValue=*/ 0, bAllowInterpolation, bOverrideParent);
+		ParamData2->Metadata->CreateAttribute<int64>(Int1Attribute, /*DefaultValue=*/ 1, bAllowInterpolation, bOverrideParent);
+		ParamData2->Metadata->CreateAttribute<FVector>(InvalidAttribute, /*DefaultValue=*/ FVector::ZeroVector, bAllowInterpolation, bOverrideParent);
+		ParamData2->Metadata->AddEntry();
+
+		Settings = NewObject<UPCGMetadataBooleanSettings>();
+	}
+	
+	bool RunInternal(const EPCGMetadataBooleanOperation Op, const FName Attr1, const FName Attr2, const bool bExpectedResult, bool bIsValid = true)
+	{
+		if (!ParamData1)
+		{
+			Init();
+		}
 
-	ParamData1->Metadata->CreateAttribute<bool>(TrueAttribute, /*DefaultValue=*/ true, bAllowInterpolation, bOverrideParent);
-	ParamData1->Metadata->CreateAttribute<bool>(FalseAttribute, /*DefaultValue=*/ false, bAllowInterpolation, bOverrideParent);
-	ParamData1->Metadata->CreateAttribute<int64>(InvalidAttribute, /*DefaultValue=*/ 0, bAllowInterpolation, bOverrideParent);
-	ParamData1->Metadata->AddEntry();
+		PCGTestsCommon::FTestData TestData;
+		TestData.Reset(Settings);
+		
+		FPCGElementPtr MetadataBooleanElement = TestData.Settings->GetElement();
 
-	ParamData2->Metadata->CreateAttribute<bool>(TrueAttribute, /*DefaultValue=*/ true, bAllowInterpolation, bOverrideParent);
-	ParamData2->Metadata->CreateAttribute<bool>(FalseAttribute, /*DefaultValue=*/ false, bAllowInterpolation, bOverrideParent);
-	ParamData2->Metadata->CreateAttribute<int64>(InvalidAttribute, /*DefaultValue=*/ 0, bAllowInterpolation, bOverrideParent);
-	ParamData2->Metadata->AddEntry();
+		FPCGTaggedData& ParamTaggedData1 = TestData.InputData.TaggedData.Emplace_GetRef();
+		ParamTaggedData1.Data = ParamData1;
+		ParamTaggedData1.Pin = Op == EPCGMetadataBooleanOperation::Not ? PCGPinConstants::DefaultInputLabel : PCGMetadataSettingsBaseConstants::DoubleInputFirstLabel;
+
+		FPCGTaggedData& ParamTaggedData2 = TestData.InputData.TaggedData.Emplace_GetRef();
+		ParamTaggedData2.Data = ParamData2;
+		ParamTaggedData2.Pin = PCGMetadataSettingsBaseConstants::DoubleInputSecondLabel;
+
+		Settings->Operation = Op;
+		Settings->InputSource1.SetAttributeName(Attr1);
+		Settings->InputSource2.SetAttributeName(Attr2);
+		Settings->OutputTarget.SetAttributeName(OutputAttributeName);
+		Settings->ForceOutputConnections[0] = true;
 
-	auto ValidateOp = [this, &MetadataBooleanElement, &TestData, &Settings](bool bExpectedResult, bool bIsValid = true) -> bool
-	{
 		TUniquePtr<FPCGContext> Context = TestData.InitializeTestContext();
 
-		while (!MetadataBooleanElement->Execute(Context.Get()))
-		{
-		}
+		while (!MetadataBooleanElement->Execute(Context.Get())){}
 
 		const TArray<FPCGTaggedData> Outputs = Context->OutputData.GetInputsByPin(PCGPinConstants::DefaultOutputLabel);
 
@@ -64,184 +86,402 @@ bool FPCGMetadataBooleanOpTest::RunTest(const FString& Parameters)
 			return TestTrue(TEXT("Invalid op == no output"), Outputs.IsEmpty());
 		}
 
-		if (!TestEqual(TEXT("Number of outputs"), Outputs.Num(), 1))
-		{
-			return false;
-		}
+		UTEST_EQUAL("Number of outputs", Outputs.Num(), 1)
 
 		const UPCGData* OutputData = Outputs[0].Data;
+		UTEST_TRUE("Output data is param data", OutputData->IsA<UPCGParamData>())
+		
+		const UPCGMetadata* OutMetadata = OutputData->ConstMetadata();
+		UTEST_NOT_NULL("Output metadata exists", OutMetadata)
+
+		const FPCGMetadataAttribute<bool>* OutAttribute = OutMetadata->GetConstTypedAttribute<bool>(Settings->OutputTarget.GetName());
+		UTEST_NOT_NULL("Output attribute exists", OutAttribute);
+		UTEST_EQUAL("Output attribute has an entry", OutAttribute->GetNumberOfEntries(), 1);
+
+		const bool Out = OutAttribute->GetValueFromItemKey(PCGInvalidEntryKey);
+		UTEST_EQUAL("Op has the expected result", Out, bExpectedResult);
+		
+		return true;
+	}
 
-		if (!TestTrue(TEXT("Output data is param data"), OutputData->IsA<UPCGParamData>()))
+	bool Run(const EPCGMetadataBooleanOperation Op, const FName Attr1, const FName Attr2, const bool bExpectedResult, bool bIsValid = true)
+	{
+		if (!RunInternal(Op, Attr1, Attr2, bExpectedResult, bIsValid))
 		{
+			const UEnum* Enum = StaticEnum<EPCGMetadataBooleanOperation>();
+			check(Enum);
+			const FName OpName(Enum->GetNameStringByValue(static_cast<int64>(Op)));
+			
+			UE_LOG(LogPCG, Error, TEXT("%s %s %s == %s failed"), *Attr1.ToString(), *OpName.ToString(), *Attr2.ToString(), bExpectedResult ? TEXT("true") : TEXT("false"));
 			return false;
 		}
-
-		const UPCGMetadata* OutMetadata = CastChecked<UPCGParamData>(OutputData)->Metadata;
-
-		if (!TestNotNull(TEXT("Output metadata exists"), OutMetadata))
+		else
 		{
-			return false;
+			return true;
 		}
+	}
+	
+	UPCGMetadataBooleanSettings* Settings = nullptr;
+	UPCGParamData* ParamData1 = nullptr;
+	UPCGParamData* ParamData2 = nullptr;
+};
+
+IMPLEMENT_CUSTOM_SIMPLE_AUTOMATION_TEST(FPCGMetadataBooleanNotTest, FPCGMetadataBooleanOpTest, "Plugins.PCG.Metadata.BooleanOp.Not", PCGTestsCommon::TestFlags)
+IMPLEMENT_CUSTOM_SIMPLE_AUTOMATION_TEST(FPCGMetadataBooleanAndTest, FPCGMetadataBooleanOpTest, "Plugins.PCG.Metadata.BooleanOp.And", PCGTestsCommon::TestFlags)
+IMPLEMENT_CUSTOM_SIMPLE_AUTOMATION_TEST(FPCGMetadataBooleanOrTest, FPCGMetadataBooleanOpTest, "Plugins.PCG.Metadata.BooleanOp.Or", PCGTestsCommon::TestFlags)
+IMPLEMENT_CUSTOM_SIMPLE_AUTOMATION_TEST(FPCGMetadataBooleanXorTest, FPCGMetadataBooleanOpTest, "Plugins.PCG.Metadata.BooleanOp.Xor", PCGTestsCommon::TestFlags)
+IMPLEMENT_CUSTOM_SIMPLE_AUTOMATION_TEST(FPCGMetadataBooleanNandTest, FPCGMetadataBooleanOpTest, "Plugins.PCG.Metadata.BooleanOp.Nand", PCGTestsCommon::TestFlags)
+IMPLEMENT_CUSTOM_SIMPLE_AUTOMATION_TEST(FPCGMetadataBooleanNorTest, FPCGMetadataBooleanOpTest, "Plugins.PCG.Metadata.BooleanOp.Nor", PCGTestsCommon::TestFlags)
+IMPLEMENT_CUSTOM_SIMPLE_AUTOMATION_TEST(FPCGMetadataBooleanXnorTest, FPCGMetadataBooleanOpTest, "Plugins.PCG.Metadata.BooleanOp.Xnor", PCGTestsCommon::TestFlags)
+IMPLEMENT_CUSTOM_SIMPLE_AUTOMATION_TEST(FPCGMetadataBooleanImplyTest, FPCGMetadataBooleanOpTest, "Plugins.PCG.Metadata.BooleanOp.Imply", PCGTestsCommon::TestFlags)
+IMPLEMENT_CUSTOM_SIMPLE_AUTOMATION_TEST(FPCGMetadataBooleanNimplyTest, FPCGMetadataBooleanOpTest, "Plugins.PCG.Metadata.BooleanOp.Nimply", PCGTestsCommon::TestFlags)
+IMPLEMENT_CUSTOM_SIMPLE_AUTOMATION_TEST(FPCGMetadataBooleanInvalidTest, FPCGMetadataBooleanOpTest, "Plugins.PCG.Metadata.BooleanOp.Invalid", PCGTestsCommon::TestFlags)
+
+bool FPCGMetadataBooleanNotTest::RunTest(const FString& Parameters)
+{
+	// Attr1, bExpectedResult
+	TArray<TTuple<FName, bool>> TestValues =
+	{
+		{TrueAttribute, false},
+		{FalseAttribute, true},
+		{Int0Attribute, true},
+		{Int1Attribute, false}
+	};
 
-		const FPCGMetadataAttributeBase* OutAttributeBase = OutMetadata->GetConstAttribute(Settings->OutputTarget.GetName());
-		if (!TestNotNull(TEXT("Output attribute exists"), OutAttributeBase))
+	for (const auto [Attr1, bExpectedResult] : TestValues)
+	{
+		if (!Run(EPCGMetadataBooleanOperation::Not, Attr1, NAME_None, bExpectedResult))
 		{
 			return false;
 		}
+	}
 
-		if (!TestEqual(TEXT("Output attribute is a bool"), OutAttributeBase->GetTypeId(), (uint16)PCG::Private::MetadataTypes<bool>::Id))
-		{
-			return false;
-		}
+	return true;
+}
 
-		if (!TestEqual(TEXT("Output attribute has an entry"), OutAttributeBase->GetNumberOfEntries(), 1))
+bool FPCGMetadataBooleanOrTest::RunTest(const FString& Parameters)
+{	
+	// Attr1, Attr2, bExpectedResult
+	TArray<TTuple<FName, FName, bool>> TestValues =
+	{
+		{TrueAttribute, FalseAttribute, true},
+		{TrueAttribute, TrueAttribute, true},
+		{TrueAttribute, Int0Attribute, true},
+		{TrueAttribute, Int1Attribute, true},
+		
+		{FalseAttribute, FalseAttribute, false},
+		{FalseAttribute, TrueAttribute, true},
+		{FalseAttribute, Int0Attribute, false},
+		{FalseAttribute, Int1Attribute, true},
+
+		{Int0Attribute, FalseAttribute, false},
+		{Int0Attribute, TrueAttribute, true},
+		{Int0Attribute, Int0Attribute, false},
+		{Int0Attribute, Int1Attribute, true},
+
+		{Int1Attribute, FalseAttribute, true},
+		{Int1Attribute, TrueAttribute, true},
+		{Int1Attribute, Int0Attribute, true},
+		{Int1Attribute, Int1Attribute, true},
+	};
+
+	for (const auto [Attr1, Attr2, bExpectedResult] : TestValues)
+	{
+		if (!Run(EPCGMetadataBooleanOperation::Or, Attr1, Attr2, bExpectedResult))
 		{
 			return false;
 		}
+	}
 
-		const FPCGMetadataAttribute<bool>* OutAttribute = static_cast<const FPCGMetadataAttribute<bool>*>(OutAttributeBase);
+	return true;
+}
 
-		bool Out = OutAttribute->GetValueFromItemKey(PCGInvalidEntryKey);
+bool FPCGMetadataBooleanAndTest::RunTest(const FString& Parameters)
+{
+	// Attr1, Attr2, bExpectedResult
+	TArray<TTuple<FName, FName, bool>> TestValues
+	{
+		{TrueAttribute, FalseAttribute, false},
+		{TrueAttribute, TrueAttribute, true},
+		{TrueAttribute, Int0Attribute, false},
+		{TrueAttribute, Int1Attribute, true},
+				
+		{FalseAttribute, FalseAttribute, false},
+		{FalseAttribute, TrueAttribute, false},
+		{FalseAttribute, Int0Attribute, false},
+		{FalseAttribute, Int1Attribute, false},
+
+		{Int0Attribute, FalseAttribute, false},
+		{Int0Attribute, TrueAttribute, false},
+		{Int0Attribute, Int0Attribute, false},
+		{Int0Attribute, Int1Attribute, false},
+
+		{Int1Attribute, FalseAttribute, false},
+		{Int1Attribute, TrueAttribute, true},
+		{Int1Attribute, Int0Attribute, false},
+		{Int1Attribute, Int1Attribute, true},
+	};
 
-		if (!TestEqual(TEXT("Output default value"), Out, bExpectedResult))
+	for (const auto [Attr1, Attr2, bExpectedResult] : TestValues)
+	{
+		if (!Run(EPCGMetadataBooleanOperation::And, Attr1, Attr2, bExpectedResult))
 		{
 			return false;
 		}
+	}
 
-		return true;
-	};
-
-	bool bTestPassed = true;
-
-	/////////////////////////////////////////
-	// Not
-	/////////////////////////////////////////
-	ParamTaggedData1.Pin = PCGPinConstants::DefaultInputLabel;
-
-	Settings->Operation = EPCGMetadataBooleanOperation::Not;
+	return true;
+}
 
+bool FPCGMetadataBooleanXorTest::RunTest(const FString& Parameters)
+{
+	// Attr1, Attr2, bExpectedResult
+	TArray<TTuple<FName, FName, bool>> TestValues
 	{
-		Settings->InputSource1.SetAttributeName(TrueAttribute);
-		bTestPassed &= ValidateOp(/*bExpectedResult=*/ false);
-	}
+		{TrueAttribute, FalseAttribute, true},
+		{TrueAttribute, TrueAttribute, false},
+		{TrueAttribute, Int0Attribute, true},
+		{TrueAttribute, Int1Attribute, false},
+			
+		{FalseAttribute, FalseAttribute, false},
+		{FalseAttribute, TrueAttribute, true},
+		{FalseAttribute, Int0Attribute, false},
+		{FalseAttribute, Int1Attribute, true},
+
+		{Int0Attribute, FalseAttribute, false},
+		{Int0Attribute, TrueAttribute, true},
+		{Int0Attribute, Int0Attribute, false},
+		{Int0Attribute, Int1Attribute, true},
+
+		{Int1Attribute, FalseAttribute, true},
+		{Int1Attribute, TrueAttribute, false},
+		{Int1Attribute, Int0Attribute, true},
+		{Int1Attribute, Int1Attribute, false},
+	};
 
+	for (const auto [Attr1, Attr2, bExpectedResult] : TestValues)
 	{
-		Settings->InputSource1.SetAttributeName(FalseAttribute);
-		bTestPassed &= ValidateOp(true);
+		if (!Run(EPCGMetadataBooleanOperation::Xor, Attr1, Attr2, bExpectedResult))
+		{
+			return false;
+		}
 	}
 
-	/////////////////////////////////////////
-	// And
-	/////////////////////////////////////////
-
-	// Validate invalid op when there is type mismatch.
-	ParamTaggedData1.Pin = PCGMetadataSettingsBaseConstants::DoubleInputFirstLabel;
-	ParamTaggedData2.Pin = PCGMetadataSettingsBaseConstants::DoubleInputSecondLabel;
-	Settings->Operation = EPCGMetadataBooleanOperation::And;
-
-	AddExpectedError(TEXT("Attribute/Property 'Invalid' from pin"), EAutomationExpectedErrorFlags::Contains, 3);
+	return true;
+}
 
+bool FPCGMetadataBooleanNandTest::RunTest(const FString& Parameters)
+{
+	// Attr1, Attr2, bExpectedResult
+	TArray<TTuple<FName, FName, bool>> TestValues
 	{
-		Settings->InputSource1.SetAttributeName(InvalidAttribute);
-		Settings->InputSource2.SetAttributeName(FalseAttribute);
-		bTestPassed &= ValidateOp(/*bExpectedResult=*/false, /*bIsValid=*/ false);
-	}
+		{TrueAttribute, FalseAttribute, true},
+		{TrueAttribute, TrueAttribute, false},
+		{TrueAttribute, Int0Attribute, true},
+		{TrueAttribute, Int1Attribute, false},
+		
+		{FalseAttribute, FalseAttribute, true},
+		{FalseAttribute, TrueAttribute, true},
+		{FalseAttribute, Int0Attribute, true},
+		{FalseAttribute, Int1Attribute, true},
+
+		{Int0Attribute, FalseAttribute, true},
+		{Int0Attribute, TrueAttribute, true},
+		{Int0Attribute, Int0Attribute, true},
+		{Int0Attribute, Int1Attribute, true},
+
+		{Int1Attribute, FalseAttribute, true},
+		{Int1Attribute, TrueAttribute, false},
+		{Int1Attribute, Int0Attribute, true},
+		{Int1Attribute, Int1Attribute, false},
+	};
 
+	for (const auto [Attr1, Attr2, bExpectedResult] : TestValues)
 	{
-		Settings->InputSource1.SetAttributeName(FalseAttribute);
-		Settings->InputSource2.SetAttributeName(InvalidAttribute);
-		bTestPassed &= ValidateOp(/*bExpectedResult=*/false, /*bIsValid=*/ false);
+		if (!Run(EPCGMetadataBooleanOperation::Nand, Attr1, Attr2, bExpectedResult))
+		{
+			return false;
+		}
 	}
 
-	{
-		Settings->InputSource1.SetAttributeName(InvalidAttribute);
-		Settings->InputSource2.SetAttributeName(InvalidAttribute);
-		bTestPassed &= ValidateOp(/*bExpectedResult=*/false, /*bIsValid=*/ false);
-	}
+	return true;
+}
 
-	// Valid tests
+bool FPCGMetadataBooleanNorTest::RunTest(const FString& Parameters)
+{
+	// Attr1, Attr2, bExpectedResult
+	TArray<TTuple<FName, FName, bool>> TestValues =
 	{
-		Settings->InputSource1.SetAttributeName(TrueAttribute);
-		Settings->InputSource2.SetAttributeName(TrueAttribute);
-		bTestPassed &= ValidateOp(/*bExpectedResult=*/true);
-	}
+		{TrueAttribute, FalseAttribute, false},
+		{TrueAttribute, TrueAttribute, false},
+		{TrueAttribute, Int0Attribute, false},
+		{TrueAttribute, Int1Attribute, false},
+		
+		{FalseAttribute, FalseAttribute, true},
+		{FalseAttribute, TrueAttribute, false},
+		{FalseAttribute, Int0Attribute, true},
+		{FalseAttribute, Int1Attribute, false},
+
+		{Int0Attribute, FalseAttribute, true},
+		{Int0Attribute, TrueAttribute, false},
+		{Int0Attribute, Int0Attribute, true},
+		{Int0Attribute, Int1Attribute, false},
+
+		{Int1Attribute, FalseAttribute, false},
+		{Int1Attribute, TrueAttribute, false},
+		{Int1Attribute, Int0Attribute, false},
+		{Int1Attribute, Int1Attribute, false},
+	};
 
+	for (const auto [Attr1, Attr2, bExpectedResult] : TestValues)
 	{
-		Settings->InputSource1.SetAttributeName(FalseAttribute);
-		Settings->InputSource2.SetAttributeName(TrueAttribute);
-		bTestPassed &= ValidateOp(/*bExpectedResult=*/false);
+		if (!Run(EPCGMetadataBooleanOperation::Nor, Attr1, Attr2, bExpectedResult))
+		{
+			return false;
+		}
 	}
 
+	return true;
+}
+
+bool FPCGMetadataBooleanXnorTest::RunTest(const FString& Parameters)
+{
+	// Attr1, Attr2, bExpectedResult
+	TArray<TTuple<FName, FName, bool>> TestValues
 	{
-		Settings->InputSource1.SetAttributeName(TrueAttribute);
-		Settings->InputSource2.SetAttributeName(FalseAttribute);
-		bTestPassed &= ValidateOp(/*bExpectedResult=*/false);
-	}
+		{TrueAttribute, FalseAttribute, false},
+		{TrueAttribute, TrueAttribute, true},
+		{TrueAttribute, Int0Attribute, false},
+		{TrueAttribute, Int1Attribute, true},
+		
+		{FalseAttribute, FalseAttribute, true},
+		{FalseAttribute, TrueAttribute, false},
+		{FalseAttribute, Int0Attribute, true},
+		{FalseAttribute, Int1Attribute, false},
+
+		{Int0Attribute, FalseAttribute, true},
+		{Int0Attribute, TrueAttribute, false},
+		{Int0Attribute, Int0Attribute, true},
+		{Int0Attribute, Int1Attribute, false},
+
+		{Int1Attribute, FalseAttribute, false},
+		{Int1Attribute, TrueAttribute, true},
+		{Int1Attribute, Int0Attribute, false},
+		{Int1Attribute, Int1Attribute, true},
+	};
 
+	for (const auto [Attr1, Attr2, bExpectedResult] : TestValues)
 	{
-		Settings->InputSource1.SetAttributeName(FalseAttribute);
-		Settings->InputSource2.SetAttributeName(FalseAttribute);
-		bTestPassed &= ValidateOp(/*bExpectedResult=*/false);
+		if (!Run(EPCGMetadataBooleanOperation::Xnor, Attr1, Attr2, bExpectedResult))
+		{
+			return false;
+		}
 	}
 
+	return true;
+}
 
-	/////////////////////////////////////////
-	// Or
-	/////////////////////////////////////////
-	Settings->Operation = EPCGMetadataBooleanOperation::Or;
+bool FPCGMetadataBooleanImplyTest::RunTest(const FString& Parameters)
+{
+	// Attr1, Attr2, bExpectedResult
+	TArray<TTuple<FName, FName, bool>> TestValues =
 	{
-		Settings->InputSource1.SetAttributeName(TrueAttribute);
-		Settings->InputSource2.SetAttributeName(TrueAttribute);
-		bTestPassed &= ValidateOp(/*bExpectedResult=*/true);
-	}
+		{TrueAttribute, FalseAttribute, false},
+		{TrueAttribute, TrueAttribute, true},
+		{TrueAttribute, Int0Attribute, false},
+		{TrueAttribute, Int1Attribute, true},
+		
+		{FalseAttribute, FalseAttribute, true},
+		{FalseAttribute, TrueAttribute, true},
+		{FalseAttribute, Int0Attribute, true},
+		{FalseAttribute, Int1Attribute, true},
+
+		{Int0Attribute, FalseAttribute, true},
+		{Int0Attribute, TrueAttribute, true},
+		{Int0Attribute, Int0Attribute, true},
+		{Int0Attribute, Int1Attribute, true},
+
+		{Int1Attribute, FalseAttribute, false},
+		{Int1Attribute, TrueAttribute, true},
+		{Int1Attribute, Int0Attribute, false},
+		{Int1Attribute, Int1Attribute, true},
+	};
 
+	for (const auto [Attr1, Attr2, bExpectedResult] : TestValues)
 	{
-		Settings->InputSource1.SetAttributeName(FalseAttribute);
-		Settings->InputSource2.SetAttributeName(TrueAttribute);
-		bTestPassed &= ValidateOp(/*bExpectedResult=*/true);
+		if (!Run(EPCGMetadataBooleanOperation::Imply, Attr1, Attr2, bExpectedResult))
+		{
+			return false;
+		}
 	}
 
+	return true;
+}
+
+bool FPCGMetadataBooleanNimplyTest::RunTest(const FString& Parameters)
+{
+	// Attr1, Attr2, bExpectedResult
+	TArray<TTuple<FName, FName, bool>> TestValues
 	{
-		Settings->InputSource1.SetAttributeName(TrueAttribute);
-		Settings->InputSource2.SetAttributeName(FalseAttribute);
-		bTestPassed &= ValidateOp(/*bExpectedResult=*/true);
-	}
+		{TrueAttribute, FalseAttribute, true},
+		{TrueAttribute, TrueAttribute, false},
+		{TrueAttribute, Int0Attribute, true},
+		{TrueAttribute, Int1Attribute, false},
+		
+		{FalseAttribute, FalseAttribute, false},
+		{FalseAttribute, TrueAttribute, false},
+		{FalseAttribute, Int0Attribute, false},
+		{FalseAttribute, Int1Attribute, false},
+
+		{Int0Attribute, FalseAttribute, false},
+		{Int0Attribute, TrueAttribute, false},
+		{Int0Attribute, Int0Attribute, false},
+		{Int0Attribute, Int1Attribute, false},
+
+		{Int1Attribute, FalseAttribute, true},
+		{Int1Attribute, TrueAttribute, false},
+		{Int1Attribute, Int0Attribute, true},
+		{Int1Attribute, Int1Attribute, false},
+	};
 
+	for (const auto [Attr1, Attr2, bExpectedResult] : TestValues)
 	{
-		Settings->InputSource1.SetAttributeName(FalseAttribute);
-		Settings->InputSource2.SetAttributeName(FalseAttribute);
-		bTestPassed &= ValidateOp(/*bExpectedResult=*/false);
+		if (!Run(EPCGMetadataBooleanOperation::Nimply, Attr1, Attr2, bExpectedResult))
+		{
+			return false;
+		}
 	}
 
-	/////////////////////////////////////////
-	// Xor
-	/////////////////////////////////////////
-	Settings->Operation = EPCGMetadataBooleanOperation::Xor;
+	return true;
+}
+
+bool FPCGMetadataBooleanInvalidTest::RunTest(const FString& Parameters)
+{
+	AddExpectedError(TEXT("Attribute/Property 'Invalid' from pin InA is not a supported type ('Vector')"), EAutomationExpectedErrorFlags::Contains, 2, /*bIsRegex=*/false);
+	AddExpectedError(TEXT("Attribute/Property 'Invalid' from pin InB is not a supported type ('Vector')"), EAutomationExpectedErrorFlags::Contains, 1, /*bIsRegex=*/false);
+	AddExpectedError(TEXT("Attribute/Property 'Invalid' from pin In is not a supported type ('Vector')"), EAutomationExpectedErrorFlags::Contains, 1, /*bIsRegex=*/false);
+	
+	if (!Run(EPCGMetadataBooleanOperation::Not, InvalidAttribute, InvalidAttribute, /*bExpectedResult*/ false, /*bIsValid=*/ false))
 	{
-		Settings->InputSource1.SetAttributeName(TrueAttribute);
-		Settings->InputSource2.SetAttributeName(TrueAttribute);
-		bTestPassed &= ValidateOp(/*bExpectedResult=*/false);
+		return false;
 	}
-
+	
+	if (!Run(EPCGMetadataBooleanOperation::Or, InvalidAttribute, FalseAttribute, /*bExpectedResult*/ false, /*bIsValid=*/ false))
 	{
-		Settings->InputSource1.SetAttributeName(FalseAttribute);
-		Settings->InputSource2.SetAttributeName(TrueAttribute);
-		bTestPassed &= ValidateOp(/*bExpectedResult=*/true);
+		return false;
 	}
 
+	if (!Run(EPCGMetadataBooleanOperation::Or, FalseAttribute, InvalidAttribute, /*bExpectedResult*/ false, /*bIsValid=*/ false))
 	{
-		Settings->InputSource1.SetAttributeName(TrueAttribute);
-		Settings->InputSource2.SetAttributeName(FalseAttribute);
-		bTestPassed &= ValidateOp(/*bExpectedResult=*/true);
+		return false;
 	}
 
+	if (!Run(EPCGMetadataBooleanOperation::Or, InvalidAttribute, InvalidAttribute, /*bExpectedResult*/ false, /*bIsValid=*/ false))
 	{
-		Settings->InputSource1.SetAttributeName(FalseAttribute);
-		Settings->InputSource2.SetAttributeName(FalseAttribute);
-		bTestPassed &= ValidateOp(/*bExpectedResult=*/false);
+		return false;
 	}
-
-	return bTestPassed;
+	
+	return true;
 }
 
 #endif // WITH_EDITOR
diff --git a/Engine/Plugins/PCG/Source/PCG/Public/Elements/Metadata/PCGMetadataBitwiseOpElement.h b/Engine/Plugins/PCG/Source/PCG/Public/Elements/Metadata/PCGMetadataBitwiseOpElement.h
--- a/Engine/Plugins/PCG/Source/PCG/Public/Elements/Metadata/PCGMetadataBitwiseOpElement.h
+++ b/Engine/Plugins/PCG/Source/PCG/Public/Elements/Metadata/PCGMetadataBitwiseOpElement.h
@@ -12,9 +12,16 @@ enum class EPCGMetadataBitwiseOperation : uint16
 	And,
 	Not,
 	Or,
-	Xor
+	Xor,
+	ShiftLeft UMETA(ToolTip="<< operation"),
+	ShiftRight UMETA(ToolTip=">> operation")
 };
 
+namespace PCGMetadataBitwiseConstants
+{
+	const FLazyName ShiftInputLabel = TEXT("Shift");
+}
+
 UCLASS(MinimalAPI, BlueprintType, ClassGroup = (Procedural))
 class UPCGMetadataBitwiseSettings : public UPCGMetadataSettingsBase
 {
diff --git a/Engine/Plugins/PCG/Source/PCG/Public/Elements/Metadata/PCGMetadataBooleanOpElement.h b/Engine/Plugins/PCG/Source/PCG/Public/Elements/Metadata/PCGMetadataBooleanOpElement.h
--- a/Engine/Plugins/PCG/Source/PCG/Public/Elements/Metadata/PCGMetadataBooleanOpElement.h
+++ b/Engine/Plugins/PCG/Source/PCG/Public/Elements/Metadata/PCGMetadataBooleanOpElement.h
@@ -12,7 +12,12 @@ enum class EPCGMetadataBooleanOperation : uint16
 	And,
 	Not,
 	Or,
-	Xor
+	Xor,
+	Nand,
+	Nor,
+	Xnor,
+	Imply UMETA(ToolTip="A -> B False only when (A=1, B=0)"),
+	Nimply UMETA(ToolTip="A -/-> B True only when (A=1, B=0)"),
 };
 
 UCLASS(MinimalAPI, BlueprintType, ClassGroup = (Procedural))
