[PCG] Add  "AddModulo" operation in maths node

* Do operation (A + B) % C
* Useful to get the ($Index + 1) % NumPoints

#rnx

diff --git a/Engine/Plugins/PCG/Source/PCG/Private/Elements/Metadata/PCGMetadataMaths.inl b/Engine/Plugins/PCG/Source/PCG/Private/Elements/Metadata/PCGMetadataMaths.inl
--- a/Engine/Plugins/PCG/Source/PCG/Private/Elements/Metadata/PCGMetadataMaths.inl
+++ b/Engine/Plugins/PCG/Source/PCG/Private/Elements/Metadata/PCGMetadataMaths.inl
@@ -504,4 +504,42 @@ namespace PCGMetadataMaths
 	{
 		return -Value;
 	}
+
+	////////////////////
+	// AddModulo
+	////////////////////
+	template <typename T>
+	inline T AddModulo(const T& Value1, const T& Value2, const T& Value3)
+	{
+		// If Result is negative, add Value3. (i.e. A=0, B=-1, C=5 Result=-1, returns 4.
+		const T Result = Modulo(Value1 + Value2, Value3);
+		return Result >= 0 ? Result : (Result + Value3);
+	}
+
+	template<>
+	inline FVector2D AddModulo(const FVector2D& Value1, const FVector2D& Value2, const FVector2D& Value3)
+	{
+		return FVector2D(
+			AddModulo(Value1.X, Value2.X, Value3.X),
+			AddModulo(Value1.Y, Value2.Y, Value3.X));
+	}
+
+	template<>
+	inline FVector AddModulo(const FVector& Value1, const FVector& Value2, const FVector& Value3)
+	{
+		return FVector(
+			AddModulo(Value1.X, Value2.X, Value3.X),
+			AddModulo(Value1.Y, Value2.Y, Value3.Y),
+			AddModulo(Value1.Z, Value2.Z, Value3.Z));
+	}
+
+	template<>
+	inline FVector4 AddModulo(const FVector4& Value1, const FVector4& Value2, const FVector4& Value3)
+	{
+		return FVector4(
+			AddModulo(Value1.X, Value2.X, Value3.X),
+			AddModulo(Value1.Y, Value2.Y, Value3.Y),
+			AddModulo(Value1.Z, Value2.Z, Value3.Z),
+			AddModulo(Value1.W, Value2.W, Value3.W));
+	}
 }
diff --git a/Engine/Plugins/PCG/Source/PCG/Private/Elements/Metadata/PCGMetadataMathsOpElement.cpp b/Engine/Plugins/PCG/Source/PCG/Private/Elements/Metadata/PCGMetadataMathsOpElement.cpp
--- a/Engine/Plugins/PCG/Source/PCG/Private/Elements/Metadata/PCGMetadataMathsOpElement.cpp
+++ b/Engine/Plugins/PCG/Source/PCG/Private/Elements/Metadata/PCGMetadataMathsOpElement.cpp
@@ -162,6 +162,8 @@ namespace PCGMetadataMathsSettings
 			return PCGMetadataMaths::Lerp(Value1, Value2, Value3);
 		case EPCGMetadataMathsOperation::MulAdd:
 			return Value1 + Value2 * Value3;
+		case EPCGMetadataMathsOperation::AddModulo:
+			return PCGMetadataMaths::AddModulo(Value1, Value2, Value3);
 		default:
 			return T{};
 		}
diff --git a/Engine/Plugins/PCG/Source/PCG/Private/Tests/Elements/Metadata/PCGMetadataMathsOpTest.cpp b/Engine/Plugins/PCG/Source/PCG/Private/Tests/Elements/Metadata/PCGMetadataMathsOpTest.cpp
--- a/Engine/Plugins/PCG/Source/PCG/Private/Tests/Elements/Metadata/PCGMetadataMathsOpTest.cpp
+++ b/Engine/Plugins/PCG/Source/PCG/Private/Tests/Elements/Metadata/PCGMetadataMathsOpTest.cpp
@@ -579,6 +579,23 @@ bool FPCGMetadataMathsOp_Negate::RunTest(const FString& Parameters)
 	return bSuccess;
 }
 
+IMPLEMENT_CUSTOM_SIMPLE_AUTOMATION_TEST(FPCGMetadataMathsOp_AddModulo, FPCGMetadataMathsOpTest, "Plugins.PCG.Metadata.MathsOp.AddModulo", PCGTestsCommon::TestFlags)
+
+bool FPCGMetadataMathsOp_AddModulo::RunTest(const FString& Parameters)
+{
+	bool bSuccess = true;
+
+	using EPCGMetadataMathsOperation::AddModulo;
+	
+	// ({0, 1, 2} + 1) % 3 -> {1, 2, 0}
+	bSuccess &= ExecuteTest(TestParams(AddModulo, GenerateParamData<int32>({{0, 1, 2}, {1}, {3}}), TArray({0, 0, 0}), TArray({1, 2, 0})));
+	
+	// ({0, 1, 2} - 1) % 3 -> {2, 0, 1}
+	bSuccess &= ExecuteTest(TestParams(AddModulo, GenerateParamData<int64>({{0ll, 1ll, 2ll}, {-1ll}, {3ll}}), TArray({0ll, 0ll, 0ll}), TArray({2ll, 0ll, 1ll})));
+	
+	return bSuccess;
+}
+
 // TODO: Round, Sqrt, Floor, Ceil, Max, Min, Pow, ClampMin, ClampMax, Modulo, Set, Clamp, Lerp
 
 #endif // WITH_EDITOR
diff --git a/Engine/Plugins/PCG/Source/PCG/Public/Elements/Metadata/PCGMetadataMathsOpElement.h b/Engine/Plugins/PCG/Source/PCG/Public/Elements/Metadata/PCGMetadataMathsOpElement.h
--- a/Engine/Plugins/PCG/Source/PCG/Public/Elements/Metadata/PCGMetadataMathsOpElement.h
+++ b/Engine/Plugins/PCG/Source/PCG/Public/Elements/Metadata/PCGMetadataMathsOpElement.h
@@ -43,7 +43,8 @@ enum class EPCGMetadataMathsOperation : uint16
 	TernaryOp = 1 << 12 UMETA(Hidden),
 	Clamp,
 	Lerp,
-	MulAdd UMETA(Tooltip = "Multiply Add (A + B * C)", SearchHints = "*+")
+	MulAdd UMETA(Tooltip = "Multiply Add (A + B * C)", SearchHints = "*+"),
+	AddModulo UMETA(Tooltip = "Add Modulo (A + B) % C. If B is negative: (A + C + B) % C", SearchHints = "+%")
 };
 ENUM_CLASS_FLAGS(EPCGMetadataMathsOperation);
 
