[PCG] GPU: Fix ensures and crash when an upstream pin is of unsupported type like Any.

In the past we had ensure()s in this case, but it recently became a hard crash. Move to warnings instead.

#rnx

diff --git a/Engine/Plugins/PCG/Source/PCG/Private/Compute/DataInterfaces/PCGDataCollectionDataInterface.cpp b/Engine/Plugins/PCG/Source/PCG/Private/Compute/DataInterfaces/PCGDataCollectionDataInterface.cpp
--- a/Engine/Plugins/PCG/Source/PCG/Private/Compute/DataInterfaces/PCGDataCollectionDataInterface.cpp
+++ b/Engine/Plugins/PCG/Source/PCG/Private/Compute/DataInterfaces/PCGDataCollectionDataInterface.cpp
@@ -9,6 +9,7 @@
 #include "PCGSettings.h"
 #include "PCGSubsystem.h"
 #include "Compute/PCGComputeCommon.h"
+#include "Compute/PCGComputeKernel.h"
 #include "Compute/PCGDataBinding.h"
 #include "Compute/PCGDataUploadAdaptor.h"
 #include "Compute/Data/PCGProxyForGPUData.h"
@@ -559,7 +560,14 @@ bool UPCGDataCollectionDataProvider::PrepareForExecute_GameThread(UPCGDataBindin
 		return false;
 	}
 
-	check(PinDesc);
+	if (!PinDesc)
+	{
+		UE_LOG(LogPCG, Warning, TEXT("UPCGDataCollectionDataProvider: missing kernel pin data description. Kernel='%s', output pin label='%s'"),
+			ProducerKernel ? *ProducerKernel->GetName() : TEXT("MISSING"),
+			*OutputPinLabel.ToString());
+
+		return true;
+	}
 
 	if (!bBufferSizeValidated)
 	{
diff --git a/Engine/Plugins/PCG/Source/PCG/Private/Compute/Elements/PCGCustomHLSLKernel.cpp b/Engine/Plugins/PCG/Source/PCG/Private/Compute/Elements/PCGCustomHLSLKernel.cpp
--- a/Engine/Plugins/PCG/Source/PCG/Private/Compute/Elements/PCGCustomHLSLKernel.cpp
+++ b/Engine/Plugins/PCG/Source/PCG/Private/Compute/Elements/PCGCustomHLSLKernel.cpp
@@ -247,7 +247,11 @@ bool UPCGCustomHLSLKernel::ComputeOutputBindingDataDesc(const UPCGComputeGraph*
 
 			if (!InBinding->ComputeKernelPinDataDesc(FirstKernelPin, &InputDesc))
 			{
-				ensure(false);
+				UE_LOG(LogPCG, Warning, TEXT("Kernel pin not registered in compute graph. May be due to unsupported pin type. KernelIndex=%d, PinLabel='%s', Input=%d"),
+					FirstKernelPin.KernelIndex,
+					*FirstKernelPin.PinLabel.ToString(),
+					FirstKernelPin.bIsInput);
+
 				return false;
 			}
 		}
@@ -1726,9 +1730,7 @@ bool UPCGCustomHLSLKernel::AreAttributesValid(FPCGContext* InContext, FText* Out
 
 			for (const FPCGPinProperties& Pin : CustomHLSLSettings->InputPins)
 			{
-				const FPCGDataCollectionDesc* DataDesc = DataBinding->GetCachedKernelPinDataDesc(this, Pin.Label, /*bIsInputPin=*/true);
-
-				if (ensure(DataDesc))
+				if (const FPCGDataCollectionDesc* DataDesc = DataBinding->GetCachedKernelPinDataDesc(this, Pin.Label, /*bIsInputPin=*/true))
 				{
 					InputPinDescs.Add(Pin.Label, DataDesc);
 				}
@@ -1736,9 +1738,7 @@ bool UPCGCustomHLSLKernel::AreAttributesValid(FPCGContext* InContext, FText* Out
 
 			for (const FPCGPinProperties& Pin : CustomHLSLSettings->OutputPins)
 			{
-				const FPCGDataCollectionDesc* DataDesc = DataBinding->GetCachedKernelPinDataDesc(this, Pin.Label, /*bIsInputPin=*/false);
-
-				if (ensure(DataDesc))
+				if (const FPCGDataCollectionDesc* DataDesc = DataBinding->GetCachedKernelPinDataDesc(this, Pin.Label, /*bIsInputPin=*/false))
 				{
 					OutputPinDescs.Add(Pin.Label, DataDesc);
 				}
@@ -1910,7 +1910,7 @@ int UPCGCustomHLSLKernel::GetElementCountForInputPin(const FPCGPinProperties& In
 
 	const FPCGDataCollectionDesc* InputDesc = InBinding->GetCachedKernelPinDataDesc(this, InInputPinProps.Label, /*bIsInputPin=*/true);
 
-	return ensure(InputDesc) ? InputDesc->ComputeDataElementCount(InInputPinProps.AllowedTypes) : 0;
+	return InputDesc ? InputDesc->ComputeDataElementCount(InInputPinProps.AllowedTypes) : 0;
 }
 
 #undef LOCTEXT_NAMESPACE
diff --git a/Engine/Plugins/PCG/Source/PCG/Private/Compute/PCGComputeGraph.cpp b/Engine/Plugins/PCG/Source/PCG/Private/Compute/PCGComputeGraph.cpp
--- a/Engine/Plugins/PCG/Source/PCG/Private/Compute/PCGComputeGraph.cpp
+++ b/Engine/Plugins/PCG/Source/PCG/Private/Compute/PCGComputeGraph.cpp
@@ -148,7 +148,11 @@ int32 UPCGComputeGraph::GetBindingIndex(const FPCGKernelPin& InKernelPin) const
 	}
 	else
 	{
-		ensure(false);
+		UE_LOG(LogPCG, Warning, TEXT("Kernel pin not registered in compute graph. May be due to unsupported pin type. KernelIndex=%d, PinLabel='%s', Input=%d"),
+			InKernelPin.KernelIndex,
+			*InKernelPin.PinLabel.ToString(),
+			InKernelPin.bIsInput);
+
 		return INDEX_NONE;
 	}
 }
diff --git a/Engine/Plugins/PCG/Source/PCG/Private/Compute/PCGDataBinding.cpp b/Engine/Plugins/PCG/Source/PCG/Private/Compute/PCGDataBinding.cpp
--- a/Engine/Plugins/PCG/Source/PCG/Private/Compute/PCGDataBinding.cpp
+++ b/Engine/Plugins/PCG/Source/PCG/Private/Compute/PCGDataBinding.cpp
@@ -172,7 +172,7 @@ const FPCGDataCollectionDesc* UPCGDataBinding::GetCachedKernelPinDataDesc(int32
 	const FPCGDataCollectionDesc* FoundDescription = DataDescriptionCache.Find(InGraphBindingIndex);
 
 	// We should never land here, all data descriptions should have been computed at this point.
-	ensureMsgf(FoundDescription, TEXT("Cache miss when trying to look up a kernel pin data description."));
+	UE_LOG(LogPCG, Warning, TEXT("Cache miss while retrieving kernel pin data description, binding index %d."), InGraphBindingIndex);
 
 	return FoundDescription;
 }
