[PCG] GPU: Fix failed conversion of ensure to warning.

#rnx

diff --git a/Engine/Plugins/PCG/Source/PCG/Private/Compute/PCGDataBinding.cpp b/Engine/Plugins/PCG/Source/PCG/Private/Compute/PCGDataBinding.cpp
--- a/Engine/Plugins/PCG/Source/PCG/Private/Compute/PCGDataBinding.cpp
+++ b/Engine/Plugins/PCG/Source/PCG/Private/Compute/PCGDataBinding.cpp
@@ -171,8 +171,11 @@ const FPCGDataCollectionDesc* UPCGDataBinding::GetCachedKernelPinDataDesc(int32
 
 	const FPCGDataCollectionDesc* FoundDescription = DataDescriptionCache.Find(InGraphBindingIndex);
 
-	// We should never land here, all data descriptions should have been computed at this point.
-	UE_LOG(LogPCG, Warning, TEXT("Cache miss while retrieving kernel pin data description, binding index %d."), InGraphBindingIndex);
+	if (!FoundDescription)
+	{
+		// Should not happen - all data descriptions should have been computed at this point.
+		UE_LOG(LogPCG, Warning, TEXT("Cache miss while retrieving kernel pin data description, binding index %d."), InGraphBindingIndex);
+	}
 
 	return FoundDescription;
 }
