[PCG] GPU: Fix for crash in metadata partion node from using wrong index.

#rnx

diff --git a/Engine/Plugins/PCG/Source/PCG/Private/Elements/Metadata/PCGMetadataPartitionKernel.cpp b/Engine/Plugins/PCG/Source/PCG/Private/Elements/Metadata/PCGMetadataPartitionKernel.cpp
--- a/Engine/Plugins/PCG/Source/PCG/Private/Elements/Metadata/PCGMetadataPartitionKernel.cpp
+++ b/Engine/Plugins/PCG/Source/PCG/Private/Elements/Metadata/PCGMetadataPartitionKernel.cpp
@@ -142,7 +142,7 @@ bool UPCGMetadataPartitionKernel::ComputeOutputBindingDataDesc(const UPCGCompute
 
 	FPCGDataCollectionDesc PinDesc;
 
-	for (int32 DataIndex = 0; DataIndex < InBinding->InputDataCollection.TaggedData.Num(); ++DataIndex)
+	for (int32 DataIndex = 0; DataIndex < AnalysisDataIndices.Num(); ++DataIndex)
 	{
 		const int32 AnalysisDataIndex = AnalysisDataIndices[DataIndex];
 		
