[PCG] GPU: Allow kernel to decide which data interface to create for a pin.

This gives the kernel flexibility to override the type of DI, or override settings on the DI as needed.

Also fix issue where the settings was invited to set the current dynamic pin type which was stomping over the type of kernel pins.

#rnx

diff --git a/Engine/Plugins/PCG/Source/PCG/Internal/Compute/PCGComputeCommon.h b/Engine/Plugins/PCG/Source/PCG/Internal/Compute/PCGComputeCommon.h
--- a/Engine/Plugins/PCG/Source/PCG/Internal/Compute/PCGComputeCommon.h
+++ b/Engine/Plugins/PCG/Source/PCG/Internal/Compute/PCGComputeCommon.h
@@ -6,13 +6,18 @@
 
 #include "PCGComputeCommon.generated.h"
 
+class UPCGComputeDataInterface;
 class UPCGComputeKernel;
 class UPCGData;
 class UPCGDataBinding;
+class UPCGNode;
+class UObject;
 class UPCGSettings;
+struct FPCGGPUCompilationContext;
 struct FPCGContext;
 struct FPCGDataCollectionDesc;
 struct FPCGDataDesc;
+struct FPCGPinProperties;
 struct FPCGPinPropertiesGPU;
 struct FPCGTaggedData;
 
@@ -127,6 +132,19 @@ namespace PCGComputeHelpers
 
 #if WITH_EDITOR
 	void ConvertObjectPathToShaderFilePath(FString& InOutPath);
+
+	struct FCreateDataInterfaceParams
+	{
+		FPCGGPUCompilationContext* Context = nullptr;
+		const FPCGPinProperties* PinProperties = nullptr;
+		const UPCGComputeKernel* ProducerKernel = nullptr;
+		UObject* ObjectOuter = nullptr;
+		bool bProducedByCPU = false;
+		bool bRequiresExport = false;
+		const UPCGNode* NodeForDebug = nullptr;
+	};
+
+	UPCGComputeDataInterface* CreateOutputPinDataInterface(const FCreateDataInterfaceParams& InParams);
 #endif
 }
 
diff --git a/Engine/Plugins/PCG/Source/PCG/Private/Compute/PCGComputeCommon.cpp b/Engine/Plugins/PCG/Source/PCG/Private/Compute/PCGComputeCommon.cpp
--- a/Engine/Plugins/PCG/Source/PCG/Private/Compute/PCGComputeCommon.cpp
+++ b/Engine/Plugins/PCG/Source/PCG/Private/Compute/PCGComputeCommon.cpp
@@ -6,11 +6,18 @@
 #include "PCGContext.h"
 #include "PCGParamData.h"
 #include "PCGSubsystem.h"
+#include "Compute/PCGComputeKernel.h"
 #include "Compute/PCGDataBinding.h"
 #include "Compute/PCGPinPropertiesGPU.h"
 #include "Compute/Data/PCGProxyForGPUData.h"
+#include "Compute/DataInterfaces/PCGDataCollectionDataInterface.h"
+#include "Compute/DataInterfaces/PCGLandscapeDataInterface.h"
+#include "Compute/DataInterfaces/PCGStaticMeshDataInterface.h"
+#include "Compute/DataInterfaces/PCGTextureDataInterface.h"
+#include "Compute/DataInterfaces/PCGVirtualTextureDataInterface.h"
 #include "Data/PCGBasePointData.h"
 #include "Data/PCGTextureData.h"
+#include "Graph/PCGGPUGraphCompilationContext.h"
 #include "Editor/IPCGEditorModule.h"
 
 #include "DynamicRHI.h"
@@ -190,6 +197,76 @@ namespace PCGComputeHelpers
 		// Shader compilation result parsing will break if it finds ':' where it doesn't expect.
 		InOutPath.ReplaceCharInline(TEXT(':'), TEXT('@'));
 	}
+
+	UPCGComputeDataInterface* CreateOutputPinDataInterface(const FCreateDataInterfaceParams& InParams)
+	{
+		check(InParams.Context);
+		check(InParams.PinProperties);
+		check(InParams.ObjectOuter);
+
+		FPCGGPUCompilationContext& Context = *InParams.Context;
+		const FPCGPinProperties& PinProperties = *InParams.PinProperties;
+
+		UPCGComputeDataInterface* DataInterface = nullptr;
+
+		switch (PinProperties.AllowedTypes)
+		{
+		case EPCGDataType::Point:
+		case EPCGDataType::Param:
+		{
+			UPCGDataCollectionDataInterface* DataInterfacePCGData = Context.NewObject_AnyThread<UPCGDataCollectionDataInterface>(InParams.ObjectOuter);
+
+			// If data comes from a CPU task, upload it to the GPU.
+			ensureMsgf(!InParams.bRequiresExport || !InParams.bProducedByCPU, TEXT("Download from GPU only relevant for data produced on GPU."));
+
+			DataInterfacePCGData->SetRequiresExport(InParams.bRequiresExport);
+
+			DataInterfacePCGData->SetElementCountMultiplier(InParams.ProducerKernel ? InParams.ProducerKernel->GetElementCountMultiplier(PinProperties.Label) : 1);
+			DataInterfacePCGData->SetRequiresZeroInitialization(InParams.ProducerKernel && InParams.ProducerKernel->DoesOutputPinRequireZeroInitialization(PinProperties.Label));
+
+			DataInterface = DataInterfacePCGData;
+			break;
+		}
+		case EPCGDataType::VirtualTexture:
+		{
+			DataInterface = Context.NewObject_AnyThread<UPCGVirtualTextureDataInterface>(InParams.ObjectOuter);
+			break;
+		}
+		case EPCGDataType::BaseTexture:
+		case EPCGDataType::Texture:
+		case EPCGDataType::RenderTarget:
+		{
+			UPCGTextureDataInterface* TextureDataInterface = Context.NewObject_AnyThread<UPCGTextureDataInterface>(InParams.ObjectOuter);
+			TextureDataInterface->SetRequiresExport(InParams.bRequiresExport);
+			TextureDataInterface->SetInitializeFromDataCollection(InParams.bProducedByCPU);
+
+			DataInterface = TextureDataInterface;
+
+			break;
+		}
+		case EPCGDataType::Landscape:
+		{
+			DataInterface = Context.NewObject_AnyThread<UPCGLandscapeDataInterface>(InParams.ObjectOuter);
+			break;
+		}
+		case EPCGDataType::StaticMeshResource:
+		{
+			DataInterface = Context.NewObject_AnyThread<UPCGStaticMeshDataInterface>(InParams.ObjectOuter);
+			break;
+		}
+		default:
+		{
+			UE_LOG(LogPCG, Error, TEXT("Unsupported connected upstream pin '%s' on node '%s' with type %s. Consider adding a conversion to a supported type such as Point."),
+				*PinProperties.Label.ToString(),
+				InParams.NodeForDebug ? *InParams.NodeForDebug->GetNodeTitle(EPCGNodeTitleType::ListView).ToString() : TEXT("MISSING"),
+				StaticEnum<EPCGDataType>() ? *StaticEnum<EPCGDataType>()->GetNameStringByValue(static_cast<int64>(PinProperties.AllowedTypes)) : TEXT("MISSING")
+			);
+			break;
+		}
+		}
+
+		return DataInterface;
+	}
 #endif
 }
 
diff --git a/Engine/Plugins/PCG/Source/PCG/Private/Compute/PCGComputeKernel.cpp b/Engine/Plugins/PCG/Source/PCG/Private/Compute/PCGComputeKernel.cpp
--- a/Engine/Plugins/PCG/Source/PCG/Private/Compute/PCGComputeKernel.cpp
+++ b/Engine/Plugins/PCG/Source/PCG/Private/Compute/PCGComputeKernel.cpp
@@ -221,6 +221,11 @@ void UPCGComputeKernel::GatherAdditionalSources(TArray<TObjectPtr<UComputeSource
 		}
 	}
 }
+
+UPCGComputeDataInterface* UPCGComputeKernel::CreateOutputPinDataInterface(const PCGComputeHelpers::FCreateDataInterfaceParams& InParams) const
+{
+	return PCGComputeHelpers::CreateOutputPinDataInterface(InParams);
+}
 #endif
 
 void UPCGComputeKernel::ComputeDataDescFromPinProperties(const FPCGPinPropertiesGPU& OutputPinProps, const TArrayView<const FPCGPinProperties>& InInputPinProps, UPCGDataBinding* InBinding, FPCGDataCollectionDesc& OutPinDesc) const
diff --git a/Engine/Plugins/PCG/Source/PCG/Private/Compute/PCGComputeKernel.h b/Engine/Plugins/PCG/Source/PCG/Private/Compute/PCGComputeKernel.h
--- a/Engine/Plugins/PCG/Source/PCG/Private/Compute/PCGComputeKernel.h
+++ b/Engine/Plugins/PCG/Source/PCG/Private/Compute/PCGComputeKernel.h
@@ -130,6 +130,9 @@ public:
 
 	/** Whether the pin matching the given label should be eligible for inspect/debug. */
 	bool IsPinInternal(FName InPinLabel) const { return InternalPinLabels.Contains(InPinLabel); }
+
+	/** Create a data interface for the given output pin. */
+	virtual UPCGComputeDataInterface* CreateOutputPinDataInterface(const PCGComputeHelpers::FCreateDataInterfaceParams& InParams) const;
 #endif
 
 	/** Compute a description of data that will be output from pin InOutputPinLabel. */
diff --git a/Engine/Plugins/PCG/Source/PCG/Private/Graph/PCGGraphCompilerGPU.cpp b/Engine/Plugins/PCG/Source/PCG/Private/Graph/PCGGraphCompilerGPU.cpp
--- a/Engine/Plugins/PCG/Source/PCG/Private/Graph/PCGGraphCompilerGPU.cpp
+++ b/Engine/Plugins/PCG/Source/PCG/Private/Graph/PCGGraphCompilerGPU.cpp
@@ -11,13 +11,8 @@
 #include "Compute/PCGPinPropertiesGPU.h"
 #include "Compute/DataInterfaces/PCGComputeDataInterface.h"
 #include "Compute/DataInterfaces/PCGCustomKernelDataInterface.h"
-#include "Compute/DataInterfaces/PCGDataCollectionDataInterface.h"
 #include "Compute/DataInterfaces/PCGDataLabelResolverDataInterface.h"
 #include "Compute/DataInterfaces/PCGDebugDataInterface.h"
-#include "Compute/DataInterfaces/PCGLandscapeDataInterface.h"
-#include "Compute/DataInterfaces/PCGStaticMeshDataInterface.h"
-#include "Compute/DataInterfaces/PCGTextureDataInterface.h"
-#include "Compute/DataInterfaces/PCGVirtualTextureDataInterface.h"
 #include "Compute/Elements/PCGComputeGraphElement.h"
 #include "Elements/PCGStaticMeshSpawnerKernel.h"
 #include "Graph/PCGGPUGraphCompilationContext.h"
@@ -910,85 +905,44 @@ void FPCGGraphCompilerGPU::CreateDataInterfaces(
 			ProducerKernel = *FoundKernel;
 		}
 
-		const bool bUpstreamIsGPUTask = InCollapsedTasks.Contains(InProducerTaskId);
-
-		EPCGDataType PinType = InOutputPinProperties.AllowedTypes;
-
-		// Dynamically typed pins could have a different type
 		const UPCGSettings* ProducerSettings = InCompiledTasks[InProducerTaskId].Node ? InCompiledTasks[InProducerTaskId].Node->GetSettings() : nullptr;
-		const UPCGPin* ProducerPin = InCompiledTasks[InProducerTaskId].Node ? InCompiledTasks[InProducerTaskId].Node->GetOutputPin(InOutputPinProperties.Label) : nullptr;
-		if (ProducerSettings && ProducerPin)
-		{
-			PinType = ProducerSettings->GetCurrentPinTypes(ProducerPin);
-		}
+		FPCGPinProperties PinProperties = InOutputPinProperties;
+
+		PCGComputeHelpers::FCreateDataInterfaceParams Params;
+		Params.Context = &InOutContext;
+		Params.PinProperties = &PinProperties;
+		Params.ProducerKernel = ProducerKernel;
+		Params.ObjectOuter = InOutComputeGraph;
+		Params.bProducedByCPU = !InCollapsedTasks.Contains(InProducerTaskId);
+		Params.bRequiresExport = bRequiresExport;
+		Params.NodeForDebug = InCompiledTasks[InProducerTaskId].Node;
 
 		UPCGComputeDataInterface* DataInterface = nullptr;
 
-		switch (PinType)
+		if (ProducerKernel)
 		{
-		case EPCGDataType::Point:
-		case EPCGDataType::Param:
+			// Allow kernel to decide which DI should be created.
+			DataInterface = ProducerKernel->CreateOutputPinDataInterface(Params);
+		}
+		else
 		{
-			UPCGDataCollectionDataInterface* DataInterfacePCGData = InOutContext.NewObject_AnyThread<UPCGDataCollectionDataInterface>(InOutComputeGraph);
-
-			// If data comes from a CPU task, upload it to the GPU.
-			ensureMsgf(!bRequiresExport || bUpstreamIsGPUTask, TEXT("Download from GPU only relevant for data produced on GPU."));
-
-			DataInterfacePCGData->SetRequiresExport(bRequiresExport);
-
-			if (ProducerPin)
+			// Resolve current pin type.
+			const UPCGPin* ProducerPin = InCompiledTasks[InProducerTaskId].Node ? InCompiledTasks[InProducerTaskId].Node->GetOutputPin(PinProperties.Label) : nullptr;
+			if (ProducerSettings && ProducerPin)
 			{
-				DataInterfacePCGData->SetElementCountMultiplier(ProducerKernel ? ProducerKernel->GetElementCountMultiplier(ProducerPin->Properties.Label) : 1);
-				DataInterfacePCGData->SetRequiresZeroInitialization(ProducerKernel && ProducerKernel->DoesOutputPinRequireZeroInitialization(ProducerPin->Properties.Label));
+				PinProperties.AllowedTypes = ProducerSettings->GetCurrentPinTypes(ProducerPin);
 			}
 
-			DataInterface = DataInterfacePCGData;
-			break;
-		}
-		case EPCGDataType::VirtualTexture:
-		{
-			DataInterface = InOutContext.NewObject_AnyThread<UPCGVirtualTextureDataInterface>(InOutComputeGraph);
-			break;
-		}
-		case EPCGDataType::BaseTexture:
-		case EPCGDataType::Texture:
-		case EPCGDataType::RenderTarget:
-		{
-			UPCGTextureDataInterface* TextureDataInterface = InOutContext.NewObject_AnyThread<UPCGTextureDataInterface>(InOutComputeGraph);
-			TextureDataInterface->SetRequiresExport(bRequiresExport);
-			TextureDataInterface->SetInitializeFromDataCollection(!bUpstreamIsGPUTask);
-
-			DataInterface = TextureDataInterface;
-
-			break;
-		}
-		case EPCGDataType::Landscape:
-		{
-			DataInterface = InOutContext.NewObject_AnyThread<UPCGLandscapeDataInterface>(InOutComputeGraph);
-			break;
-		}
-		case EPCGDataType::StaticMeshResource:
-		{
-			DataInterface = InOutContext.NewObject_AnyThread<UPCGStaticMeshDataInterface>(InOutComputeGraph);
-			break;
-		}
-		default:
-		{
-			UE_LOG(LogPCG, Error, TEXT("Unsupported connected upstream pin '%s' on node '%s' with type %s. Consider adding a conversion to a supported type such as Point."),
-				*InOutputPinProperties.Label.ToString(),
-				InCompiledTasks[InProducerTaskId].Node ? *InCompiledTasks[InProducerTaskId].Node->GetNodeTitle(EPCGNodeTitleType::ListView).ToString() : TEXT("MISSING"),
-				StaticEnum<EPCGDataType>() ? *StaticEnum<EPCGDataType>()->GetNameStringByValue(static_cast<int64>(PinType)) : TEXT("MISSING")
-			);
-			break;
-		}
+			// Use default logic for DI creation.
+			DataInterface = PCGComputeHelpers::CreateOutputPinDataInterface(Params);
 		}
 
 		if (DataInterface)
 		{
 			DataInterface->ProducerKernel = ProducerKernel;
 			DataInterface->SetProducerSettings(ProducerSettings);
-			DataInterface->SetOutputPin(InOutputPinProperties.Label);
-			DataInterface->SetProducedByCPU(!bUpstreamIsGPUTask);
+			DataInterface->SetOutputPin(PinProperties.Label);
+			DataInterface->SetProducedByCPU(Params.bProducedByCPU);
 		}
 
 		return DataInterface;
@@ -997,63 +951,80 @@ void FPCGGraphCompilerGPU::CreateDataInterfaces(
 	// Create DIs for output pins (either output pins within this compute graph or outside).
 	for (FPCGTaskId TaskId : InCollapsedTasks)
 	{
+		TArray<FPCGPinProperties> OutputPinsProperties;
+
 		// Create DIs for all output pins regardless of outbound connections, because the kernels currently need their outputs to be bound to valid resources.
-		if (const UPCGSettings* Settings = InCompiledTasks[TaskId].Node ? InCompiledTasks[TaskId].Node->GetSettings() : nullptr)
+
+		if (UPCGComputeKernel* const* Kernel = InTaskIdToKernel.Find(TaskId))
 		{
-			// Create all the output data interfaces.
-			for (const FPCGPinProperties& OutputPinProperties : Settings->AllOutputPinProperties())
+			// Get output pins from kernel.
+			TArray<FPCGPinPropertiesGPU> OutputPins;
+			(*Kernel)->GetOutputPins(OutputPins);
+
+			for (const FPCGPinPropertiesGPU& Pin : OutputPins)
 			{
-				if (OutputPinProperties.Label == NAME_None)
-				{
-					continue;
-				}
+				OutputPinsProperties.Add(Pin);
+			}
+		}
+		else if (const UPCGSettings* Settings = InCompiledTasks[TaskId].Node ? InCompiledTasks[TaskId].Node->GetSettings() : nullptr)
+		{
+			// Get output pins from settings, and apply any dynamic pin types.
+			OutputPinsProperties = Settings->AllOutputPinProperties();
+		}
 
-				if (InOutPinToDataInterface.Contains({ TaskId, OutputPinProperties.Label }))
-				{
-					ensure(false);
-					continue;
-				}
+		// Create all the output data interfaces.
+		for (const FPCGPinProperties& OutputPinProperties : OutputPinsProperties) 
+		{
+			if (OutputPinProperties.Label == NAME_None)
+			{
+				continue;
+			}
 
-				// Request buffer export from compute graph if any downstream task is outside of this compute graph.
-				bool bRequiresExport = false;
-				if (const TArray<FPCGTaskId>* Successors = InTaskSuccessors.Find(TaskId))
+			if (InOutPinToDataInterface.Contains({ TaskId, OutputPinProperties.Label }))
+			{
+				ensure(false);
+				continue;
+			}
+
+			// Request buffer export from compute graph if any downstream task is outside of this compute graph.
+			bool bRequiresExport = false;
+			if (const TArray<FPCGTaskId>* Successors = InTaskSuccessors.Find(TaskId))
+			{
+				for (FPCGTaskId Successor : *Successors)
 				{
-					for (FPCGTaskId Successor : *Successors)
+					for (const FPCGGraphTaskInput& Input : InCompiledTasks[Successor].Inputs)
 					{
-						for (const FPCGGraphTaskInput& Input : InCompiledTasks[Successor].Inputs)
+						if (Input.TaskId == TaskId && Input.UpstreamPin && !InCollapsedTasks.Contains(Successor))
 						{
-							if (Input.TaskId == TaskId && Input.UpstreamPin && !InCollapsedTasks.Contains(Successor))
-							{
-								bRequiresExport = true;
-								break;
-							}
+							bRequiresExport = true;
+							break;
 						}
 					}
 				}
+			}
 
-				if (UPCGComputeDataInterface* OutputDI = CreateOutputDataInterface(TaskId, bRequiresExport, OutputPinProperties))
-				{
-					// Get the compute graph virtual pin corresponding to this pin, if one was created.
-					const FName* VirtualPinLabel = InOriginalToVirtualPin.Find({ TaskId, OutputPinProperties.Label, /*bIsInputPin*/false });
-					OutputDI->SetOutputPin(OutputPinProperties.Label, VirtualPinLabel);
+			if (UPCGComputeDataInterface* OutputDI = CreateOutputDataInterface(TaskId, bRequiresExport, OutputPinProperties))
+			{
+				// Get the compute graph virtual pin corresponding to this pin, if one was created.
+				const FName* VirtualPinLabel = InOriginalToVirtualPin.Find({ TaskId, OutputPinProperties.Label, /*bIsInputPin*/false });
+				OutputDI->SetOutputPin(OutputPinProperties.Label, VirtualPinLabel);
 
-					InOutComputeGraph->DataInterfaces.Add(OutputDI);
+				InOutComputeGraph->DataInterfaces.Add(OutputDI);
 
-					InOutPinToDataInterface.Add({ TaskId, OutputPinProperties.Label }, OutputDI);
+				InOutPinToDataInterface.Add({ TaskId, OutputPinProperties.Label }, OutputDI);
 
-					// Iterate over downstream connections and register each downstream pin.
-					if (const TArray<FPCGTaskId>* Successors = InTaskSuccessors.Find(TaskId))
+				// Iterate over downstream connections and register each downstream pin.
+				if (const TArray<FPCGTaskId>* Successors = InTaskSuccessors.Find(TaskId))
+				{
+					for (FPCGTaskId Successor : *Successors)
 					{
-						for (FPCGTaskId Successor : *Successors)
+						for (const FPCGGraphTaskInput& Input : InCompiledTasks[Successor].Inputs)
 						{
-							for (const FPCGGraphTaskInput& Input : InCompiledTasks[Successor].Inputs)
+							if (Input.TaskId == TaskId && Input.UpstreamPin && Input.DownstreamPin)
 							{
-								if (Input.TaskId == TaskId && Input.UpstreamPin && Input.DownstreamPin)
-								{
-									// Get the compute graph virtual pin corresponding to this pin, if one was created.
-									const FName* VirtualInputPinLabel = InOriginalToVirtualPin.Find({ TaskId, Input.DownstreamPin->Label, /*bIsInputPin=*/true });
-									OutputDI->AddDownstreamInputPin(Input.DownstreamPin->Label, VirtualInputPinLabel);
-								}
+								// Get the compute graph virtual pin corresponding to this pin, if one was created.
+								const FName* VirtualInputPinLabel = InOriginalToVirtualPin.Find({ TaskId, Input.DownstreamPin->Label, /*bIsInputPin=*/true });
+								OutputDI->AddDownstreamInputPin(Input.DownstreamPin->Label, VirtualInputPinLabel);
 							}
 						}
 					}
@@ -1254,16 +1225,12 @@ void FPCGGraphCompilerGPU::CompileComputeGraph(
 			FName PinLabel = NAME_None;
 		};
 
-		const int NumInputPins = Settings->InputPinProperties().Num();
-		const int NumOutputPins = Settings->OutputPinProperties().Num();
-
 		TArray<FDataInterfaceInfo> InputDataInterfaces;
 		TArray<FDataInterfaceInfo> OutputDataInterfaces;
-		InputDataInterfaces.Reserve(NumInputPins);
-		OutputDataInterfaces.Reserve(NumOutputPins);
+		InputDataInterfaces.Reserve(8);
+		OutputDataInterfaces.Reserve(8);
 
-		TMap</*OriginalPinLabel*/FName, /*VirtualPinLabel*/FName> OriginalToVirtualInputPin;
-		OriginalToVirtualInputPin.Reserve(NumInputPins);
+		TMap</*OriginalPinLabel*/FName, /*VirtualPinLabel*/FName, TInlineSetAllocator<8>> OriginalToVirtualInputPin;
 
 		auto CreateDataResolverDataInterface = [&InOutContext, Kernel = KernelWithBindings.Kernel, InOutComputeGraph](TArray<FDataInterfaceInfo>& InOutDataInterfaces, FName PinLabel, bool bIsInput)
 		{
