[PCG] Added comma-separated tag list to create target actor node

diff --git a/Engine/Plugins/PCG/Source/PCG/Private/Elements/PCGCreateTargetActor.cpp b/Engine/Plugins/PCG/Source/PCG/Private/Elements/PCGCreateTargetActor.cpp
--- a/Engine/Plugins/PCG/Source/PCG/Private/Elements/PCGCreateTargetActor.cpp
+++ b/Engine/Plugins/PCG/Source/PCG/Private/Elements/PCGCreateTargetActor.cpp
@@ -484,6 +484,9 @@ bool FPCGCreateTargetActorElement::ExecuteInternal(FPCGContext* Context) const
 
 		GeneratedActor->Tags.Add(PCGHelpers::DefaultPCGActorTag);
 
+		const TArray<FString> AdditionalTags = PCGHelpers::GetStringArrayFromCommaSeparatedList(Settings->CommaSeparatedActorTags);
+		GeneratedActor->Tags.Append(AdditionalTags);
+
 #if WITH_EDITOR
 		if (!Settings->ActorLabel.IsEmpty())
 		{
diff --git a/Engine/Plugins/PCG/Source/PCG/Public/Elements/PCGCreateTargetActor.h b/Engine/Plugins/PCG/Source/PCG/Public/Elements/PCGCreateTargetActor.h
--- a/Engine/Plugins/PCG/Source/PCG/Public/Elements/PCGCreateTargetActor.h
+++ b/Engine/Plugins/PCG/Source/PCG/Public/Elements/PCGCreateTargetActor.h
@@ -75,6 +75,10 @@ public:
 	UPROPERTY(meta = (PCG_Overridable))
 	FTransform ActorPivot;
 
+	/** Comma-separated list of tags to add to the newly created target actor. */
+	UPROPERTY(EditAnywhere, BlueprintReadWrite, Category = Settings, meta = (PCG_Overridable))
+	FString CommaSeparatedActorTags;
+
 	/** Override the default property values on the created target actor. Applied before post-process functions. */
 	UPROPERTY(EditAnywhere, BlueprintReadWrite, Category = Settings)
 	TArray<FPCGObjectPropertyOverrideDescription> PropertyOverrideDescriptions;
