PCG Attribute accessor comment out spammy getter/setter trace events

diff --git a/Engine/Plugins/PCG/Source/PCG/Public/Metadata/Accessors/IPCGAttributeAccessorTpl.h b/Engine/Plugins/PCG/Source/PCG/Public/Metadata/Accessors/IPCGAttributeAccessorTpl.h
--- a/Engine/Plugins/PCG/Source/PCG/Public/Metadata/Accessors/IPCGAttributeAccessorTpl.h
+++ b/Engine/Plugins/PCG/Source/PCG/Public/Metadata/Accessors/IPCGAttributeAccessorTpl.h
@@ -34,7 +34,7 @@ private:
 	template<typename U>
 	bool InternalGetRange(TArrayView<U>& OutValues, int32 Index, const IPCGAttributeAccessorKeys& Keys, EPCGAttributeAccessorFlags Flags) const
 	{
-		TRACE_CPUPROFILER_EVENT_SCOPE(IPCGAttributeAccessorT::InternalGetRange);
+		//TRACE_CPUPROFILER_EVENT_SCOPE(IPCGAttributeAccessorT::InternalGetRange);
 
 		using T = typename Derived::Type;
 		const Derived* This = static_cast<const Derived*>(this);
@@ -92,7 +92,7 @@ private:
 	template<typename U>
 	bool InternalSetRange(TArrayView<const U>& InValues, int32 Index, IPCGAttributeAccessorKeys& Keys, EPCGAttributeAccessorFlags Flags)
 	{
-		TRACE_CPUPROFILER_EVENT_SCOPE(IPCGAttributeAccessorT::InternalSetRange);
+		//TRACE_CPUPROFILER_EVENT_SCOPE(IPCGAttributeAccessorT::InternalSetRange);
 
 		// Can't set if read only, and discard any set if there are too many values.
 		if (bReadOnly || (Index + InValues.Num()) > Keys.GetNum())
