[PCG] Follow up fix for dynamic pin types which broken combination of GPU attribute partition and SM Spawner.

One of the create data interface paths no longer resolved the pin type correctly. This change shifts the pin type resolve upstream to where the kernels are created.

It may be possible to push this type resolution further upstream to the CPU graph compiler itself. I tried this but ran into resistance (it broke some pin property comparison in the dynamic culling logic). Leaving for furture investigation.

#rnx

diff --git a/Engine/Plugins/PCG/Source/PCG/Private/Graph/PCGGraphCompilerGPU.cpp b/Engine/Plugins/PCG/Source/PCG/Private/Graph/PCGGraphCompilerGPU.cpp
--- a/Engine/Plugins/PCG/Source/PCG/Private/Graph/PCGGraphCompilerGPU.cpp
+++ b/Engine/Plugins/PCG/Source/PCG/Private/Graph/PCGGraphCompilerGPU.cpp
@@ -434,6 +434,22 @@ void FPCGGraphCompilerGPU::ExpandGPUNodeKernelsToTasks(
 						// Match - add a new input wire to the new task.
 						FPCGGraphTaskInput& NewTaskInput = NewKernelTask.Inputs.Add_GetRef(*OriginalTaskInput);
 
+						// Resolve dynamic pin type here to simplify later logic.
+						// todo_pcg: Possibly the graph compiler should resolve dynamic pins during compilation. In a quick test this broke dynamic culling logic
+						// due to mismatching pin properties.
+						if (NewTaskInput.UpstreamPin)
+						{
+							const UPCGSettings* ProducerSettings = InOutCompiledTasks[NewTaskInput.TaskId].Node ? InOutCompiledTasks[NewTaskInput.TaskId].Node->GetSettings() : nullptr;
+							const UPCGPin* ProducerPin = InOutCompiledTasks[NewTaskInput.TaskId].Node
+								? InOutCompiledTasks[NewTaskInput.TaskId].Node->GetOutputPin(NewTaskInput.UpstreamPin->Label)
+								: nullptr;
+
+							if (ProducerSettings && ProducerPin)
+							{
+								NewTaskInput.UpstreamPin->AllowedTypes = ProducerSettings->GetCurrentPinTypes(ProducerPin);
+							}
+						}
+
 						// Update the pin properties with those specified by the kernel.
 						NewTaskInput.DownstreamPin = KernelInputProps;
 					}
@@ -926,13 +942,6 @@ void FPCGGraphCompilerGPU::CreateDataInterfaces(
 		}
 		else
 		{
-			// Resolve current pin type.
-			const UPCGPin* ProducerPin = InCompiledTasks[InProducerTaskId].Node ? InCompiledTasks[InProducerTaskId].Node->GetOutputPin(PinProperties.Label) : nullptr;
-			if (ProducerSettings && ProducerPin)
-			{
-				PinProperties.AllowedTypes = ProducerSettings->GetCurrentPinTypes(ProducerPin);
-			}
-
 			// Use default logic for DI creation.
 			DataInterface = PCGComputeHelpers::CreateOutputPinDataInterface(Params);
 		}
