[PCG] Fill in missing coverage of GPU data transfer indicators.

* Edges that cross grids without an explicit grid size node now show the upload indicator on the upstream output pin.
* Add missing case was when a data collection can not be passed through a CPU node and forces a CPU readback, for example when Gather'ing data from multiple data collections which is not yet supported.

#rnx

diff --git a/Engine/Plugins/PCG/Source/PCG/Internal/Compute/PCGComputeCommon.h b/Engine/Plugins/PCG/Source/PCG/Internal/Compute/PCGComputeCommon.h
--- a/Engine/Plugins/PCG/Source/PCG/Internal/Compute/PCGComputeCommon.h
+++ b/Engine/Plugins/PCG/Source/PCG/Internal/Compute/PCGComputeCommon.h
@@ -145,6 +145,9 @@ namespace PCGComputeHelpers
 	};
 
 	UPCGComputeDataInterface* CreateOutputPinDataInterface(const FCreateDataInterfaceParams& InParams);
+
+	void NotifyGPUToCPUReadback(UPCGDataBinding* InBinding, const UPCGComputeKernel* InKernel, const UPCGSettings* InSettings);
+	void NotifyCPUToGPUUpload(UPCGDataBinding* InBinding, const UPCGComputeKernel* InKernel, const UPCGSettings* InSettings);
 #endif
 }
 
diff --git a/Engine/Plugins/PCG/Source/PCG/Private/Compute/DataInterfaces/PCGDataCollectionDataInterface.cpp b/Engine/Plugins/PCG/Source/PCG/Private/Compute/DataInterfaces/PCGDataCollectionDataInterface.cpp
--- a/Engine/Plugins/PCG/Source/PCG/Private/Compute/DataInterfaces/PCGDataCollectionDataInterface.cpp
+++ b/Engine/Plugins/PCG/Source/PCG/Private/Compute/DataInterfaces/PCGDataCollectionDataInterface.cpp
@@ -589,7 +589,7 @@ bool UPCGDataCollectionDataProvider::PrepareForExecute_GameThread(UPCGDataBindin
 			DataAdaptor = MakeShared<FPCGDataUploadAdaptor>(InBinding, PinDesc, DownstreamInputPinLabelAliases[0]);
 		}
 
-		const bool bPreparationDone = !DataAdaptor || DataAdaptor->PrepareData_GameThread();
+		const bool bPreparationDone = !DataAdaptor || DataAdaptor->PrepareData_GameThread(ProducerKernel, GetProducerSettings());
 
 #if WITH_EDITOR
 		if (bPreparationDone && DataAdaptor && DataAdaptor->IsUploadingFromCPU())
diff --git a/Engine/Plugins/PCG/Source/PCG/Private/Compute/PCGComputeCommon.cpp b/Engine/Plugins/PCG/Source/PCG/Private/Compute/PCGComputeCommon.cpp
--- a/Engine/Plugins/PCG/Source/PCG/Private/Compute/PCGComputeCommon.cpp
+++ b/Engine/Plugins/PCG/Source/PCG/Private/Compute/PCGComputeCommon.cpp
@@ -6,6 +6,7 @@
 #include "PCGContext.h"
 #include "PCGParamData.h"
 #include "PCGSubsystem.h"
+#include "Compute/PCGComputeGraph.h"
 #include "Compute/PCGComputeKernel.h"
 #include "Compute/PCGDataBinding.h"
 #include "Compute/PCGPinPropertiesGPU.h"
@@ -267,6 +268,50 @@ namespace PCGComputeHelpers
 
 		return DataInterface;
 	}
+
+	static void NotifyHelper(UPCGDataBinding* InBinding, const UPCGComputeKernel* InKernel, const UPCGSettings* InSettings, const TFunction<void(FPCGContext*, const UPCGNode*)>& InLambda)
+	{
+		if (!ensure(InBinding))
+		{
+			return;
+		}
+
+		TSharedPtr<FPCGContextHandle> ContextHandle = InBinding->ContextHandle.Pin();
+		if (FPCGContext* Context = ContextHandle ? ContextHandle->GetContext() : nullptr)
+		{
+			const UPCGNode* IndicatorNode = nullptr;
+
+			if (InKernel)
+			{
+				IndicatorNode = InBinding->Graph ? InBinding->Graph->GetKernelNode(InKernel) : nullptr;
+			}
+			else if (InSettings)
+			{
+				IndicatorNode = Cast<UPCGNode>(InSettings->GetOuter());
+			}
+
+			if (IndicatorNode)
+			{
+				InLambda(Context, IndicatorNode);
+			}
+		}
+	}
+
+	void NotifyGPUToCPUReadback(UPCGDataBinding* InBinding, const UPCGComputeKernel* InKernel, const UPCGSettings* InSettings)
+	{
+		NotifyHelper(InBinding, InKernel, InSettings, [](FPCGContext* InContext, const UPCGNode* InNode)
+		{
+			InContext->ExecutionSource->GetExecutionState().GetInspection().NotifyGPUToCPUReadback(InNode, InContext->GetStack());
+		});
+	}
+
+	void NotifyCPUToGPUUpload(UPCGDataBinding* InBinding, const UPCGComputeKernel* InKernel, const UPCGSettings* InSettings)
+	{
+		NotifyHelper(InBinding, InKernel, InSettings, [](FPCGContext* InContext, const UPCGNode* InNode)
+		{
+			InContext->ExecutionSource->GetExecutionState().GetInspection().NotifyCPUToGPUUpload(InNode, InContext->GetStack());
+		});
+	}
 #endif
 }
 
diff --git a/Engine/Plugins/PCG/Source/PCG/Private/Compute/PCGComputeGraph.cpp b/Engine/Plugins/PCG/Source/PCG/Private/Compute/PCGComputeGraph.cpp
--- a/Engine/Plugins/PCG/Source/PCG/Private/Compute/PCGComputeGraph.cpp
+++ b/Engine/Plugins/PCG/Source/PCG/Private/Compute/PCGComputeGraph.cpp
@@ -347,6 +347,12 @@ void UPCGComputeGraph::DebugLogDataDescriptions(const UPCGDataBinding* InBinding
 #endif
 }
 
+const UPCGNode* UPCGComputeGraph::GetKernelNode(const UComputeKernel* InKernel) const
+{
+	const int FoundIndex = KernelInvocations.IndexOfByPredicate([InKernel](const TObjectPtr<UComputeKernel>& Kernel) { return Kernel.Get() == InKernel; });
+	return (FoundIndex != INDEX_NONE) ? KernelToNode[FoundIndex].Get() : nullptr;
+}
+
 #ifdef PCG_GPU_KERNEL_PROFILING
 bool UPCGComputeGraph::ShouldRepeatDispatch() const
 {
diff --git a/Engine/Plugins/PCG/Source/PCG/Private/Compute/PCGComputeGraph.h b/Engine/Plugins/PCG/Source/PCG/Private/Compute/PCGComputeGraph.h
--- a/Engine/Plugins/PCG/Source/PCG/Private/Compute/PCGComputeGraph.h
+++ b/Engine/Plugins/PCG/Source/PCG/Private/Compute/PCGComputeGraph.h
@@ -97,6 +97,8 @@ public:
 	/** Logs detailed data descriptions for all the kernels in the compute graph. */
 	void DebugLogDataDescriptions(const UPCGDataBinding* InBinding) const;
 
+	const UPCGNode* GetKernelNode(const UComputeKernel* InKernel) const;
+
 #ifdef PCG_GPU_KERNEL_PROFILING
 	/** Whether to repeatedly dispatch every frame. Helpful for profiling. */
 	bool ShouldRepeatDispatch() const;
diff --git a/Engine/Plugins/PCG/Source/PCG/Private/Compute/PCGDataUploadAdaptor.cpp b/Engine/Plugins/PCG/Source/PCG/Private/Compute/PCGDataUploadAdaptor.cpp
--- a/Engine/Plugins/PCG/Source/PCG/Private/Compute/PCGDataUploadAdaptor.cpp
+++ b/Engine/Plugins/PCG/Source/PCG/Private/Compute/PCGDataUploadAdaptor.cpp
@@ -6,6 +6,7 @@
 #include "Compute/PCGComputeCommon.h"
 #include "Compute/PCGComputeGraph.h"
 #include "Compute/PCGDataBinding.h"
+#include "Compute/PCGKernelHelpers.h"
 #include "Compute/Data/PCGProxyForGPUData.h"
 
 #include "RenderCaptureInterface.h"
@@ -103,7 +104,7 @@ FPCGDataUploadAdaptor::FPCGDataUploadAdaptor(UPCGDataBinding* InDataBinding, con
 }
 
 /** Do any preparation work such as data readbacks. Returns true when preparation is complete. */
-bool FPCGDataUploadAdaptor::PrepareData_GameThread()
+bool FPCGDataUploadAdaptor::PrepareData_GameThread(const UPCGComputeKernel* InProducerKernel, const UPCGSettings* InProducerSettings)
 {
 	TRACE_CPUPROFILER_EVENT_SCOPE(FPCGDataUploadAdaptor::PrepareData_GameThread);
 
@@ -165,6 +166,10 @@ bool FPCGDataUploadAdaptor::PrepareData_GameThread()
 
 			if (const UPCGProxyForGPUData* Proxy = Cast<UPCGProxyForGPUData>(TaggedData.Data))
 			{
+#if WITH_EDITOR
+				PCGComputeHelpers::NotifyGPUToCPUReadback(Binding, InProducerKernel, InProducerSettings);
+#endif
+
 				UPCGProxyForGPUData::FReadbackResult Result = Proxy->GetCPUData(/*InContext=*/nullptr);
 
 				if (!Result.bComplete)
@@ -184,6 +189,10 @@ bool FPCGDataUploadAdaptor::PrepareData_GameThread()
 			return false;
 		}
 
+#if WITH_EDITOR
+		PCGComputeHelpers::NotifyCPUToGPUUpload(Binding, InProducerKernel, InProducerSettings);
+#endif
+
 		TargetDataCollectionDesc.PackDataCollection(Binding->InputDataCollection, InputPinLabel, Binding, PackedDataCollection);
 	}
 
diff --git a/Engine/Plugins/PCG/Source/PCG/Private/Compute/PCGDataUploadAdaptor.h b/Engine/Plugins/PCG/Source/PCG/Private/Compute/PCGDataUploadAdaptor.h
--- a/Engine/Plugins/PCG/Source/PCG/Private/Compute/PCGDataUploadAdaptor.h
+++ b/Engine/Plugins/PCG/Source/PCG/Private/Compute/PCGDataUploadAdaptor.h
@@ -8,7 +8,9 @@
 #include "RendererInterface.h"
 
 class FRDGBuilder;
+class UPCGComputeKernel;
 class UPCGDataBinding;
+class UPCGSettings;
 enum class EPCGExportMode : uint8;
 
 /** A helper to ingest incoming data objects and pass the required data to the data interface. */
@@ -18,7 +20,7 @@ public:
 	FPCGDataUploadAdaptor(UPCGDataBinding* InDataBinding, const FPCGDataCollectionDesc* InTargetDataCollectionDesc, FName InInputPinLabel);
 
 	/** Do any preparation work such as data readbacks. Returns true when preparation is complete. */
-	bool PrepareData_GameThread();
+	bool PrepareData_GameThread(const UPCGComputeKernel* InProducerKernel, const UPCGSettings* InProducerSettings);
 
 	FRDGBufferSRVRef GetAttributeRemapBufferSRV(FRDGBuilder& InGraphBuilder, int32& OutFirstRemappedAttributeId);
 
diff --git a/Engine/Plugins/PCG/Source/PCG/Private/Graph/PCGGraphCompilerGPU.cpp b/Engine/Plugins/PCG/Source/PCG/Private/Graph/PCGGraphCompilerGPU.cpp
--- a/Engine/Plugins/PCG/Source/PCG/Private/Graph/PCGGraphCompilerGPU.cpp
+++ b/Engine/Plugins/PCG/Source/PCG/Private/Graph/PCGGraphCompilerGPU.cpp
@@ -17,6 +17,7 @@
 #include "Elements/PCGStaticMeshSpawnerKernel.h"
 #include "Graph/PCGGPUGraphCompilationContext.h"
 #include "Graph/PCGGraphCompiler.h"
+#include "Graph/PCGGraphExecutor.h"
 
 #include "ComputeFramework/ComputeKernel.h"
 #include "HAL/IConsoleManager.h"
@@ -920,8 +921,23 @@ void FPCGGraphCompilerGPU::CreateDataInterfaces(
 		{
 			ProducerKernel = *FoundKernel;
 		}
+		
+		const UPCGSettings* ProducerSettings = nullptr;
+
+		if (InCompiledTasks[InProducerTaskId].Node)
+		{
+			ProducerSettings = InCompiledTasks[InProducerTaskId].Node->GetSettings();
+		}
+		else if (InCompiledTasks[InProducerTaskId].Element && InCompiledTasks[InProducerTaskId].Element->IsGridLinkage())
+		{
+			// Special case - edges that cross grids without an explicit grid size node have a special grid linkage element. In this case
+			// look up the producer settings from the link. This enables the data upload indicator to be displayed on the producer.
+			const PCGGraphExecutor::FPCGGridLinkageElement& LinkElement = static_cast<const PCGGraphExecutor::FPCGGridLinkageElement&>(*InCompiledTasks[InProducerTaskId].Element);
+			const UPCGPin* UpstreamPin = LinkElement.GetUpstreamPin();
+			const UPCGNode* Node = UpstreamPin ? UpstreamPin->Node : nullptr;
+			ProducerSettings = Node ? Node->GetSettings() : nullptr;
+		}
 
-		const UPCGSettings* ProducerSettings = InCompiledTasks[InProducerTaskId].Node ? InCompiledTasks[InProducerTaskId].Node->GetSettings() : nullptr;
 		FPCGPinProperties PinProperties = InOutputPinProperties;
 
 		PCGComputeHelpers::FCreateDataInterfaceParams Params;
@@ -977,7 +993,7 @@ void FPCGGraphCompilerGPU::CreateDataInterfaces(
 		}
 		else if (const UPCGSettings* Settings = InCompiledTasks[TaskId].Node ? InCompiledTasks[TaskId].Node->GetSettings() : nullptr)
 		{
-			// Get output pins from settings, and apply any dynamic pin types.
+			// Get output pins from settings.
 			OutputPinsProperties = Settings->AllOutputPinProperties();
 		}
 
