[PCG] GPU: Fix repeated data description computations by routing back through data binding (which has cache).

#rnx

diff --git a/Engine/Plugins/PCG/Source/PCG/Private/Compute/PCGComputeGraph.cpp b/Engine/Plugins/PCG/Source/PCG/Private/Compute/PCGComputeGraph.cpp
--- a/Engine/Plugins/PCG/Source/PCG/Private/Compute/PCGComputeGraph.cpp
+++ b/Engine/Plugins/PCG/Source/PCG/Private/Compute/PCGComputeGraph.cpp
@@ -124,7 +124,7 @@ bool UPCGComputeGraph::ComputeKernelBindingDataDesc(int32 InBindingIndex, UPCGDa
 	else if (GraphEdge.bKernelInput)
 	{
 		const int32* UpstreamBindingIndex = DownstreamToUpstreamBinding.Find(InBindingIndex);
-		return ensure(UpstreamBindingIndex) && ComputeKernelBindingDataDesc(*UpstreamBindingIndex, InBinding, OutDataDesc);
+		return ensure(UpstreamBindingIndex) && InBinding->ComputeKernelBindingDataDesc(*UpstreamBindingIndex, &OutDataDesc);
 	}
 	else
 	{
diff --git a/Engine/Plugins/PCG/Source/PCG/Private/Compute/PCGComputeKernel.cpp b/Engine/Plugins/PCG/Source/PCG/Private/Compute/PCGComputeKernel.cpp
--- a/Engine/Plugins/PCG/Source/PCG/Private/Compute/PCGComputeKernel.cpp
+++ b/Engine/Plugins/PCG/Source/PCG/Private/Compute/PCGComputeKernel.cpp
@@ -158,9 +158,7 @@ bool UPCGComputeKernel::IsKernelDataValid(FPCGContext* InContext) const
 					// when looking up the pin data description, so only consider InitPins which actually exist.
 					if (Node->GetInputPin(InitPin))
 					{
-						const FPCGDataCollectionDesc* InputDesc = DataBinding->GetCachedKernelPinDataDesc(this, InitPin, /*bIsInputPin=*/true);
-
-						if (ensure(InputDesc))
+						if (const FPCGDataCollectionDesc* InputDesc = DataBinding->GetCachedKernelPinDataDesc(this, InitPin, /*bIsInputPin=*/true))
 						{
 							const uint32 InputNumData = InputDesc->DataDescs.Num();
 
@@ -247,7 +245,7 @@ void UPCGComputeKernel::ComputeDataDescFromPinProperties(const FPCGPinProperties
 
 				FPCGDataCollectionDesc& InputDesc = InputDescs.AddDefaulted_GetRef();
 				const FPCGKernelPin KernelPin(KernelIndex, PinToInitFrom, /*bIsInput=*/true);
-				ensure(InBinding->ComputeKernelPinDataDesc(KernelPin, &InputDesc));
+				InBinding->ComputeKernelPinDataDesc(KernelPin, &InputDesc);
 			}
 		}
 
diff --git a/Engine/Plugins/PCG/Source/PCG/Private/Compute/PCGDataBinding.cpp b/Engine/Plugins/PCG/Source/PCG/Private/Compute/PCGDataBinding.cpp
--- a/Engine/Plugins/PCG/Source/PCG/Private/Compute/PCGDataBinding.cpp
+++ b/Engine/Plugins/PCG/Source/PCG/Private/Compute/PCGDataBinding.cpp
@@ -104,11 +104,6 @@ void UPCGDataBinding::PrimeDataDescriptionCache()
 
 bool UPCGDataBinding::ComputeKernelPinDataDesc(const FPCGKernelPin& InKernelPin, FPCGDataCollectionDesc* OutDesc)
 {
-	// Uncomment to be notified when a kernel pin description is computed late:
-	//ensureAlwaysMsgf(!IsInGameThread(), TEXT("Data descriptions should only be computed when priming the cache, which takes place on a worker thread."));
-	ensure(!bIsDataDescriptionCachePrimed);
-	ensure(bTablesInitialized);
-
 	if (!ensure(Graph))
 	{
 		return false;
@@ -122,7 +117,22 @@ bool UPCGDataBinding::ComputeKernelPinDataDesc(const FPCGKernelPin& InKernelPin,
 		return false;
 	}
 
-	if (FPCGDataCollectionDesc* FoundDesc = DataDescriptionCache.Find(GraphBindingIndex))
+	return ComputeKernelBindingDataDesc(GraphBindingIndex, OutDesc);
+}
+
+bool UPCGDataBinding::ComputeKernelBindingDataDesc(int32 InGraphBindingIndex, FPCGDataCollectionDesc* OutDesc)
+{
+	// Uncomment to be notified when a kernel pin description is computed late:
+	//ensureAlwaysMsgf(!IsInGameThread(), TEXT("Data descriptions should only be computed when priming the cache, which takes place on a worker thread."));
+	ensure(!bIsDataDescriptionCachePrimed);
+	ensure(bTablesInitialized);
+
+	if (!ensure(Graph))
+	{
+		return false;
+	}
+
+	if (FPCGDataCollectionDesc* FoundDesc = DataDescriptionCache.Find(InGraphBindingIndex))
 	{
 		if (OutDesc)
 		{
@@ -134,12 +144,16 @@ bool UPCGDataBinding::ComputeKernelPinDataDesc(const FPCGKernelPin& InKernelPin,
 
 	FPCGDataCollectionDesc ComputedDesc;
 	
-	if (Graph->ComputeKernelBindingDataDesc(GraphBindingIndex, this, ComputedDesc))
+	if (Graph->ComputeKernelBindingDataDesc(InGraphBindingIndex, this, ComputedDesc))
 	{
-		DataDescriptionCache.Add(GraphBindingIndex, ComputedDesc);
-
-		if (OutDesc)
+		if (!OutDesc)
+		{
+			DataDescriptionCache.Add(InGraphBindingIndex, MoveTemp(ComputedDesc));
+		}
+		else
 		{
+			DataDescriptionCache.Add(InGraphBindingIndex, ComputedDesc);
+
 			*OutDesc = MoveTemp(ComputedDesc);
 		}
 
diff --git a/Engine/Plugins/PCG/Source/PCG/Private/Compute/PCGDataBinding.h b/Engine/Plugins/PCG/Source/PCG/Private/Compute/PCGDataBinding.h
--- a/Engine/Plugins/PCG/Source/PCG/Private/Compute/PCGDataBinding.h
+++ b/Engine/Plugins/PCG/Source/PCG/Private/Compute/PCGDataBinding.h
@@ -107,6 +107,7 @@ public:
 
 	/** Computes a description of data for a single pin in the compute graph and caches it. */
 	bool ComputeKernelPinDataDesc(const FPCGKernelPin& InKernelPin, FPCGDataCollectionDesc* OutDesc);
+	bool ComputeKernelBindingDataDesc(int32 InGraphBindingIndex, FPCGDataCollectionDesc* OutDesc);
 
 	/** Get description of data produced by a pin. */
 	const FPCGDataCollectionDesc* GetCachedKernelPinDataDesc(const UPCGComputeKernel* InKernel, FName InPinLabel, bool bIsInput) const;
