[PCG] Fix static analysis warning.

#rnx

diff --git a/Engine/Plugins/PCG/Source/PCG/Private/Graph/PCGGraphCompilerGPU.cpp b/Engine/Plugins/PCG/Source/PCG/Private/Graph/PCGGraphCompilerGPU.cpp
--- a/Engine/Plugins/PCG/Source/PCG/Private/Graph/PCGGraphCompilerGPU.cpp
+++ b/Engine/Plugins/PCG/Source/PCG/Private/Graph/PCGGraphCompilerGPU.cpp
@@ -934,7 +934,7 @@ void FPCGGraphCompilerGPU::CreateDataInterfaces(
 			// look up the producer settings from the link. This enables the data upload indicator to be displayed on the producer.
 			const PCGGraphExecutor::FPCGGridLinkageElement& LinkElement = static_cast<const PCGGraphExecutor::FPCGGridLinkageElement&>(*InCompiledTasks[InProducerTaskId].Element);
 			const UPCGPin* UpstreamPin = LinkElement.GetUpstreamPin();
-			const UPCGNode* Node = UpstreamPin ? UpstreamPin->Node : nullptr;
+			const UPCGNode* Node = UpstreamPin ? UpstreamPin->Node.Get() : nullptr;
 			ProducerSettings = Node ? Node->GetSettings() : nullptr;
 		}
 
