[PCG] Add support to call Pre/PostEditChange on objects on ApplyOnObjects

#rnx

diff --git a/Engine/Plugins/PCG/Source/PCG/Private/Elements/PCGApplyOnActor.cpp b/Engine/Plugins/PCG/Source/PCG/Private/Elements/PCGApplyOnActor.cpp
--- a/Engine/Plugins/PCG/Source/PCG/Private/Elements/PCGApplyOnActor.cpp
+++ b/Engine/Plugins/PCG/Source/PCG/Private/Elements/PCGApplyOnActor.cpp
@@ -122,6 +122,12 @@ bool FPCGApplyOnObjectElement::ExecuteInternal(FPCGContext* Context) const
 
 	TArray<TPair<UObject*, int32>> TargetObjectsAndIndices;
 
+#if WITH_EDITOR
+	const bool bPropagateObjectChangeEvent = Settings->bPropagateObjectChangeEvent;
+#else 
+	const bool bPropagateObjectChangeEvent = false;
+#endif // WITH_EDITOR
+
 	int CurrentPathIndex = 0;
 	while (CurrentPathIndex < ThisContext->PathsToObjectsAndDataIndex.Num())
 	{
@@ -134,7 +140,7 @@ bool FPCGApplyOnObjectElement::ExecuteInternal(FPCGContext* Context) const
 			++CurrentPathIndex;
 		}
 
-		PCGObjectPropertyOverrideHelpers::ApplyOverrides(Settings->PropertyOverrideDescriptions, TargetObjectsAndIndices, PCGApplyOnActorConstants::ObjectPropertyOverridesLabel, InputIndex, Context);
+		PCGObjectPropertyOverrideHelpers::ApplyOverrides(Settings->PropertyOverrideDescriptions, TargetObjectsAndIndices, PCGApplyOnActorConstants::ObjectPropertyOverridesLabel, InputIndex, Context, bPropagateObjectChangeEvent);
 
 		for (const TPair<UObject*, int32>& TargetObjectAndIndex : TargetObjectsAndIndices)
 		{
diff --git a/Engine/Plugins/PCG/Source/PCG/Private/Metadata/PCGObjectPropertyOverride.cpp b/Engine/Plugins/PCG/Source/PCG/Private/Metadata/PCGObjectPropertyOverride.cpp
--- a/Engine/Plugins/PCG/Source/PCG/Private/Metadata/PCGObjectPropertyOverride.cpp
+++ b/Engine/Plugins/PCG/Source/PCG/Private/Metadata/PCGObjectPropertyOverride.cpp
@@ -23,7 +23,7 @@ namespace PCGObjectPropertyOverrideHelpers
 		return ObjectOverridePinProperties;
 	}
 
-	void ApplyOverrides(const TArray<FPCGObjectPropertyOverrideDescription>& InObjectPropertyOverrideDescriptions, const TArray<TPair<UObject*, int32>>& TargetObjectToOverrideIndices, FName OverridesPinLabel, int32 InInputDataIndex, FPCGContext* Context)
+	void ApplyOverrides(const TArray<FPCGObjectPropertyOverrideDescription>& InObjectPropertyOverrideDescriptions, const TArray<TPair<UObject*, int32>>& TargetObjectToOverrideIndices, FName OverridesPinLabel, int32 InInputDataIndex, FPCGContext* Context, bool bPropagateEditChangeEvent)
 	{
 		if (!Context || TargetObjectToOverrideIndices.IsEmpty())
 		{
@@ -61,7 +61,7 @@ namespace PCGObjectPropertyOverrideHelpers
 			int32 InputKeyIndex = TargetObjectAndIndex.Value;
 
 			FPCGObjectOverrides ObjectOverrides(TargetObject);
-			ObjectOverrides.Initialize(InObjectPropertyOverrideDescriptions, TargetObject, OverrideData, Context);
+			ObjectOverrides.Initialize(InObjectPropertyOverrideDescriptions, TargetObject, OverrideData, Context, bPropagateEditChangeEvent);
 			if (!ObjectOverrides.Apply(InputKeyIndex)) // Use the First Entry of the param data for override (similar to what is done in Parameter Overrides in FPCGContext)
 			{
 				PCGLog::LogErrorOnGraph(FText::Format(LOCTEXT("ApplyOverrideFailed", "Failed to apply property overrides to object '%s'."), FText::FromName(TargetObject->GetClass()->GetFName())), Context);
@@ -69,17 +69,18 @@ namespace PCGObjectPropertyOverrideHelpers
 		}
 	}
 
-	void ApplyOverridesFromParams(const TArray<FPCGObjectPropertyOverrideDescription>& InObjectPropertyOverrideDescriptions, UObject* TargetObject, FName OverridesPinLabel, FPCGContext* Context)
+	void ApplyOverridesFromParams(const TArray<FPCGObjectPropertyOverrideDescription>& InObjectPropertyOverrideDescriptions, UObject* TargetObject, FName OverridesPinLabel, FPCGContext* Context, bool bPropagateEditChangeEvent)
 	{
 		ApplyOverrides(InObjectPropertyOverrideDescriptions,
 			{ TPair<UObject*, int32>(TargetObject, /*Entry index=*/0) },
 			OverridesPinLabel,
 			/*InInputDataIndex=*/0,
-			Context);
+			Context,
+			bPropagateEditChangeEvent);
 	}
 }
 
-void FPCGObjectSingleOverride::Initialize(const FPCGAttributePropertySelector& InputSelector, const FString& OutputProperty, const UStruct* TemplateClass, const UPCGData* SourceData, FPCGContext* Context)
+void FPCGObjectSingleOverride::Initialize(const FPCGAttributePropertySelector& InputSelector, const FString& OutputProperty, const UStruct* TemplateClass, const UPCGData* SourceData, FPCGContext* Context, bool bComputePropertyEditChain)
 {
 	InputKeys = PCGAttributeAccessorHelpers::CreateConstKeys(SourceData, InputSelector);
 	ObjectOverrideInputAccessor = PCGAttributeAccessorHelpers::CreateConstAccessor(SourceData, InputSelector);
@@ -95,11 +96,18 @@ void FPCGObjectSingleOverride::Initialize(const FPCGAttributePropertySelector& I
 	const TArray<FString>& ExtraNames = OutputSelector.GetExtraNames();
 	if (ExtraNames.IsEmpty())
 	{
-		if (const FProperty* Property = TemplateClass->FindPropertyByName(FName(OutputProperty)))
+		if (FProperty* Property = TemplateClass->FindPropertyByName(FName(OutputProperty)))
 		{
 			ObjectOverrideOutputAccessor = PCGAttributeAccessorHelpers::CreatePropertyAccessor(Property);
 			TArray<const FStructProperty*> EncounteredStructProps;
 			bWillNeedLoading = Property->ContainsObjectReference(EncounteredStructProps, EPropertyObjectReferenceType::Strong);
+
+#if WITH_EDITOR
+			if (bComputePropertyEditChain)
+			{
+				EditPropertyChain.AddHead(Property);
+			}
+#endif // WITH_EDITOR
 		}
 	}
 	else
@@ -120,6 +128,16 @@ void FPCGObjectSingleOverride::Initialize(const FPCGAttributePropertySelector& I
 			bWillNeedLoading = PropertyChain.Last()->ContainsObjectReference(EncounteredStructProps, EPropertyObjectReferenceType::Strong);
 
 			ObjectOverrideOutputAccessor = PCGAttributeAccessorHelpers::CreatePropertyChainAccessor(std::move(PropertyChain));
+
+#if WITH_EDITOR
+			if (bComputePropertyEditChain)
+			{
+				for (const FProperty* Property : PropertyChain)
+				{
+					EditPropertyChain.AddTail(const_cast<FProperty*>(Property));
+				}
+			}
+#endif // WITH_EDITOR
 		}
 	}
 
@@ -160,6 +178,27 @@ bool FPCGObjectSingleOverride::IsValid() const
 	return InputKeys.IsValid() && ObjectOverrideInputAccessor.IsValid() && ObjectOverrideOutputAccessor.IsValid() && ObjectOverrideFunction;
 }
 
+void FPCGObjectSingleOverride::PreApply(UObject* TargetObject)
+{
+#if WITH_EDITOR
+	if (TargetObject && !EditPropertyChain.IsEmpty())
+	{
+		TargetObject->PreEditChange(EditPropertyChain);
+	}
+#endif // WITH_EDITOR
+}
+
+void FPCGObjectSingleOverride::PostApply(UObject* TargetObject)
+{
+#if WITH_EDITOR
+	if (TargetObject && !EditPropertyChain.IsEmpty())
+	{
+		FPropertyChangedChainEvent ChangedEvent{EditPropertyChain, FPropertyChangedEvent{EditPropertyChain.GetHead()->GetValue(), EPropertyChangeType::ValueSet}};
+		TargetObject->PostEditChangeChainProperty(ChangedEvent);
+	}
+#endif // WITH_EDITOR
+}
+
 bool FPCGObjectSingleOverride::Apply(int32 InputKeyIndex, IPCGAttributeAccessorKeys& OutputKey)
 {
 	return Invoke(ObjectOverrideFunction, this, InputKeyIndex, OutputKey);
diff --git a/Engine/Plugins/PCG/Source/PCG/Public/Elements/PCGApplyOnActor.h b/Engine/Plugins/PCG/Source/PCG/Public/Elements/PCGApplyOnActor.h
--- a/Engine/Plugins/PCG/Source/PCG/Public/Elements/PCGApplyOnActor.h
+++ b/Engine/Plugins/PCG/Source/PCG/Public/Elements/PCGApplyOnActor.h
@@ -59,6 +59,12 @@ public:
 	UPROPERTY(BlueprintReadWrite, EditAnywhere, Category = Settings, AdvancedDisplay)
 	bool bSilenceErrorOnEmptyObjectPath = false;
 
+#if WITH_EDITORONLY_DATA
+	/** Overrides can propagate the change to the object if it is necessary. May trigger expensive downstream computation or infinite refresh loop if PCG listens to changes to this object. Only works in Editor. */
+	UPROPERTY(BlueprintReadWrite, EditAnywhere, Category = "Settings", AdvancedDisplay)
+	bool bPropagateObjectChangeEvent = false;
+#endif // WITH_EDITORONLY_DATA
+
 	/** By default, object loading is asynchronous, can force it synchronous if needed. */
 	UPROPERTY(BlueprintReadWrite, EditAnywhere, Category = "Settings|Debug")
 	bool bSynchronousLoad = false;
diff --git a/Engine/Plugins/PCG/Source/PCG/Public/Metadata/PCGObjectPropertyOverride.h b/Engine/Plugins/PCG/Source/PCG/Public/Metadata/PCGObjectPropertyOverride.h
--- a/Engine/Plugins/PCG/Source/PCG/Public/Metadata/PCGObjectPropertyOverride.h
+++ b/Engine/Plugins/PCG/Source/PCG/Public/Metadata/PCGObjectPropertyOverride.h
@@ -52,9 +52,9 @@ namespace PCGObjectPropertyOverrideHelpers
 	PCG_API FPCGPinProperties CreateObjectPropertiesOverridePin(FName Label, const FText& Tooltip);
 
 	/** Apply property overrides to the TargetObject directly from the ObjectPropertiesOverride pin. Use CreateObjectPropertiesOverridePin(). */
-	PCG_API void ApplyOverridesFromParams(const TArray<FPCGObjectPropertyOverrideDescription>& InObjectPropertyOverrideDescriptions, UObject* TargetObject, FName OverridesPinLabel, FPCGContext* Context);
+	PCG_API void ApplyOverridesFromParams(const TArray<FPCGObjectPropertyOverrideDescription>& InObjectPropertyOverrideDescriptions, UObject* TargetObject, FName OverridesPinLabel, FPCGContext* Context, bool bPropagateEditChangeEvent = false);
 
-	PCG_API void ApplyOverrides(const TArray<FPCGObjectPropertyOverrideDescription>& InObjectPropertyOverrideDescriptions, const TArray<TPair<UObject*, int32>>& TargetObjectAndIndex, FName OverridesPinLabel, int32 InputDataIndex, FPCGContext* Context);
+	PCG_API void ApplyOverrides(const TArray<FPCGObjectPropertyOverrideDescription>& InObjectPropertyOverrideDescriptions, const TArray<TPair<UObject*, int32>>& TargetObjectAndIndex, FName OverridesPinLabel, int32 InputDataIndex, FPCGContext* Context, bool bPropagateEditChangeEvent = false);
 }
 
 /**
@@ -65,15 +65,28 @@ namespace PCGObjectPropertyOverrideHelpers
 */
 struct FPCGObjectSingleOverride
 {
-	/** Initialize the single object override. Call before using Apply(InputKeyIndex, OutputKey). */
-	PCG_API void Initialize(const FPCGAttributePropertySelector& InputSelector, const FString& OutputProperty, const UStruct* TemplateClass, const UPCGData* SourceData, FPCGContext* Context);
+	// Can't copy but can move
+	FPCGObjectSingleOverride() = default;
+	FPCGObjectSingleOverride(const FPCGObjectSingleOverride&) = delete;
+	FPCGObjectSingleOverride(FPCGObjectSingleOverride&&) = default;
+	FPCGObjectSingleOverride& operator=(const FPCGObjectSingleOverride&) = delete;
+	FPCGObjectSingleOverride& operator=(FPCGObjectSingleOverride&&) = default;
+
+	/** Initialize the single object override. Call before using Apply(InputKeyIndex, OutputKey). Can also compute the edit property chain if it is needed to propagate the change. Editor only. */
+	PCG_API void Initialize(const FPCGAttributePropertySelector& InputSelector, const FString& OutputProperty, const UStruct* TemplateClass, const UPCGData* SourceData, FPCGContext* Context, bool bComputePropertyEditChain = false);
 
 	/** Returns true if initialization succeeded in creating the accessors and accessor keys. */
 	PCG_API bool IsValid() const;
 
+	/** To be called before the Apply, if it is needed to propagate the change to the target object. */
+	PCG_API void PreApply(UObject* TargetObject);
+
 	/** Applies a single property override to the object by reading from the InputAccessor at the given KeyIndex, and writing to the OutputKey which represents the object property. */
 	PCG_API bool Apply(int32 InputKeyIndex, IPCGAttributeAccessorKeys& OutputKey);
 
+	/** To be called after the Apply, if it is needed to propagate the change to the target object. */
+	PCG_API void PostApply(UObject* TargetObject);
+
 	/** Gathers overrides into an array, if they need loading. */
 	PCG_API void GatherAllOverridesToLoad(TArray<FSoftObjectPath>& OutObjectsToLoad) const;
 
@@ -82,6 +95,11 @@ private:
 	TUniquePtr<const IPCGAttributeAccessor> ObjectOverrideInputAccessor;
 	TUniquePtr<IPCGAttributeAccessor> ObjectOverrideOutputAccessor;
 
+#if WITH_EDITOR
+	// To be propagated to the object if requested, allow the object to react to those overrides.
+	FEditPropertyChain EditPropertyChain;
+#endif // WITH_EDITOR
+
 	bool bWillNeedLoading = false;
 
 	// InputKeyIndex, OutputKeys
@@ -120,8 +138,14 @@ struct FPCGObjectOverrides
 	FPCGObjectOverrides(T* TemplateObject) : OutputKey(TemplateObject)
 	{}
 
+	// Can't copy but can move
+	FPCGObjectOverrides(const FPCGObjectOverrides&) = delete;
+	FPCGObjectOverrides(FPCGObjectOverrides&&) = default;
+	FPCGObjectOverrides& operator=(const FPCGObjectOverrides&) = delete;
+	FPCGObjectOverrides& operator=(FPCGObjectOverrides&&) = default;
+
 	/** Initialize the object overrides. Call before using Apply(InputKeyIndex). */
-	void Initialize(const TArray<FPCGObjectPropertyOverrideDescription>& OverrideDescriptions, T* TemplateObject, const UPCGData* SourceData, FPCGContext* Context)
+	void Initialize(const TArray<FPCGObjectPropertyOverrideDescription>& OverrideDescriptions, T* TemplateObject, const UPCGData* SourceData, FPCGContext* Context, bool bInPropagateEditChangeEvent = false)
 	{
 		bInitialized = false;
 
@@ -131,12 +155,15 @@ struct FPCGObjectOverrides
 			return;
 		}
 
+		CachedOutputObject = nullptr; 
 		OutputKey = FPCGAttributeAccessorKeysSingleObjectPtr<T>(TemplateObject);
+		bPropagateEditChangeEvent = bInPropagateEditChangeEvent;
 
 		const UStruct* ClassObject = nullptr;
 		if constexpr (TIsDerivedFrom<T, UObject>::Value)
 		{
 			ClassObject = TemplateObject->GetClass();
+			CachedOutputObject = bPropagateEditChangeEvent ? TemplateObject : nullptr;
 		}
 		else
 		{
@@ -150,12 +177,12 @@ struct FPCGObjectOverrides
 			const FPCGAttributePropertyInputSelector InputSelector = OverrideDescriptions[i].InputSource.CopyAndFixLast(SourceData);
 			const FString& OutputProperty = OverrideDescriptions[i].PropertyTarget;
 
-			FPCGObjectSingleOverride Override;
-			Override.Initialize(InputSelector, OutputProperty, ClassObject, SourceData, Context);
+			FPCGObjectSingleOverride Override{};
+			Override.Initialize(InputSelector, OutputProperty, ClassObject, SourceData, Context, bPropagateEditChangeEvent);
 
 			if (Override.IsValid())
 			{
-				ObjectSingleOverrides.Add(std::move(Override));
+				ObjectSingleOverrides.Emplace(MoveTemp(Override));
 			}
 			else
 			{
@@ -166,14 +193,31 @@ struct FPCGObjectOverrides
 		bInitialized = true;
 	}
 
-	/** Applies each property override to the object by reading from the InputAccessor at the given KeyIndex, and writing to the OutputKey which represents the object property. */
+	/**
+	 * Applies each property override to the object by reading from the InputAccessor at the given KeyIndex, and writing to the OutputKey which represents the object property.
+	 * Can also propagate an edit change to the target object if we need to propagate the information. Editor only.
+	 */
 	bool Apply(int32 InputKeyIndex)
 	{
 		bool bAllSucceeded = true;
 
 		for (FPCGObjectSingleOverride& ObjectSingleOverride : ObjectSingleOverrides)
 		{
+#if WITH_EDITOR
+			if (bPropagateEditChangeEvent && CachedOutputObject)
+			{
+				ObjectSingleOverride.PreApply(CachedOutputObject);
+			}
+#endif // WITH_EDITOR
+
 			bAllSucceeded &= ObjectSingleOverride.Apply(InputKeyIndex, OutputKey);
+
+#if WITH_EDITOR
+			if (bPropagateEditChangeEvent && CachedOutputObject)
+			{
+				ObjectSingleOverride.PostApply(CachedOutputObject);
+			}
+#endif // WITH_EDITOR
 		}
 
 		return bAllSucceeded;
@@ -185,6 +229,10 @@ struct FPCGObjectOverrides
 		if (IsValid())
 		{
 			OutputKey = FPCGAttributeAccessorKeysSingleObjectPtr<T>(InTemplateObject);
+			if constexpr (TIsDerivedFrom<T, UObject>::Value)
+			{
+				CachedOutputObject = bPropagateEditChangeEvent ? InTemplateObject : nullptr;
+			}
 		}
 	}
 
@@ -201,6 +249,8 @@ struct FPCGObjectOverrides
 	}
 
 private:
+	bool bPropagateEditChangeEvent = false;
+	UObject* CachedOutputObject = nullptr;
 	FPCGAttributeAccessorKeysSingleObjectPtr<T> OutputKey;
 	TArray<FPCGObjectSingleOverride> ObjectSingleOverrides;
 	bool bInitialized = false;
