[PCG] Fix static analysis issue

#rnx

diff --git a/Engine/Plugins/PCG/Source/PCG/Private/Metadata/PCGObjectPropertyOverride.cpp b/Engine/Plugins/PCG/Source/PCG/Private/Metadata/PCGObjectPropertyOverride.cpp
--- a/Engine/Plugins/PCG/Source/PCG/Private/Metadata/PCGObjectPropertyOverride.cpp
+++ b/Engine/Plugins/PCG/Source/PCG/Private/Metadata/PCGObjectPropertyOverride.cpp
@@ -127,8 +127,6 @@ void FPCGObjectSingleOverride::Initialize(const FPCGAttributePropertySelector& I
 			TArray<const FStructProperty*> EncounteredStructProps;
 			bWillNeedLoading = PropertyChain.Last()->ContainsObjectReference(EncounteredStructProps, EPropertyObjectReferenceType::Strong);
 
-			ObjectOverrideOutputAccessor = PCGAttributeAccessorHelpers::CreatePropertyChainAccessor(std::move(PropertyChain));
-
 #if WITH_EDITOR
 			if (bComputePropertyEditChain)
 			{
@@ -138,6 +136,8 @@ void FPCGObjectSingleOverride::Initialize(const FPCGAttributePropertySelector& I
 				}
 			}
 #endif // WITH_EDITOR
+
+			ObjectOverrideOutputAccessor = PCGAttributeAccessorHelpers::CreatePropertyChainAccessor(MoveTemp(PropertyChain));
 		}
 	}
 
