[PCG] Unify data desc element counts into a single vector 4

#rnx

diff --git a/Engine/Plugins/PCG/Source/PCG/Internal/Compute/Data/PCGProxyForGPUData.h b/Engine/Plugins/PCG/Source/PCG/Internal/Compute/Data/PCGProxyForGPUData.h
--- a/Engine/Plugins/PCG/Source/PCG/Internal/Compute/Data/PCGProxyForGPUData.h
+++ b/Engine/Plugins/PCG/Source/PCG/Internal/Compute/Data/PCGProxyForGPUData.h
@@ -49,7 +49,7 @@ public:
 	UE_API FReadbackResult GetCPUData(FPCGContext* InContext) const;
 
 	/** Returns the element count for this data. Does not trigger a GPU->CPU readback. */
-	UE_API int GetElementCount() const;
+	UE_API FIntVector4 GetElementCount() const;
 
 	/** Returns a description of this data. Does not trigger a GPU->CPU readback. */
 	UE_API bool GetDescription(FPCGDataDesc& OutDescription) const;
diff --git a/Engine/Plugins/PCG/Source/PCG/Internal/Compute/PCGComputeCommon.h b/Engine/Plugins/PCG/Source/PCG/Internal/Compute/PCGComputeCommon.h
--- a/Engine/Plugins/PCG/Source/PCG/Internal/Compute/PCGComputeCommon.h
+++ b/Engine/Plugins/PCG/Source/PCG/Internal/Compute/PCGComputeCommon.h
@@ -48,6 +48,16 @@ enum class EPCGExportMode : uint8
 };
 ENUM_CLASS_FLAGS(EPCGExportMode);
 
+/** Dimensionality of a data description element count. */
+UENUM()
+enum class EPCGElementDimension : uint8
+{
+	One,
+	Two,
+	Three,
+	Four
+};
+
 namespace PCGComputeConstants
 {
 	constexpr int MAX_NUM_ATTRS = 128;
@@ -93,10 +103,11 @@ namespace PCGComputeConstants
 namespace PCGComputeHelpers
 {
 	/** Gets the element count for a given data. E.g. number of points in a PointData, number of metadata entries in a ParamData, etc. */
-	int GetElementCount(const UPCGData* InData);
+	FIntVector4 GetElementCount(const UPCGData* InData);
 
-	/** Gets the 2D element count for a given data. E.g. texture size for UPCGTextureData. */
-	FIntPoint GetElementCount2D(const UPCGData* InData);
+	/** Gets the element dimension for a given data. E.g. OneD for PointData, TwoD for TextureData, etc. */
+	EPCGElementDimension GetElementDimension(const UPCGData* InData);
+	EPCGElementDimension GetElementDimension(EPCGDataType InDataType);
 
 	/** True if 'Type' is valid on a GPU input pin. */
 	bool IsTypeAllowedAsInput(EPCGDataType Type);
diff --git a/Engine/Plugins/PCG/Source/PCG/Internal/Compute/PCGDataForGPU.h b/Engine/Plugins/PCG/Source/PCG/Internal/Compute/PCGDataForGPU.h
--- a/Engine/Plugins/PCG/Source/PCG/Internal/Compute/PCGDataForGPU.h
+++ b/Engine/Plugins/PCG/Source/PCG/Internal/Compute/PCGDataForGPU.h
@@ -18,6 +18,7 @@ class UPCGMetadata;
 class UPCGPin;
 class UPCGSettings;
 struct FPCGContext;
+enum class EPCGElementMultiplicity : uint8;
 enum class EPCGMetadataTypes : uint8;
 enum class EPCGPointNativeProperties : uint32;
 
@@ -157,9 +158,10 @@ struct FPCGDataDesc
 {
 public:
 	FPCGDataDesc() = default;
-	explicit FPCGDataDesc(EPCGDataType InType);
 	FPCGDataDesc(EPCGDataType InType, int32 InElementCount);
-	FPCGDataDesc(EPCGDataType InType, int32 InElementCount, FIntPoint InElementCount2D);
+	FPCGDataDesc(EPCGDataType InType, FIntPoint InElementCount);
+	FPCGDataDesc(EPCGDataType InType, FIntVector3 InElementCount);
+	FPCGDataDesc(EPCGDataType InType, FIntVector4 InElementCount);
 	FPCGDataDesc(const FPCGTaggedData& InTaggedData, const UPCGDataBinding* InBinding);
 
 	uint64 ComputePackedSize() const;
@@ -172,9 +174,21 @@ public:
 
 	void AddAttribute(FPCGKernelAttributeKey InAttribute, const UPCGDataBinding* InBinding, const TArray<int32>* InOptionalUniqueStringKeys = nullptr);
 
-	int32 GetElementCountForAttribute(const FPCGKernelAttributeDesc& AttributeDesc) const;
+	EPCGDataType GetType() const { return Type; }
 
-	bool IsDomain2D() const { return !!(Type & EPCGDataType::BaseTexture); }
+	FIntVector4 GetElementCount() const { return ElementCount; }
+	FIntVector4 GetElementCountForAttribute(const FPCGKernelAttributeDesc& AttributeDesc) const;
+	EPCGElementDimension GetElementDimension() const { return ElementDimension; }
+	int32 ComputeTotalElementCount() const;
+
+	void AddElementCount(int32 InElementCountToAdd);
+	void AddElementCount(FIntPoint InElementCountToAdd);
+	void AddElementCount(FIntVector3 InElementCountToAdd);
+	void AddElementCount(FIntVector4 InElementCountToAdd);
+	void ScaleElementCount(int32 InMultiplier) { ElementCount *= InMultiplier; }
+
+	FPCGDataDesc& CombineElementCount(const FPCGDataDesc& Other, EPCGElementMultiplicity Multiplicity);
+	static FPCGDataDesc CombineElementCount(const FPCGDataDesc& A, const FPCGDataDesc& B, EPCGElementMultiplicity Multiplicity);
 
 	void AllocateProperties(EPCGPointNativeProperties InProperties) { AllocatedPointProperties |= InProperties; }
 	EPCGPointNativeProperties GetAllocatedProperties() const { return AllocatedPointProperties; }
@@ -186,14 +200,15 @@ private:
 	void InitializeAttributeDescs(const UPCGData* InData, const UPCGDataBinding* InBinding);
 
 public:
-	EPCGDataType Type = EPCGDataType::Point;
+	// @todo_pcg: Encapsulate these members.
 	TArray<FPCGKernelAttributeDesc> AttributeDescs;
-	// @todo_pcg: These counts should be merged into one FIntVector
-	int32 ElementCount = 0;
-	FIntPoint ElementCount2D = FIntPoint::ZeroValue;
 	TArray<int32, TInlineAllocator<4>> TagStringKeys;
 
 private:
+	EPCGDataType Type = EPCGDataType::Point;
+	FIntVector4 ElementCount = FIntVector4::ZeroValue;
+	EPCGElementDimension ElementDimension = EPCGElementDimension::One;
+
 	/** Properties that should be fully allocated for this data. Only applies to EPCGDataType::Point currently. */
 	EPCGPointNativeProperties AllocatedPointProperties = EPCGPointNativeProperties::None;
 };
@@ -220,10 +235,10 @@ struct FPCGDataCollectionDesc
 	EPCGUnpackDataCollectionResult UnpackDataCollection(FPCGContext* InContext, const TArray<uint8>& InPackedData, FName InPin, const TArray<FString>& InStringTable, FPCGDataCollection& OutDataCollection) const;
 
 	/** Compute total number of processing elements. */
-	uint32 ComputeDataElementCount() const;
+	uint32 ComputeTotalElementCount() const;
 
 	/** Compute total number of processing elements of the given type. */
-	uint32 ComputeDataElementCount(EPCGDataType InDataType) const;
+	uint32 ComputeTotalElementCount(EPCGDataType InDataType) const;
 
 	/** Aggregate another data description. */
 	void Combine(const FPCGDataCollectionDesc& Other);
diff --git a/Engine/Plugins/PCG/Source/PCG/Private/Compute/BuiltInKernels/PCGCountUniqueAttributeValuesKernel.cpp b/Engine/Plugins/PCG/Source/PCG/Private/Compute/BuiltInKernels/PCGCountUniqueAttributeValuesKernel.cpp
--- a/Engine/Plugins/PCG/Source/PCG/Private/Compute/BuiltInKernels/PCGCountUniqueAttributeValuesKernel.cpp
+++ b/Engine/Plugins/PCG/Source/PCG/Private/Compute/BuiltInKernels/PCGCountUniqueAttributeValuesKernel.cpp
@@ -110,19 +110,18 @@ bool UPCGCountUniqueAttributeValuesKernel::ComputeOutputBindingDataDesc(const UP
 
 	if (AttributeId != INDEX_NONE)
 	{
-		// Base data description object, copied N times below to populate the final description.
-		FPCGDataDesc CountDataDesc;
-		CountDataDesc.Type = EPCGDataType::Param;
-		CountDataDesc.AddAttribute(FPCGKernelAttributeKey(PCGCountUniqueAttributeValuesConstants::ValueAttributeName, EPCGKernelAttributeType::Int), InBinding);
-		CountDataDesc.AddAttribute(FPCGKernelAttributeKey(PCGCountUniqueAttributeValuesConstants::ValueCountAttributeName, EPCGKernelAttributeType::Int), InBinding);
-
 		// Compute a set of (value, count) pairs across all data.
 		TArray<int32> UniqueStringKeyValues;
 		InputDesc.GetUniqueStringKeyValues(AttributeId, UniqueStringKeyValues);
 
 		// If the highest string key value is 3, we'll allocate 4 counters. String key value 0 is reserved for empty/null string.
 		const int32* MaxStringKeyValue = Algo::MaxElement(UniqueStringKeyValues);
-		CountDataDesc.ElementCount = (MaxStringKeyValue ? *MaxStringKeyValue : 0) + 1;
+		const int32 ElementCount = (MaxStringKeyValue ? *MaxStringKeyValue : 0) + 1;
+
+		// Base data description object, copied N times below to populate the final description.
+		FPCGDataDesc CountDataDesc(EPCGDataType::Param, ElementCount);
+		CountDataDesc.AddAttribute(FPCGKernelAttributeKey(PCGCountUniqueAttributeValuesConstants::ValueAttributeName, EPCGKernelAttributeType::Int), InBinding);
+		CountDataDesc.AddAttribute(FPCGKernelAttributeKey(PCGCountUniqueAttributeValuesConstants::ValueCountAttributeName, EPCGKernelAttributeType::Int), InBinding);
 
 		if (bEmitPerDataCounts)
 		{
@@ -155,7 +154,7 @@ int UPCGCountUniqueAttributeValuesKernel::ComputeThreadCount(const UPCGDataBindi
 {
 	if (const FPCGDataCollectionDesc* InputPinDesc = InBinding->GetCachedKernelPinDataDesc(this, PCGPinConstants::DefaultInputLabel, /*bIsInput=*/true); ensure(InputPinDesc))
 	{
-		const int ElementCount = InputPinDesc->ComputeDataElementCount(EPCGDataType::Point);
+		const int ElementCount = InputPinDesc->ComputeTotalElementCount();
 
 		const int32 PartitionAttributeId = InBinding->GetAttributeId(AttributeName, EPCGKernelAttributeType::StringKey);
 		const int ValueCount = (PartitionAttributeId != INDEX_NONE) ? InputPinDesc->GetNumStringKeyValues(PartitionAttributeId) : 0;
diff --git a/Engine/Plugins/PCG/Source/PCG/Private/Compute/Data/PCGProxyForGPUData.cpp b/Engine/Plugins/PCG/Source/PCG/Private/Compute/Data/PCGProxyForGPUData.cpp
--- a/Engine/Plugins/PCG/Source/PCG/Private/Compute/Data/PCGProxyForGPUData.cpp
+++ b/Engine/Plugins/PCG/Source/PCG/Private/Compute/Data/PCGProxyForGPUData.cpp
@@ -33,7 +33,7 @@ EPCGDataType UPCGProxyForGPUData::GetUnderlyingDataType() const
 {
 	if (TSharedPtr<const FPCGProxyForGPUDataCollection> DataCollectionGPU = GetGPUInfo())
 	{
-		return DataCollectionGPU->GetDescription().DataDescs[DataIndexInCollection].Type;
+		return DataCollectionGPU->GetDescription().DataDescs[DataIndexInCollection].GetType();
 	}
 
 	return EPCGDataType::None;
@@ -44,14 +44,14 @@ TSharedPtr<const FPCGProxyForGPUDataCollection> UPCGProxyForGPUData::GetInputDat
 	return DataCollectionOnGPU;
 }
 
-int UPCGProxyForGPUData::GetElementCount() const
+FIntVector4 UPCGProxyForGPUData::GetElementCount() const
 {
 	if (TSharedPtr<const FPCGProxyForGPUDataCollection> DataCollectionGPU = GetGPUInfo())
 	{
-		return DataCollectionGPU->GetDescription().DataDescs[DataIndexInCollection].ElementCount;
+		return DataCollectionGPU->GetDescription().DataDescs[DataIndexInCollection].GetElementCount();
 	}
 
-	return 0;
+	return FIntVector4::ZeroValue;
 }
 
 bool UPCGProxyForGPUData::GetDescription(FPCGDataDesc& OutDescription) const
diff --git a/Engine/Plugins/PCG/Source/PCG/Private/Compute/DataInterfaces/Elements/PCGSkinnedMeshSpawnerDataInterface.cpp b/Engine/Plugins/PCG/Source/PCG/Private/Compute/DataInterfaces/Elements/PCGSkinnedMeshSpawnerDataInterface.cpp
--- a/Engine/Plugins/PCG/Source/PCG/Private/Compute/DataInterfaces/Elements/PCGSkinnedMeshSpawnerDataInterface.cpp
+++ b/Engine/Plugins/PCG/Source/PCG/Private/Compute/DataInterfaces/Elements/PCGSkinnedMeshSpawnerDataInterface.cpp
@@ -320,7 +320,7 @@ bool UPCGSkinnedMeshSpawnerDataProvider::PrepareForExecute_GameThread(UPCGDataBi
 
 		for (const FPCGDataDesc& Desc : InputDataDesc->DataDescs)
 		{
-			if (Desc.ElementCount <= 0)
+			if (Desc.GetElementCount().X <= 0)
 			{
 				continue;
 			}
@@ -470,7 +470,7 @@ void UPCGSkinnedMeshSpawnerDataProvider::CreatePrimitiveDescriptors(FPCGContext*
 		return;
 	}
 
-	const int32 TotalInputPointCount = InputDataDesc->ComputeDataElementCount(EPCGDataType::Point);
+	const int32 TotalInputPointCount = InputDataDesc->ComputeTotalElementCount();
 	if (TotalInputPointCount == 0)
 	{
 		return;
@@ -519,7 +519,7 @@ void UPCGSkinnedMeshSpawnerDataProvider::CreatePrimitiveDescriptors(FPCGContext*
 		{
 			for (const FPCGDataDesc& Desc : InputDataDesc->DataDescs)
 			{
-				if (Desc.ElementCount <= 0)
+				if (Desc.GetElementCount().X <= 0)
 				{
 					continue;
 				}
@@ -534,7 +534,7 @@ void UPCGSkinnedMeshSpawnerDataProvider::CreatePrimitiveDescriptors(FPCGContext*
 							{
 								if (AttributeStringKeyValue > 0)
 								{
-									StringKeyToInstanceCount.FindOrAdd(AttributeStringKeyValue, 0) += Desc.ElementCount;
+									StringKeyToInstanceCount.FindOrAdd(AttributeStringKeyValue, 0) += Desc.GetElementCount().X;
 								}
 							}
 
diff --git a/Engine/Plugins/PCG/Source/PCG/Private/Compute/DataInterfaces/Elements/PCGStaticMeshSpawnerDataInterface.cpp b/Engine/Plugins/PCG/Source/PCG/Private/Compute/DataInterfaces/Elements/PCGStaticMeshSpawnerDataInterface.cpp
--- a/Engine/Plugins/PCG/Source/PCG/Private/Compute/DataInterfaces/Elements/PCGStaticMeshSpawnerDataInterface.cpp
+++ b/Engine/Plugins/PCG/Source/PCG/Private/Compute/DataInterfaces/Elements/PCGStaticMeshSpawnerDataInterface.cpp
@@ -321,7 +321,7 @@ bool UPCGStaticMeshSpawnerDataProvider::PrepareForExecute_GameThread(UPCGDataBin
 
 		for (const FPCGDataDesc& Desc : InputDataDesc->DataDescs)
 		{
-			if (Desc.ElementCount <= 0)
+			if (Desc.GetElementCount().X <= 0)
 			{
 				continue;
 			}
@@ -479,7 +479,7 @@ void UPCGStaticMeshSpawnerDataProvider::CreatePrimitiveDescriptors(FPCGContext*
 		return;
 	}
 
-	const int32 TotalInputPointCount = InputDataDesc->ComputeDataElementCount(EPCGDataType::Point);
+	const int32 TotalInputPointCount = InputDataDesc->ComputeTotalElementCount();
 	if (TotalInputPointCount == 0)
 	{
 		return;
@@ -527,7 +527,7 @@ void UPCGStaticMeshSpawnerDataProvider::CreatePrimitiveDescriptors(FPCGContext*
 		{
 			for (const FPCGDataDesc& Desc : InputDataDesc->DataDescs)
 			{
-				if (Desc.ElementCount <= 0)
+				if (Desc.GetElementCount().X <= 0)
 				{
 					continue;
 				}
@@ -542,7 +542,7 @@ void UPCGStaticMeshSpawnerDataProvider::CreatePrimitiveDescriptors(FPCGContext*
 							{
 								if (AttributeStringKeyValue > 0)
 								{
-									StringKeyToInstanceCount.FindOrAdd(AttributeStringKeyValue, 0) += Desc.GetElementCountForAttribute(AttributeDesc);
+									StringKeyToInstanceCount.FindOrAdd(AttributeStringKeyValue, 0) += Desc.GetElementCountForAttribute(AttributeDesc).X;
 								}
 							}
 
diff --git a/Engine/Plugins/PCG/Source/PCG/Private/Compute/DataInterfaces/Elements/PCGTransformPointsDataInterface.cpp b/Engine/Plugins/PCG/Source/PCG/Private/Compute/DataInterfaces/Elements/PCGTransformPointsDataInterface.cpp
--- a/Engine/Plugins/PCG/Source/PCG/Private/Compute/DataInterfaces/Elements/PCGTransformPointsDataInterface.cpp
+++ b/Engine/Plugins/PCG/Source/PCG/Private/Compute/DataInterfaces/Elements/PCGTransformPointsDataInterface.cpp
@@ -198,7 +198,7 @@ bool UPCGTransformPointsDataProvider::PrepareForExecute_GameThread(UPCGDataBindi
 
 		for (const FPCGDataDesc& Desc : InputDataDesc->DataDescs)
 		{
-			if (Desc.ElementCount <= 0)
+			if (Desc.GetElementCount().X <= 0)
 			{
 				continue;
 			}
diff --git a/Engine/Plugins/PCG/Source/PCG/Private/Compute/DataInterfaces/PCGTextureDataInterface.cpp b/Engine/Plugins/PCG/Source/PCG/Private/Compute/DataInterfaces/PCGTextureDataInterface.cpp
--- a/Engine/Plugins/PCG/Source/PCG/Private/Compute/DataInterfaces/PCGTextureDataInterface.cpp
+++ b/Engine/Plugins/PCG/Source/PCG/Private/Compute/DataInterfaces/PCGTextureDataInterface.cpp
@@ -79,7 +79,7 @@ FPCGTextureBindingInfo::FPCGTextureBindingInfo(const UPCGBaseTextureData* InText
 FPCGTextureBindingInfo::FPCGTextureBindingInfo(const FPCGDataDesc& InDataDesc, const FTransform& InTransform)
 {
 	// @todo_pcg: Need to support other information here too, like format, slice index, resource type, etc.
-	Size = FIntPoint(FMath::Max(InDataDesc.ElementCount2D.X, 1), FMath::Max(InDataDesc.ElementCount2D.Y, 1));
+	Size = FIntPoint(FMath::Max(InDataDesc.GetElementCount().X, 1), FMath::Max(InDataDesc.GetElementCount().Y, 1));
 	Transform = InTransform;
 	TextureBounds = FBox(FVector(-1.0f, -1.0f, 0.0f), FVector(1.0f, 1.0f, 0.0f)).TransformBy(Transform);
 }
diff --git a/Engine/Plugins/PCG/Source/PCG/Private/Compute/Elements/PCGCullPointsOutsideActorBoundsKernel.cpp b/Engine/Plugins/PCG/Source/PCG/Private/Compute/Elements/PCGCullPointsOutsideActorBoundsKernel.cpp
--- a/Engine/Plugins/PCG/Source/PCG/Private/Compute/Elements/PCGCullPointsOutsideActorBoundsKernel.cpp
+++ b/Engine/Plugins/PCG/Source/PCG/Private/Compute/Elements/PCGCullPointsOutsideActorBoundsKernel.cpp
@@ -34,7 +34,7 @@ bool UPCGCullPointsOutsideActorBoundsKernel::ComputeOutputBindingDataDesc(const
 int UPCGCullPointsOutsideActorBoundsKernel::ComputeThreadCount(const UPCGDataBinding* InBinding) const
 {
 	const FPCGDataCollectionDesc* OutputPinDesc = InBinding->GetCachedKernelPinDataDesc(this, PCGPinConstants::DefaultOutputLabel, /*bIsInput=*/false);
-	return ensure(OutputPinDesc) ? OutputPinDesc->ComputeDataElementCount() : 0;
+	return ensure(OutputPinDesc) ? OutputPinDesc->ComputeTotalElementCount() : 0;
 }
 
 #if WITH_EDITOR
diff --git a/Engine/Plugins/PCG/Source/PCG/Private/Compute/Elements/PCGCustomHLSLKernel.cpp b/Engine/Plugins/PCG/Source/PCG/Private/Compute/Elements/PCGCustomHLSLKernel.cpp
--- a/Engine/Plugins/PCG/Source/PCG/Private/Compute/Elements/PCGCustomHLSLKernel.cpp
+++ b/Engine/Plugins/PCG/Source/PCG/Private/Compute/Elements/PCGCustomHLSLKernel.cpp
@@ -266,9 +266,7 @@ bool UPCGCustomHLSLKernel::ComputeOutputBindingDataDesc(const UPCGComputeGraph*
 	else if (InOutputPinLabel == FirstOutputPinLabel && CustomHLSLSettings->KernelType == EPCGKernelType::TextureGenerator)
 	{
 		// Generators always produce a single texture data with known size.
-		FPCGDataDesc& NewDataDesc = OutputDataDesc.DataDescs.Emplace_GetRef();
-		NewDataDesc.Type = EPCGDataType::BaseTexture;
-		NewDataDesc.ElementCount2D = CustomHLSLSettings->NumElements2D;
+		OutputDataDesc.DataDescs.Emplace(EPCGDataType::BaseTexture, FIntVector4(CustomHLSLSettings->NumElements2D.X, CustomHLSLSettings->NumElements2D.Y, 0, 0));
 	}
 	else if (ensure(OutputPinProperties))
 	{
@@ -623,7 +621,7 @@ int UPCGCustomHLSLKernel::ComputeThreadCount(const UPCGDataBinding* InBinding) c
 			if (const FPCGPinPropertiesGPU* OutputPin = CustomHLSLSettings->OutputPins.IsEmpty() ? nullptr : &CustomHLSLSettings->OutputPins[0])
 			{
 				const FPCGDataCollectionDesc* Desc = InBinding->GetCachedKernelPinDataDesc(this, OutputPin->Label, /*bIsInput=*/false);
-				ThreadCount = Desc ? OutputPin->GetElementCountMultiplier() * Desc->ComputeDataElementCount(OutputPin->AllowedTypes) : 0;
+				ThreadCount = Desc ? OutputPin->GetElementCountMultiplier() * Desc->ComputeTotalElementCount(OutputPin->AllowedTypes) : 0;
 			}
 		}
 		else if (CustomHLSLSettings->DispatchThreadCount == EPCGDispatchThreadCount::FromProductOfInputPins)
@@ -1910,7 +1908,7 @@ int UPCGCustomHLSLKernel::GetElementCountForInputPin(const FPCGPinProperties& In
 
 	const FPCGDataCollectionDesc* InputDesc = InBinding->GetCachedKernelPinDataDesc(this, InInputPinProps.Label, /*bIsInputPin=*/true);
 
-	return InputDesc ? InputDesc->ComputeDataElementCount(InInputPinProps.AllowedTypes) : 0;
+	return InputDesc ? InputDesc->ComputeTotalElementCount(InInputPinProps.AllowedTypes) : 0;
 }
 
 #undef LOCTEXT_NAMESPACE
diff --git a/Engine/Plugins/PCG/Source/PCG/Private/Compute/Elements/PCGNormalToDensityKernel.cpp b/Engine/Plugins/PCG/Source/PCG/Private/Compute/Elements/PCGNormalToDensityKernel.cpp
--- a/Engine/Plugins/PCG/Source/PCG/Private/Compute/Elements/PCGNormalToDensityKernel.cpp
+++ b/Engine/Plugins/PCG/Source/PCG/Private/Compute/Elements/PCGNormalToDensityKernel.cpp
@@ -43,7 +43,7 @@ bool UPCGNormalToDensityKernel::ComputeOutputBindingDataDesc(const UPCGComputeGr
 int UPCGNormalToDensityKernel::ComputeThreadCount(const UPCGDataBinding* InBinding) const
 {
 	const FPCGDataCollectionDesc* OutputPinDesc = InBinding->GetCachedKernelPinDataDesc(this, PCGPinConstants::DefaultOutputLabel, /*bIsInput=*/false);
-	return ensure(OutputPinDesc) ? OutputPinDesc->ComputeDataElementCount() : 0;
+	return ensure(OutputPinDesc) ? OutputPinDesc->ComputeTotalElementCount() : 0;
 }
 
 #if WITH_EDITOR
diff --git a/Engine/Plugins/PCG/Source/PCG/Private/Compute/Elements/PCGTransformPointsKernel.cpp b/Engine/Plugins/PCG/Source/PCG/Private/Compute/Elements/PCGTransformPointsKernel.cpp
--- a/Engine/Plugins/PCG/Source/PCG/Private/Compute/Elements/PCGTransformPointsKernel.cpp
+++ b/Engine/Plugins/PCG/Source/PCG/Private/Compute/Elements/PCGTransformPointsKernel.cpp
@@ -44,7 +44,7 @@ bool UPCGTransformPointsKernel::ComputeOutputBindingDataDesc(const UPCGComputeGr
 int UPCGTransformPointsKernel::ComputeThreadCount(const UPCGDataBinding* InBinding) const
 {
 	const FPCGDataCollectionDesc* OutputPinDesc = InBinding->GetCachedKernelPinDataDesc(this, PCGPinConstants::DefaultOutputLabel, /*bIsInput=*/false);
-	return ensure(OutputPinDesc) ? OutputPinDesc->ComputeDataElementCount() : 0;
+	return ensure(OutputPinDesc) ? OutputPinDesc->ComputeTotalElementCount() : 0;
 }
 
 #if WITH_EDITOR
diff --git a/Engine/Plugins/PCG/Source/PCG/Private/Compute/PCGComputeCommon.cpp b/Engine/Plugins/PCG/Source/PCG/Private/Compute/PCGComputeCommon.cpp
--- a/Engine/Plugins/PCG/Source/PCG/Private/Compute/PCGComputeCommon.cpp
+++ b/Engine/Plugins/PCG/Source/PCG/Private/Compute/PCGComputeCommon.cpp
@@ -53,35 +53,43 @@ namespace PCGComputeHelpers
 		0.5f,
 		TEXT("Maximum GPU buffer size as proportion of total available graphics memory."));
 
-	int GetElementCount(const UPCGData* InData)
+	FIntVector4 GetElementCount(const UPCGData* InData)
 	{
+		FIntVector4 ElementCount = FIntVector4::ZeroValue;
+
 		if (const UPCGBasePointData* PointData = Cast<UPCGBasePointData>(InData))
 		{
-			return PointData->GetNumPoints();
+			ElementCount.X = PointData->GetNumPoints();
 		}
 		else if (const UPCGParamData* ParamData = Cast<UPCGParamData>(InData))
 		{
 			if (const UPCGMetadata* Metadata = ParamData->ConstMetadata())
 			{
-				return Metadata->GetItemCountForChild();
+				ElementCount.X = Metadata->GetItemCountForChild();
 			}
 		}
 		else if (const UPCGProxyForGPUData* Proxy = Cast<UPCGProxyForGPUData>(InData))
 		{
-			return Proxy->GetElementCount();
+			ElementCount = Proxy->GetElementCount();
+		}
+		else if (const UPCGBaseTextureData* TextureData = Cast<UPCGBaseTextureData>(InData))
+		{
+			const FIntPoint TextureSize = TextureData->GetTextureSize();
+			ElementCount.X = TextureSize.X;
+			ElementCount.Y = TextureSize.Y;
 		}
 
-		return 0;
+		return ElementCount;
 	}
 
-	FIntPoint GetElementCount2D(const UPCGData* InData)
+	EPCGElementDimension GetElementDimension(const UPCGData* InData)
 	{
-		if (const UPCGBaseTextureData* TextureData = Cast<UPCGBaseTextureData>(InData))
-		{
-			return TextureData->GetTextureSize();
-		}
+		return InData ? GetElementDimension(InData->GetDataType()) : EPCGElementDimension::One;
+	}
 
-		return FIntPoint::ZeroValue;
+	EPCGElementDimension GetElementDimension(EPCGDataType InDataType)
+	{
+		return !!(InDataType & EPCGDataType::BaseTexture) ? EPCGElementDimension::Two : EPCGElementDimension::One;
 	}
 
 	bool IsTypeAllowedAsInput(EPCGDataType Type)
diff --git a/Engine/Plugins/PCG/Source/PCG/Private/Compute/PCGComputeGraph.cpp b/Engine/Plugins/PCG/Source/PCG/Private/Compute/PCGComputeGraph.cpp
--- a/Engine/Plugins/PCG/Source/PCG/Private/Compute/PCGComputeGraph.cpp
+++ b/Engine/Plugins/PCG/Source/PCG/Private/Compute/PCGComputeGraph.cpp
@@ -211,17 +211,8 @@ void UPCGComputeGraph::DebugLogDataDescriptions(const UPCGDataBinding* InBinding
 			const FPCGDataDesc& DataDesc = InDataDesc->DataDescs[DataIndex];
 
 			UE_LOG(LogPCG, Display, TEXT("\t\tData Index %d"), DataIndex);
-			UE_LOG(LogPCG, Display, TEXT("\t\t\tType: %s"), *PCGDataTypeEnum->GetDisplayNameTextByValue(static_cast<int64>(DataDesc.Type)).ToString());
-			
-			if (DataDesc.IsDomain2D())
-			{
-				UE_LOG(LogPCG, Display, TEXT("\t\t\tNum Elements: (%d, %d)"), DataDesc.ElementCount2D.X, DataDesc.ElementCount2D.Y);
-			}
-			else
-			{
-				UE_LOG(LogPCG, Display, TEXT("\t\t\tNum Elements: %d"), DataDesc.ElementCount);
-			}
-
+			UE_LOG(LogPCG, Display, TEXT("\t\t\tType: %s"), *PCGDataTypeEnum->GetDisplayNameTextByValue(static_cast<int64>(DataDesc.GetType())).ToString());
+			UE_LOG(LogPCG, Display, TEXT("\t\t\tNum Elements: %s (%s)"), *DataDesc.GetElementCount().ToString(), *UEnum::GetValueAsString(DataDesc.GetElementDimension()));
 			UE_LOG(LogPCG, Display, TEXT("\t\t\tAttributes (%d)"), DataDesc.AttributeDescs.Num());
 
 			TSharedPtr<const FPCGProxyForGPUDataCollection> UpstreamDataCollectionInfo;
diff --git a/Engine/Plugins/PCG/Source/PCG/Private/Compute/PCGComputeKernel.cpp b/Engine/Plugins/PCG/Source/PCG/Private/Compute/PCGComputeKernel.cpp
--- a/Engine/Plugins/PCG/Source/PCG/Private/Compute/PCGComputeKernel.cpp
+++ b/Engine/Plugins/PCG/Source/PCG/Private/Compute/PCGComputeKernel.cpp
@@ -275,20 +275,39 @@ void UPCGComputeKernel::ComputeDataDescFromPinProperties(const FPCGPinProperties
 		{
 			for (int DataIndex = 0; DataIndex < MaxDataCount; ++DataIndex)
 			{
-				// Set element count to 0 for now, but we will overwrite it.
-				FPCGDataDesc& DataDesc = OutPinDesc.DataDescs.Emplace_GetRef(OutputPinProps.AllowedTypes, /*SetNumElements=*/0);
-				const bool bDomainIs2D = DataDesc.IsDomain2D();
-
-				int TotalNumElements = 0; // Total number of elements computed for this data index.
-				FIntPoint TotalNumElements2D = FIntPoint::ZeroValue;
+				FIntVector4 InitialNumElements = FIntVector4::ZeroValue; // Total number of elements computed for this data index.
 
 				// If the element count is a product, we have to start at 1 instead of 0 or we'll end up with 0!
 				if (Props.ElementMultiplicity == EPCGElementMultiplicity::Product && NumInitPins > 0)
 				{
-					TotalNumElements = 1;
-					TotalNumElements2D = FIntPoint(1, 1);
+					const EPCGElementDimension ElementDimension = PCGComputeHelpers::GetElementDimension(OutputPinProps.AllowedTypes);
+
+					if (ElementDimension == EPCGElementDimension::One)
+					{
+						InitialNumElements.X = 1;
+					}
+					else if (ElementDimension == EPCGElementDimension::Two)
+					{
+						InitialNumElements.X = 1;
+						InitialNumElements.Y = 1;
+					}
+					else if (ElementDimension == EPCGElementDimension::Three)
+					{
+						InitialNumElements.X = 1;
+						InitialNumElements.Y = 1;
+						InitialNumElements.Z = 1;
+					}
+					else if (ElementDimension == EPCGElementDimension::Four)
+					{
+						InitialNumElements.X = 1;
+						InitialNumElements.Y = 1;
+						InitialNumElements.Z = 1;
+						InitialNumElements.W = 1;
+					}
 				}
 
+				FPCGDataDesc& DataDesc = OutPinDesc.DataDescs.Emplace_GetRef(OutputPinProps.AllowedTypes, InitialNumElements);
+
 				// For each data index, loop over all the pins and create the uber-data.
 				for (int InputPinIndex = 0; InputPinIndex < NumInitPins; ++InputPinIndex)
 				{
@@ -303,42 +322,30 @@ void UPCGComputeKernel::ComputeDataDescFromPinProperties(const FPCGPinProperties
 					}
 
 					const FPCGDataDesc& InputDataDesc = InputDesc.DataDescs.IsValidIndex(ClampedDataIndex) ? InputDesc.DataDescs[ClampedDataIndex] : FPCGDataDesc(EPCGDataType::Any, 0);
-					const bool bSourceDomainIs2D = InputDataDesc.IsDomain2D();
+					const EPCGElementDimension SourceDimension = InputDataDesc.GetElementDimension();
 
 					if (Props.ElementCountMode == EPCGElementCountMode::FromInputData)
 					{
-						if (Props.ElementMultiplicity == EPCGElementMultiplicity::Product)
+						DataDesc.CombineElementCount(InputDataDesc, Props.ElementMultiplicity);
+					}
+					else if (Props.ElementCountMode == EPCGElementCountMode::Fixed)
+					{
+						const EPCGElementDimension ElementDimension = PCGComputeHelpers::GetElementDimension(OutputPinProps.AllowedTypes);
+
+						if (ElementDimension == EPCGElementDimension::One)
 						{
-							if (bDomainIs2D)
-							{
-								TotalNumElements2D *= bSourceDomainIs2D ? (InputDataDesc.ElementCount2D.X * InputDataDesc.ElementCount2D.Y) : InputDataDesc.ElementCount;
-							}
-							else
-							{
-								TotalNumElements *= bSourceDomainIs2D ? (InputDataDesc.ElementCount2D.X * InputDataDesc.ElementCount2D.Y) : InputDataDesc.ElementCount;
-							}
+							DataDesc.AddElementCount(Props.ElementCount);
 						}
-						else if (Props.ElementMultiplicity == EPCGElementMultiplicity::Sum)
+						else if (ElementDimension == EPCGElementDimension::Two)
 						{
-							if (bDomainIs2D)
-							{
-								TotalNumElements2D += bSourceDomainIs2D ? InputDataDesc.ElementCount2D : InputDataDesc.ElementCount;
-							}
-							else
-							{
-								TotalNumElements += bSourceDomainIs2D ? (InputDataDesc.ElementCount2D.X * InputDataDesc.ElementCount2D.Y) : InputDataDesc.ElementCount;
-							}
+							DataDesc.AddElementCount(Props.NumElements2D);
 						}
 						else
 						{
+							// Not supported yet.
 							checkNoEntry();
 						}
 					}
-					else if (Props.ElementCountMode == EPCGElementCountMode::Fixed)
-					{
-						TotalNumElements += Props.ElementCount;
-						TotalNumElements2D += Props.NumElements2D;
-					}
 					else
 					{
 						checkNoEntry();
@@ -357,9 +364,6 @@ void UPCGComputeKernel::ComputeDataDescFromPinProperties(const FPCGPinProperties
 
 					DataDesc.AllocateProperties(InputDataDesc.GetAllocatedProperties());
 				}
-
-				DataDesc.ElementCount = TotalNumElements;
-				DataDesc.ElementCount2D = TotalNumElements2D;
 			}
 		};
 
@@ -370,18 +374,8 @@ void UPCGComputeKernel::ComputeDataDescFromPinProperties(const FPCGPinProperties
 			{
 				for (const FPCGDataDesc& InputDataDesc : InputDescs[0].DataDescs)
 				{
-					FPCGDataDesc& DataDesc = OutPinDesc.DataDescs.Emplace_GetRef(OutputPinProps.AllowedTypes);
-					const bool bDomainIs2D = DataDesc.IsDomain2D();
-					const bool bSourceDomainIs2D = InputDataDesc.IsDomain2D();
-
-					if (bDomainIs2D)
-					{
-						DataDesc.ElementCount2D = bSourceDomainIs2D ? InputDataDesc.ElementCount2D : InputDataDesc.ElementCount;
-					}
-					else
-					{
-						DataDesc.ElementCount = bSourceDomainIs2D ? InputDataDesc.ElementCount2D.X * InputDataDesc.ElementCount2D.Y : InputDataDesc.ElementCount;
-					}
+					FPCGDataDesc& DataDesc = OutPinDesc.DataDescs.Emplace_GetRef(OutputPinProps.AllowedTypes, /*InitialElementCount=*/0);
+					DataDesc.CombineElementCount(InputDataDesc, EPCGElementMultiplicity::Sum);
 
 					if (Props.AttributeInheritanceMode == EPCGAttributeInheritanceMode::CopyAttributeSetup)
 					{
@@ -407,73 +401,30 @@ void UPCGComputeKernel::ComputeDataDescFromPinProperties(const FPCGPinProperties
 						{
 							for (const FPCGDataDesc& OtherInputDataDesc : OtherInputDesc.DataDescs)
 							{
-								FPCGDataDesc& DataDesc = OutPinDesc.DataDescs.Emplace_GetRef(OutputPinProps.AllowedTypes);
-								const bool bDomainIs2D = DataDesc.IsDomain2D();
-								const bool bSourceDomainIs2D = InputDataDesc.IsDomain2D();
-								const bool bOtherSourceDomainIs2D = OtherInputDataDesc.IsDomain2D();
+								FPCGDataDesc& DataDesc = OutPinDesc.DataDescs.Emplace_GetRef(OutputPinProps.AllowedTypes, /*InitialElementCount=*/0);
 
 								if (Props.ElementCountMode == EPCGElementCountMode::FromInputData)
 								{
-									if (Props.ElementMultiplicity == EPCGElementMultiplicity::Product)
+									DataDesc.CombineElementCount(FPCGDataDesc::CombineElementCount(InputDataDesc, OtherInputDataDesc, Props.ElementMultiplicity), EPCGElementMultiplicity::Sum);
+								}
+								else if (Props.ElementCountMode == EPCGElementCountMode::Fixed)
+								{
+									const EPCGElementDimension ElementDimension = PCGComputeHelpers::GetElementDimension(OutputPinProps.AllowedTypes);
+
+									if (ElementDimension == EPCGElementDimension::One)
 									{
-										if (bDomainIs2D)
-										{
-											if (bSourceDomainIs2D)
-											{
-												DataDesc.ElementCount2D = bOtherSourceDomainIs2D ? InputDataDesc.ElementCount2D * OtherInputDataDesc.ElementCount2D : InputDataDesc.ElementCount2D * OtherInputDataDesc.ElementCount;
-											}
-											else
-											{
-												DataDesc.ElementCount2D = bOtherSourceDomainIs2D ?  OtherInputDataDesc.ElementCount2D * InputDataDesc.ElementCount : InputDataDesc.ElementCount * OtherInputDataDesc.ElementCount;
-											}
-										}
-										else
-										{
-											if (bSourceDomainIs2D)
-											{
-												DataDesc.ElementCount = (bOtherSourceDomainIs2D ? InputDataDesc.ElementCount2D * OtherInputDataDesc.ElementCount2D : InputDataDesc.ElementCount2D * OtherInputDataDesc.ElementCount).X;
-											}
-											else
-											{
-												DataDesc.ElementCount = (bOtherSourceDomainIs2D ?  OtherInputDataDesc.ElementCount2D * InputDataDesc.ElementCount : InputDataDesc.ElementCount * OtherInputDataDesc.ElementCount).X;
-											}
-										}
+										DataDesc.AddElementCount(Props.ElementCount);
 									}
-									else if (Props.ElementMultiplicity == EPCGElementMultiplicity::Sum)
+									else if (ElementDimension == EPCGElementDimension::Two)
 									{
-										if (bDomainIs2D)
-										{
-											if (bSourceDomainIs2D)
-											{
-												DataDesc.ElementCount2D = bOtherSourceDomainIs2D ? InputDataDesc.ElementCount2D + OtherInputDataDesc.ElementCount2D : InputDataDesc.ElementCount2D + OtherInputDataDesc.ElementCount;
-											}
-											else
-											{
-												DataDesc.ElementCount2D = bOtherSourceDomainIs2D ?  OtherInputDataDesc.ElementCount2D + InputDataDesc.ElementCount : InputDataDesc.ElementCount + OtherInputDataDesc.ElementCount;
-											}
-										}
-										else
-										{
-											if (bSourceDomainIs2D)
-											{
-												DataDesc.ElementCount = (bOtherSourceDomainIs2D ? InputDataDesc.ElementCount2D + OtherInputDataDesc.ElementCount2D : InputDataDesc.ElementCount2D + OtherInputDataDesc.ElementCount).X;
-											}
-											else
-											{
-												DataDesc.ElementCount = (bOtherSourceDomainIs2D ?  OtherInputDataDesc.ElementCount2D + InputDataDesc.ElementCount : InputDataDesc.ElementCount + OtherInputDataDesc.ElementCount).X;
-											}
-										}
+										DataDesc.AddElementCount(Props.NumElements2D);
 									}
 									else
 									{
+										// Not supported yet.
 										checkNoEntry();
 									}
 								}
-								else if (Props.ElementCountMode == EPCGElementCountMode::Fixed)
-								{
-									DataDesc.ElementCount = Props.ElementCount;
-									DataDesc.ElementCount2D = Props.NumElements2D;
-								}
 								else
 								{
 									checkNoEntry();
@@ -542,15 +493,29 @@ void UPCGComputeKernel::ComputeDataDescFromPinProperties(const FPCGPinProperties
 		const uint32 Multiplier = OutputPinProps.GetElementCountMultiplier();
 		for (FPCGDataDesc& Desc : OutPinDesc.DataDescs)
 		{
-			Desc.ElementCount *= Multiplier;
-			Desc.ElementCount2D *= Multiplier;
+			Desc.ScaleElementCount(Multiplier);
 		}
 	}
 	else if (Props.InitializationMode == EPCGPinInitMode::Custom)
 	{
 		for (int I = 0; I < Props.DataCount; ++I)
 		{
-			OutPinDesc.DataDescs.Emplace(OutputPinProps.AllowedTypes, Props.ElementCount, Props.NumElements2D);
+			const EPCGElementDimension ElementDimension = PCGComputeHelpers::GetElementDimension(OutputPinProps.AllowedTypes);
+
+			if (ElementDimension == EPCGElementDimension::One)
+			{
+				OutPinDesc.DataDescs.Emplace(OutputPinProps.AllowedTypes, Props.ElementCount);
+			}
+			else if (ElementDimension == EPCGElementDimension::Two)
+			{
+				OutPinDesc.DataDescs.Emplace(OutputPinProps.AllowedTypes, Props.NumElements2D);
+			}
+			else
+			{
+				// Not supported yet.
+				checkNoEntry();
+			}
+
 		}
 	}
 	else
diff --git a/Engine/Plugins/PCG/Source/PCG/Private/Compute/PCGDataForGPU.cpp b/Engine/Plugins/PCG/Source/PCG/Private/Compute/PCGDataForGPU.cpp
--- a/Engine/Plugins/PCG/Source/PCG/Private/Compute/PCGDataForGPU.cpp
+++ b/Engine/Plugins/PCG/Source/PCG/Private/Compute/PCGDataForGPU.cpp
@@ -14,6 +14,7 @@
 #include "Compute/PCGComputeCommon.h"
 #include "Compute/PCGComputeGraph.h"
 #include "Compute/PCGDataBinding.h"
+#include "Compute/PCGPinPropertiesGPU.h"
 #include "Compute/Data/PCGProxyForGPUData.h"
 #include "Data/PCGBasePointData.h"
 #include "Data/PCGPointData.h"
@@ -744,21 +745,27 @@ bool FPCGKernelAttributeDesc::operator==(const FPCGKernelAttributeDesc& Other) c
 	return AttributeId == Other.AttributeId && AttributeKey == Other.AttributeKey;
 }
 
-FPCGDataDesc::FPCGDataDesc(EPCGDataType InType)
-	: FPCGDataDesc(InType, 0, FIntPoint::ZeroValue)
+FPCGDataDesc::FPCGDataDesc(EPCGDataType InType, int32 InElementCount)
+	: FPCGDataDesc(InType, FIntVector4(InElementCount, 0, 0, 0))
 {
 }
 
-FPCGDataDesc::FPCGDataDesc(EPCGDataType InType, int InElementCount)
-	: FPCGDataDesc(InType, InElementCount, FIntPoint::ZeroValue)
+FPCGDataDesc::FPCGDataDesc(EPCGDataType InType, FIntPoint InElementCount)
+	: FPCGDataDesc(InType, FIntVector4(InElementCount.X, InElementCount.Y, 0, 0))
 {
 }
 
-FPCGDataDesc::FPCGDataDesc(EPCGDataType InType, int InElementCount, FIntPoint InElementCount2D)
+FPCGDataDesc::FPCGDataDesc(EPCGDataType InType, FIntVector3 InElementCount)
+	: FPCGDataDesc(InType, FIntVector4(InElementCount.X, InElementCount.Y, InElementCount.Z, 0))
+{
+}
+
+FPCGDataDesc::FPCGDataDesc(EPCGDataType InType, FIntVector4 InElementCount)
 	: Type(InType)
 	, ElementCount(InElementCount)
-	, ElementCount2D(InElementCount2D)
 {
+	ElementDimension = PCGComputeHelpers::GetElementDimension(InType);
+
 	InitializeAttributeDescs(/*InData=*/nullptr, /*InBinding*/nullptr);
 }
 
@@ -768,7 +775,7 @@ FPCGDataDesc::FPCGDataDesc(const FPCGTaggedData& InTaggedData, const UPCGDataBin
 
 	Type = InTaggedData.Data->GetDataType();
 	ElementCount = PCGComputeHelpers::GetElementCount(InTaggedData.Data);
-	ElementCount2D = PCGComputeHelpers::GetElementCount2D(InTaggedData.Data);
+	ElementDimension = PCGComputeHelpers::GetElementDimension(InTaggedData.Data);
 
 	TagStringKeys.Reserve(InTaggedData.Tags.Num());
 
@@ -787,12 +794,13 @@ FPCGDataDesc::FPCGDataDesc(const FPCGTaggedData& InTaggedData, const UPCGDataBin
 uint64 FPCGDataDesc::ComputePackedSize() const
 {
 	check(PCGComputeHelpers::IsTypeAllowedInDataCollection(Type));
+	ensure(ElementDimension == EPCGElementDimension::One);
 
 	uint64 DataSizeBytes = DATA_HEADER_SIZE_BYTES;
 
 	for (const FPCGKernelAttributeDesc& AttributeDesc : AttributeDescs)
 	{
-		const uint64 NumValues = static_cast<uint64>(GetElementCountForAttribute(AttributeDesc));
+		const uint64 NumValues = static_cast<uint64>(GetElementCountForAttribute(AttributeDesc).X);
 		DataSizeBytes += static_cast<uint64>(PCGDataForGPUHelpers::GetAttributeTypeStrideBytes(AttributeDesc.AttributeKey.Type)) * NumValues;
 	}
 
@@ -847,13 +855,13 @@ void FPCGDataDesc::AddAttribute(FPCGKernelAttributeKey InAttribute, const UPCGDa
 	AttributeDescs.Emplace(AttributeId, InAttribute.Type, InAttribute.Identifier, InOptionalUniqueStringKeys);
 }
 
-int32 FPCGDataDesc::GetElementCountForAttribute(const FPCGKernelAttributeDesc& AttributeDesc) const
+FIntVector4 FPCGDataDesc::GetElementCountForAttribute(const FPCGKernelAttributeDesc& AttributeDesc) const
 {
 	// A buffer that is not fully allocated uses only a single value to represent all elements.
 	// Data domain only has a single element.
 	if (!IsAttributeAllocated(AttributeDesc) || AttributeDesc.AttributeKey.Identifier.MetadataDomain == PCGMetadataDomainID::Data)
 	{
-		return 1;
+		return FIntVector4(1, 0, 0, 0);
 	}
 	else
 	{
@@ -861,6 +869,153 @@ int32 FPCGDataDesc::GetElementCountForAttribute(const FPCGKernelAttributeDesc& A
 	}
 }
 
+int32 FPCGDataDesc::ComputeTotalElementCount() const
+{
+	switch (ElementDimension)
+	{
+	default: // Fallthrough
+	case EPCGElementDimension::One:
+		return ElementCount.X;
+	case EPCGElementDimension::Two:
+		return ElementCount.X * ElementCount.Y;
+	case EPCGElementDimension::Three:
+		return ElementCount.X * ElementCount.Y * ElementCount.Z;
+	case EPCGElementDimension::Four:
+		return ElementCount.X * ElementCount.Y * ElementCount.Z * ElementCount.W;
+	}
+}
+
+void FPCGDataDesc::AddElementCount(int32 InElementCountToAdd)
+{
+	ensure(ElementDimension == EPCGElementDimension::One);
+	ElementCount.X += InElementCountToAdd;
+}
+
+void FPCGDataDesc::AddElementCount(FIntPoint InElementCountToAdd)
+{
+	ensure(ElementDimension == EPCGElementDimension::Two);
+	ElementCount.X += InElementCountToAdd.X;
+	ElementCount.Y += InElementCountToAdd.Y;
+}
+
+void FPCGDataDesc::AddElementCount(FIntVector3 InElementCountToAdd)
+{
+	ensure(ElementDimension == EPCGElementDimension::Three);
+	ElementCount.X += InElementCountToAdd.X;
+	ElementCount.Y += InElementCountToAdd.Y;
+	ElementCount.Z += InElementCountToAdd.Z;
+}
+
+void FPCGDataDesc::AddElementCount(FIntVector4 InElementCountToAdd)
+{
+	ensure(ElementDimension == EPCGElementDimension::Four);
+	ElementCount += InElementCountToAdd;
+}
+
+FPCGDataDesc& FPCGDataDesc::CombineElementCount(const FPCGDataDesc& Other, EPCGElementMultiplicity Multiplicity)
+{
+	if (Multiplicity == EPCGElementMultiplicity::Sum)
+	{
+		if (ElementDimension == EPCGElementDimension::One)
+		{
+			// Special case for 1D processing, sum by collapsing the other element count into a single value.
+			ElementCount.X += Other.ComputeTotalElementCount();
+		}
+		else if (ElementDimension == EPCGElementDimension::Two)
+		{
+			if (Other.ElementDimension == EPCGElementDimension::One)
+			{
+				ElementCount.X += Other.ElementCount.X;
+				ElementCount.Y += Other.ElementCount.X;
+			}
+			else
+			{
+				ElementCount += Other.ElementCount;
+				ElementCount.Z = 0;
+				ElementCount.W = 0;
+			}
+		}
+		else if (ElementDimension == EPCGElementDimension::Three)
+		{
+			if (Other.ElementDimension == EPCGElementDimension::One)
+			{
+				ElementCount.X += Other.ElementCount.X;
+				ElementCount.Y += Other.ElementCount.X;
+				ElementCount.Z += Other.ElementCount.X;
+			}
+			else
+			{
+				ElementCount += Other.ElementCount;
+				ElementCount.W = 0;
+			}
+		}
+		else if (ElementDimension == EPCGElementDimension::Four)
+		{
+			if (Other.ElementDimension == EPCGElementDimension::One)
+			{
+				ElementCount.X += Other.ElementCount.X;
+				ElementCount.Y += Other.ElementCount.X;
+				ElementCount.Z += Other.ElementCount.X;
+				ElementCount.W += Other.ElementCount.X;
+			}
+			else
+			{
+				ElementCount += Other.ElementCount;
+			}
+		}
+	}
+	else
+	{
+		if (ElementDimension == EPCGElementDimension::One)
+		{
+			// Special case for 1D processing, scale by collapsing the other element count into a single value.
+			ElementCount.X *= Other.ComputeTotalElementCount();
+		}
+		else if (ElementDimension == EPCGElementDimension::Two)
+		{
+			if (Other.ElementDimension == EPCGElementDimension::One)
+			{
+				ElementCount *= Other.ElementCount.X;
+			}
+			else
+			{
+				ElementCount *= Other.ElementCount;
+			}
+		}
+		else if (ElementDimension == EPCGElementDimension::Three)
+		{
+			if (Other.ElementDimension == EPCGElementDimension::One)
+			{
+				ElementCount *= Other.ElementCount.X;
+			}
+			else
+			{
+				ElementCount *= Other.ElementCount;
+			}
+		}
+		else if (ElementDimension == EPCGElementDimension::Four)
+		{
+			if (Other.ElementDimension == EPCGElementDimension::One)
+			{
+				ElementCount *= Other.ElementCount.X;
+			}
+			else
+			{
+				ElementCount *= Other.ElementCount;
+			}
+		}
+	}
+
+	return *this;
+}
+
+FPCGDataDesc FPCGDataDesc::CombineElementCount(const FPCGDataDesc& A, const FPCGDataDesc& B, EPCGElementMultiplicity Multiplicity)
+{
+	FPCGDataDesc Desc = A;
+	Desc.CombineElementCount(B, Multiplicity);
+	return Desc;
+}
+
 bool FPCGDataDesc::IsAttributeAllocated(int32 InAttributeId) const
 {
 	// @todo_pcg: Only point properties are supported for single value ranges. All other attributes are fully allocated.
@@ -1070,15 +1225,16 @@ void FPCGDataCollectionDesc::WriteHeader(TArray<uint32>& OutPackedDataCollection
 	for (int32 DataIndex = 0; DataIndex < DataDescs.Num(); ++DataIndex)
 	{
 		const FPCGDataDesc& DataDesc = DataDescs[DataIndex];
+		ensure(DataDesc.GetElementDimension() == EPCGElementDimension::One);
 
 		// Data i: type ID
-		if (DataDesc.Type == EPCGDataType::Param)
+		if (DataDesc.GetType() == EPCGDataType::Param)
 		{
 			OutPackedDataCollectionHeader[WriteAddressUints++] = PARAM_DATA_TYPE_ID;
 		}
 		else
 		{
-			ensure(DataDesc.Type == EPCGDataType::Point);
+			ensure(DataDesc.GetType() == EPCGDataType::Point);
 			OutPackedDataCollectionHeader[WriteAddressUints++] = POINT_DATA_TYPE_ID;
 		}
 
@@ -1086,7 +1242,7 @@ void FPCGDataCollectionDesc::WriteHeader(TArray<uint32>& OutPackedDataCollection
 		OutPackedDataCollectionHeader[WriteAddressUints++] = DataDesc.AttributeDescs.Num();
 
 		// Data i: element count
-		OutPackedDataCollectionHeader[WriteAddressUints++] = DataDesc.ElementCount;
+		OutPackedDataCollectionHeader[WriteAddressUints++] = DataDesc.GetElementCount().X;
 
 		const uint32 DataAttributesHeaderStartAddressBytes = WriteAddressUints << 2;
 
@@ -1094,7 +1250,7 @@ void FPCGDataCollectionDesc::WriteHeader(TArray<uint32>& OutPackedDataCollection
 		{
 			const FPCGKernelAttributeDesc& AttributeDesc = DataDesc.AttributeDescs[AttrIndex];
 			const int AttributeStride = PCGDataForGPUHelpers::GetAttributeTypeStrideBytes(AttributeDesc.AttributeKey.Type);
-			const int AttributeElementsCount = DataDesc.GetElementCountForAttribute(AttributeDesc);
+			const int AttributeElementsCount = DataDesc.GetElementCountForAttribute(AttributeDesc).X;
 
 			// Scatter from attributes that are present into header which has slots for all possible attributes.
 			WriteAddressUints = (AttributeDesc.AttributeId * ATTRIBUTE_HEADER_SIZE_BYTES + DataAttributesHeaderStartAddressBytes) >> 2;
@@ -1155,6 +1311,7 @@ void FPCGDataCollectionDesc::PackDataCollection(const FPCGDataCollection& InData
 	for (int32 DataIndex = 0; DataIndex < InputData.Num(); ++DataIndex)
 	{
 		const FPCGDataDesc& DataDesc = DataDescs[DataIndex];
+		ensure(DataDesc.GetElementDimension() == EPCGElementDimension::One);
 
 		if (Cast<UPCGProxyForGPUData>(InputData[DataIndex].Data))
 		{
@@ -1163,7 +1320,7 @@ void FPCGDataCollectionDesc::PackDataCollection(const FPCGDataCollection& InData
 		}
 
 		// No work to do if there are no elements to process.
-		if (DataDesc.ElementCount <= 0)
+		if (DataDesc.GetElementCount().X <= 0)
 		{
 			continue;
 		}
@@ -1352,7 +1509,7 @@ void FPCGDataCollectionDesc::PackDataCollection(const FPCGDataCollection& InData
 
 				// Pack attribute. Validate first element only for perf.
 				ensure(PCGDataForGPUHelpers::PackAttributeHelper(AttributeBase, AttributeDesc, /*InEntryKey*/0, InDataBinding->GetStringTable(), OutPackedDataCollection, AddressUints));
-				for (int32 ElementIndex = 1; ElementIndex < DataDesc.GetElementCountForAttribute(AttributeDesc); ++ElementIndex)
+				for (int32 ElementIndex = 1; ElementIndex < DataDesc.GetElementCountForAttribute(AttributeDesc).X; ++ElementIndex)
 				{
 					PCGDataForGPUHelpers::PackAttributeHelper(AttributeBase, AttributeDesc, /*InEntryKey*/ElementIndex, InDataBinding->GetStringTable(), OutPackedDataCollection, AddressUints);
 				}
@@ -1393,19 +1550,20 @@ EPCGUnpackDataCollectionResult FPCGDataCollectionDesc::UnpackDataCollection(FPCG
 
 	for (uint32 DataIndex = 0; DataIndex < NumData; ++DataIndex)
 	{
-		const uint32 DataHeaderAddress = (DATA_COLLECTION_HEADER_SIZE_BYTES + DATA_HEADER_SIZE_BYTES * DataIndex) / sizeof(uint32);
+		const FPCGDataDesc& DataDesc = DataDescs[DataIndex];
+		ensure(DataDesc.GetElementDimension() == EPCGElementDimension::One);
 
+		const uint32 DataHeaderAddress = (DATA_COLLECTION_HEADER_SIZE_BYTES + DATA_HEADER_SIZE_BYTES * DataIndex) / sizeof(uint32);
 		const uint32 TypeId =        DataAsUint[DataHeaderAddress + 0];
 		const uint32 NumAttributes = DataAsUint[DataHeaderAddress + 1];
 		
 		uint32 NumElements = DataAsUint[DataHeaderAddress + 2];
 		// Use tighter/more refined element count if we have one:
-		if (DataDescs[DataIndex].ElementCount >= 0)
+		if (DataDesc.GetElementCount().X >= 0)
 		{
-			NumElements = FMath::Min(NumElements, static_cast<uint32>(DataDescs[DataIndex].ElementCount));
+			NumElements = FMath::Min(NumElements, static_cast<uint32>(DataDesc.GetElementCount().X));
 		}
 
-		const FPCGDataDesc& DataDesc = DataDescs[DataIndex];
 		const TArray<FPCGKernelAttributeDesc>& AttributeDescs = DataDesc.AttributeDescs;
 		check(NumAttributes == AttributeDescs.Num());
 		check(AttributeDescs.Num() <= MAX_NUM_ATTRS);
@@ -1720,7 +1878,7 @@ EPCGUnpackDataCollectionResult FPCGDataCollectionDesc::UnpackDataCollection(FPCG
 
 					if (FPCGMetadataAttributeBase* AttributeBase = MetadataAttributes[AttributeDesc.AttributeId])
 					{
-						ensure(PCGDataForGPUHelpers::UnpackAttributeHelper(InContext, PackedData, AttributeDesc, InStringTable, AddressUints, DataDesc.GetElementCountForAttribute(AttributeDesc), OutPointData));
+						ensure(PCGDataForGPUHelpers::UnpackAttributeHelper(InContext, PackedData, AttributeDesc, InStringTable, AddressUints, DataDesc.GetElementCountForAttribute(AttributeDesc).X, OutPointData));
 					}
 				}
 			}
@@ -1830,24 +1988,24 @@ EPCGUnpackDataCollectionResult FPCGDataCollectionDesc::UnpackDataCollection(FPCG
 	return EPCGUnpackDataCollectionResult::Success;
 }
 
-uint32 FPCGDataCollectionDesc::ComputeDataElementCount() const
+uint32 FPCGDataCollectionDesc::ComputeTotalElementCount() const
 {
-	return ComputeDataElementCount(EPCGDataType::Any);
+	return ComputeTotalElementCount(EPCGDataType::Any);
 }
 
-uint32 FPCGDataCollectionDesc::ComputeDataElementCount(EPCGDataType InDataType) const
+uint32 FPCGDataCollectionDesc::ComputeTotalElementCount(EPCGDataType InDataType) const
 {
-	uint32 ElementCount = 0;
+	uint32 ProcessingElementCount = 0;
 
 	for (const FPCGDataDesc& DataDesc : DataDescs)
 	{
-		if (!!(DataDesc.Type & InDataType))
+		if (!!(DataDesc.GetType() & InDataType))
 		{
-			ElementCount += DataDesc.IsDomain2D() ? (DataDesc.ElementCount2D.X * DataDesc.ElementCount2D.Y) : DataDesc.ElementCount;
+			ProcessingElementCount += DataDesc.ComputeTotalElementCount();
 		}
 	}
 
-	return ElementCount;
+	return ProcessingElementCount;
 }
 
 void FPCGDataCollectionDesc::Combine(const FPCGDataCollectionDesc& Other)
diff --git a/Engine/Plugins/PCG/Source/PCG/Private/Compute/PCGDataUploadAdaptor.cpp b/Engine/Plugins/PCG/Source/PCG/Private/Compute/PCGDataUploadAdaptor.cpp
--- a/Engine/Plugins/PCG/Source/PCG/Private/Compute/PCGDataUploadAdaptor.cpp
+++ b/Engine/Plugins/PCG/Source/PCG/Private/Compute/PCGDataUploadAdaptor.cpp
@@ -256,7 +256,8 @@ FRDGBufferSRVRef FPCGDataUploadAdaptor::GetDataElementCountsBufferSRV(FRDGBuilde
 
 	for (const FPCGDataDesc& DataDesc : TargetDataCollectionDesc.DataDescs)
 	{
-		DataElementCounts.Add(DataDesc.ElementCount);
+		ensure(DataDesc.GetElementDimension() == EPCGElementDimension::One);
+		DataElementCounts.Add(DataDesc.GetElementCount().X);
 	}
 
 	if (DataElementCounts.IsEmpty())
diff --git a/Engine/Plugins/PCG/Source/PCG/Private/Elements/Metadata/PCGMetadataPartitionKernel.cpp b/Engine/Plugins/PCG/Source/PCG/Private/Elements/Metadata/PCGMetadataPartitionKernel.cpp
--- a/Engine/Plugins/PCG/Source/PCG/Private/Elements/Metadata/PCGMetadataPartitionKernel.cpp
+++ b/Engine/Plugins/PCG/Source/PCG/Private/Elements/Metadata/PCGMetadataPartitionKernel.cpp
@@ -187,9 +187,7 @@ bool UPCGMetadataPartitionKernel::ComputeOutputBindingDataDesc(const UPCGCompute
 		{
 			if (AttributeValueToCount[Value] > 0)
 			{
-				FPCGDataDesc& DataDesc = PinDesc.DataDescs.Emplace_GetRef();
-				DataDesc.Type = EPCGDataType::Point;
-				DataDesc.ElementCount = AttributeValueToCount[Value];
+				FPCGDataDesc& DataDesc = PinDesc.DataDescs.Emplace_GetRef(EPCGDataType::Point, AttributeValueToCount[Value]);
 				DataDesc.AttributeDescs.Append(InputPinDesc.DataDescs[DataIndex].AttributeDescs);
 				DataDesc.AllocateProperties(InputPinDesc.DataDescs[DataIndex].GetAllocatedProperties());
 
@@ -220,7 +218,7 @@ bool UPCGMetadataPartitionKernel::ComputeOutputBindingDataDesc(const UPCGCompute
 int UPCGMetadataPartitionKernel::ComputeThreadCount(const UPCGDataBinding* InBinding) const
 {
 	const FPCGDataCollectionDesc* OutputPinDesc = InBinding->GetCachedKernelPinDataDesc(this, PCGPinConstants::DefaultOutputLabel, /*bIsInput=*/false);
-	return ensure(OutputPinDesc) ? OutputPinDesc->ComputeDataElementCount(EPCGDataType::Point) : 0;
+	return ensure(OutputPinDesc) ? OutputPinDesc->ComputeTotalElementCount() : 0;
 }
 
 #if WITH_EDITOR
diff --git a/Engine/Plugins/PCG/Source/PCG/Private/Elements/PCGCopyPointsAnalysisKernel.cpp b/Engine/Plugins/PCG/Source/PCG/Private/Elements/PCGCopyPointsAnalysisKernel.cpp
--- a/Engine/Plugins/PCG/Source/PCG/Private/Elements/PCGCopyPointsAnalysisKernel.cpp
+++ b/Engine/Plugins/PCG/Source/PCG/Private/Elements/PCGCopyPointsAnalysisKernel.cpp
@@ -63,20 +63,11 @@ bool UPCGCopyPointsAnalysisKernel::ComputeOutputBindingDataDesc(const UPCGComput
 	// Output: a single attribute set with a single attribute true/false value per NumSource*NumTarget possible outputs.
 	FPCGDataCollectionDesc OutDesc;
 
-	FPCGDataDesc& DataDesc = OutDesc.DataDescs.Emplace_GetRef();
-	DataDesc.Type = EPCGDataType::Param;
-
 	const UPCGCopyPointsSettings* CopyPointSettings = CastChecked<UPCGCopyPointsSettings>(GetSettings());
 
-	if (CopyPointSettings->bCopyEachSourceOnEveryTarget)
-	{
-		DataDesc.ElementCount = NumSourceDatas * NumTargetDatas;
-	}
-	else
-	{
-		// Validation should have ensured already that these are valid: (S, T) = (N, N), (N, 1) or (1, N).
-		DataDesc.ElementCount = FMath::Max(NumSourceDatas, NumTargetDatas);
-	}
+	// Validation should have ensured already that these are valid: (S, T) = (N, N), (N, 1) or (1, N).
+	const int32 ElementCount = CopyPointSettings->bCopyEachSourceOnEveryTarget ? (NumSourceDatas * NumTargetDatas) : FMath::Max(NumSourceDatas, NumTargetDatas);
+	FPCGDataDesc& DataDesc = OutDesc.DataDescs.Emplace_GetRef(EPCGDataType::Param, ElementCount);
 
 	FPCGKernelAttributeKey SelectedFlagAttributeKey = FPCGKernelAttributeKey(PCGCopyPointsConstants::SelectedFlagAttributeName, EPCGKernelAttributeType::Bool);
 	DataDesc.AddAttribute(SelectedFlagAttributeKey, InBinding);
@@ -89,7 +80,7 @@ bool UPCGCopyPointsAnalysisKernel::ComputeOutputBindingDataDesc(const UPCGComput
 int UPCGCopyPointsAnalysisKernel::ComputeThreadCount(const UPCGDataBinding* InBinding) const
 {
 	const FPCGDataCollectionDesc* OutputPinDesc = InBinding->GetCachedKernelPinDataDesc(this, PCGPinConstants::DefaultOutputLabel, /*bIsInput=*/false);
-	return ensure(OutputPinDesc) ? OutputPinDesc->ComputeDataElementCount(EPCGDataType::Point) : 0;
+	return ensure(OutputPinDesc) ? OutputPinDesc->ComputeTotalElementCount() : 0;
 }
 
 #if WITH_EDITOR
diff --git a/Engine/Plugins/PCG/Source/PCG/Private/Elements/PCGCopyPointsKernel.cpp b/Engine/Plugins/PCG/Source/PCG/Private/Elements/PCGCopyPointsKernel.cpp
--- a/Engine/Plugins/PCG/Source/PCG/Private/Elements/PCGCopyPointsKernel.cpp
+++ b/Engine/Plugins/PCG/Source/PCG/Private/Elements/PCGCopyPointsKernel.cpp
@@ -128,7 +128,7 @@ bool UPCGCopyPointsKernel::ComputeOutputBindingDataDesc(const UPCGComputeGraph*
 			const FPCGDataDesc& SourceDesc = SourcePinDesc.DataDescs[SourceIndex];
 			const FPCGDataDesc& TargetDesc = TargetPinDesc.DataDescs[TargetIndex];
 
-			FPCGDataDesc& ResultDataDesc = PinDesc.DataDescs.Emplace_GetRef(EPCGDataType::Point, SourceDesc.ElementCount * TargetDesc.ElementCount);
+			FPCGDataDesc& ResultDataDesc = PinDesc.DataDescs.Emplace_GetRef(EPCGDataType::Point, SourceDesc.GetElementCount() * TargetDesc.GetElementCount());
 
 			if (CPSettings->AttributeInheritance == EPCGCopyPointsMetadataInheritanceMode::SourceFirst)
 			{
@@ -210,7 +210,7 @@ bool UPCGCopyPointsKernel::ComputeOutputBindingDataDesc(const UPCGComputeGraph*
 int UPCGCopyPointsKernel::ComputeThreadCount(const UPCGDataBinding* InBinding) const
 {
 	const FPCGDataCollectionDesc* OutputPinDesc = InBinding->GetCachedKernelPinDataDesc(this, PCGPinConstants::DefaultOutputLabel, /*bIsInput=*/false);
-	return ensure(OutputPinDesc) ? OutputPinDesc->ComputeDataElementCount(EPCGDataType::Point) : 0;
+	return ensure(OutputPinDesc) ? OutputPinDesc->ComputeTotalElementCount() : 0;
 }
 
 #if WITH_EDITOR
diff --git a/Engine/Plugins/PCG/Source/PCG/Private/Elements/PCGSkinnedMeshSpawnerKernel.cpp b/Engine/Plugins/PCG/Source/PCG/Private/Elements/PCGSkinnedMeshSpawnerKernel.cpp
--- a/Engine/Plugins/PCG/Source/PCG/Private/Elements/PCGSkinnedMeshSpawnerKernel.cpp
+++ b/Engine/Plugins/PCG/Source/PCG/Private/Elements/PCGSkinnedMeshSpawnerKernel.cpp
@@ -60,7 +60,7 @@ bool UPCGSkinnedMeshSpawnerKernel::ComputeOutputBindingDataDesc(const UPCGComput
 int UPCGSkinnedMeshSpawnerKernel::ComputeThreadCount(const UPCGDataBinding* InBinding) const
 {
 	const FPCGDataCollectionDesc* InputPinDesc = InBinding->GetCachedKernelPinDataDesc(this, PCGPinConstants::DefaultInputLabel, /*bIsInput=*/true);
-	return ensure(InputPinDesc) ? InputPinDesc->ComputeDataElementCount(EPCGDataType::Point) : 0;
+	return ensure(InputPinDesc) ? InputPinDesc->ComputeTotalElementCount() : 0;
 }
 
 #if WITH_EDITOR
diff --git a/Engine/Plugins/PCG/Source/PCG/Private/Elements/PCGStaticMeshSpawnerKernel.cpp b/Engine/Plugins/PCG/Source/PCG/Private/Elements/PCGStaticMeshSpawnerKernel.cpp
--- a/Engine/Plugins/PCG/Source/PCG/Private/Elements/PCGStaticMeshSpawnerKernel.cpp
+++ b/Engine/Plugins/PCG/Source/PCG/Private/Elements/PCGStaticMeshSpawnerKernel.cpp
@@ -118,7 +118,7 @@ bool UPCGStaticMeshSpawnerKernel::ComputeOutputBindingDataDesc(const UPCGCompute
 int UPCGStaticMeshSpawnerKernel::ComputeThreadCount(const UPCGDataBinding* InBinding) const
 {
 	const FPCGDataCollectionDesc* InputPinDesc = InBinding->GetCachedKernelPinDataDesc(this, PCGPinConstants::DefaultInputLabel, /*bIsInput=*/true);
-	return ensure(InputPinDesc) ? InputPinDesc->ComputeDataElementCount(EPCGDataType::Point) : 0;
+	return ensure(InputPinDesc) ? InputPinDesc->ComputeTotalElementCount() : 0;
 }
 
 #if WITH_EDITOR
