[PCG] Avoid accessing PCGComponent managed resources if they are not accessible.

FIxes crash hit by user where resources where not accessible, likely due to a delayed cleanup.

#rnx

diff --git a/Engine/Plugins/PCG/Source/PCG/Private/PCGActorAndComponentMapping.cpp b/Engine/Plugins/PCG/Source/PCG/Private/PCGActorAndComponentMapping.cpp
--- a/Engine/Plugins/PCG/Source/PCG/Private/PCGActorAndComponentMapping.cpp
+++ b/Engine/Plugins/PCG/Source/PCG/Private/PCGActorAndComponentMapping.cpp
@@ -1901,7 +1901,7 @@ void FPCGActorAndComponentMapping::OnSelectionKeyChanged(const FPCGSelectionKey&
 		OuterActor->GetComponents<UPCGComponent>(OriginatingComponents);
 		for (UPCGComponent* OriginatingComponent : OriginatingComponents)
 		{
-			if (OriginatingComponent->IsAnyObjectManagedByResource(ChangedObjects))
+			if (!OriginatingComponent->IsCleaningUp() && OriginatingComponent->IsAnyObjectManagedByResource(ChangedObjects))
 			{
 				ChangedObjectDynamicTrackingPriority = OriginatingComponent->GetExecutionState().GetDynamicTrackingPriority();
 				break;
@@ -1921,6 +1921,7 @@ void FPCGActorAndComponentMapping::OnSelectionKeyChanged(const FPCGSelectionKey&
 		if (!InComponent
 			|| ChangedObjects.Contains(InComponent)
 			|| (bInNoRefreshOwner && ChangedObjects.Contains(InComponent->GetOwner()))
+			|| InComponent->IsCleaningUp()
 			|| InComponent->IsAnyObjectManagedByResource(ChangedObjects)
 			|| (ChangedObjectDynamicTrackingPriority.IsSet() && InComponent->GetExecutionState().GetDynamicTrackingPriority() > ChangedObjectDynamicTrackingPriority.GetValue()))
 		{
diff --git a/Engine/Plugins/PCG/Source/PCG/Private/PCGComponent.cpp b/Engine/Plugins/PCG/Source/PCG/Private/PCGComponent.cpp
--- a/Engine/Plugins/PCG/Source/PCG/Private/PCGComponent.cpp
+++ b/Engine/Plugins/PCG/Source/PCG/Private/PCGComponent.cpp
@@ -415,6 +415,12 @@ bool UPCGComponent::IsAnyObjectManagedByResource(const TArrayView<const UObject*
 {
 	TRACE_CPUPROFILER_EVENT_SCOPE(UPCGComponent::IsAnyObjectManagedByResource);
 	
+	if (GeneratedResourcesInaccessible)
+	{
+		ensure(false);
+		return false;
+	}
+
 	bool bIsManagedByResource = false;
 	ForEachConstManagedResource([&bIsManagedByResource, &InObjects](const UPCGManagedResource* ManagedResource)
 	{
