[PCG] Add option to disable default input creation from top-level graphs. Will become default eventually.

diff --git a/Engine/Plugins/PCG/Source/PCG/Private/Graph/PCGGraphExecutor.cpp b/Engine/Plugins/PCG/Source/PCG/Private/Graph/PCGGraphExecutor.cpp
--- a/Engine/Plugins/PCG/Source/PCG/Private/Graph/PCGGraphExecutor.cpp
+++ b/Engine/Plugins/PCG/Source/PCG/Private/Graph/PCGGraphExecutor.cpp
@@ -2969,6 +2969,15 @@ bool FPCGFetchInputElement::ExecuteInternal(FPCGContext* Context) const
 
 	check(Context->Node);
 
+	// Skip retrieval if the graph is marked as not providing default inputs
+	if (const UPCGGraph* Graph = Context->Node->GetGraph())
+	{
+		if (!Graph->HasDefaultConstructedInputs())
+		{
+			return true;
+		}
+	}
+
 	if (Context->Node->IsOutputPinConnected(PCGPinConstants::DefaultInputLabel))
 	{
 		if (UPCGData* PCGData = Component->GetPCGData())
diff --git a/Engine/Plugins/PCG/Source/PCG/Private/PCGGraph.cpp b/Engine/Plugins/PCG/Source/PCG/Private/PCGGraph.cpp
--- a/Engine/Plugins/PCG/Source/PCG/Private/PCGGraph.cpp
+++ b/Engine/Plugins/PCG/Source/PCG/Private/PCGGraph.cpp
@@ -1853,7 +1853,8 @@ void UPCGGraph::PostEditChangeChainProperty(FPropertyChangedChainEvent& Property
 
 	const FName PropertyName = PropertyChangedEvent.GetPropertyName();
 
-	if (PropertyName == GET_MEMBER_NAME_CHECKED(UPCGGraph, bLandscapeUsesMetadata))
+	if (PropertyName == GET_MEMBER_NAME_CHECKED(UPCGGraph, bLandscapeUsesMetadata) ||
+		PropertyName == GET_MEMBER_NAME_CHECKED(UPCGGraph, bHasDefaultConstructedInputs))
 	{
 		NotifyGraphChanged(EPCGChangeType::Input);
 	}
diff --git a/Engine/Plugins/PCG/Source/PCG/Public/PCGGraph.h b/Engine/Plugins/PCG/Source/PCG/Public/PCGGraph.h
--- a/Engine/Plugins/PCG/Source/PCG/Public/PCGGraph.h
+++ b/Engine/Plugins/PCG/Source/PCG/Public/PCGGraph.h
@@ -293,6 +293,7 @@ public:
 	UE_API uint32 GetDefaultGridSize() const;
 	bool IsHierarchicalGenerationEnabled() const { return bUseHierarchicalGeneration; }
 	bool Use2DGrid() const { return bUse2DGrid; }
+	bool HasDefaultConstructedInputs() const { return bHasDefaultConstructedInputs; }
 
 	using FComputeGraphInstanceKey = TPair</*GridSize*/uint32, /*Compute graph index*/int32>;
 	using FComputeGraphInstancePool = TMap<FComputeGraphInstanceKey, TArray<TSharedPtr<FComputeGraphInstance>>>;
@@ -547,6 +548,10 @@ protected:
 	UPROPERTY(EditAnywhere, Category = Settings, meta = (DisplayName = "2D Grid"))
 	bool bUse2DGrid = true;
 
+	/** Controls whether the default pins (In, ...) are provided data on top-level graphs automatically. Note this will eventually change to be turned off by default. */
+	UPROPERTY(BlueprintReadWrite, EditAnywhere, Category = Settings)
+	bool bHasDefaultConstructedInputs = true;
+
 	/** Execution grid size for nodes. */
 	mutable TMap<const UPCGNode*, uint32> NodeToGridSize;
 
