[PCG] Add polygon metadata support. Further support to inherit metadata from points/splines and following operations will be done in a separate CL.

diff --git a/Engine/Plugins/PCG/Source/PCG/Private/Data/PCGPolyLineData.cpp b/Engine/Plugins/PCG/Source/PCG/Private/Data/PCGPolyLineData.cpp
--- a/Engine/Plugins/PCG/Source/PCG/Private/Data/PCGPolyLineData.cpp
+++ b/Engine/Plugins/PCG/Source/PCG/Private/Data/PCGPolyLineData.cpp
@@ -32,6 +32,11 @@ FVector::FReal UPCGPolyLineData::GetLength() const
 	return Length;
 }
 
+int UPCGPolyLineData::GetNumVertices() const
+{
+	return GetNumSegments() + (IsClosed() ? 0 : 1);
+}
+
 FTransform UPCGPolyLineData::K2_GetTransformAtDistance(int SegmentIndex, double Distance, FBox& OutBounds, bool bWorldSpace) const
 {
 	return GetTransformAtDistance(SegmentIndex, Distance, bWorldSpace, &OutBounds);
diff --git a/Engine/Plugins/PCG/Source/PCG/Private/Data/PCGPolygon2DData.cpp b/Engine/Plugins/PCG/Source/PCG/Private/Data/PCGPolygon2DData.cpp
--- a/Engine/Plugins/PCG/Source/PCG/Private/Data/PCGPolygon2DData.cpp
+++ b/Engine/Plugins/PCG/Source/PCG/Private/Data/PCGPolygon2DData.cpp
@@ -6,11 +6,16 @@
 #include "PCGCustomVersion.h"
 #include "Data/PCGPointData.h"
 #include "Data/PCGPointArrayData.h"
+#include "Metadata/Accessors/PCGAttributeAccessorHelpers.h"
+#include "Metadata/Accessors/PCGPropertyAccessor.h"
+#include "Metadata/Accessors/PCGPolygon2DAccessor.h"
 
 #include "SegmentTypes.h"
 #include "Serialization/ArchiveCrc32.h"
 #include "UObject/Package.h"
 
+#define LOCTEXT_NAMESPACE "PCGPolygon2DData"
+
 namespace PCGPolygon2DData
 {
 	// @todo_pcg : could suggest to move this to the polygon implementation
@@ -51,6 +56,117 @@ namespace PCGPolygon2DData
 			}
 		}
 	}
+
+	TUniquePtr<IPCGAttributeAccessor> CreatePropertyAccessor(EPCGPolygon2DProperties Property, bool bQuiet, bool bConst)
+	{
+		if (bConst)
+		{
+			switch (Property)
+			{
+			case EPCGPolygon2DProperties::Position:
+				return MakeUnique<FPCGPolygon2DVerticesAccessor<FVector, EPCGPolygon2DProperties::Position>>(bConst);
+			case EPCGPolygon2DProperties::Rotation:
+				return MakeUnique<FPCGPolygon2DVerticesAccessor<FQuat, EPCGPolygon2DProperties::Rotation>>(bConst);
+			case EPCGPolygon2DProperties::SegmentIndex:
+				return MakeUnique<FPCGPolygon2DVerticesAccessor<int32, EPCGPolygon2DProperties::SegmentIndex>>(bConst);
+			case EPCGPolygon2DProperties::HoleIndex:
+				return MakeUnique<FPCGPolygon2DVerticesAccessor<int32, EPCGPolygon2DProperties::HoleIndex>>(bConst);
+			case EPCGPolygon2DProperties::SegmentLength:
+				return MakeUnique<FPCGPolygon2DVerticesAccessor<double, EPCGPolygon2DProperties::SegmentLength>>(bConst);
+			case EPCGPolygon2DProperties::LocalPosition:
+				return MakeUnique<FPCGPolygon2DVerticesAccessor<FVector2d, EPCGPolygon2DProperties::LocalPosition>>(bConst);
+			case EPCGPolygon2DProperties::LocalRotation:
+				return MakeUnique<FPCGPolygon2DVerticesAccessor<FQuat, EPCGPolygon2DProperties::LocalRotation>>(bConst);
+			default:
+				if (!bQuiet)
+				{
+					UE_LOG(LogPCG, Error, TEXT("EPCGPolygon2DProperties value '%d' does not exist."), Property);
+				}
+				return TUniquePtr<IPCGAttributeAccessor>();
+			}
+		}
+		else
+		{
+			switch (Property)
+			{
+			case EPCGPolygon2DProperties::Position:
+				return MakeUnique<FPCGPolygon2DVerticesAccessor<FVector, EPCGPolygon2DProperties::Position>>(bConst);
+			case EPCGPolygon2DProperties::LocalPosition:
+				return MakeUnique<FPCGPolygon2DVerticesAccessor<FVector2d, EPCGPolygon2DProperties::LocalPosition>>(bConst);
+			case EPCGPolygon2DProperties::Rotation: //fall-through
+			case EPCGPolygon2DProperties::SegmentIndex: //fall-through
+			case EPCGPolygon2DProperties::HoleIndex: //fall-through
+			case EPCGPolygon2DProperties::SegmentLength:
+			case EPCGPolygon2DProperties::LocalRotation:
+				if (!bQuiet)
+				{
+					UE_LOG(LogPCG, Error, TEXT("EPCGPolygon2DProperties value '%d' is read-only."), Property);
+				}
+				return TUniquePtr<IPCGAttributeAccessor>();
+			default:
+				if (!bQuiet)
+				{
+					UE_LOG(LogPCG, Error, TEXT("EPCGPolygon2DProperties value '%d' does not exist."), Property);
+				}
+				return TUniquePtr<IPCGAttributeAccessor>();
+			}
+		}
+	}
+
+	TUniquePtr<IPCGAttributeAccessor> CreateDataPropertyAccessor(EPCGPolygon2DDataProperties DataProperty, bool bQuiet, bool bConst)
+	{
+		if (bConst)
+		{
+			switch(DataProperty)
+			{
+			case EPCGPolygon2DDataProperties::Transform:
+				return MakeUnique<FPCGPolygon2DDataAccessor<FTransform, EPCGPolygon2DDataProperties::Transform>>(bConst);
+			case EPCGPolygon2DDataProperties::Area:
+				return MakeUnique<FPCGPolygon2DDataAccessor<double, EPCGPolygon2DDataProperties::Area>>(bConst);
+			case EPCGPolygon2DDataProperties::Perimeter:
+				return MakeUnique<FPCGPolygon2DDataAccessor<double, EPCGPolygon2DDataProperties::Perimeter>>(bConst);
+			case EPCGPolygon2DDataProperties::BoundsMin:
+				return MakeUnique<FPCGPolygon2DDataAccessor<FVector2d, EPCGPolygon2DDataProperties::BoundsMin>>(bConst);
+			case EPCGPolygon2DDataProperties::BoundsMax:
+				return MakeUnique<FPCGPolygon2DDataAccessor<FVector2d, EPCGPolygon2DDataProperties::BoundsMax>>(bConst);
+			case EPCGPolygon2DDataProperties::SegmentCount:
+				return MakeUnique<FPCGPolygon2DDataAccessor<int32, EPCGPolygon2DDataProperties::SegmentCount>>(bConst);
+			case EPCGPolygon2DDataProperties::OuterSegmentCount:
+				return MakeUnique<FPCGPolygon2DDataAccessor<int32, EPCGPolygon2DDataProperties::OuterSegmentCount>>(bConst);
+			case EPCGPolygon2DDataProperties::HoleCount:
+				return MakeUnique<FPCGPolygon2DDataAccessor<int32, EPCGPolygon2DDataProperties::HoleCount>>(bConst);
+			case EPCGPolygon2DDataProperties::LongestOuterSegmentIndex:
+				return MakeUnique<FPCGPolygon2DDataAccessor<int32, EPCGPolygon2DDataProperties::LongestOuterSegmentIndex>>(bConst);
+			case EPCGPolygon2DDataProperties::IsClockwise:
+				return MakeUnique<FPCGPolygon2DDataAccessor<bool, EPCGPolygon2DDataProperties::IsClockwise>>(bConst);
+			default:
+				if (!bQuiet)
+				{
+					UE_LOG(LogPCG, Error, TEXT("EPCGPolygon2DDataProperties value '%d' does not exist."), DataProperty);
+				}
+				return TUniquePtr<IPCGAttributeAccessor>();
+			}
+		}
+		else
+		{
+			// Only the transform can be set at the data level.
+			if (DataProperty == EPCGPolygon2DDataProperties::Transform)
+			{
+				return MakeUnique<FPCGPolygon2DDataAccessor<FTransform, EPCGPolygon2DDataProperties::Transform>>(bConst);
+			}
+			else
+			{
+				return TUniquePtr<IPCGAttributeAccessor>();
+			}
+		}
+	}
+} // namespace PCGPolygon2DData
+
+UPCGPolygon2DData::UPCGPolygon2DData(const FObjectInitializer& ObjectInitializer)
+	: Super(ObjectInitializer)
+{
+	check(Metadata);
+	Metadata->SetupDomain(PCGMetadataDomainID::Elements, /*bIsDefault=*/true);
 }
 
 void UPCGPolygon2DData::Serialize(FArchive& InArchive)
@@ -61,6 +177,7 @@ void UPCGPolygon2DData::Serialize(FArchive& InArchive)
 
 	PCGPolygon2DData::SerializePolygon(Polygon, InArchive);
 	InArchive << Transform;
+	InArchive << VertexEntryKeys;
 }
 
 void UPCGPolygon2DData::AddToCrc(FArchiveCrc32& Ar, bool bFullDataCrc) const
@@ -88,6 +205,33 @@ void UPCGPolygon2DData::AddToCrc(FArchiveCrc32& Ar, bool bFullDataCrc) const
 	}
 }
 
+FPCGMetadataDomainID UPCGPolygon2DData::GetMetadataDomainIDFromSelector(const FPCGAttributePropertySelector& InSelector) const
+{
+	const FName DomainName = InSelector.GetDomainName();
+
+	if (DomainName == PCGPolygon2DData::VertexDomainName)
+	{
+		return PCGMetadataDomainID::Elements;
+	}
+	else
+	{
+		return Super::GetMetadataDomainIDFromSelector(InSelector);
+	}
+}
+
+bool UPCGPolygon2DData::SetDomainFromDomainID(const FPCGMetadataDomainID& InDomainID, FPCGAttributePropertySelector& InOutSelector) const
+{
+	if (InDomainID == PCGMetadataDomainID::Elements)
+	{
+		InOutSelector.SetDomainName(PCGPolygon2DData::VertexDomainName, /*bResetExtraNames=*/false);
+		return true;
+	}
+	else
+	{
+		return Super::SetDomainFromDomainID(InDomainID, InOutSelector);
+	}
+}
+
 bool UPCGPolygon2DData::SamplePoint(const FTransform& InTransform, const FBox& InBounds, FPCGPoint& OutPoint, UPCGMetadata* OutMetadata) const
 {
 	// Project point down to the 2d plane.
@@ -99,16 +243,37 @@ bool UPCGPolygon2DData::SamplePoint(const FTransform& InTransform, const FBox& I
 		return false;
 	}
 
+	const FVector2D Center = FVector2D(BoundsInLocalSpace.GetCenter());
+
 	// @todo_pcg: Check if the bounds (as a polygon) overlap with the polygon.
 	// Note that the polygon class has an overlap test, but not at the general polygon level
 	// We'll test just the center until then.
-	if (!Polygon.Contains(FVector2D(BoundsInLocalSpace.GetCenter())))
+	if (!Polygon.Contains(Center))
 	{
 		return false;
 	}
 
 	new(&OutPoint) FPCGPoint(InTransform, /*Density=*/1.0f, /*Seed=*/0);
 	OutPoint.SetLocalBounds(InBounds);
+
+	// @todo_pcg: this could be done a bit more cleaner in conjunction with the contains test.
+	// The issue here is that the key should be driven by the closest segment where the point is "inside" not outside.
+	// In the mean time, this might be a reasonable approximation.
+	if (OutMetadata)
+	{
+		int HoleIndex = INDEX_NONE;
+		int SegmentIndex = INDEX_NONE;
+		double DistanceParameter = 0.0;
+
+		double NearestSegmentSquaredDistance = Polygon.DistanceSquared(Center, HoleIndex, SegmentIndex, DistanceParameter);
+		const double SegmentRatio = Polygon.Segment(SegmentIndex, HoleIndex).ConvertToUnitRange(DistanceParameter);
+
+		const int GlobalSegmentIndex = SegmentAndHoleIndicesToSegmentIndex[TPair<int, int>(SegmentIndex, HoleIndex)];
+		const double NearestKey = GlobalSegmentIndex + SegmentRatio;
+
+		WriteMetadataToPoint(NearestKey, OutPoint, OutMetadata);
+	}
+
 	return true;
 }
 
@@ -121,8 +286,8 @@ bool UPCGPolygon2DData::ProjectPoint(const FTransform& InTransform, const FBox&
 	// @todo_pcg: Check if the bounds (as a polygon) overlap with the polygon.
 	// Note that the polygon class has an overlap test, but not at the general polygon level
 	// We'll test just the center until then.
-	const FVector Center = BoundsInLocalSpace.GetCenter();
-	if (!Polygon.Contains(FVector2D(Center)))
+	const FVector2D Center = FVector2D(BoundsInLocalSpace.GetCenter());
+	if (!Polygon.Contains(Center))
 	{
 		return false;
 	}
@@ -141,15 +306,106 @@ bool UPCGPolygon2DData::ProjectPoint(const FTransform& InTransform, const FBox&
 		OutPoint.Transform.SetRotation(UpVectorRotation * InTransform.GetRotation());
 	}
 
+	// @todo_pcg: this could be done a bit more cleaner in conjunction with the contains test.
+	// The issue here is that the key should be driven by the closest segment where the point is "inside" not outside.
+	// In the mean time, this might be a reasonable approximation.
+	if (OutMetadata)
+	{
+		int HoleIndex = INDEX_NONE;
+		int SegmentIndex = INDEX_NONE;
+		double DistanceParameter = 0.0;
+
+		double NearestSegmentSquaredDistance = Polygon.DistanceSquared(Center, HoleIndex, SegmentIndex, DistanceParameter);
+		const double SegmentRatio = Polygon.Segment(SegmentIndex, HoleIndex).ConvertToUnitRange(DistanceParameter);
+
+		const int GlobalSegmentIndex = SegmentAndHoleIndicesToSegmentIndex[TPair<int, int>(SegmentIndex, HoleIndex)];
+		const double NearestKey = GlobalSegmentIndex + SegmentRatio;
+
+		WriteMetadataToPoint(NearestKey, OutPoint, OutMetadata);
+	}
+
 	// Implementation note - the polygon transform has always a unit scale, so there's nothing to project against.
 	return true;
 }
 
+void UPCGPolygon2DData::WriteMetadataToPoint(float InputKey, FPCGPoint& OutPoint, UPCGMetadata* OutMetadata) const
+{
+	return WriteMetadataToEntry(InputKey, OutPoint.MetadataEntry, OutMetadata);
+}
+
+void UPCGPolygon2DData::WriteMetadataToEntry(float InputKey, PCGMetadataEntryKey& OutEntryKey, UPCGMetadata* OutMetadata) const
+{
+	if (!OutMetadata || !ensure(Metadata) || VertexEntryKeys.IsEmpty())
+	{
+		return;
+	}
+
+	const FPCGMetadataDomain* MetadataDomain = Metadata->GetConstMetadataDomain(PCGMetadataDomainID::Elements);
+	FPCGMetadataDomain* OutMetadataDomain = OutMetadata->GetMetadataDomain(PCGMetadataDomainID::Elements);
+
+	if (!MetadataDomain || !OutMetadataDomain)
+	{
+		return;
+	}
+
+	const double InterpolationRatio = FMath::Frac(InputKey);
+	const int CurrentVertexIndex = FMath::FloorToInt32(InputKey);
+
+	check(VertexEntryKeys.IsValidIndex(CurrentVertexIndex));
+	OutMetadataDomain->InitializeOnSet(OutEntryKey, VertexEntryKeys[CurrentVertexIndex], MetadataDomain);
+
+	// Implementation note: there's no need to compute property interpolation if we're directly on a vertex.
+	if (InterpolationRatio <= UE_DOUBLE_SMALL_NUMBER)
+	{
+		return;
+	}
+
+	int NextVertexIndex = INDEX_NONE;
+	int CurrentSegmentIndex = INDEX_NONE;
+	int CurrentHoleIndex = INDEX_NONE;
+
+	if (const TPair<int, int>* CurrentPair = SegmentIndexToSegmentAndHoleIndices.Find(CurrentVertexIndex))
+	{
+		CurrentSegmentIndex = CurrentPair->Key;
+		CurrentHoleIndex = CurrentPair->Value;
+	}
+	else
+	{
+		// invalid index
+		return;
+	}
+
+	if (const int* TentativeNextVertex = SegmentAndHoleIndicesToSegmentIndex.Find(TPair<int, int>(CurrentSegmentIndex + 1, CurrentHoleIndex)))
+	{
+		NextVertexIndex = *TentativeNextVertex;
+	}
+	else if (const int* TentativeFirstVertex = SegmentAndHoleIndicesToSegmentIndex.Find(TPair<int, int>(0, CurrentHoleIndex)))
+	{
+		NextVertexIndex = *TentativeFirstVertex;
+	}
+	else
+	{
+		// Invalid index
+		return;
+	}
+
+	check(VertexEntryKeys.IsValidIndex(NextVertexIndex));
+
+	TStaticArray<TPair<PCGMetadataEntryKey, float>, 2> Coefficients;
+	Coefficients[0] = { VertexEntryKeys[CurrentVertexIndex], 1.0 - InterpolationRatio };
+	Coefficients[1] = { VertexEntryKeys[NextVertexIndex], InterpolationRatio };
+	OutMetadataDomain->ComputeWeightedAttribute(OutEntryKey, Coefficients, MetadataDomain);
+}
+
 UPCGSpatialData* UPCGPolygon2DData::CopyInternal(FPCGContext* Context) const
 {
 	UPCGPolygon2DData* NewPolygonData = FPCGContext::NewObject_AnyThread<UPCGPolygon2DData>(Context);
 	NewPolygonData->Polygon = Polygon;
+	NewPolygonData->VertexEntryKeys = VertexEntryKeys;
 	NewPolygonData->Transform = Transform;
+	NewPolygonData->SegmentIndexToSegmentAndHoleIndices = SegmentIndexToSegmentAndHoleIndices;
+	NewPolygonData->SegmentAndHoleIndicesToSegmentIndex = SegmentAndHoleIndicesToSegmentIndex;
+	NewPolygonData->SegmentCount = SegmentCount;
 	return NewPolygonData;
 }
 
@@ -181,19 +437,46 @@ const UPCGBasePointData* UPCGPolygon2DData::CreateBasePointData(FPCGContext* Con
 
 	PointData->InitializeFromDataWithParams(InitializeFromDataParams);
 
-	// Allocate properties & points
-	// @todo_pcg: add metadata entries if the vertices had entries
-	const EPCGPointNativeProperties PropertiesToAllocate = EPCGPointNativeProperties::Transform;
+	// Implementation note: since polygons can have holes, we will always export the hole index property as an attribute.
+	FName HoleIndexAttributeName = FName(StaticEnum<EPCGPolygon2DProperties>()->GetNameStringByValue(static_cast<int64>(EPCGPolygon2DProperties::HoleIndex)));
+	PointData->Metadata->FindOrCreateAttribute<int>(HoleIndexAttributeName, -1, /*bAllowsInterpolation=*/false);
+
+	// Allocate properties & points.
+	const EPCGPointNativeProperties PropertiesToAllocate = EPCGPointNativeProperties::Transform | EPCGPointNativeProperties::MetadataEntry;
 
-	const int NumPoints = GetNumSegments();
+	const int NumPoints = GetNumVertices();
 	PointData->SetNumPoints(NumPoints);
 	PointData->AllocateProperties(PropertiesToAllocate);
 
+	const bool bHasVertexMetadata = !VertexEntryKeys.IsEmpty();
+	check(VertexEntryKeys.IsEmpty() || VertexEntryKeys.Num() == NumPoints);
+	
 	FPCGPointValueRanges OutRanges(PointData, /*bAllocate=*/false);
 
 	for (int PointIndex = 0; PointIndex < NumPoints; ++PointIndex)
 	{
 		OutRanges.TransformRange[PointIndex] = GetTransformAtDistance(PointIndex, 0.0, /*bWorldSpace=*/true);
+
+		if (bHasVertexMetadata)
+		{
+			OutRanges.MetadataEntryRange[PointIndex] = VertexEntryKeys[PointIndex];
+		}
+
+		PointData->Metadata->InitializeOnSet(OutRanges.MetadataEntryRange[PointIndex]);
+	}
+
+	// Write $HoleIndex to HoleIndex attribute
+	if(ensure(!SegmentIndexToSegmentAndHoleIndices.IsEmpty()))
+	{
+		TArray<int> HoleIndices;
+		HoleIndices.Reserve(SegmentIndexToSegmentAndHoleIndices.Num());
+		for (int VertexIndex = 0; VertexIndex < NumPoints; ++VertexIndex)
+		{
+			HoleIndices.Add(SegmentIndexToSegmentAndHoleIndices[VertexIndex].Value);
+		}
+
+		FPCGAttributePropertyOutputSelector HoleIndexSelector = FPCGAttributePropertySelector::CreateAttributeSelector<FPCGAttributePropertyOutputSelector>(HoleIndexAttributeName);
+		PCGAttributeAccessorHelpers::WriteAllValues<int>(PointData, HoleIndexSelector, HoleIndices, /*SourceSelector=*/nullptr, Context);
 	}
 
 	return PointData;
@@ -276,13 +559,14 @@ float UPCGPolygon2DData::GetInputKeyAtDistance(int SegmentIndex, double Distance
 	}
 	else if (SegmentIndex >= SegmentCount)
 	{
-		return 1.0;
+		return SegmentCount;
 	}
 	else
 	{
 		const TPair<int, int> SegmentMapping = SegmentIndexToSegmentAndHoleIndices[SegmentIndex];
 		const UE::Geometry::FSegment2d Segment = Polygon.Segment(SegmentMapping.Key, SegmentMapping.Value);
-		return (static_cast<double>(SegmentIndex) / SegmentCount) + (0.5 * Distance / Segment.Extent);
+		// @todo_pcg : clamp to 1 the distance/segment.length?
+		return SegmentIndex + (Distance / Segment.Length());
 	}
 }
 
@@ -296,23 +580,33 @@ double UPCGPolygon2DData::GetDistanceAtSegmentStart(int SegmentIndex) const
 	{
 		TPair<int, int> SegmentMapping = SegmentIndexToSegmentAndHoleIndices[CurrentSegmentIndex];
 		UE::Geometry::FSegment2d Segment = Polygon.Segment(SegmentMapping.Key, SegmentMapping.Value);
-		Distance += 2.0 * Segment.Extent;
+		Distance += Segment.Length();
 		++CurrentSegmentIndex;
 	}
 
 	return Distance;
 }
 
-void UPCGPolygon2DData::SetPolygon(const UE::Geometry::FGeneralPolygon2d& InPolygon)
+void UPCGPolygon2DData::SetPolygon(const UE::Geometry::FGeneralPolygon2d& InPolygon, TConstArrayView<PCGMetadataEntryKey>* InOptionalEntryKeys)
 {
 	Polygon = InPolygon;
 	UpdateMappings();
+
+	if (InOptionalEntryKeys && InOptionalEntryKeys->Num() == SegmentIndexToSegmentAndHoleIndices.Num())
+	{
+		VertexEntryKeys = *InOptionalEntryKeys;
+	}
 }
 
-void UPCGPolygon2DData::SetPolygon(UE::Geometry::FGeneralPolygon2d&& InPolygon)
+void UPCGPolygon2DData::SetPolygon(UE::Geometry::FGeneralPolygon2d&& InPolygon, TConstArrayView<PCGMetadataEntryKey>* InOptionalEntryKeys)
 {
 	Polygon = MoveTemp(InPolygon);
 	UpdateMappings();
+
+	if (InOptionalEntryKeys && InOptionalEntryKeys->Num() == SegmentIndexToSegmentAndHoleIndices.Num())
+	{
+		VertexEntryKeys = *InOptionalEntryKeys;
+	}
 }
 
 void UPCGPolygon2DData::SetTransform(const FTransform& InTransform, bool bCheckWinding)
@@ -336,12 +630,15 @@ void UPCGPolygon2DData::SetTransform(const FTransform& InTransform, bool bCheckW
 void UPCGPolygon2DData::UpdateMappings()
 {
 	SegmentIndexToSegmentAndHoleIndices.Reset();
+	SegmentAndHoleIndicesToSegmentIndex.Reset();
 	SegmentCount = 0;
 
 	const UE::Geometry::FPolygon2d& Outer = Polygon.GetOuter();
 	for (int VertexIndex = 0; VertexIndex < Outer.VertexCount(); ++VertexIndex)
 	{
-		SegmentIndexToSegmentAndHoleIndices.Add(SegmentCount++, TPair<int, int>(VertexIndex, -1));
+		SegmentIndexToSegmentAndHoleIndices.Add(SegmentCount, TPair<int, int>(VertexIndex, -1));
+		SegmentAndHoleIndicesToSegmentIndex.Add(TPair<int, int>(VertexIndex, -1), SegmentCount);
+		++SegmentCount;
 	}
 
 	const TArray<UE::Geometry::FPolygon2d>& Holes = Polygon.GetHoles();
@@ -350,7 +647,158 @@ void UPCGPolygon2DData::UpdateMappings()
 		const UE::Geometry::FPolygon2d& Hole = Holes[HoleIndex];
 		for (int VertexIndex = 0; VertexIndex < Hole.VertexCount(); ++VertexIndex)
 		{
-			SegmentIndexToSegmentAndHoleIndices.Add(SegmentCount++, TPair<int, int>(VertexIndex, HoleIndex));
+			SegmentIndexToSegmentAndHoleIndices.Add(SegmentCount, TPair<int, int>(VertexIndex, HoleIndex));
+			SegmentAndHoleIndicesToSegmentIndex.Add(TPair<int, int>(VertexIndex, HoleIndex), SegmentCount);
+			++SegmentCount;
+		}
+	}
+}
+
+void UPCGPolygon2DData::AllocateMetadataEntries()
+{
+	const int32 NumVertices = GetNumVertices();
+
+	if (VertexEntryKeys.Num() != NumVertices)
+	{
+		VertexEntryKeys.Empty(NumVertices);
+	}
+
+	if (VertexEntryKeys.IsEmpty())
+	{
+		VertexEntryKeys.Reserve(NumVertices);
+		for (int32 i = 0; i < NumVertices; ++i)
+		{
+			VertexEntryKeys.Add(PCGInvalidEntryKey);
+		}
+	}
+}
+
+TUniquePtr<IPCGAttributeAccessor> UPCGPolygon2DData::CreateStaticAccessor(const UPCGPolygon2DData* Data, const FPCGAttributePropertySelector& InSelector, bool bQuiet, bool bConst)
+{
+	const FName DomainName = InSelector.GetDomainName();
+	if (InSelector.GetSelection() == EPCGAttributePropertySelection::Property && (DomainName.IsNone() || DomainName == PCGPolygon2DData::VertexDomainName))
+	{
+		const FName PropertyName = InSelector.GetName();
+
+		if (const int64 EnumValue = StaticEnum<EPCGPolygon2DProperties>()->GetValueByName(PropertyName); EnumValue != INDEX_NONE)
+		{
+			const EPCGPolygon2DProperties NativeProperty = static_cast<EPCGPolygon2DProperties>(EnumValue);
+			return PCGPolygon2DData::CreatePropertyAccessor(NativeProperty, bQuiet, bConst);
+		}
+
+		if (!bQuiet)
+		{
+			PCGLog::LogErrorOnGraph(FText::Format(LOCTEXT("FailCreateAccessor", "Property {0} is not valid for a polygon vertex."), FText::FromName(PropertyName)));
+		}
+	}
+	else if (InSelector.GetSelection() == EPCGAttributePropertySelection::Property && DomainName == PCGDataConstants::DataDomainName)
+	{
+		const FName PropertyName = InSelector.GetName();
+
+		if (const int64 EnumValue = StaticEnum<EPCGPolygon2DDataProperties>()->GetValueByName(PropertyName); EnumValue != INDEX_NONE)
+		{
+			const EPCGPolygon2DDataProperties NativeDataProperty = static_cast<EPCGPolygon2DDataProperties>(EnumValue);
+			return PCGPolygon2DData::CreateDataPropertyAccessor(NativeDataProperty, bQuiet, bConst);
 		}
 	}
-}
\ No newline at end of file
+
+	return {};
+}
+
+TUniquePtr<IPCGAttributeAccessorKeys> UPCGPolygon2DData::CreateAccessorKeys(const FPCGAttributePropertySelector& InSelector, bool bQuiet)
+{
+	const EPCGAttributePropertySelection Selection = InSelector.GetSelection();
+	const FPCGMetadataDomainID DomainID = GetMetadataDomainIDFromSelector(InSelector);
+
+	// Global data
+	if (DomainID == PCGMetadataDomainID::Data && (Selection == EPCGAttributePropertySelection::Property || Selection == EPCGAttributePropertySelection::ExtraProperty))
+	{
+		return MakeUnique<FPCGAttributeAccessorKeysPolygon2DData>(this, /*bInGlobalData=*/true);
+	}
+
+	// Vertices data
+	if ((DomainID.IsDefault() || DomainID == PCGMetadataDomainID::Elements) && (Selection == EPCGAttributePropertySelection::Property || Selection == EPCGAttributePropertySelection::ExtraProperty))
+	{
+		return MakeUnique<FPCGAttributeAccessorKeysPolygon2DData>(this, /*bInGlobalData=*/false);
+	}
+
+	// Metadata in the vertices domain
+	if ((DomainID.IsDefault() || DomainID == PCGMetadataDomainID::Elements) && Selection == EPCGAttributePropertySelection::Attribute)
+	{
+		return MakeUnique<FPCGAttributeAccessorKeysPolygon2DDataEntries>(this);
+	}
+
+	return nullptr;
+}
+
+TUniquePtr<const IPCGAttributeAccessorKeys> UPCGPolygon2DData::CreateConstAccessorKeys(const FPCGAttributePropertySelector& InSelector, bool bQuiet) const
+{
+	const EPCGAttributePropertySelection Selection = InSelector.GetSelection();
+	const FPCGMetadataDomainID DomainID = GetMetadataDomainIDFromSelector(InSelector);
+
+	// Global data
+	if (DomainID == PCGMetadataDomainID::Data && (Selection == EPCGAttributePropertySelection::Property || Selection == EPCGAttributePropertySelection::ExtraProperty))
+	{
+		return MakeUnique<FPCGAttributeAccessorKeysPolygon2DData>(this, /*bInGlobalData=*/true);
+	}
+
+	// Vertices data
+	if ((DomainID.IsDefault() || DomainID == PCGMetadataDomainID::Elements) && (Selection == EPCGAttributePropertySelection::Property || Selection == EPCGAttributePropertySelection::ExtraProperty))
+	{
+		return MakeUnique<FPCGAttributeAccessorKeysPolygon2DData>(this, /*bInGlobalData=*/false);
+	}
+
+	// Metadata in the vertices domain
+	if ((DomainID.IsDefault() || DomainID == PCGMetadataDomainID::Elements) && Selection == EPCGAttributePropertySelection::Attribute)
+	{
+		return MakeUnique<FPCGAttributeAccessorKeysPolygon2DDataEntries>(this);
+	}
+
+	return nullptr;
+}
+
+FPCGAttributeAccessorMethods UPCGPolygon2DData::GetPolygon2DAccessorMethods()
+{
+	auto CreateAccessorFunc = [](UPCGData* InData, const FPCGAttributePropertySelector& InSelector, bool bQuiet) -> TUniquePtr<IPCGAttributeAccessor>
+		{
+			return CreateStaticAccessor(CastChecked<UPCGPolygon2DData>(InData), InSelector, bQuiet, /*bConst=*/false);
+		};
+
+	auto CreateConstAccessorFunc = [](const UPCGData* InData, const FPCGAttributePropertySelector& InSelector, bool bQuiet) -> TUniquePtr<const IPCGAttributeAccessor>
+		{
+			return CreateStaticAccessor(CastChecked<UPCGPolygon2DData>(InData), InSelector, bQuiet, /*bConst=*/true);
+		};
+
+	auto CreateAccessorKeysFunc = [](UPCGData* InData, const FPCGAttributePropertySelector& InSelector, bool bQuiet) -> TUniquePtr<IPCGAttributeAccessorKeys>
+		{
+			return CastChecked<UPCGPolygon2DData>(InData)->CreateAccessorKeys(InSelector, bQuiet);
+		};
+
+	auto CreateConstAccessorKeysFunc = [](const UPCGData* InData, const FPCGAttributePropertySelector& InSelector, bool bQuiet) -> TUniquePtr<const IPCGAttributeAccessorKeys>
+		{
+			return CastChecked<UPCGPolygon2DData>(InData)->CreateConstAccessorKeys(InSelector, bQuiet);
+		};
+
+	FPCGAttributeAccessorMethods Methods
+	{
+		.CreateAccessorFunc = CreateAccessorFunc,
+		.CreateConstAccessorFunc = CreateConstAccessorFunc,
+		.CreateAccessorKeysFunc = CreateAccessorKeysFunc,
+		.CreateConstAccessorKeysFunc = CreateConstAccessorKeysFunc
+	};
+
+#if WITH_EDITOR
+	TArray<FText, TInlineAllocator<2>> Menus = { LOCTEXT("Polygon2VerticesSelectorMenuEntry", "Polygon 2d") };
+	FText& SubMenu = Menus.Emplace_GetRef();
+
+	SubMenu = LOCTEXT("PolygonVerticesSelectorMenuEntryPoints", "Vertices");
+	Methods.FillSelectorMenuEntryFromEnum<EPCGPolygon2DProperties>(Menus);
+
+	SubMenu = LOCTEXT("PolygonSelectorMenuEntryGlobal", "Global");
+	Methods.FillSelectorMenuEntryFromEnum<EPCGPolygon2DDataProperties>(Menus);
+#endif // WITH_EDITOR
+
+	return Methods;
+}
+
+#undef LOCTEXT_NAMESPACE
\ No newline at end of file
diff --git a/Engine/Plugins/PCG/Source/PCG/Private/Elements/Polygon/PCGCreatePolygon2D.cpp b/Engine/Plugins/PCG/Source/PCG/Private/Elements/Polygon/PCGCreatePolygon2D.cpp
--- a/Engine/Plugins/PCG/Source/PCG/Private/Elements/Polygon/PCGCreatePolygon2D.cpp
+++ b/Engine/Plugins/PCG/Source/PCG/Private/Elements/Polygon/PCGCreatePolygon2D.cpp
@@ -47,6 +47,11 @@ bool FPCGCreatePolygonElement::ExecuteInternal(FPCGContext* Context) const
 			// @todo_pcg: this method doesn't allow proper metadata inheritance, we will need to update it.
 			const double ErrorTolerance = FMath::Max(Settings->SplineMaxDiscretizationError, 0.001);
 			SplineData->SplineStruct.ConvertSplineToPolyLine(ESplineCoordinateSpace::World, ErrorTolerance, Positions);
+
+			if (bInputFromClosedSpline && !Positions.IsEmpty())
+			{
+				Positions.Pop();
+			}
 		}
 		else if (const UPCGBasePointData* InputPointData = Cast<UPCGBasePointData>(Input.Data))
 		{
diff --git a/Engine/Plugins/PCG/Source/PCG/Private/Metadata/Accessors/PCGAttributeAccessorKeys.cpp b/Engine/Plugins/PCG/Source/PCG/Private/Metadata/Accessors/PCGAttributeAccessorKeys.cpp
--- a/Engine/Plugins/PCG/Source/PCG/Private/Metadata/Accessors/PCGAttributeAccessorKeys.cpp
+++ b/Engine/Plugins/PCG/Source/PCG/Private/Metadata/Accessors/PCGAttributeAccessorKeys.cpp
@@ -5,6 +5,7 @@
 #include "PCGPoint.h"
 #include "Data/PCGBasePointData.h"
 #include "Data/PCGPointData.h"
+#include "Data/PCGPolygon2DData.h"
 #include "Data/PCGSplineData.h"
 #include "Metadata/PCGMetadata.h"
 #include "Metadata/PCGMetadataAttribute.h"
@@ -354,23 +355,27 @@ bool FPCGAttributeAccessorKeysPointsSubset::GetMetadataEntryKeys(int32 InStart,
 }
 
 /////////////////////////////////////////////////////////
-FPCGAttributeAccessorKeysSplineData::FPCGAttributeAccessorKeysSplineData()
-	: FPCGAttributeAccessorKeysSingleObjectPtr()
+template<typename PolyLineType>
+FPCGAttributeAccessorKeysPolyLineData<PolyLineType>::FPCGAttributeAccessorKeysPolyLineData()
+	: FPCGAttributeAccessorKeysSingleObjectPtr<PolyLineType>()
 {}
 
-FPCGAttributeAccessorKeysSplineData::FPCGAttributeAccessorKeysSplineData(UPCGSplineData* InPtr, bool bInGlobalData)
-	: FPCGAttributeAccessorKeysSingleObjectPtr(InPtr)
+template<typename PolyLineType>
+FPCGAttributeAccessorKeysPolyLineData<PolyLineType>::FPCGAttributeAccessorKeysPolyLineData(PolyLineType* InPtr, bool bInGlobalData)
+	: FPCGAttributeAccessorKeysSingleObjectPtr<PolyLineType>(InPtr)
 	, bGlobalData(bInGlobalData)
 {}
 
-FPCGAttributeAccessorKeysSplineData::FPCGAttributeAccessorKeysSplineData(const UPCGSplineData* InPtr, bool bInGlobalData)
-	: FPCGAttributeAccessorKeysSingleObjectPtr(InPtr)
+template<typename PolyLineType>
+FPCGAttributeAccessorKeysPolyLineData<PolyLineType>::FPCGAttributeAccessorKeysPolyLineData(const PolyLineType* InPtr, bool bInGlobalData)
+	: FPCGAttributeAccessorKeysSingleObjectPtr<PolyLineType>(InPtr)
 	, bGlobalData(bInGlobalData)
 {}
 
-int32 FPCGAttributeAccessorKeysSplineData::GetNum() const
+template<typename PolyLineType>
+int32 FPCGAttributeAccessorKeysPolyLineData<PolyLineType>::GetNum() const
 {
-	if (!Ptr)
+	if (!this->Ptr)
 	{
 		return 0; 
 	}
@@ -380,20 +385,24 @@ int32 FPCGAttributeAccessorKeysSplineData::GetNum() const
 	}
 	else
 	{
-		return Ptr->SplineStruct.GetSplinePointsPosition().Points.Num();
+		return this->Ptr->GetNumVertices();
 	}
 }
 
-/////////////////////////////////////////////////////////
+// Instantiate templated classes
+template class FPCGAttributeAccessorKeysPolyLineData<UPCGSplineData>;
+template class FPCGAttributeAccessorKeysPolyLineData<UPCGPolygon2DData>;
 
-FPCGAttributeAccessorKeysSplineDataEntries::FPCGAttributeAccessorKeysSplineDataEntries(const UPCGSplineData* InSplineData)
-	: Ptr(InSplineData)
+/////////////////////////////////////////////////////////
+template<typename PolyLineType>
+FPCGAttributeAccessorKeysPolyLineDataEntries<PolyLineType>::FPCGAttributeAccessorKeysPolyLineDataEntries(const PolyLineType* InPolyLineData)
+	: Ptr(InPolyLineData)
 {
-	check(InSplineData);
+	check(InPolyLineData);
 
 	bIsReadOnly = true;
 	// const_cast the view
-	TConstArrayView<PCGMetadataEntryKey> ConstEntries = InSplineData->SplineStruct.GetConstControlPointsEntryKeys();
+	TConstArrayView<PCGMetadataEntryKey> ConstEntries = InPolyLineData->GetConstVerticesEntryKeys();
 	Entries = MakeArrayView(const_cast<PCGMetadataEntryKey*>(ConstEntries.GetData()), ConstEntries.Num());
 	
 	if (Entries.IsEmpty())
@@ -403,17 +412,23 @@ FPCGAttributeAccessorKeysSplineDataEntries::FPCGAttributeAccessorKeysSplineDataE
 	}
 }
 
-FPCGAttributeAccessorKeysSplineDataEntries::FPCGAttributeAccessorKeysSplineDataEntries(UPCGSplineData* InSplineData)
-	: Ptr(InSplineData)
+template<typename PolyLineType>
+FPCGAttributeAccessorKeysPolyLineDataEntries<PolyLineType>::FPCGAttributeAccessorKeysPolyLineDataEntries(PolyLineType* InPolyLineData)
+	: Ptr(InPolyLineData)
 {
-	check(InSplineData);
+	check(InPolyLineData);
 	
 	bIsReadOnly = false;
-	InSplineData->SplineStruct.AllocateMetadataEntries();
-	Entries = InSplineData->SplineStruct.GetMutableControlPointsEntryKeys();
+	InPolyLineData->AllocateMetadataEntries();
+	Entries = InPolyLineData->GetMutableVerticesEntryKeys();
 }
 
-int32 FPCGAttributeAccessorKeysSplineDataEntries::GetNum() const
+template<typename PolyLineType>
+int32 FPCGAttributeAccessorKeysPolyLineDataEntries<PolyLineType>::GetNum() const
 {
-	return Ptr ? Ptr->SplineStruct.GetSplinePointsPosition().Points.Num() : 0;
+	return Ptr ? Ptr->GetNumVertices() : 0;
 }
+
+// Instantiate templated classes
+template class FPCGAttributeAccessorKeysPolyLineDataEntries<UPCGSplineData>;
+template class FPCGAttributeAccessorKeysPolyLineDataEntries<UPCGPolygon2DData>;
\ No newline at end of file
diff --git a/Engine/Plugins/PCG/Source/PCG/Private/PCGModule.cpp b/Engine/Plugins/PCG/Source/PCG/Private/PCGModule.cpp
--- a/Engine/Plugins/PCG/Source/PCG/Private/PCGModule.cpp
+++ b/Engine/Plugins/PCG/Source/PCG/Private/PCGModule.cpp
@@ -8,7 +8,9 @@
 #include "Compute/PCGComputeKernel.h"
 #include "Data/PCGBasePointData.h"
 #include "Data/PCGPointData.h"
+#include "Data/PCGPolyLineData.h"
 #include "Data/PCGSplineData.h"
+#include "Data/PCGPolygon2DData.h"
 
 #include "Interfaces/IPluginManager.h"
 #include "Misc/CoreDelegates.h"
@@ -72,6 +74,7 @@ void FPCGModule::StartupModule()
 	AttributeAccessorFactory.RegisterMethods<UPCGBasePointData>(UPCGBasePointData::GetPointAccessorMethods());
 	// @todo_pcg: Eventually remove the UPCGPointData method registration because the UPCGBasePointData accessors are compatible
 	AttributeAccessorFactory.RegisterMethods<UPCGPointData>(UPCGPointData::GetPointAccessorMethods());
+	AttributeAccessorFactory.RegisterMethods<UPCGPolygon2DData>(UPCGPolygon2DData::GetPolygon2DAccessorMethods());
 	AttributeAccessorFactory.RegisterMethods<UPCGSplineData>(UPCGSplineData::GetSplineAccessorMethods());
 
 #if WITH_EDITOR
@@ -98,9 +101,10 @@ void FPCGModule::PreExit()
 {
 	// Unregistering accessor methods
 	AttributeAccessorFactory.UnregisterMethods<UPCGSplineData>();
+	AttributeAccessorFactory.UnregisterMethods<UPCGPolygon2DData>();
 	AttributeAccessorFactory.UnregisterMethods<UPCGPointData>();
-	AttributeAccessorFactory.UnregisterDefaultMethods();
 	AttributeAccessorFactory.UnregisterMethods<UPCGBasePointData>();
+	AttributeAccessorFactory.UnregisterDefaultMethods();
 
 #if WITH_EDITOR
 	DeregisterNativeElementDeterminismTests();
diff --git a/Engine/Plugins/PCG/Source/PCG/Public/Data/PCGPolyLineData.h b/Engine/Plugins/PCG/Source/PCG/Public/Data/PCGPolyLineData.h
--- a/Engine/Plugins/PCG/Source/PCG/Public/Data/PCGPolyLineData.h
+++ b/Engine/Plugins/PCG/Source/PCG/Public/Data/PCGPolyLineData.h
@@ -2,7 +2,7 @@
 
 #pragma once
 
-#include "PCGSpatialData.h"
+#include "Data/PCGSpatialData.h"
 
 #include "PCGPolyLineData.generated.h"
 
@@ -31,6 +31,9 @@ public:
 	UFUNCTION(BlueprintCallable, Category = "PolylineData")
 	virtual int GetNumSegments() const PURE_VIRTUAL(UPCGPolyLineData::GetNumSegments, return 0;);
 
+	UFUNCTION(BlueprintCallable, Category = "PolylineData")
+	UE_API virtual int GetNumVertices() const;
+
 	/** Get the length of a specific segment of the line. */
 	UFUNCTION(BlueprintCallable, Category = "PolylineData")
 	virtual double GetSegmentLength(int SegmentIndex) const PURE_VIRTUAL(UPCGPolyLineData::GetSegmentLength, return 0;);
@@ -96,6 +99,8 @@ public:
 	virtual void WriteMetadataToPoint(float InputKey, FPCGPoint& OutPoint, UPCGMetadata* OutMetadata) const {}
 
 	virtual TConstArrayView<PCGMetadataEntryKey> GetConstVerticesEntryKeys() const { return {}; }
+	virtual void AllocateMetadataEntries() {}
+	virtual TArrayView<PCGMetadataEntryKey> GetMutableVerticesEntryKeys() { return {}; }
 };
 
 #undef UE_API
diff --git a/Engine/Plugins/PCG/Source/PCG/Public/Data/PCGPolygon2DData.h b/Engine/Plugins/PCG/Source/PCG/Public/Data/PCGPolygon2DData.h
--- a/Engine/Plugins/PCG/Source/PCG/Public/Data/PCGPolygon2DData.h
+++ b/Engine/Plugins/PCG/Source/PCG/Public/Data/PCGPolygon2DData.h
@@ -5,11 +5,44 @@
 #include "Data/PCGPolyLineData.h"
 
 #include "PCGPointData.h"
+#include "Metadata/PCGMetadataCommon.h"
 
 #include "Curve/GeneralPolygon2.h"
 
 #include "PCGPolygon2DData.generated.h"
 
+UENUM()
+enum class EPCGPolygon2DProperties : uint8
+{
+	Position UMETA(Tooltip = "Location of the vertex in world coordinates."),
+	Rotation UMETA(Tooltip = "Rotation of the vertex in world coordinates."),
+	SegmentIndex UMETA(Tooltip = "Segment index of the vertex in the current polygon (outer or hole)"),
+	HoleIndex UMETA(Tooltip = "Index of the hole the vertex is inside of. -1 on the Outer polygon."),
+	SegmentLength UMETA(Tooltip = "Length of the segment starting from the vertex."),
+	LocalPosition UMETA(Tooltip = "Location of the vertex in 2d coordinates."),
+	LocalRotation UMETA(Tooltip = "Rotation of the vertex in 2d coordinates.")
+};
+
+UENUM()
+enum class EPCGPolygon2DDataProperties : uint8
+{
+	Transform UMETA(Tooltip = "Transform of the 2d polygon.", PCG_MetadataDomain="Data"),
+	Area UMETA(Tooltip = "Area of the 2d polygon.", PCG_MetadataDomain = "Data"),
+	Perimeter UMETA(Tooltip = "Perimeter of the 2d polygon.", PCG_MetadataDomain = "Data"),
+	BoundsMin UMETA(Tooltip = "Minimum point of the bounds of the 2d polygon in 2d space.", PCG_MetadataDomain = "Data"),
+	BoundsMax UMETA(Tooltip = "Minimum point of the bounds of the 2d polygon in 2d space.", PCG_MetadataDomain = "Data"),
+	SegmentCount UMETA(Tooltip = "Number of segments in the 2d polygon (includes hole segments).", PCG_MetadataDomain = "Data"),
+	OuterSegmentCount UMETA(Tooltip = "Number of segments in the 2d polygon (only the outer polygon).", PCG_MetadataDomain = "Data"),
+	HoleCount UMETA(Tooltip = "Number of holes in the 2d polygon.", PCG_MetadataDomain = "Data"),
+	LongestOuterSegmentIndex UMETA(Tooltip = "Index of the longest segment in the polygon.", PCG_MetadataDomain = "Data"),
+	IsClockwise UMETA(Tooltip = "Returns whether the 2d polygon is clockwise or counter-clocwise.", PCG_MetadataDomain = "Data")
+};
+
+namespace PCGPolygon2DData
+{
+	const FName VertexDomainName = "Vertices";
+}
+
 /** Data representing a single 2D polygon with a 3D transform (for spatial operations). */
 UCLASS(MinimalAPI, BlueprintType, ClassGroup = (Procedural))
 class UPCGPolygon2DData : public UPCGPolyLineData
@@ -17,14 +50,20 @@ class UPCGPolygon2DData : public UPCGPolyLineData
 	GENERATED_BODY()
 
 public:
+	PCG_API UPCGPolygon2DData(const FObjectInitializer& ObjectInitializer);
+
 	// ~Begin UObject interface
 	PCG_API virtual void Serialize(FArchive& InArchive) override;
 	// ~End UObject interface
 
 	// ~Begin UPCGData interface
 	virtual EPCGDataType GetDataType() const override { return EPCGDataType::Polygon2D; }
-	// @todo_pcg: supported domain ids.. elements?
 	PCG_API virtual void AddToCrc(FArchiveCrc32& Ar, bool bFullDataCrc) const override;
+
+	virtual FPCGMetadataDomainID GetDefaultMetadataDomainID() const override { return PCGMetadataDomainID::Elements; }
+	virtual TArray<FPCGMetadataDomainID> GetAllSupportedMetadataDomainIDs() const override { return { PCGMetadataDomainID::Data, PCGMetadataDomainID::Elements }; }
+	PCG_API virtual FPCGMetadataDomainID GetMetadataDomainIDFromSelector(const FPCGAttributePropertySelector& InSelector) const override;
+	PCG_API virtual bool SetDomainFromDomainID(const FPCGMetadataDomainID& InDomainID, FPCGAttributePropertySelector& InOutSelector) const;
 	// ~End UPCGData interface
 
 	// ~Begin UPCGSpatialData interface
@@ -33,9 +72,10 @@ public:
 	PCG_API virtual bool SamplePoint(const FTransform& Transform, const FBox& Bounds, FPCGPoint& OutPoint, UPCGMetadata* OutMetadata) const override;
 	PCG_API virtual bool ProjectPoint(const FTransform& Transform, const FBox& Bounds, const FPCGProjectionParams& Params, FPCGPoint& OutPoint, UPCGMetadata* OutMetadata) const override;
 
+	PCG_API virtual void WriteMetadataToPoint(float InputKey, FPCGPoint& OutPoint, UPCGMetadata* OutMetadata) const override;
+
 protected:
 	virtual UPCGSpatialData* CopyInternal(FPCGContext* Context) const override;
-	// @todo_pcg: initialize target metadata?
 	// ~End UPCGSpatialData interface
 
 public:
@@ -47,22 +87,39 @@ public:
 	// ~Begin UPCGPolyLineData interface - all these functions will return world space information (e.g. including scale, ...) unless specified
 	virtual FTransform GetTransform() const override { return Transform; }
 	PCG_API virtual int GetNumSegments() const override;
+	PCG_API virtual int GetNumVertices() const override { return GetNumSegments(); }
 	PCG_API virtual double GetSegmentLength(int SegmentIndex) const override;
 	PCG_API virtual FVector GetLocationAtAlpha(float Alpha) const override;
 	PCG_API virtual FTransform GetTransformAtAlpha(float Alpha) const override;
 	PCG_API virtual FTransform GetTransformAtDistance(int SegmentIndex, double Distance, bool bWorldSpace = true, FBox* OutBounds = nullptr) const override;
 	PCG_API virtual float GetInputKeyAtDistance(int SegmentIndex, double Distance) const override;
 	PCG_API virtual double GetDistanceAtSegmentStart(int SegmentIndex) const override;
-	// @todo_pcg: write metadata to point + get const vertices entry keys
+	virtual bool IsClosed() const { return true; }
+
+	virtual TConstArrayView<PCGMetadataEntryKey> GetConstVerticesEntryKeys() const override { return VertexEntryKeys; }
+	PCG_API virtual void AllocateMetadataEntries() override;
+	virtual TArrayView<PCGMetadataEntryKey> GetMutableVerticesEntryKeys() override { return VertexEntryKeys; }
 	// ~End UPCGPolyLineData interface
 
-	// @todo_pcg maybe static create attribute accessor methods if we support elements
+	/** Static helper to create an accessor on a data that doesn't yet exist, as accessors for spline data don't rely on existing data. */
+	static PCG_API TUniquePtr<IPCGAttributeAccessor> CreateStaticAccessor(const UPCGPolygon2DData* InData, const FPCGAttributePropertySelector& InSelector, bool bQuiet = false, bool bConst = false);
+
+	/** Helper to create an accessor keys for data that doesn't yet exist. */
+	PCG_API TUniquePtr<IPCGAttributeAccessorKeys> CreateAccessorKeys(const FPCGAttributePropertySelector& InSelector, bool bQuiet = false);
+	PCG_API TUniquePtr<const IPCGAttributeAccessorKeys> CreateConstAccessorKeys(const FPCGAttributePropertySelector& InSelector, bool bQuiet = false) const;
+
+	// Get the functions to the accessor factory
+	static PCG_API FPCGAttributeAccessorMethods GetPolygon2DAccessorMethods();
+
 public:
 	const UE::Geometry::FGeneralPolygon2d& GetPolygon() const { return Polygon; }
+	const TMap<int, TPair<int, int>>& GetSegmentIndexToSegmentAndHoleIndices() const { return SegmentIndexToSegmentAndHoleIndices; }
+
+	PCG_API void SetPolygon(const UE::Geometry::FGeneralPolygon2d& InPolygon, TConstArrayView<PCGMetadataEntryKey>* InOptionalEntryKeys = nullptr);
+	PCG_API void SetPolygon(UE::Geometry::FGeneralPolygon2d&& InPolygon, TConstArrayView<PCGMetadataEntryKey>* InOptionalEntryKeys = nullptr);
+	PCG_API void SetTransform(const FTransform& InTransform, bool bCheckWinding = false);
 
-	void SetPolygon(const UE::Geometry::FGeneralPolygon2d& InPolygon);
-	void SetPolygon(UE::Geometry::FGeneralPolygon2d&& InPolygon);
-	void SetTransform(const FTransform& InTransform, bool bCheckWinding = false);
+	PCG_API void WriteMetadataToEntry(float InputKey, PCGMetadataEntryKey& OutEntryKey, UPCGMetadata* OutMetadata) const;
 
 protected:
 	const UPCGBasePointData* CreateBasePointData(FPCGContext* Context, TSubclassOf<UPCGBasePointData> PointDataClass) const;
@@ -74,10 +131,15 @@ protected:
 	// Note: can't be a property
 	UE::Geometry::FGeneralPolygon2d Polygon;
 
+	// In order, outer then holes
+	UPROPERTY()
+	TArray<int64> VertexEntryKeys;
+
 	UPROPERTY()
 	FTransform Transform;
 
 private:
 	TMap<int, TPair<int, int>> SegmentIndexToSegmentAndHoleIndices;
+	TMap<TPair<int, int>, int> SegmentAndHoleIndicesToSegmentIndex;
 	int SegmentCount = 0;
 };
diff --git a/Engine/Plugins/PCG/Source/PCG/Public/Data/PCGSplineData.h b/Engine/Plugins/PCG/Source/PCG/Public/Data/PCGSplineData.h
--- a/Engine/Plugins/PCG/Source/PCG/Public/Data/PCGSplineData.h
+++ b/Engine/Plugins/PCG/Source/PCG/Public/Data/PCGSplineData.h
@@ -3,9 +3,9 @@
 #pragma once
 
 #include "Math/InterpCurve.h"
-#include "PCGPolyLineData.h"
-#include "PCGProjectionData.h"
-#include "PCGSplineStruct.h"
+#include "Data/PCGPolyLineData.h"
+#include "Data/PCGProjectionData.h"
+#include "Data/PCGSplineStruct.h"
 #include "Elements/PCGProjectionParams.h"
 
 #include "PCGSplineData.generated.h"
@@ -87,7 +87,10 @@ public:
 	UE_API virtual FTransform GetTransformAtAlpha(float Alpha) const override;
 	virtual bool IsClosed() const override { return SplineStruct.bClosedLoop; }
 	UE_API virtual void WriteMetadataToPoint(float InputKey, FPCGPoint& OutPoint, UPCGMetadata* OutMetadata) const override;
+
+	virtual void AllocateMetadataEntries() override { SplineStruct.AllocateMetadataEntries(); }
 	virtual TConstArrayView<PCGMetadataEntryKey> GetConstVerticesEntryKeys() const override { return SplineStruct.GetConstControlPointsEntryKeys(); }
+	virtual TArrayView<PCGMetadataEntryKey> GetMutableVerticesEntryKeys() override { return SplineStruct.GetMutableControlPointsEntryKeys(); }
 	//~End UPCGPolyLineData interface
 	
 	/** Static helper to create an accessor on a data that doesn't yet exist, as accessors for spline data don't rely on existing data. */
diff --git a/Engine/Plugins/PCG/Source/PCG/Public/Metadata/Accessors/PCGAttributeAccessorKeys.h b/Engine/Plugins/PCG/Source/PCG/Public/Metadata/Accessors/PCGAttributeAccessorKeys.h
--- a/Engine/Plugins/PCG/Source/PCG/Public/Metadata/Accessors/PCGAttributeAccessorKeys.h
+++ b/Engine/Plugins/PCG/Source/PCG/Public/Metadata/Accessors/PCGAttributeAccessorKeys.h
@@ -11,6 +11,8 @@ class FPCGMetadataDomain;
 struct FPCGPoint;
 class UPCGMetadata;
 class UPCGSplineData;
+class UPCGPolygon2DData;
+class UPCGPolyLineData;
 class UPCGBasePointData;
 
 ///////////////////////////////////////////////////////////////////////
@@ -434,38 +436,51 @@ protected:
 };
 
 /////////////////////////////////////////////////////////////////
-
 /**
-* Unique Key around a single spline data.
+* Unique Key around a polyline data.
 */
-class FPCGAttributeAccessorKeysSplineData : public FPCGAttributeAccessorKeysSingleObjectPtr<UPCGSplineData>
+template<typename PolyLineType>
+class FPCGAttributeAccessorKeysPolyLineData : public FPCGAttributeAccessorKeysSingleObjectPtr<PolyLineType>
 {
 public:
-	FPCGAttributeAccessorKeysSplineData();
-	explicit FPCGAttributeAccessorKeysSplineData(UPCGSplineData* InPtr, bool bInGlobalData);
-	explicit FPCGAttributeAccessorKeysSplineData(const UPCGSplineData* InPtr, bool bInGlobalData);
-	
+	FPCGAttributeAccessorKeysPolyLineData();
+	explicit FPCGAttributeAccessorKeysPolyLineData(PolyLineType* InPtr, bool bInGlobalData);
+	explicit FPCGAttributeAccessorKeysPolyLineData(const PolyLineType* InPtr, bool bInGlobalData);
+
 	virtual int32 GetNum() const override;
-	
+
 private:
 	bool bGlobalData = false;
 };
 
+/**
+* Aliases for polyline keys
+*/
+typedef FPCGAttributeAccessorKeysPolyLineData<UPCGSplineData> FPCGAttributeAccessorKeysSplineData;
+typedef FPCGAttributeAccessorKeysPolyLineData<UPCGPolygon2DData> FPCGAttributeAccessorKeysPolygon2DData;
+
 /**
 * Keys for metadata on a spline data
 */
-class FPCGAttributeAccessorKeysSplineDataEntries : public FPCGAttributeAccessorKeysEntries
+template<typename PolyLineType>
+class FPCGAttributeAccessorKeysPolyLineDataEntries : public FPCGAttributeAccessorKeysEntries
 {
 public:
-	explicit FPCGAttributeAccessorKeysSplineDataEntries(const UPCGSplineData* InSplineData);
-	explicit FPCGAttributeAccessorKeysSplineDataEntries(UPCGSplineData* InSplineData);
+	explicit FPCGAttributeAccessorKeysPolyLineDataEntries(const PolyLineType* InPolyLineData);
+	explicit FPCGAttributeAccessorKeysPolyLineDataEntries(PolyLineType* InPolyLineData);
 
 	virtual int32 GetNum() const override;
 
 private:
-	const UPCGSplineData* Ptr = nullptr;
+	const PolyLineType* Ptr = nullptr;
 };
 
+/**
+* Aliases for polyline keys
+*/
+typedef FPCGAttributeAccessorKeysPolyLineDataEntries<UPCGSplineData> FPCGAttributeAccessorKeysSplineDataEntries;
+typedef FPCGAttributeAccessorKeysPolyLineDataEntries<UPCGPolygon2DData> FPCGAttributeAccessorKeysPolygon2DDataEntries;
+
 template <typename ObjectType>
 inline bool IPCGAttributeAccessorKeys::GetKeys(int32 InStart, TArrayView<ObjectType*> OutKeys)
 {
diff --git a/Engine/Plugins/PCG/Source/PCG/Public/Metadata/Accessors/PCGPolygon2DAccessor.h b/Engine/Plugins/PCG/Source/PCG/Public/Metadata/Accessors/PCGPolygon2DAccessor.h
new file mode 100644
--- /dev/null
+++ b/Engine/Plugins/PCG/Source/PCG/Public/Metadata/Accessors/PCGPolygon2DAccessor.h
@@ -0,0 +1,335 @@
+// Copyright Epic Games, Inc. All Rights Reserved.
+
+#pragma once
+
+#include "Data/PCGPolygon2DData.h"
+#include "Metadata/Accessors/PCGPropertyAccessor.h"
+
+/**
+* Templated accessor for polygon vertices accessor.
+*/
+template<typename T, EPCGPolygon2DProperties Target>
+class FPCGPolygon2DVerticesAccessor : public IPCGAttributeAccessorT<FPCGPolygon2DVerticesAccessor<T, Target>>
+{
+public:
+	using Type = T;
+	using Super = IPCGAttributeAccessorT<FPCGPolygon2DVerticesAccessor<T, Target>>;
+
+	FPCGPolygon2DVerticesAccessor(bool bIsReadOnly)
+		: Super(bIsReadOnly)
+	{
+	}
+
+	bool GetRangeImpl(TArrayView<Type> OutValues, int32 Index, const IPCGAttributeAccessorKeys& Keys) const
+	{
+		const void* ContainerKeys = nullptr;
+		TArrayView<const void*> ContainerKeysView(&ContainerKeys, 1);
+		if (!Keys.GetKeys(Index, ContainerKeysView))
+		{
+			return false;
+		}
+
+		// Validation to not access keys that are not the expected type. Done after the GetKeys, as we also want to discard other type of incompatible
+		// keys (like a Default Metadata entry key)
+		if (!ensure(Keys.IsClassSupported(UPCGPolygon2DData::StaticClass())))
+		{
+			return false;
+		}
+
+		const UPCGPolygon2DData* Data = static_cast<const UPCGPolygon2DData*>(ContainerKeys);
+		const FTransform& Transform = Data->GetTransform();
+		const UE::Geometry::FGeneralPolygon2d& Polygon = Data->GetPolygon();
+		const TMap<int, TPair<int, int>>& SegmentIndexToSegmentAndHoleIndices = Data->GetSegmentIndexToSegmentAndHoleIndices();
+
+		const int32 NumPoints = SegmentIndexToSegmentAndHoleIndices.Num();
+		if (NumPoints == 0)
+		{
+			return false;
+		}
+
+		for (int32 i = 0; i < OutValues.Num(); ++i)
+		{
+			const int32 CurrIndex = (Index + i) % NumPoints;
+			auto [SegmentIndex, HoleIndex] = SegmentIndexToSegmentAndHoleIndices[CurrIndex];
+
+			if constexpr (Target == EPCGPolygon2DProperties::Position)
+			{
+				static_assert(std::is_same_v<Type, FVector>);
+				const UE::Geometry::FSegment2d Segment = Polygon.Segment(SegmentIndex, HoleIndex);
+				OutValues[i] = Transform.TransformPosition(FVector(Segment.StartPoint(), 0.0));
+			}
+			else if constexpr (Target == EPCGPolygon2DProperties::Rotation)
+			{
+				static_assert(std::is_same_v<Type, FQuat>);
+				//@todo_pcg: there is a Polygon.GetNormal that we could use, but it blends normals across vertices.
+				const UE::Geometry::FSegment2d Segment = Polygon.Segment(SegmentIndex, HoleIndex);
+				OutValues[i] = Transform.TransformRotation(FRotationMatrix::MakeFromXZ(FVector(Segment.Direction, 0.0), FVector::UpVector).ToQuat());
+			}
+			else if constexpr (Target == EPCGPolygon2DProperties::SegmentIndex)
+			{
+				static_assert(std::is_same_v<Type, int32>);
+				OutValues[i] = SegmentIndex;
+			}
+			else if constexpr (Target == EPCGPolygon2DProperties::HoleIndex)
+			{
+				static_assert(std::is_same_v<Type, int32>);
+				OutValues[i] = HoleIndex;
+			}
+			else if constexpr (Target == EPCGPolygon2DProperties::SegmentLength)
+			{
+				static_assert(std::is_same_v<Type, double>);
+				const UE::Geometry::FSegment2d Segment = Polygon.Segment(SegmentIndex, HoleIndex);
+				OutValues[i] = Segment.Length();
+			}
+			else if constexpr (Target == EPCGPolygon2DProperties::LocalPosition)
+			{
+				static_assert(std::is_same_v<Type, FVector2d>);
+				const UE::Geometry::FSegment2d Segment = Polygon.Segment(SegmentIndex, HoleIndex);
+				OutValues[i] = Segment.StartPoint();
+			}
+			else if constexpr (Target == EPCGPolygon2DProperties::LocalRotation)
+			{
+				static_assert(std::is_same_v<Type, FQuat>);
+				const UE::Geometry::FSegment2d Segment = Polygon.Segment(SegmentIndex, HoleIndex);
+				OutValues[i] = FRotationMatrix::MakeFromXZ(FVector(Segment.Direction, 0.0), FVector::UpVector).ToQuat();
+			}
+			else
+			{
+				// Pitfall static assert
+				static_assert(!std::is_same_v<Type, Type>);
+			}
+		}
+
+		return true;
+	}
+
+	bool SetRangeImpl(TArrayView<const Type> InValues, int32 Index, IPCGAttributeAccessorKeys& Keys, EPCGAttributeAccessorFlags)
+	{
+		void* ContainerKeys = nullptr;
+		TArrayView<void*> ContainerKeysView(&ContainerKeys, 1);
+		if (!Keys.GetKeys(Index, ContainerKeysView))
+		{
+			return false;
+		}
+
+		// Validation to not access keys that are not the expected type. Done after the GetKeys, as we also want to discard other type of incompatible
+		// keys (like a Default Metadata entry key)
+		if (!ensure(Keys.IsClassSupported(UPCGPolygon2DData::StaticClass())))
+		{
+			return false;
+		}
+
+		UPCGPolygon2DData* Data = static_cast<UPCGPolygon2DData*>(ContainerKeys);
+		const FTransform& Transform = Data->GetTransform();
+		const UE::Geometry::FGeneralPolygon2d& Polygon = Data->GetPolygon();
+		const TMap<int, TPair<int, int>>& SegmentIndexToSegmentAndHoleIndices = Data->GetSegmentIndexToSegmentAndHoleIndices();
+
+		const int32 NumPoints = SegmentIndexToSegmentAndHoleIndices.Num();
+		if (NumPoints == 0)
+		{
+			return false;
+		}
+
+		if constexpr (Target != EPCGPolygon2DProperties::Position && Target != EPCGPolygon2DProperties::LocalPosition)
+		{
+			return false;
+		}
+		else
+		{
+			FTransform InverseTransform;
+
+			if constexpr (Target == EPCGPolygon2DProperties::Position)
+			{
+				InverseTransform = Transform.Inverse();
+			}
+
+			// Need to recreate a new polygon from scratch here.
+			UE::Geometry::FGeneralPolygon2d NewPolygon;
+			TArray<TArray<FVector2d>> Vertices;
+			Vertices.Reserve(1 + Polygon.GetHoles().Num());
+			Vertices.Add(Polygon.GetOuter().GetVertices());
+			for (const UE::Geometry::FPolygon2d& Hole : Polygon.GetHoles())
+			{
+				Vertices.Add(Hole.GetVertices());
+			}
+
+			for (int i = 0; i < InValues.Num(); ++i)
+			{
+				const int32 CurrIndex = (Index + i) % NumPoints;
+				auto [SegmentIndex, HoleIndex] = SegmentIndexToSegmentAndHoleIndices[CurrIndex];
+
+				if constexpr (Target == EPCGPolygon2DProperties::Position)
+				{
+					static_assert(std::is_same_v<Type, FVector3d>);
+					FVector LocalPosition = InverseTransform.TransformPosition(InValues[i]);
+					Vertices[HoleIndex + 1][SegmentIndex] = FVector2d(LocalPosition.X, LocalPosition.Y);
+				}
+				else if constexpr (Target == EPCGPolygon2DProperties::LocalPosition)
+				{
+					static_assert(std::is_same_v<Type, FVector2d>);
+					Vertices[HoleIndex + 1][SegmentIndex] = InValues[i];
+				}
+				else
+				{
+					// Pitfall static assert
+					static_assert(!std::is_same_v<Type, Type>);
+				}
+			}
+
+			NewPolygon.SetOuter(UE::Geometry::FPolygon2d(MoveTemp(Vertices[0])));
+			for (int i = 1; i < Vertices.Num(); ++i)
+			{
+				NewPolygon.AddHole(UE::Geometry::FPolygon2d(MoveTemp(Vertices[i])));
+			}
+
+			Data->SetPolygon(NewPolygon);
+
+			return true;
+		}
+	}
+};
+
+/**
+* Templated accessor for global polygon2d data.
+*/
+template<typename T, EPCGPolygon2DDataProperties Target>
+class FPCGPolygon2DDataAccessor : public IPCGAttributeAccessorT<FPCGPolygon2DDataAccessor<T, Target>>
+{
+public:
+	using Type = T;
+	using Super = IPCGAttributeAccessorT<FPCGPolygon2DDataAccessor<T, Target>>;
+
+	FPCGPolygon2DDataAccessor(bool bIsReadOnly)
+		: Super(bIsReadOnly)
+	{
+	}
+
+	bool GetRangeImpl(TArrayView<Type> OutValues, int32 Index, const IPCGAttributeAccessorKeys& Keys) const
+	{
+		const void* ContainerKeys = nullptr;
+		TArrayView<const void*> ContainerKeysView(&ContainerKeys, 1);
+		if (!Keys.GetKeys(Index, ContainerKeysView))
+		{
+			return false;
+		}
+
+		// Validation to not access keys that are not the expected type. Done after the GetKeys, as we also want to discard other type of incompatible
+		// keys (like a Default Metadata entry key)
+		if (!ensure(Keys.IsClassSupported(UPCGPolygon2DData::StaticClass())))
+		{
+			return false;
+		}
+
+		const UPCGPolygon2DData* Data = static_cast<const UPCGPolygon2DData*>(ContainerKeys);
+		const FTransform& Transform = Data->GetTransform();
+		const UE::Geometry::FGeneralPolygon2d& Polygon = Data->GetPolygon();
+
+		for (int32 i = 0; i < OutValues.Num(); ++i)
+		{
+			if constexpr (Target == EPCGPolygon2DDataProperties::Transform)
+			{
+				static_assert(std::is_same_v<Type, FTransform>);
+				OutValues[i] = Transform;
+			}
+			else if constexpr (Target == EPCGPolygon2DDataProperties::Area)
+			{
+				static_assert(std::is_same_v<Type, double>);
+				OutValues[i] = Polygon.SignedArea();
+			}
+			else if constexpr (Target == EPCGPolygon2DDataProperties::Perimeter)
+			{
+				static_assert(std::is_same_v<Type, double>);
+				OutValues[i] = Polygon.Perimeter();
+			}
+			else if constexpr (Target == EPCGPolygon2DDataProperties::BoundsMin)
+			{
+				static_assert(std::is_same_v<Type, FVector2d>);
+				OutValues[i] = Polygon.Bounds().Min;
+			}
+			else if constexpr (Target == EPCGPolygon2DDataProperties::BoundsMax)
+			{
+				static_assert(std::is_same_v<Type, FVector2d>);
+				OutValues[i] = Polygon.Bounds().Max;
+			}
+			else if constexpr (Target == EPCGPolygon2DDataProperties::SegmentCount)
+			{
+				static_assert(std::is_same_v<Type, int32>);
+				OutValues[i] = Data->GetNumSegments();
+			}
+			else if constexpr (Target == EPCGPolygon2DDataProperties::OuterSegmentCount)
+			{
+				static_assert(std::is_same_v<Type, int32>);
+				OutValues[i] = Polygon.GetOuter().GetVertices().Num();
+			}
+			else if constexpr (Target == EPCGPolygon2DDataProperties::HoleCount)
+			{
+				static_assert(std::is_same_v<Type, int32>);
+				OutValues[i] = Polygon.GetHoles().Num();
+			}
+			else if constexpr (Target == EPCGPolygon2DDataProperties::LongestOuterSegmentIndex)
+			{
+				int32 LongestOuterSegmentIndex = INDEX_NONE;
+				double MaxLength = -DBL_MAX;
+				const int32 OuterSegmentCount = Polygon.GetOuter().GetVertices().Num();
+				for (int SegIndex = 0; SegIndex < OuterSegmentCount; ++SegIndex)
+				{
+					const double SegLength = Polygon.Segment(SegIndex, -1).Length();
+					if (SegLength > MaxLength)
+					{
+						MaxLength = SegLength;
+						LongestOuterSegmentIndex = SegIndex;
+					}
+				}
+
+				static_assert(std::is_same_v<Type, int32>);
+				OutValues[i] = LongestOuterSegmentIndex;
+			}
+			else if constexpr (Target == EPCGPolygon2DDataProperties::IsClockwise)
+			{
+				static_assert(std::is_same_v<Type, bool>);
+				OutValues[i] = Polygon.OuterIsClockwise();
+			}
+			else
+			{
+				// Pitfall static assert
+				static_assert(!std::is_same_v<Type, Type>);
+			}
+		}
+
+		return true;
+	}
+
+	bool SetRangeImpl(TArrayView<const Type> InValues, int32 Index, IPCGAttributeAccessorKeys& Keys, EPCGAttributeAccessorFlags AccessorFlags)
+	{
+		void* ContainerKeys = nullptr;
+		TArrayView<void*> ContainerKeysView(&ContainerKeys, 1);
+		if (!Keys.GetKeys(Index, ContainerKeysView))
+		{
+			return false;
+		}
+
+		// Validation to not access keys that are not the expected type. Done after the GetKeys, as we also want to discard other type of incompatible
+		// keys (like a Default Metadata entry key)
+		if (!ensure(Keys.IsClassSupported(UPCGPolygon2DData::StaticClass())))
+		{
+			return false;
+		}
+
+		UPCGPolygon2DData* Data = static_cast<UPCGPolygon2DData*>(ContainerKeys);
+
+		if constexpr (Target == EPCGPolygon2DDataProperties::Transform)
+		{
+			static_assert(std::is_same_v<Type, FTransform>);
+			for (int i = 0; i < InValues.Num(); ++i)
+			{
+				Data->SetTransform(InValues[i], /*bCheckWinding=*/true);
+			}
+
+			return true;
+		}
+		else
+		{
+			return false;
+		}
+	}
+};
\ No newline at end of file
diff --git a/Engine/Plugins/PCG/Source/PCG/Public/Metadata/Accessors/PCGSplineAccessor.h b/Engine/Plugins/PCG/Source/PCG/Public/Metadata/Accessors/PCGSplineAccessor.h
--- a/Engine/Plugins/PCG/Source/PCG/Public/Metadata/Accessors/PCGSplineAccessor.h
+++ b/Engine/Plugins/PCG/Source/PCG/Public/Metadata/Accessors/PCGSplineAccessor.h
@@ -2,8 +2,8 @@
 
 #pragma once
 
-#include "PCGPropertyAccessor.h"
 #include "Data/PCGSplineStruct.h"
+#include "Metadata/Accessors/PCGPropertyAccessor.h"
 
 #include "Components/SplineComponent.h"
 
diff --git a/Engine/Plugins/PCG/Source/PCGEditor/Private/DataVisualizations/PCGPolygon2DDataVisualization.cpp b/Engine/Plugins/PCG/Source/PCGEditor/Private/DataVisualizations/PCGPolygon2DDataVisualization.cpp
--- a/Engine/Plugins/PCG/Source/PCGEditor/Private/DataVisualizations/PCGPolygon2DDataVisualization.cpp
+++ b/Engine/Plugins/PCG/Source/PCGEditor/Private/DataVisualizations/PCGPolygon2DDataVisualization.cpp
@@ -4,11 +4,127 @@
 
 #include "Data/PCGPolygon2DData.h"
 #include "Data/PCGPolygon2DInteriorData.h"
+#include "DataVisualizations/PCGDataVisualizationHelpers.h"
 
 #include "AdvancedPreviewScene.h"
 #include "EditorViewportClient.h"
 #include "Components/SplineComponent.h"
 
+FPCGTableVisualizerInfo FPCGPolygon2DDataVisualization::GetTableVisualizerInfoWithDomain(const UPCGData* Data, const FPCGMetadataDomainID& DomainID) const
+{
+	using namespace PCGDataVisualizationHelpers;
+
+	const UPCGPolygon2DData* PolygonData = Cast<UPCGPolygon2DData>(Data);
+	if (!PolygonData)
+	{
+		if (const UPCGPolygon2DInteriorSurfaceData* PolygonSurfaceData = Cast<UPCGPolygon2DInteriorSurfaceData>(Data))
+		{
+			PolygonData = PolygonSurfaceData->GetPolygonData();
+		}
+	}
+
+	FPCGTableVisualizerInfo Info;
+	Info.Data = PolygonData;
+
+	if (!PolygonData)
+	{
+		return Info;
+	}
+
+	if (DomainID == PCGMetadataDomainID::Data)
+	{
+		// Column Sorting
+		FPCGAttributePropertySelector IndexSelector = FPCGAttributePropertySelector::CreateExtraPropertySelector(EPCGExtraProperties::Index, PCGDataConstants::DataDomainName);
+		AddColumnInfo(Info, PolygonData, IndexSelector);
+		Info.SortingColumn = Info.ColumnInfos.Last().Id;
+
+		// Columns
+		AddPropertyEnumColumnInfo<FTransform>(Info, PolygonData, EPCGPolygon2DDataProperties::Transform);
+		AddPropertyEnumColumnInfo<double>(Info, PolygonData, EPCGPolygon2DDataProperties::Area);
+		AddPropertyEnumColumnInfo<double>(Info, PolygonData, EPCGPolygon2DDataProperties::Perimeter);
+		AddPropertyEnumColumnInfo<FVector2d>(Info, PolygonData, EPCGPolygon2DDataProperties::BoundsMin);
+		AddPropertyEnumColumnInfo<FVector2d>(Info, PolygonData, EPCGPolygon2DDataProperties::BoundsMax);
+		AddPropertyEnumColumnInfo<int32>(Info, PolygonData, EPCGPolygon2DDataProperties::SegmentCount);
+		AddPropertyEnumColumnInfo<int32>(Info, PolygonData, EPCGPolygon2DDataProperties::OuterSegmentCount);
+		AddPropertyEnumColumnInfo<int32>(Info, PolygonData, EPCGPolygon2DDataProperties::HoleCount);
+		AddPropertyEnumColumnInfo<int32>(Info, PolygonData, EPCGPolygon2DDataProperties::LongestOuterSegmentIndex);
+		AddPropertyEnumColumnInfo<bool>(Info, PolygonData, EPCGPolygon2DDataProperties::IsClockwise);
+
+		// Add Metadata Columns
+		CreateMetadataColumnInfos(PolygonData, Info, PCGMetadataDomainID::Data);
+	}
+	else if (DomainID == PCGMetadataDomainID::Elements)
+	{
+		// Column Sorting
+		FPCGAttributePropertySelector IndexSelector = FPCGAttributePropertySelector::CreateExtraPropertySelector(EPCGExtraProperties::Index, PCGPolygon2DData::VertexDomainName);
+		AddColumnInfo(Info, PolygonData, IndexSelector);
+		Info.SortingColumn = Info.ColumnInfos.Last().Id;
+
+		// Columns
+		AddPropertyEnumColumnInfo<FVector>(Info, PolygonData, EPCGPolygon2DProperties::Position);
+		AddPropertyEnumColumnInfo<FRotator>(Info, PolygonData, EPCGPolygon2DProperties::Rotation);
+		AddPropertyEnumColumnInfo<int32>(Info, PolygonData, EPCGPolygon2DProperties::SegmentIndex);
+		AddPropertyEnumColumnInfo<int32>(Info, PolygonData, EPCGPolygon2DProperties::HoleIndex);
+		AddPropertyEnumColumnInfo<double>(Info, PolygonData, EPCGPolygon2DProperties::SegmentLength);
+		AddPropertyEnumColumnInfo<FVector2d>(Info, PolygonData, EPCGPolygon2DProperties::LocalPosition);
+		AddPropertyEnumColumnInfo<FRotator>(Info, PolygonData, EPCGPolygon2DProperties::LocalRotation);
+
+		// Add Metadata Columns
+		CreateMetadataColumnInfos(PolygonData, Info, PCGMetadataDomainID::Elements);
+
+		// Focus on data behavior
+		Info.FocusOnDataCallback = [](const UPCGData* Data, TArrayView<const int> Indices)
+		{
+			const UPCGPolygon2DData* PolygonData = Cast<UPCGPolygon2DData>(Data);
+			if (!PolygonData)
+			{
+				if (const UPCGPolygon2DInteriorSurfaceData* PolygonSurfaceData = Cast<UPCGPolygon2DInteriorSurfaceData>(Data))
+				{
+					PolygonData = PolygonSurfaceData->GetPolygonData();
+				}
+			}
+
+			if (!PolygonData)
+			{
+				return;
+			}
+
+			FBox BoundingBox(EForceInit::ForceInit);
+
+			if (Indices.IsEmpty())
+			{
+				BoundingBox = PolygonData->GetBounds();
+			}
+			else
+			{
+				const FTransform& Transform = PolygonData->GetTransform();
+				const TMap<int, TPair<int, int>>& VertexIndexToSegmentAndHoleIndices = PolygonData->GetSegmentIndexToSegmentAndHoleIndices();
+				const UE::Geometry::FGeneralPolygon2d& Polygon = PolygonData->GetPolygon();
+				const FVector HalfExtents(50.0);
+
+				for (const int& Index : Indices)
+				{
+					if (const TPair<int, int>* SegmentAndHoleIndex = VertexIndexToSegmentAndHoleIndices.Find(Index))
+					{
+						FVector2D VertexPosition = Polygon.Segment(SegmentAndHoleIndex->Key, SegmentAndHoleIndex->Value).StartPoint();
+						FVector VertexWorldPosition = Transform.TransformPosition(FVector(VertexPosition, 0.0));
+
+						FBox PointBoundingBox(VertexWorldPosition + HalfExtents, VertexWorldPosition - HalfExtents);
+						BoundingBox += PointBoundingBox;
+					}
+				}
+			}
+
+			if (GEditor && BoundingBox.IsValid)
+			{
+				GEditor->MoveViewportCamerasToBox(BoundingBox, /*bActiveViewportOnly=*/true, /*DrawDebugBoxTimeInSeconds=*/2.5f);
+			}
+		};
+	}
+
+	return Info;
+}
+
 FPCGSetupSceneFunc FPCGPolygon2DDataVisualization::GetViewportSetupFunc(const UPCGData* Data) const
 {
 	return[this, Data](FPCGSceneSetupParams& InOutParams)
diff --git a/Engine/Plugins/PCG/Source/PCGEditor/Public/DataVisualizations/PCGPolygon2DDataVisualization.h b/Engine/Plugins/PCG/Source/PCGEditor/Public/DataVisualizations/PCGPolygon2DDataVisualization.h
--- a/Engine/Plugins/PCG/Source/PCGEditor/Public/DataVisualizations/PCGPolygon2DDataVisualization.h
+++ b/Engine/Plugins/PCG/Source/PCGEditor/Public/DataVisualizations/PCGPolygon2DDataVisualization.h
@@ -11,6 +11,7 @@ class FPCGPolygon2DDataVisualization : public IPCGSpatialDataVisualization
 {
 public:
 	// ~Begin IPCGDataVisualization interface
+	PCGEDITOR_API virtual FPCGTableVisualizerInfo GetTableVisualizerInfoWithDomain(const UPCGData* Data, const FPCGMetadataDomainID& DomainID) const override;
 	PCGEDITOR_API virtual FPCGSetupSceneFunc GetViewportSetupFunc(const UPCGData* Data) const override;
 	// ~End IPCGDataVisualization interface
 };
\ No newline at end of file
