[PCG] Overhaul of type info for PCG types
* Add a new concept of FPCGDataTypeInfo, that is a UStruct and can be inherited to provide a hierarchy of type
* Each type will implement a new struct to inherit from it.
* All those types will be gathered in a registry that will build a hierarchy tree
* Add a new concept FPCGDataTypeIdentifier that is meant to replace EPCGDataType
* API is similar as the EPCGDataType for aggregationn/composition/intersection, but can query the registry
* Registry can be queried to find if 2 types are compatible with each other.
* For streamlining the integration, Macros are here to reduce the boilerplate code.
* All supported types are converted
* Backend only at the moment, no replacement of the existing system.
* Moved GetDataFunctionRegistry in the new Registry folder.

diff --git a/Engine/Plugins/PCG/Source/PCG/Internal/Compute/Data/PCGProxyForGPUData.h b/Engine/Plugins/PCG/Source/PCG/Internal/Compute/Data/PCGProxyForGPUData.h
--- a/Engine/Plugins/PCG/Source/PCG/Internal/Compute/Data/PCGProxyForGPUData.h
+++ b/Engine/Plugins/PCG/Source/PCG/Internal/Compute/Data/PCGProxyForGPUData.h
@@ -12,6 +12,14 @@
 
 struct FPCGContext;
 
+USTRUCT()
+struct FPCGDataTypeInfoProxyForGPU : public FPCGDataTypeInfo
+{
+	GENERATED_BODY()
+
+	PCG_DECLARE_LEGACY_TYPE_INFO(EPCGDataType::ProxyForGPU)
+};
+
 /** A proxy for data residing on the GPU with functionality to read the data back to the CPU. */
 UCLASS(MinimalAPI, ClassGroup = (Procedural), DisplayName = "GPU Proxy")
 class UPCGProxyForGPUData : public UPCGData
@@ -22,13 +30,15 @@ public:
 	UE_API void Initialize(TSharedPtr<FPCGProxyForGPUDataCollection> InDataCollection, int InDataIndexInCollection);
 
 	//~Begin UPCGData interface
-	virtual EPCGDataType GetDataType() const override { return EPCGDataType::ProxyForGPU; }
+	PCG_ASSIGN_DEFAULT_TYPE_INFO(FPCGDataTypeInfoProxyForGPU)
+
 	UE_API virtual void AddToCrc(FArchiveCrc32& Ar, bool bFullDataCrc) const override;
 	virtual bool CanBeSerialized() const override { return false; }
 	virtual bool HoldsTransientResources() const override { return true; }
 	virtual bool IsCacheable() const override { return false; }
 	UE_API virtual void ReleaseTransientResources(const TCHAR* InReason = nullptr);
 	UE_API virtual EPCGDataType GetUnderlyingDataType() const override;
+	UE_API virtual FPCGDataTypeBaseId GetUnderlyingDataTypeId() const override;
 	//~End UPCGData interface
 
 	/** Returns the GPU info. Returns null if the buffer has been discarded. */
diff --git a/Engine/Plugins/PCG/Source/PCG/Private/Compute/Data/PCGProxyForGPUData.cpp b/Engine/Plugins/PCG/Source/PCG/Private/Compute/Data/PCGProxyForGPUData.cpp
--- a/Engine/Plugins/PCG/Source/PCG/Private/Compute/Data/PCGProxyForGPUData.cpp
+++ b/Engine/Plugins/PCG/Source/PCG/Private/Compute/Data/PCGProxyForGPUData.cpp
@@ -6,6 +6,8 @@
 
 #include UE_INLINE_GENERATED_CPP_BY_NAME(PCGProxyForGPUData)
 
+PCG_DEFINE_TYPE_INFO(FPCGDataTypeInfoProxyForGPU, UPCGProxyForGPUData)
+
 void UPCGProxyForGPUData::Initialize(TSharedPtr<FPCGProxyForGPUDataCollection> InDataCollection, int InDataIndexInCollection)
 {
 	DataCollectionOnGPU = InDataCollection;
@@ -39,6 +41,12 @@ EPCGDataType UPCGProxyForGPUData::GetUnderlyingDataType() const
 	return EPCGDataType::None;
 }
 
+FPCGDataTypeBaseId UPCGProxyForGPUData::GetUnderlyingDataTypeId() const
+{
+	// @todo_pcg: Update GPU to use FPCGDataTypeBaseId instead of EPCGDataType
+	return FPCGDataTypeBaseId::MakeFromLegacyType(GetUnderlyingDataType());
+}
+
 TSharedPtr<const FPCGProxyForGPUDataCollection> UPCGProxyForGPUData::GetInputDataCollectionInfo() const
 {
 	return DataCollectionOnGPU;
diff --git a/Engine/Plugins/PCG/Source/PCG/Private/Data/PCGBasePointData.cpp b/Engine/Plugins/PCG/Source/PCG/Private/Data/PCGBasePointData.cpp
--- a/Engine/Plugins/PCG/Source/PCG/Private/Data/PCGBasePointData.cpp
+++ b/Engine/Plugins/PCG/Source/PCG/Private/Data/PCGBasePointData.cpp
@@ -5,6 +5,7 @@
 #include "PCGContext.h"
 #include "Data/PCGPointArrayData.h"
 #include "Data/PCGPointData.h"
+#include "Elements/PCGConvertToAttributeSet.h"
 #include "Helpers/PCGHelpers.h"
 #include "Helpers/PCGPointHelpers.h"
 #include "Helpers/PCGTagHelpers.h"
@@ -26,6 +27,8 @@ static TAutoConsoleVariable<bool> CVarCacheFullPointDataCrc(
 	true,
 	TEXT("Enable fine-grained CRC of point data for change tracking on elements that request it, rather than using data UID."));
 
+PCG_DEFINE_TYPE_INFO(FPCGDataTypeInfoPoint, UPCGBasePointData)
+
 namespace PCGPointHelpers
 {
 	bool GetDistanceRatios(const FTransform& InTransform, const FVector& InBoundsMin, const FVector& InBoundsMax, const float InSteepness, const FVector& InPosition, FVector& OutRatios)
@@ -423,6 +426,24 @@ namespace PCGBasePointData
 	}
 }
 
+bool FPCGDataTypeInfoPoint::SupportsConversion(const FPCGDataTypeIdentifier& Other) const
+{
+	// Point can convert to param
+	return Other == EPCGDataType::Param || FPCGDataTypeInfoConcrete::SupportsConversion(Other);
+}
+
+TSubclassOf<UPCGSettings> FPCGDataTypeInfoPoint::GetDefaultConversionSettings(const FPCGDataTypeIdentifier& Other) const
+{
+	if (Other == EPCGDataType::Param)
+	{
+		return UPCGConvertToAttributeSetSettings::StaticClass();
+	}
+	else
+	{
+		return FPCGDataTypeInfoConcrete::GetDefaultConversionSettings(Other);
+	}
+}
+
 UPCGBasePointData::UPCGBasePointData(const FObjectInitializer& ObjectInitializer)
 	: Super(ObjectInitializer)
 {
diff --git a/Engine/Plugins/PCG/Source/PCG/Private/Data/PCGDynamicMeshData.cpp b/Engine/Plugins/PCG/Source/PCG/Private/Data/PCGDynamicMeshData.cpp
--- a/Engine/Plugins/PCG/Source/PCG/Private/Data/PCGDynamicMeshData.cpp
+++ b/Engine/Plugins/PCG/Source/PCG/Private/Data/PCGDynamicMeshData.cpp
@@ -16,6 +16,8 @@
 
 #include UE_INLINE_GENERATED_CPP_BY_NAME(PCGDynamicMeshData)
 
+PCG_DEFINE_TYPE_INFO(FPCGDataTypeInfoDynamicMesh, UPCGDynamicMeshData)
+
 UPCGDynamicMeshData::UPCGDynamicMeshData(const FObjectInitializer& ObjectInitializer)
 	: UPCGSpatialData(ObjectInitializer)
 {
@@ -82,7 +84,10 @@ void UPCGDynamicMeshData::AddToCrc(FArchiveCrc32& Ar, bool bFullDataCrc) const
 
 	if (bFullDataCrc)
 	{
-		// Implementation note: metadata not supported at this point.
+		if (Metadata)
+		{
+			Metadata->Serialize(Ar);
+		}
 
 		FString ClassName = StaticClass()->GetPathName();
 		Ar << ClassName;
@@ -177,10 +182,10 @@ const UPCGBasePointData* UPCGDynamicMeshData::ToBasePointData(FPCGContext* Conte
 
 	DynamicMesh->ProcessMesh([&Vertices](const FDynamicMesh3& ReadMesh)
 		{
-			Vertices.Init(FVector::ZeroVector, ReadMesh.MaxVertexID());
+			Vertices.Reserve(ReadMesh.VertexCount());
 			for (int32 vid : ReadMesh.VertexIndicesItr())
 			{
-				Vertices[vid] = static_cast<FVector>(ReadMesh.GetVertex(vid));
+				Vertices.Add(ReadMesh.GetVertex(vid));
 			}
 		});
 
diff --git a/Engine/Plugins/PCG/Source/PCG/Private/Data/PCGGetDataFunctionRegistry.cpp b/Engine/Plugins/PCG/Source/PCG/Private/Data/PCGGetDataFunctionRegistry.cpp
deleted file mode 100644
diff --git a/Engine/Plugins/PCG/Source/PCG/Private/Data/PCGLandscapeData.cpp b/Engine/Plugins/PCG/Source/PCG/Private/Data/PCGLandscapeData.cpp
--- a/Engine/Plugins/PCG/Source/PCG/Private/Data/PCGLandscapeData.cpp
+++ b/Engine/Plugins/PCG/Source/PCG/Private/Data/PCGLandscapeData.cpp
@@ -33,6 +33,8 @@ namespace PCGLandscapeDataConstants
 	const FName ComponentYAttribute = TEXT("ComponentY");
 }
 
+PCG_DEFINE_TYPE_INFO(FPCGDataTypeInfoLandscape, UPCGLandscapeData)
+
 void UPCGLandscapeData::Initialize(const TArray<TWeakObjectPtr<ALandscapeProxy>>& InLandscapes, const FBox& InBounds, const FPCGLandscapeDataProps& InDataProps)
 {
 	TSet<ALandscapeProxy*> LandscapesToIgnore;
diff --git a/Engine/Plugins/PCG/Source/PCG/Private/Data/PCGLandscapeSplineData.cpp b/Engine/Plugins/PCG/Source/PCG/Private/Data/PCGLandscapeSplineData.cpp
--- a/Engine/Plugins/PCG/Source/PCG/Private/Data/PCGLandscapeSplineData.cpp
+++ b/Engine/Plugins/PCG/Source/PCG/Private/Data/PCGLandscapeSplineData.cpp
@@ -40,6 +40,8 @@ namespace PCGLandscapeDataHelpers
 	}
 }
 
+PCG_DEFINE_TYPE_INFO(FPCGDataTypeInfoLandscapeSpline, UPCGLandscapeSplineData)
+
 bool UPCGLandscapeSplineData::CheckSpline() const
 {
 	if (!Spline.IsValid())
diff --git a/Engine/Plugins/PCG/Source/PCG/Private/Data/PCGPolyLineData.cpp b/Engine/Plugins/PCG/Source/PCG/Private/Data/PCGPolyLineData.cpp
--- a/Engine/Plugins/PCG/Source/PCG/Private/Data/PCGPolyLineData.cpp
+++ b/Engine/Plugins/PCG/Source/PCG/Private/Data/PCGPolyLineData.cpp
@@ -5,6 +5,8 @@
 
 #include UE_INLINE_GENERATED_CPP_BY_NAME(PCGPolyLineData)
 
+PCG_DEFINE_TYPE_INFO(FPCGDataTypeInfoPolyline, UPCGPolyLineData)
+
 FBox UPCGPolyLineData::GetBounds() const
 {
 	FBox Bounds(EForceInit::ForceInit);
diff --git a/Engine/Plugins/PCG/Source/PCG/Private/Data/PCGPolygon2DData.cpp b/Engine/Plugins/PCG/Source/PCG/Private/Data/PCGPolygon2DData.cpp
--- a/Engine/Plugins/PCG/Source/PCG/Private/Data/PCGPolygon2DData.cpp
+++ b/Engine/Plugins/PCG/Source/PCG/Private/Data/PCGPolygon2DData.cpp
@@ -6,16 +6,20 @@
 #include "PCGCustomVersion.h"
 #include "Data/PCGPointData.h"
 #include "Data/PCGPointArrayData.h"
+#include "Elements/Polygon/PCGSurfaceFromPolygon2D.h"
 #include "Metadata/Accessors/PCGAttributeAccessorHelpers.h"
 #include "Metadata/Accessors/PCGPropertyAccessor.h"
 #include "Metadata/Accessors/PCGPolygon2DAccessor.h"
 
+
 #include "SegmentTypes.h"
 #include "Serialization/ArchiveCrc32.h"
 #include "UObject/Package.h"
 
 #define LOCTEXT_NAMESPACE "PCGPolygon2DData"
 
+PCG_DEFINE_TYPE_INFO(FPCGDataTypeInfoPolygon2D, UPCGPolygon2DData)
+
 namespace PCGPolygon2DData
 {
 	// @todo_pcg : could suggest to move this to the polygon implementation
@@ -169,6 +173,24 @@ UPCGPolygon2DData::UPCGPolygon2DData(const FObjectInitializer& ObjectInitializer
 	Metadata->SetupDomain(PCGMetadataDomainID::Elements, /*bIsDefault=*/true);
 }
 
+bool FPCGDataTypeInfoPolygon2D::SupportsConversion(const FPCGDataTypeIdentifier& Other) const
+{
+	// Polygon can convert to surface
+	return Other == EPCGDataType::Surface || FPCGDataTypeInfoPolyline::SupportsConversion(Other);
+}
+
+TSubclassOf<UPCGSettings> FPCGDataTypeInfoPolygon2D::GetDefaultConversionSettings(const FPCGDataTypeIdentifier& Other) const
+{
+	if (Other == EPCGDataType::Surface)
+	{
+		return UPCGCreateSurfaceFromPolygon2DSettings::StaticClass();
+	}
+	else
+	{
+		return FPCGDataTypeInfoPolyline::GetDefaultConversionSettings(Other);
+	}
+}
+
 void UPCGPolygon2DData::Serialize(FArchive& InArchive)
 {
 	Super::Serialize(InArchive);
diff --git a/Engine/Plugins/PCG/Source/PCG/Private/Data/PCGPrimitiveData.cpp b/Engine/Plugins/PCG/Source/PCG/Private/Data/PCGPrimitiveData.cpp
--- a/Engine/Plugins/PCG/Source/PCG/Private/Data/PCGPrimitiveData.cpp
+++ b/Engine/Plugins/PCG/Source/PCG/Private/Data/PCGPrimitiveData.cpp
@@ -11,6 +11,8 @@
 
 #include UE_INLINE_GENERATED_CPP_BY_NAME(PCGPrimitiveData)
 
+PCG_DEFINE_TYPE_INFO(FPCGDataTypeInfoPrimitive, UPCGPrimitiveData)
+
 void UPCGPrimitiveData::Initialize(UPrimitiveComponent* InPrimitive)
 {
 	check(InPrimitive);
diff --git a/Engine/Plugins/PCG/Source/PCG/Private/Data/PCGRenderTargetData.cpp b/Engine/Plugins/PCG/Source/PCG/Private/Data/PCGRenderTargetData.cpp
--- a/Engine/Plugins/PCG/Source/PCG/Private/Data/PCGRenderTargetData.cpp
+++ b/Engine/Plugins/PCG/Source/PCG/Private/Data/PCGRenderTargetData.cpp
@@ -8,6 +8,8 @@
 
 #include UE_INLINE_GENERATED_CPP_BY_NAME(PCGRenderTargetData)
 
+PCG_DEFINE_TYPE_INFO(FPCGDataTypeInfoRenderTarget, UPCGRenderTargetData)
+
 void UPCGRenderTargetData::Initialize(UTextureRenderTarget2D* InRenderTarget, const FTransform& InTransform, bool bInSkipReadbackToCPU, bool bInTakeOwnershipOfRenderTarget)
 {
 	RenderTarget = InRenderTarget;
diff --git a/Engine/Plugins/PCG/Source/PCG/Private/Data/PCGResourceData.cpp b/Engine/Plugins/PCG/Source/PCG/Private/Data/PCGResourceData.cpp
--- a/Engine/Plugins/PCG/Source/PCG/Private/Data/PCGResourceData.cpp
+++ b/Engine/Plugins/PCG/Source/PCG/Private/Data/PCGResourceData.cpp
@@ -6,6 +6,8 @@
 
 #include UE_INLINE_GENERATED_CPP_BY_NAME(PCGResourceData)
 
+PCG_DEFINE_TYPE_INFO(FPCGDataTypeInfoResource, UPCGResourceData)
+
 // @todo_pcg: A preferable API might be to hold the load handle as a member and provide functions to query the state.
 TSharedPtr<FStreamableHandle> UPCGResourceData::RequestResourceLoad(bool bAsynchronous) const
 {
diff --git a/Engine/Plugins/PCG/Source/PCG/Private/Data/PCGSpatialData.cpp b/Engine/Plugins/PCG/Source/PCG/Private/Data/PCGSpatialData.cpp
--- a/Engine/Plugins/PCG/Source/PCG/Private/Data/PCGSpatialData.cpp
+++ b/Engine/Plugins/PCG/Source/PCG/Private/Data/PCGSpatialData.cpp
@@ -9,6 +9,7 @@
 #include "Data/PCGPointArrayData.h"
 #include "Data/PCGProjectionData.h"
 #include "Data/PCGUnionData.h"
+#include "Elements/PCGCollapseElement.h"
 #include "Elements/PCGExecuteBlueprint.h"
 
 #include "Serialization/ArchiveCrc32.h"
@@ -34,6 +35,29 @@ FPCGInitializeFromDataParams::FPCGInitializeFromDataParams(const UPCGSpatialData
 {
 }
 
+PCG_DEFINE_TYPE_INFO(FPCGDataTypeInfoSpatial, UPCGSpatialData)
+
+bool FPCGDataTypeInfoSpatial::SupportsConversion(const FPCGDataTypeIdentifier& Other) const
+{
+	// Spatial can convert to points.
+	return Other == EPCGDataType::Point || FPCGDataTypeInfo::SupportsConversion(Other);
+}
+
+TSubclassOf<UPCGSettings> FPCGDataTypeInfoSpatial::GetDefaultConversionSettings(const FPCGDataTypeIdentifier& Other) const
+{
+	if (Other == EPCGDataType::Point)
+	{
+		return UPCGCollapseSettings::StaticClass();
+	}
+	else
+	{
+		return FPCGDataTypeInfo::GetDefaultConversionSettings(Other);
+	}
+}
+
+PCG_DEFINE_TYPE_INFO_WITHOUT_CLASS(FPCGDataTypeInfoConcrete)
+PCG_DEFINE_TYPE_INFO_WITHOUT_CLASS(FPCGDataTypeInfoComposite)
+
 UPCGSpatialData::UPCGSpatialData(const FObjectInitializer& ObjectInitializer)
 	: Super(ObjectInitializer)
 {
diff --git a/Engine/Plugins/PCG/Source/PCG/Private/Data/PCGSplineData.cpp b/Engine/Plugins/PCG/Source/PCG/Private/Data/PCGSplineData.cpp
--- a/Engine/Plugins/PCG/Source/PCG/Private/Data/PCGSplineData.cpp
+++ b/Engine/Plugins/PCG/Source/PCG/Private/Data/PCGSplineData.cpp
@@ -8,6 +8,7 @@
 #include "Data/PCGPolyLineData.h"
 #include "Data/PCGProjectionData.h"
 #include "Data/PCGSpatialData.h"
+#include "Elements/PCGCreateSurfaceFromSpline.h"
 #include "Elements/PCGSplineSampler.h"
 #include "Helpers/PCGHelpers.h"
 #include "Metadata/Accessors/PCGSplineAccessor.h"
@@ -19,6 +20,26 @@
 
 #define LOCTEXT_NAMESPACE "PCGSplineData"
 
+PCG_DEFINE_TYPE_INFO(FPCGDataTypeInfoSpline, UPCGSplineData)
+
+bool FPCGDataTypeInfoSpline::SupportsConversion(const FPCGDataTypeIdentifier& Other) const
+{
+	// Spline can convert to surface
+	return Other == EPCGDataType::Surface || FPCGDataTypeInfoPolyline::SupportsConversion(Other);
+}
+
+TSubclassOf<UPCGSettings> FPCGDataTypeInfoSpline::GetDefaultConversionSettings(const FPCGDataTypeIdentifier& Other) const
+{
+	if (Other == EPCGDataType::Surface)
+	{
+		return UPCGCreateSurfaceFromSplineSettings::StaticClass();
+	}
+	else
+	{
+		return FPCGDataTypeInfoPolyline::GetDefaultConversionSettings(Other);
+	}
+}
+
 UPCGSplineData::UPCGSplineData()
 {
 	check(Metadata);
diff --git a/Engine/Plugins/PCG/Source/PCG/Private/Data/PCGStaticMeshResourceData.cpp b/Engine/Plugins/PCG/Source/PCG/Private/Data/PCGStaticMeshResourceData.cpp
--- a/Engine/Plugins/PCG/Source/PCG/Private/Data/PCGStaticMeshResourceData.cpp
+++ b/Engine/Plugins/PCG/Source/PCG/Private/Data/PCGStaticMeshResourceData.cpp
@@ -4,6 +4,8 @@
 
 #include UE_INLINE_GENERATED_CPP_BY_NAME(PCGStaticMeshResourceData)
 
+PCG_DEFINE_TYPE_INFO(FPCGDataTypeInfoStaticMeshResource, UPCGStaticMeshResourceData)
+
 void UPCGStaticMeshResourceData::Initialize(TSoftObjectPtr<UStaticMesh> InStaticMesh)
 {
 	StaticMesh = InStaticMesh;
diff --git a/Engine/Plugins/PCG/Source/PCG/Private/Data/PCGSurfaceData.cpp b/Engine/Plugins/PCG/Source/PCG/Private/Data/PCGSurfaceData.cpp
--- a/Engine/Plugins/PCG/Source/PCG/Private/Data/PCGSurfaceData.cpp
+++ b/Engine/Plugins/PCG/Source/PCG/Private/Data/PCGSurfaceData.cpp
@@ -5,6 +5,8 @@
 
 #include UE_INLINE_GENERATED_CPP_BY_NAME(PCGSurfaceData)
 
+PCG_DEFINE_TYPE_INFO(FPCGDataTypeInfoSurface, UPCGSurfaceData)
+
 void UPCGSurfaceData::CopyBaseSurfaceData(UPCGSurfaceData* NewSurfaceData) const
 {
 	NewSurfaceData->Transform = Transform;
diff --git a/Engine/Plugins/PCG/Source/PCG/Private/Data/PCGTextureData.cpp b/Engine/Plugins/PCG/Source/PCG/Private/Data/PCGTextureData.cpp
--- a/Engine/Plugins/PCG/Source/PCG/Private/Data/PCGTextureData.cpp
+++ b/Engine/Plugins/PCG/Source/PCG/Private/Data/PCGTextureData.cpp
@@ -17,6 +17,9 @@
 
 #include UE_INLINE_GENERATED_CPP_BY_NAME(PCGTextureData)
 
+PCG_DEFINE_TYPE_INFO(FPCGDataTypeInfoBaseTexture, UPCGBaseTextureData)
+PCG_DEFINE_TYPE_INFO(FPCGDataTypeInfoTexture, UPCGTextureData)
+
 namespace PCGTextureSamplingHelpers
 {
 	static int32 GTriggerReadbackCaptureDispatches = 0;
diff --git a/Engine/Plugins/PCG/Source/PCG/Private/Data/PCGVirtualTextureData.cpp b/Engine/Plugins/PCG/Source/PCG/Private/Data/PCGVirtualTextureData.cpp
--- a/Engine/Plugins/PCG/Source/PCG/Private/Data/PCGVirtualTextureData.cpp
+++ b/Engine/Plugins/PCG/Source/PCG/Private/Data/PCGVirtualTextureData.cpp
@@ -11,6 +11,8 @@
 
 #include UE_INLINE_GENERATED_CPP_BY_NAME(PCGVirtualTextureData)
 
+PCG_DEFINE_TYPE_INFO(FPCGDataTypeInfoVirtualTexture, UPCGVirtualTextureData)
+
 void UPCGVirtualTextureData::Initialize(const URuntimeVirtualTextureComponent* InVirtualTextureComponent)
 {
 	check(InVirtualTextureComponent);
diff --git a/Engine/Plugins/PCG/Source/PCG/Private/Data/PCGVolumeData.cpp b/Engine/Plugins/PCG/Source/PCG/Private/Data/PCGVolumeData.cpp
--- a/Engine/Plugins/PCG/Source/PCG/Private/Data/PCGVolumeData.cpp
+++ b/Engine/Plugins/PCG/Source/PCG/Private/Data/PCGVolumeData.cpp
@@ -17,6 +17,8 @@
 
 #include UE_INLINE_GENERATED_CPP_BY_NAME(PCGVolumeData)
 
+PCG_DEFINE_TYPE_INFO(FPCGDataTypeInfoVolume, UPCGVolumeData)
+
 UPCGVolumeData::~UPCGVolumeData()
 {
 	ReleaseInternalBodyInstance();
diff --git a/Engine/Plugins/PCG/Source/PCG/Private/Data/Registry/PCGDataType.cpp b/Engine/Plugins/PCG/Source/PCG/Private/Data/Registry/PCGDataType.cpp
new file mode 100644
--- /dev/null
+++ b/Engine/Plugins/PCG/Source/PCG/Private/Data/Registry/PCGDataType.cpp
@@ -0,0 +1,186 @@
+// Copyright Epic Games, Inc. All Rights Reserved.
+
+#include "Data/Registry/PCGDataType.h"
+
+#include "PCGData.h"
+#include "PCGParamData.h"
+#include "PCGSettings.h"
+#include "Compute/Data/PCGProxyForGPUData.h"
+#include "Data/PCGBasePointData.h"
+#include "Data/PCGDynamicMeshData.h"
+#include "Data/PCGLandscapeData.h"
+#include "Data/PCGLandscapeSplineData.h"
+#include "Data/PCGPolygon2DData.h"
+#include "Data/PCGRenderTargetData.h"
+#include "Data/PCGPrimitiveData.h"
+#include "Data/PCGSpatialData.h"
+#include "Data/PCGSplineData.h"
+#include "Data/PCGStaticMeshResourceData.h"
+#include "Data/PCGTextureData.h"
+#include "Data/PCGVirtualTextureData.h"
+#include "Data/PCGVolumeData.h"
+#include "Helpers/PCGSubgraphHelpers.h"
+
+PCG_DEFINE_TYPE_INFO(FPCGDataTypeInfo, UPCGData)
+PCG_DEFINE_TYPE_INFO_WITHOUT_CLASS(FPCGDataTypeInfoOther)
+
+FPCGDataTypeBaseId::FPCGDataTypeBaseId(const UScriptStruct* InStruct)
+	: Struct(InStruct && InStruct->IsChildOf(FPCGDataTypeInfo::StaticStruct()) ? InStruct : nullptr)
+{
+}
+
+bool FPCGDataTypeBaseId::IsValid() const
+{
+	return Struct != nullptr;
+}
+
+bool FPCGDataTypeBaseId::IsChildOf(const FPCGDataTypeBaseId& Other) const
+{
+	return IsValid() && Other.IsValid() && Struct->IsChildOf(Other.Struct);
+}
+
+FString FPCGDataTypeBaseId::ToString() const
+{
+	if (!IsValid())
+	{
+		return TEXT("None");
+	}
+	else if (Struct == FPCGDataTypeInfo::StaticStruct())
+	{
+		return TEXT("Any");
+	}
+	else
+	{
+		// Strip `PCGDataTypeInfo` of the name
+		static const FString Prefix = FPCGDataTypeInfo::StaticStruct()->GetName();
+		FString Result = Struct->GetName();
+		Result.RemoveFromStart(Prefix);
+		return Result;
+	}
+}
+
+bool FPCGDataTypeInfo::SupportsConversion(const FPCGDataTypeIdentifier& Other) const
+{
+	// Any/Spatial/Composite can convert to concrete
+	if (Other == EPCGDataType::Concrete)
+	{
+		const EPCGDataType ThisType = GetAssociatedLegacyType();
+		return ThisType == EPCGDataType::Any || ThisType == EPCGDataType::Spatial || ThisType == EPCGDataType::Composite;
+	}
+	else
+	{
+		return false;
+	}
+}
+
+UPCGNode* FPCGDataTypeInfo::AddConversionNode(const FPCGDataTypeIdentifier& Other, UPCGGraph* InGraph, UPCGPin* InUpstreamPin, UPCGPin* InDownstreamPin) const
+{
+	if (const TSubclassOf<UPCGSettings> SettingsClass = GetDefaultConversionSettings(Other))
+	{
+		return FPCGSubgraphHelpers::SpawnNodeAndConnect(InGraph, InUpstreamPin, InDownstreamPin, SettingsClass);
+	}
+	else
+	{
+		return nullptr;
+	}
+}
+
+TSubclassOf<UPCGSettings> FPCGDataTypeInfo::GetDefaultConversionSettings(const FPCGDataTypeIdentifier& Other) const
+{
+	return nullptr;
+}
+
+FPCGDataTypeBaseId FPCGDataTypeBaseId::MakeFromLegacyType(EPCGDataType DataType)
+{
+	switch (DataType)
+	{
+	case EPCGDataType::Any:
+		return Construct<FPCGDataTypeInfo>();
+	case EPCGDataType::Point:
+		return Construct<FPCGDataTypeInfoPoint>();
+	case EPCGDataType::Spline:
+		return Construct<FPCGDataTypeInfoSpline>();
+	case EPCGDataType::LandscapeSpline:
+		return Construct<FPCGDataTypeInfoLandscapeSpline>();
+	case EPCGDataType::PolyLine:
+		return Construct<FPCGDataTypeInfoPolyline>();
+	case EPCGDataType::Polygon2D:
+		return Construct<FPCGDataTypeInfoPolygon2D>();
+	case EPCGDataType::Landscape:
+		return Construct<FPCGDataTypeInfoLandscape>();
+	case EPCGDataType::Texture:
+		return Construct<FPCGDataTypeInfoTexture>();
+	case EPCGDataType::RenderTarget:
+		return Construct<FPCGDataTypeInfoRenderTarget>();
+	case EPCGDataType::VirtualTexture:
+		return Construct<FPCGDataTypeInfoVirtualTexture>();
+	case EPCGDataType::BaseTexture:
+		return Construct<FPCGDataTypeInfoBaseTexture>();
+	case EPCGDataType::Surface:
+		return Construct<FPCGDataTypeInfoSurface>();
+	case EPCGDataType::Volume:
+		return Construct<FPCGDataTypeInfoVolume>();
+	case EPCGDataType::Primitive:
+		return Construct<FPCGDataTypeInfoPrimitive>();
+	case EPCGDataType::DynamicMesh:
+		return Construct<FPCGDataTypeInfoDynamicMesh>();
+	case EPCGDataType::StaticMeshResource:
+		return Construct<FPCGDataTypeInfoStaticMeshResource>();
+	case EPCGDataType::Concrete:
+		return Construct<FPCGDataTypeInfoConcrete>();
+	case EPCGDataType::Composite:
+		return Construct<FPCGDataTypeInfoComposite>();
+	case EPCGDataType::Spatial:
+		return Construct<FPCGDataTypeInfoSpatial>();
+	case EPCGDataType::ProxyForGPU:
+		return Construct<FPCGDataTypeInfoProxyForGPU>();
+	case EPCGDataType::Param:
+		return Construct<FPCGDataTypeInfoParam>();
+	// @todo_pcg To be uncommented when we do the switch. 
+	// case EPCGDataType::Settings:
+	// 	return Construct<FPCGDataTypeInfoSettings>();
+	case EPCGDataType::Other:
+		return Construct<FPCGDataTypeInfoOther>();
+	case EPCGDataType::None: // fall-through
+	default:
+		return FPCGDataTypeBaseId{};
+	}
+}
+
+EPCGDataType FPCGDataTypeBaseId::AsLegacyType() const
+{
+#define PCG_DATA_TYPE_TEST(TestStruct) if (Struct == TestStruct::StaticStruct()) { return TestStruct::GetStaticAssociatedLegacyType(); }
+		
+	PCG_DATA_TYPE_TEST(FPCGDataTypeInfo)
+	PCG_DATA_TYPE_TEST(FPCGDataTypeInfoSpatial)
+	PCG_DATA_TYPE_TEST(FPCGDataTypeInfoConcrete)
+	PCG_DATA_TYPE_TEST(FPCGDataTypeInfoPoint)
+	PCG_DATA_TYPE_TEST(FPCGDataTypeInfoComposite)
+	PCG_DATA_TYPE_TEST(FPCGDataTypeInfoParam)
+	PCG_DATA_TYPE_TEST(FPCGDataTypeInfoSpline)
+	PCG_DATA_TYPE_TEST(FPCGDataTypeInfoLandscapeSpline)
+	PCG_DATA_TYPE_TEST(FPCGDataTypeInfoPolyline)
+	PCG_DATA_TYPE_TEST(FPCGDataTypeInfoPolygon2D)
+	PCG_DATA_TYPE_TEST(FPCGDataTypeInfoLandscape)
+	PCG_DATA_TYPE_TEST(FPCGDataTypeInfoTexture)
+	PCG_DATA_TYPE_TEST(FPCGDataTypeInfoRenderTarget)
+	PCG_DATA_TYPE_TEST(FPCGDataTypeInfoVirtualTexture)
+	PCG_DATA_TYPE_TEST(FPCGDataTypeInfoBaseTexture)
+	PCG_DATA_TYPE_TEST(FPCGDataTypeInfoSurface)
+	PCG_DATA_TYPE_TEST(FPCGDataTypeInfoVolume)
+	PCG_DATA_TYPE_TEST(FPCGDataTypeInfoPrimitive)
+	PCG_DATA_TYPE_TEST(FPCGDataTypeInfoDynamicMesh)
+	PCG_DATA_TYPE_TEST(FPCGDataTypeInfoStaticMeshResource)
+	PCG_DATA_TYPE_TEST(FPCGDataTypeInfoConcrete)
+	PCG_DATA_TYPE_TEST(FPCGDataTypeInfoComposite)
+	PCG_DATA_TYPE_TEST(FPCGDataTypeInfoSpatial)
+	PCG_DATA_TYPE_TEST(FPCGDataTypeInfoProxyForGPU)
+	PCG_DATA_TYPE_TEST(FPCGDataTypeInfoParam)
+	// @todo_pcg To be uncommented when we do the switch.
+	// PCG_DATA_TYPE_TEST(FPCGDataTypeInfoSettings)
+	PCG_DATA_TYPE_TEST(FPCGDataTypeInfoOther)
+		
+#undef PCG_DATA_TYPE_TEST
+
+	return EPCGDataType::None;
+}
diff --git a/Engine/Plugins/PCG/Source/PCG/Private/Data/Registry/PCGDataTypeIdentifier.cpp b/Engine/Plugins/PCG/Source/PCG/Private/Data/Registry/PCGDataTypeIdentifier.cpp
new file mode 100644
--- /dev/null
+++ b/Engine/Plugins/PCG/Source/PCG/Private/Data/Registry/PCGDataTypeIdentifier.cpp
@@ -0,0 +1,312 @@
+// Copyright Epic Games, Inc. All Rights Reserved.
+
+#include "Data/Registry/PCGDataTypeIdentifier.h"
+
+#include "PCGContext.h"
+#include "PCGData.h"
+#include "PCGModule.h"
+
+#include "Algo/AnyOf.h"
+
+namespace PCGDataTypeIdentifier
+{
+	/**
+	* Sorted by hand, with the highest number of bits set first
+	* When we initialize an identifier from a EPCGDataType, we need to go in hierarchical order, by checking the highest in the hierarchy
+	* (aka. the ones with the most bits sets), to be sure to always have the most compact representation.
+	* (i.e. if there are EPCGDataType::Texture | EPCGDataType::RenderTarget, we have to make sure to check for EPCGDataType::BaseTexture first, as it is the parent)
+	*/
+	static constexpr const TStaticArray<EPCGDataType, 23>& HierarchyDataType{
+			EPCGDataType::Any,         // 30 bits
+			EPCGDataType::Spatial,     // 12 bits
+			EPCGDataType::Concrete,    // 11 bits
+			EPCGDataType::Surface,     // 4 bits
+			EPCGDataType::PolyLine,    // 3 bits
+			EPCGDataType::BaseTexture, // 2 bits
+		
+			EPCGDataType::Point,
+			EPCGDataType::Spline,
+			EPCGDataType::LandscapeSpline,
+			EPCGDataType::Polygon2D,
+			EPCGDataType::Landscape,
+			EPCGDataType::RenderTarget,
+			EPCGDataType::Texture,
+			EPCGDataType::VirtualTexture,
+			EPCGDataType::Volume,
+			EPCGDataType::Primitive,
+			EPCGDataType::DynamicMesh,
+			EPCGDataType::StaticMeshResource,
+			EPCGDataType::Composite,
+			EPCGDataType::ProxyForGPU,
+			EPCGDataType::Param,
+			EPCGDataType::Settings,
+			EPCGDataType::Other,
+		};
+}
+
+FPCGDataTypeIdentifier::FPCGDataTypeIdentifier(EPCGDataType LegacyType)
+{
+	FPCGDataTypeBaseId BaseId = FPCGDataTypeBaseId::MakeFromLegacyType(LegacyType);
+	if (BaseId.IsValid())
+	{
+		IDs.Emplace(BaseId);
+		return;
+	}
+
+	// Go through all the enum, from the root types to the child types.
+	for (EPCGDataType Type : PCGDataTypeIdentifier::HierarchyDataType)
+	{
+		if ((Type & LegacyType) == Type)
+		{
+			BaseId = FPCGDataTypeBaseId::MakeFromLegacyType(Type);
+			if (BaseId.IsValid())
+			{
+				IDs.Add(BaseId);
+				LegacyType &= ~Type;
+			}
+		}
+
+		if (!LegacyType)
+		{
+			break;
+		}
+	}
+}
+
+FPCGDataTypeIdentifier::FPCGDataTypeIdentifier(TSubclassOf<UPCGData> Class)
+	: FPCGDataTypeIdentifier(Class ? Class->GetDefaultObject<UPCGData>()->GetDataTypeId() : FPCGDataTypeBaseId{})
+{}
+
+FPCGDataTypeIdentifier::FPCGDataTypeIdentifier(TConstArrayView<FPCGDataTypeIdentifier> InIDs)
+{
+	for (const FPCGDataTypeIdentifier& ID : InIDs)
+	{
+		for (const FPCGDataTypeBaseId& BaseID : ID.GetIDs())
+		{
+			IDs.AddUnique(BaseID);
+		}
+	}
+}
+
+FPCGDataTypeIdentifier::FPCGDataTypeIdentifier(TConstArrayView<TSubclassOf<UPCGData>> Classes)
+{
+	for (TSubclassOf<UPCGData> Class : Classes)
+	{
+		IDs.AddUnique(Class->GetDefaultObject<UPCGData>()->GetDataTypeId());
+	}
+}
+
+FPCGDataTypeIdentifier::operator EPCGDataType() const
+{
+	EPCGDataType Result = EPCGDataType::None;
+	for (const FPCGDataTypeBaseId& ID : IDs)
+	{
+		Result |= ID.AsLegacyType();
+	}
+
+	return Result;
+}
+
+FPCGDataTypeIdentifier& FPCGDataTypeIdentifier::operator|=(EPCGDataType OtherLegacyType)
+{
+	*this = operator|(OtherLegacyType);
+	return *this;
+}
+
+FPCGDataTypeIdentifier& FPCGDataTypeIdentifier::operator|=(const FPCGDataTypeIdentifier& Other)
+{
+	for (const FPCGDataTypeBaseId& ID : Other.GetIDs())
+	{
+		if (ID.IsValid())
+		{
+			IDs.AddUnique(ID);
+		}
+	}
+	
+	return *this;
+}
+
+FPCGDataTypeIdentifier FPCGDataTypeIdentifier::operator|(EPCGDataType OtherLegacyType) const
+{
+	return *this | FPCGDataTypeIdentifier{OtherLegacyType};
+}
+
+FPCGDataTypeIdentifier FPCGDataTypeIdentifier::operator|(const FPCGDataTypeIdentifier& Other) const
+{
+	FPCGDataTypeIdentifier Result = *this;
+	Result |= Other;
+	return Result;
+}
+
+FPCGDataTypeIdentifier& FPCGDataTypeIdentifier::operator&=(EPCGDataType OtherLegacyType)
+{
+	*this = operator&(OtherLegacyType);
+	return *this;
+}
+
+FPCGDataTypeIdentifier& FPCGDataTypeIdentifier::operator&=(const FPCGDataTypeIdentifier& Other)
+{
+	TArray<FPCGDataTypeBaseId> CommonTypes;
+	
+	for (const FPCGDataTypeBaseId& ThisID : IDs)
+	{
+		if (!ensure(ThisID.IsValid()))
+		{
+			continue;
+		}
+		
+		for (const FPCGDataTypeBaseId& OtherID : Other.GetIDs())
+		{
+			if (!ensure(OtherID.IsValid()))
+			{
+				continue;
+			}
+			
+			if (ThisID.IsChildOf(OtherID))
+			{
+				CommonTypes.AddUnique(ThisID);
+			}
+			else if (OtherID.IsChildOf(ThisID))
+			{
+				CommonTypes.AddUnique(OtherID);
+			}
+		}
+	}
+
+	IDs = MoveTemp(CommonTypes);
+	
+	return *this;
+}
+
+FPCGDataTypeIdentifier FPCGDataTypeIdentifier::operator&(EPCGDataType OtherLegacyType) const
+{
+	return *this & FPCGDataTypeIdentifier{OtherLegacyType};
+}
+
+FPCGDataTypeIdentifier FPCGDataTypeIdentifier::operator&(const FPCGDataTypeIdentifier& Other) const
+{
+	FPCGDataTypeIdentifier Result = *this;
+	Result &= Other;
+	return Result;
+}
+
+FPCGDataTypeIdentifier FPCGDataTypeIdentifier::Compose(const FPCGDataTypeIdentifier& Other)
+{
+	return Compose({*this, Other});
+}
+
+FPCGDataTypeIdentifier FPCGDataTypeIdentifier::Compose(TConstArrayView<FPCGDataTypeIdentifier> IDs)
+{
+	return FPCGModule::GetConstDataTypeRegistry().GetIdentifiersComposition(IDs);
+}
+
+bool FPCGDataTypeIdentifier::operator!() const
+{
+	return !IsValid();
+}
+
+bool FPCGDataTypeIdentifier::SupportsType(EPCGDataType OtherLegacyType) const
+{
+	const FPCGDataTypeIdentifier OtherID(OtherLegacyType);
+	if (OtherID.IsComposition())
+	{
+		return Algo::AllOf(OtherID.GetIDs(), [this](const FPCGDataTypeBaseId& ID) { return IDs.Contains(ID); });
+	}
+	else
+	{
+		return SupportsType(FPCGDataTypeIdentifier(OtherLegacyType));
+	}
+}
+
+TConstArrayView<FPCGDataTypeBaseId> FPCGDataTypeIdentifier::GetIDs() const
+{
+	return IDs;
+};
+
+bool FPCGDataTypeIdentifier::IsSameType(const FPCGDataTypeIdentifier& Other) const
+{
+	return (!IsValid() && !Other.IsValid()) || (IDs.Num() == Other.IDs.Num() && SupportsType(Other));
+}
+
+bool FPCGDataTypeIdentifier::SupportsType(const FPCGDataTypeIdentifier& Other) const
+{
+	return IsValid() && Algo::AllOf(Other.GetIDs(), [this](const FPCGDataTypeBaseId& ID) { return IDs.Contains(ID); });
+}
+
+bool FPCGDataTypeIdentifier::IsWider(const FPCGDataTypeIdentifier& Other) const
+{
+	if (IsSameType(Other))
+	{
+		return false;
+	}
+	
+	const FPCGDataTypeIdentifier Intersection = *this & Other;
+	return Intersection != *this;
+}
+
+bool FPCGDataTypeIdentifier::IsIdentical(const FPCGDataTypeIdentifier& Other) const
+{
+	return IsSameType(Other) && CustomSubtype == Other.CustomSubtype;
+}
+
+bool FPCGDataTypeIdentifier::SerializeFromMismatchedTag(const struct FPropertyTag& Tag, FStructuredArchive::FSlot Slot)
+{
+	const UEnum* Enum = StaticEnum<EPCGDataType>();
+	check(Enum);
+	
+	if (Tag.GetType().IsEnum(Enum->GetFName()))
+	{
+		// Enum are serialized as FName (cf. TryLoadEnumValueByName)
+		FName EnumValueName;
+		Slot << EnumValueName;
+		
+		if (EnumValueName != NAME_None)
+        {
+        	int64 EnumValue = Enum->GetValueOrBitfieldFromString(*EnumValueName.ToString());
+        	if (EnumValue != INDEX_NONE)
+        	{
+        		FPCGDataTypeIdentifier Temp = FPCGDataTypeIdentifier(static_cast<EPCGDataType>(EnumValue));
+        		IDs = Temp.GetIDs();
+        	}
+        }
+
+		return true;
+	}
+	else
+	{
+		return false;
+	}
+}
+
+FString FPCGDataTypeIdentifier::ToString() const
+{
+	return IsValid() ? FString::JoinBy(IDs, TEXT(" | "), [](const FPCGDataTypeBaseId& ID) { return ID.ToString(); }) : TEXT("None");
+}
+
+FText FPCGDataTypeIdentifier::ToDisplayText() const
+{
+	return FText::FromString(ToString());
+}
+
+FPCGDataTypeIdentifier UPCGDataTypeIdentifierHelpers::GetIdentifierFromDefaultType(EPCGExclusiveDataType DataType)
+{
+	const UEnum* DataTypeEnum = StaticEnum<EPCGDataType>();
+	const UEnum* ExclusiveDataTypeEnum = StaticEnum<EPCGExclusiveDataType>();
+
+	EPCGDataType InternalDataType = EPCGDataType::Other;
+
+	if (DataTypeEnum && ExclusiveDataTypeEnum)
+	{
+		FName ExclusiveDataTypeName = ExclusiveDataTypeEnum->GetNameByValue(static_cast<__underlying_type(EPCGExclusiveDataType)>(DataType));
+		if (ensure(ExclusiveDataTypeName != NAME_None))
+		{
+			const int64 MatchingType = DataTypeEnum->GetValueByName(ExclusiveDataTypeName);
+			if (ensure(MatchingType != INDEX_NONE))
+			{
+				InternalDataType = static_cast<EPCGDataType>(MatchingType);
+			}
+		}
+	}
+
+	return FPCGDataTypeIdentifier(InternalDataType);
+}
diff --git a/Engine/Plugins/PCG/Source/PCG/Private/Data/Registry/PCGDataTypeRegistry.cpp b/Engine/Plugins/PCG/Source/PCG/Private/Data/Registry/PCGDataTypeRegistry.cpp
new file mode 100644
--- /dev/null
+++ b/Engine/Plugins/PCG/Source/PCG/Private/Data/Registry/PCGDataTypeRegistry.cpp
@@ -0,0 +1,677 @@
+// Copyright Epic Games, Inc. All Rights Reserved.
+
+#include "Data/Registry/PCGDataTypeRegistry.h"
+
+#include "PCGCommon.h"
+#include "PCGData.h"
+
+#if WITH_EDITOR
+#include "Tests/PCGTestsCommon.h"
+#endif // WITH_EDITOR
+
+#include "UObject/UObjectIterator.h"
+
+namespace PCGDataTypeHierarchy
+{
+/**
+ * Node for the hierarchy tree that can go up (to the parent) and down (to the children)
+ */
+struct FNode
+{
+	~FNode();
+
+	void AddChild(FNode* Child);
+
+	FPCGDataTypeBaseId BaseID;
+
+	FNode* Parent = nullptr;
+
+	// Use pointers for parent stability. Otherwise, we can have an array allocation in the grandparent that invalidate the parent pointer.
+	TArray<FNode*> Children;
+};
+
+/**
+* Hierarchy tree used to query the hierarchy for the types, to find common ancestors for example.
+* Expected to be filled by the registry after all the types are loaded.
+*/
+struct FTree
+{
+	FTree();
+
+	void AddEntry(const FPCGDataTypeBaseId& BaseID);
+	void RemoveEntry(const FPCGDataTypeBaseId& BaseID);
+
+	/**
+	 * Verify if TypeA and TypeB has a hierarchy relationship.
+	 * OutCommonAncestor will contain the type if of the type that is higher in the hierarchy than TypeA and TypeB.
+	 * If TypeA is a child of TypeB, OutCommonAncestor will be TypeB, and vice versa.
+	 * returns false if identifiers are invalid.
+	 */
+	bool QueryHierarchy(const FPCGDataTypeIdentifier& TypeA, const FPCGDataTypeIdentifier& TypeB, FPCGDataTypeIdentifier& OutCommonAncestor) const;
+
+	enum class ESearchCommand
+	{
+		ExpandAndContinue,
+		Continue,
+		Stop
+	};
+	
+	template <typename Func> requires std::is_invocable_r_v<FTree::ESearchCommand, Func, FNode*, int32>
+	void BreadthSearch(Func InFunc) const;
+
+	template <typename Func> requires std::is_invocable_r_v<FTree::ESearchCommand, Func, FNode*, int32>
+	void DepthSearch(Func InFunc) const;
+
+	FNode Root;
+};
+	
+template <typename Func> requires std::is_invocable_r_v<FTree::ESearchCommand, Func, FNode*, int32>
+void FTree::BreadthSearch(Func InFunc) const
+{
+	// Const-cast to not have 2 versions, const and non-const, since those are private.
+	// Will have the node and its depth.
+	TArray<TTuple<FNode*, int32>, TInlineAllocator<64>> List{{const_cast<FNode*>(&Root), 0}};
+	int32 i = 0;
+
+	while (i < List.Num())
+	{
+		FNode* Curr = List[i].Get<0>();
+		const int32 Depth = List[i++].Get<1>();
+
+		const ESearchCommand Command = InFunc(Curr, Depth);
+	
+		if (Command == ESearchCommand::Stop)
+		{
+			return;
+		}
+
+		if (Command == ESearchCommand::ExpandAndContinue)
+		{
+			for (FNode* Child : Curr->Children)
+			{
+				List.Emplace(Child, Depth + 1);
+			}
+		}
+	}
+}
+
+template <typename Func> requires std::is_invocable_r_v<FTree::ESearchCommand, Func, FNode*, int32>
+void FTree::DepthSearch(Func InFunc) const
+{
+	// Const-cast to not have 2 versions, const and non-const, since those are private.
+	// Will have the node and its depth.
+	TArray<TTuple<FNode*, int32>, TInlineAllocator<64>> Stack{{const_cast<FNode*>(&Root), 0}};
+
+	while (!Stack.IsEmpty())
+	{
+		TTuple<FNode*, int32> It = Stack.Pop(EAllowShrinking::No);
+		FNode* Curr = It.Get<0>();
+		const int32 Depth = It.Get<1>();
+	
+		const ESearchCommand Command = InFunc(Curr, Depth);
+	
+		if (Command == ESearchCommand::Stop)
+		{
+			return;
+		}
+
+		if (Command == ESearchCommand::ExpandAndContinue)
+		{
+			for (FNode* Child : Curr->Children)
+			{
+				Stack.Emplace(Child, Depth + 1);
+			}
+		}
+	}
+}
+
+FNode::~FNode()
+{
+	for (FNode* Child : Children)
+	{
+		delete Child;
+	}
+}
+
+void FNode::AddChild(FNode* Child)
+{
+	Child->Parent = this;
+	Children.Add(Child);
+}
+
+FTree::FTree()
+{
+	Root.BaseID = FPCGDataTypeInfo::StaticStruct();
+}
+
+void FTree::AddEntry(const FPCGDataTypeBaseId& BaseID)
+{
+	FNode* Parent = &Root;
+
+	FNode* NewNode = new FNode();
+	NewNode->BaseID = BaseID;
+
+	const UScriptStruct* EntryStruct = BaseID.GetStruct();
+
+	BreadthSearch([&Parent, EntryStruct](FNode* Curr, int32 Depth) -> ESearchCommand
+	{
+		if (EntryStruct->IsChildOf(Curr->BaseID.GetStruct()))
+		{
+			Parent = Curr;
+		}
+
+		// Don't stop until the end.
+		return ESearchCommand::ExpandAndContinue;
+	});
+
+	// If the parent have children, we need to check for all of them if they are a child of the new entry
+	int32 i = 0 ;
+	while (i < Parent->Children.Num())
+	{
+		FNode* Curr = Parent->Children[i];
+		if (Curr->BaseID.GetStruct()->IsChildOf(EntryStruct))
+		{
+			NewNode->AddChild(Curr);
+			Parent->Children.RemoveAt(i);
+		}
+		else
+		{
+			++i;
+		}
+	}
+
+	// Finally setup the new node
+	Parent->AddChild(NewNode);
+}
+
+void FTree::RemoveEntry(const FPCGDataTypeBaseId& BaseID)
+{
+	DepthSearch([&BaseID](FNode* Curr, int32 Depth) -> ESearchCommand
+	{
+		if (Curr->BaseID == BaseID)
+		{
+			for (FNode* Child : Curr->Children)
+			{
+				Child->Parent = Curr->Parent;
+			}
+
+			Curr->Parent->Children.Remove(Curr);
+			Curr->Parent->Children.Append(Curr->Children);
+
+			delete Curr;
+		
+			return ESearchCommand::Stop;
+		}
+		else
+		{
+			return ESearchCommand::ExpandAndContinue;
+		}
+	});
+}
+
+bool FTree::QueryHierarchy(const FPCGDataTypeIdentifier& TypeA, const FPCGDataTypeIdentifier& TypeB,
+	FPCGDataTypeIdentifier& OutCommonAncestor) const
+{
+	if (!TypeA.IsValid() || !TypeB.IsValid())
+	{
+		return false;
+	}
+
+	if (TypeA.IsSameType(TypeB))
+	{
+		OutCommonAncestor = TypeA;
+		return true;
+	}
+
+	const bool TypeAIsComposition = TypeA.IsComposition();
+	const bool TypeBIsComposition = TypeB.IsComposition();
+
+	// In case one is a composition and not the other, check if they support the non-composition.
+	if (TypeAIsComposition && !TypeBIsComposition && TypeA.SupportsType(TypeB))
+	{
+		OutCommonAncestor = TypeA;
+		return true;
+	}
+
+	if (TypeBIsComposition && !TypeAIsComposition && TypeB.SupportsType(TypeA))
+	{
+		OutCommonAncestor = TypeB;
+		return true;
+	}
+
+	// Breadth search to find which comes first in the hierarchy
+	const FNode* FirstFound = nullptr;
+	const FNode* SecondFound = nullptr;
+	BreadthSearch([&FirstFound, &SecondFound, &TypeA, &TypeB](const FNode* Curr, int32 Depth) -> ESearchCommand
+	{
+		bool bFound = false;
+		const FPCGDataTypeIdentifier CurrID{Curr->BaseID};
+
+		if (TypeA.SupportsType(CurrID))
+		{
+			if (!FirstFound)
+			{
+				FirstFound = Curr;
+				return ESearchCommand::ExpandAndContinue;
+			}
+
+			bFound = true;
+		}
+
+		if (TypeB.SupportsType(CurrID))
+		{
+			if (!FirstFound)
+			{
+				FirstFound = Curr;
+				return ESearchCommand::ExpandAndContinue;
+			}
+
+			bFound = true;
+		}
+
+		if (!bFound)
+		{
+			return ESearchCommand::ExpandAndContinue;
+		}
+		else
+		{
+			SecondFound = Curr;
+			return ESearchCommand::Stop;
+		}
+	});
+
+	if (!FirstFound || !SecondFound)
+	{
+		return false;
+	}
+
+	using FHierarchyArray = TArray<const FNode*, TInlineAllocator<16>>;
+
+	auto BuildHierarchy = [](const FNode* Start)-> FHierarchyArray
+	{
+		FHierarchyArray Hierarchy;
+		const FNode* UpstreamCurr = Start;
+			
+		while (UpstreamCurr)
+		{
+			Hierarchy.Add(UpstreamCurr);
+			UpstreamCurr = UpstreamCurr->Parent;
+		}
+
+		return Hierarchy;
+	};
+
+	const FHierarchyArray FirstHierarchy = BuildHierarchy(FirstFound);
+	const FHierarchyArray SecondHierarchy = BuildHierarchy(SecondFound);
+
+	int32 i = FirstHierarchy.Num() - 1;
+	int32 j = SecondHierarchy.Num() - 1;
+
+	while (i >= 0 && j >= 0 && FirstHierarchy[i] == SecondHierarchy[j])
+	{
+		--i;
+		--j;
+	}
+
+	// All types have a common root type (FPCGDataTypeInfo) so FirstHierarchy and SecondHierarchy will always have the same "Last" member (same ancestor).
+	// By construction, FirstHierarchy[i+1] will never be out of bounds.
+	OutCommonAncestor = FPCGDataTypeIdentifier{FirstHierarchy[i+1]->BaseID};
+
+	return OutCommonAncestor.IsValid();
+}
+} // namespace PCGDataTypeHierarchy
+
+FPCGDataTypeRegistry::FPCGDataTypeRegistry()
+{
+	HierarchyTree = new PCGDataTypeHierarchy::FTree();
+	
+	// Pre-register the base type
+	Mapping.Emplace(FPCGDataTypeInfo::StaticStruct(), FPCGDataTypeInfo::StaticStruct());
+}
+
+void FPCGDataTypeRegistry::RegisterEntry(const FPCGDataTypeBaseId& BaseID)
+{
+	check(IsInGameThread());
+	
+	if (!ensureMsgf(!Mapping.Contains(BaseID), TEXT("Tried to register a PCG data type that has the same identifier as another already registered.")))
+	{
+		return;
+	}
+
+	Mapping.Emplace(BaseID, BaseID.GetStruct());
+	HierarchyTree->AddEntry(BaseID);
+}
+
+FPCGDataTypeRegistry::~FPCGDataTypeRegistry()
+{
+	delete HierarchyTree;
+}
+
+EPCGDataTypeCompatibilityResult FPCGDataTypeRegistry::IsCompatible(const FPCGDataTypeIdentifier& InType, const FPCGDataTypeIdentifier& OutType) const
+{
+	if (!InType.IsValid() || !OutType.IsValid())
+	{
+		return EPCGDataTypeCompatibilityResult::UnknownType;
+	}
+
+	if (InType.IsSameType(OutType))
+	{
+		return EPCGDataTypeCompatibilityResult::Compatible;
+	}
+
+	// Run the intersection between the two.
+	const FPCGDataTypeIdentifier Intersection = InType & OutType;
+
+	// If the intersection gives InType, it means OutType accept all InType, so it is compatible
+	if (Intersection.IsSameType(InType))
+	{
+		return EPCGDataTypeCompatibilityResult::Compatible;
+	}
+	
+	// If the intersection gives OutType, or something else that is valid, it means that InType is wider than OutType, so we might need a filter or a conversion.
+	// Conversion can only happen if InType and OutType are not compositions
+	if (Intersection.IsSameType(OutType) || Intersection.IsValid())
+	{
+		if (!InType.IsComposition() && !OutType.IsComposition())
+		{
+			const TInstancedStruct<FPCGDataTypeInfo>* InDefIt = Mapping.Find(InType.GetID());
+			const TInstancedStruct<FPCGDataTypeInfo>* OutDefIt = Mapping.Find(OutType.GetID());
+
+			if (!InDefIt || !OutDefIt)
+			{
+				return EPCGDataTypeCompatibilityResult::UnknownType;
+			}
+
+			return InDefIt->Get().SupportsConversion(OutType) ? EPCGDataTypeCompatibilityResult::RequireConversion : EPCGDataTypeCompatibilityResult::RequireFilter;
+		}
+		else
+		{
+			return EPCGDataTypeCompatibilityResult::RequireFilter;
+		}
+	}
+
+	// If we have no intersection and either of them is a composition, it's ambiguous, as we might require multiple conversions,
+	// so mark it as not compatible and let the user filter and convert as they see fit.
+	if (InType.IsComposition() || OutType.IsComposition())
+	{
+		return EPCGDataTypeCompatibilityResult::NotCompatible;
+	}
+
+	// Finally we check if we have a conversion
+	const TInstancedStruct<FPCGDataTypeInfo>* InDefIt = Mapping.Find(InType.GetID());
+	const TInstancedStruct<FPCGDataTypeInfo>* OutDefIt = Mapping.Find(OutType.GetID());
+
+	if (!InDefIt || !OutDefIt)
+	{
+		return EPCGDataTypeCompatibilityResult::UnknownType;
+	}
+
+	return InDefIt->Get().SupportsConversion(OutType) ? EPCGDataTypeCompatibilityResult::RequireConversion : EPCGDataTypeCompatibilityResult::NotCompatible;
+}
+
+const FPCGDataTypeInfo* FPCGDataTypeRegistry::GetTypeInfo(const FPCGDataTypeBaseId& ID) const
+{
+	const TInstancedStruct<FPCGDataTypeInfo>* TypeInfo = Mapping.Find(ID);
+	return TypeInfo ? TypeInfo->GetPtr() : nullptr;
+}
+
+/**
+ * Goal: From a list of identifiers, aggregate all the types, and find the common ancestor of them all
+ * Algorithm overview:
+ * - Gather all the unique types
+ * - Initialize result to the first unique type and for the rest find the common ancestor with result. The common ancestor becomes the result and we continue.
+ * - If we ever find Any, we can stop and return it.
+ */
+FPCGDataTypeIdentifier FPCGDataTypeRegistry::GetIdentifiersUnion(TConstArrayView<FPCGDataTypeIdentifier> Identifiers) const
+{
+	if (Identifiers.IsEmpty())
+	{
+		return FPCGDataTypeIdentifier{};
+	}
+
+	static const FPCGDataTypeBaseId AnyType = FPCGDataTypeBaseId::Construct<FPCGDataTypeInfo>();
+	
+	// Extract all unique IDs
+	TArray<FPCGDataTypeBaseId, TInlineAllocator<32>> UniqueIDs;
+
+	for (const FPCGDataTypeIdentifier& Identifier : Identifiers)
+	{
+		for (const FPCGDataTypeBaseId& ID : Identifier.GetIDs())
+		{
+			// If we have the "Any" type in our composition, it is the parent of all the types, so the composition is "Any"
+			if (ID == AnyType)
+			{
+				return FPCGDataTypeIdentifier{AnyType};
+			}
+			
+			UniqueIDs.AddUnique(ID);
+		}
+	}
+
+	if (UniqueIDs.IsEmpty())
+	{
+		return FPCGDataTypeIdentifier{};
+	}
+	
+	FPCGDataTypeIdentifier Result{UniqueIDs[0]};
+
+	for (int32 i = 1; i < UniqueIDs.Num(); ++i)
+	{
+		const FPCGDataTypeBaseId& ID = UniqueIDs[i];
+		
+		if (Result.IsSameType(ID))
+		{
+			continue;
+		}
+
+		FPCGDataTypeIdentifier CommonAncestor{};
+		if (!HierarchyTree->QueryHierarchy(Result, ID, CommonAncestor))
+		{
+			return FPCGDataTypeIdentifier{};
+		}
+
+		if (CommonAncestor.IsSameType(AnyType))
+		{
+			return FPCGDataTypeIdentifier{AnyType};
+		}
+
+		Result = CommonAncestor;
+	}
+
+	return Result;
+}
+
+/**
+ * Goal: From a list of identifiers, aggregate all the types, and reduce types that are fully represent the parent type
+ *       (i.e. if we compose DataB and DataC that are the only children of DataA, the composition should return DataA)
+ * Algorithm overview:
+ * - Gather all the unique types
+ * - Do a depth search in the hierarchy tree to look for our unique types
+ *    + When we find a type, we mark at which depth it is and add it to our list of nodes.
+ *    + Since it's Depth first, any type that we found that is at a greater depth that the one marked is a child, and can be discarded for the composition.
+ *    + At the moment we have a Depth greater or equal than the one marked, we are on another branch, and we reset our marked depth.
+ * - When all of our types are found, we sort them by depth in descending order (deeper first), we need the descending order because we might merge 2 types
+ *   that can be merged with something at a higher level of the hierarchy. (i.e. D and E are child of B, and B and C are child of A, combining C, D and E should give A.
+ * - We iterate through our sorted list, and we look for all the other nodes that are on the same depth (potential siblings). We count all of them that share the same parent.
+ * - If the number of siblings is equal to the number of children of the parent type, we can merge them all into the parent type.
+ * - When the merge is done, we return the aggregation of the all the remaining types.
+ */
+FPCGDataTypeIdentifier FPCGDataTypeRegistry::GetIdentifiersComposition(TConstArrayView<FPCGDataTypeIdentifier> Identifiers) const
+{
+	if (Identifiers.IsEmpty())
+	{
+		return FPCGDataTypeIdentifier{};
+	}
+
+	static const FPCGDataTypeBaseId AnyType = FPCGDataTypeBaseId::Construct<FPCGDataTypeInfo>();
+	
+	// Extract all unique IDs
+	TArray<FPCGDataTypeBaseId, TInlineAllocator<32>> UniqueIDs;
+
+	for (const FPCGDataTypeIdentifier& Identifier : Identifiers)
+	{
+		for (const FPCGDataTypeBaseId& ID : Identifier.GetIDs())
+		{
+			// If we have the "Any" type in our composition, it is the parent of all the types, so the composition is "Any"
+			if (ID == AnyType)
+			{
+				return FPCGDataTypeIdentifier{AnyType};
+			}
+			
+			UniqueIDs.AddUnique(ID);
+		}
+	}
+
+	// If we just have a single one, return this one
+	if (UniqueIDs.Num() == 1)
+	{
+		return FPCGDataTypeIdentifier{UniqueIDs[0]};
+	}
+
+	// Find all the nodes, and sort them by depth.
+	TArray<TTuple<const PCGDataTypeHierarchy::FNode*, int32>, TInlineAllocator<32>> Nodes;
+	Nodes.Reserve(UniqueIDs.Num());
+
+	HierarchyTree->DepthSearch([&UniqueIDs, &Nodes](const PCGDataTypeHierarchy::FNode* Node, int32 Depth)
+	{
+		if (UniqueIDs.Contains(Node->BaseID))
+		{
+			Nodes.Emplace(Node, Depth);
+			// We found one, no need to check the children, as they won't contribute to the composition.
+			return PCGDataTypeHierarchy::FTree::ESearchCommand::Continue;
+		}
+
+		// Continue until we find them all.
+		return PCGDataTypeHierarchy::FTree::ESearchCommand::ExpandAndContinue;
+	});
+
+	// Sort by depth, starting at the deepest.
+	Nodes.Sort([](const auto& LH, const auto& RH) { return LH.template Get<1>() > RH.template Get<1>(); });
+
+	// At that point, we have only types that are on different branches.
+	// If they are on the same depth, we try to reduce them
+	int32 i = 0;
+	TArray<int32, TInlineAllocator<32>> Indices;
+	
+	while (i < Nodes.Num() - 1)
+	{
+		int32 SameParentNum = 1;
+		Indices.SetNum(0, EAllowShrinking::No);
+		
+		for (int32 j = i + 1; j < Nodes.Num(); ++j)
+		{
+			if (Nodes[i].Get<1>() == Nodes[j].Get<1>() && Nodes[i].Get<0>()->Parent == Nodes[j].Get<0>()->Parent)
+			{
+				SameParentNum++;
+				Indices.Add(j);
+			}
+		}
+
+		if (SameParentNum > 1 && SameParentNum == Nodes[i].Get<0>()->Parent->Children.Num())
+		{
+			// All the types at a given level are combined. Remove all of them and replace it with the parent type.
+			// The parent type can't be in the composition, because if it were, those children won't have been added while doing the depth search
+			Nodes[i].Get<0>() = Nodes[i].Get<0>()->Parent;
+			Nodes[i].Get<1>()--;
+			for (int32 j = Indices.Num() - 1; j >= 0; --j)
+			{
+				Nodes.RemoveAt(Indices[j]);
+			}
+
+			// We have a merge, so try again at index i.
+		}
+		else
+		{
+			// Otherwise, check the next
+			++i;
+		}
+	}
+
+	// Returns the aggregation of the types after the reduce
+	FPCGDataTypeIdentifier Result{};
+	Result.IDs.Reserve(Nodes.Num());
+	for (const auto& Node : Nodes)
+	{
+		Result.IDs.Emplace(Node.Get<0>()->BaseID);
+	}
+
+	return Result;
+}
+
+void FPCGDataTypeRegistry::RegisterKnownTypes()
+{
+	const UScriptStruct* BaseStruct = FPCGDataTypeInfo::StaticStruct();
+	
+	for (TObjectIterator<UScriptStruct> It; It; ++It)
+	{
+		// Skip FPCGDataTypeInfo as it is already register by default
+		if (*It != BaseStruct && It->IsChildOf(BaseStruct))
+		{
+			RegisterEntry(FPCGDataTypeBaseId{*It});
+		}
+	}
+}
+
+
+// Tests for hierarchy
+#if WITH_EDITOR
+
+IMPLEMENT_CUSTOM_SIMPLE_AUTOMATION_TEST(FPCGHierarchyTest, FPCGTestBaseClass, "Plugins.PCG.DataTypeRegistry.Hierarchy", PCGTestsCommon::TestFlags)
+
+bool FPCGHierarchyTest::RunTest(const FString& Parameters)
+{
+	TTuple<FPCGDataTypeIdentifier, FPCGDataTypeIdentifier, FPCGDataTypeIdentifier> TestValues[] =
+	{
+		{EPCGDataType::Any, EPCGDataType::Spatial, EPCGDataType::Any},
+		{EPCGDataType::Any, EPCGDataType::Point, EPCGDataType::Any},
+		{EPCGDataType::Point, EPCGDataType::Point, EPCGDataType::Point},
+		{EPCGDataType::Spatial, EPCGDataType::Point, EPCGDataType::Spatial},
+		{EPCGDataType::Spatial, EPCGDataType::Volume, EPCGDataType::Spatial},
+		{EPCGDataType::Spatial, EPCGDataType::Surface, EPCGDataType::Spatial},
+		{EPCGDataType::Spatial, EPCGDataType::Spline, EPCGDataType::Spatial},
+		{EPCGDataType::Spline, EPCGDataType::Point, EPCGDataType::Concrete},
+		{EPCGDataType::Volume, EPCGDataType::Point, EPCGDataType::Concrete},
+		{EPCGDataType::Spline, EPCGDataType::LandscapeSpline, EPCGDataType::PolyLine},
+		{EPCGDataType::Landscape, EPCGDataType::Texture, EPCGDataType::Surface},
+		{EPCGDataType::RenderTarget, EPCGDataType::Texture, EPCGDataType::BaseTexture},
+		{EPCGDataType::VirtualTexture, EPCGDataType::Texture, EPCGDataType::Surface},
+		{EPCGDataType::Point, EPCGDataType::Param, EPCGDataType::Any},
+		{EPCGDataType::PointOrParam, EPCGDataType::Param, EPCGDataType::PointOrParam},
+		{EPCGDataType::PointOrParam, EPCGDataType::Point, EPCGDataType::PointOrParam},
+		{EPCGDataType::PointOrParam, EPCGDataType::Volume, EPCGDataType::Any},
+	};
+
+	const FPCGDataTypeRegistry& Registry = FPCGModule::GetConstDataTypeRegistry();
+
+	const UEnum* Enum = StaticEnum<EPCGDataType>();
+	check(Enum);
+
+	for (const auto& [TypeAIdentifier, TypeBIdentifier, ExpectedCommonAncestorIdentifier] : TestValues)
+	{
+		FPCGDataTypeIdentifier CommonAncestorAToB{};
+		FPCGDataTypeIdentifier CommonAncestorBToA{};
+		
+		UTEST_TRUE(
+			FString::Printf(TEXT("QueryHierarchy between %s and %s is valid"), *TypeAIdentifier.ToString(), *TypeBIdentifier.ToString()),
+			Registry.HierarchyTree->QueryHierarchy(TypeAIdentifier, TypeBIdentifier, CommonAncestorAToB)
+		);
+		
+		UTEST_TRUE(
+			FString::Printf(TEXT("QueryHierarchy between %s and %s is valid"), *TypeBIdentifier.ToString(), *TypeAIdentifier.ToString()),
+			Registry.HierarchyTree->QueryHierarchy(TypeBIdentifier, TypeAIdentifier, CommonAncestorBToA)
+		);
+
+		UTEST_TRUE(
+			FString::Printf(TEXT("Expected ancestor between %s and %s is %s"), *TypeAIdentifier.ToString(), *TypeBIdentifier.ToString(), *ExpectedCommonAncestorIdentifier.ToString()),
+			ExpectedCommonAncestorIdentifier.IsSameType(CommonAncestorAToB)
+		);
+		
+		UTEST_TRUE(
+			FString::Printf(TEXT("Expected ancestor between %s and %s is %s"), *TypeBIdentifier.ToString(), *TypeAIdentifier.ToString(), *ExpectedCommonAncestorIdentifier.ToString()),
+			ExpectedCommonAncestorIdentifier.IsSameType(CommonAncestorBToA)
+		);
+	}
+
+	return true;
+}
+#endif // WITH_EDITOR
\ No newline at end of file
diff --git a/Engine/Plugins/PCG/Source/PCG/Private/Data/Registry/PCGGetDataFunctionRegistry.cpp b/Engine/Plugins/PCG/Source/PCG/Private/Data/Registry/PCGGetDataFunctionRegistry.cpp
new file mode 100644
--- /dev/null
+++ b/Engine/Plugins/PCG/Source/PCG/Private/Data/Registry/PCGGetDataFunctionRegistry.cpp
@@ -0,0 +1,331 @@
+// Copyright Epic Games, Inc. All Rights Reserved.
+
+#include "Data/Registry/PCGGetDataFunctionRegistry.h"
+
+#include "PCGCommon.h"
+#include "PCGComponent.h"
+#include "PCGContext.h"
+#include "PCGData.h"
+#include "PCGGraph.h"
+#include "Compute/PCGComputeCommon.h"
+#include "Data/PCGBasePointData.h"
+#include "Data/PCGCollisionShapeData.h"
+#include "Data/PCGIntersectionData.h"
+#include "Data/PCGLandscapeData.h"
+#include "Data/PCGLandscapeSplineData.h"
+#include "Data/PCGPrimitiveData.h"
+#include "Data/PCGProjectionData.h"
+#include "Data/PCGSpatialData.h"
+#include "Data/PCGSplineData.h"
+#include "Data/PCGUnionData.h"
+#include "Data/PCGVirtualTextureData.h"
+#include "Data/PCGVolumeData.h"
+#include "Grid/PCGPartitionActor.h"
+#include "Helpers/PCGActorHelpers.h"
+#include "Helpers/PCGBlueprintHelpers.h"
+#include "Helpers/PCGHelpers.h"
+
+#include "LandscapeComponent.h"
+#include "LandscapeProxy.h"
+#include "LandscapeSplinesComponent.h"
+#include "Components/ActorComponent.h"
+#include "Components/BillboardComponent.h"
+#include "Components/InstancedStaticMeshComponent.h"
+#include "Components/RuntimeVirtualTextureComponent.h"
+#include "Components/ShapeComponent.h"
+#include "Components/SplineComponent.h"
+#include "Components/SplineMeshComponent.h"
+#include "GameFramework/Actor.h"
+#include "GameFramework/Volume.h"
+#include "VT/RuntimeVirtualTexture.h"
+
+FPCGGetDataFunctionRegistry::FFunctionHandle FPCGGetDataFunctionRegistry::RegisterDataFromActorFunction(const FPCGGetDataFunctionRegistry::FDataFromActorFunction& InFunction)
+{
+	check(IsInGameThread());
+	ActorParsingFunctions.Emplace(InFunction, NextFunctionHandle);
+	return NextFunctionHandle++;
+}
+
+void FPCGGetDataFunctionRegistry::UnregisterDataFromActorFunction(FPCGGetDataFunctionRegistry::FFunctionHandle InFunctionHandle)
+{
+	check(IsInGameThread());
+	int FunctionIndex = ActorParsingFunctions.IndexOfByPredicate([InFunctionHandle](const TPair<FDataFromActorFunction, FFunctionHandle>& InPair) { return InPair.Value == InFunctionHandle; });
+	if (FunctionIndex != INDEX_NONE)
+	{
+		ActorParsingFunctions.RemoveAt(FunctionIndex);
+	}
+}
+
+FPCGGetDataFunctionRegistry::FFunctionHandle FPCGGetDataFunctionRegistry::RegisterDataFromComponentFunction(const FPCGGetDataFunctionRegistry::FDataFromComponentFunction& InFunction)
+{
+	check(IsInGameThread());
+	ComponentParsingFunctions.Emplace(InFunction, NextFunctionHandle);
+	return NextFunctionHandle++;
+}
+
+void FPCGGetDataFunctionRegistry::UnregisterDataFromComponentFunction(FPCGGetDataFunctionRegistry::FFunctionHandle InFunctionHandle)
+{
+	check(IsInGameThread());
+	int FunctionIndex = ComponentParsingFunctions.IndexOfByPredicate([InFunctionHandle](const TPair<FDataFromComponentFunction, FFunctionHandle>& InPair) { return InPair.Value == InFunctionHandle; });
+	if (FunctionIndex != INDEX_NONE)
+	{
+		ComponentParsingFunctions.RemoveAt(FunctionIndex);
+	}
+}
+
+int FPCGGetDataFunctionRegistry::GetDataFromActor(FPCGContext* InContext, const FPCGGetDataFunctionRegistryParams& InParams, AActor* InActor, FPCGGetDataFunctionRegistryOutput& Output) const
+{
+	int DataCountBefore = Output.Collection.TaggedData.Num();
+
+	// Try custom getters first
+	for (const TPair<FDataFromActorFunction, FFunctionHandle>& ActorParsingPair : ActorParsingFunctions)
+	{
+		if (ActorParsingPair.Key(InContext, InParams, InActor, Output))
+		{
+			break;
+		}
+	}
+
+	if (Output.Collection.TaggedData.Num() != DataCountBefore)
+	{
+		return Output.Collection.TaggedData.Num() - DataCountBefore;
+	}
+
+	// Default actor parsing, if any
+	return DefaultDataFromActor(InContext, InParams, InActor, Output);
+}
+
+int FPCGGetDataFunctionRegistry::DefaultDataFromActor(FPCGContext* InContext, const FPCGGetDataFunctionRegistryParams& InParams, AActor* Actor, FPCGGetDataFunctionRegistryOutput& Output) const
+{
+	check(Actor);
+
+	// Early out if the actor gets rejected by the component selector
+	if (InParams.ComponentSelector && !InParams.ComponentSelector->FilterActor(Actor))
+	{
+		return 0;
+	}
+
+	const EPCGDataType InDataFilter = InParams.DataTypeFilter;
+	const bool bParseActor = InParams.bParseActor;
+	const UPCGComponent* Component = InParams.SourceComponent;
+
+	auto NameTagsToStringTags = [](const FName& InName) { return InName.ToString(); };
+	TSet<FString> ActorTags;
+	Algo::Transform(Actor->Tags, ActorTags, NameTagsToStringTags);
+
+	// Fill in collection based on the data on the given actor.
+	// Some actor types we will forego full parsing to build strictly on the actor existence, such as partition actors, volumes and landscape
+	// TODO: review the !bParseActor cases - it might make sense to have just a point for a partition actor, even if we preintersect it.
+	APCGPartitionActor* PartitionActor = (!!(InDataFilter & EPCGDataType::Spatial) || !!(InDataFilter & EPCGDataType::Volume)) ? Cast<APCGPartitionActor>(Actor) : nullptr;
+	ALandscapeProxy* LandscapeActor = !!(InDataFilter & EPCGDataType::Landscape) ? Cast<ALandscapeProxy>(Actor) : nullptr;
+	AVolume* VolumeActor = !!(InDataFilter & EPCGDataType::Volume) ? Cast<AVolume>(Actor) : nullptr;
+	if (!bParseActor && !!(InDataFilter & EPCGDataType::Point))
+	{
+		UPCGBasePointData* Data = FPCGContext::NewPointData_AnyThread(InContext);
+		Data->InitializeFromActor(Actor, &Output.bSanitizedTagAttributeNames);
+
+		FPCGTaggedData& TaggedData = Output.Collection.TaggedData.Emplace_GetRef();
+		TaggedData.Data = Data;
+		TaggedData.Tags = ActorTags;
+	}
+	else if (PartitionActor)
+	{
+		check(!Component || Component->GetOwner() == Actor); // Invalid processing otherwise because of the this usage
+
+		UPCGVolumeData* VolumeData = FPCGContext::NewObject_AnyThread<UPCGVolumeData>(InContext);
+		UPCGSpatialData* Result = VolumeData;
+		if (InDataFilter == EPCGDataType::Volume)
+		{
+			VolumeData->Initialize(PCGHelpers::GetGridBounds(Actor, Component));
+		}
+		else
+		{
+			VolumeData->Initialize(PartitionActor->GetFixedBounds());
+
+			if (Component && PartitionActor->GetOriginalComponent(Component))
+			{
+				PCG_EXECUTION_CACHE_VALIDATION_CREATE_ORIGINAL_SCOPE(const_cast<UPCGComponent*>(Component));
+				// Important note: we do NOT call the collection version here, as we want to have a union if that's the case
+				const UPCGSpatialData* OriginalComponentSpatialData = Cast<const UPCGSpatialData>(Component->GetOriginalActorPCGData());
+
+				if (OriginalComponentSpatialData)
+				{
+					Result = Result->IntersectWith(nullptr, OriginalComponentSpatialData);
+				}
+			}
+		}
+
+		FPCGTaggedData& TaggedData = Output.Collection.TaggedData.Emplace_GetRef();
+		TaggedData.Data = Result;
+		// No need to keep partition actor tags, though we might want to push PCG grid GUID at some point
+	}
+	else if (LandscapeActor)
+	{
+		UPCGLandscapeData* Data = FPCGContext::NewObject_AnyThread<UPCGLandscapeData>(InContext);
+		const UPCGGraph* PCGGraph = Component ? Component->GetGraph() : nullptr;
+
+		FPCGLandscapeDataProps LandscapeDataProps;
+		LandscapeDataProps.bGetHeightOnly = false;
+		LandscapeDataProps.bGetLayerWeights = (!PCGGraph || PCGGraph->bLandscapeUsesMetadata);
+
+		Data->Initialize({ LandscapeActor }, PCGHelpers::GetGridBounds(Actor, Component), LandscapeDataProps);
+
+		FPCGTaggedData& TaggedData = Output.Collection.TaggedData.Emplace_GetRef();
+		TaggedData.Data = Data;
+		TaggedData.Tags = ActorTags;
+	}
+	else if (VolumeActor)
+	{
+		UPCGVolumeData* Data = FPCGContext::NewObject_AnyThread<UPCGVolumeData>(InContext);
+		Data->Initialize(VolumeActor);
+
+		FPCGTaggedData& TaggedData = Output.Collection.TaggedData.Emplace_GetRef();
+		TaggedData.Data = Data;
+		TaggedData.Tags = ActorTags;
+	}
+	else
+	{
+		// Use component-based approach
+		int AddedData = 0;
+		TInlineComponentArray<UActorComponent*> ActorComponents;
+		Actor->GetComponents(ActorComponents);
+
+		for (UActorComponent* ActorComponent : ActorComponents)
+		{
+			AddedData += GetDataFromComponent(InContext, InParams, ActorComponent, Output);
+		}
+		
+		// IMPLEMENTATION NOTE: API BREAK HERE, we normally would have returned a single point here if there were no valid components
+		return AddedData;
+	}
+
+	return 1;
+}
+
+int FPCGGetDataFunctionRegistry::GetDataFromComponent(FPCGContext* InContext, const FPCGGetDataFunctionRegistryParams& InParams, UActorComponent* InComponent, FPCGGetDataFunctionRegistryOutput& Output) const
+{
+	check(InComponent);
+
+	if (InParams.ComponentSelector && !InParams.ComponentSelector->FilterComponent(InComponent))
+	{
+		return 0;
+	}
+
+	int DataCountBefore = Output.Collection.TaggedData.Num();
+
+	for (const TPair<FDataFromComponentFunction, FFunctionHandle>& ComponentParsingPair : ComponentParsingFunctions)
+	{
+		if (ComponentParsingPair.Key(InContext, InParams, InComponent, Output))
+		{
+			break;
+		}
+	}
+
+	if (Output.Collection.TaggedData.Num() != DataCountBefore)
+	{
+		return Output.Collection.TaggedData.Num() - DataCountBefore;
+	}
+
+	// Default component parsing, if any
+	return DefaultDataFromComponent(InContext, InParams, InComponent, Output);
+}
+
+int FPCGGetDataFunctionRegistry::DefaultDataFromComponent(FPCGContext* InContext, const FPCGGetDataFunctionRegistryParams& InParams, UActorComponent* InComponent, FPCGGetDataFunctionRegistryOutput& Output) const
+{
+	check(InComponent);
+
+	if (InParams.bIgnorePCGGeneratedComponents && InComponent->ComponentTags.Contains(PCGHelpers::DefaultPCGTag))
+	{
+		return 0;
+	}
+
+	auto NameTagsToStringTags = [](const FName& InName) { return InName.ToString(); };
+	const EPCGDataType InDataFilter = InParams.DataTypeFilter;
+
+	// Landscape spline component
+	// Splines component
+	// Shape component
+	// Primitive component, but ignore billboard components
+	// RuntimeVirtualTexture component
+	FPCGTaggedData TaggedData;
+
+	if (!TaggedData.Data && !!(InDataFilter & EPCGDataType::Spline))
+	{
+		if (ULandscapeSplinesComponent* LandscapeSplineComponent = Cast<ULandscapeSplinesComponent>(InComponent))
+		{
+			UPCGLandscapeSplineData* SplineData = FPCGContext::NewObject_AnyThread<UPCGLandscapeSplineData>(InContext);
+			SplineData->Initialize(LandscapeSplineComponent);
+			TaggedData.Data = SplineData;
+		}
+		else if (USplineComponent* SplineComponent = Cast<USplineComponent>(InComponent))
+		{
+			UPCGSplineData* SplineData = FPCGContext::NewObject_AnyThread<UPCGSplineData>(InContext);
+			SplineData->Initialize(SplineComponent);
+			TaggedData.Data = SplineData;
+		}
+	}
+
+	if (!TaggedData.Data && !!(InDataFilter & EPCGDataType::Primitive))
+	{
+		if (UShapeComponent* ShapeComponent = Cast<UShapeComponent>(InComponent); ShapeComponent && UPCGCollisionShapeData::IsSupported(ShapeComponent))
+		{
+			UPCGCollisionShapeData* ShapeData = FPCGContext::NewObject_AnyThread<UPCGCollisionShapeData>(InContext);
+			ShapeData->Initialize(ShapeComponent);
+			TaggedData.Data = ShapeData;
+		}
+		else if (UBillboardComponent* BillboardComponent = Cast<UBillboardComponent>(InComponent))
+		{
+			// Exception: skip the billboard component
+		}
+		else if (UPrimitiveComponent* PrimitiveComponent = Cast<UPrimitiveComponent>(InComponent))
+		{
+			UPCGPrimitiveData* PrimitiveData = FPCGContext::NewObject_AnyThread<UPCGPrimitiveData>(InContext);
+			PrimitiveData->Initialize(PrimitiveComponent);
+			TaggedData.Data = PrimitiveData;
+		}
+	}
+
+	if (!TaggedData.Data && !!(InDataFilter & EPCGDataType::VirtualTexture))
+	{
+		if (URuntimeVirtualTextureComponent* VirtualTextureComponent = Cast<URuntimeVirtualTextureComponent>(InComponent))
+		{
+			UPCGVirtualTextureData* VirtualTextureData = FPCGContext::NewObject_AnyThread<UPCGVirtualTextureData>(InContext);
+			VirtualTextureData->Initialize(VirtualTextureComponent);
+			TaggedData.Data = VirtualTextureData;
+
+			if (const URuntimeVirtualTexture* VirtualTexture = VirtualTextureData->GetRuntimeVirtualTexture())
+			{
+				FString VirtualTextureName;
+				VirtualTexture->GetName(VirtualTextureName);
+
+				const UEnum* MaterialTypeEnum = StaticEnum<ERuntimeVirtualTextureMaterialType>();
+				check(MaterialTypeEnum);
+
+				const FString MaterialTypeString = MaterialTypeEnum->GetNameStringByValue(static_cast<int64>(VirtualTexture->GetMaterialType()));
+
+				TaggedData.Tags.Add(PCGComputeHelpers::GetPrefixedDataLabel(VirtualTextureName));
+				TaggedData.Tags.Add(PCGComputeHelpers::GetPrefixedDataLabel(MaterialTypeString));
+			}
+
+		}
+	}
+
+	if (TaggedData.Data)
+	{
+		Algo::Transform(InComponent->ComponentTags, TaggedData.Tags, NameTagsToStringTags);
+
+		if (InParams.bAddActorTags && InComponent->GetOwner())
+		{
+			TSet<FString> ActorTags;
+			Algo::Transform(InComponent->GetOwner()->Tags, ActorTags, NameTagsToStringTags);
+			TaggedData.Tags.Append(ActorTags);
+		}
+
+		Output.Collection.TaggedData.Add(MoveTemp(TaggedData));
+		return 1;
+	}	
+	else
+	{
+		return 0;
+	}
+}
\ No newline at end of file
diff --git a/Engine/Plugins/PCG/Source/PCG/Private/Helpers/PCGSubgraphHelpers.cpp b/Engine/Plugins/PCG/Source/PCG/Private/Helpers/PCGSubgraphHelpers.cpp
--- a/Engine/Plugins/PCG/Source/PCG/Private/Helpers/PCGSubgraphHelpers.cpp
+++ b/Engine/Plugins/PCG/Source/PCG/Private/Helpers/PCGSubgraphHelpers.cpp
@@ -815,4 +815,57 @@ UPCGGraph* FPCGSubgraphHelpers::CollapseIntoSubgraphWithReason(UPCGGraph* InOrig
 	return NewPCGGraph;
 }
 
+UPCGNode* FPCGSubgraphHelpers::SpawnNodeAndConnect(UPCGGraph* InGraph, UPCGPin* InUpstreamPin, UPCGPin* InDownstreamPin, TSubclassOf<UPCGSettings> InSettingsClass, const TFunction<bool(UPCGSettings*)>& PostSettingsCreation)
+{
+	check(InGraph && InUpstreamPin && InDownstreamPin && InSettingsClass);
+	UPCGSettings* Settings;
+	UPCGNode* NewNode = InGraph->AddNodeOfType(InSettingsClass, Settings);
+	
+	if (NewNode)
+	{
+		if (PostSettingsCreation.IsSet() && PostSettingsCreation(Settings))
+		{
+			NewNode->UpdateAfterSettingsChangeDuringCreation();	
+		}
+		
+		UPCGNode* UpstreamNode = InUpstreamPin->Node;
+		UPCGNode* DownstreamNode = InDownstreamPin->Node;
+		check(UpstreamNode && DownstreamNode);
+
+#if WITH_EDITOR
+		// Put the new node between A & B but make it stay within a radius of B to keep things tidy.
+		NewNode->PositionX = (UpstreamNode->PositionX + DownstreamNode->PositionX) / 2;
+		NewNode->PositionY = (UpstreamNode->PositionY + DownstreamNode->PositionY) / 2;
+
+		// A hand tweaked distance that keeps it reasonably close.
+		constexpr float MaxDistFromB = 200;
+		const FVector2D OffsetFromB(NewNode->PositionX - DownstreamNode->PositionX, NewNode->PositionY - DownstreamNode->PositionY);
+		const float Dist = OffsetFromB.Length();
+		if (Dist > MaxDistFromB)
+		{
+			const float Scale = MaxDistFromB / Dist;
+			NewNode->PositionX = DownstreamNode->PositionX + Scale * OffsetFromB.X;
+			NewNode->PositionY = DownstreamNode->PositionY + Scale * OffsetFromB.Y;
+		}
+#endif // WITH_EDITOR
+
+		// Connect by default the first in and out pins of the new node.
+		if (ensure(!NewNode->GetInputPins().IsEmpty()))
+		{
+			check(NewNode->GetInputPins()[0]);
+			const FName NewNodeInputPinLabel = NewNode->GetInputPins()[0]->Properties.Label;
+			UpstreamNode->AddEdgeTo(InUpstreamPin->Properties.Label, NewNode, NewNodeInputPinLabel);
+		}
+
+		if (ensure(!NewNode->GetOutputPins().IsEmpty()))
+		{
+			check(NewNode->GetOutputPins()[0]);
+			const FName NewNodeOutputPinLabel = NewNode->GetOutputPins()[0]->Properties.Label;
+			NewNode->AddEdgeTo(NewNodeOutputPinLabel, DownstreamNode, InDownstreamPin->Properties.Label);
+		}
+	}
+	
+	return NewNode;
+}
+
 #undef LOCTEXT_NAMESPACE
\ No newline at end of file
diff --git a/Engine/Plugins/PCG/Source/PCG/Private/PCGModule.cpp b/Engine/Plugins/PCG/Source/PCG/Private/PCGModule.cpp
--- a/Engine/Plugins/PCG/Source/PCG/Private/PCGModule.cpp
+++ b/Engine/Plugins/PCG/Source/PCG/Private/PCGModule.cpp
@@ -38,20 +38,40 @@ FPCGGraphChangedDelegate FPCGModule::OnGraphChangedDelegate;
 namespace PCGModule
 {
 	FPCGModule* PCGModulePtr = nullptr;
+
+	static const TCHAR* DataTypeRegistryError = TEXT("Data type operations like compatibility or intersection can't run before the module is loaded."
+		"The logic will not function in a CDO and must be deferred until after the load.");
 }
 
 FPCGModule& FPCGModule::GetPCGModuleChecked()
 {
-	return PCGModule::PCGModulePtr ? *PCGModule::PCGModulePtr : FModuleManager::GetModuleChecked<FPCGModule>(TEXT("PCG"));
+	check(PCGModule::PCGModulePtr);
+	return *PCGModule::PCGModulePtr;
+}
+
+const FPCGDataTypeRegistry& FPCGModule::GetConstDataTypeRegistry()
+{
+	checkf(IsPCGModuleLoaded(), TEXT("%s"), PCGModule::DataTypeRegistryError);
+	return GetPCGModuleChecked().DataTypeRegistry;
+}
+
+FPCGDataTypeRegistry& FPCGModule::GetMutableDataTypeRegistry()
+{
+	checkf(IsPCGModuleLoaded(), TEXT("%s"), PCGModule::DataTypeRegistryError);
+	return GetPCGModuleChecked().DataTypeRegistry;
 }
 
 bool FPCGModule::IsPCGModuleLoaded()
 {
-	return FModuleManager::Get().IsModuleLoaded(TEXT("PCG"));
+	return PCGModule::PCGModulePtr != nullptr;
 }
 
 void FPCGModule::StartupModule()
 {
+	// Cache for fast access
+	check(!PCGModule::PCGModulePtr);
+	PCGModule::PCGModulePtr = this;
+	
 #if WITH_EDITOR
 	PCGDeterminismTests::FNativeTestRegistry::Create();
 
@@ -65,10 +85,6 @@ void FPCGModule::StartupModule()
 	LoadAdditionalKernelSources();
 #endif // WITH_EDITOR
 
-	// Cache for fast access
-	check(!PCGModule::PCGModulePtr);
-	PCGModule::PCGModulePtr = this;
-
 	// Registering accessor methods
 	AttributeAccessorFactory.RegisterDefaultMethods();
 	AttributeAccessorFactory.RegisterMethods<UPCGBasePointData>(UPCGBasePointData::GetPointAccessorMethods());
@@ -84,6 +100,8 @@ void FPCGModule::StartupModule()
 
 	// Register onto the PreExit, because we need the class to be still valid to remove them from the mapping
 	FCoreDelegates::OnPreExit.AddRaw(this, &FPCGModule::PreExit);
+	// Also register on post init to gather all PCG types to register them
+	FCoreDelegates::OnPostEngineInit.AddRaw(this, &FPCGModule::OnPostInitEngine);
 
 	TickDelegateHandle = FTSTicker::GetCoreTicker().AddTicker(FTickerDelegate::CreateRaw(this, &FPCGModule::Tick));
 }
@@ -92,6 +110,7 @@ void FPCGModule::ShutdownModule()
 {
 	FTSTicker::RemoveTicker(TickDelegateHandle);
 
+	FCoreDelegates::OnPostEngineInit.RemoveAll(this);
 	FCoreDelegates::OnPreExit.RemoveAll(this);
 
 	PCGModule::PCGModulePtr = nullptr;
@@ -113,6 +132,11 @@ void FPCGModule::PreExit()
 #endif // WITH_EDITOR
 }
 
+void FPCGModule::OnPostInitEngine()
+{
+	DataTypeRegistry.RegisterKnownTypes();
+}
+
 void FPCGModule::ExecuteNextTick(TFunction<void()> TickFunction)
 {
 	FScopeLock Lock(&ExecuteNextTickLock);
diff --git a/Engine/Plugins/PCG/Source/PCG/Private/PCGParamData.cpp b/Engine/Plugins/PCG/Source/PCG/Private/PCGParamData.cpp
--- a/Engine/Plugins/PCG/Source/PCG/Private/PCGParamData.cpp
+++ b/Engine/Plugins/PCG/Source/PCG/Private/PCGParamData.cpp
@@ -2,7 +2,7 @@
 
 #include "PCGParamData.h"
 #include "PCGContext.h"
-
+#include "Elements/PCGCollapseElement.h"
 #include "Elements/PCGExecuteBlueprint.h"
 #include "Metadata/PCGAttributePropertySelector.h"
 #include "Metadata/Accessors/IPCGAttributeAccessor.h"
@@ -13,6 +13,26 @@
 
 #include UE_INLINE_GENERATED_CPP_BY_NAME(PCGParamData)
 
+PCG_DEFINE_TYPE_INFO(FPCGDataTypeInfoParam, UPCGParamData)
+
+bool FPCGDataTypeInfoParam::SupportsConversion(const FPCGDataTypeIdentifier& Other) const
+{
+	// Param can convert to points
+	return Other == EPCGDataType::Point || FPCGDataTypeInfo::SupportsConversion(Other);
+}
+
+TSubclassOf<UPCGSettings> FPCGDataTypeInfoParam::GetDefaultConversionSettings(const FPCGDataTypeIdentifier& Other) const
+{
+	if (Other == EPCGDataType::Point)
+	{
+		return UPCGConvertToPointDataSettings::StaticClass();
+	}
+	else
+	{
+		return FPCGDataTypeInfo::GetDefaultConversionSettings(Other);
+	}
+}
+
 UPCGParamData::UPCGParamData(const FObjectInitializer& ObjectInitializer)
 	: Super(ObjectInitializer)
 {
diff --git a/Engine/Plugins/PCG/Source/PCG/Private/Tests/Registry/PCGDataTypeRegistryTests.cpp b/Engine/Plugins/PCG/Source/PCG/Private/Tests/Registry/PCGDataTypeRegistryTests.cpp
new file mode 100644
--- /dev/null
+++ b/Engine/Plugins/PCG/Source/PCG/Private/Tests/Registry/PCGDataTypeRegistryTests.cpp
@@ -0,0 +1,541 @@
+// Copyright Epic Games, Inc. All Rights Reserved.
+
+#if WITH_EDITOR
+
+#include "Tests/PCGTestsCommon.h"
+
+#include "PCGData.h"
+#include "PCGModule.h"
+#include "PCGParamData.h"
+#include "Compute/Data/PCGProxyForGPUData.h"
+#include "Data/PCGBasePointData.h"
+#include "Data/PCGDynamicMeshData.h"
+#include "Data/PCGLandscapeData.h"
+#include "Data/PCGLandscapeSplineData.h"
+#include "Data/PCGPolygon2DData.h"
+#include "Data/PCGPolyLineData.h"
+#include "Data/PCGPrimitiveData.h"
+#include "Data/PCGRenderTargetData.h"
+#include "Data/PCGSpatialData.h"
+#include "Data/PCGSplineData.h"
+#include "Data/PCGStaticMeshResourceData.h"
+#include "Data/PCGTextureData.h"
+#include "Data/PCGVirtualTextureData.h"
+#include "Data/PCGVolumeData.h"
+
+class FPCGDataTypeRegistryTestBase : public FPCGTestBaseClass
+{
+	using FPCGTestBaseClass::FPCGTestBaseClass;
+
+protected:
+	static bool IsRegistered(const FPCGDataTypeBaseId& ID)
+	{
+		return FPCGModule::GetConstDataTypeRegistry().IsRegistered(ID);
+	}
+};
+
+IMPLEMENT_CUSTOM_SIMPLE_AUTOMATION_TEST(FPCGByClassTest, FPCGDataTypeRegistryTestBase, "Plugins.PCG.DataTypeRegistry.ByClass", PCGTestsCommon::TestFlags)
+IMPLEMENT_CUSTOM_SIMPLE_AUTOMATION_TEST(FPCGByDefaultTypeTest, FPCGDataTypeRegistryTestBase, "Plugins.PCG.DataTypeRegistry.ByDefaultType", PCGTestsCommon::TestFlags)
+IMPLEMENT_CUSTOM_SIMPLE_AUTOMATION_TEST(FPCGCompatibilityTest, FPCGDataTypeRegistryTestBase, "Plugins.PCG.DataTypeRegistry.Compatibility", PCGTestsCommon::TestFlags)
+IMPLEMENT_CUSTOM_SIMPLE_AUTOMATION_TEST(FPCGAggregationTest, FPCGDataTypeRegistryTestBase, "Plugins.PCG.DataTypeRegistry.Aggregation", PCGTestsCommon::TestFlags)
+IMPLEMENT_CUSTOM_SIMPLE_AUTOMATION_TEST(FPCGIntersectionTest, FPCGDataTypeRegistryTestBase, "Plugins.PCG.DataTypeRegistry.Intersection", PCGTestsCommon::TestFlags)
+IMPLEMENT_CUSTOM_SIMPLE_AUTOMATION_TEST(FPCGCompositionTest, FPCGDataTypeRegistryTestBase, "Plugins.PCG.DataTypeRegistry.Composition", PCGTestsCommon::TestFlags)
+IMPLEMENT_CUSTOM_SIMPLE_AUTOMATION_TEST(FPCGUnicityTest, FPCGDataTypeRegistryTestBase, "Plugins.PCG.DataTypeRegistry.Unicity", PCGTestsCommon::TestFlags)
+IMPLEMENT_CUSTOM_SIMPLE_AUTOMATION_TEST(FPCGWiderTest, FPCGDataTypeRegistryTestBase, "Plugins.PCG.DataTypeRegistry.Wider", PCGTestsCommon::TestFlags)
+IMPLEMENT_CUSTOM_SIMPLE_AUTOMATION_TEST(FPCGUnionTest, FPCGDataTypeRegistryTestBase, "Plugins.PCG.DataTypeRegistry.Union", PCGTestsCommon::TestFlags)
+
+bool FPCGByClassTest::RunTest(const FString& Parameters)
+{
+	const FPCGDataTypeRegistry& Registry = FPCGModule::GetConstDataTypeRegistry();
+
+	TArray<FPCGDataTypeIdentifier> Identifiers;
+
+	TTuple<const TCHAR*, FPCGDataTypeIdentifier> TestValues[] =
+	{
+		{TEXT("Any"), FPCGDataTypeIdentifier::Construct<UPCGData>()},
+		{TEXT("Spatial"), FPCGDataTypeIdentifier::Construct<UPCGSpatialData>()},
+		{TEXT("Point"), FPCGDataTypeIdentifier::Construct<UPCGBasePointData>()},
+		{TEXT("Volume"), FPCGDataTypeIdentifier::Construct<UPCGVolumeData>()},
+		{TEXT("Primitive"), FPCGDataTypeIdentifier::Construct<UPCGPrimitiveData>()},
+		{TEXT("Polyline"), FPCGDataTypeIdentifier::Construct<UPCGPolyLineData>()},
+		{TEXT("Spline"), FPCGDataTypeIdentifier::Construct<UPCGSplineData>()},
+		{TEXT("Polygon2D"), FPCGDataTypeIdentifier::Construct<UPCGPolygon2DData>()},
+		{TEXT("Landscape spline"), FPCGDataTypeIdentifier::Construct<UPCGLandscapeSplineData>()},
+		{TEXT("Landscape"), FPCGDataTypeIdentifier::Construct<UPCGLandscapeData>()},
+		{TEXT("Base texture"), FPCGDataTypeIdentifier::Construct<UPCGBaseTextureData>()},
+		{TEXT("Texture"), FPCGDataTypeIdentifier::Construct<UPCGTextureData>()},
+		{TEXT("Proxy for GPU"), FPCGDataTypeIdentifier::Construct<UPCGProxyForGPUData>()},
+		{TEXT("Static mesh resource"), FPCGDataTypeIdentifier::Construct<UPCGStaticMeshResourceData>()},
+		{TEXT("Virtual Texture"), FPCGDataTypeIdentifier::Construct<UPCGVirtualTextureData>()},
+		{TEXT("Render target"), FPCGDataTypeIdentifier::Construct<UPCGRenderTargetData>()},
+		//{TEXT("Settings"), FPCGDataTypeIdentifier::Construct<UPCGSettings>()},
+		{TEXT("Param"), FPCGDataTypeIdentifier::Construct<UPCGParamData>()},
+		{TEXT("DynamicMesh"), FPCGDataTypeIdentifier::Construct<UPCGDynamicMeshData>()},
+	};
+
+	for (const auto& [TypeName, Identifier] : TestValues)
+	{
+		UTEST_TRUE(FString::Printf(TEXT("%s type is registered"), TypeName), Identifier.IsValid());
+		
+		UTEST_FALSE(
+			FString::Printf(TEXT("%s type identifier is unique"), TypeName),
+			Identifiers.ContainsByPredicate([&Identifier](const FPCGDataTypeIdentifier& Other) { return Identifier.IsSameType(Other); })
+		);
+		
+		Identifiers.Add(Identifier);
+	}
+	
+	return true;
+}
+
+bool FPCGByDefaultTypeTest::RunTest(const FString& Parameters)
+{
+	const FPCGDataTypeRegistry& Registry = FPCGModule::GetConstDataTypeRegistry();
+
+	// Test invalid first
+	UTEST_FALSE(FString::Printf(TEXT("%s type is invalid"), TEXT("None")), FPCGDataTypeIdentifier{EPCGDataType::None}.IsValid());
+
+	TArray<FPCGDataTypeIdentifier> Identifiers;
+
+	TTuple<const TCHAR*, FPCGDataTypeIdentifier> TestValues[] =
+	{
+		{TEXT("Any"), FPCGDataTypeIdentifier(EPCGDataType::Any)},
+		{TEXT("Spatial"), FPCGDataTypeIdentifier(EPCGDataType::Spatial)},
+		{TEXT("Point"), FPCGDataTypeIdentifier(EPCGDataType::Point)},
+		{TEXT("Volume"), FPCGDataTypeIdentifier(EPCGDataType::Volume)},
+		{TEXT("Primitive"), FPCGDataTypeIdentifier(EPCGDataType::Primitive)},
+		{TEXT("Polyline"), FPCGDataTypeIdentifier(EPCGDataType::PolyLine)},
+		{TEXT("Spline"), FPCGDataTypeIdentifier(EPCGDataType::Spline)},
+		{TEXT("Polygon2D"), FPCGDataTypeIdentifier(EPCGDataType::Polygon2D)},
+		{TEXT("Landscape spline"), FPCGDataTypeIdentifier(EPCGDataType::LandscapeSpline)},
+		{TEXT("Landscape"), FPCGDataTypeIdentifier(EPCGDataType::Landscape)},
+		{TEXT("Base texture"), FPCGDataTypeIdentifier(EPCGDataType::BaseTexture)},
+		{TEXT("Texture"), FPCGDataTypeIdentifier(EPCGDataType::Texture)},
+		{TEXT("Proxy for GPU"), FPCGDataTypeIdentifier(EPCGDataType::ProxyForGPU)},
+		{TEXT("Static mesh resource"), FPCGDataTypeIdentifier(EPCGDataType::StaticMeshResource)},
+		{TEXT("Virtual Texture"), FPCGDataTypeIdentifier(EPCGDataType::VirtualTexture)},
+		{TEXT("Render target"),FPCGDataTypeIdentifier(EPCGDataType::RenderTarget)},
+		//{TEXT("Settings"), FPCGDataTypeIdentifier(EPCGDataType::Settings)},
+		{TEXT("Param"), FPCGDataTypeIdentifier(EPCGDataType::Param)},
+		{TEXT("PointOrParam"), FPCGDataTypeIdentifier(EPCGDataType::PointOrParam)},
+		{TEXT("DynamicMesh"), FPCGDataTypeIdentifier(EPCGDataType::DynamicMesh)}
+	};
+
+	for (const auto& [TypeName, Identifier] : TestValues)
+	{
+		UTEST_TRUE(FString::Printf(TEXT("%s type is valid"), TypeName), Identifier.IsValid());
+		bool bRegistered = Algo::AllOf(Identifier.GetIDs(), [](const FPCGDataTypeBaseId ID) { return IsRegistered(ID); });
+		
+		UTEST_TRUE(FString::Printf(TEXT("%s type is registered"), TypeName), bRegistered);
+		
+		UTEST_FALSE(
+			FString::Printf(TEXT("%s type identifier is unique"), TypeName),
+			Identifiers.ContainsByPredicate([&Identifier](const FPCGDataTypeIdentifier& Other) { return Identifier.IsSameType(Other); })
+		);
+		
+		Identifiers.Add(Identifier);
+	}
+	
+	return true;
+}
+
+bool FPCGCompatibilityTest::RunTest(const FString& Parameters)
+{
+	TTuple<EPCGDataType, EPCGDataType, EPCGDataTypeCompatibilityResult> TestValues[] =
+	{
+		{EPCGDataType::Any, EPCGDataType::None, EPCGDataTypeCompatibilityResult::UnknownType},
+		{EPCGDataType::None, EPCGDataType::None, EPCGDataTypeCompatibilityResult::UnknownType},
+		{EPCGDataType::None, EPCGDataType::Any, EPCGDataTypeCompatibilityResult::UnknownType},
+		
+		{EPCGDataType::Any, EPCGDataType::Spatial, EPCGDataTypeCompatibilityResult::RequireFilter},
+		{EPCGDataType::Any, EPCGDataType::Concrete, EPCGDataTypeCompatibilityResult::RequireConversion},
+		{EPCGDataType::Any, EPCGDataType::Point, EPCGDataTypeCompatibilityResult::RequireFilter},
+		
+		{EPCGDataType::None, EPCGDataType::None, EPCGDataTypeCompatibilityResult::UnknownType},
+		
+		{EPCGDataType::Spatial, EPCGDataType::Any, EPCGDataTypeCompatibilityResult::Compatible},
+		{EPCGDataType::Point, EPCGDataType::Any, EPCGDataTypeCompatibilityResult::Compatible},
+		{EPCGDataType::Any, EPCGDataType::Any, EPCGDataTypeCompatibilityResult::Compatible},
+		
+		{EPCGDataType::Point, EPCGDataType::Point, EPCGDataTypeCompatibilityResult::Compatible},
+		
+		{EPCGDataType::Spatial, EPCGDataType::Concrete, EPCGDataTypeCompatibilityResult::RequireConversion},
+		{EPCGDataType::Spatial, EPCGDataType::Point, EPCGDataTypeCompatibilityResult::RequireConversion},
+		{EPCGDataType::Spatial, EPCGDataType::Volume, EPCGDataTypeCompatibilityResult::RequireFilter},
+		{EPCGDataType::Spatial, EPCGDataType::Surface, EPCGDataTypeCompatibilityResult::RequireFilter},
+		{EPCGDataType::Spatial, EPCGDataType::Spline, EPCGDataTypeCompatibilityResult::RequireFilter},
+		
+		{EPCGDataType::Composite, EPCGDataType::Concrete, EPCGDataTypeCompatibilityResult::RequireConversion},
+		
+		{EPCGDataType::Spline, EPCGDataType::Point, EPCGDataTypeCompatibilityResult::RequireConversion},
+		{EPCGDataType::Volume, EPCGDataType::Point, EPCGDataTypeCompatibilityResult::RequireConversion},
+		
+		{EPCGDataType::Spline, EPCGDataType::Surface, EPCGDataTypeCompatibilityResult::RequireConversion},
+		{EPCGDataType::Polygon2D, EPCGDataType::Surface, EPCGDataTypeCompatibilityResult::RequireConversion},
+		
+		{EPCGDataType::Spline, EPCGDataType::LandscapeSpline, EPCGDataTypeCompatibilityResult::NotCompatible},
+		{EPCGDataType::Spline, EPCGDataType::PolyLine, EPCGDataTypeCompatibilityResult::Compatible},
+		{EPCGDataType::LandscapeSpline, EPCGDataType::PolyLine, EPCGDataTypeCompatibilityResult::Compatible},
+		
+		{EPCGDataType::Landscape, EPCGDataType::Texture, EPCGDataTypeCompatibilityResult::NotCompatible},
+		{EPCGDataType::Landscape, EPCGDataType::Surface, EPCGDataTypeCompatibilityResult::Compatible},
+		{EPCGDataType::Surface, EPCGDataType::Landscape, EPCGDataTypeCompatibilityResult::RequireFilter},
+		
+		{EPCGDataType::RenderTarget, EPCGDataType::Texture, EPCGDataTypeCompatibilityResult::NotCompatible},
+		{EPCGDataType::BaseTexture, EPCGDataType::Texture, EPCGDataTypeCompatibilityResult::RequireFilter},
+		{EPCGDataType::Texture, EPCGDataType::BaseTexture, EPCGDataTypeCompatibilityResult::Compatible},
+		{EPCGDataType::BaseTexture, EPCGDataType::RenderTarget, EPCGDataTypeCompatibilityResult::RequireFilter},
+		{EPCGDataType::RenderTarget, EPCGDataType::BaseTexture, EPCGDataTypeCompatibilityResult::Compatible},
+		{EPCGDataType::VirtualTexture, EPCGDataType::Texture, EPCGDataTypeCompatibilityResult::NotCompatible},
+		{EPCGDataType::VirtualTexture, EPCGDataType::Surface, EPCGDataTypeCompatibilityResult::Compatible},
+		
+		{EPCGDataType::Point, EPCGDataType::Param, EPCGDataTypeCompatibilityResult::RequireConversion},
+		{EPCGDataType::Param, EPCGDataType::Point, EPCGDataTypeCompatibilityResult::RequireConversion},
+		
+		{EPCGDataType::PointOrParam, EPCGDataType::Param, EPCGDataTypeCompatibilityResult::RequireFilter},
+		{EPCGDataType::PointOrParam, EPCGDataType::Point, EPCGDataTypeCompatibilityResult::RequireFilter},
+		{EPCGDataType::Point, EPCGDataType::PointOrParam, EPCGDataTypeCompatibilityResult::Compatible},
+		{EPCGDataType::Param, EPCGDataType::Param, EPCGDataTypeCompatibilityResult::Compatible},
+
+		
+		{EPCGDataType::Spline | EPCGDataType::Point, EPCGDataType::Spline, EPCGDataTypeCompatibilityResult::RequireFilter},
+		{EPCGDataType::Spline | EPCGDataType::Point, EPCGDataType::Point, EPCGDataTypeCompatibilityResult::RequireFilter},
+		{EPCGDataType::Spline | EPCGDataType::Point, EPCGDataType::Param, EPCGDataTypeCompatibilityResult::NotCompatible},
+		{EPCGDataType::Spline | EPCGDataType::Point, EPCGDataType::Spline | EPCGDataType::Landscape, EPCGDataTypeCompatibilityResult::RequireFilter},
+		
+		{EPCGDataType::Param, EPCGDataType::Spline | EPCGDataType::Point, EPCGDataTypeCompatibilityResult::NotCompatible},
+		{EPCGDataType::Spline, EPCGDataType::Spline | EPCGDataType::Point, EPCGDataTypeCompatibilityResult::Compatible},
+		{EPCGDataType::Point, EPCGDataType::Spline | EPCGDataType::Point, EPCGDataTypeCompatibilityResult::Compatible},
+		{EPCGDataType::Spline | EPCGDataType::Landscape, EPCGDataType::Spline | EPCGDataType::Point, EPCGDataTypeCompatibilityResult::RequireFilter},
+	};
+
+	const FPCGDataTypeRegistry& Registry = FPCGModule::GetConstDataTypeRegistry();
+
+	const UEnum* Enum = StaticEnum<EPCGDataType>();
+	check(Enum);
+
+	for (const auto& [TypeA, TypeB, ExpectedCompatibility] : TestValues)
+	{
+		const FName TypeAName = Enum->GetNameByValue(static_cast<int64>(TypeA));
+		const FName TypeBName = Enum->GetNameByValue(static_cast<int64>(TypeB));
+
+		const FPCGDataTypeIdentifier TypeAIdentifier{TypeA};
+		const FPCGDataTypeIdentifier TypeBIdentifier{TypeB};
+		
+		UTEST_EQUAL(
+			FString::Printf(TEXT("Compatiblity between %s and %s is expected."),
+				*TypeAName.ToString(),
+				*TypeBName.ToString()),
+			static_cast<int64>(Registry.IsCompatible(TypeAIdentifier, TypeBIdentifier)),
+			static_cast<int64>(ExpectedCompatibility)
+		);
+	}
+
+	return true;
+}
+
+/**
+ * Aggregation should return all the base types.
+ */
+bool FPCGAggregationTest::RunTest(const FString& Parameters)
+{
+	TTuple<FPCGDataTypeIdentifier, FPCGDataTypeIdentifier> TestValues[] =
+	{
+		{
+			FPCGDataTypeIdentifier{EPCGDataType::Point} | EPCGDataType::None, 
+			FPCGDataTypeIdentifier::Construct<UPCGPointData>()
+		},
+		{ 
+			EPCGDataType::Point | EPCGDataType::Param, 
+			FPCGDataTypeIdentifier::Construct<UPCGPointData, UPCGParamData>()
+		},
+		{
+			EPCGDataType::Point | EPCGDataType::Param,
+			FPCGDataTypeIdentifier::Construct<UPCGParamData, UPCGPointData>()
+		},
+		{
+			EPCGDataType::PointOrParam,
+			FPCGDataTypeIdentifier::Construct<UPCGPointArrayData, UPCGParamData>()
+		},
+		{
+			EPCGDataType::PointOrParam,
+			FPCGDataTypeIdentifier::Construct<UPCGParamData, UPCGBasePointData>()
+		},
+		{
+			EPCGDataType::Spline | EPCGDataType::Volume,
+			FPCGDataTypeIdentifier::Construct<UPCGSplineData, UPCGVolumeData>()
+		},
+		{
+			EPCGDataType::Texture | EPCGDataType::Point | EPCGDataType::Landscape,
+			FPCGDataTypeIdentifier::Construct<UPCGTextureData, UPCGBasePointData, UPCGLandscapeData>()
+		},
+	};
+
+	const UEnum* Enum = StaticEnum<EPCGDataType>();
+	check(Enum);
+
+	for (const auto& [Type, ExpectedIdentifier] : TestValues)
+	{
+		const FPCGDataTypeIdentifier TypeIdentifier{Type};
+		
+		PRAGMA_DISABLE_DEPRECATION_WARNINGS
+		const FString TypeName = Enum->GetValueOrBitfieldAsString(static_cast<int64>(Type));
+		PRAGMA_ENABLE_DEPRECATION_WARNINGS
+		
+		UTEST_TRUE(FString::Printf(TEXT("Composition for %s is valid."), *TypeName), TypeIdentifier.IsSameType(ExpectedIdentifier))
+	}
+
+	return true;
+}
+
+/**
+ * Intersection should return the common type, or none if there is no common type.
+ */
+bool FPCGIntersectionTest::RunTest(const FString& Parameters)
+{
+	TTuple<TArray<EPCGDataType>, FPCGDataTypeIdentifier> TestValues[] =
+	{
+		{ 
+			TArray{EPCGDataType::Point, EPCGDataType::None}, 
+			FPCGDataTypeIdentifier{}
+		},
+		{ 
+			TArray{EPCGDataType::Point, EPCGDataType::Param}, 
+			FPCGDataTypeIdentifier{}
+		},
+		{
+			TArray{EPCGDataType::Spline, EPCGDataType::Volume},
+			FPCGDataTypeIdentifier{}
+		},
+		{
+			TArray{EPCGDataType::Surface, EPCGDataType::Point, EPCGDataType::Landscape},
+			FPCGDataTypeIdentifier{}
+		},
+		{
+			TArray{EPCGDataType::Texture,  EPCGDataType::Texture},
+			FPCGDataTypeIdentifier{EPCGDataType::Texture}
+		},
+		{
+			TArray{EPCGDataType::Point,  EPCGDataType::PointOrParam},
+			FPCGDataTypeIdentifier{EPCGDataType::Point}
+		},
+		{
+			TArray{EPCGDataType::Param,  EPCGDataType::PointOrParam},
+			FPCGDataTypeIdentifier{EPCGDataType::Param}
+		},
+		{
+			TArray{EPCGDataType::Surface,  EPCGDataType::Landscape},
+			FPCGDataTypeIdentifier{EPCGDataType::Landscape}
+		},
+		{
+			TArray{EPCGDataType::Surface,  EPCGDataType::Texture},
+			FPCGDataTypeIdentifier{EPCGDataType::Texture}
+		},
+		{
+			TArray{EPCGDataType::Any,  EPCGDataType::Spatial},
+			FPCGDataTypeIdentifier{EPCGDataType::Spatial}
+		},
+	};
+
+	const UEnum* Enum = StaticEnum<EPCGDataType>();
+	check(Enum);
+
+	for (const auto& [Types, ExpectedIdentifier] : TestValues)
+	{
+		FPCGDataTypeIdentifier TypeIdentifier{Types[0]};
+		FString TypesNames = TypeIdentifier.ToString();
+		for (int32 i = 1; i < Types.Num(); ++i)
+		{
+			const FPCGDataTypeIdentifier OtherType{Types[i]};
+			TypeIdentifier &= OtherType;
+			TypesNames += TEXT(" & ") + OtherType.ToString();
+		}
+
+		UTEST_TRUE(FString::Printf(TEXT("Intersection for %s is %s."), *TypesNames, *ExpectedIdentifier.ToString()), TypeIdentifier.IsSameType(ExpectedIdentifier))
+	}
+
+	return true;
+}
+
+/**
+ * Composition: Return the aggregation with other types, with reduction. (ie. if Data B and C are child of Data A, Composition of B and C gives A)
+ */
+bool FPCGCompositionTest::RunTest(const FString& Parameters)
+{
+	TTuple<TArray<FPCGDataTypeIdentifier>, FPCGDataTypeIdentifier> TestValues[] =
+	{
+		{
+			TArray{FPCGDataTypeIdentifier{EPCGDataType::Point}, FPCGDataTypeIdentifier{EPCGDataType::None}},
+			FPCGDataTypeIdentifier{EPCGDataType::Point}
+		},
+		{
+			TArray{FPCGDataTypeIdentifier{EPCGDataType::Any}, FPCGDataTypeIdentifier{EPCGDataType::Spatial}},
+			FPCGDataTypeIdentifier{EPCGDataType::Any}
+		},
+		{
+			TArray{FPCGDataTypeIdentifier{EPCGDataType::Spatial}, FPCGDataTypeIdentifier{EPCGDataType::Concrete}},
+			FPCGDataTypeIdentifier{EPCGDataType::Spatial}
+		},
+		{
+			TArray{FPCGDataTypeIdentifier{EPCGDataType::Spatial}, (FPCGDataTypeIdentifier{EPCGDataType::Composite})},
+			FPCGDataTypeIdentifier{EPCGDataType::Spatial}
+		},
+		{
+			TArray{FPCGDataTypeIdentifier{EPCGDataType::Concrete}, FPCGDataTypeIdentifier{EPCGDataType::Composite}},
+			FPCGDataTypeIdentifier{EPCGDataType::Spatial}
+		},
+		{
+			TArray{FPCGDataTypeIdentifier{EPCGDataType::Texture}, FPCGDataTypeIdentifier{EPCGDataType::Composite}},
+			FPCGDataTypeIdentifier{EPCGDataType::Texture | EPCGDataType::Composite}
+		},
+		{
+			TArray{FPCGDataTypeIdentifier{EPCGDataType::Concrete}, FPCGDataTypeIdentifier{EPCGDataType::Point}},
+			FPCGDataTypeIdentifier{EPCGDataType::Concrete}
+		},
+		{
+			TArray{FPCGDataTypeIdentifier{EPCGDataType::Surface}, FPCGDataTypeIdentifier{EPCGDataType::Texture}},
+			FPCGDataTypeIdentifier{EPCGDataType::Surface}
+		},
+		{
+			TArray{FPCGDataTypeIdentifier{EPCGDataType::Spline}, FPCGDataTypeIdentifier{EPCGDataType::LandscapeSpline}, FPCGDataTypeIdentifier{EPCGDataType::Polygon2D}},
+			FPCGDataTypeIdentifier{EPCGDataType::PolyLine}
+		},
+		{
+			TArray{FPCGDataTypeIdentifier{EPCGDataType::Spline}, FPCGDataTypeIdentifier{EPCGDataType::LandscapeSpline}, FPCGDataTypeIdentifier{EPCGDataType::Polygon2D}, FPCGDataTypeIdentifier{EPCGDataType::Point} },
+			FPCGDataTypeIdentifier::Construct<UPCGPolyLineData, UPCGPointData>()
+		},
+		{
+			TArray{FPCGDataTypeIdentifier{EPCGDataType::Texture}, FPCGDataTypeIdentifier{EPCGDataType::RenderTarget}},
+			FPCGDataTypeIdentifier{EPCGDataType::BaseTexture}
+		},
+		{
+			TArray{FPCGDataTypeIdentifier{EPCGDataType::Texture}, FPCGDataTypeIdentifier{EPCGDataType::RenderTarget}, FPCGDataTypeIdentifier{EPCGDataType::VirtualTexture}, FPCGDataTypeIdentifier{EPCGDataType::Landscape}},
+			FPCGDataTypeIdentifier{EPCGDataType::Surface}
+		},
+		{
+			TArray{FPCGDataTypeIdentifier{EPCGDataType::Point}, FPCGDataTypeIdentifier{EPCGDataType::Param}},
+			FPCGDataTypeIdentifier::Construct<UPCGPointData, UPCGParamData>()
+		},
+		{
+			TArray{FPCGDataTypeIdentifier{EPCGDataType::Point}, FPCGDataTypeIdentifier{EPCGDataType::Param}},
+			FPCGDataTypeIdentifier::Construct<UPCGParamData, UPCGPointData>()
+		},
+		{
+			TArray{FPCGDataTypeIdentifier{EPCGDataType::Spline}, FPCGDataTypeIdentifier{EPCGDataType::Volume}},
+			FPCGDataTypeIdentifier::Construct<UPCGSplineData, UPCGVolumeData>()
+		},
+		{
+			TArray{FPCGDataTypeIdentifier{EPCGDataType::Texture}, FPCGDataTypeIdentifier{EPCGDataType::Point}, FPCGDataTypeIdentifier{EPCGDataType::Landscape}},
+			FPCGDataTypeIdentifier::Construct<UPCGTextureData, UPCGPointData, UPCGLandscapeData>()
+		},
+	};
+
+	for (const auto& [TypeIdentifiers, ExpectedIdentifier] : TestValues)
+	{
+		const FPCGDataTypeIdentifier CompositionType = FPCGDataTypeIdentifier::Compose(TypeIdentifiers);
+		
+		const FString TypeName = FPCGDataTypeIdentifier(TypeIdentifiers).ToString();
+		const FString ExpectedTypeName = ExpectedIdentifier.ToString();
+
+		UTEST_TRUE(FString::Printf(TEXT("Composition of %s gives %s."), *TypeName, *ExpectedTypeName), CompositionType.IsSameType(ExpectedIdentifier))
+	}
+
+	return true;
+}
+
+/**
+ * We should never have duplicates in the IDs when we construct/combine/compose the same types.
+ */
+bool FPCGUnicityTest::RunTest(const FString& Parameters)
+{
+	const FPCGDataTypeIdentifier SameTypeWithClassTpl = FPCGDataTypeIdentifier::Construct<UPCGData, UPCGData>();
+	UTEST_EQUAL("Construct with same class in template gives just 1 ID", SameTypeWithClassTpl.GetIDs().Num(), 1);
+
+	const FPCGDataTypeIdentifier SameTypeWithClass = FPCGDataTypeIdentifier({UPCGPointData::StaticClass(), UPCGPointData::StaticClass()});
+	UTEST_EQUAL("Construct with same class gives just 1 ID", SameTypeWithClass.GetIDs().Num(), 1);
+
+	const FPCGDataTypeIdentifier SameTypeWithBaseIds = FPCGDataTypeIdentifier({FPCGDataTypeBaseId::Construct<FPCGDataTypeInfoLandscape>(), FPCGDataTypeBaseId::Construct<FPCGDataTypeInfoLandscape>()});
+	UTEST_EQUAL("Construct with same base ids gives just 1 ID", SameTypeWithBaseIds.GetIDs().Num(), 1);
+
+	const FPCGDataTypeIdentifier SameTypeWithIds = FPCGDataTypeIdentifier({SameTypeWithClassTpl, SameTypeWithClassTpl, SameTypeWithClassTpl});
+	UTEST_EQUAL("Construct with same ids gives just 1 ID", SameTypeWithIds.GetIDs().Num(), 1);
+
+	const FPCGDataTypeIdentifier SameTypeWithDefaultTypeOr = SameTypeWithClass | EPCGDataType::Point;
+	UTEST_EQUAL("Bitwise Or with same default type gives just 1 ID", SameTypeWithDefaultTypeOr.GetIDs().Num(), 1);
+
+	const FPCGDataTypeIdentifier SameTypeWithOr = SameTypeWithClass | SameTypeWithClass;
+	UTEST_EQUAL("Bitwise Or with same ids gives just 1 ID", SameTypeWithOr.GetIDs().Num(), 1);
+
+	FPCGDataTypeIdentifier SameTypeOr = SameTypeWithClass;
+	SameTypeOr |= SameTypeWithClass;
+	UTEST_EQUAL("Bitwise Or with self gives just 1 ID", SameTypeOr.GetIDs().Num(), 1);
+
+	SameTypeOr |= EPCGDataType::Point;
+	UTEST_EQUAL("Bitwise Or on self with default type gives just 1 ID", SameTypeOr.GetIDs().Num(), 1);
+
+	const FPCGDataTypeIdentifier SameTypeCompose = SameTypeWithClass.Compose({SameTypeWithClass, SameTypeWithClass});
+	UTEST_EQUAL("Compose with same type gives just 1 ID", SameTypeCompose.GetIDs().Num(), 1);
+
+	return true;
+}
+
+bool FPCGWiderTest::RunTest(const FString& Parameters)
+{
+	// Test: Left is wider than right. Result in the bool
+	TTuple<FPCGDataTypeIdentifier, FPCGDataTypeIdentifier, bool> TestValues[] =
+	{
+		{EPCGDataType::Any, EPCGDataType::Concrete, true},
+		{EPCGDataType::Concrete, EPCGDataType::Any, false},
+		{EPCGDataType::PointOrParam, EPCGDataType::Point, true},
+		{EPCGDataType::PointOrParam, EPCGDataType::Param, true},
+		{EPCGDataType::Param, EPCGDataType::PointOrParam, false},
+		{EPCGDataType::Point, EPCGDataType::PointOrParam, false},
+		{EPCGDataType::PolyLine, EPCGDataType::Spline, true},
+		{EPCGDataType::Spline, EPCGDataType::PolyLine, false},
+		{EPCGDataType::Spline, EPCGDataType::LandscapeSpline, true},
+		{EPCGDataType::LandscapeSpline, EPCGDataType::Spline, true},
+		{EPCGDataType::Spline | EPCGDataType::Point, EPCGDataType::Point | EPCGDataType::LandscapeSpline, true}
+	};
+
+	for (const auto& [LeftIdentifier, RightIdentifier, ExpectedResult] : TestValues)
+	{
+		if (ExpectedResult)
+		{
+			UTEST_TRUE(FString::Printf(TEXT("%s is wider than %s."), *LeftIdentifier.ToString(), *RightIdentifier.ToString()), LeftIdentifier.IsWider(RightIdentifier))
+		}
+		else
+		{
+			UTEST_FALSE(FString::Printf(TEXT("%s is not wider than %s."), *LeftIdentifier.ToString(), *RightIdentifier.ToString()), LeftIdentifier.IsWider(RightIdentifier))
+		}
+	}
+
+	return true;
+}
+
+bool FPCGUnionTest::RunTest(const FString& Parameters)
+{
+	TTuple<FPCGDataTypeIdentifier, FPCGDataTypeIdentifier> TestValues[] =
+	{
+		{FPCGDataTypeIdentifier{EPCGDataType::Point} | EPCGDataType::None, EPCGDataType::Point},
+		{FPCGDataTypeIdentifier{EPCGDataType::Any} | EPCGDataType::Concrete, EPCGDataType::Any},
+		{FPCGDataTypeIdentifier{EPCGDataType::Param} | EPCGDataType::Spatial, EPCGDataType::Any},
+		{FPCGDataTypeIdentifier{EPCGDataType::Param} | EPCGDataType::Point, EPCGDataType::Any},
+		{FPCGDataTypeIdentifier{EPCGDataType::Texture} | EPCGDataType::RenderTarget, EPCGDataType::BaseTexture},
+		{FPCGDataTypeIdentifier{EPCGDataType::Landscape} | EPCGDataType::RenderTarget, EPCGDataType::Surface},
+		{FPCGDataTypeIdentifier{EPCGDataType::PointOrParam}, EPCGDataType::Any},
+		{FPCGDataTypeIdentifier{EPCGDataType::Point} | EPCGDataType::Landscape | EPCGDataType::Polygon2D, EPCGDataType::Concrete},
+		{FPCGDataTypeIdentifier{EPCGDataType::Point} | EPCGDataType::Landscape | EPCGDataType::Polygon2D | EPCGDataType::Composite, EPCGDataType::Spatial},
+	};
+	
+	const FPCGDataTypeRegistry& Registry = FPCGModule::GetConstDataTypeRegistry();
+	
+	for (const auto& [Identifiers, ExpectedIdentifier] : TestValues)
+	{
+		const FPCGDataTypeIdentifier Union = Registry.GetIdentifiersUnion({Identifiers});
+		UTEST_TRUE(FString::Printf(TEXT("Union of %s is %s."), *Identifiers.ToString(), *ExpectedIdentifier.ToString()), Union.IsSameType(ExpectedIdentifier))
+	}
+	
+	return true;
+}
+
+#endif // WITH_EDITOR
\ No newline at end of file
diff --git a/Engine/Plugins/PCG/Source/PCG/Public/Data/PCGBasePointData.h b/Engine/Plugins/PCG/Source/PCG/Public/Data/PCGBasePointData.h
--- a/Engine/Plugins/PCG/Source/PCG/Public/Data/PCGBasePointData.h
+++ b/Engine/Plugins/PCG/Source/PCG/Public/Data/PCGBasePointData.h
@@ -18,6 +18,17 @@ namespace PCGPointDataConstants
 	const FName ElementsDomainName = "Points";
 }
 
+USTRUCT()
+struct FPCGDataTypeInfoPoint : public FPCGDataTypeInfoConcrete
+{
+	GENERATED_BODY()
+
+	PCG_DECLARE_LEGACY_TYPE_INFO(EPCGDataType::Point)
+
+	PCG_API virtual bool SupportsConversion(const FPCGDataTypeIdentifier& Other) const override;
+	PCG_API virtual TSubclassOf<UPCGSettings> GetDefaultConversionSettings(const FPCGDataTypeIdentifier& Other) const override;
+};
+
 UCLASS(MinimalAPI, Abstract, BlueprintType, ClassGroup = (Procedural))
 class UPCGBasePointData : public UPCGSpatialData
 {
@@ -40,6 +51,7 @@ public:
 	// ~End UPCGSpatialData interface
 		
 	// ~Begin UPCGData Interface
+	PCG_ASSIGN_DEFAULT_TYPE_INFO(FPCGDataTypeInfoPoint)
 	
 	/** Make a pass on Metadata to flatten parenting and only keep entries used by points. */
 	UE_API virtual void Flatten() override;
@@ -62,8 +74,6 @@ public:
 	UE_API void AddSinglePointFromActor(AActor* InActor, bool* bOutOptionalSanitizedTagAttributeName = nullptr);
 
 	virtual bool IsValidRef(const PCGPointOctree::FPointRef& InPointRef) const PURE_VIRTUAL(UPCGBasePointData::IsValidRef, return false;);
-	
-	virtual EPCGDataType GetDataType() const override { return EPCGDataType::Point; }
 
 	virtual int32 GetNumPoints() const PURE_VIRTUAL(UPCGBasePointData::GetNumPoints, return 0;);
 	bool IsEmpty() const { return GetNumPoints() == 0; }
diff --git a/Engine/Plugins/PCG/Source/PCG/Public/Data/PCGCollisionShapeData.h b/Engine/Plugins/PCG/Source/PCG/Public/Data/PCGCollisionShapeData.h
--- a/Engine/Plugins/PCG/Source/PCG/Public/Data/PCGCollisionShapeData.h
+++ b/Engine/Plugins/PCG/Source/PCG/Public/Data/PCGCollisionShapeData.h
@@ -2,6 +2,7 @@
 
 #pragma once
 
+#include "PCGPrimitiveData.h"
 #include "PCGSpatialData.h"
 
 #include "CollisionShape.h"
@@ -24,7 +25,7 @@ public:
 	PCG_API static bool IsSupported(UShapeComponent* InShape);
 
 	// ~Begin UPCGData interface
-	virtual EPCGDataType GetDataType() const override { return EPCGDataType::Primitive; }
+	PCG_ASSIGN_DEFAULT_TYPE_INFO(FPCGDataTypeInfoPrimitive)
 	virtual void AddToCrc(FArchiveCrc32& Ar, bool bFullDataCrc) const override;
 	// ~End UPCGData interface
 
diff --git a/Engine/Plugins/PCG/Source/PCG/Public/Data/PCGCollisionWrapperData.h b/Engine/Plugins/PCG/Source/PCG/Public/Data/PCGCollisionWrapperData.h
--- a/Engine/Plugins/PCG/Source/PCG/Public/Data/PCGCollisionWrapperData.h
+++ b/Engine/Plugins/PCG/Source/PCG/Public/Data/PCGCollisionWrapperData.h
@@ -2,6 +2,7 @@
 
 #pragma once
 
+#include "PCGPrimitiveData.h"
 #include "PCGSpatialData.h"
 #include "Data/PCGBasePointData.h"
 #include "Metadata/PCGAttributePropertySelector.h"
@@ -71,7 +72,7 @@ public:
 	void FinalizeInitializationEx(const TArray<FSoftObjectPath>& InMeshPaths);
 
 	// ~Begin UPCGData Interface
-	virtual EPCGDataType GetDataType() const override { return EPCGDataType::Primitive; }
+	PCG_ASSIGN_DEFAULT_TYPE_INFO(FPCGDataTypeInfoPrimitive)
 	virtual FPCGCrc ComputeCrc(bool bFullDataCrc) const override;
 	virtual void AddToCrc(FArchiveCrc32& Ar, bool bFullDataCrc) const override;
 	virtual void GetResourceSizeEx(FResourceSizeEx& CumulativeResourceSize) override;
diff --git a/Engine/Plugins/PCG/Source/PCG/Public/Data/PCGDifferenceData.h b/Engine/Plugins/PCG/Source/PCG/Public/Data/PCGDifferenceData.h
--- a/Engine/Plugins/PCG/Source/PCG/Public/Data/PCGDifferenceData.h
+++ b/Engine/Plugins/PCG/Source/PCG/Public/Data/PCGDifferenceData.h
@@ -61,7 +61,7 @@ public:
 	// ~End UObject interface
 
 	// ~Begin UPCGData interface
-	virtual EPCGDataType GetDataType() const override { return EPCGDataType::Spatial; }
+	PCG_ASSIGN_DEFAULT_TYPE_INFO(FPCGDataTypeInfoComposite)
 	UE_API virtual void VisitDataNetwork(TFunctionRef<void(const UPCGData*)> Action) const override;
 
 protected:
diff --git a/Engine/Plugins/PCG/Source/PCG/Public/Data/PCGDynamicMeshData.h b/Engine/Plugins/PCG/Source/PCG/Public/Data/PCGDynamicMeshData.h
--- a/Engine/Plugins/PCG/Source/PCG/Public/Data/PCGDynamicMeshData.h
+++ b/Engine/Plugins/PCG/Source/PCG/Public/Data/PCGDynamicMeshData.h
@@ -16,6 +16,14 @@ class UDynamicMesh;
 class UDynamicMeshComponent;
 class UMaterialInterface;
 
+USTRUCT()
+struct FPCGDataTypeInfoDynamicMesh: public FPCGDataTypeInfo
+{
+	GENERATED_BODY()
+
+	PCG_DECLARE_LEGACY_TYPE_INFO(EPCGDataType::DynamicMesh);
+};
+
 UCLASS(MinimalAPI, BlueprintType, ClassGroup = (Procedural))
 class UPCGDynamicMeshData : public UPCGSpatialData
 {
@@ -37,7 +45,7 @@ public:
 	void K2_Initialize(UDynamicMesh* InMesh, const TArray<UMaterialInterface*>& InMaterials, bool bCanTakeOwnership = false) { Initialize(InMesh, bCanTakeOwnership, InMaterials); }
 	
 	// ~Begin UPCGData interface
-	virtual EPCGDataType GetDataType() const override { return EPCGDataType::DynamicMesh; }
+	PCG_ASSIGN_DEFAULT_TYPE_INFO(FPCGDataTypeInfoDynamicMesh)
 	PCG_API virtual void AddToCrc(FArchiveCrc32& Ar, bool bFullDataCrc) const override;
 	// ~End UPCGData interface
 
diff --git a/Engine/Plugins/PCG/Source/PCG/Public/Data/PCGGetDataFunctionRegistry.h b/Engine/Plugins/PCG/Source/PCG/Public/Data/PCGGetDataFunctionRegistry.h
--- a/Engine/Plugins/PCG/Source/PCG/Public/Data/PCGGetDataFunctionRegistry.h
+++ b/Engine/Plugins/PCG/Source/PCG/Public/Data/PCGGetDataFunctionRegistry.h
@@ -2,75 +2,8 @@
 
 #pragma once
 
-#include "PCGCommon.h"
-#include "PCGData.h"
+#include "CoreTypes.h"
 
-#define UE_API PCG_API
+UE_DEPRECATED_HEADER(5.7, "This header has moved to Data/Registry/PCGGetDataFunctionRegistry.h")
 
-class AActor;
-class UActorComponent;
-class UPCGComponent;
-struct FPCGContext;
-struct FPCGComponentSelectorSettings;
-
-/** Struct to hold selection criteria */
-struct FPCGGetDataFunctionRegistryParams
-{
-	const UPCGComponent* SourceComponent = nullptr;
-	const FPCGComponentSelectorSettings* ComponentSelector = nullptr;
-	EPCGDataType DataTypeFilter = EPCGDataType::Any;
-	bool bParseActor = true;
-	bool bIgnorePCGGeneratedComponents = true;
-	bool bAddActorTags = true;
-};
-
-struct FPCGGetDataFunctionRegistryOutput
-{
-	FPCGDataCollection Collection;
-	bool bSanitizedTagAttributeNames = false;
-};
-
-/** Registry to hold actor & component to PCG data construction mapping */
-struct FPCGGetDataFunctionRegistry
-{
-public:
-	using FDataFromActorFunction = TFunction<bool(FPCGContext*, const FPCGGetDataFunctionRegistryParams&, AActor*, FPCGGetDataFunctionRegistryOutput&)>;
-	using FDataFromComponentFunction = TFunction<bool(FPCGContext*, const FPCGGetDataFunctionRegistryParams&, UActorComponent*, FPCGGetDataFunctionRegistryOutput&)>;
-	using FFunctionHandle = uint64;
-
-	FPCGGetDataFunctionRegistry() = default;
-	FPCGGetDataFunctionRegistry(const FPCGGetDataFunctionRegistry&) = delete;
-	FPCGGetDataFunctionRegistry(FPCGGetDataFunctionRegistry&&) = default;
-	FPCGGetDataFunctionRegistry& operator=(const FPCGGetDataFunctionRegistry&) = delete;
-	FPCGGetDataFunctionRegistry& operator=(FPCGGetDataFunctionRegistry&&) = default;
-	virtual ~FPCGGetDataFunctionRegistry() = default;
-
-	/** Register an actor to PCG data function. */
-	UE_API FFunctionHandle RegisterDataFromActorFunction(const FDataFromActorFunction& InFunction);
-
-	/** Unregisters an actor to PCG data function. */
-	UE_API void UnregisterDataFromActorFunction(FFunctionHandle InFunctionHandle);
-
-	/** Appends actor data to the given collection and returns the number of data added. */
-	UE_API int GetDataFromActor(FPCGContext* InContext, const FPCGGetDataFunctionRegistryParams& InParams, AActor* InActor, FPCGGetDataFunctionRegistryOutput& OutDataCollection) const;
-
-	/** Register a component to PCG data function. */
-	UE_API FFunctionHandle RegisterDataFromComponentFunction(const FDataFromComponentFunction& InFunction);
-
-	/** Unregisters a component to PCG data function. */
-	UE_API void UnregisterDataFromComponentFunction(FFunctionHandle InFunctionHandle);
-
-	/** Append component data to the given collection, returns the number of data added. */
-	UE_API int GetDataFromComponent(FPCGContext* InContext, const FPCGGetDataFunctionRegistryParams& InParams, UActorComponent* InComponent, FPCGGetDataFunctionRegistryOutput& OutDataCollection) const;
-
-private:
-	UE_API int DefaultDataFromActor(FPCGContext* InContext, const FPCGGetDataFunctionRegistryParams& InParams, AActor* InActor, FPCGGetDataFunctionRegistryOutput& OutDataCollection) const;
-	UE_API int DefaultDataFromComponent(FPCGContext* InContext, const FPCGGetDataFunctionRegistryParams& InParams, UActorComponent* InComponent, FPCGGetDataFunctionRegistryOutput& OutDataCollection) const;
-
-	FFunctionHandle NextFunctionHandle = 0;
-
-	TArray<TPair<FDataFromActorFunction, FFunctionHandle>> ActorParsingFunctions;
-	TArray<TPair<FDataFromComponentFunction, FFunctionHandle>> ComponentParsingFunctions;
-};
-
-#undef UE_API
+#include "Data/Registry/PCGGetDataFunctionRegistry.h"
\ No newline at end of file
diff --git a/Engine/Plugins/PCG/Source/PCG/Public/Data/PCGIntersectionData.h b/Engine/Plugins/PCG/Source/PCG/Public/Data/PCGIntersectionData.h
--- a/Engine/Plugins/PCG/Source/PCG/Public/Data/PCGIntersectionData.h
+++ b/Engine/Plugins/PCG/Source/PCG/Public/Data/PCGIntersectionData.h
@@ -33,7 +33,7 @@ public:
 	// ~End UObject interface
 
 	// ~Begin UPCGData interface
-	virtual EPCGDataType GetDataType() const override { return EPCGDataType::Spatial; }
+	PCG_ASSIGN_DEFAULT_TYPE_INFO(FPCGDataTypeInfoComposite)
 	UE_API virtual void VisitDataNetwork(TFunctionRef<void(const UPCGData*)> Action) const override;
 
 protected:
diff --git a/Engine/Plugins/PCG/Source/PCG/Public/Data/PCGLandscapeData.h b/Engine/Plugins/PCG/Source/PCG/Public/Data/PCGLandscapeData.h
--- a/Engine/Plugins/PCG/Source/PCG/Public/Data/PCGLandscapeData.h
+++ b/Engine/Plugins/PCG/Source/PCG/Public/Data/PCGLandscapeData.h
@@ -52,6 +52,15 @@ struct FPCGLandscapeDataProps
 	bool bSampleVirtualTextureNormals = false;
 };
 
+
+USTRUCT()
+struct FPCGDataTypeInfoLandscape : public FPCGDataTypeInfoSurface
+{
+	GENERATED_BODY()
+
+	PCG_DECLARE_LEGACY_TYPE_INFO(EPCGDataType::Landscape)
+};
+
 /**
 * Landscape data access abstraction for PCG. Supports multi-landscape access, but it assumes that they are not overlapping.
 */
@@ -68,7 +77,7 @@ public:
 	// ~End UObject interface
 
 	// ~Begin UPCGData interface
-	virtual EPCGDataType GetDataType() const override { return EPCGDataType::Landscape; }
+	PCG_ASSIGN_DEFAULT_TYPE_INFO(FPCGDataTypeInfoLandscape)
 	UE_API virtual void AddToCrc(FArchiveCrc32& Ar, bool bFullDataCrc) const override;
 	// ~End UPCGData interface
 
diff --git a/Engine/Plugins/PCG/Source/PCG/Public/Data/PCGLandscapeSplineData.h b/Engine/Plugins/PCG/Source/PCG/Public/Data/PCGLandscapeSplineData.h
--- a/Engine/Plugins/PCG/Source/PCG/Public/Data/PCGLandscapeSplineData.h
+++ b/Engine/Plugins/PCG/Source/PCG/Public/Data/PCGLandscapeSplineData.h
@@ -12,6 +12,14 @@ class UPCGSpatialData;
 
 class ULandscapeSplinesComponent;
 
+USTRUCT()
+struct FPCGDataTypeInfoLandscapeSpline : public FPCGDataTypeInfoPolyline
+{
+	GENERATED_BODY()
+
+	PCG_DECLARE_LEGACY_TYPE_INFO(EPCGDataType::LandscapeSpline)
+};
+
 UCLASS(MinimalAPI, BlueprintType, ClassGroup=(Procedural))
 class UPCGLandscapeSplineData : public UPCGPolyLineData
 {
@@ -24,7 +32,7 @@ public:
 	// ~End UObject interface
 
 	// ~Begin UPCGData interface
-	virtual EPCGDataType GetDataType() const override { return EPCGDataType::LandscapeSpline; }
+	PCG_ASSIGN_DEFAULT_TYPE_INFO(FPCGDataTypeInfoLandscapeSpline)
 	UE_API virtual void AddToCrc(FArchiveCrc32& Ar, bool bFullDataCrc) const override;
 	// ~End UPCGData interface
 
diff --git a/Engine/Plugins/PCG/Source/PCG/Public/Data/PCGPolyLineData.h b/Engine/Plugins/PCG/Source/PCG/Public/Data/PCGPolyLineData.h
--- a/Engine/Plugins/PCG/Source/PCG/Public/Data/PCGPolyLineData.h
+++ b/Engine/Plugins/PCG/Source/PCG/Public/Data/PCGPolyLineData.h
@@ -8,6 +8,14 @@
 
 #define UE_API PCG_API
 
+USTRUCT()
+struct FPCGDataTypeInfoPolyline : public FPCGDataTypeInfoConcrete
+{
+	GENERATED_BODY()
+
+	PCG_DECLARE_LEGACY_TYPE_INFO(EPCGDataType::PolyLine)
+};
+
 UCLASS(MinimalAPI, Abstract, BlueprintType, ClassGroup = (Procedural))
 class UPCGPolyLineData : public UPCGSpatialDataWithPointCache
 {
@@ -15,7 +23,7 @@ class UPCGPolyLineData : public UPCGSpatialDataWithPointCache
 
 public:
 	// ~Begin UPCGData interface
-	virtual EPCGDataType GetDataType() const override { return EPCGDataType::PolyLine; }
+	PCG_ASSIGN_DEFAULT_TYPE_INFO(FPCGDataTypeInfoPolyline)
 	// ~End UPCGData interface
 
 	//~Begin UPCGSpatialData interface
diff --git a/Engine/Plugins/PCG/Source/PCG/Public/Data/PCGPolygon2DData.h b/Engine/Plugins/PCG/Source/PCG/Public/Data/PCGPolygon2DData.h
--- a/Engine/Plugins/PCG/Source/PCG/Public/Data/PCGPolygon2DData.h
+++ b/Engine/Plugins/PCG/Source/PCG/Public/Data/PCGPolygon2DData.h
@@ -43,6 +43,17 @@ namespace PCGPolygon2DData
 	const FName VertexDomainName = "Vertices";
 }
 
+USTRUCT()
+struct FPCGDataTypeInfoPolygon2D : public FPCGDataTypeInfoPolyline
+{
+	GENERATED_BODY()
+
+	PCG_DECLARE_LEGACY_TYPE_INFO(EPCGDataType::Polygon2D);
+
+	PCG_API virtual bool SupportsConversion(const FPCGDataTypeIdentifier& Other) const override;
+	PCG_API virtual TSubclassOf<UPCGSettings> GetDefaultConversionSettings(const FPCGDataTypeIdentifier& Other) const override;
+};
+
 /** Data representing a single 2D polygon with a 3D transform (for spatial operations). */
 UCLASS(MinimalAPI, BlueprintType, ClassGroup = (Procedural))
 class UPCGPolygon2DData : public UPCGPolyLineData
@@ -57,7 +68,7 @@ public:
 	// ~End UObject interface
 
 	// ~Begin UPCGData interface
-	virtual EPCGDataType GetDataType() const override { return EPCGDataType::Polygon2D; }
+	PCG_ASSIGN_DEFAULT_TYPE_INFO(FPCGDataTypeInfoPolygon2D)
 	PCG_API virtual void AddToCrc(FArchiveCrc32& Ar, bool bFullDataCrc) const override;
 
 	virtual FPCGMetadataDomainID GetDefaultMetadataDomainID() const override { return PCGMetadataDomainID::Elements; }
diff --git a/Engine/Plugins/PCG/Source/PCG/Public/Data/PCGPrimitiveData.h b/Engine/Plugins/PCG/Source/PCG/Public/Data/PCGPrimitiveData.h
--- a/Engine/Plugins/PCG/Source/PCG/Public/Data/PCGPrimitiveData.h
+++ b/Engine/Plugins/PCG/Source/PCG/Public/Data/PCGPrimitiveData.h
@@ -8,6 +8,14 @@
 
 #define UE_API PCG_API
 
+USTRUCT()
+struct FPCGDataTypeInfoPrimitive : public FPCGDataTypeInfoConcrete
+{
+	GENERATED_BODY()
+
+	PCG_DECLARE_LEGACY_TYPE_INFO(EPCGDataType::Primitive)
+};
+
 UCLASS(MinimalAPI, BlueprintType, ClassGroup = (Procedural))
 class UPCGPrimitiveData : public UPCGSpatialDataWithPointCache
 {
@@ -17,7 +25,7 @@ public:
 	UE_API void Initialize(UPrimitiveComponent* InPrim);
 
 	// ~Begin UPCGData interface
-	virtual EPCGDataType GetDataType() const override { return EPCGDataType::Primitive; }
+	PCG_ASSIGN_DEFAULT_TYPE_INFO(FPCGDataTypeInfoPrimitive)
 	UE_API virtual void AddToCrc(FArchiveCrc32& Ar, bool bFullDataCrc) const override;
 	// ~End UPCGData interface
 
diff --git a/Engine/Plugins/PCG/Source/PCG/Public/Data/PCGRenderTargetData.h b/Engine/Plugins/PCG/Source/PCG/Public/Data/PCGRenderTargetData.h
--- a/Engine/Plugins/PCG/Source/PCG/Public/Data/PCGRenderTargetData.h
+++ b/Engine/Plugins/PCG/Source/PCG/Public/Data/PCGRenderTargetData.h
@@ -11,6 +11,15 @@
 class UTexture;
 class UTextureRenderTarget2D;
 
+
+USTRUCT()
+struct FPCGDataTypeInfoRenderTarget : public FPCGDataTypeInfoBaseTexture
+{
+	GENERATED_BODY()
+
+	PCG_DECLARE_LEGACY_TYPE_INFO(EPCGDataType::RenderTarget)
+};
+
 //TODO: It's possible that caching the result in this class is not as efficient as it could be
 // if we expect to sample in different ways (e.g. channel) in the same render target
 UCLASS(MinimalAPI, BlueprintType, ClassGroup = (Procedural))
@@ -20,7 +29,7 @@ class UPCGRenderTargetData : public UPCGBaseTextureData
 
 public:
 	// ~Begin UPCGData interface
-	virtual EPCGDataType GetDataType() const override { return EPCGDataType::RenderTarget; }
+	PCG_ASSIGN_DEFAULT_TYPE_INFO(FPCGDataTypeInfoRenderTarget);
 	virtual void AddToCrc(FArchiveCrc32& Ar, bool bFullDataCrc) const override;
 	virtual bool HoldsTransientResources() const override { return bOwnsRenderTarget; }
 	virtual bool IsCacheable() const { return Super::IsCacheable() && !bOwnsRenderTarget; }
diff --git a/Engine/Plugins/PCG/Source/PCG/Public/Data/PCGResourceData.h b/Engine/Plugins/PCG/Source/PCG/Public/Data/PCGResourceData.h
--- a/Engine/Plugins/PCG/Source/PCG/Public/Data/PCGResourceData.h
+++ b/Engine/Plugins/PCG/Source/PCG/Public/Data/PCGResourceData.h
@@ -10,6 +10,14 @@
 
 #define UE_API PCG_API
 
+USTRUCT()
+struct FPCGDataTypeInfoResource : public FPCGDataTypeInfo
+{
+	GENERATED_BODY()
+
+	PCG_DECLARE_LEGACY_TYPE_INFO(EPCGDataType::Resource)
+};
+
 /** Data that wrap/represent an asset, like a Static Mesh or Texture. */
 UCLASS(MinimalAPI, Abstract, BlueprintType, ClassGroup = (Procedural))
 class UPCGResourceData : public UPCGData
@@ -18,7 +26,7 @@ class UPCGResourceData : public UPCGData
 
 public:
 	//~ Begin UPCGData interface
-	virtual EPCGDataType GetDataType() const override { return EPCGDataType::Resource; }
+	PCG_ASSIGN_DEFAULT_TYPE_INFO(FPCGDataTypeInfoResource)
 	//~ End UPCGData interface
 
 	virtual FSoftObjectPath GetResourcePath() const PURE_VIRTUAL(UPCGResourceData::GetResourcePath, return FSoftObjectPath(););
diff --git a/Engine/Plugins/PCG/Source/PCG/Public/Data/PCGSpatialData.h b/Engine/Plugins/PCG/Source/PCG/Public/Data/PCGSpatialData.h
--- a/Engine/Plugins/PCG/Source/PCG/Public/Data/PCGSpatialData.h
+++ b/Engine/Plugins/PCG/Source/PCG/Public/Data/PCGSpatialData.h
@@ -61,6 +61,34 @@ struct FPCGInitializeFromDataParams
 	FPCGMetadataInitializeParams MetadataInitializeParams;
 };
 
+USTRUCT()
+struct FPCGDataTypeInfoSpatial : public FPCGDataTypeInfo
+{
+	GENERATED_BODY()
+
+	PCG_DECLARE_LEGACY_TYPE_INFO(EPCGDataType::Spatial)
+
+	PCG_API virtual bool SupportsConversion(const FPCGDataTypeIdentifier& Other) const override;
+	PCG_API virtual TSubclassOf<UPCGSettings> GetDefaultConversionSettings(const FPCGDataTypeIdentifier& Other) const override;
+};
+
+// Subclasses of UPCGSpatial
+USTRUCT()
+struct FPCGDataTypeInfoConcrete : public FPCGDataTypeInfoSpatial
+{
+	GENERATED_BODY()
+
+	PCG_DECLARE_LEGACY_TYPE_INFO(EPCGDataType::Concrete)
+};
+
+USTRUCT()
+struct FPCGDataTypeInfoComposite : public FPCGDataTypeInfoSpatial
+{
+	GENERATED_BODY()
+
+	PCG_DECLARE_LEGACY_TYPE_INFO(EPCGDataType::Composite)
+};
+
 /**
 * "Concrete" data base class for PCG generation
 * This will be the base class for data classes that actually represent
@@ -83,7 +111,8 @@ public:
 	//~End UObject Interface
 
 	// ~Begin UPCGData interface
-	virtual EPCGDataType GetDataType() const override { return EPCGDataType::Spatial; }
+	PCG_ASSIGN_DEFAULT_TYPE_INFO(FPCGDataTypeInfoSpatial)
+
 	UE_API virtual void AddToCrc(FArchiveCrc32& Ar, bool bFullDataCrc) const override;
 
 	UE_API virtual bool HasCachedLastSelector() const override;
diff --git a/Engine/Plugins/PCG/Source/PCG/Public/Data/PCGSplineData.h b/Engine/Plugins/PCG/Source/PCG/Public/Data/PCGSplineData.h
--- a/Engine/Plugins/PCG/Source/PCG/Public/Data/PCGSplineData.h
+++ b/Engine/Plugins/PCG/Source/PCG/Public/Data/PCGSplineData.h
@@ -22,6 +22,17 @@ namespace PCGSplineData
 	const FName ControlPointDomainName = "ControlPoints";
 }
 
+USTRUCT()
+struct FPCGDataTypeInfoSpline : public FPCGDataTypeInfoPolyline
+{
+	GENERATED_BODY()
+
+	PCG_DECLARE_LEGACY_TYPE_INFO(EPCGDataType::Spline)
+
+	PCG_API virtual bool SupportsConversion(const FPCGDataTypeIdentifier& Other) const override;
+	PCG_API virtual TSubclassOf<UPCGSettings> GetDefaultConversionSettings(const FPCGDataTypeIdentifier& Other) const override;
+};
+
 UCLASS(MinimalAPI, BlueprintType, ClassGroup = (Procedural))
 class UPCGSplineData : public UPCGPolyLineData
 {
@@ -59,7 +70,7 @@ public:
 	UE_API void ApplyTo(USplineComponent* InSpline) const;
 
 	// ~Begin UPCGData interface
-	virtual EPCGDataType GetDataType() const override { return EPCGDataType::Spline; }
+	PCG_ASSIGN_DEFAULT_TYPE_INFO(FPCGDataTypeInfoSpline)
 	UE_API virtual void AddToCrc(FArchiveCrc32& Ar, bool bFullDataCrc) const override;
 
 	// To be enabled when we are sure Default translate well from Data to Elements (or we have a path for deprecation)
diff --git a/Engine/Plugins/PCG/Source/PCG/Public/Data/PCGStaticMeshResourceData.h b/Engine/Plugins/PCG/Source/PCG/Public/Data/PCGStaticMeshResourceData.h
--- a/Engine/Plugins/PCG/Source/PCG/Public/Data/PCGStaticMeshResourceData.h
+++ b/Engine/Plugins/PCG/Source/PCG/Public/Data/PCGStaticMeshResourceData.h
@@ -10,6 +10,14 @@
 
 class UStaticMesh;
 
+USTRUCT()
+struct FPCGDataTypeInfoStaticMeshResource : public FPCGDataTypeInfoResource
+{
+	GENERATED_BODY()
+
+	PCG_DECLARE_LEGACY_TYPE_INFO(EPCGDataType::StaticMeshResource)
+};
+
 /** Data that wraps a Static Mesh soft object path. */
 UCLASS(MinimalAPI, BlueprintType, ClassGroup = (Procedural))
 class UPCGStaticMeshResourceData : public UPCGResourceData
@@ -20,7 +28,7 @@ public:
 	UE_API void Initialize(TSoftObjectPtr<UStaticMesh> InStaticMesh);
 
 	//~ Begin UPCGData interface
-	virtual EPCGDataType GetDataType() const override { return EPCGDataType::StaticMeshResource; }
+	PCG_ASSIGN_DEFAULT_TYPE_INFO(FPCGDataTypeInfoStaticMeshResource)
 	//~ End UPCGData interface
 
 	//~ Begin UPCGResourceData interface
diff --git a/Engine/Plugins/PCG/Source/PCG/Public/Data/PCGSurfaceData.h b/Engine/Plugins/PCG/Source/PCG/Public/Data/PCGSurfaceData.h
--- a/Engine/Plugins/PCG/Source/PCG/Public/Data/PCGSurfaceData.h
+++ b/Engine/Plugins/PCG/Source/PCG/Public/Data/PCGSurfaceData.h
@@ -8,6 +8,14 @@
 
 #define UE_API PCG_API
 
+USTRUCT()
+struct FPCGDataTypeInfoSurface : public FPCGDataTypeInfoConcrete
+{
+	GENERATED_BODY()
+
+	PCG_DECLARE_LEGACY_TYPE_INFO(EPCGDataType::Surface)
+};
+
 UCLASS(MinimalAPI, Abstract, BlueprintType, ClassGroup = (Procedural))
 class UPCGSurfaceData : public UPCGSpatialDataWithPointCache
 {
@@ -15,7 +23,7 @@ class UPCGSurfaceData : public UPCGSpatialDataWithPointCache
 
 public:
 	// ~Begin UPCGData interface
-	virtual EPCGDataType GetDataType() const override { return EPCGDataType::Surface; }
+	PCG_ASSIGN_DEFAULT_TYPE_INFO(FPCGDataTypeInfoSurface)
 	// ~End UPCGData interface
 
 	//~Begin UPCGSpatialData interface
diff --git a/Engine/Plugins/PCG/Source/PCG/Public/Data/PCGTextureData.h b/Engine/Plugins/PCG/Source/PCG/Public/Data/PCGTextureData.h
--- a/Engine/Plugins/PCG/Source/PCG/Public/Data/PCGTextureData.h
+++ b/Engine/Plugins/PCG/Source/PCG/Public/Data/PCGTextureData.h
@@ -63,6 +63,22 @@ namespace PCGTextureSamplingHelpers
 	TOptional<bool> CanGPUTextureBeCPUAccessed(UTexture2D* Texture);
 }
 
+USTRUCT()
+struct FPCGDataTypeInfoBaseTexture : public FPCGDataTypeInfoSurface
+{
+	GENERATED_BODY()
+
+	PCG_DECLARE_LEGACY_TYPE_INFO(EPCGDataType::BaseTexture)
+};
+
+USTRUCT()
+struct FPCGDataTypeInfoTexture : public FPCGDataTypeInfoBaseTexture
+{
+	GENERATED_BODY()
+
+	PCG_DECLARE_LEGACY_TYPE_INFO(EPCGDataType::Texture)
+};
+
 UCLASS(MinimalAPI, Abstract)
 class UPCGBaseTextureData : public UPCGSurfaceData
 {
@@ -74,7 +90,7 @@ public:
 	// ~End UObject interface
 
 	// ~Begin UPCGData interface
-	virtual EPCGDataType GetDataType() const override { return EPCGDataType::BaseTexture; }
+	PCG_ASSIGN_DEFAULT_TYPE_INFO(FPCGDataTypeInfoBaseTexture)
 	// ~End UPCGData interface
 
 	//~Begin UPCGSpatialData interface
@@ -195,7 +211,7 @@ public:
 	EPCGTextureResourceType GetTextureResourceType() const override { return ResourceType; }
 
 	//~Begin UPCGData interface
-	virtual EPCGDataType GetDataType() const override { return EPCGDataType::Texture; }
+	PCG_ASSIGN_DEFAULT_TYPE_INFO(FPCGDataTypeInfoTexture)
 	virtual void AddToCrc(FArchiveCrc32& Ar, bool bFullDataCrc) const override;
 	// Don't hold onto exported buffers currently as graphics memory usage (and lifetimes) may cause issues.
 	virtual bool HoldsTransientResources() const override { return ResourceType == EPCGTextureResourceType::ExportedTexture; }
diff --git a/Engine/Plugins/PCG/Source/PCG/Public/Data/PCGUnionData.h b/Engine/Plugins/PCG/Source/PCG/Public/Data/PCGUnionData.h
--- a/Engine/Plugins/PCG/Source/PCG/Public/Data/PCGUnionData.h
+++ b/Engine/Plugins/PCG/Source/PCG/Public/Data/PCGUnionData.h
@@ -39,7 +39,7 @@ public:
 	void SetDensityFunction(EPCGUnionDensityFunction InDensityFunction) { DensityFunction = InDensityFunction; }
 
 	// ~Begin UPCGData interface
-	virtual EPCGDataType GetDataType() const override { return EPCGDataType::Spatial; }
+	PCG_ASSIGN_DEFAULT_TYPE_INFO(FPCGDataTypeInfoComposite)
 	UE_API virtual void VisitDataNetwork(TFunctionRef<void(const UPCGData*)> Action) const override;
 
 protected:
diff --git a/Engine/Plugins/PCG/Source/PCG/Public/Data/PCGVirtualTextureData.h b/Engine/Plugins/PCG/Source/PCG/Public/Data/PCGVirtualTextureData.h
--- a/Engine/Plugins/PCG/Source/PCG/Public/Data/PCGVirtualTextureData.h
+++ b/Engine/Plugins/PCG/Source/PCG/Public/Data/PCGVirtualTextureData.h
@@ -9,6 +9,14 @@
 class URuntimeVirtualTexture;
 class URuntimeVirtualTextureComponent;
 
+USTRUCT()
+struct FPCGDataTypeInfoVirtualTexture : public FPCGDataTypeInfoSurface
+{
+	GENERATED_BODY()
+
+	PCG_DECLARE_LEGACY_TYPE_INFO(EPCGDataType::VirtualTexture)
+};
+
 UCLASS()
 class UPCGVirtualTextureData : public UPCGSurfaceData
 {
@@ -18,7 +26,7 @@ public:
 	void Initialize(const URuntimeVirtualTextureComponent* InVirtualTextureComponent);
 
 	//~Begin UPCGData interface
-	virtual EPCGDataType GetDataType() const override { return EPCGDataType::VirtualTexture; }
+	PCG_ASSIGN_DEFAULT_TYPE_INFO(FPCGDataTypeInfoVirtualTexture)
 	//~End UPCGData interface
 
 	//~Begin UPCGSpatialData interface
diff --git a/Engine/Plugins/PCG/Source/PCG/Public/Data/PCGVolumeData.h b/Engine/Plugins/PCG/Source/PCG/Public/Data/PCGVolumeData.h
--- a/Engine/Plugins/PCG/Source/PCG/Public/Data/PCGVolumeData.h
+++ b/Engine/Plugins/PCG/Source/PCG/Public/Data/PCGVolumeData.h
@@ -11,6 +11,14 @@
 class AVolume;
 struct FBodyInstance;
 
+USTRUCT()
+struct FPCGDataTypeInfoVolume : public FPCGDataTypeInfoConcrete
+{
+	GENERATED_BODY()
+
+	PCG_DECLARE_LEGACY_TYPE_INFO(EPCGDataType::Volume)
+};
+
 UCLASS(MinimalAPI, BlueprintType, ClassGroup = (Procedural))
 class UPCGVolumeData : public UPCGSpatialDataWithPointCache
 {
@@ -24,7 +32,7 @@ public:
 	UE_API void Initialize(const FBox& InBounds);
 
 	// ~Begin UPCGData interface
-	virtual EPCGDataType GetDataType() const override { return EPCGDataType::Volume; }
+	PCG_ASSIGN_DEFAULT_TYPE_INFO(FPCGDataTypeInfoVolume)
 	UE_API virtual void AddToCrc(FArchiveCrc32& Ar, bool bFullDataCrc) const override;
 	// ~End UPCGData interface
 
diff --git a/Engine/Plugins/PCG/Source/PCG/Public/Data/Registry/PCGDataType.h b/Engine/Plugins/PCG/Source/PCG/Public/Data/Registry/PCGDataType.h
new file mode 100644
--- /dev/null
+++ b/Engine/Plugins/PCG/Source/PCG/Public/Data/Registry/PCGDataType.h
@@ -0,0 +1,132 @@
+// Copyright Epic Games, Inc. All Rights Reserved.
+
+#pragma once
+
+#include "PCGCommon.h"
+
+#include "PCGDataType.generated.h"
+
+struct FPCGDataTypeIdentifier;
+class UPCGData;
+class UPCGGraph;
+class UPCGNode;
+class UPCGPin;
+class UPCGSettings;
+
+/**
+ * To declare a new type:
+ * - Make a new UClass inheriting from UPCGData (or any of its children)
+ * - Make a new UStruct inheriting from FPCGDataTypeInfo, and use the PCG_DECLARE_TYPE_INFO(DLL_EXPORT_NAME) in the header, replacing the DLL_EXPORT_NAME by your lib dll export (ie. PCG_API). Leave it empty if you don't want to export it.
+ * - In the cpp, use the PCG_DEFINE_TYPE_INFO(PCGDataTypeInfoClass, PCGDataClass), replacing PCGDataClass by your new UClass and PCGDataTypeInfoClass by your new UStruct
+ * - In the UClass, use the PCG_ASSIGN_TYPE_INFO(PCGDataTypeInfoClass) replacing PCGDataTypeInfoClass by your new UStruct
+ * - Implement the virtual classes in the type info for any support for different colors/icons/conversions etc...
+ *
+ * If the new type doesn't have an associated PCGData class, use PCG_DEFINE_TYPE_INFO_WITHOUT_CLASS
+ */
+#define PCG_DECLARE_TYPE_INFO(DLL_EXPORT_NAME) \
+	DLL_EXPORT_NAME static TSubclassOf<UPCGData> GetStaticAssociatedClass(); \
+	DLL_EXPORT_NAME virtual TSubclassOf<UPCGData> GetAssociatedClass() const override; \
+	static FPCGDataTypeBaseId AsID() { return FPCGDataTypeBaseId(StaticStruct()); }
+
+#define PCG_DEFINE_TYPE_INFO(PCGDataTypeInfoClass, PCGDataClass) \
+	TSubclassOf<UPCGData> PCGDataTypeInfoClass::GetStaticAssociatedClass() { return PCGDataClass::StaticClass(); } \
+	TSubclassOf<UPCGData> PCGDataTypeInfoClass::GetAssociatedClass() const { return PCGDataClass::StaticClass(); };
+
+#define PCG_DEFINE_TYPE_INFO_WITHOUT_CLASS(PCGDataTypeInfoClass) \
+	TSubclassOf<UPCGData> PCGDataTypeInfoClass::GetStaticAssociatedClass() { return nullptr; } \
+	TSubclassOf<UPCGData> PCGDataTypeInfoClass::GetAssociatedClass() const { return nullptr; };
+
+#define PCG_ASSIGN_TYPE_INFO(PCGDataTypeInfoClass) \
+	using TypeInfo = PCGDataTypeInfoClass; \
+	virtual FPCGDataTypeBaseId GetDataTypeId() const override { return FPCGDataTypeBaseId::Construct<PCGDataTypeInfoClass>(); }
+
+// Convenient macro for default types in PCG (that are defined in EPCGDataType)
+#define PCG_DECLARE_LEGACY_TYPE_INFO(LegacyDataType) \
+	PCG_DECLARE_TYPE_INFO(PCG_API) \
+	virtual EPCGDataType GetAssociatedLegacyType() const override{ return LegacyDataType; } \
+	static EPCGDataType GetStaticAssociatedLegacyType() { return LegacyDataType; }
+
+#define PCG_ASSIGN_DEFAULT_TYPE_INFO(PCGDataTypeInfoClass) \
+	PCG_ASSIGN_TYPE_INFO(PCGDataTypeInfoClass); \
+	PRAGMA_DISABLE_DEPRECATION_WARNINGS \
+	virtual EPCGDataType GetDataType() const { return PCGDataTypeInfoClass::GetStaticAssociatedLegacyType(); } \
+	PRAGMA_ENABLE_DEPRECATION_WARNINGS
+
+/**
+ * Store the information about a type.
+ * Used by FPCGDataTypeBaseId as the UScriptStruct (the meta-struct) and will have a "default object" initialized in the registry for the virtual calls.
+ */
+USTRUCT()
+struct FPCGDataTypeInfo
+{
+	GENERATED_BODY()
+
+	virtual ~FPCGDataTypeInfo() = default;
+
+	PCG_API static TSubclassOf<UPCGData> GetStaticAssociatedClass();
+	PCG_API virtual TSubclassOf<UPCGData> GetAssociatedClass() const;
+
+	virtual EPCGDataType GetAssociatedLegacyType() const { return EPCGDataType::Any; }
+	static EPCGDataType GetStaticAssociatedLegacyType() { return EPCGDataType::Any; }
+	
+	PCG_API virtual bool SupportsConversion(const FPCGDataTypeIdentifier& Other) const;
+	PCG_API virtual UPCGNode* AddConversionNode(const FPCGDataTypeIdentifier& Other, UPCGGraph* InGraph, UPCGPin* InUpstreamPin, UPCGPin* InDownstreamPin) const;
+	PCG_API virtual TSubclassOf<UPCGSettings> GetDefaultConversionSettings(const FPCGDataTypeIdentifier& Other) const;
+};
+
+/**
+ * Wrapper around an object ptr of a UScriptStruct, that must be a child of FPCGDataTypeInfo
+ * Represent a single type, and is aggregated into a FPCGDataTypeIdentifier.
+ * Have a mapping from and to the EPCGDataType, for all legacy data types.
+ */
+USTRUCT(BlueprintType)
+struct FPCGDataTypeBaseId
+{
+	GENERATED_BODY()
+
+	PCG_API static FPCGDataTypeBaseId MakeFromLegacyType(EPCGDataType LegacyType);
+	PCG_API EPCGDataType AsLegacyType() const;
+
+	PCG_API FPCGDataTypeBaseId(const UScriptStruct* InStruct = nullptr);
+
+	template <typename T> requires std::is_base_of_v<FPCGDataTypeInfo, T>
+	static FPCGDataTypeBaseId Construct()
+	{
+		return FPCGDataTypeBaseId(T::StaticStruct());
+	}
+
+	friend int32 GetTypeHash(const FPCGDataTypeBaseId& This)
+	{
+		return GetTypeHash(This.Struct);
+	}
+
+	bool operator==(const FPCGDataTypeBaseId& Other) const
+	{
+		return Struct == Other.Struct;
+	}
+
+	PCG_API bool IsValid() const;
+	PCG_API bool IsChildOf(const FPCGDataTypeBaseId& Other) const;
+
+	template <typename T> requires std::is_base_of_v<FPCGDataTypeInfo, T>
+	bool IsChildOf()
+	{
+		return IsChildOf(FPCGDataTypeBaseId(T::StaticStruct()));
+	}
+
+	PCG_API FString ToString() const;
+
+	const UScriptStruct* GetStruct() const { return Struct; }
+
+private:
+	UPROPERTY()
+	TObjectPtr<const UScriptStruct> Struct = nullptr;
+};
+
+USTRUCT()
+struct FPCGDataTypeInfoOther: public FPCGDataTypeInfo
+{
+	GENERATED_BODY()
+
+	PCG_DECLARE_LEGACY_TYPE_INFO(EPCGDataType::Other);
+};
diff --git a/Engine/Plugins/PCG/Source/PCG/Public/Data/Registry/PCGDataTypeIdentifier.h b/Engine/Plugins/PCG/Source/PCG/Public/Data/Registry/PCGDataTypeIdentifier.h
new file mode 100644
--- /dev/null
+++ b/Engine/Plugins/PCG/Source/PCG/Public/Data/Registry/PCGDataTypeIdentifier.h
@@ -0,0 +1,180 @@
+﻿// Copyright Epic Games, Inc. All Rights Reserved.
+
+#pragma once
+
+#include "PCGCommon.h"
+#include "PCGDataType.h"
+#include "Algo/AllOf.h"
+
+#include "Containers/Array.h"
+#include "Kismet/BlueprintFunctionLibrary.h"
+#include "Templates/SubclassOf.h"
+
+#include "PCGDataTypeIdentifier.generated.h"
+
+class FPCGDataTypeRegistry;
+class UPCGData;
+
+USTRUCT(BlueprintType)
+struct FPCGDataTypeIdentifier
+{
+	friend FPCGDataTypeRegistry;
+	
+	GENERATED_BODY()
+
+	FPCGDataTypeIdentifier() = default;
+
+	FPCGDataTypeIdentifier(FPCGDataTypeBaseId ID)
+	{
+		IDs.Emplace(ID);
+	}
+	
+	PCG_API explicit FPCGDataTypeIdentifier(TSubclassOf<UPCGData> Class);
+
+    /** Composition ctors. */
+    PCG_API explicit FPCGDataTypeIdentifier(TConstArrayView<FPCGDataTypeIdentifier> InIDs);
+    PCG_API explicit FPCGDataTypeIdentifier(TConstArrayView<TSubclassOf<UPCGData>> Classes);
+
+	// ------- Begin deprecation with EPCGDataType --------
+	// Not explicit to facilitate deprecation
+	PCG_API FPCGDataTypeIdentifier(EPCGDataType LegacyType);
+	
+	// @todo_pcg: Implicit for now, should be explicit before next release with deprecation path
+	UE_DEPRECATED(5.7, "FPCGDataTypeIdentifier is replacing EPCGDataType")
+	PCG_API operator EPCGDataType() const;
+	
+	// @todo_pcg: Implicit for now, should be explicit before next release with deprecation path
+	UE_DEPRECATED(5.7, "FPCGDataTypeIdentifier is replacing EPCGDataType")
+	operator uint32() const
+	{
+		PRAGMA_DISABLE_DEPRECATION_WARNINGS
+		return static_cast<uint32>(static_cast<EPCGDataType>(*this));
+		PRAGMA_ENABLE_DEPRECATION_WARNINGS
+	}
+
+	FPCGDataTypeIdentifier& operator=(EPCGDataType LegacyType)
+	{
+		*this = FPCGDataTypeIdentifier(LegacyType);
+		return *this;
+	}
+
+	bool operator==(EPCGDataType OtherLegacyType) const
+	{
+		return IsSameType(FPCGDataTypeIdentifier(OtherLegacyType));
+	}
+	
+	PCG_API bool operator!() const;
+	
+	PCG_API FPCGDataTypeIdentifier& operator|=(EPCGDataType OtherLegacyType);
+	PCG_API FPCGDataTypeIdentifier operator|(EPCGDataType OtherLegacyType) const;
+	PCG_API FPCGDataTypeIdentifier& operator&=(EPCGDataType OtherLegacyType);
+	PCG_API FPCGDataTypeIdentifier operator&(EPCGDataType OtherLegacyType) const;
+	
+	PCG_API bool SupportsType(EPCGDataType OtherLegacyType) const;
+	// ------- End deprecation with EPCGDataType --------
+
+	// To mimic flags
+	bool operator==(const FPCGDataTypeIdentifier& Other) const
+	{
+		return IsIdentical(Other);
+	}
+
+	PCG_API FPCGDataTypeIdentifier& operator|=(const FPCGDataTypeIdentifier& Other);
+	PCG_API FPCGDataTypeIdentifier operator|(const FPCGDataTypeIdentifier& Other) const;
+	PCG_API FPCGDataTypeIdentifier& operator&=(const FPCGDataTypeIdentifier& Other);
+	PCG_API FPCGDataTypeIdentifier operator&(const FPCGDataTypeIdentifier& Other) const;
+
+	// Return the composition with other types, with reduction. (ie. if Data B and C are child of Data A, Composition of B and C gives A)
+	// CAREFUL: This can't be used before the Data Type Registry is done initializing (so can't be used before the Engine Post Init).
+	PCG_API FPCGDataTypeIdentifier Compose(const FPCGDataTypeIdentifier& Other);
+	PCG_API static FPCGDataTypeIdentifier Compose(TConstArrayView<FPCGDataTypeIdentifier> IDs);
+	
+	template <typename ...T> requires (std::is_base_of_v<UPCGData, T> && ...)
+	static FPCGDataTypeIdentifier Construct()
+	{
+		FPCGDataTypeIdentifier Result;
+		Result.IDs.Reserve(sizeof...(T));
+		(Result.IDs.AddUnique(FPCGDataTypeIdentifier(T::StaticClass()).GetID()), ...);
+		return Result;
+	}
+
+	/** Can't be called when the identifier is a composition. Need to use GetIDs in that case. */
+	FPCGDataTypeBaseId GetID() const { check(!IsComposition() && IsValid()); return IDs[0]; };
+	
+	PCG_API TConstArrayView<FPCGDataTypeBaseId> GetIDs() const;
+
+	/** A type identifier can be a composition of multiple types. */
+	bool IsComposition() const {return IDs.Num() > 1; }
+	
+	bool IsValid() const { return !IDs.IsEmpty() && Algo::AllOf(IDs, [](const FPCGDataTypeBaseId& ID) { return ID.IsValid(); }); }
+
+	/** Return true if the IDs are the same, but not necessary the subtype. */
+	PCG_API bool IsSameType(const FPCGDataTypeIdentifier& Other) const;
+
+	/** Equivalent to IsSameType if it is not a composition. Otherwise, return true if the OtherID is in the IDs. */
+	PCG_API bool SupportsType(const FPCGDataTypeIdentifier& Other) const;
+
+	/** Return true if this identifier supports different/more types than Other. Equivalent of !!(*this & ~Other) in bit flags form. */
+	PCG_API bool IsWider(const FPCGDataTypeIdentifier& Other) const;
+
+	/** Return true if the IDs and subtypes are the same. */
+	PCG_API bool IsIdentical(const FPCGDataTypeIdentifier& Other) const;
+
+	/** To support previously saved types, that used EPCGDataType, we need to define this function to de-serialize the new class using the old. And add a trait (see below). */
+	PCG_API bool SerializeFromMismatchedTag(const struct FPropertyTag& Tag, FStructuredArchive::FSlot Slot);
+
+	friend uint32 GetTypeHash(const FPCGDataTypeIdentifier& ID)
+	{
+		return HashCombine(GetTypeHash(ID.IDs), ID.CustomSubtype);
+	}
+
+	PCG_API FString ToString() const;
+	PCG_API FText ToDisplayText() const;
+	
+private:
+	UPROPERTY(BlueprintReadOnly, VisibleAnywhere, Category = "PCG", meta = (AllowPrivateAccess))
+	TArray<FPCGDataTypeBaseId> IDs;
+
+public:
+	UPROPERTY(BlueprintReadWrite, EditAnywhere, Category = "PCG")
+	int32 CustomSubtype = 0;
+};
+
+template<>
+struct TStructOpsTypeTraits<FPCGDataTypeIdentifier> : public TStructOpsTypeTraitsBase2<FPCGDataTypeIdentifier>
+{
+	enum
+	{
+		WithStructuredSerializeFromMismatchedTag = true
+	};
+};
+
+/**
+* Helper class to allow the BP to call the custom setters and getters on FPCGDataTypeIdentifier.
+*/
+UCLASS()
+class UPCGDataTypeIdentifierHelpers : public UBlueprintFunctionLibrary
+{
+	GENERATED_BODY()
+
+public:
+	UFUNCTION(BlueprintCallable, Category = "PCG|PCGDataTypeIdentifier")
+	static FPCGDataTypeIdentifier GetIdentifierFromClass(TSubclassOf<UPCGData> Class) { return FPCGDataTypeIdentifier(Class); };
+
+	UFUNCTION(BlueprintCallable, Category = "PCG|PCGDataTypeIdentifier")
+	static PCG_API FPCGDataTypeIdentifier GetIdentifierFromDefaultType(EPCGExclusiveDataType DataType);
+
+	UFUNCTION(BlueprintCallable, Category = "PCG|PCGDataTypeIdentifier", meta = (ScriptMethod))
+	static bool IsValid(UPARAM(Ref) const FPCGDataTypeIdentifier& This) { return This.IsValid(); }
+
+	UFUNCTION(BlueprintCallable, Category = "PCG|PCGDataTypeIdentifier", meta = (ScriptMethod))
+	static bool IsComposition(UPARAM(Ref) const FPCGDataTypeIdentifier& This) { return This.IsComposition(); }
+
+	/** Return true if the IDs are the same, but not necessary the subtype. */
+	UFUNCTION(BlueprintCallable, Category = "PCG|PCGDataTypeIdentifier", meta = (ScriptMethod))
+	static bool IsSameType(UPARAM(Ref) const FPCGDataTypeIdentifier& This, UPARAM(Ref) const FPCGDataTypeIdentifier& Other) { return This.IsSameType(Other); }
+
+	/** Return true if the IDs and subtypes are the same. */
+	UFUNCTION(BlueprintCallable, Category = "PCG|PCGDataTypeIdentifier", meta = (ScriptMethod))
+	static bool IsIdentical(UPARAM(Ref) const FPCGDataTypeIdentifier& This, UPARAM(Ref) const FPCGDataTypeIdentifier& Other) { return This.IsIdentical(Other); }
+};
\ No newline at end of file
diff --git a/Engine/Plugins/PCG/Source/PCG/Public/Data/Registry/PCGDataTypeRegistry.h b/Engine/Plugins/PCG/Source/PCG/Public/Data/Registry/PCGDataTypeRegistry.h
new file mode 100644
--- /dev/null
+++ b/Engine/Plugins/PCG/Source/PCG/Public/Data/Registry/PCGDataTypeRegistry.h
@@ -0,0 +1,88 @@
+﻿// Copyright Epic Games, Inc. All Rights Reserved.
+
+#pragma once
+
+#include "PCGCommon.h"
+
+#include "Data/Registry/PCGDataTypeIdentifier.h"
+#include "StructUtils/InstancedStruct.h"
+
+class FPCGDataTypeRegistry;
+// Friend for tests
+class FPCGDataTypeRegistryTestBase;
+class FPCGHierarchyTest;
+class FPCGModule;
+
+namespace PCGDataTypeHierarchy
+{
+	struct FTree;
+}
+
+enum class EPCGDataTypeCompatibilityResult
+{
+	// Types are not compatible, they can't be connected together.
+	NotCompatible = 0,
+
+	// Types are fully compatible, they can be connected as-is.
+	Compatible = 1,
+
+	// Types are compatible, but requires a filter to be added for the connection.
+	RequireFilter = 2,
+
+	// Types are compatible, but require conversion nodes to be added for the connection.
+	RequireConversion = 3,
+
+	// Error if the identifier is not found
+	UnknownType = 4
+};
+
+namespace PCGDataTypeCompatibilityResult
+{
+	inline bool IsValid(const EPCGDataTypeCompatibilityResult Result)
+	{
+		return Result != EPCGDataTypeCompatibilityResult::NotCompatible && Result != EPCGDataTypeCompatibilityResult::UnknownType;
+	}
+}
+
+/**
+ * Registry to hold PCG data types and the hierarchy.
+ */
+class FPCGDataTypeRegistry
+{
+public:
+	// Moveable but not copyable
+	FPCGDataTypeRegistry();
+	FPCGDataTypeRegistry(const FPCGDataTypeRegistry&) = delete;
+	FPCGDataTypeRegistry(FPCGDataTypeRegistry&&) = default;
+	FPCGDataTypeRegistry& operator=(const FPCGDataTypeRegistry&) = delete;
+	FPCGDataTypeRegistry& operator=(FPCGDataTypeRegistry&&) = default;
+	
+	~FPCGDataTypeRegistry();
+
+	/** Is it possible to connect the InType into a OutType. */
+	PCG_API EPCGDataTypeCompatibilityResult IsCompatible(const FPCGDataTypeIdentifier& InType, const FPCGDataTypeIdentifier& OutType) const;
+
+	/** Return the type info, to allow to call virtual functions. */
+	PCG_API const FPCGDataTypeInfo* GetTypeInfo(const FPCGDataTypeBaseId& ID) const;
+
+	/** Returns the data type for the union of all the identifiers passed as input. It is the common ancestor of all of them. */
+	PCG_API FPCGDataTypeIdentifier GetIdentifiersUnion(TConstArrayView<FPCGDataTypeIdentifier> Identifiers) const;
+
+	/** Returns the data type for the composition of all the identifiers passed as input. It will be simplified if possible (like composition of spline and landscape spline gives a polyline). */
+	PCG_API FPCGDataTypeIdentifier GetIdentifiersComposition(TConstArrayView<FPCGDataTypeIdentifier> Identifiers) const;
+	
+private:
+	friend FPCGDataTypeRegistryTestBase;
+	friend FPCGHierarchyTest;
+	friend FPCGModule;
+	
+	void RegisterKnownTypes();
+	void RegisterEntry(const FPCGDataTypeBaseId& BaseID);
+	
+	bool IsRegistered(const FPCGDataTypeBaseId& BaseID) const { return Mapping.Contains(BaseID); }
+	
+	TMap<FPCGDataTypeBaseId, TInstancedStruct<FPCGDataTypeInfo>> Mapping;
+	// Can't be a UniquePtr because of incomplete type.
+	PCGDataTypeHierarchy::FTree* HierarchyTree = nullptr;
+};
+
diff --git a/Engine/Plugins/PCG/Source/PCG/Public/Data/Registry/PCGGetDataFunctionRegistry.h b/Engine/Plugins/PCG/Source/PCG/Public/Data/Registry/PCGGetDataFunctionRegistry.h
new file mode 100644
--- /dev/null
+++ b/Engine/Plugins/PCG/Source/PCG/Public/Data/Registry/PCGGetDataFunctionRegistry.h
@@ -0,0 +1,76 @@
+// Copyright Epic Games, Inc. All Rights Reserved.
+
+#pragma once
+
+#include "PCGCommon.h"
+#include "PCGData.h"
+
+#define UE_API PCG_API
+
+class AActor;
+class UActorComponent;
+class UPCGComponent;
+struct FPCGContext;
+struct FPCGComponentSelectorSettings;
+
+/** Struct to hold selection criteria */
+struct FPCGGetDataFunctionRegistryParams
+{
+	const UPCGComponent* SourceComponent = nullptr;
+	const FPCGComponentSelectorSettings* ComponentSelector = nullptr;
+	EPCGDataType DataTypeFilter = EPCGDataType::Any;
+	bool bParseActor = true;
+	bool bIgnorePCGGeneratedComponents = true;
+	bool bAddActorTags = true;
+};
+
+struct FPCGGetDataFunctionRegistryOutput
+{
+	FPCGDataCollection Collection;
+	bool bSanitizedTagAttributeNames = false;
+};
+
+/** Registry to hold actor & component to PCG data construction mapping */
+struct FPCGGetDataFunctionRegistry
+{
+public:
+	using FDataFromActorFunction = TFunction<bool(FPCGContext*, const FPCGGetDataFunctionRegistryParams&, AActor*, FPCGGetDataFunctionRegistryOutput&)>;
+	using FDataFromComponentFunction = TFunction<bool(FPCGContext*, const FPCGGetDataFunctionRegistryParams&, UActorComponent*, FPCGGetDataFunctionRegistryOutput&)>;
+	using FFunctionHandle = uint64;
+
+	FPCGGetDataFunctionRegistry() = default;
+	FPCGGetDataFunctionRegistry(const FPCGGetDataFunctionRegistry&) = delete;
+	FPCGGetDataFunctionRegistry(FPCGGetDataFunctionRegistry&&) = default;
+	FPCGGetDataFunctionRegistry& operator=(const FPCGGetDataFunctionRegistry&) = delete;
+	FPCGGetDataFunctionRegistry& operator=(FPCGGetDataFunctionRegistry&&) = default;
+	virtual ~FPCGGetDataFunctionRegistry() = default;
+
+	/** Register an actor to PCG data function. */
+	UE_API FFunctionHandle RegisterDataFromActorFunction(const FDataFromActorFunction& InFunction);
+
+	/** Unregisters an actor to PCG data function. */
+	UE_API void UnregisterDataFromActorFunction(FFunctionHandle InFunctionHandle);
+
+	/** Appends actor data to the given collection and returns the number of data added. */
+	UE_API int GetDataFromActor(FPCGContext* InContext, const FPCGGetDataFunctionRegistryParams& InParams, AActor* InActor, FPCGGetDataFunctionRegistryOutput& OutDataCollection) const;
+
+	/** Register a component to PCG data function. */
+	UE_API FFunctionHandle RegisterDataFromComponentFunction(const FDataFromComponentFunction& InFunction);
+
+	/** Unregisters a component to PCG data function. */
+	UE_API void UnregisterDataFromComponentFunction(FFunctionHandle InFunctionHandle);
+
+	/** Append component data to the given collection, returns the number of data added. */
+	UE_API int GetDataFromComponent(FPCGContext* InContext, const FPCGGetDataFunctionRegistryParams& InParams, UActorComponent* InComponent, FPCGGetDataFunctionRegistryOutput& OutDataCollection) const;
+
+private:
+	UE_API int DefaultDataFromActor(FPCGContext* InContext, const FPCGGetDataFunctionRegistryParams& InParams, AActor* InActor, FPCGGetDataFunctionRegistryOutput& OutDataCollection) const;
+	UE_API int DefaultDataFromComponent(FPCGContext* InContext, const FPCGGetDataFunctionRegistryParams& InParams, UActorComponent* InComponent, FPCGGetDataFunctionRegistryOutput& OutDataCollection) const;
+
+	FFunctionHandle NextFunctionHandle = 0;
+
+	TArray<TPair<FDataFromActorFunction, FFunctionHandle>> ActorParsingFunctions;
+	TArray<TPair<FDataFromComponentFunction, FFunctionHandle>> ComponentParsingFunctions;
+};
+
+#undef UE_API
diff --git a/Engine/Plugins/PCG/Source/PCG/Public/Helpers/PCGSubgraphHelpers.h b/Engine/Plugins/PCG/Source/PCG/Public/Helpers/PCGSubgraphHelpers.h
--- a/Engine/Plugins/PCG/Source/PCG/Public/Helpers/PCGSubgraphHelpers.h
+++ b/Engine/Plugins/PCG/Source/PCG/Public/Helpers/PCGSubgraphHelpers.h
@@ -3,9 +3,12 @@
 #pragma once
 
 #include "Containers/Array.h"
+#include "Templates/SubclassOf.h"
 
 class UPCGGraph;
 class UPCGNode;
+class UPCGPin;
+class UPCGSettings;
 class UObject;
 
 class FPCGSubgraphHelpers
@@ -13,4 +16,7 @@ class FPCGSubgraphHelpers
 public:
 	static PCG_API UPCGGraph* CollapseIntoSubgraphWithReason(UPCGGraph* InOriginalGraph, const TArray<UPCGNode*>& InNodesToCollapse, const TArray<UObject*>& InExtraNodesToCollapse, FText& OutFailReason, UPCGGraph* OptionalPreAllocatedGraph = nullptr);
 	static PCG_API UPCGGraph* CollapseIntoSubgraph(UPCGGraph* InOriginalGraph, const TArray<UPCGNode*>& InNodesToCollapse, const TArray<UObject*>& InExtraNodesToCollapse, UPCGGraph* OptionalPreAllocatedGraph = nullptr);
+
+	/** Generic function to spawn a node and connect it to the default in and out pins. */
+	static PCG_API UPCGNode* SpawnNodeAndConnect(UPCGGraph* InGraph, UPCGPin* InUpstreamPin, UPCGPin* InDownstreamPin, TSubclassOf<UPCGSettings> InSettingsClass, const TFunction<bool(UPCGSettings*)>& PostSettingsCreation = {});
 };
\ No newline at end of file
diff --git a/Engine/Plugins/PCG/Source/PCG/Public/PCGData.h b/Engine/Plugins/PCG/Source/PCG/Public/PCGData.h
--- a/Engine/Plugins/PCG/Source/PCG/Public/PCGData.h
+++ b/Engine/Plugins/PCG/Source/PCG/Public/PCGData.h
@@ -5,6 +5,7 @@
 #include "PCGCommon.h"
 #include "PCGCrc.h"
 #include "Data/PCGDataPtrWrapper.h"
+#include "Data/Registry/PCGDataType.h"
 #include "Metadata/PCGAttributePropertySelector.h"
 #include "Metadata/PCGMetadataCommon.h"
 
@@ -53,7 +54,13 @@ class UPCGData : public UObject
 
 public:
 	UE_API UPCGData(const FObjectInitializer& ObjectInitializer);
-	virtual EPCGDataType GetDataType() const { return EPCGDataType::None; }
+
+	// @todo_pcg To be deprecated when we switch to the new type
+	//UE_DEPRECATED(5.7, "Use the new GetDataTypeId")
+	virtual EPCGDataType GetDataType() const { return EPCGDataType::Any; }
+
+	using TypeInfo = FPCGDataTypeInfo;
+	virtual FPCGDataTypeBaseId GetDataTypeId() const { return FPCGDataTypeBaseId::Construct<FPCGDataTypeInfo>(); }
 
 	/** Returns a Crc for this and any connected data. */
 	UE_API FPCGCrc GetOrComputeCrc(bool bFullDataCrc) const;
@@ -73,7 +80,14 @@ public:
 	virtual void ReleaseTransientResources(const TCHAR* InReason = nullptr) {}
 
 	/** If this data is a proxy, returns the underlying data type, otherwise returns this data type. */
-	virtual EPCGDataType GetUnderlyingDataType() const { return GetDataType(); }
+	// @todo_pcg To be deprecated when we switch to the new type
+	// UE_DEPRECATED(5.7, "Use the new GetUnderlyingDataTypeInfo")
+	virtual EPCGDataType GetUnderlyingDataType() const
+	{
+		return GetUnderlyingDataTypeId().AsLegacyType();
+	}
+
+	virtual FPCGDataTypeBaseId GetUnderlyingDataTypeId() const { return GetDataTypeId(); }
 
 	/** Unique ID for this object instance. */
 	UPROPERTY(Transient)
diff --git a/Engine/Plugins/PCG/Source/PCG/Public/PCGModule.h b/Engine/Plugins/PCG/Source/PCG/Public/PCGModule.h
--- a/Engine/Plugins/PCG/Source/PCG/Public/PCGModule.h
+++ b/Engine/Plugins/PCG/Source/PCG/Public/PCGModule.h
@@ -3,7 +3,8 @@
 #pragma once
 
 #include "PCGDataVisualizationRegistry.h"
-#include "Data/PCGGetDataFunctionRegistry.h"
+#include "Data/Registry/PCGDataTypeRegistry.h"
+#include "Data/Registry/PCGGetDataFunctionRegistry.h"
 #include "Hash/PCGObjectHash.h"
 #include "Metadata/Accessors/PCGAttributeAccessorFactory.h"
 #include "Utils/PCGLogErrors.h"
@@ -41,11 +42,15 @@ public:
 	//~ End IModuleInterface implementation
 
 	void PreExit();
+	void OnPostInitEngine();
 
 	PCG_API static FPCGModule& GetPCGModuleChecked();
 	static const FPCGGetDataFunctionRegistry& ConstGetDataFunctionRegistry() { return GetPCGModuleChecked().GetDataFunctionRegistry; }
 	static FPCGGetDataFunctionRegistry& MutableGetDataFunctionRegistry() { return GetPCGModuleChecked().GetDataFunctionRegistry; }
 
+	PCG_API static const FPCGDataTypeRegistry& GetConstDataTypeRegistry();
+	PCG_API static FPCGDataTypeRegistry& GetMutableDataTypeRegistry();
+
 	static const FPCGAttributeAccessorFactory& GetConstAttributeAccessorFactory() { return GetPCGModuleChecked().AttributeAccessorFactory; }
 	static FPCGAttributeAccessorFactory& GetMutableAttributeAccessorFactory() { return GetPCGModuleChecked().AttributeAccessorFactory; }
 	
@@ -72,6 +77,7 @@ private:
 
 	FPCGGetDataFunctionRegistry GetDataFunctionRegistry;
 	FPCGAttributeAccessorFactory AttributeAccessorFactory;
+	FPCGDataTypeRegistry DataTypeRegistry;
 
 #if WITH_EDITOR
 	FPCGObjectHashFactory ObjectHashFactory;
diff --git a/Engine/Plugins/PCG/Source/PCG/Public/PCGParamData.h b/Engine/Plugins/PCG/Source/PCG/Public/PCGParamData.h
--- a/Engine/Plugins/PCG/Source/PCG/Public/PCGParamData.h
+++ b/Engine/Plugins/PCG/Source/PCG/Public/PCGParamData.h
@@ -14,6 +14,17 @@ namespace PCGParamDataConstants
 	const FName ElementsDomainName = "Elements";
 }
 
+USTRUCT()
+struct FPCGDataTypeInfoParam : public FPCGDataTypeInfo
+{
+	GENERATED_BODY()
+
+	PCG_DECLARE_LEGACY_TYPE_INFO(EPCGDataType::Param)
+
+	PCG_API virtual bool SupportsConversion(const FPCGDataTypeIdentifier& Other) const override;
+	PCG_API virtual TSubclassOf<UPCGSettings> GetDefaultConversionSettings(const FPCGDataTypeIdentifier& Other) const override;
+};
+
 /**
 * Class to hold execution parameters that will be consumed in nodes of the graph
 */
@@ -26,7 +37,7 @@ public:
 	UE_API UPCGParamData(const FObjectInitializer& ObjectInitializer);
 
 	// ~Begin UPCGData interface
-	virtual EPCGDataType GetDataType() const override { return EPCGDataType::Param; }
+	PCG_ASSIGN_DEFAULT_TYPE_INFO(FPCGDataTypeInfoParam)
 	UE_API virtual void AddToCrc(FArchiveCrc32& Ar, bool bFullDataCrc) const override;
 
 	UE_API virtual bool HasCachedLastSelector() const override;
diff --git a/Engine/Plugins/PCGInterops/PCGGeometryScriptInterop/Source/PCGGeometryScriptInterop/Private/Elements/PCGGetDynamicMeshData.cpp b/Engine/Plugins/PCGInterops/PCGGeometryScriptInterop/Source/PCGGeometryScriptInterop/Private/Elements/PCGGetDynamicMeshData.cpp
--- a/Engine/Plugins/PCGInterops/PCGGeometryScriptInterop/Source/PCGGeometryScriptInterop/Private/Elements/PCGGetDynamicMeshData.cpp
+++ b/Engine/Plugins/PCGInterops/PCGGeometryScriptInterop/Source/PCGGeometryScriptInterop/Private/Elements/PCGGetDynamicMeshData.cpp
@@ -4,7 +4,7 @@
 
 #include "PCGContext.h"
 #include "Data/PCGDynamicMeshData.h"
-#include "Data/PCGGetDataFunctionRegistry.h"
+#include "Data/Registry/PCGGetDataFunctionRegistry.h"
 #include "Helpers/PCGGeometryHelpers.h"
 #include "Helpers/PCGHelpers.h"
 
