[PCG] Editor Mode [Beta] - Create new editor mode for PCG with initial toolbox setup

diff --git a/Engine/Plugins/PCG/Source/PCG/Public/PCGGraph.h b/Engine/Plugins/PCG/Source/PCG/Public/PCGGraph.h
--- a/Engine/Plugins/PCG/Source/PCG/Public/PCGGraph.h
+++ b/Engine/Plugins/PCG/Source/PCG/Public/PCGGraph.h
@@ -174,6 +174,10 @@ public:
 	/** Override of the color for the subgraph node for this graph. */
 	UPROPERTY(BlueprintReadWrite, EditAnywhere, Category = AssetInfo, meta = (EditCondition = "AreOverridesEnabled && bOverrideColor", EditConditionHides, PCGNoHash))
 	FLinearColor Color = FLinearColor::White;
+
+	/** Signifies the graph has compatibility with the annotated tags. Multiple tags may be comma separated. Ex. "PCGTool1, PCGTool2, PCGTool3" */
+	UPROPERTY(BlueprintReadWrite, EditAnywhere, Category = AssetInfo, AssetRegistrySearchable, meta=(EditCondition = "bExposeToLibrary", EditConditionHides, PCGNoHash))
+	FString CompatibleToolTags;
 #endif
 
 	template <typename T>
diff --git a/Engine/Plugins/PCG/Source/PCGEditor/PCGEditor.Build.cs b/Engine/Plugins/PCG/Source/PCGEditor/PCGEditor.Build.cs
--- a/Engine/Plugins/PCG/Source/PCGEditor/PCGEditor.Build.cs
+++ b/Engine/Plugins/PCG/Source/PCGEditor/PCGEditor.Build.cs
@@ -35,8 +35,8 @@ namespace UnrealBuildTool.Rules
 					"AppFramework",
 					"ApplicationCore",
 					"AssetDefinition",
-					"AssetTools",
 					"AssetRegistry",
+					"AssetTools",
 					"BlueprintGraph",
 					"ContentBrowser",
 					"ContentBrowserData",
@@ -51,22 +51,25 @@ namespace UnrealBuildTool.Rules
 					"GameProjectGeneration",
 					"GraphEditor",
 					"InputCore",
+					"InteractiveToolsFramework",
 					"Kismet",
 					"KismetWidgets",
+					"LevelEditor",
 					"PropertyEditor",
 					"RHI",
 					"RenderCore",
+					"SceneOutliner",
 					"Slate",
 					"SlateCore",
 					"SourceControl",
+					"StatusBar",
 					"StructUtilsEditor",
 					"ToolMenus",
 					"ToolWidgets",
 					"TypedElementFramework",
 					"TypedElementRuntime",
 					"UnrealEd",
-					"LevelEditor",
-					"SceneOutliner"
+					"WidgetRegistration",
 				});
 		}
 	}
diff --git a/Engine/Plugins/PCG/Source/PCGEditor/Private/EditorMode/PCGEdMode.cpp b/Engine/Plugins/PCG/Source/PCGEditor/Private/EditorMode/PCGEdMode.cpp
new file mode 100644
--- /dev/null
+++ b/Engine/Plugins/PCG/Source/PCGEditor/Private/EditorMode/PCGEdMode.cpp
@@ -0,0 +1,431 @@
+// Copyright Epic Games, Inc. All Rights Reserved.
+
+#include "EditorMode/PCGEdMode.h"
+
+#include "EditorMode/PCGEdModeCommands.h"
+#include "EditorMode/PCGEdModeToolkit.h"
+#include "EditorMode/Tools/Line/PCGDrawSplineTool.h"
+#include "EditorMode/Tools/Volume/PCGCreateVolumeTool.h"
+
+#include "ContextObjectStore.h"
+#include "EngineAnalytics.h"
+#include "IAnalyticsProviderET.h"
+#include "ILevelEditor.h"
+#include "InteractiveToolManager.h"
+#include "Selection.h"
+#include "UnrealEdGlobals.h" // GUnrealEd
+#include "BaseGizmos/TransformGizmoUtil.h"
+#include "Components/StaticMeshComponent.h"
+#include "Editor/UnrealEdEngine.h"
+#include "Tools/EdModeInteractiveToolsContext.h"
+
+// Commands
+#include "EditorModeManager.h"
+#include "InteractiveToolQueryInterfaces.h"
+#include "LevelEditor.h"
+#include "SLevelViewport.h"
+#include "Application/ThrottleManager.h"
+
+#include "ToolContextInterfaces.h"
+#include "ToolMenus.h"
+#include "EditorMode/PCGEdModeSettings.h"
+#include "EditorMode/PCGEdModeStyle.h"
+#include "EditorMode/Tools/Paint/PCGPaintTool.h"
+#include "Engine/World.h" // FWorldDelegates
+#include "Framework/Application/SlateApplication.h"
+#include "ViewportToolbar/UnrealEdViewportToolbar.h"
+
+#include UE_INLINE_GENERATED_CPP_BY_NAME(PCGEdMode)
+
+#define LOCTEXT_NAMESPACE "UPCGEditorMode"
+
+const FEditorModeID UPCGEditorMode::EM_PCGEditorModeId = TEXT("EM_PCGEditorMode");
+
+namespace PCGEditorMode
+{
+	namespace Constants
+	{
+		const FName PCGEditorModeOwnerName = "PCGEditorModeViewportToolBar";
+		const FText PCGEditorModeDisplayName = LOCTEXT("PCGEditorModeDisplayName", "PCG (Experimental)");
+	}
+
+	namespace Helpers
+	{
+		FString GetToolName(const UInteractiveTool& Tool)
+		{
+			const FString* ToolName = FTextInspector::GetSourceString(Tool.GetToolInfo().ToolDisplayName);
+			return ToolName ? *ToolName : FString(TEXT("<Invalid ToolName>"));
+		}
+	}
+}
+
+UPCGEditorMode::UPCGEditorMode()
+{
+	Info = FEditorModeInfo(
+		EM_PCGEditorModeId,
+		PCGEditorMode::Constants::PCGEditorModeDisplayName,
+		FSlateIcon("PCGEditorModeStyle", "PCGEditorModeIcon"),
+		GetDefault<UPCGEditorModeSettings>()->bEnableEditorMode,
+		/*@todo_pcg: TBD...*/std::numeric_limits<int32>::max());
+}
+
+void UPCGEditorMode::RegisterEditorMode()
+{
+	FPCGEditorModeStyle::Register();
+	FPCGEditorModeCommands::Register();
+	FPCGEditorModePaletteCommands::Register();
+	FPCGEditorModeToolCommands::Register();
+	FPCGEditorModeToolCommands::RegisterAllToolCommands();
+}
+
+void UPCGEditorMode::UnregisterEditorMode()
+{
+	FPCGEditorModeToolCommands::UnregisterAllToolCommands();
+	FPCGEditorModeToolCommands::Unregister();
+	FPCGEditorModePaletteCommands::Unregister();
+	FPCGEditorModeCommands::Unregister();
+	FPCGEditorModeStyle::Unregister();
+}
+
+void UPCGEditorMode::ActorSelectionChangeNotify()
+{
+	// @todo_pcg: Validate and accept/deny selection.
+}
+
+bool UPCGEditorMode::ProcessEditDelete()
+{
+	if (UEdMode::ProcessEditDelete())
+	{
+		return true;
+	}
+
+	UInteractiveToolManager* ToolManager = GetToolManager();
+	// Disable deleting in an Accept-style tool which could be volatile.
+	if (ToolManager && ToolManager->HasAnyActiveTool() && ToolManager->GetActiveTool(EToolSide::Mouse)->HasAccept())
+	{
+		ToolManager->DisplayMessage(
+			LOCTEXT("CannotDeleteWarning", "Cannot delete objects while this Tool is active"), EToolMessageLevel::UserWarning);
+		return true;
+	}
+
+	return false;
+}
+
+bool UPCGEditorMode::ProcessEditCut()
+{
+	if (UInteractiveToolManager* ToolManager = GetToolManager())
+	{
+		// Disable deleting in an Accept-style tool because it can result in crashes if we are deleting target object
+		if (ToolManager->HasAnyActiveTool() && ToolManager->GetActiveTool(EToolSide::Mouse)->HasAccept())
+		{
+			ToolManager->DisplayMessage(
+				LOCTEXT("CannotCutWarning", "Cannot cut objects while this Tool is active"), EToolMessageLevel::UserWarning);
+			return true;
+		}
+	}
+
+	return false;
+}
+
+bool UPCGEditorMode::CanAutoSave() const
+{
+	// @todo_pcg: Autosave should be fine, since all tools' artifacts are transient, but this needs to be verified.
+	// For now, just disable autosave if a tool is active.
+	return !GetToolManager() || (GetToolManager()->HasAnyActiveTool() == false);
+}
+
+bool UPCGEditorMode::GetPivotForOrbit(FVector& OutPivot) const
+{
+	// @todo_pcg: Some features may want to override this.
+
+	if (GCurrentLevelEditingViewportClient)
+	{
+		OutPivot = GCurrentLevelEditingViewportClient->GetViewTransform().GetLookAt();
+		return true;
+	}
+
+	return false;
+}
+
+void UPCGEditorMode::Enter()
+{
+	UEdMode::Enter();
+
+	UInteractiveToolManager* ToolManager = GetToolManager();
+	if (!ToolManager)
+	{
+		return;
+	}
+
+	ToolManager->OnToolPostBuild.AddUObject(this, &UPCGEditorMode::OnToolPostBuild);
+	ToolManager->OnToolShutdownRequest.BindLambda([this](UInteractiveToolManager*, UInteractiveTool*, const EToolShutdownType ShutdownType)
+	{
+		if (UEditorInteractiveToolsContext* ToolsContext = GetInteractiveToolsContext())
+		{
+			ToolsContext->EndTool(ShutdownType);
+			return true;
+		}
+		else
+		{
+			return false;
+		}
+	});
+
+	const FPCGEditorModeToolCommands& ToolManagerCommands = FPCGEditorModeToolCommands::Get();
+
+	// @todo_pcg: Could be abstracted out and have the tool identifier come from the label
+	UPCGDrawSplineToolBuilder* DrawSplineBuilder = NewObject<UPCGDrawSplineToolBuilder>();
+	DrawSplineBuilder->SetIsSplineSurface(false);
+	RegisterTool(ToolManagerCommands.EnableDrawSplineTool, TEXT("DrawSpline"), DrawSplineBuilder);
+	UPCGDrawSplineToolBuilder* DrawSurfaceBuilder = NewObject<UPCGDrawSplineToolBuilder>();
+	DrawSurfaceBuilder->SetIsSplineSurface(true);
+	RegisterTool(ToolManagerCommands.EnableDrawSurfaceTool, TEXT("DrawSurface"), DrawSurfaceBuilder);
+
+	RegisterTool(ToolManagerCommands.EnablePaintTool, TEXT("Paint"), NewObject<UPCGPaintToolBuilder>());
+	RegisterTool(ToolManagerCommands.EnableCreateVolumeTool, TEXT("CreateVolume"), NewObject<UPCGCreateVolumeToolBuilder>());
+
+	ToolManager->SelectActiveToolType(EToolSide::Left, TEXT("Line"));
+}
+
+void UPCGEditorMode::Exit()
+{
+	// Clear realtime viewport override
+	ConfigureRealTimeViewportsOverride(false);
+
+	if (UInteractiveToolManager* ToolManager = GetToolManager())
+	{
+		ToolManager->OnToolShutdownRequest.Unbind();
+		ToolManager->OnToolPostBuild.RemoveAll(this);
+	}
+
+	// Call base Exit method to ensure proper cleanup
+	UEdMode::Exit();
+}
+
+bool UPCGEditorMode::ShouldToolStartBeAllowed(const FString& ToolIdentifier) const
+{
+	if (UInteractiveToolManager* ToolManager = GetToolManager())
+	{
+		// @todo_pcg: As more tools are added, some may require scene compatibility.
+	}
+
+	return Super::ShouldToolStartBeAllowed(ToolIdentifier);
+}
+
+bool UPCGEditorMode::HasCustomViewportFocus() const
+{
+	if (Super::HasCustomViewportFocus())
+	{
+		return true;
+	}
+
+	// @todo_pcg: Some features may want to override this.
+
+	return false;
+}
+
+FBox UPCGEditorMode::ComputeCustomViewportFocus() const
+{
+	auto ProcessFocusBoxFunc = [](FBox& FocusBoxInOut)
+	{
+		const double MaxDimension = FocusBoxInOut.GetExtent().GetMax();
+		FocusBoxInOut = FocusBoxInOut.ExpandBy(MaxDimension * 0.2);
+	};
+
+	FBox FocusBox = Super::ComputeCustomViewportFocus();
+	if (FocusBox.IsValid)
+	{
+		ProcessFocusBoxFunc(FocusBox);
+		return FocusBox;
+	}
+
+	// @todo_pcg: Some features may want to override this.
+
+	// No custom focus box. Return a default (invalid) box
+	return FBox();
+}
+
+void UPCGEditorMode::CreateToolkit()
+{
+	Toolkit = MakeShared<FPCGEditorModeToolkit>();
+}
+
+void UPCGEditorMode::OnToolStarted(UInteractiveToolManager* Manager, UInteractiveTool* Tool)
+{
+	// Disable slate throttling so that Tool background computes responding to sliders can properly be processed
+	// on Tool Tick. Otherwise, when a Tool kicks off a background update in a background thread, the computed
+	// result will be ignored until the user moves the slider, ie you cannot hold down the mouse and wait to see
+	// the result. This apparently broken behavior is currently by-design.
+	FSlateThrottleManager::Get().DisableThrottle(true);
+
+	bIsToolActive = true;
+}
+
+void UPCGEditorMode::OnToolEnded(UInteractiveToolManager* Manager, UInteractiveTool* Tool)
+{
+	bIsToolActive = false;
+
+	// Give the editor back control over pivot location
+	GUnrealEd->UpdatePivotLocationForSelection();
+
+	// Re-enable slate throttling (see OnToolStarted)
+	FSlateThrottleManager::Get().DisableThrottle(false);
+}
+
+void UPCGEditorMode::OnToolPostBuild(
+	UInteractiveToolManager* InToolManager,
+	EToolSide InSide,
+	UInteractiveTool* InBuiltTool,
+	UInteractiveToolBuilder* InToolBuilder,
+	const FToolBuilderState& ToolState)
+{
+	// @todo_pcg: Anything the editor mode might need to track after a tool is built can go here.
+}
+
+void UPCGEditorMode::OnEditorClosed() const
+{
+	// On editor close, Exit() should run to clean up, but this happens very late.
+	// Close out any active Tools or Selections to mitigate any late-destruction issues.
+
+	UInteractiveToolManager* ToolManager = GetToolManager();
+
+	if (GetModeManager() != nullptr
+		&& GetInteractiveToolsContext() != nullptr
+		&& ToolManager != nullptr
+		&& ToolManager->HasAnyActiveTool())
+	{
+		ToolManager->DeactivateTool(EToolSide::Mouse, EToolShutdownType::Cancel);
+	}
+}
+
+void UPCGEditorMode::AcceptActiveToolActionOrTool() const
+{
+	UInteractiveToolManager* ToolManager = GetToolManager();
+
+	if (ToolManager && ToolManager->HasAnyActiveTool())
+	{
+		UInteractiveTool* Tool = ToolManager->GetActiveTool(EToolSide::Mouse);
+		IInteractiveToolNestedAcceptCancelAPI* CancelAPI = Cast<IInteractiveToolNestedAcceptCancelAPI>(Tool);
+		if (CancelAPI && CancelAPI->SupportsNestedAcceptCommand() && CancelAPI->CanCurrentlyNestedAccept())
+		{
+			if (CancelAPI->ExecuteNestedAcceptCommand())
+			{
+				return;
+			}
+		}
+	}
+
+	UEditorInteractiveToolsContext* ToolsContext = GetInteractiveToolsContext();
+	const EToolShutdownType ShutdownType = ToolsContext->CanAcceptActiveTool() ? EToolShutdownType::Accept : EToolShutdownType::Completed;
+	ToolsContext->EndTool(ShutdownType);
+}
+
+void UPCGEditorMode::CancelActiveToolActionOrTool() const
+{
+	UInteractiveToolManager* ToolManager = GetToolManager();
+
+	if (ToolManager && ToolManager->HasAnyActiveTool())
+	{
+		UInteractiveTool* Tool = ToolManager->GetActiveTool(EToolSide::Mouse);
+		IInteractiveToolNestedAcceptCancelAPI* CancelAPI = Cast<IInteractiveToolNestedAcceptCancelAPI>(Tool);
+		if (CancelAPI && CancelAPI->SupportsNestedCancelCommand() && CancelAPI->CanCurrentlyNestedCancel())
+		{
+			if (CancelAPI->ExecuteNestedCancelCommand())
+			{
+				return;
+			}
+		}
+	}
+
+	UEditorInteractiveToolsContext* ToolsContext = GetInteractiveToolsContext();
+	const EToolShutdownType ShutdownType = ToolsContext->CanCancelActiveTool() ? EToolShutdownType::Cancel : EToolShutdownType::Completed;
+	ToolsContext->EndTool(ShutdownType);
+}
+
+void UPCGEditorMode::ConfigureRealTimeViewportsOverride(const bool bEnable)
+{
+	const FLevelEditorModule& LevelEditorModule = FModuleManager::GetModuleChecked<FLevelEditorModule>("LevelEditor");
+	const TSharedPtr<ILevelEditor> LevelEditor = LevelEditorModule.GetFirstLevelEditor();
+	if (LevelEditor.IsValid())
+	{
+		TArray<TSharedPtr<SLevelViewport>> Viewports = LevelEditor->GetViewports();
+		for (const TSharedPtr<SLevelViewport>& ViewportWindow : Viewports)
+		{
+			if (ViewportWindow.IsValid())
+			{
+				FEditorViewportClient& Viewport = ViewportWindow->GetAssetViewportClient();
+				if (bEnable)
+				{
+					Viewport.AddRealtimeOverride(bEnable, PCGEditorMode::Constants::PCGEditorModeDisplayName);
+				}
+				else
+				{
+					Viewport.RemoveRealtimeOverride(PCGEditorMode::Constants::PCGEditorModeDisplayName, false);
+				}
+			}
+		}
+	}
+}
+
+// @todo_pcg: Activate when custom tools are supported.
+// void UPCGEditorMode::RegisterCustomTool(
+// 	TSharedPtr<FUICommandInfo> UICommand,
+// 	FString ToolIdentifier,
+// 	UInteractiveToolBuilder* Builder,
+// 	const TFunction<bool(UInteractiveToolManager*, EToolSide)>& ExecuteAction,
+// 	const TFunction<bool(UInteractiveToolManager*, EToolSide)>& CanExecuteAction,
+// 	const TFunction<bool(UInteractiveToolManager*, EToolSide)>& IsActionChecked) const
+// {
+// 	if (!Toolkit.IsValid())
+// 	{
+// 		return;
+// 	}
+//
+// 	constexpr EToolsContextScope ToolScope = EToolsContextScope::EdMode;
+// 	UEditorInteractiveToolsContext* UseToolsContext = GetInteractiveToolsContext(ToolScope);
+// 	if (ensure(UseToolsContext != nullptr) == false)
+// 	{
+// 		return;
+// 	}
+//
+// 	constexpr EToolSide ToolSide = EToolSide::Left;
+// 	const TSharedRef<FUICommandList>& CommandList = Toolkit->GetToolkitCommands();
+// 	UseToolsContext->ToolManager->RegisterToolType(ToolIdentifier, Builder);
+// 	CommandList->MapAction(UICommand,
+// 		FExecuteAction::CreateWeakLambda(UseToolsContext, [this, ToolIdentifier, UseToolsContext, ExecuteAction]
+// 		{
+// 			if (!ExecuteAction || !ExecuteAction(UseToolsContext->ToolManager, ToolSide))
+// 			{
+// 				UseToolsContext->StartTool(ToolIdentifier);
+// 			}
+// 		}),
+// 		FCanExecuteAction::CreateWeakLambda(UseToolsContext, [this, ToolIdentifier, UseToolsContext, CanExecuteAction]()
+// 		{
+// 			bool bResult = false;
+// 			if (ShouldToolStartBeAllowed(ToolIdentifier))
+// 			{
+// 				bResult = CanExecuteAction ? CanExecuteAction(UseToolsContext->ToolManager, ToolSide) : UseToolsContext->ToolManager->CanActivateTool(ToolSide, ToolIdentifier);
+// 			}
+// 			return bResult;
+// 		}),
+// 		FIsActionChecked::CreateWeakLambda(UseToolsContext, [this, ToolIdentifier, UseToolsContext, IsActionChecked]()
+// 		{
+// 			if (IsActionChecked)
+// 			{
+// 				return IsActionChecked(UseToolsContext->ToolManager, ToolSide);
+// 			}
+//
+// 			if (UseToolsContext->ToolManager->GetActiveTool(ToolSide))
+// 			{
+// 				// Read the ActiveToolType rather than the ActiveToolName so that clients can leverage
+// 				// that storage on the ToolsContext through SelectActiveToolType to notify of changes
+// 				// in the active tool without actually requiring a tool change.
+// 				const FString& ActiveToolIdentifier = UseToolsContext->ToolManager->GetActiveToolType(ToolSide);
+// 				return ActiveToolIdentifier == ToolIdentifier;
+// 			}
+// 			return false;
+// 		}),
+// 		EUIActionRepeatMode::RepeatDisabled);
+// }
+
+#undef LOCTEXT_NAMESPACE
diff --git a/Engine/Plugins/PCG/Source/PCGEditor/Private/EditorMode/PCGEdModeCommands.cpp b/Engine/Plugins/PCG/Source/PCGEditor/Private/EditorMode/PCGEdModeCommands.cpp
new file mode 100644
--- /dev/null
+++ b/Engine/Plugins/PCG/Source/PCGEditor/Private/EditorMode/PCGEdModeCommands.cpp
@@ -0,0 +1,148 @@
+// Copyright Epic Games, Inc. All Rights Reserved.
+
+#include "EditorMode/PCGEdModeCommands.h"
+
+#include "EditorMode/PCGEdModeSettings.h"
+#include "EditorMode/PCGEdModeStyle.h"
+#include "EditorMode/Tools/Line/PCGDrawSplineTool.h"
+#include "EditorMode/Tools/Paint/PCGPaintTool.h"
+#include "EditorMode/Tools/Volume/PCGCreateVolumeTool.h"
+
+#define LOCTEXT_NAMESPACE "PCGEditorMode"
+
+struct FSlateIcon;
+
+FPCGEditorModeCommands::FPCGEditorModeCommands()
+	: TCommands(/*InContextName=*/TEXT("PCGEditorModeCommands")
+							  , /*InContextDesc=*/LOCTEXT("PCGEditorModeCommands", "PCG Editor Mode")
+							  , /*InContextParent=*/NAME_None
+							  , /*InStyleSetName=*/FPCGEditorModeStyle::Get().GetStyleSetName()) {}
+
+void FPCGEditorModeCommands::RegisterCommands() {}
+
+FPCGEditorModePaletteCommands::FPCGEditorModePaletteCommands()
+	: TCommands(TEXT("PCGEditorModePaletteCommands")
+			, LOCTEXT("PCGEditorModePaletteCommands", "PCG Editor Mode - Palette")
+			, /*Parent=*/NAME_None
+			, /*IconStyleSet=*/FPCGEditorModeStyle::Get().GetStyleSetName()) {}
+
+/**
+ * @param CommandInfoSP The shared pointer for the FUICommandInfo
+ * @param Name The name of the tool
+ * @param IconStyleName The style name for the icon to be used for the button
+ */
+#define PCG_CONTEXT_COMMAND_INFO(CommandInfoSP, Name, IconStyleName)\
+FUICommandInfo::MakeCommandInfo(AsShared()\
+	, CommandInfoSP\
+	, TEXT("Enter" Name "Context")\
+	, LOCTEXT(Name "ContextLabel", Name)\
+	, LOCTEXT(Name "ContextToolTip", "Enter '" Name "' context and open corresponding toolkit")\
+	, FSlateIcon(FPCGEditorModeStyle::Get().GetStyleSetName(), IconStyleName)\
+	, EUserInterfaceActionType::ToggleButton\
+	, FInputChord())
+
+void FPCGEditorModePaletteCommands::RegisterCommands()
+{
+	/*** PCG EDITOR CONTEXTS **/
+	PCG_CONTEXT_COMMAND_INFO(LoadLineContextPalette, "Line", "PCGEditorMode.Tools.DrawSpline");
+	PCG_CONTEXT_COMMAND_INFO(LoadPaintContextPalette, "Paint", "PCGEditorMode.Context.Paint");
+	PCG_CONTEXT_COMMAND_INFO(LoadVolumeContextPalette, "Volume", "PCGEditorMode.Context.Volume");
+}
+
+#undef PCG_CONTEXT_COMMAND_INFO
+
+FPCGEditorModeToolCommands::FPCGEditorModeToolCommands()
+	: TInteractiveToolCommands(TEXT("PCGEditorModeToolCommands")
+						   , LOCTEXT("PCGEditorModeToolCommands", "PCG Editor Mode - Tool Commands")
+						   , /*Parent=*/NAME_None
+						   , /*IconStyleSet=*/FPCGEditorModeStyle::Get().GetStyleSetName()) {}
+
+/**
+  * @param CommandInfoSP The shared pointer for the FUICommandInfo
+ * @param Name The name of the tool
+ * @param Label The user-forward label for the tool
+ * @param ToolTip The tooltip for the tool
+ * @param IconStyleName The style name for the icon to be used for the button
+ */
+#define PCG_TOOL_COMMAND_INFO(CommandInfoSP, Name, Label, ToolTip, IconStyleName)\
+FUICommandInfo::MakeCommandInfo(AsShared()\
+	, CommandInfoSP\
+	, TEXT(Name"Tool")\
+	, LOCTEXT(Name"ToolLabel", Label)\
+	, LOCTEXT(Name"ToolTip", ToolTip)\
+	, FSlateIcon(FPCGEditorModeStyle::Get().GetStyleSetName(), IconStyleName)\
+	, EUserInterfaceActionType::ToggleButton\
+	, FInputChord())
+
+void FPCGEditorModeToolCommands::RegisterCommands()
+{
+	TInteractiveToolCommands::RegisterCommands();
+
+	/*** PCG EDITOR MODE TOOLS **/
+	PCG_TOOL_COMMAND_INFO(EnableDrawSplineTool, "DrawSpline", "Draw Spline", "Use the 'Draw Spline' tool to build a spline for PCG.", "PCGEditorMode.Tools.DrawSpline");
+	PCG_TOOL_COMMAND_INFO(EnableDrawSurfaceTool, "DrawSurface", "Draw Surface", "Use the 'Draw Surface' tool to build a spline for PCG.", "PCGEditorMode.Tools.DrawSurface");
+	PCG_TOOL_COMMAND_INFO(EnablePaintTool, "Paint", "Paint", "Use the 'Paint' tool to apply PCG points with a brush.", "PCGEditorMode.Context.Paint");
+	PCG_TOOL_COMMAND_INFO(EnableCreateVolumeTool, "CreateVolume", "Create Volume", "Create a PCG volume in the scene.", "PCGEditorMode.Tools.CreateVolume");
+}
+
+#undef PCG_TOOL_COMMAND_INFO
+
+void FPCGEditorModeToolCommands::RegisterAllToolCommands()
+{
+	FPCGDrawSplineCommands::Register();
+	FPCGPaintCommands::Register();
+	FPCGCreateVolumeCommands::Register();
+}
+
+void FPCGEditorModeToolCommands::UnregisterAllToolCommands()
+{
+	FPCGCreateVolumeCommands::Unregister();
+	FPCGPaintCommands::Unregister();
+	FPCGDrawSplineCommands::Unregister();
+}
+
+void FPCGEditorModeToolCommands::GetToolDefaultObjectList(TArray<UInteractiveTool*>& ToolCDOs)
+{
+	// @todo_pcg: To be populated with tool commands for bindings.
+}
+
+TArray<TSharedPtr<FUICommandInfo>> FPCGEditorModeToolCommands::GetLineContextToolCommands() const
+{
+	return {EnableDrawSplineTool, EnableDrawSurfaceTool};
+}
+
+TArray<TSharedPtr<FUICommandInfo>> FPCGEditorModeToolCommands::GetPaintContextToolCommands() const
+{
+	return {EnablePaintTool};
+}
+
+TArray<TSharedPtr<FUICommandInfo>> FPCGEditorModeToolCommands::GetVolumeContextToolCommands() const
+{
+	return {EnableCreateVolumeTool};
+}
+
+/**
+ * Macro for defining a constructor and a helper function to retrieve 
+ * @param ToolCommandClassName The tool's class name
+ * @param ToolNameLabel A name string for the command context
+ * @param ToolClassName The class name of the tool
+ */
+#define PCG_DEFINE_TOOL_COMMANDS(ToolCommandClassName, Label, Name, ToolClassName )\
+ToolCommandClassName::ToolCommandClassName()\
+	: TInteractiveToolCommands<ToolCommandClassName>(\
+		Name,\
+		NSLOCTEXT("PCGEditorModeTools", "PCG" Label "ToolCommands", "PCG Editor Tools - '" Name "' Tool"),\
+		NAME_None,\
+		FPCGEditorModeStyle::Get().GetStyleSetName()) {}\
+void ToolCommandClassName::GetToolDefaultObjectList(TArray<UInteractiveTool*>& ToolCDOs)\
+{\
+ToolCDOs.Add(GetMutableDefault<ToolClassName>());\
+}
+PCG_DEFINE_TOOL_COMMANDS(FPCGDrawSplineCommands, "DrawSpline", "Draw Spline", UPCGDrawSplineTool);
+PCG_DEFINE_TOOL_COMMANDS(FPCGDrawSurfaceCommands, "DrawSurface", "Draw Surface", UPCGDrawSplineTool);
+PCG_DEFINE_TOOL_COMMANDS(FPCGPaintCommands, "DrawSurface", "Draw Surface", UPCGPaintTool);
+PCG_DEFINE_TOOL_COMMANDS(FPCGCreateVolumeCommands, "CreateVolume", "Create Volume", UPCGCreateVolumeTool);
+
+#undef PCG_DEFINE_TOOL_COMMANDS
+
+#undef LOCTEXT_NAMESPACE
diff --git a/Engine/Plugins/PCG/Source/PCGEditor/Private/EditorMode/PCGEdModeStyle.cpp b/Engine/Plugins/PCG/Source/PCGEditor/Private/EditorMode/PCGEdModeStyle.cpp
new file mode 100644
--- /dev/null
+++ b/Engine/Plugins/PCG/Source/PCGEditor/Private/EditorMode/PCGEdModeStyle.cpp
@@ -0,0 +1,108 @@
+// Copyright Epic Games, Inc. All Rights Reserved.
+
+#include "EditorMode/PCGEdModeStyle.h"
+#include "SlateOptMacros.h"
+#include "Brushes/SlateImageBrush.h"
+#include "Interfaces/IPluginManager.h"
+#include "Misc/Paths.h"
+#include "Styling/CoreStyle.h"
+#include "Styling/SlateStyle.h"
+#include "Styling/SlateStyleMacros.h"
+#include "Styling/SlateStyleRegistry.h"
+#include "Styling/StarshipCoreStyle.h"
+#include "Styling/StyleColors.h"
+#include "Styling/ToolBarStyle.h"
+#include "Util/ColorConstants.h"
+
+#include "Textures/SlateIcon.h"
+
+namespace PCGEdModeStyle::Constants
+{
+	constexpr TCHAR StyleSet[] = TEXT("PCGEditorModeStyle");
+	constexpr TCHAR DefaultBrushStyle[] = TEXT("PCGEditorMode.Warning");
+}
+
+void FPCGEditorModeStyle::Register()
+{
+	FSlateStyleRegistry::RegisterSlateStyle(Get());
+}
+
+void FPCGEditorModeStyle::Unregister()
+{
+	FSlateStyleRegistry::UnRegisterSlateStyle(Get());
+}
+
+FPCGEditorModeStyle::FPCGEditorModeStyle()
+	: FSlateStyleSet(PCGEdModeStyle::Constants::StyleSet)
+{
+	using namespace PCGEdModeStyle::Constants;
+	// Const icon sizes
+	const FVector2D Icon16x16(16.0f, 16.0f);
+	const FVector2D Icon20x20(20.0f, 20.0f);
+	const FVector2D Icon40x40(40.0f, 40.0f);
+
+	// Default back to CoreStyle
+	SetParentStyleName("CoreStyle");
+
+	FSlateStyleSet::SetContentRoot(FPaths::EngineDir() / TEXT("Plugins/PCG/Content"));
+	FSlateStyleSet::SetCoreContentRoot(FPaths::EngineContentDir() / TEXT("Slate"));
+
+	const FTextBlockStyle& NormalText = FAppStyle::Get().GetWidgetStyle<FTextBlockStyle>("NormalText");
+
+	// General
+	{
+		// Set a default brush to indicate the brush is missing
+		DefaultBrush = new IMAGE_BRUSH_SVG("Icons/EdMode/Warning", Icon40x40);
+		Set(DefaultBrushStyle, DefaultBrush);
+
+		// @todo_pcg: Get a 20x20 version
+		Set("PCGEditorModeIcon", new IMAGE_BRUSH_SVG("Icons/PCG_64", Icon20x20));
+	}
+
+	// Editor Mode Contexts
+	{
+		Set("PCGEditorMode.Context.Line", new IMAGE_BRUSH_SVG("Icons/EdMode/DrawSpline", Icon20x20));
+		Set("PCGEditorMode.Context.Paint", new IMAGE_BRUSH_SVG("Icons/EdMode/Paint", Icon20x20));
+		Set("PCGEditorMode.Context.Volume", new IMAGE_BRUSH_SVG("Icons/EdMode/Volume", Icon20x20));
+	}
+
+	// Tools
+	{
+		Set("PCGEditorMode.Tools.DrawSpline", new IMAGE_BRUSH_SVG("Icons/EdMode/DrawSpline", Icon20x20));
+		Set("PCGEditorMode.Tools.DrawSurface", new IMAGE_BRUSH_SVG("Icons/EdMode/DrawSurface", Icon20x20));
+
+		Set("PCGEditorMode.Tools.Paint", new IMAGE_BRUSH_SVG("Icons/EdMode/Paint", Icon20x20));
+
+		Set("PCGEditorMode.Tools.CreateVolume", new IMAGE_BRUSH_SVG("Icons/EdMode/Volume", Icon20x20));
+	}
+
+	// Backdrops and backgrounds
+	{
+		FLinearColor ViewportOverlayBackgroundColor = FStyleColors::Foreground.GetSpecifiedColor();
+		ViewportOverlayBackgroundColor.A = 0.1f;
+		Set("PCGEditorMode.OverlayBackgroundBrush", new FSlateRoundedBoxBrush(ViewportOverlayBackgroundColor, 8.0f, FStyleColors::Dropdown, 0.0));
+
+		const FLinearColor ContextPaletteColor = FStyleColors::Dropdown.GetSpecifiedColor();
+		Set("PCGEditorMode.ContextPaletteBrush",
+			new FSlateRoundedBoxBrush(ContextPaletteColor, FVector4(8.0f, 0.0f, 0.0f, 8.0f), FStyleColors::Dropdown, 0.0));
+
+		const FLinearColor ContextToolsPaletteColor = FStyleColors::Panel.GetSpecifiedColor();
+		Set("PCGEditorMode.ContextToolsPaletteBrush",
+			new FSlateRoundedBoxBrush(ContextToolsPaletteColor, FVector4(0.0f, 8.0f, 8.0f, 0.0f), FStyleColors::Dropdown, 0.0));
+	}
+
+	// Button styles
+	{
+		FCheckBoxStyle ToggleButtonStyle = FAppStyle::Get().GetWidgetStyle<FCheckBoxStyle>("ToggleButtonCheckbox");
+		ToggleButtonStyle.SetCheckedPressedImage(FSlateRoundedBoxBrush(FStyleColors::Primary, 4.0f));
+		ToggleButtonStyle.SetPadding(FMargin(4.0f, 2.0f));
+		Set("PaletteToggleButton", ToggleButtonStyle);
+	}
+}
+
+const FPCGEditorModeStyle& FPCGEditorModeStyle::Get()
+{
+	static FPCGEditorModeStyle Instance;
+	return Instance;
+}
+
diff --git a/Engine/Plugins/PCG/Source/PCGEditor/Private/EditorMode/PCGEdModeToolkit.cpp b/Engine/Plugins/PCG/Source/PCGEditor/Private/EditorMode/PCGEdModeToolkit.cpp
new file mode 100644
--- /dev/null
+++ b/Engine/Plugins/PCG/Source/PCGEditor/Private/EditorMode/PCGEdModeToolkit.cpp
@@ -0,0 +1,399 @@
+// Copyright Epic Games, Inc. All Rights Reserved.
+
+#include "EditorMode/PCGEdModeToolkit.h"
+
+#include "EditorMode/PCGEdModeCommands.h"
+#include "EditorMode/PCGEdModeSettings.h"
+#include "EditorMode/Tools/PCGInteractiveTool.h"
+
+#include "EditorModeManager.h"
+#include "LevelEditorViewport.h"
+#include "PropertyEditorModule.h"
+#include "Toolkits/AssetEditorModeUILayer.h"
+#include "Toolkits/BaseToolkit.h"
+#include "Tools/EdModeInteractiveToolsContext.h"
+#include "Tools/UEdMode.h"
+
+#include "SPrimaryButton.h"
+#include "Widgets/Input/STextComboBox.h"
+
+#define LOCTEXT_NAMESPACE "PCGEditorModeToolkit"
+
+void FPCGEditorModeToolkit::Init(const TSharedPtr<IToolkitHost>& InitToolkitHost, const TWeakObjectPtr<UEdMode> OwningMode)
+{
+	// @todo_pcg: To be removed with UI overhaul (shared with Modeling Tools) in the future
+	bUsesToolkitBuilder = true;
+
+	FModeToolkit::Init(InitToolkitHost, OwningMode);
+
+	RegisterToolkitCallbacks();
+
+	RegisterPalettes();
+
+	CreateWidgets();
+
+	ClearToolNotification();
+	ClearToolWarning();
+
+	SAssignNew(ToolkitWidget, SBorder)
+	.HAlign(HAlign_Fill)
+	.Padding(0)
+	.BorderImage(FAppStyle::GetBrush("ToolPanel.GroupBorder"))
+	[
+		ToolkitBuilder->GenerateWidget()->AsShared()
+	];
+
+	ShowModeWarnings();
+}
+
+FName FPCGEditorModeToolkit::GetToolkitFName() const
+{
+	return FName("PCGEditorMode");
+}
+
+FText FPCGEditorModeToolkit::GetBaseToolkitName() const
+{
+	return LOCTEXT("DisplayName", "PCGEditorMode Tool");
+}
+
+void FPCGEditorModeToolkit::GetToolPaletteNames(TArray<FName>& PaletteNames) const
+{
+	PaletteNames.Emplace("PCG Toolkit");
+}
+
+FText FPCGEditorModeToolkit::GetToolPaletteDisplayName(const FName PaletteName) const
+{
+	return FText::FromName(PaletteName);
+}
+
+void FPCGEditorModeToolkit::OnToolPaletteChanged(const FName PaletteName)
+{
+	// There is currently only one palette
+	ensure(false);
+}
+
+FText FPCGEditorModeToolkit::GetActiveToolDisplayName() const
+{
+	return CachedActiveToolInfo.ToolDisplayName;
+}
+
+FText FPCGEditorModeToolkit::GetActiveToolMessage() const
+{
+	return CachedActiveToolInfo.ToolDisplayMessage;
+}
+
+void FPCGEditorModeToolkit::OnToolStarted(UInteractiveToolManager* Manager, UInteractiveTool* Tool)
+{
+	UpdateActiveTool();
+
+	if (IsHosted())
+	{
+		GetToolkitHost()->AddViewportOverlayWidget(AcceptToolWidget.ToSharedRef());
+		// Force a refresh of the tool widget with the new tool properties.
+		AcceptToolWidget->Invalidate(EInvalidateWidgetReason::Volatility);
+	}
+}
+
+void FPCGEditorModeToolkit::OnToolEnded(UInteractiveToolManager* Manager, UInteractiveTool* Tool)
+{
+	if (IsHosted())
+	{
+		GetToolkitHost()->RemoveViewportOverlayWidget(AcceptToolWidget.ToSharedRef());
+	}
+
+	ModeDetailsView->SetObject(nullptr);
+
+	ToolkitBuilder->SetActiveToolDisplayName(FText::GetEmpty());
+	ToolkitBuilder->SetActivePaletteCommandsVisibility(EVisibility::Visible);
+
+	ClearToolNotification();
+	ClearToolWarning();
+
+	if (UPCGInteractiveTool* CurrentTool = CastChecked<UPCGInteractiveTool>(GetScriptableEditorMode()->GetToolManager(EToolsContextScope::EdMode)->GetActiveTool(EToolSide::Left), ECastCheckedType::NullAllowed))
+	{
+		UnregisterToolCallbacks(CurrentTool);
+	}
+}
+
+void FPCGEditorModeToolkit::PostToolNotification(const FText& Message)
+{
+	ClearToolNotification();
+
+	UStatusBarSubsystem* StatusBarSubsystem = GEditor->GetEditorSubsystem<UStatusBarSubsystem>();
+	const TSharedPtr<FAssetEditorModeUILayer> ModeUILayerPtr = ModeUILayer.Pin();
+	if (StatusBarSubsystem && ModeUILayerPtr)
+	{
+		ActiveToolMessageHandle = StatusBarSubsystem->PushStatusBarMessage(ModeUILayerPtr->GetStatusBarName(), Message);
+	}
+}
+
+void FPCGEditorModeToolkit::ClearToolNotification()
+{
+	UStatusBarSubsystem* StatusBarSubsystem = GEditor->GetEditorSubsystem<UStatusBarSubsystem>();
+	const TSharedPtr<FAssetEditorModeUILayer> ModeUILayerPtr = ModeUILayer.Pin();
+	if (StatusBarSubsystem && ModeUILayerPtr)
+	{
+		StatusBarSubsystem->PopStatusBarMessage(ModeUILayerPtr->GetStatusBarName(), ActiveToolMessageHandle);
+	}
+
+	ActiveToolMessageHandle.Reset();
+}
+
+void FPCGEditorModeToolkit::PostToolWarning(const FText& Message)
+{
+	check(ToolkitSections && ToolkitSections->ToolWarningArea);
+	ToolkitSections->ToolWarningArea->SetText(Message);
+	ToolkitSections->ToolWarningArea->SetVisibility(EVisibility::Visible);
+}
+
+void FPCGEditorModeToolkit::ClearToolWarning()
+{
+	check(ToolkitSections && ToolkitSections->ToolWarningArea);
+	ToolkitSections->ToolWarningArea->SetText(FText());
+	ToolkitSections->ToolWarningArea->SetVisibility(EVisibility::Collapsed);
+}
+
+void FPCGEditorModeToolkit::ShowModeWarnings()
+{
+	FText WarningText{};
+	if (GEditor->bIsSimulatingInEditor)
+	{
+		WarningText = LOCTEXT("SimulatingWarning", "Cannot use PCG Editor Mode while simulating.");
+	}
+	else if (GEditor->PlayWorld != nullptr)
+	{
+		WarningText = LOCTEXT("PIEWarning", "Cannot use PCG Editor Mode in PIE.");
+	}
+
+	if (!WarningText.IdenticalTo(ActiveWarning))
+	{
+		ActiveWarning = WarningText;
+		ToolkitSections->ModeWarningArea->SetVisibility(ActiveWarning.IsEmpty() ? EVisibility::Collapsed : EVisibility::Visible);
+		ToolkitSections->ModeWarningArea->SetText(ActiveWarning);
+	}
+}
+
+void FPCGEditorModeToolkit::OnActiveViewportChanged(TSharedPtr<IAssetViewport> OldViewport, TSharedPtr<IAssetViewport> NewViewport) const
+{
+	// @todo_pcg: Needs to be confirmed in the long run once the toolkit is fleshed out. For now, just cancel the tool.
+	if (ToolkitBuilder->HasActivePalette())
+	{
+		GetScriptableEditorMode()->GetInteractiveToolsContext()->EndTool(EToolShutdownType::Cancel);
+	}
+}
+
+const FSlateBrush* FPCGEditorModeToolkit::GetActiveToolIcon() const
+{
+	return CachedActiveToolInfo.ToolIcon;
+}
+
+bool FPCGEditorModeToolkit::IsAcceptButtonEnabled() const
+{
+	return GetScriptableEditorMode()->GetInteractiveToolsContext()->CanAcceptActiveTool();
+}
+
+bool FPCGEditorModeToolkit::IsCancelButtonEnabled() const
+{
+	return GetScriptableEditorMode()->GetInteractiveToolsContext()->CanCancelActiveTool();
+}
+
+EVisibility FPCGEditorModeToolkit::GetAcceptButtonVisibility() const
+{
+	return GetScriptableEditorMode()->GetInteractiveToolsContext()->ActiveToolHasAccept() ? EVisibility::Visible : EVisibility::Collapsed;
+}
+
+EVisibility FPCGEditorModeToolkit::GetCancelButtonVisibility() const
+{
+	return GetScriptableEditorMode()->GetInteractiveToolsContext()->ActiveToolHasAccept() ? EVisibility::Visible : EVisibility::Collapsed;
+}
+
+FReply FPCGEditorModeToolkit::OnAcceptButtonClicked() const
+{
+	GetScriptableEditorMode()->GetInteractiveToolsContext()->EndTool(EToolShutdownType::Accept);
+	return FReply::Handled();
+}
+
+FReply FPCGEditorModeToolkit::OnCancelButtonClicked() const
+{
+	GetScriptableEditorMode()->GetInteractiveToolsContext()->EndTool(EToolShutdownType::Cancel);
+	return FReply::Unhandled();
+}
+
+void FPCGEditorModeToolkit::RegisterToolkitCallbacks()
+{
+	// Subscribe to receive settings change updates.
+	UPCGEditorModeSettings* PCGEditorModeSettings = GetMutableDefault<UPCGEditorModeSettings>();
+	PCGEditorModeSettings->OnSettingChanged().AddSP(SharedThis(this), &FPCGEditorModeToolkit::OnEditorModeSettingsChanged);
+
+	GetScriptableEditorMode()->GetInteractiveToolsContext(EToolsContextScope::EdMode)->OnToolNotificationMessage.AddSP(this, &FPCGEditorModeToolkit::PostToolNotification);
+	GetScriptableEditorMode()->GetInteractiveToolsContext(EToolsContextScope::EdMode)->OnToolWarningMessage.AddSP(this, &FPCGEditorModeToolkit::PostToolWarning);
+
+	GetToolkitHost()->OnActiveViewportChanged().AddSP(this, &FPCGEditorModeToolkit::OnActiveViewportChanged);
+
+	
+}
+
+void FPCGEditorModeToolkit::UnregisterToolkitCallbacks()
+{
+	GetToolkitHost()->OnActiveViewportChanged().RemoveAll(this);
+
+	GetScriptableEditorMode()->GetInteractiveToolsContext(EToolsContextScope::EdMode)->OnToolNotificationMessage.RemoveAll(this);
+	GetScriptableEditorMode()->GetInteractiveToolsContext(EToolsContextScope::EdMode)->OnToolWarningMessage.RemoveAll(this);
+
+	UPCGEditorModeSettings* PCGEditorModeSettings = GetMutableDefault<UPCGEditorModeSettings>();
+	PCGEditorModeSettings->OnSettingChanged().RemoveAll(this);
+}
+
+void FPCGEditorModeToolkit::RegisterToolCallbacks(UPCGInteractiveTool* CurrentTool)
+{
+	if (CurrentTool)
+	{
+		CurrentTool->OnPropertySetsModified.AddSP(this, &FPCGEditorModeToolkit::UpdateActiveTool);
+		CurrentTool->OnPropertyModifiedDirectlyByTool.AddSP(this, &FPCGEditorModeToolkit::InvalidateCachedDetailPanelState);
+		CurrentTool->OnToolWarningChanged.AddSP(this, &FPCGEditorModeToolkit::PostToolWarning);
+	}
+}
+
+void FPCGEditorModeToolkit::UnregisterToolCallbacks(UPCGInteractiveTool* CurrentTool) const
+{
+	if (CurrentTool)
+	{
+		CurrentTool->OnPropertySetsModified.RemoveAll(this);
+		CurrentTool->OnPropertyModifiedDirectlyByTool.RemoveAll(this);
+		CurrentTool->OnToolWarningChanged.RemoveAll(this);
+	}
+}
+
+void FPCGEditorModeToolkit::OnEditorModeSettingsChanged(UObject* InObject, FPropertyChangedEvent& InChangedEvent)
+{
+	// @todo_pcg: A safety precaution, but this may not be needed.
+	if (InChangedEvent.GetMemberPropertyName() == GET_MEMBER_NAME_CHECKED(UPCGEditorModeSettings, bEnableEditorMode))
+	{
+		GetScriptableEditorMode()->GetToolManager(EToolsContextScope::EdMode)->DeactivateTool(EToolSide::Left, EToolShutdownType::Cancel);
+	}
+}
+
+void FPCGEditorModeToolkit::UpdateActiveTool()
+{
+	UInteractiveTool* CurrentTool = GetScriptableEditorMode()->GetToolManager(EToolsContextScope::EdMode)->GetActiveTool(EToolSide::Left);
+	if (UPCGInteractiveTool* CurrentPCGTool = CastChecked<UPCGInteractiveTool>(CurrentTool, ECastCheckedType::NullAllowed))
+	{
+		CachedActiveToolInfo = CurrentPCGTool->GetToolInfo();
+
+		ToolkitBuilder->SetActiveToolDisplayName(CachedActiveToolInfo.ToolDisplayName);
+
+		if (const UPCGEditorModeSettings* Settings = GetDefault<UPCGEditorModeSettings>())
+		{
+			if (Settings->bHideToolButtonsDuringActiveTool)
+			{
+				ToolkitBuilder->SetActivePaletteCommandsVisibility(EVisibility::Collapsed);
+			}
+		}
+
+		// Before actually changing the detail panel, we need to see where the current keyboard focus is, because
+		// if it's inside the detail panel, we'll need to reset it to the detail panel as a whole, else we might
+		// lose it entirely when that detail panel element gets destroyed (which would make us unable to receive any
+		// hotkey presses until the user clicks somewhere).
+		const TSharedPtr<SWidget> FocusedWidget = FSlateApplication::Get().GetKeyboardFocusedWidget();
+		if (FocusedWidget != ModeDetailsView)
+		{
+			// Search upward from the currently focused widget
+			TSharedPtr<SWidget> CurrentWidget = FocusedWidget;
+			while (CurrentWidget.IsValid())
+			{
+				if (CurrentWidget == ModeDetailsView)
+				{
+					// Reset focus to the detail panel as a whole to avoid losing it when the inner elements change.
+					FSlateApplication::Get().SetKeyboardFocus(ModeDetailsView);
+					break;
+				}
+
+				CurrentWidget = CurrentWidget->GetParentWidget();
+			}
+		}
+
+		ModeDetailsView->SetObjects({CurrentPCGTool->GetToolProperties(true)});
+		RegisterToolCallbacks(CurrentPCGTool);
+	}
+
+	// @todo_pcg: If a tool property needs to invalidate the cached details panel, do it here. 
+	// InvalidateCachedDetailPanelState(nullptr);
+}
+
+void FPCGEditorModeToolkit::InvalidateCachedDetailPanelState(UObject* ChangedObject) const
+{
+	ModeDetailsView->InvalidateCachedState();
+}
+
+void FPCGEditorModeToolkit::RegisterPalettes()
+{
+	const FPCGEditorModePaletteCommands& PaletteCommands = FPCGEditorModePaletteCommands::Get();
+	const FPCGEditorModeToolCommands& ToolCommands = FPCGEditorModeToolCommands::Get();
+
+	ToolkitSections = MakeShared<FToolkitSections>();
+	ToolkitSections->DetailsView = ModeDetailsView;
+
+	FToolkitBuilderArgs ToolkitBuilderArgs(GetScriptableEditorMode()->GetModeInfo().ToolbarCustomizationName);
+	ToolkitBuilderArgs.ToolkitCommandList = GetToolkitCommands();
+	ToolkitBuilderArgs.ToolkitSections = ToolkitSections;
+	ToolkitBuilderArgs.SelectedCategoryTitleVisibility = EVisibility::Collapsed;
+
+	ToolkitBuilder = MakeShared<FToolkitBuilder>(ToolkitBuilderArgs);
+
+	const TSharedPtr<FToolPalette> LineToolPalette = MakeShared<FToolPalette>(PaletteCommands.LoadLineContextPalette, ToolCommands.GetLineContextToolCommands());
+	ToolkitBuilder->AddPalette(std::move(LineToolPalette));
+
+	const TSharedPtr<FToolPalette> PaintToolPalette = MakeShared<FToolPalette>(PaletteCommands.LoadPaintContextPalette, ToolCommands.GetPaintContextToolCommands());
+	ToolkitBuilder->AddPalette(std::move(PaintToolPalette));
+
+	const TSharedPtr<FToolPalette> VolumeToolPalette = MakeShared<FToolPalette>(PaletteCommands.LoadVolumeContextPalette, ToolCommands.GetVolumeContextToolCommands());
+	ToolkitBuilder->AddPalette(std::move(VolumeToolPalette));
+
+	ToolkitBuilder->SetActivePaletteOnLoad(PaletteCommands.LoadLineContextPalette.Get());
+
+	ToolkitBuilder->UpdateWidget();
+
+	// If selected palette changes, make sure we are showing the palette command buttons, which may be hidden by active Tool.
+	ActivePaletteChangedHandle = ToolkitBuilder->OnActivePaletteChanged.AddLambda([this]()
+	{
+		ToolkitBuilder->SetActivePaletteCommandsVisibility(EVisibility::Visible);
+	});
+}
+
+void FPCGEditorModeToolkit::CreateWidgets()
+{
+	SAssignNew(ToolkitSections->ModeWarningArea, STextBlock)
+	.Text(FText::GetEmpty())
+	.AutoWrapText(true)
+	.Font(FCoreStyle::GetDefaultFontStyle("Bold", 9))
+	.ColorAndOpacity(FSlateColor(FLinearColor(0.9f, 0.15f, 0.15f)))
+	.Visibility(EVisibility::Collapsed);
+
+	SAssignNew(ToolkitSections->ToolWarningArea, STextBlock)
+	.Text(FText::GetEmpty())
+	.AutoWrapText(true)
+	.Font(FCoreStyle::GetDefaultFontStyle("Bold", 9))
+	.ColorAndOpacity(FSlateColor(FLinearColor(0.9f, 0.15f, 0.15f)));
+
+	const TSharedPtr<SVerticalBox> ToolkitWidgetVBox = SNew(SVerticalBox);
+	ToolkitWidgetVBox->AddSlot().AutoHeight().HAlign(HAlign_Fill).Padding(5)
+	[
+		ToolkitSections->ModeWarningArea->AsShared()
+	];
+
+	SAssignNew(AcceptToolWidget, SPCGEdModeAcceptToolOverlay)
+	.ActiveToolDisplayName(this, &FPCGEditorModeToolkit::GetActiveToolDisplayName)
+	.ActiveToolIcon(this, &FPCGEditorModeToolkit::GetActiveToolIcon)
+	.AcceptButtonLabel(LOCTEXT("AcceptButtonLabel", "Accept"))
+	.CancelButtonLabel(LOCTEXT("CancelButtonLabel", "Cancel"))
+	.AcceptButtonTooltip(LOCTEXT("AcceptButtonTooltip", "Confirm the tool operation and apply the result."))
+	.CancelButtonTooltip(LOCTEXT("CancelButtonTooltip", "Cancel the current tool operation."))
+	.IsAcceptButtonEnabled(this, &FPCGEditorModeToolkit::IsAcceptButtonEnabled)
+	.IsCancelButtonEnabled(this, &FPCGEditorModeToolkit::IsCancelButtonEnabled)
+	.GetAcceptButtonVisibility(this, &FPCGEditorModeToolkit::GetAcceptButtonVisibility)
+	.GetCancelButtonVisibility(this, &FPCGEditorModeToolkit::GetCancelButtonVisibility)
+	.OnAcceptButtonClicked(this, &FPCGEditorModeToolkit::OnAcceptButtonClicked)
+	.OnCancelButtonClicked(this, &FPCGEditorModeToolkit::OnCancelButtonClicked);
+}
+
+#undef LOCTEXT_NAMESPACE
diff --git a/Engine/Plugins/PCG/Source/PCGEditor/Private/EditorMode/Tools/Helpers/PCGEdModeActorHelpers.cpp b/Engine/Plugins/PCG/Source/PCGEditor/Private/EditorMode/Tools/Helpers/PCGEdModeActorHelpers.cpp
new file mode 100644
--- /dev/null
+++ b/Engine/Plugins/PCG/Source/PCGEditor/Private/EditorMode/Tools/Helpers/PCGEdModeActorHelpers.cpp
@@ -0,0 +1,79 @@
+// Copyright Epic Games, Inc. All Rights Reserved.
+
+#pragma once
+
+#include "EditorMode/Tools/Helpers/PCGEdModeActorHelpers.h"
+
+#include "PCGComponent.h"
+#include "ActorFactories/ActorFactoryEmptyActor.h"
+#include "Helpers/PCGActorHelpers.h"
+
+#define LOCTEXT_NAMESPACE "PCGEditorMode"
+
+namespace PCG::EditorMode
+{
+	namespace Actor
+	{
+		static const FLazyName DefaultWorkingActorName = "PCG Tool Working Actor";
+
+		AActor* Spawn(UWorld* World, const FName ActorName, const FTransform& Transform)
+		{
+			if (!ensure(World && !ActorName.IsNone()))
+			{
+				return nullptr;
+			}
+
+			FActorSpawnParameters SpawnParams;
+			SpawnParams.Name = ActorName;
+			SpawnParams.NameMode = FActorSpawnParameters::ESpawnActorNameMode::Requested;
+
+			UPCGActorHelpers::FSpawnDefaultActorParams SpawnDefaultActorParams(World, AActor::StaticClass(), Transform, SpawnParams);
+			SpawnDefaultActorParams.bForceStaticMobility = false;
+
+			return UPCGActorHelpers::SpawnDefaultActor(SpawnDefaultActorParams);
+		}
+
+		AActor* SpawnWorking(UWorld* World, const FTransform& Transform)
+		{
+			FActorSpawnParameters SpawnParams;
+			SpawnParams.Name = DefaultWorkingActorName;
+			SpawnParams.NameMode = FActorSpawnParameters::ESpawnActorNameMode::Requested;
+			SpawnParams.ObjectFlags = RF_Transient;
+
+			UPCGActorHelpers::FSpawnDefaultActorParams SpawnDefaultActorParams(World, AActor::StaticClass(), Transform, SpawnParams);
+
+			return UPCGActorHelpers::SpawnDefaultActor(SpawnDefaultActorParams);
+		}
+	} // namespace PCG::EditorMode::Actor
+
+	namespace PCGComponent
+	{
+		UPCGComponent* Find(const AActor* Owner)
+		{
+			if (!ensure(Owner))
+			{
+				return nullptr;
+			}
+
+			return Owner->FindComponentByClass<UPCGComponent>();
+		}
+
+		UPCGComponent* Create(AActor* Owner, const FName ComponentName)
+		{
+			if (!ensure(Owner && !ComponentName.IsNone()))
+			{
+				return nullptr;
+			}
+
+			UPCGComponent* PCGComponent = NewObject<UPCGComponent>(Owner, ComponentName);
+			PCGComponent->RegisterComponent();
+			Owner->AddInstanceComponent(PCGComponent);
+			return PCGComponent;
+		}
+
+		UPCGComponent* CreateWorking(AActor* Owner)
+		{
+			return Create(Owner, Actor::DefaultWorkingActorName);
+		}
+	} // namespace PCG::EditorMode::PCGComponent
+} // namespace PCG::EditorMode
diff --git a/Engine/Plugins/PCG/Source/PCGEditor/Private/EditorMode/Tools/Helpers/PCGEdModeSceneQueryHelpers.cpp b/Engine/Plugins/PCG/Source/PCGEditor/Private/EditorMode/Tools/Helpers/PCGEdModeSceneQueryHelpers.cpp
new file mode 100644
--- /dev/null
+++ b/Engine/Plugins/PCG/Source/PCGEditor/Private/EditorMode/Tools/Helpers/PCGEdModeSceneQueryHelpers.cpp
@@ -0,0 +1,114 @@
+// Copyright Epic Games, Inc. All Rights Reserved.
+
+#pragma once
+
+#include "EditorMode/Tools/Helpers/PCGEdModeSceneQueryHelpers.h"
+
+#include "CollisionQueryParams.h"
+#include "BaseGizmos/GizmoMath.h"
+#include "Components/PrimitiveComponent.h"
+#include "Engine/HitResult.h"
+#include "Engine/World.h"
+
+namespace PCG::EditorMode::Scene
+{
+	bool IsHitTargetVisible(const FHitResult& HitResult)
+	{
+		if (const AActor* Actor = HitResult.GetActor())
+		{
+			if (Actor->IsHiddenEd())
+			{
+				return false;
+			}
+		}
+
+		if (const UPrimitiveComponent* Component = HitResult.GetComponent())
+		{
+			if (!Component->IsVisibleInEditor())
+			{
+				return false;
+			}
+		}
+
+		return true;
+	}
+
+	TOptional<FHitResult> TraceToNearestObject(const FInputDeviceRay& Ray, const FViewRayParams& Params)
+	{
+		if (!Params.World)
+		{
+			return {};
+		}
+
+		const FCollisionObjectQueryParams ObjectQueryParams(FCollisionObjectQueryParams::AllObjects);
+		FCollisionQueryParams QueryParams = FCollisionQueryParams::DefaultQueryParam;
+		QueryParams.bTraceComplex = true;
+
+		TArray<FHitResult> OutHits;
+		const FVector& Origin = Ray.WorldRay.Origin;
+		const FVector End = Origin + Ray.WorldRay.Direction * Params.Distance;
+		if (!Params.World->LineTraceMultiByObjectType(OutHits, Origin, End, ObjectQueryParams, QueryParams))
+		{
+			return {};
+		}
+
+		for (const FHitResult& Hit : OutHits)
+		{
+			// Check for a visible result
+			bool bHitIsValid = IsHitTargetVisible(Hit);
+			// If invisible, check if it's explicitly allowed.
+			bHitIsValid |= Params.AllowedInvisibleComponents && Params.AllowedInvisibleComponents->Contains(Hit.GetComponent());
+			// Check that it's not explicitly ignored.
+			bHitIsValid &= !Params.ComponentsToIgnore || !Params.ComponentsToIgnore->Contains(Hit.GetComponent());
+
+			if (bHitIsValid)
+			{
+				return Hit;
+			}
+		}
+
+		return {};
+	}
+
+	TOptional<FViewHitResult> ViewportRay(const FInputDeviceRay& Ray, const FViewRayParams& Params)
+	{
+		if (TOptional<FHitResult> Result = TraceToNearestObject(Ray, Params))
+		{
+			return FViewHitResult(Result->ImpactPoint, Result->ImpactNormal, Result->Distance);
+		}
+		else
+		{
+			return {};
+		}
+	}
+
+	TOptional<FViewHitResult> ViewportRayAgainstPlane(const FInputDeviceRay& Ray, const FViewRayParams& Params, const FPlane& Plane)
+	{
+		bool bHitPlane = false;
+		FVector IntersectionPoint;
+		GizmoMath::RayPlaneIntersectionPoint(Plane.GetOrigin(), Plane.GetNormal(), Ray.WorldRay.Origin, Ray.WorldRay.Direction, bHitPlane, IntersectionPoint);
+
+		if (const TOptional<FViewHitResult> HitResult = ViewportRay(Ray, Params))
+		{
+			// If the plane wasn't hit or the hit result was closer, return it.
+			if (!bHitPlane || FVector::DistSquared(Ray.WorldRay.Origin, HitResult->ImpactPosition) < FVector::DistSquared(Ray.WorldRay.Origin, IntersectionPoint))
+			{
+				return HitResult;
+			}
+		}
+
+		if (bHitPlane)
+		{
+			return FViewHitResult(IntersectionPoint, Plane.GetNormal(), FVector::Distance(Ray.WorldRay.Origin, IntersectionPoint));
+		}
+
+		return {};
+	}
+
+	TOptional<FViewHitResult> ViewportRayAgainstCameraPlane(const FInputDeviceRay& Ray, const FViewRayParams& Params, const FViewCameraState& CameraState)
+	{
+		const FVector PlaneOrigin = CameraState.Position + CameraState.Forward() * Params.Distance;
+		const FVector PlaneNormal = CameraState.bIsOrthographic ? -Ray.WorldRay.Direction : FVector::UnitZ();
+		return ViewportRayAgainstPlane(Ray, Params, FPlane(PlaneOrigin, PlaneNormal));
+	}
+} // namespace PCG::EditorMode::Scene
diff --git a/Engine/Plugins/PCG/Source/PCGEditor/Private/EditorMode/Tools/Helpers/PCGEdModeSceneVisualizeHelpers.cpp b/Engine/Plugins/PCG/Source/PCGEditor/Private/EditorMode/Tools/Helpers/PCGEdModeSceneVisualizeHelpers.cpp
new file mode 100644
--- /dev/null
+++ b/Engine/Plugins/PCG/Source/PCGEditor/Private/EditorMode/Tools/Helpers/PCGEdModeSceneVisualizeHelpers.cpp
@@ -0,0 +1,160 @@
+// Copyright Epic Games, Inc. All Rights Reserved.
+
+#pragma once
+
+#include "EditorMode/Tools/Helpers/PCGEdModeSceneVisualizeHelpers.h"
+
+#include "PrimitiveDrawInterface.h"
+#include "ToolContextInterfaces.h"
+#include "Components/SplineComponent.h"
+#include "Settings/LevelEditorViewportSettings.h"
+
+namespace PCG::EditorMode::Scene::Visualize
+{
+	// Taken in large directly from SplineComponentVisualizer::DrawVisualization in SplineComponentVisualizer.h
+	void Spline(const USplineComponent& Spline, IToolsContextRenderAPI& RenderAPI, int32 LatestIndex, bool bShouldVisualizeScale, bool bShouldVisualizeTangents)
+	{
+		FPrimitiveDrawInterface* PDI = RenderAPI.GetPrimitiveDrawInterface();
+		const FInterpCurveVector& SplineInfo = Spline.GetSplinePointsPosition();
+
+		const bool bIsSplineEditable = !Spline.bModifiedByConstructionScript;
+
+		// @todo_pcg: Special color for read-only splines
+		const FColor NormalColor = FColor::White;
+		const FColor SelectedColor = FColor::White;
+		const FColor TangentColor = FColor::White;
+		const float PointSize = 10.0f + (bIsSplineEditable ? GetDefault<ULevelEditorViewportSettings>()->SelectedSplinePointSizeAdjustment : 0.0f);
+
+		const float VisualizationWidthScale = Spline.ScaleVisualizationWidth;
+		const float TangentScale = GetDefault<ULevelEditorViewportSettings>()->SplineTangentScale;
+
+		FVector PreviousKeyPosition = !SplineInfo.Points.IsEmpty() ? Spline.GetLocationAtSplinePoint(0, ESplineCoordinateSpace::World) : FVector::ZeroVector;
+		FVector PreviousKeyRightVector = !SplineInfo.Points.IsEmpty() ? Spline.GetRightVectorAtSplinePoint(0, ESplineCoordinateSpace::World) : FVector::ZeroVector;
+		FVector PreviousKeyScale = !SplineInfo.Points.IsEmpty() ? Spline.GetUpVectorAtSplinePoint(0, ESplineCoordinateSpace::World) : FVector::ZeroVector;
+
+		const int32 NumPoints = SplineInfo.Points.Num();
+		const int32 NumSegments = SplineInfo.bIsLooped ? NumPoints : NumPoints - 1;
+		for (int32 KeyIndex = 0; KeyIndex < NumSegments + 1; ++KeyIndex)
+		{
+			const FVector KeyPosition = Spline.GetLocationAtSplinePoint(KeyIndex, ESplineCoordinateSpace::World);
+			const FVector KeyRightVector = Spline.GetRightVectorAtSplinePoint(KeyIndex, ESplineCoordinateSpace::World);
+			const FVector KeyUpVector = Spline.GetUpVectorAtSplinePoint(KeyIndex, ESplineCoordinateSpace::World);
+			const FVector KeyScale = Spline.GetScaleAtSplinePoint(KeyIndex) * VisualizationWidthScale;
+			const FVector ArriveTangent = Spline.GetArriveTangentAtSplinePoint(KeyIndex, ESplineCoordinateSpace::World) * TangentScale;
+			const FVector LeaveTangent = Spline.GetLeaveTangentAtSplinePoint(KeyIndex, ESplineCoordinateSpace::World) * TangentScale;
+
+			const FColor KeyColor = (KeyIndex == LatestIndex) ? SelectedColor : NormalColor;
+
+			// Draw the keypoint and up/right vectors
+			if (KeyIndex < NumPoints)
+			{
+				// Draws the tangents separately with translucent lines.
+				// @todo_pcg: Separate visualization into arrive and leave tangents
+				if (bShouldVisualizeTangents)
+				{
+					PDI->DrawTranslucentLine(KeyPosition, KeyPosition - ArriveTangent, TangentColor, SDPG_Foreground);
+					PDI->DrawTranslucentLine(KeyPosition, KeyPosition + LeaveTangent, TangentColor, SDPG_Foreground);
+				}
+
+				// If true, will create a thickness of lines around the spline
+				if (bShouldVisualizeScale)
+				{
+					PDI->DrawLine(KeyPosition, KeyPosition - KeyRightVector * KeyScale.Y, KeyColor, SDPG_Foreground);
+					PDI->DrawLine(KeyPosition, KeyPosition + KeyRightVector * KeyScale.Y, KeyColor, SDPG_Foreground);
+					PDI->DrawLine(KeyPosition, KeyPosition + KeyUpVector * KeyScale.Z, KeyColor, SDPG_Foreground);
+
+					// Points used to visualize the radial dimension
+					constexpr int32 ArcPoints = 20;
+					FVector OldArcPos = KeyPosition + KeyRightVector * KeyScale.Y;
+					for (int32 ArcIndex = 1; ArcIndex <= ArcPoints; ++ArcIndex)
+					{
+						float Sin;
+						float Cos;
+						FMath::SinCos(&Sin, &Cos, ArcIndex * PI / ArcPoints);
+						const FVector NewArcPos = KeyPosition + Cos * KeyRightVector * KeyScale.Y + Sin * KeyUpVector * KeyScale.Z;
+						PDI->DrawLine(OldArcPos, NewArcPos, KeyColor, SDPG_Foreground);
+						OldArcPos = NewArcPos;
+					}
+				}
+
+				PDI->DrawPoint(KeyPosition, KeyColor, PointSize, SDPG_Foreground);
+			}
+
+			// If not the first keypoint, draw a line to the previous keypoint.
+			if (KeyIndex > 0)
+			{
+				const FColor LineColor = NormalColor;
+
+				const int32 PreviousIndex = KeyIndex - 1;
+				const FColor BeginColor = (LatestIndex == PreviousIndex) ? SelectedColor : NormalColor;
+				const FColor EndColor = (LatestIndex == KeyIndex) ? SelectedColor : NormalColor;
+
+				// Find position on first keyframe.
+				const FVector& PreviousPosition = PreviousKeyPosition;
+				const FVector& PreviousRightVector = PreviousKeyRightVector;
+				const FVector& PreviousScale = PreviousKeyScale;
+
+				// Then draw a line for each substep.
+				constexpr int32 NumSteps = 20;
+				float SegmentLineThickness = GetDefault<ULevelEditorViewportSettings>()->SplineLineThicknessAdjustment;
+
+				for (int32 StepIndex = 1; StepIndex <= NumSteps; ++StepIndex)
+				{
+					const float StepRatio = static_cast<float>(StepIndex) / static_cast<float>(NumSteps);
+					const float Key = PreviousIndex + StepRatio;
+
+					FColor StepColor;
+					// Grade the lerp to full over 80%, then clamp.
+					constexpr float PartialGradientProportion = 0.8f;
+					const float LerpRatio = (1.0f - StepRatio) / PartialGradientProportion;
+					StepColor = FMath::Lerp(BeginColor.ReinterpretAsLinear(), EndColor.ReinterpretAsLinear(), LerpRatio).ToFColor(false);
+
+					// Draw the segment.
+					PDI->DrawLine(PreviousPosition, KeyPosition, StepColor, SDPG_Foreground, SegmentLineThickness);
+
+					if (bShouldVisualizeScale)
+					{
+						PDI->DrawLine(PreviousPosition - PreviousRightVector * PreviousScale.Y, KeyPosition - KeyRightVector * KeyScale.Y, LineColor, SDPG_Foreground);
+						PDI->DrawLine(PreviousPosition + PreviousRightVector * PreviousScale.Y, KeyPosition + KeyRightVector * KeyScale.Y, LineColor, SDPG_Foreground);
+					}
+
+					// Visualize Up Vector
+					const float HalfVisualizationWidthScale = VisualizationWidthScale * 0.5f;
+					const FVector NewUpVector = Spline.GetUpVectorAtSplineInputKey(Key, ESplineCoordinateSpace::World);
+					PDI->DrawLine(KeyPosition, KeyPosition + NewUpVector * HalfVisualizationWidthScale, LineColor, SDPG_Foreground);
+					PDI->DrawLine(KeyPosition, KeyPosition + KeyRightVector * HalfVisualizationWidthScale, LineColor, SDPG_Foreground);
+
+					PreviousKeyPosition = KeyPosition;
+					PreviousKeyRightVector = KeyRightVector;
+					PreviousKeyScale = KeyScale;
+				}
+			}
+		}
+	}
+
+	void SplineTangent(const USplineComponent& Spline, IToolsContextRenderAPI& RenderAPI, const int32 SplinePointIndex)
+	{
+		if (!ensure(SplinePointIndex >= 0 && SplinePointIndex < Spline.GetNumberOfSplinePoints()))
+		{
+			return;
+		}
+
+		FPrimitiveDrawInterface* PDI = RenderAPI.GetPrimitiveDrawInterface();
+		check(PDI);
+
+		const float TangentScale = FMath::Min(0.5f, GetDefault<ULevelEditorViewportSettings>()->SplineTangentScale);
+		const float TangentHandleSize = 8.0f + GetDefault<ULevelEditorViewportSettings>()->SplineTangentHandleSizeAdjustment;
+
+		const FVector Location = Spline.GetLocationAtSplinePoint(SplinePointIndex, ESplineCoordinateSpace::World);
+		const FVector LeaveTangent = Spline.GetLeaveTangentAtSplinePoint(SplinePointIndex, ESplineCoordinateSpace::World) * TangentScale;
+		const FVector ArriveTangent = Spline.bAllowDiscontinuousSpline ? Spline.GetArriveTangentAtSplinePoint(SplinePointIndex, ESplineCoordinateSpace::World) * TangentScale : LeaveTangent;
+
+		const FColor Color = FColor::White;
+
+		PDI->DrawLine(Location, Location - ArriveTangent, Color, SDPG_Foreground);
+		PDI->DrawLine(Location, Location + LeaveTangent, Color, SDPG_Foreground);
+
+		PDI->DrawPoint(Location + LeaveTangent, Color, TangentHandleSize, SDPG_Foreground);
+		PDI->DrawPoint(Location - ArriveTangent, Color, TangentHandleSize, SDPG_Foreground);
+	}
+}
diff --git a/Engine/Plugins/PCG/Source/PCGEditor/Private/EditorMode/Tools/Helpers/PCGEdModeSplineHelpers.cpp b/Engine/Plugins/PCG/Source/PCGEditor/Private/EditorMode/Tools/Helpers/PCGEdModeSplineHelpers.cpp
new file mode 100644
--- /dev/null
+++ b/Engine/Plugins/PCG/Source/PCGEditor/Private/EditorMode/Tools/Helpers/PCGEdModeSplineHelpers.cpp
@@ -0,0 +1,113 @@
+// Copyright Epic Games, Inc. All Rights Reserved.
+
+#pragma once
+
+#include "EditorMode/Tools/Helpers/PCGEdModeSplineHelpers.h"
+
+#include "Components/SplineComponent.h"
+#include "UObject/UObjectGlobals.h"
+
+#define LOCTEXT_NAMESPACE "PCGEditorMode"
+
+namespace PCG::EditorMode::Spline
+{
+	static const FLazyName DefaultWorkingSplineName = "PCG Tool Working Spline";
+
+	USplineComponent* CreateComponent(
+		AActor* Actor,
+		const FName ComponentName,
+		const USplineComponent* TemplateComponent,
+		const bool bTransactional)
+	{
+		if (!ensure(Actor && !ComponentName.IsNone()))
+		{
+			return nullptr;
+		}
+
+		if (bTransactional)
+		{
+			Actor->Modify();
+		}
+
+		const FName Name = MakeUniqueObjectName(Actor, TemplateComponent ? TemplateComponent->StaticClass() : USplineComponent::StaticClass(), ComponentName);
+
+		USplineComponent* Spline = nullptr;
+		const EObjectFlags Flags = RF_Transient | (bTransactional ? RF_Transactional : RF_NoFlags);
+		if (TemplateComponent)
+		{
+			Spline = DuplicateObject<USplineComponent>(TemplateComponent, Actor, Name);
+			Spline->SetFlags(Flags);
+		}
+		else
+		{
+			Spline = NewObject<USplineComponent>(Actor, USplineComponent::StaticClass(), Name, Flags);
+		}
+
+		Actor->AddInstanceComponent(Spline);
+		Spline->RegisterComponent();
+		Spline->ResetRelativeTransform();
+		Actor->PostEditChange();
+
+		return Spline;
+	}
+
+	USplineComponent* CreateWorkingComponent(AActor* Actor)
+	{
+		return CreateComponent(Actor, DefaultWorkingSplineName, /*TemplateComponent=*/nullptr, /*bTransactional=*/false);
+	}
+
+	USplineComponent* FindComponent(const AActor* Actor, const int32 TargetIndex)
+	{
+		if (!ensure(Actor))
+		{
+			return nullptr;
+		}
+
+		TInlineComponentArray<USplineComponent*> SplineComponents;
+		Actor->GetComponents<USplineComponent>(SplineComponents);
+		if (SplineComponents.IsValidIndex(TargetIndex))
+		{
+			return SplineComponents[TargetIndex];
+		}
+		else if (!SplineComponents.IsEmpty())
+		{
+			return SplineComponents[0];
+		}
+
+		return nullptr;
+	}
+
+	void CopySplinePoints(const USplineComponent& Source, USplineComponent& Destination, const bool bTransactional)
+	{
+		if (bTransactional)
+		{
+			Destination.Modify();
+		}
+
+		Destination.ClearSplinePoints();
+		Destination.bSplineHasBeenEdited = true;
+
+		/**
+		 * Iterate here (rather than just copying over the SplineCurves data) to transform the data properly into the
+		 * coordinate space of the target component.
+		 */
+		const int32 NumSplinePoints = Source.GetNumberOfSplinePoints();
+		for (int32 SplinePointIndex = 0; SplinePointIndex < NumSplinePoints; ++SplinePointIndex)
+		{
+			Destination.AddSplinePoint(Source.GetLocationAtSplinePoint(SplinePointIndex, ESplineCoordinateSpace::World),
+									   ESplineCoordinateSpace::World, false);
+			Destination.SetUpVectorAtSplinePoint(SplinePointIndex, Source.GetUpVectorAtSplinePoint(SplinePointIndex, ESplineCoordinateSpace::World),
+												 ESplineCoordinateSpace::World, false);
+			Destination.SetTangentsAtSplinePoint(SplinePointIndex,
+												 Source.GetArriveTangentAtSplinePoint(SplinePointIndex, ESplineCoordinateSpace::World),
+												 Source.GetLeaveTangentAtSplinePoint(SplinePointIndex, ESplineCoordinateSpace::World),
+												 ESplineCoordinateSpace::World, false);
+			Destination.SetSplinePointType(SplinePointIndex, Source.GetSplinePointType(SplinePointIndex), false);
+		}
+
+		Destination.SetClosedLoop(Source.IsClosedLoop());
+		Destination.UpdateSpline();
+	}
+} // namespace PCG::EditorMode::Spline
+
+#undef LOCTEXT_NAMESPACE
diff --git a/Engine/Plugins/PCG/Source/PCGEditor/Private/EditorMode/Tools/Line/PCGDrawSplineTool.cpp b/Engine/Plugins/PCG/Source/PCGEditor/Private/EditorMode/Tools/Line/PCGDrawSplineTool.cpp
new file mode 100644
--- /dev/null
+++ b/Engine/Plugins/PCG/Source/PCGEditor/Private/EditorMode/Tools/Line/PCGDrawSplineTool.cpp
@@ -0,0 +1,756 @@
+// Copyright Epic Games, Inc. All Rights Reserved.
+
+#include "EditorMode/Tools/Line/PCGDrawSplineTool.h"
+
+#include "PCGComponent.h"
+#include "EditorMode/PCGEdModeSettings.h"
+#include "EditorMode/PCGEdModeStyle.h"
+#include "EditorMode/Tools/Helpers/PCGEdModeSceneVisualizeHelpers.h"
+#include "EditorMode/Tools/Helpers/PCGEdModeSplineHelpers.h"
+
+#include "BaseBehaviors/SingleClickOrDragBehavior.h"
+#include "Components/SplineComponent.h"
+#include "Settings/LevelEditorViewportSettings.h"
+
+#include UE_INLINE_GENERATED_CPP_BY_NAME(PCGDrawSplineTool)
+
+#define LOCTEXT_NAMESPACE "UPCGDrawSplineTool"
+
+namespace PCGDrawSplineTool
+{
+	using namespace PCG::EditorMode;
+
+	// Helper classes for making undo/redo transactions
+	namespace CommandChange
+	{
+		class FSplineCommandChange : public FToolCommandChange
+		{
+		public:
+			FSplineCommandChange(USplineComponent* InWorkingSpline)
+				: WorkingSpline(InWorkingSpline) {}
+
+			// These pass the working spline to the overloads below
+			virtual void Apply(UObject* Object) override
+			{
+				if (!ensure(Object && Object->IsA<UPCGDrawSplineTool>() && WorkingSpline))
+				{
+					return;
+				}
+
+				Apply(*WorkingSpline);
+			}
+
+			virtual void Revert(UObject* Object) override
+			{
+				if (!ensure(Object && Object->IsA<UPCGDrawSplineTool>() && WorkingSpline))
+				{
+					return;
+				}
+
+				Revert(*WorkingSpline);
+			}
+
+		protected:
+			virtual void Apply(USplineComponent& Spline) = 0;
+			virtual void Revert(USplineComponent& Spline) = 0;
+
+		private:
+			USplineComponent* WorkingSpline = nullptr;
+		};
+
+		// Undoes a point addition with an auto tangent
+		class FSimplePointInsertionChange final : public FSplineCommandChange
+		{
+		public:
+			FSimplePointInsertionChange(
+				USplineComponent* InWorkingSpline,
+				const FVector& HitLocationIn,
+				const FVector& UpVectorIn)
+				: FSplineCommandChange(InWorkingSpline)
+				, HitLocation(HitLocationIn)
+				, UpVector(UpVectorIn) {}
+
+			virtual void Apply(USplineComponent& Spline) override
+			{
+				Spline.AddSplinePoint(HitLocation, ESplineCoordinateSpace::World, false);
+				const int32 PointIndex = Spline.GetNumberOfSplinePoints() - 1;
+				Spline.SetUpVectorAtSplinePoint(PointIndex, UpVector, ESplineCoordinateSpace::World, true);
+			}
+
+			virtual void Revert(USplineComponent& Spline) override
+			{
+				if (ensure(Spline.GetNumberOfSplinePoints() > 0))
+				{
+					Spline.RemoveSplinePoint(Spline.GetNumberOfSplinePoints() - 1, true);
+				}
+			}
+
+			virtual FString ToString() const override
+			{
+				return TEXT("FSimplePointInsertionChange");
+			}
+
+		protected:
+			FVector HitLocation;
+			FVector UpVector;
+		};
+
+		// Undoes a point addition with an explicit tangent
+		class FTangentPointInsertionChange final : public FSplineCommandChange
+		{
+		public:
+			FTangentPointInsertionChange(USplineComponent* InWorkingSpline,
+										 const FVector& HitLocationIn,
+										 const FVector& UpVectorIn,
+										 const FVector& TangentIn)
+				: FSplineCommandChange(InWorkingSpline)
+				, HitLocation(HitLocationIn)
+				, UpVector(UpVectorIn)
+				, Tangent(TangentIn) {}
+
+			virtual void Apply(USplineComponent& Spline) override
+			{
+				Spline.AddSplinePoint(HitLocation, ESplineCoordinateSpace::World, false);
+				const int32 PointIndex = Spline.GetNumberOfSplinePoints() - 1;
+				Spline.SetUpVectorAtSplinePoint(PointIndex, UpVector, ESplineCoordinateSpace::World, false);
+				Spline.SetTangentAtSplinePoint(PointIndex, Tangent, ESplineCoordinateSpace::World, true);
+			}
+
+			virtual void Revert(USplineComponent& Spline) override
+			{
+				if (ensure(Spline.GetNumberOfSplinePoints() > 0))
+				{
+					Spline.RemoveSplinePoint(Spline.GetNumberOfSplinePoints() - 1, true);
+				}
+			}
+
+			virtual FString ToString() const override
+			{
+				return TEXT("FTangentPointInsertionChange");
+			}
+
+		protected:
+			FVector HitLocation;
+			FVector UpVector;
+			FVector Tangent;
+		};
+
+		// Undoes a free draw stroke (multiple points at once)
+		class FStrokeInsertionChange final : public FSplineCommandChange
+		{
+		public:
+			FStrokeInsertionChange(USplineComponent* InWorkingSpline,
+								   const TArray<FVector>& HitLocationsIn,
+								   const TArray<FVector>& UpVectorsIn)
+				: FSplineCommandChange(InWorkingSpline)
+				, HitLocations(HitLocationsIn)
+				, UpVectors(UpVectorsIn)
+			{
+				if (!ensure(HitLocations.Num() == UpVectors.Num()))
+				{
+					const int32 Num = FMath::Min(HitLocations.Num(), UpVectors.Num());
+					HitLocations.SetNum(Num);
+					UpVectors.SetNum(Num);
+				}
+			}
+
+			virtual void Apply(USplineComponent& Spline) override
+			{
+				for (int32 i = 0; i < HitLocations.Num(); ++i)
+				{
+					Spline.AddSplinePoint(HitLocations[i], ESplineCoordinateSpace::World, false);
+					const int32 PointIndex = Spline.GetNumberOfSplinePoints() - 1;
+					Spline.SetUpVectorAtSplinePoint(PointIndex, UpVectors[i], ESplineCoordinateSpace::World, false);
+				}
+				Spline.UpdateSpline();
+			}
+
+			virtual void Revert(USplineComponent& Spline) override
+			{
+				for (int32 i = 0; i < HitLocations.Num(); ++i)
+				{
+					if (!ensure(Spline.GetNumberOfSplinePoints() > 0))
+					{
+						break;
+					}
+					Spline.RemoveSplinePoint(Spline.GetNumberOfSplinePoints() - 1, false);
+				}
+				Spline.UpdateSpline();
+			}
+
+			virtual FString ToString() const override
+			{
+				return TEXT("FStrokeInsertionChange");
+			}
+
+		protected:
+			TArray<FVector> HitLocations;
+			TArray<FVector> UpVectors;
+		};
+	} // namespace CommandChange
+
+	namespace Constants
+	{
+		const FText ToolName = LOCTEXT("ToolName", "DrawSpline");
+		const FText ToolDescription = LOCTEXT("DrawSplineDisplayMessage", "Draw a spline and distribute spline data directly to a PCG Graph Instance.");
+		static const FText AddPointTransactionName = LOCTEXT("AddPointTransactionName", "Add Point");
+	}
+} // namespace PCGDrawSplineTool
+
+UPCGDrawSplineToolSettings::UPCGDrawSplineToolSettings()
+	: SplineName(GetDefault<UPCGEditorModeSettings>()->DefaultNewSplineComponentName)
+{
+	if (!HasAnyFlags(RF_ClassDefaultObject))
+	{
+		UPCGDrawSplineToolSettings::TryLoadDefaultGraph();
+	}
+}
+
+FText UPCGDrawSplineToolSettings::GetToolTag() const
+{
+	return PCGDrawSplineTool::Constants::ToolName;
+}
+
+bool UPCGDrawSplineToolSettings::ApplyTool(UWorld* World, const USplineComponent* ToolSplineComponent, FText& OutWarning) const
+{
+	using namespace PCG::EditorMode;
+
+	check(World);
+	check(ToolSplineComponent);
+
+	if (!TargetGraph)
+	{
+		OutWarning = LOCTEXT("InvalidGraphInstance", "[Spline Tool] Failed: Invalid Graph Instance");
+		return false;
+	}
+
+	AActor* Actor = nullptr;
+	UPCGComponent* PCGComponent = nullptr;
+
+	switch (Interaction)
+	{
+		case EPCGToolInteractionMode::CreateNewActor:
+		{
+			Actor = Actor::Spawn(World, NewActorName, FTransform::Identity);
+		}
+		break;
+		case EPCGToolInteractionMode::AddToSelectedActor:
+		{
+			if (!TargetActor.IsValid())
+			{
+				OutWarning = LOCTEXT("InvalidActorWarning", "Invalid Actor!");
+				return false;
+			}
+
+			Actor = TargetActor.Get();
+			PCGComponent = PCGComponent::Find(Actor);
+		}
+		break;
+		default:
+			checkNoEntry();
+	}
+
+	// @todo_pcg: Will overwrite the current spline. Could be an option.
+	USplineComponent* NewSpline = Spline::FindComponent(Actor);
+	if (!NewSpline)
+	{
+		NewSpline = Spline::CreateComponent(Actor, SplineName, ToolSplineComponent, /*bTransactional=*/true);
+		check(NewSpline);
+	}
+	else
+	{
+		Spline::CopySplinePoints(*ToolSplineComponent, *NewSpline, /*bTransactional=*/true);
+	}
+
+	if (TargetGraph)
+	{
+		if (!PCGComponent)
+		{
+			PCGComponent = PCGComponent::Create(Actor, NewPCGComponentName);
+			check(PCGComponent);
+		}
+
+		PCGComponent->SetGraph(TargetGraph);
+
+		if (bAutoGenerate)
+		{
+			PCGComponent->Generate();
+		}
+	}
+
+	return true;
+}
+
+FSoftObjectPath UPCGDrawSplineToolSettings::GetDefaultGraph() const
+{
+	// @todo_pcg: Once finalized, the default graph asset for this tool is returned here. Ex:
+	// return FSoftObjectPath(TEXT("/PCG/EdMode/ToolGraphs/DrawSpline/PCG_Spline.PCG_Spline"));
+	return FSoftObjectPath{};
+}
+
+void UPCGDrawSplineTool::Setup()
+{
+	using namespace PCGDrawSplineTool;
+
+	Super::Setup();
+
+	// @todo_pcg: These are set up to be re-used between edits. For now, they should be wiped every time.
+	check(!WorkingActor.IsValid())
+	check(!WorkingComponent.IsValid())
+	check(!WorkingSpline)
+
+	UPCGDrawSplineToolSettings* Settings = CastChecked<UPCGDrawSplineToolSettings>(ToolSettings);
+	CachedMinPointSpacingSquared = Settings->MinPointSpacing * Settings->MinPointSpacing;
+
+	WorkingActor = Actor::SpawnWorking(GetContextualWorld());
+	check(WorkingActor.IsValid());
+
+	// @todo_pcg: Maybe a helper to validate the graph?
+	if (!Settings->TargetGraph)
+	{
+		NotifyError(LOCTEXT("InvalidGraph", "Invalid or non-existent graph. Shutting down."));
+		GetToolManager()->ShutdownAfterError();
+		return;
+	}
+
+	WorkingComponent = PCGComponent::CreateWorking(WorkingActor.Get());
+	check(WorkingComponent.IsValid());
+	WorkingComponent->SetGraph(Settings->TargetGraph);
+
+	WorkingSpline = Spline::CreateWorkingComponent(WorkingActor.Get());
+	check(WorkingSpline);
+
+	WorkingSpline->ClearSplinePoints();
+	WorkingSpline->SetClosedLoop(bIsSplineSurface);
+	WorkingSpline->bSplineHasBeenEdited = true;
+
+	ClickOrDragBehavior = NewObject<USingleClickOrDragInputBehavior>();
+	ClickOrDragBehavior->Initialize(this, this);
+	AddInputBehavior(ClickOrDragBehavior);
+
+	// @todo_pcg: Make adjustments to this and check when properties change.
+	CachedRayParams =
+	{
+		.World = GetWorld()
+	};
+
+	Settings->WatchProperty(Settings->TargetActor, [this](TWeakObjectPtr<AActor> TargetActor)
+	{
+		if (!TargetActor.IsValid())
+		{
+			SetWarning(LOCTEXT("TargetActorSetToInvalid", "Target actor is invalid!"));
+		}
+	});
+
+	Settings->SilentUpdateWatched();
+}
+
+void UPCGDrawSplineTool::Shutdown(const EToolShutdownType ShutdownType)
+{
+	using namespace PCGDrawSplineTool;
+
+	if (ShutdownType == EToolShutdownType::Accept && WorkingSpline && WorkingSpline->GetNumberOfSplinePoints() > 0)
+	{
+		ApplyToolInternal();
+	}
+
+	if (WorkingComponent.IsValid())
+	{
+		WorkingComponent->CleanupLocalImmediate(/*bRemoveAllComponents=*/true, /*bCleanUpLocalComponents=*/true);
+		WorkingComponent->bActivated = false;
+		WorkingComponent->DestroyComponent();
+	}
+
+	if (WorkingSpline)
+	{
+		WorkingSpline->DestroyComponent();
+		WorkingSpline = nullptr;
+	}
+
+	if (WorkingActor.IsValid())
+	{
+		WorkingActor->MarkAsGarbage();
+		WorkingActor = nullptr;
+	}
+
+	Super::Shutdown(ShutdownType);
+}
+
+void UPCGDrawSplineTool::AddSplinePoint(const FVector& HitLocation, const FVector& HitNormal) const
+{
+	using namespace PCGDrawSplineTool;
+	if (!WorkingSpline)
+	{
+		return;
+	}
+
+	const int32 NumSplinePoints = WorkingSpline->GetNumberOfSplinePoints();
+	const FVector NewNormal = GetNewPointNormal(HitNormal);
+
+	FVector FinalLocation = HitLocation;
+	const UPCGDrawSplineToolSettings* Settings = CastChecked<UPCGDrawSplineToolSettings>(ToolSettings);
+	switch (Settings->Target.OffsetMode)
+	{
+		case EPCGToolDrawTargetOffset::DistanceOffset:
+			check(HitNormal.IsNormalized());
+			FinalLocation += Settings->Target.OffsetDistance * NewNormal;
+			break;
+		case EPCGToolDrawTargetOffset::ExplicitOffset:
+			FinalLocation += Settings->Target.OffsetVector;
+			break;
+		default: // fall-through
+		case EPCGToolDrawTargetOffset::HitLocation:
+			break;
+	}
+
+	WorkingSpline->AddSplinePoint(FinalLocation, ESplineCoordinateSpace::World, /*bUpdate=*/false);
+	WorkingSpline->SetUpVectorAtSplinePoint(NumSplinePoints, NewNormal, ESplineCoordinateSpace::World, /*bUpdate=*/true);
+}
+
+FVector UPCGDrawSplineTool::GetNewPointNormal(const FVector& HitNormal) const
+{
+	const UPCGDrawSplineToolSettings* Settings = CastChecked<UPCGDrawSplineToolSettings>(ToolSettings);
+	const int32 NumSplinePoints = WorkingSpline->GetNumberOfSplinePoints();
+	switch (Settings->Target.NormalMode)
+	{
+		case EPCGToolDrawTargetNormal::AlignToPrevious:
+		{
+			check(NumSplinePoints > 0);
+			if (NumSplinePoints == 1)
+			{
+				/*
+				 * If there's only one point, GetUpVectorAtSplinePoint is unreliable because it seeks to build a quaternion
+				 * from the tangent and the set-up vector, and the tangent is zero. We want to use the "stored" up
+				 * vector directly.
+				 */
+				const FVector LocalUpVector = WorkingSpline->GetRotationAtSplinePoint(0, ESplineCoordinateSpace::Local).RotateVector(WorkingSpline->DefaultUpVector);
+				return WorkingSpline->GetComponentTransform().TransformVectorNoScale(LocalUpVector);
+			}
+			else if (NumSplinePoints > 1)
+			{
+				return WorkingSpline->GetUpVectorAtSplinePoint(NumSplinePoints - 1, ESplineCoordinateSpace::World);
+			}
+
+			checkNoEntry();
+			return HitNormal;
+		}
+		case EPCGToolDrawTargetNormal::WorldUp:
+			return FVector::UnitZ();
+		case EPCGToolDrawTargetNormal::Explicit:
+			return Settings->Target.Normal;
+		default: // fall-through
+		case EPCGToolDrawTargetNormal::HitNormal:
+			return HitNormal;
+	}
+}
+
+void UPCGDrawSplineTool::ApplyToolInternal() const
+{
+	GetToolManager()->BeginUndoTransaction(LOCTEXT("DrawSplineTransactionName", "[PCG Tool] Draw Spline"));
+	const UPCGDrawSplineToolSettings* Settings = CastChecked<UPCGDrawSplineToolSettings>(ToolSettings);
+
+	FText Warnings;
+	if (!Settings->ApplyTool(GetContextualWorldChecked(), WorkingSpline.Get(), Warnings))
+	{
+		NotifyError(Warnings);
+	}
+
+	GetToolManager()->EndUndoTransaction();
+}
+
+FInputRayHit UPCGDrawSplineTool::IsHitByClick(const FInputDeviceRay& ClickPosition)
+{
+	using namespace PCG::EditorMode;
+	if (TOptional<Scene::FViewHitResult> Result = ViewportRay(ClickPosition, CachedRayParams))
+	{
+		return Result->ToInputRayHit();
+	}
+
+	return {};
+}
+
+void UPCGDrawSplineTool::OnClicked(const FInputDeviceRay& ClickPosition)
+{
+	using namespace PCGDrawSplineTool;
+
+	if (TOptional<Scene::FViewHitResult> Result = ViewportRay(ClickPosition, CachedRayParams))
+	{
+		const UPCGDrawSplineToolSettings* Settings = CastChecked<UPCGDrawSplineToolSettings>(ToolSettings);
+		switch (Settings->DrawMode)
+		{
+			case EPCGDrawSplineDrawMode::ClickAutoTangent:
+			case EPCGDrawSplineDrawMode::FreeDraw:
+			{
+				AddSplinePoint(Result->ImpactPosition, Result->ImpactNormal);
+
+				const int32 PointIndex = WorkingSpline->GetNumberOfSplinePoints() - 1;
+				GetToolManager()->EmitObjectChange(this,
+												   MakeUnique<CommandChange::FSimplePointInsertionChange>(
+													   WorkingSpline.Get(),
+													   Result->ImpactPosition,
+													   WorkingSpline->GetUpVectorAtSplinePoint(PointIndex, ESplineCoordinateSpace::World)),
+												   Constants::AddPointTransactionName);
+				break;
+			}
+			case EPCGDrawSplineDrawMode::TangentDrag:
+			{
+				AddSplinePoint(Result->ImpactPosition, Result->ImpactNormal);
+
+				const int32 PointIndex = WorkingSpline->GetNumberOfSplinePoints() - 1;
+				WorkingSpline->SetTangentAtSplinePoint(PointIndex,
+													   FVector::Zero(),
+													   ESplineCoordinateSpace::World, true);
+
+				GetToolManager()->EmitObjectChange(
+					this,
+					MakeUnique<CommandChange::FTangentPointInsertionChange>(
+						WorkingSpline.Get(),
+						Result->ImpactPosition,
+						WorkingSpline->GetUpVectorAtSplinePoint(PointIndex, ESplineCoordinateSpace::World),
+						FVector::Zero()),
+					Constants::AddPointTransactionName);
+				break;
+			}
+		}
+	}
+}
+
+FInputRayHit UPCGDrawSplineTool::CanBeginClickDragSequence(const FInputDeviceRay& DragPosition)
+{
+	using namespace PCG::EditorMode::Scene;
+	if (TOptional<FViewHitResult> Result = ViewportRay(DragPosition, CachedRayParams))
+	{
+		return Result->ToInputRayHit();
+	}
+
+	return {};
+}
+
+void UPCGDrawSplineTool::OnClickPress(const FInputDeviceRay& ClickPosition)
+{
+	using namespace PCG::EditorMode::Scene;
+	if (TOptional<FViewHitResult> Result = ViewportRay(ClickPosition, CachedRayParams))
+	{
+		AddSplinePoint(Result->ImpactPosition, Result->ImpactNormal);
+
+		const UPCGDrawSplineToolSettings* Settings = CastChecked<UPCGDrawSplineToolSettings>(ToolSettings);
+		if (Settings->DrawMode == EPCGDrawSplineDrawMode::FreeDraw)
+		{
+			// Cache the index for where this stroke began
+			FreeDrawStrokeStartIndex = WorkingSpline->GetNumberOfSplinePoints() - 1;
+		}
+	}
+}
+
+void UPCGDrawSplineTool::OnClickDrag(const FInputDeviceRay& DragPosition)
+{
+	using namespace PCG::EditorMode::Scene;
+
+	const int32 NumSplinePoints = WorkingSpline->GetNumberOfSplinePoints();
+	if (!ensure(NumSplinePoints > 0))
+	{
+		return;
+	}
+
+	if (TOptional<FViewHitResult> Result = ViewportRay(DragPosition, CachedRayParams))
+	{
+		const int32 FinalSplinePoint = NumSplinePoints - 1;
+		const UPCGDrawSplineToolSettings* Settings = CastChecked<UPCGDrawSplineToolSettings>(ToolSettings);
+		switch (Settings->DrawMode)
+		{
+			case EPCGDrawSplineDrawMode::ClickAutoTangent:
+			{
+				// Drag the last placed point
+				const FVector Normal = GetNewPointNormal(Result->ImpactNormal);
+				WorkingSpline->SetLocationAtSplinePoint(FinalSplinePoint, Result->ImpactPosition, ESplineCoordinateSpace::World, false);
+				WorkingSpline->SetUpVectorAtSplinePoint(FinalSplinePoint, Normal, ESplineCoordinateSpace::World, true);
+
+				break;
+			}
+			case EPCGDrawSplineDrawMode::TangentDrag:
+			{
+				// Set the tangent
+				const FVector LastPoint = WorkingSpline->GetLocationAtSplinePoint(FinalSplinePoint, ESplineCoordinateSpace::World);
+				const FVector Tangent = (Result->ImpactPosition - LastPoint) / FMath::Min(0.5f, GetDefault<ULevelEditorViewportSettings>()->SplineTangentScale);
+				WorkingSpline->SetTangentAtSplinePoint(FinalSplinePoint, Tangent, ESplineCoordinateSpace::World, true);
+				bDrawTangentForPreviousPoint = true;
+				break;
+			}
+
+			case EPCGDrawSplineDrawMode::FreeDraw:
+			{
+				/*
+				 * Instead of dragging the first placed point (which gets placed in OnClickPress), drag a second "preview" one
+				 * until far enough from the previous to place a new control point.
+				 */
+				if (!bFreeDrawPlacedPreviewPoint)
+				{
+					AddSplinePoint(Result->ImpactPosition, Result->ImpactNormal);
+					bFreeDrawPlacedPreviewPoint = true;
+				}
+				else
+				{
+					FVector FinalLocation = Result->ImpactPosition;
+					const FVector Normal = GetNewPointNormal(Result->ImpactNormal);
+
+					switch (Settings->Target.OffsetMode)
+					{
+						case EPCGToolDrawTargetOffset::DistanceOffset:
+							check(Result->ImpactNormal.IsNormalized());
+							FinalLocation += Settings->Target.OffsetDistance * Normal;
+							break;
+						case EPCGToolDrawTargetOffset::ExplicitOffset:
+							FinalLocation += Settings->Target.OffsetVector;
+							break;
+						default: // fall-through
+						case EPCGToolDrawTargetOffset::HitLocation:
+							break;
+					}
+
+					WorkingSpline->SetLocationAtSplinePoint(FinalSplinePoint, FinalLocation, ESplineCoordinateSpace::World, /*bUpdateSpline=*/false);
+					WorkingSpline->SetUpVectorAtSplinePoint(FinalSplinePoint, Normal, ESplineCoordinateSpace::World, /*bUpdateSpline=*/true);
+					const FVector PreviousPoint = WorkingSpline->GetLocationAtSplinePoint(FinalSplinePoint - 1, ESplineCoordinateSpace::World);
+					if (FVector::DistSquared(FinalLocation, PreviousPoint) >= CachedMinPointSpacingSquared)
+					{
+						AddSplinePoint(Result->ImpactPosition, Result->ImpactNormal);
+					}
+				}
+				break;
+			}
+		}
+	}
+}
+
+void UPCGDrawSplineTool::OnClickRelease(const FInputDeviceRay& ReleasePosition)
+{
+	// Activate the final position
+	OnClickDrag(ReleasePosition);
+	OnTerminateDragSequence();
+}
+
+void UPCGDrawSplineTool::OnTerminateDragSequence()
+{
+	using namespace PCGDrawSplineTool;
+	check(WorkingSpline);
+
+	bDrawTangentForPreviousPoint = false;
+	bFreeDrawPlacedPreviewPoint = false;
+
+	const int32 NumSplinePoints = WorkingSpline->GetNumberOfSplinePoints();
+	const int32 PreviousSplinePoint = NumSplinePoints - 1;
+
+	const UPCGDrawSplineToolSettings* Settings = CastChecked<UPCGDrawSplineToolSettings>(ToolSettings);
+	// Emit the appropriate undo transaction
+	switch (Settings->DrawMode)
+	{
+		case EPCGDrawSplineDrawMode::ClickAutoTangent:
+		{
+			GetToolManager()->EmitObjectChange(
+				this,
+				MakeUnique<CommandChange::FSimplePointInsertionChange>(
+					WorkingSpline.Get(),
+					WorkingSpline->GetLocationAtSplinePoint(PreviousSplinePoint, ESplineCoordinateSpace::World),
+					WorkingSpline->GetUpVectorAtSplinePoint(PreviousSplinePoint, ESplineCoordinateSpace::World)),
+				Constants::AddPointTransactionName);
+			break;
+		}
+		case EPCGDrawSplineDrawMode::TangentDrag:
+		{
+			GetToolManager()->EmitObjectChange(
+				this,
+				MakeUnique<CommandChange::FTangentPointInsertionChange>(
+					WorkingSpline.Get(),
+					WorkingSpline->GetLocationAtSplinePoint(PreviousSplinePoint, ESplineCoordinateSpace::World),
+					WorkingSpline->GetUpVectorAtSplinePoint(PreviousSplinePoint, ESplineCoordinateSpace::World),
+					WorkingSpline->GetTangentAtSplinePoint(PreviousSplinePoint, ESplineCoordinateSpace::World)),
+				Constants::AddPointTransactionName);
+			break;
+		}
+
+		case EPCGDrawSplineDrawMode::FreeDraw:
+		{
+			TArray<FVector> HitLocations;
+			TArray<FVector> UpVectors;
+			for (int32 Index = FreeDrawStrokeStartIndex; Index < NumSplinePoints; ++Index)
+			{
+				HitLocations.Add(WorkingSpline->GetLocationAtSplinePoint(Index, ESplineCoordinateSpace::World));
+				UpVectors.Add(WorkingSpline->GetUpVectorAtSplinePoint(Index, ESplineCoordinateSpace::World));
+			}
+
+			GetToolManager()->EmitObjectChange(
+				this,
+				MakeUnique<CommandChange::FStrokeInsertionChange>(WorkingSpline.Get(), HitLocations, UpVectors),
+				Constants::AddPointTransactionName);
+			break;
+		}
+	}
+}
+
+void UPCGDrawSplineTool::OnTick(const float DeltaTime)
+{
+	if (!WorkingSpline)
+	{
+		GetToolManager()->PostActiveToolShutdownRequest(
+			this,
+			EToolShutdownType::Cancel,
+			/*bShowUnexpectedShutdownMessage=*/true,
+			LOCTEXT("LostWorkingSpline", "The Draw Spline tool must close because the in-progress spline has been unexpectedly deleted."));
+	}
+
+	Super::OnTick(DeltaTime);
+}
+
+void UPCGDrawSplineTool::Render(IToolsContextRenderAPI* RenderAPI)
+{
+	using namespace PCG::EditorMode::Scene;
+
+	if (WorkingSpline)
+	{
+		const int32 LastPointIndex = WorkingSpline->GetNumberOfSplinePoints() - 1;
+
+		if (bDrawTangentForPreviousPoint)
+		{
+			Visualize::SplineTangent(*WorkingSpline, *RenderAPI, LastPointIndex);
+		}
+
+		Visualize::Spline(*WorkingSpline, *RenderAPI, LastPointIndex, /*bShouldVisualizeScale=*/false, /*bShouldVisualizeTangents=*/true);
+	}
+}
+
+void UPCGDrawSplineTool::OnRefresh()
+{
+	if (WorkingSpline)
+	{
+		WorkingSpline->PostEditChange();
+	}
+
+	Super::OnRefresh();
+}
+
+bool UPCGDrawSplineTool::CanAccept() const
+{
+	return WorkingSpline && WorkingSpline->GetNumberOfSplinePoints() > 0;
+}
+
+bool UPCGDrawSplineToolBuilder::CanBuildTool(const FToolBuilderState& SceneState) const
+{
+	return true;
+}
+
+UInteractiveTool* UPCGDrawSplineToolBuilder::BuildTool(const FToolBuilderState& SceneState) const
+{
+	UPCGDrawSplineTool* NewTool = NewObject<UPCGDrawSplineTool>(SceneState.ToolManager);
+
+	NewTool->SetIsSplineSurface(bIsSplineSurface);
+
+	NewTool->SetToolInfo(
+	{
+		.ToolDisplayName = PCGDrawSplineTool::Constants::ToolName,
+		.ToolDisplayMessage = PCGDrawSplineTool::Constants::ToolDescription,
+		.ToolIcon = FPCGEditorModeStyle::Get().GetBrush("PCGEditorMode.Tools.DrawSpline")
+	});
+
+	return NewTool;
+}
+
+#undef LOCTEXT_NAMESPACE
diff --git a/Engine/Plugins/PCG/Source/PCGEditor/Private/EditorMode/Tools/PCGInteractiveTool.cpp b/Engine/Plugins/PCG/Source/PCGEditor/Private/EditorMode/Tools/PCGInteractiveTool.cpp
new file mode 100644
--- /dev/null
+++ b/Engine/Plugins/PCG/Source/PCGEditor/Private/EditorMode/Tools/PCGInteractiveTool.cpp
@@ -0,0 +1,112 @@
+// Copyright Epic Games, Inc. All Rights Reserved.
+
+#include "EditorMode/Tools/PCGInteractiveTool.h"
+
+#include "PCGComponent.h"
+#include "PCGEditorModule.h"
+#include "EditorMode/PCGEdModeHelpers.h"
+#include "EditorMode/PCGEdModeSettings.h"
+#include "EditorMode/Tools/PCGInteractiveToolSettings.h"
+
+#include UE_INLINE_GENERATED_CPP_BY_NAME(PCGInteractiveTool)
+
+void UPCGInteractiveTool::Setup()
+{
+	// `GetSettingsClass` must be overridden with a valid settings class for this tool.
+	if (const UClass* SettingsClass = GetSettingsClass())
+	{
+		ToolSettings = NewObject<UPCGInteractiveToolSettings>(this, SettingsClass, SettingsClass->GetDefaultObjectName());
+	}
+
+	// Validate the classes were generated.
+	check(ToolSettings)
+
+	ToolSettings->RestoreProperties(this);
+	AddToolPropertySource(ToolSettings);
+
+	// Do some housecleaning checks. Must be defined prior. Can define in the `PostBuildTool` override.
+	check(!GetToolInfo().ToolDisplayName.IsEmpty())
+	check(!GetToolInfo().ToolDisplayMessage.IsEmpty())
+	check(GetToolInfo().ToolIcon)
+
+	Super::Setup();
+}
+
+void UPCGInteractiveTool::Shutdown(const EToolShutdownType ShutdownType)
+{
+	Super::Shutdown(ShutdownType);
+
+	// Save the updates to the tool settings
+	ToolSettings->SaveProperties(this);
+
+	RemoveToolPropertySource(ToolSettings);
+}
+
+void UPCGInteractiveTool::OnTick(const float DeltaTime)
+{
+	GraphRefreshTimer -= DeltaTime;
+	if (GraphRefreshTimer < 0.f)
+	{
+		OnRefresh();
+	}
+}
+
+void UPCGInteractiveTool::OnRefresh()
+{
+	if (WorkingComponent.IsValid() && !WorkingComponent->IsCleaningUp() && !WorkingComponent->IsGenerating())
+	{
+		WorkingComponent->GenerateLocal(true);
+	}
+
+	// Reset the refresh timer
+	GraphRefreshTimer = GetDefault<UPCGEditorModeSettings>()->GraphRefreshRate;
+}
+
+void UPCGInteractiveTool::SetTargetActor(AActor* InTargetActor)
+{
+	if (ToolSettings)
+	{
+		ToolSettings->TargetActor = InTargetActor;
+	}
+}
+
+TWeakObjectPtr<AActor> UPCGInteractiveTool::GetTargetActor() const
+{
+	return ToolSettings ? ToolSettings->TargetActor : nullptr;
+}
+
+UWorld* UPCGInteractiveTool::GetContextualWorld() const
+{
+	return GetToolManager()->GetContextQueriesAPI()->GetCurrentEditingWorld();
+}
+
+UWorld* UPCGInteractiveTool::GetContextualWorldChecked() const
+{
+	UWorld* World = GetContextualWorld();
+	check(World);
+	return World;
+}
+
+void UPCGInteractiveTool::SetWarning(const FText& Warning) const
+{
+	WarningMessage = Warning;
+	OnToolWarningChanged.Broadcast(WarningMessage);
+}
+
+void UPCGInteractiveTool::ClearWarning() const
+{
+	WarningMessage = FText::GetEmpty();
+	OnToolWarningChanged.Broadcast(WarningMessage);
+}
+
+void UPCGInteractiveTool::NotifyError(const FText& ErrorMessage) const
+{
+	SetWarning(ErrorMessage);
+
+	if (GetDefault<UPCGEditorModeSettings>()->bShowEditorToastOnToolErrors)
+	{
+		PCG::EditorMode::DispatchEditorNotification(ErrorMessage);
+	}
+
+	UE_LOG(LogPCGEditor, Error, TEXT("[%sTool]: %s"), *GetToolInfo().ToolDisplayName.ToString(), *ErrorMessage.ToString());
+}
diff --git a/Engine/Plugins/PCG/Source/PCGEditor/Private/EditorMode/Tools/PCGInteractiveToolBuilder.cpp b/Engine/Plugins/PCG/Source/PCGEditor/Private/EditorMode/Tools/PCGInteractiveToolBuilder.cpp
new file mode 100644
--- /dev/null
+++ b/Engine/Plugins/PCG/Source/PCGEditor/Private/EditorMode/Tools/PCGInteractiveToolBuilder.cpp
@@ -0,0 +1,23 @@
+// Copyright Epic Games, Inc. All Rights Reserved.
+
+#include "EditorMode/Tools/PCGInteractiveToolBuilder.h"
+
+#include "InteractiveToolManager.h"
+
+#include UE_INLINE_GENERATED_CPP_BY_NAME(PCGInteractiveToolBuilder)
+
+bool UPCGInteractiveToolBuilder::CanBuildTool(const FToolBuilderState& SceneState) const
+{
+	return Super::CanBuildTool(SceneState);
+}
+
+UInteractiveTool* UPCGInteractiveToolBuilder::BuildTool(const FToolBuilderState& SceneState) const
+{
+	return Super::BuildTool(SceneState);
+}
+
+void UPCGInteractiveToolBuilder::PostBuildTool(UInteractiveTool* Tool, const FToolBuilderState& SceneState) const
+{
+	Super::PostBuildTool(Tool, SceneState);
+	SceneState.ToolManager->DisplayMessage(Tool->GetToolInfo().ToolDisplayMessage, EToolMessageLevel::UserNotification);
+}
diff --git a/Engine/Plugins/PCG/Source/PCGEditor/Private/EditorMode/Tools/PCGInteractiveToolSettings.cpp b/Engine/Plugins/PCG/Source/PCGEditor/Private/EditorMode/Tools/PCGInteractiveToolSettings.cpp
new file mode 100644
--- /dev/null
+++ b/Engine/Plugins/PCG/Source/PCGEditor/Private/EditorMode/Tools/PCGInteractiveToolSettings.cpp
@@ -0,0 +1,40 @@
+// Copyright Epic Games, Inc. All Rights Reserved.
+
+#include "EditorMode/Tools/PCGInteractiveToolSettings.h"
+
+#include "PCGEditorModule.h"
+#include "EditorMode/PCGEdModeHelpers.h"
+#include "EditorMode/PCGEdModeSettings.h"
+#include "Helpers/PCGHelpers.h"
+
+#include UE_INLINE_GENERATED_CPP_BY_NAME(PCGInteractiveToolSettings)
+
+UPCGInteractiveToolSettings::UPCGInteractiveToolSettings()
+	: NewActorName(GetDefault<UPCGEditorModeSettings>()->DefaultNewActorName)
+	, NewPCGComponentName(GetDefault<UPCGEditorModeSettings>()->DefaultNewPCGComponentName) {}
+
+void UPCGInteractiveToolSettings::TryLoadDefaultGraph()
+{
+	check(!HasAnyFlags(RF_ClassDefaultObject));
+	TargetGraph = CastChecked<UPCGGraphInterface>(GetDefaultGraph().TryLoad(), ECastCheckedType::NullAllowed);
+}
+
+bool UPCGInteractiveToolSettings::GraphAssetFilter(const FAssetData& AssetData) const
+{
+	// This can be checked by subclasses by either finding a tag or more specifically checking for values.
+	if (AssetData.FindTag(GET_MEMBER_NAME_CHECKED(UPCGGraph, CompatibleToolTags)))
+	{
+		FString ToolTagsString;
+		TArray<FString> Tags;
+		if (AssetData.GetTagValue<FString>(GET_MEMBER_NAME_CHECKED(UPCGGraph, CompatibleToolTags), ToolTagsString))
+		{
+			Tags = PCGHelpers::GetStringArrayFromCommaSeparatedList(ToolTagsString);
+		}
+
+		return !Tags.FindByPredicate([ToolTag = GetToolTag().ToString()](const FString& InTag) { return InTag.Equals(ToolTag, ESearchCase::IgnoreCase); });
+	}
+
+	// True == filter the asset
+	return true;
+}
+
diff --git a/Engine/Plugins/PCG/Source/PCGEditor/Private/EditorMode/Tools/Paint/PCGPaintTool.cpp b/Engine/Plugins/PCG/Source/PCGEditor/Private/EditorMode/Tools/Paint/PCGPaintTool.cpp
new file mode 100644
--- /dev/null
+++ b/Engine/Plugins/PCG/Source/PCGEditor/Private/EditorMode/Tools/Paint/PCGPaintTool.cpp
@@ -0,0 +1,82 @@
+// Copyright Epic Games, Inc. All Rights Reserved.
+
+#include "EditorMode/Tools/Paint/PCGPaintTool.h"
+
+#include "EditorMode/PCGEdModeStyle.h"
+
+#include UE_INLINE_GENERATED_CPP_BY_NAME(PCGPaintTool)
+
+#define LOCTEXT_NAMESPACE "UPCGPaintTool"
+
+namespace PCGPaintTool::Constants
+{
+	const FText ToolName = LOCTEXT("ToolName", "PaintTool");
+	const FText ToolDescription = LOCTEXT("DisplayMessage", "Tool not yet implemented.");
+}
+
+UPCGPaintSettings::UPCGPaintSettings(const FObjectInitializer& ObjectInitializer)
+	: Brush(CastChecked<UPCGEdModeBrushBase>(
+		ObjectInitializer.CreateEditorOnlyDefaultSubobject(
+			this,
+			GET_MEMBER_NAME_CHECKED(UPCGPaintSettings, Brush),
+			UPCGEdModeBrushBase::StaticClass(),
+			/*bIsTransient=*/false)))
+{}
+
+FText UPCGPaintSettings::GetToolTag() const
+{
+	return PCGPaintTool::Constants::ToolName;
+}
+
+FSoftObjectPath UPCGPaintSettings::GetDefaultGraph() const
+{
+	return FSoftObjectPath{};
+}
+
+bool UPCGPaintSettings::ApplyTool(UWorld* World, FText& OutWarning) const
+{
+	// @todo_pcg: Apply the output of the tool here.
+	return false;
+}
+
+void UPCGPaintTool::Setup()
+{
+	// @todo_pcg: Set up the tool here.
+	Super::Setup();
+
+	UPCGPaintSettings* Settings = CastChecked<UPCGPaintSettings>(ToolSettings);
+
+	NotifyError(LOCTEXT("ToolNotYetImplemented", "The tool is not yet implemented."));
+}
+
+TSubclassOf<UPCGInteractiveToolSettings> UPCGPaintTool::GetSettingsClass() const
+{
+	return UPCGPaintSettings::StaticClass();
+}
+
+bool UPCGPaintTool::CanAccept() const
+{
+	// @todo_pcg: true once implemented
+	return false;
+}
+
+bool UPCGPaintToolBuilder::CanBuildTool(const FToolBuilderState& SceneState) const
+{
+	return true;
+}
+
+UInteractiveTool* UPCGPaintToolBuilder::BuildTool(const FToolBuilderState& SceneState) const
+{
+	UPCGPaintTool* NewTool = NewObject<UPCGPaintTool>(SceneState.ToolManager);
+
+	NewTool->SetToolInfo(
+	{
+		.ToolDisplayName = PCGPaintTool::Constants::ToolName,
+		.ToolDisplayMessage = PCGPaintTool::Constants::ToolDescription,
+		.ToolIcon = FPCGEditorModeStyle::Get().GetBrush("PCGEditorMode.Tools.Paint")
+	});
+
+	return NewTool;
+}
+
+#undef LOCTEXT_NAMESPACE
diff --git a/Engine/Plugins/PCG/Source/PCGEditor/Private/EditorMode/Tools/Volume/PCGCreateVolumeTool.cpp b/Engine/Plugins/PCG/Source/PCGEditor/Private/EditorMode/Tools/Volume/PCGCreateVolumeTool.cpp
new file mode 100644
--- /dev/null
+++ b/Engine/Plugins/PCG/Source/PCGEditor/Private/EditorMode/Tools/Volume/PCGCreateVolumeTool.cpp
@@ -0,0 +1,120 @@
+// Copyright Epic Games, Inc. All Rights Reserved.
+
+#include "EditorMode/Tools/Volume/PCGCreateVolumeTool.h"
+
+#include "EditorMode/PCGEdModeStyle.h"
+
+#include UE_INLINE_GENERATED_CPP_BY_NAME(PCGCreateVolumeTool)
+
+#define LOCTEXT_NAMESPACE "UPCGCreateVolumeTool"
+
+namespace PCGCreateVolumeTool::Constants
+{
+	const FText ToolName = LOCTEXT("ToolName", "CreateVolume");
+}
+
+FText UPCGCreateVolumeToolSettings::GetToolTag() const
+{
+	return PCGCreateVolumeTool::Constants::ToolName;
+}
+
+FSoftObjectPath UPCGCreateVolumeToolSettings::GetDefaultGraph() const
+{
+	return FSoftObjectPath{};
+}
+
+bool UPCGCreateVolumeToolSettings::ApplyTool(UWorld* World, FText& OutWarning) const
+{
+	// @todo_pcg: Apply the output of the tool here.
+	return false;
+}
+
+void UPCGCreateVolumeTool::Setup()
+{
+	// @todo_pcg: Set up the tool here.
+	Super::Setup();
+
+	UPCGCreateVolumeToolSettings* Settings = CastChecked<UPCGCreateVolumeToolSettings>(ToolSettings);
+
+	NotifyError(LOCTEXT("ToolNotYetImplemented", "The tool is not yet implemented."));
+}
+
+void UPCGCreateVolumeTool::Shutdown(EToolShutdownType ShutdownType)
+{
+	// @todo_pcg: Clean up the tool here.
+}
+
+void UPCGCreateVolumeTool::OnTick(const float DeltaTime)
+{
+	Super::OnTick(DeltaTime);
+}
+
+bool UPCGCreateVolumeTool::CanAccept() const
+{
+	// @todo_pcg: true once implemented
+	return false;
+}
+
+void UPCGCreateVolumeTool::Render(IToolsContextRenderAPI* RenderAPI)
+{
+	FPrimitiveDrawInterface* PDI = RenderAPI->GetPrimitiveDrawInterface();
+	check(PDI);
+
+	const UPCGCreateVolumeToolSettings* Settings = CastChecked<UPCGCreateVolumeToolSettings>(ToolSettings);
+
+	/** @todo_pcg: Three click process -
+	 *	First click - Point (Location)
+	 *	Second click - 2D Plane (X, Y)
+	 *	Third click - 3D Box (Z)
+	 */
+}
+
+TSubclassOf<UPCGInteractiveToolSettings> UPCGCreateVolumeTool::GetSettingsClass() const
+{
+	return UPCGCreateVolumeToolSettings::StaticClass();
+}
+
+FInputRayHit UPCGCreateVolumeTool::IsHitByClick(const FInputDeviceRay& ClickPosition)
+{
+	using namespace PCG::EditorMode::Scene;
+	if (TOptional<FViewHitResult> Result = ViewportRay(ClickPosition, CachedRayParams))
+	{
+		return Result->ToInputRayHit();
+	}
+
+	return {};
+}
+
+void UPCGCreateVolumeTool::OnClicked(const FInputDeviceRay& ClickPosition)
+{
+	UPCGCreateVolumeToolSettings* Settings = CastChecked<UPCGCreateVolumeToolSettings>(ToolSettings);
+
+	/** @todo_pcg: Three click process -
+	 *	First click - Point (Location)
+	 *	Second click - 2D Plane (X, Y)
+	 *	Third click - 3D Box (Z)
+	 */
+
+	// @todo_pcg: Different chords to translate, rotate, scale at any step of the process?
+}
+
+bool UPCGCreateVolumeToolBuilder::CanBuildTool(const FToolBuilderState& SceneState) const
+{
+	return true;
+}
+
+UInteractiveTool* UPCGCreateVolumeToolBuilder::BuildTool(const FToolBuilderState& SceneState) const
+{
+	UPCGCreateVolumeTool* NewTool = NewObject<UPCGCreateVolumeTool>(SceneState.ToolManager);
+
+	NewTool->SetToolInfo(
+	{
+		.ToolDisplayName = PCGCreateVolumeTool::Constants::ToolName,
+		.ToolDisplayMessage = LOCTEXT("CreateVolumeToolDisplayMessage", "Tool not yet implemented."),
+		.ToolIcon = FPCGEditorModeStyle::Get().GetBrush("PCGEditorMode.Tools.CreateVolume")
+	});
+
+	return NewTool;
+}
+
+#undef LOCTEXT_NAMESPACE
diff --git a/Engine/Plugins/PCG/Source/PCGEditor/Private/EditorMode/Widgets/SPCGEdModeAcceptToolOverlay.cpp b/Engine/Plugins/PCG/Source/PCGEditor/Private/EditorMode/Widgets/SPCGEdModeAcceptToolOverlay.cpp
new file mode 100644
--- /dev/null
+++ b/Engine/Plugins/PCG/Source/PCGEditor/Private/EditorMode/Widgets/SPCGEdModeAcceptToolOverlay.cpp
@@ -0,0 +1,75 @@
+// Copyright Epic Games, Inc. All Rights Reserved.
+
+#include "EditorMode/Widgets/SPCGEdModeAcceptToolOverlay.h"
+
+#include "SPrimaryButton.h"
+#include "Widgets/SBoxPanel.h"
+#include "Widgets/Images/SImage.h"
+#include "Widgets/Text/STextBlock.h"
+
+void SPCGEdModeAcceptToolOverlay::Construct(const FArguments& InArgs)
+{
+	Args = InArgs;
+
+	// Non-volatile. Should remain static until tool changes.
+	SetCanTick(false);
+
+	this->ChildSlot
+	[
+		SNew(SHorizontalBox)
+		+ SHorizontalBox::Slot()
+		.HAlign(HAlign_Center)
+		.VAlign(VAlign_Bottom)
+		.Padding(FMargin(0.0f, 0.0f, 0.f, 15.f))
+		[
+			SNew(SBorder)
+			.BorderImage(FAppStyle::Get().GetBrush("EditorViewport.OverlayBrush"))
+			.Padding(8.f)
+			[
+				SNew(SHorizontalBox)
+				+ SHorizontalBox::Slot()
+				.AutoWidth()
+				.VAlign(VAlign_Center)
+				.Padding(FMargin(0.f, 0.f, 8.f, 0.f))
+				[
+					SNew(SImage)
+					.Image(Args._ActiveToolIcon)
+				]
+				+ SHorizontalBox::Slot()
+				.AutoWidth()
+				.VAlign(VAlign_Center)
+				.Padding(FMargin(0.f, 0.f, 8.f, 0.f))
+				[
+					SNew(STextBlock)
+					.Text(Args._ActiveToolDisplayName)
+				]
+				+ SHorizontalBox::Slot()
+				.AutoWidth()
+				.Padding(FMargin(0.0, 0.f, 2.f, 0.f))
+				[
+					SNew(SPrimaryButton)
+					.Text(Args._AcceptButtonLabel)
+					.ToolTipText(Args._AcceptButtonTooltip)
+					.OnClicked(Args._OnAcceptButtonClicked)
+					.IsEnabled(Args._IsAcceptButtonEnabled)
+					.Visibility(Args._GetAcceptButtonVisibility)
+				]
+				+ SHorizontalBox::Slot()
+				.AutoWidth()
+				.Padding(FMargin(2.0, 0.f, 0.f, 0.f))
+				[
+					SNew(SButton)
+					.HAlign(HAlign_Center)
+					.TextStyle(FAppStyle::Get(), "DialogButtonText")
+					.Text(Args._CancelButtonLabel)
+					.ToolTipText(Args._CancelButtonTooltip)
+					.OnClicked(Args._OnCancelButtonClicked)
+					.IsEnabled(Args._IsCancelButtonEnabled)
+					.Visibility(Args._GetCancelButtonVisibility)
+				]
+			]
+		]
+	];
+
+	SlatePrepass();
+}
diff --git a/Engine/Plugins/PCG/Source/PCGEditor/Private/PCGEditorModule.cpp b/Engine/Plugins/PCG/Source/PCGEditor/Private/PCGEditorModule.cpp
--- a/Engine/Plugins/PCG/Source/PCGEditor/Private/PCGEditorModule.cpp
+++ b/Engine/Plugins/PCG/Source/PCGEditor/Private/PCGEditorModule.cpp
@@ -66,6 +66,7 @@
 #include "ToolMenus.h"
 #include "Details/EnumSelectorDetails.h"
 #include "Editor/SceneOutliner/Public/ISceneOutliner.h"
+#include "EditorMode/PCGEdMode.h"
 #include "Framework/MultiBox/MultiBoxBuilder.h"
 #include "Framework/Notifications/NotificationManager.h"
 #include "UObject/ObjectSaveContext.h"
@@ -91,6 +92,8 @@ void FPCGEditorModule::StartupModule()
 	FPCGEditorSpawnNodeCommands::Register();
 	FPCGNodeSourceEditorTextBoxCommands::Register();
 
+	UPCGEditorMode::RegisterEditorMode();
+
 	GraphNodeFactory = MakeShareable(new FPCGEditorGraphNodeFactory());
 	FEdGraphUtilities::RegisterVisualNodeFactory(GraphNodeFactory);
 
@@ -124,6 +127,8 @@ void FPCGEditorModule::ShutdownModule()
 	UnregisterDetailsCustomizations();
 	UnregisterMenuExtensions();
 
+	UPCGEditorMode::UnregisterEditorMode();
+
 	FPCGNodeSourceEditorTextBoxCommands::Unregister();
 	FPCGEditorSpawnNodeCommands::Unregister();
 	FPCGEditorCommands::Unregister();
diff --git a/Engine/Plugins/PCG/Source/PCGEditor/Public/EditorMode/PCGEdMode.h b/Engine/Plugins/PCG/Source/PCGEditor/Public/EditorMode/PCGEdMode.h
new file mode 100644
--- /dev/null
+++ b/Engine/Plugins/PCG/Source/PCGEditor/Public/EditorMode/PCGEdMode.h
@@ -0,0 +1,66 @@
+// Copyright Epic Games, Inc. All Rights Reserved.
+
+#pragma once
+
+#include "EdMode.h"
+#include "InteractiveToolManager.h"
+
+#include "PCGEdMode.generated.h"
+
+/** The unique state for the editor to be in, specific to PCG's toolkits. */
+UCLASS(Transient)
+class UPCGEditorMode : public UEdMode
+{
+	GENERATED_BODY()
+
+public:
+	/** Unique name for PCG Editor Mode */
+	PCGEDITOR_API const static FEditorModeID EM_PCGEditorModeId;
+
+	UPCGEditorMode();
+
+	static void RegisterEditorMode();
+	static void UnregisterEditorMode();
+
+	//~Begin UEdMode interface
+	virtual void ActorSelectionChangeNotify() override;
+	virtual bool ProcessEditDelete() override;
+	virtual bool ProcessEditCut() override;
+	virtual bool CanAutoSave() const override;
+	/** Override the viewport pivot, if needed. */
+	virtual bool GetPivotForOrbit(FVector& OutPivot) const override;
+	/** Focus event when the Editor enters PCG Editor Mode. */
+	virtual void Enter() override;
+	/** Focus event when the Editor leaves PCG Editor Mode. */
+	virtual void Exit() override;
+	/** Used to restrict starting a tool, if a certain state is required beforehand. */
+	virtual bool ShouldToolStartBeAllowed(const FString& ToolIdentifier) const override;
+	virtual bool HasCustomViewportFocus() const override;
+	virtual FBox ComputeCustomViewportFocus() const override;
+	//~End UEdMode interface
+
+protected:
+	virtual void CreateToolkit() override;
+
+	virtual void OnToolStarted(UInteractiveToolManager* Manager, UInteractiveTool* Tool) override;
+	virtual void OnToolEnded(UInteractiveToolManager* Manager, UInteractiveTool* Tool) override;
+	virtual void OnToolPostBuild(UInteractiveToolManager* InToolManager, EToolSide InSide, UInteractiveTool* InBuiltTool, UInteractiveToolBuilder* InToolBuilder, const FToolBuilderState& ToolState);
+
+	void OnEditorClosed() const;
+	void AcceptActiveToolActionOrTool() const;
+	void CancelActiveToolActionOrTool() const;
+
+	void ConfigureRealTimeViewportsOverride(bool bEnable);
+
+private:
+	// @todo_pcg: Activate when custom tools are supported.
+	// void RegisterCustomTool(
+	// 	TSharedPtr<FUICommandInfo> UICommand,
+	// 	FString ToolIdentifier,
+	// 	UInteractiveToolBuilder* Builder,
+	// 	const TFunction<bool(UInteractiveToolManager*, EToolSide)>& ExecuteAction,
+	// 	const TFunction<bool(UInteractiveToolManager*, EToolSide)>& CanExecuteAction,
+	// 	const TFunction<bool(UInteractiveToolManager*, EToolSide)>& IsActionChecked) const;
+
+	bool bIsToolActive = false;
+};
diff --git a/Engine/Plugins/PCG/Source/PCGEditor/Public/EditorMode/PCGEdModeCommands.h b/Engine/Plugins/PCG/Source/PCGEditor/Public/EditorMode/PCGEdModeCommands.h
new file mode 100644
--- /dev/null
+++ b/Engine/Plugins/PCG/Source/PCGEditor/Public/EditorMode/PCGEdModeCommands.h
@@ -0,0 +1,96 @@
+// Copyright Epic Games, Inc. All Rights Reserved.
+
+#pragma once
+
+#include "Framework/Commands/Commands.h"
+#include "Tools/InteractiveToolsCommands.h"
+
+/** General actions related to the PCG Editor Mode. */
+class PCGEDITOR_API FPCGEditorModeCommands : public TCommands<FPCGEditorModeCommands>
+{
+public:
+	FPCGEditorModeCommands();
+
+	virtual void RegisterCommands() override;
+
+	// @todo_pcg: Any file menu, hotkeys, etc related to the editor mode in general can go here.
+};
+
+/** Actions related to the PCG Editor Mode Toolkit Palette. */
+class PCGEDITOR_API FPCGEditorModePaletteCommands : public TCommands<FPCGEditorModePaletteCommands>
+{
+public:
+	FPCGEditorModePaletteCommands();
+
+	/*** PCG EDITOR CONTEXTS **/
+	TSharedPtr<FUICommandInfo> LoadLineContextPalette;
+	TSharedPtr<FUICommandInfo> LoadPaintContextPalette;
+	TSharedPtr<FUICommandInfo> LoadVolumeContextPalette;
+
+	// @todo_pcg: Selection/Filtering actions.
+
+	virtual void RegisterCommands() override;
+
+	// @todo_pcg: Allow extension for custom tools. See Modeling Tools as an example.
+};
+
+/** Tool specific actions. */
+class FPCGEditorModeToolCommands : public TInteractiveToolCommands<FPCGEditorModeToolCommands>
+{
+public:
+	FPCGEditorModeToolCommands();
+
+	virtual void RegisterCommands() override;
+
+	/** Register all Tool command sets. Call this in module startup */
+	static void RegisterAllToolCommands();
+
+	/** Unregister all Tool command sets. Call this from module shutdown. */
+	static void UnregisterAllToolCommands();
+
+	/** Needed to retrieve the CDOs to call GetActionSet directly on the tool. */
+	virtual void GetToolDefaultObjectList(TArray<UInteractiveTool*>& ToolCDOs) override;
+
+	// @todo_pcg: Implement if needed for custom tool commands.
+	/**
+	 * Add the tool-specific commands from an interactive tool to a specified UICommandList 
+	 * Call this when the active tool changes (eg on ToolManager.OnToolStarted / OnToolEnded)
+	 */
+	// static void AddToolCommandBinding(UInteractiveTool* Tool, TSharedPtr<FUICommandList> UICommandList);
+	// static void RemoveToolCommandBinding(UInteractiveTool* Tool, TSharedPtr<FUICommandList> UICommandList);
+
+	/*** PCG EDITOR MODE TOOLS - Line Context **/
+	TSharedPtr<FUICommandInfo> EnableDrawSplineTool;
+	TSharedPtr<FUICommandInfo> EnableDrawSurfaceTool;
+	TArray<TSharedPtr<FUICommandInfo>> GetLineContextToolCommands() const;
+
+	/*** PCG EDITOR MODE TOOLS - Paint Context **/
+	TSharedPtr<FUICommandInfo> EnablePaintTool;
+	TArray<TSharedPtr<FUICommandInfo>> GetPaintContextToolCommands() const;
+
+	/*** PCG EDITOR MODE TOOLS - Volume Context **/
+	TSharedPtr<FUICommandInfo> EnableCreateVolumeTool;
+	TArray<TSharedPtr<FUICommandInfo>> GetVolumeContextToolCommands() const;
+};
+
+/**
+ * Macro for the class declaration for each tool as a TInteractiveToolCommands class
+ * @todo_pcg: This is basically done elsewhere for other toolkits, but can probably be modernized using C++ Concepts 
+ * @param ToolCommandClassName The tool's class name
+ */
+#define DECLARE_TOOL_COMMANDS(ToolCommandClassName) \
+class ToolCommandClassName : public TInteractiveToolCommands<ToolCommandClassName> \
+{\
+public:\
+ToolCommandClassName();\
+\
+virtual void GetToolDefaultObjectList(TArray<UInteractiveTool*>& ToolCDOs) override;\
+};\
+
+// Define all tools commands here
+DECLARE_TOOL_COMMANDS(FPCGDrawSplineCommands);
+DECLARE_TOOL_COMMANDS(FPCGDrawSurfaceCommands);
+DECLARE_TOOL_COMMANDS(FPCGPaintCommands);
+DECLARE_TOOL_COMMANDS(FPCGCreateVolumeCommands);
+
+#undef DECLARE_TOOL_COMMANDS
diff --git a/Engine/Plugins/PCG/Source/PCGEditor/Public/EditorMode/PCGEdModeHelpers.h b/Engine/Plugins/PCG/Source/PCGEditor/Public/EditorMode/PCGEdModeHelpers.h
new file mode 100644
--- /dev/null
+++ b/Engine/Plugins/PCG/Source/PCGEditor/Public/EditorMode/PCGEdModeHelpers.h
@@ -0,0 +1,20 @@
+// Copyright Epic Games, Inc. All Rights Reserved.
+
+#pragma once
+
+#include "Framework/Notifications/NotificationManager.h"
+#include "Widgets/Notifications/SNotificationList.h"
+
+namespace PCG::EditorMode
+{
+	static constexpr float DefaultNotificationDuration = 2.5f;
+
+	inline void DispatchEditorNotification(const FText& Text, const FText& SubText = FText{}, const float Duration = DefaultNotificationDuration)
+	{
+		FNotificationInfo NotificationInfo(std::move(Text));
+		NotificationInfo.ExpireDuration = Duration;
+		NotificationInfo.bFireAndForget = true;
+		NotificationInfo.SubText = std::move(SubText);
+		FSlateNotificationManager::Get().AddNotification(NotificationInfo);
+	}
+}
diff --git a/Engine/Plugins/PCG/Source/PCGEditor/Public/EditorMode/PCGEdModeSettings.h b/Engine/Plugins/PCG/Source/PCGEditor/Public/EditorMode/PCGEdModeSettings.h
new file mode 100644
--- /dev/null
+++ b/Engine/Plugins/PCG/Source/PCGEditor/Public/EditorMode/PCGEdModeSettings.h
@@ -0,0 +1,57 @@
+// Copyright Epic Games, Inc. All Rights Reserved.
+
+#pragma once
+
+#include "Engine/DeveloperSettings.h"
+
+#include "PCGEdModeSettings.generated.h"
+
+/** User preferences related to the PCG Editor Mode. */
+UCLASS(Config = EditorPerProjectUserSettings, DisplayName = "PCG Editor Mode Settings")
+class UPCGEditorModeSettings : public UDeveloperSettings
+{
+	GENERATED_BODY()
+
+public:
+	UPCGEditorModeSettings(const FObjectInitializer& ObjectInitializer)
+		: UDeveloperSettings(ObjectInitializer)
+		, bEnableEditorMode(false) // @todo_pcg: Added to the constructor to avoid false as default. Remove when permanently enabling toolkit.
+	{}
+
+	virtual FName GetContainerName() const override { return FName("Editor"); }
+	virtual FName GetCategoryName() const override { return FName("Plugins"); }
+	virtual FName GetSectionName() const override { return FName("PCGEditorMode"); }
+
+	/** Editor Mode user configurations. */
+	UPROPERTY(EditAnywhere, Config, Category = "General", meta = (DisplayName = "Enable the PCG Editor Mode Toolkit (requires restart)", ConfigRestartRequired = true))
+	bool bEnableEditorMode = true;
+
+	/** The rate at which PCG Graphs will be refreshed while actively editing with a tool. */
+	UPROPERTY(EditAnywhere, Category = "Component", meta = (
+		ForceUnits = "Seconds",
+		ConsoleVariable = "pcg.editormode.GraphRefreshRate"))
+	float GraphRefreshRate = 1.5f;
+
+	/** Hide the tool buttons from the UI when a tool is actively being used. */
+	UPROPERTY(EditAnywhere, Config, Category = "Tools", meta = (
+		EditCondition = "bEnableEditorMode",
+		EditConditionHides,
+		ConsoleVariable = "pcg.editormode.HideToolButtonsDuringActiveTool"))
+	bool bHideToolButtonsDuringActiveTool = false;
+
+	/** Enable an editor toast issued when an Editor Mode tool has an error. */
+	UPROPERTY(EditAnywhere, Config, Category = "Tools", meta = (
+		EditCondition = "bEnableEditorMode",
+		EditConditionHides,
+		ConsoleVariable = "pcg.editormode.ShowEditorToastOnToolErrors"))
+	bool bShowEditorToastOnToolErrors = true;
+
+	UPROPERTY(EditAnywhere, Category = "Tools|Output", meta = (EditCondition = "bEnableEditorMode", EditConditionHides))
+	FName DefaultNewActorName = "PCG Tool Actor";
+
+	UPROPERTY(EditAnywhere, Category = "Tools|Output", meta = (EditCondition = "bEnableEditorMode", EditConditionHides))
+	FName DefaultNewPCGComponentName = "PCG Tool Component";
+
+	UPROPERTY(EditAnywhere, Category = "Tools|Output|Spline", meta = (EditCondition = "bEnableEditorMode", EditConditionHides))
+	FName DefaultNewSplineComponentName = "PCG Tool Spline Component";
+};
diff --git a/Engine/Plugins/PCG/Source/PCGEditor/Public/EditorMode/PCGEdModeStyle.h b/Engine/Plugins/PCG/Source/PCGEditor/Public/EditorMode/PCGEdModeStyle.h
new file mode 100644
--- /dev/null
+++ b/Engine/Plugins/PCG/Source/PCGEditor/Public/EditorMode/PCGEdModeStyle.h
@@ -0,0 +1,20 @@
+// Copyright Epic Games, Inc. All Rights Reserved.
+
+#pragma once
+
+#include "Styling/SlateStyle.h"
+
+struct FSlateIcon;
+
+/** Style details specific to the PCG Editor Mode. */
+class FPCGEditorModeStyle : public FSlateStyleSet
+{
+public:
+	static void Register();
+	static void Unregister();
+
+	static const FPCGEditorModeStyle& Get();
+
+private:
+	FPCGEditorModeStyle();
+};
diff --git a/Engine/Plugins/PCG/Source/PCGEditor/Public/EditorMode/PCGEdModeToolkit.h b/Engine/Plugins/PCG/Source/PCGEditor/Public/EditorMode/PCGEdModeToolkit.h
new file mode 100644
--- /dev/null
+++ b/Engine/Plugins/PCG/Source/PCGEditor/Public/EditorMode/PCGEdModeToolkit.h
@@ -0,0 +1,86 @@
+// Copyright Epic Games, Inc. All Rights Reserved.
+
+#pragma once
+
+#include "InteractiveTool.h"
+#include "StatusBarSubsystem.h"
+#include "Toolkits/BaseToolkit.h"
+
+#include "Widgets/SPCGEdModeAcceptToolOverlay.h"
+
+class UPCGInteractiveTool;
+
+/** The basic toolkit needed to show the docked editor mode panel and set up the tool shelves and details views. */
+class FPCGEditorModeToolkit : public FModeToolkit
+{
+public:
+	// ~Begin FModeToolkit interface
+	virtual void Init(const TSharedPtr<IToolkitHost>& InitToolkitHost, TWeakObjectPtr<UEdMode> OwningMode) override;
+	virtual FName GetToolkitFName() const override;
+	virtual FText GetBaseToolkitName() const override;
+
+	/** Returns the PCG Editor Mode specific tabs in the mode toolbar */
+	virtual void GetToolPaletteNames(TArray<FName>& PaletteNames) const override;
+	virtual FText GetToolPaletteDisplayName(FName PaletteName) const override;
+	virtual void OnToolPaletteChanged(FName PaletteName) override;
+	virtual bool HasIntegratedToolPalettes() const override { return false; }
+
+	virtual FText GetActiveToolDisplayName() const override;
+	virtual FText GetActiveToolMessage() const override;
+
+	virtual void OnToolStarted(UInteractiveToolManager* Manager, UInteractiveTool* Tool) override;
+	virtual void OnToolEnded(UInteractiveToolManager* Manager, UInteractiveTool* Tool) override;
+
+	// Keep track of the widget for customization.
+	virtual TSharedPtr<SWidget> GetInlineContent() const override { return ToolkitWidget; }
+	// ~End FModeToolkit interface
+
+	virtual void PostToolNotification(const FText& Message);
+	virtual void ClearToolNotification();
+
+	virtual void PostToolWarning(const FText& Message);
+	virtual void ClearToolWarning();
+
+	virtual void ShowModeWarnings();
+
+	/*** UI Events ***/
+	void OnActiveViewportChanged(TSharedPtr<IAssetViewport>, TSharedPtr<IAssetViewport>) const;
+
+	/*** Accept/Cancel Window ***/
+	const FSlateBrush* GetActiveToolIcon() const;
+	bool IsAcceptButtonEnabled() const;
+	bool IsCancelButtonEnabled() const;
+	EVisibility GetAcceptButtonVisibility() const;
+	EVisibility GetCancelButtonVisibility() const;
+	FReply OnAcceptButtonClicked() const;
+	FReply OnCancelButtonClicked() const;
+
+private:
+	void RegisterToolkitCallbacks();
+	void UnregisterToolkitCallbacks();
+
+	void RegisterToolCallbacks(UPCGInteractiveTool* CurrentTool);
+	void UnregisterToolCallbacks(UPCGInteractiveTool* CurrentTool) const;
+
+	void OnEditorModeSettingsChanged(UObject* InObject, FPropertyChangedEvent& InChangedEvent);
+
+	void UpdateActiveTool();
+	void InvalidateCachedDetailPanelState(UObject* ChangedObject) const;
+
+	/** A utility function to register the tool palettes with the ToolkitBuilder. */
+	void RegisterPalettes();
+	/** A utility function to create additional widgets for the Toolkit. */
+	void CreateWidgets();
+
+	/** UI overlay in the viewport to accept or cancel the current tool. */
+	TSharedPtr<SPCGEdModeAcceptToolOverlay> AcceptToolWidget;
+
+	/** Currently only used to cancel the current tool. */
+	FDelegateHandle ActivePaletteChangedHandle;
+
+	FStatusBarMessageHandle ActiveToolMessageHandle;
+	FText ActiveWarning{};
+	FInteractiveToolInfo CachedActiveToolInfo;
+};
+
+#undef LOCTEXT_NAMESPACE
diff --git a/Engine/Plugins/PCG/Source/PCGEditor/Public/EditorMode/Tools/Helpers/PCGEdModeActorHelpers.h b/Engine/Plugins/PCG/Source/PCGEditor/Public/EditorMode/Tools/Helpers/PCGEdModeActorHelpers.h
new file mode 100644
--- /dev/null
+++ b/Engine/Plugins/PCG/Source/PCGEditor/Public/EditorMode/Tools/Helpers/PCGEdModeActorHelpers.h
@@ -0,0 +1,28 @@
+// Copyright Epic Games, Inc. All Rights Reserved.
+
+#pragma once
+
+#include "Helpers/PCGActorHelpers.h"
+
+namespace PCG::EditorMode
+{
+	namespace Actor
+	{
+		/** Spawn an actor for the toolkit with spawn parameters. */
+		AActor* Spawn(const UPCGActorHelpers::FSpawnDefaultActorParams& SpawnParams);
+		/** Spawn an actor for the toolkit with explicit parameters. */
+		AActor* Spawn(UWorld* World, FName ActorName, const FTransform& Transform = FTransform::Identity);
+		/** Helper to spawn a temporary working actor for a tool. */
+		AActor* SpawnWorking(UWorld* World, const FTransform& Transform = FTransform::Identity);
+	}
+
+	namespace PCGComponent
+	{
+		/** Find a PCG Component on an actor if it exists. */
+		UPCGComponent* Find(const AActor* Owner);
+		/** Spawn a PCG Component on an actor. */
+		UPCGComponent* Create(AActor* Owner, FName ComponentName);
+		/** Helper to spawn a temporary working PCG component for a tool. */
+		UPCGComponent* CreateWorking(AActor* Owner);
+	}
+}
diff --git a/Engine/Plugins/PCG/Source/PCGEditor/Public/EditorMode/Tools/Helpers/PCGEdModeSceneQueryHelpers.h b/Engine/Plugins/PCG/Source/PCGEditor/Public/EditorMode/Tools/Helpers/PCGEdModeSceneQueryHelpers.h
new file mode 100644
--- /dev/null
+++ b/Engine/Plugins/PCG/Source/PCGEditor/Public/EditorMode/Tools/Helpers/PCGEdModeSceneQueryHelpers.h
@@ -0,0 +1,58 @@
+// Copyright Epic Games, Inc. All Rights Reserved.
+
+#pragma once
+
+#include "InputState.h"
+#include "Math/Vector.h"
+#include "Misc/Optional.h"
+
+struct FViewCameraState;
+struct FHitResult;
+
+namespace PCG::EditorMode::Scene
+{
+	/** Custom parameters for UI ray casts from the current view into the scene. */
+	struct FViewRayParams
+	{
+		const UWorld* World = nullptr;
+		double Distance = UE_DOUBLE_BIG_NUMBER;
+		const TArray<const UPrimitiveComponent*>* ComponentsToIgnore = nullptr;
+		const TArray<const UPrimitiveComponent*>* AllowedInvisibleComponents = nullptr;
+	};
+
+	/** The resulting data from a UI ray cast hit. */
+	struct FViewHitResult
+	{
+		FVector ImpactPosition;
+		FVector ImpactNormal;
+		double Depth;
+
+		FInputRayHit ToInputRayHit() const
+		{
+			FInputRayHit OutRayHit;
+			OutRayHit.bHit = true;
+			OutRayHit.bHasHitNormal = true;
+			OutRayHit.HitNormal = ImpactNormal;
+			OutRayHit.HitDepth = Depth;
+			OutRayHit.HitObject = nullptr;
+			OutRayHit.HitOwner = nullptr;
+
+			return OutRayHit;
+		}
+	};
+
+	/** The hit actor or primitive component is visible. */
+	bool IsHitTargetVisible(const FHitResult& HitResult);
+
+	/** Cast a ray and if a hit is recorded, return the hit result nearest object. */
+	TOptional<FHitResult> TraceToNearestObject(const FInputDeviceRay& Ray, const FViewRayParams& Params);
+
+	/** Cast a ray from the view and if a hit is recorded, return a condensed view-related hit result. */
+	TOptional<FViewHitResult> ViewportRay(const FInputDeviceRay& Ray, const FViewRayParams& Params);
+
+	/** Cast a ray from the view to a plane in space. If a hit is recorded, return a condensed view-related hit result. */
+	TOptional<FViewHitResult> ViewportRayAgainstPlane(const FInputDeviceRay& Ray, const FViewRayParams& Params, const FPlane& Plane);
+
+	/** Cast a ray from the view to a plane parallel to the user view. If a hit is recorded, return a condensed view-related hit result. */
+	TOptional<FViewHitResult> ViewportRayAgainstCameraPlane(const FInputDeviceRay& Ray, const FViewRayParams& Params, const FViewCameraState& CameraState);
+} // namespace PCG::EditorMode::Scene
diff --git a/Engine/Plugins/PCG/Source/PCGEditor/Public/EditorMode/Tools/Helpers/PCGEdModeSceneVisualizeHelpers.h b/Engine/Plugins/PCG/Source/PCGEditor/Public/EditorMode/Tools/Helpers/PCGEdModeSceneVisualizeHelpers.h
new file mode 100644
--- /dev/null
+++ b/Engine/Plugins/PCG/Source/PCGEditor/Public/EditorMode/Tools/Helpers/PCGEdModeSceneVisualizeHelpers.h
@@ -0,0 +1,17 @@
+// Copyright Epic Games, Inc. All Rights Reserved.
+
+#pragma once
+
+#include "CoreMinimal.h"
+
+class IToolsContextRenderAPI;
+class USplineComponent;
+
+namespace PCG::EditorMode::Scene::Visualize
+{
+	// Visualize a spline within the editor viewport.
+	void Spline(const USplineComponent& Spline, IToolsContextRenderAPI& RenderAPI, int32 LatestIndex = INDEX_NONE, bool bShouldVisualizeScale = false, bool bShouldVisualizeTangents = true);
+
+	// Visualize a spline tangent within the editor viewport.
+	void SplineTangent(const USplineComponent& Spline, IToolsContextRenderAPI& RenderAPI, const int32 SplinePointIndex);
+}
diff --git a/Engine/Plugins/PCG/Source/PCGEditor/Public/EditorMode/Tools/Helpers/PCGEdModeSplineHelpers.h b/Engine/Plugins/PCG/Source/PCGEditor/Public/EditorMode/Tools/Helpers/PCGEdModeSplineHelpers.h
new file mode 100644
--- /dev/null
+++ b/Engine/Plugins/PCG/Source/PCGEditor/Public/EditorMode/Tools/Helpers/PCGEdModeSplineHelpers.h
@@ -0,0 +1,23 @@
+// Copyright Epic Games, Inc. All Rights Reserved.
+
+#pragma once
+
+#include "CoreMinimal.h"
+
+class USplineComponent;
+class AActor;
+
+namespace PCG::EditorMode::Spline
+{
+	USplineComponent* CreateComponent(
+		AActor* Actor,
+		FName ComponentName,
+		const USplineComponent* TemplateComponent = nullptr,
+		const bool bTransactional = false);
+
+	USplineComponent* CreateWorkingComponent(AActor* Actor);
+
+	USplineComponent* FindComponent(const AActor* Actor, const int32 TargetIndex = INDEX_NONE);
+
+	void CopySplinePoints(const USplineComponent& Source, USplineComponent& Destination, const bool bTransactional = false);
+}
diff --git a/Engine/Plugins/PCG/Source/PCGEditor/Public/EditorMode/Tools/Line/PCGDrawSplineTool.h b/Engine/Plugins/PCG/Source/PCGEditor/Public/EditorMode/Tools/Line/PCGDrawSplineTool.h
new file mode 100644
--- /dev/null
+++ b/Engine/Plugins/PCG/Source/PCGEditor/Public/EditorMode/Tools/Line/PCGDrawSplineTool.h
@@ -0,0 +1,129 @@
+// Copyright Epic Games, Inc. All Rights Reserved.
+
+#pragma once
+
+#include "EditorMode/Tools/PCGInteractiveTool.h"
+#include "EditorMode/Tools/PCGInteractiveToolBuilder.h"
+#include "EditorMode/Tools/PCGInteractiveToolCommon.h"
+#include "EditorMode/Tools/PCGInteractiveToolSettings.h"
+#include "EditorMode/Tools/Helpers/PCGEdModeSceneQueryHelpers.h"
+
+#include "BaseBehaviors/BehaviorTargetInterfaces.h"
+
+#include "PCGDrawSplineTool.generated.h"
+
+class USplineComponent;
+class USingleClickOrDragInputBehavior;
+
+UENUM(BlueprintType)
+enum class EPCGDrawSplineDrawMode : uint8
+{
+	TangentDrag UMETA(ToolTip = "Click to place a point and then drag to set its tangent. Clicking without dragging will create sharp corners."),
+	ClickAutoTangent UMETA(ToolTip = "Click and drag new points, with the tangent set automatically."),
+	FreeDraw UMETA(ToolTip = "Drag to place multiple points, with spacing controlled by Min Point Spacing.")
+};
+
+UCLASS(BlueprintType)
+class PCGEDITOR_API UPCGDrawSplineToolSettings : public UPCGInteractiveToolSettings
+{
+	GENERATED_BODY()
+
+public:
+	UPCGDrawSplineToolSettings();
+
+	// ~Begin UPCGInteractiveToolSettings interface
+	virtual FText GetToolTag() const override;
+	/** Apply the results of the tool. */
+	bool ApplyTool(UWorld* World, const USplineComponent* ToolSplineComponent, FText& OutWarning) const;
+	virtual FSoftObjectPath GetDefaultGraph() const override;
+	// ~End UPCGInteractiveToolSettings interface
+
+	UPROPERTY(EditAnywhere, Category = "Spline")
+	EPCGDrawSplineDrawMode DrawMode = EPCGDrawSplineDrawMode::TangentDrag;
+
+	/** Minimum world distance to cover while free drawing before a point will be generated. */
+	UPROPERTY(EditAnywhere, Category = "Spline", meta = (
+			ClampMin = 0,
+			UIMax = 1000,
+			EditCondition = "DrawMode == EPCGDrawSplineDrawMode::FreeDraw",
+			EditConditionHides))
+	double MinPointSpacing = 200.0;
+
+	/** The abstract target upon which to draw the spline. By default, it will be a raycast to the world. */
+	UPROPERTY(EditAnywhere, Category = "Target", meta = (ShowOnlyInnerProperties))
+	FPCGToolDrawTarget Target{};
+
+	UPROPERTY(EditAnywhere, Category = "PCG", meta = (TransientToolProperty))
+	FName SplineName;
+};
+
+UCLASS(Experimental)
+class PCGEDITOR_API UPCGDrawSplineTool : public UPCGInteractiveTool, public IClickBehaviorTarget, public IClickDragBehaviorTarget
+{
+	GENERATED_BODY()
+
+public:
+	// ~Begin UInteractiveTool interface
+	virtual void Setup() override;
+	virtual void Shutdown(EToolShutdownType ShutdownType) override;
+	virtual void OnTick(float DeltaTime) override;
+
+	virtual bool HasCancel() const override { return true; }
+	virtual bool HasAccept() const override { return true; }
+	virtual bool CanAccept() const override;
+	virtual void Render(IToolsContextRenderAPI* RenderAPI) override;
+	// ~End UInteractiveTool interface
+
+	// ~Begin UPCGInteractiveTool interface
+	virtual void OnRefresh() override;
+	virtual TSubclassOf<UPCGInteractiveToolSettings> GetSettingsClass() const override { return UPCGDrawSplineToolSettings::StaticClass(); }
+	// ~End UPCGInteractiveTool interface
+
+	// ~Begin IClickBehaviorTarget interface
+	virtual FInputRayHit IsHitByClick(const FInputDeviceRay& ClickPosition) override;
+	virtual void OnClicked(const FInputDeviceRay& ClickPosition) override;
+	// ~End IClickBehaviorTarget interface
+
+	// ~Begin IClickDragBehaviorTarget interface
+	virtual FInputRayHit CanBeginClickDragSequence(const FInputDeviceRay& DragPosition) override;
+	virtual void OnClickPress(const FInputDeviceRay& ClickPosition) override;
+	virtual void OnClickDrag(const FInputDeviceRay& DragPosition) override;
+	virtual void OnClickRelease(const FInputDeviceRay& ReleasePosition) override;
+	virtual void OnTerminateDragSequence() override;
+	// ~End IClickDragBehaviorTarget interface
+
+	void SetIsSplineSurface(const bool bInIsSplineSurface) { bIsSplineSurface = bInIsSplineSurface; }
+
+private:
+	void AddSplinePoint(const FVector& HitLocation, const FVector& HitNormal) const;
+	FVector GetNewPointNormal(const FVector& HitNormal) const;
+	void ApplyToolInternal() const;
+
+	UPROPERTY(Instanced)
+	TObjectPtr<USingleClickOrDragInputBehavior> ClickOrDragBehavior = nullptr;
+
+	UPROPERTY(Transient)
+	TObjectPtr<USplineComponent> WorkingSpline = nullptr;
+
+	PCG::EditorMode::Scene::FViewRayParams CachedRayParams{};
+	bool bIsSplineSurface = false;
+	bool bDrawTangentForPreviousPoint = false;
+	bool bFreeDrawPlacedPreviewPoint = false;
+	int32 FreeDrawStrokeStartIndex = 0;
+	double CachedMinPointSpacingSquared = 0;
+};
+
+UCLASS(Transient)
+class PCGEDITOR_API UPCGDrawSplineToolBuilder : public UPCGInteractiveToolBuilder
+{
+	GENERATED_BODY()
+
+public:
+	virtual bool CanBuildTool(const FToolBuilderState& SceneState) const override;
+	virtual UInteractiveTool* BuildTool(const FToolBuilderState& SceneState) const override;
+
+	void SetIsSplineSurface(const bool bInIsSplineSurface) { bIsSplineSurface = bInIsSplineSurface; }
+
+private:
+	bool bIsSplineSurface = false;
+};
diff --git a/Engine/Plugins/PCG/Source/PCGEditor/Public/EditorMode/Tools/PCGInteractiveTool.h b/Engine/Plugins/PCG/Source/PCGEditor/Public/EditorMode/Tools/PCGInteractiveTool.h
new file mode 100644
--- /dev/null
+++ b/Engine/Plugins/PCG/Source/PCGEditor/Public/EditorMode/Tools/PCGInteractiveTool.h
@@ -0,0 +1,63 @@
+// Copyright Epic Games, Inc. All Rights Reserved.
+
+#pragma once
+
+#include "InteractiveTool.h"
+#include "InteractiveToolManager.h"
+
+#include "PCGInteractiveTool.generated.h"
+
+DECLARE_MULTICAST_DELEGATE_OneParam(OnPCGInteractiveToolWarningChanged, const FText& Warning);
+
+class UPCGComponent;
+class UPCGInteractiveToolSettings;
+
+UCLASS(Abstract, Transient)
+class PCGEDITOR_API UPCGInteractiveTool : public UInteractiveTool
+{
+	GENERATED_BODY()
+
+public:
+	// ~Begin UInteractiveTool interface
+	virtual void Setup() override;
+	virtual void Shutdown(EToolShutdownType ShutdownType) override;
+	virtual void OnTick(float DeltaTime) override;
+	// ~End UInteractiveTool interface
+
+	virtual void OnRefresh();
+	virtual void SetTargetActor(AActor* InTargetActor);
+	virtual TWeakObjectPtr<AActor> GetTargetActor() const;
+
+	UWorld* GetContextualWorld() const;
+	UWorld* GetContextualWorldChecked() const;
+
+	OnPCGInteractiveToolWarningChanged OnToolWarningChanged;
+
+protected:
+	/** The concrete class of the tool. This must be defined in the subclass. */
+	virtual TSubclassOf<UPCGInteractiveToolSettings> GetSettingsClass() const PURE_VIRTUAL(UPCGInteractiveTool::GetSettingsClass, return nullptr;)
+
+	/** Set the warning for the tool. */
+	virtual void SetWarning(const FText& Warning) const;
+	/** Clear the warning for the tool. */
+	virtual void ClearWarning() const;
+	/** Notify the tool of an error to log, display, and optionally toast. */
+	virtual void NotifyError(const FText& ErrorMessage) const;
+
+	// @todo_pcg: Convert to a target object, rather than actor and component, to ease transition into SceneGraph
+	/** The tool can work directly on this temporary actor until effects are applied. */
+	TWeakObjectPtr<AActor> WorkingActor = nullptr;
+	/** The tool can work directly on this temporary PCG component until effects are applied. */
+	TWeakObjectPtr<UPCGComponent> WorkingComponent = nullptr;
+
+public:
+	UPROPERTY(EditAnywhere, Category = "Tool Settings", meta = (ShowOnlyInnerProperties))
+	TObjectPtr<UPCGInteractiveToolSettings> ToolSettings = nullptr;
+
+private:
+	TWeakObjectPtr<UWorld> TargetWorld = nullptr;
+	mutable FText WarningMessage;
+
+	/** The current time left before a graph refresh (in seconds). */
+	float GraphRefreshTimer = 0.f;
+};
diff --git a/Engine/Plugins/PCG/Source/PCGEditor/Public/EditorMode/Tools/PCGInteractiveToolBuilder.h b/Engine/Plugins/PCG/Source/PCGEditor/Public/EditorMode/Tools/PCGInteractiveToolBuilder.h
new file mode 100644
--- /dev/null
+++ b/Engine/Plugins/PCG/Source/PCGEditor/Public/EditorMode/Tools/PCGInteractiveToolBuilder.h
@@ -0,0 +1,20 @@
+// Copyright Epic Games, Inc. All Rights Reserved.
+
+#pragma once
+
+#include "InteractiveToolBuilder.h"
+
+#include "PCGInteractiveToolBuilder.generated.h"
+
+UCLASS()
+class PCGEDITOR_API UPCGInteractiveToolBuilder : public UInteractiveToolBuilder
+{
+	GENERATED_BODY()
+
+public:
+	// ~Begin UInteractiveToolBuilder interface
+	virtual bool CanBuildTool(const FToolBuilderState& SceneState) const override;
+	virtual UInteractiveTool* BuildTool(const FToolBuilderState& SceneState) const override;
+	virtual void PostBuildTool(UInteractiveTool* Tool, const FToolBuilderState& SceneState) const override;
+	// ~End UInteractiveToolBuilder interface
+};
diff --git a/Engine/Plugins/PCG/Source/PCGEditor/Public/EditorMode/Tools/PCGInteractiveToolCommon.h b/Engine/Plugins/PCG/Source/PCGEditor/Public/EditorMode/Tools/PCGInteractiveToolCommon.h
new file mode 100644
--- /dev/null
+++ b/Engine/Plugins/PCG/Source/PCGEditor/Public/EditorMode/Tools/PCGInteractiveToolCommon.h
@@ -0,0 +1,68 @@
+// Copyright Epic Games, Inc. All Rights Reserved.
+
+#pragma once
+
+#include "Helpers/PCGEdModeSceneQueryHelpers.h"
+
+#include "PCGInteractiveToolCommon.generated.h"
+
+UENUM()
+enum class EPCGToolInteractionMode : uint8
+{
+	CreateNewActor,
+	AddToSelectedActor
+};
+
+UENUM()
+enum class EPCGToolDrawTargetOffset : uint8
+{
+	// Drawn points will be at the hit location.
+	HitLocation,
+	// Drawn points will be at the hit location with a distance offset.
+	DistanceOffset,
+	// Drawn points will be at the hit location with a custom vector offset.
+	ExplicitOffset
+};
+
+UENUM()
+enum class EPCGToolDrawTargetNormal : uint8
+{
+	// Drawn points will adopt the impact normal.
+	HitNormal,
+	// Drawn points will adopt the world up normal.
+	WorldUp,
+	// Align to the previous spline point's up vector.
+	AlignToPrevious,
+	// Drawn points will adopt an explicit normal.
+	Explicit
+};
+
+USTRUCT()
+struct PCGEDITOR_API FPCGToolDrawTarget
+{
+	GENERATED_BODY()
+
+	/** How to choose the direction to offset points from the drawn surface. */
+	UPROPERTY(EditAnywhere, Category = "Draw Target")
+	EPCGToolDrawTargetOffset OffsetMode = EPCGToolDrawTargetOffset::HitLocation;
+
+	/** How far to offset drawn points from the clicked surface, along the normal. */
+	UPROPERTY(EditAnywhere, Category = "Draw Target", DisplayName = "Offset", meta = (UIMin = 0, UIMax = 100, EditCondition = "OffsetMode == EPCGToolDrawTargetOffset::DistanceOffset", EditConditionHides))
+	double OffsetDistance = 100.0;
+
+	/** Manual vector to offset drawn points from the clicked surface. */
+	UPROPERTY(EditAnywhere, Category = "Draw Target", DisplayName = "Offset", meta = (UIMin = 0, UIMax = 100, EditCondition = "OffsetMode == EPCGToolDrawTargetOffset::ExplicitOffset", EditConditionHides))
+	FVector OffsetVector = FVector(0, 0, 100);
+
+	/** How to choose the direction to offset points from the drawn surface. */
+	UPROPERTY(EditAnywhere, Category = "Draw Target")
+	EPCGToolDrawTargetNormal NormalMode = EPCGToolDrawTargetNormal::HitNormal;
+
+	/** Manual vector to offset drawn points from the clicked surface. */
+	UPROPERTY(EditAnywhere, Category = "Draw Target", meta = (UIMin = 0, UIMax = 100, EditCondition = "NormalMode == EPCGToolDrawTargetNormal::Explicit", EditConditionHides))
+	FVector Normal;
+
+	// @todo_pcg: Target object (such as a mesh)
+
+	// @todo_pcg: Target plane
+};
diff --git a/Engine/Plugins/PCG/Source/PCGEditor/Public/EditorMode/Tools/PCGInteractiveToolSettings.h b/Engine/Plugins/PCG/Source/PCGEditor/Public/EditorMode/Tools/PCGInteractiveToolSettings.h
new file mode 100644
--- /dev/null
+++ b/Engine/Plugins/PCG/Source/PCGEditor/Public/EditorMode/Tools/PCGInteractiveToolSettings.h
@@ -0,0 +1,58 @@
+// Copyright Epic Games, Inc. All Rights Reserved.
+
+#pragma once
+
+#include "PCGGraph.h"
+#include "PCGInteractiveToolCommon.h"
+#include "Helpers/PCGEdModeActorHelpers.h"
+
+#include "InteractiveTool.h"
+#include "NativeGameplayTags.h"
+
+#include "PCGInteractiveToolSettings.generated.h"
+
+UCLASS(Abstract)
+class PCGEDITOR_API UPCGInteractiveToolSettings : public UInteractiveToolPropertySet
+{
+	GENERATED_BODY()
+
+public:
+	UPCGInteractiveToolSettings();
+
+	/** Tool tag for use with filtering compatible graphs. */
+	virtual FText GetToolTag() const PURE_VIRTUAL(UPCGInteractiveToolSettings::GetToolTag, return FText::GetEmpty(););
+	/** A soft object path to the graph asset to load by default for this tool. */
+	virtual FSoftObjectPath GetDefaultGraph() const PURE_VIRTUAL(UPCGInteractiveToolSettings::GetDefaultGraph, return FSoftObjectPath{};); // @todo_pcg: Eventually, this could be a user preference per tool
+
+protected:
+	/** Filter details view graph selection. */
+	UFUNCTION()
+	virtual bool GraphAssetFilter(const FAssetData& AssetData) const;
+	/** Try to load the default graph into TargetGraph. */
+	virtual void TryLoadDefaultGraph();
+
+public:
+	UPROPERTY(BlueprintReadOnly, EditAnywhere, Category = "PCG", meta = (NoResetToDefault, GetAssetFilter = "GraphAssetFilter"))
+	TObjectPtr<UPCGGraphInterface> TargetGraph = nullptr;
+
+	/** Automatically generate the PCG component while using the tool. */
+	UPROPERTY(EditAnywhere, Category = "PCG")
+	bool bAutoGenerate = true;
+
+	/** The tool can create a new actor or modify an existing one. */
+	UPROPERTY(EditAnywhere, Category = "PCG")
+	EPCGToolInteractionMode Interaction = EPCGToolInteractionMode::CreateNewActor;
+
+	// @todo_pcg: Convert to a target object, rather than actor, to ease transition into SceneGraph
+	/** The selected actor to modify with the tool. */
+	UPROPERTY(EditAnywhere, Category = "PCG", meta = (EditCondition = "Interaction == EPCGToolInteractionMode::AddToSelectedActor", EditConditionHides, TransientToolProperty))
+	TWeakObjectPtr<AActor> TargetActor = nullptr;
+
+	/** The name to be issued to a newly created actor when the tool is applied. */
+	UPROPERTY(EditAnywhere, Category = "PCG", meta = (EditCondition = "Interaction == EPCGToolInteractionMode::CreateNewActor", EditConditionHides, TransientToolProperty))
+	FName NewActorName;
+
+	/** The name to be issued to a newly created PCG component when the tool is applied. */
+	UPROPERTY(EditAnywhere, Category = "PCG", meta = (TransientToolProperty))
+	FName NewPCGComponentName;
+};
diff --git a/Engine/Plugins/PCG/Source/PCGEditor/Public/EditorMode/Tools/Paint/PCGEdModeBrushBase.h b/Engine/Plugins/PCG/Source/PCGEditor/Public/EditorMode/Tools/Paint/PCGEdModeBrushBase.h
new file mode 100644
--- /dev/null
+++ b/Engine/Plugins/PCG/Source/PCGEditor/Public/EditorMode/Tools/Paint/PCGEdModeBrushBase.h
@@ -0,0 +1,16 @@
+// Copyright Epic Games, Inc. All Rights Reserved.
+
+#pragma once
+
+#include "BaseTools/BaseBrushTool.h"
+
+#include "PCGEdModeBrushBase.generated.h"
+
+UCLASS()
+class PCGEDITOR_API UPCGEdModeBrushBase : public UBrushBaseProperties
+{
+	GENERATED_BODY()
+
+public:
+	UPCGEdModeBrushBase() = default;
+};
diff --git a/Engine/Plugins/PCG/Source/PCGEditor/Public/EditorMode/Tools/Paint/PCGPaintTool.h b/Engine/Plugins/PCG/Source/PCGEditor/Public/EditorMode/Tools/Paint/PCGPaintTool.h
new file mode 100644
--- /dev/null
+++ b/Engine/Plugins/PCG/Source/PCGEditor/Public/EditorMode/Tools/Paint/PCGPaintTool.h
@@ -0,0 +1,52 @@
+// Copyright Epic Games, Inc. All Rights Reserved.
+
+#pragma once
+
+#include "PCGEdModeBrushBase.h"
+#include "EditorMode/Tools/PCGInteractiveTool.h"
+#include "EditorMode/Tools/PCGInteractiveToolBuilder.h"
+#include "EditorMode/Tools/PCGInteractiveToolSettings.h"
+
+#include "PCGPaintTool.generated.h"
+
+UCLASS()
+class PCGEDITOR_API UPCGPaintSettings : public UPCGInteractiveToolSettings
+{
+	GENERATED_BODY()
+
+	UPCGPaintSettings(const FObjectInitializer& ObjectInitializer);
+
+	// ~Begin UPCGInteractiveToolSettings interface
+	virtual FText GetToolTag() const override;
+	virtual FSoftObjectPath GetDefaultGraph() const override;
+	// ~End UPCGInteractiveToolSettings interface
+
+	bool ApplyTool(UWorld* World, FText& OutWarning) const;
+
+public:
+	UPROPERTY(EditAnywhere, Category = "Brush", Instanced, meta = (ShowOnlyInnerProperties))
+	TObjectPtr<UPCGEdModeBrushBase> Brush = nullptr;
+};
+
+UCLASS(Transient)
+class UPCGPaintTool : public UPCGInteractiveTool
+{
+	GENERATED_BODY()
+
+public:
+	// ~Begin UPCGInteractiveToolSettings interface
+	virtual void Setup() override;
+	virtual TSubclassOf<UPCGInteractiveToolSettings> GetSettingsClass() const override;
+	virtual bool CanAccept() const override;
+	// ~End UPCGInteractiveToolSettings interface
+};
+
+UCLASS(Transient)
+class UPCGPaintToolBuilder : public UPCGInteractiveToolBuilder
+{
+	GENERATED_BODY()
+
+public:
+	virtual bool CanBuildTool(const FToolBuilderState& SceneState) const override;
+	virtual UInteractiveTool* BuildTool(const FToolBuilderState& SceneState) const override;
+};
diff --git a/Engine/Plugins/PCG/Source/PCGEditor/Public/EditorMode/Tools/Volume/PCGCreateVolumeTool.h b/Engine/Plugins/PCG/Source/PCGEditor/Public/EditorMode/Tools/Volume/PCGCreateVolumeTool.h
new file mode 100644
--- /dev/null
+++ b/Engine/Plugins/PCG/Source/PCGEditor/Public/EditorMode/Tools/Volume/PCGCreateVolumeTool.h
@@ -0,0 +1,81 @@
+// Copyright Epic Games, Inc. All Rights Reserved.
+
+#pragma once
+
+#include "EditorMode/Tools/PCGInteractiveTool.h"
+#include "EditorMode/Tools/PCGInteractiveToolBuilder.h"
+#include "EditorMode/Tools/PCGInteractiveToolSettings.h"
+#include "EditorMode/Tools/Helpers/PCGEdModeSceneQueryHelpers.h"
+#include "BaseBehaviors/BehaviorTargetInterfaces.h"
+
+#include "PCGCreateVolumeTool.generated.h"
+
+UCLASS()
+class PCGEDITOR_API UPCGCreateVolumeToolSettings : public UPCGInteractiveToolSettings
+{
+	GENERATED_BODY()
+
+public:
+	// ~Begin UPCGInteractiveToolSettings interface
+	virtual FText GetToolTag() const override;
+	virtual FSoftObjectPath GetDefaultGraph() const override;
+	// ~End UPCGInteractiveToolSettings interface
+
+	bool ApplyTool(UWorld* World, FText& OutWarning) const;
+
+	UPROPERTY(EditAnywhere, Category = "Volume")
+	FVector Dimensions = FVector::OneVector * 100;
+
+	UPROPERTY(EditAnywhere, Category = "Volume")
+	FVector Location = FVector::ZeroVector;
+
+    UPROPERTY(EditAnywhere, Category = "Volume")
+    FRotator Rotation = FRotator::ZeroRotator;
+
+    UPROPERTY(EditAnywhere, Category = "Volume")
+    FVector Scale = FVector::OneVector;
+
+	// @todo_pcg: Origin modes: Center, Corner
+
+	// @todo_pcg: Snap to grid
+};
+
+UCLASS(Transient, MinimalAPI)
+class UPCGCreateVolumeTool : public UPCGInteractiveTool, public IClickBehaviorTarget
+{
+	GENERATED_BODY()
+
+public:
+	// ~Begin UInteractiveTool interface
+	virtual bool HasCancel() const override { return true; }
+	virtual bool HasAccept() const override { return true; }
+
+	virtual void Setup() override;
+	virtual void Shutdown(EToolShutdownType ShutdownType) override;
+	virtual void OnTick(float DeltaTime) override;
+	virtual bool CanAccept() const override;
+	virtual void Render(IToolsContextRenderAPI* RenderAPI) override;
+	// ~End UInteractiveTool interface
+
+	// ~Begin UPCGInteractiveTool interface
+	virtual TSubclassOf<UPCGInteractiveToolSettings> GetSettingsClass() const override;
+	// ~End UPCGInteractiveTool interface
+
+	// ~Begin IClickBehaviorTarget interface
+	virtual FInputRayHit IsHitByClick(const FInputDeviceRay& ClickPosition) override;
+	virtual void OnClicked(const FInputDeviceRay& ClickPosition) override;
+	// ~End IClickBehaviorTarget interface
+
+private:
+	PCG::EditorMode::Scene::FViewRayParams CachedRayParams;
+};
+
+UCLASS(Transient, MinimalAPI)
+class UPCGCreateVolumeToolBuilder : public UPCGInteractiveToolBuilder
+{
+	GENERATED_BODY()
+
+public:
+	virtual bool CanBuildTool(const FToolBuilderState& SceneState) const override;
+	virtual UInteractiveTool* BuildTool(const FToolBuilderState& SceneState) const override;
+};
diff --git a/Engine/Plugins/PCG/Source/PCGEditor/Public/EditorMode/Widgets/SPCGEdModeAcceptToolOverlay.h b/Engine/Plugins/PCG/Source/PCGEditor/Public/EditorMode/Widgets/SPCGEdModeAcceptToolOverlay.h
new file mode 100644
--- /dev/null
+++ b/Engine/Plugins/PCG/Source/PCGEditor/Public/EditorMode/Widgets/SPCGEdModeAcceptToolOverlay.h
@@ -0,0 +1,41 @@
+// Copyright Epic Games, Inc. All Rights Reserved.
+
+#pragma once
+
+#include "Framework/SlateDelegates.h"
+
+#include "Widgets/SCompoundWidget.h"
+
+/** A simple Accept/Cancel widget for the PCG Editor Mode Toolkit. */
+class SPCGEdModeAcceptToolOverlay : public SCompoundWidget
+{
+	SLATE_BEGIN_ARGS(SPCGEdModeAcceptToolOverlay)
+		: _ActiveToolDisplayName(FText::GetEmpty())
+		, _ActiveToolIcon(nullptr)
+		, _IsAcceptButtonEnabled(false)
+		, _IsCancelButtonEnabled(false)
+		, _GetAcceptButtonVisibility(EVisibility::Collapsed)
+		, _GetCancelButtonVisibility(EVisibility::Collapsed) {}
+
+		SLATE_ATTRIBUTE(FText, ActiveToolDisplayName)
+		SLATE_ATTRIBUTE(const FSlateBrush*, ActiveToolIcon)
+
+		SLATE_ATTRIBUTE(FText, AcceptButtonLabel)
+		SLATE_ATTRIBUTE(FText, CancelButtonLabel)
+		SLATE_ATTRIBUTE(FText, AcceptButtonTooltip)
+		SLATE_ATTRIBUTE(FText, CancelButtonTooltip)
+		SLATE_ATTRIBUTE(bool, IsAcceptButtonEnabled)
+		SLATE_ATTRIBUTE(bool, IsCancelButtonEnabled)
+		SLATE_ATTRIBUTE(EVisibility, GetAcceptButtonVisibility)
+		SLATE_ATTRIBUTE(EVisibility, GetCancelButtonVisibility)
+
+		SLATE_EVENT(FOnClicked, OnAcceptButtonClicked)
+		SLATE_EVENT(FOnClicked, OnCancelButtonClicked)
+
+	SLATE_END_ARGS()
+
+	void Construct(const FArguments& InArgs);
+
+private:
+	FArguments Args;
+};
diff --git a/Engine/Source/Runtime/InteractiveToolsFramework/Public/InteractiveTool.h b/Engine/Source/Runtime/InteractiveToolsFramework/Public/InteractiveTool.h
--- a/Engine/Source/Runtime/InteractiveToolsFramework/Public/InteractiveTool.h
+++ b/Engine/Source/Runtime/InteractiveToolsFramework/Public/InteractiveTool.h
@@ -37,6 +37,7 @@ class UInputBehavior;
 class UInteractiveTool;
 class UInteractiveToolManager;
 struct FPropertyChangedEvent;
+struct FSlateBrush;
 
 /** Passed to UInteractiveTool::Shutdown to indicate how Tool should shut itself down*/
 UENUM(BlueprintType)
@@ -57,6 +58,10 @@ struct FInteractiveToolInfo
 {
 	/** Name of Tool. May be FText::Empty(), but will default to Tool->GetClass()->GetDisplayNameText() in InteractiveTool constructor */
 	FText ToolDisplayName = FText::GetEmpty();
+	/** Display message that can be used as a tool tip or to provide information about the tool. */
+	FText ToolDisplayMessage = FText::GetEmpty();
+	/** The icon to use for the tool ui command. */
+	const FSlateBrush* ToolIcon = nullptr;
 };
 
 
@@ -623,6 +628,22 @@ public:
 		DefaultToolInfo.ToolDisplayName = NewName;
 	}
 
+	/**
+	 * Set the Tool's Display Message
+	 */
+	virtual void SetToolDisplayMessage(const FText& NewDisplayMessage)
+	{
+		DefaultToolInfo.ToolDisplayMessage = NewDisplayMessage;
+	}
+
+	/**
+	 * Set the Tool's Icon
+	 */
+	virtual void SetToolIcon(const FSlateBrush* NewIcon)
+	{
+		DefaultToolInfo.ToolIcon = NewIcon;
+	}
+
 private:
 	/**
 	 * ToolInfo for this Tool
