[PCG] Fix missing undef LOCTEXT_NAMESPACE and SA Warning from PVS-Studio for Draw Spline Tool unreachable code

#rnx

diff --git a/Engine/Plugins/PCG/Source/PCGEditor/Private/EditorMode/Tools/Helpers/PCGEdModeActorHelpers.cpp b/Engine/Plugins/PCG/Source/PCGEditor/Private/EditorMode/Tools/Helpers/PCGEdModeActorHelpers.cpp
--- a/Engine/Plugins/PCG/Source/PCGEditor/Private/EditorMode/Tools/Helpers/PCGEdModeActorHelpers.cpp
+++ b/Engine/Plugins/PCG/Source/PCGEditor/Private/EditorMode/Tools/Helpers/PCGEdModeActorHelpers.cpp
@@ -77,3 +77,5 @@ namespace PCG::EditorMode
 		}
 	} // namespace PCG::EditorMode::PCGComponent
 } // namespace PCG::EditorMode
+
+#undef LOCTEXT_NAMESPACE
diff --git a/Engine/Plugins/PCG/Source/PCGEditor/Private/EditorMode/Tools/Line/PCGDrawSplineTool.cpp b/Engine/Plugins/PCG/Source/PCGEditor/Private/EditorMode/Tools/Line/PCGDrawSplineTool.cpp
--- a/Engine/Plugins/PCG/Source/PCGEditor/Private/EditorMode/Tools/Line/PCGDrawSplineTool.cpp
+++ b/Engine/Plugins/PCG/Source/PCGEditor/Private/EditorMode/Tools/Line/PCGDrawSplineTool.cpp
@@ -415,7 +415,6 @@ FVector UPCGDrawSplineTool::GetNewPointNormal(const FVector& HitNormal) const
 	{
 		case EPCGToolDrawTargetNormal::AlignToPrevious:
 		{
-			check(NumSplinePoints > 0);
 			if (NumSplinePoints == 1)
 			{
 				/*
@@ -430,9 +429,11 @@ FVector UPCGDrawSplineTool::GetNewPointNormal(const FVector& HitNormal) const
 			{
 				return WorkingSpline->GetUpVectorAtSplinePoint(NumSplinePoints - 1, ESplineCoordinateSpace::World);
 			}
-
-			checkNoEntry();
-			return HitNormal;
+			else // NumSplinePoints is negative, which should never happen.
+			{
+				ensure(false);
+				return HitNormal;
+			}
 		}
 		case EPCGToolDrawTargetNormal::WorldUp:
 			return FVector::UnitZ();
