[PCG Editor Mode] Fixed a nullcheck on Editor Servers from PCGPaintTool

#rnx

diff --git a/Engine/Plugins/PCG/Source/PCGEditor/Private/EditorMode/Tools/Paint/PCGPaintTool.cpp b/Engine/Plugins/PCG/Source/PCGEditor/Private/EditorMode/Tools/Paint/PCGPaintTool.cpp
--- a/Engine/Plugins/PCG/Source/PCGEditor/Private/EditorMode/Tools/Paint/PCGPaintTool.cpp
+++ b/Engine/Plugins/PCG/Source/PCGEditor/Private/EditorMode/Tools/Paint/PCGPaintTool.cpp
@@ -15,13 +15,13 @@ namespace PCGPaintTool::Constants
 }
 
 UPCGPaintSettings::UPCGPaintSettings(const FObjectInitializer& ObjectInitializer)
-	: Brush(CastChecked<UPCGEdModeBrushBase>(
+	: Brush(Cast<UPCGEdModeBrushBase>(
 		ObjectInitializer.CreateEditorOnlyDefaultSubobject(
 			this,
 			GET_MEMBER_NAME_CHECKED(UPCGPaintSettings, Brush),
 			UPCGEdModeBrushBase::StaticClass(),
-			/*bIsTransient=*/false)))
-{}
+			/*bIsTransient=*/true)))
+{} // @todo_pcg: The brush will inevitably be serialized somewhere as the user modifies it.
 
 FText UPCGPaintSettings::GetToolTag() const
 {
