[PCG Editor Mode] Clean up multiple issues for SA, Camera plane normal and Spline Draw bug

#rnx

diff --git a/Engine/Plugins/PCG/Source/PCGEditor/Private/EditorMode/PCGEdModeToolkit.cpp b/Engine/Plugins/PCG/Source/PCGEditor/Private/EditorMode/PCGEdModeToolkit.cpp
--- a/Engine/Plugins/PCG/Source/PCGEditor/Private/EditorMode/PCGEdModeToolkit.cpp
+++ b/Engine/Plugins/PCG/Source/PCGEditor/Private/EditorMode/PCGEdModeToolkit.cpp
@@ -109,7 +109,7 @@ void FPCGEditorModeToolkit::OnToolEnded(UInteractiveToolManager* Manager, UInter
 	ClearToolNotification();
 	ClearToolWarning();
 
-	if (UPCGInteractiveTool* CurrentTool = CastChecked<UPCGInteractiveTool>(GetScriptableEditorMode()->GetToolManager(EToolsContextScope::EdMode)->GetActiveTool(EToolSide::Left), ECastCheckedType::NullAllowed))
+	if (UPCGInteractiveTool* CurrentTool = CastChecked<UPCGInteractiveTool>(Tool, ECastCheckedType::NullAllowed))
 	{
 		UnregisterToolCallbacks(CurrentTool);
 	}
@@ -216,7 +216,7 @@ FReply FPCGEditorModeToolkit::OnAcceptButtonClicked() const
 FReply FPCGEditorModeToolkit::OnCancelButtonClicked() const
 {
 	GetScriptableEditorMode()->GetInteractiveToolsContext()->EndTool(EToolShutdownType::Cancel);
-	return FReply::Unhandled();
+	return FReply::Handled();
 }
 
 void FPCGEditorModeToolkit::RegisterToolkitCallbacks()
diff --git a/Engine/Plugins/PCG/Source/PCGEditor/Private/EditorMode/Tools/Helpers/PCGEdModeSceneQueryHelpers.cpp b/Engine/Plugins/PCG/Source/PCGEditor/Private/EditorMode/Tools/Helpers/PCGEdModeSceneQueryHelpers.cpp
--- a/Engine/Plugins/PCG/Source/PCGEditor/Private/EditorMode/Tools/Helpers/PCGEdModeSceneQueryHelpers.cpp
+++ b/Engine/Plugins/PCG/Source/PCGEditor/Private/EditorMode/Tools/Helpers/PCGEdModeSceneQueryHelpers.cpp
@@ -106,7 +106,7 @@ namespace PCG::EditorMode::Scene
 	TOptional<FViewHitResult> ViewportRayAgainstCameraPlane(const FInputDeviceRay& Ray, const FViewRayParams& Params, const FViewCameraState& CameraState)
 	{
 		const FVector PlaneOrigin = CameraState.Position + CameraState.Forward() * Params.Distance;
-		const FVector PlaneNormal = CameraState.bIsOrthographic ? -Ray.WorldRay.Direction : FVector::UnitZ();
+		const FVector PlaneNormal = -CameraState.Forward();
 		return ViewportRayAgainstPlane(Ray, Params, FPlane(PlaneOrigin, PlaneNormal));
 	}
 } // namespace PCG::EditorMode::Scene
diff --git a/Engine/Plugins/PCG/Source/PCGEditor/Private/EditorMode/Tools/Helpers/PCGEdModeSceneVisualizeHelpers.cpp b/Engine/Plugins/PCG/Source/PCGEditor/Private/EditorMode/Tools/Helpers/PCGEdModeSceneVisualizeHelpers.cpp
--- a/Engine/Plugins/PCG/Source/PCGEditor/Private/EditorMode/Tools/Helpers/PCGEdModeSceneVisualizeHelpers.cpp
+++ b/Engine/Plugins/PCG/Source/PCGEditor/Private/EditorMode/Tools/Helpers/PCGEdModeSceneVisualizeHelpers.cpp
@@ -10,6 +10,7 @@
 namespace PCG::EditorMode::Scene::Visualize
 {
 	// Taken in large directly from SplineComponentVisualizer::DrawVisualization in SplineComponentVisualizer.h
+	// Note: this currently overdraws the new spline visualizer, but can be used for additional things like scale, normals, etc.
 	void Spline(const USplineComponent& Spline, IToolsContextRenderAPI& RenderAPI, int32 LatestIndex, bool bShouldVisualizeScale, bool bShouldVisualizeTangents)
 	{
 		FPrimitiveDrawInterface* PDI = RenderAPI.GetPrimitiveDrawInterface();
@@ -101,6 +102,11 @@ namespace PCG::EditorMode::Scene::Visualize
 					const float StepRatio = static_cast<float>(StepIndex) / static_cast<float>(NumSteps);
 					const float Key = PreviousIndex + StepRatio;
 
+					// Find the step position
+					const FVector StepPosition = Spline.GetLocationAtSplinePoint(Key, ESplineCoordinateSpace::World);
+					const FVector StepRightVector = Spline.GetRightVectorAtSplinePoint(Key, ESplineCoordinateSpace::World);
+					const FVector StepScale = Spline.GetScaleAtSplinePoint(Key);
+
 					FColor StepColor;
 					// Grade the lerp to full over 80%, then clamp.
 					constexpr float PartialGradientProportion = 0.8f;
@@ -108,23 +114,23 @@ namespace PCG::EditorMode::Scene::Visualize
 					StepColor = FMath::Lerp(BeginColor.ReinterpretAsLinear(), EndColor.ReinterpretAsLinear(), LerpRatio).ToFColor(false);
 
 					// Draw the segment.
-					PDI->DrawLine(PreviousPosition, KeyPosition, StepColor, SDPG_Foreground, SegmentLineThickness);
+					PDI->DrawLine(PreviousPosition, StepPosition, StepColor, SDPG_Foreground, SegmentLineThickness);
 
 					if (bShouldVisualizeScale)
 					{
-						PDI->DrawLine(PreviousPosition - PreviousRightVector * PreviousScale.Y, KeyPosition - KeyRightVector * KeyScale.Y, LineColor, SDPG_Foreground);
-						PDI->DrawLine(PreviousPosition + PreviousRightVector * PreviousScale.Y, KeyPosition + KeyRightVector * KeyScale.Y, LineColor, SDPG_Foreground);
+						PDI->DrawLine(PreviousPosition - PreviousRightVector * PreviousScale.Y, StepPosition - StepRightVector * StepScale.Y, LineColor, SDPG_Foreground);
+						PDI->DrawLine(PreviousPosition + PreviousRightVector * PreviousScale.Y, StepPosition + StepRightVector * StepScale.Y, LineColor, SDPG_Foreground);
 					}
 
 					// Visualize Up Vector
 					const float HalfVisualizationWidthScale = VisualizationWidthScale * 0.5f;
-					const FVector NewUpVector = Spline.GetUpVectorAtSplineInputKey(Key, ESplineCoordinateSpace::World);
-					PDI->DrawLine(KeyPosition, KeyPosition + NewUpVector * HalfVisualizationWidthScale, LineColor, SDPG_Foreground);
-					PDI->DrawLine(KeyPosition, KeyPosition + KeyRightVector * HalfVisualizationWidthScale, LineColor, SDPG_Foreground);
+					const FVector StepUpVector = Spline.GetUpVectorAtSplineInputKey(Key, ESplineCoordinateSpace::World);
+					PDI->DrawLine(StepPosition, StepPosition + StepUpVector * HalfVisualizationWidthScale, LineColor, SDPG_Foreground);
+					PDI->DrawLine(StepPosition, StepPosition + StepRightVector * HalfVisualizationWidthScale, LineColor, SDPG_Foreground);
 
-					PreviousKeyPosition = KeyPosition;
-					PreviousKeyRightVector = KeyRightVector;
-					PreviousKeyScale = KeyScale;
+					PreviousKeyPosition = StepPosition;
+					PreviousKeyRightVector = StepRightVector;
+					PreviousKeyScale = StepUpVector;
 				}
 			}
 		}
diff --git a/Engine/Plugins/PCG/Source/PCGEditor/Private/EditorMode/Tools/Line/PCGDrawSplineTool.cpp b/Engine/Plugins/PCG/Source/PCGEditor/Private/EditorMode/Tools/Line/PCGDrawSplineTool.cpp
--- a/Engine/Plugins/PCG/Source/PCGEditor/Private/EditorMode/Tools/Line/PCGDrawSplineTool.cpp
+++ b/Engine/Plugins/PCG/Source/PCGEditor/Private/EditorMode/Tools/Line/PCGDrawSplineTool.cpp
@@ -429,7 +429,7 @@ FVector UPCGDrawSplineTool::GetNewPointNormal(const FVector& HitNormal) const
 			{
 				return WorkingSpline->GetUpVectorAtSplinePoint(NumSplinePoints - 1, ESplineCoordinateSpace::World);
 			}
-			else // NumSplinePoints is negative, which should never happen.
+			else // NumSplinePoints is zero or negative, which should never happen.
 			{
 				ensure(false);
 				return HitNormal;
@@ -640,6 +640,8 @@ void UPCGDrawSplineTool::OnTerminateDragSequence()
 	const int32 NumSplinePoints = WorkingSpline->GetNumberOfSplinePoints();
 	const int32 PreviousSplinePoint = NumSplinePoints - 1;
 
+	check(PreviousSplinePoint >= 0);
+
 	const UPCGDrawSplineToolSettings* Settings = CastChecked<UPCGDrawSplineToolSettings>(ToolSettings);
 	// Emit the appropriate undo transaction
 	switch (Settings->DrawMode)
diff --git a/Engine/Plugins/PCG/Source/PCGEditor/Public/EditorMode/PCGEdModeToolkit.h b/Engine/Plugins/PCG/Source/PCGEditor/Public/EditorMode/PCGEdModeToolkit.h
--- a/Engine/Plugins/PCG/Source/PCGEditor/Public/EditorMode/PCGEdModeToolkit.h
+++ b/Engine/Plugins/PCG/Source/PCGEditor/Public/EditorMode/PCGEdModeToolkit.h
@@ -82,5 +82,3 @@ private:
 	FText ActiveWarning{};
 	FInteractiveToolInfo CachedActiveToolInfo;
 };
-
-#undef LOCTEXT_NAMESPACE
