[PCG] Add support for property remap when converting back and forth between Attribute Set and Points

#rnx

diff --git a/Engine/Plugins/PCG/Source/PCG/Private/Elements/PCGCollapseElement.cpp b/Engine/Plugins/PCG/Source/PCG/Private/Elements/PCGCollapseElement.cpp
--- a/Engine/Plugins/PCG/Source/PCG/Private/Elements/PCGCollapseElement.cpp
+++ b/Engine/Plugins/PCG/Source/PCG/Private/Elements/PCGCollapseElement.cpp
@@ -6,6 +6,7 @@
 #include "PCGParamData.h"
 #include "Data/PCGBasePointData.h"
 #include "Data/PCGSpatialData.h"
+#include "Metadata/Accessors/PCGAttributeAccessorHelpers.h"
 
 #include "UObject/FortniteMainBranchObjectVersion.h"
 
@@ -111,6 +112,59 @@ bool FPCGCollapseElement::ExecuteInternal(FPCGContext* InContext) const
 				MetadataEntryRange[PointIndex] = PointIndex;
 			}
 
+			const UPCGConvertToPointDataSettings* ConvertToPointDataSettings = InContext->GetInputSettings<UPCGConvertToPointDataSettings>();
+			if (ConvertToPointDataSettings && ConvertToPointDataSettings->bMatchAttributeNamesWithPropertyNames)
+			{
+				TArray<FName> AttributeNames;
+				TArray<EPCGMetadataTypes> AttributeTypes;
+
+				PointData->ConstMetadata()->GetAttributes(AttributeNames, AttributeTypes);
+				for (const FName AttributeName : AttributeNames)
+				{
+					FPCGAttributePropertyOutputSelector OutputSelector = FPCGAttributePropertySelector::CreatePropertySelector<FPCGAttributePropertyOutputSelector>(AttributeName);
+
+					// Need to create the accessor to know if it exists.
+					if (!PCGAttributeAccessorHelpers::CreateAccessor(PointData, OutputSelector, /*bQuiet=*/true))
+					{
+						continue;
+					}
+
+					PCGMetadataHelpers::FPCGCopyAttributeParams CopyParams
+					{
+						.SourceData = ParamData,
+						.TargetData = PointData,
+						.InputSource = FPCGAttributePropertySelector::CreateAttributeSelector<FPCGAttributePropertyInputSelector>(AttributeName),
+						.OutputTarget = MoveTemp(OutputSelector),
+						.bSameOrigin = true
+					};
+
+					if (PCGMetadataHelpers::CopyAttribute(CopyParams) && ConvertToPointDataSettings->bDeleteOriginalRemappedAttribute)
+					{
+						PointData->MutableMetadata()->DeleteAttribute(AttributeName);
+					}
+				}
+			}
+
+			if (ConvertToPointDataSettings && !ConvertToPointDataSettings->AttributeToConvert.IsEmpty())
+			{
+				for (const auto& [InputSource, OutputTarget] : ConvertToPointDataSettings->AttributeToConvert)
+				{
+					PCGMetadataHelpers::FPCGCopyAttributeParams CopyParams
+					{
+						.SourceData = ParamData,
+						.TargetData = PointData,
+						.InputSource = InputSource,
+						.OutputTarget = OutputTarget,
+						.bSameOrigin = true
+					};
+
+					if (PCGMetadataHelpers::CopyAttribute(CopyParams) && ConvertToPointDataSettings->bDeleteOriginalRemappedAttribute)
+					{
+						PointData->MutableMetadata()->DeleteAttribute(InputSource.CopyAndFixLast(ParamData).GetName());
+					}
+				}
+			}
+
 			Output.Data = PointData;
 			Context->bShouldComputeFullOutputDataCrc = true;
 		}
diff --git a/Engine/Plugins/PCG/Source/PCG/Private/Elements/PCGConvertToAttributeSet.cpp b/Engine/Plugins/PCG/Source/PCG/Private/Elements/PCGConvertToAttributeSet.cpp
--- a/Engine/Plugins/PCG/Source/PCG/Private/Elements/PCGConvertToAttributeSet.cpp
+++ b/Engine/Plugins/PCG/Source/PCG/Private/Elements/PCGConvertToAttributeSet.cpp
@@ -42,6 +42,9 @@ bool FPCGConvertToAttributeSetElement::ExecuteInternal(FPCGContext* Context) con
 	TRACE_CPUPROFILER_EVENT_SCOPE(FPCGConvertToAttributeSetElement::Execute);
 	check(Context);
 
+	const UPCGConvertToAttributeSetSettings* Settings = Context->GetInputSettings<UPCGConvertToAttributeSetSettings>();
+	check(Settings);
+
 	TArray<FPCGTaggedData> Inputs = Context->InputData.GetInputsByPin(PCGPinConstants::DefaultInputLabel);
 	TArray<FPCGTaggedData>& Outputs = Context->OutputData.TaggedData;
 
@@ -57,7 +60,7 @@ bool FPCGConvertToAttributeSetElement::ExecuteInternal(FPCGContext* Context) con
 		check(SourceMetadata);
 
 		// Note: this is very similar in idea to UPCGPointData::Flatten
-		if (SourceMetadata->GetAttributeCount() == 0 || PointData->GetNumPoints() == 0)
+		if ((Settings->PropertiesToConvert.IsEmpty() && SourceMetadata->GetAttributeCount() == 0) || PointData->GetNumPoints() == 0)
 		{
 			continue;
 		}
@@ -76,6 +79,40 @@ bool FPCGConvertToAttributeSetElement::ExecuteInternal(FPCGContext* Context) con
 		Params.DomainInitializeParams.Emplace(PCGMetadataDomainID::Elements, std::move(ElementsDomainParams));
 		Metadata->InitializeAsCopy(Params);
 
+		if (!Settings->PropertiesToConvert.IsEmpty())
+		{
+			for (const auto& [Source, Target] : Settings->PropertiesToConvert)
+			{
+				// @Source is treated as @SourceName
+				FPCGAttributePropertyOutputSelector OutputSelector =
+					Target.GetAttributeName() == PCGMetadataAttributeConstants::SourceAttributeName
+					? FPCGAttributePropertySelector::CreateAttributeSelector<FPCGAttributePropertyOutputSelector>(PCGMetadataAttributeConstants::SourceNameAttributeName)
+					: Target;
+				
+				// Verify if the target attribute exists
+				if (!Settings->bOverrideExistingAttributes)
+				{
+					const FPCGAttributePropertyInputSelector InputSelector = Source.CopyAndFixLast(Input.Data);
+					OutputSelector = OutputSelector.CopyAndFixSource(&InputSelector, Input.Data);
+
+					if (ParamData->ConstMetadata()->HasAttribute(OutputSelector.GetName()))
+					{
+						continue;
+					}
+				}
+				
+				PCGMetadataHelpers::FPCGCopyAttributeParams CopyParams
+				{
+					.SourceData = Input.Data,
+					.TargetData = ParamData,
+					.InputSource = Source,
+					.OutputTarget = MoveTemp(OutputSelector)
+				};
+
+				PCGMetadataHelpers::CopyAttribute(CopyParams);
+			}
+		}
+		
 		FPCGTaggedData& Output = Outputs.Emplace_GetRef(Input);
 		Output.Data = ParamData;
 	}
diff --git a/Engine/Plugins/PCG/Source/PCG/Public/Elements/PCGCollapseElement.h b/Engine/Plugins/PCG/Source/PCG/Public/Elements/PCGCollapseElement.h
--- a/Engine/Plugins/PCG/Source/PCG/Public/Elements/PCGCollapseElement.h
+++ b/Engine/Plugins/PCG/Source/PCG/Public/Elements/PCGCollapseElement.h
@@ -58,6 +58,19 @@ public:
 protected:
 	virtual TArray<FPCGPinProperties> InputPinProperties() const override;
 	//~End UPCGSettings interface
+
+public:
+	/** Try to match attribute names that matches point property names to automatically set them. */
+	UPROPERTY(BlueprintReadWrite, EditAnywhere, Category="Settings", meta = (PCG_Overridable))
+	bool bMatchAttributeNamesWithPropertyNames = false;
+	
+	/** Can remap attribute to point properties at the same time. */
+	UPROPERTY(BlueprintReadWrite, EditAnywhere, Category="Settings", meta = (PCG_Overridable))
+	TMap<FPCGAttributePropertyInputSelector, FPCGAttributePropertyOutputSelector> AttributeToConvert;
+
+	/** When an attribute is remapped, delete the original. */
+	UPROPERTY(BlueprintReadWrite, EditAnywhere, Category="Settings", meta = (PCG_Overridable))
+	bool bDeleteOriginalRemappedAttribute = true;
 };
 
 struct FPCGCollapseContext : public FPCGContext
diff --git a/Engine/Plugins/PCG/Source/PCG/Public/Elements/PCGConvertToAttributeSet.h b/Engine/Plugins/PCG/Source/PCG/Public/Elements/PCGConvertToAttributeSet.h
--- a/Engine/Plugins/PCG/Source/PCG/Public/Elements/PCGConvertToAttributeSet.h
+++ b/Engine/Plugins/PCG/Source/PCG/Public/Elements/PCGConvertToAttributeSet.h
@@ -29,6 +29,15 @@ protected:
 	virtual TArray<FPCGPinProperties> OutputPinProperties() const override;
 	virtual FPCGElementPtr CreateElement() const override;
 	//~End UPCGSettings interface
+
+public:
+	/** Can add properties to convert at the same time, with a potential remap. */
+	UPROPERTY(BlueprintReadWrite, EditAnywhere, Category="Settings", meta = (PCG_Overridable))
+	TMap<FPCGAttributePropertyInputSelector, FPCGAttributePropertyOutputSelector> PropertiesToConvert;
+
+	/** If properties are converted and there is already an attribute with the same name, it can be overridden by the property. */
+	UPROPERTY(BlueprintReadWrite, EditAnywhere, Category="Settings", meta = (PCG_Overridable))
+	bool bOverrideExistingAttributes = false;
 };
 
 class FPCGConvertToAttributeSetElement : public IPCGElement
