[PCG] Adding support for metadata in polygon2d operations.

diff --git a/Engine/Plugins/PCG/Source/PCG/Private/Data/PCGPolygon2DData.cpp b/Engine/Plugins/PCG/Source/PCG/Private/Data/PCGPolygon2DData.cpp
--- a/Engine/Plugins/PCG/Source/PCG/Private/Data/PCGPolygon2DData.cpp
+++ b/Engine/Plugins/PCG/Source/PCG/Private/Data/PCGPolygon2DData.cpp
@@ -61,6 +61,68 @@ namespace PCGPolygon2DData
 		}
 	}
 
+	void ReversePolygon(UE::Geometry::FGeneralPolygon2d& Polygon, TArray<int64>& PolygonEntryKeys)
+	{
+		Polygon.Reverse();
+		if (!PolygonEntryKeys.IsEmpty())
+		{
+			int VertexOffset = 0;
+			int VertexCount = Polygon.GetOuter().VertexCount();
+			Algo::Reverse(&PolygonEntryKeys[VertexOffset], VertexCount);
+			VertexOffset += VertexCount;
+
+			for (const UE::Geometry::FPolygon2d& Hole : Polygon.GetHoles())
+			{
+				VertexCount = Hole.VertexCount();
+				Algo::Reverse(&PolygonEntryKeys[VertexOffset], VertexCount);
+				VertexOffset += VertexCount;
+			}
+		}
+	}
+
+	int PolygonVertexCount(const UE::Geometry::FGeneralPolygon2d& InPoly)
+	{
+		int Count = InPoly.GetOuter().VertexCount();
+		for (const UE::Geometry::FPolygon2d& Hole : InPoly.GetHoles())
+		{
+			Count += Hole.VertexCount();
+		}
+
+		return Count;
+	}
+
+	int GetVertexIndex(const UE::Geometry::FGeneralPolygon2d& InPoly, int InSegmentIndex, int InHoleIndex)
+	{
+		int StartVertexIndex = -1;
+		int EndVertexIndex = -1;
+		GetStartEndVertexIndices(InPoly, InSegmentIndex, InHoleIndex, StartVertexIndex, EndVertexIndex);
+
+		return StartVertexIndex;
+	}
+
+	void GetStartEndVertexIndices(const UE::Geometry::FGeneralPolygon2d& InPoly, int InSegmentIndex, int InHoleIndex, int& OutStartVertexIndex, int& OutEndVertexIndex)
+	{
+		int Offset = 0;
+		int CurrentPolyVertexCount = InPoly.GetOuter().VertexCount();
+
+		if (InHoleIndex >= 0)
+		{
+			Offset = CurrentPolyVertexCount;
+		}
+
+		for (int HoleIndex = 0; HoleIndex <= InHoleIndex; ++HoleIndex)
+		{
+			CurrentPolyVertexCount = InPoly.GetHoles()[HoleIndex].VertexCount();
+			if (HoleIndex < InHoleIndex)
+			{
+				Offset += CurrentPolyVertexCount;
+			}
+		}
+
+		OutStartVertexIndex = Offset + InSegmentIndex;
+		OutEndVertexIndex = Offset + ((InSegmentIndex + 1) % CurrentPolyVertexCount);
+	}
+
 	TUniquePtr<IPCGAttributeAccessor> CreatePropertyAccessor(EPCGPolygon2DProperties Property, bool bQuiet, bool bConst)
 	{
 		if (bConst)
@@ -644,8 +706,7 @@ void UPCGPolygon2DData::SetTransform(const FTransform& InTransform, bool bCheckW
 
 	if (bNeedsToReversePolygon)
 	{
-		// @todo_pcg - check how that affects the polygon operations
-		Polygon.Reverse();
+		PCGPolygon2DData::ReversePolygon(Polygon, VertexEntryKeys);
 	}
 }
 
diff --git a/Engine/Plugins/PCG/Source/PCG/Private/Elements/Polygon/PCGCreatePolygon2D.cpp b/Engine/Plugins/PCG/Source/PCG/Private/Elements/Polygon/PCGCreatePolygon2D.cpp
--- a/Engine/Plugins/PCG/Source/PCG/Private/Elements/Polygon/PCGCreatePolygon2D.cpp
+++ b/Engine/Plugins/PCG/Source/PCG/Private/Elements/Polygon/PCGCreatePolygon2D.cpp
@@ -7,6 +7,7 @@
 #include "Data/PCGSplineData.h"
 #include "Data/PCGPolygon2DData.h"
 #include "Utils/PCGLogErrors.h"
+#include "Helpers/PCGPolygon2DProcessingHelpers.h"
 
 #include "CompGeom/PolygonTriangulation.h"
 #include "Curve/PolygonOffsetUtils.h"
@@ -30,6 +31,7 @@ bool FPCGCreatePolygonElement::ExecuteInternal(FPCGContext* Context) const
 	TArray<FPCGTaggedData> Inputs = Context->InputData.GetInputsByPin(PCGPinConstants::DefaultInputLabel);
 
 	TArray<FVector> Positions;
+	TArray<PCGMetadataEntryKey> EntryKeys;
 
 	for (const FPCGTaggedData& Input : Inputs)
 	{
@@ -37,6 +39,11 @@ bool FPCGCreatePolygonElement::ExecuteInternal(FPCGContext* Context) const
 
 		bool bInputIsOpen = false;
 		bool bInputFromClosedSpline = false;
+		bool bBuildEntryKeysBySampling = false;
+		bool bInputHasMetadata = false;
+
+		Positions.Reset();
+		EntryKeys.Reset();
 
 		if (const UPCGSplineData* SplineData = Cast<UPCGSplineData>(Input.Data))
 		{
@@ -52,6 +59,9 @@ bool FPCGCreatePolygonElement::ExecuteInternal(FPCGContext* Context) const
 			{
 				Positions.Pop();
 			}
+
+			bInputHasMetadata = !SplineData->GetConstVerticesEntryKeys().IsEmpty();
+			bBuildEntryKeysBySampling = bInputHasMetadata;
 		}
 		else if (const UPCGBasePointData* InputPointData = Cast<UPCGBasePointData>(Input.Data))
 		{
@@ -61,13 +71,29 @@ bool FPCGCreatePolygonElement::ExecuteInternal(FPCGContext* Context) const
 
 			if (TransformRange.Num() > 0)
 			{
-				Positions.Reset();
 				Positions.Reserve(TransformRange.Num());
 				for (const FTransform& Transform : TransformRange)
 				{
 					Positions.Emplace(Transform.GetLocation());
 				}
 			}
+
+			const TConstPCGValueRange<PCGMetadataEntryKey> EntryKeysRange = InputPointData->GetConstMetadataEntryValueRange();
+
+			if (EntryKeysRange.Num() > 0)
+			{
+				bInputHasMetadata = true;
+
+				// There's no point in getting the entry keys if we're extruding the input path.
+				if (Settings->InputType != EPCGCreatePolygonInputType::ForceOpen)
+				{
+					EntryKeys.Reserve(EntryKeysRange.Num());
+					for (const PCGMetadataEntryKey& EK : EntryKeysRange)
+					{
+						EntryKeys.Add(EK);
+					}
+				}
+			}
 		}
 		else
 		{
@@ -134,14 +160,23 @@ bool FPCGCreatePolygonElement::ExecuteInternal(FPCGContext* Context) const
 			{
 				for (int PolyIndex = 0; PolyIndex < PolygonOffset.Result.Num(); ++PolyIndex)
 				{
+					UE::Geometry::FGeneralPolygon2d& Poly = PolygonOffset.Result[PolyIndex];
+
+					if (Poly.OuterIsClockwise())
+					{
+						Poly.Reverse();
+					}
+
 					UPCGPolygon2DData* Polygon = CreatePolygonData();
 
-					if (PolygonOffset.Result[PolyIndex].OuterIsClockwise())
+					if (bInputHasMetadata)
 					{
-						PolygonOffset.Result[PolyIndex].Reverse();
+						PCGPolygon2DProcessingHelpers::ComputeEntryKeysByNearestSegmentSampling(SpatialData, Settings->OpenPolygonWidth, Poly, PlaneTransform, Polygon->MutableMetadata());
 					}
 
-					Polygon->SetPolygon(MoveTemp(PolygonOffset.Result[PolyIndex]));
+					TConstArrayView<PCGMetadataEntryKey> EntryKeysView(EntryKeys);
+
+					Polygon->SetPolygon(MoveTemp(Poly), EntryKeysView.IsEmpty() ? nullptr : &EntryKeysView);
 					Polygon->SetTransform(PlaneTransform);
 				}
 			}
@@ -164,7 +199,15 @@ bool FPCGCreatePolygonElement::ExecuteInternal(FPCGContext* Context) const
 			}
 
 			UPCGPolygon2DData* Polygon = CreatePolygonData();
-			Polygon->SetPolygon(MoveTemp(NewPolygon));
+
+			// Once we've created the polygon object we can assign new metadata entries as needed.
+			if (bBuildEntryKeysBySampling)
+			{
+				EntryKeys = PCGPolygon2DProcessingHelpers::ComputeEntryKeysBySampling(SpatialData, Positions, Polygon->MutableMetadata());
+			}
+
+			TConstArrayView<PCGMetadataEntryKey> EntryKeysView(EntryKeys);
+			Polygon->SetPolygon(MoveTemp(NewPolygon), EntryKeysView.IsEmpty() ? nullptr : &EntryKeysView);
 			Polygon->SetTransform(PlaneTransform);
 
 			// @todo_pcg: Create hole polygons / multi polygons...
diff --git a/Engine/Plugins/PCG/Source/PCG/Private/Elements/Polygon/PCGOffsetPolygon2D.cpp b/Engine/Plugins/PCG/Source/PCG/Private/Elements/Polygon/PCGOffsetPolygon2D.cpp
--- a/Engine/Plugins/PCG/Source/PCG/Private/Elements/Polygon/PCGOffsetPolygon2D.cpp
+++ b/Engine/Plugins/PCG/Source/PCG/Private/Elements/Polygon/PCGOffsetPolygon2D.cpp
@@ -4,6 +4,7 @@
 
 #include "PCGContext.h"
 #include "Data/PCGPolygon2DData.h"
+#include "Helpers/PCGPolygon2DProcessingHelpers.h"
 
 #include "Curve/PolygonOffsetUtils.h"
 
@@ -28,7 +29,7 @@ bool FPCGOffsetPolygonElement::ExecuteInternal(FPCGContext* Context) const
 		}
 
 		const UE::Geometry::FGeneralPolygon2d& Polygon = PolygonData->GetPolygon();
-		TArrayView<const UE::Geometry::FGeneralPolygon2d> PolyView = MakeArrayView(&Polygon, 1);
+		TConstArrayView<UE::Geometry::FGeneralPolygon2d> PolyView = MakeConstArrayView(&Polygon, 1);
 		TArray<UE::Geometry::FGeneralPolygon2d> OutPolys;
 
 		bool bOperationSuccess = false;
@@ -66,13 +67,37 @@ bool FPCGOffsetPolygonElement::ExecuteInternal(FPCGContext* Context) const
 
 		if (bOperationSuccess)
 		{
+			TArray<PCGPolygon2DProcessingHelpers::PolygonVertexMapping> VertexMapping;
+			PCGPolygon2DProcessingHelpers::PolyToMetadataInfoMap PolyToMetadataInfoMap;
+
+			if (Settings->bInheritMetadata)
+			{
+				TArray<PCGMetadataEntryKey> PolygonEntryKeys;
+				PolygonEntryKeys = PolygonData->GetConstVerticesEntryKeys();
+				PCGPolygon2DProcessingHelpers::AddToPolyMetadataInfoMap(PolyToMetadataInfoMap, Polygon, PolygonData->ConstMetadata(), PolygonEntryKeys);
+				VertexMapping = PCGPolygon2DProcessingHelpers::BuildVertexMapping(OutPolys, PolyView, {});
+			}
+
 			for (int OutPolyIndex = 0; OutPolyIndex < OutPolys.Num(); ++OutPolyIndex)
 			{
 				UPCGPolygon2DData* OutPolyData = FPCGContext::NewObject_AnyThread<UPCGPolygon2DData>(Context);
 
-				OutPolyData->InitializeFromData(PolygonData);
+				FPCGInitializeFromDataParams InitializeParams(PolygonData);
+				InitializeParams.bInheritMetadata = InitializeParams.bInheritAttributes = Settings->bInheritMetadata;
+				InitializeParams.bInheritSpatialData = true;
+
+				OutPolyData->InitializeFromDataWithParams(InitializeParams);
 				Context->OutputData.TaggedData.Add_GetRef(Input).Data = OutPolyData;
-				OutPolyData->SetPolygon(OutPolys[OutPolyIndex]);
+
+				TArray<PCGMetadataEntryKey> EntryKeys;
+				if (Settings->bInheritMetadata)
+				{
+					EntryKeys = PCGPolygon2DProcessingHelpers::ComputeEntryKeysAndInitializeAttributes(OutPolyData->MutableMetadata(), VertexMapping[OutPolyIndex], PolyToMetadataInfoMap);
+				}
+
+				TConstArrayView<const PCGMetadataEntryKey> EntryKeysView = MakeConstArrayView(EntryKeys);
+
+				OutPolyData->SetPolygon(OutPolys[OutPolyIndex], EntryKeysView.IsEmpty() ? nullptr : &EntryKeysView);
 				OutPolyData->SetTransform(PolygonData->GetTransform());
 			}
 		}
diff --git a/Engine/Plugins/PCG/Source/PCG/Private/Elements/Polygon/PCGPolygon2DOperation.cpp b/Engine/Plugins/PCG/Source/PCG/Private/Elements/Polygon/PCGPolygon2DOperation.cpp
--- a/Engine/Plugins/PCG/Source/PCG/Private/Elements/Polygon/PCGPolygon2DOperation.cpp
+++ b/Engine/Plugins/PCG/Source/PCG/Private/Elements/Polygon/PCGPolygon2DOperation.cpp
@@ -4,6 +4,7 @@
 
 #include "PCGContext.h"
 #include "Data/PCGPolygon2DData.h"
+#include "Helpers/PCGPolygon2DProcessingHelpers.h"
 
 #include "Curve/PolygonIntersectionUtils.h"
 
@@ -77,6 +78,7 @@ EPCGElementExecutionLoopMode FPCGPolygon2DOperationElement::ExecutionLoopMode(co
 bool FPCGPolygon2DOperationElement::ExecuteInternal(FPCGContext* Context) const
 {
 	TRACE_CPUPROFILER_EVENT_SCOPE(FPCGPolygon2DOperationElement::Execute);
+	using namespace PCGPolygon2DProcessingHelpers;
 
 	const UPCGPolygon2DOperationSettings* Settings = Context->GetInputSettings<UPCGPolygon2DOperationSettings>();
 	check(Settings);
@@ -87,16 +89,26 @@ bool FPCGPolygon2DOperationElement::ExecuteInternal(FPCGContext* Context) const
 	bool bHasTransform = false;
 	FTransform Transform = FTransform::Identity;
 	TArray<UE::Geometry::FGeneralPolygon2d> Polygons;
+	TArray<TArray<PCGMetadataEntryKey>> PolygonsVertexEntryKeys;
+	TArray<const UPCGMetadata*> PolygonsMetadata;
 	TArray<int> PolyIndexToInputIndex;
 	TArray<UE::Geometry::FGeneralPolygon2d> ClipPolygons;
+	TArray<TArray<PCGMetadataEntryKey>> ClipPolygonsVertexEntryKeys;
+	TArray<const UPCGMetadata*> ClipPolygonsMetadata;
 
-	auto AddPolygon = [&bHasTransform, &Transform](const UE::Geometry::FGeneralPolygon2d& InPoly, const FTransform& InTransform, TArray<UE::Geometry::FGeneralPolygon2d>& OutPolyList)
+	auto AddPolygon = [&bHasTransform, &Transform](const UPCGPolygon2DData* InData, TArray<UE::Geometry::FGeneralPolygon2d>& OutPolyList, TArray<TArray<PCGMetadataEntryKey>>& OutPolyVertexEntryKeysList, TArray<const UPCGMetadata*>& OutMetadataList)
 	{
+		check(InData);
+		const UE::Geometry::FGeneralPolygon2d& InPoly = InData->GetPolygon();
+		const FTransform& InTransform = InData->GetTransform();
+		TConstArrayView<PCGMetadataEntryKey> InEntryKeys = InData->GetConstVerticesEntryKeys();
+
+		OutMetadataList.Add(InData->ConstMetadata());
+
 		if (!bHasTransform)
 		{
 			UE::Geometry::FGeneralPolygon2d& OutPoly = OutPolyList.Add_GetRef(InPoly);
-			// @todo_pcg: having clockwise polygons should be supported because it allows more flexibility
-			ensure(!OutPoly.OuterIsClockwise());
+			OutPolyVertexEntryKeysList.Emplace(InEntryKeys);
 
 			Transform = InTransform;
 			bHasTransform = true;
@@ -105,36 +117,30 @@ bool FPCGPolygon2DOperationElement::ExecuteInternal(FPCGContext* Context) const
 		{
 			// @todo_pcg: we could validate that the normal are parallel.
 			UE::Geometry::FGeneralPolygon2d& Poly = OutPolyList.Emplace_GetRef();
+			TArray<PCGMetadataEntryKey>& EntryKeys = OutPolyVertexEntryKeysList.Emplace_GetRef(InEntryKeys);
 			// Important note: since this can technically flip the normal, that would in turn change the polygon orientation.
 			// Unfortunately, this means that the winding of the polygon might change.
 			// However, the polygon transform isn't aware of this, so we need to rebuild the polygon from scratch here instead.
 			const FTransform RelativeTransform = InTransform.GetRelativeTransform(Transform);
 
-			auto TransformPoly = [&RelativeTransform](UE::Geometry::FPolygon2d& Poly, bool bIsHole)
+			Poly = InPoly;
+			Poly.Transform([&RelativeTransform](const FVector2D& InVertex)
 			{
-				Poly.Transform([&RelativeTransform](const FVector2d& InVertex)
-				{
-					const FVector TransformedVertex = RelativeTransform.TransformPosition(FVector(InVertex, 0.0));
-					return FVector2d(TransformedVertex.X, TransformedVertex.Y);
-				});
-
-				if (Poly.IsClockwise() != bIsHole)
-				{
-					Poly.Reverse();
-				}
-			};
-
-			UE::Geometry::FPolygon2d Outer = InPoly.GetOuter();
-			TransformPoly(Outer, false);
+				const FVector TransformedVertex = RelativeTransform.TransformPosition(FVector(InVertex, 0.0));
+				return FVector2D(TransformedVertex.X, TransformedVertex.Y);
+			});
 
-			Poly.SetOuter(MoveTemp(Outer));
-
-			for (const UE::Geometry::FPolygon2d& InHole : InPoly.GetHoles())
+			// If the outer was reversed because of the operation (due to a normal sign toggle, for example - and because the Transform doesn't update the flag),
+			// We need to fix this up so that the reported type of the polygon is in line with the reported type of the outer.
+			if (Poly.OuterIsClockwise() != Poly.GetOuter().IsClockwise())
 			{
-				UE::Geometry::FPolygon2d Hole = InHole;
-				TransformPoly(Hole, true);
+				PCGPolygon2DData::ReversePolygon(Poly, EntryKeys);
+
+				// Finally, reupdate the outer so the flag is put back to its good value
+				UE::Geometry::FPolygon2d TempOuter = Poly.GetOuter();
+				Poly.SetOuter(MoveTemp(TempOuter));
 
-				Poly.AddHole(MoveTemp(Hole));
+				check(Poly.OuterIsClockwise() == InPoly.OuterIsClockwise());
 			}
 		}
 	};
@@ -144,7 +150,7 @@ bool FPCGPolygon2DOperationElement::ExecuteInternal(FPCGContext* Context) const
 		if (const UPCGPolygon2DData* Polygon = Cast<UPCGPolygon2DData>(Inputs[InputIndex].Data))
 		{
 			PolyIndexToInputIndex.Add(InputIndex);
-			AddPolygon(Polygon->GetPolygon(), Polygon->GetTransform(), Polygons);
+			AddPolygon(Polygon, Polygons, PolygonsVertexEntryKeys, PolygonsMetadata);
 		}
 	}
 
@@ -152,33 +158,67 @@ bool FPCGPolygon2DOperationElement::ExecuteInternal(FPCGContext* Context) const
 	{
 		if (const UPCGPolygon2DData* Polygon = Cast<UPCGPolygon2DData>(ClipInputs[ClipIndex].Data))
 		{
-			AddPolygon(Polygon->GetPolygon(), Polygon->GetTransform(), ClipPolygons);
+			AddPolygon(Polygon, ClipPolygons, ClipPolygonsVertexEntryKeys, ClipPolygonsMetadata);
 		}
 	}
 
+	const EPCGPolygonOperationMetadataMode MetadataMode = Settings->MetadataMode;
+
+	// Build a map from polygon ptr to metadata for ease of access later
+	PolyToMetadataInfoMap PolyToMetadataInfoMap;
+	if (MetadataMode != EPCGPolygonOperationMetadataMode::None)
+	{
+		AddToPolyMetadataInfoMap(PolyToMetadataInfoMap, Polygons, PolygonsMetadata, PolygonsVertexEntryKeys);
+	}
+
+	if (MetadataMode == EPCGPolygonOperationMetadataMode::Full)
+	{
+		AddToPolyMetadataInfoMap(PolyToMetadataInfoMap, ClipPolygons, ClipPolygonsMetadata, ClipPolygonsVertexEntryKeys);
+	}
+
 	TArray<UE::Geometry::FGeneralPolygon2d> OutPolys;
 
-	auto AddToOutput = [Context, &Transform](TArray<UE::Geometry::FGeneralPolygon2d>& OutPolys, const FPCGTaggedData* InputPtr)
+	auto AddToOutput = [Context, &Transform, &PolyToMetadataInfoMap, MetadataMode](TArray<UE::Geometry::FGeneralPolygon2d>& Polys, const FPCGTaggedData* InputPtr, TArrayView<UE::Geometry::FGeneralPolygon2d> PrimaryPolygons, TArrayView<UE::Geometry::FGeneralPolygon2d> SecondaryPolygons)
 	{
-		for (int OutPolyIndex = 0; OutPolyIndex < OutPolys.Num(); ++OutPolyIndex)
+		TArray<PolygonVertexMapping> VertexMapping;
+
+		if (MetadataMode != EPCGPolygonOperationMetadataMode::None)
+		{
+			TArrayView<UE::Geometry::FGeneralPolygon2d> FilteredSecondaryPolygons = ((MetadataMode == EPCGPolygonOperationMetadataMode::Full) ? SecondaryPolygons : TArrayView<UE::Geometry::FGeneralPolygon2d>());
+			VertexMapping = BuildVertexMapping(Polys, PrimaryPolygons, FilteredSecondaryPolygons);
+		}
+
+		for (int PolyIndex = 0; PolyIndex < Polys.Num(); ++PolyIndex)
 		{
 			UPCGPolygon2DData* OutPolyData = FPCGContext::NewObject_AnyThread<UPCGPolygon2DData>(Context);
 
 			if (InputPtr)
 			{
-				OutPolyData->InitializeFromData(Cast<UPCGSpatialData>(InputPtr->Data));
+				FPCGInitializeFromDataParams InitializeParams(Cast<UPCGSpatialData>(InputPtr->Data));
+				InitializeParams.bInheritMetadata = InitializeParams.bInheritAttributes = (MetadataMode != EPCGPolygonOperationMetadataMode::None);
+				InitializeParams.bInheritSpatialData = true;
+
+				OutPolyData->InitializeFromDataWithParams(InitializeParams);
 				Context->OutputData.TaggedData.Add_GetRef(*InputPtr).Data = OutPolyData;
 			}
 			else
 			{
 				Context->OutputData.TaggedData.Emplace_GetRef().Data = OutPolyData;
 			}
+			
+			TArray<PCGMetadataEntryKey> EntryKeys;
+			if (MetadataMode != EPCGPolygonOperationMetadataMode::None)
+			{
+				EntryKeys = ComputeEntryKeysAndInitializeAttributes(OutPolyData->MutableMetadata(), VertexMapping[PolyIndex], PolyToMetadataInfoMap);
+			}
 
-			OutPolyData->SetPolygon(MoveTemp(OutPolys[OutPolyIndex]));
+			TConstArrayView<PCGMetadataEntryKey> EntryKeysView = MakeConstArrayView(EntryKeys);
+
+			OutPolyData->SetPolygon(MoveTemp(Polys[PolyIndex]), EntryKeysView.IsEmpty() ? nullptr : &EntryKeysView);
 			OutPolyData->SetTransform(Transform);
 		}
 
-		OutPolys.Reset();
+		Polys.Reset();
 	};
 
 	if (Settings->Operation == EPCGPolygonOperation::Union)
@@ -192,7 +232,7 @@ bool FPCGPolygon2DOperationElement::ExecuteInternal(FPCGContext* Context) const
 		{
 			if (UE::Geometry::PolygonsUnion(Polygons, OutPolys, /*bCopyInputOnFailure=*/false))
 			{
-				AddToOutput(OutPolys, nullptr);
+				AddToOutput(OutPolys, nullptr, Polygons, {});
 			}
 			else if (!Settings->bQuiet)
 			{
@@ -234,7 +274,7 @@ bool FPCGPolygon2DOperationElement::ExecuteInternal(FPCGContext* Context) const
 			}
 
 			OutPolys = MoveTemp(CurrentPolygons);
-			AddToOutput(OutPolys, nullptr);
+			AddToOutput(OutPolys, nullptr, Polygons, {});
 		}
 	}
 	else if (ClipPolygons.IsEmpty())
@@ -249,10 +289,12 @@ bool FPCGPolygon2DOperationElement::ExecuteInternal(FPCGContext* Context) const
 			const FPCGTaggedData& CurrentInput = Inputs[PolyIndexToInputIndex[PolygonIndex]];
 
 			UE::Geometry::FGeneralPolygon2d& Poly = Polygons[PolygonIndex];
+			TArray<PCGMetadataEntryKey>& EntryKeys = PolygonsVertexEntryKeys[PolygonIndex];
 			OutPolys.Reset();
 			bool bOperationOk = false;
 
 			TArrayView<UE::Geometry::FGeneralPolygon2d> SubjPolygons = MakeArrayView(&Poly, 1);
+			TArrayView<TArray<PCGMetadataEntryKey>> SubjEntryKeys = MakeArrayView(&EntryKeys, 1);
 
 			if (Settings->Operation == EPCGPolygonOperation::Difference ||
 				Settings->Operation == EPCGPolygonOperation::Intersection ||
@@ -262,7 +304,7 @@ bool FPCGPolygon2DOperationElement::ExecuteInternal(FPCGContext* Context) const
 					(Settings->Operation == EPCGPolygonOperation::Intersection && UE::Geometry::PolygonsIntersection(SubjPolygons, ClipPolygons, OutPolys)) ||
 					(Settings->Operation == EPCGPolygonOperation::ExclusiveOr && UE::Geometry::PolygonsExclusiveOr(SubjPolygons, ClipPolygons, OutPolys)))
 				{
-					AddToOutput(OutPolys, &CurrentInput);
+					AddToOutput(OutPolys, &CurrentInput, SubjPolygons, ClipPolygons);
 				}
 				else if(!Settings->bQuiet)
 				{
@@ -274,10 +316,13 @@ bool FPCGPolygon2DOperationElement::ExecuteInternal(FPCGContext* Context) const
 				for (int ClipIndex = 0; ClipIndex < ClipPolygons.Num(); ++ClipIndex)
 				{
 					OutPolys.Reset();
+
 					TArrayView<UE::Geometry::FGeneralPolygon2d> ClipSubsetPolygons = MakeArrayView(&ClipPolygons[ClipIndex], 1);
+					TArrayView<TArray<PCGMetadataEntryKey>> ClipSubsetEntryKeys = MakeArrayView(&ClipPolygonsVertexEntryKeys[ClipIndex], 1);
+
 					if (UE::Geometry::PolygonsIntersection(SubjPolygons, ClipSubsetPolygons, OutPolys))
 					{
-						AddToOutput(OutPolys, &CurrentInput);
+						AddToOutput(OutPolys, &CurrentInput, SubjPolygons, ClipPolygons);
 					}
 					else if (!Settings->bQuiet)
 					{
diff --git a/Engine/Plugins/PCG/Source/PCG/Private/Helpers/PCGPolygon2DProcessingHelpers.cpp b/Engine/Plugins/PCG/Source/PCG/Private/Helpers/PCGPolygon2DProcessingHelpers.cpp
new file mode 100644
--- /dev/null
+++ b/Engine/Plugins/PCG/Source/PCG/Private/Helpers/PCGPolygon2DProcessingHelpers.cpp
@@ -0,0 +1,303 @@
+// Copyright Epic Games, Inc. All Rights Reserved.
+
+#include "Helpers/PCGPolygon2DProcessingHelpers.h"
+
+#include "PCGModule.h"
+#include "Data/PCGBasePointData.h"
+#include "Data/PCGPolygon2DData.h"
+#include "Data/PCGSplineData.h"
+#include "Metadata/PCGMetadata.h"
+
+namespace PCGPolygon2DProcessingHelpers
+{
+
+void ForAllVertices(UE::Geometry::FGeneralPolygon2d& InPoly, TFunctionRef<bool(const FVector2D&, int /*VertexIndex*/, int /*SegmentIndex*/, int /*HoleIndex*/)> InFunc)
+{
+	bool bExit = false;
+	int Offset = 0;
+
+	auto ProcessPoly = [&InFunc, &Offset](const UE::Geometry::FPolygon2d& Poly, int HoleIndex) -> bool
+	{
+		const TArray<FVector2D>& Vertices = Poly.GetVertices();
+		for (int VertexIndex = 0; VertexIndex < Vertices.Num(); ++VertexIndex)
+		{
+			if (!InFunc(Vertices[VertexIndex], Offset + VertexIndex, VertexIndex, HoleIndex))
+			{
+				return false;
+			}
+		}
+
+		Offset += Vertices.Num();
+		return true;
+	};
+
+	if (!ProcessPoly(InPoly.GetOuter(), -1))
+	{
+		return;
+	}
+
+	int HoleIndex = 0;
+	for (const UE::Geometry::FPolygon2d& Hole : InPoly.GetHoles())
+	{
+		if (!ProcessPoly(Hole, HoleIndex++))
+		{
+			return;
+		}
+	}
+}
+
+bool DistanceSquaredToPolygons(const FVector2D& InVertex,
+	TArrayView<const UE::Geometry::FGeneralPolygon2d> InPolys,
+	double Tolerance,
+	int& OutPolyIndex,
+	int& OutSegmentIndex,
+	int& OutHoleIndex,
+	double& OutT,
+	double& OutMinDistance,
+	double InitMinDistance)
+{
+	int MinPolyIndex = -1;
+	int MinSegmentIndex = -1;
+	int MinHoleIndex = -1;
+	double MinT = 0.0;
+	double MinDistance = InitMinDistance;
+
+	for(int PolyIndex = 0; PolyIndex < InPolys.Num(); ++PolyIndex)
+	{
+		int HoleIndex = -1;
+		int SegmentIndex = -1;
+		double T = 0.0;
+		const double Distance = InPolys[PolyIndex].DistanceSquared(InVertex, HoleIndex, SegmentIndex, T);
+		if (Distance < MinDistance)
+		{
+			MinDistance = Distance;
+			MinPolyIndex = PolyIndex;
+			MinSegmentIndex = SegmentIndex;
+			MinHoleIndex = HoleIndex;
+			MinT = T;
+
+			if (Distance < Tolerance)
+			{
+				break;
+			}
+		}
+	}
+
+	if (MinPolyIndex >= 0)
+	{
+		OutPolyIndex = MinPolyIndex;
+		OutSegmentIndex = MinSegmentIndex;
+		OutHoleIndex = MinHoleIndex;
+		OutT = MinT;
+		OutMinDistance = MinDistance;
+	}
+
+	return MinDistance < Tolerance;
+}
+
+// Matches the vertices from the Polys to either vertices or edges on the Primary/Secondary polygons.
+TArray<PolygonVertexMapping> BuildVertexMapping(TArray<UE::Geometry::FGeneralPolygon2d>& Polys, TConstArrayView<UE::Geometry::FGeneralPolygon2d> PrimaryPolygons, TConstArrayView<UE::Geometry::FGeneralPolygon2d> SecondaryPolygons)
+{
+	TRACE_CPUPROFILER_EVENT_SCOPE(PCGPolygon2DProcessingHelpers::BuildVertexMapping);
+	TArray<PolygonVertexMapping> VertexMappings;
+	VertexMappings.SetNum(Polys.Num());
+
+	for (int PolyIndex = 0; PolyIndex < Polys.Num(); ++PolyIndex)
+	{
+		UE::Geometry::FGeneralPolygon2d& Poly = Polys[PolyIndex];
+		PolygonVertexMapping& PolyVertexMapping = VertexMappings[PolyIndex];
+
+		ForAllVertices(Poly, [&PolyVertexMapping, &PrimaryPolygons, &SecondaryPolygons](const FVector2D& Vertex, int VertexIndex, int SegmentIndex, int HoleIndex)
+		{
+			check(VertexIndex == PolyVertexMapping.Num());
+			VertexToSourcePolygonMapping& VertexMapping = PolyVertexMapping.Emplace_GetRef();
+
+			double MinDistanceToPrimaryPolygons = DBL_MAX;
+			int NearestPolyIndex = -1;
+			double T = 0.0;
+
+			// Test against primary polygons first
+			bool bDone = DistanceSquaredToPolygons(Vertex, PrimaryPolygons, UE_SMALL_NUMBER, NearestPolyIndex, VertexMapping.SegmentIndex, VertexMapping.HoleIndex, T, MinDistanceToPrimaryPolygons, DBL_MAX);
+
+			if (PrimaryPolygons.IsValidIndex(NearestPolyIndex))
+			{
+				VertexMapping.Poly = &PrimaryPolygons[NearestPolyIndex];
+			}
+
+			// Otherwise, test against secondary polygons
+			if (!bDone)
+			{
+				double MinDistanceToSecondaryPolygons = DBL_MAX;
+				NearestPolyIndex = -1;
+				bDone = DistanceSquaredToPolygons(Vertex, SecondaryPolygons, UE_SMALL_NUMBER, NearestPolyIndex, VertexMapping.SegmentIndex, VertexMapping.HoleIndex, T, MinDistanceToSecondaryPolygons, MinDistanceToPrimaryPolygons);
+
+				if (SecondaryPolygons.IsValidIndex(NearestPolyIndex))
+				{
+					VertexMapping.Poly = &SecondaryPolygons[NearestPolyIndex];
+				}
+			}
+
+			// Temporarily, we'll want to ensure here, but if it happened, we might be fine with it, if the final distance is reasonable.
+			if (ensure(VertexMapping.Poly))
+			{
+				VertexMapping.Ratio = VertexMapping.Poly->Segment(VertexMapping.SegmentIndex, VertexMapping.HoleIndex).ConvertToUnitRange(T);
+				PCGPolygon2DData::GetStartEndVertexIndices(*VertexMapping.Poly, VertexMapping.SegmentIndex, VertexMapping.HoleIndex, VertexMapping.StartVertexIndex, VertexMapping.EndVertexIndex);
+			}
+
+			return true; // always process all vertices
+		});
+	}
+
+	return VertexMappings;
+}
+
+void AddToPolyMetadataInfoMap(PolyToMetadataInfoMap& OutMap, const UE::Geometry::FGeneralPolygon2d& Poly, const UPCGMetadata* Metadata, const TArray<PCGMetadataEntryKey>& EntryKeys)
+{
+	OutMap.Add(&Poly, MetadataInfo(Metadata, &EntryKeys));
+}
+
+void AddToPolyMetadataInfoMap(PolyToMetadataInfoMap& OutMap, const TArray<UE::Geometry::FGeneralPolygon2d>& Polys, const TArray<const UPCGMetadata*>& Metadata, const TArray<TArray<PCGMetadataEntryKey>>& EntryKeys)
+{
+	check(Polys.Num() == Metadata.Num() && Polys.Num() == EntryKeys.Num());
+	for (int i = 0; i < Polys.Num(); ++i)
+	{
+		OutMap.Add(&Polys[i], MetadataInfo(Metadata[i], &EntryKeys[i]));
+	}
+}
+
+TArray<PCGMetadataEntryKey> ComputeEntryKeysAndInitializeAttributes(UPCGMetadata* OutMetadata, const PolygonVertexMapping& InVertexMapping, const PolyToMetadataInfoMap& InInfoMap)
+{
+	TRACE_CPUPROFILER_EVENT_SCOPE(PCGPolygon2DProcessingHelpers::ComputeEntryKeys);
+	TArray<PCGMetadataEntryKey> EntryKeys;
+	EntryKeys.Reserve(InVertexMapping.Num());
+
+	TSet<const UPCGMetadata*> VisitedMetadata;
+
+	for (int VertexIndex = 0; VertexIndex < InVertexMapping.Num(); ++VertexIndex)
+	{
+		const VertexToSourcePolygonMapping& VertexMapping = InVertexMapping[VertexIndex];
+
+		const UE::Geometry::FGeneralPolygon2d* SourcePolygon = VertexMapping.Poly;
+		check(SourcePolygon);
+
+		const auto [SourceMetadata, SourceEntryKeys] = InInfoMap[SourcePolygon];
+
+		// Propagate attributes from source metadata to out metadata
+		if (!VisitedMetadata.Contains(SourceMetadata))
+		{
+			VisitedMetadata.Add(SourceMetadata);
+			OutMetadata->AddAttributes(SourceMetadata);
+		}
+
+		PCGMetadataEntryKey& NewKey = EntryKeys.Add_GetRef(PCGInvalidEntryKey);
+
+		if (SourceEntryKeys->IsEmpty())
+		{
+			continue;
+		}
+
+		OutMetadata->InitializeOnSet(NewKey);
+
+		TStaticArray<TPair<PCGMetadataEntryKey, float>, 2> Coefficients;
+		Coefficients[0] = { (*SourceEntryKeys)[VertexMapping.StartVertexIndex], 1.0f - VertexMapping.Ratio };
+		Coefficients[1] = { (*SourceEntryKeys)[VertexMapping.EndVertexIndex], VertexMapping.Ratio };
+		OutMetadata->ComputeWeightedAttribute(NewKey, Coefficients, SourceMetadata);
+	}
+
+	return EntryKeys;
+}
+
+TArray<PCGMetadataEntryKey> ComputeEntryKeysBySampling(const UPCGSpatialData* InputData, TConstArrayView<FVector> Positions, UPCGMetadata* OutMetadata)
+{
+	check(InputData);
+	check(OutMetadata);
+
+	TArray<PCGMetadataEntryKey> EntryKeys;
+	EntryKeys.Reserve(Positions.Num());
+
+	const FBox DefaultBox(FVector(-1, -1, -1), FVector(1, 1, 1));
+	for (const FVector& Position : Positions)
+	{
+		FPCGPoint OutPoint;
+		InputData->SamplePoint(FTransform(Position), DefaultBox, OutPoint, OutMetadata);
+		EntryKeys.Add(OutPoint.MetadataEntry);
+	}
+
+	return EntryKeys;
+}
+
+TArray<PCGMetadataEntryKey> ComputeEntryKeysByNearestSegmentSampling(const UPCGSpatialData* InputData, double InSearchRadius, UE::Geometry::FGeneralPolygon2d& InPoly, const FTransform& InPolyTransform, UPCGMetadata* OutMetadata)
+{
+	check(InputData);
+	check(OutMetadata);
+
+	TArray<PCGMetadataEntryKey> EntryKeys;
+	EntryKeys.Reserve(PCGPolygon2DData::PolygonVertexCount(InPoly));
+
+	if (const UPCGSplineData* SplineData = Cast<UPCGSplineData>(InputData))
+	{
+		ForAllVertices(InPoly, [SplineData, &InPolyTransform, OutMetadata, &EntryKeys](const FVector2D& Vertex, int VertexIndex, int SegmentIndex, int HoleIndex)
+		{
+			FVector Position = InPolyTransform.TransformPosition(FVector(Vertex, 0.0));
+
+			float InputKey = SplineData->SplineStruct.FindInputKeyClosestToWorldLocation(Position);
+
+			PCGMetadataEntryKey& EntryKey = EntryKeys.Emplace_GetRef(PCGInvalidEntryKey);
+			SplineData->WriteMetadataToEntry(InputKey, EntryKey, OutMetadata);
+
+			return true; // always process all points
+		});
+	}
+	else if (const UPCGBasePointData* PointData = Cast<UPCGBasePointData>(InputData))
+	{
+		const PCGPointOctree::FPointOctree& Octree = PointData->GetPointOctree();
+		const TConstPCGValueRange<FTransform> PointTransforms = PointData->GetConstTransformValueRange();
+		const TConstPCGValueRange<PCGMetadataEntryKey> PointEntryKeys = PointData->GetConstMetadataEntryValueRange();
+		const UPCGMetadata* SourceMetadata = PointData->ConstMetadata();
+
+		ForAllVertices(InPoly, [&Octree, &PointTransforms, &PointEntryKeys, &InPolyTransform, OutMetadata, &EntryKeys, InSearchRadius, SourceMetadata](const FVector2D& Vertex, int VertexIndex, int SegmentIndex, int HoleIndex)
+		{
+			TArray<TPair<PCGMetadataEntryKey, float>> WeightedSourceEntries;
+			FVector Position = InPolyTransform.TransformPosition(FVector(Vertex, 0.0));
+			FBoxCenterAndExtent SearchBounds(Position, FVector(InSearchRadius, InSearchRadius, InSearchRadius));
+			const double SqrSearchRadius = FMath::Square(InSearchRadius);
+
+			Octree.FindElementsWithBoundsTest(SearchBounds, [&Position, &PointTransforms, &PointEntryKeys, &WeightedSourceEntries, SqrSearchRadius](const PCGPointOctree::FPointRef& PointRef)
+			{
+				const double SquaredDistance = FVector::DistSquared(Position, PointTransforms[PointRef.Index].GetLocation());
+				if (SquaredDistance < SqrSearchRadius)
+				{
+					WeightedSourceEntries.Emplace(PointEntryKeys[PointRef.Index], 1.0f - (SquaredDistance / SqrSearchRadius));
+				}
+			});
+
+			float TotalWeight = 0.0;
+			for (TPair<PCGMetadataEntryKey, float>& WeightedEntry : WeightedSourceEntries)
+			{
+				TotalWeight += WeightedEntry.Value;
+			}
+
+			if (TotalWeight > 0)
+			{
+				for (TPair<PCGMetadataEntryKey, float>& WeightedEntry : WeightedSourceEntries)
+				{
+					WeightedEntry.Value /= TotalWeight;
+				}
+			}
+
+			PCGMetadataEntryKey& EntryKey = EntryKeys.Emplace_GetRef(PCGInvalidEntryKey);
+			OutMetadata->InitializeOnSet(EntryKey);
+			OutMetadata->ComputeWeightedAttribute(EntryKey, WeightedSourceEntries, SourceMetadata);
+
+			return true;
+		});
+	}
+	else
+	{
+		UE_LOG(LogPCG, Error, TEXT("Unsupported type in PCGPolygon2DProcessingHelpers::ComputeEntryKeysByNearestSegmentSampling"));
+	}
+
+	return EntryKeys;
+}
+
+} // end namespace PCGPolygon2DProcessingHelpers
\ No newline at end of file
diff --git a/Engine/Plugins/PCG/Source/PCG/Private/Helpers/PCGPolygon2DProcessingHelpers.h b/Engine/Plugins/PCG/Source/PCG/Private/Helpers/PCGPolygon2DProcessingHelpers.h
new file mode 100644
--- /dev/null
+++ b/Engine/Plugins/PCG/Source/PCG/Private/Helpers/PCGPolygon2DProcessingHelpers.h
@@ -0,0 +1,53 @@
+// Copyright Epic Games, Inc. All Rights Reserved.
+
+#pragma once
+
+#include "Metadata/PCGMetadataCommon.h"
+
+#include "Curve/GeneralPolygon2.h"
+
+class UPCGMetadata;
+class UPCGSpatialData;
+
+namespace PCGPolygon2DProcessingHelpers
+{
+	/** Visits all vertices in the polygon, starting from the outer. Return false in the InFunc to exit the loop prematurely. */
+	void ForAllVertices(UE::Geometry::FGeneralPolygon2d& InPoly, TFunctionRef<bool(const FVector2D& /*VertexPosition*/, int /*VertexIndex*/, int /*SegmentIndex*/, int /*HoleIndex*/)> InFunc);
+	
+	/** Finds the nearest segment in any of the polygons and returns it's information. Note that the Out variables will be written to only if there is a distance found that's smaller than InitMinDistance. Returns true if the vertex found is inside the given tolerance (short-circuits the rest of the visit.) */
+	bool DistanceSquaredToPolygons(const FVector2D& InVertex, TArrayView<const UE::Geometry::FGeneralPolygon2d> InPolys, double Tolerance, int& OutPolyIndex, int& OutSegmentIndex, int& OutHoleIndex, double& OutT, double& OutMinDistance, double InitMinDistance = DBL_MAX);
+
+	using MetadataInfo = TPair<const UPCGMetadata*, const TArray<PCGMetadataEntryKey>*>;
+	using PolyToMetadataInfoMap = TMap<const UE::Geometry::FGeneralPolygon2d*, MetadataInfo>;
+	
+	/** Builds a map from polygon pointer to metadata pointer & entry keys */
+	void AddToPolyMetadataInfoMap(PolyToMetadataInfoMap& OutMap, const UE::Geometry::FGeneralPolygon2d& Poly, const UPCGMetadata* Metadata, const TArray<PCGMetadataEntryKey>& EntryKeys);
+	void AddToPolyMetadataInfoMap(PolyToMetadataInfoMap& OutMap, const TArray<UE::Geometry::FGeneralPolygon2d>& Polys, const TArray<const UPCGMetadata*>& Metadata, const TArray<TArray<PCGMetadataEntryKey>>& EntryKeys);
+
+	struct VertexToSourcePolygonMapping
+	{
+		const UE::Geometry::FGeneralPolygon2d* Poly = nullptr;
+		int SegmentIndex = -1;
+		int HoleIndex = -1;
+		int StartVertexIndex = -1;
+		int EndVertexIndex = -1;
+		double Ratio = 0.0;
+	};
+
+	using PolygonVertexMapping = TArray<VertexToSourcePolygonMapping>;
+	
+	/** Builds a mapping for all vertices of all polys to their nearest vertex/edge in the primary & secondary polygons.
+	* Warning: this is an expensive operation.
+	*/
+	TArray<PolygonVertexMapping> BuildVertexMapping(TArray<UE::Geometry::FGeneralPolygon2d>& Polys, TConstArrayView<UE::Geometry::FGeneralPolygon2d> PrimaryPolygons, TConstArrayView<UE::Geometry::FGeneralPolygon2d> SecondaryPolygons);
+
+	/** Collates vertex mapping information to allocate metadata entries with proper interpolation. */
+	TArray<PCGMetadataEntryKey> ComputeEntryKeysAndInitializeAttributes(UPCGMetadata* OutMetadata, const PolygonVertexMapping& InVertexMapping, const PolyToMetadataInfoMap& InInfoMap);
+
+	/** Computes entry keys from given positions and the original input data (spline or points), on a sampling basis. Assumes that the positions are "on" the existing edges. */
+	TArray<PCGMetadataEntryKey> ComputeEntryKeysBySampling(const UPCGSpatialData* InputData, TConstArrayView<FVector> Positions, UPCGMetadata* OutMetadata);
+
+	/** Computes entry keys for polygon vertices by sampling against nearest segments on spatial data (spline or points). */
+	TArray<PCGMetadataEntryKey> ComputeEntryKeysByNearestSegmentSampling(const UPCGSpatialData* InputData, double InSearchRadius, UE::Geometry::FGeneralPolygon2d& InPoly, const FTransform& InPolyTransform, UPCGMetadata* OutMetadata);
+
+}
\ No newline at end of file
diff --git a/Engine/Plugins/PCG/Source/PCG/Public/Data/PCGPolygon2DData.h b/Engine/Plugins/PCG/Source/PCG/Public/Data/PCGPolygon2DData.h
--- a/Engine/Plugins/PCG/Source/PCG/Public/Data/PCGPolygon2DData.h
+++ b/Engine/Plugins/PCG/Source/PCG/Public/Data/PCGPolygon2DData.h
@@ -41,6 +41,18 @@ enum class EPCGPolygon2DDataProperties : uint8
 namespace PCGPolygon2DData
 {
 	const FName VertexDomainName = "Vertices";
+
+	/** Reverses the given polygon & its entry keys with proper hole support. */
+	void ReversePolygon(UE::Geometry::FGeneralPolygon2d& Polygon, TArray<int64>& PolygonEntryKeys);
+
+	/** Counts the number of vertices in a polygon, including the hole vertices. */
+	int PolygonVertexCount(const UE::Geometry::FGeneralPolygon2d& InPoly);
+
+	/** Returns the vertex index (0 to number of vertices - 1) from a segment index & hole index. */
+	int GetVertexIndex(const UE::Geometry::FGeneralPolygon2d& InPoly, int InSegmentIndex, int InHoleIndex);
+
+	/** Returns the vertex indices of a segment start and end. */
+	void GetStartEndVertexIndices(const UE::Geometry::FGeneralPolygon2d& InPoly, int InSegmentIndex, int InHoleIndex, int& OutStartVertexIndex, int& OutEndVertexIndex);
 }
 
 USTRUCT()
diff --git a/Engine/Plugins/PCG/Source/PCG/Public/Elements/Polygon/PCGOffsetPolygon2D.h b/Engine/Plugins/PCG/Source/PCG/Public/Elements/Polygon/PCGOffsetPolygon2D.h
--- a/Engine/Plugins/PCG/Source/PCG/Public/Elements/Polygon/PCGOffsetPolygon2D.h
+++ b/Engine/Plugins/PCG/Source/PCG/Public/Elements/Polygon/PCGOffsetPolygon2D.h
@@ -50,6 +50,9 @@ public:
 	UPROPERTY(BlueprintReadWrite, EditAnywhere, Category = Settings)
 	EPCGPolygonOffsetOperation Operation = EPCGPolygonOffsetOperation::Offset;
 
+	UPROPERTY(BlueprintReadWrite, EditAnywhere, Category = Settings)
+	bool bInheritMetadata = true;
+
 	UPROPERTY(BlueprintReadWrite, EditAnywhere, Category = Settings, meta = (PCG_Overridable))
 	double Offset = 100.0;
 	
diff --git a/Engine/Plugins/PCG/Source/PCG/Public/Elements/Polygon/PCGPolygon2DOperation.h b/Engine/Plugins/PCG/Source/PCG/Public/Elements/Polygon/PCGPolygon2DOperation.h
--- a/Engine/Plugins/PCG/Source/PCG/Public/Elements/Polygon/PCGPolygon2DOperation.h
+++ b/Engine/Plugins/PCG/Source/PCG/Public/Elements/Polygon/PCGPolygon2DOperation.h
@@ -18,6 +18,14 @@ enum class EPCGPolygonOperation : uint8
 	ExclusiveOr,
 };
 
+UENUM()
+enum class EPCGPolygonOperationMetadataMode : uint8
+{
+	None,
+	SourceOnly,
+	Full
+};
+
 class FPCGPolygon2DOperationElement : public IPCGElement
 {
 public:
@@ -55,6 +63,9 @@ public:
 	UPROPERTY(BlueprintReadWrite, EditAnywhere, Category = Settings)
 	EPCGPolygonOperation Operation = EPCGPolygonOperation::Intersection;
 
+	UPROPERTY(BlueprintReadWrite, EditAnywhere, Category = Settings)
+	EPCGPolygonOperationMetadataMode MetadataMode = EPCGPolygonOperationMetadataMode::Full;
+
 	UPROPERTY(BlueprintReadWrite, EditAnywhere, Category = Settings)
 	bool bQuiet = false;
 };
\ No newline at end of file
