[PCG] Fixed code replacement bugs in poly 2d code.

#rnx

diff --git a/Engine/Plugins/PCG/Source/PCG/Private/Elements/Polygon/PCGCreatePolygon2D.cpp b/Engine/Plugins/PCG/Source/PCG/Private/Elements/Polygon/PCGCreatePolygon2D.cpp
--- a/Engine/Plugins/PCG/Source/PCG/Private/Elements/Polygon/PCGCreatePolygon2D.cpp
+++ b/Engine/Plugins/PCG/Source/PCG/Private/Elements/Polygon/PCGCreatePolygon2D.cpp
@@ -171,7 +171,7 @@ bool FPCGCreatePolygonElement::ExecuteInternal(FPCGContext* Context) const
 
 					if (bInputHasMetadata)
 					{
-						PCGPolygon2DProcessingHelpers::ComputeEntryKeysByNearestSegmentSampling(SpatialData, Settings->OpenPolygonWidth, Poly, PlaneTransform, Polygon->MutableMetadata());
+						EntryKeys = PCGPolygon2DProcessingHelpers::ComputeEntryKeysByNearestSegmentSampling(SpatialData, Settings->OpenPolygonWidth, Poly, PlaneTransform, Polygon->MutableMetadata());
 					}
 
 					TConstArrayView<PCGMetadataEntryKey> EntryKeysView(EntryKeys);
diff --git a/Engine/Plugins/PCG/Source/PCG/Private/Helpers/PCGPolygon2DProcessingHelpers.cpp b/Engine/Plugins/PCG/Source/PCG/Private/Helpers/PCGPolygon2DProcessingHelpers.cpp
--- a/Engine/Plugins/PCG/Source/PCG/Private/Helpers/PCGPolygon2DProcessingHelpers.cpp
+++ b/Engine/Plugins/PCG/Source/PCG/Private/Helpers/PCGPolygon2DProcessingHelpers.cpp
@@ -176,9 +176,13 @@ TArray<PCGMetadataEntryKey> ComputeEntryKeysAndInitializeAttributes(UPCGMetadata
 	for (int VertexIndex = 0; VertexIndex < InVertexMapping.Num(); ++VertexIndex)
 	{
 		const VertexToSourcePolygonMapping& VertexMapping = InVertexMapping[VertexIndex];
+		PCGMetadataEntryKey& NewKey = EntryKeys.Add_GetRef(PCGInvalidEntryKey);
 
 		const UE::Geometry::FGeneralPolygon2d* SourcePolygon = VertexMapping.Poly;
-		check(SourcePolygon);
+		if (!SourcePolygon)
+		{
+			continue;
+		}
 
 		const auto [SourceMetadata, SourceEntryKeys] = InInfoMap[SourcePolygon];
 
@@ -189,8 +193,6 @@ TArray<PCGMetadataEntryKey> ComputeEntryKeysAndInitializeAttributes(UPCGMetadata
 			OutMetadata->AddAttributes(SourceMetadata);
 		}
 
-		PCGMetadataEntryKey& NewKey = EntryKeys.Add_GetRef(PCGInvalidEntryKey);
-
 		if (SourceEntryKeys->IsEmpty())
 		{
 			continue;
