[PCG] Fix crash when enabling GPU execution for PCGDataNum node

#rnx

diff --git a/Engine/Plugins/PCG/Source/PCG/Private/Compute/BuiltInKernels/PCGDataNumKernel.cpp b/Engine/Plugins/PCG/Source/PCG/Private/Compute/BuiltInKernels/PCGDataNumKernel.cpp
--- a/Engine/Plugins/PCG/Source/PCG/Private/Compute/BuiltInKernels/PCGDataNumKernel.cpp
+++ b/Engine/Plugins/PCG/Source/PCG/Private/Compute/BuiltInKernels/PCGDataNumKernel.cpp
@@ -5,6 +5,7 @@
 #include "PCGContext.h"
 #include "Compute/DataInterfaces/BuiltInKernels/PCGDataNumDataInterface.h"
 #include "Compute/Elements/PCGComputeGraphElement.h"
+#include "Elements/PCGDataNum.h"
 #include "Graph/PCGGPUGraphCompilationContext.h"
 
 #include "ShaderCompilerCore.h"
@@ -36,55 +37,6 @@ int UPCGDataNumKernel::ComputeThreadCount(const UPCGDataBinding* InBinding) cons
 	return 1;
 }
 
-bool UPCGDataNumKernel::IsKernelDataValid(FPCGContext* InContext) const
-{
-	TRACE_CPUPROFILER_EVENT_SCOPE(UPCGDataNumKernel::IsKernelDataValid);
-
-	if (!Super::IsKernelDataValid(InContext))
-	{
-		return false;
-	}
-
-	if (!InContext || !InContext->IsComputeContext())
-	{
-		return true;
-	}
-
-	const FPCGComputeGraphContext* Context = static_cast<FPCGComputeGraphContext*>(InContext);
-	if (const UPCGDataBinding* DataBinding = Context->DataBinding.Get())
-	{
-		if (const FPCGDataCollectionDesc* OutputDataDesc = DataBinding->GetCachedKernelPinDataDesc(this, PCGPinConstants::DefaultOutputLabel, /*bIsInputPin=*/false))
-		{
-			FPCGKernelAttributeDesc AttributeDesc;
-			bool bConflictingTypesFound = false;
-			bool bPresentOnAllData = false;
-			if (OutputDataDesc->GetAttributeDesc(OutCountAttributeKey.Identifier, AttributeDesc, bConflictingTypesFound, bPresentOnAllData))
-			{
-				if (OutCountAttributeKey.Type == EPCGKernelAttributeType::Invalid)
-				{
-					PCG_KERNEL_VALIDATION_ERR(InContext, GetSettings(), FText::Format(
-						LOCTEXT("InvalidAttributeKeyType", "Output attribute for the count could not be created with an 'invalid' type."),
-						FText::FromString(OutCountAttributeKey.Identifier.ToString())));
-
-					return false;
-				}
-
-				return true;
-			}
-			else
-			{
-				PCG_KERNEL_VALIDATION_ERR(InContext, GetSettings(), FText::Format(
-					LOCTEXT("AttributeMissing", "Output attribute '{0}' not found."),
-					FText::FromString(OutCountAttributeKey.Identifier.ToString())));
-
-				return false;
-			}
-		}
-	}
-
-	return false;
-}
-
 #if WITH_EDITOR
 void UPCGDataNumKernel::CreateAdditionalInputDataInterfaces(FPCGGPUCompilationContext& InOutContext, UObject* InObjectOuter, TArray<TObjectPtr<UComputeDataInterface>>& OutDataInterfaces) const
 {
@@ -112,9 +64,14 @@ void UPCGDataNumKernel::GetOutputPins(TArray<FPCGPinPropertiesGPU>& OutPins) con
 	OutPins.Emplace(PCGPinConstants::DefaultOutputLabel, EPCGDataType::Param);
 }
 
-void UPCGDataNumKernel::SetOutputAttribute(const FName InAttributeName)
+#if WITH_EDITOR
+void UPCGDataNumKernel::InitializeInternal()
 {
-	OutCountAttributeKey = FPCGKernelAttributeKey(InAttributeName, EPCGKernelAttributeType::Int);
+	Super::InitializeInternal();
+
+	const UPCGDataNumSettings* Settings = CastChecked<UPCGDataNumSettings>(GetSettings());
+	OutCountAttributeKey = FPCGKernelAttributeKey(FPCGAttributeIdentifier(Settings->OutputAttributeName), EPCGKernelAttributeType::Int);
 }
+#endif
 
 #undef LOCTEXT_NAMESPACE
diff --git a/Engine/Plugins/PCG/Source/PCG/Private/Compute/BuiltInKernels/PCGDataNumKernel.h b/Engine/Plugins/PCG/Source/PCG/Private/Compute/BuiltInKernels/PCGDataNumKernel.h
--- a/Engine/Plugins/PCG/Source/PCG/Private/Compute/BuiltInKernels/PCGDataNumKernel.h
+++ b/Engine/Plugins/PCG/Source/PCG/Private/Compute/BuiltInKernels/PCGDataNumKernel.h
@@ -16,7 +16,6 @@ class UPCGDataNumKernel : public UPCGComputeKernel
 public:
 	virtual bool ComputeOutputBindingDataDesc(const UPCGComputeGraph* InGraph, FName InOutputPinLabel, UPCGDataBinding* InBinding, FPCGDataCollectionDesc& OutDataDesc) const override;
 	virtual int ComputeThreadCount(const UPCGDataBinding* Binding) const override;
-	virtual bool IsKernelDataValid(FPCGContext* InContext) const override;
 
 #if WITH_EDITOR
 	virtual const TCHAR* GetSourceFilePath() const override { return TEXT("/Plugin/PCG/Private/BuiltInKernels/PCGDataNum.usf"); }
@@ -27,8 +26,12 @@ public:
 	virtual void GetInputPins(TArray<FPCGPinProperties>& OutPins) const override;
 	virtual void GetOutputPins(TArray<FPCGPinPropertiesGPU>& OutPins) const override;
 
-	void SetOutputAttribute(FName InAttributeName);
+protected:
+#if WITH_EDITOR
+	virtual void InitializeInternal() override;
+#endif
 
 private:
+	UPROPERTY()
 	FPCGKernelAttributeKey OutCountAttributeKey;
 };
