[PCG] Fix crash when trying to spawn a spline mesh with an invalid component class

#rnx

diff --git a/Engine/Plugins/PCG/Source/PCG/Private/Elements/PCGSpawnSplineMesh.cpp b/Engine/Plugins/PCG/Source/PCG/Private/Elements/PCGSpawnSplineMesh.cpp
--- a/Engine/Plugins/PCG/Source/PCG/Private/Elements/PCGSpawnSplineMesh.cpp
+++ b/Engine/Plugins/PCG/Source/PCG/Private/Elements/PCGSpawnSplineMesh.cpp
@@ -217,7 +217,7 @@ bool FPCGSpawnSplineMeshElement::ExecuteInternal(FPCGContext* InContext) const
 			
 			UStaticMesh* StaticMesh = IterState.SMCBuilderParams.Descriptor.StaticMesh.Get();
 
-			if (!StaticMesh)
+			if (!StaticMesh || !IterState.SMCBuilderParams.Descriptor.ComponentClass)
 			{
 				IterState.ElementIndex++;
 				continue;
diff --git a/Engine/Plugins/PCG/Source/PCG/Private/Helpers/PCGActorHelpers.cpp b/Engine/Plugins/PCG/Source/PCG/Private/Helpers/PCGActorHelpers.cpp
--- a/Engine/Plugins/PCG/Source/PCG/Private/Helpers/PCGActorHelpers.cpp
+++ b/Engine/Plugins/PCG/Source/PCG/Private/Helpers/PCGActorHelpers.cpp
@@ -405,7 +405,7 @@ UPCGManagedSplineMeshComponent* UPCGActorHelpers::GetOrCreateManagedSplineMeshCo
 
 	const UStaticMesh* StaticMesh = InParams.Descriptor.StaticMesh;
 
-	if (!StaticMesh)
+	if (!StaticMesh || !InParams.Descriptor.ComponentClass)
 	{
 		return nullptr;
 	}
@@ -467,7 +467,7 @@ UPCGManagedSplineMeshComponent* UPCGActorHelpers::GetOrCreateManagedSplineMeshCo
 
 	FString ComponentName = TEXT("PCGSplineMeshComponent_") + StaticMesh->GetName();
 	const EObjectFlags ObjectFlags = (InSourceComponent->IsInPreviewMode() ? RF_Transient | RF_NonPIEDuplicateTransient : RF_NoFlags);
-	USplineMeshComponent* SplineMeshComponent = InParams.Descriptor.CreateComponent(InTargetActor, MakeUniqueObjectName(InTargetActor, USplineMeshComponent::StaticClass(), FName(ComponentName)), ObjectFlags);
+	USplineMeshComponent* SplineMeshComponent = InParams.Descriptor.CreateComponent(InTargetActor, MakeUniqueObjectName(InTargetActor, InParams.Descriptor.ComponentClass, FName(ComponentName)), ObjectFlags);
 
 	// Init Component
 	{
