[PCG] Fix missing pin for AddModulo

#rnx

diff --git a/Engine/Plugins/PCG/Source/PCG/Private/Elements/Metadata/PCGMetadataMathsOpElement.cpp b/Engine/Plugins/PCG/Source/PCG/Private/Elements/Metadata/PCGMetadataMathsOpElement.cpp
--- a/Engine/Plugins/PCG/Source/PCG/Private/Elements/Metadata/PCGMetadataMathsOpElement.cpp
+++ b/Engine/Plugins/PCG/Source/PCG/Private/Elements/Metadata/PCGMetadataMathsOpElement.cpp
@@ -38,7 +38,8 @@ namespace PCGMetadataMathsSettings
 
 		if (PCGMetadataMathsSettings::IsBinaryOp(Operation)
 			|| Operation == EPCGMetadataMathsOperation::Lerp
-			|| Operation == EPCGMetadataMathsOperation::MulAdd)
+			|| Operation == EPCGMetadataMathsOperation::MulAdd
+			|| Operation == EPCGMetadataMathsOperation::AddModulo)
 		{
 			return PCGMetadataSettingsBaseConstants::DoubleInputFirstLabel;
 		}
