[PCG] Fix oversight in PCGEdModeSceneVisualizeHelpers.cpp Spline visualization helper
Also, fixed linting issue with FPCGToolDrawTarget Normal not being default initialized to Z Up

#rnx

diff --git a/Engine/Plugins/PCG/Source/PCGEditor/Private/EditorMode/Tools/Helpers/PCGEdModeSceneVisualizeHelpers.cpp b/Engine/Plugins/PCG/Source/PCGEditor/Private/EditorMode/Tools/Helpers/PCGEdModeSceneVisualizeHelpers.cpp
--- a/Engine/Plugins/PCG/Source/PCGEditor/Private/EditorMode/Tools/Helpers/PCGEdModeSceneVisualizeHelpers.cpp
+++ b/Engine/Plugins/PCG/Source/PCGEditor/Private/EditorMode/Tools/Helpers/PCGEdModeSceneVisualizeHelpers.cpp
@@ -89,9 +89,9 @@ namespace PCG::EditorMode::Scene::Visualize
 				const FColor EndColor = (LatestIndex == KeyIndex) ? SelectedColor : NormalColor;
 
 				// Find position on first keyframe.
-				const FVector& PreviousPosition = PreviousKeyPosition;
-				const FVector& PreviousRightVector = PreviousKeyRightVector;
-				const FVector& PreviousScale = PreviousKeyScale;
+				FVector PreviousPosition = PreviousKeyPosition;
+				FVector PreviousRightVector = PreviousKeyRightVector;
+				FVector PreviousScale = PreviousKeyScale;
 
 				// Then draw a line for each substep.
 				constexpr int32 NumSteps = 20;
@@ -128,10 +128,14 @@ namespace PCG::EditorMode::Scene::Visualize
 					PDI->DrawLine(StepPosition, StepPosition + StepUpVector * HalfVisualizationWidthScale, LineColor, SDPG_Foreground);
 					PDI->DrawLine(StepPosition, StepPosition + StepRightVector * HalfVisualizationWidthScale, LineColor, SDPG_Foreground);
 
-					PreviousKeyPosition = StepPosition;
-					PreviousKeyRightVector = StepRightVector;
-					PreviousKeyScale = StepUpVector;
+					PreviousPosition = StepPosition;
+					PreviousRightVector = StepRightVector;
+					PreviousScale = StepScale;
 				}
+
+				PreviousKeyPosition = KeyPosition;
+				PreviousKeyRightVector = KeyRightVector;
+				PreviousKeyScale = KeyScale;
 			}
 		}
 	}
diff --git a/Engine/Plugins/PCG/Source/PCGEditor/Public/EditorMode/Tools/PCGInteractiveToolCommon.h b/Engine/Plugins/PCG/Source/PCGEditor/Public/EditorMode/Tools/PCGInteractiveToolCommon.h
--- a/Engine/Plugins/PCG/Source/PCGEditor/Public/EditorMode/Tools/PCGInteractiveToolCommon.h
+++ b/Engine/Plugins/PCG/Source/PCGEditor/Public/EditorMode/Tools/PCGInteractiveToolCommon.h
@@ -60,7 +60,7 @@ struct PCGEDITOR_API FPCGToolDrawTarget
 
 	/** Manual vector to offset drawn points from the clicked surface. */
 	UPROPERTY(EditAnywhere, Category = "Draw Target", meta = (UIMin = 0, UIMax = 100, EditCondition = "NormalMode == EPCGToolDrawTargetNormal::Explicit", EditConditionHides))
-	FVector Normal;
+	FVector Normal = FVector::ZAxisVector;
 
 	// @todo_pcg: Target object (such as a mesh)
 
