[PCG] Expose the Graph Parameters to higher order in details view

diff --git a/Engine/Plugins/PCG/Source/PCGEditor/Private/Details/PCGComponentDetails.cpp b/Engine/Plugins/PCG/Source/PCGEditor/Private/Details/PCGComponentDetails.cpp
--- a/Engine/Plugins/PCG/Source/PCGEditor/Private/Details/PCGComponentDetails.cpp
+++ b/Engine/Plugins/PCG/Source/PCGEditor/Private/Details/PCGComponentDetails.cpp
@@ -3,6 +3,7 @@
 #include "Details/PCGComponentDetails.h"
 #include "PCGComponent.h"
 #include "PCGEditorStyle.h"
+#include "PCGGraph.h"
 #include "PCGSubsystem.h"
 #include "Framework/Application/SlateApplication.h"
 #include "Widgets/Text/STextBlock.h"
@@ -12,6 +13,7 @@
 #include "DetailWidgetRow.h"
 #include "Modules/ModuleManager.h"
 #include "LevelEditor.h"
+#include "Algo/NoneOf.h"
 
 #define LOCTEXT_NAMESPACE "PCGComponentDetails"
 
@@ -55,7 +57,7 @@ void FPCGComponentDetails::CustomizeDetails(IDetailLayoutBuilder& DetailBuilder)
 		// Add all properties in the category in order
 		PCGCategory.GetDefaultProperties(AllProperties, bSimpleProperties, bAdvancedProperties);
 
-		for (auto& Property : AllProperties)
+		for (TSharedRef<IPropertyHandle>& Property : AllProperties)
 		{
 			PCGCategory.AddProperty(Property);
 		}
@@ -74,6 +76,19 @@ void FPCGComponentDetails::CustomizeDetails(IDetailLayoutBuilder& DetailBuilder)
 		}
 	}
 
+	// Only promote the special view if customized objects are not a component themselves
+	if (Algo::NoneOf(ObjectsBeingCustomized, [](const TWeakObjectPtr<UObject>& Object) { return !Object.IsValid() || Object.Get()->IsA<UPCGComponent>(); }))
+	{
+		TArray<UObject*> Objects;
+		Objects.Reserve(SelectedComponents.Num());
+		Algo::Transform(SelectedComponents, Objects, [](const TWeakObjectPtr<UPCGComponent>& Component) -> UObject*
+		{
+			return Component.IsValid() ? CastChecked<UObject>(Component.Get()->GetGraphInstance(), ECastCheckedType::NullAllowed) : nullptr;
+		});
+		PCGCategory.AddExternalObjectProperty(Objects, TEXT("Graph"));
+		PCGCategory.AddExternalObjectProperty(Objects, TEXT("ParametersOverrides"));
+	}
+
 	FDetailWidgetRow& NewRow = PCGCategory.AddCustomRow(FText::GetEmpty());
 
 	NewRow.ValueContent()
