[PCG] Improve attribute reduce to be able to reduce directly into the data domain

* Add unit tests for this new operation
* Add unit tests for operations that were not tested

#rnx

diff --git a/Engine/Plugins/PCG/Source/PCG/Private/Elements/PCGAttributeReduceElement.cpp b/Engine/Plugins/PCG/Source/PCG/Private/Elements/PCGAttributeReduceElement.cpp
--- a/Engine/Plugins/PCG/Source/PCG/Private/Elements/PCGAttributeReduceElement.cpp
+++ b/Engine/Plugins/PCG/Source/PCG/Private/Elements/PCGAttributeReduceElement.cpp
@@ -182,6 +182,11 @@ TArray<FPCGPreConfiguredSettingsInfo> UPCGAttributeReduceSettings::GetPreconfigu
 }
 #endif
 
+bool UPCGAttributeReduceSettings::HasDynamicPins() const
+{
+	return bWriteToDataDomain;
+}
+
 void UPCGAttributeReduceSettings::ApplyPreconfiguredSettings(const FPCGPreConfiguredSettingsInfo& PreconfiguredInfo)
 {
 	if (const UEnum* EnumPtr = StaticEnum<EPCGAttributeReduceOperation>())
@@ -245,7 +250,7 @@ TArray<FPCGPinProperties> UPCGAttributeReduceSettings::InputPinProperties() cons
 TArray<FPCGPinProperties> UPCGAttributeReduceSettings::OutputPinProperties() const
 {
 	TArray<FPCGPinProperties> PinProperties;
-	PinProperties.Emplace(PCGPinConstants::DefaultOutputLabel, EPCGDataType::Param);
+	PinProperties.Emplace(PCGPinConstants::DefaultOutputLabel, bWriteToDataDomain ? EPCGDataType::Any : EPCGDataType::Param);
 
 	return PinProperties;
 }
@@ -266,7 +271,7 @@ bool FPCGAttributeReduceElement::ExecuteInternal(FPCGContext* Context) const
 
 	TArray<FPCGTaggedData> Inputs = Context->InputData.GetInputsByPin(PCGPinConstants::DefaultInputLabel);
 	TArray<FPCGTaggedData>& Outputs = Context->OutputData.TaggedData;
-	UPCGParamData* OutputParams = nullptr;
+	UPCGData* OutputData = nullptr;
 	FPCGMetadataAttributeBase* NewAttribute = nullptr;
 
 	for (int32 i = 0; i < Inputs.Num(); ++i)
@@ -279,6 +284,13 @@ bool FPCGAttributeReduceElement::ExecuteInternal(FPCGContext* Context) const
 			continue;
 		}
 
+		if (Settings->bWriteToDataDomain && !InputData->IsSupportedMetadataDomainID(PCGMetadataDomainID::Data))
+		{
+			static const FPCGAttributePropertySelector DataSelector = FPCGAttributePropertySelector::CreateAttributeSelector(NAME_None, PCGDataConstants::DataDomainName);
+			PCGLog::Metadata::LogInvalidMetadataDomain(DataSelector, Context);
+			continue;
+		}
+
 		const FPCGAttributePropertyInputSelector InputSource = Settings->InputSource.CopyAndFixLast(InputData);
 		const FName OutputAttributeName = (Settings->OutputAttributeName == PCGMetadataAttributeConstants::SourceNameAttributeName) ? InputSource.GetName() : Settings->OutputAttributeName;
 
@@ -291,36 +303,51 @@ bool FPCGAttributeReduceElement::ExecuteInternal(FPCGContext* Context) const
 			continue;
 		}
 
-		auto DoOperation = [&Accessor, &Keys, Settings, &OutputParams, &NewAttribute, OutputAttributeName, &Context](auto DummyValue) -> bool
+		const bool bMergeOutputAttributes = Settings->ShouldMergeOutputAttributes();
+
+		auto DoOperation = [&Accessor, &Keys, Settings, &OutputData, &NewAttribute, OutputAttributeName, &Context, InputData, bMergeOutputAttributes](auto DummyValue) -> bool
 		{
 			using AttributeType = decltype(DummyValue);
 
 			const EPCGAttributeReduceOperation Operation = Settings->Operation;
-			const bool bMergeOutputAttributes = Settings->bMergeOutputAttributes;
 
-			bool bCreatedNewParams = false;
-			if (!OutputParams || !bMergeOutputAttributes)
+			bool bCreatedNewData = false;
+			if (!OutputData || !bMergeOutputAttributes)
 			{
-				OutputParams = FPCGContext::NewObject_AnyThread<UPCGParamData>(Context);
-				bCreatedNewParams = true;
+				if (Settings->bWriteToDataDomain)
+				{
+					OutputData = InputData->DuplicateData(Context);
+				}
+				else
+				{
+					OutputData = FPCGContext::NewObject_AnyThread<UPCGParamData>(Context);
+				}
+
+				bCreatedNewData = true;
 			}
 
-			auto DoOperationInternal = [&OutputParams, &Accessor, &Keys, Settings, bCreatedNewParams, &NewAttribute, OutputAttributeName](auto DummyOutputValue) -> bool
+			auto DoOperationInternal = [&OutputData, &Accessor, &Keys, Settings, bCreatedNewData, &NewAttribute, OutputAttributeName, bMergeOutputAttributes](auto DummyOutputValue) -> bool
 			{
 				using OutAttributeType = decltype(DummyOutputValue);
 				const EPCGAttributeReduceOperation Operation = Settings->Operation;
-				const bool bMergeOutputAttributes = Settings->bMergeOutputAttributes;
 				const FString& JoinDelimiter = Settings->JoinDelimiter;
 
 				OutAttributeType OutputValue = PCG::Private::MetadataTraits<OutAttributeType>::ZeroValue();
+				FPCGMetadataDomain* OutputDomain = Settings->bWriteToDataDomain ? OutputData->MutableMetadata()->GetMetadataDomain(PCGMetadataDomainID::Data) : OutputData->MutableMetadata()->GetDefaultMetadataDomain();
+				check(OutputDomain);
 
-				if (bCreatedNewParams)
+				if (bCreatedNewData)
 				{
-					NewAttribute = OutputParams->Metadata->CreateAttribute<OutAttributeType>(OutputAttributeName, OutputValue, /*bAllowInterpolation=*/ true, /*bOverrideParent=*/false);
+					if (OutputDomain->HasAttribute(OutputAttributeName))
+					{
+						OutputDomain->DeleteAttribute(OutputAttributeName);
+					}
+
+					NewAttribute = OutputDomain->CreateAttribute<OutAttributeType>(OutputAttributeName, OutputValue, /*bAllowInterpolation=*/ true, /*bOverrideParent=*/false);
 
 					if (!NewAttribute)
 					{
-						OutputParams = nullptr;
+						OutputData = nullptr;
 						return false;
 					}
 				}
@@ -344,7 +371,7 @@ bool FPCGAttributeReduceElement::ExecuteInternal(FPCGContext* Context) const
 						bSuccess = PCGAttributeReduceElement::Sum<AttributeType>(*Keys, *Accessor, OutputValue);
 						break;
 					case EPCGAttributeReduceOperation::Join:
-						bSuccess = PCGAttributeReduceElement::Join<AttributeType>(*Keys, *Accessor, Settings->JoinDelimiter, OutputValue);
+						bSuccess = PCGAttributeReduceElement::Join<AttributeType>(*Keys, *Accessor, JoinDelimiter, OutputValue);
 					default:
 						break;
 					}
@@ -354,7 +381,7 @@ bool FPCGAttributeReduceElement::ExecuteInternal(FPCGContext* Context) const
 					switch (Operation)
 					{
 					case EPCGAttributeReduceOperation::Join:
-						bSuccess = PCGAttributeReduceElement::Join<AttributeType>(*Keys, *Accessor, Settings->JoinDelimiter, OutputValue);
+						bSuccess = PCGAttributeReduceElement::Join<AttributeType>(*Keys, *Accessor, JoinDelimiter, OutputValue);
 					default:
 						break;
 					}
@@ -365,9 +392,15 @@ bool FPCGAttributeReduceElement::ExecuteInternal(FPCGContext* Context) const
 					FPCGMetadataAttribute<OutAttributeType>* TypedNewAttribute = static_cast<FPCGMetadataAttribute<OutAttributeType>*>(NewAttribute);
 					check(TypedNewAttribute);
 
+					PCGMetadataEntryKey EntryKey = PCGFirstEntryKey;
+					if (!Settings->bWriteToDataDomain || OutputDomain->GetItemCountForChild() == 0)
+					{
+						EntryKey = OutputDomain->AddEntry();
+					}
+
 					// Implementation note: since the default value does not match the value computed here
 					// and because we might have multiple entries, we need to set it in the attribute
-					TypedNewAttribute->SetValue(OutputParams->Metadata->AddEntry(), OutputValue);
+					TypedNewAttribute->SetValue(EntryKey, OutputValue);
 				}
 
 				return bSuccess;
@@ -389,10 +422,10 @@ bool FPCGAttributeReduceElement::ExecuteInternal(FPCGContext* Context) const
 			continue;
 		}
 
-		if (ensure(OutputParams) && (Outputs.IsEmpty() || !Settings->bMergeOutputAttributes))
+		if (ensure(OutputData) && (Outputs.IsEmpty() || !bMergeOutputAttributes))
 		{
 			FPCGTaggedData& Output = Outputs.Add_GetRef(Inputs[i]);
-			Output.Data = OutputParams;
+			Output.Data = OutputData;
 		}
 	}
 
@@ -402,7 +435,8 @@ bool FPCGAttributeReduceElement::ExecuteInternal(FPCGContext* Context) const
 EPCGElementExecutionLoopMode FPCGAttributeReduceElement::ExecutionLoopMode(const UPCGSettings* InSettings) const
 {
 	const UPCGAttributeReduceSettings* Settings = Cast<UPCGAttributeReduceSettings>(InSettings);
-	return (!Settings || !Settings->bMergeOutputAttributes) ? EPCGElementExecutionLoopMode::SinglePrimaryPin : EPCGElementExecutionLoopMode::NotALoop;
+	const bool bMergeOutputAttributes = Settings && Settings->ShouldMergeOutputAttributes();
+	return !bMergeOutputAttributes ? EPCGElementExecutionLoopMode::SinglePrimaryPin : EPCGElementExecutionLoopMode::NotALoop;
 }
 
 #undef LOCTEXT_NAMESPACE
diff --git a/Engine/Plugins/PCG/Source/PCG/Private/Tests/Elements/PCGAttributeReduceTest.cpp b/Engine/Plugins/PCG/Source/PCG/Private/Tests/Elements/PCGAttributeReduceTest.cpp
--- a/Engine/Plugins/PCG/Source/PCG/Private/Tests/Elements/PCGAttributeReduceTest.cpp
+++ b/Engine/Plugins/PCG/Source/PCG/Private/Tests/Elements/PCGAttributeReduceTest.cpp
@@ -6,6 +6,7 @@
 #include "PCGComponent.h"
 #include "PCGParamData.h"
 #include "Data/PCGPointData.h"
+#include "Data/PCGSplineData.h"
 #include "Elements/PCGAttributeReduceElement.h"
 #include "Metadata/PCGMetadata.h"
 #include "Metadata/PCGMetadataAttributeTpl.h"
@@ -16,7 +17,7 @@ class FPCGAttributeReduceTests : public FPCGTestBaseClass
 
 protected:
 	template <typename... T>
-	bool RunAndValidate(TArray<FPCGTaggedData> InData, EPCGAttributeReduceOperation Operation, FPCGAttributePropertyInputSelector Selector, bool bMerged, T&& ...ExpectedValues)
+	bool RunAndValidate(TArray<FPCGTaggedData> InData, EPCGAttributeReduceOperation Operation, FPCGAttributePropertyInputSelector Selector, bool bMerged, bool bInDataDomain, T&& ...ExpectedValues)
 	{
 		AddInfo(*FString::Printf(TEXT("Test with '%s' as input attribute"), *Selector.GetDisplayText().ToString()));
 
@@ -34,6 +35,7 @@ protected:
 		TestData.InputData.TaggedData.Append(std::move(InData));
 		Settings->InputSource = std::move(Selector);
 		Settings->bMergeOutputAttributes = bMerged;
+		Settings->bWriteToDataDomain = bInDataDomain;
 		Settings->Operation = Operation;
 		Settings->OutputAttributeName = NAME_None;
 
@@ -54,27 +56,39 @@ protected:
 			UTEST_EQUAL("Not Merged output is the same number of inputs.", Outputs.Num(), NumExpectedValues);
 		}
 
-		// All outputs should be param data
-		TArray<const UPCGParamData*> AllParams;
-		AllParams.Reserve(Outputs.Num());
+		// All outputs should be param data or same data as input
+		TArray<const UPCGData*> AllData;
+		AllData.Reserve(Outputs.Num());
 		for (int32 i = 0; i < Outputs.Num(); ++i)
 		{
-			const UPCGParamData* Param = Cast<const UPCGParamData>(Outputs[i].Data);
-			UTEST_NOT_NULL(*FString::Printf(TEXT("Output %d is a param data"), i), Param);
-			AllParams.Add(Param);
+			if (bInDataDomain)
+			{
+				// i+1 for input data as the first one is the settings.
+				UTEST_TRUE(*FString::Printf(TEXT("Output %d is a the same type as input"), i), Outputs[i].Data && Outputs[i].Data->IsA(TestData.InputData.TaggedData[i+1].Data->StaticClass()));
+			}
+			else
+			{
+				UTEST_NOT_NULL(*FString::Printf(TEXT("Output %d is a param data"), i), Cast<UPCGParamData>(Outputs[i].Data));
+			}
+			
+			AllData.Add(Outputs[i].Data);
 		}
 
 		// Finally validate all values matches, use a fold expression with Index noting the current param to check or entry to check
 		int32 Index = 0;
 		bool bSuccess = true;
-		([&Index, this, &bSuccess, bMerged, &AllParams, &ExpectedValues]()
+
+		([&Index, this, &bSuccess, bMerged, &AllData, &ExpectedValues, bInDataDomain]()
 		{
 			if (!bSuccess)
 			{
 				return;
 			}
 
-			const FPCGMetadataAttribute<T>* Attribute = (bMerged ? AllParams[0] : AllParams[Index])->Metadata->GetConstTypedAttribute<T>(NAME_None);
+			const UPCGData* Data = bMerged ? AllData[0] : AllData[Index];
+			const FPCGMetadataDomain* MetadataDomain = bInDataDomain ? Data->ConstMetadata()->GetConstMetadataDomain(PCGMetadataDomainID::Data) : Data->ConstMetadata()->GetConstDefaultMetadataDomain();
+
+			const FPCGMetadataAttribute<T>* Attribute = MetadataDomain->GetConstTypedAttribute<T>(NAME_None);
 			bSuccess = TestNotNull(TEXT("Attribute 'None' exists and is the right type."), Attribute);
 
 			if (bSuccess)
@@ -155,6 +169,42 @@ protected:
 		return Inputs;
 	}
 
+	TArray<FPCGTaggedData> PrepareSplineAndSelectors(FPCGAttributePropertyInputSelector& PositionSelector, FPCGAttributePropertyInputSelector& IntSelector)
+	{
+		const FName IntName = TEXT("int");
+
+		UPCGSplineData* SplineData = NewObject<UPCGSplineData>();
+		FPCGMetadataAttribute<int32>* IntAttribute = SplineData->Metadata->CreateAttribute<int32>(IntName, 0, /*bAllowsInterpolation=*/ true, /*bOverrideParent=*/false);
+
+		constexpr int32 PointCount = 5;
+		TArray<FSplinePoint> SplinePoints;
+		TArray<PCGMetadataEntryKey> SplineEntryKeys;
+		
+		SplinePoints.Reserve(PointCount);
+		SplineEntryKeys.Reserve(PointCount);
+		
+		for (int32 i = 0; i < PointCount; ++i)
+		{
+			const FVector Position{ 3.0 * i, 3.0 * i + 1.0, 3.0 * i + 2 };
+			PCGMetadataEntryKey& EntryKey = SplineEntryKeys.Emplace_GetRef(SplineData->Metadata->AddEntry());
+
+			IntAttribute->SetValue(EntryKey, i);
+
+			SplinePoints.Emplace(i, Position);
+		}
+
+		SplineData->Initialize(SplinePoints, /*bInClosedLoop=*/false, FTransform::Identity, MoveTemp(SplineEntryKeys));
+
+		TArray<FPCGTaggedData> Inputs;
+		FPCGTaggedData& Data = Inputs.Emplace_GetRef();
+		Data.Data = SplineData;
+
+		IntSelector.SetAttributeName(IntName);
+		PositionSelector.SetPropertyName(TEXT("Position"));
+
+		return Inputs;
+	}
+
 	TArray<FPCGTaggedData> PrepareMixedAndSelector(FPCGAttributePropertyInputSelector& Selector)
 	{
 		const FName DoubleName = TEXT("double");
@@ -194,10 +244,18 @@ protected:
 IMPLEMENT_CUSTOM_SIMPLE_AUTOMATION_TEST(FPCGAttributeReducePoints_Average, FPCGAttributeReduceTests, "Plugins.PCG.AttributeReduce.Points.Average", PCGTestsCommon::TestFlags)
 IMPLEMENT_CUSTOM_SIMPLE_AUTOMATION_TEST(FPCGAttributeReducePoints_Min, FPCGAttributeReduceTests, "Plugins.PCG.AttributeReduce.Points.Min", PCGTestsCommon::TestFlags)
 IMPLEMENT_CUSTOM_SIMPLE_AUTOMATION_TEST(FPCGAttributeReducePoints_Max, FPCGAttributeReduceTests, "Plugins.PCG.AttributeReduce.Points.Max", PCGTestsCommon::TestFlags)
+IMPLEMENT_CUSTOM_SIMPLE_AUTOMATION_TEST(FPCGAttributeReducePoints_Sum, FPCGAttributeReduceTests, "Plugins.PCG.AttributeReduce.Points.Sum", PCGTestsCommon::TestFlags)
+IMPLEMENT_CUSTOM_SIMPLE_AUTOMATION_TEST(FPCGAttributeReducePoints_Join, FPCGAttributeReduceTests, "Plugins.PCG.AttributeReduce.Points.Join", PCGTestsCommon::TestFlags)
+IMPLEMENT_CUSTOM_SIMPLE_AUTOMATION_TEST(FPCGAttributeReducePoints_DataDomain, FPCGAttributeReduceTests, "Plugins.PCG.AttributeReduce.Points.DataDomain", PCGTestsCommon::TestFlags)
 
 IMPLEMENT_CUSTOM_SIMPLE_AUTOMATION_TEST(FPCGAttributeReduceParams_Average, FPCGAttributeReduceTests, "Plugins.PCG.AttributeReduce.Params.Average", PCGTestsCommon::TestFlags)
 IMPLEMENT_CUSTOM_SIMPLE_AUTOMATION_TEST(FPCGAttributeReduceParams_Min, FPCGAttributeReduceTests, "Plugins.PCG.AttributeReduce.Params.Min", PCGTestsCommon::TestFlags)
 IMPLEMENT_CUSTOM_SIMPLE_AUTOMATION_TEST(FPCGAttributeReduceParams_Max, FPCGAttributeReduceTests, "Plugins.PCG.AttributeReduce.Params.Max", PCGTestsCommon::TestFlags)
+IMPLEMENT_CUSTOM_SIMPLE_AUTOMATION_TEST(FPCGAttributeReduceParams_Sum, FPCGAttributeReduceTests, "Plugins.PCG.AttributeReduce.Params.Sum", PCGTestsCommon::TestFlags)
+IMPLEMENT_CUSTOM_SIMPLE_AUTOMATION_TEST(FPCGAttributeReduceParams_Join, FPCGAttributeReduceTests, "Plugins.PCG.AttributeReduce.Params.Join", PCGTestsCommon::TestFlags)
+IMPLEMENT_CUSTOM_SIMPLE_AUTOMATION_TEST(FPCGAttributeReduceParams_DataDomain, FPCGAttributeReduceTests, "Plugins.PCG.AttributeReduce.Params.DataDomain", PCGTestsCommon::TestFlags)
+
+IMPLEMENT_CUSTOM_SIMPLE_AUTOMATION_TEST(FPCGAttributeReduceSpline_DataDomain, FPCGAttributeReduceTests, "Plugins.PCG.AttributeReduce.Spline.DataDomain", PCGTestsCommon::TestFlags)
 
 IMPLEMENT_CUSTOM_SIMPLE_AUTOMATION_TEST(FPCGAttributeReduceMixed_NoMerge, FPCGAttributeReduceTests, "Plugins.PCG.AttributeReduce.Mixed.NoMerge", PCGTestsCommon::TestFlags)
 IMPLEMENT_CUSTOM_SIMPLE_AUTOMATION_TEST(FPCGAttributeReduceMixed_Merge, FPCGAttributeReduceTests, "Plugins.PCG.AttributeReduce.Mixed.Merge", PCGTestsCommon::TestFlags)
@@ -215,8 +273,8 @@ bool FPCGAttributeReducePoints_Average::RunTest(const FString& Parameters)
 
 	bool bSuccess = true;
 
-	bSuccess &= RunAndValidate(Inputs, EPCGAttributeReduceOperation::Average, std::move(InputSourceDensity), /*bMerged=*/ false, 0.3);
-	bSuccess &= RunAndValidate(Inputs, EPCGAttributeReduceOperation::Average, std::move(InputSourcePosition), /*bMerged=*/ false, FVector(6.0, 7.0, 8.0));
+	bSuccess &= RunAndValidate(Inputs, EPCGAttributeReduceOperation::Average, std::move(InputSourceDensity), /*bMerged=*/ false, /*bInDataDomain=*/ false, 0.3);
+	bSuccess &= RunAndValidate(Inputs, EPCGAttributeReduceOperation::Average, std::move(InputSourcePosition), /*bMerged=*/ false, /*bInDataDomain=*/ false, FVector(6.0, 7.0, 8.0));
 
 	return bSuccess;
 }
@@ -228,8 +286,8 @@ bool FPCGAttributeReducePoints_Min::RunTest(const FString& Parameters)
 
 	bool bSuccess = true;
 
-	bSuccess &= RunAndValidate(Inputs, EPCGAttributeReduceOperation::Min, std::move(InputSourceDensity), /*bMerged=*/ false, 0.1);
-	bSuccess &= RunAndValidate(Inputs, EPCGAttributeReduceOperation::Min, std::move(InputSourcePosition), /*bMerged=*/ false, FVector(0.0, 1.0, 2.0));
+	bSuccess &= RunAndValidate(Inputs, EPCGAttributeReduceOperation::Min, std::move(InputSourceDensity), /*bMerged=*/ false, /*bInDataDomain=*/ false, 0.1);
+	bSuccess &= RunAndValidate(Inputs, EPCGAttributeReduceOperation::Min, std::move(InputSourcePosition), /*bMerged=*/ false, /*bInDataDomain=*/ false, FVector(0.0, 1.0, 2.0));
 
 	return bSuccess;
 }
@@ -241,8 +299,46 @@ bool FPCGAttributeReducePoints_Max::RunTest(const FString& Parameters)
 
 	bool bSuccess = true;
 
-	bSuccess &= RunAndValidate(Inputs, EPCGAttributeReduceOperation::Max, std::move(InputSourceDensity), /*bMerged=*/ false, 0.5);
-	bSuccess &= RunAndValidate(Inputs, EPCGAttributeReduceOperation::Max, std::move(InputSourcePosition), /*bMerged=*/ false, FVector(12.0, 13.0, 14.0));
+	bSuccess &= RunAndValidate(Inputs, EPCGAttributeReduceOperation::Max, std::move(InputSourceDensity), /*bMerged=*/ false, /*bInDataDomain=*/ false, 0.5);
+	bSuccess &= RunAndValidate(Inputs, EPCGAttributeReduceOperation::Max, std::move(InputSourcePosition), /*bMerged=*/ false, /*bInDataDomain=*/ false, FVector(12.0, 13.0, 14.0));
+
+	return bSuccess;
+}
+
+bool FPCGAttributeReducePoints_Sum::RunTest(const FString& Parameters)
+{
+	FPCGAttributePropertyInputSelector InputSourceDensity, InputSourcePosition;
+	TArray<FPCGTaggedData> Inputs = PreparePointsAndSelectors(InputSourceDensity, InputSourcePosition);
+
+	bool bSuccess = true;
+
+	bSuccess &= RunAndValidate(Inputs, EPCGAttributeReduceOperation::Sum, std::move(InputSourceDensity), /*bMerged=*/ false, /*bInDataDomain=*/ false, 1.5);
+	bSuccess &= RunAndValidate(Inputs, EPCGAttributeReduceOperation::Sum, std::move(InputSourcePosition), /*bMerged=*/ false, /*bInDataDomain=*/ false, FVector(30.0, 35.0, 40.0));
+
+	return bSuccess;
+}
+
+bool FPCGAttributeReducePoints_Join::RunTest(const FString& Parameters)
+{
+	FPCGAttributePropertyInputSelector InputSourceDensity, InputSourcePosition;
+	TArray<FPCGTaggedData> Inputs = PreparePointsAndSelectors(InputSourceDensity, InputSourcePosition);
+
+	bool bSuccess = true;
+
+	bSuccess &= RunAndValidate(Inputs, EPCGAttributeReduceOperation::Join, std::move(InputSourceDensity), /*bMerged=*/ false, /*bInDataDomain=*/ false, FString(TEXT("0.100000, 0.200000, 0.300000, 0.400000, 0.500000")));
+
+	return bSuccess;
+}
+
+bool FPCGAttributeReducePoints_DataDomain::RunTest(const FString& Parameters)
+{
+	FPCGAttributePropertyInputSelector InputSourceDensity, InputSourcePosition;
+	TArray<FPCGTaggedData> Inputs = PreparePointsAndSelectors(InputSourceDensity, InputSourcePosition);
+
+	bool bSuccess = true;
+
+	bSuccess &= RunAndValidate(Inputs, EPCGAttributeReduceOperation::Average, std::move(InputSourceDensity), /*bMerged=*/ false, /*bInDataDomain=*/ true, 0.3);
+	bSuccess &= RunAndValidate(Inputs, EPCGAttributeReduceOperation::Average, std::move(InputSourcePosition), /*bMerged=*/ false, /*bInDataDomain=*/ true, FVector(6.0, 7.0, 8.0));
 
 	return bSuccess;
 }
@@ -254,8 +350,8 @@ bool FPCGAttributeReduceParams_Average::RunTest(const FString& Parameters)
 
 	bool bSuccess = true;
 
-	bSuccess &= RunAndValidate(Inputs, EPCGAttributeReduceOperation::Average, std::move(InputSourceInt), /*bMerged=*/ false, 2);
-	bSuccess &= RunAndValidate(Inputs, EPCGAttributeReduceOperation::Average, std::move(InputSourceVec2), /*bMerged=*/ false, FVector2D(4.0, 5.0));
+	bSuccess &= RunAndValidate(Inputs, EPCGAttributeReduceOperation::Average, std::move(InputSourceInt), /*bMerged=*/ false, /*bInDataDomain=*/ false, 2);
+	bSuccess &= RunAndValidate(Inputs, EPCGAttributeReduceOperation::Average, std::move(InputSourceVec2), /*bMerged=*/ false, /*bInDataDomain=*/ false, FVector2D(4.0, 5.0));
 
 	return bSuccess;
 }
@@ -267,8 +363,8 @@ bool FPCGAttributeReduceParams_Min::RunTest(const FString& Parameters)
 
 	bool bSuccess = true;
 
-	bSuccess &= RunAndValidate(Inputs, EPCGAttributeReduceOperation::Min, std::move(InputSourceInt), /*bMerged=*/ false, 0);
-	bSuccess &= RunAndValidate(Inputs, EPCGAttributeReduceOperation::Min, std::move(InputSourceVec2), /*bMerged=*/ false, FVector2D(0.0, 1.0));
+	bSuccess &= RunAndValidate(Inputs, EPCGAttributeReduceOperation::Min, std::move(InputSourceInt), /*bMerged=*/ false, /*bInDataDomain=*/ false, 0);
+	bSuccess &= RunAndValidate(Inputs, EPCGAttributeReduceOperation::Min, std::move(InputSourceVec2), /*bMerged=*/ false, /*bInDataDomain=*/ false, FVector2D(0.0, 1.0));
 
 	return bSuccess;
 }
@@ -280,8 +376,59 @@ bool FPCGAttributeReduceParams_Max::RunTest(const FString& Parameters)
 
 	bool bSuccess = true;
 
-	bSuccess &= RunAndValidate(Inputs, EPCGAttributeReduceOperation::Max, std::move(InputSourceInt), /*bMerged=*/ false, 4);
-	bSuccess &= RunAndValidate(Inputs, EPCGAttributeReduceOperation::Max, std::move(InputSourceVec2), /*bMerged=*/ false, FVector2D(8.0, 9.0));
+	bSuccess &= RunAndValidate(Inputs, EPCGAttributeReduceOperation::Max, std::move(InputSourceInt), /*bMerged=*/ false, /*bInDataDomain=*/ false, 4);
+	bSuccess &= RunAndValidate(Inputs, EPCGAttributeReduceOperation::Max, std::move(InputSourceVec2), /*bMerged=*/ false, /*bInDataDomain=*/ false, FVector2D(8.0, 9.0));
+
+	return bSuccess;
+}
+
+bool FPCGAttributeReduceParams_Sum::RunTest(const FString& Parameters)
+{
+	FPCGAttributePropertyInputSelector InputSourceInt, InputSourceVec2;
+	TArray<FPCGTaggedData> Inputs = PrepareParamAndSelectors(InputSourceInt, InputSourceVec2);
+
+	bool bSuccess = true;
+
+	bSuccess &= RunAndValidate(Inputs, EPCGAttributeReduceOperation::Sum, std::move(InputSourceInt), /*bMerged=*/ false, /*bInDataDomain=*/ false, 10);
+	bSuccess &= RunAndValidate(Inputs, EPCGAttributeReduceOperation::Sum, std::move(InputSourceVec2), /*bMerged=*/ false, /*bInDataDomain=*/ false, FVector2D(20.0, 25.0));
+
+	return bSuccess;
+}
+
+bool FPCGAttributeReduceParams_Join::RunTest(const FString& Parameters)
+{
+	FPCGAttributePropertyInputSelector InputSourceInt, InputSourceVec2;
+	TArray<FPCGTaggedData> Inputs = PrepareParamAndSelectors(InputSourceInt, InputSourceVec2);
+
+	bool bSuccess = true;
+
+	bSuccess &= RunAndValidate(Inputs, EPCGAttributeReduceOperation::Join, std::move(InputSourceInt), /*bMerged=*/ false, /*bInDataDomain=*/ false, FString(TEXT("0, 1, 2, 3, 4")));
+
+	return bSuccess;
+}
+
+bool FPCGAttributeReduceParams_DataDomain::RunTest(const FString& Parameters)
+{
+	FPCGAttributePropertyInputSelector InputSourceInt, InputSourceVec2;
+	TArray<FPCGTaggedData> Inputs = PrepareParamAndSelectors(InputSourceInt, InputSourceVec2);
+
+	bool bSuccess = true;
+
+	bSuccess &= RunAndValidate(Inputs, EPCGAttributeReduceOperation::Average, std::move(InputSourceInt), /*bMerged=*/ false, /*bInDataDomain=*/ true, 2);
+	bSuccess &= RunAndValidate(Inputs, EPCGAttributeReduceOperation::Average, std::move(InputSourceVec2), /*bMerged=*/ false, /*bInDataDomain=*/ true, FVector2D(4.0, 5.0));
+
+	return bSuccess;
+}
+
+bool FPCGAttributeReduceSpline_DataDomain::RunTest(const FString& Parameters)
+{
+	FPCGAttributePropertyInputSelector InputSourcePosition, InputSourceInt;
+	TArray<FPCGTaggedData> Inputs = PrepareSplineAndSelectors(InputSourcePosition, InputSourceInt);
+
+	bool bSuccess = true;
+
+	bSuccess &= RunAndValidate(Inputs, EPCGAttributeReduceOperation::Average, std::move(InputSourceInt), /*bMerged=*/ false, /*bInDataDomain=*/ true, 2);
+	bSuccess &= RunAndValidate(Inputs, EPCGAttributeReduceOperation::Average, std::move(InputSourcePosition), /*bMerged=*/ false, /*bInDataDomain=*/ true, FVector(6.0, 7.0, 8.0));
 
 	return bSuccess;
 }
@@ -291,7 +438,7 @@ bool FPCGAttributeReduceMixed_NoMerge::RunTest(const FString& Parameters)
 	FPCGAttributePropertyInputSelector InputSource;
 	TArray<FPCGTaggedData> Inputs = PrepareMixedAndSelector(InputSource);
 
-	return RunAndValidate(Inputs, EPCGAttributeReduceOperation::Average, std::move(InputSource), /*bMerged=*/ false, 0.3, 0.6);
+	return RunAndValidate(Inputs, EPCGAttributeReduceOperation::Average, std::move(InputSource), /*bMerged=*/ false, /*bInDataDomain=*/ false, 0.3, 0.6);
 }
 
 bool FPCGAttributeReduceMixed_Merge::RunTest(const FString& Parameters)
@@ -299,7 +446,7 @@ bool FPCGAttributeReduceMixed_Merge::RunTest(const FString& Parameters)
 	FPCGAttributePropertyInputSelector InputSource;
 	TArray<FPCGTaggedData> Inputs = PrepareMixedAndSelector(InputSource);
 
-	return RunAndValidate(Inputs, EPCGAttributeReduceOperation::Max, std::move(InputSource), /*bMerged=*/ true, 0.5, 1.0);
+	return RunAndValidate(Inputs, EPCGAttributeReduceOperation::Max, std::move(InputSource), /*bMerged=*/ true, /*bInDataDomain=*/ false, 0.5, 1.0);
 }
 
 bool FPCGAttributeReduceFirstInputInvalid_NoMerge::RunTest(const FString& Parameters)
@@ -313,7 +460,7 @@ bool FPCGAttributeReduceFirstInputInvalid_NoMerge::RunTest(const FString& Parame
 
 	AddExpectedError(TEXT("Invalid input data."));
 
-	return RunAndValidate(Inputs, EPCGAttributeReduceOperation::Min, std::move(InputSource), /*bMerged=*/ false, 0.1, 0.2);
+	return RunAndValidate(Inputs, EPCGAttributeReduceOperation::Min, std::move(InputSource), /*bMerged=*/ false, /*bInDataDomain=*/ false, 0.1, 0.2);
 }
 
 bool FPCGAttributeReduceFirstInputInvalid_Merge::RunTest(const FString& Parameters)
@@ -327,7 +474,7 @@ bool FPCGAttributeReduceFirstInputInvalid_Merge::RunTest(const FString& Paramete
 
 	AddExpectedError(TEXT("Invalid input data."));
 
-	return RunAndValidate(Inputs, EPCGAttributeReduceOperation::Min, std::move(InputSource), /*bMerged=*/ true, 0.1, 0.2);
+	return RunAndValidate(Inputs, EPCGAttributeReduceOperation::Min, std::move(InputSource), /*bMerged=*/ true, /*bInDataDomain=*/ false, 0.1, 0.2);
 }
 
 bool FPCGAttributeReduceIncompatibleType::RunTest(const FString& Parameters)
@@ -362,10 +509,10 @@ bool FPCGAttributeReduceIncompatibleType::RunTest(const FString& Parameters)
 	bool bSuccess = true;
 
 	AddExpectedError(TEXT("Operation was not compatible with the attribute type String or could not create attribute 'None' for input 0"));
-	bSuccess &= RunAndValidate(Inputs, EPCGAttributeReduceOperation::Average, std::move(InputSourceString), /*bMerged=*/ false);
+	bSuccess &= RunAndValidate(Inputs, EPCGAttributeReduceOperation::Average, std::move(InputSourceString), /*bMerged=*/ false, /*bInDataDomain=*/ false);
 
 	AddExpectedError(TEXT("Operation was not compatible with the attribute type Quaternion or could not create attribute 'None' for input 0"));
-	bSuccess &= RunAndValidate(Inputs, EPCGAttributeReduceOperation::Min, std::move(InputSourceQuat), /*bMerged=*/ false);
+	bSuccess &= RunAndValidate(Inputs, EPCGAttributeReduceOperation::Min, std::move(InputSourceQuat), /*bMerged=*/ false,  /*bInDataDomain=*/ false);
 
 	return bSuccess;
 }
\ No newline at end of file
diff --git a/Engine/Plugins/PCG/Source/PCG/Public/Elements/PCGAttributeReduceElement.h b/Engine/Plugins/PCG/Source/PCG/Public/Elements/PCGAttributeReduceElement.h
--- a/Engine/Plugins/PCG/Source/PCG/Public/Elements/PCGAttributeReduceElement.h
+++ b/Engine/Plugins/PCG/Source/PCG/Public/Elements/PCGAttributeReduceElement.h
@@ -44,11 +44,14 @@ public:
 	virtual TArray<FPCGPreConfiguredSettingsInfo> GetPreconfiguredInfo() const override;
 	virtual bool OnlyExposePreconfiguredSettings() const override { return true; }
 #endif
+	virtual bool HasDynamicPins() const override;
 
 	virtual void ApplyPreconfiguredSettings(const FPCGPreConfiguredSettingsInfo& PreconfiguredInfo) override;
 
 	virtual FString GetAdditionalTitleInformation() const override;
 
+	bool ShouldMergeOutputAttributes() const { return !bWriteToDataDomain && bMergeOutputAttributes; }
+
 protected:
 #if WITH_EDITOR
 	virtual EPCGChangeType GetChangeTypeForProperty(const FName& InPropertyName) const override { return Super::GetChangeTypeForProperty(InPropertyName) | EPCGChangeType::Cosmetic; }
@@ -61,6 +64,10 @@ public:
 	UPROPERTY(BlueprintReadWrite, EditAnywhere, Category = Settings, meta = (PCG_Overridable))
 	FPCGAttributePropertyInputSelector InputSource;
 
+	/** By default the reduce output the result into an new attribute set, but it can also be written into the data domain of the input. */
+	UPROPERTY(BlueprintReadWrite, EditAnywhere, Category = Settings)
+	bool bWriteToDataDomain = false;
+
 	UPROPERTY(BlueprintReadWrite, EditAnywhere, Category = Settings, meta = (PCG_Overridable))
 	FName OutputAttributeName = NAME_None;
 
@@ -71,7 +78,7 @@ public:
 	FString JoinDelimiter = FString(", ");
 
 	/** Option to merge all results into a single attribute set with multiple entries, instead of multiple attribute sets with a single value in them.*/
-	UPROPERTY(BlueprintReadWrite, EditAnywhere, Category = Settings, meta = (PCG_Overridable))
+	UPROPERTY(BlueprintReadWrite, EditAnywhere, Category = Settings, meta = (PCG_Overridable, EditCondition = "!bWriteToDataDomain", EditConditionHides))
 	bool bMergeOutputAttributes = false;
 
 #if WITH_EDITORONLY_DATA
