[PCG] New "Extract Attribute" node to extract any attribute from any data and domain

* Factorize code for Get Attribute From Point Index
* Allow Get Attribute From Point Index to work with data domains too.

#rnx

diff --git a/Engine/Plugins/PCG/Source/PCG/Private/Elements/Metadata/PCGExtractAttribute.cpp b/Engine/Plugins/PCG/Source/PCG/Private/Elements/Metadata/PCGExtractAttribute.cpp
new file mode 100644
--- /dev/null
+++ b/Engine/Plugins/PCG/Source/PCG/Private/Elements/Metadata/PCGExtractAttribute.cpp
@@ -0,0 +1,152 @@
+// Copyright Epic Games, Inc. All Rights Reserved.
+
+#include "Elements/Metadata/PCGExtractAttribute.h"
+
+#include "PCGContext.h"
+#include "PCGParamData.h"
+#include "Metadata/PCGMetadataAttributeTpl.h"
+#include "Metadata/Accessors/IPCGAttributeAccessor.h"
+#include "Metadata/Accessors/PCGAttributeAccessorHelpers.h"
+#include "Metadata/Accessors/PCGAttributeAccessorKeys.h"
+#include "Utils/PCGLogErrors.h"
+
+#define LOCTEXT_NAMESPACE "PCGExtractAttributeElement"
+
+void PCGExtractAttribute::ExtractAttribute(const PCGExtractAttribute::FExtractAttributeParams& Params)
+{
+	check(Params.Context && Params.InputSource && Params.OutputAttributeName)
+	TArray<FPCGTaggedData> Inputs = Params.Context->InputData.GetInputsByPin(Params.InputLabel);
+
+	for (int32 InputIndex = 0; InputIndex < Inputs.Num(); ++InputIndex)
+	{
+		const FPCGTaggedData& Input = Inputs[InputIndex];
+
+		if (!Input.Data || (Params.OptionalClassRequirement.IsSet() && *Params.OptionalClassRequirement && !Input.Data->IsA(*Params.OptionalClassRequirement)))
+		{
+			PCGLog::InputOutput::LogInvalidInputDataError(Params.Context);
+			continue;
+		}
+
+		TArray<FPCGTaggedData>& Outputs = Params.Context->OutputData.TaggedData;
+
+		const FPCGAttributePropertyInputSelector InputSource = Params.InputSource->CopyAndFixLast(Input.Data);
+		const FPCGAttributePropertyOutputSelector OutputTarget = Params.OutputAttributeName->CopyAndFixSource(&InputSource);
+
+		TUniquePtr<const IPCGAttributeAccessor> Accessor = PCGAttributeAccessorHelpers::CreateConstAccessor(Input.Data, InputSource);
+		TUniquePtr<const IPCGAttributeAccessorKeys> Keys = PCGAttributeAccessorHelpers::CreateConstKeys(Input.Data, InputSource);
+
+		if (!Accessor || !Keys)
+		{
+			PCGLog::Metadata::LogFailToCreateAccessorError(InputSource, Params.Context);
+			continue;
+		}
+
+		const int32 Index = Params.Index;
+
+		if (Index < 0 || Index >= Keys->GetNum())
+		{
+			PCGLog::LogErrorOnGraph(FText::Format(LOCTEXT("IndexOutOfBounds", "Index for input {0} is out of bounds. Index: {1}; Number of Elements: {2}"), InputIndex, Index, Keys->GetNum()), Params.Context);
+			continue;
+		}
+		
+		const FName OutputAttributeName = OutputTarget.GetName();
+		
+		UPCGParamData* OutputParamData = FPCGContext::NewObject_AnyThread<UPCGParamData>(Params.Context);
+
+		auto ExtractAttribute = [&Params, OutputAttributeName, &OutputParamData, &Accessor, &Index, &Keys]<typename AttributeType>(AttributeType) -> bool
+		{
+			AttributeType Value{};
+
+			// Should never fail, as OutputType == Accessor->UnderlyingType
+			if (!ensure(Accessor->Get<AttributeType>(Value, Index, *Keys)))
+			{
+				return false;
+			}
+
+			FPCGMetadataAttribute<AttributeType>* NewAttribute = static_cast<FPCGMetadataAttribute<AttributeType>*>(
+				OutputParamData->Metadata->CreateAttribute<AttributeType>(OutputAttributeName, Value, /*bAllowInterpolation=*/true, /*bOverrideParent=*/false));
+
+			if (!NewAttribute)
+			{
+				PCGLog::Metadata::LogFailToCreateAttributeError<AttributeType>(OutputAttributeName, Params.Context);
+				return false;
+			}
+
+			OutputParamData->Metadata->AddEntry();
+
+			return true;
+		};
+
+		if (PCGMetadataAttribute::CallbackWithRightType(Accessor->GetUnderlyingType(), ExtractAttribute))
+		{
+			FPCGTaggedData& Output = Outputs.Add_GetRef(Input);
+			Output.Data = OutputParamData;
+			Output.Pin = Params.OutputLabel;
+
+			if (Params.OnSuccessExtractionCallback.IsSet())
+			{
+				Params.OnSuccessExtractionCallback(Input);
+			}
+		}
+	}
+}
+
+#if WITH_EDITOR
+FName UPCGExtractAttributeSettings::GetDefaultNodeName() const
+{
+	return FName(TEXT("ExtractAttribute"));
+}
+
+FText UPCGExtractAttributeSettings::GetDefaultNodeTitle() const
+{
+	return LOCTEXT("NodeTitle", "Extract Attribute");
+}
+
+FText UPCGExtractAttributeSettings::GetNodeTooltipText() const
+{
+	return LOCTEXT("NodeTooltip", "Extract Attribute");
+}
+#endif // WITH_EDITOR
+
+FPCGElementPtr UPCGExtractAttributeSettings::CreateElement() const
+{
+	return MakeShared<FPCGExtractAttributeElement>();
+}
+
+TArray<FPCGPinProperties> UPCGExtractAttributeSettings::InputPinProperties() const
+{
+	TArray<FPCGPinProperties> Properties;
+	Properties.Emplace_GetRef(PCGPinConstants::DefaultInputLabel, EPCGDataType::Any).SetRequiredPin();	
+	return Properties;
+}
+
+TArray<FPCGPinProperties> UPCGExtractAttributeSettings::OutputPinProperties() const
+{ 
+	TArray<FPCGPinProperties> Properties;
+	Properties.Emplace(PCGPinConstants::DefaultOutputLabel, EPCGDataType::Param);	
+	return Properties;
+}
+
+bool FPCGExtractAttributeElement::ExecuteInternal(FPCGContext* Context) const
+{
+	TRACE_CPUPROFILER_EVENT_SCOPE(FPCGExtractAttributeElement::Execute);
+
+	check(Context);
+
+	const UPCGExtractAttributeSettings* Settings = Context->GetInputSettings<UPCGExtractAttributeSettings>();
+	check(Settings);
+
+	PCGExtractAttribute::FExtractAttributeParams Params =
+	{
+		.Context = Context,
+		.InputSource = &Settings->InputSource,
+		.Index = Settings->Index,
+		.OutputAttributeName = &Settings->OutputAttributeName,
+	};
+
+	PCGExtractAttribute::ExtractAttribute(Params);
+
+	return true;
+}
+
+#undef LOCTEXT_NAMESPACE
diff --git a/Engine/Plugins/PCG/Source/PCG/Private/Elements/PCGAttributeGetFromPointIndexElement.cpp b/Engine/Plugins/PCG/Source/PCG/Private/Elements/PCGAttributeGetFromPointIndexElement.cpp
--- a/Engine/Plugins/PCG/Source/PCG/Private/Elements/PCGAttributeGetFromPointIndexElement.cpp
+++ b/Engine/Plugins/PCG/Source/PCG/Private/Elements/PCGAttributeGetFromPointIndexElement.cpp
@@ -4,12 +4,11 @@
 
 #include "PCGContext.h"
 #include "PCGCustomVersion.h"
+#include "PCGData.h"
 #include "PCGPin.h"
-#include "PCGParamData.h"
-#include "Data/PCGPointData.h"
+#include "Data/PCGBasePointData.h"
+#include "Elements/Metadata/PCGExtractAttribute.h"
 #include "Helpers/PCGHelpers.h"
-#include "Metadata/Accessors/IPCGAttributeAccessor.h"
-#include "Metadata/Accessors/PCGAttributeAccessorKeys.h"
 #include "Metadata/Accessors/PCGAttributeAccessorHelpers.h"
 
 #include UE_INLINE_GENERATED_CPP_BY_NAME(PCGAttributeGetFromPointIndexElement)
@@ -99,105 +98,47 @@ bool FPCGAttributeGetFromPointIndexElement::ExecuteInternal(FPCGContext* Context
 	const UPCGAttributeGetFromPointIndexSettings* Settings = Context->GetInputSettings<UPCGAttributeGetFromPointIndexSettings>();
 	check(Settings);
 
-	TArray<FPCGTaggedData> Inputs = Context->InputData.GetInputsByPin(PCGPinConstants::DefaultInputLabel);
-
-	for (int32 InputIndex = 0; InputIndex < Inputs.Num(); ++InputIndex)
+	auto CreatePointData = [Context, Index = Settings->Index](const FPCGTaggedData& Input)
 	{
-		const FPCGTaggedData& Input = Inputs[InputIndex];
-		const UPCGBasePointData* PointData = Cast<UPCGBasePointData>(Input.Data);
-
-		if (!PointData)
-		{
-			PCGE_LOG(Error, GraphAndLog, FText::Format(LOCTEXT("InputNotPointData", "Input {0} is not a point data"), InputIndex));
-			continue;
-		}
-
-		const int32 Index = Settings->Index;
-
-		if (Index < 0 || Index >= PointData->GetNumPoints())
-		{
-			PCGE_LOG(Error, GraphAndLog, FText::Format(LOCTEXT("IndexOutOfBounds", "Index for input {0} is out of bounds. Index: {1}; Number of Points: {2}"), InputIndex, Index, PointData->GetNumPoints()));
-			continue;
-		}
-
-		TArray<FPCGTaggedData>& Outputs = Context->OutputData.TaggedData;
-
-		const FPCGAttributePropertyInputSelector InputSource = Settings->InputSource.CopyAndFixLast(PointData);
-		const FPCGAttributePropertyOutputSelector OutputTarget = Settings->OutputAttributeName.CopyAndFixSource(&InputSource);
-		
-		const FName OutputAttributeName = OutputTarget.GetName();
-
-		TUniquePtr<const IPCGAttributeAccessor> Accessor = PCGAttributeAccessorHelpers::CreateConstAccessor(PointData, InputSource);
-		FPCGAttributeAccessorKeysPointsSubset PointKey(PointData, { Index });
-
-		if (Accessor.IsValid())
-		{
-			UPCGParamData* OutputParamData = FPCGContext::NewObject_AnyThread<UPCGParamData>(Context);
-
-			auto ExtractAttribute = [this, Context, OutputAttributeName, &OutputParamData, &Accessor, &PointKey, InputIndex](auto DummyValue) -> bool
-			{
-				using AttributeType = decltype(DummyValue);
-
-				AttributeType Value{};
-
-				// Should never fail, as OutputType == Accessor->UnderlyingType
-				if (!ensure(Accessor->Get<AttributeType>(Value, PointKey)))
-				{
-					return false;
-				}
-
-				FPCGMetadataAttribute<AttributeType>* NewAttribute = static_cast<FPCGMetadataAttribute<AttributeType>*>(
-					OutputParamData->Metadata->CreateAttribute<AttributeType>(OutputAttributeName, Value, /*bAllowInterpolation=*/true, /*bOverrideParent=*/false));
-
-				if (!NewAttribute)
-				{
-					PCGE_LOG(Error, GraphAndLog, FText::Format(LOCTEXT("ErrorCreatingTargetAttribute", "Error while creating target attribute '{0}' for output {1}"), FText::FromName(OutputAttributeName), InputIndex));
-					return false;
-				}
-
-				OutputParamData->Metadata->AddEntry();
-
-				return true;
-			};
-
-			if (PCGMetadataAttribute::CallbackWithRightType(Accessor->GetUnderlyingType(), ExtractAttribute))
-			{
 #if !WITH_EDITOR
-				// Add the point
-				// Eschew output creation only in non-editor builds
-				if (Context->Node && Context->Node->IsOutputPinConnected(PCGAttributeGetFromPointIndexConstants::OutputPointLabel))
+		// Add the point
+		// Eschew output creation only in non-editor builds
+		if (Context->Node && Context->Node->IsOutputPinConnected(PCGAttributeGetFromPointIndexConstants::OutputPointLabel))
 #endif
-				{
-					UPCGBasePointData* OutputPointData = FPCGContext::NewPointData_AnyThread(Context);
-
-					FPCGInitializeFromDataParams InitializeFromDataParams(PointData);
-					InitializeFromDataParams.bInheritSpatialData = false;
-					OutputPointData->InitializeFromDataWithParams(InitializeFromDataParams);
+		{
+			const UPCGBasePointData* PointData = CastChecked<UPCGBasePointData>(Input.Data);
+			UPCGBasePointData* OutputPointData = FPCGContext::NewPointData_AnyThread(Context);
 
-					OutputPointData->SetNumPoints(1);
-					OutputPointData->CopyUnallocatedPropertiesFrom(PointData);
+			FPCGInitializeFromDataParams InitializeFromDataParams(PointData);
+			InitializeFromDataParams.bInheritSpatialData = false;
+			OutputPointData->InitializeFromDataWithParams(InitializeFromDataParams);
 
-					const FConstPCGPointValueRanges InRanges(PointData);
-					FPCGPointValueRanges OutRanges(OutputPointData, /*bAllocate=*/false);
+			OutputPointData->SetNumPoints(1);
+			OutputPointData->CopyUnallocatedPropertiesFrom(PointData);
 
-					OutRanges.SetFromValueRanges(0, InRanges, Index);
+			const FConstPCGPointValueRanges InRanges(PointData);
+			FPCGPointValueRanges OutRanges(OutputPointData, /*bAllocate=*/false);
 
-					FPCGTaggedData& OutputPoint = Outputs.Add_GetRef(Input);
-					OutputPoint.Data = OutputPointData;
-					OutputPoint.Pin = PCGAttributeGetFromPointIndexConstants::OutputPointLabel;
-				}
+			OutRanges.SetFromValueRanges(0, InRanges, Index);
 
-				// And the attribute
-				FPCGTaggedData& OutputAttribute = Outputs.Add_GetRef(Input);
-				OutputAttribute.Data = OutputParamData;
-				OutputAttribute.Pin = PCGAttributeGetFromPointIndexConstants::OutputAttributeLabel;
-			}
-		}
-		else
-		{
-			PCGE_LOG(Warning, GraphAndLog, FText::Format(LOCTEXT("AttributeNotFound", "Cannot find attribute/property '{0}' in input {1}"), FText::FromName(InputSource.GetName()), InputIndex));
+			FPCGTaggedData& OutputPoint = Context->OutputData.TaggedData.Add_GetRef(Input);
+			OutputPoint.Data = OutputPointData;
+			OutputPoint.Pin = PCGAttributeGetFromPointIndexConstants::OutputPointLabel;
 		}
-	}
+	};
+
+	PCGExtractAttribute::FExtractAttributeParams Params =
+	{
+		.Context = Context,
+		.InputSource = &Settings->InputSource,
+		.Index = Settings->Index,
+		.OutputAttributeName = &Settings->OutputAttributeName,
+		.OptionalClassRequirement = TSubclassOf<UPCGData>(UPCGBasePointData::StaticClass()),
+		.OutputLabel =  PCGAttributeGetFromPointIndexConstants::OutputAttributeLabel,
+		.OnSuccessExtractionCallback = CreatePointData
+	};
+
+	PCGExtractAttribute::ExtractAttribute(Params);
 
 	return true;
 }
diff --git a/Engine/Plugins/PCG/Source/PCG/Public/Elements/Metadata/PCGExtractAttribute.h b/Engine/Plugins/PCG/Source/PCG/Public/Elements/Metadata/PCGExtractAttribute.h
new file mode 100644
--- /dev/null
+++ b/Engine/Plugins/PCG/Source/PCG/Public/Elements/Metadata/PCGExtractAttribute.h
@@ -0,0 +1,97 @@
+// Copyright Epic Games, Inc. All Rights Reserved.
+
+#pragma once
+
+#include "PCGSettings.h"
+
+#include "Misc/Optional.h"
+#include "Templates/Function.h"
+#include "Templates/SubclassOf.h"
+
+#include "PCGExtractAttribute.generated.h"
+
+struct FPCGAttributePropertyInputSelector;
+struct FPCGAttributePropertyOutputSelector;
+struct FPCGContext;
+struct FPCGTaggedData;
+class UPCGData;
+
+namespace PCGExtractAttribute
+{
+	struct FExtractAttributeParams
+	{
+		/** Context that contains the input data. Cannot be null when passed to ExtractAttribute. */
+		FPCGContext* Context = nullptr;
+		
+		/** Selector for the attribute to extract. Cannot be null when passed to ExtractAttribute. */
+		const FPCGAttributePropertyInputSelector* InputSource = nullptr;
+		
+		/** Index of the element to extract the attribute from. */
+		const int32 Index = 0;
+		
+		/** Selector for the output attribute to write into. Cannot be null when passed to ExtractAttribute. */
+		const FPCGAttributePropertyOutputSelector* OutputAttributeName = nullptr;
+		
+		/** Optional class to check the input data. */
+		TOptional<TSubclassOf<UPCGData>> OptionalClassRequirement = {};
+		
+		/** Input pin label to read the data from. */
+		FName InputLabel = PCGPinConstants::DefaultInputLabel;
+		
+		/** Output pin label to output the result attribute set to. */
+		FName OutputLabel = PCGPinConstants::DefaultOutputLabel;
+		
+		/** Extra function that can be called after the extraction succeeded with the input data. */
+		TFunction<void(const FPCGTaggedData&)> OnSuccessExtractionCallback;
+	};
+
+	/**
+	 * Go through all inputs from the context input data, and try to extract an attribute at the provided index into a new Attribute Set.
+	 * Input are read from the context and output is written in the context too.
+	 */
+	void ExtractAttribute(const FExtractAttributeParams& Params);
+}
+
+/**
+ * Extract an attribute at a given index into a new attribute set.
+ * Support any domain. Index needs to be in range of valid indexes for the given domain.
+ */
+UCLASS(MinimalAPI, BlueprintType, ClassGroup = (Procedural))
+class UPCGExtractAttributeSettings : public UPCGSettings
+{
+	GENERATED_BODY()
+
+public:
+	//~Begin UPCGSettings interface
+#if WITH_EDITOR
+	virtual FName GetDefaultNodeName() const override;
+	virtual FText GetDefaultNodeTitle() const override;
+    virtual FText GetNodeTooltipText() const override;
+	virtual EPCGSettingsType GetType() const override { return EPCGSettingsType::Metadata; }
+#endif
+
+protected:
+	virtual TArray<FPCGPinProperties> InputPinProperties() const override;
+	virtual TArray<FPCGPinProperties> OutputPinProperties() const override;
+	virtual FPCGElementPtr CreateElement() const override;
+	//~End UPCGSettings interface
+
+public:
+	UPROPERTY(BlueprintReadWrite, EditAnywhere, Category = Settings, meta = (PCG_Overridable))
+	FPCGAttributePropertyInputSelector InputSource;
+
+	UPROPERTY(BlueprintReadWrite, EditAnywhere, Category = Settings, meta = (PCG_Overridable))
+	int32 Index = 0;
+
+	UPROPERTY(BlueprintReadWrite, EditAnywhere, Category = Settings, meta = (PCG_Overridable, PCG_DiscardPropertySelection, PCG_DiscardExtraSelection))
+	FPCGAttributePropertyOutputSelector OutputAttributeName;
+};
+
+class FPCGExtractAttributeElement : public IPCGElement
+{
+protected:
+	virtual bool ExecuteInternal(FPCGContext* InContext) const override;
+    virtual bool SupportsBasePointDataInputs(FPCGContext* InContext) const override { return true; }
+    virtual EPCGElementExecutionLoopMode ExecutionLoopMode(const UPCGSettings* Settings) const override { return EPCGElementExecutionLoopMode::SinglePrimaryPin; }
+};
+
