Added additional polygon intersection functions; clipping polygons from path(s) and clipping path(s) from polygon(s).

diff --git a/Engine/Plugins/PCG/Source/PCG/Private/Elements/PCGSplitSplines.cpp b/Engine/Plugins/PCG/Source/PCG/Private/Elements/PCGSplitSplines.cpp
--- a/Engine/Plugins/PCG/Source/PCG/Private/Elements/PCGSplitSplines.cpp
+++ b/Engine/Plugins/PCG/Source/PCG/Private/Elements/PCGSplitSplines.cpp
@@ -18,6 +18,151 @@ namespace PCGSplitSplineConstants
 	const FName SplitInfoLabel = TEXT("SplitInfo");
 }
 
+namespace PCGSplitSpline
+{
+	void SplitSpline(const UPCGSplineData* Spline, UPCGSplineData* OutSplitSpline, TConstArrayView<FSplinePoint> SplinePoints, TConstArrayView<PCGMetadataEntryKey> SplinePointsEntryKeys, double StartKey, double EndKey)
+	{
+		check(Spline);
+		check(OutSplitSpline);
+		check(StartKey < EndKey || !Spline || Spline->IsClosed());
+
+		TArray<FSplinePoint> SplitSplinePoints;
+		TArray<int64> SplitSplinePointsEntryKeys;
+
+		auto AddExistingControlPoint = [&SplitSplinePoints, &SplitSplinePointsEntryKeys, &SplinePoints, &SplinePointsEntryKeys](const int ControlPointIndex)
+		{
+			SplitSplinePoints.Add(SplinePoints[ControlPointIndex]);
+			// Add metadata if needed, e.g. if the original data had entry keys.
+			if (!SplinePointsEntryKeys.IsEmpty())
+			{
+				SplitSplinePointsEntryKeys.Add(SplinePointsEntryKeys[ControlPointIndex]);
+			}
+		};
+
+		auto AddSplitPoint = [&OutSplitSpline, &SplitSplinePoints, &SplitSplinePointsEntryKeys, &Spline, &SplinePointsEntryKeys](const double Key)
+		{
+			FSplinePoint SplitPoint;
+			SplitPoint.InputKey = Key;
+			SplitPoint.Position = Spline->SplineStruct.GetLocationAtSplineInputKey(Key, ESplineCoordinateSpace::Local);
+			SplitPoint.ArriveTangent = Spline->SplineStruct.GetTangentAtSplineInputKey(Key, ESplineCoordinateSpace::Local);
+			SplitPoint.LeaveTangent = SplitPoint.ArriveTangent;
+			SplitPoint.Rotation = Spline->SplineStruct.GetQuaternionAtSplineInputKey(Key, ESplineCoordinateSpace::Local).Rotator();
+			SplitPoint.Scale = Spline->SplineStruct.GetScaleAtSplineInputKey(Key);
+			SplitPoint.Type = ESplinePointType::CurveCustomTangent;
+
+			// The tangent is always computed from the original spline's perspective so we will assume the key range is always one.
+			SplitPoint.ArriveTangent *= FMath::Frac(Key);
+			SplitPoint.LeaveTangent *= 1.0 - FMath::Frac(Key);
+
+			SplitSplinePoints.Add(SplitPoint);
+
+			// Add metadata if needed, e.g. if the original data had entry keys.
+			if (!SplinePointsEntryKeys.IsEmpty())
+			{
+				PCGMetadataEntryKey& EntryKey = SplitSplinePointsEntryKeys.Add_GetRef(PCGInvalidEntryKey);
+				Spline->WriteMetadataToEntry(Key, EntryKey, OutSplitSpline->MutableMetadata());
+			}
+		};
+
+		int FirstControlPointInSegment = FMath::CeilToInt(StartKey);
+		int LastControlPointInSegment = FMath::FloorToInt(EndKey);
+
+		if (Spline->IsClosed() && EndKey <= StartKey) // allow for loop
+		{
+			LastControlPointInSegment += SplinePoints.Num();
+		}
+
+		bool bNeedsStartTangentAdjustment = false;
+		bool bNeedsEndTangentAdjustment = false;
+
+		// Add left split point, prioritize existing control points when possible
+		const double KeyEpsilon = 1.0e-6f;
+		const double FracStartKey = FMath::Frac(StartKey);
+		if (FracStartKey < KeyEpsilon)
+		{
+			if (FracStartKey > 0) // otherwise point is already added in the control point loop
+			{
+				AddExistingControlPoint(FMath::FloorToInt(StartKey));
+			}
+		}
+		else if (1.0 - FracStartKey < KeyEpsilon)
+		{
+			// will be added by the control point loop
+		}
+		else
+		{
+			AddSplitPoint(StartKey);
+			bNeedsStartTangentAdjustment = true;
+		}
+
+		// Add control points that are taken as-is from the original spline points
+		for (int ControlPointIndex = FirstControlPointInSegment; ControlPointIndex <= LastControlPointInSegment; ++ControlPointIndex)
+		{
+			AddExistingControlPoint((ControlPointIndex % SplinePoints.Num()));
+		}
+
+		// Add right split point
+		const double FracEndKey = FMath::Frac(EndKey);
+		if (FracEndKey < KeyEpsilon)
+		{
+			// will be already added by the control point loop
+		}
+		else if (1.0 - FracEndKey < KeyEpsilon)
+		{
+			AddExistingControlPoint(FMath::CeilToInt(EndKey));
+		}
+		else
+		{
+			AddSplitPoint(EndKey);
+			bNeedsEndTangentAdjustment = true;
+		}
+
+		// Adjust tangents if required
+		if (bNeedsStartTangentAdjustment)
+		{
+			check(SplitSplinePoints.Num() >= 1);
+			// assumes key range is 1.
+			const double Ratio = (1.0 - FMath::Frac(SplitSplinePoints[0].InputKey));
+			SplitSplinePoints[1].ArriveTangent *= Ratio;
+		}
+
+		if (bNeedsEndTangentAdjustment && SplitSplinePoints.Num() >= 2)
+		{
+			// assumes key range is 1.
+			const double Ratio = FMath::Frac(SplitSplinePoints.Last().InputKey);
+			SplitSplinePoints[SplitSplinePoints.Num() - 2].LeaveTangent *= Ratio;
+		}
+
+		// Finally, mark all tangents as custom and reset the keys.
+		for (int SPIndex = 0; SPIndex < SplitSplinePoints.Num(); ++SPIndex)
+		{
+			FSplinePoint& SplinePoint = SplitSplinePoints[SPIndex];
+			SplinePoint.InputKey = static_cast<float>(SPIndex);
+			if (SplinePoint.Type == ESplinePointType::Curve || SplinePoint.Type == ESplinePointType::CurveClamped)
+			{
+				SplitSplinePoints[SPIndex].Type = ESplinePointType::CurveCustomTangent;
+			}
+		}
+
+		// Initialize data now
+		OutSplitSpline->Initialize(SplitSplinePoints,
+			/*bIsClosedLoop*/false,
+			Spline->SplineStruct.GetTransform(),
+			SplitSplinePointsEntryKeys);
+	}
+
+	void SplitSpline(const UPCGSplineData* Spline, UPCGSplineData* OutSplitSpline, double StartKey, double EndKey)
+	{
+		check(Spline);
+		check(OutSplitSpline);
+
+		TArray<FSplinePoint> SplinePoints = Spline->GetSplinePoints();
+		TConstArrayView<PCGMetadataEntryKey> SplinePointsEntryKeys = Spline->GetConstVerticesEntryKeys();
+
+		SplitSpline(Spline, OutSplitSpline, SplinePoints, SplinePointsEntryKeys, StartKey, EndKey);
+	}
+}
+
 UPCGSplitSplinesSettings::UPCGSplitSplinesSettings()
 {
 	OutputOriginatingSplineIndex = FPCGAttributePropertySelector::CreateAttributeSelector<FPCGAttributePropertyOutputSelector>(TEXT("OriginatingSplineIndex"), TEXT("Data"));
@@ -112,7 +257,7 @@ bool FPCGSplitSplineElement::ExecuteInternal(FPCGContext* Context) const
 		}
 
 		TArray<FSplinePoint> SplinePoints = SplineData->GetSplinePoints();
-		TArray<int64> SplinePointsEntryKeys = SplineData->GetMetadataEntryKeysForSplinePoints();
+		TConstArrayView<PCGMetadataEntryKey> SplinePointsEntryKeys = SplineData->GetConstVerticesEntryKeys();
 		check(SplinePoints.Num() == SplinePointsEntryKeys.Num() || SplinePointsEntryKeys.IsEmpty());
 
 		// Build keys list
@@ -252,126 +397,7 @@ bool FPCGSplitSplineElement::ExecuteInternal(FPCGContext* Context) const
 			UPCGSplineData* NewSplineData = FPCGContext::NewObject_AnyThread<UPCGSplineData>(Context);
 			NewSplineData->InitializeFromData(SplineData);
 
-			TArray<FSplinePoint> NewSplinePoints;
-			TArray<int64> NewSplinePointsEntryKeys;
-
-			auto AddExistingControlPoint = [&NewSplinePoints, &NewSplinePointsEntryKeys, &SplinePoints, &SplinePointsEntryKeys](const int ControlPointIndex)
-			{
-				NewSplinePoints.Add(SplinePoints[ControlPointIndex]);
-				if (!SplinePointsEntryKeys.IsEmpty())
-				{
-					NewSplinePointsEntryKeys.Add(SplinePointsEntryKeys[ControlPointIndex]);
-				}
-			};
-
-			auto AddSplitPoint = [&SplineData, &NewSplineData, &NewSplinePoints, &NewSplinePointsEntryKeys, &SplinePoints, &SplinePointsEntryKeys](const double Key)
-			{
-					FSplinePoint SplitPoint;
-					SplitPoint.InputKey = Key;
-					SplitPoint.Position = SplineData->SplineStruct.GetLocationAtSplineInputKey(Key, ESplineCoordinateSpace::Local);
-					SplitPoint.ArriveTangent = SplineData->SplineStruct.GetTangentAtSplineInputKey(Key, ESplineCoordinateSpace::Local);
-					SplitPoint.LeaveTangent = SplitPoint.ArriveTangent;
-					SplitPoint.Rotation = SplineData->SplineStruct.GetQuaternionAtSplineInputKey(Key, ESplineCoordinateSpace::Local).Rotator();
-					SplitPoint.Scale = SplineData->SplineStruct.GetScaleAtSplineInputKey(Key);
-					SplitPoint.Type = ESplinePointType::CurveCustomTangent;
-
-					// The tangent is always computed from the original spline's perspective so we will assume the key range is always one.
-					SplitPoint.ArriveTangent *= FMath::Frac(Key);
-					SplitPoint.LeaveTangent *= 1.0 - FMath::Frac(Key);
-
-					NewSplinePoints.Add(SplitPoint);
-
-					// Add metadata if needed
-					if (!SplinePointsEntryKeys.IsEmpty())
-					{
-						PCGMetadataEntryKey EntryKey = PCGInvalidEntryKey;
-						SplineData->WriteMetadataToEntry(Key, EntryKey, NewSplineData->MutableMetadata());
-					}
-			};
-
-			int FirstControlPointInSegment = FMath::CeilToInt(StartKey);
-			int LastControlPointInSegment = FMath::FloorToInt(EndKey);
-
-			if (bIsClosed && EndKey <= StartKey) // allow for loop
-			{
-				LastControlPointInSegment += SplinePoints.Num();
-			}
-
-			bool bNeedsStartTangentAdjustment = false;
-			bool bNeedsEndTangentAdjustment = false;
-
-			// Add left split point, prioritize existing control points when possible
-			const double FracStartKey = FMath::Frac(StartKey);
-			if (FracStartKey < KeyEpsilon)
-			{
-				if (FracStartKey > 0) // otherwise point is already added in the control point loop
-				{
-					AddExistingControlPoint(FMath::FloorToInt(StartKey));
-				}
-			}
-			else if (1.0 - FracStartKey < KeyEpsilon)
-			{
-				// will be added by the control point loop
-			}
-			else
-			{
-				AddSplitPoint(StartKey);
-				bNeedsStartTangentAdjustment = true;
-			}
-			
-			// Add control points that are taken as-is from the original spline points
-			for (int ControlPointIndex = FirstControlPointInSegment; ControlPointIndex <= LastControlPointInSegment; ++ControlPointIndex)
-			{
-				AddExistingControlPoint((ControlPointIndex % SplinePoints.Num()));
-			}
-
-			// Add right split point
-			const double FracEndKey = FMath::Frac(EndKey);
-			if (FracEndKey < KeyEpsilon)
-			{
-				// will be already added by the control point loop
-			}
-			else if (1.0 - FracEndKey < KeyEpsilon)
-			{
-				AddExistingControlPoint(FMath::CeilToInt(EndKey));
-			}
-			else
-			{
-				AddSplitPoint(EndKey);
-				bNeedsEndTangentAdjustment = true;
-			}
-
-			// Adjust tangents if required
-			if (bNeedsStartTangentAdjustment)
-			{
-				// assumes key range is 1.
-				const double Ratio = (1.0 - FMath::Frac(NewSplinePoints[0].InputKey));
-				NewSplinePoints[1].ArriveTangent *= Ratio;
-			}
-
-			if (bNeedsEndTangentAdjustment)
-			{
-				// assumes key range is 1.
-				const double Ratio = FMath::Frac(NewSplinePoints.Last().InputKey);
-				NewSplinePoints[NewSplinePoints.Num() - 2].LeaveTangent *= Ratio;
-			}
-
-			// Finally, mark all tangents as custom and reset the keys.
-			for (int SPIndex = 0; SPIndex < NewSplinePoints.Num(); ++SPIndex)
-			{
-				FSplinePoint& SplinePoint = NewSplinePoints[SPIndex];
-				SplinePoint.InputKey = static_cast<float>(SPIndex);
-				if (SplinePoint.Type == ESplinePointType::Curve || SplinePoint.Type == ESplinePointType::CurveClamped)
-				{
-					NewSplinePoints[SPIndex].Type = ESplinePointType::CurveCustomTangent;
-				}
-			}
-
-			// Initialize data now
-			NewSplineData->Initialize(NewSplinePoints,
-				/*bIsClosedLoop*/false,
-				SplineData->SplineStruct.GetTransform(),
-				NewSplinePointsEntryKeys);
+			PCGSplitSpline::SplitSpline(SplineData, NewSplineData, SplinePoints, SplinePointsEntryKeys, StartKey, EndKey);
 
 			// Write originating spline index if required
 			if (Settings->bShouldOutputOriginatingSplineIndex)
diff --git a/Engine/Plugins/PCG/Source/PCG/Private/Elements/Polygon/PCGClipPaths.cpp b/Engine/Plugins/PCG/Source/PCG/Private/Elements/Polygon/PCGClipPaths.cpp
new file mode 100644
--- /dev/null
+++ b/Engine/Plugins/PCG/Source/PCG/Private/Elements/Polygon/PCGClipPaths.cpp
@@ -0,0 +1,206 @@
+// Copyright Epic Games, Inc. All Rights Reserved.
+
+#include "Elements/Polygon/PCGClipPaths.h"
+
+#include "PCGContext.h"
+#include "Data/PCGBasePointData.h"
+#include "Data/PCGSplineData.h"
+#include "Data/PCGPolygon2DData.h"
+#include "Elements/PCGSplitSplines.h"
+#include "Elements/Polygon/PCGPolygon2DUtils.h"
+#include "Helpers/PCGPolygon2DProcessingHelpers.h"
+
+#include "Curve/PolygonIntersectionUtils.h"
+#include "UObject/Package.h"
+
+#define LOCTEXT_NAMESPACE "PCGClipPathsElement"
+
+#if WITH_EDITOR
+FText UPCGClipPathsSettings::GetDefaultNodeTitle() const
+{
+	return LOCTEXT("NodeTitle", "Clip Paths");
+}
+
+FText UPCGClipPathsSettings::GetNodeTooltipText() const
+{
+	return LOCTEXT("NodeTooltip", "Clips paths against the polygons, to have the paths either inscribed in the polygons (intersection) or outside the polygons (difference).");
+}
+#endif // WITH_EDITOR
+
+FString UPCGClipPathsSettings::GetAdditionalTitleInformation() const
+{
+	if (const UEnum* EnumPtr = StaticEnum<EPCGClipPathOperation>())
+	{
+		return EnumPtr->GetDisplayNameTextByValue(static_cast<int64>(Operation)).ToString();
+	}
+	else
+	{
+		return FString();
+	}
+}
+
+TArray<FPCGPinProperties> UPCGClipPathsSettings::InputPinProperties() const
+{
+	TArray<FPCGPinProperties> PinProperties;
+	// @todo_pcg: add landscape spline support?
+	PinProperties.Emplace_GetRef(PCGPinConstants::DefaultInputLabel, EPCGDataType::Point | EPCGDataType::Spline).SetRequiredPin();
+	PinProperties.Emplace(PCGPolygon2DUtils::Constants::ClipPolysLabel, EPCGDataType::Polygon2D);
+
+	return PinProperties;
+}
+
+TArray<FPCGPinProperties> UPCGClipPathsSettings::OutputPinProperties() const
+{
+	TArray<FPCGPinProperties> PinProperties;
+	// @todo_pcg: add landscape spline support?
+	PinProperties.Emplace(PCGPinConstants::DefaultOutputLabel, EPCGDataType::Point | EPCGDataType::Spline);
+
+	return PinProperties;
+}
+
+bool FPCGClipPathsElement::ExecuteInternal(FPCGContext* Context) const
+{
+	TRACE_CPUPROFILER_EVENT_SCOPE(FPCGCLipPathsElement::Execute);
+	const UPCGClipPathsSettings* Settings = Context->GetInputSettings<UPCGClipPathsSettings>();
+	check(Settings);
+
+	TArray<FPCGTaggedData> InputPaths = Context->InputData.GetInputsByPin(PCGPinConstants::DefaultInputLabel);
+	TArray<FPCGTaggedData> ClipInputs = Context->InputData.GetInputsByPin(PCGPolygon2DUtils::Constants::ClipPolysLabel);
+
+	// Gather polygons for clipping. Note that we will not gather entry keys here as we will infer metadata only from the input paths.
+	bool bHasTransform = false;
+	FTransform PolyTransform;
+	TArray<UE::Geometry::FGeneralPolygon2d> ClipPolys;
+	ClipPolys.Reserve(ClipInputs.Num());
+
+	for (const FPCGTaggedData& ClipInput : ClipInputs)
+	{
+		if (const UPCGPolygon2DData* ClipPolyData = Cast<UPCGPolygon2DData>(ClipInput.Data))
+		{
+			PCGPolygon2DProcessingHelpers::AddPolygon(ClipPolyData, ClipPolys, bHasTransform, PolyTransform);
+		}
+	}
+
+	// Implementation note - there's an assumption that the input points/splines are on the same plane as the polygons.
+	// However, that's not guaranteed in anyway, and if we don't have polygons to clip against, then we will return the inputs as is.
+	if (!bHasTransform)
+	{
+		Context->OutputData.TaggedData = InputPaths;
+		return true;
+	}
+
+	TArray<FVector> Positions;
+	TArray<FVector2D> Positions2D;
+
+	for (const FPCGTaggedData& Input : InputPaths)
+	{
+		// @todo_pcg: we probably want to accept closed splines as paths...
+		if (!PCGPolygon2DProcessingHelpers::GetPath(Input.Data, PolyTransform, Settings->SplineMaxDiscretizationError, Positions2D))
+		{
+			// Unsuppported type
+			continue;
+		}
+		
+		TArray<TArrayView<FVector2D>> PathView;
+		PathView.Add(MakeArrayView(Positions2D));
+
+		TArray<TArray<FVector2D>> Result;
+
+		bool bOperationSuccessful =
+			(Settings->Operation == EPCGClipPathOperation::Intersection && UE::Geometry::PathsIntersection(PathView, ClipPolys, Result)) ||
+			(Settings->Operation == EPCGClipPathOperation::Difference && UE::Geometry::PathsDifference(PathView, ClipPolys, Result));
+
+		if (!bOperationSuccessful)
+		{
+			continue;
+		}
+		
+		TArray<UPCGSpatialData*> PathData;
+		PathData.Reserve(Result.Num());
+
+		for (const TArray<FVector2D>& ResultPath : Result)
+		{
+			if (ResultPath.Num() < 2)
+			{
+				continue;
+			}
+
+			if (const UPCGSplineData* InputSpline = Cast<UPCGSplineData>(Input.Data))
+			{
+				// Implementation note: while there's a small assumption that the spline is on the polygon plane,
+				// in practice as long as it is parallel to the polygons it's fine as the minimum distance will be the right one.
+				FVector PathStart = PolyTransform.TransformPosition(FVector(ResultPath[0], 0.0));
+				FVector PathEnd = PolyTransform.TransformPosition(FVector(ResultPath.Last(), 0.0));
+
+				double StartKey = InputSpline->SplineStruct.FindInputKeyClosestToWorldLocation(PathStart);
+				double EndKey = InputSpline->SplineStruct.FindInputKeyClosestToWorldLocation(PathEnd);
+
+				UPCGSplineData* SplitSplineData = FPCGContext::NewObject_AnyThread<UPCGSplineData>(Context);
+				SplitSplineData->InitializeFromData(InputSpline);
+
+				PCGSplitSpline::SplitSpline(InputSpline, SplitSplineData, StartKey, EndKey);
+
+				PathData.Add(SplitSplineData);
+			}
+			else if (const UPCGBasePointData* InputPointData = Cast<UPCGBasePointData>(Input.Data))
+			{
+				// Implementation note: stronger assumption that the points are on the plane here.
+				// If they aren't, then the sampling method will not work.
+				// @todo_pcg : Improve this; we could find key-equivalent and split the points instead.
+				TArray<FVector> PathPositions;
+				PathPositions.Reserve(ResultPath.Num());
+
+				for (const FVector2D& PathVertex : ResultPath)
+				{
+					PathPositions.Add(PolyTransform.TransformPosition(FVector(PathVertex, 0.0)));
+				}
+
+				UPCGBasePointData* SplitPointPath = FPCGContext::NewPointData_AnyThread(Context);
+				FPCGInitializeFromDataParams InitializeDataParams(InputPointData);
+				InitializeDataParams.bInheritMetadata = InitializeDataParams.bInheritAttributes = true;
+				InitializeDataParams.bInheritSpatialData = false;
+
+				SplitPointPath->InitializeFromDataWithParams(InitializeDataParams);
+
+				SplitPointPath->SetNumPoints(ResultPath.Num());
+
+				const bool bHasMetadata = !!(InputPointData->GetAllocatedProperties() & EPCGPointNativeProperties::MetadataEntry);
+				SplitPointPath->AllocateProperties(EPCGPointNativeProperties::Transform | (bHasMetadata ? EPCGPointNativeProperties::MetadataEntry : EPCGPointNativeProperties::None));
+
+				TPCGValueRange<FTransform> PathTransformRange = SplitPointPath->GetTransformValueRange(/*bAllocate=*/false);
+				for(int PathVertexIndex = 0; PathVertexIndex < PathPositions.Num(); ++PathVertexIndex)
+				{
+					PathTransformRange[PathVertexIndex] = FTransform(PathPositions[PathVertexIndex]);
+				}
+
+				if (bHasMetadata)
+				{
+					TPCGValueRange<PCGMetadataEntryKey> SplitPointEntryKeysRange = SplitPointPath->GetMetadataEntryValueRange(/*bAllocate=*/false);
+					TArray<PCGMetadataEntryKey> SplitPointEntryKeys = PCGPolygon2DProcessingHelpers::ComputeEntryKeysBySampling(InputPointData, Positions, SplitPointPath->MutableMetadata());
+					for (int Index = 0; Index < SplitPointEntryKeys.Num(); ++Index)
+					{
+						SplitPointEntryKeysRange[Index] = SplitPointEntryKeys[Index];
+					}
+				}
+
+				PathData.Add(SplitPointPath);
+			}
+			else
+			{
+				// shouldn't happen
+				check(0);
+				continue;
+			}
+		}
+
+		// Create output data for all new path data.
+		for (UPCGSpatialData* PathDatum : PathData)
+		{
+			Context->OutputData.TaggedData.Add_GetRef(Input).Data = PathDatum;
+		}
+	}
+
+	return true;
+}
+
+#undef LOCTEXT_NAMESPACE
\ No newline at end of file
diff --git a/Engine/Plugins/PCG/Source/PCG/Private/Elements/Polygon/PCGPolygon2DOperation.cpp b/Engine/Plugins/PCG/Source/PCG/Private/Elements/Polygon/PCGPolygon2DOperation.cpp
--- a/Engine/Plugins/PCG/Source/PCG/Private/Elements/Polygon/PCGPolygon2DOperation.cpp
+++ b/Engine/Plugins/PCG/Source/PCG/Private/Elements/Polygon/PCGPolygon2DOperation.cpp
@@ -10,11 +10,6 @@
 
 #define LOCTEXT_NAMESPACE "PCGPolygon2DOperation"
 
-namespace PCGPolygon2DOperationConstants
-{
-	const FName ClipPolysLabel = TEXT("Clip Polygons");
-}
-
 #if WITH_EDITOR
 FText UPCGPolygon2DOperationSettings::GetDefaultNodeTitle() const
 {
@@ -55,9 +50,16 @@ TArray<FPCGPinProperties> UPCGPolygon2DOperationSettings::InputPinProperties() c
 {
 	TArray<FPCGPinProperties> PinProperties = PCGPolygon2DUtils::DefaultPolygonInputPinProperties();
 
-	if (Operation != EPCGPolygonOperation::Union && Operation != EPCGPolygonOperation::InnerIntersection)
+	if(Operation == EPCGPolygonOperation::Difference ||
+		Operation == EPCGPolygonOperation::Intersection ||
+		Operation == EPCGPolygonOperation::PairwiseIntersection ||
+		Operation == EPCGPolygonOperation::ExclusiveOr)
+	{
+		PinProperties.Emplace(PCGPolygon2DUtils::Constants::ClipPolysLabel, EPCGDataType::Polygon2D);
+	}
+	else if (Operation == EPCGPolygonOperation::CutWithPaths)
 	{
-		PinProperties.Emplace_GetRef(PCGPolygon2DOperationConstants::ClipPolysLabel, EPCGDataType::Polygon2D).SetRequiredPin();
+		PinProperties.Emplace(PCGPolygon2DUtils::Constants::ClipPathsLabel, EPCGDataType::Point | EPCGDataType::Spline);
 	}
 
 	return PinProperties;
@@ -84,7 +86,16 @@ bool FPCGPolygon2DOperationElement::ExecuteInternal(FPCGContext* Context) const
 	check(Settings);
 
 	TArray<FPCGTaggedData> Inputs = Context->InputData.GetInputsByPin(PCGPinConstants::DefaultInputLabel);
-	TArray<FPCGTaggedData> ClipInputs = Context->InputData.GetInputsByPin(PCGPolygon2DOperationConstants::ClipPolysLabel);
+	TArray<FPCGTaggedData> ClipInputs;
+	
+	if (Settings->Operation == EPCGPolygonOperation::CutWithPaths)
+	{
+		ClipInputs = Context->InputData.GetInputsByPin(PCGPolygon2DUtils::Constants::ClipPathsLabel);
+	}
+	else
+	{
+		ClipInputs = Context->InputData.GetInputsByPin(PCGPolygon2DUtils::Constants::ClipPolysLabel);
+	}
 
 	bool bHasTransform = false;
 	FTransform Transform = FTransform::Identity;
@@ -96,69 +107,42 @@ bool FPCGPolygon2DOperationElement::ExecuteInternal(FPCGContext* Context) const
 	TArray<TArray<PCGMetadataEntryKey>> ClipPolygonsVertexEntryKeys;
 	TArray<const UPCGMetadata*> ClipPolygonsMetadata;
 
-	auto AddPolygon = [&bHasTransform, &Transform](const UPCGPolygon2DData* InData, TArray<UE::Geometry::FGeneralPolygon2d>& OutPolyList, TArray<TArray<PCGMetadataEntryKey>>& OutPolyVertexEntryKeysList, TArray<const UPCGMetadata*>& OutMetadataList)
+	for(int InputIndex = 0; InputIndex < Inputs.Num(); ++InputIndex)
 	{
-		check(InData);
-		const UE::Geometry::FGeneralPolygon2d& InPoly = InData->GetPolygon();
-		const FTransform& InTransform = InData->GetTransform();
-		TConstArrayView<PCGMetadataEntryKey> InEntryKeys = InData->GetConstVerticesEntryKeys();
-
-		OutMetadataList.Add(InData->ConstMetadata());
-
-		if (!bHasTransform)
+		if (const UPCGPolygon2DData* Polygon = Cast<UPCGPolygon2DData>(Inputs[InputIndex].Data))
 		{
-			UE::Geometry::FGeneralPolygon2d& OutPoly = OutPolyList.Add_GetRef(InPoly);
-			OutPolyVertexEntryKeysList.Emplace(InEntryKeys);
-
-			Transform = InTransform;
-			bHasTransform = true;
+			PolyIndexToInputIndex.Add(InputIndex);
+			PCGPolygon2DProcessingHelpers::AddPolygon(Polygon, Polygons, &PolygonsVertexEntryKeys, &PolygonsMetadata, bHasTransform, Transform);
 		}
-		else
-		{
-			// @todo_pcg: we could validate that the normal are parallel.
-			UE::Geometry::FGeneralPolygon2d& Poly = OutPolyList.Emplace_GetRef();
-			TArray<PCGMetadataEntryKey>& EntryKeys = OutPolyVertexEntryKeysList.Emplace_GetRef(InEntryKeys);
-			// Important note: since this can technically flip the normal, that would in turn change the polygon orientation.
-			// Unfortunately, this means that the winding of the polygon might change.
-			// However, the polygon transform isn't aware of this, so we need to rebuild the polygon from scratch here instead.
-			const FTransform RelativeTransform = InTransform.GetRelativeTransform(Transform);
-
-			Poly = InPoly;
-			Poly.Transform([&RelativeTransform](const FVector2D& InVertex)
-			{
-				const FVector TransformedVertex = RelativeTransform.TransformPosition(FVector(InVertex, 0.0));
-				return FVector2D(TransformedVertex.X, TransformedVertex.Y);
-			});
-
-			// If the outer was reversed because of the operation (due to a normal sign toggle, for example - and because the Transform doesn't update the flag),
-			// We need to fix this up so that the reported type of the polygon is in line with the reported type of the outer.
-			if (Poly.OuterIsClockwise() != Poly.GetOuter().IsClockwise())
-			{
-				PCGPolygon2DData::ReversePolygon(Poly, EntryKeys);
-
-				// Finally, reupdate the outer so the flag is put back to its good value
-				UE::Geometry::FPolygon2d TempOuter = Poly.GetOuter();
-				Poly.SetOuter(MoveTemp(TempOuter));
+	}
 
-				check(Poly.OuterIsClockwise() == InPoly.OuterIsClockwise());
-			}
+	if (Settings->Operation == EPCGPolygonOperation::CutWithPaths)
+	{
+		UE::Geometry::FAxisAlignedBox2d PolyBounds;
+		for (const UE::Geometry::FGeneralPolygon2d& Poly : Polygons)
+		{
+			PolyBounds.Contain(Poly.Bounds());
 		}
-	};
 
-	for(int InputIndex = 0; InputIndex < Inputs.Num(); ++InputIndex)
-	{
-		if (const UPCGPolygon2DData* Polygon = Cast<UPCGPolygon2DData>(Inputs[InputIndex].Data))
+		for (int ClipIndex = 0; ClipIndex < ClipInputs.Num(); ++ClipIndex)
 		{
-			PolyIndexToInputIndex.Add(InputIndex);
-			AddPolygon(Polygon, Polygons, PolygonsVertexEntryKeys, PolygonsMetadata);
+			UE::Geometry::FGeneralPolygon2d ClipPoly;
+			if (PCGPolygon2DProcessingHelpers::GetEnclosingPolygonFromPath(ClipInputs[ClipIndex].Data, Transform, Settings->SplineMaxDiscretizationError, PolyBounds, ClipPoly))
+			{
+				ClipPolygons.Add(MoveTemp(ClipPoly));
+				ClipPolygonsVertexEntryKeys.Emplace(); // no metadata
+				ClipPolygonsMetadata.Add(nullptr); // no metadata
+			}
 		}
 	}
-
-	for (int ClipIndex = 0; ClipIndex < ClipInputs.Num(); ++ClipIndex)
+	else
 	{
-		if (const UPCGPolygon2DData* Polygon = Cast<UPCGPolygon2DData>(ClipInputs[ClipIndex].Data))
+		for (int ClipIndex = 0; ClipIndex < ClipInputs.Num(); ++ClipIndex)
 		{
-			AddPolygon(Polygon, ClipPolygons, ClipPolygonsVertexEntryKeys, ClipPolygonsMetadata);
+			if (const UPCGPolygon2DData* Polygon = Cast<UPCGPolygon2DData>(ClipInputs[ClipIndex].Data))
+			{
+				PCGPolygon2DProcessingHelpers::AddPolygon(Polygon, ClipPolygons, &ClipPolygonsVertexEntryKeys, &ClipPolygonsMetadata, bHasTransform, Transform);
+			}
 		}
 	}
 
@@ -289,12 +273,9 @@ bool FPCGPolygon2DOperationElement::ExecuteInternal(FPCGContext* Context) const
 			const FPCGTaggedData& CurrentInput = Inputs[PolyIndexToInputIndex[PolygonIndex]];
 
 			UE::Geometry::FGeneralPolygon2d& Poly = Polygons[PolygonIndex];
-			TArray<PCGMetadataEntryKey>& EntryKeys = PolygonsVertexEntryKeys[PolygonIndex];
 			OutPolys.Reset();
-			bool bOperationOk = false;
 
 			TArrayView<UE::Geometry::FGeneralPolygon2d> SubjPolygons = MakeArrayView(&Poly, 1);
-			TArrayView<TArray<PCGMetadataEntryKey>> SubjEntryKeys = MakeArrayView(&EntryKeys, 1);
 
 			if (Settings->Operation == EPCGPolygonOperation::Difference ||
 				Settings->Operation == EPCGPolygonOperation::Intersection ||
@@ -318,7 +299,6 @@ bool FPCGPolygon2DOperationElement::ExecuteInternal(FPCGContext* Context) const
 					OutPolys.Reset();
 
 					TArrayView<UE::Geometry::FGeneralPolygon2d> ClipSubsetPolygons = MakeArrayView(&ClipPolygons[ClipIndex], 1);
-					TArrayView<TArray<PCGMetadataEntryKey>> ClipSubsetEntryKeys = MakeArrayView(&ClipPolygonsVertexEntryKeys[ClipIndex], 1);
 
 					if (UE::Geometry::PolygonsIntersection(SubjPolygons, ClipSubsetPolygons, OutPolys))
 					{
@@ -330,6 +310,61 @@ bool FPCGPolygon2DOperationElement::ExecuteInternal(FPCGContext* Context) const
 					}
 				}
 			}
+			else if (Settings->Operation == EPCGPolygonOperation::CutWithPaths)
+			{
+				// Cutting with paths is a bit more involved, as we need to take the polygon, then intersect + difference it, and get 1-2 polygons out.
+				// Then we need to cut these individually with the next path polygon and so on.
+				TArray<UE::Geometry::FGeneralPolygon2d> OutPositivePolys;
+				TArray<UE::Geometry::FGeneralPolygon2d> OutNegativePolys;
+				TArray<UE::Geometry::FGeneralPolygon2d> CurrentGenerationPolygons;
+				TArray<UE::Geometry::FGeneralPolygon2d> NextGenerationPolygons;
+
+				CurrentGenerationPolygons.Add(Poly);
+				bool bSuccess = true;
+
+				for (int ClipIndex = 0; ClipIndex < ClipPolygons.Num(); ++ClipIndex)
+				{
+					TArrayView<UE::Geometry::FGeneralPolygon2d> ClipSubsetPolygons = MakeArrayView(&ClipPolygons[ClipIndex], 1);
+
+					for (int GenerationPolyIndex = 0; GenerationPolyIndex < CurrentGenerationPolygons.Num(); ++GenerationPolyIndex)
+					{
+						TArrayView<UE::Geometry::FGeneralPolygon2d> CurrentSubjPolygons = MakeArrayView(&CurrentGenerationPolygons[GenerationPolyIndex], 1);
+						if (UE::Geometry::PolygonsIntersection(CurrentSubjPolygons, ClipSubsetPolygons, OutPositivePolys) && UE::Geometry::PolygonsDifference(CurrentSubjPolygons, ClipSubsetPolygons, OutNegativePolys))
+						{
+							NextGenerationPolygons.Append(OutPositivePolys);
+							NextGenerationPolygons.Append(OutNegativePolys);
+
+							OutPositivePolys.Reset();
+							OutNegativePolys.Reset();
+						}
+						else
+						{
+							bSuccess = false;
+							break;
+						}
+					}
+
+					if (bSuccess)
+					{
+						CurrentGenerationPolygons = MoveTemp(NextGenerationPolygons);
+						NextGenerationPolygons.Reset();
+					}
+					else
+					{
+						break;
+					}
+				}
+
+				if (bSuccess)
+				{
+					OutPolys = MoveTemp(CurrentGenerationPolygons);
+					AddToOutput(OutPolys, &CurrentInput, SubjPolygons, ClipPolygons);
+				}
+				else if(!Settings->bQuiet)
+				{
+					PCGLog::LogWarningOnGraph(LOCTEXT("OperationCutWithPathsFailed", "Polygon operation (cut with paths) failed."), Context);
+				}
+			}
 			else
 			{
 				PCGLog::LogErrorOnGraph(LOCTEXT("InvalidOperation", "Invalid polygon operation."), Context);
diff --git a/Engine/Plugins/PCG/Source/PCG/Private/Helpers/PCGPolygon2DProcessingHelpers.cpp b/Engine/Plugins/PCG/Source/PCG/Private/Helpers/PCGPolygon2DProcessingHelpers.cpp
--- a/Engine/Plugins/PCG/Source/PCG/Private/Helpers/PCGPolygon2DProcessingHelpers.cpp
+++ b/Engine/Plugins/PCG/Source/PCG/Private/Helpers/PCGPolygon2DProcessingHelpers.cpp
@@ -11,6 +11,220 @@
 namespace PCGPolygon2DProcessingHelpers
 {
 
+bool GetPath(const UPCGData* InData, const FTransform& PolyTransform, double DiscretizationError, TArray<FVector2D>& OutPath)
+{
+	OutPath.Reset();
+
+	TArray<FVector> Positions;
+
+	if (const UPCGSplineData* SplineData = Cast<UPCGSplineData>(InData))
+	{
+		const double ErrorTolerance = FMath::Max(DiscretizationError, 0.001);
+		// Implementation note: we will not remove the last point in the case of a closed spline because we want the full path regardless.
+		SplineData->SplineStruct.ConvertSplineToPolyLine(ESplineCoordinateSpace::World, ErrorTolerance, Positions);
+	}
+	else if(const UPCGBasePointData* PointData = Cast<UPCGBasePointData>(InData))
+	{
+		const TConstPCGValueRange<FTransform> TransformRange = PointData->GetConstTransformValueRange();
+		Positions.Reserve(TransformRange.Num());
+		for (const FTransform& Transform : TransformRange)
+		{
+			Positions.Emplace(Transform.GetLocation());
+		}
+	}
+	else
+	{
+		// Unsupported type
+		return false;
+	}
+
+	OutPath.Reserve(Positions.Num());
+	for (const FVector& Position : Positions)
+	{
+		FVector Position2D = PolyTransform.InverseTransformPosition(Position);
+		OutPath.Emplace(Position2D.X, Position2D.Y);
+	}
+
+	return true;
+}
+
+bool GetEnclosingPolygonFromPath(const UPCGData* InData, const FTransform& PolyTransform, double DiscretizationError, UE::Geometry::FAxisAlignedBox2d Bounds, UE::Geometry::FGeneralPolygon2d& OutPoly)
+{
+	// Idea: 
+	// 1. create a polygon where the first segment(s) is the input path
+	// 2. complete the polygon using the points from the bounding box,
+	//      where we will select the first vertex on the left from either the other points on the bounding box or the segment start.
+	TArray<FVector2D> Path;
+	if (!GetPath(InData, PolyTransform, DiscretizationError, Path))
+	{
+		// Unsupported type
+		return false;
+	}
+
+	if (Path.Num() < 2)
+	{
+		// Invalid path
+		return false;
+	}
+
+	// Make bounds slightly larger to prevent issues with edges on the polygons that are directly on the bounds.
+	const double SafeBoundarySize = FMath::Max(1.0, 0.01 * FMath::Max(Bounds.Width(), Bounds.Height()));
+	Bounds.Expand(SafeBoundarySize);
+
+	TStaticArray<FVector2D, 4> Corners;
+	for (int CornerIndex = 0; CornerIndex < 4; ++CornerIndex)
+	{
+		Corners[CornerIndex] = Bounds.GetCorner(CornerIndex);
+	}
+
+	TArray<FVector2D> TargetPoints;
+
+	TargetPoints.Reset();
+	TargetPoints.Append(Corners);
+	TargetPoints.Add(Path[0]);
+
+	bool bDone = false;
+	while (!bDone)
+	{
+		FVector2D LastPosition = Path.Last();
+		FVector2D LastDirection = (Path.Last() - Path.Last(1)).GetSafeNormal();
+		double MinAngle = DBL_MAX;
+		int MinIndex = INDEX_NONE;
+
+		// We'll select the point that has the smallest angle vs. the previous direction, counting 0 from the left, up to 2*pi at the right.
+		// e.g. we're doing similar to a convex hull with point sorting.
+		for (int TargetIndex = 0; TargetIndex < TargetPoints.Num(); ++TargetIndex)
+		{
+			// Compute angle from (Target - Last) vs last direction.
+			FVector2D TargetDirection = (TargetPoints[TargetIndex] - LastPosition).GetSafeNormal();
+			FVector2D ProjectedVector = FVector2D(TargetDirection | LastDirection, LastDirection ^ TargetDirection);
+			double Angle = FMath::Atan2(ProjectedVector.Y, ProjectedVector.X);
+
+			// Atan2 returns a value in the ]-pi, pi] range, but we want a value in the [0, 2pi[ range (left to right).
+			if (Angle < 0)
+			{
+				Angle += 2.0 * UE_DOUBLE_PI;
+			}
+
+			ensure(Angle >= 0);
+
+			if (Angle < MinAngle)
+			{
+				MinAngle = Angle;
+				MinIndex = TargetIndex;
+			}
+		}
+
+		// Two conditions here:
+		// Either we're adding a point from the bounding box, or we're done with the polygon.
+		if (MinIndex == TargetPoints.Num() - 1)
+		{
+			// we're done, we selected the first path point.
+			bDone = true;
+		}
+		else
+		{
+			Path.Add(TargetPoints[MinIndex]);
+			TargetPoints.RemoveAt(MinIndex);
+		}
+	}
+
+	OutPoly = UE::Geometry::FGeneralPolygon2d(UE::Geometry::FPolygon2d(MoveTemp(Path)));
+	return true;
+}
+
+TArray<UE::Geometry::FGeneralPolygon2d> GetPolygons(const TArray<const UPCGPolygon2DData*>& InPolygonData, bool& bTransformSet, FTransform& ReferenceTransform)
+{
+	TArray<UE::Geometry::FGeneralPolygon2d> OutPolys;
+	OutPolys.Reserve(InPolygonData.Num());
+
+	for (const UPCGPolygon2DData* PolyData : InPolygonData)
+	{
+		AddPolygon(PolyData, OutPolys, bTransformSet, ReferenceTransform);
+	}
+
+	return OutPolys;
+}
+
+void AddPolygon(const UPCGPolygon2DData* InPolygonData, TArray<UE::Geometry::FGeneralPolygon2d>& OutPolys, bool& bTransformSet, FTransform& ReferenceTransform)
+{
+	AddPolygon(InPolygonData, OutPolys, nullptr, nullptr, bTransformSet, ReferenceTransform);
+}
+
+void AddPolygon(const UPCGPolygon2DData* InPolygonData,
+	TArray<UE::Geometry::FGeneralPolygon2d>& OutPolys,
+	TArray<TArray<PCGMetadataEntryKey>>* OutPolyVertexEntryKeysList,
+	TArray<const UPCGMetadata*>* OutMetadataList,
+	bool& bTransformSet, 
+	FTransform& ReferenceTransform)
+{
+	check(InPolygonData);
+	const UE::Geometry::FGeneralPolygon2d& InPoly = InPolygonData->GetPolygon();
+	const FTransform& InTransform = InPolygonData->GetTransform();
+	TConstArrayView<PCGMetadataEntryKey> InEntryKeys = InPolygonData->GetConstVerticesEntryKeys();
+
+	if (OutMetadataList)
+	{
+		OutMetadataList->Add(InPolygonData->ConstMetadata());
+	}
+
+	if (!bTransformSet)
+	{
+		OutPolys.Add(InPoly);
+
+		if (OutPolyVertexEntryKeysList)
+		{
+			OutPolyVertexEntryKeysList->Emplace(InEntryKeys);
+		}
+
+		ReferenceTransform = InTransform;
+		bTransformSet = true;
+	}
+	else
+	{
+		// @todo_pcg: we could validate that the normal are parallel.
+		UE::Geometry::FGeneralPolygon2d& Poly = OutPolys.Emplace_GetRef();
+
+		if (OutPolyVertexEntryKeysList)
+		{
+			OutPolyVertexEntryKeysList->Emplace(InEntryKeys);
+		}
+
+		// Important note: since this can technically flip the normal, that would in turn change the polygon orientation.
+		// Unfortunately, this means that the winding of the polygon might change.
+		// However, the polygon transform isn't aware of this, so we need to rebuild the polygon from scratch here instead.
+		const FTransform RelativeTransform = InTransform.GetRelativeTransform(ReferenceTransform);
+
+		Poly = InPoly;
+		Poly.Transform([&RelativeTransform](const FVector2D& InVertex)
+		{
+			const FVector TransformedVertex = RelativeTransform.TransformPosition(FVector(InVertex, 0.0));
+			return FVector2D(TransformedVertex.X, TransformedVertex.Y);
+		});
+
+		// If the outer was reversed because of the operation (due to a normal sign toggle, for example - and because the Transform doesn't update the flag),
+		// We need to fix this up so that the reported type of the polygon is in line with the reported type of the outer.
+		if (Poly.OuterIsClockwise() != Poly.GetOuter().IsClockwise())
+		{
+			if (OutPolyVertexEntryKeysList)
+			{
+				PCGPolygon2DData::ReversePolygon(Poly, OutPolyVertexEntryKeysList->Last());
+			}
+			else
+			{
+				Poly.Reverse();
+			}
+
+			// Finally, reupdate the outer so the flag is put back to its good value
+			UE::Geometry::FPolygon2d TempOuter = Poly.GetOuter();
+			Poly.SetOuter(MoveTemp(TempOuter));
+
+			// At this point, the polygon should have been properly reset, but it might not appear so if it is degenerate, in which case we'll just take it as is.
+			check(Poly.OuterIsClockwise() == InPoly.OuterIsClockwise() || FMath::Abs(Poly.SignedArea()) < UE_KINDA_SMALL_NUMBER);
+		}
+	}
+}
+
 void ForAllVertices(UE::Geometry::FGeneralPolygon2d& InPoly, TFunctionRef<bool(const FVector2D&, int /*VertexIndex*/, int /*SegmentIndex*/, int /*HoleIndex*/)> InFunc)
 {
 	bool bExit = false;
diff --git a/Engine/Plugins/PCG/Source/PCG/Private/Helpers/PCGPolygon2DProcessingHelpers.h b/Engine/Plugins/PCG/Source/PCG/Private/Helpers/PCGPolygon2DProcessingHelpers.h
--- a/Engine/Plugins/PCG/Source/PCG/Private/Helpers/PCGPolygon2DProcessingHelpers.h
+++ b/Engine/Plugins/PCG/Source/PCG/Private/Helpers/PCGPolygon2DProcessingHelpers.h
@@ -6,11 +6,26 @@
 
 #include "Curve/GeneralPolygon2.h"
 
+class UPCGData;
 class UPCGMetadata;
+class UPCGPolygon2DData;
 class UPCGSpatialData;
 
 namespace PCGPolygon2DProcessingHelpers
 {
+	/** Returns the path representing the input data, projected in the PolyTransform local space. */
+	bool GetPath(const UPCGData* InData, const FTransform& PolyTransform, double DiscretizationError, TArray<FVector2D>& OutPath);
+
+	/** Gets the path from the input data and builds a closed polygon using provided bounds. Output polygon will be counterclockwise. */
+	bool GetEnclosingPolygonFromPath(const UPCGData* InData, const FTransform& PolyTransform, double DiscretizationError, UE::Geometry::FAxisAlignedBox2d Bounds, UE::Geometry::FGeneralPolygon2d& OutPoly);
+
+	/** Returns the polygons from the input data, normalizing their orientation/transforms based on the input/output transform. */
+	[[nodiscard]] TArray<UE::Geometry::FGeneralPolygon2d> GetPolygons(const TArray<const UPCGPolygon2DData*>& InPolygonData, bool& bTransformSet, FTransform& ReferenceTransform);
+
+	/** Adds & normalizes a polygon to a polygon array, conforming or setting the given transform. */
+	void AddPolygon(const UPCGPolygon2DData* InPolygonData, TArray<UE::Geometry::FGeneralPolygon2d>& OutPolys, bool& bTransformSet, FTransform& ReferenceTransform);
+	void AddPolygon(const UPCGPolygon2DData* InPolygonData, TArray<UE::Geometry::FGeneralPolygon2d>& OutPolys, TArray<TArray<PCGMetadataEntryKey>>* OutEntryKeys, TArray<const UPCGMetadata*>* OutMetadataList, bool& bTransformSet, FTransform& ReferenceTransform);
+
 	/** Visits all vertices in the polygon, starting from the outer. Return false in the InFunc to exit the loop prematurely. */
 	void ForAllVertices(UE::Geometry::FGeneralPolygon2d& InPoly, TFunctionRef<bool(const FVector2D& /*VertexPosition*/, int /*VertexIndex*/, int /*SegmentIndex*/, int /*HoleIndex*/)> InFunc);
 	
@@ -39,15 +54,15 @@ namespace PCGPolygon2DProcessingHelpers
 	/** Builds a mapping for all vertices of all polys to their nearest vertex/edge in the primary & secondary polygons.
 	* Warning: this is an expensive operation.
 	*/
-	TArray<PolygonVertexMapping> BuildVertexMapping(TArray<UE::Geometry::FGeneralPolygon2d>& Polys, TConstArrayView<UE::Geometry::FGeneralPolygon2d> PrimaryPolygons, TConstArrayView<UE::Geometry::FGeneralPolygon2d> SecondaryPolygons);
+	[[nodiscard]] TArray<PolygonVertexMapping> BuildVertexMapping(TArray<UE::Geometry::FGeneralPolygon2d>& Polys, TConstArrayView<UE::Geometry::FGeneralPolygon2d> PrimaryPolygons, TConstArrayView<UE::Geometry::FGeneralPolygon2d> SecondaryPolygons);
 
 	/** Collates vertex mapping information to allocate metadata entries with proper interpolation. */
-	TArray<PCGMetadataEntryKey> ComputeEntryKeysAndInitializeAttributes(UPCGMetadata* OutMetadata, const PolygonVertexMapping& InVertexMapping, const PolyToMetadataInfoMap& InInfoMap);
+	[[nodiscard]] TArray<PCGMetadataEntryKey> ComputeEntryKeysAndInitializeAttributes(UPCGMetadata* OutMetadata, const PolygonVertexMapping& InVertexMapping, const PolyToMetadataInfoMap& InInfoMap);
 
 	/** Computes entry keys from given positions and the original input data (spline or points), on a sampling basis. Assumes that the positions are "on" the existing edges. */
-	TArray<PCGMetadataEntryKey> ComputeEntryKeysBySampling(const UPCGSpatialData* InputData, TConstArrayView<FVector> Positions, UPCGMetadata* OutMetadata);
+	[[nodiscard]] TArray<PCGMetadataEntryKey> ComputeEntryKeysBySampling(const UPCGSpatialData* InputData, TConstArrayView<FVector> Positions, UPCGMetadata* OutMetadata);
 
 	/** Computes entry keys for polygon vertices by sampling against nearest segments on spatial data (spline or points). */
-	TArray<PCGMetadataEntryKey> ComputeEntryKeysByNearestSegmentSampling(const UPCGSpatialData* InputData, double InSearchRadius, UE::Geometry::FGeneralPolygon2d& InPoly, const FTransform& InPolyTransform, UPCGMetadata* OutMetadata);
+	[[nodiscard]] TArray<PCGMetadataEntryKey> ComputeEntryKeysByNearestSegmentSampling(const UPCGSpatialData* InputData, double InSearchRadius, UE::Geometry::FGeneralPolygon2d& InPoly, const FTransform& InPolyTransform, UPCGMetadata* OutMetadata);
 
 }
\ No newline at end of file
diff --git a/Engine/Plugins/PCG/Source/PCG/Public/Elements/PCGSplitSplines.h b/Engine/Plugins/PCG/Source/PCG/Public/Elements/PCGSplitSplines.h
--- a/Engine/Plugins/PCG/Source/PCG/Public/Elements/PCGSplitSplines.h
+++ b/Engine/Plugins/PCG/Source/PCG/Public/Elements/PCGSplitSplines.h
@@ -3,6 +3,7 @@
 #pragma once
 
 #include "PCGSettings.h"
+#include "Data/PCGSplineData.h"
 
 #include "PCGSplitSplines.generated.h"
 
@@ -15,6 +16,12 @@ enum class EPCGSplitSplineMode
 	ByPredicateOnControlPoints,
 };
 
+namespace PCGSplitSpline
+{
+	void SplitSpline(const UPCGSplineData* Spline, UPCGSplineData* SplitSpline, double StartKey, double EndKey);
+	void SplitSpline(const UPCGSplineData* Spline, UPCGSplineData* SplitSpline, TConstArrayView<FSplinePoint> SplinePoints, TConstArrayView<PCGMetadataEntryKey> SplinePointsEntryKeys, double StartKey, double EndKey);
+}
+
 /** Splits spline at a specific distance(s), key(s) or at certain values. */
 UCLASS(MinimalAPI, BlueprintType, ClassGroup=(Procedural))
 class UPCGSplitSplinesSettings : public UPCGSettings
diff --git a/Engine/Plugins/PCG/Source/PCG/Public/Elements/Polygon/PCGClipPaths.h b/Engine/Plugins/PCG/Source/PCG/Public/Elements/Polygon/PCGClipPaths.h
new file mode 100644
--- /dev/null
+++ b/Engine/Plugins/PCG/Source/PCG/Public/Elements/Polygon/PCGClipPaths.h
@@ -0,0 +1,56 @@
+// Copyright Epic Games, Inc. All Rights Reserved.
+
+#pragma once
+
+#include "PCGSettings.h"
+
+#include "PCGClipPaths.generated.h"
+
+UENUM()
+enum class EPCGClipPathOperation : uint8
+{
+	Intersection,
+	Difference
+};
+
+class FPCGClipPathsElement : public IPCGElement
+{
+public:
+	virtual bool ExecuteInternal(FPCGContext* InContext) const override;
+	virtual bool SupportsBasePointDataInputs(FPCGContext* InContext) const override { return true; }
+	virtual EPCGElementExecutionLoopMode ExecutionLoopMode(const UPCGSettings* Setings) const override { return EPCGElementExecutionLoopMode::SinglePrimaryPin; }
+};
+
+/**
+* Clips paths (points or splines) using provided polygons
+*/
+UCLASS(MinimalAPI, BlueprintType, ClassGroup=(Procedural))
+class UPCGClipPathsSettings : public UPCGSettings
+{
+	GENERATED_BODY()
+
+public:
+	//~Begin UPCGSettings interface
+#if WITH_EDITOR
+	virtual FName GetDefaultNodeName() const override { return FName(TEXT("ClipPaths")); }
+	virtual FText GetDefaultNodeTitle() const override;
+	virtual FText GetNodeTooltipText() const override;
+	virtual EPCGSettingsType GetType() const override { return EPCGSettingsType::Spatial; }
+#endif
+	virtual FString GetAdditionalTitleInformation() const override;
+
+protected:
+	virtual TArray<FPCGPinProperties> InputPinProperties() const override;
+	virtual TArray<FPCGPinProperties> OutputPinProperties() const override;
+	virtual FPCGElementPtr CreateElement() const override { return MakeShared<FPCGClipPathsElement>(); }
+	//~End UPCGSettings interface
+
+public:
+	/** Controls whether the paths will be clipped to be inside or outside the clip polygons. */
+	UPROPERTY(EditAnywhere, BlueprintReadWrite, Category = Settings)
+	EPCGClipPathOperation Operation = EPCGClipPathOperation::Intersection;
+
+	/** Maximum squared distance before we need to subdivide a segment again as part of the spline discretization to a polygon. */
+	UPROPERTY(EditAnywhere, BlueprintReadWrite, Category = Settings, AdvancedDisplay, meta = (ClampMin = 0.001, PCG_Overridable))
+	double SplineMaxDiscretizationError = 1.0;
+};
\ No newline at end of file
diff --git a/Engine/Plugins/PCG/Source/PCG/Public/Elements/Polygon/PCGCreatePolygon2D.h b/Engine/Plugins/PCG/Source/PCG/Public/Elements/Polygon/PCGCreatePolygon2D.h
--- a/Engine/Plugins/PCG/Source/PCG/Public/Elements/Polygon/PCGCreatePolygon2D.h
+++ b/Engine/Plugins/PCG/Source/PCG/Public/Elements/Polygon/PCGCreatePolygon2D.h
@@ -59,7 +59,7 @@ public:
 	UPROPERTY(EditAnywhere, BlueprintReadWrite, Category = Settings, meta = (EditCondition = "InputType != EPCGCreatePolygonInputType::ForceClosed", PCG_Overridable))
 	double OpenPolygonWidth = 100.0;
 
-	/* Maximum squared distance before we need to subdivide a segment again as part of the spline discretization to a polygon. */
+	/** Maximum squared distance before we need to subdivide a segment again as part of the spline discretization to a polygon. */
 	UPROPERTY(EditAnywhere, BlueprintReadWrite, Category = Settings, AdvancedDisplay, meta = (ClampMin = 0.001, PCG_Overridable))
 	double SplineMaxDiscretizationError = 1.0;
 };
\ No newline at end of file
diff --git a/Engine/Plugins/PCG/Source/PCG/Public/Elements/Polygon/PCGPolygon2DOperation.h b/Engine/Plugins/PCG/Source/PCG/Public/Elements/Polygon/PCGPolygon2DOperation.h
--- a/Engine/Plugins/PCG/Source/PCG/Public/Elements/Polygon/PCGPolygon2DOperation.h
+++ b/Engine/Plugins/PCG/Source/PCG/Public/Elements/Polygon/PCGPolygon2DOperation.h
@@ -16,14 +16,15 @@ enum class EPCGPolygonOperation : uint8
 	PairwiseIntersection,
 	InnerIntersection,
 	ExclusiveOr,
+	CutWithPaths UMETA(Tooltip="Cuts polygons with paths if and only if both ends of the path are outside the input polygons.")
 };
 
 UENUM()
 enum class EPCGPolygonOperationMetadataMode : uint8
 {
-	None,
-	SourceOnly,
-	Full
+	None UMETA(Tooltip="Output polygons will have their default attribute values only."),
+	SourceOnly UMETA(Tooltip="Output polygons will have their attribute values based on the input polygons only."),
+	Full UMETA(Tooltip="Output polygons will have their attribute values be based on input polygons first, then on clip polygons.")
 };
 
 class FPCGPolygon2DOperationElement : public IPCGElement
@@ -60,12 +61,19 @@ protected:
 	//~End UPCGSettings interface
 
 public:
+	/** Controls operation to be performed on the input polygons (vs. the clip polygons). */
 	UPROPERTY(BlueprintReadWrite, EditAnywhere, Category = Settings)
 	EPCGPolygonOperation Operation = EPCGPolygonOperation::Intersection;
 
+	/** Controls whether metadata on the resulting output will be either not computed & defaulted ('None'), based on the input polygons only ('SourceOnly') or using both the source and the clip polygons. */
 	UPROPERTY(BlueprintReadWrite, EditAnywhere, Category = Settings)
 	EPCGPolygonOperationMetadataMode MetadataMode = EPCGPolygonOperationMetadataMode::Full;
 
-	UPROPERTY(BlueprintReadWrite, EditAnywhere, Category = Settings)
+	/** Maximum squared distance before we need to subdivide a segment again as part of the spline discretization to a path. **/
+	UPROPERTY(BlueprintReadWrite, EditAnywhere, Category = Settings, AdvancedDisplay, meta = (ClampMin = 0.001, PCG_Overridable, EditCondition = "Operation==EPCGPolygonOperation::CutWithPaths", EditConditionHides))
+	double SplineMaxDiscretizationError = 1.0;
+
+	/** Controls whether the operation can log warnings/errors if it fails. */
+	UPROPERTY(BlueprintReadWrite, EditAnywhere, Category = Settings, AdvancedDisplay)
 	bool bQuiet = false;
 };
\ No newline at end of file
diff --git a/Engine/Plugins/PCG/Source/PCG/Public/Elements/Polygon/PCGPolygon2DUtils.h b/Engine/Plugins/PCG/Source/PCG/Public/Elements/Polygon/PCGPolygon2DUtils.h
--- a/Engine/Plugins/PCG/Source/PCG/Public/Elements/Polygon/PCGPolygon2DUtils.h
+++ b/Engine/Plugins/PCG/Source/PCG/Public/Elements/Polygon/PCGPolygon2DUtils.h
@@ -6,6 +6,12 @@
 
 namespace PCGPolygon2DUtils
 {
+	namespace Constants
+	{
+		const FLazyName ClipPolysLabel = TEXT("Clip Polygons");
+		const FLazyName ClipPathsLabel = TEXT("Clip Paths");
+	}
+
 	TArray<FPCGPinProperties> DefaultPolygonInputPinProperties();
 	TArray<FPCGPinProperties> DefaultPolygonOutputPinProperties();
 }
\ No newline at end of file
