[PCG] Change FPCGDynamicTrackingPriority to double so that value is signed and priority can be inserted between others using fractions
- Invert priority value: lower is higher priority
- Interops better with other systems using floating point priority
- Unedited Priority is always in the [0,~0.00007] range so anything lower when user specified will be higher prio and anything over will be lower prio when comparing to default priorities

diff --git a/Engine/Plugins/PCG/Source/PCG/Private/PCGActorAndComponentMapping.cpp b/Engine/Plugins/PCG/Source/PCG/Private/PCGActorAndComponentMapping.cpp
--- a/Engine/Plugins/PCG/Source/PCG/Private/PCGActorAndComponentMapping.cpp
+++ b/Engine/Plugins/PCG/Source/PCG/Private/PCGActorAndComponentMapping.cpp
@@ -1914,7 +1914,7 @@ void FPCGActorAndComponentMapping::OnSelectionKeyChanged(const FPCGSelectionKey&
 	//   * It is one of the objects that changed.
 	//   * We should not refresh the owner and its owner is one of the objects that changed.
 	//   * The objects that changed are managed by it (safe-guard, would probably induce an infinite loop).
-	//   * The objects that changed are managed by a lower tracking priority component.
+	//   * The objects that changed are managed by a lower tracking priority component. (lower tracking priority value means higher priority)
 	//   * One of the objects that changed is ignored by its original component.
 	auto ShouldDiscardComponent = [&ChangedObjects, &ChangedObjectDynamicTrackingPriority, bInNoRefreshOwner](const UPCGComponent* InComponent)
 	{
@@ -1923,7 +1923,7 @@ void FPCGActorAndComponentMapping::OnSelectionKeyChanged(const FPCGSelectionKey&
 			|| (bInNoRefreshOwner && ChangedObjects.Contains(InComponent->GetOwner()))
 			|| InComponent->IsCleaningUp()
 			|| InComponent->IsAnyObjectManagedByResource(ChangedObjects)
-			|| (ChangedObjectDynamicTrackingPriority.IsSet() && InComponent->GetExecutionState().GetDynamicTrackingPriority() > ChangedObjectDynamicTrackingPriority.GetValue()))
+			|| (ChangedObjectDynamicTrackingPriority.IsSet() && InComponent->GetExecutionState().GetDynamicTrackingPriority() < ChangedObjectDynamicTrackingPriority.GetValue()))
 		{
 			return true;
 		}
diff --git a/Engine/Plugins/PCG/Source/PCG/Private/PCGComponent.cpp b/Engine/Plugins/PCG/Source/PCG/Private/PCGComponent.cpp
--- a/Engine/Plugins/PCG/Source/PCG/Private/PCGComponent.cpp
+++ b/Engine/Plugins/PCG/Source/PCG/Private/PCGComponent.cpp
@@ -153,6 +153,13 @@ namespace PCGComponent
 		bool bShouldLoadActors;
 #endif
 	};
+
+	[[nodiscard]] double RoundTrackingPriority(double InTrackingPriority)
+	{
+		// Round to 4 decimals (10^4)
+		constexpr double Factor = 10000;
+		return FMath::RoundToDouble(InTrackingPriority * Factor) / Factor;
+	}
 }
 
 UPCGComponent::UPCGComponent(const FObjectInitializer& InObjectInitializer)
@@ -2497,6 +2504,13 @@ void UPCGComponent::PostEditChangeProperty(FPropertyChangedEvent& PropertyChange
 		DirtyGenerated();
 		Refresh();
 	}
+	else if (PropName == GET_MEMBER_NAME_CHECKED(UPCGComponent, TrackingPriority))
+	{
+		TrackingPriority = PCGComponent::RoundTrackingPriority(TrackingPriority);
+
+		DirtyGenerated();
+		Refresh();
+	}
 	else if (PropName == GET_MEMBER_NAME_CHECKED(UPCGComponent, SchedulingPolicyClass))
 	{
 		// We don't need to refresh the component here because this does not effect generation behavior, only scheduling behavior.
@@ -4138,6 +4152,7 @@ void FPCGComponentInstanceData::ApplyToComponent(UActorComponent* Component, con
 			PCGComponent->bWasGeneratedThisSession = ConstructionSourceComponent->bWasGeneratedThisSession;
 
 			PCGComponent->ExecutionInspection.InspectionCounter = ConstructionSourceComponent->ExecutionInspection.InspectionCounter;
+			PCGComponent->TrackingPriority = PCGComponent::RoundTrackingPriority(PCGComponent->TrackingPriority);
 #endif
 		}
 
@@ -4317,8 +4332,14 @@ FPCGDynamicTrackingPriority FPCGComponentExecutionState::GetDynamicTrackingPrior
 {
 	check(Component);
 
-	// Give user control over priority by making UPCGComponent::Priority the most significant part of the resulting priority
-	return static_cast<uint64>((static_cast<uint64>(Component->TrackingPriority) << 32) | GetTypeHash(Component->GetPathName()));
+	// Hash [0, 1] 
+	double Fraction = static_cast<double>(GetTypeHash(Component->GetPathName())) / std::numeric_limits<uint32>::max();
+	// Divide Fraction by 2^14 so that max value most significant decimal is the 5th (1.0 / 16384 -> 0.00006103515625)
+	// TrackingPriority is rounded to 4 decimals so that Fraction (Hash) doesn't interfere with it (or almost never, in which case difference in value would be minimal).
+	Fraction /= 16384;
+
+	// So that leading integer and fractional user part stay the same add or subtract the hash based fraction based on sign
+	return Component->TrackingPriority >= 0 ? (double)Component->TrackingPriority + Fraction : (double)Component->TrackingPriority - Fraction;
 }
 
 #endif // WITH_EDITOR
diff --git a/Engine/Plugins/PCG/Source/PCG/Public/PCGComponent.h b/Engine/Plugins/PCG/Source/PCG/Public/PCGComponent.h
--- a/Engine/Plugins/PCG/Source/PCG/Public/PCGComponent.h
+++ b/Engine/Plugins/PCG/Source/PCG/Public/PCGComponent.h
@@ -357,9 +357,9 @@ public:
 	UPROPERTY(BlueprintReadWrite, EditAnywhere, Category = "Editing Settings", meta = (DisplayPriority = 451))
 	bool bIgnoreLandscapeTracking = false;
 
-	/** Tracking priority used to solve tracking dependencies between PCG Components. By default object path is used to create a default priority (order) but overriding this value gives control to the user. */
-	UPROPERTY(BlueprintReadWrite, EditAnywhere, Category= "Editing Settings", meta = (DisplayPriority = 452, ClampMin=0, UIMin=0))
-	int32 TrackingPriority = 0; 
+	/** Tracking priority used to solve tracking dependencies between PCG Components. By default object path is used to create a default priority (order) but overriding this value gives control to the user. (smaller value is higher priority) */
+	UPROPERTY(BlueprintReadWrite, EditAnywhere, Category= "Editing Settings", meta = (DisplayPriority = 452, ClampMin="-10000.0", ClampMax="10000.0", UIMin="-10000.0", UIMax="10000.0", Delta="1"))
+	double TrackingPriority = 0; 
 
 	UPROPERTY(BlueprintReadOnly, VisibleAnywhere, Transient, Category = Debug, meta = (NoResetToDefault))
 	bool bDirtyGenerated = false;
diff --git a/Engine/Plugins/PCG/Source/PCG/Public/PCGGraphExecutionStateInterface.h b/Engine/Plugins/PCG/Source/PCG/Public/PCGGraphExecutionStateInterface.h
--- a/Engine/Plugins/PCG/Source/PCG/Public/PCGGraphExecutionStateInterface.h
+++ b/Engine/Plugins/PCG/Source/PCG/Public/PCGGraphExecutionStateInterface.h
@@ -23,7 +23,7 @@ namespace PCGUtils
 	class FExtraCapture;
 }
 
-using FPCGDynamicTrackingPriority = uint64;
+using FPCGDynamicTrackingPriority = double;
 
 /**
 * Interface returned by a IPCGGraphExecutionSource that is queried / updated during execution of a PCG Graph.
