[PCG] Log an error when Landscape DI fails to acquire any height texture

#rnx

diff --git a/Engine/Plugins/PCG/Source/PCG/Private/Compute/DataInterfaces/PCGLandscapeDataInterface.cpp b/Engine/Plugins/PCG/Source/PCG/Private/Compute/DataInterfaces/PCGLandscapeDataInterface.cpp
--- a/Engine/Plugins/PCG/Source/PCG/Private/Compute/DataInterfaces/PCGLandscapeDataInterface.cpp
+++ b/Engine/Plugins/PCG/Source/PCG/Private/Compute/DataInterfaces/PCGLandscapeDataInterface.cpp
@@ -680,7 +680,9 @@ void FPCGLandscapeDataProviderProxy::SetShaderParameters(FParameters& OutShaderP
 		SetBaseColorVirtualTextureParameters_Default(OutShaderParameters);
 	}
 
-	if (!SetHeightVirtualTextureParameters(OutShaderParameters))
+	const bool bSetHeightVT = SetHeightVirtualTextureParameters(OutShaderParameters);
+
+	if (!bSetHeightVT)
 	{
 		SetHeightVirtualTextureParameters_Default(OutShaderParameters);
 	}
@@ -690,11 +692,18 @@ void FPCGLandscapeDataProviderProxy::SetShaderParameters(FParameters& OutShaderP
 		SetNormalVirtualTextureParameters_Default(OutShaderParameters);
 	}
 
-	if (!SetCollisionHeightTextureParameters(OutShaderParameters))
+	const bool bSetCollisionHeightTexture = SetCollisionHeightTextureParameters(OutShaderParameters);
+
+	if (!bSetCollisionHeightTexture)
 	{
 		SetCollisionHeightTextureParameters_Defaults(OutShaderParameters);
 	}
 
+	if (!bSetHeightVT && !bSetCollisionHeightTexture)
+	{
+		UE_LOG(LogPCG, Error, TEXT("PCGLandscapeDataInterface: Failed to initialize a height texture. If you are trying to use a virtual texture, make sure the RVT volumes are created."))
+	}
+
 	OutShaderParameters.LandscapeGridSize = (FVector2f)LandscapeGridSize;
 }
 
