[PCG] Fix PCG RT data failing to initialize for some formats

#rnx

diff --git a/Engine/Plugins/PCG/Source/PCG/Private/Data/PCGRenderTargetData.cpp b/Engine/Plugins/PCG/Source/PCG/Private/Data/PCGRenderTargetData.cpp
--- a/Engine/Plugins/PCG/Source/PCG/Private/Data/PCGRenderTargetData.cpp
+++ b/Engine/Plugins/PCG/Source/PCG/Private/Data/PCGRenderTargetData.cpp
@@ -30,12 +30,16 @@ void UPCGRenderTargetData::Initialize(UTextureRenderTarget2D* InRenderTarget, co
 			{
 				TRACE_CPUPROFILER_EVENT_SCOPE(UPCGRenderTargetData::Initialize::ReadData);
 				ETextureRenderTargetFormat Format = RenderTarget->RenderTargetFormat;
-				if ((Format == RTF_RGBA16f) || (Format == RTF_RGBA32f) || (Format == RTF_RGBA8) || (Format == RTF_R8) || (Format == RTF_RGB10A2))
+				if ((Format == RTF_RGBA16f) || (Format == RTF_RGBA32f) || (Format == RTF_RGBA8) || (Format == RTF_RGB10A2))
 				{
 					const FIntRect Rect = FIntRect(0, 0, RenderTarget->SizeX, RenderTarget->SizeY);
 					const FReadSurfaceDataFlags ReadPixelFlags(RCM_MinMax);
 					RTResource->ReadLinearColorPixels(ColorData, ReadPixelFlags, Rect);
 				}
+				else
+				{
+					UE_LOG(LogPCG, Warning, TEXT("Could not CPU-initialize UPCGRenderTargetData with render target format '%s'."), *UEnum::GetValueAsString(Format));
+				}
 			}
 		}
 
