[PCG] Code clarifications from the clip paths cl.

#rnx

diff --git a/Engine/Plugins/PCG/Source/PCG/Private/Elements/PCGSplitSplines.cpp b/Engine/Plugins/PCG/Source/PCG/Private/Elements/PCGSplitSplines.cpp
--- a/Engine/Plugins/PCG/Source/PCG/Private/Elements/PCGSplitSplines.cpp
+++ b/Engine/Plugins/PCG/Source/PCG/Private/Elements/PCGSplitSplines.cpp
@@ -24,7 +24,13 @@ namespace PCGSplitSpline
 	{
 		check(Spline);
 		check(OutSplitSpline);
-		check(StartKey < EndKey || !Spline || Spline->IsClosed());
+
+		if (!Spline->IsClosed() && StartKey > EndKey)
+		{
+			double TempKey = StartKey;
+			StartKey = EndKey;
+			EndKey = TempKey;
+		}
 
 		TArray<FSplinePoint> SplitSplinePoints;
 		TArray<int64> SplitSplinePointsEntryKeys;
@@ -118,19 +124,21 @@ namespace PCGSplitSpline
 		}
 
 		// Adjust tangents if required
-		if (bNeedsStartTangentAdjustment)
+		if (SplitSplinePoints.Num() >= 2)
 		{
-			check(SplitSplinePoints.Num() >= 1);
-			// assumes key range is 1.
-			const double Ratio = (1.0 - FMath::Frac(SplitSplinePoints[0].InputKey));
-			SplitSplinePoints[1].ArriveTangent *= Ratio;
-		}
+			if (bNeedsStartTangentAdjustment)
+			{
+				// assumes key range is 1.
+				const double Ratio = (1.0 - FMath::Frac(SplitSplinePoints[0].InputKey));
+				SplitSplinePoints[1].ArriveTangent *= Ratio;
+			}
 
-		if (bNeedsEndTangentAdjustment && SplitSplinePoints.Num() >= 2)
-		{
-			// assumes key range is 1.
-			const double Ratio = FMath::Frac(SplitSplinePoints.Last().InputKey);
-			SplitSplinePoints[SplitSplinePoints.Num() - 2].LeaveTangent *= Ratio;
+			if (bNeedsEndTangentAdjustment)
+			{
+				// assumes key range is 1.
+				const double Ratio = FMath::Frac(SplitSplinePoints.Last().InputKey);
+				SplitSplinePoints[SplitSplinePoints.Num() - 2].LeaveTangent *= Ratio;
+			}
 		}
 
 		// Finally, mark all tangents as custom and reset the keys.
diff --git a/Engine/Plugins/PCG/Source/PCG/Private/Elements/Polygon/PCGClipPaths.cpp b/Engine/Plugins/PCG/Source/PCG/Private/Elements/Polygon/PCGClipPaths.cpp
--- a/Engine/Plugins/PCG/Source/PCG/Private/Elements/Polygon/PCGClipPaths.cpp
+++ b/Engine/Plugins/PCG/Source/PCG/Private/Elements/Polygon/PCGClipPaths.cpp
@@ -176,7 +176,7 @@ bool FPCGClipPathsElement::ExecuteInternal(FPCGContext* Context) const
 				if (bHasMetadata)
 				{
 					TPCGValueRange<PCGMetadataEntryKey> SplitPointEntryKeysRange = SplitPointPath->GetMetadataEntryValueRange(/*bAllocate=*/false);
-					TArray<PCGMetadataEntryKey> SplitPointEntryKeys = PCGPolygon2DProcessingHelpers::ComputeEntryKeysBySampling(InputPointData, Positions, SplitPointPath->MutableMetadata());
+					TArray<PCGMetadataEntryKey> SplitPointEntryKeys = PCGPolygon2DProcessingHelpers::ComputeEntryKeysBySampling(InputPointData, PathPositions, SplitPointPath->MutableMetadata());
 					for (int Index = 0; Index < SplitPointEntryKeys.Num(); ++Index)
 					{
 						SplitPointEntryKeysRange[Index] = SplitPointEntryKeys[Index];
diff --git a/Engine/Plugins/PCG/Source/PCG/Public/Elements/Polygon/PCGPolygon2DOperation.h b/Engine/Plugins/PCG/Source/PCG/Public/Elements/Polygon/PCGPolygon2DOperation.h
--- a/Engine/Plugins/PCG/Source/PCG/Public/Elements/Polygon/PCGPolygon2DOperation.h
+++ b/Engine/Plugins/PCG/Source/PCG/Public/Elements/Polygon/PCGPolygon2DOperation.h
@@ -16,7 +16,7 @@ enum class EPCGPolygonOperation : uint8
 	PairwiseIntersection,
 	InnerIntersection,
 	ExclusiveOr,
-	CutWithPaths UMETA(Tooltip="Cuts polygons with paths if and only if both ends of the path are outside the input polygons.")
+	CutWithPaths UMETA(Tooltip="Cuts polygons with paths by completing them using the polygon bounds. If both ends of the path are not outside the input polygons, the results might be incorrect.")
 };
 
 UENUM()
