[PCG] Support overrides on params that must be known at dispatch time

diff --git a/Engine/Plugins/PCG/Shaders/Private/Elements/PCGCopyPoints.usf b/Engine/Plugins/PCG/Shaders/Private/Elements/PCGCopyPoints.usf
--- a/Engine/Plugins/PCG/Shaders/Private/Elements/PCGCopyPoints.usf
+++ b/Engine/Plugins/PCG/Shaders/Private/Elements/PCGCopyPoints.usf
@@ -21,9 +21,7 @@ void PCGCopyPointsCS(uint3 GroupId : SV_GroupID, uint GroupIndex : SV_GroupIndex
 
 	const uint SourceNumData = Source_GetNumData();
 	const uint TargetNumData = Target_GetNumData();
-
-	// @todo_pcg: Must be known statically, so graph splits are required in order to support overriding this property.
-	const uint CopyEachSourceOnEveryTarget = CopyPoints_GetCopyEachSourceOnEveryTarget();
+	const uint CopyEachSourceOnEveryTarget = bCopyEachSourceOnEveryTarget_GetOverridableValue();
 
 	// Get source and target data indices for this output data index.
 	uint SourceDataIndex, TargetDataIndex;
@@ -53,6 +51,7 @@ void PCGCopyPointsCS(uint3 GroupId : SV_GroupID, uint GroupIndex : SV_GroupIndex
 	const uint TargetSeed = Target_GetSeed(TargetDataIndex, TargetPointIndexInData);
 
 	const uint RotationInheritance = RotationInheritance_GetOverridableValue();
+
 	uint ApplyTargetRotationToPositions = 1;
 	if (RotationInheritance == PCG_COPY_POINTS_INHERITANCE_MODE_Source)
 	{
@@ -67,9 +66,9 @@ void PCGCopyPointsCS(uint3 GroupId : SV_GroupID, uint GroupIndex : SV_GroupIndex
 	}
 
 	// @todo_pcg: Must be known statically, so graph splits are required in order to support overriding these properties.
-	const uint ColorInheritance = CopyPoints_GetColorInheritance();
-	const uint SeedInheritance = CopyPoints_GetSeedInheritance();
-	const uint AttributeInheritance = CopyPoints_GetAttributeInheritance();
+	const uint ColorInheritance = ColorInheritance_GetOverridableValue();
+	const uint SeedInheritance = SeedInheritance_GetOverridableValue();
+	const uint AttributeInheritance = AttributeInheritance_GetOverridableValue();
 
 	const float3 ScaleForPosition = ApplyTargetScaleToPositions > 0 ? TargetScale : float3(1, 1, 1);
 
diff --git a/Engine/Plugins/PCG/Shaders/Private/Elements/PCGCopyPointsAnalysis.usf b/Engine/Plugins/PCG/Shaders/Private/Elements/PCGCopyPointsAnalysis.usf
--- a/Engine/Plugins/PCG/Shaders/Private/Elements/PCGCopyPointsAnalysis.usf
+++ b/Engine/Plugins/PCG/Shaders/Private/Elements/PCGCopyPointsAnalysis.usf
@@ -33,7 +33,7 @@ void PCGCopyPointsAnalysisCS(uint3 GroupId : SV_GroupID, uint GroupIndex : SV_Gr
 	uint SourceDataIndex = 0;
 	uint TargetDataIndex = 0;
 
-	if (CopyPoints_GetCopyEachSourceOnEveryTarget())
+	if (bCopyEachSourceOnEveryTarget_GetOverridableValue())
 	{
 		// Output is attribute set with SourceNumData * TargetNumData bools. Kernel executes one thread per output.
 		// copypoints(S = ((src 0), (src 1)), T = ((tgt 0), (tgt 1))) = (src 0 x tgt 0), (src 0 x tgt 1), (src 1 x tgt 0), (src 1 tgt 1))
diff --git a/Engine/Plugins/PCG/Source/PCG/Private/Compute/DataInterfaces/Elements/PCGCopyPointsAnalysisDataInterface.cpp b/Engine/Plugins/PCG/Source/PCG/Private/Compute/DataInterfaces/Elements/PCGCopyPointsAnalysisDataInterface.cpp
--- a/Engine/Plugins/PCG/Source/PCG/Private/Compute/DataInterfaces/Elements/PCGCopyPointsAnalysisDataInterface.cpp
+++ b/Engine/Plugins/PCG/Source/PCG/Private/Compute/DataInterfaces/Elements/PCGCopyPointsAnalysisDataInterface.cpp
@@ -20,6 +20,8 @@
 
 void UPCGCopyPointsAnalysisDataInterface::GetSupportedInputs(TArray<FShaderFunctionDefinition>& OutFunctions) const
 {
+	Super::GetSupportedInputs(OutFunctions);
+
 	OutFunctions.AddDefaulted_GetRef()
 		.SetName(TEXT("CopyPoints_GetMatchAttributeId"))
 		.AddReturnType(FShaderValueType::Get(EShaderFundamentalType::Int));
@@ -27,16 +29,12 @@ void UPCGCopyPointsAnalysisDataInterface::GetSupportedInputs(TArray<FShaderFunct
 	OutFunctions.AddDefaulted_GetRef()
 		.SetName(TEXT("CopyPoints_GetSelectedFlagAttributeId"))
 		.AddReturnType(FShaderValueType::Get(EShaderFundamentalType::Int));
-
-	OutFunctions.AddDefaulted_GetRef()
-		.SetName(TEXT("CopyPoints_GetCopyEachSourceOnEveryTarget"))
-		.AddReturnType(FShaderValueType::Get(EShaderFundamentalType::Uint));
 }
 
 BEGIN_SHADER_PARAMETER_STRUCT(FPCGCopyPointsAnalysisDataInterfaceParameters,)
 	SHADER_PARAMETER(int32, MatchAttributeId)
 	SHADER_PARAMETER(int32, SelectedFlagAttributeId)
-	SHADER_PARAMETER(uint32, CopyEachSourceOnEveryTarget)
+	SHADER_PARAMETER(uint32, bCopyEachSourceOnEveryTarget)
 END_SHADER_PARAMETER_STRUCT()
 
 void UPCGCopyPointsAnalysisDataInterface::GetShaderParameters(TCHAR const* UID, FShaderParametersMetadataBuilder& InOutBuilder, FShaderParametersMetadataAllocations& InOutAllocations) const
@@ -46,6 +44,8 @@ void UPCGCopyPointsAnalysisDataInterface::GetShaderParameters(TCHAR const* UID,
 
 void UPCGCopyPointsAnalysisDataInterface::GetHLSL(FString& OutHLSL, FString const& InDataInterfaceName) const
 {
+	Super::GetHLSL(OutHLSL, InDataInterfaceName);
+
 	TMap<FString, FStringFormatArg> TemplateArgs =
 	{
 		{ TEXT("DataInterfaceName"), InDataInterfaceName },
@@ -54,11 +54,9 @@ void UPCGCopyPointsAnalysisDataInterface::GetHLSL(FString& OutHLSL, FString cons
 	OutHLSL += FString::Format(TEXT(
 		"int {DataInterfaceName}_MatchAttributeId;\n"
 		"int {DataInterfaceName}_SelectedFlagAttributeId;\n"
-		"uint {DataInterfaceName}_CopyEachSourceOnEveryTarget;\n"
 		"\n"
 		"int CopyPoints_GetMatchAttributeId_{DataInterfaceName}() { return {DataInterfaceName}_MatchAttributeId; }\n"
 		"int CopyPoints_GetSelectedFlagAttributeId_{DataInterfaceName}() { return {DataInterfaceName}_SelectedFlagAttributeId; }\n"
-		"uint CopyPoints_GetCopyEachSourceOnEveryTarget_{DataInterfaceName}() { return {DataInterfaceName}_CopyEachSourceOnEveryTarget; }\n"
 		), TemplateArgs);
 }
 
@@ -67,20 +65,6 @@ UComputeDataProvider* UPCGCopyPointsAnalysisDataInterface::CreateDataProvider()
 	return NewObject<UPCGCopyPointsAnalysisDataProvider>();
 }
 
-void UPCGCopyPointsAnalysisDataProvider::Initialize(const UComputeDataInterface* InDataInterface, UObject* InBinding, uint64 InInputMask, uint64 InOutputMask)
-{
-	TRACE_CPUPROFILER_EVENT_SCOPE(UPCGCopyPointsAnalysisDataProvider::Initialize);
-
-	Super::Initialize(InDataInterface, InBinding, InInputMask, InOutputMask);
-
-	const UPCGCopyPointsAnalysisDataInterface* DataInterface = CastChecked<UPCGCopyPointsAnalysisDataInterface>(InDataInterface);
-	UPCGDataBinding* Binding = CastChecked<UPCGDataBinding>(InBinding);
-
-	Settings = CastChecked<UPCGCopyPointsSettings>(DataInterface->ProducerKernel->GetSettings());
-
-	bCopyEachSourceOnEveryTarget = Settings->bCopyEachSourceOnEveryTarget;
-}
-
 bool UPCGCopyPointsAnalysisDataProvider::PrepareForExecute_GameThread(UPCGDataBinding* InBinding)
 {
 	if (!Super::PrepareForExecute_GameThread(InBinding))
@@ -88,32 +72,42 @@ bool UPCGCopyPointsAnalysisDataProvider::PrepareForExecute_GameThread(UPCGDataBi
 		return false;
 	}
 
+	const FPCGKernelParams* KernelParams = InBinding->GetCachedKernelParams(ProducerKernel);
+
+	if (!ensure(KernelParams))
+	{
+		return true;
+	}
+
+	Params.bCopyEachSourceOnEveryTarget = KernelParams->GetValueBool(GET_MEMBER_NAME_CHECKED(UPCGCopyPointsSettings, bCopyEachSourceOnEveryTarget));
+
 	// Resolve attribute IDs from incoming data.
-	const FName MatchAttributeName = Settings->MatchAttribute;
-	MatchAttributeId = InBinding->GetAttributeId(MatchAttributeName, EPCGKernelAttributeType::Int);
-	ensure(MatchAttributeId != INDEX_NONE);
+	const FName MatchAttributeName = KernelParams->GetValueName(GET_MEMBER_NAME_CHECKED(UPCGCopyPointsSettings, MatchAttribute));
+	Params.MatchAttributeId = InBinding->GetAttributeId(MatchAttributeName, EPCGKernelAttributeType::Int);
+	ensure(Params.MatchAttributeId != INDEX_NONE);
 
-	SelectedFlagAttributeId = InBinding->GetAttributeId(PCGCopyPointsConstants::SelectedFlagAttributeName, EPCGKernelAttributeType::Bool);
-	ensure(SelectedFlagAttributeId != INDEX_NONE);
+	Params.SelectedFlagAttributeId = InBinding->GetAttributeId(PCGCopyPointsConstants::SelectedFlagAttributeName, EPCGKernelAttributeType::Bool);
+	ensure(Params.SelectedFlagAttributeId != INDEX_NONE);
 
 	return true;
 }
 
 FComputeDataProviderRenderProxy* UPCGCopyPointsAnalysisDataProvider::GetRenderProxy()
 {
-	FPCGCopyPointsAnalysisDataProviderProxy::FCopyPointsAnalysisData_RenderThread ProxyData;
-	ProxyData.MatchAttributeId = MatchAttributeId;
-	ProxyData.SelectedFlagAttributeId = SelectedFlagAttributeId;
-	ProxyData.bCopyEachSourceOnEveryTarget = bCopyEachSourceOnEveryTarget;
+	return new FPCGCopyPointsAnalysisDataProviderProxy(Params);
+}
 
-	return new FPCGCopyPointsAnalysisDataProviderProxy(MoveTemp(ProxyData));
+void UPCGCopyPointsAnalysisDataProvider::Reset()
+{
+	Params = {};
+	Super::Reset();
 }
 
 bool FPCGCopyPointsAnalysisDataProviderProxy::IsValid(FValidationData const& InValidationData) const
 {
 	return InValidationData.ParameterStructSize == sizeof(FParameters)
-		&& Data.MatchAttributeId != INDEX_NONE
-		&& Data.SelectedFlagAttributeId != INDEX_NONE;
+		&& Params.MatchAttributeId != INDEX_NONE
+		&& Params.SelectedFlagAttributeId != INDEX_NONE;
 }
 
 void FPCGCopyPointsAnalysisDataProviderProxy::GatherDispatchData(FDispatchData const& InDispatchData)
@@ -122,8 +116,8 @@ void FPCGCopyPointsAnalysisDataProviderProxy::GatherDispatchData(FDispatchData c
 	for (int32 InvocationIndex = 0; InvocationIndex < ParameterArray.Num(); ++InvocationIndex)
 	{
 		FParameters& Parameters = ParameterArray[InvocationIndex];
-		Parameters.MatchAttributeId = Data.MatchAttributeId;
-		Parameters.SelectedFlagAttributeId = Data.SelectedFlagAttributeId;
-		Parameters.CopyEachSourceOnEveryTarget = Data.bCopyEachSourceOnEveryTarget ? 1 : 0;
+		Parameters.MatchAttributeId = Params.MatchAttributeId;
+		Parameters.SelectedFlagAttributeId = Params.SelectedFlagAttributeId;
+		Parameters.bCopyEachSourceOnEveryTarget = Params.bCopyEachSourceOnEveryTarget ? 1 : 0;
 	}
 }
diff --git a/Engine/Plugins/PCG/Source/PCG/Private/Compute/DataInterfaces/Elements/PCGCopyPointsAnalysisDataInterface.h b/Engine/Plugins/PCG/Source/PCG/Private/Compute/DataInterfaces/Elements/PCGCopyPointsAnalysisDataInterface.h
--- a/Engine/Plugins/PCG/Source/PCG/Private/Compute/DataInterfaces/Elements/PCGCopyPointsAnalysisDataInterface.h
+++ b/Engine/Plugins/PCG/Source/PCG/Private/Compute/DataInterfaces/Elements/PCGCopyPointsAnalysisDataInterface.h
@@ -10,9 +10,16 @@
 class FPCGCopyPointsAnalysisDataInterfaceParameters;
 class UPCGCopyPointsSettings;
 
+struct FPCGCopyPointsAnalysisParams
+{
+	int32 MatchAttributeId = INDEX_NONE;
+	int32 SelectedFlagAttributeId = INDEX_NONE;
+	bool bCopyEachSourceOnEveryTarget = true;
+};
+
 /** Data Interface to marshal Copy Points settings to the GPU. */
 UCLASS(ClassGroup = (Procedural))
-class UPCGCopyPointsAnalysisDataInterface : public UPCGComputeDataInterface
+class UPCGCopyPointsAnalysisDataInterface : public UPCGKernelParamsDataInterface
 {
 	GENERATED_BODY()
 
@@ -27,43 +34,26 @@ public:
 };
 
 UCLASS()
-class UPCGCopyPointsAnalysisDataProvider : public UPCGComputeDataProvider
+class UPCGCopyPointsAnalysisDataProvider : public UPCGKernelParamsDataProvider
 {
 	GENERATED_BODY()
 
 public:
 	//~ Begin UComputeDataProvider Interface
-	virtual void Initialize(const UComputeDataInterface* InDataInterface, UObject* InBinding, uint64 InInputMask, uint64 InOutputMask) override;
 	virtual bool PrepareForExecute_GameThread(UPCGDataBinding* InBinding) override;
 	FComputeDataProviderRenderProxy* GetRenderProxy() override;
+	virtual void Reset() override;
 	//~ End UComputeDataProvider Interface
 
 public:
-	UPROPERTY()
-	TObjectPtr<const UPCGCopyPointsSettings> Settings;
-
-	UPROPERTY()
-	int32 MatchAttributeId = INDEX_NONE;
-
-	UPROPERTY()
-	int32 SelectedFlagAttributeId = INDEX_NONE;
-
-	UPROPERTY()
-	bool bCopyEachSourceOnEveryTarget = true;
+	FPCGCopyPointsAnalysisParams Params;
 };
 
 class FPCGCopyPointsAnalysisDataProviderProxy : public FComputeDataProviderRenderProxy
 {
 public:
-	struct FCopyPointsAnalysisData_RenderThread
-	{
-		int32 MatchAttributeId = INDEX_NONE;
-		int32 SelectedFlagAttributeId = INDEX_NONE;
-		bool bCopyEachSourceOnEveryTarget = true;
-	};
-
-	FPCGCopyPointsAnalysisDataProviderProxy(FCopyPointsAnalysisData_RenderThread InData)
-		: Data(MoveTemp(InData))
+	FPCGCopyPointsAnalysisDataProviderProxy(const FPCGCopyPointsAnalysisParams& InParams)
+		: Params(InParams)
 	{}
 
 	//~ Begin FComputeDataProviderRenderProxy Interface
@@ -74,5 +64,5 @@ public:
 protected:
 	using FParameters = FPCGCopyPointsAnalysisDataInterfaceParameters;
 
-	FCopyPointsAnalysisData_RenderThread Data;
+	FPCGCopyPointsAnalysisParams Params;
 };
diff --git a/Engine/Plugins/PCG/Source/PCG/Private/Compute/DataInterfaces/Elements/PCGCopyPointsDataInterface.cpp b/Engine/Plugins/PCG/Source/PCG/Private/Compute/DataInterfaces/Elements/PCGCopyPointsDataInterface.cpp
--- a/Engine/Plugins/PCG/Source/PCG/Private/Compute/DataInterfaces/Elements/PCGCopyPointsDataInterface.cpp
+++ b/Engine/Plugins/PCG/Source/PCG/Private/Compute/DataInterfaces/Elements/PCGCopyPointsDataInterface.cpp
@@ -26,32 +26,7 @@
 
 void UPCGCopyPointsDataInterface::GetSupportedInputs(TArray<FShaderFunctionDefinition>& OutFunctions) const
 {
-	const UPCGSettings* Settings = ProducerKernel->GetSettings();
-	check(Settings);
-
-	// @todo_pcg: Pull this into the base class or elsewhere to make supporting overrides easier.
-	for (const FPCGSettingsOverridableParam& OverridableParam : Settings->OverridableParamsForGPU())
-	{
-		OutFunctions.AddDefaulted_GetRef()
-			.SetName(FString::Format(TEXT("CopyPoints_Get{0}Internal"), { OverridableParam.Label.ToString()}))
-			.AddReturnType(FShaderValueType::Get(EShaderFundamentalType::Uint)); // @todo_pcg: Determine the correct type from the overridable param
-	}
-
-	OutFunctions.AddDefaulted_GetRef()
-		.SetName(TEXT("CopyPoints_GetColorInheritance"))
-		.AddReturnType(FShaderValueType::Get(EShaderFundamentalType::Uint));
-
-	OutFunctions.AddDefaulted_GetRef()
-		.SetName(TEXT("CopyPoints_GetSeedInheritance"))
-		.AddReturnType(FShaderValueType::Get(EShaderFundamentalType::Uint));
-
-	OutFunctions.AddDefaulted_GetRef()
-		.SetName(TEXT("CopyPoints_GetAttributeInheritance"))
-		.AddReturnType(FShaderValueType::Get(EShaderFundamentalType::Uint));
-
-	OutFunctions.AddDefaulted_GetRef()
-		.SetName(TEXT("CopyPoints_GetCopyEachSourceOnEveryTarget"))
-		.AddReturnType(FShaderValueType::Get(EShaderFundamentalType::Bool));
+	Super::GetSupportedInputs(OutFunctions);
 
 	OutFunctions.AddDefaulted_GetRef()
 		.SetName(TEXT("CopyPoints_GetSourceAndTargetDataIndices"))
@@ -70,7 +45,7 @@ BEGIN_SHADER_PARAMETER_STRUCT(FPCGCopyPointsDataInterfaceParameters,)
 	SHADER_PARAMETER(uint32, ColorInheritance)
 	SHADER_PARAMETER(uint32, SeedInheritance)
 	SHADER_PARAMETER(uint32, AttributeInheritance)
-	SHADER_PARAMETER(uint32, CopyEachSourceOnEveryTarget)
+	SHADER_PARAMETER(uint32, bCopyEachSourceOnEveryTarget)
 	SHADER_PARAMETER_RDG_BUFFER_SRV(StructuredBuffer<FUintVector2>, SourceAndTargetDataIndices)
 END_SHADER_PARAMETER_STRUCT()
 
@@ -100,36 +75,16 @@ void UPCGCopyPointsDataInterface::GetDefines(FComputeKernelDefinitionSet& OutDef
 
 void UPCGCopyPointsDataInterface::GetHLSL(FString& OutHLSL, FString const& InDataInterfaceName) const
 {
+	Super::GetHLSL(OutHLSL, InDataInterfaceName);
+
 	TMap<FString, FStringFormatArg> TemplateArgs =
 	{
 		{ TEXT("DataInterfaceName"), InDataInterfaceName },
 	};
 	
-	const UPCGSettings* Settings = ProducerKernel->GetSettings();
-	check(Settings);
-
-	// @todo_pcg: Pull this into the base class or elsewhere to make supporting overrides easier.
-	for (const FPCGSettingsOverridableParam& OverridableParam : Settings->OverridableParamsForGPU())
-	{
-		// @todo_pcg: Determine the correct type from the overridable param
-		OutHLSL += FString::Format(TEXT(
-				"uint {0}_{1};\n"
-				"uint CopyPoints_Get{1}Internal_{0}() { return {0}_{1}; }\n"
-			), { InDataInterfaceName, OverridableParam.Label.ToString() });
-	}
-
 	OutHLSL += FString::Format(TEXT(
-		"uint {DataInterfaceName}_ColorInheritance;\n"
-		"uint {DataInterfaceName}_SeedInheritance;\n"
-		"uint {DataInterfaceName}_AttributeInheritance;\n"
-		"uint {DataInterfaceName}_CopyEachSourceOnEveryTarget;\n"
 		"StructuredBuffer<uint2> {DataInterfaceName}_SourceAndTargetDataIndices;\n"
 		"\n"
-		"uint CopyPoints_GetColorInheritance_{DataInterfaceName}() { return {DataInterfaceName}_ColorInheritance; }\n"
-		"uint CopyPoints_GetSeedInheritance_{DataInterfaceName}() { return {DataInterfaceName}_SeedInheritance; }\n"
-		"uint CopyPoints_GetAttributeInheritance_{DataInterfaceName}() { return {DataInterfaceName}_AttributeInheritance; }\n"
-		"uint CopyPoints_GetCopyEachSourceOnEveryTarget_{DataInterfaceName}() { return {DataInterfaceName}_CopyEachSourceOnEveryTarget; }\n"
-		"\n"
 		"void CopyPoints_GetSourceAndTargetDataIndices_{DataInterfaceName}(uint InOutputDataIndex, int InSourceDataCount, int InTargetDataCount, out uint OutSourceIndex, out uint OutTargetIndex)\n"
 		"{\n"
 		"#if USE_INPUT_DATA_INDICES_BUFFER\n"
@@ -137,7 +92,8 @@ void UPCGCopyPointsDataInterface::GetHLSL(FString& OutHLSL, FString const& InDat
 		"	OutSourceIndex = Indices[0];\n"
 		"	OutTargetIndex = Indices[1];\n"
 		"#else\n"
-		"	if (CopyPoints_GetCopyEachSourceOnEveryTarget_{DataInterfaceName}())\n"
+		// Note: Because bCopyEachSourceOnEveryTarget requires a graph split to be overidden, it is safe to directly access instead of calling GetOverridableValue(), which can only be called from kernel source.
+		"	if ({DataInterfaceName}_bCopyEachSourceOnEveryTarget)\n"
 		"	{\n"
 		"		OutSourceIndex = InOutputDataIndex / InTargetDataCount;\n"
 		"		OutTargetIndex = InOutputDataIndex % InTargetDataCount;\n"
@@ -167,8 +123,16 @@ bool UPCGCopyPointsDataProvider::PerformPreExecuteReadbacks_GameThread(UPCGDataB
 		return false;
 	}
 
-	const UPCGCopyPointsSettings* Settings = CastChecked<UPCGCopyPointsSettings>(ProducerKernel->GetSettings());
-	if (!Settings->bMatchBasedOnAttribute)
+	KernelParams = InBinding->GetCachedKernelParams(ProducerKernel);
+
+	if (!ensure(KernelParams))
+	{
+		return true;
+	}
+
+	const bool bMatchBasedOnAttribute = KernelParams->GetValueBool(GET_MEMBER_NAME_CHECKED(UPCGCopyPointsSettings, bMatchBasedOnAttribute));
+
+	if (!bMatchBasedOnAttribute)
 	{
 		// If we're not matching based on attribute then we don't need to do any readbacks.
 		return true;
@@ -190,14 +154,19 @@ bool UPCGCopyPointsDataProvider::PrepareForExecute_GameThread(UPCGDataBinding* I
 	TRACE_CPUPROFILER_EVENT_SCOPE(UPCGCopyPointsDataProvider::PrepareForExecute_GameThread);
 	check(InBinding);
 
+	if (!ensure(KernelParams))
+	{
+		return true;
+	}
+
 	if (!Super::PrepareForExecute_GameThread(InBinding))
 	{
 		return false;
 	}
 
-	const UPCGCopyPointsSettings* Settings = CastChecked<UPCGCopyPointsSettings>(ProducerKernel->GetSettings());
+	const bool bMatchBasedOnAttribute = KernelParams->GetValueBool(GET_MEMBER_NAME_CHECKED(UPCGCopyPointsSettings, bMatchBasedOnAttribute));
 
-	if (Settings->bMatchBasedOnAttribute)
+	if (bMatchBasedOnAttribute)
 	{
 		const int32 AnalysisDataIndex = InBinding->GetFirstInputDataIndex(ProducerKernel, PCGCopyPointsConstants::SelectedFlagsPinLabel);
 		if (AnalysisDataIndex == INDEX_NONE)
@@ -213,7 +182,9 @@ bool UPCGCopyPointsDataProvider::PrepareForExecute_GameThread(UPCGDataBinding* I
 			return true;
 		}
 
-		const int32 NumOutputs = Settings->bCopyEachSourceOnEveryTarget ? SourcePinDesc->DataDescs.Num() * TargetPinDesc->DataDescs.Num() : FMath::Max(SourcePinDesc->DataDescs.Num(), TargetPinDesc->DataDescs.Num());
+		const bool bCopyEachSourceOnEveryTarget = KernelParams->GetValueBool(GET_MEMBER_NAME_CHECKED(UPCGCopyPointsSettings, bCopyEachSourceOnEveryTarget));
+
+		const int32 NumOutputs = bCopyEachSourceOnEveryTarget ? SourcePinDesc->DataDescs.Num() * TargetPinDesc->DataDescs.Num() : FMath::Max(SourcePinDesc->DataDescs.Num(), TargetPinDesc->DataDescs.Num());
 		SourceAndTargetDataIndices.Reserve(NumOutputs);
 
 		const UPCGParamData* AnalysisResultsData = Cast<UPCGParamData>(InBinding->InputDataCollection.TaggedData[AnalysisDataIndex].Data);
@@ -224,7 +195,7 @@ bool UPCGCopyPointsDataProvider::PrepareForExecute_GameThread(UPCGDataBinding* I
 
 		if (SelectedFlagsAttribute)
 		{
-			if (Settings->bCopyEachSourceOnEveryTarget)
+			if (bCopyEachSourceOnEveryTarget)
 			{
 				for (int32 S = 0; S < SourcePinDesc->DataDescs.Num(); ++S)
 				{
@@ -258,18 +229,29 @@ bool UPCGCopyPointsDataProvider::PrepareForExecute_GameThread(UPCGDataBinding* I
 
 FComputeDataProviderRenderProxy* UPCGCopyPointsDataProvider::GetRenderProxy()
 {
-	const UPCGCopyPointsSettings* Settings = CastChecked<UPCGCopyPointsSettings>(ProducerKernel->GetSettings());
-
 	FPCGCopyPointsDataProviderProxy::FCopyPointsData_RenderThread ProxyData;
-	ProxyData.RotationInheritance = static_cast<uint32>(Settings->RotationInheritance);
-	ProxyData.ApplyTargetRotationToPositions = Settings->bApplyTargetRotationToPositions ? 1u : 0u;
-	ProxyData.ScaleInheritance = static_cast<uint32>(Settings->ScaleInheritance);
-	ProxyData.ApplyTargetScaleToPositions = Settings->bApplyTargetScaleToPositions ? 1u : 0u;
-	ProxyData.ColorInheritance = static_cast<uint32>(Settings->ColorInheritance);
-	ProxyData.SeedInheritance = static_cast<uint32>(Settings->SeedInheritance);
-	ProxyData.AttributeInheritance = static_cast<uint32>(Settings->AttributeInheritance);
-	ProxyData.bCopyEachSourceOnEveryTarget = static_cast<uint32>(Settings->bCopyEachSourceOnEveryTarget);
-	ProxyData.SourceAndTargetDataIndices = SourceAndTargetDataIndices;
+
+	if (ensure(KernelParams))
+	{
+		const EPCGCopyPointsInheritanceMode RotationInheritance = KernelParams->GetValueEnum<EPCGCopyPointsInheritanceMode>(GET_MEMBER_NAME_CHECKED(UPCGCopyPointsSettings, RotationInheritance));
+		const bool bApplyTargetRotationToPositions = KernelParams->GetValueBool(GET_MEMBER_NAME_CHECKED(UPCGCopyPointsSettings, bApplyTargetRotationToPositions));
+		const EPCGCopyPointsInheritanceMode ScaleInheritance = KernelParams->GetValueEnum<EPCGCopyPointsInheritanceMode>(GET_MEMBER_NAME_CHECKED(UPCGCopyPointsSettings, ScaleInheritance));
+		const bool bApplyTargetScaleToPositions = KernelParams->GetValueBool(GET_MEMBER_NAME_CHECKED(UPCGCopyPointsSettings, bApplyTargetScaleToPositions));
+		const EPCGCopyPointsInheritanceMode ColorInheritance = KernelParams->GetValueEnum<EPCGCopyPointsInheritanceMode>(GET_MEMBER_NAME_CHECKED(UPCGCopyPointsSettings, ColorInheritance));
+		const EPCGCopyPointsInheritanceMode SeedInheritance = KernelParams->GetValueEnum<EPCGCopyPointsInheritanceMode>(GET_MEMBER_NAME_CHECKED(UPCGCopyPointsSettings, SeedInheritance));
+		const EPCGCopyPointsMetadataInheritanceMode AttributeInheritance = KernelParams->GetValueEnum<EPCGCopyPointsMetadataInheritanceMode>(GET_MEMBER_NAME_CHECKED(UPCGCopyPointsSettings, AttributeInheritance));
+		const bool bCopyEachSourceOnEveryTarget = KernelParams->GetValueBool(GET_MEMBER_NAME_CHECKED(UPCGCopyPointsSettings, bCopyEachSourceOnEveryTarget));
+
+		ProxyData.RotationInheritance = static_cast<uint32>(RotationInheritance);
+		ProxyData.ApplyTargetRotationToPositions = bApplyTargetRotationToPositions ? 1u : 0u;
+		ProxyData.ScaleInheritance = static_cast<uint32>(ScaleInheritance);
+		ProxyData.ApplyTargetScaleToPositions = bApplyTargetScaleToPositions ? 1u : 0u;
+		ProxyData.ColorInheritance = static_cast<uint32>(ColorInheritance);
+		ProxyData.SeedInheritance = static_cast<uint32>(SeedInheritance);
+		ProxyData.AttributeInheritance = static_cast<uint32>(AttributeInheritance);
+		ProxyData.bCopyEachSourceOnEveryTarget = static_cast<uint32>(bCopyEachSourceOnEveryTarget);
+		ProxyData.SourceAndTargetDataIndices = SourceAndTargetDataIndices;
+	}
 
 	return new FPCGCopyPointsDataProviderProxy(MoveTemp(ProxyData));
 }
@@ -278,6 +260,7 @@ void UPCGCopyPointsDataProvider::Reset()
 {
 	Super::Reset();
 
+	KernelParams = nullptr;
 	SourceAndTargetDataIndices.Empty();
 }
 
@@ -344,7 +327,7 @@ void FPCGCopyPointsDataProviderProxy::GatherDispatchData(FDispatchData const& In
 		Parameters.ColorInheritance = Data.ColorInheritance;
 		Parameters.SeedInheritance = Data.SeedInheritance;
 		Parameters.AttributeInheritance = Data.AttributeInheritance;
-		Parameters.CopyEachSourceOnEveryTarget = Data.bCopyEachSourceOnEveryTarget;
+		Parameters.bCopyEachSourceOnEveryTarget = Data.bCopyEachSourceOnEveryTarget;
 		Parameters.SourceAndTargetDataIndices = SourceAndTargetDataIndicesBufferSRV;
 	}
 }
diff --git a/Engine/Plugins/PCG/Source/PCG/Private/Compute/DataInterfaces/Elements/PCGCopyPointsDataInterface.h b/Engine/Plugins/PCG/Source/PCG/Private/Compute/DataInterfaces/Elements/PCGCopyPointsDataInterface.h
--- a/Engine/Plugins/PCG/Source/PCG/Private/Compute/DataInterfaces/Elements/PCGCopyPointsDataInterface.h
+++ b/Engine/Plugins/PCG/Source/PCG/Private/Compute/DataInterfaces/Elements/PCGCopyPointsDataInterface.h
@@ -3,16 +3,15 @@
 #pragma once
 
 #include "Compute/DataInterfaces/PCGComputeDataInterface.h"
-#include "ComputeFramework/ComputeDataProvider.h"
 
 #include "PCGCopyPointsDataInterface.generated.h"
 
 class FPCGCopyPointsDataInterfaceParameters;
-class UPCGCopyPointsSettings;
+struct FPCGKernelParams;
 
 /** Data Interface to marshal Copy Points settings to the GPU. */
 UCLASS(ClassGroup = (Procedural))
-class UPCGCopyPointsDataInterface : public UPCGComputeDataInterface
+class UPCGCopyPointsDataInterface : public UPCGKernelParamsDataInterface
 {
 	GENERATED_BODY()
 
@@ -29,7 +28,7 @@ public:
 };
 
 UCLASS()
-class UPCGCopyPointsDataProvider : public UPCGComputeDataProvider
+class UPCGCopyPointsDataProvider : public UPCGKernelParamsDataProvider
 {
 	GENERATED_BODY()
 
@@ -45,7 +44,9 @@ public:
 	//~ End UComputeDataProvider Interface
 
 public:
-	UPROPERTY()
+	/** This property bag instance is owned by the KernelParamsCache, so we are not responsible for its lifetime. */
+	const FPCGKernelParams* KernelParams = nullptr;
+
 	TArray<FUintVector2> SourceAndTargetDataIndices;
 };
 
diff --git a/Engine/Plugins/PCG/Source/PCG/Private/Compute/DataInterfaces/PCGComputeDataInterface.cpp b/Engine/Plugins/PCG/Source/PCG/Private/Compute/DataInterfaces/PCGComputeDataInterface.cpp
--- a/Engine/Plugins/PCG/Source/PCG/Private/Compute/DataInterfaces/PCGComputeDataInterface.cpp
+++ b/Engine/Plugins/PCG/Source/PCG/Private/Compute/DataInterfaces/PCGComputeDataInterface.cpp
@@ -10,6 +10,8 @@
 #include "Compute/PCGComputeKernel.h"
 #include "Compute/PCGDataBinding.h"
 
+#include "ComputeFramework/ShaderParamTypeDefinition.h"
+
 #include UE_INLINE_GENERATED_CPP_BY_NAME(PCGComputeDataInterface)
 
 #if WITH_EDITOR
@@ -218,3 +220,33 @@ bool UPCGExportableDataProvider::PrepareForExecute_GameThread(UPCGDataBinding* I
 
 	return true;
 }
+
+void UPCGKernelParamsDataInterface::GetSupportedInputs(TArray<FShaderFunctionDefinition>& OutFunctions) const
+{
+	if (ProducerKernel)
+	{
+		for (const FPCGKernelOverridableParam& OverridableParam : ProducerKernel->GetCachedOverridableParams())
+		{
+			// @todo_pcg: Broaden type support
+			OutFunctions.AddDefaulted_GetRef()
+				.SetName(FString::Format(TEXT("Get{0}Internal"), { OverridableParam.Label.ToString(), }))
+				.AddReturnType(FShaderValueType::Get(EShaderFundamentalType::Uint));
+		}
+	}
+}
+
+void UPCGKernelParamsDataInterface::GetHLSL(FString& OutHLSL, FString const& InDataInterfaceName) const
+{
+	if (ProducerKernel)
+	{
+		for (const FPCGKernelOverridableParam& OverridableParam : ProducerKernel->GetCachedOverridableParams())
+		{
+			// @todo_pcg: Broaden type support.
+			OutHLSL += FString::Format(TEXT(
+				"uint {0}_{1};\n"
+				"uint Get{1}Internal_{0}() { return {0}_{1}; }\n"
+			), { InDataInterfaceName, OverridableParam.Label.ToString(), });
+		}
+	}
+}
+
diff --git a/Engine/Plugins/PCG/Source/PCG/Private/Compute/DataInterfaces/PCGComputeDataInterface.h b/Engine/Plugins/PCG/Source/PCG/Private/Compute/DataInterfaces/PCGComputeDataInterface.h
--- a/Engine/Plugins/PCG/Source/PCG/Private/Compute/DataInterfaces/PCGComputeDataInterface.h
+++ b/Engine/Plugins/PCG/Source/PCG/Private/Compute/DataInterfaces/PCGComputeDataInterface.h
@@ -161,3 +161,24 @@ protected:
 	TWeakObjectPtr<UPCGDataBinding> Binding;
 	const FPCGDataCollectionDesc* PinDesc = nullptr;
 };
+
+/** Interface for marshalling settings/params to the GPU. */
+UCLASS(Abstract, ClassGroup = (Procedural))
+class UPCGKernelParamsDataInterface : public UPCGComputeDataInterface
+{
+	GENERATED_BODY()
+
+public:
+	//~ Begin UComputeDataInterface Interface
+	void GetSupportedInputs(TArray<FShaderFunctionDefinition>& OutFunctions) const override;
+	void GetHLSL(FString& OutHLSL, FString const& InDataInterfaceName) const override;
+	//~ End UComputeDataInterface Interface
+};
+
+UCLASS(Abstract, ClassGroup (Procedural))
+class UPCGKernelParamsDataProvider : public UPCGComputeDataProvider
+{
+	GENERATED_BODY()
+
+	// Nothing implemented, but provided to avoid breaking changes in the future in case it becomes necessary.
+};
diff --git a/Engine/Plugins/PCG/Source/PCG/Private/Compute/DataInterfaces/PCGDataCollectionDataInterface.cpp b/Engine/Plugins/PCG/Source/PCG/Private/Compute/DataInterfaces/PCGDataCollectionDataInterface.cpp
--- a/Engine/Plugins/PCG/Source/PCG/Private/Compute/DataInterfaces/PCGDataCollectionDataInterface.cpp
+++ b/Engine/Plugins/PCG/Source/PCG/Private/Compute/DataInterfaces/PCGDataCollectionDataInterface.cpp
@@ -257,6 +257,10 @@ void UPCGDataCollectionDataInterface::GetSupportedInputs(TArray<FShaderFunctionD
 			.SetName(TEXT("GetFirstAttributeAsBool"))
 			.AddReturnType(EShaderFundamentalType::Bool);
 
+		OutFunctions.AddDefaulted_GetRef()
+			.SetName(TEXT("GetFirstAttributeAsUint"))
+			.AddReturnType(EShaderFundamentalType::Uint);
+
 		OutFunctions.AddDefaulted_GetRef()
 			.SetName(TEXT("GetFirstAttributeAsInt"))
 			.AddReturnType(EShaderFundamentalType::Int);
diff --git a/Engine/Plugins/PCG/Source/PCG/Private/Compute/Elements/PCGComputeGraphElement.cpp b/Engine/Plugins/PCG/Source/PCG/Private/Compute/Elements/PCGComputeGraphElement.cpp
--- a/Engine/Plugins/PCG/Source/PCG/Private/Compute/Elements/PCGComputeGraphElement.cpp
+++ b/Engine/Plugins/PCG/Source/PCG/Private/Compute/Elements/PCGComputeGraphElement.cpp
@@ -292,6 +292,27 @@ bool FPCGComputeGraphElement::ExecuteInternal(FPCGContext* InContext) const
 
 		InitializeBindingTask.Wait();
 
+		Context->ExecutionSubPhase = EPCGComputeGraphExecutionPhase::InitializeKernelParams;
+
+		if (Context->AsyncState.ShouldStop())
+		{
+			return false;
+		}
+	}
+	case EPCGComputeGraphExecutionPhase::InitializeKernelParams: // Fallthrough
+	{
+		TRACE_CPUPROFILER_EVENT_SCOPE(InitializeKernelParams);
+
+		if (UPCGDataBinding* DataBinding = Context->DataBinding.Get(); ensure(DataBinding))
+		{
+			// @todo_pcg: Move this to a worker thread.
+			if (!DataBinding->InitializeKernelParams(Context))
+			{
+				SleepUntilNextFrame();
+				return false;
+			}
+		}
+
 		Context->ExecutionSubPhase = EPCGComputeGraphExecutionPhase::PreExecuteReadbacks;
 
 		if (Context->AsyncState.ShouldStop())
diff --git a/Engine/Plugins/PCG/Source/PCG/Private/Compute/Elements/PCGComputeGraphElement.h b/Engine/Plugins/PCG/Source/PCG/Private/Compute/Elements/PCGComputeGraphElement.h
--- a/Engine/Plugins/PCG/Source/PCG/Private/Compute/Elements/PCGComputeGraphElement.h
+++ b/Engine/Plugins/PCG/Source/PCG/Private/Compute/Elements/PCGComputeGraphElement.h
@@ -22,6 +22,7 @@ enum class EPCGComputeGraphExecutionPhase : uint8
 	None,
 	GetComputeGraph,
 	InitializeDataBindingAndComputeGraph,
+	InitializeKernelParams,
 	PreExecuteReadbacks,
 	PrimeDataDescriptionsAndValidateData,
 	PrepareForExecute,
diff --git a/Engine/Plugins/PCG/Source/PCG/Private/Compute/PCGComputeGraph.cpp b/Engine/Plugins/PCG/Source/PCG/Private/Compute/PCGComputeGraph.cpp
--- a/Engine/Plugins/PCG/Source/PCG/Private/Compute/PCGComputeGraph.cpp
+++ b/Engine/Plugins/PCG/Source/PCG/Private/Compute/PCGComputeGraph.cpp
@@ -162,6 +162,16 @@ void UPCGComputeGraph::GetKernelPins(TArray<FPCGKernelPin>& OutKernelPins) const
 	KernelPinToFirstBinding.GenerateKeyArray(OutKernelPins);
 }
 
+int32 UPCGComputeGraph::GetNumKernels() const
+{
+	return KernelInvocations.Num();
+}
+
+const UComputeKernel* UPCGComputeGraph::GetKernel(int32 InKernelIndex) const
+{
+	return KernelInvocations.IsValidIndex(InKernelIndex) ? KernelInvocations[InKernelIndex] : nullptr;
+}
+
 bool UPCGComputeGraph::AreGraphSettingsValid(FPCGContext* InContext) const
 {
 	TRACE_CPUPROFILER_EVENT_SCOPE(UPCGComputeGraph::AreGraphSettingsValid);
@@ -305,7 +315,7 @@ void UPCGComputeGraph::DebugLogDataDescriptions(const UPCGDataBinding* InBinding
 			UE_LOG(LogPCG, Display, TEXT("Kernel: %s"), *PCGKernel->GetName());
 
 			TArray<FPCGPinProperties> InputPins;
-			PCGKernel->GetInputPins(InputPins);
+			PCGKernel->GetInputPinsAndOverridablePins(InputPins);
 			for (const FPCGPinProperties& Pin : InputPins)
 			{
 				if (Pin.Label != NAME_None)
diff --git a/Engine/Plugins/PCG/Source/PCG/Private/Compute/PCGComputeGraph.h b/Engine/Plugins/PCG/Source/PCG/Private/Compute/PCGComputeGraph.h
--- a/Engine/Plugins/PCG/Source/PCG/Private/Compute/PCGComputeGraph.h
+++ b/Engine/Plugins/PCG/Source/PCG/Private/Compute/PCGComputeGraph.h
@@ -88,6 +88,9 @@ public:
 	/** Populates an array with all of the kernel pins in the compute graph. */
 	void GetKernelPins(TArray<FPCGKernelPin>& OutKernelPins) const;
 
+	int32 GetNumKernels() const;
+	const UComputeKernel* GetKernel(int32 InKernelIndex) const;
+
 	/** Validates and logs information about the graph known at compile time. */
 	bool AreGraphSettingsValid(FPCGContext* InContext) const;
 
diff --git a/Engine/Plugins/PCG/Source/PCG/Private/Compute/PCGComputeKernel.cpp b/Engine/Plugins/PCG/Source/PCG/Private/Compute/PCGComputeKernel.cpp
--- a/Engine/Plugins/PCG/Source/PCG/Private/Compute/PCGComputeKernel.cpp
+++ b/Engine/Plugins/PCG/Source/PCG/Private/Compute/PCGComputeKernel.cpp
@@ -57,9 +57,20 @@ void UPCGComputeKernel::Initialize(const FPCGComputeKernelParams& InParams)
 		PCGKernelFlags |= EPCGComputeKernelFlags::LogDataDescriptions;
 	}
 
-	if (InParams.bRequiresOverridableParams)
+	if (InParams.bRequiresOverridableParams && ResolvedSettings)
 	{
-		PCGKernelFlags |= EPCGComputeKernelFlags::RequiresOverridableParams;
+		// Copy any GPU supported overridable params from the settings.
+		for (const FPCGSettingsOverridableParam& OverridableParam : ResolvedSettings->OverridableParams())
+		{
+			if (OverridableParam.bSupportsGPU)
+			{
+				FPCGKernelOverridableParam& KernelOverridableParam = CachedOverridableParams.Emplace_GetRef();
+				KernelOverridableParam.Label = OverridableParam.Label;
+				KernelOverridableParam.PropertiesNames = OverridableParam.PropertiesNames;
+				KernelOverridableParam.bRequiresGPUReadback = OverridableParam.bRequiresGPUReadback;
+				KernelOverridableParam.bIsPropertyOverriddenByPin = ResolvedSettings->IsPropertyOverriddenByPin(OverridableParam.PropertiesNames);
+			}
+		}
 	}
 
 	bInitialized = true;
@@ -220,6 +231,19 @@ void UPCGComputeKernel::GatherAdditionalSources(TArray<TObjectPtr<UComputeSource
 	}
 }
 
+bool UPCGComputeKernel::SplitGraphAtInput(FName InPinLabel) const
+{
+	for (const FPCGKernelOverridableParam& OverridableParam : CachedOverridableParams)
+	{
+		if (OverridableParam.bRequiresGPUReadback && OverridableParam.Label == InPinLabel)
+		{
+			return true;
+		}
+	}
+
+	return false;
+}
+
 UPCGComputeDataInterface* UPCGComputeKernel::CreateOutputPinDataInterface(const PCGComputeHelpers::FCreateDataInterfaceParams& InParams) const
 {
 	return PCGComputeHelpers::CreateOutputPinDataInterface(InParams);
@@ -524,6 +548,12 @@ void UPCGComputeKernel::ComputeDataDescFromPinProperties(const FPCGPinProperties
 	}
 }
 
+void UPCGComputeKernel::GetInputPinsAndOverridablePins(TArray<FPCGPinProperties>& OutPins) const
+{
+	GetInputPins(OutPins);
+	GetOverridableParamPins(OutPins);
+}
+
 #if WITH_EDITOR
 bool UPCGComputeKernel::PerformStaticValidation()
 {
@@ -541,10 +571,8 @@ bool UPCGComputeKernel::ValidatePCGNode(TArray<FPCGKernelLogEntry>& OutLogEntrie
 #if PCG_KERNEL_LOGGING_ENABLED
 	if (PCGComputeKernel::CVarWarnOnOverridePinUsage.GetValueOnAnyThread())
 	{
-		// @todo_pcg: There are three classes of properties that still do not support overrides:
-		// 1. Properties which must be known at dispatch time (e.g. UPCGCopyPointsSettings::bCopyEachSourceToEveryTarget), so to support them we will need injected graph splits when they are overridden.
-		// 2. Properties which must be known at compile time (e.g. UPCGCopyPointsSettings::bMatchBasedOnAttribute), which would require compiling two versions of the kernel graph and branching dynamically at execution time.
-		// 3. Generic 'Overrides' pin available on all nodes with override pins.
+		// Note: Properties which must be known at compile time (e.g. UPCGCopyPointsSettings::bMatchBasedOnAttribute), are not supported.
+		// They would require compiling two versions of the kernel graph and branching dynamically at execution time.
 		// Note, nothing currently validates that params of this nature are not overridden, so implement new overrides with caution.
 
 		const UPCGNode* Node = Cast<UPCGNode>(Settings->GetOuter());
@@ -559,19 +587,14 @@ bool UPCGComputeKernel::ValidatePCGNode(TArray<FPCGKernelLogEntry>& OutLogEntrie
 				EPCGKernelLogVerbosity::Warning);
 		}
 
-		for (TFieldIterator<FProperty> PropIt(Settings->GetClass()); PropIt; ++PropIt)
+		for (const FPCGSettingsOverridableParam& OverridableParam : Settings->OverridableParams())
 		{
-			const FName PropertyName = (*PropIt)->GetFName();
-
-			if (!Settings->OverridableParamsForGPU().ContainsByPredicate([PropertyName](const FPCGSettingsOverridableParam& OverridableParam) { return OverridableParam.Label == PropertyName; }))
+			if (!OverridableParam.bSupportsGPU && Settings->IsPropertyOverriddenByPin(OverridableParam.PropertiesNames))
 			{
-				if (Settings->IsPropertyOverriddenByPin(PropertyName))
-				{
-					OutLogEntries.Emplace(FText::Format(
-						LOCTEXT("UnsupportedOverrideError", "Tried to override parameter '{0}', but this parameter does not support overrides when executing on the GPU at this time."),
-						FText::FromName(PropertyName)),
-						EPCGKernelLogVerbosity::Warning);
-				}
+				OutLogEntries.Emplace(FText::Format(
+					LOCTEXT("UnsupportedOverrideError", "Tried to override parameter '{0}', but this parameter does not support overrides when executing on the GPU at this time."),
+					FText::FromName(OverridableParam.Label)),
+					EPCGKernelLogVerbosity::Warning);
 			}
 		}
 	}
@@ -618,17 +641,14 @@ bool UPCGComputeKernel::AreInputEdgesValid(TArray<FPCGKernelLogEntry>& OutLogEnt
 }
 #endif
 
-void UPCGComputeKernel::FillOverridableParamsPins(TArray<FPCGPinProperties>& OutPins) const
+void UPCGComputeKernel::GetOverridableParamPins(TArray<FPCGPinProperties>& OutPins) const
 {
-	if (const UPCGSettings* Settings = GetSettings())
+	for (const FPCGKernelOverridableParam& OverridableParam : CachedOverridableParams)
 	{
-		for (const FPCGSettingsOverridableParam& OverridableParam : Settings->OverridableParamsForGPU())
+		if (OverridableParam.bIsPropertyOverriddenByPin)
 		{
-			if (Settings->IsPropertyOverriddenByPin(OverridableParam.Label))
-			{
-				FPCGPinProperties& ParamPin = OutPins.Emplace_GetRef(OverridableParam.Label, EPCGDataType::Param, /*bInAllowMultipleConnections=*/false, /*bAllowMultipleData=*/false);
-				ParamPin.SetOverrideOrUserParamPin();
-			}
+			FPCGPinProperties& ParamPin = OutPins.Emplace_GetRef(OverridableParam.Label, EPCGDataType::Param, /*bInAllowMultipleConnections=*/false, /*bAllowMultipleData=*/false);
+			ParamPin.SetOverrideOrUserParamPin();
 		}
 	}
 }
diff --git a/Engine/Plugins/PCG/Source/PCG/Private/Compute/PCGComputeKernel.h b/Engine/Plugins/PCG/Source/PCG/Private/Compute/PCGComputeKernel.h
--- a/Engine/Plugins/PCG/Source/PCG/Private/Compute/PCGComputeKernel.h
+++ b/Engine/Plugins/PCG/Source/PCG/Private/Compute/PCGComputeKernel.h
@@ -2,9 +2,10 @@
 
 #pragma once
 
-#include "ComputeFramework/ComputeKernel.h"
+#include "PCGSettings.h"
 #include "Compute/PCGComputeCommon.h"
 
+#include "ComputeFramework/ComputeKernel.h"
 #include "UObject/ObjectKey.h"
 
 #include "PCGComputeKernel.generated.h"
@@ -31,11 +32,6 @@ enum class EPCGComputeKernelFlags : uint32
 	LogDataDescriptions = 1 << 1,
 	/** Dispatch compute graph containing this kernel every frame to allow profiling. Has no effect unless PCG_GPU_KERNEL_PROFILING defined. */
 	RepeatDispatchForProfiling = 1 << 2,
-	/**
-	 * Indicates that this kernel will use parameters which are marked as OverridableGPU on the settings.
-	 * This is necessary to correctly wire up the kernel and provide code generated functions for reading the overridden values.
-	 */
-	RequiresOverridableParams = 1 << 3,
 };
 ENUM_CLASS_FLAGS(EPCGComputeKernelFlags);
 
@@ -75,6 +71,26 @@ struct FPCGComputeKernelParams
 };
 #endif
 
+/** Represents a parameter that can be overridden on a kernel. */
+USTRUCT()
+struct FPCGKernelOverridableParam
+{
+	GENERATED_BODY()
+
+	UPROPERTY()
+	FName Label = NAME_None;
+
+	UPROPERTY()
+	TArray<FName> PropertiesNames;
+
+	UPROPERTY()
+	bool bRequiresGPUReadback = false;
+
+	/** True if this param is actually overidden by a pin. */
+	UPROPERTY()
+	bool bIsPropertyOverriddenByPin = false;
+};
+
 UCLASS(Abstract)
 class UPCGComputeKernel : public UComputeKernel
 {
@@ -122,6 +138,9 @@ public:
 	/** Create additional output data interfaces to marshal any required output data. */
 	virtual void CreateAdditionalOutputDataInterfaces(FPCGGPUCompilationContext& InOutContext, UObject* InObjectOuter, TArray<TObjectPtr<UComputeDataInterface>>& OutDataInterfaces) const {}
 
+	/** Whether compute graphs should be split at this input pin. Useful for overridable parameters. */
+	virtual bool SplitGraphAtInput(FName InPinLabel) const;
+
 	/** Whether compute graphs should be split at the output of this node. Useful for reading back runtime statistics or diagnostic info before continuing. */
 	virtual bool SplitGraphAtOutput() const { return false; }
 
@@ -168,12 +187,15 @@ public:
 	virtual void GetInputPins(TArray<FPCGPinProperties>& OutPins) const {}
 	virtual void GetOutputPins(TArray<FPCGPinPropertiesGPU>& OutPins) const {}
 
+	/** Return the concatenation of GetInputPins() and GetOverridableParamPins(). */
+	void GetInputPinsAndOverridablePins(TArray<FPCGPinProperties>& OutPins) const;
+
+	const TArray<FPCGKernelOverridableParam>& GetCachedOverridableParams() const { return CachedOverridableParams; }
+
 	bool GetLogDataDescriptions() const { return !!(PCGKernelFlags & EPCGComputeKernelFlags::LogDataDescriptions); }
 
 	bool HasStaticValidationErrors() const { return !!(PCGKernelFlags & EPCGComputeKernelFlags::HasStaticValidationErrors); }
 
-	bool RequiresOverridableParams() const { return !!(PCGKernelFlags & EPCGComputeKernelFlags::RequiresOverridableParams); }
-
 #ifdef PCG_GPU_KERNEL_PROFILING
 	void EnableRepeatDispatch() { PCGKernelFlags |= EPCGComputeKernelFlags::RepeatDispatchForProfiling; }
 
@@ -199,7 +221,7 @@ protected:
 #endif
 
 	/** Collects all param pins which are overridable on the GPU. */
-	void FillOverridableParamsPins(TArray<FPCGPinProperties>& OutPins) const;
+	void GetOverridableParamPins(TArray<FPCGPinProperties>& OutPins) const;
 
 protected:
 	UPROPERTY()
@@ -220,6 +242,10 @@ protected:
 	UPROPERTY()
 	TArray<FPCGKernelLogEntry> StaticLogEntries;
 
+	/** Overridable params supported by this kernel, cached during compilation. */
+	UPROPERTY()
+	TArray<FPCGKernelOverridableParam> CachedOverridableParams;
+
 #if WITH_EDITOR
 	TSet<FName> InternalPinLabels;
 #endif
diff --git a/Engine/Plugins/PCG/Source/PCG/Private/Compute/PCGDataBinding.cpp b/Engine/Plugins/PCG/Source/PCG/Private/Compute/PCGDataBinding.cpp
--- a/Engine/Plugins/PCG/Source/PCG/Private/Compute/PCGDataBinding.cpp
+++ b/Engine/Plugins/PCG/Source/PCG/Private/Compute/PCGDataBinding.cpp
@@ -5,7 +5,6 @@
 #include "PCGContext.h"
 #include "PCGNode.h"
 #include "PCGParamData.h"
-#include "PCGSettings.h"
 #include "Compute/PCGComputeCommon.h"
 #include "Compute/PCGComputeGraph.h"
 #include "Compute/PCGComputeKernel.h"
@@ -63,6 +62,7 @@ void UPCGDataBinding::InitializeTables(FPCGContext* InContext)
 
 void UPCGDataBinding::ReleaseTransientResources()
 {
+	KernelParamsCache.Reset();
 	SourceBufferAttributeToGraphAttributeIndex.Reset();
 	GraphAttributeToSourceBufferAttributeIndex.Reset();
 	SourceBufferToGraphStringKey.Reset();
@@ -320,6 +320,21 @@ bool UPCGDataBinding::ReadbackInputDataToCPU(int32 InInputDataIndex)
 	return true;
 }
 
+bool UPCGDataBinding::InitializeKernelParams(FPCGContext* InContext)
+{
+	return KernelParamsCache.Initialize(InContext, this);
+}
+
+const FPCGKernelParams* UPCGDataBinding::GetCachedKernelParams(const UPCGComputeKernel* InKernel) const
+{
+	return InKernel ? GetCachedKernelParams(InKernel->GetKernelIndex()) : nullptr;
+}
+
+const FPCGKernelParams* UPCGDataBinding::GetCachedKernelParams(int32 InKernelIndex) const
+{
+	return KernelParamsCache.GetCachedKernelParams(InKernelIndex);
+}
+
 void UPCGDataBinding::AddInputDataAttributesToTable()
 {
 	for (const FPCGTaggedData& Data : InputDataCollection.TaggedData)
diff --git a/Engine/Plugins/PCG/Source/PCG/Private/Compute/PCGDataBinding.h b/Engine/Plugins/PCG/Source/PCG/Private/Compute/PCGDataBinding.h
--- a/Engine/Plugins/PCG/Source/PCG/Private/Compute/PCGDataBinding.h
+++ b/Engine/Plugins/PCG/Source/PCG/Private/Compute/PCGDataBinding.h
@@ -3,7 +3,9 @@
 #pragma once
 
 #include "PCGData.h"
+#include "PCGSettings.h"
 #include "Compute/PCGDataForGPU.h"
+#include "Compute/PCGKernelParamsCache.h"
 
 #include "PCGDataBinding.generated.h"
 
@@ -148,6 +150,15 @@ public:
 
 	void DebugLogDataDescriptions() const;
 
+	/** Initialize the kernel params cache. Performs readbacks of overridable params that require it, and caches all values for quick access from kernels and data providers.
+	* Returns false while readback is in progress.
+	*/
+	bool InitializeKernelParams(FPCGContext* InContext);
+
+	/** Look up the cached parameter struct for a kernel. */
+	const FPCGKernelParams* GetCachedKernelParams(const UPCGComputeKernel* InKernel) const;
+	const FPCGKernelParams* GetCachedKernelParams(int32 InKernelIndex) const;
+
 private:
 	/** Loop over all metadata attributes in all input data ensure all attributes are registered in attribute table. */
 	void AddInputDataAttributesToTable();
@@ -207,4 +218,7 @@ protected:
 	std::atomic<bool> bIsDataDescriptionCachePrimed = false;
 
 	std::atomic<bool> bTablesInitialized = false;
+
+	/** Cache of overridden settings/values per-kernel, built during compute graph execution. */
+	FPCGKernelParamsCache KernelParamsCache;
 };
diff --git a/Engine/Plugins/PCG/Source/PCG/Private/Compute/PCGKernelHelpers.h b/Engine/Plugins/PCG/Source/PCG/Private/Compute/PCGKernelHelpers.h
--- a/Engine/Plugins/PCG/Source/PCG/Private/Compute/PCGKernelHelpers.h
+++ b/Engine/Plugins/PCG/Source/PCG/Private/Compute/PCGKernelHelpers.h
@@ -43,6 +43,7 @@ namespace PCGKernelHelpers
 		FPCGComputeKernelParams KernelParams;
 		KernelParams.Settings = InParams.OwnerSettings;
 		KernelParams.bLogDescriptions = InParams.bDumpDataDescriptions;
+		KernelParams.bRequiresOverridableParams = InParams.bRequiresOverridableParams;
 
 		KernelType* Kernel = InCompilationContext.NewObject_AnyThread<KernelType>(InParams.ObjectOuter);
 
diff --git a/Engine/Plugins/PCG/Source/PCG/Private/Compute/PCGKernelParamsCache.cpp b/Engine/Plugins/PCG/Source/PCG/Private/Compute/PCGKernelParamsCache.cpp
new file mode 100644
--- /dev/null
+++ b/Engine/Plugins/PCG/Source/PCG/Private/Compute/PCGKernelParamsCache.cpp
@@ -0,0 +1,334 @@
+// Copyright Epic Games, Inc. All Rights Reserved.
+
+#include "Compute/PCGKernelParamsCache.h"
+
+#include "PCGContext.h"
+#include "PCGParamData.h"
+#include "PCGSettings.h"
+#include "Compute/PCGComputeGraph.h"
+#include "Compute/PCGDataBinding.h"
+#include "Metadata/Accessors/IPCGAttributeAccessor.h"
+#include "Metadata/Accessors/PCGAttributeAccessorHelpers.h"
+#include "Metadata/Accessors/PCGAttributeAccessorKeys.h"
+
+#define LOCTEXT_NAMESPACE "PCGKernelParamsCache"
+
+namespace PCGKernelParamsHelpers
+{
+	void OverrideKernelParamValue(FPCGContext* InContext, const UPCGSettings* InSettings, FInstancedPropertyBag& InOutKernelParams, const UStruct* KernelParamsStructType, FName InParamName, const FProperty* OverrideProperty, void* OverrideValuePtr, const UPCGData* InData)
+	{
+		TRACE_CPUPROFILER_EVENT_SCOPE(OverrideKernelParamValue);
+		check(KernelParamsStructType);
+		check(OverrideProperty);
+		check(OverrideValuePtr);
+
+		const UPCGParamData* OverrideData = Cast<UPCGParamData>(InData);
+
+		if (!OverrideData)
+		{
+			PCG_KERNEL_VALIDATION_ERR(InContext, InSettings, FText::Format(LOCTEXT("NoOverrideData", "Override pin '{0}' has no data, we can't override it."), FText::FromName(InParamName)));
+			return;
+		}
+
+		const UPCGMetadata* Metadata = InData ? InData->ConstMetadata() : nullptr;
+		const FPCGMetadataAttributeBase* AttributeBase = Metadata ? Metadata->GetConstAttribute(Metadata->GetLatestAttributeNameOrNone()) : nullptr;
+
+		if (!AttributeBase)
+		{
+			PCG_KERNEL_VALIDATION_ERR(InContext, InSettings, FText::Format(LOCTEXT("NoOverrideAttribute", "Data on override pin '{0}' has no attributes, we can't override it."), FText::FromName(InParamName)));
+			return;
+		}
+
+		auto LogUnsupportedAttributeTypeError = [InContext, InSettings, InParamName, OverrideProperty, AttributeBase]()
+		{
+			PCG_KERNEL_VALIDATION_ERR(
+				InContext,
+				InSettings,
+				FText::Format(LOCTEXT("UnsupportedAttributeType", "Tried to override pin '{0}', but attribute of type '{1}' does not support conversion to expected type '{2}'."),
+					FText::FromName(InParamName),
+					StaticEnum<EPCGMetadataTypes>()->GetDisplayNameTextByValue((int64)AttributeBase->GetTypeId()),
+					FText::FromString(OverrideProperty->GetClass()->GetName())));
+		};
+
+		if (const FNumericProperty* NumericProperty = CastField<FNumericProperty>(OverrideProperty))
+		{
+			if (NumericProperty->IsFloatingPoint())
+			{
+				double Value = 0.0f;
+
+				switch (AttributeBase->GetTypeId())
+				{
+				case PCG::Private::MetadataTypes<float>::Id:
+					Value = static_cast<const FPCGMetadataAttribute<int32>*>(AttributeBase)->GetValueFromItemKey(PCGFirstEntryKey);
+					break;
+				case PCG::Private::MetadataTypes<double>::Id:
+					Value = static_cast<const FPCGMetadataAttribute<int64>*>(AttributeBase)->GetValueFromItemKey(PCGFirstEntryKey);
+					break;
+				default:
+					LogUnsupportedAttributeTypeError();
+					return;
+				}
+
+				NumericProperty->SetFloatingPointPropertyValue(OverrideValuePtr, Value);
+			}
+			else
+			{
+				int64 Value = 0;
+
+				switch (AttributeBase->GetTypeId())
+				{
+				case PCG::Private::MetadataTypes<int32>::Id:
+					Value = static_cast<const FPCGMetadataAttribute<int32>*>(AttributeBase)->GetValueFromItemKey(PCGFirstEntryKey);
+					break;
+				case PCG::Private::MetadataTypes<int64>::Id:
+					Value = static_cast<const FPCGMetadataAttribute<int64>*>(AttributeBase)->GetValueFromItemKey(PCGFirstEntryKey);
+					break;
+				default:
+					LogUnsupportedAttributeTypeError();
+					return;
+				}
+
+				NumericProperty->SetIntPropertyValue(OverrideValuePtr, Value);
+			}
+		}
+		else if (const FBoolProperty* BoolProperty = CastField<FBoolProperty>(OverrideProperty))
+		{
+			bool Value = false;
+
+			switch (AttributeBase->GetTypeId())
+			{
+			case PCG::Private::MetadataTypes<int32>::Id:
+				Value = static_cast<const FPCGMetadataAttribute<int32>*>(AttributeBase)->GetValueFromItemKey(PCGFirstEntryKey) != 0;
+				break;
+			case PCG::Private::MetadataTypes<int64>::Id:
+				Value = static_cast<const FPCGMetadataAttribute<int64>*>(AttributeBase)->GetValueFromItemKey(PCGFirstEntryKey) != 0;
+				break;
+			case PCG::Private::MetadataTypes<bool>::Id:
+				Value = static_cast<const FPCGMetadataAttribute<bool>*>(AttributeBase)->GetValueFromItemKey(PCGFirstEntryKey);
+				break;
+			default:
+				LogUnsupportedAttributeTypeError();
+				return;
+			}
+
+			BoolProperty->SetPropertyValue(OverrideValuePtr, Value);
+		}
+		else
+		{
+			// @todo_pcg: Broaden type support.
+			PCG_KERNEL_VALIDATION_ERR(
+				InContext,
+				InSettings,
+				FText::Format(LOCTEXT("UnsupportedPropertyType", "Tried to override pin '{0}', but property of type '{1}' is not supported."),
+					FText::FromName(InParamName),
+					FText::FromString(OverrideProperty->GetClass()->GetName())));
+		}
+	}
+}
+
+bool FPCGKernelParams::GetValueBool(FName InParamName) const
+{
+	const FBoolProperty* Property = nullptr;
+	const void* ValuePtr = nullptr;
+
+	if (const FPropertyAndAddress* PropertyAndAddress = OverriddenProperties.Find(InParamName))
+	{
+		Property = CastField<FBoolProperty>(PropertyAndAddress->Key);
+		ValuePtr = PropertyAndAddress->Value;
+	}
+	else if (const UPCGSettings* Settings = WeakSettings.Get())
+	{
+		Property = CastField<FBoolProperty>(Settings->GetClass()->FindPropertyByName(InParamName));
+		ValuePtr = Property ? Property->ContainerPtrToValuePtr<void>(Settings) : nullptr;
+	}
+
+	return ensure(Property && ValuePtr) ? Property->GetPropertyValue(ValuePtr) : false;
+}
+
+FName FPCGKernelParams::GetValueName(FName InParamName) const
+{
+	const FNameProperty* Property = nullptr;
+	const void* ValuePtr = nullptr;
+
+	if (const FPropertyAndAddress* PropertyAndAddress = OverriddenProperties.Find(InParamName))
+	{
+		Property = CastField<FNameProperty>(PropertyAndAddress->Key);
+		ValuePtr = PropertyAndAddress->Value;
+	}
+	else if (const UPCGSettings* Settings = WeakSettings.Get())
+	{
+		Property = CastField<FNameProperty>(Settings->GetClass()->FindPropertyByName(InParamName));
+		ValuePtr = Property ? Property->ContainerPtrToValuePtr<void>(Settings) : nullptr;
+	}
+
+	return ensure(Property && ValuePtr) ? Property->GetPropertyValue(ValuePtr) : NAME_None;
+}
+
+bool FPCGKernelParamsCache::Initialize(FPCGContext* InContext, UPCGDataBinding* InDataBinding)
+{
+	TRACE_CPUPROFILER_EVENT_SCOPE(FPCGKernelParamsCache::Initialize);
+
+	if (!ensure(InDataBinding))
+	{
+		return true;
+	}
+
+	CacheOverridableParamDataIndices(InContext, InDataBinding);
+	
+	if (!ReadbackOverridableParamData(InDataBinding))
+	{
+		return false;
+	}
+
+	CreateAndCacheKernelParams(InContext, InDataBinding);
+
+	return true;
+}
+
+const FPCGKernelParams* FPCGKernelParamsCache::GetCachedKernelParams(int32 InKernelIndex) const
+{
+	const FPCGKernelParams* FoundKernelParams = KernelParams.Find(InKernelIndex);
+	return FoundKernelParams ? FoundKernelParams : nullptr;
+}
+
+void FPCGKernelParamsCache::Reset()
+{
+	KernelParams.Reset();
+	OverridableParamDataIndices.Reset();
+}
+
+void FPCGKernelParamsCache::CacheOverridableParamDataIndices(FPCGContext* InContext, const UPCGDataBinding* InDataBinding)
+{
+	TRACE_CPUPROFILER_EVENT_SCOPE(FPCGKernelParamsCache::CacheOverridableParamDataIndices);
+	check(InDataBinding);
+
+	if (OverridableParamDataIndices.IsEmpty())
+	{
+		if (!ensure(InDataBinding->Graph))
+		{
+			return;
+		}
+
+		for (int32 KernelIndex = 0; KernelIndex < InDataBinding->Graph->GetNumKernels(); ++KernelIndex)
+		{
+			const UPCGComputeKernel* Kernel = Cast<UPCGComputeKernel>(InDataBinding->Graph->GetKernel(KernelIndex));
+
+			if (ensure(Kernel))
+			{
+				for (const FPCGKernelOverridableParam& OverridableParam : Kernel->GetCachedOverridableParams())
+				{
+					if (OverridableParam.bIsPropertyOverriddenByPin && OverridableParam.bRequiresGPUReadback)
+					{
+						const int32 OverridableParamDataIndex = InDataBinding->GetFirstInputDataIndex(Kernel, OverridableParam.Label);
+
+						if (OverridableParamDataIndex != INDEX_NONE)
+						{
+							TArray<TPair<FPCGKernelOverridableParam, int32>>& OverridableParamDataIndicesArray = OverridableParamDataIndices.FindOrAdd(KernelIndex);
+							OverridableParamDataIndicesArray.Emplace(OverridableParam, OverridableParamDataIndex);
+						}
+						else
+						{
+							PCG_KERNEL_VALIDATION_WARN(InContext, Kernel->GetSettings(), FText::Format(LOCTEXT("NoOverrideDataFound", "No data found on connected override pin '{0}'."), FText::FromName(OverridableParam.Label)));
+						}
+					}
+				}
+			}
+		}
+	}
+}
+
+bool FPCGKernelParamsCache::ReadbackOverridableParamData(UPCGDataBinding* InDataBinding)
+{
+	TRACE_CPUPROFILER_EVENT_SCOPE(FPCGKernelParamsCache::ReadbackOverridableParamData);
+	check(InDataBinding);
+
+	// Readback any overridable params that require readback.
+	bool bAllReadBack = true;
+
+	for (const TPair<int32, TArray<TPair<FPCGKernelOverridableParam, int32>>>& OverridableParamDataIndicesArray : OverridableParamDataIndices)
+	{
+		for (const TPair<FPCGKernelOverridableParam, int32>& OverridableParam : OverridableParamDataIndicesArray.Value)
+		{
+			// Readback data - poll until readback complete.
+			// @todo_pcg: There could be a hot path for reading back overridable param values, since they are a single-attribute, single-entry attribute set of a known type.
+			if (!InDataBinding->ReadbackInputDataToCPU(OverridableParam.Value))
+			{
+				bAllReadBack = false;
+			}
+		}
+
+		if (!bAllReadBack)
+		{
+			return false;
+		}
+	}
+
+	return true;
+}
+
+void FPCGKernelParamsCache::CreateAndCacheKernelParams(FPCGContext* InContext, const UPCGDataBinding* InDataBinding)
+{
+	TRACE_CPUPROFILER_EVENT_SCOPE(FPCGKernelParamsCache::CreateAndCacheKernelParams);
+	check(InDataBinding);
+
+	if (!ensure(InDataBinding->Graph))
+	{
+		return;
+	}
+
+	for (int32 KernelIndex = 0; KernelIndex < InDataBinding->Graph->GetNumKernels(); ++KernelIndex)
+	{
+		const UPCGComputeKernel* Kernel = Cast<UPCGComputeKernel>(InDataBinding->Graph->GetKernel(KernelIndex));
+		const UPCGSettings* Settings = ensure(Kernel) ? Kernel->GetSettings() : nullptr;
+
+		if (!Settings)
+		{
+			continue;
+		}
+
+		FPCGKernelParams NewKernelParams;
+		NewKernelParams.WeakSettings = Settings;
+
+		// Apply any overrides that exist.
+		if (const TArray<TPair<FPCGKernelOverridableParam, int32>>* OverridableParams = OverridableParamDataIndices.Find(KernelIndex))
+		{
+			TRACE_CPUPROFILER_EVENT_SCOPE(OverrideKernelParamValues);
+
+			NewKernelParams.OverriddenValues = Settings->CreateEmptyPropertyBagInstance();
+
+			const UStruct* KernelParamsStructType = NewKernelParams.OverriddenValues.GetPropertyBagStruct();
+
+			if (!ensure(KernelParamsStructType))
+			{
+				continue;
+			}
+
+			for (const TPair<FPCGKernelOverridableParam, int32>& OverridableParamAndDataIndex : *OverridableParams)
+			{
+				const FPCGKernelOverridableParam& OverridableParam = OverridableParamAndDataIndex.Key;
+				const UPCGData* Data = InDataBinding->InputDataCollection.TaggedData[OverridableParamAndDataIndex.Value].Data;
+				const FProperty* OverrideProperty = KernelParamsStructType ? KernelParamsStructType->FindPropertyByName(OverridableParam.Label) : nullptr;
+
+				if (!OverrideProperty)
+				{
+					PCG_KERNEL_VALIDATION_ERR(InContext, Settings, FText::Format(LOCTEXT("OverridePropertyNotFound", "Overridden property '{0}' was not found in container '{1}'."), FText::FromName(OverridableParam.Label), FText::FromString(KernelParamsStructType->GetName())));
+					continue;
+				}
+
+				void* OverrideValuePtr = OverrideProperty->ContainerPtrToValuePtr<void>(NewKernelParams.OverriddenValues.GetMutableValue().GetMemory());
+
+				if (!ensure(OverrideValuePtr))
+				{
+					continue;
+				}
+
+				PCGKernelParamsHelpers::OverrideKernelParamValue(InContext, Settings, NewKernelParams.OverriddenValues, KernelParamsStructType, OverridableParam.Label, OverrideProperty, OverrideValuePtr, Data);
+				NewKernelParams.OverriddenProperties.FindOrAdd(OverridableParam.Label) = { OverrideProperty, OverrideValuePtr };
+			}
+		}
+
+		KernelParams.FindOrAdd(KernelIndex) = MoveTemp(NewKernelParams);
+	}
+}
+
+#undef LOCTEXT_NAMESPACE
diff --git a/Engine/Plugins/PCG/Source/PCG/Private/Compute/PCGKernelParamsCache.h b/Engine/Plugins/PCG/Source/PCG/Private/Compute/PCGKernelParamsCache.h
new file mode 100644
--- /dev/null
+++ b/Engine/Plugins/PCG/Source/PCG/Private/Compute/PCGKernelParamsCache.h
@@ -0,0 +1,91 @@
+// Copyright Epic Games, Inc. All Rights Reserved.
+
+#pragma once
+
+#include "Compute/PCGComputeKernel.h"
+
+#include "StructUtils/PropertyBag.h"
+
+class UPCGDataBinding;
+class UPCGSettings;
+struct FPCGContext;
+
+/** Parameters for compute kernel execution. */
+struct FPCGKernelParams
+{
+public:
+	bool GetValueBool(FName InParamName) const;
+	FName GetValueName(FName InParamName) const;
+
+	template<typename EnumValueType>
+	EnumValueType GetValueEnum(FName InParamName) const
+	{
+		static_assert(TIsEnum<EnumValueType>::Value, "Should only call this with enum types.");
+
+		EnumValueType Value{};
+
+		const FEnumProperty* Property = nullptr;
+		const void* ValuePtr = nullptr;
+
+		if (const FPropertyAndAddress* PropertyAndAddress = OverriddenProperties.Find(InParamName))
+		{
+			Property = CastField<FEnumProperty>(PropertyAndAddress->Key);
+			ValuePtr = PropertyAndAddress->Value;
+		}
+		else if (const UPCGSettings* Settings = WeakSettings.Get())
+		{
+			Property = CastField<FEnumProperty>(Settings->GetClass()->FindPropertyByName(InParamName));
+			ValuePtr = Property ? Property->ContainerPtrToValuePtr<void>(Settings) : nullptr;
+		}
+
+		const FNumericProperty* NumericProperty = Property ? Property->GetUnderlyingProperty() : nullptr;
+		return static_cast<EnumValueType>(ensure(NumericProperty && ValuePtr) ? NumericProperty->GetSignedIntPropertyValue(ValuePtr) : 0);
+	}
+
+private:
+	/** Optional underlying settings object to get default values that are not overridden. */
+	TWeakObjectPtr<const UPCGSettings> WeakSettings;
+
+	/** Property bag instance that mirrors the WeakSettings. Contains the actual overridden values. */
+	FInstancedPropertyBag OverriddenValues;
+
+	using FPropertyAndAddress = TPair<const FProperty*, const void*>;
+
+	/** Maps param name to its property and an address to read the overridden value. This will only be populated for params which are overidden and read back for CPU logic. */
+	TMap<FName, FPropertyAndAddress> OverriddenProperties;
+
+	friend struct FPCGKernelParamsCache;
+};
+
+/** Helper struct to cache parameters for compute kernels. */
+struct FPCGKernelParamsCache
+{
+public:
+	/** Performs a readback of any kernels with overridable parameters that require readback, and caches them for quick access from kernels and data providers.
+	* Returns false while readback is in progress.
+	*/
+	bool Initialize(FPCGContext* InContext, UPCGDataBinding* InDataBinding);
+
+	/** Look up the cached parameter struct for a kernel. */
+	const FPCGKernelParams* GetCachedKernelParams(int32 InKernelIndex) const;
+
+	/** Clear the cache. */
+	void Reset();
+
+private:
+	/** Cache readback indices for any overridable param data that require readback. */
+	void CacheOverridableParamDataIndices(FPCGContext* InContext, const UPCGDataBinding* InDataBinding);
+
+	/** Poll readback for overridable param data. Returns false while readback is in progress.  */
+	bool ReadbackOverridableParamData(UPCGDataBinding* InDataBinding);
+
+	/** Create and cache a property bag instance to store kernel param values and apply overrides if they exist. */
+	void CreateAndCacheKernelParams(FPCGContext* InContext, const UPCGDataBinding* InDataBinding);
+
+private:
+	/** Maps kernels to their cached kernel parameter values. */
+	TMap</*KernelIndex*/int32, FPCGKernelParams> KernelParams;
+
+	/** Maps kernels to their overridable params and an index into the DataBinding input data collection that holds the corresponding param data. */
+	TMap</*KernelIndex*/int32, TArray<TPair<FPCGKernelOverridableParam, int32>>> OverridableParamDataIndices;
+};
diff --git a/Engine/Plugins/PCG/Source/PCG/Private/Elements/PCGCopyPoints.cpp b/Engine/Plugins/PCG/Source/PCG/Private/Elements/PCGCopyPoints.cpp
--- a/Engine/Plugins/PCG/Source/PCG/Private/Elements/PCGCopyPoints.cpp
+++ b/Engine/Plugins/PCG/Source/PCG/Private/Elements/PCGCopyPoints.cpp
@@ -31,13 +31,11 @@ void UPCGCopyPointsSettings::CreateKernels(FPCGGPUCompilationContext& InOutConte
 	CreateParams.bRequiresOverridableParams = true;
 	CreateParams.NodeInputPinsToWire = { PCGCopyPointsConstants::SourcePointsLabel, PCGCopyPointsConstants::TargetPointsLabel };
 
-	UPCGComputeKernel* CopyKernel = PCGKernelHelpers::CreateKernel<UPCGCopyPointsKernel>(InOutContext, CreateParams, OutKernels, OutEdges);
+	UPCGCopyPointsKernel* CopyKernel = PCGKernelHelpers::CreateKernel<UPCGCopyPointsKernel>(InOutContext, CreateParams, OutKernels, OutEdges);
+	CopyKernel->SetMatchBasedOnAttribute(bMatchBasedOnAttribute);
 
 	if (bMatchBasedOnAttribute)
 	{
-		// Analysis kernel does not require any of the overridden properties used by the CopyPoints settings.
-		CreateParams.bRequiresOverridableParams = false;
-		
 		// Don't wire analysis kernel to node output pin, wire manually to the copy kernel below.
 		CreateParams.NodeOutputPinsToWire.Empty();
 
diff --git a/Engine/Plugins/PCG/Source/PCG/Private/Elements/PCGCopyPointsAnalysisKernel.cpp b/Engine/Plugins/PCG/Source/PCG/Private/Elements/PCGCopyPointsAnalysisKernel.cpp
--- a/Engine/Plugins/PCG/Source/PCG/Private/Elements/PCGCopyPointsAnalysisKernel.cpp
+++ b/Engine/Plugins/PCG/Source/PCG/Private/Elements/PCGCopyPointsAnalysisKernel.cpp
@@ -35,7 +35,7 @@ bool UPCGCopyPointsAnalysisKernel::IsKernelDataValid(FPCGContext* InContext) con
 		return true;
 	}
 
-	return PCGCopyPointsKernel::IsKernelDataValid(this, CastChecked<UPCGCopyPointsSettings>(GetSettings()), static_cast<FPCGComputeGraphContext*>(InContext));
+	return PCGCopyPointsKernel::IsKernelDataValid(this, static_cast<FPCGComputeGraphContext*>(InContext));
 }
 
 bool UPCGCopyPointsAnalysisKernel::ComputeOutputBindingDataDesc(const UPCGComputeGraph* InGraph, FName InOutputPinLabel, UPCGDataBinding* InBinding, FPCGDataCollectionDesc& OutDataDesc) const
@@ -43,6 +43,13 @@ bool UPCGCopyPointsAnalysisKernel::ComputeOutputBindingDataDesc(const UPCGComput
 	check(InGraph);
 	check(InBinding);
 
+	const FPCGKernelParams* KernelParams = InBinding->GetCachedKernelParams(this);
+
+	if (!ensure(KernelParams))
+	{
+		return false;
+	}
+
 	// Code assumes single output pin.
 	if (!ensure(InOutputPinLabel == PCGPinConstants::DefaultOutputLabel))
 	{
@@ -63,10 +70,10 @@ bool UPCGCopyPointsAnalysisKernel::ComputeOutputBindingDataDesc(const UPCGComput
 	// Output: a single attribute set with a single attribute true/false value per NumSource*NumTarget possible outputs.
 	FPCGDataCollectionDesc OutDesc;
 
-	const UPCGCopyPointsSettings* CopyPointSettings = CastChecked<UPCGCopyPointsSettings>(GetSettings());
+	const bool bCopyEachSourceOnEveryTarget = KernelParams->GetValueBool(GET_MEMBER_NAME_CHECKED(UPCGCopyPointsSettings, bCopyEachSourceOnEveryTarget));
 
 	// Validation should have ensured already that these are valid: (S, T) = (N, N), (N, 1) or (1, N).
-	const int32 ElementCount = CopyPointSettings->bCopyEachSourceOnEveryTarget ? (NumSourceDatas * NumTargetDatas) : FMath::Max(NumSourceDatas, NumTargetDatas);
+	const int32 ElementCount = bCopyEachSourceOnEveryTarget ? (NumSourceDatas * NumTargetDatas) : FMath::Max(NumSourceDatas, NumTargetDatas);
 	FPCGDataDesc& DataDesc = OutDesc.DataDescs.Emplace_GetRef(EPCGDataType::Param, ElementCount);
 
 	FPCGKernelAttributeKey SelectedFlagAttributeKey = FPCGKernelAttributeKey(PCGCopyPointsConstants::SelectedFlagAttributeName, EPCGKernelAttributeType::Bool);
diff --git a/Engine/Plugins/PCG/Source/PCG/Private/Elements/PCGCopyPointsKernel.cpp b/Engine/Plugins/PCG/Source/PCG/Private/Elements/PCGCopyPointsKernel.cpp
--- a/Engine/Plugins/PCG/Source/PCG/Private/Elements/PCGCopyPointsKernel.cpp
+++ b/Engine/Plugins/PCG/Source/PCG/Private/Elements/PCGCopyPointsKernel.cpp
@@ -36,7 +36,7 @@ bool UPCGCopyPointsKernel::IsKernelDataValid(FPCGContext* InContext) const
 		return true;
 	}
 
-	return PCGCopyPointsKernel::IsKernelDataValid(this, CastChecked<UPCGCopyPointsSettings>(GetSettings()), static_cast<FPCGComputeGraphContext*>(InContext));
+	return PCGCopyPointsKernel::IsKernelDataValid(this, static_cast<FPCGComputeGraphContext*>(InContext));
 }
 
 bool UPCGCopyPointsKernel::ComputeOutputBindingDataDesc(const UPCGComputeGraph* InGraph, FName InOutputPinLabel, UPCGDataBinding* InBinding, FPCGDataCollectionDesc& OutDataDesc) const
@@ -44,7 +44,12 @@ bool UPCGCopyPointsKernel::ComputeOutputBindingDataDesc(const UPCGComputeGraph*
 	check(InGraph);
 	check(InBinding);
 
-	const UPCGCopyPointsSettings* CPSettings = CastChecked<UPCGCopyPointsSettings>(GetSettings());
+	const FPCGKernelParams* KernelParams = InBinding->GetCachedKernelParams(this);
+
+	if (!ensure(KernelParams))
+	{
+		return false;
+	}
 
 	// Code assumes single output pin.
 	if (!ensure(InOutputPinLabel == PCGPinConstants::DefaultOutputLabel))
@@ -55,8 +60,8 @@ bool UPCGCopyPointsKernel::ComputeOutputBindingDataDesc(const UPCGComputeGraph*
 	// A graph split was injected before this kernel. We expect to find selected flags
 	int NumSelected = INDEX_NONE;
 	TArray<bool> Selected;
-	
-	if (CPSettings->bMatchBasedOnAttribute)
+
+	if (bMatchBasedOnAttribute)
 	{
 		const int32 AnalysisDataIndex = InBinding->GetFirstInputDataIndex(this, PCGCopyPointsConstants::SelectedFlagsPinLabel);
 
@@ -106,14 +111,19 @@ bool UPCGCopyPointsKernel::ComputeOutputBindingDataDesc(const UPCGComputeGraph*
 
 	FPCGDataCollectionDesc PinDesc;
 
+	const bool bCopyEachSourceOnEveryTarget = KernelParams->GetValueBool(GET_MEMBER_NAME_CHECKED(UPCGCopyPointsSettings, bCopyEachSourceOnEveryTarget));
+	const EPCGCopyPointsInheritanceMode ColorInheritance = KernelParams->GetValueEnum<EPCGCopyPointsInheritanceMode>(GET_MEMBER_NAME_CHECKED(UPCGCopyPointsSettings, ColorInheritance));
+	const EPCGCopyPointsInheritanceMode SeedInheritance = KernelParams->GetValueEnum<EPCGCopyPointsInheritanceMode>(GET_MEMBER_NAME_CHECKED(UPCGCopyPointsSettings, SeedInheritance));
+	const EPCGCopyPointsMetadataInheritanceMode AttributeInheritance = KernelParams->GetValueEnum<EPCGCopyPointsMetadataInheritanceMode>(GET_MEMBER_NAME_CHECKED(UPCGCopyPointsSettings, AttributeInheritance));
+	const EPCGCopyPointsTagInheritanceMode TagInheritance = KernelParams->GetValueEnum<EPCGCopyPointsTagInheritanceMode>(GET_MEMBER_NAME_CHECKED(UPCGCopyPointsSettings, TagInheritance));
 	const int32 NumSources = SourcePinDesc.DataDescs.Num();
 	const int32 NumTargets = TargetPinDesc.DataDescs.Num();
-	const int32 NumIterations = CPSettings->bCopyEachSourceOnEveryTarget ? NumSources * NumTargets : FMath::Max(NumSources, NumTargets);
+	const int32 NumIterations = bCopyEachSourceOnEveryTarget ? NumSources * NumTargets : FMath::Max(NumSources, NumTargets);
 
 	// Properties that might need to be allocated
 	const EPCGPointNativeProperties OptionalProperties = EPCGPointNativeProperties::Seed | EPCGPointNativeProperties::Color;
 
-	if (NumSources > 0 && NumTargets > 0 && (CPSettings->bCopyEachSourceOnEveryTarget || NumSources == NumTargets || NumSources == 1 || NumTargets == 1))
+	if (NumSources > 0 && NumTargets > 0 && (bCopyEachSourceOnEveryTarget || NumSources == NumTargets || NumSources == 1 || NumTargets == 1))
 	{
 		for (int32 I = 0; I < NumIterations; ++I)
 		{
@@ -122,15 +132,15 @@ bool UPCGCopyPointsKernel::ComputeOutputBindingDataDesc(const UPCGComputeGraph*
 				continue;
 			}
 
-			const int32 SourceIndex = CPSettings->bCopyEachSourceOnEveryTarget ? (I / NumTargets) : FMath::Min(I, NumSources - 1);
-			const int32 TargetIndex = CPSettings->bCopyEachSourceOnEveryTarget ? (I % NumTargets) : FMath::Min(I, NumTargets - 1);
+			const int32 SourceIndex = bCopyEachSourceOnEveryTarget ? (I / NumTargets) : FMath::Min(I, NumSources - 1);
+			const int32 TargetIndex = bCopyEachSourceOnEveryTarget ? (I % NumTargets) : FMath::Min(I, NumTargets - 1);
 
 			const FPCGDataDesc& SourceDesc = SourcePinDesc.DataDescs[SourceIndex];
 			const FPCGDataDesc& TargetDesc = TargetPinDesc.DataDescs[TargetIndex];
 
 			FPCGDataDesc& ResultDataDesc = PinDesc.DataDescs.Emplace_GetRef(EPCGDataType::Point, SourceDesc.GetElementCount() * TargetDesc.GetElementCount());
 
-			if (CPSettings->AttributeInheritance == EPCGCopyPointsMetadataInheritanceMode::SourceFirst)
+			if (AttributeInheritance == EPCGCopyPointsMetadataInheritanceMode::SourceFirst)
 			{
 				ResultDataDesc.AttributeDescs = SourceDesc.AttributeDescs;
 
@@ -139,7 +149,7 @@ bool UPCGCopyPointsKernel::ComputeOutputBindingDataDesc(const UPCGComputeGraph*
 					ResultDataDesc.AttributeDescs.AddUnique(AttrDesc);
 				}
 			}
-			else if (CPSettings->AttributeInheritance == EPCGCopyPointsMetadataInheritanceMode::TargetFirst)
+			else if (AttributeInheritance == EPCGCopyPointsMetadataInheritanceMode::TargetFirst)
 			{
 				ResultDataDesc.AttributeDescs = TargetDesc.AttributeDescs;
 
@@ -148,16 +158,16 @@ bool UPCGCopyPointsKernel::ComputeOutputBindingDataDesc(const UPCGComputeGraph*
 					ResultDataDesc.AttributeDescs.AddUnique(AttrDesc);
 				}
 			}
-			else if (CPSettings->AttributeInheritance == EPCGCopyPointsMetadataInheritanceMode::SourceOnly)
+			else if (AttributeInheritance == EPCGCopyPointsMetadataInheritanceMode::SourceOnly)
 			{
 				ResultDataDesc.AttributeDescs = SourceDesc.AttributeDescs;
 			}
-			else if (CPSettings->AttributeInheritance == EPCGCopyPointsMetadataInheritanceMode::TargetOnly)
+			else if (AttributeInheritance == EPCGCopyPointsMetadataInheritanceMode::TargetOnly)
 			{
 				ResultDataDesc.AttributeDescs = TargetDesc.AttributeDescs;
 			}
 
-			if (CPSettings->TagInheritance == EPCGCopyPointsTagInheritanceMode::Source || CPSettings->TagInheritance == EPCGCopyPointsTagInheritanceMode::Both)
+			if (TagInheritance == EPCGCopyPointsTagInheritanceMode::Source || TagInheritance == EPCGCopyPointsTagInheritanceMode::Both)
 			{
 				for (const int32 SourceTagStringKey : SourceDesc.TagStringKeys)
 				{
@@ -165,7 +175,7 @@ bool UPCGCopyPointsKernel::ComputeOutputBindingDataDesc(const UPCGComputeGraph*
 				}
 			}
 			
-			if (CPSettings->TagInheritance == EPCGCopyPointsTagInheritanceMode::Target || CPSettings->TagInheritance == EPCGCopyPointsTagInheritanceMode::Both)
+			if (TagInheritance == EPCGCopyPointsTagInheritanceMode::Target || TagInheritance == EPCGCopyPointsTagInheritanceMode::Both)
 			{
 				for (const int32 TargetTagStringKey : TargetDesc.TagStringKeys)
 				{
@@ -196,8 +206,8 @@ bool UPCGCopyPointsKernel::ComputeOutputBindingDataDesc(const UPCGComputeGraph*
 			};
 
 			PropertiesToAllocate |= EPCGPointNativeProperties::Transform;
-			PropertiesToAllocate |= GetPropertyToAllocate(CPSettings->SeedInheritance, EPCGPointNativeProperties::Seed);
-			PropertiesToAllocate |= GetPropertyToAllocate(CPSettings->ColorInheritance, EPCGPointNativeProperties::Color);
+			PropertiesToAllocate |= GetPropertyToAllocate(SeedInheritance, EPCGPointNativeProperties::Seed);
+			PropertiesToAllocate |= GetPropertyToAllocate(ColorInheritance, EPCGPointNativeProperties::Color);
 
 			ResultDataDesc.AllocateProperties(PropertiesToAllocate);
 		}
@@ -214,42 +224,6 @@ int UPCGCopyPointsKernel::ComputeThreadCount(const UPCGDataBinding* InBinding) c
 }
 
 #if WITH_EDITOR
-FString UPCGCopyPointsKernel::GetCookedSource(FPCGGPUCompilationContext& InOutContext) const
-{
-	const FString SourceFile = Super::GetCookedSource(InOutContext);
-
-	FString OverridableParamFunctions = "";
-
-	if (const UPCGSettings* Settings = GetSettings())
-	{
-		for (const FPCGSettingsOverridableParam& OverridableParam : Settings->OverridableParamsForGPU())
-		{
-			// @todo_pcg: Support overriding more than just Int/Uint. Should be able to get the type from the property.
-
-			if (Settings->IsPropertyOverriddenByPin(OverridableParam.Label))
-			{
-				OverridableParamFunctions += FString::Format(TEXT(
-					"uint {0}_GetOverridableValue()\n"
-					"{\n"
-					"	return (uint){0}_GetFirstAttributeAsInt();\n"
-					"}\n"
-					"\n"), { OverridableParam.Label.ToString() });
-			}
-			else
-			{
-				OverridableParamFunctions += FString::Format(TEXT(
-					"uint {0}_GetOverridableValue()\n"
-					"{\n"
-					"	return CopyPoints_Get{0}Internal();\n"
-					"}\n"
-					"\n"), { OverridableParam.Label.ToString() });
-			}
-		}
-	}
-
-	return OverridableParamFunctions + SourceFile;
-}
-
 void UPCGCopyPointsKernel::CreateAdditionalInputDataInterfaces(FPCGGPUCompilationContext& InOutContext, UObject* InObjectOuter, TArray<TObjectPtr<UComputeDataInterface>>& OutDataInterfaces) const
 {
 	Super::CreateAdditionalInputDataInterfaces(InOutContext, InObjectOuter, OutDataInterfaces);
@@ -263,18 +237,14 @@ void UPCGCopyPointsKernel::CreateAdditionalInputDataInterfaces(FPCGGPUCompilatio
 
 void UPCGCopyPointsKernel::GetInputPins(TArray<FPCGPinProperties>& OutPins) const
 {
-	const UPCGCopyPointsSettings* CPSettings = CastChecked<UPCGCopyPointsSettings>(GetSettings());
-
 	OutPins.Emplace(PCGCopyPointsConstants::SourcePointsLabel, EPCGDataType::Point);
 	OutPins.Emplace(PCGCopyPointsConstants::TargetPointsLabel, EPCGDataType::Point);
 
-	if (CPSettings->bMatchBasedOnAttribute)
+	if (bMatchBasedOnAttribute)
 	{
 		// One bool flag per output data that signals if the output data should be computed or not.
 		OutPins.Emplace(PCGCopyPointsConstants::SelectedFlagsPinLabel, EPCGDataType::Param);
 	}
-
-	FillOverridableParamsPins(OutPins);
 }
 
 void UPCGCopyPointsKernel::GetOutputPins(TArray<FPCGPinPropertiesGPU>& OutPins) const
diff --git a/Engine/Plugins/PCG/Source/PCG/Private/Elements/PCGCopyPointsKernelShared.cpp b/Engine/Plugins/PCG/Source/PCG/Private/Elements/PCGCopyPointsKernelShared.cpp
--- a/Engine/Plugins/PCG/Source/PCG/Private/Elements/PCGCopyPointsKernelShared.cpp
+++ b/Engine/Plugins/PCG/Source/PCG/Private/Elements/PCGCopyPointsKernelShared.cpp
@@ -11,9 +11,8 @@
 
 #define LOCTEXT_NAMESPACE "PCGCopyPointsKernel"
 
-bool PCGCopyPointsKernel::IsKernelDataValid(const UPCGComputeKernel* InKernel, const UPCGCopyPointsSettings* InCopyPointSettings, const FPCGComputeGraphContext* InContext)
+bool PCGCopyPointsKernel::IsKernelDataValid(const UPCGComputeKernel* InKernel, const FPCGComputeGraphContext* InContext)
 {
-	check(InCopyPointSettings);
 	check(InContext);
 	check(InKernel);
 
@@ -33,22 +32,34 @@ bool PCGCopyPointsKernel::IsKernelDataValid(const UPCGComputeKernel* InKernel, c
 			return false;
 		}
 
+		const UPCGSettings* Settings = InKernel->GetSettings();
+		const FPCGKernelParams* KernelParams = DataBinding->GetCachedKernelParams(InKernel);
+
+		if (!ensure(KernelParams))
+		{
+			return false;
+		}
+
+		const bool bCopyEachSourceOnEveryTarget = KernelParams->GetValueBool(GET_MEMBER_NAME_CHECKED(UPCGCopyPointsSettings, bCopyEachSourceOnEveryTarget));
+		const bool bMatchBasedOnAttribute = KernelParams->GetValueBool(GET_MEMBER_NAME_CHECKED(UPCGCopyPointsSettings, bMatchBasedOnAttribute));
+		const FName MatchAttribute = KernelParams->GetValueName(GET_MEMBER_NAME_CHECKED(UPCGCopyPointsSettings, MatchAttribute));
+
 		const int32 NumSources = SourcePinDesc->DataDescs.Num();
 		const int32 NumTargets = TargetPinDesc->DataDescs.Num();
 
-		if (!InCopyPointSettings->bCopyEachSourceOnEveryTarget && (NumSources != NumTargets && NumSources != 1 && NumTargets != 1 && NumSources != 0 && NumTargets != 0))
+		if (!bCopyEachSourceOnEveryTarget && (NumSources != NumTargets && NumSources != 1 && NumTargets != 1 && NumSources != 0 && NumTargets != 0))
 		{
-			PCG_KERNEL_VALIDATION_ERR(InContext, InCopyPointSettings, FText::Format(
+			PCG_KERNEL_VALIDATION_ERR(InContext, Settings, FText::Format(
 				LOCTEXT("NumDataMismatch", "Num Sources ({0}), mismatches with Num Targets ({1}). Only supports N:N, 1:N and N:1 operation."),
 				NumSources,
 				NumTargets));
 		}
 
-		if (InCopyPointSettings->bMatchBasedOnAttribute)
+		if (bMatchBasedOnAttribute)
 		{
-			auto ValidateAttributeExists = [InContext, InCopyPointSettings](const FPCGDataCollectionDesc* PinDataDesc) -> bool
+			auto ValidateAttributeExists = [InContext, Settings, MatchAttribute](const FPCGDataCollectionDesc* PinDataDesc) -> bool
 			{
-				const FName MatchAttributeName = InCopyPointSettings->MatchAttribute;
+				const FName MatchAttributeName = MatchAttribute;
 
 				// todo_pcg: Can generalize this to any type?
 				constexpr EPCGKernelAttributeType MatchAttributeType = EPCGKernelAttributeType::Int;
@@ -57,7 +68,7 @@ bool PCGCopyPointsKernel::IsKernelDataValid(const UPCGComputeKernel* InKernel, c
 				{
 					if (!DataDesc.ContainsAttribute(MatchAttributeName, MatchAttributeType))
 					{
-						PCG_KERNEL_VALIDATION_ERR(InContext, InCopyPointSettings, FText::Format(
+						PCG_KERNEL_VALIDATION_ERR(InContext, Settings, FText::Format(
 							LOCTEXT("MatchAttributeMissing", "Match attribute '{0}' not found, this attribute must be present on all input data, and be of type Integer."),
 							FText::FromName(MatchAttributeName)));
 
diff --git a/Engine/Plugins/PCG/Source/PCG/Private/Graph/PCGGraphCompilerGPU.cpp b/Engine/Plugins/PCG/Source/PCG/Private/Graph/PCGGraphCompilerGPU.cpp
--- a/Engine/Plugins/PCG/Source/PCG/Private/Graph/PCGGraphCompilerGPU.cpp
+++ b/Engine/Plugins/PCG/Source/PCG/Private/Graph/PCGGraphCompilerGPU.cpp
@@ -124,24 +124,52 @@ void FPCGGraphCompilerGPU::CollectGPUNodeSubsets(
 				}
 			};
 
-			// GPU task. Split at output if the kernel requests it.
+			// GPU task. Split at edges if the kernel requests it.
 			if (InGPUCompatibleTaskIds.Contains(InTaskId))
 			{
 				UPCGComputeKernel* const* FoundKernel = InTaskIdToKernel.Find(InTaskId);
-				if (FoundKernel && (*FoundKernel)->SplitGraphAtOutput())
+				UPCGComputeKernel* ProducerKernel = FoundKernel ? *FoundKernel : nullptr;
+
+				if (ProducerKernel)
 				{
-					if (const TArray<FPCGTaskId>* Successors = InTaskSuccessors.Find(InTaskId))
+					for (const FPCGGraphTaskInput& TaskInput : InCompiledTasks[InTaskId].Inputs)
 					{
-						for (FPCGTaskId SuccessorTaskId : *Successors)
+						if (!TaskInput.DownstreamPin.IsSet())
+						{
+							continue;
+						}
+
+						if (ProducerKernel->SplitGraphAtInput(TaskInput.DownstreamPin->Label))
 						{
-							if (InGPUCompatibleTaskIds.Contains(SuccessorTaskId))
+							// Don't need to split if the upstream task is already on a different GPU island.
+							if (InGPUCompatibleTaskIds.Contains(TaskInput.TaskId))
 							{
-								const uint32 SuccessorIsland = ConnectedGPUNodeIslandIDs[SuccessorTaskId];
+								const uint32 UpstreamIslandId = ConnectedGPUNodeIslandIDs[TaskInput.TaskId];
 
-								if (ConnectedGPUNodeIslandIDs[InTaskId] == SuccessorIsland)
+								if (ConnectedGPUNodeIslandIDs[InTaskId] == UpstreamIslandId)
 								{
 									// Propagate a new island ID to all downstream GPU tasks within the island.
-									PropagateIslandIDDownstream(SuccessorTaskId, SuccessorIsland, NextIslandId++, PropagateIslandIDDownstream);
+									PropagateIslandIDDownstream(InTaskId, ConnectedGPUNodeIslandIDs[InTaskId], NextIslandId++, PropagateIslandIDDownstream);
+								}
+							}
+						}
+					}
+
+					if (ProducerKernel->SplitGraphAtOutput())
+					{
+						if (const TArray<FPCGTaskId>* Successors = InTaskSuccessors.Find(InTaskId))
+						{
+							for (FPCGTaskId SuccessorTaskId : *Successors)
+							{
+								if (InGPUCompatibleTaskIds.Contains(SuccessorTaskId))
+								{
+									const uint32 SuccessorIsland = ConnectedGPUNodeIslandIDs[SuccessorTaskId];
+
+									if (ConnectedGPUNodeIslandIDs[InTaskId] == SuccessorIsland)
+									{
+										// Propagate a new island ID to all downstream GPU tasks within the island.
+										PropagateIslandIDDownstream(SuccessorTaskId, SuccessorIsland, NextIslandId++, PropagateIslandIDDownstream);
+									}
 								}
 							}
 						}
@@ -291,14 +319,11 @@ void FPCGGraphCompilerGPU::ExpandGPUNodeKernelsToTasks(
 		// Wire up overridden parameters for any kernels that require them.
 		for (UPCGComputeKernel* NodeKernel : NodeKernels)
 		{
-			if (NodeKernel->RequiresOverridableParams())
+			for (const FPCGKernelOverridableParam& OverridableParam : NodeKernel->GetCachedOverridableParams())
 			{
-				for (const FPCGSettingsOverridableParam& OverridableParam : Settings->OverridableParamsForGPU())
+				if (OverridableParam.bIsPropertyOverriddenByPin)
 				{
-					if (Settings->IsPropertyOverriddenByPin(OverridableParam.Label))
-					{
-						NodeKernelEdges.Emplace(FPCGPinReference(OverridableParam.Label), FPCGPinReference(NodeKernel, OverridableParam.Label));
-					}
+					NodeKernelEdges.Emplace(FPCGPinReference(OverridableParam.Label), FPCGPinReference(NodeKernel, OverridableParam.Label));
 				}
 			}
 		}
@@ -373,7 +398,7 @@ void FPCGGraphCompilerGPU::ExpandGPUNodeKernelsToTasks(
 
 			// Get kernel inputs in preparation for wiring.
 			KernelInputs.Reset();
-			Kernel->GetInputPins(KernelInputs);
+			Kernel->GetInputPinsAndOverridablePins(KernelInputs);
 
 			if (Settings->bTriggerRenderCapture)
 			{
@@ -1476,7 +1501,36 @@ void FPCGGraphCompilerGPU::CompileComputeGraph(
 			KernelSource->GroupSize = FIntVector(PCGComputeConstants::THREAD_GROUP_SIZE, 1, 1);
 
 			// All kernels require ComputeShaderUtils.ush, so inject that before anything else.
-			FString Source = TEXT("#include \"/Engine/Private/ComputeShaderUtils.ush\"\n\n") + KernelWithBindings.Kernel->GetCookedSource(InOutContext);
+			FString Source = TEXT("#include \"/Engine/Private/ComputeShaderUtils.ush\"\n\n");
+
+			// Generate parameter override functions.
+			for (const FPCGSettingsOverridableParam& OverridableParam : Settings->OverridableParams())
+			{
+				if (OverridableParam.bSupportsGPU)
+				{
+					// @todo_pcg: Broaden type support.
+					if (!OverridableParam.bRequiresGPUReadback && Settings->IsPropertyOverriddenByPin(OverridableParam.PropertiesNames))
+					{
+						Source += FString::Format(TEXT(
+							"uint {0}_GetOverridableValue()\n"
+							"{\n"
+							"	return {0}_GetFirstAttributeAsUint();\n"
+							"}\n"
+							"\n"), { OverridableParam.Label.ToString() });
+					}
+					else
+					{
+						Source += FString::Format(TEXT(
+							"uint {0}_GetOverridableValue()\n"
+							"{\n"
+							"	return Get{0}Internal();\n"
+							"}\n"
+							"\n"), { OverridableParam.Label.ToString() });
+					}
+				}
+			}
+ 
+			Source += KernelWithBindings.Kernel->GetCookedSource(InOutContext);
 			KernelSource->SetSource(MoveTemp(Source));
 			KernelWithBindings.Kernel->GatherAdditionalSources(KernelSource->AdditionalSources);
 
diff --git a/Engine/Plugins/PCG/Source/PCG/Private/Helpers/PCGSettingsHelpers.cpp b/Engine/Plugins/PCG/Source/PCG/Private/Helpers/PCGSettingsHelpers.cpp
--- a/Engine/Plugins/PCG/Source/PCG/Private/Helpers/PCGSettingsHelpers.cpp
+++ b/Engine/Plugins/PCG/Source/PCG/Private/Helpers/PCGSettingsHelpers.cpp
@@ -437,6 +437,8 @@ namespace PCGSettingsHelpers
 					}
 
 #if WITH_EDITOR
+					Param.bSupportsGPU = ChildParam.bSupportsGPU;
+					Param.bRequiresGPUReadback = ChildParam.bRequiresGPUReadback;
 					GatherAliases(Property, Param);
 					// For all already found aliases, add them to the current param, with the index incremented by 1.
 					for (TPair<int32, FPCGPropertyAliases>& It : ChildParam.MapOfAliases)
@@ -486,6 +488,8 @@ namespace PCGSettingsHelpers
 				Param.bHasNameClash = bHasNameClash;
 #if WITH_EDITOR
 				Param.Label = bHasNameClash ? FName(Param.GetPropertyPath()) : Label;
+				Param.bSupportsGPU = Property->HasMetaData(PCGObjectMetadata::OverridableCPUAndGPU) || Property->HasMetaData(PCGObjectMetadata::OverridableCPUAndGPUWithReadback);
+				Param.bRequiresGPUReadback = Property->HasMetaData(PCGObjectMetadata::OverridableCPUAndGPUWithReadback);
 				GatherAliases(Property, Param);
 #endif // WITH_EDITOR
 			}
diff --git a/Engine/Plugins/PCG/Source/PCG/Private/PCGContext.cpp b/Engine/Plugins/PCG/Source/PCG/Private/PCGContext.cpp
--- a/Engine/Plugins/PCG/Source/PCG/Private/PCGContext.cpp
+++ b/Engine/Plugins/PCG/Source/PCG/Private/PCGContext.cpp
@@ -36,14 +36,6 @@ TAutoConsoleVariable<bool> CVarPCGEnablePointArrayData(
 namespace PCGContextHelpers
 {
 	thread_local bool bIsInitializingSettings = false;
-
-	template <typename T>
-	bool GetOverrideParamValue(const IPCGAttributeAccessor& InAccessor, T& OutValue)
-	{
-		// Override were using the first entry (0) by default.
-		FPCGAttributeAccessorKeysEntries FirstEntry(PCGFirstEntryKey);
-		return InAccessor.Get<T>(OutValue, FirstEntry, EPCGAttributeAccessorFlags::AllowBroadcastAndConstructible);
-	}
 }
 
 FPCGContext::~FPCGContext()
@@ -252,7 +244,7 @@ void FPCGContext::InitializeSettings(bool bSkipPostLoad)
 						}
 
 						FSoftObjectPath ObjectPath;
-						if (PCGContextHelpers::GetOverrideParamValue(*AttributeAccessor, ObjectPath))
+						if (PCGAttributeAccessorHelpers::GetOverrideParamValue(*AttributeAccessor, ObjectPath))
 						{
 							FGCScopeGuard GCScope;
 							if (!ObjectPath.ResolveObject())
@@ -422,7 +414,7 @@ void FPCGContext::OverrideSettings()
 			FPCGAttributeAccessorKeysEntries FirstEntry(PCGFirstEntryKey);
 					
 			PropertyType Value{};
-			if (!PCGContextHelpers::GetOverrideParamValue(*AttributeAccessor, Value))
+			if (!PCGAttributeAccessorHelpers::GetOverrideParamValue(*AttributeAccessor, Value))
 			{
 				PCGE_LOG_C(Warning, GraphAndLog, this, FText::Format(LOCTEXT("ConversionFailed", "Parameter '{0}' cannot be converted from attribute '{1}'"), FText::FromName(Param.Label), FText::FromName(AttributeName)));
 				return false;
diff --git a/Engine/Plugins/PCG/Source/PCG/Private/PCGSettings.cpp b/Engine/Plugins/PCG/Source/PCG/Private/PCGSettings.cpp
--- a/Engine/Plugins/PCG/Source/PCG/Private/PCGSettings.cpp
+++ b/Engine/Plugins/PCG/Source/PCG/Private/PCGSettings.cpp
@@ -11,15 +11,11 @@
 #include "PCGPin.h"
 #include "PCGSubgraph.h"
 #include "PCGSubsystem.h"
-#include "Compute/PCGComputeCommon.h"
-#include "Compute/PCGComputeGraph.h"
-#include "Compute/PCGDataBinding.h"
-#include "Compute/PCGDataForGPU.h"
 #include "Elements/PCGAddTag.h"
 #include "Helpers/PCGHelpers.h"
 #include "Helpers/PCGSettingsHelpers.h"
-#include "Metadata/Accessors/PCGAttributeAccessorHelpers.h"
 #include "Metadata/Accessors/IPCGAttributeAccessor.h"
+#include "Metadata/Accessors/PCGAttributeAccessorHelpers.h"
 
 #include "Misc/Crc.h"
 #include "Serialization/ArchiveObjectCrc32.h"
@@ -282,6 +278,7 @@ void UPCGSettings::PostEditUndo()
 	// CachedOverridableParams was reset to previous value
 	// Therefore we need to rebuild the properties array since it is transient.
 	InitializeCachedOverridableParams();
+	InitializeOverriddenSettingsPropertyBag();
 
 	Super::PostEditUndo();
 }
@@ -307,6 +304,8 @@ void UPCGSettings::PostLoad()
 	InitializeCachedOverridableParams(/*bReset=*/false);
 #endif // WITH_EDITOR
 
+	InitializeOverriddenSettingsPropertyBag();
+
 #if WITH_EDITOR
 	if (ExecutionMode_DEPRECATED != EPCGSettingsExecutionMode::Enabled)
 	{
@@ -327,6 +326,8 @@ void UPCGSettings::PostInitProperties()
 	InitializeCachedOverridableParams();
 #endif //WITH_EDITOR
 
+	InitializeOverriddenSettingsPropertyBag();
+
 	// For new objects, use type hash as a seed so we don't get multiple nodes doing randomness in an identical way (we had a case where
 	// a Point Subset followed by a Density Noise produced uniform densities because the random seeds/streams were correllated).
 	// It is still possible to chain nodes of the same type (Point Subset -> Point Subset) and see correllated behaviour, in which
@@ -420,6 +421,16 @@ bool UPCGSettings::HasOverridableParam(FName InParamName) const
 	}) != nullptr;
 }
 
+FInstancedPropertyBag UPCGSettings::CreateEmptyPropertyBagInstance() const
+{
+	TRACE_CPUPROFILER_EVENT_SCOPE(UPCGSettings::CreateEmptyPropertyBagInstance);
+
+	FInstancedPropertyBag Instance;
+	Instance.InitializeFromBagStruct(OverriddenSettingsPropertyBag);
+
+	return Instance;
+}
+
 void UPCGSettings::AddExecutionDependencyPin(TArray<FPCGPinProperties>& OutPins) const
 {
 	if (HasExecutionDependencyPin())
@@ -730,6 +741,7 @@ void UPCGSettings::PostPaste()
 	// Params are not properly built when importing from pasted text (because param
 	// Properties are transient/unserialized), so we should reinitialize in post
 	InitializeCachedOverridableParams(/*bReset=*/true);
+	InitializeOverriddenSettingsPropertyBag();
 }
 
 bool UPCGSettings::CanUseExecutionDependencyPinForCulling() const
@@ -911,9 +923,11 @@ bool UPCGSettings::IsPinUsedByNodeExecution(const UPCGPin* InPin) const
 	{
 		auto MatchesPinLabel = [PinLabel = InPin->Properties.Label](const FPCGSettingsOverridableParam& OverridableParam) { return OverridableParam.Label == PinLabel; };
 
+		const FPCGSettingsOverridableParam* Param = OverridableParams().FindByPredicate(MatchesPinLabel);
+
 		// Pin is not in use if we are in GPU execution mode and the pin is an OverridableParam that doesn't support GPU overrides.
 		// Generic overrides pin is also not supported.
-		if ((OverridableParams().ContainsByPredicate(MatchesPinLabel) && !OverridableParamsForGPU().ContainsByPredicate(MatchesPinLabel)) || InPin->Properties.Label == PCGPinConstants::DefaultParamsLabel)
+		if ((Param && !Param->bSupportsGPU) || InPin->Properties.Label == PCGPinConstants::DefaultParamsLabel)
 		{
 			return false;
 		}
@@ -951,15 +965,8 @@ TArray<FPCGSettingsOverridableParam> UPCGSettings::GatherOverridableParams() con
 	PCGSettingsHelpers::FPCGGetAllOverridableParamsConfig Config;
 	Config.bUseSeed = UseSeed();
 	Config.IncludeMetadataValues.Add(PCGObjectMetadata::Overridable);
-	Config.ExcludeMetadataValues.Add(PCGObjectMetadata::NotOverridable);
-
-	return PCGSettingsHelpers::GetAllOverridableParams(GetClass(), Config);
-}
-
-TArray<FPCGSettingsOverridableParam> UPCGSettings::GatherOverridableParamsGPU() const
-{
-	PCGSettingsHelpers::FPCGGetAllOverridableParamsConfig Config;
-	Config.IncludeMetadataValues.Add(PCGObjectMetadata::OverridableGPU);
+	Config.IncludeMetadataValues.Add(PCGObjectMetadata::OverridableCPUAndGPU);
+	Config.IncludeMetadataValues.Add(PCGObjectMetadata::OverridableCPUAndGPUWithReadback);
 	Config.ExcludeMetadataValues.Add(PCGObjectMetadata::NotOverridable);
 
 	return PCGSettingsHelpers::GetAllOverridableParams(GetClass(), Config);
@@ -1021,11 +1028,6 @@ void UPCGSettings::InitializeCachedOverridableParams(bool bReset)
 			}
 		}
 	}
-
-	if (CachedOverridableParamsGPU.IsEmpty() || bReset)
-	{
-		CachedOverridableParamsGPU = GatherOverridableParamsGPU();
-	}
 #endif // WITH_EDITOR
 
 	// Reset the value
@@ -1101,6 +1103,28 @@ void UPCGSettings::InitializeCachedOverridableParams(bool bReset)
 	}
 }
 
+void UPCGSettings::InitializeOverriddenSettingsPropertyBag()
+{
+	TRACE_CPUPROFILER_EVENT_SCOPE(UPCGSettings::InitializeOverriddenSettingsPropertyBag);
+
+	const UClass* BaseClass = UPCGSettings::StaticClass();
+	const UClass* SettingsClass = GetClass();
+
+	TArray<FPropertyBagPropertyDesc> Descriptors;
+
+	for (TFieldIterator<FProperty> It(SettingsClass); It; ++It)
+	{
+		// Only include properties declared in classes *after* the abstract base class UPCGSettings (almost all properties in UPCGSettings are not interesting for override).
+		// @todo_pcg: We should probably whitelist some UPCGSettings properties (e.g. Seed).
+		if (It->GetOwnerClass()->IsChildOf(BaseClass) && It->GetOwnerClass() != BaseClass)
+		{
+			Descriptors.Emplace(It->GetFName(), *It);
+		}
+	}
+
+	OverriddenSettingsPropertyBag = UPropertyBag::GetOrCreateFromDescs(Descriptors);
+}
+
 void UPCGSettings::FixingOverridableParamPropertyClass(FPCGSettingsOverridableParam& Param) const
 {
 	if (!GetClass() || Param.PropertiesNames.IsEmpty())
diff --git a/Engine/Plugins/PCG/Source/PCG/Public/Elements/PCGCopyPoints.h b/Engine/Plugins/PCG/Source/PCG/Public/Elements/PCGCopyPoints.h
--- a/Engine/Plugins/PCG/Source/PCG/Public/Elements/PCGCopyPoints.h
+++ b/Engine/Plugins/PCG/Source/PCG/Public/Elements/PCGCopyPoints.h
@@ -64,39 +64,39 @@ protected:
 
 public:
 	/** The method used to determine output point rotation */
-	UPROPERTY(BlueprintReadWrite, EditAnywhere, Category = Settings, meta = (PCG_Overridable, PCG_OverridableGPU))
+	UPROPERTY(BlueprintReadWrite, EditAnywhere, Category = Settings, meta = (PCG_OverridableCPUAndGPU))
 	EPCGCopyPointsInheritanceMode RotationInheritance = EPCGCopyPointsInheritanceMode::Relative;
 
 	/** If this option is set, points will have their source position transformed using the target transform with rotation */
-	UPROPERTY(BlueprintReadWrite, EditAnywhere, Category = Settings, meta = (PCG_Overridable, PCG_OverridableGPU, EditCondition = "RotationInheritance==EPCGCopyPointsInheritanceMode::Source", EditConditionHides))
+	UPROPERTY(BlueprintReadWrite, EditAnywhere, Category = Settings, meta = (PCG_OverridableCPUAndGPU, EditCondition = "RotationInheritance==EPCGCopyPointsInheritanceMode::Source", EditConditionHides))
 	bool bApplyTargetRotationToPositions = true;
 
 	/** The method used to determine output point scale */
-	UPROPERTY(BlueprintReadWrite, EditAnywhere, Category = Settings, meta = (PCG_Overridable, PCG_OverridableGPU))
+	UPROPERTY(BlueprintReadWrite, EditAnywhere, Category = Settings, meta = (PCG_OverridableCPUAndGPU))
 	EPCGCopyPointsInheritanceMode ScaleInheritance = EPCGCopyPointsInheritanceMode::Relative;
 
 	/** If this option is set, points will have their source position transformed using the target transform with scale */
-	UPROPERTY(BlueprintReadWrite, EditAnywhere, Category = Settings, meta = (PCG_Overridable, PCG_OverridableGPU, EditCondition="ScaleInheritance==EPCGCopyPointsInheritanceMode::Source", EditConditionHides))
+	UPROPERTY(BlueprintReadWrite, EditAnywhere, Category = Settings, meta = (PCG_OverridableCPUAndGPU, EditCondition="ScaleInheritance==EPCGCopyPointsInheritanceMode::Source", EditConditionHides))
 	bool bApplyTargetScaleToPositions = true;
 
 	/** The method used to determine output point color */
-	UPROPERTY(BlueprintReadWrite, EditAnywhere, Category = Settings, meta = (PCG_Overridable))
+	UPROPERTY(BlueprintReadWrite, EditAnywhere, Category = Settings, meta = (PCG_OverridableCPUAndGPUWithReadback))
 	EPCGCopyPointsInheritanceMode ColorInheritance = EPCGCopyPointsInheritanceMode::Relative;
 
 	/** The method used to determine output seed values. Relative recomputes the seed from the new location. */
-	UPROPERTY(BlueprintReadWrite, EditAnywhere, Category = Settings, meta = (PCG_Overridable))
+	UPROPERTY(BlueprintReadWrite, EditAnywhere, Category = Settings, meta = (PCG_OverridableCPUAndGPUWithReadback))
 	EPCGCopyPointsInheritanceMode SeedInheritance = EPCGCopyPointsInheritanceMode::Relative;
 
 	/** The method used to determine output data attributes */
-	UPROPERTY(BlueprintReadWrite, EditAnywhere, Category = Settings, meta = (PCG_Overridable))
+	UPROPERTY(BlueprintReadWrite, EditAnywhere, Category = Settings, meta = (PCG_OverridableCPUAndGPUWithReadback))
 	EPCGCopyPointsMetadataInheritanceMode AttributeInheritance = EPCGCopyPointsMetadataInheritanceMode::SourceFirst;
 
 	/** The method used to determine the output data tags */
-	UPROPERTY(BlueprintReadWrite, EditAnywhere, Category = Settings, meta = (PCG_Overridable))
+	UPROPERTY(BlueprintReadWrite, EditAnywhere, Category = Settings, meta = (PCG_OverridableCPUAndGPUWithReadback))
 	EPCGCopyPointsTagInheritanceMode TagInheritance = EPCGCopyPointsTagInheritanceMode::Both;
 
 	/** If this option is set, each source point data will be copied to every target point data (cartesian product), producing N * M point data. Otherwise, will do a N:N (or N:1 or 1:N) operation, producing N point data. */
-	UPROPERTY(BlueprintReadWrite, EditAnywhere, Category = Settings, meta = (PCG_Overridable))
+	UPROPERTY(BlueprintReadWrite, EditAnywhere, Category = Settings, meta = (PCG_OverridableCPUAndGPUWithReadback))
 	bool bCopyEachSourceOnEveryTarget = true;
 
 	// Note: Not overridable currently. Used to make compile time decision about which kernels to emit. If this needs to be overridable then we
@@ -107,6 +107,7 @@ public:
 	UPROPERTY(BlueprintReadWrite, EditAnywhere, Category = Settings, AdvancedDisplay, meta = (EditCondition = "bExecuteOnGPU", EditConditionHides))
 	bool bMatchBasedOnAttribute = false;
 
+	// @todo_pcg: Could be overridden when Name overrides are supported.
 	/** Attribute that must be present on both the source data and the target point, and have the same value.
 	* EXPERIMENTAL: This feature is work in progress and is currently GPU-only. It may not be stable, and may be removed in the future.
 	*/
diff --git a/Engine/Plugins/PCG/Source/PCG/Public/Elements/PCGCopyPointsKernel.h b/Engine/Plugins/PCG/Source/PCG/Public/Elements/PCGCopyPointsKernel.h
--- a/Engine/Plugins/PCG/Source/PCG/Public/Elements/PCGCopyPointsKernel.h
+++ b/Engine/Plugins/PCG/Source/PCG/Public/Elements/PCGCopyPointsKernel.h
@@ -16,15 +16,22 @@ class UPCGCopyPointsKernel : public UPCGComputeKernel
 	GENERATED_BODY()
 
 public:
+	//~ Begin UPCGComputeKernel Interface
 	virtual bool IsKernelDataValid(FPCGContext* InContext) const override;
 	virtual bool ComputeOutputBindingDataDesc(const UPCGComputeGraph* InGraph, FName InOutputPinLabel, UPCGDataBinding* InBinding, FPCGDataCollectionDesc& OutDataDesc) const override;
 	virtual int ComputeThreadCount(const UPCGDataBinding* Binding) const override;
 #if WITH_EDITOR
-	virtual FString GetCookedSource(FPCGGPUCompilationContext& InOutContext) const override;
 	virtual const TCHAR* GetSourceFilePath() const override { return TEXT("/Plugin/PCG/Private/Elements/PCGCopyPoints.usf"); }
 	virtual FString GetEntryPoint() const override { return TEXT("PCGCopyPointsCS"); }
 	void CreateAdditionalInputDataInterfaces(FPCGGPUCompilationContext& InOutContext, UObject* InObjectOuter, TArray<TObjectPtr<UComputeDataInterface>>& OutDataInterfaces) const override;
 #endif
 	virtual void GetInputPins(TArray<FPCGPinProperties>& OutPins) const override;
 	virtual void GetOutputPins(TArray<FPCGPinPropertiesGPU>& OutPins) const override;
+	//~ End UPCGComputeKernel Interface
+
+	void SetMatchBasedOnAttribute(bool bInMatchBasedOnAttribute) { bMatchBasedOnAttribute = bInMatchBasedOnAttribute; }
+
+protected:
+	UPROPERTY()
+	bool bMatchBasedOnAttribute = false;
 };
diff --git a/Engine/Plugins/PCG/Source/PCG/Public/Elements/PCGCopyPointsKernelShared.h b/Engine/Plugins/PCG/Source/PCG/Public/Elements/PCGCopyPointsKernelShared.h
--- a/Engine/Plugins/PCG/Source/PCG/Public/Elements/PCGCopyPointsKernelShared.h
+++ b/Engine/Plugins/PCG/Source/PCG/Public/Elements/PCGCopyPointsKernelShared.h
@@ -3,12 +3,10 @@
 #pragma once
 
 class UPCGComputeKernel;
-class UPCGCopyPointsSettings;
 struct FPCGComputeGraphContext;
-struct FPCGContext;
 
 namespace PCGCopyPointsKernel
 {
 	/** Performs data validation common to all copy points kernels. */
-	bool IsKernelDataValid(const UPCGComputeKernel* InKernel, const UPCGCopyPointsSettings* InCopyPointSettings, const FPCGComputeGraphContext* InContext);
+	bool IsKernelDataValid(const UPCGComputeKernel* InKernel, const FPCGComputeGraphContext* InContext);
 }
diff --git a/Engine/Plugins/PCG/Source/PCG/Public/Metadata/Accessors/PCGAttributeAccessorHelpers.h b/Engine/Plugins/PCG/Source/PCG/Public/Metadata/Accessors/PCGAttributeAccessorHelpers.h
--- a/Engine/Plugins/PCG/Source/PCG/Public/Metadata/Accessors/PCGAttributeAccessorHelpers.h
+++ b/Engine/Plugins/PCG/Source/PCG/Public/Metadata/Accessors/PCGAttributeAccessorHelpers.h
@@ -348,4 +348,12 @@ namespace PCGAttributeAccessorHelpers
 
 		InArray = MoveTemp(SortedArray);
 	}
+
+	template <typename T>
+	bool GetOverrideParamValue(const IPCGAttributeAccessor& InAccessor, T& OutValue)
+	{
+		// Override were using the first entry (0) by default.
+		FPCGAttributeAccessorKeysEntries FirstEntry(PCGFirstEntryKey);
+		return InAccessor.Get<T>(OutValue, FirstEntry, EPCGAttributeAccessorFlags::AllowBroadcastAndConstructible);
+	}
 }
diff --git a/Engine/Plugins/PCG/Source/PCG/Public/PCGCommon.h b/Engine/Plugins/PCG/Source/PCG/Public/PCGCommon.h
--- a/Engine/Plugins/PCG/Source/PCG/Public/PCGCommon.h
+++ b/Engine/Plugins/PCG/Source/PCG/Public/PCGCommon.h
@@ -234,7 +234,8 @@ namespace Icons
 namespace PCGObjectMetadata
 {
 	const FName Overridable = TEXT("PCG_Overridable");
-	const FName OverridableGPU = TEXT("PCG_OverridableGPU");
+	const FName OverridableCPUAndGPU = TEXT("PCG_OverridableCPUAndGPU");
+	const FName OverridableCPUAndGPUWithReadback = TEXT("PCG_OverridableCPUAndGPUWithReadback");
 	const FName NotOverridable = TEXT("PCG_NotOverridable");
 	const FName OverridableChildProperties = TEXT("PCG_OverridableChildProperties");
 	const FName OverrideAliases = TEXT("PCG_OverrideAliases");
@@ -245,11 +246,14 @@ namespace PCGObjectMetadata
 	// Metadata usable in UPROPERTY for customizing the behavior when displaying the property in a property panel or graph node
 	enum
 	{
-		/// [PropertyMetadata] Indicates that the property is overridable by params.
+		/// [PropertyMetadata] Indicates that the property is overridable by params on the CPU.
 		PCG_Overridable,
 
-		/// [PropertyMetadata] Indicates that the property is overridable by params on the GPU.
-		PCG_OverridableGPU,
+		/// [PropertyMetadata] Indicates that the property is overridable by params on the CPU and GPU.
+		PCG_OverridableCPUAndGPU,
+
+		/// [PropertyMetadata] Indicates that the property is overridable by params on the CPU and GPU and the overridden value will be read back for use in CPU computation during compute graph dispatch.
+		PCG_OverridableCPUAndGPUWithReadback,
 
 		/// [PropertyMetadata] Indicates that the property is not-overridable by params. Used in structs to hide some parameters
 		PCG_NotOverridable,
diff --git a/Engine/Plugins/PCG/Source/PCG/Public/PCGSettings.h b/Engine/Plugins/PCG/Source/PCG/Public/PCGSettings.h
--- a/Engine/Plugins/PCG/Source/PCG/Public/PCGSettings.h
+++ b/Engine/Plugins/PCG/Source/PCG/Public/PCGSettings.h
@@ -22,6 +22,7 @@ class UPCGComputeKernel;
 class UPCGGraph;
 class UPCGNode;
 class UPCGSettings;
+class UPropertyBag;
 struct FPCGGPUCompilationContext;
 struct FPCGKernelEdge;
 struct FPCGPinProperties;
@@ -111,6 +112,16 @@ struct FPCGSettingsOverridableParam
 	UPROPERTY()
 	bool bHasNameClash = false;
 
+	/** Whether this overridable param is supported on the GPU. */
+	UPROPERTY()
+	bool bSupportsGPU = false;
+
+	/** Whether this overridable param requires a readback to CPU when overridden on the GPU. Necessary for params which can drive CPU logic during
+	 * ComputeGraph dispatch, for example any params that affect thread counts, buffer sizes, data descriptions, etc.
+	 */
+	UPROPERTY()
+	bool bRequiresGPUReadback = false;
+
 	bool HasAliases() const { return !MapOfAliases.IsEmpty(); }
 
 	// Transient
@@ -466,9 +477,6 @@ public:
 	/** List of all the overridable params available for these settings. */
 	virtual const TArray<FPCGSettingsOverridableParam>& OverridableParams() const { return CachedOverridableParams; }
 
-	/** List of all the overridable params available for these settings on the GPU. */
-	const TArray<FPCGSettingsOverridableParam>& OverridableParamsForGPU() const { return CachedOverridableParamsGPU; }
-
 	/** Check if we have some override. Can be overriden to force params pin for example */
 	virtual bool HasOverridableParams() const { return !CachedOverridableParams.IsEmpty(); }
 
@@ -478,6 +486,9 @@ public:
 	/** Check if we need to hook the output of the pre-task to this. One use is to compute overrides in the subgraph element and pass the overrides as data, to all nodes that needs it. */
 	virtual bool RequiresDataFromPreTask() const { return false; }
 
+	/** Creates an empty property bag instance from this settings. Values will be default initialized, not set from the settings instance. */
+	FInstancedPropertyBag CreateEmptyPropertyBagInstance() const;
+
 protected:
 	/** Iterate over OverridableParams to automatically add param pins to the list. */
 	UE_API void FillOverridableParamsPins(TArray<FPCGPinProperties>& OutPins) const;
@@ -485,7 +496,6 @@ protected:
 #if WITH_EDITOR
 	/** Can be overridden to add more custom params (like in BP). */
 	UE_API virtual TArray<FPCGSettingsOverridableParam> GatherOverridableParams() const;
-	TArray<FPCGSettingsOverridableParam> GatherOverridableParamsGPU() const;
 #endif // WITH_EDITOR
 
 	/** Called after intialization to construct the list of all overridable params. 
@@ -495,6 +505,10 @@ protected:
 	*/
 	UE_API void InitializeCachedOverridableParams(bool bReset = false);
 
+	// @todo_pcg: This could potentially be done once per class rather than once per class instance.
+	/** Called on load to cache the overridden settings property bag. */
+	void InitializeOverriddenSettingsPropertyBag();
+
 	/**
 	* There is a weird issue where the BP class is not set correctly in some Server cases.
 	* We can try to recover if the PropertyClass is null.
@@ -505,9 +519,8 @@ protected:
 	UPROPERTY(meta = (PCGNoHash))
 	TArray<FPCGSettingsOverridableParam> CachedOverridableParams;
 
-	/** Needs to be serialized since property metadata (used to populate this array) is not available at runtime. */
-	UPROPERTY()
-	TArray<FPCGSettingsOverridableParam> CachedOverridableParamsGPU;
+	/** Property bag used to create instances of overridden parameter structs for execution. Lifetime is managed by the reflection system. */
+	const UPropertyBag* OverriddenSettingsPropertyBag = nullptr;
 
 	// We need to make sure that if we have hard references that are overridable, and they are overriden by paths
 	// on objects that are not yet loaded, that we are loading it on the main thread.
