[PCG] Fix misleading comments and remove unacceptable word

#rnx

diff --git a/Engine/Plugins/PCG/Source/PCG/Private/Compute/DataInterfaces/Elements/PCGCopyPointsDataInterface.h b/Engine/Plugins/PCG/Source/PCG/Private/Compute/DataInterfaces/Elements/PCGCopyPointsDataInterface.h
--- a/Engine/Plugins/PCG/Source/PCG/Private/Compute/DataInterfaces/Elements/PCGCopyPointsDataInterface.h
+++ b/Engine/Plugins/PCG/Source/PCG/Private/Compute/DataInterfaces/Elements/PCGCopyPointsDataInterface.h
@@ -44,7 +44,7 @@ public:
 	//~ End UComputeDataProvider Interface
 
 public:
-	/** This property bag instance is owned by the KernelParamsCache, so we are not responsible for its lifetime. */
+	/** The kernel params are owned by the KernelParamsCache, so we are not responsible for its lifetime. */
 	const FPCGKernelParams* KernelParams = nullptr;
 
 	TArray<FUintVector2> SourceAndTargetDataIndices;
diff --git a/Engine/Plugins/PCG/Source/PCG/Private/PCGSettings.cpp b/Engine/Plugins/PCG/Source/PCG/Private/PCGSettings.cpp
--- a/Engine/Plugins/PCG/Source/PCG/Private/PCGSettings.cpp
+++ b/Engine/Plugins/PCG/Source/PCG/Private/PCGSettings.cpp
@@ -1115,7 +1115,7 @@ void UPCGSettings::InitializeOverriddenSettingsPropertyBag()
 	for (TFieldIterator<FProperty> It(SettingsClass); It; ++It)
 	{
 		// Only include properties declared in classes *after* the abstract base class UPCGSettings (almost all properties in UPCGSettings are not interesting for override).
-		// @todo_pcg: We should probably whitelist some UPCGSettings properties (e.g. Seed).
+		// @todo_pcg: We should probably allow some UPCGSettings properties (e.g. Seed).
 		if (It->GetOwnerClass()->IsChildOf(BaseClass) && It->GetOwnerClass() != BaseClass)
 		{
 			Descriptors.Emplace(It->GetFName(), *It);
