[PCG] Add new property for data that supports elements: [REDACTED_HANDLE].$NumElements.

* Useful to combine with Filter Data by Attribute
* Add concept of read-only for properties in the Selector menu to not show for output.
* Applied this concept to Polygon2D properties.

#rnx

diff --git a/Engine/Plugins/PCG/Source/PCG/Private/Elements/Metadata/PCGExtractAttribute.cpp b/Engine/Plugins/PCG/Source/PCG/Private/Elements/Metadata/PCGExtractAttribute.cpp
--- a/Engine/Plugins/PCG/Source/PCG/Private/Elements/Metadata/PCGExtractAttribute.cpp
+++ b/Engine/Plugins/PCG/Source/PCG/Private/Elements/Metadata/PCGExtractAttribute.cpp
@@ -48,9 +48,9 @@ void PCGExtractAttribute::ExtractAttribute(const PCGExtractAttribute::FExtractAt
 			PCGLog::LogErrorOnGraph(FText::Format(LOCTEXT("IndexOutOfBounds", "Index for input {0} is out of bounds. Index: {1}; Number of Elements: {2}"), InputIndex, Index, Keys->GetNum()), Params.Context);
 			continue;
 		}
-		
+
 		const FName OutputAttributeName = OutputTarget.GetName();
-		
+
 		UPCGParamData* OutputParamData = FPCGContext::NewObject_AnyThread<UPCGParamData>(Params.Context);
 
 		auto ExtractAttribute = [&Params, OutputAttributeName, &OutputParamData, &Accessor, &Index, &Keys]<typename AttributeType>(AttributeType) -> bool
@@ -104,7 +104,8 @@ FText UPCGExtractAttributeSettings::GetDefaultNodeTitle() const
 
 FText UPCGExtractAttributeSettings::GetNodeTooltipText() const
 {
-	return LOCTEXT("NodeTooltip", "Extract Attribute");
+	return LOCTEXT("NodeTooltip", "Extract an attribute at a given index into a new attribute set.\n"
+		"Support any domain. Index needs to be in range of valid indexes for the given domain.");
 }
 #endif // WITH_EDITOR
 
@@ -116,14 +117,14 @@ FPCGElementPtr UPCGExtractAttributeSettings::CreateElement() const
 TArray<FPCGPinProperties> UPCGExtractAttributeSettings::InputPinProperties() const
 {
 	TArray<FPCGPinProperties> Properties;
-	Properties.Emplace_GetRef(PCGPinConstants::DefaultInputLabel, EPCGDataType::Any).SetRequiredPin();	
+	Properties.Emplace_GetRef(PCGPinConstants::DefaultInputLabel, EPCGDataType::Any).SetRequiredPin();
 	return Properties;
 }
 
 TArray<FPCGPinProperties> UPCGExtractAttributeSettings::OutputPinProperties() const
 { 
 	TArray<FPCGPinProperties> Properties;
-	Properties.Emplace(PCGPinConstants::DefaultOutputLabel, EPCGDataType::Param);	
+	Properties.Emplace(PCGPinConstants::DefaultOutputLabel, EPCGDataType::Param);
 	return Properties;
 }
 
diff --git a/Engine/Plugins/PCG/Source/PCG/Private/Helpers/PCGHelpers.cpp b/Engine/Plugins/PCG/Source/PCG/Private/Helpers/PCGHelpers.cpp
--- a/Engine/Plugins/PCG/Source/PCG/Private/Helpers/PCGHelpers.cpp
+++ b/Engine/Plugins/PCG/Source/PCG/Private/Helpers/PCGHelpers.cpp
@@ -4,6 +4,7 @@
 
 #include "PCGComponent.h"
 #include "PCGContext.h"
+#include "PCGData.h"
 #include "PCGGraph.h"
 #include "PCGModule.h"
 #include "PCGSettings.h"
@@ -11,6 +12,7 @@
 #include "PCGWorldActor.h"
 #include "Grid/PCGPartitionActor.h"
 #include "Helpers/PCGActorHelpers.h"
+#include "Metadata/Accessors/PCGAttributeAccessorHelpers.h"
 
 #include "Components/BillboardComponent.h"
 #include "Landscape.h"
@@ -786,7 +788,21 @@ namespace PCGHelpers
 		}
 	}
 
+	int32 GetNumberOfElements(const UPCGData* InData)
+	{
+		if (!InData->IsSupportedMetadataDomainID(PCGMetadataDomainID::Elements))
+		{
+			return 0;
+		}
+		
+		FPCGAttributePropertyInputSelector ElementSelector;
+		ElementSelector.SetAttributeName(NAME_None);
+		InData->SetDomainFromDomainID(PCGMetadataDomainID::Elements, ElementSelector);
+		TUniquePtr<const IPCGAttributeAccessorKeys> ElementKeys = PCGAttributeAccessorHelpers::CreateConstKeys(InData, ElementSelector);
+		return ElementKeys ? ElementKeys->GetNum() : 0;
+	}
 #if UE_ENABLE_DEBUG_DRAWING
+
 	void DebugDrawGenerationVolume(FPCGContext* InContext, const FColor* InOverrideColor)
 	{
 		if (!InContext || !InContext->ExecutionSource.IsValid())
diff --git a/Engine/Plugins/PCG/Source/PCG/Private/Metadata/Accessors/PCGAttributeAccessorFactory.cpp b/Engine/Plugins/PCG/Source/PCG/Private/Metadata/Accessors/PCGAttributeAccessorFactory.cpp
--- a/Engine/Plugins/PCG/Source/PCG/Private/Metadata/Accessors/PCGAttributeAccessorFactory.cpp
+++ b/Engine/Plugins/PCG/Source/PCG/Private/Metadata/Accessors/PCGAttributeAccessorFactory.cpp
@@ -4,10 +4,12 @@
 
 #include "PCGModule.h"
 #include "Data/PCGSpatialData.h"
+#include "Helpers/PCGHelpers.h"
 #include "Metadata/PCGAttributePropertySelector.h"
 #include "Metadata/PCGMetadata.h"
 #include "Metadata/PCGMetadataAttribute.h"
 #include "Metadata/Accessors/PCGAttributeAccessorHelpers.h"
+#include "Metadata/Accessors/PCGCustomAccessor.h"
 
 namespace PCGAttributeAccessorFactory
 {
@@ -24,7 +26,7 @@ namespace PCGAttributeAccessorFactory
 			{
 				if (!bQuiet)
 				{
-					UE_LOG(LogPCG, Error, TEXT("[FPCGAttributeAccessorFactory::DefaultAccessor] Expected to select an extra property but the data doesn't support this property."));
+					UE_LOG(LogPCG, Error, TEXT("[FPCGAttributeAccessorFactory::DefaultAccessor] Expected to select an extra property '%s' but the data doesn't support it or is only constant."), *InSelector.GetName().ToString());
 				}
 				
 				return nullptr;
@@ -62,13 +64,28 @@ namespace PCGAttributeAccessorFactory
 	
 		if (InSelector.GetSelection() == EPCGAttributePropertySelection::ExtraProperty)
 		{
+			if (InSelector.GetExtraProperty() == EPCGExtraProperties::NumElements)
+			{
+				if (InSelector.GetDomainName() != PCGDataConstants::DataDomainName || !InData->IsSupportedMetadataDomainID(PCGMetadataDomainID::Elements))
+				{
+					if (!bQuiet)
+					{
+						UE_LOG(LogPCG, Error, TEXT("[FPCGAttributeAccessorFactory::DefaultConstAccessor] Num elements is only supported on @Data domain and for data that have elements."));
+					}
+					
+					return nullptr;
+				}
+				
+				return MakeUnique<FPCGConstantValueAccessor<int32>>(PCGHelpers::GetNumberOfElements(InData));
+			}
+			
 			Accessor = PCGAttributeAccessorHelpers::CreateExtraAccessor(InSelector.GetExtraProperty());
 
 			if (!Accessor.IsValid())
 			{
 				if (!bQuiet)
 				{
-					UE_LOG(LogPCG, Error, TEXT("[FPCGAttributeAccessorFactory::DefaultConstAccessor] Expected to select an extra property but the data doesn't support this property."));
+					UE_LOG(LogPCG, Error, TEXT("[FPCGAttributeAccessorFactory::DefaultAccessor] Expected to select an extra property '%s' but the data doesn't support it."), *InSelector.GetName().ToString());
 				}
 				
 				return nullptr;
@@ -121,6 +138,19 @@ namespace PCGAttributeAccessorFactory
 	
 	TUniquePtr<const IPCGAttributeAccessorKeys> CreateDefaultConstAccessorKeys(const UPCGData* InData, const FPCGAttributePropertySelector& InSelector, const bool bQuiet)
 	{
+		if (InSelector.GetSelection() == EPCGAttributePropertySelection::ExtraProperty && InSelector.GetExtraProperty() == EPCGExtraProperties::NumElements)
+		{
+			if (InSelector.GetDomainName() == PCGDataConstants::DataDomainName && InData->IsSupportedMetadataDomainID(PCGMetadataDomainID::Elements))
+			{
+				// Make a dummy key
+				return MakeUnique<const FPCGAttributeAccessorKeysEntries>(PCGInvalidEntryKey);
+			}
+			else
+			{
+				return nullptr;
+			}
+		}
+		
 		const UPCGMetadata* Metadata = InData->ConstMetadata();
 		if (!Metadata)
 		{
@@ -218,7 +248,9 @@ void FPCGAttributeAccessorMethods::FillSelectorMenuEntryFromEnum(const UEnum* En
 			Selector.SetDomainName(*MetadataDomain);
 		}
 
-		CurrentMenu->Entries.Emplace(std::move(EnumName), std::move(Tooltip), std::move(Selector));
+		const bool bReadOnly = EnumType->HasMetaData(*PCGObjectMetadata::PropertyReadOnly.ToString(), i);
+
+		CurrentMenu->Entries.Emplace(std::move(EnumName), std::move(Tooltip), std::move(Selector), bReadOnly);
 	}
 }
 #endif // WITH_EDITOR
diff --git a/Engine/Plugins/PCG/Source/PCG/Public/Data/PCGPolygon2DData.h b/Engine/Plugins/PCG/Source/PCG/Public/Data/PCGPolygon2DData.h
--- a/Engine/Plugins/PCG/Source/PCG/Public/Data/PCGPolygon2DData.h
+++ b/Engine/Plugins/PCG/Source/PCG/Public/Data/PCGPolygon2DData.h
@@ -15,27 +15,27 @@ UENUM()
 enum class EPCGPolygon2DProperties : uint8
 {
 	Position UMETA(Tooltip = "Location of the vertex in world coordinates."),
-	Rotation UMETA(Tooltip = "Rotation of the vertex in world coordinates."),
-	SegmentIndex UMETA(Tooltip = "Segment index of the vertex in the current polygon (outer or hole)"),
-	HoleIndex UMETA(Tooltip = "Index of the hole the vertex is inside of. -1 on the Outer polygon."),
-	SegmentLength UMETA(Tooltip = "Length of the segment starting from the vertex."),
+	Rotation UMETA(Tooltip = "Rotation of the vertex in world coordinates.", PCG_PropertyReadOnly),
+	SegmentIndex UMETA(Tooltip = "Segment index of the vertex in the current polygon (outer or hole)", PCG_PropertyReadOnly),
+	HoleIndex UMETA(Tooltip = "Index of the hole the vertex is inside of. -1 on the Outer polygon.", PCG_PropertyReadOnly),
+	SegmentLength UMETA(Tooltip = "Length of the segment starting from the vertex.", PCG_PropertyReadOnly),
 	LocalPosition UMETA(Tooltip = "Location of the vertex in 2d coordinates."),
-	LocalRotation UMETA(Tooltip = "Rotation of the vertex in 2d coordinates.")
+	LocalRotation UMETA(Tooltip = "Rotation of the vertex in 2d coordinates.", PCG_PropertyReadOnly)
 };
 
 UENUM()
 enum class EPCGPolygon2DDataProperties : uint8
 {
 	Transform UMETA(Tooltip = "Transform of the 2d polygon.", PCG_MetadataDomain="Data"),
-	Area UMETA(Tooltip = "Area of the 2d polygon.", PCG_MetadataDomain = "Data"),
-	Perimeter UMETA(Tooltip = "Perimeter of the 2d polygon.", PCG_MetadataDomain = "Data"),
-	BoundsMin UMETA(Tooltip = "Minimum point of the bounds of the 2d polygon in 2d space.", PCG_MetadataDomain = "Data"),
-	BoundsMax UMETA(Tooltip = "Minimum point of the bounds of the 2d polygon in 2d space.", PCG_MetadataDomain = "Data"),
-	SegmentCount UMETA(Tooltip = "Number of segments in the 2d polygon (includes hole segments).", PCG_MetadataDomain = "Data"),
-	OuterSegmentCount UMETA(Tooltip = "Number of segments in the 2d polygon (only the outer polygon).", PCG_MetadataDomain = "Data"),
-	HoleCount UMETA(Tooltip = "Number of holes in the 2d polygon.", PCG_MetadataDomain = "Data"),
-	LongestOuterSegmentIndex UMETA(Tooltip = "Index of the longest segment in the polygon.", PCG_MetadataDomain = "Data"),
-	IsClockwise UMETA(Tooltip = "Returns whether the 2d polygon is clockwise or counter-clocwise.", PCG_MetadataDomain = "Data")
+	Area UMETA(Tooltip = "Area of the 2d polygon.", PCG_MetadataDomain = "Data", PCG_PropertyReadOnly),
+	Perimeter UMETA(Tooltip = "Perimeter of the 2d polygon.", PCG_MetadataDomain = "Data", PCG_PropertyReadOnly),
+	BoundsMin UMETA(Tooltip = "Minimum point of the bounds of the 2d polygon in 2d space.", PCG_MetadataDomain = "Data", PCG_PropertyReadOnly),
+	BoundsMax UMETA(Tooltip = "Minimum point of the bounds of the 2d polygon in 2d space.", PCG_MetadataDomain = "Data", PCG_PropertyReadOnly),
+	SegmentCount UMETA(Tooltip = "Number of segments in the 2d polygon (includes hole segments).", PCG_MetadataDomain = "Data", PCG_PropertyReadOnly),
+	OuterSegmentCount UMETA(Tooltip = "Number of segments in the 2d polygon (only the outer polygon).", PCG_MetadataDomain = "Data", PCG_PropertyReadOnly),
+	HoleCount UMETA(Tooltip = "Number of holes in the 2d polygon.", PCG_MetadataDomain = "Data", PCG_PropertyReadOnly),
+	LongestOuterSegmentIndex UMETA(Tooltip = "Index of the longest segment in the polygon.", PCG_MetadataDomain = "Data", PCG_PropertyReadOnly),
+	IsClockwise UMETA(Tooltip = "Returns whether the 2d polygon is clockwise or counter-clocwise.", PCG_MetadataDomain = "Data", PCG_PropertyReadOnly)
 };
 
 namespace PCGPolygon2DData
diff --git a/Engine/Plugins/PCG/Source/PCG/Public/Helpers/PCGHelpers.h b/Engine/Plugins/PCG/Source/PCG/Public/Helpers/PCGHelpers.h
--- a/Engine/Plugins/PCG/Source/PCG/Public/Helpers/PCGHelpers.h
+++ b/Engine/Plugins/PCG/Source/PCG/Public/Helpers/PCGHelpers.h
@@ -20,6 +20,7 @@ class ALandscapeProxy;
 class IPCGGraphExecutionSource;
 class UPCGComponent;
 class UPCGGraph;
+class UPCGData;
 class UPCGMetadata;
 class UPCGSettings;
 class UWorld;
@@ -186,6 +187,10 @@ namespace PCGHelpers
 		}
 	}
 
+	// Create keys for a hypothetical attribute in the 'Elements' domain, get its count, and create a const value accessor for retrieving it
+	// Only available if the data supports the Elements domain 
+	PCG_API int32 GetNumberOfElements(const UPCGData* InData);
+
 #if UE_ENABLE_DEBUG_DRAWING 
 	PCG_API void DebugDrawGenerationVolume(FPCGContext* InContext, const FColor* InOverrideColor = nullptr);
 #endif
diff --git a/Engine/Plugins/PCG/Source/PCG/Public/Metadata/Accessors/PCGAttributeAccessorFactory.h b/Engine/Plugins/PCG/Source/PCG/Public/Metadata/Accessors/PCGAttributeAccessorFactory.h
--- a/Engine/Plugins/PCG/Source/PCG/Public/Metadata/Accessors/PCGAttributeAccessorFactory.h
+++ b/Engine/Plugins/PCG/Source/PCG/Public/Metadata/Accessors/PCGAttributeAccessorFactory.h
@@ -30,10 +30,11 @@ struct FPCGAttributeSelectorMenu
 
 struct FPCGAttributeSelectorMenuEntry
 {
-	FPCGAttributeSelectorMenuEntry(FText InLabel, FText InTooltip, FPCGAttributePropertySelector InSelector)
+	FPCGAttributeSelectorMenuEntry(FText InLabel, FText InTooltip, FPCGAttributePropertySelector InSelector, bool bInReadOnly = false)
 		: Label(std::move(InLabel))
 		, Tooltip(std::move(InTooltip))
 		, Selector(std::move(InSelector))
+		, bReadOnly(bInReadOnly)
 	{}
 
 	/** Label for the entry. */
@@ -44,6 +45,9 @@ struct FPCGAttributeSelectorMenuEntry
 
 	/** Selector to copy if this entry is selected. */
 	FPCGAttributePropertySelector Selector;
+
+	/** If the entry is read-only, it will only show for FPCGAttributePropertyInputSelector. */
+	bool bReadOnly = false;
 };
 #endif // WITH_EDITOR
 
diff --git a/Engine/Plugins/PCG/Source/PCG/Public/Metadata/PCGAttributePropertySelector.h b/Engine/Plugins/PCG/Source/PCG/Public/Metadata/PCGAttributePropertySelector.h
--- a/Engine/Plugins/PCG/Source/PCG/Public/Metadata/PCGAttributePropertySelector.h
+++ b/Engine/Plugins/PCG/Source/PCG/Public/Metadata/PCGAttributePropertySelector.h
@@ -22,7 +22,8 @@ enum class EPCGAttributePropertySelection
 UENUM()
 enum class EPCGExtraProperties : uint8
 {
-	Index
+	Index UMETA(PCG_PropertyReadOnly),
+	NumElements UMETA(ToolTip = "Data domain property. Only supported on data that have elements.", PCG_PropertyReadOnly, PCG_MetadataDomain="Data")
 };
 
 class FArchiveCrc32;
diff --git a/Engine/Plugins/PCG/Source/PCG/Public/PCGCommon.h b/Engine/Plugins/PCG/Source/PCG/Public/PCGCommon.h
--- a/Engine/Plugins/PCG/Source/PCG/Public/PCGCommon.h
+++ b/Engine/Plugins/PCG/Source/PCG/Public/PCGCommon.h
@@ -233,15 +233,16 @@ namespace Icons
 // Metadata used by PCG
 namespace PCGObjectMetadata
 {
-	const FName Overridable = TEXT("PCG_Overridable");
-	const FName OverridableCPUAndGPU = TEXT("PCG_OverridableCPUAndGPU");
-	const FName OverridableCPUAndGPUWithReadback = TEXT("PCG_OverridableCPUAndGPUWithReadback");
-	const FName NotOverridable = TEXT("PCG_NotOverridable");
-	const FName OverridableChildProperties = TEXT("PCG_OverridableChildProperties");
-	const FName OverrideAliases = TEXT("PCG_OverrideAliases");
-	const FName DiscardPropertySelection = TEXT("PCG_DiscardPropertySelection");
-	const FName DiscardExtraSelection = TEXT("PCG_DiscardExtraSelection");
-	const FName EnumMetadataDomain = TEXT("PCG_MetadataDomain");
+	const FLazyName Overridable = TEXT("PCG_Overridable");
+	const FLazyName OverridableCPUAndGPU = TEXT("PCG_OverridableCPUAndGPU");
+	const FLazyName OverridableCPUAndGPUWithReadback = TEXT("PCG_OverridableCPUAndGPUWithReadback");
+	const FLazyName NotOverridable = TEXT("PCG_NotOverridable");
+	const FLazyName OverridableChildProperties = TEXT("PCG_OverridableChildProperties");
+	const FLazyName OverrideAliases = TEXT("PCG_OverrideAliases");
+	const FLazyName DiscardPropertySelection = TEXT("PCG_DiscardPropertySelection");
+	const FLazyName DiscardExtraSelection = TEXT("PCG_DiscardExtraSelection");
+	const FLazyName EnumMetadataDomain = TEXT("PCG_MetadataDomain");
+	const FLazyName PropertyReadOnly = TEXT("PCG_PropertyReadOnly");
 
 	// Metadata usable in UPROPERTY for customizing the behavior when displaying the property in a property panel or graph node
 	enum
@@ -282,7 +283,10 @@ namespace PCGObjectMetadata
 	enum
 	{
 		/// [EnumMetadata] Specify the domain for this entry.
-		PCG_MetadataDomain
+		PCG_MetadataDomain,
+
+		/// [EnumMetadata] Mark the property read-only, won't show in Output selectors
+		PCG_PropertyReadOnly
 	};
 }
 
diff --git a/Engine/Plugins/PCG/Source/PCGEditor/Private/Details/PCGAttributePropertySelectorDetails.cpp b/Engine/Plugins/PCG/Source/PCGEditor/Private/Details/PCGAttributePropertySelectorDetails.cpp
--- a/Engine/Plugins/PCG/Source/PCGEditor/Private/Details/PCGAttributePropertySelectorDetails.cpp
+++ b/Engine/Plugins/PCG/Source/PCGEditor/Private/Details/PCGAttributePropertySelectorDetails.cpp
@@ -25,13 +25,13 @@ namespace PCGAttributePropertySelectorDetails
 	static const FText SetAttributePropertyNameTransaction = LOCTEXT("SetAttributePropertyName", "[PCG] Set Attribute/Property Name");
 
 	template <typename EnumType>
-	void BuildMenuSectionFromEnum(FMenuBuilder& InMenuBuilder, FPCGAttributePropertySelectorDetails* InDetailsObject, void(FPCGAttributePropertySelectorDetails::*InCallback)(EnumType))
+	void BuildMenuSectionFromEnum(FMenuBuilder& InMenuBuilder, FPCGAttributePropertySelectorDetails* InDetailsObject, void(FPCGAttributePropertySelectorDetails::*InCallback)(EnumType), bool bIsInput)
 	{
 		if (const UEnum* EnumPtr = StaticEnum<EnumType>())
 		{
 			for (int32 i = 0; i < EnumPtr->NumEnums() - 1; ++i)
 			{
-				if (EnumPtr->HasMetaData(TEXT("Hidden"), i))
+				if (EnumPtr->HasMetaData(TEXT("Hidden"), i) || (!bIsInput && EnumPtr->HasMetaData(*PCGObjectMetadata::PropertyReadOnly.ToString(), i)))
 				{
 					continue;
 				}
@@ -264,7 +264,7 @@ TSharedRef<SWidget> FPCGAttributePropertySelectorDetails::GenerateExtraMenu()
 	FMenuBuilder MenuBuilder(true, nullptr);
 
 	FStructProperty* StructProperty = CastFieldChecked<FStructProperty>(PropertyHandle->GetProperty());
-	const bool bIsInput = StructProperty->Struct == FPCGAttributePropertyInputSelector::StaticStruct();
+	const bool bIsInput = StructProperty->Struct->IsChildOf<FPCGAttributePropertyInputSelector>();
 
 	// @todo_pcg: This submenu could be context friendly to the property type an only show when relevant.
 	MenuBuilder.AddSubMenu(
@@ -312,7 +312,7 @@ TSharedRef<SWidget> FPCGAttributePropertySelectorDetails::GenerateExtraMenu()
 
 		TWeakPtr<FPCGAttributePropertySelectorDetails> ThisWeak = StaticCastWeakPtr<FPCGAttributePropertySelectorDetails>(AsWeak());
 
-		auto BuildMenu = [ThisWeak](const FPCGAttributeSelectorMenu& InSelectorMenu, FMenuBuilder& MenuBuilder, auto RecurseCallback) -> void
+		auto BuildMenu = [ThisWeak, bIsInput](const FPCGAttributeSelectorMenu& InSelectorMenu, FMenuBuilder& MenuBuilder, auto RecurseCallback) -> void
 		{
 			for (const FPCGAttributeSelectorMenu& SubMenu : InSelectorMenu.SubMenus)
 			{
@@ -326,6 +326,11 @@ TSharedRef<SWidget> FPCGAttributePropertySelectorDetails::GenerateExtraMenu()
 
 			for (const FPCGAttributeSelectorMenuEntry& MenuEntry : InSelectorMenu.Entries)
 			{
+				if (!bIsInput && MenuEntry.bReadOnly)
+				{
+					continue;
+				}
+				
 				auto SetSelector = [ThisWeak, Selector = MenuEntry.Selector]()
 				{
 					if (FPCGAttributePropertySelectorDetails* This = ThisWeak.Pin().Get())
@@ -357,7 +362,7 @@ TSharedRef<SWidget> FPCGAttributePropertySelectorDetails::GenerateExtraMenu()
 	{
 		MenuBuilder.BeginSection(TEXT("OtherProperties"), LOCTEXT("OtherPropertiesHeader", "Other Properties"));
 		{
-			PCGAttributePropertySelectorDetails::BuildMenuSectionFromEnum<EPCGExtraProperties>(MenuBuilder, this, &FPCGAttributePropertySelectorDetails::SetExtraProperty);
+			PCGAttributePropertySelectorDetails::BuildMenuSectionFromEnum<EPCGExtraProperties>(MenuBuilder, this, &FPCGAttributePropertySelectorDetails::SetExtraProperty, bIsInput);
 		}
 		MenuBuilder.EndSection();
 	}
@@ -410,7 +415,21 @@ void FPCGAttributePropertySelectorDetails::SetAttributeName(FName NewName)
 
 void FPCGAttributePropertySelectorDetails::SetExtraProperty(EPCGExtraProperties EnumValue)
 {
-	PCGAttributePropertySelectorDetails::SetValue(GetStruct(), PropertyHandle, EnumValue, [](FPCGAttributePropertySelector* InStruct, EPCGExtraProperties InValue) -> bool { return InStruct && InStruct->SetExtraProperty(InValue); });
+	const UEnum* EnumPtr = StaticEnum<EPCGExtraProperties>();
+	const int32 EnumIndex = EnumPtr ? EnumPtr->GetIndexByValue(static_cast<int64>(EnumValue)) : INDEX_NONE;
+	const FName DomainName = EnumPtr ? FName(EnumPtr->GetMetaData(*PCGObjectMetadata::EnumMetadataDomain.ToString(), EnumIndex)) : FName(NAME_None);
+	
+	PCGAttributePropertySelectorDetails::SetValue(GetStruct(), PropertyHandle, EnumValue, [DomainName](FPCGAttributePropertySelector* InStruct, EPCGExtraProperties InValue) -> bool
+		{
+			if (!InStruct)
+			{
+				return false;
+			}
+
+			bool bHasChanged = InStruct->SetExtraProperty(InValue);
+			bHasChanged |= InStruct->SetDomainName(DomainName);
+			return bHasChanged;
+		});
 }
 
 void FPCGAttributePropertySelectorDetails::SetSelector(FPCGAttributePropertySelector InSelector)
