[PCG] Fix overrides inside structs not appearing on nodes

#rnx

diff --git a/Engine/Plugins/PCG/Source/PCG/Private/Helpers/PCGSettingsHelpers.cpp b/Engine/Plugins/PCG/Source/PCG/Private/Helpers/PCGSettingsHelpers.cpp
--- a/Engine/Plugins/PCG/Source/PCG/Private/Helpers/PCGSettingsHelpers.cpp
+++ b/Engine/Plugins/PCG/Source/PCG/Private/Helpers/PCGSettingsHelpers.cpp
@@ -379,6 +379,8 @@ namespace PCGSettingsHelpers
 				RecurseConfig.bUseSeed = true;
 #if WITH_EDITOR
 				RecurseConfig.IncludeMetadataValues.Remove(PCGObjectMetadata::Overridable);
+				RecurseConfig.IncludeMetadataValues.Remove(PCGObjectMetadata::OverridableCPUAndGPU);
+				RecurseConfig.IncludeMetadataValues.Remove(PCGObjectMetadata::OverridableCPUAndGPUWithReadback);
 
 				// Also check if we have some filtering on the child properties
 				TArray<FString> ChildPropertiesFiltering;
