[PCG] Add primitive factories to abstract away details of different primitive types.

#rnx

diff --git a/Engine/Plugins/PCG/Source/PCG/Private/Components/PCGProceduralISMComponent.cpp b/Engine/Plugins/PCG/Source/PCG/Private/Components/PCGProceduralISMComponent.cpp
--- a/Engine/Plugins/PCG/Source/PCG/Private/Components/PCGProceduralISMComponent.cpp
+++ b/Engine/Plugins/PCG/Source/PCG/Private/Components/PCGProceduralISMComponent.cpp
@@ -519,20 +519,6 @@ void UPCGProceduralISMComponent::GetSceneProxyDesc(FInstancedStaticMeshSceneProx
 	OutSceneProxyDesc.bUseGpuLodSelection = true;
 }
 
-void UPCGProceduralISMComponent::BuildSceneDesc(FPrimitiveSceneProxyDesc* InSceneProxyDesc, FPrimitiveSceneDesc& OutPrimitiveSceneDesc)
-{
-	check(InSceneProxyDesc);
-
-	OutPrimitiveSceneDesc.SceneProxy = GetSceneProxy();
-	OutPrimitiveSceneDesc.ProxyDesc = InSceneProxyDesc;
-	OutPrimitiveSceneDesc.PrimitiveSceneData = &GetSceneData();
-	OutPrimitiveSceneDesc.RenderMatrix = GetRenderMatrix();
-	OutPrimitiveSceneDesc.AttachmentRootPosition = GetComponentLocation();
-	OutPrimitiveSceneDesc.Bounds = WorldBounds;
-	OutPrimitiveSceneDesc.LocalBounds = GetLocalBounds();
-	OutPrimitiveSceneDesc.Mobility = InSceneProxyDesc->Mobility;
-}
-
 void UPCGProceduralISMComponent::ValidateComponentSetup()
 {
 	if (bAffectDynamicIndirectLighting)
diff --git a/Engine/Plugins/PCG/Source/PCG/Private/Components/PCGProceduralISMComponent.h b/Engine/Plugins/PCG/Source/PCG/Private/Components/PCGProceduralISMComponent.h
--- a/Engine/Plugins/PCG/Source/PCG/Private/Components/PCGProceduralISMComponent.h
+++ b/Engine/Plugins/PCG/Source/PCG/Private/Components/PCGProceduralISMComponent.h
@@ -128,8 +128,6 @@ protected:
 public:
 	void GetSceneProxyDesc(FInstancedStaticMeshSceneProxyDesc& OutSceneProxyDesc) const;
 
-	void BuildSceneDesc(FPrimitiveSceneProxyDesc* InSceneProxyDesc, FPrimitiveSceneDesc& OutPrimitiveSceneDesc);
-
 protected:
 	/** Validate component setup, disable features that are not supported (typically those that require instance transforms on CPU). */
 	virtual void ValidateComponentSetup();
diff --git a/Engine/Plugins/PCG/Source/PCG/Private/Compute/DataInterfaces/Elements/PCGSkinnedMeshSpawnerDataInterface.cpp b/Engine/Plugins/PCG/Source/PCG/Private/Compute/DataInterfaces/Elements/PCGSkinnedMeshSpawnerDataInterface.cpp
--- a/Engine/Plugins/PCG/Source/PCG/Private/Compute/DataInterfaces/Elements/PCGSkinnedMeshSpawnerDataInterface.cpp
+++ b/Engine/Plugins/PCG/Source/PCG/Private/Compute/DataInterfaces/Elements/PCGSkinnedMeshSpawnerDataInterface.cpp
@@ -11,6 +11,7 @@
 #include "Compute/PCGDataBinding.h"
 #include "Compute/BuiltInKernels/PCGCountUniqueAttributeValuesKernel.h"
 #include "Compute/Data/PCGProxyForGPUData.h"
+#include "Compute/PrimitiveFactories/PCGPrimitiveFactoryISKMC.h"
 #include "Elements/PCGSkinnedMeshSpawner.h"
 #include "Elements/PCGSkinnedMeshSpawnerKernel.h"
 #include "InstanceDataPackers/PCGSkinnedMeshInstanceDataPackerBase.h"
@@ -653,37 +654,11 @@ bool UPCGSkinnedMeshSpawnerDataProvider::SetupPrimitives(FPCGContext* InContext,
 	{
 		const FPCGSoftSkinnedMeshComponentDescriptor& Desc = PrimitiveDescriptors[NumPrimitivesProcessed];
 
-		FPCGSkinnedMeshComponentBuilderParams Params;
-		Params.Descriptor = Desc;
-		Params.NumCustomDataFloats = 0; // This is the num custom floats allocated in CPU memory.
+		TSharedPtr<FPCGPrimitiveFactoryISKMC> Primitive = MakeShared<FPCGPrimitiveFactoryISKMC>();
 
-		// If the root actor we're binding to is movable, then the component should be movable by default
-		if (USceneComponent* SceneComponent = TargetActor->GetRootComponent())
+		if (Primitive->Create(InContext, TargetActor, Desc, Settings->UID, Settings->GetSettingsCrc()))
 		{
-			Params.Descriptor.Mobility = SceneComponent->Mobility;
-		}
-
-		Params.SettingsCrc = Settings->GetSettingsCrc();
-		ensure(Params.SettingsCrc.IsValid());
-
-		UPCGComponent* SourceComponent = Cast<UPCGComponent>(InContext->ExecutionSource.Get());
-		check(SourceComponent);
-
-		// TODO - not sure this is a good managed resource object to use, could make procedural variant?
-		UPCGManagedISKMComponent* ManagedComponent = UPCGActorHelpers::GetOrCreateManagedABMC(TargetActor, SourceComponent, Params, InContext);
-
-		if (ManagedComponent)
-		{
-			ManagedComponent->SetCrc(InContext->DependenciesCrc);
-
-			// Don't bother registering the resource change as we're transient anyway.
-			//InContext->TouchedResources.Emplace(MISMC);
-
-			FPCGSpawnerPrimitive Primitive;
-			Primitive.Component = ManagedComponent->GetComponent();
-			Primitive.NumInstances = Desc.NumInstancesGPUOnly;
-
-			Primitives->Primitives.Add(MoveTemp(Primitive));
+			Primitives->PrimitiveFactories.Add(Primitive);
 
 			if (Settings->bApplyMeshBoundsToPoints)
 			{
diff --git a/Engine/Plugins/PCG/Source/PCG/Private/Compute/DataInterfaces/Elements/PCGStaticMeshSpawnerDataInterface.cpp b/Engine/Plugins/PCG/Source/PCG/Private/Compute/DataInterfaces/Elements/PCGStaticMeshSpawnerDataInterface.cpp
--- a/Engine/Plugins/PCG/Source/PCG/Private/Compute/DataInterfaces/Elements/PCGStaticMeshSpawnerDataInterface.cpp
+++ b/Engine/Plugins/PCG/Source/PCG/Private/Compute/DataInterfaces/Elements/PCGStaticMeshSpawnerDataInterface.cpp
@@ -10,6 +10,7 @@
 #include "Compute/PCGDataBinding.h"
 #include "Compute/BuiltInKernels/PCGCountUniqueAttributeValuesKernel.h"
 #include "Compute/Data/PCGProxyForGPUData.h"
+#include "Compute/PrimitiveFactories/PCGPrimitiveFactoryPISMC.h"
 #include "Elements/PCGStaticMeshSpawner.h"
 #include "Elements/PCGStaticMeshSpawnerKernel.h"
 #include "InstanceDataPackers/PCGInstanceDataPackerBase.h"
@@ -738,27 +739,11 @@ bool UPCGStaticMeshSpawnerDataProvider::SetupPrimitives(FPCGContext* InContext,
 	{
 		const FPCGProceduralISMComponentDescriptor& Desc = PrimitiveDescriptors[NumPrimitivesProcessed];
 
-		FPCGProceduralISMCBuilderParameters Params;
-		Params.Descriptor = Desc;
-		Params.bAllowDescriptorChanges = false;
-
-		UPCGComponent* SourceComponent = Cast<UPCGComponent>(InContext->ExecutionSource.Get());
-		check(SourceComponent);
-
-		UPCGManagedProceduralISMComponent* MISMC = PCGManagedProceduralISMComponent::GetOrCreateManagedProceduralISMC(TargetActor, SourceComponent, Settings->UID, Params);
-
-		if (MISMC)
+		TSharedPtr<FPCGPrimitiveFactoryPISMC> Factory = MakeShared<FPCGPrimitiveFactoryPISMC>();
+		
+		if (Factory->Create(InContext, TargetActor, Desc, Settings->UID))
 		{
-			MISMC->SetCrc(InContext->DependenciesCrc);
-
-			// Don't bother registering the resource change as we're transient anyway.
-			//InContext->TouchedResources.Emplace(MISMC);
-
-			FPCGSpawnerPrimitive Primitive;
-			Primitive.Component = MISMC->GetComponent();
-			Primitive.NumInstances = Desc.NumInstances;
-
-			Primitives->Primitives.Add(MoveTemp(Primitive));
+			Primitives->PrimitiveFactories.Add(Factory);
 
 			if (Settings->bApplyMeshBoundsToPoints)
 			{
diff --git a/Engine/Plugins/PCG/Source/PCG/Private/Compute/DataInterfaces/PCGInstanceDataInterface.cpp b/Engine/Plugins/PCG/Source/PCG/Private/Compute/DataInterfaces/PCGInstanceDataInterface.cpp
--- a/Engine/Plugins/PCG/Source/PCG/Private/Compute/DataInterfaces/PCGInstanceDataInterface.cpp
+++ b/Engine/Plugins/PCG/Source/PCG/Private/Compute/DataInterfaces/PCGInstanceDataInterface.cpp
@@ -2,12 +2,14 @@
 
 #include "PCGInstanceDataInterface.h"
 
+#include "PCGContext.h"
 #include "PCGModule.h"
 #include "PCGSettings.h"
 #include "PCGSceneWriterCS.h"
 #include "Components/PCGProceduralISMComponent.h"
 #include "Compute/PCGComputeKernel.h"
 #include "Compute/PCGDataBinding.h"
+#include "Compute/PrimitiveFactories/IPCGRuntimePrimitiveFactory.h"
 
 #include "ComputeWorkerInterface.h"
 #include "GlobalRenderResources.h"
@@ -159,23 +161,39 @@ bool UPCGInstanceDataProvider::PrepareForExecute_GameThread(UPCGDataBinding* InB
 		return true;
 	}
 
-	if (!ensure(FoundPrimitives->Primitives.Num() <= PCGComputeConstants::MAX_PRIMITIVE_COMPONENTS_PER_SPAWNER))
+	if (!ensure(FoundPrimitives->PrimitiveFactories.Num() <= PCGComputeConstants::MAX_PRIMITIVE_COMPONENTS_PER_SPAWNER))
 	{
 		// Last resort - should be clamped earlier during setup and not come to this.
-		FoundPrimitives->Primitives.SetNum(PCGComputeConstants::MAX_PRIMITIVE_COMPONENTS_PER_SPAWNER);
+		FoundPrimitives->PrimitiveFactories.SetNum(PCGComputeConstants::MAX_PRIMITIVE_COMPONENTS_PER_SPAWNER);
 	}
 
 	NumInstancesAllPrimitives = 0;
-	Primitives.Reserve(FoundPrimitives->Primitives.Num());
+	Primitives.Reserve(FoundPrimitives->PrimitiveFactories.Num());
 
-	for (const FPCGSpawnerPrimitive& Primitive : FoundPrimitives->Primitives)
+	for (TSharedPtr<IPCGRuntimePrimitiveFactory> Primitive : FoundPrimitives->PrimitiveFactories)
 	{
-		Primitives.Add(FPCGPrimitiveComponentAndNumInstances(Primitive.Component, Primitive.NumInstances));
-		NumInstancesAllPrimitives += Primitive.NumInstances;
+		if (!ensure(Primitive))
+		{
+			continue;
+		}
+
+		FPrimitiveSceneProxy* SceneProxy = Primitive->GetSceneProxy();
+		if (!SceneProxy)
+		{
+			// Not ready yet, try again later.
+			UE_LOG(LogPCG, Verbose, TEXT("FPCGInstanceDataProvider: One or more scene proxies were not ready. Will try again on the next tick."));
+			return false;
+		}
+
+		const int32 NumInstances = Primitive->GetNumInstances();
+		Primitives.Add(FPCGPrimitiveAndNumInstances(SceneProxy, NumInstances));
+		NumInstancesAllPrimitives += NumInstances;
 	}
 
 	NumCustomFloatsPerInstance = FoundPrimitives->NumCustomFloats;
 
+	ContextHandle = InBinding->ContextHandle;
+
 	return true;
 }
 
@@ -193,7 +211,7 @@ FComputeDataProviderRenderProxy* UPCGInstanceDataProvider::GetRenderProxy()
 {
 	TRACE_CPUPROFILER_EVENT_SCOPE(UPCGInstanceDataProvider::GetRenderProxy);
 
-	return new FPCGInstanceDataProviderProxy(Primitives, NumInstancesAllPrimitives, NumCustomFloatsPerInstance, MakeWeakObjectPtr(this));
+	return new FPCGInstanceDataProviderProxy(Primitives, NumInstancesAllPrimitives, NumCustomFloatsPerInstance, MakeWeakObjectPtr(this), ContextHandle);
 }
 
 void UPCGInstanceDataProvider::Reset()
@@ -204,13 +222,15 @@ void UPCGInstanceDataProvider::Reset()
 	NumInstancesAllPrimitives = 0;
 	NumCustomFloatsPerInstance = 0;
 	bWroteInstances = false;
+	ContextHandle.Reset();
 }
 
-FPCGInstanceDataProviderProxy::FPCGInstanceDataProviderProxy(TArray<FPCGPrimitiveComponentAndNumInstances> InPrimitives, uint32 InNumInstancesAllPrimitives, uint32 InNumCustomFloatsPerInstance, TWeakObjectPtr<UPCGInstanceDataProvider> InDataProvider)
+FPCGInstanceDataProviderProxy::FPCGInstanceDataProviderProxy(TArray<FPCGPrimitiveAndNumInstances> InPrimitives, uint32 InNumInstancesAllPrimitives, uint32 InNumCustomFloatsPerInstance, TWeakObjectPtr<UPCGInstanceDataProvider> InDataProvider, TWeakPtr<FPCGContextHandle> InContextHandle)
 	: Primitives(MoveTemp(InPrimitives))
 	, NumInstancesAllPrimitives(InNumInstancesAllPrimitives)
 	, NumCustomFloatsPerInstance(InNumCustomFloatsPerInstance)
 	, DataProvider(InDataProvider)
+	, ContextHandle(InContextHandle)
 {
 	bIsValid = true;
 
@@ -229,7 +249,7 @@ FPCGInstanceDataProviderProxy::FPCGInstanceDataProviderProxy(TArray<FPCGPrimitiv
 		bIsValid = false;
 	}
 
-	if (Algo::AnyOf(Primitives, [](const FPCGPrimitiveComponentAndNumInstances& InPrimitive) { return !InPrimitive.Get<0>().IsValid(); }))
+	if (Algo::AnyOf(Primitives, [](const FPCGPrimitiveAndNumInstances& InPrimitive) { return !InPrimitive.Get<0>(); }))
 	{
 		UE_LOG(LogPCG, Warning, TEXT("FPCGInstanceDataProviderProxy: One or more components were invalid, proxy is invalid, compute graph will not execute."));
 		bIsValid = false;
@@ -303,44 +323,40 @@ void FPCGInstanceDataProviderProxy::PostSubmit(FComputeContext& Context) const
 	const TRefCountPtr<FRDGPooledBuffer> InstanceCustomFloatDataExported = Context.GraphBuilder.ConvertToExternalBuffer(InstanceCustomFloatData);
 	const TRefCountPtr<FRDGPooledBuffer> WriteCountersExported = Context.GraphBuilder.ConvertToExternalBuffer(WriteCounters);
 	
-	FTSTicker::GetCoreTicker().AddTicker(TEXT("ApplyPrimitiveSceneUpdates"), 0.0f, [DataProviderWeak=DataProvider, DataProviderGeneration=DataProviderGeneration, InstanceDataExported, InstanceCustomFloatDataExported, WriteCountersExported](float)
+	FTSTicker::GetCoreTicker().AddTicker(TEXT("ApplyPrimitiveSceneUpdates"), 0.0f, [DataProviderWeak=DataProvider, DataProviderGeneration=DataProviderGeneration, InstanceDataExported, InstanceCustomFloatDataExported, WriteCountersExported, ContextHandle=ContextHandle](float)
 	{
-		TStrongObjectPtr<UPCGInstanceDataProvider> DataProviderInner = DataProviderWeak.Pin();
+		UPCGInstanceDataProvider* DataProviderInner = DataProviderWeak.Get();
 		if (!DataProviderInner || DataProviderInner->GenerationCounter != DataProviderGeneration)
 		{
 			UE_LOG(LogPCG, Verbose, TEXT("Data provider object lost, GPU instancing will fail."));
 			return false;
 		}
 
-		bool bAllProxiesAreReady = true;
+		// All instance data is stored in a single buffer, so this is used to give the scene writer an index to the first primitive.
+		int32 CumulativeInstanceCount = 0;
 
-		for (const FPCGPrimitiveComponentAndNumInstances& PrimitiveAndInstanceCount : DataProviderInner->Primitives)
+		FSceneInterface* Scene = nullptr;
+
+		if (TSharedPtr<FPCGContextHandle> SharedHandle = ContextHandle.Pin())
 		{
-			if (UPrimitiveComponent* PrimitiveComponent = PrimitiveAndInstanceCount.Get<0>().Get())
+			if (FPCGContext* ContextPtr = SharedHandle->GetContext())
 			{
-				bAllProxiesAreReady &= PrimitiveComponent->GetSceneProxy() != nullptr;
+				if (UWorld* World = ContextPtr->ExecutionSource.IsValid() ? ContextPtr->ExecutionSource->GetExecutionState().GetWorld() : nullptr)
+				{
+					Scene = World->Scene;
+				}
 			}
 		}
 
-		if (!bAllProxiesAreReady)
+		if (!Scene)
 		{
-			UE_LOG(LogPCG, Verbose, TEXT("FPCGInstanceDataProvider: One or more scene proxies were not ready. Will try again on the next tick."));
-			return true;
+			UE_LOG(LogPCG, Error, TEXT("Missing scene encountered during instancing, should not happen."));
+			return false;
 		}
 
-		// All instance data is stored in a single buffer, so this is used to give the scene writer an index to the first primitive.
-		int32 CumulativeInstanceCount = 0;
-
 		for (int PrimitiveIndex = 0; PrimitiveIndex < DataProviderInner->Primitives.Num(); ++PrimitiveIndex)
 		{
-			FPCGPrimitiveComponentAndNumInstances& PrimitiveAndInstanceCount = DataProviderInner->Primitives[PrimitiveIndex];
-
-			UPrimitiveComponent* PrimitiveComponent = PrimitiveAndInstanceCount.Get<0>().Get();
-			if (!PrimitiveComponent)
-			{
-				UE_LOG(LogPCG, Warning, TEXT("Null primitive component encountered during instancing, should not happen."));
-				continue;
-			}
+			FPCGPrimitiveAndNumInstances& PrimitiveAndInstanceCount = DataProviderInner->Primitives[PrimitiveIndex];
 
 			const int32 PrimitiveNumInstances = PrimitiveAndInstanceCount.Get<1>();
 			if (PrimitiveNumInstances <= 0)
@@ -349,98 +365,71 @@ void FPCGInstanceDataProviderProxy::PostSubmit(FComputeContext& Context) const
 				continue;
 			}
 
-			FPrimitiveSceneProxyDesc* SceneProxyDesc = nullptr;
-
-			// TODO: Slight mess, could add a GetSceneProxyDesc function to the base descriptor class. However we plan
-			// to move away from components and this code is temporary.
-			FInstancedStaticMeshSceneProxyDesc PISMCSceneProxyDesc;
-			FInstancedSkinnedMeshSceneProxyDesc ISKMSceneProxyDesc;
-
-			if (UPCGProceduralISMComponent* PISM = Cast<UPCGProceduralISMComponent>(PrimitiveComponent))
-			{
-				PISM->GetSceneProxyDesc(PISMCSceneProxyDesc);
-				SceneProxyDesc = &PISMCSceneProxyDesc;
-			}
-			else if (UInstancedSkinnedMeshComponent* ISKM = Cast<UInstancedSkinnedMeshComponent>(PrimitiveComponent))
-			{
-				ISKMSceneProxyDesc = FInstancedSkinnedMeshSceneProxyDesc(ISKM);
-				SceneProxyDesc = &ISKMSceneProxyDesc;
-			}
+			FPrimitiveSceneProxy* SceneProxy = PrimitiveAndInstanceCount.Get<0>();
+			check(SceneProxy);
 
-			if (!SceneProxyDesc)
+			if (SceneProxy->GetPrimitiveSceneInfo()->GetInstanceSceneDataOffset() == -1)
 			{
-				UE_LOG(LogPCG, Warning, TEXT("Primitive '%s' of unsupported type, instancing aborted."), *PrimitiveComponent->GetName());
+				UE_LOG(LogPCG, Error, TEXT("Primitive not allocated in GPU scene, should not happen."));
 				continue;
 			}
 
-			FPrimitiveSceneDesc PrimitiveSceneDesc;
-
-			// TODO try add FPrimitiveSceneDesc::InitFromPrimitiveComponent
-			if (UPCGProceduralISMComponent* PISM = Cast<UPCGProceduralISMComponent>(PrimitiveComponent))
-			{
-				PISM->BuildSceneDesc(SceneProxyDesc, PrimitiveSceneDesc);
-			}
-			else if (UInstancedSkinnedMeshComponent* ISKM = Cast<UInstancedSkinnedMeshComponent>(PrimitiveComponent))
-			{
-				ISKM->BuildSceneDesc(SceneProxyDesc, PrimitiveSceneDesc);
-			}
-
-			if (FSceneInterface* Scene = PrimitiveComponent->GetWorld() ? PrimitiveComponent->GetWorld()->Scene : nullptr)
-			{
 #if UE_BUILD_SHIPPING
-				const bool bTriggerCapture = false;
+			const bool bTriggerCapture = false;
 #else
-				const bool bTriggerCapture = (PCGInstanceDataInterface::TriggerGPUCaptureDispatchIndex > 0);
-				PCGInstanceDataInterface::TriggerGPUCaptureDispatchIndex = FMath::Max(PCGInstanceDataInterface::TriggerGPUCaptureDispatchIndex - 1, 0);
+			const bool bTriggerCapture = (PCGInstanceDataInterface::TriggerGPUCaptureDispatchIndex > 0);
+			PCGInstanceDataInterface::TriggerGPUCaptureDispatchIndex = FMath::Max(PCGInstanceDataInterface::TriggerGPUCaptureDispatchIndex - 1, 0);
 #endif // UE_BUILD_SHIPPING
 
-				Scene->UpdatePrimitiveInstancesFromCompute(&PrimitiveSceneDesc, FGPUSceneWriteDelegate::CreateLambda([PrimitiveIndex, PrimitiveNumInstances, CumulativeInstanceCount, bTriggerCapture, InstanceDataExported, InstanceCustomFloatDataExported, WriteCountersExported](FRDGBuilder& InGraphBuilder, const FGPUSceneWriteDelegateParams& Params)
-				{
-					RDG_EVENT_SCOPE(InGraphBuilder, "SceneComputeUpdateInterface->EnqueueUpdate");
-					check(Params.PersistentPrimitiveId != (uint32)INDEX_NONE);
+			FPrimitiveSceneDesc PrimitiveSceneDesc;
+			PrimitiveSceneDesc.SceneProxy = SceneProxy;
+
+			Scene->UpdatePrimitiveInstancesFromCompute(&PrimitiveSceneDesc, FGPUSceneWriteDelegate::CreateLambda([PrimitiveIndex, PrimitiveNumInstances, CumulativeInstanceCount, bTriggerCapture, InstanceDataExported, InstanceCustomFloatDataExported, WriteCountersExported](FRDGBuilder& InGraphBuilder, const FGPUSceneWriteDelegateParams& Params)
+			{
+				RDG_EVENT_SCOPE(InGraphBuilder, "SceneComputeUpdateInterface->EnqueueUpdate");
+				check(Params.PersistentPrimitiveId != (uint32)INDEX_NONE);
 
 #if !UE_BUILD_SHIPPING
-					RenderCaptureInterface::FScopedCapture RenderCapture(bTriggerCapture, InGraphBuilder, TEXT("SceneComputeUpdateInterface->EnqueueUpdate"));
+				RenderCaptureInterface::FScopedCapture RenderCapture(bTriggerCapture, InGraphBuilder, TEXT("SceneComputeUpdateInterface->EnqueueUpdate"));
 #endif // !UE_BUILD_SHIPPING
 
-					FRDGBufferRef InstanceDataInner  = InGraphBuilder.RegisterExternalBuffer(InstanceDataExported);
-					FRDGBufferRef InstanceCustomFloatDataInner = InGraphBuilder.RegisterExternalBuffer(InstanceCustomFloatDataExported);
-					FRDGBufferRef WriteCountersInner = InGraphBuilder.RegisterExternalBuffer(WriteCountersExported);
-
-					// Write instances
-					FPCGSceneWriterCS::FParameters* Parameters = InGraphBuilder.AllocParameters<FPCGSceneWriterCS::FParameters>();
-					Parameters->InPrimitiveIndex = PrimitiveIndex;
-					Parameters->InNumInstancesAllocatedInGPUScene = PrimitiveNumInstances;
-					Parameters->InInstanceOffset = CumulativeInstanceCount;
-					Parameters->InInstanceData = InGraphBuilder.CreateSRV(InstanceDataInner);
-					Parameters->InInstanceCustomFloatData = InGraphBuilder.CreateSRV(InstanceCustomFloatDataInner);
-					Parameters->InWriteCounters = InGraphBuilder.CreateSRV(WriteCountersInner);
-					Parameters->InPrimitiveId = Params.PersistentPrimitiveId;
-					PRAGMA_DISABLE_DEPRECATION_WARNINGS
-					Parameters->GPUSceneWriterParameters = Params.GPUWriteParams;
-					PRAGMA_ENABLE_DEPRECATION_WARNINGS
-					Parameters->InCustomDataCount = Params.NumCustomDataFloats;
-					Parameters->InPayloadDataFlags = Params.PackedInstanceSceneDataFlags;
+				FRDGBufferRef InstanceDataInner  = InGraphBuilder.RegisterExternalBuffer(InstanceDataExported);
+				FRDGBufferRef InstanceCustomFloatDataInner = InGraphBuilder.RegisterExternalBuffer(InstanceCustomFloatDataExported);
+				FRDGBufferRef WriteCountersInner = InGraphBuilder.RegisterExternalBuffer(WriteCountersExported);
+
+				// Write instances
+				FPCGSceneWriterCS::FParameters* Parameters = InGraphBuilder.AllocParameters<FPCGSceneWriterCS::FParameters>();
+				Parameters->InPrimitiveIndex = PrimitiveIndex;
+				Parameters->InNumInstancesAllocatedInGPUScene = PrimitiveNumInstances;
+				Parameters->InInstanceOffset = CumulativeInstanceCount;
+				Parameters->InInstanceData = InGraphBuilder.CreateSRV(InstanceDataInner);
+				Parameters->InInstanceCustomFloatData = InGraphBuilder.CreateSRV(InstanceCustomFloatDataInner);
+				Parameters->InWriteCounters = InGraphBuilder.CreateSRV(WriteCountersInner);
+				Parameters->InPrimitiveId = Params.PersistentPrimitiveId;
+				PRAGMA_DISABLE_DEPRECATION_WARNINGS
+				Parameters->GPUSceneWriterParameters = Params.GPUWriteParams;
+				PRAGMA_ENABLE_DEPRECATION_WARNINGS
+				Parameters->InCustomDataCount = Params.NumCustomDataFloats;
+				Parameters->InPayloadDataFlags = Params.PackedInstanceSceneDataFlags;
 
 #if PCG_INSTANCE_DATA_LOGGING
-					UE_LOG(LogPCG, Log, TEXT("\tScene writer delegate [%d]:\tPrimitive ID %u,\tsource instance offset %u,\tInstanceSceneDataOffset %u, num instances %u"),
-						Parameters->InPrimitiveIndex,
-						Parameters->InPrimitiveId,
-						Parameters->InInstanceOffset,
-						Params.InstanceSceneDataOffset,
-						Parameters->InNumInstancesAllocatedInGPUScene);
+				UE_LOG(LogPCG, Log, TEXT("\tScene writer delegate [%d]:\tPrimitive ID %u,\tsource instance offset %u,\tInstanceSceneDataOffset %u, num instances %u"),
+					Parameters->InPrimitiveIndex,
+					Parameters->InPrimitiveId,
+					Parameters->InInstanceOffset,
+					Params.InstanceSceneDataOffset,
+					Parameters->InNumInstancesAllocatedInGPUScene);
 #endif
 
-					// We're using a custom compute shader here rather than using a compute graph kernel, because ultimately the scene update will incorporate the
-					// new instance data. In the future we will not write directly to the scene here.
-					TShaderMapRef<FPCGSceneWriterCS> Shader(GetGlobalShaderMap(GMaxRHIFeatureLevel));
-					int GroupCount = FMath::DivideAndRoundUp<int>(PrimitiveNumInstances, FPCGSceneWriterCS::NumThreadsPerGroup);
-					ensureAlways(GroupCount <= GRHIMaxDispatchThreadGroupsPerDimension.X);
-					GroupCount = FMath::Min(GRHIMaxDispatchThreadGroupsPerDimension.X, GroupCount);
+				// We're using a custom compute shader here rather than using a compute graph kernel, because ultimately the scene update will incorporate the
+				// new instance data. In the future we will not write directly to the scene here.
+				TShaderMapRef<FPCGSceneWriterCS> Shader(GetGlobalShaderMap(GMaxRHIFeatureLevel));
+				int GroupCount = FMath::DivideAndRoundUp<int>(PrimitiveNumInstances, FPCGSceneWriterCS::NumThreadsPerGroup);
+				ensureAlways(GroupCount <= GRHIMaxDispatchThreadGroupsPerDimension.X);
+				GroupCount = FMath::Min(GRHIMaxDispatchThreadGroupsPerDimension.X, GroupCount);
 
-					FComputeShaderUtils::AddPass(InGraphBuilder, RDG_EVENT_NAME("PCGWriteInstanceData"), ERDGPassFlags::Compute | ERDGPassFlags::NeverCull, Shader, Parameters, FIntVector(GroupCount, 1, 1));
-				}));
-			}
+				FComputeShaderUtils::AddPass(InGraphBuilder, RDG_EVENT_NAME("PCGWriteInstanceData"), ERDGPassFlags::Compute | ERDGPassFlags::NeverCull, Shader, Parameters, FIntVector(GroupCount, 1, 1));
+			}));
 
 			CumulativeInstanceCount += PrimitiveNumInstances;
 		}
diff --git a/Engine/Plugins/PCG/Source/PCG/Private/Compute/DataInterfaces/PCGInstanceDataInterface.h b/Engine/Plugins/PCG/Source/PCG/Private/Compute/DataInterfaces/PCGInstanceDataInterface.h
--- a/Engine/Plugins/PCG/Source/PCG/Private/Compute/DataInterfaces/PCGInstanceDataInterface.h
+++ b/Engine/Plugins/PCG/Source/PCG/Private/Compute/DataInterfaces/PCGInstanceDataInterface.h
@@ -12,8 +12,9 @@ class FPCGInstanceDataInterfaceParameters;
 class FPrimitiveSceneProxy;
 class UPCGSettings;
 class UComputeKernel;
+struct FPCGContextHandle;
 
-using FPCGPrimitiveComponentAndNumInstances = TPair<TWeakObjectPtr<UPrimitiveComponent>, int>;
+using FPCGPrimitiveAndNumInstances = TPair<FPrimitiveSceneProxy*, int32>;
 
 /** Data Interface for writing instance data (transform and custom floats) to an intermediate buffer, which is then injected into the GPU Scene. */
 UCLASS(ClassGroup = (Procedural))
@@ -56,7 +57,7 @@ public:
 	virtual void Reset() override;
 	//~ End UComputeDataProvider Interface
 
-	TArray<FPCGPrimitiveComponentAndNumInstances> Primitives;
+	TArray<FPCGPrimitiveAndNumInstances> Primitives;
 
 	uint32 NumInstancesAllPrimitives = 0;
 
@@ -65,14 +66,13 @@ public:
 	/** Whether instance data has been applied to scene and operation is complete. */
 	bool bWroteInstances = false;
 
-	/** Bumped each time the data provider is initialized or reset, so that async callbacks can detect if they originated from a previous usage of the data provider and no-op. */
-	std::atomic<uint64> GenerationCounter = 0;
+	TWeakPtr<FPCGContextHandle> ContextHandle;
 };
 
 class FPCGInstanceDataProviderProxy : public FComputeDataProviderRenderProxy
 {
 public:
-	FPCGInstanceDataProviderProxy(TArray<FPCGPrimitiveComponentAndNumInstances> InPrimitives, uint32 InNumInstancesAllPrimitives, uint32 InNumCustomFloatsPerInstance, TWeakObjectPtr<UPCGInstanceDataProvider> InDataProvider/*, uint64 InDataProviderGeneration*/);
+	FPCGInstanceDataProviderProxy(TArray<FPCGPrimitiveAndNumInstances> InPrimitives, uint32 InNumInstancesAllPrimitives, uint32 InNumCustomFloatsPerInstance, TWeakObjectPtr<UPCGInstanceDataProvider> InDataProvider, TWeakPtr<FPCGContextHandle> InContextHandle);
 
 	//~ Begin FComputeDataProviderRenderProxy Interface
 	bool IsValid(FValidationData const& InValidationData) const override;
@@ -86,7 +86,7 @@ protected:
 
 	bool bIsValid = false;
 
-	TArray<FPCGPrimitiveComponentAndNumInstances> Primitives;
+	TArray<FPCGPrimitiveAndNumInstances> Primitives;
 
 	/** Instance transforms of all instances across all primitives. Three float4s per instance. */
 	FRDGBufferRef InstanceData = nullptr;
@@ -108,4 +108,6 @@ protected:
 
 	TWeakObjectPtr<UPCGInstanceDataProvider> DataProvider;
 	uint64 DataProviderGeneration = 0;
+
+	TWeakPtr<FPCGContextHandle> ContextHandle;
 };
diff --git a/Engine/Plugins/PCG/Source/PCG/Private/Compute/PCGDataBinding.h b/Engine/Plugins/PCG/Source/PCG/Private/Compute/PCGDataBinding.h
--- a/Engine/Plugins/PCG/Source/PCG/Private/Compute/PCGDataBinding.h
+++ b/Engine/Plugins/PCG/Source/PCG/Private/Compute/PCGDataBinding.h
@@ -9,6 +9,7 @@
 
 #include "PCGDataBinding.generated.h"
 
+class IPCGRuntimePrimitiveFactory;
 class UComputeKernel;
 class UPCGComponent;
 class UPCGComputeGraph;
@@ -35,8 +36,7 @@ struct FPCGSpawnerPrimitives
 {
 	GENERATED_BODY()
 
-	UPROPERTY()
-	TArray<FPCGSpawnerPrimitive> Primitives;
+	TArray<TSharedPtr<IPCGRuntimePrimitiveFactory>> PrimitiveFactories;
 
 	TArray<FBox> PrimitiveMeshBounds;
 
diff --git a/Engine/Plugins/PCG/Source/PCG/Private/Compute/PrimitiveFactories/IPCGRuntimePrimitiveFactory.h b/Engine/Plugins/PCG/Source/PCG/Private/Compute/PrimitiveFactories/IPCGRuntimePrimitiveFactory.h
new file mode 100644
--- /dev/null
+++ b/Engine/Plugins/PCG/Source/PCG/Private/Compute/PrimitiveFactories/IPCGRuntimePrimitiveFactory.h
@@ -0,0 +1,15 @@
+// Copyright Epic Games, Inc. All Rights Reserved.
+
+#pragma once
+
+#include "CoreMinimal.h"
+
+class FPrimitiveSceneProxy;
+
+/** Interface for helper object that creates and adds primitives to the scene. */
+class IPCGRuntimePrimitiveFactory
+{
+public:
+	virtual FPrimitiveSceneProxy* GetSceneProxy() const = 0;
+	virtual int32 GetNumInstances() const = 0;
+};
diff --git a/Engine/Plugins/PCG/Source/PCG/Private/Compute/PrimitiveFactories/PCGPrimitiveFactoryISKMC.cpp b/Engine/Plugins/PCG/Source/PCG/Private/Compute/PrimitiveFactories/PCGPrimitiveFactoryISKMC.cpp
new file mode 100644
--- /dev/null
+++ b/Engine/Plugins/PCG/Source/PCG/Private/Compute/PrimitiveFactories/PCGPrimitiveFactoryISKMC.cpp
@@ -0,0 +1,41 @@
+// Copyright Epic Games, Inc. All Rights Reserved.
+
+#include "Compute/PrimitiveFactories/PCGPrimitiveFactoryISKMC.h"
+
+#include "PCGComponent.h"
+#include "PCGContext.h"
+#include "PCGManagedResource.h"
+#include "Components/InstancedSkinnedMeshComponent.h"
+#include "Helpers/PCGActorHelpers.h"
+
+bool FPCGPrimitiveFactoryISKMC::Create(FPCGContext* InContext, AActor* InTargetActor, const FPCGSoftSkinnedMeshComponentDescriptor& InDesc, uint64 InSettingsUID, const FPCGCrc& InSettingsCrc)
+{
+	UPCGComponent* SourceComponent = Cast<UPCGComponent>(InContext->ExecutionSource.Get());
+	check(SourceComponent);
+
+	FPCGSkinnedMeshComponentBuilderParams Params;
+	Params.Descriptor = InDesc;
+	Params.NumCustomDataFloats = 0; // This is the num custom floats allocated in CPU memory.
+
+	// If the root actor we're binding to is movable, then the component should be movable by default
+	if (USceneComponent* SceneComponent = InTargetActor->GetRootComponent())
+	{
+		Params.Descriptor.Mobility = SceneComponent->Mobility;
+	}
+
+	ensure(InSettingsCrc.IsValid());
+	Params.SettingsCrc = InSettingsCrc;
+
+	// TODO - not sure this is a good managed resource object to use, could make procedural variant?
+	UPCGManagedISKMComponent* ManagedComponent = UPCGActorHelpers::GetOrCreateManagedABMC(InTargetActor, SourceComponent, Params, InContext);
+
+	if (ManagedComponent)
+	{
+		Component = ManagedComponent->GetComponent();
+		NumInstances = InDesc.NumInstancesGPUOnly;
+
+		ManagedComponent->SetCrc(InContext->DependenciesCrc);
+	}
+
+	return ManagedComponent != nullptr;
+}
diff --git a/Engine/Plugins/PCG/Source/PCG/Private/Compute/PrimitiveFactories/PCGPrimitiveFactoryISKMC.h b/Engine/Plugins/PCG/Source/PCG/Private/Compute/PrimitiveFactories/PCGPrimitiveFactoryISKMC.h
new file mode 100644
--- /dev/null
+++ b/Engine/Plugins/PCG/Source/PCG/Private/Compute/PrimitiveFactories/PCGPrimitiveFactoryISKMC.h
@@ -0,0 +1,17 @@
+// Copyright Epic Games, Inc. All Rights Reserved.
+
+#pragma once
+
+#include "Compute/PrimitiveFactories/PCGPrimitiveFactoryPrimitiveComponentBase.h"
+
+class AActor;
+struct FPCGContext;
+struct FPCGCrc;
+struct FPCGSoftSkinnedMeshComponentDescriptor;
+
+class FPCGPrimitiveFactoryISKMC : public FPCGPrimitiveFactoryPrimitiveComponentBase
+{
+public:
+	/** Creates a Instanced Skinned Mesh Component on the target actor, and registers it with the execution source. */
+	bool Create(FPCGContext* InContext, AActor* InTargetActor, const FPCGSoftSkinnedMeshComponentDescriptor& InDesc, uint64 InSettingsUID, const FPCGCrc& InSettingsCrc);
+};
diff --git a/Engine/Plugins/PCG/Source/PCG/Private/Compute/PrimitiveFactories/PCGPrimitiveFactoryPISMC.cpp b/Engine/Plugins/PCG/Source/PCG/Private/Compute/PrimitiveFactories/PCGPrimitiveFactoryPISMC.cpp
new file mode 100644
--- /dev/null
+++ b/Engine/Plugins/PCG/Source/PCG/Private/Compute/PrimitiveFactories/PCGPrimitiveFactoryPISMC.cpp
@@ -0,0 +1,29 @@
+// Copyright Epic Games, Inc. All Rights Reserved.
+
+#include "Compute/PrimitiveFactories/PCGPrimitiveFactoryPISMC.h"
+
+#include "PCGComponent.h"
+#include "PCGContext.h"
+#include "Components/PCGProceduralISMComponent.h"
+
+bool FPCGPrimitiveFactoryPISMC::Create(FPCGContext* InContext, AActor* InTargetActor, const FPCGProceduralISMComponentDescriptor& InDesc, uint64 InSettingsUID)
+{
+	UPCGComponent* SourceComponent = Cast<UPCGComponent>(InContext->ExecutionSource.Get());
+	check(SourceComponent);
+
+	FPCGProceduralISMCBuilderParameters Params;
+	Params.Descriptor = InDesc;
+	Params.bAllowDescriptorChanges = false;
+
+	UPCGManagedProceduralISMComponent* ManagedComponent = PCGManagedProceduralISMComponent::GetOrCreateManagedProceduralISMC(InTargetActor, SourceComponent, InSettingsUID, Params);
+
+	if (ManagedComponent)
+	{
+		Component = ManagedComponent->GetComponent();
+		NumInstances = InDesc.NumInstances;
+
+		ManagedComponent->SetCrc(InContext->DependenciesCrc);
+	}
+
+	return ManagedComponent != nullptr;
+}
diff --git a/Engine/Plugins/PCG/Source/PCG/Private/Compute/PrimitiveFactories/PCGPrimitiveFactoryPISMC.h b/Engine/Plugins/PCG/Source/PCG/Private/Compute/PrimitiveFactories/PCGPrimitiveFactoryPISMC.h
new file mode 100644
--- /dev/null
+++ b/Engine/Plugins/PCG/Source/PCG/Private/Compute/PrimitiveFactories/PCGPrimitiveFactoryPISMC.h
@@ -0,0 +1,16 @@
+// Copyright Epic Games, Inc. All Rights Reserved.
+
+#pragma once
+
+#include "Compute/PrimitiveFactories/PCGPrimitiveFactoryPrimitiveComponentBase.h"
+
+class AActor;
+struct FPCGContext;
+struct FPCGProceduralISMComponentDescriptor;
+
+class FPCGPrimitiveFactoryPISMC : public FPCGPrimitiveFactoryPrimitiveComponentBase
+{
+public:
+	/** Creates a Procedural Instanced Static Mesh Component on the target actor, and registers it with the execution source. */
+	bool Create(FPCGContext* InContext, AActor* InTargetActor, const FPCGProceduralISMComponentDescriptor& InDesc, uint64 InSettingsUID);
+};
diff --git a/Engine/Plugins/PCG/Source/PCG/Private/Compute/PrimitiveFactories/PCGPrimitiveFactoryPrimitiveComponentBase.cpp b/Engine/Plugins/PCG/Source/PCG/Private/Compute/PrimitiveFactories/PCGPrimitiveFactoryPrimitiveComponentBase.cpp
new file mode 100644
--- /dev/null
+++ b/Engine/Plugins/PCG/Source/PCG/Private/Compute/PrimitiveFactories/PCGPrimitiveFactoryPrimitiveComponentBase.cpp
@@ -0,0 +1,10 @@
+// Copyright Epic Games, Inc. All Rights Reserved.
+
+#include "Compute/PrimitiveFactories/PCGPrimitiveFactoryPrimitiveComponentBase.h"
+
+#include "Components/PrimitiveComponent.h"
+
+FPrimitiveSceneProxy* FPCGPrimitiveFactoryPrimitiveComponentBase::GetSceneProxy() const
+{
+	return Component.IsValid() ? Component->GetSceneProxy() : nullptr;
+}
diff --git a/Engine/Plugins/PCG/Source/PCG/Private/Compute/PrimitiveFactories/PCGPrimitiveFactoryPrimitiveComponentBase.h b/Engine/Plugins/PCG/Source/PCG/Private/Compute/PrimitiveFactories/PCGPrimitiveFactoryPrimitiveComponentBase.h
new file mode 100644
--- /dev/null
+++ b/Engine/Plugins/PCG/Source/PCG/Private/Compute/PrimitiveFactories/PCGPrimitiveFactoryPrimitiveComponentBase.h
@@ -0,0 +1,23 @@
+// Copyright Epic Games, Inc. All Rights Reserved.
+
+#pragma once
+
+#include "Compute/PrimitiveFactories/IPCGRuntimePrimitiveFactory.h"
+
+#include "UObject/WeakObjectPtr.h"
+
+class FPrimitiveSceneProxy;
+class UPrimitiveComponent;
+
+class FPCGPrimitiveFactoryPrimitiveComponentBase : public IPCGRuntimePrimitiveFactory
+{
+public:
+	//~ Begin IPCGRuntimePrimitiveFactory interface
+	virtual FPrimitiveSceneProxy* GetSceneProxy() const override;
+	virtual int32 GetNumInstances() const override { return NumInstances; }
+	//~ End IPCGRuntimePrimitiveFactory interface
+
+protected:
+	TWeakObjectPtr<UPrimitiveComponent> Component;
+	int32 NumInstances = -1;
+};
diff --git a/Engine/Source/Runtime/Engine/Classes/Components/InstancedSkinnedMeshComponent.h b/Engine/Source/Runtime/Engine/Classes/Components/InstancedSkinnedMeshComponent.h
--- a/Engine/Source/Runtime/Engine/Classes/Components/InstancedSkinnedMeshComponent.h
+++ b/Engine/Source/Runtime/Engine/Classes/Components/InstancedSkinnedMeshComponent.h
@@ -155,8 +155,6 @@ public:
 #endif
 	//~ End USkinnedMeshComponent Interface
 
-	ENGINE_API void BuildSceneDesc(FPrimitiveSceneProxyDesc* InSceneProxyDesc, FPrimitiveSceneDesc& OutPrimitiveSceneDesc);
-
 public:
 	void SetNumGPUInstances(int32 InCount);
 
@@ -292,6 +290,10 @@ private:
 
 private:
 	void SetSkinnedAssetCallback();
+
+public:
+	UE_DEPRECATED(5.7, "Method removed.")
+	ENGINE_API void BuildSceneDesc(FPrimitiveSceneProxyDesc* InSceneProxyDesc, FPrimitiveSceneDesc& OutPrimitiveSceneDesc) {}
 };
 
 /** Helper class used to preserve state across blueprint re-instancing */
diff --git a/Engine/Source/Runtime/Engine/Private/Components/InstancedSkinnedMeshComponent.cpp b/Engine/Source/Runtime/Engine/Private/Components/InstancedSkinnedMeshComponent.cpp
--- a/Engine/Source/Runtime/Engine/Private/Components/InstancedSkinnedMeshComponent.cpp
+++ b/Engine/Source/Runtime/Engine/Private/Components/InstancedSkinnedMeshComponent.cpp
@@ -921,20 +921,6 @@ void UInstancedSkinnedMeshComponent::PostAssetCompilation()
 
 #endif 
 
-void UInstancedSkinnedMeshComponent::BuildSceneDesc(FPrimitiveSceneProxyDesc* InSceneProxyDesc, FPrimitiveSceneDesc& OutPrimitiveSceneDesc) // TODO try add FPrimitiveSceneDesc::InitFromPrimitiveComponent
-{
-	check(InSceneProxyDesc);
-
-	OutPrimitiveSceneDesc.SceneProxy = GetSceneProxy();
-	OutPrimitiveSceneDesc.ProxyDesc = InSceneProxyDesc;
-	OutPrimitiveSceneDesc.PrimitiveSceneData = &GetSceneData();
-	OutPrimitiveSceneDesc.RenderMatrix = GetRenderMatrix();
-	OutPrimitiveSceneDesc.AttachmentRootPosition = GetComponentLocation();
-	OutPrimitiveSceneDesc.LocalBounds = CalcBounds(FTransform::Identity);
-	OutPrimitiveSceneDesc.Bounds = CalcBounds(GetComponentToWorld());
-	OutPrimitiveSceneDesc.Mobility = InSceneProxyDesc->Mobility;
-}
-
 FInstanceDataManagerSourceDataDesc UInstancedSkinnedMeshComponent::GetComponentDesc(ERHIFeatureLevel::Type FeatureLevel)
 {
 	FInstanceDataManagerSourceDataDesc ComponentDesc;
