[PCG] Add "SetWorldObject" on the PCG Data Exporter for BPs

* Objects can't be cast to World in BPs, and when doing an utility widget, Get Selected Assets returns an array of objects
* Will do the cast internally, returns false if the cast failed.

#rnx

diff --git a/Engine/Plugins/PCG/Source/PCGEditor/Private/PCGLevelToAsset.cpp b/Engine/Plugins/PCG/Source/PCGEditor/Private/PCGLevelToAsset.cpp
--- a/Engine/Plugins/PCG/Source/PCGEditor/Private/PCGLevelToAsset.cpp
+++ b/Engine/Plugins/PCG/Source/PCGEditor/Private/PCGLevelToAsset.cpp
@@ -572,6 +572,12 @@ void UPCGLevelToAsset::SetWorld(UWorld* World)
 	WorldToExport = World;
 }
 
+bool UPCGLevelToAsset::SetWorldObject(UObject* WorldObject)
+{
+	WorldToExport = Cast<UWorld>(WorldObject);
+	return WorldToExport != nullptr;
+}
+
 UWorld* UPCGLevelToAsset::GetWorld() const
 {
 	return WorldToExport;
diff --git a/Engine/Plugins/PCG/Source/PCGEditor/Public/PCGLevelToAsset.h b/Engine/Plugins/PCG/Source/PCGEditor/Public/PCGLevelToAsset.h
--- a/Engine/Plugins/PCG/Source/PCGEditor/Public/PCGLevelToAsset.h
+++ b/Engine/Plugins/PCG/Source/PCGEditor/Public/PCGLevelToAsset.h
@@ -34,6 +34,10 @@ public:
 	UFUNCTION(BlueprintCallable, Category="PCG|IO")
 	UE_API void SetWorld(UWorld* World);
 
+	/** Set the world to export from a UObject. The WorldObject must be a World. Return false if the object is not a world or null. */
+	UFUNCTION(BlueprintCallable, Category = "PCG|IO")
+	UE_API bool SetWorldObject(UObject* WorldObject);
+
 	UFUNCTION(BlueprintCallable, Category = "PCG|IO")
 	UE_API UWorld* GetWorld() const;
 
