[PCG] Custom widget for PCG Data Type Identifier

* Add "Difference" operation on the Data Type Registry
* Fix few issues that arose from previous submit

#rnx

diff --git a/Engine/Plugins/PCG/Source/PCG/Internal/Compute/Data/PCGProxyForGPUData.h b/Engine/Plugins/PCG/Source/PCG/Internal/Compute/Data/PCGProxyForGPUData.h
--- a/Engine/Plugins/PCG/Source/PCG/Internal/Compute/Data/PCGProxyForGPUData.h
+++ b/Engine/Plugins/PCG/Source/PCG/Internal/Compute/Data/PCGProxyForGPUData.h
@@ -18,6 +18,10 @@ struct FPCGDataTypeInfoProxyForGPU : public FPCGDataTypeInfo
 	GENERATED_BODY()
 
 	PCG_DECLARE_LEGACY_TYPE_INFO(EPCGDataType::ProxyForGPU)
+
+#if WITH_EDITOR
+	virtual bool Hidden() const override { return true; }
+#endif // WITH_EDITOR
 };
 
 /** A proxy for data residing on the GPU with functionality to read the data back to the CPU. */
diff --git a/Engine/Plugins/PCG/Source/PCG/Private/Data/Registry/PCGDataType.cpp b/Engine/Plugins/PCG/Source/PCG/Private/Data/Registry/PCGDataType.cpp
--- a/Engine/Plugins/PCG/Source/PCG/Private/Data/Registry/PCGDataType.cpp
+++ b/Engine/Plugins/PCG/Source/PCG/Private/Data/Registry/PCGDataType.cpp
@@ -36,7 +36,7 @@ bool FPCGDataTypeBaseId::IsValid() const
 
 bool FPCGDataTypeBaseId::IsChildOf(const FPCGDataTypeBaseId& Other) const
 {
-	return IsValid() && Other.IsValid() && Struct->IsChildOf(Other.Struct);
+	return Struct && Struct->IsChildOf(Other.Struct);
 }
 
 FString FPCGDataTypeBaseId::ToString() const
@@ -90,6 +90,18 @@ TSubclassOf<UPCGSettings> FPCGDataTypeInfo::GetDefaultConversionSettings(const F
 	return nullptr;
 }
 
+#if WITH_EDITOR
+FText FPCGDataTypeInfo::GetTooltip() const
+{
+	return FText{};
+}
+
+bool FPCGDataTypeInfo::Hidden() const
+{
+	return false;
+}
+#endif //WITH_EDITOR
+
 FPCGDataTypeBaseId FPCGDataTypeBaseId::MakeFromLegacyType(EPCGDataType DataType)
 {
 	switch (DataType)
@@ -136,9 +148,8 @@ FPCGDataTypeBaseId FPCGDataTypeBaseId::MakeFromLegacyType(EPCGDataType DataType)
 		return Construct<FPCGDataTypeInfoProxyForGPU>();
 	case EPCGDataType::Param:
 		return Construct<FPCGDataTypeInfoParam>();
-	// @todo_pcg To be uncommented when we do the switch. 
-	// case EPCGDataType::Settings:
-	// 	return Construct<FPCGDataTypeInfoSettings>();
+	//case EPCGDataType::Settings:
+	//	return Construct<FPCGDataTypeInfoSettings>();
 	case EPCGDataType::Other:
 		return Construct<FPCGDataTypeInfoOther>();
 	case EPCGDataType::None: // fall-through
@@ -176,8 +187,7 @@ EPCGDataType FPCGDataTypeBaseId::AsLegacyType() const
 	PCG_DATA_TYPE_TEST(FPCGDataTypeInfoSpatial)
 	PCG_DATA_TYPE_TEST(FPCGDataTypeInfoProxyForGPU)
 	PCG_DATA_TYPE_TEST(FPCGDataTypeInfoParam)
-	// @todo_pcg To be uncommented when we do the switch.
-	// PCG_DATA_TYPE_TEST(FPCGDataTypeInfoSettings)
+	//PCG_DATA_TYPE_TEST(FPCGDataTypeInfoSettings)
 	PCG_DATA_TYPE_TEST(FPCGDataTypeInfoOther)
 		
 #undef PCG_DATA_TYPE_TEST
diff --git a/Engine/Plugins/PCG/Source/PCG/Private/Data/Registry/PCGDataTypeIdentifier.cpp b/Engine/Plugins/PCG/Source/PCG/Private/Data/Registry/PCGDataTypeIdentifier.cpp
--- a/Engine/Plugins/PCG/Source/PCG/Private/Data/Registry/PCGDataTypeIdentifier.cpp
+++ b/Engine/Plugins/PCG/Source/PCG/Private/Data/Registry/PCGDataTypeIdentifier.cpp
@@ -49,7 +49,7 @@ FPCGDataTypeIdentifier::FPCGDataTypeIdentifier(EPCGDataType LegacyType)
 	FPCGDataTypeBaseId BaseId = FPCGDataTypeBaseId::MakeFromLegacyType(LegacyType);
 	if (BaseId.IsValid())
 	{
-		IDs.Emplace(BaseId);
+		Ids.Emplace(BaseId);
 		return;
 	}
 
@@ -61,7 +61,7 @@ FPCGDataTypeIdentifier::FPCGDataTypeIdentifier(EPCGDataType LegacyType)
 			BaseId = FPCGDataTypeBaseId::MakeFromLegacyType(Type);
 			if (BaseId.IsValid())
 			{
-				IDs.Add(BaseId);
+				Ids.Add(BaseId);
 				LegacyType &= ~Type;
 			}
 		}
@@ -81,9 +81,9 @@ FPCGDataTypeIdentifier::FPCGDataTypeIdentifier(TConstArrayView<FPCGDataTypeIdent
 {
 	for (const FPCGDataTypeIdentifier& ID : InIDs)
 	{
-		for (const FPCGDataTypeBaseId& BaseID : ID.GetIDs())
+		for (const FPCGDataTypeBaseId& BaseID : ID.GetIds())
 		{
-			IDs.AddUnique(BaseID);
+			Ids.AddUnique(BaseID);
 		}
 	}
 }
@@ -92,14 +92,14 @@ FPCGDataTypeIdentifier::FPCGDataTypeIdentifier(TConstArrayView<TSubclassOf<UPCGD
 {
 	for (TSubclassOf<UPCGData> Class : Classes)
 	{
-		IDs.AddUnique(Class->GetDefaultObject<UPCGData>()->GetDataTypeId());
+		Ids.AddUnique(Class->GetDefaultObject<UPCGData>()->GetDataTypeId());
 	}
 }
 
 FPCGDataTypeIdentifier::operator EPCGDataType() const
 {
 	EPCGDataType Result = EPCGDataType::None;
-	for (const FPCGDataTypeBaseId& ID : IDs)
+	for (const FPCGDataTypeBaseId& ID : Ids)
 	{
 		Result |= ID.AsLegacyType();
 	}
@@ -115,11 +115,11 @@ FPCGDataTypeIdentifier& FPCGDataTypeIdentifier::operator|=(EPCGDataType OtherLeg
 
 FPCGDataTypeIdentifier& FPCGDataTypeIdentifier::operator|=(const FPCGDataTypeIdentifier& Other)
 {
-	for (const FPCGDataTypeBaseId& ID : Other.GetIDs())
+	for (const FPCGDataTypeBaseId& ID : Other.GetIds())
 	{
 		if (ID.IsValid())
 		{
-			IDs.AddUnique(ID);
+			Ids.AddUnique(ID);
 		}
 	}
 	
@@ -148,14 +148,14 @@ FPCGDataTypeIdentifier& FPCGDataTypeIdentifier::operator&=(const FPCGDataTypeIde
 {
 	TArray<FPCGDataTypeBaseId> CommonTypes;
 	
-	for (const FPCGDataTypeBaseId& ThisID : IDs)
+	for (const FPCGDataTypeBaseId& ThisID : Ids)
 	{
 		if (!ensure(ThisID.IsValid()))
 		{
 			continue;
 		}
 		
-		for (const FPCGDataTypeBaseId& OtherID : Other.GetIDs())
+		for (const FPCGDataTypeBaseId& OtherID : Other.GetIds())
 		{
 			if (!ensure(OtherID.IsValid()))
 			{
@@ -173,7 +173,7 @@ FPCGDataTypeIdentifier& FPCGDataTypeIdentifier::operator&=(const FPCGDataTypeIde
 		}
 	}
 
-	IDs = MoveTemp(CommonTypes);
+	Ids = MoveTemp(CommonTypes);
 	
 	return *this;
 }
@@ -185,11 +185,45 @@ FPCGDataTypeIdentifier FPCGDataTypeIdentifier::operator&(EPCGDataType OtherLegac
 
 FPCGDataTypeIdentifier FPCGDataTypeIdentifier::operator&(const FPCGDataTypeIdentifier& Other) const
 {
-	FPCGDataTypeIdentifier Result = *this;
-	Result &= Other;
+	TArray<FPCGDataTypeBaseId> CommonTypes;
+	
+	for (const FPCGDataTypeBaseId& ThisID : Ids)
+	{
+		for (const FPCGDataTypeBaseId& OtherID : Other.GetIds())
+		{
+			if (ThisID.IsChildOf(OtherID))
+			{
+				CommonTypes.AddUnique(ThisID);
+			}
+			else if (OtherID.IsChildOf(ThisID))
+			{
+				CommonTypes.AddUnique(OtherID);
+			}
+		}
+	}
+
+	FPCGDataTypeIdentifier Result;
+	Result.Ids = MoveTemp(CommonTypes);
+
 	return Result;
 }
 
+bool FPCGDataTypeIdentifier::Intersects(const FPCGDataTypeIdentifier& Other) const
+{
+	for (const FPCGDataTypeBaseId& ThisID : Ids)
+	{
+		for (const FPCGDataTypeBaseId& OtherID : Other.GetIds())
+		{
+			if (ThisID.IsChildOf(OtherID) || OtherID.IsChildOf(ThisID))
+			{
+				return true;
+			}
+		}
+	}
+
+	return false;
+}
+
 FPCGDataTypeIdentifier FPCGDataTypeIdentifier::Compose(const FPCGDataTypeIdentifier& Other)
 {
 	return Compose({*this, Other});
@@ -210,7 +244,7 @@ bool FPCGDataTypeIdentifier::SupportsType(EPCGDataType OtherLegacyType) const
 	const FPCGDataTypeIdentifier OtherID(OtherLegacyType);
 	if (OtherID.IsComposition())
 	{
-		return Algo::AllOf(OtherID.GetIDs(), [this](const FPCGDataTypeBaseId& ID) { return IDs.Contains(ID); });
+		return Algo::AllOf(OtherID.GetIds(), [this](const FPCGDataTypeBaseId& ID) { return Ids.Contains(ID); });
 	}
 	else
 	{
@@ -218,19 +252,31 @@ bool FPCGDataTypeIdentifier::SupportsType(EPCGDataType OtherLegacyType) const
 	}
 }
 
-TConstArrayView<FPCGDataTypeBaseId> FPCGDataTypeIdentifier::GetIDs() const
+FPCGDataTypeBaseId FPCGDataTypeIdentifier::GetId() const
+{
+	if (ensure(!IsComposition() && IsValid()))
+	{
+		return Ids[0];
+	}
+	else
+	{
+		return FPCGDataTypeBaseId{};
+	}
+}
+
+TConstArrayView<FPCGDataTypeBaseId> FPCGDataTypeIdentifier::GetIds() const
 {
-	return IDs;
+	return Ids;
 };
 
 bool FPCGDataTypeIdentifier::IsSameType(const FPCGDataTypeIdentifier& Other) const
 {
-	return (!IsValid() && !Other.IsValid()) || (IDs.Num() == Other.IDs.Num() && SupportsType(Other));
+	return (!IsValid() && !Other.IsValid()) || (Ids.Num() == Other.Ids.Num() && SupportsType(Other));
 }
 
 bool FPCGDataTypeIdentifier::SupportsType(const FPCGDataTypeIdentifier& Other) const
 {
-	return IsValid() && Algo::AllOf(Other.GetIDs(), [this](const FPCGDataTypeBaseId& ID) { return IDs.Contains(ID); });
+	return IsValid() && Algo::AllOf(Other.GetIds(), [this](const FPCGDataTypeBaseId& ID) { return Ids.Contains(ID); });
 }
 
 bool FPCGDataTypeIdentifier::IsWider(const FPCGDataTypeIdentifier& Other) const
@@ -241,7 +287,7 @@ bool FPCGDataTypeIdentifier::IsWider(const FPCGDataTypeIdentifier& Other) const
 	}
 	
 	const FPCGDataTypeIdentifier Intersection = *this & Other;
-	return Intersection != *this;
+	return !Intersection.IsSameType(*this);
 }
 
 bool FPCGDataTypeIdentifier::IsIdentical(const FPCGDataTypeIdentifier& Other) const
@@ -266,7 +312,7 @@ bool FPCGDataTypeIdentifier::SerializeFromMismatchedTag(const struct FPropertyTa
         	if (EnumValue != INDEX_NONE)
         	{
         		FPCGDataTypeIdentifier Temp = FPCGDataTypeIdentifier(static_cast<EPCGDataType>(EnumValue));
-        		IDs = Temp.GetIDs();
+        		Ids = Temp.GetIds();
         	}
         }
 
@@ -280,7 +326,7 @@ bool FPCGDataTypeIdentifier::SerializeFromMismatchedTag(const struct FPropertyTa
 
 FString FPCGDataTypeIdentifier::ToString() const
 {
-	return IsValid() ? FString::JoinBy(IDs, TEXT(" | "), [](const FPCGDataTypeBaseId& ID) { return ID.ToString(); }) : TEXT("None");
+	return IsValid() ? FString::JoinBy(Ids, TEXT(" | "), [](const FPCGDataTypeBaseId& ID) { return ID.ToString(); }) : TEXT("None");
 }
 
 FText FPCGDataTypeIdentifier::ToDisplayText() const
@@ -288,7 +334,7 @@ FText FPCGDataTypeIdentifier::ToDisplayText() const
 	return FText::FromString(ToString());
 }
 
-FPCGDataTypeIdentifier UPCGDataTypeIdentifierHelpers::GetIdentifierFromDefaultType(EPCGExclusiveDataType DataType)
+FPCGDataTypeIdentifier UPCGDataTypeIdentifierHelpers::GetIdentifierFromLegacyType(EPCGExclusiveDataType LegacyDataType)
 {
 	const UEnum* DataTypeEnum = StaticEnum<EPCGDataType>();
 	const UEnum* ExclusiveDataTypeEnum = StaticEnum<EPCGExclusiveDataType>();
@@ -297,7 +343,7 @@ FPCGDataTypeIdentifier UPCGDataTypeIdentifierHelpers::GetIdentifierFromDefaultTy
 
 	if (DataTypeEnum && ExclusiveDataTypeEnum)
 	{
-		FName ExclusiveDataTypeName = ExclusiveDataTypeEnum->GetNameByValue(static_cast<__underlying_type(EPCGExclusiveDataType)>(DataType));
+		FName ExclusiveDataTypeName = ExclusiveDataTypeEnum->GetNameByValue(static_cast<__underlying_type(EPCGExclusiveDataType)>(LegacyDataType));
 		if (ensure(ExclusiveDataTypeName != NAME_None))
 		{
 			const int64 MatchingType = DataTypeEnum->GetValueByName(ExclusiveDataTypeName);
diff --git a/Engine/Plugins/PCG/Source/PCG/Private/Data/Registry/PCGDataTypeRegistry.cpp b/Engine/Plugins/PCG/Source/PCG/Private/Data/Registry/PCGDataTypeRegistry.cpp
--- a/Engine/Plugins/PCG/Source/PCG/Private/Data/Registry/PCGDataTypeRegistry.cpp
+++ b/Engine/Plugins/PCG/Source/PCG/Private/Data/Registry/PCGDataTypeRegistry.cpp
@@ -48,24 +48,17 @@ struct FTree
 	 * returns false if identifiers are invalid.
 	 */
 	bool QueryHierarchy(const FPCGDataTypeIdentifier& TypeA, const FPCGDataTypeIdentifier& TypeB, FPCGDataTypeIdentifier& OutCommonAncestor) const;
-
-	enum class ESearchCommand
-	{
-		ExpandAndContinue,
-		Continue,
-		Stop
-	};
 	
-	template <typename Func> requires std::is_invocable_r_v<FTree::ESearchCommand, Func, FNode*, int32>
+	template <typename Func> requires std::is_invocable_r_v<FPCGDataTypeRegistry::ESearchCommand, Func, FNode*, int32>
 	void BreadthSearch(Func InFunc) const;
 
-	template <typename Func> requires std::is_invocable_r_v<FTree::ESearchCommand, Func, FNode*, int32>
+	template <typename Func> requires std::is_invocable_r_v<FPCGDataTypeRegistry::ESearchCommand, Func, FNode*, int32>
 	void DepthSearch(Func InFunc) const;
 
 	FNode Root;
 };
 	
-template <typename Func> requires std::is_invocable_r_v<FTree::ESearchCommand, Func, FNode*, int32>
+template <typename Func> requires std::is_invocable_r_v<FPCGDataTypeRegistry::ESearchCommand, Func, FNode*, int32>
 void FTree::BreadthSearch(Func InFunc) const
 {
 	// Const-cast to not have 2 versions, const and non-const, since those are private.
@@ -78,14 +71,14 @@ void FTree::BreadthSearch(Func InFunc) const
 		FNode* Curr = List[i].Get<0>();
 		const int32 Depth = List[i++].Get<1>();
 
-		const ESearchCommand Command = InFunc(Curr, Depth);
+		const FPCGDataTypeRegistry::ESearchCommand Command = InFunc(Curr, Depth);
 	
-		if (Command == ESearchCommand::Stop)
+		if (Command == FPCGDataTypeRegistry::ESearchCommand::Stop)
 		{
 			return;
 		}
 
-		if (Command == ESearchCommand::ExpandAndContinue)
+		if (Command == FPCGDataTypeRegistry::ESearchCommand::ExpandAndContinue)
 		{
 			for (FNode* Child : Curr->Children)
 			{
@@ -95,7 +88,7 @@ void FTree::BreadthSearch(Func InFunc) const
 	}
 }
 
-template <typename Func> requires std::is_invocable_r_v<FTree::ESearchCommand, Func, FNode*, int32>
+template <typename Func> requires std::is_invocable_r_v<FPCGDataTypeRegistry::ESearchCommand, Func, FNode*, int32>
 void FTree::DepthSearch(Func InFunc) const
 {
 	// Const-cast to not have 2 versions, const and non-const, since those are private.
@@ -108,14 +101,14 @@ void FTree::DepthSearch(Func InFunc) const
 		FNode* Curr = It.Get<0>();
 		const int32 Depth = It.Get<1>();
 	
-		const ESearchCommand Command = InFunc(Curr, Depth);
+		const FPCGDataTypeRegistry::ESearchCommand Command = InFunc(Curr, Depth);
 	
-		if (Command == ESearchCommand::Stop)
+		if (Command == FPCGDataTypeRegistry::ESearchCommand::Stop)
 		{
 			return;
 		}
 
-		if (Command == ESearchCommand::ExpandAndContinue)
+		if (Command == FPCGDataTypeRegistry::ESearchCommand::ExpandAndContinue)
 		{
 			for (FNode* Child : Curr->Children)
 			{
@@ -153,7 +146,7 @@ void FTree::AddEntry(const FPCGDataTypeBaseId& BaseID)
 
 	const UScriptStruct* EntryStruct = BaseID.GetStruct();
 
-	BreadthSearch([&Parent, EntryStruct](FNode* Curr, int32 Depth) -> ESearchCommand
+	BreadthSearch([&Parent, EntryStruct](FNode* Curr, int32 Depth) -> FPCGDataTypeRegistry::ESearchCommand
 	{
 		if (EntryStruct->IsChildOf(Curr->BaseID.GetStruct()))
 		{
@@ -161,7 +154,7 @@ void FTree::AddEntry(const FPCGDataTypeBaseId& BaseID)
 		}
 
 		// Don't stop until the end.
-		return ESearchCommand::ExpandAndContinue;
+		return FPCGDataTypeRegistry::ESearchCommand::ExpandAndContinue;
 	});
 
 	// If the parent have children, we need to check for all of them if they are a child of the new entry
@@ -186,7 +179,7 @@ void FTree::AddEntry(const FPCGDataTypeBaseId& BaseID)
 
 void FTree::RemoveEntry(const FPCGDataTypeBaseId& BaseID)
 {
-	DepthSearch([&BaseID](FNode* Curr, int32 Depth) -> ESearchCommand
+	DepthSearch([&BaseID](FNode* Curr, int32 Depth) -> FPCGDataTypeRegistry::ESearchCommand
 	{
 		if (Curr->BaseID == BaseID)
 		{
@@ -196,15 +189,18 @@ void FTree::RemoveEntry(const FPCGDataTypeBaseId& BaseID)
 			}
 
 			Curr->Parent->Children.Remove(Curr);
+			// Transfer ownership
 			Curr->Parent->Children.Append(Curr->Children);
+			// Make sure to empty to avoid the delete to destroy the children.
+			Curr->Children.Empty();
 
 			delete Curr;
 		
-			return ESearchCommand::Stop;
+			return FPCGDataTypeRegistry::ESearchCommand::Stop;
 		}
 		else
 		{
-			return ESearchCommand::ExpandAndContinue;
+			return FPCGDataTypeRegistry::ESearchCommand::ExpandAndContinue;
 		}
 	});
 }
@@ -242,7 +238,7 @@ bool FTree::QueryHierarchy(const FPCGDataTypeIdentifier& TypeA, const FPCGDataTy
 	// Breadth search to find which comes first in the hierarchy
 	const FNode* FirstFound = nullptr;
 	const FNode* SecondFound = nullptr;
-	BreadthSearch([&FirstFound, &SecondFound, &TypeA, &TypeB](const FNode* Curr, int32 Depth) -> ESearchCommand
+	BreadthSearch([&FirstFound, &SecondFound, &TypeA, &TypeB](const FNode* Curr, int32 Depth) -> FPCGDataTypeRegistry::ESearchCommand
 	{
 		bool bFound = false;
 		const FPCGDataTypeIdentifier CurrID{Curr->BaseID};
@@ -252,7 +248,7 @@ bool FTree::QueryHierarchy(const FPCGDataTypeIdentifier& TypeA, const FPCGDataTy
 			if (!FirstFound)
 			{
 				FirstFound = Curr;
-				return ESearchCommand::ExpandAndContinue;
+				return FPCGDataTypeRegistry::ESearchCommand::ExpandAndContinue;
 			}
 
 			bFound = true;
@@ -263,7 +259,7 @@ bool FTree::QueryHierarchy(const FPCGDataTypeIdentifier& TypeA, const FPCGDataTy
 			if (!FirstFound)
 			{
 				FirstFound = Curr;
-				return ESearchCommand::ExpandAndContinue;
+				return FPCGDataTypeRegistry::ESearchCommand::ExpandAndContinue;
 			}
 
 			bFound = true;
@@ -271,12 +267,12 @@ bool FTree::QueryHierarchy(const FPCGDataTypeIdentifier& TypeA, const FPCGDataTy
 
 		if (!bFound)
 		{
-			return ESearchCommand::ExpandAndContinue;
+			return FPCGDataTypeRegistry::ESearchCommand::ExpandAndContinue;
 		}
 		else
 		{
 			SecondFound = Curr;
-			return ESearchCommand::Stop;
+			return FPCGDataTypeRegistry::ESearchCommand::Stop;
 		}
 	});
 
@@ -323,7 +319,7 @@ bool FTree::QueryHierarchy(const FPCGDataTypeIdentifier& TypeA, const FPCGDataTy
 
 FPCGDataTypeRegistry::FPCGDataTypeRegistry()
 {
-	HierarchyTree = new PCGDataTypeHierarchy::FTree();
+	HierarchyTree = MakeUnique<PCGDataTypeHierarchy::FTree>();
 	
 	// Pre-register the base type
 	Mapping.Emplace(FPCGDataTypeInfo::StaticStruct(), FPCGDataTypeInfo::StaticStruct());
@@ -342,10 +338,7 @@ void FPCGDataTypeRegistry::RegisterEntry(const FPCGDataTypeBaseId& BaseID)
 	HierarchyTree->AddEntry(BaseID);
 }
 
-FPCGDataTypeRegistry::~FPCGDataTypeRegistry()
-{
-	delete HierarchyTree;
-}
+FPCGDataTypeRegistry::~FPCGDataTypeRegistry() = default;
 
 EPCGDataTypeCompatibilityResult FPCGDataTypeRegistry::IsCompatible(const FPCGDataTypeIdentifier& InType, const FPCGDataTypeIdentifier& OutType) const
 {
@@ -374,8 +367,8 @@ EPCGDataTypeCompatibilityResult FPCGDataTypeRegistry::IsCompatible(const FPCGDat
 	{
 		if (!InType.IsComposition() && !OutType.IsComposition())
 		{
-			const TInstancedStruct<FPCGDataTypeInfo>* InDefIt = Mapping.Find(InType.GetID());
-			const TInstancedStruct<FPCGDataTypeInfo>* OutDefIt = Mapping.Find(OutType.GetID());
+			const TInstancedStruct<FPCGDataTypeInfo>* InDefIt = Mapping.Find(InType.GetId());
+			const TInstancedStruct<FPCGDataTypeInfo>* OutDefIt = Mapping.Find(OutType.GetId());
 
 			if (!InDefIt || !OutDefIt)
 			{
@@ -398,8 +391,8 @@ EPCGDataTypeCompatibilityResult FPCGDataTypeRegistry::IsCompatible(const FPCGDat
 	}
 
 	// Finally we check if we have a conversion
-	const TInstancedStruct<FPCGDataTypeInfo>* InDefIt = Mapping.Find(InType.GetID());
-	const TInstancedStruct<FPCGDataTypeInfo>* OutDefIt = Mapping.Find(OutType.GetID());
+	const TInstancedStruct<FPCGDataTypeInfo>* InDefIt = Mapping.Find(InType.GetId());
+	const TInstancedStruct<FPCGDataTypeInfo>* OutDefIt = Mapping.Find(OutType.GetId());
 
 	if (!InDefIt || !OutDefIt)
 	{
@@ -436,7 +429,7 @@ FPCGDataTypeIdentifier FPCGDataTypeRegistry::GetIdentifiersUnion(TConstArrayView
 
 	for (const FPCGDataTypeIdentifier& Identifier : Identifiers)
 	{
-		for (const FPCGDataTypeBaseId& ID : Identifier.GetIDs())
+		for (const FPCGDataTypeBaseId& ID : Identifier.GetIds())
 		{
 			// If we have the "Any" type in our composition, it is the parent of all the types, so the composition is "Any"
 			if (ID == AnyType)
@@ -510,7 +503,7 @@ FPCGDataTypeIdentifier FPCGDataTypeRegistry::GetIdentifiersComposition(TConstArr
 
 	for (const FPCGDataTypeIdentifier& Identifier : Identifiers)
 	{
-		for (const FPCGDataTypeBaseId& ID : Identifier.GetIDs())
+		for (const FPCGDataTypeBaseId& ID : Identifier.GetIds())
 		{
 			// If we have the "Any" type in our composition, it is the parent of all the types, so the composition is "Any"
 			if (ID == AnyType)
@@ -538,11 +531,11 @@ FPCGDataTypeIdentifier FPCGDataTypeRegistry::GetIdentifiersComposition(TConstArr
 		{
 			Nodes.Emplace(Node, Depth);
 			// We found one, no need to check the children, as they won't contribute to the composition.
-			return PCGDataTypeHierarchy::FTree::ESearchCommand::Continue;
+			return ESearchCommand::Continue;
 		}
 
 		// Continue until we find them all.
-		return PCGDataTypeHierarchy::FTree::ESearchCommand::ExpandAndContinue;
+		return ESearchCommand::ExpandAndContinue;
 	});
 
 	// Sort by depth, starting at the deepest.
@@ -589,15 +582,111 @@ FPCGDataTypeIdentifier FPCGDataTypeRegistry::GetIdentifiersComposition(TConstArr
 
 	// Returns the aggregation of the types after the reduce
 	FPCGDataTypeIdentifier Result{};
-	Result.IDs.Reserve(Nodes.Num());
+	Result.Ids.Reserve(Nodes.Num());
 	for (const auto& Node : Nodes)
 	{
-		Result.IDs.Emplace(Node.Get<0>()->BaseID);
+		Result.Ids.Emplace(Node.Get<0>()->BaseID);
 	}
 
 	return Result;
 }
 
+/**
+ * Goal: From a source and list of difference, return the source type without any type that are in the difference.
+ * Algorithm overview:
+ * - Get all the leaves for the source type
+ * - Remove any leaf that is a child of a type in the difference
+ * - Re-compose everything.
+ */
+FPCGDataTypeIdentifier FPCGDataTypeRegistry::GetIdentifiersDifference(const FGetIdentifiersDifferenceParams& Params) const
+{
+	if (Params.DifferenceIdentifiers.IsEmpty() || !Params.SourceIdentifier || !Params.SourceIdentifier->IsValid())
+	{
+		return Params.SourceIdentifier ? *Params.SourceIdentifier : FPCGDataTypeIdentifier{};
+	}
+
+	const FPCGDataTypeIdentifier& SourceIdentifier = *Params.SourceIdentifier; 
+	
+	// Extract all leaves
+	TArray<FPCGDataTypeBaseId, TInlineAllocator<32>> SourceLeaves;
+	int32 CurrentFoundDepth = -1;
+	int32 Found = 0;
+
+	HierarchyTree->DepthSearch([&CurrentFoundDepth, &SourceIdentifier, &Params, &Found, &SourceLeaves](PCGDataTypeHierarchy::FNode* Curr, int32 Depth) -> ESearchCommand
+	{
+		if (CurrentFoundDepth == Depth)
+		{
+			CurrentFoundDepth = -1;
+		}
+			
+		if (CurrentFoundDepth == -1 && SourceIdentifier.GetIds().Contains(Curr->BaseID))
+		{
+			CurrentFoundDepth = Depth;
+			Found++;
+		}
+			
+		if (CurrentFoundDepth != -1 && CurrentFoundDepth <= Depth && Curr->Children.IsEmpty())
+		{
+			const bool bShouldInclude = (Params.Filter == FGetIdentifiersDifferenceParams::IncludeFilteredTypes) == (Params.FilteredTypes.Contains(Curr->BaseID));
+
+			if (bShouldInclude)
+			{
+				SourceLeaves.Add(Curr->BaseID);
+			}
+		}
+
+		return CurrentFoundDepth == -1 && Found == SourceIdentifier.GetIds().Num() ? ESearchCommand::Stop : ESearchCommand::ExpandAndContinue;
+	});
+
+	bool bHasRemoved = false;
+	for (const FPCGDataTypeIdentifier& Difference : Params.DifferenceIdentifiers)
+	{
+		for (const FPCGDataTypeBaseId& DifferenceId : Difference.GetIds())
+		{
+			while (!SourceLeaves.IsEmpty())
+			{
+				int32 Index = SourceLeaves.IndexOfByPredicate([&DifferenceId](const FPCGDataTypeBaseId& SourceId) { return SourceId.IsChildOf(DifferenceId);});
+				if (Index == INDEX_NONE)
+				{
+					break;
+				}
+
+				bHasRemoved = true;
+				SourceLeaves.RemoveAtSwap(Index);
+			}
+
+			if (SourceLeaves.IsEmpty())
+			{
+				break;
+			}
+		}
+
+		if (SourceLeaves.IsEmpty())
+		{
+			break;
+		}
+	}
+
+	if (!bHasRemoved)
+	{
+		return SourceIdentifier;
+	}
+
+	TArray<FPCGDataTypeIdentifier, TInlineAllocator<32>> FinalIdentifiers;
+	FinalIdentifiers.Reserve(SourceLeaves.Num());
+	Algo::Transform(SourceLeaves, FinalIdentifiers, [](const FPCGDataTypeBaseId& Id) { return Id; });
+	return GetIdentifiersComposition(FinalIdentifiers);
+}
+
+void FPCGDataTypeRegistry::HierarchyDepthSearch(TFunctionRef<ESearchCommand(const FPCGDataTypeBaseId& Id, int32 Depth)> Callback) const
+{
+	HierarchyTree->DepthSearch([&Callback](const PCGDataTypeHierarchy::FNode* Node, int32 Depth)
+	{
+		check(Node);
+		return Callback(Node->BaseID, Depth);
+	});
+}
+
 void FPCGDataTypeRegistry::RegisterKnownTypes()
 {
 	const UScriptStruct* BaseStruct = FPCGDataTypeInfo::StaticStruct();
@@ -612,7 +701,6 @@ void FPCGDataTypeRegistry::RegisterKnownTypes()
 	}
 }
 
-
 // Tests for hierarchy
 #if WITH_EDITOR
 
diff --git a/Engine/Plugins/PCG/Source/PCG/Private/Tests/Registry/PCGDataTypeRegistryTests.cpp b/Engine/Plugins/PCG/Source/PCG/Private/Tests/Registry/PCGDataTypeRegistryTests.cpp
--- a/Engine/Plugins/PCG/Source/PCG/Private/Tests/Registry/PCGDataTypeRegistryTests.cpp
+++ b/Engine/Plugins/PCG/Source/PCG/Private/Tests/Registry/PCGDataTypeRegistryTests.cpp
@@ -124,7 +124,7 @@ bool FPCGByDefaultTypeTest::RunTest(const FString& Parameters)
 	for (const auto& [TypeName, Identifier] : TestValues)
 	{
 		UTEST_TRUE(FString::Printf(TEXT("%s type is valid"), TypeName), Identifier.IsValid());
-		bool bRegistered = Algo::AllOf(Identifier.GetIDs(), [](const FPCGDataTypeBaseId ID) { return IsRegistered(ID); });
+		bool bRegistered = Algo::AllOf(Identifier.GetIds(), [](const FPCGDataTypeBaseId ID) { return IsRegistered(ID); });
 		
 		UTEST_TRUE(FString::Printf(TEXT("%s type is registered"), TypeName), bRegistered);
 		
@@ -449,32 +449,32 @@ bool FPCGCompositionTest::RunTest(const FString& Parameters)
 bool FPCGUnicityTest::RunTest(const FString& Parameters)
 {
 	const FPCGDataTypeIdentifier SameTypeWithClassTpl = FPCGDataTypeIdentifier::Construct<UPCGData, UPCGData>();
-	UTEST_EQUAL("Construct with same class in template gives just 1 ID", SameTypeWithClassTpl.GetIDs().Num(), 1);
+	UTEST_EQUAL("Construct with same class in template gives just 1 ID", SameTypeWithClassTpl.GetIds().Num(), 1);
 
 	const FPCGDataTypeIdentifier SameTypeWithClass = FPCGDataTypeIdentifier({UPCGPointData::StaticClass(), UPCGPointData::StaticClass()});
-	UTEST_EQUAL("Construct with same class gives just 1 ID", SameTypeWithClass.GetIDs().Num(), 1);
+	UTEST_EQUAL("Construct with same class gives just 1 ID", SameTypeWithClass.GetIds().Num(), 1);
 
 	const FPCGDataTypeIdentifier SameTypeWithBaseIds = FPCGDataTypeIdentifier({FPCGDataTypeBaseId::Construct<FPCGDataTypeInfoLandscape>(), FPCGDataTypeBaseId::Construct<FPCGDataTypeInfoLandscape>()});
-	UTEST_EQUAL("Construct with same base ids gives just 1 ID", SameTypeWithBaseIds.GetIDs().Num(), 1);
+	UTEST_EQUAL("Construct with same base ids gives just 1 ID", SameTypeWithBaseIds.GetIds().Num(), 1);
 
 	const FPCGDataTypeIdentifier SameTypeWithIds = FPCGDataTypeIdentifier({SameTypeWithClassTpl, SameTypeWithClassTpl, SameTypeWithClassTpl});
-	UTEST_EQUAL("Construct with same ids gives just 1 ID", SameTypeWithIds.GetIDs().Num(), 1);
+	UTEST_EQUAL("Construct with same ids gives just 1 ID", SameTypeWithIds.GetIds().Num(), 1);
 
 	const FPCGDataTypeIdentifier SameTypeWithDefaultTypeOr = SameTypeWithClass | EPCGDataType::Point;
-	UTEST_EQUAL("Bitwise Or with same default type gives just 1 ID", SameTypeWithDefaultTypeOr.GetIDs().Num(), 1);
+	UTEST_EQUAL("Bitwise Or with same default type gives just 1 ID", SameTypeWithDefaultTypeOr.GetIds().Num(), 1);
 
 	const FPCGDataTypeIdentifier SameTypeWithOr = SameTypeWithClass | SameTypeWithClass;
-	UTEST_EQUAL("Bitwise Or with same ids gives just 1 ID", SameTypeWithOr.GetIDs().Num(), 1);
+	UTEST_EQUAL("Bitwise Or with same ids gives just 1 ID", SameTypeWithOr.GetIds().Num(), 1);
 
 	FPCGDataTypeIdentifier SameTypeOr = SameTypeWithClass;
 	SameTypeOr |= SameTypeWithClass;
-	UTEST_EQUAL("Bitwise Or with self gives just 1 ID", SameTypeOr.GetIDs().Num(), 1);
+	UTEST_EQUAL("Bitwise Or with self gives just 1 ID", SameTypeOr.GetIds().Num(), 1);
 
 	SameTypeOr |= EPCGDataType::Point;
-	UTEST_EQUAL("Bitwise Or on self with default type gives just 1 ID", SameTypeOr.GetIDs().Num(), 1);
+	UTEST_EQUAL("Bitwise Or on self with default type gives just 1 ID", SameTypeOr.GetIds().Num(), 1);
 
 	const FPCGDataTypeIdentifier SameTypeCompose = SameTypeWithClass.Compose({SameTypeWithClass, SameTypeWithClass});
-	UTEST_EQUAL("Compose with same type gives just 1 ID", SameTypeCompose.GetIDs().Num(), 1);
+	UTEST_EQUAL("Compose with same type gives just 1 ID", SameTypeCompose.GetIds().Num(), 1);
 
 	return true;
 }
diff --git a/Engine/Plugins/PCG/Source/PCG/Public/Data/PCGDynamicMeshData.h b/Engine/Plugins/PCG/Source/PCG/Public/Data/PCGDynamicMeshData.h
--- a/Engine/Plugins/PCG/Source/PCG/Public/Data/PCGDynamicMeshData.h
+++ b/Engine/Plugins/PCG/Source/PCG/Public/Data/PCGDynamicMeshData.h
@@ -17,7 +17,7 @@ class UDynamicMeshComponent;
 class UMaterialInterface;
 
 USTRUCT()
-struct FPCGDataTypeInfoDynamicMesh: public FPCGDataTypeInfo
+struct FPCGDataTypeInfoDynamicMesh: public FPCGDataTypeInfoConcrete
 {
 	GENERATED_BODY()
 
diff --git a/Engine/Plugins/PCG/Source/PCG/Public/Data/PCGResourceData.h b/Engine/Plugins/PCG/Source/PCG/Public/Data/PCGResourceData.h
--- a/Engine/Plugins/PCG/Source/PCG/Public/Data/PCGResourceData.h
+++ b/Engine/Plugins/PCG/Source/PCG/Public/Data/PCGResourceData.h
@@ -16,6 +16,10 @@ struct FPCGDataTypeInfoResource : public FPCGDataTypeInfo
 	GENERATED_BODY()
 
 	PCG_DECLARE_LEGACY_TYPE_INFO(EPCGDataType::Resource)
+
+#if WITH_EDITOR
+	virtual bool Hidden() const override { return true; }
+#endif // WITH_EDITOR
 };
 
 /** Data that wrap/represent an asset, like a Static Mesh or Texture. */
diff --git a/Engine/Plugins/PCG/Source/PCG/Public/Data/PCGSpatialData.h b/Engine/Plugins/PCG/Source/PCG/Public/Data/PCGSpatialData.h
--- a/Engine/Plugins/PCG/Source/PCG/Public/Data/PCGSpatialData.h
+++ b/Engine/Plugins/PCG/Source/PCG/Public/Data/PCGSpatialData.h
@@ -87,6 +87,10 @@ struct FPCGDataTypeInfoComposite : public FPCGDataTypeInfoSpatial
 	GENERATED_BODY()
 
 	PCG_DECLARE_LEGACY_TYPE_INFO(EPCGDataType::Composite)
+
+#if WITH_EDITOR
+	virtual bool Hidden() const override { return true; }
+#endif // WITH_EDITOR
 };
 
 /**
diff --git a/Engine/Plugins/PCG/Source/PCG/Public/Data/Registry/PCGDataType.h b/Engine/Plugins/PCG/Source/PCG/Public/Data/Registry/PCGDataType.h
--- a/Engine/Plugins/PCG/Source/PCG/Public/Data/Registry/PCGDataType.h
+++ b/Engine/Plugins/PCG/Source/PCG/Public/Data/Registry/PCGDataType.h
@@ -26,7 +26,7 @@ class UPCGSettings;
 #define PCG_DECLARE_TYPE_INFO(DLL_EXPORT_NAME) \
 	DLL_EXPORT_NAME static TSubclassOf<UPCGData> GetStaticAssociatedClass(); \
 	DLL_EXPORT_NAME virtual TSubclassOf<UPCGData> GetAssociatedClass() const override; \
-	static FPCGDataTypeBaseId AsID() { return FPCGDataTypeBaseId(StaticStruct()); }
+	static FPCGDataTypeBaseId AsId() { return FPCGDataTypeBaseId(StaticStruct()); }
 
 #define PCG_DEFINE_TYPE_INFO(PCGDataTypeInfoClass, PCGDataClass) \
 	TSubclassOf<UPCGData> PCGDataTypeInfoClass::GetStaticAssociatedClass() { return PCGDataClass::StaticClass(); } \
@@ -72,6 +72,13 @@ struct FPCGDataTypeInfo
 	PCG_API virtual bool SupportsConversion(const FPCGDataTypeIdentifier& Other) const;
 	PCG_API virtual UPCGNode* AddConversionNode(const FPCGDataTypeIdentifier& Other, UPCGGraph* InGraph, UPCGPin* InUpstreamPin, UPCGPin* InDownstreamPin) const;
 	PCG_API virtual TSubclassOf<UPCGSettings> GetDefaultConversionSettings(const FPCGDataTypeIdentifier& Other) const;
+
+#if WITH_EDITOR
+	PCG_API virtual FText GetTooltip() const;
+	
+	// If the type should be hidden to the user.
+	PCG_API virtual bool Hidden() const;
+#endif // WITH_EDITOR
 };
 
 /**
diff --git a/Engine/Plugins/PCG/Source/PCG/Public/Data/Registry/PCGDataTypeIdentifier.h b/Engine/Plugins/PCG/Source/PCG/Public/Data/Registry/PCGDataTypeIdentifier.h
--- a/Engine/Plugins/PCG/Source/PCG/Public/Data/Registry/PCGDataTypeIdentifier.h
+++ b/Engine/Plugins/PCG/Source/PCG/Public/Data/Registry/PCGDataTypeIdentifier.h
@@ -1,4 +1,4 @@
-﻿// Copyright Epic Games, Inc. All Rights Reserved.
+// Copyright Epic Games, Inc. All Rights Reserved.
 
 #pragma once
 
@@ -24,15 +24,15 @@ struct FPCGDataTypeIdentifier
 
 	FPCGDataTypeIdentifier() = default;
 
-	FPCGDataTypeIdentifier(FPCGDataTypeBaseId ID)
+	FPCGDataTypeIdentifier(FPCGDataTypeBaseId Id)
 	{
-		IDs.Emplace(ID);
+		Ids.Emplace(Id);
 	}
 	
 	PCG_API explicit FPCGDataTypeIdentifier(TSubclassOf<UPCGData> Class);
 
     /** Composition ctors. */
-    PCG_API explicit FPCGDataTypeIdentifier(TConstArrayView<FPCGDataTypeIdentifier> InIDs);
+    PCG_API explicit FPCGDataTypeIdentifier(TConstArrayView<FPCGDataTypeIdentifier> InIds);
     PCG_API explicit FPCGDataTypeIdentifier(TConstArrayView<TSubclassOf<UPCGData>> Classes);
 
 	// ------- Begin deprecation with EPCGDataType --------
@@ -84,6 +84,9 @@ struct FPCGDataTypeIdentifier
 	PCG_API FPCGDataTypeIdentifier& operator&=(const FPCGDataTypeIdentifier& Other);
 	PCG_API FPCGDataTypeIdentifier operator&(const FPCGDataTypeIdentifier& Other) const;
 
+	// Faster to determine if there is intersection instead of !!(A & B)
+	PCG_API bool Intersects(const FPCGDataTypeIdentifier& Other) const;
+
 	// Return the composition with other types, with reduction. (ie. if Data B and C are child of Data A, Composition of B and C gives A)
 	// CAREFUL: This can't be used before the Data Type Registry is done initializing (so can't be used before the Engine Post Init).
 	PCG_API FPCGDataTypeIdentifier Compose(const FPCGDataTypeIdentifier& Other);
@@ -93,25 +96,25 @@ struct FPCGDataTypeIdentifier
 	static FPCGDataTypeIdentifier Construct()
 	{
 		FPCGDataTypeIdentifier Result;
-		Result.IDs.Reserve(sizeof...(T));
-		(Result.IDs.AddUnique(FPCGDataTypeIdentifier(T::StaticClass()).GetID()), ...);
+		Result.Ids.Reserve(sizeof...(T));
+		(Result.Ids.AddUnique(FPCGDataTypeIdentifier(T::StaticClass()).GetId()), ...);
 		return Result;
 	}
 
-	/** Can't be called when the identifier is a composition. Need to use GetIDs in that case. */
-	FPCGDataTypeBaseId GetID() const { check(!IsComposition() && IsValid()); return IDs[0]; };
+	/** Can't be called when the identifier is a composition. Need to use GetIds in that case. */
+	PCG_API FPCGDataTypeBaseId GetId() const;
 	
-	PCG_API TConstArrayView<FPCGDataTypeBaseId> GetIDs() const;
+	PCG_API TConstArrayView<FPCGDataTypeBaseId> GetIds() const;
 
 	/** A type identifier can be a composition of multiple types. */
-	bool IsComposition() const {return IDs.Num() > 1; }
+	bool IsComposition() const {return Ids.Num() > 1; }
 	
-	bool IsValid() const { return !IDs.IsEmpty() && Algo::AllOf(IDs, [](const FPCGDataTypeBaseId& ID) { return ID.IsValid(); }); }
+	bool IsValid() const { return !Ids.IsEmpty() && Algo::AllOf(Ids, [](const FPCGDataTypeBaseId& Id) { return Id.IsValid(); }); }
 
 	/** Return true if the IDs are the same, but not necessary the subtype. */
 	PCG_API bool IsSameType(const FPCGDataTypeIdentifier& Other) const;
 
-	/** Equivalent to IsSameType if it is not a composition. Otherwise, return true if the OtherID is in the IDs. */
+	/** Equivalent to IsSameType if it is not a composition. Otherwise, return true if the OtherId is in the Ids. */
 	PCG_API bool SupportsType(const FPCGDataTypeIdentifier& Other) const;
 
 	/** Return true if this identifier supports different/more types than Other. Equivalent of !!(*this & ~Other) in bit flags form. */
@@ -123,9 +126,9 @@ struct FPCGDataTypeIdentifier
 	/** To support previously saved types, that used EPCGDataType, we need to define this function to de-serialize the new class using the old. And add a trait (see below). */
 	PCG_API bool SerializeFromMismatchedTag(const struct FPropertyTag& Tag, FStructuredArchive::FSlot Slot);
 
-	friend uint32 GetTypeHash(const FPCGDataTypeIdentifier& ID)
+	friend uint32 GetTypeHash(const FPCGDataTypeIdentifier& Id)
 	{
-		return HashCombine(GetTypeHash(ID.IDs), ID.CustomSubtype);
+		return HashCombine(GetTypeHash(Id.Ids), Id.CustomSubtype);
 	}
 
 	PCG_API FString ToString() const;
@@ -133,7 +136,7 @@ struct FPCGDataTypeIdentifier
 	
 private:
 	UPROPERTY(BlueprintReadOnly, VisibleAnywhere, Category = "PCG", meta = (AllowPrivateAccess))
-	TArray<FPCGDataTypeBaseId> IDs;
+	TArray<FPCGDataTypeBaseId> Ids;
 
 public:
 	UPROPERTY(BlueprintReadWrite, EditAnywhere, Category = "PCG")
@@ -162,7 +165,7 @@ public:
 	static FPCGDataTypeIdentifier GetIdentifierFromClass(TSubclassOf<UPCGData> Class) { return FPCGDataTypeIdentifier(Class); };
 
 	UFUNCTION(BlueprintCallable, Category = "PCG|PCGDataTypeIdentifier")
-	static PCG_API FPCGDataTypeIdentifier GetIdentifierFromDefaultType(EPCGExclusiveDataType DataType);
+	static PCG_API FPCGDataTypeIdentifier GetIdentifierFromLegacyType(EPCGExclusiveDataType LegacyDataType);
 
 	UFUNCTION(BlueprintCallable, Category = "PCG|PCGDataTypeIdentifier", meta = (ScriptMethod))
 	static bool IsValid(UPARAM(Ref) const FPCGDataTypeIdentifier& This) { return This.IsValid(); }
diff --git a/Engine/Plugins/PCG/Source/PCG/Public/Data/Registry/PCGDataTypeRegistry.h b/Engine/Plugins/PCG/Source/PCG/Public/Data/Registry/PCGDataTypeRegistry.h
--- a/Engine/Plugins/PCG/Source/PCG/Public/Data/Registry/PCGDataTypeRegistry.h
+++ b/Engine/Plugins/PCG/Source/PCG/Public/Data/Registry/PCGDataTypeRegistry.h
@@ -3,6 +3,7 @@
 #pragma once
 
 #include "PCGCommon.h"
+#include "Containers/AdderRef.h"
 
 #include "Data/Registry/PCGDataTypeIdentifier.h"
 #include "StructUtils/InstancedStruct.h"
@@ -70,6 +71,34 @@ public:
 
 	/** Returns the data type for the composition of all the identifiers passed as input. It will be simplified if possible (like composition of spline and landscape spline gives a polyline). */
 	PCG_API FPCGDataTypeIdentifier GetIdentifiersComposition(TConstArrayView<FPCGDataTypeIdentifier> Identifiers) const;
+
+	struct FGetIdentifiersDifferenceParams
+	{
+		const FPCGDataTypeIdentifier* SourceIdentifier = nullptr;
+		TConstArrayView<FPCGDataTypeIdentifier> DifferenceIdentifiers;
+
+		enum EFilter
+		{
+			IncludeFilteredTypes,
+			ExcludeFilteredTypes
+		};
+
+		EFilter Filter = ExcludeFilteredTypes;
+		TConstArrayView<FPCGDataTypeBaseId> FilteredTypes;
+	};
+
+	/** Returns the data type for the difference between the Source and all the differences. Example: Difference(Spatial, Composite) gives Concrete. If there is no intersection between both, it returns SourceIdentifier. */
+	PCG_API FPCGDataTypeIdentifier GetIdentifiersDifference(const FGetIdentifiersDifferenceParams& Params) const;
+
+	enum class ESearchCommand
+	{
+		ExpandAndContinue,
+		Continue,
+		Stop
+	};
+	
+	/** Do a depth search in the hierarchy tree. Functor should return ESearchCommand to continue/stop. ExpandAndContinue will continue searching deeper in a branch. */
+	PCG_API void HierarchyDepthSearch(TFunctionRef<ESearchCommand(const FPCGDataTypeBaseId& Id, int32 Depth)> Callback) const;
 	
 private:
 	friend FPCGDataTypeRegistryTestBase;
@@ -82,7 +111,6 @@ private:
 	bool IsRegistered(const FPCGDataTypeBaseId& BaseID) const { return Mapping.Contains(BaseID); }
 	
 	TMap<FPCGDataTypeBaseId, TInstancedStruct<FPCGDataTypeInfo>> Mapping;
-	// Can't be a UniquePtr because of incomplete type.
-	PCGDataTypeHierarchy::FTree* HierarchyTree = nullptr;
+	TUniquePtr<PCGDataTypeHierarchy::FTree> HierarchyTree;
 };
 
diff --git a/Engine/Plugins/PCG/Source/PCG/Public/PCGCommon.h b/Engine/Plugins/PCG/Source/PCG/Public/PCGCommon.h
--- a/Engine/Plugins/PCG/Source/PCG/Public/PCGCommon.h
+++ b/Engine/Plugins/PCG/Source/PCG/Public/PCGCommon.h
@@ -140,7 +140,12 @@ enum class EPCGDataType : uint32
 
 	Settings = 1 << 28 UMETA(Hidden),
 	Other = 1 << 29,
-	Any = (1 << 30) - 1
+	Any = (1 << 30) - 1,
+
+	// Used as an indication that a function that returns a EPCGDataType that has gone through a default implementation.
+	// If a function returns something else than a DeprecationSentinel, it means the function is overridden, and we want to use the result,
+	// even if we have a loss of info. Show deprecation message and indicate how to upgrade the code.
+	DeprecationSentinel = (1u << 31) UMETA(Hidden)
 };
 ENUM_CLASS_FLAGS(EPCGDataType);
 
@@ -243,6 +248,7 @@ namespace PCGObjectMetadata
 	const FLazyName DiscardExtraSelection = TEXT("PCG_DiscardExtraSelection");
 	const FLazyName EnumMetadataDomain = TEXT("PCG_MetadataDomain");
 	const FLazyName PropertyReadOnly = TEXT("PCG_PropertyReadOnly");
+	const FLazyName DataTypeIdentifierSupportsComposition = TEXT("PCG_SupportsComposition");
 
 	// Metadata usable in UPROPERTY for customizing the behavior when displaying the property in a property panel or graph node
 	enum
@@ -277,6 +283,9 @@ namespace PCGObjectMetadata
 
 		/// [PropertyMetadata] For FPCGAttributePropertySelector, won't display the extra property items in the dropdown
 		PCG_DiscardExtraSelection,
+
+		/// [PropertyMetadata] For FPCGDataTypeIdentifier, will enable selecting types like BitFlag, to compose types.
+		PCG_SupportsComposition,
 	};
 
 	// Metadata usable in UENUM for customizing the entry in the attribute property selector
diff --git a/Engine/Plugins/PCG/Source/PCG/Public/PCGData.h b/Engine/Plugins/PCG/Source/PCG/Public/PCGData.h
--- a/Engine/Plugins/PCG/Source/PCG/Public/PCGData.h
+++ b/Engine/Plugins/PCG/Source/PCG/Public/PCGData.h
@@ -81,7 +81,7 @@ public:
 
 	/** If this data is a proxy, returns the underlying data type, otherwise returns this data type. */
 	// @todo_pcg To be deprecated when we switch to the new type
-	// UE_DEPRECATED(5.7, "Use the new GetUnderlyingDataTypeInfo")
+	// UE_DEPRECATED(5.7, "Use the new GetUnderlyingDataTypeId")
 	virtual EPCGDataType GetUnderlyingDataType() const
 	{
 		return GetUnderlyingDataTypeId().AsLegacyType();
diff --git a/Engine/Plugins/PCG/Source/PCG/Public/PCGInputOutputSettings.h b/Engine/Plugins/PCG/Source/PCG/Public/PCGInputOutputSettings.h
--- a/Engine/Plugins/PCG/Source/PCG/Public/PCGInputOutputSettings.h
+++ b/Engine/Plugins/PCG/Source/PCG/Public/PCGInputOutputSettings.h
@@ -71,6 +71,11 @@ public:
 	// the one passed as argument. It will change if its label collides with existing pins.
 	[[nodiscard]] PCG_API const FPCGPinProperties& AddPin(const FPCGPinProperties& NewCustomPinProperties);
 
+#if WITH_EDITOR
+	UFUNCTION()
+	static bool SupportsComposition() { return true; }
+#endif // WITH_EDITOR
+
 protected:
 	TArray<FPCGPinProperties> DefaultPinProperties(bool bInvisiblePin) const;
 	void FixPinProperties();
diff --git a/Engine/Plugins/PCG/Source/PCGEditor/Private/Details/PCGDataTypeIdentifierDetails.cpp b/Engine/Plugins/PCG/Source/PCGEditor/Private/Details/PCGDataTypeIdentifierDetails.cpp
new file mode 100644
--- /dev/null
+++ b/Engine/Plugins/PCG/Source/PCGEditor/Private/Details/PCGDataTypeIdentifierDetails.cpp
@@ -0,0 +1,220 @@
+// Copyright Epic Games, Inc. All Rights Reserved.
+
+#include "Details/PCGDataTypeIdentifierDetails.h"
+
+#include "PCGModule.h"
+#include "Helpers/PCGHelpers.h"
+
+#include "DetailLayoutBuilder.h"
+#include "DetailWidgetRow.h"
+#include "ScopedTransaction.h"
+#include "Framework/Application/SlateApplication.h"
+#include "Framework/MultiBox/MultiBoxBuilder.h"
+#include "Widgets/Input/SComboButton.h"
+
+namespace PCGDataTypeIdentifierDetails
+{
+	// Check if the identifier is a superset of other.
+	bool IsSuperSet(const FPCGDataTypeIdentifier& Identifier, const FPCGDataTypeBaseId& Other)
+	{
+		return (Identifier & Other) == Other;
+	}
+
+	/**
+	 * Look for a function on the outer object of the property or a flag in the property metadata to check if the widget
+	 * should support composition (have its bitflag-like behavior enabled) or not. 
+	 */
+	bool SupportComposition(TSharedRef<IPropertyHandle> InPropertyHandle)
+	{
+		const FString& SupportCompositionFunctionName = InPropertyHandle->GetMetaData(PCGObjectMetadata::DataTypeIdentifierSupportsComposition);
+		if (!SupportCompositionFunctionName.IsEmpty())
+		{
+			TArray<UObject*> OutObjects;
+			InPropertyHandle->GetOuterObjects(OutObjects);
+
+			if (OutObjects.IsEmpty() || !OutObjects[0])
+			{
+				return false;
+			}
+
+			UObject* Object = OutObjects[0];
+			// Strip `()` that could be at the end of the function name
+			const FString FunctionName = SupportCompositionFunctionName.Replace(TEXT("()"), TEXT(""));
+			
+			if (!Object->GetClass() || !Object->GetClass()->FindFunctionByName(*FunctionName))
+			{
+				return false;
+			}
+
+			DECLARE_DELEGATE_RetVal(bool, FSupportComposition);
+			auto Func = FSupportComposition::CreateUFunction(Object, *FunctionName);
+			return Func.IsBound() ? Func.Execute() : false;
+		}
+		else if (InPropertyHandle->HasMetaData(PCGObjectMetadata::DataTypeIdentifierSupportsComposition))
+		{
+			return true;
+		}
+		else
+		{
+			return false;
+		}
+	}
+}
+
+void FPCGDataTypeIdentifierDetails::CustomizeHeader(TSharedRef<IPropertyHandle> InPropertyHandle, FDetailWidgetRow& HeaderRow,
+	IPropertyTypeCustomizationUtils& CustomizationUtils)
+{
+	PropertyHandle = InPropertyHandle;
+	check(PropertyHandle.IsValid());
+
+	TSharedRef<SWidget> PropertyNameWidget = PropertyHandle->CreatePropertyNameWidget();
+	
+	bSupportComposition = PCGDataTypeIdentifierDetails::SupportComposition(InPropertyHandle);
+
+	const FPCGDataTypeRegistry& Registry = FPCGModule::GetConstDataTypeRegistry();
+	
+	HiddenTypes.Empty();
+	VisibleTypes.Empty();
+
+	// Do a depth search so types are ordered by category.
+	Registry.HierarchyDepthSearch([this, &Registry](const FPCGDataTypeBaseId& DataType, int32 Depth) -> FPCGDataTypeRegistry::ESearchCommand
+	{
+		const FPCGDataTypeInfo* TypeInfo = Registry.GetTypeInfo(DataType);
+		check(TypeInfo);
+
+		// Remove hidden types
+		if (TypeInfo->Hidden())
+		{
+			HiddenTypes.Add(DataType);
+		}
+		else
+		{
+			VisibleTypes.Emplace(DataType, Depth);
+		}
+
+		return FPCGDataTypeRegistry::ESearchCommand::ExpandAndContinue;
+	});
+
+	HeaderRow
+		.NameContent()
+		[
+			PropertyNameWidget
+		]
+		.ValueContent()
+		[
+			SNew(SComboButton)
+			.OnGetMenuContent(this, &FPCGDataTypeIdentifierDetails::OnGetMenuContent)
+			.ButtonContent()
+			[
+				SNew(STextBlock)
+				.Text(this, &FPCGDataTypeIdentifierDetails::GetText)
+				.ToolTipText(this, &FPCGDataTypeIdentifierDetails::GetTooltip)
+				.Font(IDetailLayoutBuilder::GetDetailFont())
+			]
+		];
+}
+
+TSharedRef<SWidget> FPCGDataTypeIdentifierDetails::OnGetMenuContent()
+{
+	const bool bCloseAfterSelection = !bSupportComposition;
+	FMenuBuilder MenuBuilder(bCloseAfterSelection, nullptr, nullptr, true);
+
+	const FPCGDataTypeRegistry& Registry = FPCGModule::GetConstDataTypeRegistry();
+
+	for (const auto& [DataType, Depth] : VisibleTypes)
+	{
+		const FPCGDataTypeInfo* TypeInfo = Registry.GetTypeInfo(DataType);
+		check(TypeInfo);
+		
+		MenuBuilder.AddMenuEntry(
+			FText::FromString(DataType.ToString()),
+			TypeInfo->GetTooltip(),
+			FSlateIcon(),
+			FUIAction
+			(
+				FExecuteAction::CreateLambda([this, DataType = DataType]()
+				{
+					FPCGDataTypeIdentifier* Identifier = GetStruct();
+					if (!ensure(Identifier))
+					{
+						return;
+					}
+
+					FScopedTransaction Transaction(NSLOCTEXT("PCGDataTypeIdentifierDetails", "ModifyIdentifier", "Modify PCG Data Type Identifier."));
+
+					PropertyHandle->NotifyPreChange();
+
+					if (bSupportComposition && FSlateApplication::Get().GetModifierKeys().IsControlDown())
+					{
+						if (!PCGDataTypeIdentifierDetails::IsSuperSet(*Identifier, DataType))
+						{
+							*Identifier = Identifier->Compose(DataType);
+						}
+						else
+						{
+							const FPCGDataTypeIdentifier DataTypeIdentifier{DataType};
+							FPCGDataTypeRegistry::FGetIdentifiersDifferenceParams Params = {
+								.SourceIdentifier = Identifier,
+								.DifferenceIdentifiers = MakeConstArrayView<FPCGDataTypeIdentifier>(&DataTypeIdentifier, 1),
+								.Filter = FPCGDataTypeRegistry::FGetIdentifiersDifferenceParams::ExcludeFilteredTypes,
+								.FilteredTypes = HiddenTypes
+							};
+							
+							FPCGDataTypeIdentifier TempIdentifier = FPCGModule::GetConstDataTypeRegistry().GetIdentifiersDifference(Params);
+							if (TempIdentifier.IsValid())
+							{
+								*Identifier = MoveTemp(TempIdentifier);
+							}
+						}
+					}
+					else
+					{
+						*Identifier = DataType;
+					}
+
+					PropertyHandle->NotifyPostChange(EPropertyChangeType::ValueSet);
+				}),
+				FCanExecuteAction(),
+				FIsActionChecked::CreateLambda([this, DataType = DataType]() -> bool
+				{
+					FPCGDataTypeIdentifier* Identifier = GetStruct();
+					return ensure(Identifier) ? PCGDataTypeIdentifierDetails::IsSuperSet(*Identifier, DataType) : false;
+				})
+			),
+			NAME_None,
+			bSupportComposition ? EUserInterfaceActionType::Check : EUserInterfaceActionType::None);
+	}
+
+	return MenuBuilder.MakeWidget();
+}
+
+FText FPCGDataTypeIdentifierDetails::GetText() const
+{
+	static const FPCGDataTypeIdentifier DefaultInvalidDataType{};
+
+	const FPCGDataTypeIdentifier* DataType = GetStruct();
+	return DataType ? DataType->ToDisplayText() : DefaultInvalidDataType.ToDisplayText();
+}
+
+FText FPCGDataTypeIdentifierDetails::GetTooltip() const
+{
+	static const FText CompositionTooltip = NSLOCTEXT("PCGDataTypeIdentifierDetails", "CompositionTooltip", "\n---\nTypes can be composed like bitflags.\n"
+	"By default, selecting a type will remove all the others, no composition.\n"
+	"A broader type will select all its subtypes (like BaseTexture selects Texture and RenderTarget, or Any selects everything)\n"
+	"Use 'Ctrl + Click' to add/remove another type to the composition.");
+
+	return FText::Format(INVTEXT("{0}{1}"), GetText(), bSupportComposition ? CompositionTooltip : FText{});
+}
+
+FPCGDataTypeIdentifier* FPCGDataTypeIdentifierDetails::GetStruct()
+{
+	void* Data = nullptr;
+	FPropertyAccess::Result Result = PropertyHandle->GetValueData(Data);
+
+	return (Result == FPropertyAccess::Success) ? static_cast<FPCGDataTypeIdentifier*>(Data) : nullptr;
+}
+
+const FPCGDataTypeIdentifier* FPCGDataTypeIdentifierDetails::GetStruct() const
+{
+	return const_cast<FPCGDataTypeIdentifierDetails&>(*this).GetStruct();
+}
diff --git a/Engine/Plugins/PCG/Source/PCGEditor/Private/PCGEditorModule.cpp b/Engine/Plugins/PCG/Source/PCGEditor/Private/PCGEditorModule.cpp
--- a/Engine/Plugins/PCG/Source/PCGEditor/Private/PCGEditorModule.cpp
+++ b/Engine/Plugins/PCG/Source/PCGEditor/Private/PCGEditorModule.cpp
@@ -38,11 +38,13 @@
 #include "DataVisualizations/PCGSplineDataVisualization.h"
 #include "DataVisualizations/PCGStaticMeshDataVisualization.h"
 #include "DataVisualizations/PCGVolumeDataVisualization.h"
+#include "Details/EnumSelectorDetails.h"
 #include "Details/PCGAttributePropertySelectorDetails.h"
 #include "Details/PCGBlueprintSettingsDetails.h"
 #include "Details/PCGComponentDetails.h"
 #include "Details/PCGComputeSourceDetails.h"
 #include "Details/PCGCustomHLSLSettingsDetails.h"
+#include "Details/PCGDataTypeIdentifierDetails.h"
 #include "Details/PCGEditableUserParameterDetails.h"
 #include "Details/PCGGraphDetails.h"
 #include "Details/PCGGraphInstanceDetails.h"
@@ -64,7 +66,6 @@
 #include "PropertyEditorModule.h"
 #include "ScopedTransaction.h"
 #include "ToolMenus.h"
-#include "Details/EnumSelectorDetails.h"
 #include "Editor/SceneOutliner/Public/ISceneOutliner.h"
 #include "EditorMode/PCGEdMode.h"
 #include "Framework/MultiBox/MultiBoxBuilder.h"
@@ -304,6 +305,7 @@ void FPCGEditorModule::RegisterDetailsCustomizations()
 	PropertyEditor.RegisterCustomPropertyTypeLayout("PCGAttributePropertyOutputNoSourceSelector", FOnGetPropertyTypeCustomizationInstance::CreateStatic(&FPCGAttributePropertySelectorDetails::MakeInstance));
 	PropertyEditor.RegisterCustomPropertyTypeLayout("PCGOverrideInstancedPropertyBag", FOnGetPropertyTypeCustomizationInstance::CreateStatic(&FPCGOverrideInstancedPropertyBagDetails::MakeInstance));
 	PropertyEditor.RegisterCustomPropertyTypeLayout("EnumSelector", FOnGetPropertyTypeCustomizationInstance::CreateStatic(&FEnumSelectorDetails::MakeInstance));
+	PropertyEditor.RegisterCustomPropertyTypeLayout("PCGDataTypeIdentifier", FOnGetPropertyTypeCustomizationInstance::CreateStatic(&FPCGDataTypeIdentifierDetails::MakeInstance));
 
 	PropertyEditor.NotifyCustomizationModuleChanged();
 }
@@ -312,17 +314,22 @@ void FPCGEditorModule::UnregisterDetailsCustomizations()
 {
 	if (FPropertyEditorModule* PropertyModule = FModuleManager::GetModulePtr<FPropertyEditorModule>("PropertyEditor"))
 	{
-		PropertyModule->UnregisterCustomClassLayout("PCGBlueprintSettings");
-		PropertyModule->UnregisterCustomClassLayout("PCGComponent");
-		PropertyModule->UnregisterCustomClassLayout("PCGGraph");
-		PropertyModule->UnregisterCustomClassLayout("PCGGraphInstance");
+		PropertyModule->UnregisterCustomClassLayout("PCGComputeSource");
+		PropertyModule->UnregisterCustomClassLayout("PCGCustomHLSLSettings");
+		PropertyModule->UnregisterCustomClassLayout("PCGUserParameterGetSettings");
 		PropertyModule->UnregisterCustomClassLayout("PCGVolume");
+		PropertyModule->UnregisterCustomClassLayout("PCGGraphInstance");
+		PropertyModule->UnregisterCustomClassLayout("PCGGraph");
+		PropertyModule->UnregisterCustomClassLayout("PCGComponent");
+		PropertyModule->UnregisterCustomClassLayout("PCGBlueprintSettings");
 
-		PropertyModule->UnregisterCustomPropertyTypeLayout("PCGAttributePropertySelector");
-		PropertyModule->UnregisterCustomPropertyTypeLayout("PCGAttributePropertyInputSelector");
-		PropertyModule->UnregisterCustomPropertyTypeLayout("PCGAttributePropertyOutputSelector");
-		PropertyModule->UnregisterCustomPropertyTypeLayout("PCGAttributePropertyOutputNoSourceSelector");
+		PropertyModule->UnregisterCustomPropertyTypeLayout("PCGDataTypeIdentifier");
+		PropertyModule->UnregisterCustomPropertyTypeLayout("EnumSelector");
 		PropertyModule->UnregisterCustomPropertyTypeLayout("PCGOverrideInstancedPropertyBag");
+		PropertyModule->UnregisterCustomPropertyTypeLayout("PCGAttributePropertyOutputNoSourceSelector");
+		PropertyModule->UnregisterCustomPropertyTypeLayout("PCGAttributePropertyOutputSelector");
+		PropertyModule->UnregisterCustomPropertyTypeLayout("PCGAttributePropertyInputSelector");
+		PropertyModule->UnregisterCustomPropertyTypeLayout("PCGAttributePropertySelector");
 
 		PropertyModule->NotifyCustomizationModuleChanged();
 	}
diff --git a/Engine/Plugins/PCG/Source/PCGEditor/Public/Details/PCGDataTypeIdentifierDetails.h b/Engine/Plugins/PCG/Source/PCGEditor/Public/Details/PCGDataTypeIdentifierDetails.h
new file mode 100644
--- /dev/null
+++ b/Engine/Plugins/PCG/Source/PCGEditor/Public/Details/PCGDataTypeIdentifierDetails.h
@@ -0,0 +1,39 @@
+// Copyright Epic Games, Inc. All Rights Reserved.
+
+#pragma once
+
+#include "IPropertyTypeCustomization.h"
+
+#include "Data/Registry/PCGDataTypeIdentifier.h"
+
+#include "Layout/Visibility.h"
+
+class SWidget;
+
+class FPCGDataTypeIdentifierDetails : public IPropertyTypeCustomization
+{
+public:
+	static TSharedRef<IPropertyTypeCustomization> MakeInstance()
+	{
+		return MakeShareable(new FPCGDataTypeIdentifierDetails);
+	}
+
+	/** ~Begin IPropertyTypeCustomization interface */
+	virtual void CustomizeHeader(TSharedRef<IPropertyHandle> InPropertyHandle, FDetailWidgetRow& HeaderRow, IPropertyTypeCustomizationUtils& CustomizationUtils) override;
+	virtual void CustomizeChildren(TSharedRef<IPropertyHandle> InPropertyHandle, IDetailChildrenBuilder& ChildBuilder, IPropertyTypeCustomizationUtils& CustomizationUtils) override {};
+	/** ~End IPropertyTypeCustomization interface */
+
+protected:
+	FPCGDataTypeIdentifier* GetStruct();
+	const FPCGDataTypeIdentifier* GetStruct() const;
+	
+	FText GetText() const;
+	FText GetTooltip() const;
+	TSharedRef<SWidget> OnGetMenuContent();
+
+	TArray<FPCGDataTypeBaseId> HiddenTypes;
+	TArray<TTuple<FPCGDataTypeBaseId, int32>> VisibleTypes;
+
+	TSharedPtr<IPropertyHandle> PropertyHandle;
+	bool bSupportComposition = false;
+};
