[PCG] Fix build

#rnx

diff --git a/Engine/Plugins/PCG/Source/PCGEditor/Private/Details/PCGDataTypeIdentifierDetails.cpp b/Engine/Plugins/PCG/Source/PCGEditor/Private/Details/PCGDataTypeIdentifierDetails.cpp
--- a/Engine/Plugins/PCG/Source/PCGEditor/Private/Details/PCGDataTypeIdentifierDetails.cpp
+++ b/Engine/Plugins/PCG/Source/PCGEditor/Private/Details/PCGDataTypeIdentifierDetails.cpp
@@ -132,7 +132,7 @@ TSharedRef<SWidget> FPCGDataTypeIdentifierDetails::OnGetMenuContent()
 			FSlateIcon(),
 			FUIAction
 			(
-				FExecuteAction::CreateLambda([this, DataType = DataType]()
+				FExecuteAction::CreateLambda([this, ThisDataType = DataType]()
 				{
 					FPCGDataTypeIdentifier* Identifier = GetStruct();
 					if (!ensure(Identifier))
@@ -146,13 +146,13 @@ TSharedRef<SWidget> FPCGDataTypeIdentifierDetails::OnGetMenuContent()
 
 					if (bSupportComposition && FSlateApplication::Get().GetModifierKeys().IsControlDown())
 					{
-						if (!PCGDataTypeIdentifierDetails::IsSuperSet(*Identifier, DataType))
+						if (!PCGDataTypeIdentifierDetails::IsSuperSet(*Identifier, ThisDataType))
 						{
-							*Identifier = Identifier->Compose(DataType);
+							*Identifier = Identifier->Compose(ThisDataType);
 						}
 						else
 						{
-							const FPCGDataTypeIdentifier DataTypeIdentifier{DataType};
+							const FPCGDataTypeIdentifier DataTypeIdentifier{ ThisDataType };
 							FPCGDataTypeRegistry::FGetIdentifiersDifferenceParams Params = {
 								.SourceIdentifier = Identifier,
 								.DifferenceIdentifiers = MakeConstArrayView<FPCGDataTypeIdentifier>(&DataTypeIdentifier, 1),
@@ -169,16 +169,16 @@ TSharedRef<SWidget> FPCGDataTypeIdentifierDetails::OnGetMenuContent()
 					}
 					else
 					{
-						*Identifier = DataType;
+						*Identifier = ThisDataType;
 					}
 
 					PropertyHandle->NotifyPostChange(EPropertyChangeType::ValueSet);
 				}),
 				FCanExecuteAction(),
-				FIsActionChecked::CreateLambda([this, DataType = DataType]() -> bool
+				FIsActionChecked::CreateLambda([this, ThisDataType = DataType]() -> bool
 				{
 					FPCGDataTypeIdentifier* Identifier = GetStruct();
-					return ensure(Identifier) ? PCGDataTypeIdentifierDetails::IsSuperSet(*Identifier, DataType) : false;
+					return ensure(Identifier) ? PCGDataTypeIdentifierDetails::IsSuperSet(*Identifier, ThisDataType) : false;
 				})
 			),
 			NAME_None,
