PCG: Expose getter for tracking priority in the "Get Execution Context Info" node.

diff --git a/Engine/Plugins/PCG/Source/PCG/Private/Elements/PCGGetExecutionContext.cpp b/Engine/Plugins/PCG/Source/PCG/Private/Elements/PCGGetExecutionContext.cpp
--- a/Engine/Plugins/PCG/Source/PCG/Private/Elements/PCGGetExecutionContext.cpp
+++ b/Engine/Plugins/PCG/Source/PCG/Private/Elements/PCGGetExecutionContext.cpp
@@ -76,46 +76,64 @@ bool FPCGGetExecutionContextElement::ExecuteInternal(FPCGContext* Context) const
 	const IPCGGraphExecutionSource* ExecutionSource = Context->ExecutionSource.Get();
 	const UPCGComponent* SourceComponent = Cast<UPCGComponent>(ExecutionSource);
 
-	bool Value = false;
-	const UWorld* SupportingWorld = ExecutionSource ? ExecutionSource->GetExecutionState().GetWorld() : nullptr;
-
-	if (Mode == EPCGGetExecutionContextMode::IsEditor || Mode == EPCGGetExecutionContextMode::IsRuntime)
-	{
-		const bool bContextIsRuntimeOrOnRuntimeGenerationMode = (SupportingWorld && SupportingWorld->IsGameWorld()) || PCGHelpers::IsRuntimeGeneration(ExecutionSource);
-		Value = (bContextIsRuntimeOrOnRuntimeGenerationMode == (Mode == EPCGGetExecutionContextMode::IsRuntime));
-	}
-	else if (Mode == EPCGGetExecutionContextMode::IsOriginal || Mode == EPCGGetExecutionContextMode::IsLocal)
-	{
-		Value = SourceComponent && (SourceComponent->IsLocalComponent() == (Mode == EPCGGetExecutionContextMode::IsLocal));
-	}
-	else if (Mode == EPCGGetExecutionContextMode::IsPartitioned)
-	{
-		Value = SourceComponent && SourceComponent->IsPartitioned();
-	}
-	else if (Mode == EPCGGetExecutionContextMode::IsRuntimeGeneration)
-	{
-		Value = PCGHelpers::IsRuntimeGeneration(ExecutionSource);
-	}
-	else if (Mode == EPCGGetExecutionContextMode::IsDedicatedServer)
+	auto IsBoolOutput = [](EPCGGetExecutionContextMode Mode)
 	{
-		Value = PCGHelpers::IsDedicatedServer(SupportingWorld);
-	}
-	else if (Mode == EPCGGetExecutionContextMode::IsListenServer)
-	{
-		Value = PCGHelpers::IsListenServer(SupportingWorld);
-	}
-	else if (Mode == EPCGGetExecutionContextMode::HasAuthority)
+		return Mode != EPCGGetExecutionContextMode::TrackingPriority;
+	};
+
+	if (IsBoolOutput(Mode))
 	{
-		Value = ExecutionSource && ExecutionSource->GetExecutionState().HasAuthority();
+		bool Value = false;
+		const UWorld* SupportingWorld = ExecutionSource ? ExecutionSource->GetExecutionState().GetWorld() : nullptr;
+
+		if (Mode == EPCGGetExecutionContextMode::IsEditor || Mode == EPCGGetExecutionContextMode::IsRuntime)
+		{
+			const bool bContextIsRuntimeOrOnRuntimeGenerationMode = (SupportingWorld && SupportingWorld->IsGameWorld()) || PCGHelpers::IsRuntimeGeneration(ExecutionSource);
+			Value = (bContextIsRuntimeOrOnRuntimeGenerationMode == (Mode == EPCGGetExecutionContextMode::IsRuntime));
+		}
+		else if (Mode == EPCGGetExecutionContextMode::IsOriginal || Mode == EPCGGetExecutionContextMode::IsLocal)
+		{
+			Value = SourceComponent && (SourceComponent->IsLocalComponent() == (Mode == EPCGGetExecutionContextMode::IsLocal));
+		}
+		else if (Mode == EPCGGetExecutionContextMode::IsPartitioned)
+		{
+			Value = SourceComponent && SourceComponent->IsPartitioned();
+		}
+		else if (Mode == EPCGGetExecutionContextMode::IsRuntimeGeneration)
+		{
+			Value = PCGHelpers::IsRuntimeGeneration(ExecutionSource);
+		}
+		else if (Mode == EPCGGetExecutionContextMode::IsDedicatedServer)
+		{
+			Value = PCGHelpers::IsDedicatedServer(SupportingWorld);
+		}
+		else if (Mode == EPCGGetExecutionContextMode::IsListenServer)
+		{
+			Value = PCGHelpers::IsListenServer(SupportingWorld);
+		}
+		else if (Mode == EPCGGetExecutionContextMode::HasAuthority)
+		{
+			Value = ExecutionSource && ExecutionSource->GetExecutionState().HasAuthority();
+		}
+		else if (Mode == EPCGGetExecutionContextMode::IsBuilder)
+		{
+#if WITH_EDITOR
+			Value = PCGHelpers::IsInWorldPartitionBuilderCommandletContext();
+#endif
+		}
+
+		FPCGMetadataAttribute<bool>* Attribute = ParamData->Metadata->CreateAttribute<bool>(PCGGetExecutionContextConstants::AttributeName, Value, /*bAllowInterpolation=*/false, /*bOverrideParent=*/false);
 	}
-	else if (Mode == EPCGGetExecutionContextMode::IsBuilder)
+	else
 	{
+		ensure(Mode == EPCGGetExecutionContextMode::TrackingPriority);
+		double TrackingPriority = 0.0;
 #if WITH_EDITOR
-		Value = PCGHelpers::IsInWorldPartitionBuilderCommandletContext();
+		TrackingPriority = ExecutionSource->GetExecutionState().GetDynamicTrackingPriority();
 #endif
+		FPCGMetadataAttribute<double>* Attribute = ParamData->Metadata->CreateAttribute<double>(PCGGetExecutionContextConstants::AttributeName, TrackingPriority, /*bAllowInterpolation=*/false, /*bOverrideParent=*/false);
 	}
 
-	FPCGMetadataAttribute<bool>* Attribute = ParamData->Metadata->CreateAttribute<bool>(PCGGetExecutionContextConstants::AttributeName, Value, /*bAllowInterpolation=*/false, /*bOverrideParent=*/false);
 	ParamData->Metadata->AddEntry();
 
 	return true;
diff --git a/Engine/Plugins/PCG/Source/PCG/Public/Elements/PCGGetExecutionContext.h b/Engine/Plugins/PCG/Source/PCG/Public/Elements/PCGGetExecutionContext.h
--- a/Engine/Plugins/PCG/Source/PCG/Public/Elements/PCGGetExecutionContext.h
+++ b/Engine/Plugins/PCG/Source/PCG/Public/Elements/PCGGetExecutionContext.h
@@ -20,6 +20,7 @@ enum class EPCGGetExecutionContextMode : uint8
 	IsListenServer UMETA(Tooltip="Returns whether the executing context is running on a listen server."),
 	HasAuthority UMETA(Tooltip="Returns whether the execution context is running with network authority."),
 	IsBuilder UMETA(Tooltip="Returns whether the execution source is generated by the PCG World Partition Builder. Always false in non-editor."),
+	TrackingPriority UMETA(Tooltip = "Gives tracking priority of the graph instance (lower values are unaffected by changes triggered by higher values). Always 0 in non-editor.")
 };
 
 /** Returns some context-specific information in the form of an attribute set. */
