[PCG] Fix CopyPointsAnalysisKernel incorrect pin type
Was succeeding anyways because Point and Param produce the same data interface, but could cause issue in the future.

#rnx

diff --git a/Engine/Plugins/PCG/Source/PCG/Private/Elements/PCGCopyPointsAnalysisKernel.cpp b/Engine/Plugins/PCG/Source/PCG/Private/Elements/PCGCopyPointsAnalysisKernel.cpp
--- a/Engine/Plugins/PCG/Source/PCG/Private/Elements/PCGCopyPointsAnalysisKernel.cpp
+++ b/Engine/Plugins/PCG/Source/PCG/Private/Elements/PCGCopyPointsAnalysisKernel.cpp
@@ -116,7 +116,7 @@ void UPCGCopyPointsAnalysisKernel::GetInputPins(TArray<FPCGPinProperties>& OutPi
 
 void UPCGCopyPointsAnalysisKernel::GetOutputPins(TArray<FPCGPinPropertiesGPU>& OutPins) const
 {
-	OutPins.Emplace(PCGPinConstants::DefaultOutputLabel, EPCGDataType::Point);
+	OutPins.Emplace(PCGPinConstants::DefaultOutputLabel, EPCGDataType::Param);
 }
 
 #undef LOCTEXT_NAMESPACE
