PCG: Fix cis complaint about unguarded ExecutionSource access in 45068873

#prelfight

diff --git a/Engine/Plugins/PCG/Source/PCG/Private/Elements/PCGGetExecutionContext.cpp b/Engine/Plugins/PCG/Source/PCG/Private/Elements/PCGGetExecutionContext.cpp
--- a/Engine/Plugins/PCG/Source/PCG/Private/Elements/PCGGetExecutionContext.cpp
+++ b/Engine/Plugins/PCG/Source/PCG/Private/Elements/PCGGetExecutionContext.cpp
@@ -129,7 +129,7 @@ bool FPCGGetExecutionContextElement::ExecuteInternal(FPCGContext* Context) const
 		ensure(Mode == EPCGGetExecutionContextMode::TrackingPriority);
 		double TrackingPriority = 0.0;
 #if WITH_EDITOR
-		TrackingPriority = ExecutionSource->GetExecutionState().GetDynamicTrackingPriority();
+		TrackingPriority = ExecutionSource ? ExecutionSource->GetExecutionState().GetDynamicTrackingPriority() : TrackingPriority;
 #endif
 		FPCGMetadataAttribute<double>* Attribute = ParamData->Metadata->CreateAttribute<double>(PCGGetExecutionContextConstants::AttributeName, TrackingPriority, /*bAllowInterpolation=*/false, /*bOverrideParent=*/false);
 	}
