[PCG] Switch to Identifiers instead of EPCGDataType enum

* Deprecation pass on all types comparison with PinProperties
* Special note on GPU stuff, as it makes it slower. To be monitored.

diff --git a/Engine/Plugins/PCG/Source/PCG/Internal/Compute/PCGComputeCommon.h b/Engine/Plugins/PCG/Source/PCG/Internal/Compute/PCGComputeCommon.h
--- a/Engine/Plugins/PCG/Source/PCG/Internal/Compute/PCGComputeCommon.h
+++ b/Engine/Plugins/PCG/Source/PCG/Internal/Compute/PCGComputeCommon.h
@@ -3,6 +3,7 @@
 #pragma once
 
 #include "PCGCommon.h"
+#include "Data/Registry/PCGDataTypeIdentifier.h"
 
 #include "PCGComputeCommon.generated.h"
 
@@ -107,16 +108,16 @@ namespace PCGComputeHelpers
 
 	/** Gets the element dimension for a given data. E.g. OneD for PointData, TwoD for TextureData, etc. */
 	EPCGElementDimension GetElementDimension(const UPCGData* InData);
-	EPCGElementDimension GetElementDimension(EPCGDataType InDataType);
+	EPCGElementDimension GetElementDimension(const FPCGDataTypeIdentifier& InDataType);
 
 	/** True if 'Type' is valid on a GPU input pin. */
-	bool IsTypeAllowedAsInput(EPCGDataType Type);
+	bool IsTypeAllowedAsInput(const FPCGDataTypeIdentifier& Type);
 
 	/** True if 'Type' is valid on a GPU output pin. */
-	bool IsTypeAllowedAsOutput(EPCGDataType Type);
+	bool IsTypeAllowedAsOutput(const FPCGDataTypeIdentifier& Type);
 
 	/** True if 'Type' is valid in a GPU data collection. Some types are only supported as DataInterfaces, and cannot be uploaded in data collections. */
-	bool IsTypeAllowedInDataCollection(EPCGDataType Type);
+	bool IsTypeAllowedInDataCollection(const FPCGDataTypeIdentifier& Type);
 
 	/** Whether metadata attributes should be read from the given data and registered for use in GPU graphs. */
 	bool ShouldImportAttributesFromData(const UPCGData* InData);
diff --git a/Engine/Plugins/PCG/Source/PCG/Internal/Compute/PCGDataForGPU.h b/Engine/Plugins/PCG/Source/PCG/Internal/Compute/PCGDataForGPU.h
--- a/Engine/Plugins/PCG/Source/PCG/Internal/Compute/PCGDataForGPU.h
+++ b/Engine/Plugins/PCG/Source/PCG/Internal/Compute/PCGDataForGPU.h
@@ -158,10 +158,10 @@ struct FPCGDataDesc
 {
 public:
 	FPCGDataDesc() = default;
-	FPCGDataDesc(EPCGDataType InType, int32 InElementCount);
-	FPCGDataDesc(EPCGDataType InType, FIntPoint InElementCount);
-	FPCGDataDesc(EPCGDataType InType, FIntVector3 InElementCount);
-	FPCGDataDesc(EPCGDataType InType, FIntVector4 InElementCount);
+	FPCGDataDesc(FPCGDataTypeIdentifier InType, int32 InElementCount);
+	FPCGDataDesc(FPCGDataTypeIdentifier InType, FIntPoint InElementCount);
+	FPCGDataDesc(FPCGDataTypeIdentifier InType, FIntVector3 InElementCount);
+	FPCGDataDesc(FPCGDataTypeIdentifier InType, FIntVector4 InElementCount);
 	FPCGDataDesc(const FPCGTaggedData& InTaggedData, const UPCGDataBinding* InBinding);
 
 	uint64 ComputePackedSize() const;
@@ -174,7 +174,7 @@ public:
 
 	void AddAttribute(FPCGKernelAttributeKey InAttribute, const UPCGDataBinding* InBinding, const TArray<int32>* InOptionalUniqueStringKeys = nullptr);
 
-	EPCGDataType GetType() const { return Type; }
+	const FPCGDataTypeIdentifier& GetType() const { return Type; }
 
 	FIntVector4 GetElementCount() const { return ElementCount; }
 	FIntVector4 GetElementCountForAttribute(const FPCGKernelAttributeDesc& AttributeDesc) const;
@@ -205,7 +205,7 @@ public:
 	TArray<int32, TInlineAllocator<4>> TagStringKeys;
 
 private:
-	EPCGDataType Type = EPCGDataType::Point;
+	FPCGDataTypeIdentifier Type = EPCGDataType::Point;
 	FIntVector4 ElementCount = FIntVector4::ZeroValue;
 	EPCGElementDimension ElementDimension = EPCGElementDimension::One;
 
@@ -238,7 +238,7 @@ struct FPCGDataCollectionDesc
 	uint32 ComputeTotalElementCount() const;
 
 	/** Compute total number of processing elements of the given type. */
-	uint32 ComputeTotalElementCount(EPCGDataType InDataType) const;
+	uint32 ComputeTotalElementCount(FPCGDataTypeIdentifier InDataType) const;
 
 	/** Aggregate another data description. */
 	void Combine(const FPCGDataCollectionDesc& Other);
diff --git a/Engine/Plugins/PCG/Source/PCG/Private/Compute/Data/PCGProxyForGPUData.cpp b/Engine/Plugins/PCG/Source/PCG/Private/Compute/Data/PCGProxyForGPUData.cpp
--- a/Engine/Plugins/PCG/Source/PCG/Private/Compute/Data/PCGProxyForGPUData.cpp
+++ b/Engine/Plugins/PCG/Source/PCG/Private/Compute/Data/PCGProxyForGPUData.cpp
@@ -33,18 +33,17 @@ void UPCGProxyForGPUData::ReleaseTransientResources(const TCHAR* InReason)
 
 EPCGDataType UPCGProxyForGPUData::GetUnderlyingDataType() const
 {
-	if (TSharedPtr<const FPCGProxyForGPUDataCollection> DataCollectionGPU = GetGPUInfo())
-	{
-		return DataCollectionGPU->GetDescription().DataDescs[DataIndexInCollection].GetType();
-	}
-
-	return EPCGDataType::None;
+	return GetUnderlyingDataTypeId().AsLegacyType();
 }
 
 FPCGDataTypeBaseId UPCGProxyForGPUData::GetUnderlyingDataTypeId() const
 {
-	// @todo_pcg: Update GPU to use FPCGDataTypeBaseId instead of EPCGDataType
-	return FPCGDataTypeBaseId::MakeFromLegacyType(GetUnderlyingDataType());
+	if (TSharedPtr<const FPCGProxyForGPUDataCollection> DataCollectionGPU = GetGPUInfo())
+	{
+		return DataCollectionGPU->GetDescription().DataDescs[DataIndexInCollection].GetType().GetId();
+	}
+
+	return FPCGDataTypeBaseId{};
 }
 
 TSharedPtr<const FPCGProxyForGPUDataCollection> UPCGProxyForGPUData::GetInputDataCollectionInfo() const
diff --git a/Engine/Plugins/PCG/Source/PCG/Private/Compute/Elements/PCGCustomHLSL.h b/Engine/Plugins/PCG/Source/PCG/Private/Compute/Elements/PCGCustomHLSL.h
--- a/Engine/Plugins/PCG/Source/PCG/Private/Compute/Elements/PCGCustomHLSL.h
+++ b/Engine/Plugins/PCG/Source/PCG/Private/Compute/Elements/PCGCustomHLSL.h
@@ -99,6 +99,11 @@ public:
 	/** Get user multiplier value applied to the dispatch thread count. */
 	uint32 GetThreadCountMultiplier() const { return IsThreadCountMultiplierInUse() ? static_cast<uint32>(ThreadCountMultiplier) : 1u; }
 
+#if WITH_EDITOR
+	UFUNCTION()
+	static bool SupportsComposition() { return false; }
+#endif // WITH_EDITOR
+
 protected:
 	/** Gets the GPU pin properties for the output pin with the given label. */
 	const FPCGPinPropertiesGPU* GetOutputPinPropertiesGPU(const FName& InPinLabel) const;
@@ -163,7 +168,7 @@ public:
 	TArray<FPCGPinProperties> InputPins = Super::DefaultPointInputPinProperties();
 
 	UPROPERTY(EditAnywhere, Category = "Settings")
-	TArray<FPCGPinPropertiesGPU> OutputPins = { FPCGPinPropertiesGPU(PCGPinConstants::DefaultOutputLabel, EPCGDataType::Point) };
+	TArray<FPCGPinPropertiesGPU> OutputPins = { FPCGPinPropertiesGPU(PCGPinConstants::DefaultOutputLabel, FPCGDataTypeIdentifier{EPCGDataType::Point}) };
 
 #if WITH_EDITOR
 	/** Holds input pin labels from PreEditChange, used in PostEditPropertyChange to update any references in output pin setup. */
diff --git a/Engine/Plugins/PCG/Source/PCG/Private/Compute/Elements/PCGCustomHLSLKernel.cpp b/Engine/Plugins/PCG/Source/PCG/Private/Compute/Elements/PCGCustomHLSLKernel.cpp
--- a/Engine/Plugins/PCG/Source/PCG/Private/Compute/Elements/PCGCustomHLSLKernel.cpp
+++ b/Engine/Plugins/PCG/Source/PCG/Private/Compute/Elements/PCGCustomHLSLKernel.cpp
@@ -53,12 +53,9 @@ namespace PCGCustomHLSLKernel
 		LookingForMultiLineComment,
 	};
 
-	FString GetDataTypeString(EPCGDataType Type)
+	FString GetDataTypeString(const FPCGDataTypeIdentifier& Type)
 	{
-		const UEnum* DataTypeEnum = StaticEnum<EPCGDataType>();
-		check(DataTypeEnum);
-
-		return DataTypeEnum->GetValueOrBitfieldAsString(static_cast<int64>(Type));
+		return Type.ToString();
 	}
 
 	// Reference implementation in FHLSLSyntaxHighlighterMarshaller::ProcessTokenizedLine()
diff --git a/Engine/Plugins/PCG/Source/PCG/Private/Compute/PCGComputeCommon.cpp b/Engine/Plugins/PCG/Source/PCG/Private/Compute/PCGComputeCommon.cpp
--- a/Engine/Plugins/PCG/Source/PCG/Private/Compute/PCGComputeCommon.cpp
+++ b/Engine/Plugins/PCG/Source/PCG/Private/Compute/PCGComputeCommon.cpp
@@ -30,20 +30,34 @@
 namespace PCGComputeConstants
 {
 	/** PCG data types supported in GPU node inputs. */
-	constexpr EPCGDataType AllowedInputTypes =
-		EPCGDataType::Point
-		| EPCGDataType::Param
-		| EPCGDataType::Landscape
-		| EPCGDataType::BaseTexture
-		| EPCGDataType::VirtualTexture
-		| EPCGDataType::Resource
-		| EPCGDataType::ProxyForGPU;
+	inline const FPCGDataTypeIdentifier& GetAllowedInputTypes()
+	{
+		static FPCGDataTypeIdentifier Result{
+			EPCGDataType::Point
+			| EPCGDataType::Param
+			| EPCGDataType::Landscape
+			| EPCGDataType::BaseTexture
+			| EPCGDataType::VirtualTexture
+			| EPCGDataType::Resource
+			| EPCGDataType::ProxyForGPU
+		};
+
+		return Result;
+	}
 
 	/** PCG data types supported in GPU node outputs. */
-	constexpr EPCGDataType AllowedOutputTypes = EPCGDataType::Point | EPCGDataType::Param | EPCGDataType::ProxyForGPU | EPCGDataType::BaseTexture;
+	inline const FPCGDataTypeIdentifier& GetAllowedOutputTypes()
+	{
+		static FPCGDataTypeIdentifier Result{ EPCGDataType::Point | EPCGDataType::Param | EPCGDataType::ProxyForGPU | EPCGDataType::BaseTexture };
+		return Result;
+	}
 
 	/** PCG data types supported in GPU data collections. */
-	constexpr EPCGDataType AllowedDataCollectionTypes = EPCGDataType::Point | EPCGDataType::Param | EPCGDataType::ProxyForGPU;
+	inline const FPCGDataTypeIdentifier& GetAllowedDataCollectionTypes()
+	{
+		static FPCGDataTypeIdentifier Result{ EPCGDataType::Point | EPCGDataType::Param | EPCGDataType::ProxyForGPU };
+		return Result;
+	}
 }
 
 namespace PCGComputeHelpers
@@ -84,27 +98,27 @@ namespace PCGComputeHelpers
 
 	EPCGElementDimension GetElementDimension(const UPCGData* InData)
 	{
-		return InData ? GetElementDimension(InData->GetDataType()) : EPCGElementDimension::One;
+		return InData ? GetElementDimension(InData->GetDataTypeId()) : EPCGElementDimension::One;
 	}
 
-	EPCGElementDimension GetElementDimension(EPCGDataType InDataType)
+	EPCGElementDimension GetElementDimension(const FPCGDataTypeIdentifier& InDataType)
 	{
 		return !!(InDataType & EPCGDataType::BaseTexture) ? EPCGElementDimension::Two : EPCGElementDimension::One;
 	}
 
-	bool IsTypeAllowedAsInput(EPCGDataType Type)
+	bool IsTypeAllowedAsInput(const FPCGDataTypeIdentifier& Type)
 	{
-		return (Type | PCGComputeConstants::AllowedInputTypes) == PCGComputeConstants::AllowedInputTypes;
+		return PCGComputeConstants::GetAllowedInputTypes().Intersects(Type);
 	}
 
-	bool IsTypeAllowedAsOutput(EPCGDataType Type)
+	bool IsTypeAllowedAsOutput(const FPCGDataTypeIdentifier& Type)
 	{
-		return (Type | PCGComputeConstants::AllowedOutputTypes) == PCGComputeConstants::AllowedOutputTypes;
+		return PCGComputeConstants::GetAllowedOutputTypes().Intersects(Type);
 	}
 
-	bool IsTypeAllowedInDataCollection(EPCGDataType Type)
+	bool IsTypeAllowedInDataCollection(const FPCGDataTypeIdentifier& Type)
 	{
-		return (Type | PCGComputeConstants::AllowedDataCollectionTypes) == PCGComputeConstants::AllowedDataCollectionTypes;
+		return PCGComputeConstants::GetAllowedDataCollectionTypes().Intersects(Type);
 	}
 
 	bool ShouldImportAttributesFromData(const UPCGData* InData)
@@ -218,10 +232,7 @@ namespace PCGComputeHelpers
 
 		UPCGComputeDataInterface* DataInterface = nullptr;
 
-		switch (PinProperties.AllowedTypes)
-		{
-		case EPCGDataType::Point:
-		case EPCGDataType::Param:
+		if (PinProperties.AllowedTypes == EPCGDataType::Point || PinProperties.AllowedTypes == EPCGDataType::Param)
 		{
 			UPCGDataCollectionDataInterface* DataInterfacePCGData = Context.NewObject_AnyThread<UPCGDataCollectionDataInterface>(InParams.ObjectOuter);
 
@@ -234,44 +245,34 @@ namespace PCGComputeHelpers
 			DataInterfacePCGData->SetRequiresZeroInitialization(InParams.ProducerKernel && InParams.ProducerKernel->DoesOutputPinRequireZeroInitialization(PinProperties.Label));
 
 			DataInterface = DataInterfacePCGData;
-			break;
 		}
-		case EPCGDataType::VirtualTexture:
+		else if (PinProperties.AllowedTypes == EPCGDataType::VirtualTexture)
 		{
 			DataInterface = Context.NewObject_AnyThread<UPCGVirtualTextureDataInterface>(InParams.ObjectOuter);
-			break;
 		}
-		case EPCGDataType::BaseTexture:
-		case EPCGDataType::Texture:
-		case EPCGDataType::RenderTarget:
+		else if (PinProperties.AllowedTypes.Intersects(EPCGDataType::BaseTexture))
 		{
 			UPCGTextureDataInterface* TextureDataInterface = Context.NewObject_AnyThread<UPCGTextureDataInterface>(InParams.ObjectOuter);
 			TextureDataInterface->SetRequiresExport(InParams.bRequiresExport);
 			TextureDataInterface->SetInitializeFromDataCollection(InParams.bProducedByCPU);
 
 			DataInterface = TextureDataInterface;
-
-			break;
 		}
-		case EPCGDataType::Landscape:
+		else if (PinProperties.AllowedTypes == EPCGDataType::Landscape)
 		{
 			DataInterface = Context.NewObject_AnyThread<UPCGLandscapeDataInterface>(InParams.ObjectOuter);
-			break;
 		}
-		case EPCGDataType::StaticMeshResource:
+		else if (PinProperties.AllowedTypes == EPCGDataType::StaticMeshResource)
 		{
 			DataInterface = Context.NewObject_AnyThread<UPCGStaticMeshDataInterface>(InParams.ObjectOuter);
-			break;
 		}
-		default:
+		else
 		{
 			UE_LOG(LogPCG, Error, TEXT("Unsupported connected upstream pin '%s' on node '%s' with type %s. Consider adding a conversion to a supported type such as Point."),
 				*PinProperties.Label.ToString(),
 				InParams.NodeForDebug ? *InParams.NodeForDebug->GetNodeTitle(EPCGNodeTitleType::ListView).ToString() : TEXT("MISSING"),
-				StaticEnum<EPCGDataType>() ? *StaticEnum<EPCGDataType>()->GetNameStringByValue(static_cast<int64>(PinProperties.AllowedTypes)) : TEXT("MISSING")
+				*PinProperties.AllowedTypes.ToString()
 			);
-			break;
-		}
 		}
 
 		return DataInterface;
diff --git a/Engine/Plugins/PCG/Source/PCG/Private/Compute/PCGComputeGraph.cpp b/Engine/Plugins/PCG/Source/PCG/Private/Compute/PCGComputeGraph.cpp
--- a/Engine/Plugins/PCG/Source/PCG/Private/Compute/PCGComputeGraph.cpp
+++ b/Engine/Plugins/PCG/Source/PCG/Private/Compute/PCGComputeGraph.cpp
@@ -205,11 +205,10 @@ bool UPCGComputeGraph::IsGraphDataValid(FPCGContext* InContext) const
 void UPCGComputeGraph::DebugLogDataDescriptions(const UPCGDataBinding* InBinding) const
 {
 #if PCG_KERNEL_LOGGING_ENABLED
-	const UEnum* PCGDataTypeEnum = StaticEnum<EPCGDataType>();
 	const UEnum* PCGKernelAttributeTypeEnum = StaticEnum<EPCGKernelAttributeType>();
-	check(PCGDataTypeEnum && PCGKernelAttributeTypeEnum);
+	check(PCGKernelAttributeTypeEnum);
 
-	auto LogDataDescription = [this, PCGDataTypeEnum, PCGKernelAttributeTypeEnum, InBinding](const FPCGDataCollectionDesc* InDataDesc, const FPCGKernelPin& InKernelPin)
+	auto LogDataDescription = [this, PCGKernelAttributeTypeEnum, InBinding](const FPCGDataCollectionDesc* InDataDesc, const FPCGKernelPin& InKernelPin)
 	{
 		if (!InDataDesc)
 		{
@@ -221,7 +220,7 @@ void UPCGComputeGraph::DebugLogDataDescriptions(const UPCGDataBinding* InBinding
 			const FPCGDataDesc& DataDesc = InDataDesc->DataDescs[DataIndex];
 
 			UE_LOG(LogPCG, Display, TEXT("\t\tData Index %d"), DataIndex);
-			UE_LOG(LogPCG, Display, TEXT("\t\t\tType: %s"), *PCGDataTypeEnum->GetDisplayNameTextByValue(static_cast<int64>(DataDesc.GetType())).ToString());
+			UE_LOG(LogPCG, Display, TEXT("\t\t\tType: %s"), *DataDesc.GetType().ToString());
 			UE_LOG(LogPCG, Display, TEXT("\t\t\tNum Elements: %s (%s)"), *DataDesc.GetElementCount().ToString(), *UEnum::GetValueAsString(DataDesc.GetElementDimension()));
 			UE_LOG(LogPCG, Display, TEXT("\t\t\tAttributes (%d)"), DataDesc.AttributeDescs.Num());
 
@@ -378,7 +377,7 @@ void UPCGComputeGraph::ComputeExternalPinDesc(FName InVirtualLabel, UPCGDataBind
 
 	for (const FPCGTaggedData& TaggedData : TaggedDatas)
 	{
-		if (!TaggedData.Data || !PCGComputeHelpers::IsTypeAllowedAsInput(TaggedData.Data->GetDataType()))
+		if (!TaggedData.Data || !PCGComputeHelpers::IsTypeAllowedAsInput(TaggedData.Data->GetDataTypeId()))
 		{
 			continue;
 		}
diff --git a/Engine/Plugins/PCG/Source/PCG/Private/Compute/PCGComputeKernel.cpp b/Engine/Plugins/PCG/Source/PCG/Private/Compute/PCGComputeKernel.cpp
--- a/Engine/Plugins/PCG/Source/PCG/Private/Compute/PCGComputeKernel.cpp
+++ b/Engine/Plugins/PCG/Source/PCG/Private/Compute/PCGComputeKernel.cpp
@@ -345,7 +345,7 @@ void UPCGComputeKernel::ComputeDataDescFromPinProperties(const FPCGPinProperties
 						ClampedDataIndex = InputDesc.DataDescs.Num() - 1;
 					}
 
-					const FPCGDataDesc& InputDataDesc = InputDesc.DataDescs.IsValidIndex(ClampedDataIndex) ? InputDesc.DataDescs[ClampedDataIndex] : FPCGDataDesc(EPCGDataType::Any, 0);
+					const FPCGDataDesc& InputDataDesc = InputDesc.DataDescs.IsValidIndex(ClampedDataIndex) ? InputDesc.DataDescs[ClampedDataIndex] : FPCGDataDesc(FPCGDataTypeIdentifier{EPCGDataType::Any}, 0);
 					const EPCGElementDimension SourceDimension = InputDataDesc.GetElementDimension();
 
 					if (Props.ElementCountMode == EPCGElementCountMode::FromInputData)
@@ -620,14 +620,14 @@ bool UPCGComputeKernel::AreInputEdgesValid(TArray<FPCGKernelLogEntry>& OutLogEnt
 			for (const UPCGEdge* InputEdge : InputPin->Edges)
 			{
 				const UPCGPin* UpstreamPin = InputEdge ? InputEdge->GetOtherPin(InputPin) : nullptr;
-				if (UpstreamPin && InputPin->GetRequiredTypeConversion(UpstreamPin) != EPCGTypeConversion::NoConversionRequired)
+				if (UpstreamPin && InputPin->GetCompatibilityWithOtherPin(UpstreamPin) != EPCGDataTypeCompatibilityResult::Compatible)
 				{
 #if PCG_KERNEL_LOGGING_ENABLED
 					OutLogEntries.Emplace(FText::Format(
 						LOCTEXT("InvalidInputPinEdge", "Unsupported connected upstream pin '{0}' on node '{1}' with type {2}. Recreate the edge to add required conversion nodes."),
 						FText::FromName(UpstreamPin->Properties.Label),
 						Node->GetNodeTitle(EPCGNodeTitleType::ListView),
-						StaticEnum<EPCGDataType>() ? StaticEnum<EPCGDataType>()->GetDisplayNameTextByValue(static_cast<int64>(UpstreamPin->Properties.AllowedTypes)) : FText::FromString(TEXT("MISSING"))),
+						FText::FromString(UpstreamPin->Properties.AllowedTypes.ToString())),
 						EPCGKernelLogVerbosity::Error);
 #endif
 
diff --git a/Engine/Plugins/PCG/Source/PCG/Private/Compute/PCGDataBinding.cpp b/Engine/Plugins/PCG/Source/PCG/Private/Compute/PCGDataBinding.cpp
--- a/Engine/Plugins/PCG/Source/PCG/Private/Compute/PCGDataBinding.cpp
+++ b/Engine/Plugins/PCG/Source/PCG/Private/Compute/PCGDataBinding.cpp
@@ -30,7 +30,7 @@ void UPCGDataBinding::Initialize(const UPCGComputeGraph* InComputeGraph, FPCGCon
 	InputDataCollection.TaggedData.Reserve(InContext->InputData.TaggedData.Num());
 	for (const FPCGTaggedData& InputData : InContext->InputData.TaggedData)
 	{
-		if (InputData.Data && PCGComputeHelpers::IsTypeAllowedAsInput(InputData.Data->GetDataType()))
+		if (InputData.Data && PCGComputeHelpers::IsTypeAllowedAsInput(FPCGDataTypeIdentifier{InputData.Data->GetDataTypeId()}))
 		{
 			InputDataCollection.TaggedData.Add(InputData);
 		}
diff --git a/Engine/Plugins/PCG/Source/PCG/Private/Compute/PCGDataForGPU.cpp b/Engine/Plugins/PCG/Source/PCG/Private/Compute/PCGDataForGPU.cpp
--- a/Engine/Plugins/PCG/Source/PCG/Private/Compute/PCGDataForGPU.cpp
+++ b/Engine/Plugins/PCG/Source/PCG/Private/Compute/PCGDataForGPU.cpp
@@ -745,22 +745,22 @@ bool FPCGKernelAttributeDesc::operator==(const FPCGKernelAttributeDesc& Other) c
 	return AttributeId == Other.AttributeId && AttributeKey == Other.AttributeKey;
 }
 
-FPCGDataDesc::FPCGDataDesc(EPCGDataType InType, int32 InElementCount)
+FPCGDataDesc::FPCGDataDesc(FPCGDataTypeIdentifier InType, int32 InElementCount)
 	: FPCGDataDesc(InType, FIntVector4(InElementCount, 0, 0, 0))
 {
 }
 
-FPCGDataDesc::FPCGDataDesc(EPCGDataType InType, FIntPoint InElementCount)
+FPCGDataDesc::FPCGDataDesc(FPCGDataTypeIdentifier InType, FIntPoint InElementCount)
 	: FPCGDataDesc(InType, FIntVector4(InElementCount.X, InElementCount.Y, 0, 0))
 {
 }
 
-FPCGDataDesc::FPCGDataDesc(EPCGDataType InType, FIntVector3 InElementCount)
+FPCGDataDesc::FPCGDataDesc(FPCGDataTypeIdentifier InType, FIntVector3 InElementCount)
 	: FPCGDataDesc(InType, FIntVector4(InElementCount.X, InElementCount.Y, InElementCount.Z, 0))
 {
 }
 
-FPCGDataDesc::FPCGDataDesc(EPCGDataType InType, FIntVector4 InElementCount)
+FPCGDataDesc::FPCGDataDesc(FPCGDataTypeIdentifier InType, FIntVector4 InElementCount)
 	: Type(InType)
 	, ElementCount(InElementCount)
 {
@@ -773,7 +773,7 @@ FPCGDataDesc::FPCGDataDesc(const FPCGTaggedData& InTaggedData, const UPCGDataBin
 {
 	check(InTaggedData.Data);
 
-	Type = InTaggedData.Data->GetDataType();
+	Type = InTaggedData.Data->GetDataTypeId();
 	ElementCount = PCGComputeHelpers::GetElementCount(InTaggedData.Data);
 	ElementDimension = PCGComputeHelpers::GetElementDimension(InTaggedData.Data);
 
@@ -1168,7 +1168,7 @@ FPCGDataCollectionDesc FPCGDataCollectionDesc::BuildFromDataCollection(
 
 	for (const FPCGTaggedData& Data : InDataCollection.TaggedData)
 	{
-		if (!Data.Data || !PCGComputeHelpers::IsTypeAllowedAsOutput(Data.Data->GetDataType()))
+		if (!Data.Data || !PCGComputeHelpers::IsTypeAllowedAsOutput(Data.Data->GetDataTypeId()))
 		{
 			continue;
 		}
@@ -1989,25 +1989,24 @@ EPCGUnpackDataCollectionResult FPCGDataCollectionDesc::UnpackDataCollection(FPCG
 }
 
 uint32 FPCGDataCollectionDesc::ComputeTotalElementCount() const
-{
-	return ComputeTotalElementCount(EPCGDataType::Any);
-}
-
-uint32 FPCGDataCollectionDesc::ComputeTotalElementCount(EPCGDataType InDataType) const
 {
 	uint32 ProcessingElementCount = 0;
 
 	for (const FPCGDataDesc& DataDesc : DataDescs)
 	{
-		if (!!(DataDesc.GetType() & InDataType))
-		{
-			ProcessingElementCount += DataDesc.ComputeTotalElementCount();
-		}
+		ProcessingElementCount += DataDesc.ComputeTotalElementCount();
 	}
 
 	return ProcessingElementCount;
 }
 
+uint32 FPCGDataCollectionDesc::ComputeTotalElementCount(FPCGDataTypeIdentifier InDataType) const
+{
+	// InDataType not used anymore
+	// @todo_pcg Remove in future pass.
+	return ComputeTotalElementCount();
+}
+
 void FPCGDataCollectionDesc::Combine(const FPCGDataCollectionDesc& Other)
 {
 	DataDescs.Append(Other.DataDescs);
diff --git a/Engine/Plugins/PCG/Source/PCG/Private/Compute/PCGPinPropertiesGPU.h b/Engine/Plugins/PCG/Source/PCG/Private/Compute/PCGPinPropertiesGPU.h
--- a/Engine/Plugins/PCG/Source/PCG/Private/Compute/PCGPinPropertiesGPU.h
+++ b/Engine/Plugins/PCG/Source/PCG/Private/Compute/PCGPinPropertiesGPU.h
@@ -140,7 +140,7 @@ struct FPCGPinPropertiesGPU : public FPCGPinProperties
 public:
 	FPCGPinPropertiesGPU() = default;
 
-	explicit FPCGPinPropertiesGPU(const FName& InLabel, EPCGDataType InAllowedTypes)
+	explicit FPCGPinPropertiesGPU(const FName& InLabel, FPCGDataTypeIdentifier InAllowedTypes)
 		: FPCGPinProperties(InLabel, InAllowedTypes)
 	{}
 
diff --git a/Engine/Plugins/PCG/Source/PCG/Private/Data/Registry/PCGDataType.cpp b/Engine/Plugins/PCG/Source/PCG/Private/Data/Registry/PCGDataType.cpp
--- a/Engine/Plugins/PCG/Source/PCG/Private/Data/Registry/PCGDataType.cpp
+++ b/Engine/Plugins/PCG/Source/PCG/Private/Data/Registry/PCGDataType.cpp
@@ -148,8 +148,8 @@ FPCGDataTypeBaseId FPCGDataTypeBaseId::MakeFromLegacyType(EPCGDataType DataType)
 		return Construct<FPCGDataTypeInfoProxyForGPU>();
 	case EPCGDataType::Param:
 		return Construct<FPCGDataTypeInfoParam>();
-	//case EPCGDataType::Settings:
-	//	return Construct<FPCGDataTypeInfoSettings>();
+	case EPCGDataType::Settings:
+		return Construct<FPCGDataTypeInfoSettings>();
 	case EPCGDataType::Other:
 		return Construct<FPCGDataTypeInfoOther>();
 	case EPCGDataType::None: // fall-through
@@ -187,7 +187,7 @@ EPCGDataType FPCGDataTypeBaseId::AsLegacyType() const
 	PCG_DATA_TYPE_TEST(FPCGDataTypeInfoSpatial)
 	PCG_DATA_TYPE_TEST(FPCGDataTypeInfoProxyForGPU)
 	PCG_DATA_TYPE_TEST(FPCGDataTypeInfoParam)
-	//PCG_DATA_TYPE_TEST(FPCGDataTypeInfoSettings)
+	PCG_DATA_TYPE_TEST(FPCGDataTypeInfoSettings)
 	PCG_DATA_TYPE_TEST(FPCGDataTypeInfoOther)
 		
 #undef PCG_DATA_TYPE_TEST
diff --git a/Engine/Plugins/PCG/Source/PCG/Private/Data/Registry/PCGDataTypeIdentifier.cpp b/Engine/Plugins/PCG/Source/PCG/Private/Data/Registry/PCGDataTypeIdentifier.cpp
--- a/Engine/Plugins/PCG/Source/PCG/Private/Data/Registry/PCGDataTypeIdentifier.cpp
+++ b/Engine/Plugins/PCG/Source/PCG/Private/Data/Registry/PCGDataTypeIdentifier.cpp
@@ -356,3 +356,13 @@ FPCGDataTypeIdentifier UPCGDataTypeIdentifierHelpers::GetIdentifierFromLegacyTyp
 
 	return FPCGDataTypeIdentifier(InternalDataType);
 }
+
+uint32 GetTypeHash(const FPCGDataTypeIdentifier& Id)
+{
+	// Hash needs to be stable to Ids order.
+	TArray<int32, TInlineAllocator<16>> IdsHash;
+	IdsHash.Reserve(Id.Ids.Num());
+	Algo::Transform(Id.Ids, IdsHash, [](const FPCGDataTypeBaseId& BaseId) { return GetTypeHash(BaseId); });
+	IdsHash.Sort();
+	return HashCombine(GetTypeHash(IdsHash), Id.CustomSubtype);
+}
diff --git a/Engine/Plugins/PCG/Source/PCG/Private/Elements/ControlFlow/PCGBooleanSelect.cpp b/Engine/Plugins/PCG/Source/PCG/Private/Elements/ControlFlow/PCGBooleanSelect.cpp
--- a/Engine/Plugins/PCG/Source/PCG/Private/Elements/ControlFlow/PCGBooleanSelect.cpp
+++ b/Engine/Plugins/PCG/Source/PCG/Private/Elements/ControlFlow/PCGBooleanSelect.cpp
@@ -31,16 +31,19 @@ FText UPCGBooleanSelectSettings::GetNodeTooltipText() const
 }
 #endif // WITH_EDITOR
 
-EPCGDataType UPCGBooleanSelectSettings::GetCurrentPinTypes(const UPCGPin* Pin) const
+FPCGDataTypeIdentifier UPCGBooleanSelectSettings::GetCurrentPinTypesID(const UPCGPin* Pin) const
 {
-	// Output pin depends on the union of both input pins 
+	// Output pin depends on the union of all non-advanced input pins. Advanced pins are dynamic user overrides and
+	// should not affect the output pin type.
 	if (Pin->IsOutputPin())
 	{
-		const EPCGDataType InputTypeUnion = GetTypeUnionOfIncidentEdges(PCGBooleanSelectConstants::InputLabelA) | GetTypeUnionOfIncidentEdges(PCGBooleanSelectConstants::InputLabelB);
-		return InputTypeUnion != EPCGDataType::None ? InputTypeUnion : EPCGDataType::Any;
+		const FPCGDataTypeIdentifier InputTypeUnion = GetTypeUnionIDOfAllIncidentEdges({PCGBooleanSelectConstants::InputLabelA, PCGBooleanSelectConstants::InputLabelB});
+		return InputTypeUnion.IsValid() ? InputTypeUnion : FPCGDataTypeIdentifier{EPCGDataType::Any};
+	}
+	else
+	{
+		return Super::GetCurrentPinTypesID(Pin);
 	}
-
-	return Super::GetCurrentPinTypes(Pin);
 }
 
 TArray<FPCGPinProperties> UPCGBooleanSelectSettings::InputPinProperties() const
diff --git a/Engine/Plugins/PCG/Source/PCG/Private/Elements/ControlFlow/PCGMultiSelect.cpp b/Engine/Plugins/PCG/Source/PCG/Private/Elements/ControlFlow/PCGMultiSelect.cpp
--- a/Engine/Plugins/PCG/Source/PCG/Private/Elements/ControlFlow/PCGMultiSelect.cpp
+++ b/Engine/Plugins/PCG/Source/PCG/Private/Elements/ControlFlow/PCGMultiSelect.cpp
@@ -103,28 +103,22 @@ FText UPCGMultiSelectSettings::GetNodeTooltipText() const
 }
 #endif // WITH_EDITOR
 
-EPCGDataType UPCGMultiSelectSettings::GetCurrentPinTypes(const UPCGPin* Pin) const
+FPCGDataTypeIdentifier UPCGMultiSelectSettings::GetCurrentPinTypesID(const UPCGPin* Pin) const
 {
 	// Output pin depends on the union of all non-advanced input pins. Advanced pins are dynamic user overrides and
 	// should not affect the output pin type.
 	if (Pin->IsOutputPin())
 	{
-		EPCGDataType InputTypeUnion = EPCGDataType::None;
-
-		const TArray<TObjectPtr<UPCGPin>>& InputPins = Pin->Node->GetInputPins();
-
-		for (int Index = 0; Index < InputPins.Num(); ++Index)
-		{
-			if (!InputPins[Index]->Properties.IsAdvancedPin())
-			{
-				InputTypeUnion |= GetTypeUnionOfIncidentEdges(InputPins[Index]->Properties.Label);
-			}
-		}
-
-		return InputTypeUnion != EPCGDataType::None ? InputTypeUnion : EPCGDataType::Any;
+		TArray<FName, TInlineAllocator<32>> PinLabels;
+		Algo::TransformIf(Pin->Node->GetInputPins(), PinLabels, [](const UPCGPin* InPin) { return InPin && !InPin->Properties.IsAdvancedPin();}, [](const UPCGPin* InPin) { return InPin->Properties.Label; });
+		
+		const FPCGDataTypeIdentifier InputTypeUnion = GetTypeUnionIDOfAllIncidentEdges(PinLabels);
+		return InputTypeUnion.IsValid() ? InputTypeUnion : FPCGDataTypeIdentifier{EPCGDataType::Any};
+	}
+	else
+	{
+		return Super::GetCurrentPinTypesID(Pin);
 	}
-
-	return Super::GetCurrentPinTypes(Pin);
 }
 
 FString UPCGMultiSelectSettings::GetAdditionalTitleInformation() const
diff --git a/Engine/Plugins/PCG/Source/PCG/Private/Elements/ControlFlow/PCGQualitySelect.cpp b/Engine/Plugins/PCG/Source/PCG/Private/Elements/ControlFlow/PCGQualitySelect.cpp
--- a/Engine/Plugins/PCG/Source/PCG/Private/Elements/ControlFlow/PCGQualitySelect.cpp
+++ b/Engine/Plugins/PCG/Source/PCG/Private/Elements/ControlFlow/PCGQualitySelect.cpp
@@ -27,23 +27,22 @@ FString UPCGQualitySelectSettings::GetAdditionalTitleInformation() const
 	return PCGQualityHelpers::GetQualityPinLabel().ToString();
 }
 
-EPCGDataType UPCGQualitySelectSettings::GetCurrentPinTypes(const UPCGPin* InPin) const
+FPCGDataTypeIdentifier UPCGQualitySelectSettings::GetCurrentPinTypesID(const UPCGPin* Pin) const
 {
-	check(InPin);
-
-	if (!InPin->IsOutputPin())
+	// Output pin depends on the union of all non-advanced input pins. Advanced pins are dynamic user overrides and
+	// should not affect the output pin type.
+	if (Pin->IsOutputPin())
 	{
-		return Super::GetCurrentPinTypes(InPin);
+		TArray<FName, TInlineAllocator<32>> PinLabels;
+		Algo::TransformIf(Pin->Node->GetInputPins(), PinLabels, [](const UPCGPin* InPin) { return InPin && !InPin->Properties.IsAdvancedPin();}, [](const UPCGPin* InPin) { return InPin->Properties.Label; });
+		
+		const FPCGDataTypeIdentifier InputTypeUnion = GetTypeUnionIDOfAllIncidentEdges(PinLabels);
+		return InputTypeUnion.IsValid() ? InputTypeUnion : FPCGDataTypeIdentifier{EPCGDataType::Any};
 	}
-
-	// In the case of this node, really we want the union of all inputs on all pins.
-	EPCGDataType InputType = EPCGDataType::None;
-	for (const FPCGPinProperties& PinProperties : InputPinProperties())
+	else
 	{
-		InputType |= GetTypeUnionOfIncidentEdges(PinProperties.Label);
+		return Super::GetCurrentPinTypesID(Pin);
 	}
-
-	return (InputType != EPCGDataType::None) ? InputType : EPCGDataType::Any;
 }
 
 TArray<FPCGPinProperties> UPCGQualitySelectSettings::InputPinProperties() const
diff --git a/Engine/Plugins/PCG/Source/PCG/Private/Elements/IO/PCGDataTableElement.cpp b/Engine/Plugins/PCG/Source/PCG/Private/Elements/IO/PCGDataTableElement.cpp
--- a/Engine/Plugins/PCG/Source/PCG/Private/Elements/IO/PCGDataTableElement.cpp
+++ b/Engine/Plugins/PCG/Source/PCG/Private/Elements/IO/PCGDataTableElement.cpp
@@ -51,31 +51,16 @@ TArray<FPCGPinProperties> UPCGLoadDataTableSettings::OutputPinProperties() const
 	return PinProperties;
 }
 
-EPCGDataType UPCGLoadDataTableSettings::GetCurrentPinTypes(const UPCGPin* InPin) const
+FPCGDataTypeIdentifier UPCGLoadDataTableSettings::GetCurrentPinTypesID(const UPCGPin* InPin) const
 {
 	check(InPin);
 	if (!InPin->IsOutputPin())
 	{
-		return Super::GetCurrentPinTypes(InPin);
+		return Super::GetCurrentPinTypesID(InPin);
 	}
 
-	const UEnum* DataTypeEnum = StaticEnum<EPCGDataType>();
-	const UEnum* ExclusiveDataTypeEnum = StaticEnum<EPCGExclusiveDataType>();
-
-	if (DataTypeEnum && ExclusiveDataTypeEnum)
-	{
-		FName ExclusiveDataTypeName = ExclusiveDataTypeEnum->GetNameByValue(static_cast<__underlying_type(EPCGExclusiveDataType)>(OutputType));
-		if (ExclusiveDataTypeName != NAME_None)
-		{
-			const int64 MatchingDataType = DataTypeEnum->GetValueByName(ExclusiveDataTypeName);
-			if (ensure(MatchingDataType != INDEX_NONE))
-			{
-				return static_cast<EPCGDataType>(MatchingDataType);
-			}
-		}
-	}
-
-	return Super::GetCurrentPinTypes(InPin);
+	FPCGDataTypeIdentifier Result = UPCGDataTypeIdentifierHelpers::GetIdentifierFromLegacyType(OutputType);
+	return Result.IsValid() ? Result : Super::GetCurrentPinTypesID(InPin); 
 }
 
 FPCGElementPtr UPCGLoadDataTableSettings::CreateElement() const
diff --git a/Engine/Plugins/PCG/Source/PCG/Private/Elements/IO/PCGLoadAssetElement.cpp b/Engine/Plugins/PCG/Source/PCG/Private/Elements/IO/PCGLoadAssetElement.cpp
--- a/Engine/Plugins/PCG/Source/PCG/Private/Elements/IO/PCGLoadAssetElement.cpp
+++ b/Engine/Plugins/PCG/Source/PCG/Private/Elements/IO/PCGLoadAssetElement.cpp
@@ -107,7 +107,7 @@ FString UPCGLoadDataAssetSettings::GetAdditionalTitleInformation() const
 	}
 }
 
-EPCGDataType UPCGLoadDataAssetSettings::GetCurrentPinTypes(const UPCGPin* InPin) const
+FPCGDataTypeIdentifier UPCGLoadDataAssetSettings::GetCurrentPinTypesID(const UPCGPin* InPin) const
 {
 	// Implementation notes: the output pin types don't depend on the input pin types,
 	// but they can change based on the asset selected, hence why they are dynamic, but we need to return the pin type as-is.
@@ -150,11 +150,11 @@ void UPCGLoadDataAssetSettings::UpdateFromData()
 
 			if (!MatchingPin)
 			{
-				NewPins.Emplace(TaggedData.Pin, TaggedData.Data->GetDataType());
+				NewPins.Emplace(TaggedData.Pin, TaggedData.Data->GetDataTypeId());
 			}
 			else
 			{
-				MatchingPin->AllowedTypes |= TaggedData.Data->GetDataType();
+				MatchingPin->AllowedTypes |= TaggedData.Data->GetDataTypeId();
 			}
 		}
 
diff --git a/Engine/Plugins/PCG/Source/PCG/Private/Elements/Metadata/PCGMetadataOpElementBase.cpp b/Engine/Plugins/PCG/Source/PCG/Private/Elements/Metadata/PCGMetadataOpElementBase.cpp
--- a/Engine/Plugins/PCG/Source/PCG/Private/Elements/Metadata/PCGMetadataOpElementBase.cpp
+++ b/Engine/Plugins/PCG/Source/PCG/Private/Elements/Metadata/PCGMetadataOpElementBase.cpp
@@ -68,17 +68,24 @@ void UPCGMetadataSettingsBase::PostLoad()
 }
 
 EPCGDataType UPCGMetadataSettingsBase::GetInputPinType(uint32 Index) const
+{
+PRAGMA_DISABLE_DEPRECATION_WARNINGS
+	return GetInputPinTypeID(Index);
+PRAGMA_ENABLE_DEPRECATION_WARNINGS
+}
+
+FPCGDataTypeIdentifier UPCGMetadataSettingsBase::GetInputPinTypeID(uint32 Index) const
 {
 	const FName PinLabel = GetInputPinLabel(Index);
-	const EPCGDataType FirstPinTypeUnion = GetTypeUnionOfIncidentEdges(PinLabel);
+	const FPCGDataTypeIdentifier FirstPinTypeUnion = GetTypeUnionIDOfIncidentEdges(PinLabel);
 
 	// If the pin is not connected but supports default values, treat it as a param
 	if (FirstPinTypeUnion == EPCGDataType::None && IsPinDefaultValueActivated(PinLabel))
 	{
-		return EPCGDataType::Param;
+		return FPCGDataTypeIdentifier{EPCGDataType::Param};
 	}
 
-	return FirstPinTypeUnion != EPCGDataType::None ? FirstPinTypeUnion : EPCGDataType::Any;
+	return FirstPinTypeUnion.IsValid() ? FirstPinTypeUnion : FPCGDataTypeIdentifier{EPCGDataType::Any};
 }
 
 TArray<FName> UPCGMetadataSettingsBase::GetOutputDataFromPinOptions() const
@@ -161,9 +168,9 @@ uint32 UPCGMetadataSettingsBase::GetInputPinToForward() const
 
 			for (uint32 InputPinIndex = 0; InputPinIndex < OperandNum; ++InputPinIndex)
 			{
-				const EPCGDataType PinType = GetInputPinType(InputPinIndex);
+				const FPCGDataTypeIdentifier PinType = GetInputPinTypeID(InputPinIndex);
 
-				if ((PinType != EPCGDataType::Any) && !!(PinType & EPCGDataType::Spatial))
+				if (!PinType.IsSameType(FPCGDataTypeIdentifier{EPCGDataType::Any}) && (PinType & EPCGDataType::Spatial).IsValid())
 				{
 					InputPinToForward = InputPinIndex;
 					break;
@@ -194,18 +201,18 @@ const UPCGParamData* UPCGMetadataSettingsBase::CreateDefaultValueParamData(FPCGC
 	return Data;
 }
 
-EPCGDataType UPCGMetadataSettingsBase::GetCurrentPinTypes(const UPCGPin* InPin) const
+FPCGDataTypeIdentifier UPCGMetadataSettingsBase::GetCurrentPinTypesID(const UPCGPin* InPin) const
 {
 	check(InPin);
 	const uint32 OperandNum = GetOperandNum();
 	if (!InPin->IsOutputPin() || OperandNum == 0)
 	{
 		// Fall back to default for input pins, or if no input pins present from which to obtain type
-		return Super::GetCurrentPinTypes(InPin);
+		return Super::GetCurrentPinTypesID(InPin);
 	}
 
 	// Output pin narrows to union of inputs on pin to forward
-	return GetInputPinType(GetInputPinToForward());
+	return GetInputPinTypeID(GetInputPinToForward());
 }
 
 #if WITH_EDITOR
diff --git a/Engine/Plugins/PCG/Source/PCG/Private/Elements/PCGAttributeFilter.cpp b/Engine/Plugins/PCG/Source/PCG/Private/Elements/PCGAttributeFilter.cpp
--- a/Engine/Plugins/PCG/Source/PCG/Private/Elements/PCGAttributeFilter.cpp
+++ b/Engine/Plugins/PCG/Source/PCG/Private/Elements/PCGAttributeFilter.cpp
@@ -205,25 +205,25 @@ void FPCGAttributeFilterThresholdSettings::OnPostLoad()
 // UPCGAttributeFilteringSettings
 ////////////////////////////////////////
 
-EPCGDataType UPCGAttributeFilteringSettings::GetCurrentPinTypes(const UPCGPin* InPin) const
+FPCGDataTypeIdentifier UPCGAttributeFilteringSettings::GetCurrentPinTypesID(const UPCGPin* InPin) const
 {
 	check(InPin);
 	if (!InPin->IsOutputPin())
 	{
-		return Super::GetCurrentPinTypes(InPin);
+		return Super::GetCurrentPinTypesID(InPin);
 	}
 
 	// Output pin narrows to union of inputs on first pin
-	const EPCGDataType InputTypeUnion = GetTypeUnionOfIncidentEdges(PCGPinConstants::DefaultInputLabel);
+	const FPCGDataTypeIdentifier InputTypeUnion = GetTypeUnionIDOfIncidentEdges(PCGPinConstants::DefaultInputLabel);
 
 	// Spatial is collapsed into points
-	if (InputTypeUnion != EPCGDataType::None && (InputTypeUnion & EPCGDataType::Spatial) == InputTypeUnion)
+	if (InputTypeUnion.IsValid() && (InputTypeUnion & EPCGDataType::Spatial) == InputTypeUnion)
 	{
-		return EPCGDataType::Point;
+		return FPCGDataTypeIdentifier{EPCGDataType::Point};
 	}
 	else
 	{
-		return (InputTypeUnion != EPCGDataType::None) ? InputTypeUnion : EPCGDataType::Any;
+		return InputTypeUnion.IsValid() ? InputTypeUnion : FPCGDataTypeIdentifier{EPCGDataType::Any};
 	}
 }
 
@@ -409,25 +409,25 @@ UPCGAttributeFilteringRangeSettings::UPCGAttributeFilteringRangeSettings()
 }
 
 
-EPCGDataType UPCGAttributeFilteringRangeSettings::GetCurrentPinTypes(const UPCGPin* InPin) const
+FPCGDataTypeIdentifier UPCGAttributeFilteringRangeSettings::GetCurrentPinTypesID(const UPCGPin* InPin) const
 {
 	check(InPin);
 	if (!InPin->IsOutputPin())
 	{
-		return Super::GetCurrentPinTypes(InPin);
+		return Super::GetCurrentPinTypesID(InPin);
 	}
 
 	// Output pin narrows to union of inputs on first pin
-	const EPCGDataType InputTypeUnion = GetTypeUnionOfIncidentEdges(PCGPinConstants::DefaultInputLabel);
+	const FPCGDataTypeIdentifier InputTypeUnion = GetTypeUnionIDOfIncidentEdges(PCGPinConstants::DefaultInputLabel);
 
 	// Spatial is collapsed into points
-	if (InputTypeUnion != EPCGDataType::None && (InputTypeUnion & EPCGDataType::Spatial) == InputTypeUnion)
+	if (InputTypeUnion.IsValid() && (InputTypeUnion & EPCGDataType::Spatial) == InputTypeUnion)
 	{
-		return EPCGDataType::Point;
+		return FPCGDataTypeIdentifier{EPCGDataType::Point};
 	}
 	else
 	{
-		return (InputTypeUnion != EPCGDataType::None) ? InputTypeUnion : EPCGDataType::Any;
+		return InputTypeUnion.IsValid() ? InputTypeUnion : FPCGDataTypeIdentifier{EPCGDataType::Any};
 	}
 }
 
diff --git a/Engine/Plugins/PCG/Source/PCG/Private/Elements/PCGAttributeNoise.cpp b/Engine/Plugins/PCG/Source/PCG/Private/Elements/PCGAttributeNoise.cpp
--- a/Engine/Plugins/PCG/Source/PCG/Private/Elements/PCGAttributeNoise.cpp
+++ b/Engine/Plugins/PCG/Source/PCG/Private/Elements/PCGAttributeNoise.cpp
@@ -207,19 +207,6 @@ TArray<FPCGPinProperties> UPCGAttributeNoiseSettings::OutputPinProperties() cons
 	return PinProperties;
 }
 
-EPCGDataType UPCGAttributeNoiseSettings::GetCurrentPinTypes(const UPCGPin* InPin) const
-{
-	check(InPin);
-	if (!InPin->IsOutputPin())
-	{
-		return Super::GetCurrentPinTypes(InPin);
-	}
-
-	// Output pin narrows to union of inputs on first pin
-	const EPCGDataType InputTypeUnion = GetTypeUnionOfIncidentEdges(PCGPinConstants::DefaultInputLabel);
-	return (InputTypeUnion != EPCGDataType::None) ? InputTypeUnion : (EPCGDataType::Any);
-}
-
 bool FPCGAttributeNoiseElement::ExecuteInternal(FPCGContext* InContext) const
 {
 	TRACE_CPUPROFILER_EVENT_SCOPE(FPCGAttributeNoiseElement::Execute);
diff --git a/Engine/Plugins/PCG/Source/PCG/Private/Elements/PCGCopyAttributes.cpp b/Engine/Plugins/PCG/Source/PCG/Private/Elements/PCGCopyAttributes.cpp
--- a/Engine/Plugins/PCG/Source/PCG/Private/Elements/PCGCopyAttributes.cpp
+++ b/Engine/Plugins/PCG/Source/PCG/Private/Elements/PCGCopyAttributes.cpp
@@ -122,17 +122,6 @@ void UPCGCopyAttributesSettings::ApplyDeprecationBeforeUpdatePins(UPCGNode* InOu
 	InOutNode->RenameInputPin(TEXT("Attribute"), PCGCopyPointsConstants::SourcePointsLabel, /*bInBroadcastUpdate=*/false);
 }
 
-EPCGDataType UPCGCopyAttributesSettings::GetCurrentPinTypes(const UPCGPin* InPin) const
-{
-	if (!InPin || !InPin->IsOutputPin())
-	{
-		return Super::GetCurrentPinTypes(InPin);
-	}
-
-	EPCGDataType OutputType = GetTypeUnionOfIncidentEdges(PCGCopyPointsConstants::TargetPointsLabel);
-	return OutputType != EPCGDataType::None ? OutputType : EPCGDataType::Any;
-}
-
 TArray<FText> UPCGCopyAttributesSettings::GetNodeTitleAliases() const
 {
 	return { LOCTEXT("CopyAttributesAliasTransfer", "Transfer Attributes") };
diff --git a/Engine/Plugins/PCG/Source/PCG/Private/Elements/PCGCreateAttribute.cpp b/Engine/Plugins/PCG/Source/PCG/Private/Elements/PCGCreateAttribute.cpp
--- a/Engine/Plugins/PCG/Source/PCG/Private/Elements/PCGCreateAttribute.cpp
+++ b/Engine/Plugins/PCG/Source/PCG/Private/Elements/PCGCreateAttribute.cpp
@@ -54,19 +54,6 @@ namespace PCGCreateAttribute
 	}
 }
 
-EPCGDataType UPCGAddAttributeSettings::GetCurrentPinTypes(const UPCGPin* InPin) const
-{
-	check(InPin);
-	if (!InPin->IsOutputPin())
-	{
-		return Super::GetCurrentPinTypes(InPin);
-	}
-
-	// Output pin narrows to union of inputs on first pin
-	const EPCGDataType PrimaryInputType = GetTypeUnionOfIncidentEdges(PCGPinConstants::DefaultInputLabel);
-	return (PrimaryInputType != EPCGDataType::None) ? PrimaryInputType : EPCGDataType::Param; // No input (None) means param.
-}
-
 #if WITH_EDITOR
 bool UPCGAddAttributeSettings::IsPinUsedByNodeExecution(const UPCGPin* InPin) const
 {
diff --git a/Engine/Plugins/PCG/Source/PCG/Private/Elements/PCGDataFromActor.cpp b/Engine/Plugins/PCG/Source/PCG/Private/Elements/PCGDataFromActor.cpp
--- a/Engine/Plugins/PCG/Source/PCG/Private/Elements/PCGDataFromActor.cpp
+++ b/Engine/Plugins/PCG/Source/PCG/Private/Elements/PCGDataFromActor.cpp
@@ -259,7 +259,7 @@ FPCGElementPtr UPCGDataFromActorSettings::CreateElement() const
 	return MakeShared<FPCGDataFromActorElement>();
 }
 
-EPCGDataType UPCGDataFromActorSettings::GetCurrentPinTypes(const UPCGPin* InPin) const
+FPCGDataTypeIdentifier UPCGDataFromActorSettings::GetCurrentPinTypesID(const UPCGPin* InPin) const
 {
 	check(InPin);
 
@@ -267,11 +267,11 @@ EPCGDataType UPCGDataFromActorSettings::GetCurrentPinTypes(const UPCGPin* InPin)
 	{
 		if (Mode == EPCGGetDataFromActorMode::GetSinglePoint)
 		{
-			return EPCGDataType::Point;
+			return FPCGDataTypeIdentifier{EPCGDataType::Point};
 		}
 		else if (Mode == EPCGGetDataFromActorMode::GetActorReference || Mode == EPCGGetDataFromActorMode::GetComponentsReference)
 		{
-			return EPCGDataType::Param;
+			return FPCGDataTypeIdentifier{EPCGDataType::Param};
 		}
 	}
 
diff --git a/Engine/Plugins/PCG/Source/PCG/Private/Elements/PCGExecuteBlueprint.cpp b/Engine/Plugins/PCG/Source/PCG/Private/Elements/PCGExecuteBlueprint.cpp
--- a/Engine/Plugins/PCG/Source/PCG/Private/Elements/PCGExecuteBlueprint.cpp
+++ b/Engine/Plugins/PCG/Source/PCG/Private/Elements/PCGExecuteBlueprint.cpp
@@ -266,24 +266,9 @@ bool UPCGBlueprintElement::IsCacheableOverride_Implementation() const
 
 int32 UPCGBlueprintElement::DynamicPinTypesOverride_Implementation(const UPCGSettings* InSettings, const UPCGPin* InPin) const
 {
-	// Default implementation will modify the Out pin type depending on the input data coming from the In pin. 
-	// If no data arrives, it's not dynamic, or it is another pin, it returns the default allowed type
-	check(InPin);
-
-	if (InSettings->HasDynamicPins() && InPin->IsOutputPin())
-	{
-		const UPCGNode* Node = Cast<const UPCGNode>(InSettings->GetOuter());
-		if (Node && Node->GetInputPin(PCGPinConstants::DefaultInputLabel) != nullptr)
-		{
-			const EPCGDataType InputTypeUnion = InSettings->GetTypeUnionOfIncidentEdges(PCGPinConstants::DefaultInputLabel);
-			if (InputTypeUnion != EPCGDataType::None)
-			{
-				return static_cast<int32>(InputTypeUnion);
-			}
-		}
-	}
-
-	return static_cast<int32>(InPin->Properties.AllowedTypes);
+PRAGMA_DISABLE_DEPRECATION_WARNINGS
+	return static_cast<int32>(InSettings->UPCGSettings::GetCurrentPinTypesID(InPin));
+PRAGMA_ENABLE_DEPRECATION_WARNINGS
 }
 
 TSet<FName> UPCGBlueprintElement::CustomInputLabels() const
@@ -594,14 +579,14 @@ bool UPCGBlueprintSettings::HasDynamicPins() const
 	return BlueprintElementInstance ? BlueprintElementInstance->bHasDynamicPins : UPCGSettings::HasDynamicPins();
 }
 
-EPCGDataType UPCGBlueprintSettings::GetCurrentPinTypes(const UPCGPin* InPin) const
+FPCGDataTypeIdentifier UPCGBlueprintSettings::GetCurrentPinTypesID(const UPCGPin* InPin) const
 {
 	check(InPin);
 
 	// We can't call a BP function while we are PostLoading, in that case (or if we don't have an instance or it's not dynamic) just return the pin allowed types.
 	if (BlueprintElementInstance && BlueprintElementInstance->bHasDynamicPins && !FUObjectThreadContext::Get().IsRoutingPostLoad)
 	{
-		return static_cast<EPCGDataType>(BlueprintElementInstance->DynamicPinTypesOverride(this, InPin));
+		return FPCGDataTypeIdentifier{static_cast<EPCGDataType>(BlueprintElementInstance->DynamicPinTypesOverride(this, InPin))};
 	}
 	else
 	{
diff --git a/Engine/Plugins/PCG/Source/PCG/Private/Elements/PCGFilterByType.cpp b/Engine/Plugins/PCG/Source/PCG/Private/Elements/PCGFilterByType.cpp
--- a/Engine/Plugins/PCG/Source/PCG/Private/Elements/PCGFilterByType.cpp
+++ b/Engine/Plugins/PCG/Source/PCG/Private/Elements/PCGFilterByType.cpp
@@ -17,8 +17,7 @@ FText UPCGFilterByTypeSettings::GetDefaultNodeTitle() const
 		return LOCTEXT("DefaultNodeTitle", "Filter Data By Type");
 	}
 
-	const FText TypeText = FText::FromString(StaticEnum<EPCGDataType>()->GetDisplayNameTextByValue(static_cast<int64>(TargetType)).ToString());
-	return FText::Format(LOCTEXT("NodeTitleFormat", "Filter Data - {0}"), TypeText);
+	return FText::Format(LOCTEXT("NodeTitleFormat", "Filter Data - {0}"), TargetType.ToDisplayText());
 }
 
 FText UPCGFilterByTypeSettings::GetNodeTooltipText() const
@@ -33,7 +32,7 @@ bool UPCGFilterByTypeSettings::GetCompactNodeIcon(FName& OutCompactNodeIcon) con
 }
 #endif // WITH_EDITOR
 
-EPCGDataType UPCGFilterByTypeSettings::GetCurrentPinTypes(const UPCGPin* InPin) const
+FPCGDataTypeIdentifier UPCGFilterByTypeSettings::GetCurrentPinTypesID(const UPCGPin* InPin) const
 {
 	check(InPin);
 	if (InPin->Properties.Label == PCGPinConstants::DefaultInFilterLabel)
@@ -42,7 +41,7 @@ EPCGDataType UPCGFilterByTypeSettings::GetCurrentPinTypes(const UPCGPin* InPin)
 	}
 	else
 	{
-		return Super::GetCurrentPinTypes(InPin);
+		return Super::GetCurrentPinTypesID(InPin);
 	}
 }
 
@@ -81,8 +80,9 @@ bool FPCGFilterByTypeElement::ExecuteInternal(FPCGContext* InContext) const
 		if (Input.Data)
 		{
 			// No match if data does not overlap in type with the target, or if data is broader than the target
-			const bool bDataTypeOverlapsWithTarget = !!(Input.Data->GetDataType() & Settings->TargetType);
-			const bool bDataTypeBroaderThanTarget = !!(Input.Data->GetDataType() & ~Settings->TargetType);
+			const FPCGDataTypeIdentifier InType{Input.Data->GetDataTypeId()};
+			const bool bDataTypeOverlapsWithTarget = InType.Intersects(Settings->TargetType);
+			const bool bDataTypeBroaderThanTarget = InType.IsWider(Settings->TargetType);
 
 			if (bDataTypeOverlapsWithTarget && !bDataTypeBroaderThanTarget)
 			{
diff --git a/Engine/Plugins/PCG/Source/PCG/Private/Elements/PCGGetActorProperty.cpp b/Engine/Plugins/PCG/Source/PCG/Private/Elements/PCGGetActorProperty.cpp
--- a/Engine/Plugins/PCG/Source/PCG/Private/Elements/PCGGetActorProperty.cpp
+++ b/Engine/Plugins/PCG/Source/PCG/Private/Elements/PCGGetActorProperty.cpp
@@ -150,7 +150,7 @@ TArray<FPCGPinProperties> UPCGGetActorPropertySettings::OutputPinProperties() co
 	return PinProperties;
 }
 
-EPCGDataType UPCGGetActorPropertySettings::GetCurrentPinTypes(const UPCGPin* InPin) const
+FPCGDataTypeIdentifier UPCGGetActorPropertySettings::GetCurrentPinTypesID(const UPCGPin* InPin) const
 {
 	check(InPin);
 	// Implementation note: since we can have an input pin in some instances, we can't rely on the
diff --git a/Engine/Plugins/PCG/Source/PCG/Private/Elements/PCGHiGenGridSize.cpp b/Engine/Plugins/PCG/Source/PCG/Private/Elements/PCGHiGenGridSize.cpp
--- a/Engine/Plugins/PCG/Source/PCG/Private/Elements/PCGHiGenGridSize.cpp
+++ b/Engine/Plugins/PCG/Source/PCG/Private/Elements/PCGHiGenGridSize.cpp
@@ -75,7 +75,7 @@ TArray<FPCGPreConfiguredSettingsInfo> UPCGHiGenGridSizeSettings::GetPreconfigure
 }
 #endif
 
-EPCGDataType UPCGHiGenGridSizeSettings::GetCurrentPinTypes(const UPCGPin* InPin) const
+FPCGDataTypeIdentifier UPCGHiGenGridSizeSettings::GetCurrentPinTypesID(const UPCGPin* InPin) const
 {
 	check(InPin);
 	// Non-dynamically-typed pins
@@ -83,10 +83,10 @@ EPCGDataType UPCGHiGenGridSizeSettings::GetCurrentPinTypes(const UPCGPin* InPin)
 	{
 		return InPin->Properties.AllowedTypes;
 	}
-
-	// Output pin narrows to union of inputs on first pin
-	const EPCGDataType InputTypeUnion = GetTypeUnionOfIncidentEdges(PCGPinConstants::DefaultInputLabel);
-	return InputTypeUnion != EPCGDataType::None ? InputTypeUnion : EPCGDataType::Any;
+	else
+	{
+		return Super::GetCurrentPinTypesID(InPin);
+	}
 }
 
 TArray<FPCGPinProperties> UPCGHiGenGridSizeSettings::InputPinProperties() const
diff --git a/Engine/Plugins/PCG/Source/PCG/Private/Elements/PCGMakeConcreteElement.cpp b/Engine/Plugins/PCG/Source/PCG/Private/Elements/PCGMakeConcreteElement.cpp
--- a/Engine/Plugins/PCG/Source/PCG/Private/Elements/PCGMakeConcreteElement.cpp
+++ b/Engine/Plugins/PCG/Source/PCG/Private/Elements/PCGMakeConcreteElement.cpp
@@ -47,13 +47,13 @@ bool FPCGMakeConcreteElement::ExecuteInternal(FPCGContext* InContext) const
 	{
 		const UPCGSpatialData* InputSpatialData = Cast<UPCGSpatialData>(Input.Data);
 		// Verify data really Spatial and reported typed of data also within Spatial.
-		if (!ensure(InputSpatialData) || !ensure(!(Input.Data->GetDataType() & ~EPCGDataType::Spatial)))
+		if (!ensure(InputSpatialData) || !ensure(Input.Data->GetDataTypeId().IsChildOf<FPCGDataTypeInfoSpatial>()))
 		{
 			continue;
 		}
 		
 		// Determine if we have data that needs to be made concrete.
-		const bool bDataIsConcrete = !(Input.Data->GetDataType() & ~EPCGDataType::Concrete);
+		const bool bDataIsConcrete = Input.Data->GetDataTypeId().IsChildOf<FPCGDataTypeInfoConcrete>();
 
 		// Add output, doing a to-point to make concrete if required.
 		FPCGTaggedData& Output = InContext->OutputData.TaggedData.Emplace_GetRef(Input);
diff --git a/Engine/Plugins/PCG/Source/PCG/Private/Elements/PCGProjectionElement.cpp b/Engine/Plugins/PCG/Source/PCG/Private/Elements/PCGProjectionElement.cpp
--- a/Engine/Plugins/PCG/Source/PCG/Private/Elements/PCGProjectionElement.cpp
+++ b/Engine/Plugins/PCG/Source/PCG/Private/Elements/PCGProjectionElement.cpp
@@ -40,17 +40,17 @@ FText UPCGProjectionSettings::GetNodeTooltipText() const
 }
 #endif
 
-EPCGDataType UPCGProjectionSettings::GetCurrentPinTypes(const UPCGPin* InPin) const
+FPCGDataTypeIdentifier UPCGProjectionSettings::GetCurrentPinTypesID(const UPCGPin* InPin) const
 {
 	if (InPin->IsOutputPin())
 	{
-		if (bForceCollapseToPoint || GetTypeUnionOfIncidentEdges(PCGPinConstants::DefaultInputLabel) == EPCGDataType::Point)
+		if (bForceCollapseToPoint || GetTypeUnionIDOfIncidentEdges(PCGPinConstants::DefaultInputLabel) == EPCGDataType::Point)
 		{
-			return EPCGDataType::Point;
+			return FPCGDataTypeIdentifier{EPCGDataType::Point};
 		}
 	}
 
-	return Super::GetCurrentPinTypes(InPin);
+	return Super::GetCurrentPinTypesID(InPin);
 }
 
 FPCGElementPtr UPCGProjectionSettings::CreateElement() const
diff --git a/Engine/Plugins/PCG/Source/PCG/Private/Elements/PCGReroute.cpp b/Engine/Plugins/PCG/Source/PCG/Private/Elements/PCGReroute.cpp
--- a/Engine/Plugins/PCG/Source/PCG/Private/Elements/PCGReroute.cpp
+++ b/Engine/Plugins/PCG/Source/PCG/Private/Elements/PCGReroute.cpp
@@ -248,10 +248,10 @@ TArray<FPCGPinProperties> UPCGNamedRerouteUsageSettings::InputPinProperties() co
 	return PinProperties;
 }
 
-EPCGDataType UPCGNamedRerouteUsageSettings::GetCurrentPinTypes(const UPCGPin* InPin) const
+FPCGDataTypeIdentifier UPCGNamedRerouteUsageSettings::GetCurrentPinTypesID(const UPCGPin* InPin) const
 {
 	// Defer to declaration if possible
-	return Declaration ? Declaration->GetCurrentPinTypes(InPin) : Super::GetCurrentPinTypes(InPin);
+	return Declaration ? Declaration->GetCurrentPinTypesID(InPin) : Super::GetCurrentPinTypesID(InPin);
 }
 
 bool FPCGRerouteElement::ExecuteInternal(FPCGContext* Context) const
diff --git a/Engine/Plugins/PCG/Source/PCG/Private/Graph/PCGGraphCompiler.cpp b/Engine/Plugins/PCG/Source/PCG/Private/Graph/PCGGraphCompiler.cpp
--- a/Engine/Plugins/PCG/Source/PCG/Private/Graph/PCGGraphCompiler.cpp
+++ b/Engine/Plugins/PCG/Source/PCG/Private/Graph/PCGGraphCompiler.cpp
@@ -566,7 +566,7 @@ void FPCGGraphCompiler::CreateGridLinkages(
 				// Make sure we have the correct type in the case of dynamically typed pins.
 				if (const UPCGSettings* UpstreamSettings = InOutCompiledTasks[TaskInput.TaskId].Node ? InOutCompiledTasks[TaskInput.TaskId].Node->GetSettings() : nullptr)
 				{
-					UpstreamPinProps.AllowedTypes = UpstreamSettings->GetCurrentPinTypes(UpstreamPin);
+					UpstreamPinProps.AllowedTypes = UpstreamSettings->GetCurrentPinTypesID(UpstreamPin);
 				}
 
 				// Create pin properties for the input of the LinkTask to convey type information. The grid linkage use the upstream output pin label on both input and
diff --git a/Engine/Plugins/PCG/Source/PCG/Private/Graph/PCGGraphCompilerGPU.cpp b/Engine/Plugins/PCG/Source/PCG/Private/Graph/PCGGraphCompilerGPU.cpp
--- a/Engine/Plugins/PCG/Source/PCG/Private/Graph/PCGGraphCompilerGPU.cpp
+++ b/Engine/Plugins/PCG/Source/PCG/Private/Graph/PCGGraphCompilerGPU.cpp
@@ -472,7 +472,7 @@ void FPCGGraphCompilerGPU::ExpandGPUNodeKernelsToTasks(
 
 							if (ProducerSettings && ProducerPin)
 							{
-								NewTaskInput.UpstreamPin->AllowedTypes = ProducerSettings->GetCurrentPinTypes(ProducerPin);
+								NewTaskInput.UpstreamPin->AllowedTypes = ProducerSettings->GetCurrentPinTypesID(ProducerPin);
 							}
 						}
 
diff --git a/Engine/Plugins/PCG/Source/PCG/Private/Helpers/PCGSettingsHelpers.cpp b/Engine/Plugins/PCG/Source/PCG/Private/Helpers/PCGSettingsHelpers.cpp
--- a/Engine/Plugins/PCG/Source/PCG/Private/Helpers/PCGSettingsHelpers.cpp
+++ b/Engine/Plugins/PCG/Source/PCG/Private/Helpers/PCGSettingsHelpers.cpp
@@ -515,7 +515,7 @@ namespace PCGSettingsHelpers
 
 	// Thread-local variables used for the FPinTypeScopeHelper implementation
 	static thread_local bool bVisitedPinsMapExists = false;
-	static thread_local TMap<const UPCGPin*, EPCGDataType, TInlineSetAllocator<32>> VisitedPinsMap;
+	static thread_local TMap<const UPCGPin*, FPCGDataTypeIdentifier, TInlineSetAllocator<32>> VisitedPinsMap;
 
 	FPinTypeScopeHelper::FPinTypeScopeHelper()
 	{
@@ -535,17 +535,24 @@ namespace PCGSettingsHelpers
 		}
 	}
 
-	EPCGDataType FPinTypeScopeHelper::GetCurrentPinType(const UPCGPin* InPin)
+	EPCGDataType FPinTypeScopeHelper::GetCurrentPinType(const UPCGPin* InPin) const
+	{
+PRAGMA_DISABLE_DEPRECATION_WARNINGS
+		return GetCurrentPinTypeID(InPin);
+PRAGMA_ENABLE_DEPRECATION_WARNINGS
+	}
+
+	FPCGDataTypeIdentifier FPinTypeScopeHelper::GetCurrentPinTypeID(const UPCGPin* InPin) const
 	{
 		// Hitting this ensure means break of contract (scope broken), this helper should never be dynamically allocated.
 		ensure(bVisitedPinsMapExists);
-		if (EPCGDataType* ExistingResult = bVisitedPinsMapExists ? VisitedPinsMap.Find(InPin) : nullptr)
+		if (FPCGDataTypeIdentifier* ExistingResult = bVisitedPinsMapExists ? VisitedPinsMap.Find(InPin) : nullptr)
 		{
 			return *ExistingResult;
 		}
 		else if(InPin)
 		{
-			EPCGDataType PinResult = InPin->GetCurrentTypes();
+			FPCGDataTypeIdentifier PinResult = InPin->GetCurrentTypesID();
 
 			if (bVisitedPinsMapExists)
 			{
@@ -556,7 +563,7 @@ namespace PCGSettingsHelpers
 		}
 		else
 		{
-			return EPCGDataType::None;
+			return {};
 		}
 	}
 }
diff --git a/Engine/Plugins/PCG/Source/PCG/Private/Helpers/PCGSubgraphHelpers.cpp b/Engine/Plugins/PCG/Source/PCG/Private/Helpers/PCGSubgraphHelpers.cpp
--- a/Engine/Plugins/PCG/Source/PCG/Private/Helpers/PCGSubgraphHelpers.cpp
+++ b/Engine/Plugins/PCG/Source/PCG/Private/Helpers/PCGSubgraphHelpers.cpp
@@ -405,15 +405,15 @@ namespace PCGSubgraphHelpersExtra
 		// For the pin type, narrow it down to the input edges (if it is an input pin)
 		if (!InPinToClone->IsOutputPin())
 		{
-			const EPCGDataType FirstPinTypeUnion = PinToCloneSettings->GetTypeUnionOfIncidentEdges(InPinToClone->Properties.Label);
-			if (FirstPinTypeUnion != EPCGDataType::None)
+			const FPCGDataTypeIdentifier FirstPinTypeUnion = PinToCloneSettings->GetTypeUnionIDOfIncidentEdges(InPinToClone->Properties.Label);
+			if (FirstPinTypeUnion.IsValid())
 			{
 				NewProperties.AllowedTypes = FirstPinTypeUnion;
 			}
 		}
 		else
 		{
-			NewProperties.AllowedTypes = PinToCloneSettings->GetCurrentPinTypes(InPinToClone);
+			NewProperties.AllowedTypes = PinToCloneSettings->GetCurrentPinTypesID(InPinToClone);
 		}
 
 		NewProperties.Label = FName(NewName);
diff --git a/Engine/Plugins/PCG/Source/PCG/Private/PCGElement.cpp b/Engine/Plugins/PCG/Source/PCG/Private/PCGElement.cpp
--- a/Engine/Plugins/PCG/Source/PCG/Private/PCGElement.cpp
+++ b/Engine/Plugins/PCG/Source/PCG/Private/PCGElement.cpp
@@ -758,7 +758,7 @@ void IPCGElement::DisabledPassThroughData(FPCGContext* Context) const
 		return;
 	}
 
-	const EPCGDataType OutputType = PassThroughOutputPin->GetCurrentTypes();
+	const FPCGDataTypeIdentifier OutputType = PassThroughOutputPin->GetCurrentTypesID();
 
 	// Grab data from pass-through pin, push it all to output pin
 	Context->OutputData.TaggedData = Context->InputData.GetInputsByPin(PassThroughInputPin->Properties.Label);
@@ -771,16 +771,10 @@ void IPCGElement::DisabledPassThroughData(FPCGContext* Context) const
 	// surface connected to an input pin of type Any, do not pass the surface through to an output pin of type Point).
 	auto InputDataShouldPassThrough = [OutputType](const FPCGTaggedData& InData)
 	{
-		EPCGDataType InputType = InData.Data ? InData.Data->GetDataType() : EPCGDataType::None;
-		const bool bInputTypeNotWiderThanOutputType = !(InputType & ~OutputType);
+		const FPCGDataTypeIdentifier InputType = InData.Data ? FPCGDataTypeIdentifier{InData.Data->GetDataTypeId()} : FPCGDataTypeIdentifier{};
+		const bool bInputTypeNotWiderThanOutputType = !InputType.IsWider(OutputType);
 
-		// Right now we allow edges from Spatial to Concrete. This can happen for example if a point processing node
-		// is receiving a Spatial data, and the node is disabled, it will want to pass the Spatial data through. In the
-		// future we will force collapses/conversions. For now, allow an incoming Spatial to pass out through a Concrete.
-		// TODO remove!
-		const bool bAllowSpatialToConcrete = !!(InputType & EPCGDataType::Spatial) && !!(OutputType & EPCGDataType::Concrete);
-
-		return (InputType != EPCGDataType::Param || OutputType == EPCGDataType::Param) && (bInputTypeNotWiderThanOutputType || bAllowSpatialToConcrete);
+		return bInputTypeNotWiderThanOutputType && (InputType != EPCGDataType::Param || OutputType == EPCGDataType::Param);
 	};
 
 	// Now remove any non-params edges, and if only one edge should come through, remove the others
@@ -998,15 +992,15 @@ void IPCGElement::CleanupAndValidateOutput(FPCGContext* Context) const
 				{
 					// Try to get dynamic current pin types, otherwise settle for static types
 					const UPCGPin* OutputPin = Context->Node ? Context->Node->GetOutputPin(OutputPinProperties[MatchIndex].Label) : nullptr;
-					const EPCGDataType PinTypes = OutputPin ? OutputPin->GetCurrentTypes() : OutputPinProperties[MatchIndex].AllowedTypes;
+					const FPCGDataTypeIdentifier PinTypes = OutputPin ? OutputPin->GetCurrentTypesID() : OutputPinProperties[MatchIndex].AllowedTypes;
 
-					EPCGDataType DataType = TaggedData.Data->GetUnderlyingDataType();
+					FPCGDataTypeIdentifier DataType = TaggedData.Data->GetUnderlyingDataTypeId();
 					const bool bTypesOverlap = !!(PinTypes & DataType);
-					const bool bTypeIsSubset = !(~PinTypes & DataType);
+					const bool bTypeIsSubset = !DataType.IsWider(PinTypes);
 					// TODO: Temporary fix for Settings directly from InputData (ie. from elements with code and not PCG nodes)
 					if ((!bTypesOverlap || !bTypeIsSubset) && DataType != EPCGDataType::Settings)
 					{
-						PCGE_LOG(Warning, GraphAndLog, FText::Format(LOCTEXT("OutputIncompatibleType", "Output data generated for pin '{0}' does not have a compatible type: '{1}'. Consider using more specific/narrower input pin types, or more general/wider output pin types."), FText::FromName(TaggedData.Pin), FText::FromString(UEnum::GetValueAsString(DataType))));
+						PCGE_LOG(Warning, GraphAndLog, FText::Format(LOCTEXT("OutputIncompatibleType", "Output data generated for pin '{0}' does not have a compatible type: '{1}'. Consider using more specific/narrower input pin types, or more general/wider output pin types."), FText::FromName(TaggedData.Pin), DataType.ToDisplayText()));
 					}
 				}
 
diff --git a/Engine/Plugins/PCG/Source/PCG/Private/PCGGraph.cpp b/Engine/Plugins/PCG/Source/PCG/Private/PCGGraph.cpp
--- a/Engine/Plugins/PCG/Source/PCG/Private/PCGGraph.cpp
+++ b/Engine/Plugins/PCG/Source/PCG/Private/PCGGraph.cpp
@@ -2253,7 +2253,7 @@ FString UPCGGraph::GetGraphDescriptionJson() const
 	auto CreatePinDescriptionJsonObject = [](const FPCGPinProperties& Pin)
 	{
 		TSharedPtr<FJsonObject> JsonObject = MakeShared<FJsonObject>();
-		JsonObject->SetStringField(TEXT("Type"), StaticEnum<EPCGDataType>()->GetValueOrBitfieldAsString(static_cast<int64>(Pin.AllowedTypes)));
+		JsonObject->SetStringField(TEXT("Type"), Pin.AllowedTypes.ToString());
 		JsonObject->SetBoolField(TEXT("Multiple"), Pin.bAllowMultipleData);
 		return JsonObject;
 	};
diff --git a/Engine/Plugins/PCG/Source/PCG/Private/PCGNode.cpp b/Engine/Plugins/PCG/Source/PCG/Private/PCGNode.cpp
--- a/Engine/Plugins/PCG/Source/PCG/Private/PCGNode.cpp
+++ b/Engine/Plugins/PCG/Source/PCG/Private/PCGNode.cpp
@@ -623,8 +623,8 @@ const UPCGPin* UPCGNode::GetPassThroughInputPin() const
 	if (const TObjectPtr<UPCGPin>* FirstPassThroughPin = Algo::FindByPredicate(AllInputPins, FindFirstPassThrough))
 	{
 		// Finally in order to be a candidate for passing through, data must be compatible.
-		const EPCGDataType InputType = FirstPassThroughPin->Get()->GetCurrentTypes();
-		const EPCGDataType OutputType = PassThroughOutput->GetCurrentTypes();
+		const FPCGDataTypeIdentifier InputType = FirstPassThroughPin->Get()->GetCurrentTypesID();
+		const FPCGDataTypeIdentifier OutputType = PassThroughOutput->GetCurrentTypesID();
 
 		// Misc note - it would be nice if we could be stricter, like line below this comment. However it would mean it will stop an Any
 		// input being passed through to a Point output, even if the incoming edge will be receiving points dynamically/during execution. If a
diff --git a/Engine/Plugins/PCG/Source/PCG/Private/PCGPin.cpp b/Engine/Plugins/PCG/Source/PCG/Private/PCGPin.cpp
--- a/Engine/Plugins/PCG/Source/PCG/Private/PCGPin.cpp
+++ b/Engine/Plugins/PCG/Source/PCG/Private/PCGPin.cpp
@@ -4,6 +4,7 @@
 
 #include "PCGCommon.h"
 #include "PCGEdge.h"
+#include "PCGModule.h"
 #include "PCGNode.h"
 #include "PCGSettings.h"
 
@@ -11,63 +12,6 @@
 
 #include UE_INLINE_GENERATED_CPP_BY_NAME(PCGPin)
 
-namespace PCGPin
-{
-	bool SortPinsAndRetrieveTypes(const UPCGPin* InPinA, const UPCGPin* InPinB, const UPCGPin*& OutUpstreamPin, const UPCGPin*& OutDownstreamPin, EPCGDataType& OutUpstreamTypes, EPCGDataType& OutDownstreamTypes)
-	{
-		check(InPinA && InPinB);
-		const bool bPinAIsOutput = InPinA->IsOutputPin();
-
-		if (!ensure(bPinAIsOutput != InPinB->IsOutputPin()))
-		{
-			// Cannot connect two pins of same polarity
-			return false;
-		}
-
-		OutUpstreamPin = bPinAIsOutput ? InPinA : InPinB;
-		OutDownstreamPin = bPinAIsOutput ? InPinB : InPinA;
-
-		check(OutUpstreamPin && OutDownstreamPin);
-		OutUpstreamTypes = OutUpstreamPin->GetCurrentTypes();
-		OutDownstreamTypes = OutDownstreamPin->Properties.AllowedTypes;
-
-		return true;
-	}
-
-	bool ArePinTypesCompatible(const EPCGDataType UpstreamTypes, const EPCGDataType DownstreamTypes)
-	{
-		if (UpstreamTypes == EPCGDataType::None || DownstreamTypes == EPCGDataType::None)
-		{
-			return false;
-		}
-
-		// Concrete can always be used as a composite - allow connections from concrete to composite
-		const bool bUpstreamInConcrete = !(UpstreamTypes & ~EPCGDataType::Concrete);
-		const bool bDownstreamIsSpatial = (DownstreamTypes == EPCGDataType::Spatial);
-		if (bUpstreamInConcrete && bDownstreamIsSpatial)
-		{
-			return true;
-		}
-
-		// Anything spatial can collapse to point as long as the downstream pin does not support any other spatial types. For example, Point | Param is collapsible, but Point | Surface is not
-		const bool bUpstreamIsSpatial = !(UpstreamTypes & ~EPCGDataType::Spatial);
-		const bool bDownstreamIsPoint = (DownstreamTypes & EPCGDataType::Spatial) == EPCGDataType::Point;
-		if (bUpstreamIsSpatial && bDownstreamIsPoint)
-		{
-			return true;
-		}
-
-		// Spline -> Surface conversion
-		if (UpstreamTypes == EPCGDataType::Spline && DownstreamTypes == EPCGDataType::Surface)
-		{
-			return true;
-		}
-
-		// Otherwise allow if there is overlap. Don't detect wide -> narrow issues - conversion nodes deal with that
-		return !!(UpstreamTypes & DownstreamTypes);
-	}
-}
-
 #if WITH_EDITOR
 namespace PCGPinPropertiesHelpers
 {
@@ -96,8 +40,13 @@ namespace PCGPinPropertiesHelpers
 }
 #endif // WITH_EDITOR
 
-FPCGPinProperties::FPCGPinProperties(const FName& InLabel, EPCGDataType InAllowedTypes, bool bInAllowMultipleConnections, bool bInAllowMultipleData, const FText& InTooltip)
-	: Label(InLabel), AllowedTypes(InAllowedTypes), bAllowMultipleData(bInAllowMultipleData)
+FPCGPinProperties::FPCGPinProperties()
+{
+	AllowedTypes = FPCGDataTypeIdentifier{EPCGDataType::Any};
+}
+
+FPCGPinProperties::FPCGPinProperties(const FName& InLabel, FPCGDataTypeIdentifier InAllowedTypes, bool bInAllowMultipleConnections, bool bInAllowMultipleData, const FText& InTooltip)
+	: Label(InLabel), AllowedTypes(MoveTemp(InAllowedTypes)), bAllowMultipleData(bInAllowMultipleData)
 #if WITH_EDITORONLY_DATA
 	, Tooltip(InTooltip)
 #endif
@@ -132,7 +81,7 @@ bool FPCGPinProperties::operator==(const FPCGPinProperties& Other) const
 uint32 GetTypeHash(const FPCGPinProperties& Value)
 {
 	uint32 Hash = GetTypeHash(Value.Label);
-	Hash = HashCombine(Hash, static_cast<uint32>(Value.AllowedTypes));
+	Hash = HashCombine(Hash, GetTypeHash(Value.AllowedTypes));
 	Hash = HashCombine(Hash, Value.bAllowMultipleConnections);
 	Hash = HashCombine(Hash, Value.bAllowMultipleData);
 	Hash = HashCombine(Hash, static_cast<uint32>(Value.Usage));
@@ -396,35 +345,43 @@ int32 UPCGPin::EdgeCount() const
 }
 
 EPCGDataType UPCGPin::GetCurrentTypes() const
+{
+PRAGMA_DISABLE_DEPRECATION_WARNINGS
+	return GetCurrentTypesID();
+PRAGMA_ENABLE_DEPRECATION_WARNINGS
+}
+
+FPCGDataTypeIdentifier UPCGPin::GetCurrentTypesID() const
 {
 	check(Node);
 	const UPCGSettings* Settings = Node->GetSettings();
-	return Settings ? Settings->GetCurrentPinTypes(this) : Properties.AllowedTypes;
+	return Settings ? Settings->GetCurrentPinTypesID(this) : Properties.AllowedTypes;
 }
 
 bool UPCGPin::IsCompatible(const UPCGPin* OtherPin) const
 {
-	if (!OtherPin)
-	{
-		return false;
-	}
+	const EPCGDataTypeCompatibilityResult Compatibility = GetCompatibilityWithOtherPin(OtherPin);
+	return Compatibility != EPCGDataTypeCompatibilityResult::NotCompatible && Compatibility != EPCGDataTypeCompatibilityResult::UnknownType;
+}
 
-	const UPCGPin* UpstreamPin = nullptr;
-	const UPCGPin* DownstreamPin = nullptr;
-	EPCGDataType UpstreamTypes = EPCGDataType::None, DownstreamTypes = EPCGDataType::None;
-	if (!PCGPin::SortPinsAndRetrieveTypes(this, OtherPin, UpstreamPin, DownstreamPin, UpstreamTypes, DownstreamTypes))
+EPCGDataTypeCompatibilityResult UPCGPin::GetCompatibilityWithOtherPin(const UPCGPin* OtherPin) const
+{
+	const bool bThisIsOutputPin = IsOutputPin();
+	if (!OtherPin || bThisIsOutputPin == OtherPin->IsOutputPin())
 	{
-		return false;
+		return EPCGDataTypeCompatibilityResult::NotCompatible;
 	}
-	check(UpstreamPin && DownstreamPin);
 
-	return PCGPin::ArePinTypesCompatible(UpstreamTypes, DownstreamTypes);
+	const UPCGPin* UpstreamPin = bThisIsOutputPin ? this : OtherPin;
+	const UPCGPin* DownstreamPin = bThisIsOutputPin ? OtherPin : this;
+
+	return FPCGModule::GetConstDataTypeRegistry().IsCompatible(UpstreamPin->GetCurrentTypesID(), DownstreamPin->Properties.AllowedTypes);
 }
 
 bool UPCGPin::IsDownstreamPinTypeCompatible(const EPCGDataType UpstreamTypes) const
 {
-	const EPCGDataType AllowedTypes = Properties.AllowedTypes;
-	return PCGPin::ArePinTypesCompatible(UpstreamTypes, AllowedTypes);
+	const EPCGDataTypeCompatibilityResult Compatibility = FPCGModule::GetConstDataTypeRegistry().IsCompatible(FPCGDataTypeIdentifier{UpstreamTypes}, Properties.AllowedTypes);
+	return Compatibility != EPCGDataTypeCompatibilityResult::NotCompatible && Compatibility != EPCGDataTypeCompatibilityResult::UnknownType;
 }
 
 bool UPCGPin::AllowsMultipleConnections() const
@@ -442,65 +399,3 @@ bool UPCGPin::CanConnect(const UPCGPin* OtherPin) const
 {
 	return OtherPin && (Edges.IsEmpty() || AllowsMultipleConnections());
 }
-
-EPCGTypeConversion UPCGPin::GetRequiredTypeConversion(const UPCGPin* InOtherPin) const
-{
-	if (!InOtherPin)
-	{
-		return EPCGTypeConversion::Failed;
-	}
-
-	const UPCGPin* UpstreamPin = nullptr;
-	const UPCGPin* DownstreamPin = nullptr;
-	EPCGDataType UpstreamTypes = EPCGDataType::None;
-	EPCGDataType DownstreamTypes = EPCGDataType::None;
-	if (!PCGPin::SortPinsAndRetrieveTypes(this, InOtherPin, UpstreamPin, DownstreamPin, UpstreamTypes, DownstreamTypes))
-	{
-		return EPCGTypeConversion::Failed;
-	}
-	check(UpstreamPin && DownstreamPin);
-
-	// Types same - trivial early out.
-	if (DownstreamTypes == UpstreamTypes)
-	{
-		return EPCGTypeConversion::NoConversionRequired;
-	}
-
-	// Type missing, or types are the same - trivial early out.
-	if (!DownstreamTypes || !UpstreamTypes)
-	{
-		return EPCGTypeConversion::Failed;
-	}
-
-	// Spatial -> Point - "To Point" conversion.
-	// Anything spatial can collapse to point as long as the downstream pin does not support any other spatial types. For example, Point | Param is collapsible, but Point | Surface is not.
-	const bool bUpstreamIsSpatial = !(UpstreamTypes & ~EPCGDataType::Spatial);
-	const bool bUpstreamIsPoint = UpstreamTypes == EPCGDataType::Point;
-	const bool bDownstreamIsSpatialAndPointOnly = (DownstreamTypes & EPCGDataType::Spatial) == EPCGDataType::Point;
-	if (bUpstreamIsSpatial && bDownstreamIsSpatialAndPointOnly)
-	{
-		return bUpstreamIsPoint ? EPCGTypeConversion::NoConversionRequired : EPCGTypeConversion::CollapseToPoint;
-	}
-
-	// Spline -> Surface conversion
-	if (UpstreamTypes == EPCGDataType::Spline && DownstreamTypes == EPCGDataType::Surface)
-	{
-		return EPCGTypeConversion::SplineToSurface;
-	}
-
-	// Any or Spatial -> Concrete - "Make Concrete" conversion. We decided to support Any as it is a superset of Concrete,
-	// and some pins are Any until they dynamically narrow, and other pins may stay Any but still want to be supported.
-	bool bUpstreamIsSpatialOrAny = (UpstreamTypes == EPCGDataType::Any) || (UpstreamTypes == EPCGDataType::Spatial);
-	if (bUpstreamIsSpatialOrAny && DownstreamTypes == EPCGDataType::Concrete)
-	{
-		return EPCGTypeConversion::MakeConcrete;
-	}
-
-	// Requires filter if upstream type is broader than downstream type.
-	if (!!(UpstreamTypes & (~DownstreamTypes)))
-	{
-		return EPCGTypeConversion::Filter;
-	}
-
-	return EPCGTypeConversion::NoConversionRequired;
-}
diff --git a/Engine/Plugins/PCG/Source/PCG/Private/PCGSettings.cpp b/Engine/Plugins/PCG/Source/PCG/Private/PCGSettings.cpp
--- a/Engine/Plugins/PCG/Source/PCG/Private/PCGSettings.cpp
+++ b/Engine/Plugins/PCG/Source/PCG/Private/PCGSettings.cpp
@@ -32,6 +32,8 @@
 
 #define LOCTEXT_NAMESPACE "PCGSettings"
 
+PCG_DEFINE_TYPE_INFO(FPCGDataTypeInfoSettings, UPCGSettings)
+
 namespace PCGSettings
 {
 	// For the deprecation of overridable params as display name to authored names
@@ -452,7 +454,7 @@ void UPCGSettings::FillOverridableParamsPins(TArray<FPCGPinProperties>& OutPins)
 	}
 
 	// Validating that we are not clashing with existing pins
-	TMap<FName, EPCGDataType> InputPinsLabelsAndTypes;
+	TMap<FName, FPCGDataTypeIdentifier> InputPinsLabelsAndTypes;
 
 	for (const FPCGPinProperties& InputPinProperties : OutPins)
 	{
@@ -479,7 +481,7 @@ void UPCGSettings::FillOverridableParamsPins(TArray<FPCGPinProperties>& OutPins)
 
 	// Adding the multi-pin connection for params.
 	// If it already exists (and is the correct type), we can keep it.
-	const EPCGDataType* PinType = InputPinsLabelsAndTypes.Find(PCGPinConstants::DefaultParamsLabel);
+	const FPCGDataTypeIdentifier* PinType = InputPinsLabelsAndTypes.Find(PCGPinConstants::DefaultParamsLabel);
 	if (PinType)
 	{
 		if (*PinType != EPCGDataType::Param)
@@ -552,25 +554,39 @@ TArray<FPCGPinProperties> UPCGSettings::AllInputPinProperties() const
 
 EPCGDataType UPCGSettings::GetTypeUnionOfIncidentEdges(const FName& PinLabel) const
 {
-	EPCGDataType Result = EPCGDataType::None;
+PRAGMA_DISABLE_DEPRECATION_WARNINGS
+	return GetTypeUnionIDOfIncidentEdges(PinLabel);
+PRAGMA_ENABLE_DEPRECATION_WARNINGS
+}
 
-	const UPCGNode* Node = Cast<UPCGNode>(GetOuter());
-	const UPCGPin* Pin = Node ? Node->GetInputPin(PinLabel) : nullptr;
-	if (Pin && Pin->EdgeCount() > 0)
-	{
-		PCGSettingsHelpers::FPinTypeScopeHelper PinTypeHelper;
+FPCGDataTypeIdentifier UPCGSettings::GetTypeUnionIDOfIncidentEdges(const FName& PinLabel) const
+{
+	return GetTypeUnionIDOfAllIncidentEdges({PinLabel});
+}
 
-		for (const UPCGEdge* Edge : Pin->Edges)
+FPCGDataTypeIdentifier UPCGSettings::GetTypeUnionIDOfAllIncidentEdges(TConstArrayView<FName> PinLabels) const
+{
+	TArray<FPCGDataTypeIdentifier, TInlineAllocator<32>> AllTypes;
+	PCGSettingsHelpers::FPinTypeScopeHelper PinTypeHelper;
+	
+	for (const FName PinLabel : PinLabels)
+	{
+		const UPCGNode* Node = Cast<UPCGNode>(GetOuter());
+		const UPCGPin* Pin = Node ? Node->GetInputPin(PinLabel) : nullptr;
+		if (Pin && Pin->EdgeCount() > 0)
 		{
-			const UPCGPin* OtherOutputPin = Edge ? Edge->GetOtherPin(Pin) : nullptr;
-			if (OtherOutputPin)
+			for (const UPCGEdge* Edge : Pin->Edges)
 			{
-				Result |= PinTypeHelper.GetCurrentPinType(OtherOutputPin);
+				const UPCGPin* OtherOutputPin = Edge ? Edge->GetOtherPin(Pin) : nullptr;
+				if (OtherOutputPin)
+				{
+					AllTypes.AddUnique(PinTypeHelper.GetCurrentPinTypeID(OtherOutputPin));
+				}
 			}
 		}
 	}
 
-	return Result;
+	return FPCGModule::GetConstDataTypeRegistry().GetIdentifiersUnion(AllTypes);
 }
 
 TArray<FPCGPinProperties> UPCGSettings::AllOutputPinProperties() const
@@ -942,17 +958,38 @@ bool UPCGSettings::DoesPinSupportPassThrough(UPCGPin* InPin) const
 	return InPin && !InPin->Properties.IsAdvancedPin();
 }
 
+PRAGMA_DISABLE_DEPRECATION_WARNINGS
 EPCGDataType UPCGSettings::GetCurrentPinTypes(const UPCGPin* InPin) const
 {
-	check(InPin);
+	return EPCGDataType::DeprecationSentinel;
+}
+PRAGMA_ENABLE_DEPRECATION_WARNINGS
 
+FPCGDataTypeIdentifier UPCGSettings::GetCurrentPinTypesID(const UPCGPin* InPin) const
+{
+	// To support old code, we call the legacy version. If we have anything different from the Sentinel, it means the function is overridden
+	// in a child class and we want to use that result to avoid any issue, even if we have a loss of info. We'll show a message to indicate
+	// the child class to have its code updated
+	// All PCG vanilla nodes are updated, and none is implementing `GetCurrentPinTypes` anymore.
+	// @todo_pcg: To be removed when GetCurrentPinTypes is removed, in 2 versions
+PRAGMA_DISABLE_DEPRECATION_WARNINGS
+	const EPCGDataType LegacyResult = GetCurrentPinTypes(InPin);
+	if (LegacyResult != EPCGDataType::DeprecationSentinel)
+	{
+		UE_LOG(LogPCG, Warning, TEXT("Class '%s' implements `GetCurrentPinTypes` and not `GetCurrentPinTypesID`. Needs to be updated."), *GetClass()->GetName());
+		return LegacyResult;
+	}
+PRAGMA_ENABLE_DEPRECATION_WARNINGS
+	
 	if (HasDynamicPins() && InPin->IsOutputPin())
 	{
 		const UPCGNode* Node = Cast<const UPCGNode>(GetOuter());
-		if (Node && Node->GetInputPin(PCGPinConstants::DefaultInputLabel) != nullptr)
+		const UPCGPin* FirstPin = Node && !Node->GetInputPins().IsEmpty() ? Node->GetInputPins()[0] : nullptr;
+		
+		if (FirstPin != nullptr && !FirstPin->Properties.IsAdvancedPin() && !FirstPin->Properties.IsDatalessPin())
 		{
-			const EPCGDataType InputTypeUnion = GetTypeUnionOfIncidentEdges(PCGPinConstants::DefaultInputLabel);
-			return InputTypeUnion != EPCGDataType::None ? InputTypeUnion : EPCGDataType::Any;
+			const FPCGDataTypeIdentifier InputTypeUnion = GetTypeUnionIDOfIncidentEdges(FirstPin->Properties.Label);
+			return InputTypeUnion.IsValid() ? InputTypeUnion : InPin->Properties.AllowedTypes;
 		}
 	}
 
diff --git a/Engine/Plugins/PCG/Source/PCG/Private/Tests/Determinism/PCGDeterminismTestsCommon.cpp b/Engine/Plugins/PCG/Source/PCG/Private/Tests/Determinism/PCGDeterminismTestsCommon.cpp
--- a/Engine/Plugins/PCG/Source/PCG/Private/Tests/Determinism/PCGDeterminismTestsCommon.cpp
+++ b/Engine/Plugins/PCG/Source/PCG/Private/Tests/Determinism/PCGDeterminismTestsCommon.cpp
@@ -92,7 +92,7 @@ namespace PCGDeterminismTests
 		// Update the data types tested
 		const UPCGData* InputPCGData = InPCGComponent->GetInputPCGData();
 		check(InputPCGData);
-		OutResult.DataTypesTested |= InputPCGData->GetDataType();
+		OutResult.DataTypesTested |= InputPCGData->GetDataTypeId();
 
 		// Clone the actor and component
 		AActor* PCGActor = InPCGComponent->GetOwner();
@@ -350,7 +350,7 @@ namespace PCGDeterminismTests
 	bool RunBasicSelfTest(const FNodeAndOptions& NodeAndOptions)
 	{
 		const UPCGNode* PCGNode = NodeAndOptions.PCGNode;
-		const TArray<TArray<EPCGDataType>>& BaseOptionsByPin = NodeAndOptions.BaseOptionsByPin;
+		const TArray<TArray<FPCGDataTypeIdentifier>>& BaseOptionsByPin = NodeAndOptions.BaseOptionsByPin;
 
 		PCGTestsCommon::FTestData TestData(NodeAndOptions.Seed, PCGNode->GetSettings());
 
@@ -364,7 +364,7 @@ namespace PCGDeterminismTests
 			// Add an input for each pin
 			for (int32 PinIndex = 0; PinIndex < BaseOptionsByPin.Num(); ++PinIndex)
 			{
-				EPCGDataType DataType = GetPermutation(PermutationIndex, PinIndex, BaseOptionsByPin);
+				FPCGDataTypeIdentifier DataType = GetPermutation(PermutationIndex, PinIndex, BaseOptionsByPin);
 				AddRandomizedInputData(TestData, DataType);
 			}
 
@@ -380,7 +380,7 @@ namespace PCGDeterminismTests
 	bool RunBasicCopiedSelfTest(const FNodeAndOptions& NodeAndOptions)
 	{
 		const UPCGNode* PCGNode = NodeAndOptions.PCGNode;
-		const TArray<TArray<EPCGDataType>>& BaseOptionsByPin = NodeAndOptions.BaseOptionsByPin;
+		const TArray<TArray<FPCGDataTypeIdentifier>>& BaseOptionsByPin = NodeAndOptions.BaseOptionsByPin;
 
 		PCGTestsCommon::FTestData FirstTestData(NodeAndOptions.Seed, PCGNode->GetSettings());
 		PCGTestsCommon::FTestData SecondTestData(NodeAndOptions.Seed, PCGNode->GetSettings());
@@ -396,7 +396,7 @@ namespace PCGDeterminismTests
 			// Add an input for each pin
 			for (int32 PinIndex = 0; PinIndex < BaseOptionsByPin.Num(); ++PinIndex)
 			{
-				EPCGDataType DataType = GetPermutation(PermutationIndex, PinIndex, BaseOptionsByPin);
+				FPCGDataTypeIdentifier DataType = GetPermutation(PermutationIndex, PinIndex, BaseOptionsByPin);
 				AddRandomizedInputData(FirstTestData, DataType);
 				AddRandomizedInputData(SecondTestData, DataType);
 			}
@@ -417,7 +417,7 @@ namespace PCGDeterminismTests
 		TRACE_CPUPROFILER_EVENT_SCOPE(PCGDeterminismTests::GetHighestDeterminismLevel);
 
 		const UPCGNode* PCGNode = NodeAndOptions.PCGNode;
-		const TArray<TArray<EPCGDataType>>& BaseOptionsByPin = NodeAndOptions.BaseOptionsByPin;
+		const TArray<TArray<FPCGDataTypeIdentifier>>& BaseOptionsByPin = NodeAndOptions.BaseOptionsByPin;
 
 		PCGTestsCommon::FTestData FirstTestData(NodeAndOptions.Seed, PCGNode->GetSettings());
 		PCGTestsCommon::FTestData SecondTestData(NodeAndOptions.Seed, PCGNode->GetSettings());
@@ -439,7 +439,7 @@ namespace PCGDeterminismTests
 				{
 					for (int32 I = 0; I < NumInputs; ++I)
 					{
-						EPCGDataType DataType = GetPermutation(PermutationIndex++, PinIndex, BaseOptionsByPin);
+						FPCGDataTypeIdentifier DataType = GetPermutation(PermutationIndex++, PinIndex, BaseOptionsByPin);
 						AddRandomizedInputData(FirstTestData, DataType);
 						AddRandomizedInputData(SecondTestData, DataType);
 					}
@@ -495,7 +495,7 @@ namespace PCGDeterminismTests
 	EDeterminismLevel GetHighestDeterminismLevelSimpleShuffleInput(const FNodeAndOptions& NodeAndOptions, int32 NumInputsPerPin, EDeterminismLevel MaxLevel)
 	{
 		const UPCGNode* PCGNode = NodeAndOptions.PCGNode;
-		const TArray<TArray<EPCGDataType>>& BaseOptionsByPin = NodeAndOptions.BaseOptionsByPin;
+		const TArray<TArray<FPCGDataTypeIdentifier>>& BaseOptionsByPin = NodeAndOptions.BaseOptionsByPin;
 
 		PCGTestsCommon::FTestData FirstTestData(NodeAndOptions.Seed, PCGNode->GetSettings());
 		PCGTestsCommon::FTestData SecondTestData(NodeAndOptions.Seed, PCGNode->GetSettings());
@@ -507,7 +507,7 @@ namespace PCGDeterminismTests
 		{
 			for (int32 I = 0; I < NumInputs; ++I)
 			{
-				EPCGDataType DataType = BaseOptionsByPin[PinIndex][I];
+				FPCGDataTypeIdentifier DataType = BaseOptionsByPin[PinIndex][I];
 				AddRandomizedInputData(FirstTestData, DataType);
 				AddRandomizedInputData(SecondTestData, DataType);
 			}
@@ -556,38 +556,46 @@ namespace PCGDeterminismTests
 		return MaxLevel;
 	}
 
-	void AddRandomizedInputData(PCGTestsCommon::FTestData& TestData, EPCGDataType DataType, const FName& PinName)
+	void AddRandomizedInputData(PCGTestsCommon::FTestData& TestData, FPCGDataTypeIdentifier DataType, const FName& PinName)
 	{
-		switch (DataType)
+		if (DataType == EPCGDataType::Point)
 		{
-		case EPCGDataType::Point:
 			AddRandomizedMultiplePointInputData(TestData, Defaults::NumTestPointsToGenerate, PinName);
 			return;
-		case EPCGDataType::Volume:
+		}
+		else if (DataType == EPCGDataType::Volume)
+		{
 			AddRandomizedVolumeInputData(TestData, PinName);
 			return;
-		case EPCGDataType::PolyLine:
+		}
+		else if (DataType == EPCGDataType::PolyLine)
+		{
 			AddRandomizedPolyLineInputData(TestData, Defaults::NumTestPolyLinePointsToGenerate, PinName);
 			return;
-		case EPCGDataType::Primitive:
+		}
+		else if (DataType == EPCGDataType::Primitive)
+		{
 			AddRandomizedPrimitiveInputData(TestData, PinName);
 			return;
-		case EPCGDataType::Surface:
+		}
+		else if (DataType == EPCGDataType::Surface)
+		{
 			AddRandomizedSurfaceInputData(TestData, PinName);
 			return;
-		case EPCGDataType::Landscape:
+		}
+		else if (DataType == EPCGDataType::Landscape)
+		{
 			AddRandomizedLandscapeInputData(TestData, PinName);
 			return;
+		}
 		// Just return to represent "empty" data
-		case EPCGDataType::None:
+		else if (!DataType.IsValid())
+		{
 			return;
-		case EPCGDataType::Any:
-		case EPCGDataType::Spatial:
-			UE_LOG(LogPCG, Error, TEXT("Specific EPCGDataType required to add to test data"));
-			break;
-		default:
-			UE_LOG(LogPCG, Error, TEXT("Unknown PCGDataType to add to test data: %s"), *UEnum::GetValueAsString(DataType));
-			break;
+		}
+		else
+		{
+			UE_LOG(LogPCG, Error, TEXT("Specific type required to add to test data"));
 		}
 	}
 
@@ -754,8 +762,8 @@ namespace PCGDeterminismTests
 		// Find comparable data from first collection
 		for (int32 I = 0; I < FirstCollection.TaggedData.Num(); ++I)
 		{
-			EPCGDataType FirstDataType = FirstCollection.TaggedData[I].Data->GetDataType();
-			EPCGDataType SecondDataType = SecondCollection.TaggedData[I].Data->GetDataType();
+			FPCGDataTypeIdentifier FirstDataType = FirstCollection.TaggedData[I].Data->GetDataTypeId();
+			FPCGDataTypeIdentifier SecondDataType = SecondCollection.TaggedData[I].Data->GetDataTypeId();
 
 			if (!DataTypeIsComparable(FirstDataType))
 			{
@@ -896,11 +904,11 @@ namespace PCGDeterminismTests
 					continue;
 				}
 
-				EPCGDataType FirstDataType = FirstCollection.TaggedData[I].Data->GetDataType();
-				EPCGDataType SecondDataType = SecondCollection.TaggedData[I].Data->GetDataType();
+				FPCGDataTypeIdentifier FirstDataType = FirstCollection.TaggedData[I].Data->GetDataTypeId();
+				FPCGDataTypeIdentifier SecondDataType = SecondCollection.TaggedData[J].Data->GetDataTypeId();
 
 				// Only compare if they are the same type and same pin label
-				if (FirstDataType == SecondDataType && FirstCollection.TaggedData[I].Pin == SecondCollection.TaggedData[I].Pin)
+				if (FirstDataType == SecondDataType && FirstCollection.TaggedData[I].Pin == SecondCollection.TaggedData[J].Pin)
 				{
 					auto CompareFunction = GetDataCollectionCompareFunction(FirstDataType);
 					if (CompareFunction(FirstCollection.TaggedData[I].Data, SecondCollection.TaggedData[J].Data))
@@ -927,10 +935,10 @@ namespace PCGDeterminismTests
 	{
 		TRACE_CPUPROFILER_EVENT_SCOPE(PCGDeterminismTests::InternalDataMatches);
 
-		check(FirstData->GetDataType() == SecondData->GetDataType());
+		check(FirstData->GetDataTypeId() == SecondData->GetDataTypeId());
 
 		// Filter output differences
-		auto CompareFunction = GetDataCompareFunction(FirstData->GetDataType(), EDeterminismLevel::OrderConsistent);
+		auto CompareFunction = GetDataCompareFunction(FirstData->GetDataTypeId(), EDeterminismLevel::OrderConsistent);
 		return CompareFunction(FirstData, SecondData);
 	}
 
@@ -1423,7 +1431,7 @@ namespace PCGDeterminismTests
 		return ComparisonIsUnimplemented(FirstData, SecondData);
 	}
 
-	bool DataTypeIsComparable(EPCGDataType DataType)
+	bool DataTypeIsComparable(FPCGDataTypeIdentifier DataType)
 	{
 		// Data types that don't need to be compared
 		if (DataType == EPCGDataType::None || DataType == EPCGDataType::Other || DataType == EPCGDataType::Settings)
@@ -1437,13 +1445,13 @@ namespace PCGDeterminismTests
 			return true;
 		}
 
-		UE_LOG(LogPCG, Warning, TEXT("Unknown data comparison type: %s"), *UEnum::GetValueAsString(DataType));
+		UE_LOG(LogPCG, Warning, TEXT("Unknown data comparison type: %s"), *DataType.ToString());
 		return false;
 	}
 
 	bool DataIsComparable(const UPCGData* Data)
 	{
-		return Data && DataTypeIsComparable(Data->GetDataType());
+		return Data && DataTypeIsComparable(Data->GetDataTypeId());
 	}
 
 	bool DataCanBeShuffled(const UPCGData* Data)
@@ -1479,10 +1487,10 @@ namespace PCGDeterminismTests
 		PCGHelpers::ShiftArrayElements<FPCGTaggedData>(TestData.InputData.TaggedData, NumShifts);
 	}
 
-	TArray<EPCGDataType> FilterTestableDataTypes(EPCGDataType AllowedTypes, int32 NumMultipleInputs)
+	TArray<FPCGDataTypeIdentifier> FilterTestableDataTypes(FPCGDataTypeIdentifier AllowedTypes, int32 NumMultipleInputs)
 	{
-		TArray<EPCGDataType> RemainingDataTypes;
-		if (!(EPCGDataType::Spatial & AllowedTypes))
+		TArray<FPCGDataTypeIdentifier> RemainingDataTypes;
+		if (!(AllowedTypes & EPCGDataType::Spatial))
 		{
 			// TODO: Expand this to other data types
 			UE_LOG(LogPCG, Error, TEXT("Only spatial data types are currently supported"));
@@ -1490,9 +1498,9 @@ namespace PCGDeterminismTests
 		}
 
 		// Add every possible combination of testable data types
-		for (EPCGDataType DataType : TestableDataTypes)
+		for (const FPCGDataTypeIdentifier& DataType : TestableDataTypes)
 		{
-			if (DataType == EPCGDataType::None || !!(AllowedTypes & DataType))
+			if (!DataType.IsValid() || !!(AllowedTypes & DataType))
 			{
 				for (int32 I = 0; I < NumMultipleInputs; ++I)
 				{
@@ -1504,22 +1512,22 @@ namespace PCGDeterminismTests
 		return RemainingDataTypes;
 	}
 
-	void RetrieveBaseOptionsPerPin(TArray<TArray<EPCGDataType>>& BaseOptionsArray,
+	void RetrieveBaseOptionsPerPin(TArray<TArray<FPCGDataTypeIdentifier>>& BaseOptionsArray,
 		const TArray<TObjectPtr<UPCGPin>>& InputPins,
-		EPCGDataType& OutDataTypesTested,
+		FPCGDataTypeIdentifier& OutDataTypesTested,
 		int32 NumMultipleInputs)
 	{
 		BaseOptionsArray.Empty();
 
 		for (int32 I = 0; I < InputPins.Num(); ++I)
 		{
-			EPCGDataType AllowedDataTypes = InputPins[I]->Properties.AllowedTypes;
+			FPCGDataTypeIdentifier AllowedDataTypes = InputPins[I]->Properties.AllowedTypes;
 			if (!DataTypeIsComparable(AllowedDataTypes))
 			{
 				continue;
 			}
 
-			TArray<EPCGDataType> TestableDataTypesForPin = FilterTestableDataTypes(AllowedDataTypes, NumMultipleInputs);
+			TArray<FPCGDataTypeIdentifier> TestableDataTypesForPin = FilterTestableDataTypes(AllowedDataTypes, NumMultipleInputs);
 			if (!TestableDataTypesForPin.IsEmpty())
 			{
 				BaseOptionsArray.Emplace(TestableDataTypesForPin);
@@ -1528,7 +1536,7 @@ namespace PCGDeterminismTests
 		}
 	}
 
-	int32 GetNumPermutations(const TArray<TArray<EPCGDataType>>& BaseOptionsArray)
+	int32 GetNumPermutations(const TArray<TArray<FPCGDataTypeIdentifier>>& BaseOptionsArray)
 	{
 		int32 NumPermutations = 1;
 
@@ -1540,7 +1548,7 @@ namespace PCGDeterminismTests
 		return NumPermutations;
 	}
 
-	EPCGDataType GetPermutation(int32 PermutationIteration, int32 PinIndex, const TArray<TArray<EPCGDataType>>& BaseOptionsPerPin)
+	FPCGDataTypeIdentifier GetPermutation(int32 PermutationIteration, int32 PinIndex, const TArray<TArray<FPCGDataTypeIdentifier>>& BaseOptionsPerPin)
 	{
 		check(PermutationIteration >= 0);
 		check(PinIndex >= 0);
@@ -1563,7 +1571,7 @@ namespace PCGDeterminismTests
 		Result = DeterminismLevel;
 	}
 
-	TFunction<bool(const UPCGData*, const UPCGData*)> GetDataCompareFunction(EPCGDataType DataType, EDeterminismLevel DeterminismLevel)
+	TFunction<bool(const UPCGData*, const UPCGData*)> GetDataCompareFunction(FPCGDataTypeIdentifier DataType, EDeterminismLevel DeterminismLevel)
 	{
 		if (!DataTypeIsComparable(DataType))
 		{
@@ -1591,7 +1599,7 @@ namespace PCGDeterminismTests
 		}
 	}
 
-	TFunction<bool(const UPCGData*, const UPCGData*)> GetDataCollectionCompareFunction(EPCGDataType DataType)
+	TFunction<bool(const UPCGData*, const UPCGData*)> GetDataCollectionCompareFunction(FPCGDataTypeIdentifier DataType)
 	{
 		// TODO: Currently only spatial data is comparable. Other types to be added.
 		if (!DataTypeIsComparable(DataType) || !(DataType & EPCGDataType::Spatial))
diff --git a/Engine/Plugins/PCG/Source/PCG/Private/Tests/GraphAuthoring/PCGGraphAuthoringTest.cpp b/Engine/Plugins/PCG/Source/PCG/Private/Tests/GraphAuthoring/PCGGraphAuthoringTest.cpp
deleted file mode 100644
diff --git a/Engine/Plugins/PCG/Source/PCG/Private/Tests/PCGTestsCommon.cpp b/Engine/Plugins/PCG/Source/PCG/Private/Tests/PCGTestsCommon.cpp
--- a/Engine/Plugins/PCG/Source/PCG/Private/Tests/PCGTestsCommon.cpp
+++ b/Engine/Plugins/PCG/Source/PCG/Private/Tests/PCGTestsCommon.cpp
@@ -218,14 +218,14 @@ namespace PCGTestsCommon
 	{
 		TArray<FPCGDataCollection> Data;
 
-		static const TMap<EPCGDataType, TFunction<UPCGData* (void)>> TypeToDataFn
+		static const TMap<FPCGDataTypeIdentifier, TFunction<UPCGData* (void)>> TypeToDataFn
 		{
-			{ EPCGDataType::Point, []() { return PCGTestsCommon::CreatePointData(); }},
-			{ EPCGDataType::PolyLine, []() { return PCGTestsCommon::CreatePolyLineData(); }},
-			{ EPCGDataType::Surface, []() { return PCGTestsCommon::CreateSurfaceData(); }},
-			{ EPCGDataType::Volume, []() { return PCGTestsCommon::CreateVolumeData(); }},
-			{ EPCGDataType::Primitive, []() { return PCGTestsCommon::CreatePrimitiveData(); }},
-			{ EPCGDataType::Param, []() { return PCGTestsCommon::CreateEmptyParamData();}}
+			{ FPCGDataTypeIdentifier{EPCGDataType::Point}, []() { return PCGTestsCommon::CreatePointData(); }},
+			{ FPCGDataTypeIdentifier{EPCGDataType::PolyLine}, []() { return PCGTestsCommon::CreatePolyLineData(); }},
+			{ FPCGDataTypeIdentifier{EPCGDataType::Surface}, []() { return PCGTestsCommon::CreateSurfaceData(); }},
+			{ FPCGDataTypeIdentifier{EPCGDataType::Volume}, []() { return PCGTestsCommon::CreateVolumeData(); }},
+			{ FPCGDataTypeIdentifier{EPCGDataType::Primitive}, []() { return PCGTestsCommon::CreatePrimitiveData(); }},
+			{ FPCGDataTypeIdentifier{EPCGDataType::Param}, []() { return PCGTestsCommon::CreateEmptyParamData();}}
 		};
 
 		// Create empty data
diff --git a/Engine/Plugins/PCG/Source/PCG/Private/Tests/Registry/PCGDataTypeRegistryTests.cpp b/Engine/Plugins/PCG/Source/PCG/Private/Tests/Registry/PCGDataTypeRegistryTests.cpp
--- a/Engine/Plugins/PCG/Source/PCG/Private/Tests/Registry/PCGDataTypeRegistryTests.cpp
+++ b/Engine/Plugins/PCG/Source/PCG/Private/Tests/Registry/PCGDataTypeRegistryTests.cpp
@@ -68,7 +68,7 @@ bool FPCGByClassTest::RunTest(const FString& Parameters)
 		{TEXT("Static mesh resource"), FPCGDataTypeIdentifier::Construct<UPCGStaticMeshResourceData>()},
 		{TEXT("Virtual Texture"), FPCGDataTypeIdentifier::Construct<UPCGVirtualTextureData>()},
 		{TEXT("Render target"), FPCGDataTypeIdentifier::Construct<UPCGRenderTargetData>()},
-		//{TEXT("Settings"), FPCGDataTypeIdentifier::Construct<UPCGSettings>()},
+		{TEXT("Settings"), FPCGDataTypeIdentifier::Construct<UPCGSettings>()},
 		{TEXT("Param"), FPCGDataTypeIdentifier::Construct<UPCGParamData>()},
 		{TEXT("DynamicMesh"), FPCGDataTypeIdentifier::Construct<UPCGDynamicMeshData>()},
 	};
@@ -115,7 +115,7 @@ bool FPCGByDefaultTypeTest::RunTest(const FString& Parameters)
 		{TEXT("Static mesh resource"), FPCGDataTypeIdentifier(EPCGDataType::StaticMeshResource)},
 		{TEXT("Virtual Texture"), FPCGDataTypeIdentifier(EPCGDataType::VirtualTexture)},
 		{TEXT("Render target"),FPCGDataTypeIdentifier(EPCGDataType::RenderTarget)},
-		//{TEXT("Settings"), FPCGDataTypeIdentifier(EPCGDataType::Settings)},
+		{TEXT("Settings"), FPCGDataTypeIdentifier(EPCGDataType::Settings)},
 		{TEXT("Param"), FPCGDataTypeIdentifier(EPCGDataType::Param)},
 		{TEXT("PointOrParam"), FPCGDataTypeIdentifier(EPCGDataType::PointOrParam)},
 		{TEXT("DynamicMesh"), FPCGDataTypeIdentifier(EPCGDataType::DynamicMesh)}
diff --git a/Engine/Plugins/PCG/Source/PCG/Public/Data/Registry/PCGDataType.h b/Engine/Plugins/PCG/Source/PCG/Public/Data/Registry/PCGDataType.h
--- a/Engine/Plugins/PCG/Source/PCG/Public/Data/Registry/PCGDataType.h
+++ b/Engine/Plugins/PCG/Source/PCG/Public/Data/Registry/PCGDataType.h
@@ -102,7 +102,7 @@ struct FPCGDataTypeBaseId
 		return FPCGDataTypeBaseId(T::StaticStruct());
 	}
 
-	friend int32 GetTypeHash(const FPCGDataTypeBaseId& This)
+	friend uint32 GetTypeHash(const FPCGDataTypeBaseId& This)
 	{
 		return GetTypeHash(This.Struct);
 	}
diff --git a/Engine/Plugins/PCG/Source/PCG/Public/Data/Registry/PCGDataTypeIdentifier.h b/Engine/Plugins/PCG/Source/PCG/Public/Data/Registry/PCGDataTypeIdentifier.h
--- a/Engine/Plugins/PCG/Source/PCG/Public/Data/Registry/PCGDataTypeIdentifier.h
+++ b/Engine/Plugins/PCG/Source/PCG/Public/Data/Registry/PCGDataTypeIdentifier.h
@@ -126,10 +126,7 @@ struct FPCGDataTypeIdentifier
 	/** To support previously saved types, that used EPCGDataType, we need to define this function to de-serialize the new class using the old. And add a trait (see below). */
 	PCG_API bool SerializeFromMismatchedTag(const struct FPropertyTag& Tag, FStructuredArchive::FSlot Slot);
 
-	friend uint32 GetTypeHash(const FPCGDataTypeIdentifier& Id)
-	{
-		return HashCombine(GetTypeHash(Id.Ids), Id.CustomSubtype);
-	}
+	PCG_API friend uint32 GetTypeHash(const FPCGDataTypeIdentifier& Id);
 
 	PCG_API FString ToString() const;
 	PCG_API FText ToDisplayText() const;
diff --git a/Engine/Plugins/PCG/Source/PCG/Public/Elements/ControlFlow/PCGBooleanSelect.h b/Engine/Plugins/PCG/Source/PCG/Public/Elements/ControlFlow/PCGBooleanSelect.h
--- a/Engine/Plugins/PCG/Source/PCG/Public/Elements/ControlFlow/PCGBooleanSelect.h
+++ b/Engine/Plugins/PCG/Source/PCG/Public/Elements/ControlFlow/PCGBooleanSelect.h
@@ -23,7 +23,7 @@ public:
 #endif
 
 	virtual bool HasDynamicPins() const override { return true; }
-	virtual EPCGDataType GetCurrentPinTypes(const UPCGPin* Pin) const override;
+	virtual FPCGDataTypeIdentifier GetCurrentPinTypesID(const UPCGPin* Pin) const override;
 
 protected:
 	virtual TArray<FPCGPinProperties> InputPinProperties() const override;
diff --git a/Engine/Plugins/PCG/Source/PCG/Public/Elements/ControlFlow/PCGMultiSelect.h b/Engine/Plugins/PCG/Source/PCG/Public/Elements/ControlFlow/PCGMultiSelect.h
--- a/Engine/Plugins/PCG/Source/PCG/Public/Elements/ControlFlow/PCGMultiSelect.h
+++ b/Engine/Plugins/PCG/Source/PCG/Public/Elements/ControlFlow/PCGMultiSelect.h
@@ -34,7 +34,7 @@ public:
 #endif // WITH_EDITOR
 
 	virtual bool HasDynamicPins() const override { return true; }
-	virtual EPCGDataType GetCurrentPinTypes(const UPCGPin* Pin) const override;
+	virtual FPCGDataTypeIdentifier GetCurrentPinTypesID(const UPCGPin* Pin) const override;
 	virtual FString GetAdditionalTitleInformation() const override;
 	virtual bool HasFlippedTitleLines() const override { return true; }
 
diff --git a/Engine/Plugins/PCG/Source/PCG/Public/Elements/ControlFlow/PCGQualitySelect.h b/Engine/Plugins/PCG/Source/PCG/Public/Elements/ControlFlow/PCGQualitySelect.h
--- a/Engine/Plugins/PCG/Source/PCG/Public/Elements/ControlFlow/PCGQualitySelect.h
+++ b/Engine/Plugins/PCG/Source/PCG/Public/Elements/ControlFlow/PCGQualitySelect.h
@@ -24,7 +24,7 @@ public:
 #endif
 	virtual FString GetAdditionalTitleInformation() const override;
 	virtual bool HasDynamicPins() const override { return true; }
-	virtual EPCGDataType GetCurrentPinTypes(const UPCGPin* InPin) const override;
+	virtual FPCGDataTypeIdentifier GetCurrentPinTypesID(const UPCGPin* Pin) const override;
 
 protected:
 	virtual TArray<FPCGPinProperties> InputPinProperties() const override;
diff --git a/Engine/Plugins/PCG/Source/PCG/Public/Elements/IO/PCGDataTableElement.h b/Engine/Plugins/PCG/Source/PCG/Public/Elements/IO/PCGDataTableElement.h
--- a/Engine/Plugins/PCG/Source/PCG/Public/Elements/IO/PCGDataTableElement.h
+++ b/Engine/Plugins/PCG/Source/PCG/Public/Elements/IO/PCGDataTableElement.h
@@ -26,7 +26,7 @@ public:
 #endif
 
 	virtual bool HasDynamicPins() const override { return true; }
-	virtual EPCGDataType GetCurrentPinTypes(const UPCGPin* InPin) const override;
+	virtual FPCGDataTypeIdentifier GetCurrentPinTypesID(const UPCGPin* InPin) const override;
 
 protected:
 	virtual TArray<FPCGPinProperties> OutputPinProperties() const override;
diff --git a/Engine/Plugins/PCG/Source/PCG/Public/Elements/IO/PCGLoadAssetElement.h b/Engine/Plugins/PCG/Source/PCG/Public/Elements/IO/PCGLoadAssetElement.h
--- a/Engine/Plugins/PCG/Source/PCG/Public/Elements/IO/PCGLoadAssetElement.h
+++ b/Engine/Plugins/PCG/Source/PCG/Public/Elements/IO/PCGLoadAssetElement.h
@@ -43,7 +43,7 @@ public:
 	virtual bool HasDynamicPins() const override { return true; }
 	virtual bool HasFlippedTitleLines() const override { return true; }
 	UE_API virtual FString GetAdditionalTitleInformation() const override;
-	UE_API virtual EPCGDataType GetCurrentPinTypes(const UPCGPin* InPin) const;
+	UE_API virtual FPCGDataTypeIdentifier GetCurrentPinTypesID(const UPCGPin* InPin) const;
 
 protected:
 	UE_API virtual FPCGElementPtr CreateElement() const override;
diff --git a/Engine/Plugins/PCG/Source/PCG/Public/Elements/Metadata/PCGMetadataOpElementBase.h b/Engine/Plugins/PCG/Source/PCG/Public/Elements/Metadata/PCGMetadataOpElementBase.h
--- a/Engine/Plugins/PCG/Source/PCG/Public/Elements/Metadata/PCGMetadataOpElementBase.h
+++ b/Engine/Plugins/PCG/Source/PCG/Public/Elements/Metadata/PCGMetadataOpElementBase.h
@@ -132,7 +132,7 @@ public:
 #endif
 	virtual bool HasFlippedTitleLines() const override { return true; }
 	virtual bool HasDynamicPins() const override { return true; }
-	UE_API virtual EPCGDataType GetCurrentPinTypes(const UPCGPin* InPin) const override;
+	UE_API virtual FPCGDataTypeIdentifier GetCurrentPinTypesID(const UPCGPin* InPin) const override;
 	UE_API virtual bool DoesPinSupportPassThrough(UPCGPin* InPin) const override;
 
 	virtual bool CanCullTaskIfUnwired() const override { return false; }
@@ -207,8 +207,11 @@ public:
 #endif // WITH_EDITORONLY_DATA
 
 protected:
-	/* Return the type union from incident edges with the support for default values. */
+	UE_DEPRECATED(5.7, "Use GetInputPinTypeID")
 	UE_API EPCGDataType GetInputPinType(uint32 Index) const;
+	
+	/* Return the type union from incident edges with the support for default values. */
+	UE_API FPCGDataTypeIdentifier GetInputPinTypeID(uint32 Index) const;
 
 	/* Return the index of the given input pin label. INDEX_NONE if not found*/
 	UE_API uint32 GetInputPinIndex(FName InPinLabel) const;
diff --git a/Engine/Plugins/PCG/Source/PCG/Public/Elements/PCGAttributeFilter.h b/Engine/Plugins/PCG/Source/PCG/Public/Elements/PCGAttributeFilter.h
--- a/Engine/Plugins/PCG/Source/PCG/Public/Elements/PCGAttributeFilter.h
+++ b/Engine/Plugins/PCG/Source/PCG/Public/Elements/PCGAttributeFilter.h
@@ -88,7 +88,7 @@ public:
 	virtual void ApplyPreconfiguredSettings(const FPCGPreConfiguredSettingsInfo& PreconfigureInfo) override;
 	virtual FString GetAdditionalTitleInformation() const override;
 	virtual bool HasDynamicPins() const override { return true; }
-	virtual EPCGDataType GetCurrentPinTypes(const UPCGPin* InPin) const override;
+	virtual FPCGDataTypeIdentifier GetCurrentPinTypesID(const UPCGPin* InPin) const override;
 
 protected:
 	virtual TArray<FPCGPinProperties> InputPinProperties() const override;
@@ -171,7 +171,7 @@ public:
 	virtual void ApplyPreconfiguredSettings(const FPCGPreConfiguredSettingsInfo& PreconfigureInfo) override;
 	virtual FString GetAdditionalTitleInformation() const override;
 	virtual bool HasDynamicPins() const override { return true; }
-	virtual EPCGDataType GetCurrentPinTypes(const UPCGPin* InPin) const override;
+	virtual FPCGDataTypeIdentifier GetCurrentPinTypesID(const UPCGPin* InPin) const override;
 
 protected:
 	virtual TArray<FPCGPinProperties> InputPinProperties() const override;
diff --git a/Engine/Plugins/PCG/Source/PCG/Public/Elements/PCGAttributeNoise.h b/Engine/Plugins/PCG/Source/PCG/Public/Elements/PCGAttributeNoise.h
--- a/Engine/Plugins/PCG/Source/PCG/Public/Elements/PCGAttributeNoise.h
+++ b/Engine/Plugins/PCG/Source/PCG/Public/Elements/PCGAttributeNoise.h
@@ -53,7 +53,6 @@ public:
 #endif
 	virtual bool UseSeed() const override { return true; }
 	virtual void ApplyPreconfiguredSettings(const FPCGPreConfiguredSettingsInfo& PreconfigureInfo) override;
-	virtual EPCGDataType GetCurrentPinTypes(const UPCGPin* InPin) const override;
 	virtual bool HasDynamicPins() const override { return true; }
 	
 
diff --git a/Engine/Plugins/PCG/Source/PCG/Public/Elements/PCGCopyAttributes.h b/Engine/Plugins/PCG/Source/PCG/Public/Elements/PCGCopyAttributes.h
--- a/Engine/Plugins/PCG/Source/PCG/Public/Elements/PCGCopyAttributes.h
+++ b/Engine/Plugins/PCG/Source/PCG/Public/Elements/PCGCopyAttributes.h
@@ -48,7 +48,6 @@ public:
 	virtual EPCGSettingsType GetType() const override { return EPCGSettingsType::Metadata; }
 	virtual void ApplyDeprecation(UPCGNode* InOutNode) override;
 	virtual void ApplyDeprecationBeforeUpdatePins(UPCGNode* InOutNode, TArray<TObjectPtr<UPCGPin>>& InputPins, TArray<TObjectPtr<UPCGPin>>& OutputPins) override;
-	virtual EPCGDataType GetCurrentPinTypes(const UPCGPin* InPin) const override;
 	virtual TArray<FText> GetNodeTitleAliases() const override;
 #endif
 	virtual bool HasDynamicPins() const override { return true; }
diff --git a/Engine/Plugins/PCG/Source/PCG/Public/Elements/PCGCreateAttribute.h b/Engine/Plugins/PCG/Source/PCG/Public/Elements/PCGCreateAttribute.h
--- a/Engine/Plugins/PCG/Source/PCG/Public/Elements/PCGCreateAttribute.h
+++ b/Engine/Plugins/PCG/Source/PCG/Public/Elements/PCGCreateAttribute.h
@@ -40,7 +40,6 @@ public:
 	virtual void ApplyStructuralDeprecation(UPCGNode* InOutNode) override;
 #endif // WITH_EDITOR
 	virtual bool HasDynamicPins() const override { return true; }
-	virtual EPCGDataType GetCurrentPinTypes(const UPCGPin* InPin) const override;
 	virtual FString GetAdditionalTitleInformation() const override;
 
 protected:
diff --git a/Engine/Plugins/PCG/Source/PCG/Public/Elements/PCGDataFromActor.h b/Engine/Plugins/PCG/Source/PCG/Public/Elements/PCGDataFromActor.h
--- a/Engine/Plugins/PCG/Source/PCG/Public/Elements/PCGDataFromActor.h
+++ b/Engine/Plugins/PCG/Source/PCG/Public/Elements/PCGDataFromActor.h
@@ -43,7 +43,7 @@ public:
 	UE_API virtual void ApplyDeprecation(UPCGNode* InOutNode) override;
 #endif
 	virtual bool HasDynamicPins() const override { return true; }
-	UE_API virtual EPCGDataType GetCurrentPinTypes(const UPCGPin* InPin) const override;
+	UE_API virtual FPCGDataTypeIdentifier GetCurrentPinTypesID(const UPCGPin* InPin) const override;
 
 	UE_API virtual FString GetAdditionalTitleInformation() const override;
 
diff --git a/Engine/Plugins/PCG/Source/PCG/Public/Elements/PCGExecuteBlueprint.h b/Engine/Plugins/PCG/Source/PCG/Public/Elements/PCGExecuteBlueprint.h
--- a/Engine/Plugins/PCG/Source/PCG/Public/Elements/PCGExecuteBlueprint.h
+++ b/Engine/Plugins/PCG/Source/PCG/Public/Elements/PCGExecuteBlueprint.h
@@ -340,7 +340,7 @@ public:
 #endif
 	virtual bool UseSeed() const override { return true; }
 	UE_API virtual bool HasDynamicPins() const override;
-	UE_API virtual EPCGDataType GetCurrentPinTypes(const UPCGPin* InPin) const override;
+	UE_API virtual FPCGDataTypeIdentifier GetCurrentPinTypesID(const UPCGPin* InPin) const override;
 
 	UE_API virtual void ApplyPreconfiguredSettings(const FPCGPreConfiguredSettingsInfo& InPreconfiguredsInfo) override;
 	UE_API virtual FString GetAdditionalTitleInformation() const override;
diff --git a/Engine/Plugins/PCG/Source/PCG/Public/Elements/PCGFilterByType.h b/Engine/Plugins/PCG/Source/PCG/Public/Elements/PCGFilterByType.h
--- a/Engine/Plugins/PCG/Source/PCG/Public/Elements/PCGFilterByType.h
+++ b/Engine/Plugins/PCG/Source/PCG/Public/Elements/PCGFilterByType.h
@@ -24,7 +24,7 @@ public:
 #endif
 	virtual bool HasExecutionDependencyPin() const override { return false; }
 	
-	virtual EPCGDataType GetCurrentPinTypes(const UPCGPin* InPin) const override;
+	virtual FPCGDataTypeIdentifier GetCurrentPinTypesID(const UPCGPin* InPin) const override;
 
 protected:
 	virtual TArray<FPCGPinProperties> OutputPinProperties() const override;
@@ -33,7 +33,7 @@ protected:
 
 public:
 	UPROPERTY(BlueprintReadWrite, EditAnywhere, Category = Settings)
-	EPCGDataType TargetType = EPCGDataType::Any;
+	FPCGDataTypeIdentifier TargetType = EPCGDataType::Any;
 
 	UPROPERTY(BlueprintReadWrite, EditAnywhere, Category = Settings)
 	bool bShowOutsideFilter = false;
diff --git a/Engine/Plugins/PCG/Source/PCG/Public/Elements/PCGGetActorProperty.h b/Engine/Plugins/PCG/Source/PCG/Public/Elements/PCGGetActorProperty.h
--- a/Engine/Plugins/PCG/Source/PCG/Public/Elements/PCGGetActorProperty.h
+++ b/Engine/Plugins/PCG/Source/PCG/Public/Elements/PCGGetActorProperty.h
@@ -38,7 +38,7 @@ public:
 	virtual bool CanDynamicallyTrackKeys() const override { return true; }
 #endif
 	virtual bool HasDynamicPins() const override { return true; }
-	virtual EPCGDataType GetCurrentPinTypes(const UPCGPin* InPin) const override;
+	virtual FPCGDataTypeIdentifier GetCurrentPinTypesID(const UPCGPin* InPin) const override;
 
 	virtual FString GetAdditionalTitleInformation() const override;
 	virtual TArray<FPCGPinProperties> OutputPinProperties() const override;
diff --git a/Engine/Plugins/PCG/Source/PCG/Public/Elements/PCGHiGenGridSize.h b/Engine/Plugins/PCG/Source/PCG/Public/Elements/PCGHiGenGridSize.h
--- a/Engine/Plugins/PCG/Source/PCG/Public/Elements/PCGHiGenGridSize.h
+++ b/Engine/Plugins/PCG/Source/PCG/Public/Elements/PCGHiGenGridSize.h
@@ -30,7 +30,7 @@ public:
 	virtual FString GetAdditionalTitleInformation() const override;
 	virtual void ApplyPreconfiguredSettings(const FPCGPreConfiguredSettingsInfo& PreconfigureInfo) override;
 	virtual bool HasDynamicPins() const override { return true; }
-	virtual EPCGDataType GetCurrentPinTypes(const UPCGPin* InPin) const override;
+	virtual FPCGDataTypeIdentifier GetCurrentPinTypesID(const UPCGPin* InPin) const override;
 
 	UPROPERTY(BlueprintReadWrite, EditAnywhere, Category = Settings, meta = (DisplayName = "HiGen Grid Size"))
 	EPCGHiGenGrid HiGenGridSize = EPCGHiGenGrid::Grid256;
diff --git a/Engine/Plugins/PCG/Source/PCG/Public/Elements/PCGOuterIntersectionElement.h b/Engine/Plugins/PCG/Source/PCG/Public/Elements/PCGOuterIntersectionElement.h
--- a/Engine/Plugins/PCG/Source/PCG/Public/Elements/PCGOuterIntersectionElement.h
+++ b/Engine/Plugins/PCG/Source/PCG/Public/Elements/PCGOuterIntersectionElement.h
@@ -29,7 +29,6 @@ public:
 	virtual FText GetNodeTooltipText() const override;
 	virtual EPCGSettingsType GetType() const override { return EPCGSettingsType::Spatial; }
 #endif // WITH_EDITOR
-	virtual bool HasDynamicPins() const override { return true; }
 	// If node disabled, don't intersect - pass through all primary edges
 	virtual bool OnlyPassThroughOneEdgeWhenDisabled() const override { return false; }
 
diff --git a/Engine/Plugins/PCG/Source/PCG/Public/Elements/PCGProjectionElement.h b/Engine/Plugins/PCG/Source/PCG/Public/Elements/PCGProjectionElement.h
--- a/Engine/Plugins/PCG/Source/PCG/Public/Elements/PCGProjectionElement.h
+++ b/Engine/Plugins/PCG/Source/PCG/Public/Elements/PCGProjectionElement.h
@@ -30,7 +30,7 @@ public:
 	virtual EPCGSettingsType GetType() const override { return EPCGSettingsType::Spatial; }
 #endif
 	virtual bool HasDynamicPins() const override { return true; }
-	virtual EPCGDataType GetCurrentPinTypes(const UPCGPin* InPin) const override;	
+	virtual FPCGDataTypeIdentifier GetCurrentPinTypesID(const UPCGPin* InPin) const override;	
 
 protected:
 	virtual TArray<FPCGPinProperties> InputPinProperties() const override;
diff --git a/Engine/Plugins/PCG/Source/PCG/Public/Elements/PCGReroute.h b/Engine/Plugins/PCG/Source/PCG/Public/Elements/PCGReroute.h
--- a/Engine/Plugins/PCG/Source/PCG/Public/Elements/PCGReroute.h
+++ b/Engine/Plugins/PCG/Source/PCG/Public/Elements/PCGReroute.h
@@ -105,7 +105,7 @@ protected:
 public:
 	/** Very counter-intuitive but reroute nodes are normally culled by other means, if they aren't we want to make sure they log errors. */
 	virtual bool CanCullTaskIfUnwired() const { return false; }
-	EPCGDataType GetCurrentPinTypes(const UPCGPin* InPin) const override;
+	FPCGDataTypeIdentifier GetCurrentPinTypesID(const UPCGPin* InPin) const override;
 
 public:
 	UPROPERTY(BlueprintReadWrite, Category = Settings)
diff --git a/Engine/Plugins/PCG/Source/PCG/Public/Helpers/PCGSettingsHelpers.h b/Engine/Plugins/PCG/Source/PCG/Public/Helpers/PCGSettingsHelpers.h
--- a/Engine/Plugins/PCG/Source/PCG/Public/Helpers/PCGSettingsHelpers.h
+++ b/Engine/Plugins/PCG/Source/PCG/Public/Helpers/PCGSettingsHelpers.h
@@ -152,7 +152,10 @@ namespace PCGSettingsHelpers
 		FPinTypeScopeHelper& operator=(const FPinTypeScopeHelper&) = delete;
 		FPinTypeScopeHelper& operator=(FPinTypeScopeHelper&&) = delete;
 
-		UE_API EPCGDataType GetCurrentPinType(const UPCGPin* InPin);
+		UE_DEPRECATED(5.7, "Use GetCurrentPinTypeID version")
+		UE_API EPCGDataType GetCurrentPinType(const UPCGPin* InPin) const;
+
+		UE_API FPCGDataTypeIdentifier GetCurrentPinTypeID(const UPCGPin* InPin) const;
 	private:
 		bool bClearMapOnDestruction = false;
 	};
diff --git a/Engine/Plugins/PCG/Source/PCG/Public/PCGPin.h b/Engine/Plugins/PCG/Source/PCG/Public/PCGPin.h
--- a/Engine/Plugins/PCG/Source/PCG/Public/PCGPin.h
+++ b/Engine/Plugins/PCG/Source/PCG/Public/PCGPin.h
@@ -3,6 +3,8 @@
 #pragma once
 
 #include "PCGCommon.h"
+#include "Data/Registry/PCGDataTypeIdentifier.h"
+#include "Data/Registry/PCGDataTypeRegistry.h"
 
 #include "Kismet/BlueprintFunctionLibrary.h"
 
@@ -43,8 +45,14 @@ struct FPCGPinProperties
 {
 	GENERATED_BODY()
 
-	FPCGPinProperties() = default;
-	UE_API explicit FPCGPinProperties(const FName& InLabel, EPCGDataType InAllowedTypes = EPCGDataType::Any, bool bInAllowMultipleConnections = true, bool bAllowMultipleData = true, const FText& InTooltip = FText::GetEmpty());
+	FPCGPinProperties();
+
+	UE_API explicit FPCGPinProperties(
+		const FName& InLabel,
+		FPCGDataTypeIdentifier InAllowedTypes = FPCGDataTypeIdentifier{EPCGDataType::Any},
+		bool bInAllowMultipleConnections = true,
+		bool bAllowMultipleData = true,
+		const FText& InTooltip = FText::GetEmpty());
 
 	UE_API bool operator==(const FPCGPinProperties& Other) const;
 
@@ -54,8 +62,8 @@ struct FPCGPinProperties
 	UPROPERTY(BlueprintReadWrite, EditAnywhere, Category = Settings)
 	EPCGPinUsage Usage = EPCGPinUsage::Normal;
 
-	UPROPERTY(BlueprintReadWrite, EditAnywhere, Category = Settings)
-	EPCGDataType AllowedTypes = EPCGDataType::Any;
+	UPROPERTY(BlueprintReadWrite, EditAnywhere, Category = Settings, meta = (PCG_SupportsComposition = "SupportsComposition()"))
+	FPCGDataTypeIdentifier AllowedTypes;
 
 	UPROPERTY(BlueprintReadWrite, EditAnywhere, Category = Settings)
 	bool bAllowMultipleData = true;
@@ -131,7 +139,7 @@ struct TStructOpsTypeTraits<FPCGPinProperties> : public TStructOpsTypeTraitsBase
 };
 
 UENUM()
-enum class EPCGTypeConversion : uint8
+enum class UE_DEPRECATED(5.7, "Not used anymore. Use the result from GetCompatibilityWithOtherPin and query the conversion function from FPCGDataTypeInfo form the FPCGDataTypeRegistry.") EPCGTypeConversion : uint8
 {
 	NoConversionRequired,
 	CollapseToPoint,
@@ -165,7 +173,9 @@ public:
 
 		if (DataTypeEnum && ExclusiveDataTypeEnum)
 		{
+			PRAGMA_DISABLE_DEPRECATION_WARNINGS
 			FName DataTypeName = DataTypeEnum->GetNameByValue(static_cast<__underlying_type(EPCGDataType)>(PinProperty.AllowedTypes));
+			PRAGMA_ENABLE_DEPRECATION_WARNINGS
 			if (DataTypeName != NAME_None)
 			{
 				const int64 MatchingType = ExclusiveDataTypeEnum->GetValueByName(DataTypeName);
@@ -187,7 +197,7 @@ public:
 	{
 		EPCGDataType DataType = GetCorrespondingDataType_Internal(AllowedType);
 
-		FPCGPinProperties PinProperties = FPCGPinProperties(Label, DataType, bAllowMultipleConnections, bAllowMultipleData);
+		FPCGPinProperties PinProperties = FPCGPinProperties(Label, FPCGDataTypeIdentifier{DataType}, bAllowMultipleConnections, bAllowMultipleData);
 		if (bIsAdvancedPin)
 		{
 			PinProperties.SetAdvancedPin();
@@ -268,6 +278,8 @@ public:
 	UE_API bool AllowsMultipleConnections() const;
 	UE_API bool AllowsMultipleData() const;
 	UE_API bool IsCompatible(const UPCGPin* OtherPin) const;
+	UE_API EPCGDataTypeCompatibilityResult GetCompatibilityWithOtherPin(const UPCGPin* OtherPin) const;
+
 	/** For checking when the new node doesn't exist yet. */
 	UE_API bool IsDownstreamPinTypeCompatible(EPCGDataType UpstreamTypes) const;
 	UE_API bool CanConnect(const UPCGPin* OtherPin) const;
@@ -292,10 +304,16 @@ public:
 
 	UE_API int32 EdgeCount() const;
 
-	/** Returns the current pin types, which can either be the static types from the pin properties, or a dynamic type based on connected edges. */
+	UE_DEPRECATED(5.7, "Use GetCurrentTypesID")
 	UE_API EPCGDataType GetCurrentTypes() const;
 
-	UE_API EPCGTypeConversion GetRequiredTypeConversion(const UPCGPin* InOtherPin) const;
+	/** Returns the current pin types, which can either be the static types from the pin properties, or a dynamic type based on connected edges. */
+	UE_API FPCGDataTypeIdentifier GetCurrentTypesID() const;
+
+	UE_DEPRECATED(5.7, "Not used anymore. Use the result from GetCompatibilityWithOtherPin and query the conversion function from FPCGDataTypeInfo form the FPCGDataTypeRegistry.")
+	PRAGMA_DISABLE_DEPRECATION_WARNINGS
+	EPCGTypeConversion GetRequiredTypeConversion(const UPCGPin* InOtherPin) const { return EPCGTypeConversion::Failed; };
+	PRAGMA_ENABLE_DEPRECATION_WARNINGS
 };
 
 #if WITH_EDITOR
diff --git a/Engine/Plugins/PCG/Source/PCG/Public/PCGSettings.h b/Engine/Plugins/PCG/Source/PCG/Public/PCGSettings.h
--- a/Engine/Plugins/PCG/Source/PCG/Public/PCGSettings.h
+++ b/Engine/Plugins/PCG/Source/PCG/Public/PCGSettings.h
@@ -6,6 +6,7 @@
 #include "PCGDebug.h"
 #include "PCGElement.h"
 #include "PCGPin.h"
+#include "Data/Registry/PCGDataType.h"
 #include "Elements/PCGActorSelector.h"
 #include "Utils/PCGPreconfiguration.h"
 #include "Tests/Determinism/PCGDeterminismSettings.h"
@@ -140,12 +141,28 @@ struct FPCGSettingsOverridableParam
 #endif // WITH_EDITOR
 };
 
+USTRUCT()
+struct FPCGDataTypeInfoSettings : public FPCGDataTypeInfo
+{
+	GENERATED_BODY()
+
+	PCG_DECLARE_LEGACY_TYPE_INFO(EPCGDataType::Settings);
+
+#if WITH_EDITOR
+	virtual bool Hidden() const override { return true; }
+#endif // WITH_EDITOR
+};
+
 UCLASS(MinimalAPI, Abstract)
 class UPCGSettingsInterface : public UPCGData
 {
 	GENERATED_BODY()
 
 public:
+	// ~Begin UPCGData interface
+	PCG_ASSIGN_DEFAULT_TYPE_INFO(FPCGDataTypeInfoSettings)
+	// ~End UPCGData interface
+	
 	virtual UPCGSettings* GetSettings() PURE_VIRTUAL(UPCGSettingsInterface::GetSettings, return nullptr;);
 	virtual const UPCGSettings* GetSettings() const PURE_VIRTUAL(UPCGSettingsInterface::GetSettings, return nullptr;);
 
@@ -200,7 +217,6 @@ class UPCGSettings : public UPCGSettingsInterface
 
 public:
 	// ~Begin UPCGData interface
-	virtual EPCGDataType GetDataType() const override { return EPCGDataType::Settings; }
 	UE_API virtual void AddToCrc(FArchiveCrc32& Ar, bool bFullDataCrc) const override;
 	// ~End UPCGData interface
 
@@ -237,9 +253,16 @@ public:
 	/** Whether this node can deactivate its output pins during execution, which can dynamically cull downstream nodes. */
 	virtual bool OutputPinsCanBeDeactivated() const { return false; }
 
-	/** Bitwise union of the allowed types of each incident edge on pin. Returns None if no common bits, or no edges. */
+	UE_DEPRECATED(5.7, "Use GetTypeUnionIDOfIncidentEdges version")
 	UE_API EPCGDataType GetTypeUnionOfIncidentEdges(const FName& PinLabel) const;
 
+	/** Bitwise union of the allowed types of each incident edge on pin. Returns None if no common bits, or no edges. */
+	UFUNCTION(BlueprintCallable, Category="Settings|DynamicPins")
+	UE_API FPCGDataTypeIdentifier GetTypeUnionIDOfIncidentEdges(const FName& PinLabel) const;
+
+	/** Bitwise union of the allowed types of each incident edge on pin. Returns None if no common bits, or no edges. */
+	UE_API FPCGDataTypeIdentifier GetTypeUnionIDOfAllIncidentEdges(TConstArrayView<FName> PinLabels) const;
+
 	// Internal functions, should not be used by any user.
 	// Return a different subset for input/output pin properties, in case of a default object.
 	UE_API virtual TArray<FPCGPinProperties> DefaultInputPinProperties() const;
@@ -367,10 +390,13 @@ public:
 	/** By default, the first non-Advanced pin that supports pass-through will be selected. */
 	UE_API virtual bool DoesPinSupportPassThrough(UPCGPin* InPin) const;
 
+	UE_DEPRECATED(5.7, "Use GetCurrentPinTypesID version")
+	UE_API virtual EPCGDataType GetCurrentPinTypes(const UPCGPin* InPin) const;
+	
 	/** Returns the current pin types, which can either be the static types from the pin properties, or a dynamic type based on connected edges.
 	* By default we set output pin types to the union of the default input pin incident edge types, if it is dynamic and the default input exists.
 	*/
-	UE_API virtual EPCGDataType GetCurrentPinTypes(const UPCGPin* InPin) const;
+	UE_API virtual FPCGDataTypeIdentifier GetCurrentPinTypesID(const UPCGPin* InPin) const;
 
 	/** Convert this owning node into the corresponding target. Returns true upon success. */
 	virtual bool ConvertNode(const FPCGPreconfiguredInfo& ConversionInfo) { return false; }
@@ -415,8 +441,14 @@ public:
 protected:
 	// BP version since EPCGDataType is uint32 and BP only supports uint8 enums.
 	/** Bitwise union of the allowed types of each incident edge on pin. Returns None type if no common bits, or no edges. Use the BP function helpers to extract the types from the result. */
+	UE_DEPRECATED(5.7, "Use GetTypeUnionIDOfIncidentEdges")
 	UFUNCTION(BlueprintCallable, Category = "Settings|DynamicPins", DisplayName = "Get Type Union Of Incident Edges")
-	int32 BP_GetTypeUnionOfIncidentEdges(const FName& PinLabel) const { return static_cast<int32>(GetTypeUnionOfIncidentEdges(PinLabel)); }
+	int32 BP_GetTypeUnionOfIncidentEdges(const FName& PinLabel) const
+	{
+		PRAGMA_DISABLE_DEPRECATION_WARNINGS
+		return static_cast<int32>(GetTypeUnionOfIncidentEdges(PinLabel));
+		PRAGMA_ENABLE_DEPRECATION_WARNINGS
+	}
 
 	virtual FPCGElementPtr CreateElement() const PURE_VIRTUAL(UPCGSettings::CreateElement, return nullptr;);
 
diff --git a/Engine/Plugins/PCG/Source/PCG/Public/Tests/Determinism/PCGDeterminismTestsCommon.h b/Engine/Plugins/PCG/Source/PCG/Public/Tests/Determinism/PCGDeterminismTestsCommon.h
--- a/Engine/Plugins/PCG/Source/PCG/Public/Tests/Determinism/PCGDeterminismTestsCommon.h
+++ b/Engine/Plugins/PCG/Source/PCG/Public/Tests/Determinism/PCGDeterminismTestsCommon.h
@@ -53,7 +53,7 @@ struct FDeterminismTestResult
 
 	/** BitFlags for which data types were tested */
 	UPROPERTY(EditAnywhere, BlueprintReadWrite, Category = Determinism)
-	EPCGDataType DataTypesTested = EPCGDataType::None;
+	FPCGDataTypeIdentifier DataTypesTested = EPCGDataType::None;
 
 	/** A mapping of [test names : test results] */
 	UPROPERTY(EditAnywhere, BlueprintReadWrite, Category = Determinism)
@@ -134,7 +134,7 @@ namespace PCGDeterminismTests
 		const UPCGNode* PCGNode;
 		const int32 Seed;
 		const bool bMultipleOptionsPerPin;
-		TArray<TArray<EPCGDataType>> BaseOptionsByPin;
+		TArray<TArray<FPCGDataTypeIdentifier>> BaseOptionsByPin;
 	};
 
 	/** Validates if a PCGNode is deterministic */
@@ -166,7 +166,7 @@ namespace PCGDeterminismTests
 		EDeterminismLevel MaxLevel = EDeterminismLevel::OrderIndependent);
 
 	/** Adds input data to test data, based on input pins' allowed types */
-	void AddRandomizedInputData(PCGTestsCommon::FTestData& TestData, EPCGDataType DataType, const FName& PinName = Defaults::TestPinName);
+	void AddRandomizedInputData(PCGTestsCommon::FTestData& TestData, FPCGDataTypeIdentifier DataType, const FName& PinName = Defaults::TestPinName);
 
 	/** Helper for adding PointData with a single Point to an InputData */
 	void AddSinglePointInputData(FPCGDataCollection& InputData, const FVector& Location, const FName& PinName = Defaults::TestPinName);
@@ -248,7 +248,7 @@ namespace PCGDeterminismTests
 	void UpdateTestResultForOverPermutationLimitError(FDeterminismTestResult& OutResult);
 
 	/** Determines if a PCG DataType is a type that should be compared */
-	bool DataTypeIsComparable(EPCGDataType DataType);
+	bool DataTypeIsComparable(FPCGDataTypeIdentifier DataType);
 	/** Validates PCGData and if its DataType is comparable */
 	bool DataIsComparable(const UPCGData* Data);
 	/** Validates that the data contains a shuffle-able array of data */
@@ -264,24 +264,24 @@ namespace PCGDeterminismTests
 	void ShiftInputOrder(PCGTestsCommon::FTestData& TestData, int32 NumShifts = 1);
 
 	/** Gets an array of the union between an InputPin's allowed types and testable types */
-	TArray<EPCGDataType> FilterTestableDataTypes(EPCGDataType AllowedDataTypes, int32 NumMultipleInputs = 1);
+	TArray<FPCGDataTypeIdentifier> FilterTestableDataTypes(FPCGDataTypeIdentifier AllowedDataTypes, int32 NumMultipleInputs = 1);
 	/** Given input pins, updates an array with a permutation base of options */
-	void RetrieveBaseOptionsPerPin(TArray<TArray<EPCGDataType>>& InBaseOptionsArray,
+	void RetrieveBaseOptionsPerPin(TArray<TArray<FPCGDataTypeIdentifier>>& InBaseOptionsArray,
 		const TArray<TObjectPtr<UPCGPin>>& InputPins,
-		EPCGDataType& OutDataTypesTested,
+		FPCGDataTypeIdentifier& OutDataTypesTested,
 		int32 NumMultipleInputs = 1);
 	/** Gets the number of permutations from a base set of options */
-	int32 GetNumPermutations(const TArray<TArray<EPCGDataType>>& BaseOptionsArray);
+	int32 GetNumPermutations(const TArray<TArray<FPCGDataTypeIdentifier>>& BaseOptionsArray);
 	/** Gets the index of a pin's permutation based on the permutation iteration */
-	EPCGDataType GetPermutation(int32 PermutationIteration, int32 PinIndex, const TArray<TArray<EPCGDataType>>& BaseOptionsPerPin);
+	FPCGDataTypeIdentifier GetPermutation(int32 PermutationIteration, int32 PinIndex, const TArray<TArray<FPCGDataTypeIdentifier>>& BaseOptionsPerPin);
 
 	/** Helper to update a tests' outgoing results */
 	void UpdateTestResults(FName TestName, FDeterminismTestResult& OutResult, EDeterminismLevel DeterminismLevel);
 
 	/** Gets a comparison function to compare two data objects */
-	TFunction<bool(const UPCGData*, const UPCGData*)> GetDataCompareFunction(EPCGDataType DataType, EDeterminismLevel DeterminismLevel);
+	TFunction<bool(const UPCGData*, const UPCGData*)> GetDataCompareFunction(FPCGDataTypeIdentifier DataType, EDeterminismLevel DeterminismLevel);
 	/** Gets a comparison function to compare two data collection entry basics */
-	TFunction<bool(const UPCGData*, const UPCGData*)> GetDataCollectionCompareFunction(EPCGDataType DataType);
+	TFunction<bool(const UPCGData*, const UPCGData*)> GetDataCollectionCompareFunction(FPCGDataTypeIdentifier DataType);
 
 	/** Executes the element with given test data */
 	void ExecuteWithTestData(PCGTestsCommon::FTestData& TestData, const UPCGNode* PCGNode);
diff --git a/Engine/Plugins/PCG/Source/PCGEditor/Private/Nodes/PCGEditorGraphNodeBase.cpp b/Engine/Plugins/PCG/Source/PCGEditor/Private/Nodes/PCGEditorGraphNodeBase.cpp
--- a/Engine/Plugins/PCG/Source/PCGEditor/Private/Nodes/PCGEditorGraphNodeBase.cpp
+++ b/Engine/Plugins/PCG/Source/PCGEditor/Private/Nodes/PCGEditorGraphNodeBase.cpp
@@ -1176,11 +1176,11 @@ FEdGraphPinType UPCGEditorGraphNodeBase::GetPinType(const UPCGPin* InPin)
 	EdPinType.ResetToDefaults();
 
 	check(InPin);
-	const EPCGDataType PinType = InPin->GetCurrentTypes();
+	const FPCGDataTypeIdentifier PinType = InPin->GetCurrentTypesID();
 
-	auto CheckType = [PinType](EPCGDataType AllowedType)
+	auto CheckType = [PinType](FPCGDataTypeIdentifier AllowedType)
 	{
-		return !!(PinType & AllowedType) && !(PinType & ~AllowedType);
+		return !!(PinType & AllowedType) && !PinType.IsWider(AllowedType);
 	};
 
 	if (CheckType(EPCGDataType::Concrete))
@@ -1273,35 +1273,21 @@ void UPCGEditorGraphNodeBase::GetPinHoverText(const UEdGraphPin& Pin, FString& H
 	const bool bIsInputPin = (Pin.Direction == EGPD_Input);
 	UPCGPin* MatchingPin = (PCGNode ? (bIsInputPin ? PCGNode->GetInputPin(Pin.PinName) : PCGNode->GetOutputPin(Pin.PinName)) : nullptr);
 
-	auto PCGDataTypeToText = [](EPCGDataType DataType)
-	{
-		TArray<FText> BitFlags;
-
-		for (uint64 BitIndex = 1; BitIndex < 8 * sizeof(EPCGDataType); ++BitIndex)
-		{
-			const int64 BitValue = static_cast<__underlying_type(EPCGDataType)>(DataType) & (1 << BitIndex);
-			if (BitValue)
-			{
-				BitFlags.Add(StaticEnum<EPCGDataType>()->GetDisplayNameTextByValue(BitValue));
-			}
-		}
-
-		return FText::Join(FText(LOCTEXT("Delimiter", " | ")), BitFlags);
-	};
-
-	auto PinTypeToText = [&PCGDataTypeToText](const FName& Category, UPCGPin* MatchingPin)
+	auto PinTypeToText = [](const FName& Category, UPCGPin* MatchingPin)
 	{
 		if (Category != NAME_None)
 		{
 			return FText::FromName(Category);
 		}
-		else if (MatchingPin && MatchingPin->GetCurrentTypes() == EPCGDataType::Any)
+			
+		const FPCGDataTypeIdentifier Type = MatchingPin ? MatchingPin->GetCurrentTypesID() : FPCGDataTypeIdentifier{};
+		if (Type == EPCGDataType::Any)
 		{
 			return FText::FromName(FName("Any"));
 		}
-		else if (MatchingPin)
+		else if (Type.IsValid())
 		{
-			return PCGDataTypeToText(MatchingPin->GetCurrentTypes());
+			return Type.ToDisplayText();
 		}
 		else
 		{
diff --git a/Engine/Plugins/PCG/Source/PCGEditor/Private/Nodes/SPCGEditorGraphNode.cpp b/Engine/Plugins/PCG/Source/PCGEditor/Private/Nodes/SPCGEditorGraphNode.cpp
--- a/Engine/Plugins/PCG/Source/PCGEditor/Private/Nodes/SPCGEditorGraphNode.cpp
+++ b/Engine/Plugins/PCG/Source/PCGEditor/Private/Nodes/SPCGEditorGraphNode.cpp
@@ -212,15 +212,17 @@ void SPCGEditorGraphNode::AddPin(const TSharedRef<SGraphPin>& PinToAdd)
 			const bool bIsMultiData = Pin->Properties.bAllowMultipleData;
 			const bool bIsMultiConnections = Pin->AllowsMultipleConnections();
 
+			const FPCGDataTypeIdentifier CurrentType = Pin->GetCurrentTypesID();
+
 			// Check for special types
-			if (Pin->GetCurrentTypes() == EPCGDataType::Param)
+			if (CurrentType.IsSameType(FPCGDataTypeIdentifier{EPCGDataType::Param}))
 			{
 				const FSlateBrush* ConnectedBrush = FPCGEditorStyle::Get().GetBrush(bIsInPin ? PCGEditorStyleConstants::Pin_Param_IN_C : PCGEditorStyleConstants::Pin_Param_OUT_C);
 				const FSlateBrush* DisconnectedBrush = FPCGEditorStyle::Get().GetBrush(bIsInPin ? PCGEditorStyleConstants::Pin_Param_IN_DC : PCGEditorStyleConstants::Pin_Param_OUT_DC);
 
 				PinToAdd->SetCustomPinIcon(ConnectedBrush, DisconnectedBrush);
 			}
-			else if (Pin->GetCurrentTypes() == EPCGDataType::Spatial)
+			else if (CurrentType.IsSameType(FPCGDataTypeIdentifier{EPCGDataType::Spatial}))
 			{
 				const FSlateBrush* ConnectedBrush = FPCGEditorStyle::Get().GetBrush(bIsInPin ? PCGEditorStyleConstants::Pin_Composite_IN_C : PCGEditorStyleConstants::Pin_Composite_OUT_C);
 				const FSlateBrush* DisconnectedBrush = FPCGEditorStyle::Get().GetBrush(bIsInPin ? PCGEditorStyleConstants::Pin_Composite_IN_DC : PCGEditorStyleConstants::Pin_Composite_OUT_DC);
diff --git a/Engine/Plugins/PCG/Source/PCGEditor/Private/PCGEditorGraphSchema.cpp b/Engine/Plugins/PCG/Source/PCGEditor/Private/PCGEditorGraphSchema.cpp
--- a/Engine/Plugins/PCG/Source/PCGEditor/Private/PCGEditorGraphSchema.cpp
+++ b/Engine/Plugins/PCG/Source/PCGEditor/Private/PCGEditorGraphSchema.cpp
@@ -14,13 +14,15 @@
 #include "Elements/PCGMakeConcreteElement.h"
 #include "Elements/PCGReroute.h"
 #include "Elements/PCGUserParameterGet.h"
+#include "Helpers/PCGHelpers.h"
+#include "Helpers/PCGPropertyHelpers.h"
+#include "Helpers/PCGSubgraphHelpers.h"
 
 #include "PCGEditor.h"
 #include "PCGEditorCommon.h"
 #include "PCGEditorGraph.h"
 #include "PCGEditorSettings.h"
 #include "PCGEditorUtils.h"
-#include "Helpers/PCGPropertyHelpers.h"
 #include "Nodes/PCGEditorGraphNodeBase.h"
 #include "Nodes/PCGEditorGraphNodeReroute.h"
 #include "Schema/PCGEditorGraphSchemaActions.h"
@@ -155,23 +157,20 @@ const FPinConnectionResponse UPCGEditorGraphSchema::CanCreateConnection(const UE
 		return FPinConnectionResponse(CONNECT_RESPONSE_DISALLOW, LOCTEXT("ConnectionFailed", "Unable to verify pins"));
 	}
 
-	if (!InputPin->IsCompatible(OutputPin))
-	{
-		return FPinConnectionResponse(CONNECT_RESPONSE_DISALLOW, LOCTEXT("ConnectionTypesIncompatible", "Pins are incompatible"));
-	}
+	const EPCGDataTypeCompatibilityResult Compatibility = InputPin->GetCompatibilityWithOtherPin(OutputPin);
 
-	const EPCGTypeConversion RequiredTypeConversion = InputPin->GetRequiredTypeConversion(OutputPin);
-	if (RequiredTypeConversion == EPCGTypeConversion::CollapseToPoint)
-	{
-		return FPinConnectionResponse(CONNECT_RESPONSE_MAKE_WITH_CONVERSION_NODE, LOCTEXT("ConnectionConversionToPoint", "Convert to Point"));
-	}
-	else if (RequiredTypeConversion == EPCGTypeConversion::Filter)
+	switch (Compatibility)
 	{
+	case EPCGDataTypeCompatibilityResult::RequireFilter:
 		return FPinConnectionResponse(CONNECT_RESPONSE_MAKE_WITH_CONVERSION_NODE, LOCTEXT("ConnectionUsingFilter", "Filter data to match type"));
-	}
-	else if (RequiredTypeConversion == EPCGTypeConversion::MakeConcrete)
-	{
-		return FPinConnectionResponse(CONNECT_RESPONSE_MAKE_WITH_CONVERSION_NODE, LOCTEXT("ConnectionUsingMakeConcrete", "Make data concrete"));
+	case EPCGDataTypeCompatibilityResult::RequireConversion:
+		return FPinConnectionResponse(CONNECT_RESPONSE_MAKE_WITH_CONVERSION_NODE, LOCTEXT("ConnectionConversion", "Convert data to match type"));
+	case EPCGDataTypeCompatibilityResult::NotCompatible: // fall-through
+	case EPCGDataTypeCompatibilityResult::UnknownType:
+		return FPinConnectionResponse(CONNECT_RESPONSE_DISALLOW, LOCTEXT("ConnectionTypesIncompatible", "Pins are incompatible"));
+	case EPCGDataTypeCompatibilityResult::Compatible:// fall-through
+	default:
+		break;
 	}
 
 	if (!InputPin->AllowsMultipleConnections() && InputPin->EdgeCount() > 0)
@@ -221,10 +220,13 @@ bool UPCGEditorGraphSchema::TryCreateConnectionInternal(UEdGraphPin* InA, UEdGra
 	UPCGNode* PCGNodeB = PCGEdGraphNodeB->GetPCGNode();
 	check(PCGNodeA && PCGNodeB);
 
-	const UPCGPin* PCGPinA = PCGNodeA->GetOutputPin(A->PinName);
-	const UPCGPin* PCGPinB = PCGNodeB->GetInputPin(B->PinName);
+	UPCGPin* PCGPinA = PCGNodeA->GetOutputPin(A->PinName);
+	UPCGPin* PCGPinB = PCGNodeB->GetInputPin(B->PinName);
 	check(PCGPinA && PCGPinB);
-	if (!PCGPinA->IsCompatible(PCGPinB))
+
+	const EPCGDataTypeCompatibilityResult Compatibility = PCGPinA->GetCompatibilityWithOtherPin(PCGPinB);
+	
+	if (!PCGDataTypeCompatibilityResult::IsValid(Compatibility))
 	{
 		return false;
 	}
@@ -235,129 +237,36 @@ bool UPCGEditorGraphSchema::TryCreateConnectionInternal(UEdGraphPin* InA, UEdGra
 	// UPCGEditorGraphSchema::TryCreateConnectionInternal is called directly by FDragConnection::DroppedOnPin
 	PCGGraph->PrimeGraphCompilationCache();
 
-	// Creates a connection via an intermediate conversion node.
-	auto ConnectViaIntermediate = [this, NodeA, NodeB](UPCGNode* IntermediateNode, UEdGraphPin* A, UEdGraphPin* B)
+	bool bShouldReconstruct = false;
+	
+	if (Compatibility == EPCGDataTypeCompatibilityResult::RequireFilter && bAddConversionNodeIfNeeded)
 	{
-		UEdGraph* Graph = NodeA->GetGraph();
-		check(Graph);
-		Graph->Modify();
-
-		FGraphNodeCreator<UPCGEditorGraphNode> NodeCreator(*Graph);
-		UPCGEditorGraphNode* ConversionNode = NodeCreator.CreateUserInvokedNode(/*bSelectNewNode=*/false);
-		ConversionNode->Construct(IntermediateNode);
-
-		// Put the conversion node between A & B but make it stay within a radius of B to keep things tidy.
-		{
-			// Initially place at mid point
-			ConversionNode->NodePosX = (NodeA->NodePosX + NodeB->NodePosX) / 2;
-			ConversionNode->NodePosY = (NodeA->NodePosY + NodeB->NodePosY) / 2;
-
-			// A hand tweaked distance that keeps it reasonably close.
-			constexpr float MaxDistFromB = 200;
-			const FVector2D OffsetFromB(ConversionNode->NodePosX - NodeB->NodePosX, ConversionNode->NodePosY - NodeB->NodePosY);
-			const float Dist = OffsetFromB.Length();
-			if (Dist > MaxDistFromB)
-			{
-				const float Scale = MaxDistFromB / Dist;
-				ConversionNode->NodePosX = NodeB->NodePosX + Scale * OffsetFromB.X;
-				ConversionNode->NodePosY = NodeB->NodePosY + Scale * OffsetFromB.Y;
-			}
-		}
-
-		NodeCreator.Finalize();
-
-		IntermediateNode->PositionX = ConversionNode->NodePosX;
-		IntermediateNode->PositionY = ConversionNode->NodePosY;
-
-		bool bModifiedA = false, bModifiedB = false;
-
-		UEdGraphPin*const* ConversionInputPin = ConversionNode->GetAllPins().FindByPredicate([](const UEdGraphPin* InPin)
-		{
-			return InPin->Direction == EGPD_Input && InPin->GetFName() == PCGPinConstants::DefaultInputLabel;
-		});
-
-		if (ensure(ConversionInputPin && *ConversionInputPin))
-		{
-			// Last argument: don't allow recursively adding conversion nodes.
-			bModifiedA = TryCreateConnectionInternal(A, *ConversionInputPin, /*bAddConversionNodeIfNeeded=*/false);
-		}
-
-		// Call GetAllPins() a second time. It's important that we wire up the pins one at a time. Wiring a pin can change dynamic pin types
-		// which can refresh the node, so we must re-query the pins after each connection is made.
-		UEdGraphPin*const* ConversionOutputPin = ConversionNode->GetAllPins().FindByPredicate([](const UEdGraphPin* InPin)
+		UPCGNode* ConversionPCGNode = FPCGSubgraphHelpers::SpawnNodeAndConnect(PCGGraph, PCGPinA, PCGPinB, UPCGFilterByTypeSettings::StaticClass(), [PCGPinB](UPCGSettings* NodeSettings)
 		{
-			return InPin->Direction == EGPD_Output && (InPin->GetFName() == PCGPinConstants::DefaultOutputLabel || InPin->GetFName() == PCGPinConstants::DefaultInFilterLabel);
+			UPCGFilterByTypeSettings* Settings = CastChecked<UPCGFilterByTypeSettings>(NodeSettings);
+			Settings->TargetType = PCGPinB->Properties.AllowedTypes;
+			return true;
 		});
 
-		if (ensure(ConversionOutputPin && *ConversionOutputPin))
+		if (ConversionPCGNode)
 		{
-			// Last argument: don't allow recursively adding conversion nodes.
-			bModifiedB = TryCreateConnectionInternal(*ConversionOutputPin, B, /*bAddConversionNodeIfNeeded=*/false);
+			bShouldReconstruct = true;
 		}
-
-		return bModifiedA || bModifiedB;
-	};
-
-	const EPCGTypeConversion Conversion = bAddConversionNodeIfNeeded ? PCGPinA->GetRequiredTypeConversion(PCGPinB) : EPCGTypeConversion::NoConversionRequired;
-	if (Conversion == EPCGTypeConversion::CollapseToPoint)
-	{
-		UPCGSettings* NodeSettings = nullptr;
-		UPCGNode* ConversionPCGNode = PCGGraph->AddNodeOfType(UPCGCollapseSettings::StaticClass(), NodeSettings);
-
-		return ConnectViaIntermediate(ConversionPCGNode, A, B);
 	}
-	else if (Conversion == EPCGTypeConversion::Filter)
+	else if (Compatibility == EPCGDataTypeCompatibilityResult::RequireConversion && bAddConversionNodeIfNeeded)
 	{
-		UPCGSettings* NodeSettings = nullptr;
-		UPCGNode* ConversionPCGNode = PCGGraph->AddNodeOfType(UPCGFilterByTypeSettings::StaticClass(), NodeSettings);
-
-		FName PinAName = A->GetFName();
-		FName PinBName = B->GetFName();
-
-		// Setup the output pin based on the conversion target type, before new node is finalized.
-		UPCGFilterByTypeSettings* Settings = CastChecked<UPCGFilterByTypeSettings>(NodeSettings);
-		Settings->TargetType = PCGPinB->Properties.AllowedTypes;
-		ConversionPCGNode->UpdateAfterSettingsChangeDuringCreation();
-
-		// Requery A & B because the update might have changed the pins.
-		UEdGraphPin* const* NewA = NodeA->GetAllPins().FindByPredicate([&PinAName](const UEdGraphPin* InPin)
-		{
-			return InPin->Direction == EGPD_Output && InPin->GetName() == PinAName;
-		});
-
-		if (NewA)
-		{
-			A = *NewA;
-		}
-
-		UEdGraphPin* const* NewB = NodeB->GetAllPins().FindByPredicate([&PinBName](const UEdGraphPin* InPin)
-		{
-			return InPin->Direction == EGPD_Input && InPin->GetName() == PinBName;
-		});
-
-		if (NewB)
+		const FPCGDataTypeRegistry& Registry = FPCGModule::GetConstDataTypeRegistry();
+		for (const FPCGDataTypeBaseId& ID : PCGPinA->Properties.AllowedTypes.GetIds())
 		{
-			B = *NewB;
+			if (const FPCGDataTypeInfo* TypeInfo = Registry.GetTypeInfo(ID); TypeInfo && TypeInfo->SupportsConversion(PCGPinB->Properties.AllowedTypes))
+			{
+				if (TypeInfo->AddConversionNode(PCGPinB->Properties.AllowedTypes, PCGGraph, PCGPinA, PCGPinB) != nullptr)
+				{
+					bShouldReconstruct = true;
+					break;
+				}
+			}
 		}
-
-		return ConnectViaIntermediate(ConversionPCGNode, A, B);
-	}
-	else if (Conversion == EPCGTypeConversion::MakeConcrete)
-	{
-		UPCGSettings* NodeSettings = nullptr;
-		UPCGNode* ConversionPCGNode = PCGGraph->AddNodeOfType(UPCGMakeConcreteSettings::StaticClass(), NodeSettings);
-
-		return ConnectViaIntermediate(ConversionPCGNode, A, B);
-	}
-	else if (Conversion == EPCGTypeConversion::SplineToSurface)
-	{
-		UPCGSettings* NodeSettings = nullptr;
-		UPCGNode* ConversionPCGNode = PCGGraph->AddNodeOfType(UPCGCreateSurfaceFromSplineSettings::StaticClass(), NodeSettings);
-
-		UPCGCreateSurfaceFromSplineSettings* Settings = CastChecked<UPCGCreateSurfaceFromSplineSettings>(NodeSettings);
-		Settings->bShouldDrawNodeCompact = true;
-
-		return ConnectViaIntermediate(ConversionPCGNode, A, B);
 	}
 	else
 	{
@@ -369,6 +278,16 @@ bool UPCGEditorGraphSchema::TryCreateConnectionInternal(UEdGraphPin* InA, UEdGra
 
 		return bModified;
 	}
+
+	if (bShouldReconstruct)
+	{
+		if (UPCGEditorGraph* Graph = Cast<UPCGEditorGraph>(NodeA->GetGraph()); ensure(Graph))
+		{
+			Graph->ReconstructGraph();
+		}
+	}
+
+	return bShouldReconstruct;
 }
 
 void UPCGEditorGraphSchema::BreakPinLinks(UEdGraphPin& TargetPin, bool bSendsNodeNotification) const
diff --git a/Engine/Plugins/PCG/Source/PCGEditor/Private/Widgets/SPCGEditorGraphDeterminism.cpp b/Engine/Plugins/PCG/Source/PCGEditor/Private/Widgets/SPCGEditorGraphDeterminism.cpp
--- a/Engine/Plugins/PCG/Source/PCGEditor/Private/Widgets/SPCGEditorGraphDeterminism.cpp
+++ b/Engine/Plugins/PCG/Source/PCGEditor/Private/Widgets/SPCGEditorGraphDeterminism.cpp
@@ -73,9 +73,7 @@ TSharedRef<SWidget> SPCGEditorGraphDeterminismRow::GenerateWidgetForColumn(const
 	}
 	else if (ColumnId == NAME_DataTypesTested)
 	{
-		FString DataTypesTestedString = *UEnum::GetValueAsString(CurrentItem->DataTypesTested);
-		DataTypesTestedString.RemoveFromStart("EPCGDataType::");
-		CellText = FText::FromString(DataTypesTestedString);
+		CellText = FText::FromString(CurrentItem->DataTypesTested.ToString());
 	}
 	else if (ColumnId == NAME_AdditionalDetails)
 	{
@@ -240,7 +238,7 @@ void SPCGEditorGraphDeterminismListView::OnSortColumnHeader(const EColumnSortPri
 			}
 			else if (SortingColumn == NAME_DataTypesTested)
 			{
-				bIsLess = A->DataTypesTested < B->DataTypesTested;
+				bIsLess = A->DataTypesTested.ToString() < B->DataTypesTested.ToString();
 			}
 			else if (SortingColumn == NAME_AdditionalDetails)
 			{
