[PCG] Allow to pass the settings of the inspected node to the viewport for extra information

* Deprecate the previous version that takes only the data
* Update the deprecation everywhere.

#rnx

diff --git a/Engine/Plugins/PCG/Source/PCG/Public/PCGDataVisualization.h b/Engine/Plugins/PCG/Source/PCG/Public/PCGDataVisualization.h
--- a/Engine/Plugins/PCG/Source/PCG/Public/PCGDataVisualization.h
+++ b/Engine/Plugins/PCG/Source/PCG/Public/PCGDataVisualization.h
@@ -126,7 +126,16 @@ PRAGMA_ENABLE_DEPRECATION_WARNINGS
 	/** Initiates an async load on any resources this data needs in order to be visualized. */
 	virtual TArray<TSharedPtr<FStreamableHandle>> LoadRequiredResources(const UPCGData* Data) const { return {}; }
 
-	/** Optionally provide a function to setup the data viewport. */
+	UE_DEPRECATED(5.7, "Use the version taking the settings and data.")
 	virtual FPCGSetupSceneFunc GetViewportSetupFunc(const UPCGData* Data) const { return FPCGSetupSceneFunc(); }
+
+	/** Optionally provide a function to setup the data viewport. */
+	virtual FPCGSetupSceneFunc GetViewportSetupFunc(const UPCGSettingsInterface* SettingsInterface, const UPCGData* Data) const
+	{
+PRAGMA_DISABLE_DEPRECATION_WARNINGS
+		// For deprecation, call the old virtual by default.
+		return GetViewportSetupFunc(Data);
+PRAGMA_ENABLE_DEPRECATION_WARNINGS
+	}
 };
 #endif
diff --git a/Engine/Plugins/PCG/Source/PCGEditor/Private/DataVisualizations/PCGBaseTextureDataVisualization.cpp b/Engine/Plugins/PCG/Source/PCGEditor/Private/DataVisualizations/PCGBaseTextureDataVisualization.cpp
--- a/Engine/Plugins/PCG/Source/PCGEditor/Private/DataVisualizations/PCGBaseTextureDataVisualization.cpp
+++ b/Engine/Plugins/PCG/Source/PCGEditor/Private/DataVisualizations/PCGBaseTextureDataVisualization.cpp
@@ -28,7 +28,7 @@ TArray<TSharedPtr<FStreamableHandle>> FPCGBaseTextureDataVisualization::LoadRequ
 	return LoadHandles;
 }
 
-FPCGSetupSceneFunc FPCGBaseTextureDataVisualization::GetViewportSetupFunc(const UPCGData* Data) const
+FPCGSetupSceneFunc FPCGBaseTextureDataVisualization::GetViewportSetupFunc(const UPCGSettingsInterface* SettingsInterface, const UPCGData* Data) const
 {
 	return [this, WeakData=TWeakObjectPtr<const UPCGBaseTextureData>(Cast<UPCGBaseTextureData>(Data))](FPCGSceneSetupParams& InOutParams)
 	{
diff --git a/Engine/Plugins/PCG/Source/PCGEditor/Private/DataVisualizations/PCGPolygon2DDataVisualization.cpp b/Engine/Plugins/PCG/Source/PCGEditor/Private/DataVisualizations/PCGPolygon2DDataVisualization.cpp
--- a/Engine/Plugins/PCG/Source/PCGEditor/Private/DataVisualizations/PCGPolygon2DDataVisualization.cpp
+++ b/Engine/Plugins/PCG/Source/PCGEditor/Private/DataVisualizations/PCGPolygon2DDataVisualization.cpp
@@ -125,7 +125,7 @@ FPCGTableVisualizerInfo FPCGPolygon2DDataVisualization::GetTableVisualizerInfoWi
 	return Info;
 }
 
-FPCGSetupSceneFunc FPCGPolygon2DDataVisualization::GetViewportSetupFunc(const UPCGData* Data) const
+FPCGSetupSceneFunc FPCGPolygon2DDataVisualization::GetViewportSetupFunc(const UPCGSettingsInterface* SettingsInterface, const UPCGData* Data) const
 {
 	return[this, Data](FPCGSceneSetupParams& InOutParams)
 	{
diff --git a/Engine/Plugins/PCG/Source/PCGEditor/Private/DataVisualizations/PCGSpatialDataVisualization.cpp b/Engine/Plugins/PCG/Source/PCGEditor/Private/DataVisualizations/PCGSpatialDataVisualization.cpp
--- a/Engine/Plugins/PCG/Source/PCGEditor/Private/DataVisualizations/PCGSpatialDataVisualization.cpp
+++ b/Engine/Plugins/PCG/Source/PCGEditor/Private/DataVisualizations/PCGSpatialDataVisualization.cpp
@@ -422,7 +422,7 @@ TArray<FPCGMetadataDomainID> IPCGSpatialDataVisualization::GetAllSupportedDomain
 	return GetDefault<UPCGBasePointData>()->GetAllSupportedMetadataDomainIDs();
 }
 
-FPCGSetupSceneFunc IPCGSpatialDataVisualization::GetViewportSetupFunc(const UPCGData* Data) const
+FPCGSetupSceneFunc IPCGSpatialDataVisualization::GetViewportSetupFunc(const UPCGSettingsInterface* SettingsInterface, const UPCGData* Data) const
 {
 	return [this, WeakData=TWeakObjectPtr<const UPCGData>(Data)](FPCGSceneSetupParams& InOutParams)
 	{
diff --git a/Engine/Plugins/PCG/Source/PCGEditor/Private/DataVisualizations/PCGSplineDataVisualization.cpp b/Engine/Plugins/PCG/Source/PCGEditor/Private/DataVisualizations/PCGSplineDataVisualization.cpp
--- a/Engine/Plugins/PCG/Source/PCGEditor/Private/DataVisualizations/PCGSplineDataVisualization.cpp
+++ b/Engine/Plugins/PCG/Source/PCGEditor/Private/DataVisualizations/PCGSplineDataVisualization.cpp
@@ -140,7 +140,7 @@ const UPCGBasePointData* IPCGSplineDataVisualization::CollapseToDebugBasePointDa
 	return nullptr;
 }
 
-FPCGSetupSceneFunc IPCGSplineDataVisualization::GetViewportSetupFunc(const UPCGData* Data) const
+FPCGSetupSceneFunc IPCGSplineDataVisualization::GetViewportSetupFunc(const UPCGSettingsInterface* SettingsInterface, const UPCGData* Data) const
 {
 	return[this, Data](FPCGSceneSetupParams& InOutParams)
 	{
diff --git a/Engine/Plugins/PCG/Source/PCGEditor/Private/DataVisualizations/PCGStaticMeshDataVisualization.cpp b/Engine/Plugins/PCG/Source/PCGEditor/Private/DataVisualizations/PCGStaticMeshDataVisualization.cpp
--- a/Engine/Plugins/PCG/Source/PCGEditor/Private/DataVisualizations/PCGStaticMeshDataVisualization.cpp
+++ b/Engine/Plugins/PCG/Source/PCGEditor/Private/DataVisualizations/PCGStaticMeshDataVisualization.cpp
@@ -30,7 +30,7 @@ TArray<TSharedPtr<FStreamableHandle>> IPCGStaticMeshDataVisualization::LoadRequi
 	return LoadHandles;
 }
 
-FPCGSetupSceneFunc IPCGStaticMeshDataVisualization::GetViewportSetupFunc(const UPCGData* Data) const
+FPCGSetupSceneFunc IPCGStaticMeshDataVisualization::GetViewportSetupFunc(const UPCGSettingsInterface* SettingsInterface, const UPCGData* Data) const
 {
 	return [this](FPCGSceneSetupParams& InOutParams)
 	{
diff --git a/Engine/Plugins/PCG/Source/PCGEditor/Private/DataVisualizations/PCGVolumeDataVisualization.cpp b/Engine/Plugins/PCG/Source/PCGEditor/Private/DataVisualizations/PCGVolumeDataVisualization.cpp
--- a/Engine/Plugins/PCG/Source/PCGEditor/Private/DataVisualizations/PCGVolumeDataVisualization.cpp
+++ b/Engine/Plugins/PCG/Source/PCGEditor/Private/DataVisualizations/PCGVolumeDataVisualization.cpp
@@ -53,7 +53,7 @@ const UPCGBasePointData* FPCGPrimitiveDataVisualization::CollapseToDebugBasePoin
 	}
 }
 
-FPCGSetupSceneFunc FPCGPrimitiveDataVisualization::GetViewportSetupFunc(const UPCGData* Data) const
+FPCGSetupSceneFunc FPCGPrimitiveDataVisualization::GetViewportSetupFunc(const UPCGSettingsInterface* SettingsInterface, const UPCGData* Data) const
 {
 	return [this, Data](FPCGSceneSetupParams& InOutParams)
 	{
diff --git a/Engine/Plugins/PCG/Source/PCGEditor/Private/Widgets/SPCGEditorGraphAttributeListView.cpp b/Engine/Plugins/PCG/Source/PCGEditor/Private/Widgets/SPCGEditorGraphAttributeListView.cpp
--- a/Engine/Plugins/PCG/Source/PCGEditor/Private/Widgets/SPCGEditorGraphAttributeListView.cpp
+++ b/Engine/Plugins/PCG/Source/PCGEditor/Private/Widgets/SPCGEditorGraphAttributeListView.cpp
@@ -1245,7 +1245,10 @@ void SPCGEditorGraphAttributeListView::RefreshViewport()
 		return;
 	}
 
-	ViewportWidget->SetupScene(LoadedResources, DataVisualization->GetViewportSetupFunc(Data));
+	const UPCGNode* PCGNode = PCGEditorGraphNode.IsValid() ? PCGEditorGraphNode->GetPCGNode() : nullptr;
+	const UPCGSettingsInterface* Settings = PCGNode ? PCGNode->GetSettingsInterface() : nullptr;
+
+	ViewportWidget->SetupScene(LoadedResources, DataVisualization->GetViewportSetupFunc(Settings, Data));
 }
 
 void SPCGEditorGraphAttributeListView::LaunchUpdateTask()
diff --git a/Engine/Plugins/PCG/Source/PCGEditor/Public/DataVisualizations/PCGBaseTextureDataVisualization.h b/Engine/Plugins/PCG/Source/PCGEditor/Public/DataVisualizations/PCGBaseTextureDataVisualization.h
--- a/Engine/Plugins/PCG/Source/PCGEditor/Public/DataVisualizations/PCGBaseTextureDataVisualization.h
+++ b/Engine/Plugins/PCG/Source/PCGEditor/Public/DataVisualizations/PCGBaseTextureDataVisualization.h
@@ -14,7 +14,7 @@ class FPCGBaseTextureDataVisualization : public IPCGSpatialDataVisualization
 public:
 	// ~Begin IPCGDataVisualization interface
 	UE_API virtual TArray<TSharedPtr<FStreamableHandle>> LoadRequiredResources(const UPCGData* Data) const override;
-	UE_API virtual FPCGSetupSceneFunc GetViewportSetupFunc(const UPCGData* Data) const override;
+	UE_API virtual FPCGSetupSceneFunc GetViewportSetupFunc(const UPCGSettingsInterface* SettingsInterface, const UPCGData* Data) const override;
 	// ~End IPCGDataVisualization interface
 };
 
diff --git a/Engine/Plugins/PCG/Source/PCGEditor/Public/DataVisualizations/PCGPolygon2DDataVisualization.h b/Engine/Plugins/PCG/Source/PCGEditor/Public/DataVisualizations/PCGPolygon2DDataVisualization.h
--- a/Engine/Plugins/PCG/Source/PCGEditor/Public/DataVisualizations/PCGPolygon2DDataVisualization.h
+++ b/Engine/Plugins/PCG/Source/PCGEditor/Public/DataVisualizations/PCGPolygon2DDataVisualization.h
@@ -12,6 +12,6 @@ class FPCGPolygon2DDataVisualization : public IPCGSpatialDataVisualization
 public:
 	// ~Begin IPCGDataVisualization interface
 	PCGEDITOR_API virtual FPCGTableVisualizerInfo GetTableVisualizerInfoWithDomain(const UPCGData* Data, const FPCGMetadataDomainID& DomainID) const override;
-	PCGEDITOR_API virtual FPCGSetupSceneFunc GetViewportSetupFunc(const UPCGData* Data) const override;
+	PCGEDITOR_API virtual FPCGSetupSceneFunc GetViewportSetupFunc(const UPCGSettingsInterface* SettingsInterface, const UPCGData* Data) const override;
 	// ~End IPCGDataVisualization interface
 };
\ No newline at end of file
diff --git a/Engine/Plugins/PCG/Source/PCGEditor/Public/DataVisualizations/PCGSpatialDataVisualization.h b/Engine/Plugins/PCG/Source/PCGEditor/Public/DataVisualizations/PCGSpatialDataVisualization.h
--- a/Engine/Plugins/PCG/Source/PCGEditor/Public/DataVisualizations/PCGSpatialDataVisualization.h
+++ b/Engine/Plugins/PCG/Source/PCGEditor/Public/DataVisualizations/PCGSpatialDataVisualization.h
@@ -26,7 +26,7 @@ public:
 	virtual FPCGMetadataDomainID GetDefaultDomainForInspection(const UPCGData* Data) const override { return PCGMetadataDomainID::Elements; }
 	UE_API virtual FString GetDomainDisplayNameForInspection(const UPCGData* Data, const FPCGMetadataDomainID& DomainID) const override;
 	UE_API virtual TArray<FPCGMetadataDomainID> GetAllSupportedDomainsForInspection(const UPCGData* Data) const override;
-	UE_API virtual FPCGSetupSceneFunc GetViewportSetupFunc(const UPCGData* Data) const override;
+	UE_API virtual FPCGSetupSceneFunc GetViewportSetupFunc(const UPCGSettingsInterface* SettingsInterface, const UPCGData* Data) const override;
 	// ~End IPCGDataVisualization interface
 
 	UE_DEPRECATED(5.6, "Implement CollapseToDebugBasePointData instead")
diff --git a/Engine/Plugins/PCG/Source/PCGEditor/Public/DataVisualizations/PCGSplineDataVisualization.h b/Engine/Plugins/PCG/Source/PCGEditor/Public/DataVisualizations/PCGSplineDataVisualization.h
--- a/Engine/Plugins/PCG/Source/PCGEditor/Public/DataVisualizations/PCGSplineDataVisualization.h
+++ b/Engine/Plugins/PCG/Source/PCGEditor/Public/DataVisualizations/PCGSplineDataVisualization.h
@@ -15,7 +15,7 @@ class IPCGSplineDataVisualization : public IPCGSpatialDataVisualization
 public:
 	// ~Begin IPCGDataVisualization interface
 	UE_API virtual FPCGTableVisualizerInfo GetTableVisualizerInfoWithDomain(const UPCGData* Data, const FPCGMetadataDomainID& DomainID) const override;
-	UE_API virtual FPCGSetupSceneFunc GetViewportSetupFunc(const UPCGData* Data) const override;
+	UE_API virtual FPCGSetupSceneFunc GetViewportSetupFunc(const UPCGSettingsInterface* SettingsInterface, const UPCGData* Data) const override;
 	// ~End IPCGDataVisualization interface
 
 	// ~Begin IPCGSpatialDataVisualization interface
diff --git a/Engine/Plugins/PCG/Source/PCGEditor/Public/DataVisualizations/PCGStaticMeshDataVisualization.h b/Engine/Plugins/PCG/Source/PCGEditor/Public/DataVisualizations/PCGStaticMeshDataVisualization.h
--- a/Engine/Plugins/PCG/Source/PCGEditor/Public/DataVisualizations/PCGStaticMeshDataVisualization.h
+++ b/Engine/Plugins/PCG/Source/PCGEditor/Public/DataVisualizations/PCGStaticMeshDataVisualization.h
@@ -15,6 +15,6 @@ public:
 	virtual void ExecuteDebugDisplay(FPCGContext* Context, const UPCGSettingsInterface* SettingsInterface, const UPCGData* Data, AActor* TargetActor) const override { /* Do nothing. */ };
 	virtual FPCGTableVisualizerInfo GetTableVisualizerInfoWithDomain(const UPCGData* Data, const FPCGMetadataDomainID& DomainID) const override;
 	virtual TArray<TSharedPtr<FStreamableHandle>> LoadRequiredResources(const UPCGData* Data) const override;
-	virtual FPCGSetupSceneFunc GetViewportSetupFunc(const UPCGData* Data) const override;
+	virtual FPCGSetupSceneFunc GetViewportSetupFunc(const UPCGSettingsInterface* SettingsInterface, const UPCGData* Data) const override;
 	// ~End IPCGDataVisualization interface
 };
diff --git a/Engine/Plugins/PCG/Source/PCGEditor/Public/DataVisualizations/PCGVolumeDataVisualization.h b/Engine/Plugins/PCG/Source/PCGEditor/Public/DataVisualizations/PCGVolumeDataVisualization.h
--- a/Engine/Plugins/PCG/Source/PCGEditor/Public/DataVisualizations/PCGVolumeDataVisualization.h
+++ b/Engine/Plugins/PCG/Source/PCGEditor/Public/DataVisualizations/PCGVolumeDataVisualization.h
@@ -15,7 +15,7 @@ class FPCGPrimitiveDataVisualization : public IPCGSpatialDataVisualization
 public:
 	// ~Begin IPCGSpatialDataVisualization interface
 	UE_API virtual const UPCGBasePointData* CollapseToDebugBasePointData(FPCGContext* Context, const UPCGData* Data) const override;
-	UE_API FPCGSetupSceneFunc GetViewportSetupFunc(const UPCGData* Data) const override;
+	UE_API FPCGSetupSceneFunc GetViewportSetupFunc(const UPCGSettingsInterface* SettingsInterface, const UPCGData* Data) const override;
 	// ~End IPCGSpatialDataVisualization interface
 
 protected:
diff --git a/Engine/Plugins/PCGInterops/PCGGeometryScriptInterop/Source/PCGGeometryScriptInterop/Private/Editor/PCGDynamicMeshDataVisualization.cpp b/Engine/Plugins/PCGInterops/PCGGeometryScriptInterop/Source/PCGGeometryScriptInterop/Private/Editor/PCGDynamicMeshDataVisualization.cpp
--- a/Engine/Plugins/PCGInterops/PCGGeometryScriptInterop/Source/PCGGeometryScriptInterop/Private/Editor/PCGDynamicMeshDataVisualization.cpp
+++ b/Engine/Plugins/PCGInterops/PCGGeometryScriptInterop/Source/PCGGeometryScriptInterop/Private/Editor/PCGDynamicMeshDataVisualization.cpp
@@ -37,7 +37,7 @@ void FPCGDynamicMeshDataVisualization::ExecuteDebugDisplay(FPCGContext* Context,
 	}
 }
 
-FPCGSetupSceneFunc FPCGDynamicMeshDataVisualization::GetViewportSetupFunc(const UPCGData* Data) const
+FPCGSetupSceneFunc FPCGDynamicMeshDataVisualization::GetViewportSetupFunc(const UPCGSettingsInterface* SettingsInterface, const UPCGData* Data) const
 {
 	return [this, WeakData=TWeakObjectPtr<const UPCGDynamicMeshData>(Cast<UPCGDynamicMeshData>(Data))](FPCGSceneSetupParams& InOutParams)
 	{
diff --git a/Engine/Plugins/PCGInterops/PCGGeometryScriptInterop/Source/PCGGeometryScriptInterop/Private/Editor/PCGDynamicMeshDataVisualization.h b/Engine/Plugins/PCGInterops/PCGGeometryScriptInterop/Source/PCGGeometryScriptInterop/Private/Editor/PCGDynamicMeshDataVisualization.h
--- a/Engine/Plugins/PCGInterops/PCGGeometryScriptInterop/Source/PCGGeometryScriptInterop/Private/Editor/PCGDynamicMeshDataVisualization.h
+++ b/Engine/Plugins/PCGInterops/PCGGeometryScriptInterop/Source/PCGGeometryScriptInterop/Private/Editor/PCGDynamicMeshDataVisualization.h
@@ -15,7 +15,7 @@ class FPCGDynamicMeshDataVisualization : public IPCGSpatialDataVisualization
 public:
 	// ~Begin IPCGDataVisualization interface
 	virtual void ExecuteDebugDisplay(FPCGContext* Context, const UPCGSettingsInterface* SettingsInterface, const UPCGData* Data, AActor* TargetActor) const override;
-	virtual FPCGSetupSceneFunc GetViewportSetupFunc(const UPCGData* Data) const override;
+	virtual FPCGSetupSceneFunc GetViewportSetupFunc(const UPCGSettingsInterface* SettingsInterface, const UPCGData* Data) const override;
 	// ~End IPCGDataVisualization interface
 };
 
