[PCG] Fixed bug causing Inline Constant not to respect required pin

diff --git a/Engine/Plugins/PCG/Source/PCG/Private/Elements/Metadata/PCGMetadataOpElementBase.cpp b/Engine/Plugins/PCG/Source/PCG/Private/Elements/Metadata/PCGMetadataOpElementBase.cpp
--- a/Engine/Plugins/PCG/Source/PCG/Private/Elements/Metadata/PCGMetadataOpElementBase.cpp
+++ b/Engine/Plugins/PCG/Source/PCG/Private/Elements/Metadata/PCGMetadataOpElementBase.cpp
@@ -415,8 +415,8 @@ EPCGMetadataTypes UPCGMetadataSettingsBase::GetPinInitialDefaultValueType(const
 
 bool UPCGMetadataSettingsBase::IsInputPinRequiredByExecution(const UPCGPin* InPin) const
 {
-	// If the default value is not activated, or the pin is connected, it becomes required.
-	return InPin && (!IsPinDefaultValueActivated(InPin->Properties.Label) || InPin->IsConnected());
+	// The pin maintains 'required' status unless the default value is both enabled and activated.
+	return InPin && (InPin->IsConnected() || !IsPinDefaultValueEnabled(InPin->Properties.Label) || !IsPinDefaultValueActivated(InPin->Properties.Label));
 }
 
 void UPCGMetadataSettingsBase::AddDefaultValuesToCrc(FArchiveCrc32& Crc32) const
@@ -434,12 +434,7 @@ TArray<FPCGPinProperties> UPCGMetadataSettingsBase::InputPinProperties() const
 		if (PinLabel != NAME_None)
 		{
 			FPCGPinProperties& PinProperty = PinProperties.Emplace_GetRef(PinLabel, EPCGDataType::Any);
-
-			const bool bSupportDefaultValue = DefaultValuesAreEnabled() && IsPinDefaultValueEnabled(PinLabel);
-			if (!bSupportDefaultValue)
-			{
-				PinProperty.SetRequiredPin();
-			}
+			PinProperty.SetRequiredPin();
 
 #if WITH_EDITOR
 			TArray<FText> AllTooltips;
@@ -464,7 +459,7 @@ TArray<FPCGPinProperties> UPCGMetadataSettingsBase::InputPinProperties() const
 				AllTooltips.Add(LOCTEXT("PinTooltipForwardInput", "This input will be forwarded to the output."));
 			}
 
-			if (bSupportDefaultValue)
+			if (DefaultValuesAreEnabled() && IsPinDefaultValueEnabled(PinLabel))
 			{
 				EPCGMetadataTypes TypeValue = GetPinDefaultValueType(PinLabel);
 				const UEnum* TypeEnum = StaticEnum<EPCGMetadataTypes>();
diff --git a/Engine/Plugins/PCG/Source/PCGEditor/Private/Nodes/SPCGEditorGraphNode.cpp b/Engine/Plugins/PCG/Source/PCGEditor/Private/Nodes/SPCGEditorGraphNode.cpp
--- a/Engine/Plugins/PCG/Source/PCGEditor/Private/Nodes/SPCGEditorGraphNode.cpp
+++ b/Engine/Plugins/PCG/Source/PCGEditor/Private/Nodes/SPCGEditorGraphNode.cpp
@@ -116,7 +116,7 @@ TSharedPtr<SGraphPin> SPCGEditorGraphNode::CreatePinWidget(UEdGraphPin* InPin) c
 	if (InPin && NodeSettings && NodeSettings->Implements<UPCGSettingsDefaultValueProvider>())
 	{
 		const IPCGSettingsDefaultValueProvider* DefaultValueInterface = CastChecked<IPCGSettingsDefaultValueProvider>(NodeSettings);
-		if (DefaultValueInterface->DefaultValuesAreEnabled() && DefaultValueInterface->IsPinDefaultValueActivated(InPin->PinName))
+		if (!InPin->HasAnyConnections() && DefaultValueInterface->DefaultValuesAreEnabled() && DefaultValueInterface->IsPinDefaultValueActivated(InPin->PinName))
 		{
 			// Set the string default value to match the settings' source of truth.
 			InPin->DefaultValue = DefaultValueInterface->GetPinDefaultValueAsString(InPin->PinName);
diff --git a/Engine/Plugins/PCGInterops/PCGPythonInterop/Source/PCGPythonInteropEditor/Private/Elements/PCGExecutePythonScript.cpp b/Engine/Plugins/PCGInterops/PCGPythonInterop/Source/PCGPythonInteropEditor/Private/Elements/PCGExecutePythonScript.cpp
--- a/Engine/Plugins/PCGInterops/PCGPythonInterop/Source/PCGPythonInteropEditor/Private/Elements/PCGExecutePythonScript.cpp
+++ b/Engine/Plugins/PCGInterops/PCGPythonInterop/Source/PCGPythonInteropEditor/Private/Elements/PCGExecutePythonScript.cpp
@@ -167,8 +167,8 @@ bool UPCGExecutePythonScriptSettings::CreateInitialDefaultValueAttribute(const F
 
 bool UPCGExecutePythonScriptSettings::IsInputPinRequiredByExecution(const UPCGPin* InPin) const
 {
-	// If the default value is not activated, or the pin is connected, it becomes required.
-	return InPin && (!IsPinDefaultValueActivated(InPin->Properties.Label) || InPin->IsConnected());
+	// The pin maintains 'required' status unless the default value is both enabled and activated.
+	return InPin && (InPin->IsConnected() || !IsPinDefaultValueEnabled(InPin->Properties.Label) || !IsPinDefaultValueActivated(InPin->Properties.Label));
 }
 
 void UPCGExecutePythonScriptSettings::ResetDefaultValues()
