[PCG] Fix Builder stalls by calling UpdateResources(/*bSync=*/!ComputeFramework::IsDeferredCompilation())

diff --git a/Engine/Plugins/PCG/Source/PCG/Private/Compute/Elements/PCGComputeGraphElement.cpp b/Engine/Plugins/PCG/Source/PCG/Private/Compute/Elements/PCGComputeGraphElement.cpp
--- a/Engine/Plugins/PCG/Source/PCG/Private/Compute/Elements/PCGComputeGraphElement.cpp
+++ b/Engine/Plugins/PCG/Source/PCG/Private/Compute/Elements/PCGComputeGraphElement.cpp
@@ -457,7 +457,7 @@ bool FPCGComputeGraphElement::ExecuteInternal(FPCGContext* InContext) const
 		{
 			TRACE_CPUPROFILER_EVENT_SCOPE(Context->ComputeGraph->UpdateResources);
 
-			Context->ComputeGraph->UpdateResources();
+			Context->ComputeGraph->UpdateResources(/*bSync=*/!ComputeFramework::IsDeferredCompilation());
 
 			SleepUntilNextFrame();
 			return false;
diff --git a/Engine/Plugins/PCG/Source/PCG/Private/Graph/PCGGraphCompilerGPU.cpp b/Engine/Plugins/PCG/Source/PCG/Private/Graph/PCGGraphCompilerGPU.cpp
--- a/Engine/Plugins/PCG/Source/PCG/Private/Graph/PCGGraphCompilerGPU.cpp
+++ b/Engine/Plugins/PCG/Source/PCG/Private/Graph/PCGGraphCompilerGPU.cpp
@@ -19,6 +19,7 @@
 #include "Graph/PCGGraphCompiler.h"
 #include "Graph/PCGGraphExecutor.h"
 
+#include "ComputeFramework/ComputeFramework.h"
 #include "ComputeFramework/ComputeKernel.h"
 #include "HAL/IConsoleManager.h"
 #include "Misc/App.h"
@@ -884,7 +885,7 @@ void FPCGGraphCompilerGPU::BuildComputeGraphTask(
 		TRACE_CPUPROFILER_EVENT_SCOPE(UpdateResources);
 
 		// Compile shader resources and create render proxies.
-		ComputeGraph->UpdateResources();
+		ComputeGraph->UpdateResources(/*bSync=*/!ComputeFramework::IsDeferredCompilation());
 	}
 }
 
diff --git a/Engine/Plugins/Runtime/ComputeFramework/Source/ComputeFramework/Public/ComputeFramework/ComputeFramework.h b/Engine/Plugins/Runtime/ComputeFramework/Source/ComputeFramework/Public/ComputeFramework/ComputeFramework.h
--- a/Engine/Plugins/Runtime/ComputeFramework/Source/ComputeFramework/Public/ComputeFramework/ComputeFramework.h
+++ b/Engine/Plugins/Runtime/ComputeFramework/Source/ComputeFramework/Public/ComputeFramework/ComputeFramework.h
@@ -18,7 +18,7 @@ namespace ComputeFramework
 	COMPUTEFRAMEWORK_API bool IsEnabled();
 
 	/** Returns true if compute graphs are compiled on first use instead of on PostLoad(). */
-	bool IsDeferredCompilation();
+	COMPUTEFRAMEWORK_API bool IsDeferredCompilation();
 
 	/** Rebuild all compute graphs. */
 	COMPUTEFRAMEWORK_API void RebuildComputeGraphs();
