[PCG] PCG Level to Asset: early out on loaded sub-levels with log

diff --git a/Engine/Plugins/PCG/Source/PCGEditor/Private/PCGLevelToAsset.cpp b/Engine/Plugins/PCG/Source/PCGEditor/Private/PCGLevelToAsset.cpp
--- a/Engine/Plugins/PCG/Source/PCGEditor/Private/PCGLevelToAsset.cpp
+++ b/Engine/Plugins/PCG/Source/PCGEditor/Private/PCGLevelToAsset.cpp
@@ -70,6 +70,12 @@ UPackage* UPCGLevelToAsset::CreateOrUpdatePCGAsset(UWorld* World, const FPCGAsse
 		return nullptr;
 	}
 
+	if (ULevelStreaming::FindStreamingLevel(World->PersistentLevel) != nullptr)
+	{
+		UE_LOG(LogPCGEditor, Error, TEXT("Unable to create asset from sub-level. If level is currently loaded as a sub-level, unload it first and try exporting again."));
+		return nullptr;
+	}
+
 	UPCGLevelToAsset* Exporter = nullptr;
 	if (ExporterSubclass)
 	{
