[PCG] Add support for skeletal meshes in BoundsFromMesh and PointFromMesh

* Deprecate StaticMesh property to Mesh property in PointFromMesh, with override alias

#rnx

diff --git a/Engine/Plugins/PCG/Source/PCG/Private/Elements/PCGBoundsFromMesh.cpp b/Engine/Plugins/PCG/Source/PCG/Private/Elements/PCGBoundsFromMesh.cpp
--- a/Engine/Plugins/PCG/Source/PCG/Private/Elements/PCGBoundsFromMesh.cpp
+++ b/Engine/Plugins/PCG/Source/PCG/Private/Elements/PCGBoundsFromMesh.cpp
@@ -15,6 +15,7 @@
 #endif // WITH_EDITOR
 
 #include "Algo/Transform.h"
+#include "Engine/SkeletalMesh.h"
 #include "Engine/StaticMesh.h"
 
 #include UE_INLINE_GENERATED_CPP_BY_NAME(PCGBoundsFromMesh)
@@ -40,7 +41,7 @@ FText UPCGBoundsFromMeshSettings::GetDefaultNodeTitle() const
 
 FText UPCGBoundsFromMeshSettings::GetNodeTooltipText() const
 {
-	return LOCTEXT("NodeTooltip", "Sets bounds min/max from mesh(es).");
+	return LOCTEXT("NodeTooltip", "Sets the bounds according to the static or skeletal mesh(es) provided in the mesh pin.");
 }
 
 bool UPCGBoundsFromMeshSettings::IsPinUsedByNodeExecution(const UPCGPin* InPin) const
@@ -224,6 +225,24 @@ bool FPCGBoundsFromMeshElement::ExecuteInternal(FPCGContext* InContext) const
 	TArray<FPCGTaggedData> Inputs = Context->InputData.GetInputsByPin(PCGPinConstants::DefaultInputLabel);
 	TArray<FPCGTaggedData>& Outputs = Context->OutputData.TaggedData;
 
+	auto BoundsFromMesh = [](const UObject* MeshObject, FBox& OutBounds) -> bool
+	{
+		if (const UStaticMesh* StaticMesh = Cast<UStaticMesh>(MeshObject))
+		{
+			OutBounds = StaticMesh->GetBoundingBox();
+			return true;
+		}
+		else if (const USkeletalMesh* SkeletalMesh = Cast<USkeletalMesh>(MeshObject))
+		{
+			OutBounds = SkeletalMesh->GetBounds().GetBox();
+			return true;
+		}
+		else
+		{
+			return false;
+		}
+	};
+
 	if (!Context->bBoundsQueried)
 	{
 #if WITH_EDITOR
@@ -235,12 +254,13 @@ bool FPCGBoundsFromMeshElement::ExecuteInternal(FPCGContext* InContext) const
 		for (int MeshIndex = 0; MeshIndex < Context->MeshesToLoad.Num(); ++MeshIndex)
 		{
 			const FSoftObjectPath& MeshPath = Context->MeshesToLoad[MeshIndex];
-			TSoftObjectPtr<UStaticMesh> Mesh(MeshPath);
+			const UObject* MeshObject = MeshPath.ResolveObject();
 
-			if (Mesh.IsValid() && Mesh.Get())
-			{
-				Context->MeshToBoundsMap.Add(MeshIndex, Mesh->GetBoundingBox());
+			FBox Bounds{};
 
+			if (BoundsFromMesh(MeshObject, Bounds))
+			{
+				Context->MeshToBoundsMap.Add(MeshIndex, MoveTemp(Bounds));
 #if WITH_EDITOR
 				DynamicTracking.AddToTracking(FPCGSelectionKey::CreateFromPath(MeshPath), /*bIsCulled=*/ false);
 #endif // WITH_EDITOR
@@ -272,15 +292,14 @@ bool FPCGBoundsFromMeshElement::ExecuteInternal(FPCGContext* InContext) const
 		// or we need to compare against our list of meshes
 		if (Context->SingleMesh.IsValid())
 		{
-			TSoftObjectPtr<UStaticMesh> SingleMesh(Context->SingleMesh);
+			const UObject* MeshObject = Context->SingleMesh.ResolveObject();
+			FBox MeshBounds{};
 
-			if (!SingleMesh)
+			if (!BoundsFromMesh(MeshObject, MeshBounds))
 			{
-				PCGE_LOG(Error, GraphAndLog, LOCTEXT("LoadStaticMeshFailed", "Failed to load StaticMesh"));
+				PCGE_LOG(Error, GraphAndLog, LOCTEXT("LoadStaticMeshFailed", "Failed to load StaticMesh or SkeletalMesh"));
 				continue;
 			}
-
-			const FBox StaticMeshBounds = SingleMesh->GetBoundingBox();
 			
 			UPCGBasePointData* OutputData = FPCGContext::NewPointData_AnyThread(Context);
 
@@ -290,15 +309,15 @@ bool FPCGBoundsFromMeshElement::ExecuteInternal(FPCGContext* InContext) const
 
 			if (OutputData->HasSpatialDataParent())
 			{
-				OutputData->SetBoundsMin(StaticMeshBounds.Min);
-				OutputData->SetBoundsMax(StaticMeshBounds.Max);
+				OutputData->SetBoundsMin(MeshBounds.Min);
+				OutputData->SetBoundsMax(MeshBounds.Max);
 			}
 			else
 			{
 				OutputData->AllocateProperties(PointData->GetAllocatedProperties() | EPCGPointNativeProperties::BoundsMin | EPCGPointNativeProperties::BoundsMax);
 				OutputData->CopyUnallocatedPropertiesFrom(PointData);
 
-				auto CopyAndSetBounds = [OutputData, PointData, &StaticMeshBounds](const int32 StartReadIndex, const int32 StartWriteIndex, const int32 Count)
+				auto CopyAndSetBounds = [OutputData, PointData, &MeshBounds](const int32 StartReadIndex, const int32 StartWriteIndex, const int32 Count)
 				{
 					int32 NumWritten = 0;
 
@@ -310,8 +329,8 @@ bool FPCGBoundsFromMeshElement::ExecuteInternal(FPCGContext* InContext) const
 						const int32 WriteIndex = StartWriteIndex + NumWritten;
 						OutRanges.SetFromValueRanges(WriteIndex, InRanges, ReadIndex);
 
-						OutRanges.BoundsMinRange[WriteIndex] = StaticMeshBounds.Min;
-						OutRanges.BoundsMaxRange[WriteIndex] = StaticMeshBounds.Max;
+						OutRanges.BoundsMinRange[WriteIndex] = MeshBounds.Min;
+						OutRanges.BoundsMaxRange[WriteIndex] = MeshBounds.Max;
 						++NumWritten;
 					}
 
diff --git a/Engine/Plugins/PCG/Source/PCG/Private/Elements/PCGPointFromMeshElement.cpp b/Engine/Plugins/PCG/Source/PCG/Private/Elements/PCGPointFromMeshElement.cpp
--- a/Engine/Plugins/PCG/Source/PCG/Private/Elements/PCGPointFromMeshElement.cpp
+++ b/Engine/Plugins/PCG/Source/PCG/Private/Elements/PCGPointFromMeshElement.cpp
@@ -11,12 +11,26 @@
 #include "Helpers/PCGDynamicTrackingHelpers.h"
 #endif // WITH_EDITOR
 
+#include "Engine/SkeletalMesh.h"
 #include "Engine/StaticMesh.h"
 
 #include UE_INLINE_GENERATED_CPP_BY_NAME(PCGPointFromMeshElement)
 
 #define LOCTEXT_NAMESPACE "PCGPointFromMeshElement"
 
+void UPCGPointFromMeshSettings::PostLoad()
+{
+	Super::PostLoad();
+
+#if WITH_EDITOR
+	if (!StaticMesh.IsNull())
+	{
+		Mesh = StaticMesh;
+		StaticMesh.Reset();
+	}
+#endif // WITH_EDITOR
+}
+
 #if WITH_EDITOR
 FText UPCGPointFromMeshSettings::GetDefaultNodeTitle() const
 {
@@ -25,17 +39,17 @@ FText UPCGPointFromMeshSettings::GetDefaultNodeTitle() const
 
 FText UPCGPointFromMeshSettings::GetNodeTooltipText() const
 {
-	return LOCTEXT("PointFromMeshNodeTooltip", "Creates a single point at the origin with an attribute named MeshPathAttributeName containing a SoftObjectPath to the StaticMesh.");
+	return LOCTEXT("PointFromMeshNodeTooltip", "Creates a single point at the origin with an attribute named MeshPathAttributeName containing a SoftObjectPath to the StaticMesh/SkeletalMesh.");
 }
 
 void UPCGPointFromMeshSettings::GetStaticTrackedKeys(FPCGSelectionKeyToSettingsMap& OutKeysToSettings, TArray<TObjectPtr<const UPCGGraph>>& OutVisitedGraphs) const
 {
-	if (StaticMesh.IsNull() || IsPropertyOverriddenByPin(GET_MEMBER_NAME_CHECKED(UPCGPointFromMeshSettings, StaticMesh)))
+	if (Mesh.IsNull() || IsPropertyOverriddenByPin(GET_MEMBER_NAME_CHECKED(UPCGPointFromMeshSettings, Mesh)))
 	{
 		return;
 	}
 
-	FPCGSelectionKey Key = FPCGSelectionKey::CreateFromPath(StaticMesh.ToSoftObjectPath());
+	FPCGSelectionKey Key = FPCGSelectionKey::CreateFromPath(Mesh.ToSoftObjectPath());
 
 	OutKeysToSettings.FindOrAdd(MoveTemp(Key)).Emplace(this, /*bCulling=*/false);
 }
@@ -55,7 +69,7 @@ bool FPCGPointFromMeshElement::PrepareDataInternal(FPCGContext* Context) const
 	const UPCGPointFromMeshSettings* Settings = Context->GetInputSettings<UPCGPointFromMeshSettings>();
 	check(Settings);
 
-	if (Settings->StaticMesh.IsNull())
+	if (Settings->Mesh.IsNull())
 	{
 		return true;
 	}
@@ -64,7 +78,7 @@ bool FPCGPointFromMeshElement::PrepareDataInternal(FPCGContext* Context) const
 
 	if (!ThisContext->WasLoadRequested())
 	{
-		return ThisContext->RequestResourceLoad(ThisContext, { Settings->StaticMesh.ToSoftObjectPath() }, !Settings->bSynchronousLoad);
+		return ThisContext->RequestResourceLoad(ThisContext, { Settings->Mesh.ToSoftObjectPath() }, !Settings->bSynchronousLoad);
 	}
 
 	return true;
@@ -79,21 +93,23 @@ bool FPCGPointFromMeshElement::ExecuteInternal(FPCGContext* Context) const
 	const UPCGPointFromMeshSettings* Settings = Context->GetInputSettings<UPCGPointFromMeshSettings>();
 	check(Settings);
 
-	if (Settings->StaticMesh.IsNull())
+	if (Settings->Mesh.IsNull())
 	{
 		return true;
 	}
 
-	if (!Settings->StaticMesh.Get())
+	const UObject* MeshObject = Settings->Mesh.Get();
+
+	if (!MeshObject || (!MeshObject->IsA<UStaticMesh>() && !MeshObject->IsA<USkeletalMesh>()))
 	{
-		PCGE_LOG(Error, GraphAndLog, LOCTEXT("LoadStaticMeshFailed", "Failed to load StaticMesh"));
+		PCGE_LOG(Error, GraphAndLog, LOCTEXT("LoadStaticMeshFailed", "Failed to load StaticMesh/SkeletalMesh"));
 		return true;
 	}
 
 #if WITH_EDITOR
-	if (Context->IsValueOverriden(GET_MEMBER_NAME_CHECKED(UPCGPointFromMeshSettings, StaticMesh)))
+	if (Context->IsValueOverriden(GET_MEMBER_NAME_CHECKED(UPCGPointFromMeshSettings, Mesh)))
 	{
-		FPCGDynamicTrackingHelper::AddSingleDynamicTrackingKey(Context, FPCGSelectionKey::CreateFromPath(Settings->StaticMesh.ToSoftObjectPath()), /*bIsCulled=*/false);
+		FPCGDynamicTrackingHelper::AddSingleDynamicTrackingKey(Context, FPCGSelectionKey::CreateFromPath(Settings->Mesh.ToSoftObjectPath()), /*bIsCulled=*/false);
 	}
 #endif // WITH_EDITOR
 
@@ -103,15 +119,26 @@ bool FPCGPointFromMeshElement::ExecuteInternal(FPCGContext* Context) const
 	
 	OutPointData->SetNumPoints(1);
 
-	// Capture StaticMesh bounds
-	const FBox StaticMeshBounds = Settings->StaticMesh->GetBoundingBox();
+	FBox MeshBounds{};
+	if (const UStaticMesh* StaticMesh = Cast<UStaticMesh>(MeshObject))
+	{
+		MeshBounds = StaticMesh->GetBoundingBox();
+	}
+	else if (const USkeletalMesh* SkeletalMesh = Cast<USkeletalMesh>(MeshObject))
+	{
+		MeshBounds = SkeletalMesh->GetBounds().GetBox();
+	}
+	else
+	{
+		checkNoEntry();
+	}
 
-	OutPointData->SetBoundsMin(StaticMeshBounds.Min);
-	OutPointData->SetBoundsMax(StaticMeshBounds.Max);
+	OutPointData->SetBoundsMin(MeshBounds.Min);
+	OutPointData->SetBoundsMax(MeshBounds.Max);
 
 	// Write StaticMesh path to MeshPathAttribute
 	check(OutPointData->Metadata);
-	OutPointData->Metadata->CreateSoftObjectPathAttribute(Settings->MeshPathAttributeName, Settings->StaticMesh.ToSoftObjectPath(), /*bAllowsInterpolation=*/false);
+	OutPointData->Metadata->CreateSoftObjectPathAttribute(Settings->MeshPathAttributeName, Settings->Mesh.ToSoftObjectPath(), /*bAllowsInterpolation=*/false);
 
 	return true;
 }
diff --git a/Engine/Plugins/PCG/Source/PCG/Public/Elements/PCGBoundsFromMesh.h b/Engine/Plugins/PCG/Source/PCG/Public/Elements/PCGBoundsFromMesh.h
--- a/Engine/Plugins/PCG/Source/PCG/Public/Elements/PCGBoundsFromMesh.h
+++ b/Engine/Plugins/PCG/Source/PCG/Public/Elements/PCGBoundsFromMesh.h
@@ -11,7 +11,7 @@
 class UStaticMesh;
 
 /**
-* Sets the bounds on the points according to the mesh(es) provided in the mesh pin.
+* Sets the bounds according to the static or skeletal mesh(es) provided in the mesh pin.
 */
 UCLASS(MinimalAPI, BlueprintType)
 class UPCGBoundsFromMeshSettings : public UPCGSettings
diff --git a/Engine/Plugins/PCG/Source/PCG/Public/Elements/PCGPointFromMeshElement.h b/Engine/Plugins/PCG/Source/PCG/Public/Elements/PCGPointFromMeshElement.h
--- a/Engine/Plugins/PCG/Source/PCG/Public/Elements/PCGPointFromMeshElement.h
+++ b/Engine/Plugins/PCG/Source/PCG/Public/Elements/PCGPointFromMeshElement.h
@@ -12,7 +12,7 @@
 
 class UStaticMesh;
 
-// PointFromMesh creates a single point at the origin with an attribute containing a SoftObjectPath to the selected UStaticMesh
+// PointFromMesh creates a single point at the origin with an attribute containing a SoftObjectPath to the selected UStaticMesh/USkeletalMesh
 UCLASS(MinimalAPI, BlueprintType, ClassGroup = (Procedural))
 class UPCGPointFromMeshSettings : public UPCGSettings
 {
@@ -20,6 +20,10 @@ class UPCGPointFromMeshSettings : public UPCGSettings
 
 public:
 	//~Begin UPCGSettings interface
+	void PostLoad();
+	//~End UPCGSettings interface
+	
+	//~Begin UPCGSettings interface
 #if WITH_EDITOR
 	virtual FName GetDefaultNodeName() const override { return FName(TEXT("PointFromMesh")); }
 	virtual FText GetDefaultNodeTitle() const override;
@@ -36,10 +40,15 @@ protected:
 	//~End UPCGSettings interface
 
 public:
-	UPROPERTY(BlueprintReadWrite, EditAnywhere, Category = Settings, meta = (PCG_Overridable))
+	UPROPERTY(BlueprintReadWrite, EditAnywhere, Category = Settings, meta = (PCG_Overridable, PCG_OverrideAliases="StaticMesh", AllowedClasses="/Script/Engine.StaticMesh, /Script/Engine.SkeletalMesh"))
+	TSoftObjectPtr<UObject> Mesh;
+
+#if WITH_EDITORONLY_DATA
+	UPROPERTY(meta = (DeprecatedProperty, DeprecationMessage = "Use Mesh property instead"))
 	TSoftObjectPtr<UStaticMesh> StaticMesh;
+#endif // WITH_EDITORONLY_DATA
 
-	/** Name of the string attribute to be created and hold a SoftObjectPath to the StaticMesh */
+	/** Name of the string attribute to be created and hold a SoftObjectPath to the Mesh */
 	UPROPERTY(BlueprintReadWrite, EditAnywhere, Category = Settings)
 	FName MeshPathAttributeName = NAME_None;
 
