[PCG] Change property to uint32 so that it serializes properly (currently never serialized as none-zero as we flatten the properties when storing graph outputs)

#rnx

diff --git a/Engine/Plugins/PCG/Source/PCG/Private/Data/PCGPointArrayData.cpp b/Engine/Plugins/PCG/Source/PCG/Private/Data/PCGPointArrayData.cpp
--- a/Engine/Plugins/PCG/Source/PCG/Private/Data/PCGPointArrayData.cpp
+++ b/Engine/Plugins/PCG/Source/PCG/Private/Data/PCGPointArrayData.cpp
@@ -121,7 +121,7 @@ void UPCGPointArrayData::InitializeSpatialDataInternal(const FPCGInitializeFromD
 		if (ParentData != InParams.Source)
 		{
 			SetNumPoints(SourceParentData->GetNumPoints());
-			InheritedProperties = EPCGPointNativeProperties::All;
+			InheritedProperties = static_cast<uint32>(EPCGPointNativeProperties::All);
 			ParentData = const_cast<UPCGPointArrayData*>(SourceParentData);
 		}
 	}
diff --git a/Engine/Plugins/PCG/Source/PCG/Public/Data/PCGPointArrayData.h b/Engine/Plugins/PCG/Source/PCG/Public/Data/PCGPointArrayData.h
--- a/Engine/Plugins/PCG/Source/PCG/Public/Data/PCGPointArrayData.h
+++ b/Engine/Plugins/PCG/Source/PCG/Public/Data/PCGPointArrayData.h
@@ -74,7 +74,7 @@ private:
 	{
 		if (bWithInheritance)
 		{
-			if (EnumHasAnyFlags(InheritedProperties, TPCGPointNativeProperty<Property>::EnumValue))
+			if (EnumHasAnyFlags(static_cast<EPCGPointNativeProperties>(InheritedProperties), TPCGPointNativeProperty<Property>::EnumValue))
 			{
 				check(ParentData);
 				return ParentData->GetProperty<Property>(bWithInheritance);
@@ -125,7 +125,7 @@ private:
 	{
 		if (bWithInheritance)
 		{
-			if (EnumHasAnyFlags(InheritedProperties, TPCGPointNativeProperty<Property>::EnumValue))
+			if (EnumHasAnyFlags(static_cast<EPCGPointNativeProperties>(InheritedProperties), TPCGPointNativeProperty<Property>::EnumValue))
 			{
 				check(ParentData);
 				return ParentData->GetProperty<Property>(bWithInheritance);
@@ -176,8 +176,9 @@ private:
 	{
 		using PropertyType = typename TPCGPointNativeProperty<NativeProperty>::Type;
 		EPCGPointNativeProperties EnumValue = TPCGPointNativeProperty<NativeProperty>::EnumValue;
+		EPCGPointNativeProperties InheritedPropertiesEnum = static_cast<EPCGPointNativeProperties>(InheritedProperties);
 
-		if (!EnumHasAnyFlags(InheritedProperties, EnumValue))
+		if (!EnumHasAnyFlags(InheritedPropertiesEnum, EnumValue))
 		{
 			return false;
 		}
@@ -197,9 +198,10 @@ private:
 			InheritedProperty->CopyTo(*Property, 0, 0, Property->Num());
 		}
 
-		EnumRemoveFlags(InheritedProperties, EnumValue);
+		EnumRemoveFlags(InheritedPropertiesEnum, EnumValue);
+		InheritedProperties = static_cast<uint32>(InheritedPropertiesEnum);
 
-		if (InheritedProperties == EPCGPointNativeProperties::None)
+		if (InheritedPropertiesEnum == EPCGPointNativeProperties::None)
 		{
 			ParentData = nullptr;
 		}
@@ -213,8 +215,8 @@ private:
 	UPROPERTY(VisibleAnywhere, Category = ParentData)
 	TObjectPtr<UPCGPointArrayData> ParentData = nullptr;
 
-	UPROPERTY(VisibleAnywhere, Category = ParentData)
-	EPCGPointNativeProperties InheritedProperties = EPCGPointNativeProperties::None;
+	UPROPERTY(VisibleAnywhere, Category = ParentData, meta = (BitMask, BitmaskEnum = "/Script/PCG.EPCGPointNativeProperties"))
+	uint32 InheritedProperties = 0;
 };
 
 #undef UE_API
diff --git a/Engine/Plugins/PCG/Source/PCG/Public/PCGPointPropertiesTraits.h b/Engine/Plugins/PCG/Source/PCG/Public/PCGPointPropertiesTraits.h
--- a/Engine/Plugins/PCG/Source/PCG/Public/PCGPointPropertiesTraits.h
+++ b/Engine/Plugins/PCG/Source/PCG/Public/PCGPointPropertiesTraits.h
@@ -7,7 +7,7 @@
 #include "PCGPointPropertiesTraits.generated.h"
 
 // Value names need to match EPCGPointProperties
-UENUM(meta = (Bitflags))
+UENUM(meta = (Bitflags, UseEnumValuesAsMaskValuesInEditor = "true"))
 enum class EPCGPointNativeProperties : uint32
 {
 	None = 0 UMETA(Hidden),
