[PCG] Component reuse: fix ISMC/ISKMC transform not being properly updated if root was rotated/scaled
- Compare the full transform and not just the location

diff --git a/Engine/Plugins/PCG/Source/PCG/Private/PCGManagedResource.cpp b/Engine/Plugins/PCG/Source/PCG/Private/PCGManagedResource.cpp
--- a/Engine/Plugins/PCG/Source/PCG/Private/PCGManagedResource.cpp
+++ b/Engine/Plugins/PCG/Source/PCG/Private/PCGManagedResource.cpp
@@ -831,12 +831,14 @@ void UPCGManagedISMComponent::MarkAsUsed()
 			RootLocation = FVector::ZeroVector;
 		}
 
-		if (bHasPreviousRootLocation != bHasRootLocation || (ISMC->GetComponentLocation() - RootLocation).SquaredLength() > UE_DOUBLE_SMALL_NUMBER)
+		// Reset the rotation/scale to be identity otherwise if the root component transform has changed, the final transform will be wrong
+		const FTransform NewComponentWorldTransform(FQuat::Identity, RootLocation, FVector::OneVector);
+
+		if (bHasPreviousRootLocation != bHasRootLocation || !ISMC->GetComponentTransform().Equals(NewComponentWorldTransform, UE_DOUBLE_SMALL_NUMBER))
 		{
-			// Reset the rotation/scale to be identity otherwise if the root component transform has changed, the final transform will be wrong
 			// Since this is technically 'moving' the ISM, we need to unregister it before moving otherwise we could get a warning that we're moving a component with static mobility
 			ISMC->UnregisterComponent();
-			ISMC->SetWorldTransform(FTransform(FQuat::Identity, RootLocation, FVector::OneVector));
+			ISMC->SetWorldTransform(NewComponentWorldTransform);
 			ISMC->RegisterComponent();
 		}
 	}
@@ -859,11 +861,14 @@ void UPCGManagedISMComponent::MarkAsReused()
 			}
 		}
 
-		if ((ISMC->GetComponentLocation() - TentativeRootLocation).SquaredLength() > UE_DOUBLE_SMALL_NUMBER)
+		// Reset the rotation/scale to be identity otherwise if the root component transform has changed, the final transform will be wrong
+		const FTransform NewComponentWorldTransform(FQuat::Identity, TentativeRootLocation, FVector::OneVector);
+
+		if (!ISMC->GetComponentTransform().Equals(NewComponentWorldTransform, UE_DOUBLE_SMALL_NUMBER))
 		{
 			// Since this is technically 'moving' the ISM, we need to unregister it before moving otherwise we could get a warning that we're moving a component with static mobility
 			ISMC->UnregisterComponent();
-			ISMC->SetWorldTransform(FTransform(FQuat::Identity, TentativeRootLocation, FVector::OneVector));
+			ISMC->SetWorldTransform(NewComponentWorldTransform);
 			ISMC->RegisterComponent();
 		}
 	}
@@ -964,12 +969,14 @@ void UPCGManagedISKMComponent::MarkAsUsed()
 			RootLocation = FVector::ZeroVector;
 		}
 
-		if (bHasPreviousRootLocation != bHasRootLocation || (ISKMC->GetComponentLocation() - RootLocation).SquaredLength() > UE_DOUBLE_SMALL_NUMBER)
+		// Reset the rotation/scale to be identity otherwise if the root component transform has changed, the final transform will be wrong
+		const FTransform NewComponentWorldTransform(FQuat::Identity, RootLocation, FVector::OneVector);
+
+		if (bHasPreviousRootLocation != bHasRootLocation || !ISKMC->GetComponentTransform().Equals(NewComponentWorldTransform, UE_DOUBLE_SMALL_NUMBER))
 		{
-			// Reset the rotation/scale to be identity otherwise if the root component transform has changed, the final transform will be wrong
 			// Since this is technically 'moving' the ABM, we need to unregister it before moving otherwise we could get a warning that we're moving a component with static mobility
 			ISKMC->UnregisterComponent();
-			ISKMC->SetWorldTransform(FTransform(FQuat::Identity, RootLocation, FVector::OneVector));
+			ISKMC->SetWorldTransform(NewComponentWorldTransform);
 			ISKMC->RegisterComponent();
 		}
 	}
@@ -992,11 +999,14 @@ void UPCGManagedISKMComponent::MarkAsReused()
 			}
 		}
 
-		if ((ISKMC->GetComponentLocation() - TentativeRootLocation).SquaredLength() > UE_DOUBLE_SMALL_NUMBER)
+		// Reset the rotation/scale to be identity otherwise if the root component transform has changed, the final transform will be wrong
+		const FTransform NewComponentWorldTransform(FQuat::Identity, TentativeRootLocation, FVector::OneVector);
+
+		if (!ISKMC->GetComponentTransform().Equals(NewComponentWorldTransform, UE_DOUBLE_SMALL_NUMBER))
 		{
 			// Since this is technically 'moving' the ABM, we need to unregister it before moving otherwise we could get a warning that we're moving a component with static mobility
 			ISKMC->UnregisterComponent();
-			ISKMC->SetWorldTransform(FTransform(FQuat::Identity, TentativeRootLocation, FVector::OneVector));
+			ISKMC->SetWorldTransform(NewComponentWorldTransform);
 			ISKMC->RegisterComponent();
 		}
 	}
