[PCG] Fix UPCGCustomHLSLKernel::PinAllocatedProperties failing to serialize during cook

#rnx

diff --git a/Engine/Plugins/PCG/Source/PCG/Private/Compute/Elements/PCGCustomHLSLKernel.cpp b/Engine/Plugins/PCG/Source/PCG/Private/Compute/Elements/PCGCustomHLSLKernel.cpp
--- a/Engine/Plugins/PCG/Source/PCG/Private/Compute/Elements/PCGCustomHLSLKernel.cpp
+++ b/Engine/Plugins/PCG/Source/PCG/Private/Compute/Elements/PCGCustomHLSLKernel.cpp
@@ -386,9 +386,9 @@ bool UPCGCustomHLSLKernel::ComputeOutputBindingDataDesc(const UPCGComputeGraph*
 	}
 
 	// Allocate any properties that were modified by the Custom HLSL source. If this data was initialized from another, it will have inherited its allocated properties already too.
-	if (const EPCGPointNativeProperties* AllocatedProperties = PinAllocatedProperties.Find(OutputPinProperties->Label))
+	if (const int32* AllocatedProperties = PinAllocatedProperties.Find(OutputPinProperties->Label))
 	{
-		OutputDataDesc.AllocatePropertiesForAllData(*AllocatedProperties);
+		OutputDataDesc.AllocatePropertiesForAllData(static_cast<EPCGPointNativeProperties>(*AllocatedProperties));
 	}
 
 	OutDataDesc = MoveTemp(OutputDataDesc);
@@ -970,38 +970,38 @@ void UPCGCustomHLSLKernel::ParseShaderSource()
 					
 					// If the kernel writes to an attribute, it must be fully allocated if it isn't already.
 					// @todo_pcg: Could have an API for setting the default/constant value in a constant value range as well, to avoid expanding to full-width allocation always.
-					EPCGPointNativeProperties& AllocatedProperties = PinAllocatedProperties.FindOrAdd(FName(PinStr));
+					int32& AllocatedProperties = PinAllocatedProperties.FindOrAdd(FName(PinStr));
 
 					if (PropertyStr == PCGCustomHLSLKernel::SetPositionKeyword
 						|| PropertyStr == PCGCustomHLSLKernel::SetRotationKeyword
 						|| PropertyStr == PCGCustomHLSLKernel::SetScaleKeyword
 						|| PropertyStr == PCGCustomHLSLKernel::SetPointTransformKeyword)
 					{
-						AllocatedProperties |= EPCGPointNativeProperties::Transform;
+						AllocatedProperties |= static_cast<int32>(EPCGPointNativeProperties::Transform);
 					}
 					else if (PropertyStr == PCGCustomHLSLKernel::SetBoundsMinKeyword)
 					{
-						AllocatedProperties |= EPCGPointNativeProperties::BoundsMin;
+						AllocatedProperties |= static_cast<int32>(EPCGPointNativeProperties::BoundsMin);
 					}
 					else if (PropertyStr == PCGCustomHLSLKernel::SetBoundsMaxKeyword)
 					{
-						AllocatedProperties |= EPCGPointNativeProperties::BoundsMax;
+						AllocatedProperties |= static_cast<int32>(EPCGPointNativeProperties::BoundsMax);
 					}
 					else if (PropertyStr == PCGCustomHLSLKernel::SetColorKeyword)
 					{
-						AllocatedProperties |= EPCGPointNativeProperties::Color;
+						AllocatedProperties |= static_cast<int32>(EPCGPointNativeProperties::Color);
 					}
 					else if (PropertyStr == PCGCustomHLSLKernel::SetDensityKeyword)
 					{
-						AllocatedProperties |= EPCGPointNativeProperties::Density;
+						AllocatedProperties |= static_cast<int32>(EPCGPointNativeProperties::Density);
 					}
 					else if (PropertyStr == PCGCustomHLSLKernel::SetSteepnessKeyword)
 					{
-						AllocatedProperties |= EPCGPointNativeProperties::Steepness;
+						AllocatedProperties |= static_cast<int32>(EPCGPointNativeProperties::Steepness);
 					}
 					else if (PropertyStr == PCGCustomHLSLKernel::SetSeedKeyword)
 					{
-						AllocatedProperties |= EPCGPointNativeProperties::Seed;
+						AllocatedProperties |= static_cast<int32>(EPCGPointNativeProperties::Seed);
 					}
 					else
 					{
diff --git a/Engine/Plugins/PCG/Source/PCG/Private/Compute/Elements/PCGCustomHLSLKernel.h b/Engine/Plugins/PCG/Source/PCG/Private/Compute/Elements/PCGCustomHLSLKernel.h
--- a/Engine/Plugins/PCG/Source/PCG/Private/Compute/Elements/PCGCustomHLSLKernel.h
+++ b/Engine/Plugins/PCG/Source/PCG/Private/Compute/Elements/PCGCustomHLSLKernel.h
@@ -9,8 +9,6 @@
 
 #include "PCGCustomHLSLKernel.generated.h"
 
-enum class EPCGPointNativeProperties : uint32;
-
 USTRUCT()
 struct FPCGParsedAttributeFunction
 {
@@ -175,7 +173,7 @@ protected:
 	UPROPERTY()
 	FPCGPinDataLabels PinDataLabels;
 
-	/** Properties that need to be fully allocated per pin. */
+	/** Properties that need to be fully allocated per pin. Note, we use an integer instead of EPCGPointNativeProperties directly because bitflag enums don't serialize correctly. */
 	UPROPERTY()
-	TMap<FName, EPCGPointNativeProperties> PinAllocatedProperties;
+	TMap<FName, /*EPCGPointNativeProperties*/int32> PinAllocatedProperties;
 };
