[PCG] Call FLinkerLoad::InvalidateExport before renaming Edge to transient package in case this happens before Edge is fully loaded
- Fixes an intermittent cooking error

#rnx

diff --git a/Engine/Plugins/PCG/Source/PCG/Private/PCGPin.cpp b/Engine/Plugins/PCG/Source/PCG/Private/PCGPin.cpp
--- a/Engine/Plugins/PCG/Source/PCG/Private/PCGPin.cpp
+++ b/Engine/Plugins/PCG/Source/PCG/Private/PCGPin.cpp
@@ -8,6 +8,7 @@
 #include "PCGNode.h"
 #include "PCGSettings.h"
 
+#include "UObject/LinkerLoad.h"
 #include "UObject/Package.h"
 
 #include UE_INLINE_GENERATED_CPP_BY_NAME(PCGPin)
@@ -218,6 +219,9 @@ bool UPCGPin::BreakEdgeTo(UPCGPin* OtherPin, TSet<UPCGNode*>* InTouchedNodes/*=
 
 			ensure(OtherPin->Edges.Remove(Edge));
 			Edges.Remove(Edge);
+
+			// In case this happens before Edge is fully loaded we need to invalidate its loading before renaming it out to the transient package.
+			FLinkerLoad::InvalidateExport(Edge);
 			Edge->Rename(/*NewName=*/nullptr, GetTransientPackage(), REN_DontCreateRedirectors);
 
 			if (InTouchedNodes)
@@ -253,6 +257,9 @@ bool UPCGPin::BreakAllEdges(TSet<UPCGNode*>* InTouchedNodes/*= nullptr*/)
 		{
 			OtherPin->Modify();
 			ensure(OtherPin->Edges.Remove(Edge));
+			
+			// In case this happens before Edge is fully loaded we need to invalidate its loading before renaming it out to the transient package.
+			FLinkerLoad::InvalidateExport(Edge);
 			Edge->Rename(/*NewName=*/nullptr, GetTransientPackage(), REN_DontCreateRedirectors);
 			bChanged = true;
 
