[PCG] Fixed Per data crc caching issue when we have N:N mappings

diff --git a/Engine/Plugins/PCG/Source/PCG/Private/Elements/PCGDistance.cpp b/Engine/Plugins/PCG/Source/PCG/Private/Elements/PCGDistance.cpp
--- a/Engine/Plugins/PCG/Source/PCG/Private/Elements/PCGDistance.cpp
+++ b/Engine/Plugins/PCG/Source/PCG/Private/Elements/PCGDistance.cpp
@@ -60,6 +60,10 @@ TArray<FPCGPinProperties> UPCGDistanceSettings::InputPinProperties() const
 	FPCGPinProperties& PinPropertySource = PinProperties.Emplace_GetRef(PCGDistance::SourceLabel, EPCGDataType::Point);
 	PinPropertySource.SetRequiredPin();
 	FPCGPinProperties& PinPropertyTarget = PinProperties.Emplace_GetRef(PCGDistance::TargetLabel, EPCGDataType::Point);
+	if (bCheckSourceAgainstRespectiveTarget)
+	{
+		PinPropertyTarget.SetRequiredPin();
+	}
 
 #if WITH_EDITOR
 	PinPropertySource.Tooltip = LOCTEXT("PCGSourcePinTooltip", "For each of the source points, a distance attribute will be calculated between it and the nearest target point.");
@@ -86,6 +90,20 @@ FPCGElementPtr UPCGDistanceSettings::CreateElement() const
 	return MakeShared<FPCGDistanceElement>();
 }
 
+#if WITH_EDITOR
+EPCGChangeType UPCGDistanceSettings::GetChangeTypeForProperty(const FName& InPropertyName) const
+{
+	EPCGChangeType ChangeType = Super::GetChangeTypeForProperty(InPropertyName);
+
+	if (InPropertyName == GET_MEMBER_NAME_CHECKED(UPCGDistanceSettings, bCheckSourceAgainstRespectiveTarget))
+	{
+		ChangeType |= EPCGChangeType::Structural;
+	}
+
+	return ChangeType;
+}
+#endif
+
 void UPCGDistanceSettings::PostLoad()
 {
 	Super::PostLoad();
@@ -110,6 +128,18 @@ void UPCGDistanceSettings::PostLoad()
 #endif // WITH_EDITOR
 }
 
+EPCGElementExecutionLoopMode FPCGDistanceElement::ExecutionLoopMode(const UPCGSettings* InSettings) const
+{
+	if (const UPCGDistanceSettings* Settings = Cast<UPCGDistanceSettings>(InSettings))
+	{
+		return Settings->bCheckSourceAgainstRespectiveTarget ? EPCGElementExecutionLoopMode::PrimaryPinAndBroadcastablePins : EPCGElementExecutionLoopMode::SinglePrimaryPin;
+	}
+	else
+	{
+		return EPCGElementExecutionLoopMode::NotALoop;
+	}
+}
+
 bool FPCGDistanceElement::ExecuteInternal(FPCGContext* Context) const
 {
 	TRACE_CPUPROFILER_EVENT_SCOPE(FPCGDistanceElement::Execute);
diff --git a/Engine/Plugins/PCG/Source/PCG/Private/Elements/PCGMatchAndSetAttributes.cpp b/Engine/Plugins/PCG/Source/PCG/Private/Elements/PCGMatchAndSetAttributes.cpp
--- a/Engine/Plugins/PCG/Source/PCG/Private/Elements/PCGMatchAndSetAttributes.cpp
+++ b/Engine/Plugins/PCG/Source/PCG/Private/Elements/PCGMatchAndSetAttributes.cpp
@@ -96,21 +96,21 @@ TArray<FPCGPinProperties> UPCGMatchAndSetAttributesSettings::InputPinProperties(
 	FPCGPinProperties& InputPinProperty = PinProperties.Emplace_GetRef(PCGPinConstants::DefaultInputLabel, EPCGDataType::PointOrParam);
 	InputPinProperty.SetRequiredPin();
 
-	PinProperties.Emplace(PCGMatchAndSetAttributesConstants::MatchDataLabel, 
-		EPCGDataType::Param, 
+	PinProperties.Emplace_GetRef(PCGMatchAndSetAttributesConstants::MatchDataLabel,
+		EPCGDataType::Param,
 		/*bAllowMultipleConnection=*/true,
 		/*bAllowMultipleData=*/true,
 		LOCTEXT("MatchDataTooltip", "Input containing the data to match to, then copy the accompanying attribute values")
-	);
+	).SetRequiredPin();
 
 	if (bFindNearest && MaxDistanceMode == EPCGMatchMaxDistanceMode::AttributeMaxDistance)
 	{
-		PinProperties.Emplace(PCGMatchAndSetAttributesConstants::MaxDistanceLabel,
+		PinProperties.Emplace_GetRef(PCGMatchAndSetAttributesConstants::MaxDistanceLabel,
 			EPCGDataType::PointOrParam,
 			/*bAllowMultipleConnections=*/true,
 			/*bAllowMultipleData=*/true,
 			LOCTEXT("MaxDistanceTooltip", "Input containing the maximum distance allowed for nearest search, selected by the Max Distance Attribute.")
-		);
+		).SetRequiredPin();
 	}
 
 	return PinProperties;
@@ -124,6 +124,21 @@ TArray<FPCGPinProperties> UPCGMatchAndSetAttributesSettings::OutputPinProperties
 	return PinProperties;
 }
 
+#if WITH_EDITOR
+EPCGChangeType UPCGMatchAndSetAttributesSettings::GetChangeTypeForProperty(const FName& InPropertyName) const
+{
+	EPCGChangeType ChangeType = Super::GetChangeTypeForProperty(InPropertyName);
+
+	if (InPropertyName == GET_MEMBER_NAME_CHECKED(UPCGMatchAndSetAttributesSettings, bFindNearest) ||
+		InPropertyName == GET_MEMBER_NAME_CHECKED(UPCGMatchAndSetAttributesSettings, MaxDistanceMode))
+	{
+		ChangeType |= EPCGChangeType::Structural;
+	}
+
+	return ChangeType;
+}
+#endif // WITH_EDITOR
+
 FPCGElementPtr UPCGMatchAndSetAttributesSettings::CreateElement() const
 {
 	return MakeShared<FPCGMatchAndSetAttributesElement>();
diff --git a/Engine/Plugins/PCG/Source/PCG/Private/PCGElement.cpp b/Engine/Plugins/PCG/Source/PCG/Private/PCGElement.cpp
--- a/Engine/Plugins/PCG/Source/PCG/Private/PCGElement.cpp
+++ b/Engine/Plugins/PCG/Source/PCG/Private/PCGElement.cpp
@@ -63,18 +63,25 @@ namespace PCGElementHelpers
 		OutPrimaryCollections.Reset();
 		OutCommonCollection.TaggedData.Reset();
 
-		TArray<FPCGPinProperties> RequiredPins = Settings->AllInputPinProperties().FilterByPredicate([](const FPCGPinProperties& Props) { return Props.IsRequiredPin(); });
-
 		// Early out
+		if (Mode == EPCGElementExecutionLoopMode::NotALoop)
+		{
+			return false;
+		}
+
+		TArray<FPCGPinProperties> RequiredPins = Settings->InputPinProperties();
+		RequiredPins.RemoveAll([](const FPCGPinProperties& Props) { return !Props.IsRequiredPin(); });
+
+		// Early out if the single pin criteria is not met. Normally this is an implementer issue
 		if (Mode == EPCGElementExecutionLoopMode::SinglePrimaryPin && RequiredPins.Num() != 1)
 		{
 			return false;
 		}
 
-		TArray<FName> RequiredPinLabels;
+		TArray<FName, TInlineAllocator<4>> RequiredPinLabels;
 		Algo::Transform(RequiredPins, RequiredPinLabels, [](const FPCGPinProperties& Props) { return Props.Label; });
 
-		TArray<FPCGDataCollection> DataPerRequiredPin;
+		TArray<FPCGDataCollection, TInlineAllocator<4>> DataPerRequiredPin;
 		DataPerRequiredPin.SetNum(RequiredPinLabels.Num());
 
 		for (int32 DataIndex = 0; DataIndex < Collection.TaggedData.Num(); ++DataIndex)
@@ -132,6 +139,39 @@ namespace PCGElementHelpers
 				}
 			}
 		}
+		else if (Mode == EPCGElementExecutionLoopMode::PrimaryPinAndBroadcastablePins)
+		{
+			// Idea - behaves like the matching primary pins, but we'll move out the data from pins that have a single data instead.
+			int32 NumberOfData = DataPerRequiredPin[0].TaggedData.Num();
+
+			// In this case, we'll just validate that the values match, and that the other required pins don't have extraneous data too.
+			for (int32 RequiredPinIndex = 1; RequiredPinIndex < DataPerRequiredPin.Num(); ++RequiredPinIndex)
+			{
+				const int32 DataCountOnPin = DataPerRequiredPin[RequiredPinIndex].TaggedData.Num();
+				if (DataCountOnPin != 1 && DataCountOnPin != NumberOfData)
+				{
+					return false;
+				}
+			}
+
+			OutPrimaryCollections.SetNum(NumberOfData);
+
+			for (int32 RequiredPinIndex = 0; RequiredPinIndex < DataPerRequiredPin.Num(); ++RequiredPinIndex)
+			{
+				if (DataPerRequiredPin[RequiredPinIndex].TaggedData.Num() == NumberOfData)
+				{
+					for (int32 DataIndex = 0; DataIndex < NumberOfData; ++DataIndex)
+					{
+						OutPrimaryCollections[DataIndex].TaggedData.Add(DataPerRequiredPin[RequiredPinIndex].TaggedData[DataIndex]);
+					}
+				}
+				else
+				{
+					ensure(DataPerRequiredPin[RequiredPinIndex].TaggedData.Num() == 1);
+					OutCommonCollection.TaggedData.Append(DataPerRequiredPin[RequiredPinIndex].TaggedData);
+				}
+			}
+		}
 		/*else if (Mode == EPCGElementExecutionLoopMode::CartesianPins)
 		{
 			int32 NumberOfCollections = 1;
@@ -369,10 +409,7 @@ void IPCGElement::PreExecutePrimaryLoopElement(FPCGContext* Context, const UPCGS
 			InputToOutputResults.Key = PrimaryDataCollection;
 			InputToOutputResults.Value = SubCollectionOutput;
 
-			for (int32 SubDataIndex = PrimaryDataCollection.TaggedData.Num() - 1; SubDataIndex >= 0; --SubDataIndex)
-			{
-				Context->InputData.TaggedData.RemoveAt(PrimaryDataCollection.TaggedData[SubDataIndex].OriginalIndex);
-			}
+			Context->InputData.TaggedData.RemoveAll([&PrimaryDataCollection](const FPCGTaggedData& Data) { return PrimaryDataCollection.TaggedData.Contains(Data); });
 		}
 	}
 
@@ -634,21 +671,21 @@ void IPCGElement::PostExecutePrimaryLoopElement(FPCGContext* Context, const UPCG
 
 		for (const FPCGDataCollection& PrimaryDataCollection : PrimaryDataCollections)
 		{
-			if (PrimaryDataCollection.TaggedData.IsEmpty() || PrimaryDataCollection.TaggedData[0].OriginalIndex == INDEX_NONE)
+			if (PrimaryDataCollection.TaggedData.IsEmpty())
 			{
 				continue;
 			}
+			
+			const int32 OriginalIndex = PrimaryDataCollection.TaggedData[0].OriginalIndex;
 
-			FPCGDataCollection SubCollectionOutput;
-			for (int32 DataIndex = 0; DataIndex < Context->OutputData.TaggedData.Num(); ++DataIndex)
+			if (OriginalIndex == INDEX_NONE)
 			{
-				const FPCGTaggedData& TaggedData = Context->OutputData.TaggedData[DataIndex];
-				if (TaggedData.OriginalIndex == PrimaryDataCollection.TaggedData[0].OriginalIndex)
-				{
-					SubCollectionOutput.TaggedData.Add(Context->OutputData.TaggedData[DataIndex]);
-				}
+				continue;
 			}
 
+			FPCGDataCollection SubCollectionOutput;
+			SubCollectionOutput.TaggedData = Context->OutputData.TaggedData.FilterByPredicate([OriginalIndex](const FPCGTaggedData& TaggedData) { return TaggedData.OriginalIndex == OriginalIndex; });
+
 			bool bCacheable = true;
 
 			for (const FPCGTaggedData& TaggedData : SubCollectionOutput.TaggedData)
@@ -688,17 +725,21 @@ void IPCGElement::PostExecutePrimaryLoopElement(FPCGContext* Context, const UPCG
 			// Since we've assigned original indices up front, we should be able to respect this here.
 			const FPCGDataCollection& CachedInputData = Context->CachedInputToOutputInternalResults[CachedCollectionIndex].Key;
 			const int32 CacheInputOriginalIndex = (CachedInputData.TaggedData.IsEmpty() ? INDEX_NONE : CachedInputData.TaggedData[0].OriginalIndex);
-			const int32 InsertInputIndex = CacheInputOriginalIndex != INDEX_NONE ? Context->InputData.TaggedData.IndexOfByPredicate([CacheInputOriginalIndex](const FPCGTaggedData& TaggedData) { return TaggedData.OriginalIndex > CacheInputOriginalIndex; }) : INDEX_NONE;
 
-			if (InsertInputIndex != INDEX_NONE)
-			{
-				Context->InputData.TaggedData.Insert(CachedInputData.TaggedData, InsertInputIndex);
-			}
-			else
+			for (const FPCGTaggedData& CachedInput : CachedInputData.TaggedData)
 			{
-				Context->InputData.TaggedData.Append(CachedInputData.TaggedData);
-			}
+				const int32 InsertInputIndex = CacheInputOriginalIndex != INDEX_NONE ? Context->InputData.TaggedData.IndexOfByPredicate([InputIndex = CachedInput.OriginalIndex](const FPCGTaggedData& TaggedData) { return TaggedData.OriginalIndex > InputIndex; }) : INDEX_NONE;
 
+				if (InsertInputIndex != INDEX_NONE)
+				{
+					Context->InputData.TaggedData.Insert(CachedInput, InsertInputIndex);
+				}
+				else
+				{
+					Context->InputData.TaggedData.Add(CachedInput);
+				}
+			}
+			
 			// Note: we're modifying the context data, but it will not be reused after
 			FPCGDataCollection& CachedOutputData = Context->CachedInputToOutputInternalResults[CachedCollectionIndex].Value;
 
diff --git a/Engine/Plugins/PCG/Source/PCG/Public/Elements/PCGDistance.h b/Engine/Plugins/PCG/Source/PCG/Public/Elements/PCGDistance.h
--- a/Engine/Plugins/PCG/Source/PCG/Public/Elements/PCGDistance.h
+++ b/Engine/Plugins/PCG/Source/PCG/Public/Elements/PCGDistance.h
@@ -45,6 +45,9 @@ protected:
 	virtual TArray<FPCGPinProperties> InputPinProperties() const override;
 	virtual TArray<FPCGPinProperties> OutputPinProperties() const override;
 	virtual FPCGElementPtr CreateElement() const override;
+#if WITH_EDITOR
+	virtual EPCGChangeType GetChangeTypeForProperty(const FName& InPropertyName) const override;
+#endif
 	//~End UPCGSettings interface
 
 public:
@@ -94,6 +97,6 @@ class FPCGDistanceElement : public IPCGElement
 {
 protected:
 	virtual bool ExecuteInternal(FPCGContext* Context) const override;
-	virtual EPCGElementExecutionLoopMode ExecutionLoopMode(const UPCGSettings* Settings) const override { return EPCGElementExecutionLoopMode::SinglePrimaryPin; }
+	virtual EPCGElementExecutionLoopMode ExecutionLoopMode(const UPCGSettings* Settings) const override;
 	virtual bool SupportsBasePointDataInputs(FPCGContext* InContext) const override { return true; }
 };
\ No newline at end of file
diff --git a/Engine/Plugins/PCG/Source/PCG/Public/Elements/PCGMatchAndSetAttributes.h b/Engine/Plugins/PCG/Source/PCG/Public/Elements/PCGMatchAndSetAttributes.h
--- a/Engine/Plugins/PCG/Source/PCG/Public/Elements/PCGMatchAndSetAttributes.h
+++ b/Engine/Plugins/PCG/Source/PCG/Public/Elements/PCGMatchAndSetAttributes.h
@@ -67,6 +67,9 @@ protected:
 	virtual TArray<FPCGPinProperties> InputPinProperties() const override;
 	virtual TArray<FPCGPinProperties> OutputPinProperties() const override;
 	virtual FPCGElementPtr CreateElement() const override;
+#if WITH_EDITOR
+	virtual EPCGChangeType GetChangeTypeForProperty(const FName& InPropertyName) const override;
+#endif
 	// ~End UPCGSettings interface
 
 public:
@@ -130,6 +133,6 @@ class FPCGMatchAndSetAttributesElement : public TPCGTimeSlicedElementBase<FPCGMa
 protected:
 	virtual bool PrepareDataInternal(FPCGContext* InContext) const override;
 	virtual bool ExecuteInternal(FPCGContext* InContext) const override;
-	virtual EPCGElementExecutionLoopMode ExecutionLoopMode(const UPCGSettings* Settings) const override { return EPCGElementExecutionLoopMode::SinglePrimaryPin; }
+	virtual EPCGElementExecutionLoopMode ExecutionLoopMode(const UPCGSettings* Settings) const override { return EPCGElementExecutionLoopMode::PrimaryPinAndBroadcastablePins; }
 	virtual bool SupportsBasePointDataInputs(FPCGContext* InContext) const override { return true; }
 };
diff --git a/Engine/Plugins/PCG/Source/PCG/Public/PCGElement.h b/Engine/Plugins/PCG/Source/PCG/Public/PCGElement.h
--- a/Engine/Plugins/PCG/Source/PCG/Public/PCGElement.h
+++ b/Engine/Plugins/PCG/Source/PCG/Public/PCGElement.h
@@ -74,6 +74,8 @@ enum class EPCGElementExecutionLoopMode : uint8
 	SinglePrimaryPin,
 	/** Loops on matching indices on required pin(s). */
 	MatchingPrimaryPins,
+	/** Loops on matching indices on required pin(s), but the second+ pin can have a single entry. E.g. 1:1, N:N or N:1 */
+	PrimaryPinAndBroadcastablePins,
 	/** Cartesian loop on required pins. */
 	// CartesianPins // TODO
 };
