[PCG] Fix GPU Transform Points rotations when using non-uniform angle ranges.

#rnx
#virtualized

diff --git a/Engine/Plugins/PCG/Shaders/Private/Elements/PCGTransformPoints.usf b/Engine/Plugins/PCG/Shaders/Private/Elements/PCGTransformPoints.usf
--- a/Engine/Plugins/PCG/Shaders/Private/Elements/PCGTransformPoints.usf
+++ b/Engine/Plugins/PCG/Shaders/Private/Elements/PCGTransformPoints.usf
@@ -79,17 +79,19 @@ void PCGTransformPointsCS(uint3 GroupId : SV_GroupID, uint GroupIndex : SV_Group
 	uint Seed = ComputeSeed(GetSeed(), InSeed);
 
 	// Position
-	float3 OutPosition = bAbsoluteOffset ? float3(0,0,0) : InPosition;
+	float3 OutPosition = bAbsoluteOffset ? float3(0, 0, 0) : InPosition;
 	OutPosition.x += lerp(OffsetMin.x, OffsetMax.x, FRand(Seed));
 	OutPosition.y += lerp(OffsetMin.y, OffsetMax.y, FRand(Seed));
 	OutPosition.z += lerp(OffsetMin.z, OffsetMax.z, FRand(Seed));
 
 	// Rotation
-	FQuat OutRotation = bAbsoluteRotation ? FQuat(0,0,0,1) : InRotation;
-	float3 RandomRotation = float3(0,0,0);
-	RandomRotation.x = lerp(RotationMin.x, RotationMax.x, FRand(Seed));
+	FQuat OutRotation = bAbsoluteRotation ? FQuat(0, 0, 0, 1) : InRotation;
+	float3 RandomRotation = float3(0, 0, 0);
+
+	// Note: Random numbers drawn in same order as CPU.
 	RandomRotation.y = lerp(RotationMin.y, RotationMax.y, FRand(Seed));
 	RandomRotation.z = lerp(RotationMin.z, RotationMax.z, FRand(Seed));
+	RandomRotation.x = lerp(RotationMin.x, RotationMax.x, FRand(Seed));
 
 	OutRotation = Quat_Multiply(OutRotation, Quat_FromRotator(RandomRotation));
 
