[PCG] Further adjustments to primitive factory after iterating further on component-less.

The component-less path has clusters of multiple primitives.

#rnx

diff --git a/Engine/Plugins/PCG/Source/PCG/Private/Compute/DataInterfaces/PCGInstanceDataInterface.cpp b/Engine/Plugins/PCG/Source/PCG/Private/Compute/DataInterfaces/PCGInstanceDataInterface.cpp
--- a/Engine/Plugins/PCG/Source/PCG/Private/Compute/DataInterfaces/PCGInstanceDataInterface.cpp
+++ b/Engine/Plugins/PCG/Source/PCG/Private/Compute/DataInterfaces/PCGInstanceDataInterface.cpp
@@ -167,27 +167,42 @@ bool UPCGInstanceDataProvider::PrepareForExecute_GameThread(UPCGDataBinding* InB
 		FoundPrimitives->PrimitiveFactories.SetNum(PCGComputeConstants::MAX_PRIMITIVE_COMPONENTS_PER_SPAWNER);
 	}
 
+	bool bAllReady = true;
+
+	for (TSharedPtr<IPCGRuntimePrimitiveFactory>& Factory : FoundPrimitives->PrimitiveFactories)
+	{
+		if (ensure(Factory))
+		{
+			// This can tick the creation of the render proxy for some factories so call on all factories rather than early'ing out.
+			bAllReady &= Factory->IsRenderStateCreated();
+		}
+	}
+
+	if (!bAllReady)
+	{
+		UE_LOG(LogPCG, Verbose, TEXT("FPCGInstanceDataProvider: One or more scene proxies were not ready. Will try again on the next tick."));
+		return false;
+	}
+
 	NumInstancesAllPrimitives = 0;
 	Primitives.Reserve(FoundPrimitives->PrimitiveFactories.Num());
 
-	for (TSharedPtr<IPCGRuntimePrimitiveFactory> Primitive : FoundPrimitives->PrimitiveFactories)
+	for (TSharedPtr<IPCGRuntimePrimitiveFactory> Factory : FoundPrimitives->PrimitiveFactories)
 	{
-		if (!ensure(Primitive))
+		if (!ensure(Factory))
 		{
 			continue;
 		}
 
-		FPrimitiveSceneProxy* SceneProxy = Primitive->GetSceneProxy();
-		if (!SceneProxy)
+		for (int32 PrimitiveIndex = 0; PrimitiveIndex < Factory->GetNumPrimitives(); ++PrimitiveIndex)
 		{
-			// Not ready yet, try again later.
-			UE_LOG(LogPCG, Verbose, TEXT("FPCGInstanceDataProvider: One or more scene proxies were not ready. Will try again on the next tick."));
-			return false;
-		}
+			FPrimitiveSceneProxy* SceneProxy = Factory->GetSceneProxy(PrimitiveIndex);
+			check(SceneProxy);
 
-		const int32 NumInstances = Primitive->GetNumInstances();
-		Primitives.Add(FPCGPrimitiveAndNumInstances(SceneProxy, NumInstances));
-		NumInstancesAllPrimitives += NumInstances;
+			const int32 NumInstances = Factory->GetNumInstances(PrimitiveIndex);
+			Primitives.Add(FPCGPrimitiveAndNumInstances(SceneProxy, NumInstances));
+			NumInstancesAllPrimitives += NumInstances;
+		}
 	}
 
 	NumCustomFloatsPerInstance = FoundPrimitives->NumCustomFloats;
@@ -323,7 +338,7 @@ void FPCGInstanceDataProviderProxy::PostSubmit(FComputeContext& Context) const
 	const TRefCountPtr<FRDGPooledBuffer> InstanceCustomFloatDataExported = Context.GraphBuilder.ConvertToExternalBuffer(InstanceCustomFloatData);
 	const TRefCountPtr<FRDGPooledBuffer> WriteCountersExported = Context.GraphBuilder.ConvertToExternalBuffer(WriteCounters);
 	
-	FTSTicker::GetCoreTicker().AddTicker(TEXT("ApplyPrimitiveSceneUpdates"), 0.0f, [DataProviderWeak=DataProvider, DataProviderGeneration=DataProviderGeneration, InstanceDataExported, InstanceCustomFloatDataExported, WriteCountersExported, ContextHandle=ContextHandle](float)
+	FTSTicker::GetCoreTicker().AddTicker(TEXT("ApplyPrimitiveSceneUpdates"), /*InDelay=*/0.0f, [DataProviderWeak=DataProvider, DataProviderGeneration=DataProviderGeneration, InstanceDataExported, InstanceCustomFloatDataExported, WriteCountersExported, ContextHandle=ContextHandle](float)
 	{
 		UPCGInstanceDataProvider* DataProviderInner = DataProviderWeak.Get();
 		if (!DataProviderInner || DataProviderInner->GenerationCounter != DataProviderGeneration)
diff --git a/Engine/Plugins/PCG/Source/PCG/Private/Compute/PCGDataBinding.h b/Engine/Plugins/PCG/Source/PCG/Private/Compute/PCGDataBinding.h
--- a/Engine/Plugins/PCG/Source/PCG/Private/Compute/PCGDataBinding.h
+++ b/Engine/Plugins/PCG/Source/PCG/Private/Compute/PCGDataBinding.h
@@ -19,18 +19,6 @@ struct FPCGComputeGraphContext;
 struct FPCGContextHandle;
 struct FPCGKernelPin;
 
-USTRUCT()
-struct FPCGSpawnerPrimitive
-{
-	GENERATED_BODY()
-
-	UPROPERTY()
-	TObjectPtr<UPrimitiveComponent> Component;
-
-	UPROPERTY()
-	int NumInstances = INDEX_NONE;
-};
-
 USTRUCT()
 struct FPCGSpawnerPrimitives
 {
diff --git a/Engine/Plugins/PCG/Source/PCG/Private/Compute/PrimitiveFactories/IPCGRuntimePrimitiveFactory.h b/Engine/Plugins/PCG/Source/PCG/Private/Compute/PrimitiveFactories/IPCGRuntimePrimitiveFactory.h
--- a/Engine/Plugins/PCG/Source/PCG/Private/Compute/PrimitiveFactories/IPCGRuntimePrimitiveFactory.h
+++ b/Engine/Plugins/PCG/Source/PCG/Private/Compute/PrimitiveFactories/IPCGRuntimePrimitiveFactory.h
@@ -10,6 +10,10 @@ class FPrimitiveSceneProxy;
 class IPCGRuntimePrimitiveFactory
 {
 public:
-	virtual FPrimitiveSceneProxy* GetSceneProxy() const = 0;
-	virtual int32 GetNumInstances() const = 0;
+	virtual ~IPCGRuntimePrimitiveFactory() = default;
+
+	virtual int32 GetNumPrimitives() const = 0;
+	virtual FPrimitiveSceneProxy* GetSceneProxy(int32 InPrimitiveIndex) const = 0;
+	virtual int32 GetNumInstances(int32 InPrimitiveIndex) const = 0;
+	virtual bool IsRenderStateCreated() const = 0;
 };
diff --git a/Engine/Plugins/PCG/Source/PCG/Private/Compute/PrimitiveFactories/PCGPrimitiveFactoryPrimitiveComponentBase.cpp b/Engine/Plugins/PCG/Source/PCG/Private/Compute/PrimitiveFactories/PCGPrimitiveFactoryPrimitiveComponentBase.cpp
--- a/Engine/Plugins/PCG/Source/PCG/Private/Compute/PrimitiveFactories/PCGPrimitiveFactoryPrimitiveComponentBase.cpp
+++ b/Engine/Plugins/PCG/Source/PCG/Private/Compute/PrimitiveFactories/PCGPrimitiveFactoryPrimitiveComponentBase.cpp
@@ -4,7 +4,13 @@
 
 #include "Components/PrimitiveComponent.h"
 
-FPrimitiveSceneProxy* FPCGPrimitiveFactoryPrimitiveComponentBase::GetSceneProxy() const
+FPrimitiveSceneProxy* FPCGPrimitiveFactoryPrimitiveComponentBase::GetSceneProxy(int32 InPrimitiveIndex) const
 {
+	ensure(InPrimitiveIndex == 0);
 	return Component.IsValid() ? Component->GetSceneProxy() : nullptr;
 }
+
+bool FPCGPrimitiveFactoryPrimitiveComponentBase::IsRenderStateCreated() const
+{
+	return GetSceneProxy(0) != nullptr;
+}
diff --git a/Engine/Plugins/PCG/Source/PCG/Private/Compute/PrimitiveFactories/PCGPrimitiveFactoryPrimitiveComponentBase.h b/Engine/Plugins/PCG/Source/PCG/Private/Compute/PrimitiveFactories/PCGPrimitiveFactoryPrimitiveComponentBase.h
--- a/Engine/Plugins/PCG/Source/PCG/Private/Compute/PrimitiveFactories/PCGPrimitiveFactoryPrimitiveComponentBase.h
+++ b/Engine/Plugins/PCG/Source/PCG/Private/Compute/PrimitiveFactories/PCGPrimitiveFactoryPrimitiveComponentBase.h
@@ -13,8 +13,10 @@ class FPCGPrimitiveFactoryPrimitiveComponentBase : public IPCGRuntimePrimitiveFa
 {
 public:
 	//~ Begin IPCGRuntimePrimitiveFactory interface
-	virtual FPrimitiveSceneProxy* GetSceneProxy() const override;
-	virtual int32 GetNumInstances() const override { return NumInstances; }
+	virtual int32 GetNumPrimitives() const override { return 1; }
+	virtual FPrimitiveSceneProxy* GetSceneProxy(int32 InPrimitiveIndex) const override;
+	virtual int32 GetNumInstances(int32 InPrimitiveIndex) const override { ensure(InPrimitiveIndex == 0); return NumInstances; }
+	virtual bool IsRenderStateCreated() const override;
 	//~ End IPCGRuntimePrimitiveFactory interface
 
 protected:
