[PCG] Made traces hit physics objects only, not queryable ones.

#rnx

diff --git a/Engine/Plugins/PCG/Source/PCGEditor/Private/EditorMode/Tools/Helpers/PCGEdModeSceneQueryHelpers.cpp b/Engine/Plugins/PCG/Source/PCGEditor/Private/EditorMode/Tools/Helpers/PCGEdModeSceneQueryHelpers.cpp
--- a/Engine/Plugins/PCG/Source/PCGEditor/Private/EditorMode/Tools/Helpers/PCGEdModeSceneQueryHelpers.cpp
+++ b/Engine/Plugins/PCG/Source/PCGEditor/Private/EditorMode/Tools/Helpers/PCGEdModeSceneQueryHelpers.cpp
@@ -38,7 +38,8 @@ namespace PCG::EditorMode::Scene
 			return {};
 		}
 
-		const FCollisionObjectQueryParams ObjectQueryParams(FCollisionObjectQueryParams::AllObjects);
+		// @todo_pcg: make this configurable
+		const FCollisionObjectQueryParams ObjectQueryParams(ECC_WorldStatic);
 		FCollisionQueryParams QueryParams = FCollisionQueryParams::DefaultQueryParam;
 		QueryParams.bTraceComplex = true;
 
@@ -52,6 +53,12 @@ namespace PCG::EditorMode::Scene
 
 		for (const FHitResult& Hit : OutHits)
 		{
+			// Skip hits on non-physical components, esp. brush components, like PCG volumes.
+			if (Hit.GetComponent() && !Hit.GetComponent()->IsPhysicsCollisionEnabled())
+			{
+				continue;
+			}
+
 			// Check for a visible result
 			bool bHitIsValid = IsHitTargetVisible(Hit);
 			// If invisible, check if it's explicitly allowed.
