[PCG] Add editor functions to get the pin color/type depending on its type id.

* There will be 2 arrays, one for compositions and one for single types
* Arrays will be ordered so a type that is a child of the other will be seen first.
* Add IsChildOf on PCGDataTypeIdentifier + unit tests
* Unify the constructors to support variadic construction of identifiers.

#rnx

diff --git a/Engine/Plugins/PCG/Source/PCG/Private/Data/Registry/PCGDataTypeIdentifier.cpp b/Engine/Plugins/PCG/Source/PCG/Private/Data/Registry/PCGDataTypeIdentifier.cpp
--- a/Engine/Plugins/PCG/Source/PCG/Private/Data/Registry/PCGDataTypeIdentifier.cpp
+++ b/Engine/Plugins/PCG/Source/PCG/Private/Data/Registry/PCGDataTypeIdentifier.cpp
@@ -44,7 +44,7 @@ namespace PCGDataTypeIdentifier
 		};
 }
 
-FPCGDataTypeIdentifier::FPCGDataTypeIdentifier(EPCGDataType LegacyType)
+void FPCGDataTypeIdentifier::ComposeFromLegacyType(EPCGDataType LegacyType)
 {
 	FPCGDataTypeBaseId BaseId = FPCGDataTypeBaseId::MakeFromLegacyType(LegacyType);
 	if (BaseId.IsValid())
@@ -73,35 +73,53 @@ FPCGDataTypeIdentifier::FPCGDataTypeIdentifier(EPCGDataType LegacyType)
 	}
 }
 
-FPCGDataTypeIdentifier::FPCGDataTypeIdentifier(TSubclassOf<UPCGData> Class)
-	: FPCGDataTypeIdentifier(Class ? Class->GetDefaultObject<UPCGData>()->GetDataTypeId() : FPCGDataTypeBaseId{})
-{}
+void FPCGDataTypeIdentifier::ComposeFromPCGDataSubclass(TSubclassOf<UPCGData> Subclass)
+{
+	if (Subclass)
+	{
+		if (FPCGDataTypeBaseId BaseId = Subclass->GetDefaultObject<UPCGData>()->GetDataTypeId(); BaseId.IsValid())
+		{
+			Ids.AddUnique(BaseId);
+		}
+	}
+}
 
-FPCGDataTypeIdentifier::FPCGDataTypeIdentifier(TConstArrayView<FPCGDataTypeIdentifier> InIDs)
+FPCGDataTypeIdentifier FPCGDataTypeIdentifier::Construct(TConstArrayView<FPCGDataTypeIdentifier> InIds)
 {
-	for (const FPCGDataTypeIdentifier& ID : InIDs)
+	FPCGDataTypeIdentifier Result;
+	
+	for (const FPCGDataTypeIdentifier& Id : InIds)
 	{
-		for (const FPCGDataTypeBaseId& BaseID : ID.GetIds())
+		for (const FPCGDataTypeBaseId& BaseId : Id.GetIds())
 		{
-			Ids.AddUnique(BaseID);
+			if (BaseId.IsValid())
+			{
+				Result.Ids.AddUnique(BaseId);
+			}
 		}
 	}
+
+	return Result;
 }
 
-FPCGDataTypeIdentifier::FPCGDataTypeIdentifier(TConstArrayView<TSubclassOf<UPCGData>> Classes)
+FPCGDataTypeIdentifier FPCGDataTypeIdentifier::Construct(TConstArrayView<TSubclassOf<UPCGData>> Classes)
 {
+	FPCGDataTypeIdentifier Result;
+	
 	for (TSubclassOf<UPCGData> Class : Classes)
 	{
-		Ids.AddUnique(Class->GetDefaultObject<UPCGData>()->GetDataTypeId());
+		Result.ComposeFromPCGDataSubclass(Class);
 	}
+
+	return Result;
 }
 
 FPCGDataTypeIdentifier::operator EPCGDataType() const
 {
 	EPCGDataType Result = EPCGDataType::None;
-	for (const FPCGDataTypeBaseId& ID : Ids)
+	for (const FPCGDataTypeBaseId& Id : Ids)
 	{
-		Result |= ID.AsLegacyType();
+		Result |= Id.AsLegacyType();
 	}
 
 	return Result;
@@ -115,11 +133,11 @@ FPCGDataTypeIdentifier& FPCGDataTypeIdentifier::operator|=(EPCGDataType OtherLeg
 
 FPCGDataTypeIdentifier& FPCGDataTypeIdentifier::operator|=(const FPCGDataTypeIdentifier& Other)
 {
-	for (const FPCGDataTypeBaseId& ID : Other.GetIds())
+	for (const FPCGDataTypeBaseId& Id : Other.GetIds())
 	{
-		if (ID.IsValid())
+		if (Id.IsValid())
 		{
-			Ids.AddUnique(ID);
+			Ids.AddUnique(Id);
 		}
 	}
 	
@@ -226,7 +244,7 @@ bool FPCGDataTypeIdentifier::Intersects(const FPCGDataTypeIdentifier& Other) con
 
 FPCGDataTypeIdentifier FPCGDataTypeIdentifier::Compose(const FPCGDataTypeIdentifier& Other)
 {
-	return Compose({*this, Other});
+	return Compose(TConstArrayView<FPCGDataTypeIdentifier>{{*this, Other}});
 }
 
 FPCGDataTypeIdentifier FPCGDataTypeIdentifier::Compose(TConstArrayView<FPCGDataTypeIdentifier> IDs)
@@ -244,7 +262,7 @@ bool FPCGDataTypeIdentifier::SupportsType(EPCGDataType OtherLegacyType) const
 	const FPCGDataTypeIdentifier OtherID(OtherLegacyType);
 	if (OtherID.IsComposition())
 	{
-		return Algo::AllOf(OtherID.GetIds(), [this](const FPCGDataTypeBaseId& ID) { return Ids.Contains(ID); });
+		return Algo::AllOf(OtherID.GetIds(), [this](const FPCGDataTypeBaseId& Id) { return Ids.Contains(Id); });
 	}
 	else
 	{
@@ -276,7 +294,7 @@ bool FPCGDataTypeIdentifier::IsSameType(const FPCGDataTypeIdentifier& Other) con
 
 bool FPCGDataTypeIdentifier::SupportsType(const FPCGDataTypeIdentifier& Other) const
 {
-	return IsValid() && Algo::AllOf(Other.GetIds(), [this](const FPCGDataTypeBaseId& ID) { return Ids.Contains(ID); });
+	return IsValid() && Algo::AllOf(Other.GetIds(), [this](const FPCGDataTypeBaseId& Id) { return Ids.Contains(Id); });
 }
 
 bool FPCGDataTypeIdentifier::IsWider(const FPCGDataTypeIdentifier& Other) const
@@ -295,6 +313,19 @@ bool FPCGDataTypeIdentifier::IsIdentical(const FPCGDataTypeIdentifier& Other) co
 	return IsSameType(Other) && CustomSubtype == Other.CustomSubtype;
 }
 
+bool FPCGDataTypeIdentifier::IsChildOf(const FPCGDataTypeIdentifier& Other) const
+{
+	if (IsComposition() || Other.IsComposition())
+	{
+		return (*this & Other).IsSameType(*this);
+	}
+	else
+	{
+		// Optimization if none are compositions
+		return GetId().IsChildOf(Other.GetId());
+	}
+}
+
 bool FPCGDataTypeIdentifier::SerializeFromMismatchedTag(const struct FPropertyTag& Tag, FStructuredArchive::FSlot Slot)
 {
 	const UEnum* Enum = StaticEnum<EPCGDataType>();
@@ -326,7 +357,7 @@ bool FPCGDataTypeIdentifier::SerializeFromMismatchedTag(const struct FPropertyTa
 
 FString FPCGDataTypeIdentifier::ToString() const
 {
-	return IsValid() ? FString::JoinBy(Ids, TEXT(" | "), [](const FPCGDataTypeBaseId& ID) { return ID.ToString(); }) : TEXT("None");
+	return IsValid() ? FString::JoinBy(Ids, TEXT(" | "), [](const FPCGDataTypeBaseId& Id) { return Id.ToString(); }) : TEXT("None");
 }
 
 FText FPCGDataTypeIdentifier::ToDisplayText() const
diff --git a/Engine/Plugins/PCG/Source/PCG/Private/Data/Registry/PCGDataTypeRegistry.cpp b/Engine/Plugins/PCG/Source/PCG/Private/Data/Registry/PCGDataTypeRegistry.cpp
--- a/Engine/Plugins/PCG/Source/PCG/Private/Data/Registry/PCGDataTypeRegistry.cpp
+++ b/Engine/Plugins/PCG/Source/PCG/Private/Data/Registry/PCGDataTypeRegistry.cpp
@@ -687,6 +687,130 @@ void FPCGDataTypeRegistry::HierarchyDepthSearch(TFunctionRef<ESearchCommand(cons
 	});
 }
 
+#if WITH_EDITOR
+namespace PCGDataTypeRegistry
+{
+	template <typename Function>
+	void AddUniqueOrdered(const FPCGDataTypeIdentifier& InId, TArray<TTuple<FPCGDataTypeIdentifier, Function>>& InContainer, Function&& InFunction)
+	{
+		bool bInserted = false;
+		for (int32 i = 0; i < InContainer.Num(); ++i)
+		{
+			const FPCGDataTypeIdentifier& CurrId = InContainer[i].template Get<0>();
+			if (!ensure(!InId.IsSameType(CurrId)))
+			{
+				UE_LOG(LogPCG, Error, TEXT("Try to register a type '%s' that is already registered"), *InId.ToString());
+				return;
+			}
+
+			// Children will come always first in the array
+			if (InId.IsChildOf(CurrId))
+			{
+				InContainer.EmplaceAt(i, InId, std::forward<Function>(InFunction));
+				bInserted = true;
+				break;
+			}
+		}
+
+		if (!bInserted)
+		{
+			InContainer.Emplace(InId, std::forward<Function>(InFunction));
+		}
+	}
+	
+	template <typename Function>
+	void RemoveFromContainer(const FPCGDataTypeIdentifier& InId, TArray<TTuple<FPCGDataTypeIdentifier, Function>>& InContainer)
+	{
+		if (const int32 Index = InContainer.IndexOfByPredicate([&InId](const auto& It) { return It.template Get<0>() == InId; }) ; Index != INDEX_NONE)
+		{
+			InContainer.RemoveAt(Index);
+		}
+	}
+
+	template <typename Function, typename ...Args>
+	decltype(auto) CallFunction(const FPCGDataTypeIdentifier& InId, const TArray<TTuple<FPCGDataTypeIdentifier, Function>>& InContainer, bool& bSuccess, Args&& ...InArgs)
+	{
+		if (const auto* It = InContainer.FindByPredicate([&InId](const auto& It) { return InId.IsChildOf(It.template Get<0>()); }))
+		{
+			bSuccess = true;
+			return It->template Get<1>()(InId, InArgs...);
+		}
+		else
+		{
+			bSuccess = false;
+			using ReturnType = std::invoke_result_t<Function, const FPCGDataTypeIdentifier&, Args...>;
+			return ReturnType{};
+		}
+	}
+
+	template <typename Function, typename ...Args>
+	decltype(auto) CallFunctionWithComposition(
+		const FPCGDataTypeRegistry& InRegistry,
+		const FPCGDataTypeIdentifier& InId,
+		const TArray<TTuple<FPCGDataTypeIdentifier, Function>>& InContainerComposition,
+		const TArray<TTuple<FPCGDataTypeIdentifier, Function>>& InContainerSingleType,
+		Args&& ...InArgs)
+	{
+		using ReturnType = std::invoke_result_t<Function, const FPCGDataTypeIdentifier&, Args...>;
+		bool bSuccess = false;
+		
+		if (InId.IsComposition())
+		{
+			ReturnType Result = PCGDataTypeRegistry::CallFunction(InId, InContainerComposition, bSuccess, InArgs...);
+			if (!bSuccess)
+			{
+				// Retry with the Union of the types
+				FPCGDataTypeIdentifier TypeUnion = InRegistry.GetIdentifiersUnion({InId});
+				if (!TypeUnion.IsValid() || !ensure(!TypeUnion.IsComposition()))
+				{
+					TypeUnion = FPCGDataTypeBaseId::Construct<FPCGDataTypeInfo>();
+				}
+
+				return PCGDataTypeRegistry::CallFunction(InId, InContainerSingleType, bSuccess, InArgs...);
+			}
+			else
+			{
+				return Result;
+			}
+		}
+		else
+		{
+			return PCGDataTypeRegistry::CallFunction(InId, InContainerSingleType, bSuccess, InArgs...);
+		}
+	}
+}
+
+void FPCGDataTypeRegistry::RegisterPinColorFunction(const FPCGDataTypeIdentifier& InId, FPCGPinColorQueryFunction InFunction)
+{
+	PCGDataTypeRegistry::AddUniqueOrdered(InId, InId.IsComposition() ? CombinationPinColors : SingleTypePinColors, MoveTemp(InFunction));
+}
+
+void FPCGDataTypeRegistry::UnregisterPinColorFunction(const FPCGDataTypeIdentifier& InId)
+{
+	PCGDataTypeRegistry::RemoveFromContainer(InId, InId.IsComposition() ? CombinationPinColors : SingleTypePinColors);
+}
+
+FLinearColor FPCGDataTypeRegistry::GetPinColor(const FPCGDataTypeIdentifier& InId) const
+{
+	return PCGDataTypeRegistry::CallFunctionWithComposition(*this, InId, CombinationPinColors, SingleTypePinColors);
+}
+
+void FPCGDataTypeRegistry::RegisterPinIconsFunction(const FPCGDataTypeIdentifier& InId, FPCGPinIconsQueryFunction InFunction)
+{
+	PCGDataTypeRegistry::AddUniqueOrdered(InId, InId.IsComposition() ? CombinationPinIcons : SingleTypePinIcons, MoveTemp(InFunction));
+}
+
+void FPCGDataTypeRegistry::UnregisterPinIconsFunction(const FPCGDataTypeIdentifier& InId)
+{
+	PCGDataTypeRegistry::RemoveFromContainer(InId, InId.IsComposition() ? CombinationPinIcons : SingleTypePinIcons);
+}
+
+TTuple<FName, FName> FPCGDataTypeRegistry::GetPinIcons(const FPCGDataTypeIdentifier& InId, const FPCGPinProperties& InProperties, bool bIsInput) const
+{
+	return PCGDataTypeRegistry::CallFunctionWithComposition(*this, InId, CombinationPinIcons, SingleTypePinIcons, InProperties, bIsInput);
+}
+#endif // WITH_EDITOR
+
 void FPCGDataTypeRegistry::RegisterKnownTypes()
 {
 	const UScriptStruct* BaseStruct = FPCGDataTypeInfo::StaticStruct();
diff --git a/Engine/Plugins/PCG/Source/PCG/Private/Tests/Registry/PCGDataTypeRegistryTests.cpp b/Engine/Plugins/PCG/Source/PCG/Private/Tests/Registry/PCGDataTypeRegistryTests.cpp
--- a/Engine/Plugins/PCG/Source/PCG/Private/Tests/Registry/PCGDataTypeRegistryTests.cpp
+++ b/Engine/Plugins/PCG/Source/PCG/Private/Tests/Registry/PCGDataTypeRegistryTests.cpp
@@ -43,6 +43,7 @@ IMPLEMENT_CUSTOM_SIMPLE_AUTOMATION_TEST(FPCGCompositionTest, FPCGDataTypeRegistr
 IMPLEMENT_CUSTOM_SIMPLE_AUTOMATION_TEST(FPCGUnicityTest, FPCGDataTypeRegistryTestBase, "Plugins.PCG.DataTypeRegistry.Unicity", PCGTestsCommon::TestFlags)
 IMPLEMENT_CUSTOM_SIMPLE_AUTOMATION_TEST(FPCGWiderTest, FPCGDataTypeRegistryTestBase, "Plugins.PCG.DataTypeRegistry.Wider", PCGTestsCommon::TestFlags)
 IMPLEMENT_CUSTOM_SIMPLE_AUTOMATION_TEST(FPCGUnionTest, FPCGDataTypeRegistryTestBase, "Plugins.PCG.DataTypeRegistry.Union", PCGTestsCommon::TestFlags)
+IMPLEMENT_CUSTOM_SIMPLE_AUTOMATION_TEST(FPCGIsChildOfTest, FPCGDataTypeRegistryTestBase, "Plugins.PCG.DataTypeRegistry.IsChildOf", PCGTestsCommon::TestFlags)
 
 bool FPCGByClassTest::RunTest(const FString& Parameters)
 {
@@ -434,7 +435,7 @@ bool FPCGCompositionTest::RunTest(const FString& Parameters)
 	{
 		const FPCGDataTypeIdentifier CompositionType = FPCGDataTypeIdentifier::Compose(TypeIdentifiers);
 		
-		const FString TypeName = FPCGDataTypeIdentifier(TypeIdentifiers).ToString();
+		const FString TypeName = FPCGDataTypeIdentifier::Construct(TypeIdentifiers).ToString();
 		const FString ExpectedTypeName = ExpectedIdentifier.ToString();
 
 		UTEST_TRUE(FString::Printf(TEXT("Composition of %s gives %s."), *TypeName, *ExpectedTypeName), CompositionType.IsSameType(ExpectedIdentifier))
@@ -451,13 +452,13 @@ bool FPCGUnicityTest::RunTest(const FString& Parameters)
 	const FPCGDataTypeIdentifier SameTypeWithClassTpl = FPCGDataTypeIdentifier::Construct<UPCGData, UPCGData>();
 	UTEST_EQUAL("Construct with same class in template gives just 1 ID", SameTypeWithClassTpl.GetIds().Num(), 1);
 
-	const FPCGDataTypeIdentifier SameTypeWithClass = FPCGDataTypeIdentifier({UPCGPointData::StaticClass(), UPCGPointData::StaticClass()});
+	const FPCGDataTypeIdentifier SameTypeWithClass = FPCGDataTypeIdentifier(UPCGPointData::StaticClass(), UPCGPointData::StaticClass());
 	UTEST_EQUAL("Construct with same class gives just 1 ID", SameTypeWithClass.GetIds().Num(), 1);
 
-	const FPCGDataTypeIdentifier SameTypeWithBaseIds = FPCGDataTypeIdentifier({FPCGDataTypeBaseId::Construct<FPCGDataTypeInfoLandscape>(), FPCGDataTypeBaseId::Construct<FPCGDataTypeInfoLandscape>()});
+	const FPCGDataTypeIdentifier SameTypeWithBaseIds = FPCGDataTypeIdentifier(FPCGDataTypeBaseId::Construct<FPCGDataTypeInfoLandscape>(), FPCGDataTypeBaseId::Construct<FPCGDataTypeInfoLandscape>());
 	UTEST_EQUAL("Construct with same base ids gives just 1 ID", SameTypeWithBaseIds.GetIds().Num(), 1);
 
-	const FPCGDataTypeIdentifier SameTypeWithIds = FPCGDataTypeIdentifier({SameTypeWithClassTpl, SameTypeWithClassTpl, SameTypeWithClassTpl});
+	const FPCGDataTypeIdentifier SameTypeWithIds = FPCGDataTypeIdentifier(SameTypeWithClassTpl, SameTypeWithClassTpl, SameTypeWithClassTpl);
 	UTEST_EQUAL("Construct with same ids gives just 1 ID", SameTypeWithIds.GetIds().Num(), 1);
 
 	const FPCGDataTypeIdentifier SameTypeWithDefaultTypeOr = SameTypeWithClass | EPCGDataType::Point;
@@ -538,4 +539,37 @@ bool FPCGUnionTest::RunTest(const FString& Parameters)
 	return true;
 }
 
+bool FPCGIsChildOfTest::RunTest(const FString& Parameters)
+{
+	TTuple<FPCGDataTypeIdentifier, FPCGDataTypeIdentifier, bool> TestValues[] =
+	{
+		{FPCGDataTypeInfo::AsId(), FPCGDataTypeInfo::AsId(), true},
+		{FPCGDataTypeInfoPoint::AsId(), FPCGDataTypeInfo::AsId(), true},
+		{FPCGDataTypeInfo::AsId(), FPCGDataTypeInfoPoint::AsId(), false},
+		{FPCGDataTypeInfoSpline::AsId(), FPCGDataTypeInfoPolyline::AsId(), true},
+		{FPCGDataTypeInfoSpline::AsId(), FPCGDataTypeInfoLandscapeSpline::AsId(), false},
+		{FPCGDataTypeIdentifier{FPCGDataTypeInfoSpline::AsId(), FPCGDataTypeInfoLandscapeSpline::AsId()}, FPCGDataTypeInfoPolyline::AsId(), true},
+		{FPCGDataTypeIdentifier{FPCGDataTypeInfoParam::AsId(), FPCGDataTypeInfoPoint::AsId()}, FPCGDataTypeIdentifier{FPCGDataTypeInfoParam::AsId(), FPCGDataTypeInfoSpatial::AsId()}, true},
+	};
+	
+	const FPCGDataTypeRegistry& Registry = FPCGModule::GetConstDataTypeRegistry();
+	
+	for (const auto& [LHS, RHS, ExpectedResult] : TestValues)
+	{
+		const bool bIsChildOf = LHS.IsChildOf(RHS);
+		const FString TestWhat = FString::Printf(TEXT("%s is child of %s."), *LHS.ToString(), *RHS.ToString());
+		
+		if (ExpectedResult)
+		{
+			UTEST_TRUE(*TestWhat, bIsChildOf)
+		}
+		else
+		{
+			UTEST_FALSE(*TestWhat, bIsChildOf)
+		}
+	}
+	
+	return true;
+}
+
 #endif // WITH_EDITOR
\ No newline at end of file
diff --git a/Engine/Plugins/PCG/Source/PCG/Public/Data/Registry/PCGDataType.h b/Engine/Plugins/PCG/Source/PCG/Public/Data/Registry/PCGDataType.h
--- a/Engine/Plugins/PCG/Source/PCG/Public/Data/Registry/PCGDataType.h
+++ b/Engine/Plugins/PCG/Source/PCG/Public/Data/Registry/PCGDataType.h
@@ -6,6 +6,7 @@
 
 #include "PCGDataType.generated.h"
 
+struct FPCGDataTypeBaseId;
 struct FPCGDataTypeIdentifier;
 class UPCGData;
 class UPCGGraph;
@@ -26,7 +27,7 @@ class UPCGSettings;
 #define PCG_DECLARE_TYPE_INFO(DLL_EXPORT_NAME) \
 	DLL_EXPORT_NAME static TSubclassOf<UPCGData> GetStaticAssociatedClass(); \
 	DLL_EXPORT_NAME virtual TSubclassOf<UPCGData> GetAssociatedClass() const override; \
-	static FPCGDataTypeBaseId AsId() { return FPCGDataTypeBaseId(StaticStruct()); }
+	static const FPCGDataTypeBaseId& AsId() { static FPCGDataTypeBaseId Id(StaticStruct()); return Id; }
 
 #define PCG_DEFINE_TYPE_INFO(PCGDataTypeInfoClass, PCGDataClass) \
 	TSubclassOf<UPCGData> PCGDataTypeInfoClass::GetStaticAssociatedClass() { return PCGDataClass::StaticClass(); } \
@@ -38,7 +39,7 @@ class UPCGSettings;
 
 #define PCG_ASSIGN_TYPE_INFO(PCGDataTypeInfoClass) \
 	using TypeInfo = PCGDataTypeInfoClass; \
-	virtual FPCGDataTypeBaseId GetDataTypeId() const override { return FPCGDataTypeBaseId::Construct<PCGDataTypeInfoClass>(); }
+	virtual const FPCGDataTypeBaseId& GetDataTypeId() const override { return PCGDataTypeInfoClass::AsId(); }
 
 // Convenient macro for default types in PCG (that are defined in EPCGDataType)
 #define PCG_DECLARE_LEGACY_TYPE_INFO(LegacyDataType) \
@@ -63,6 +64,8 @@ struct FPCGDataTypeInfo
 
 	virtual ~FPCGDataTypeInfo() = default;
 
+	static const FPCGDataTypeBaseId& AsId();
+
 	PCG_API static TSubclassOf<UPCGData> GetStaticAssociatedClass();
 	PCG_API virtual TSubclassOf<UPCGData> GetAssociatedClass() const;
 
@@ -75,7 +78,7 @@ struct FPCGDataTypeInfo
 
 #if WITH_EDITOR
 	PCG_API virtual FText GetTooltip() const;
-	
+
 	// If the type should be hidden to the user.
 	PCG_API virtual bool Hidden() const;
 #endif // WITH_EDITOR
@@ -91,11 +94,11 @@ struct FPCGDataTypeBaseId
 {
 	GENERATED_BODY()
 
+	PCG_API FPCGDataTypeBaseId(const UScriptStruct* InStruct = nullptr);
+
 	PCG_API static FPCGDataTypeBaseId MakeFromLegacyType(EPCGDataType LegacyType);
 	PCG_API EPCGDataType AsLegacyType() const;
 
-	PCG_API FPCGDataTypeBaseId(const UScriptStruct* InStruct = nullptr);
-
 	template <typename T> requires std::is_base_of_v<FPCGDataTypeInfo, T>
 	static FPCGDataTypeBaseId Construct()
 	{
@@ -116,7 +119,7 @@ struct FPCGDataTypeBaseId
 	PCG_API bool IsChildOf(const FPCGDataTypeBaseId& Other) const;
 
 	template <typename T> requires std::is_base_of_v<FPCGDataTypeInfo, T>
-	bool IsChildOf()
+	bool IsChildOf() const
 	{
 		return IsChildOf(FPCGDataTypeBaseId(T::StaticStruct()));
 	}
@@ -137,3 +140,9 @@ struct FPCGDataTypeInfoOther: public FPCGDataTypeInfo
 
 	PCG_DECLARE_LEGACY_TYPE_INFO(EPCGDataType::Other);
 };
+
+inline const FPCGDataTypeBaseId& FPCGDataTypeInfo::AsId()
+{
+	static FPCGDataTypeBaseId Id(StaticStruct());
+	return Id;
+}
\ No newline at end of file
diff --git a/Engine/Plugins/PCG/Source/PCG/Public/Data/Registry/PCGDataTypeIdentifier.h b/Engine/Plugins/PCG/Source/PCG/Public/Data/Registry/PCGDataTypeIdentifier.h
--- a/Engine/Plugins/PCG/Source/PCG/Public/Data/Registry/PCGDataTypeIdentifier.h
+++ b/Engine/Plugins/PCG/Source/PCG/Public/Data/Registry/PCGDataTypeIdentifier.h
@@ -12,6 +12,14 @@
 
 #include "PCGDataTypeIdentifier.generated.h"
 
+// A concept to know what we can construct a FPCGDataTypeIdentifier from.
+template <typename T>
+concept CPCGDataTypeIdentifierConstructible =
+	std::is_same_v<std::decay_t<T>, FPCGDataTypeIdentifier>
+	|| std::is_same_v<std::decay_t<T>, FPCGDataTypeBaseId>
+	|| std::is_convertible_v<std::decay_t<T>, TSubclassOf<UPCGData>>
+	|| std::is_same_v<std::decay_t<T>, EPCGDataType>;
+
 class FPCGDataTypeRegistry;
 class UPCGData;
 
@@ -24,21 +32,21 @@ struct FPCGDataTypeIdentifier
 
 	FPCGDataTypeIdentifier() = default;
 
-	FPCGDataTypeIdentifier(FPCGDataTypeBaseId Id)
+	// Can be constructed from another FPCGDataTypeIdentifier, FPCGDataTypeBaseId, TSubclassOf<UPCGData> or EPCGDataType
+	template <typename ...T> requires (CPCGDataTypeIdentifierConstructible<T> && ...)
+	FPCGDataTypeIdentifier(T const& ...InIds);
+
+    /** Composition constructors. Not default constructors because of ambiguity with the variadic constructor using initializer list. */
+    PCG_API static FPCGDataTypeIdentifier Construct(TConstArrayView<FPCGDataTypeIdentifier> InIds);
+    PCG_API static FPCGDataTypeIdentifier Construct(TConstArrayView<TSubclassOf<UPCGData>> Classes);
+
+	template <typename ...T> requires (std::is_base_of_v<UPCGData, T> && ...)
+	static FPCGDataTypeIdentifier Construct()
 	{
-		Ids.Emplace(Id);
+		return FPCGDataTypeIdentifier(T::StaticClass()...);
 	}
-	
-	PCG_API explicit FPCGDataTypeIdentifier(TSubclassOf<UPCGData> Class);
-
-    /** Composition ctors. */
-    PCG_API explicit FPCGDataTypeIdentifier(TConstArrayView<FPCGDataTypeIdentifier> InIds);
-    PCG_API explicit FPCGDataTypeIdentifier(TConstArrayView<TSubclassOf<UPCGData>> Classes);
 
 	// ------- Begin deprecation with EPCGDataType --------
-	// Not explicit to facilitate deprecation
-	PCG_API FPCGDataTypeIdentifier(EPCGDataType LegacyType);
-	
 	// @todo_pcg: Implicit for now, should be explicit before next release with deprecation path
 	UE_DEPRECATED(5.7, "FPCGDataTypeIdentifier is replacing EPCGDataType")
 	PCG_API operator EPCGDataType() const;
@@ -91,15 +99,6 @@ struct FPCGDataTypeIdentifier
 	// CAREFUL: This can't be used before the Data Type Registry is done initializing (so can't be used before the Engine Post Init).
 	PCG_API FPCGDataTypeIdentifier Compose(const FPCGDataTypeIdentifier& Other);
 	PCG_API static FPCGDataTypeIdentifier Compose(TConstArrayView<FPCGDataTypeIdentifier> IDs);
-	
-	template <typename ...T> requires (std::is_base_of_v<UPCGData, T> && ...)
-	static FPCGDataTypeIdentifier Construct()
-	{
-		FPCGDataTypeIdentifier Result;
-		Result.Ids.Reserve(sizeof...(T));
-		(Result.Ids.AddUnique(FPCGDataTypeIdentifier(T::StaticClass()).GetId()), ...);
-		return Result;
-	}
 
 	/** Can't be called when the identifier is a composition. Need to use GetIds in that case. */
 	PCG_API FPCGDataTypeBaseId GetId() const;
@@ -123,6 +122,9 @@ struct FPCGDataTypeIdentifier
 	/** Return true if the IDs and subtypes are the same. */
 	PCG_API bool IsIdentical(const FPCGDataTypeIdentifier& Other) const;
 
+	/** A type is a child of another type if and only if their intersection gives self. */
+	PCG_API bool IsChildOf(const FPCGDataTypeIdentifier& Other) const;
+
 	/** To support previously saved types, that used EPCGDataType, we need to define this function to de-serialize the new class using the old. And add a trait (see below). */
 	PCG_API bool SerializeFromMismatchedTag(const struct FPropertyTag& Tag, FStructuredArchive::FSlot Slot);
 
@@ -130,7 +132,11 @@ struct FPCGDataTypeIdentifier
 
 	PCG_API FString ToString() const;
 	PCG_API FText ToDisplayText() const;
-	
+
+private:
+	PCG_API void ComposeFromLegacyType(EPCGDataType LegacyType);
+	PCG_API void ComposeFromPCGDataSubclass(TSubclassOf<UPCGData> Subclass);
+
 private:
 	UPROPERTY(BlueprintReadOnly, VisibleAnywhere, Category = "PCG", meta = (AllowPrivateAccess))
 	TArray<FPCGDataTypeBaseId> Ids;
@@ -177,4 +183,35 @@ public:
 	/** Return true if the IDs and subtypes are the same. */
 	UFUNCTION(BlueprintCallable, Category = "PCG|PCGDataTypeIdentifier", meta = (ScriptMethod))
 	static bool IsIdentical(UPARAM(Ref) const FPCGDataTypeIdentifier& This, UPARAM(Ref) const FPCGDataTypeIdentifier& Other) { return This.IsIdentical(Other); }
-};
\ No newline at end of file
+};
+
+template <typename ...T> requires (CPCGDataTypeIdentifierConstructible<T> && ...)
+FPCGDataTypeIdentifier::FPCGDataTypeIdentifier(T const& ...InIds)
+{
+	([this, &InIds]()
+	{
+		if constexpr (std::is_same_v<T, FPCGDataTypeIdentifier>)
+		{
+			*this |= InIds;
+		}
+		else if constexpr (std::is_same_v<T, FPCGDataTypeBaseId>)
+		{
+			if (InIds.IsValid())
+			{
+				Ids.AddUnique(InIds);
+			}
+		}
+		else if constexpr (std::is_convertible_v<T, TSubclassOf<UPCGData>>)
+		{
+			ComposeFromPCGDataSubclass(InIds);
+		}
+		else if constexpr(std::is_same_v<T, EPCGDataType>)
+		{
+			ComposeFromLegacyType(InIds);
+		}
+		else
+		{
+			static_assert(!std::is_same_v<T, T>, "Incompatible argument");
+		}
+	}(), ...);
+}
\ No newline at end of file
diff --git a/Engine/Plugins/PCG/Source/PCG/Public/Data/Registry/PCGDataTypeRegistry.h b/Engine/Plugins/PCG/Source/PCG/Public/Data/Registry/PCGDataTypeRegistry.h
--- a/Engine/Plugins/PCG/Source/PCG/Public/Data/Registry/PCGDataTypeRegistry.h
+++ b/Engine/Plugins/PCG/Source/PCG/Public/Data/Registry/PCGDataTypeRegistry.h
@@ -8,6 +8,7 @@
 #include "Data/Registry/PCGDataTypeIdentifier.h"
 #include "StructUtils/InstancedStruct.h"
 
+struct FPCGPinProperties;
 class FPCGDataTypeRegistry;
 // Friend for tests
 class FPCGDataTypeRegistryTestBase;
@@ -99,6 +100,22 @@ public:
 	
 	/** Do a depth search in the hierarchy tree. Functor should return ESearchCommand to continue/stop. ExpandAndContinue will continue searching deeper in a branch. */
 	PCG_API void HierarchyDepthSearch(TFunctionRef<ESearchCommand(const FPCGDataTypeBaseId& Id, int32 Depth)> Callback) const;
+
+#if WITH_EDITOR
+	using FPCGPinColorQueryFunction = TFunction<FLinearColor(const FPCGDataTypeIdentifier& InId)>;
+	/** Can register a custom pin color for a given Id. Can be a composition or a single type. Custom Types will be ignored. */
+	PCG_API void RegisterPinColorFunction(const FPCGDataTypeIdentifier& InId, FPCGPinColorQueryFunction InFunction);
+	PCG_API void UnregisterPinColorFunction(const FPCGDataTypeIdentifier& InId);
+
+	PCG_API FLinearColor GetPinColor(const FPCGDataTypeIdentifier& InId) const;
+
+	using FPCGPinIconsQueryFunction = TFunction<TTuple<FName, FName>(const FPCGDataTypeIdentifier& InId, const FPCGPinProperties& InProperties, bool bIsInput)>;
+	/** Can register a custom pin icon for a given Id. Can be a composition or a single type. Custom Types will be ignored. Callback should return Slate Brush names for connected and disconnected pins. */
+	PCG_API void RegisterPinIconsFunction(const FPCGDataTypeIdentifier& InId, FPCGPinIconsQueryFunction InFunction);
+	PCG_API void UnregisterPinIconsFunction(const FPCGDataTypeIdentifier& InId);
+	
+	PCG_API TTuple<FName, FName> GetPinIcons(const FPCGDataTypeIdentifier& InId, const FPCGPinProperties& InProperties, bool bIsInput) const;
+#endif // WITH_EDITOR
 	
 private:
 	friend FPCGDataTypeRegistryTestBase;
@@ -112,5 +129,13 @@ private:
 	
 	TMap<FPCGDataTypeBaseId, TInstancedStruct<FPCGDataTypeInfo>> Mapping;
 	TUniquePtr<PCGDataTypeHierarchy::FTree> HierarchyTree;
+
+#if WITH_EDITOR
+	TArray<TTuple<FPCGDataTypeIdentifier, FPCGPinColorQueryFunction>> CombinationPinColors;
+	TArray<TTuple<FPCGDataTypeIdentifier, FPCGPinColorQueryFunction>> SingleTypePinColors;
+
+	TArray<TTuple<FPCGDataTypeIdentifier, FPCGPinIconsQueryFunction>> CombinationPinIcons;
+	TArray<TTuple<FPCGDataTypeIdentifier, FPCGPinIconsQueryFunction>> SingleTypePinIcons;
+#endif // WITH_EDITOR
 };
 
diff --git a/Engine/Plugins/PCG/Source/PCG/Public/PCGData.h b/Engine/Plugins/PCG/Source/PCG/Public/PCGData.h
--- a/Engine/Plugins/PCG/Source/PCG/Public/PCGData.h
+++ b/Engine/Plugins/PCG/Source/PCG/Public/PCGData.h
@@ -60,7 +60,7 @@ public:
 	virtual EPCGDataType GetDataType() const { return EPCGDataType::Any; }
 
 	using TypeInfo = FPCGDataTypeInfo;
-	virtual FPCGDataTypeBaseId GetDataTypeId() const { return FPCGDataTypeBaseId::Construct<FPCGDataTypeInfo>(); }
+	virtual const FPCGDataTypeBaseId& GetDataTypeId() const { return FPCGDataTypeInfo::AsId(); }
 
 	/** Returns a Crc for this and any connected data. */
 	UE_API FPCGCrc GetOrComputeCrc(bool bFullDataCrc) const;
diff --git a/Engine/Plugins/PCG/Source/PCGEditor/Private/Nodes/SPCGEditorGraphNode.cpp b/Engine/Plugins/PCG/Source/PCGEditor/Private/Nodes/SPCGEditorGraphNode.cpp
--- a/Engine/Plugins/PCG/Source/PCGEditor/Private/Nodes/SPCGEditorGraphNode.cpp
+++ b/Engine/Plugins/PCG/Source/PCGEditor/Private/Nodes/SPCGEditorGraphNode.cpp
@@ -209,64 +209,18 @@ void SPCGEditorGraphNode::AddPin(const TSharedRef<SGraphPin>& PinToAdd)
 
 		if(UPCGPin* Pin = (bIsInPin ? PCGNode->GetInputPin(PinName) : PCGNode->GetOutputPin(PinName)))
 		{
-			const bool bIsMultiData = Pin->Properties.bAllowMultipleData;
-			const bool bIsMultiConnections = Pin->AllowsMultipleConnections();
-
-			const FPCGDataTypeIdentifier CurrentType = Pin->GetCurrentTypesID();
-
-			// Check for special types
-			if (CurrentType.IsSameType(FPCGDataTypeIdentifier{EPCGDataType::Param}))
+			if (!ensure(Pin))
 			{
-				const FSlateBrush* ConnectedBrush = FPCGEditorStyle::Get().GetBrush(bIsInPin ? PCGEditorStyleConstants::Pin_Param_IN_C : PCGEditorStyleConstants::Pin_Param_OUT_C);
-				const FSlateBrush* DisconnectedBrush = FPCGEditorStyle::Get().GetBrush(bIsInPin ? PCGEditorStyleConstants::Pin_Param_IN_DC : PCGEditorStyleConstants::Pin_Param_OUT_DC);
-
-				PinToAdd->SetCustomPinIcon(ConnectedBrush, DisconnectedBrush);
+				return;
 			}
-			else if (CurrentType.IsSameType(FPCGDataTypeIdentifier{EPCGDataType::Spatial}))
-			{
-				const FSlateBrush* ConnectedBrush = FPCGEditorStyle::Get().GetBrush(bIsInPin ? PCGEditorStyleConstants::Pin_Composite_IN_C : PCGEditorStyleConstants::Pin_Composite_OUT_C);
-				const FSlateBrush* DisconnectedBrush = FPCGEditorStyle::Get().GetBrush(bIsInPin ? PCGEditorStyleConstants::Pin_Composite_IN_DC : PCGEditorStyleConstants::Pin_Composite_OUT_DC);
+			
+			const FPCGDataTypeIdentifier PinType = Pin->GetCurrentTypesID();
+			const TTuple<FName, FName> PinBrushes = FPCGModule::GetConstDataTypeRegistry().GetPinIcons(PinType, Pin->Properties, bIsInPin);
 
-				PinToAdd->SetCustomPinIcon(ConnectedBrush, DisconnectedBrush);
-			}
-			else if (Pin->Properties.Usage == EPCGPinUsage::DependencyOnly)
-			{
-				const FSlateBrush* ConnectedBrush = FPCGEditorStyle::Get().GetBrush(PCGEditorStyleConstants::Pin_Graph_Dependency_C);
-				const FSlateBrush* DisconnectedBrush = FPCGEditorStyle::Get().GetBrush(PCGEditorStyleConstants::Pin_Graph_Dependency_DC);
+			const FSlateBrush* ConnectedBrush = FPCGEditorStyle::Get().GetBrush(PinBrushes.Get<0>());
+			const FSlateBrush* DisconnectedBrush = FPCGEditorStyle::Get().GetBrush(PinBrushes.Get<1>());
 
-				PinToAdd->SetCustomPinIcon(ConnectedBrush, DisconnectedBrush);
-			}
-			else
-			{
-				// Node outputs are always single collection (SC).
-				static const FName* PinBrushes[] =
-				{
-					&PCGEditorStyleConstants::Pin_SD_SC_IN_C,
-					&PCGEditorStyleConstants::Pin_SD_SC_IN_DC,
-					&PCGEditorStyleConstants::Pin_SD_MC_IN_C,
-					&PCGEditorStyleConstants::Pin_SD_MC_IN_DC,
-					&PCGEditorStyleConstants::Pin_MD_SC_IN_C,
-					&PCGEditorStyleConstants::Pin_MD_SC_IN_DC,
-					&PCGEditorStyleConstants::Pin_MD_MC_IN_C,
-					&PCGEditorStyleConstants::Pin_MD_MC_IN_DC,
-					&PCGEditorStyleConstants::Pin_SD_SC_OUT_C,
-					&PCGEditorStyleConstants::Pin_SD_SC_OUT_DC,
-					&PCGEditorStyleConstants::Pin_SD_SC_OUT_C,
-					&PCGEditorStyleConstants::Pin_SD_SC_OUT_DC,
-					&PCGEditorStyleConstants::Pin_MD_SC_OUT_C,
-					&PCGEditorStyleConstants::Pin_MD_SC_OUT_DC,
-					&PCGEditorStyleConstants::Pin_MD_SC_OUT_C,
-					&PCGEditorStyleConstants::Pin_MD_SC_OUT_DC
-				};
-
-				const int32 ConnectedIndex = (bIsInPin ? 0 : 8) + (bIsMultiData ? 4 : 0) + (bIsMultiConnections ? 2 : 0);
-				const int32 DisconnectedIndex = ConnectedIndex + 1;
-
-				const FSlateBrush* ConnectedBrush = FPCGEditorStyle::Get().GetBrush(*PinBrushes[ConnectedIndex]);
-				const FSlateBrush* DisconnectedBrush = FPCGEditorStyle::Get().GetBrush(*PinBrushes[DisconnectedIndex]);
-
-				PinToAdd->SetCustomPinIcon(ConnectedBrush, DisconnectedBrush);
-			}
+			PinToAdd->SetCustomPinIcon(ConnectedBrush, DisconnectedBrush);
 		}
 	}
 
diff --git a/Engine/Plugins/PCG/Source/PCGEditor/Private/PCGEditorModule.cpp b/Engine/Plugins/PCG/Source/PCGEditor/Private/PCGEditorModule.cpp
--- a/Engine/Plugins/PCG/Source/PCGEditor/Private/PCGEditorModule.cpp
+++ b/Engine/Plugins/PCG/Source/PCGEditor/Private/PCGEditorModule.cpp
@@ -11,14 +11,17 @@
 #include "PCGVolumeFactory.h"
 #include "Data/PCGCollisionShapeData.h"
 #include "Data/PCGCollisionWrapperData.h"
+#include "Data/PCGDynamicMeshData.h"
 #include "Data/PCGLandscapeData.h"
 #include "Data/PCGPolygon2DData.h"
 #include "Data/PCGPolygon2DInteriorData.h"
 #include "Data/PCGPrimitiveData.h"
+#include "Data/PCGRenderTargetData.h"
 #include "Data/PCGSpatialData.h"
 #include "Data/PCGSplineData.h"
 #include "Data/PCGStaticMeshResourceData.h"
 #include "Data/PCGTextureData.h"
+#include "Data/PCGVirtualTextureData.h"
 #include "Data/PCGVolumeData.h"
 #include "Grid/PCGPartitionActor.h"
 #include "WorldPartitionBuilder/PCGWorldPartitionBuilder.h"
@@ -85,6 +88,7 @@ void FPCGEditorModule::StartupModule()
 	RegisterDetailsCustomizations();
 	RegisterSettings();
 	RegisterPCGDataVisualizations();
+	RegisterPinColorAndIcons();
 
 	UToolMenus::RegisterStartupCallback(FSimpleMulticastDelegate::FDelegate::CreateRaw(this, &FPCGEditorModule::RegisterMenuExtensions));
 
@@ -123,6 +127,7 @@ void FPCGEditorModule::ShutdownModule()
 {
 	IPCGEditorModule::SetEditorModule(nullptr);
 
+	UnregisterPinColorAndIcons();
 	UnregisterPCGDataVisualizations();
 	UnregisterSettings();
 	UnregisterDetailsCustomizations();
@@ -832,6 +837,113 @@ void FPCGEditorModule::UnregisterPCGDataVisualizations()
 	FPCGModule::GetMutablePCGDataVisualizationRegistry().InternalRegistry.Empty();
 }
 
+void FPCGEditorModule::RegisterPinColorAndIcons()
+{
+	FPCGDataTypeRegistry& InRegistry = FPCGModule::GetMutableDataTypeRegistry();
+
+	// Pin Colors
+	InRegistry.RegisterPinColorFunction(FPCGDataTypeInfoPoint::AsId(), [](const FPCGDataTypeIdentifier&) { return GetDefault<UPCGEditorSettings>()->PointDataPinColor; });
+	InRegistry.RegisterPinColorFunction(FPCGDataTypeInfoParam::AsId(), [](const FPCGDataTypeIdentifier&) { return GetDefault<UPCGEditorSettings>()->ParamDataPinColor; });
+	InRegistry.RegisterPinColorFunction(FPCGDataTypeInfoConcrete::AsId(), [](const FPCGDataTypeIdentifier&) { return GetDefault<UPCGEditorSettings>()->ConcreteDataPinColor; });
+	InRegistry.RegisterPinColorFunction(FPCGDataTypeInfoPolyline::AsId(), [](const FPCGDataTypeIdentifier&) { return GetDefault<UPCGEditorSettings>()->PolyLineDataPinColor; });
+	InRegistry.RegisterPinColorFunction(FPCGDataTypeInfoLandscape::AsId(), [](const FPCGDataTypeIdentifier&) { return GetDefault<UPCGEditorSettings>()->LandscapeDataPinColor; });
+	InRegistry.RegisterPinColorFunction(FPCGDataTypeInfoBaseTexture::AsId(), [](const FPCGDataTypeIdentifier&) { return GetDefault<UPCGEditorSettings>()->BaseTextureDataPinColor; });
+	InRegistry.RegisterPinColorFunction(FPCGDataTypeInfoTexture::AsId(), [](const FPCGDataTypeIdentifier&) { return GetDefault<UPCGEditorSettings>()->TextureDataPinColor; });
+	InRegistry.RegisterPinColorFunction(FPCGDataTypeInfoRenderTarget::AsId(), [](const FPCGDataTypeIdentifier&) { return GetDefault<UPCGEditorSettings>()->RenderTargetDataPinColor; });
+	InRegistry.RegisterPinColorFunction(FPCGDataTypeInfoVirtualTexture::AsId(), [](const FPCGDataTypeIdentifier&) { return GetDefault<UPCGEditorSettings>()->VirtualTextureDataPinColor; });
+	InRegistry.RegisterPinColorFunction(FPCGDataTypeInfoSurface::AsId(), [](const FPCGDataTypeIdentifier&) { return GetDefault<UPCGEditorSettings>()->SurfaceDataPinColor; });
+	InRegistry.RegisterPinColorFunction(FPCGDataTypeInfoVolume::AsId(), [](const FPCGDataTypeIdentifier&) { return GetDefault<UPCGEditorSettings>()->VolumeDataPinColor; });
+	InRegistry.RegisterPinColorFunction(FPCGDataTypeInfoPrimitive::AsId(), [](const FPCGDataTypeIdentifier&) { return GetDefault<UPCGEditorSettings>()->PrimitiveDataPinColor; });
+	InRegistry.RegisterPinColorFunction(FPCGDataTypeInfoDynamicMesh::AsId(), [](const FPCGDataTypeIdentifier&) { return GetDefault<UPCGEditorSettings>()->DynamicMeshPinColor; });
+	InRegistry.RegisterPinColorFunction(FPCGDataTypeInfoSpatial::AsId(), [](const FPCGDataTypeIdentifier&) { return GetDefault<UPCGEditorSettings>()->SpatialDataPinColor; });
+	InRegistry.RegisterPinColorFunction(FPCGDataTypeInfoStaticMeshResource::AsId(), [](const FPCGDataTypeIdentifier&) { return GetDefault<UPCGEditorSettings>()->StaticMeshResourcePinColor; });
+	InRegistry.RegisterPinColorFunction(FPCGDataTypeInfoOther::AsId(), [](const FPCGDataTypeIdentifier&) { return GetDefault<UPCGEditorSettings>()->UnknownDataPinColor; });
+	InRegistry.RegisterPinColorFunction(FPCGDataTypeInfo::AsId(), [](const FPCGDataTypeIdentifier&) { return GetDefault<UPCGEditorSettings>()->DefaultPinColor; });
+
+	// PinIcons. Just have catch all as most of it will be the same icons.
+	InRegistry.RegisterPinIconsFunction(FPCGDataTypeInfo::AsId(), [](const FPCGDataTypeIdentifier& InId, const FPCGPinProperties& InProperties, const bool bIsInput) -> TTuple<FName, FName>
+	{
+		const bool bIsMultiData = InProperties.bAllowMultipleData;
+		const bool bIsMultiConnections = InProperties.AllowsMultipleConnections();
+
+		FName ConnectedBrushName = NAME_None;
+		FName DisconnectedBrushName = NAME_None;
+
+		if (InId.IsSameType(FPCGDataTypeInfoParam::AsId()))
+		{
+			ConnectedBrushName = bIsInput ? PCGEditorStyleConstants::Pin_Param_IN_C : PCGEditorStyleConstants::Pin_Param_OUT_C;
+			DisconnectedBrushName = bIsInput ? PCGEditorStyleConstants::Pin_Param_IN_DC : PCGEditorStyleConstants::Pin_Param_OUT_DC;
+		}
+		else if (InId.IsSameType(FPCGDataTypeInfoSpatial::AsId()))
+		{
+			ConnectedBrushName = bIsInput ? PCGEditorStyleConstants::Pin_Composite_IN_C : PCGEditorStyleConstants::Pin_Composite_OUT_C;
+			DisconnectedBrushName = bIsInput ? PCGEditorStyleConstants::Pin_Composite_IN_DC : PCGEditorStyleConstants::Pin_Composite_OUT_DC;
+		}
+		else if (InProperties.Usage == EPCGPinUsage::DependencyOnly)
+		{
+			ConnectedBrushName = PCGEditorStyleConstants::Pin_Graph_Dependency_C;
+			DisconnectedBrushName = PCGEditorStyleConstants::Pin_Graph_Dependency_DC;
+		}
+		else
+		{
+			// Node outputs are always single collection (SC).
+			static const FName* PinBrushes[] =
+			{
+				&PCGEditorStyleConstants::Pin_SD_SC_IN_C,
+				&PCGEditorStyleConstants::Pin_SD_SC_IN_DC,
+				&PCGEditorStyleConstants::Pin_SD_MC_IN_C,
+				&PCGEditorStyleConstants::Pin_SD_MC_IN_DC,
+				&PCGEditorStyleConstants::Pin_MD_SC_IN_C,
+				&PCGEditorStyleConstants::Pin_MD_SC_IN_DC,
+				&PCGEditorStyleConstants::Pin_MD_MC_IN_C,
+				&PCGEditorStyleConstants::Pin_MD_MC_IN_DC,
+				&PCGEditorStyleConstants::Pin_SD_SC_OUT_C,
+				&PCGEditorStyleConstants::Pin_SD_SC_OUT_DC,
+				&PCGEditorStyleConstants::Pin_SD_SC_OUT_C,
+				&PCGEditorStyleConstants::Pin_SD_SC_OUT_DC,
+				&PCGEditorStyleConstants::Pin_MD_SC_OUT_C,
+				&PCGEditorStyleConstants::Pin_MD_SC_OUT_DC,
+				&PCGEditorStyleConstants::Pin_MD_SC_OUT_C,
+				&PCGEditorStyleConstants::Pin_MD_SC_OUT_DC
+			};
+
+			const int32 ConnectedIndex = (bIsInput ? 0 : 8) + (bIsMultiData ? 4 : 0) + (bIsMultiConnections ? 2 : 0);
+			const int32 DisconnectedIndex = ConnectedIndex + 1;
+
+			ConnectedBrushName = *PinBrushes[ConnectedIndex];
+			DisconnectedBrushName = *PinBrushes[DisconnectedIndex];
+		}
+
+		return {ConnectedBrushName, DisconnectedBrushName};
+	});
+}
+
+void FPCGEditorModule::UnregisterPinColorAndIcons()
+{
+	FPCGDataTypeRegistry& InRegistry = FPCGModule::GetMutableDataTypeRegistry();
+	
+	InRegistry.UnregisterPinColorFunction(FPCGDataTypeInfo::AsId());
+	InRegistry.UnregisterPinColorFunction(FPCGDataTypeInfoOther::AsId());
+	InRegistry.UnregisterPinColorFunction(FPCGDataTypeInfoStaticMeshResource::AsId());
+	InRegistry.UnregisterPinColorFunction(FPCGDataTypeInfoSpatial::AsId());
+	InRegistry.UnregisterPinColorFunction(FPCGDataTypeInfoDynamicMesh::AsId());
+	InRegistry.UnregisterPinColorFunction(FPCGDataTypeInfoPrimitive::AsId());
+	InRegistry.UnregisterPinColorFunction(FPCGDataTypeInfoVolume::AsId());
+	InRegistry.UnregisterPinColorFunction(FPCGDataTypeInfoSurface::AsId());
+	InRegistry.UnregisterPinColorFunction(FPCGDataTypeInfoVirtualTexture::AsId());
+	InRegistry.UnregisterPinColorFunction(FPCGDataTypeInfoRenderTarget::AsId());
+	InRegistry.UnregisterPinColorFunction(FPCGDataTypeInfoTexture::AsId());
+	InRegistry.UnregisterPinColorFunction(FPCGDataTypeInfoBaseTexture::AsId());
+	InRegistry.UnregisterPinColorFunction(FPCGDataTypeInfoLandscape::AsId());
+	InRegistry.UnregisterPinColorFunction(FPCGDataTypeInfoPolyline::AsId());
+	InRegistry.UnregisterPinColorFunction(FPCGDataTypeInfoConcrete::AsId());
+	InRegistry.UnregisterPinColorFunction(FPCGDataTypeInfoParam::AsId());
+	InRegistry.UnregisterPinColorFunction(FPCGDataTypeInfoPoint::AsId());
+
+	// PinIcons. Just have catch all as most of it will be the same icons.
+	InRegistry.UnregisterPinIconsFunction(FPCGDataTypeInfo::AsId());
+}
+
 void FPCGEditorModule::OnScheduleGraph(const FPCGStackContext& StackContext)
 {
 	TRACE_CPUPROFILER_EVENT_SCOPE(UPCGSubsystem::OnScheduleGraph);
diff --git a/Engine/Plugins/PCG/Source/PCGEditor/Private/Pins/SPCGEditorGraphNodePin.cpp b/Engine/Plugins/PCG/Source/PCGEditor/Private/Pins/SPCGEditorGraphNodePin.cpp
--- a/Engine/Plugins/PCG/Source/PCGEditor/Private/Pins/SPCGEditorGraphNodePin.cpp
+++ b/Engine/Plugins/PCG/Source/PCGEditor/Private/Pins/SPCGEditorGraphNodePin.cpp
@@ -3,6 +3,7 @@
 #include "SPCGEditorGraphNodePin.h"
 
 #include "PCGEditorStyle.h"
+#include "PCGModule.h"
 #include "PCGNode.h"
 #include "PCGPin.h"
 #include "PCGSettings.h"
@@ -306,10 +307,19 @@ void SPCGEditorGraphNodePin::ApplyUnusedPinStyle(FSlateColor& InOutColor) const
 	}
 }
 
-// Adapted from SGraphPin::GetPinColor
 FSlateColor SPCGEditorGraphNodePin::GetPinColor() const
 {
-	FSlateColor Color = SGraphPin::GetPinColor();
+	const UPCGPin* PCGPin = nullptr;
+	const UPCGNode* PCGNode = nullptr;
+
+	GetPCGNodeAndPin(PCGNode, PCGPin);
+
+	if (!PCGPin)
+	{
+		return SGraphPin::GetPinColor();
+	}
+	
+	FSlateColor Color{FPCGModule::GetConstDataTypeRegistry().GetPinColor(PCGPin->GetCurrentTypesID())};
 
 	ApplyUnusedPinStyle(Color);
 
diff --git a/Engine/Plugins/PCG/Source/PCGEditor/Public/PCGEditorModule.h b/Engine/Plugins/PCG/Source/PCGEditor/Public/PCGEditorModule.h
--- a/Engine/Plugins/PCG/Source/PCGEditor/Public/PCGEditorModule.h
+++ b/Engine/Plugins/PCG/Source/PCGEditor/Public/PCGEditorModule.h
@@ -59,6 +59,8 @@ protected:
 	void UnregisterSettings();
 	void RegisterPCGDataVisualizations();
 	void UnregisterPCGDataVisualizations();
+	void RegisterPinColorAndIcons();
+	void UnregisterPinColorAndIcons();
 
 	void OnLevelEditorCreated(TSharedPtr<ILevelEditor> InLevelEditor);
 	void RegisterOnEditorModeChange();
