[PCG] Fix composition retry in the PCGDataTypeRegistry

* Was not retrying on the union type

#rnx

diff --git a/Engine/Plugins/PCG/Source/PCG/Private/Data/Registry/PCGDataTypeRegistry.cpp b/Engine/Plugins/PCG/Source/PCG/Private/Data/Registry/PCGDataTypeRegistry.cpp
--- a/Engine/Plugins/PCG/Source/PCG/Private/Data/Registry/PCGDataTypeRegistry.cpp
+++ b/Engine/Plugins/PCG/Source/PCG/Private/Data/Registry/PCGDataTypeRegistry.cpp
@@ -766,7 +766,7 @@ namespace PCGDataTypeRegistry
 					TypeUnion = FPCGDataTypeBaseId::Construct<FPCGDataTypeInfo>();
 				}
 
-				return PCGDataTypeRegistry::CallFunction(InId, InContainerSingleType, bSuccess, InArgs...);
+				return PCGDataTypeRegistry::CallFunction(TypeUnion, InContainerSingleType, bSuccess, InArgs...);
 			}
 			else
 			{
