PCG Editor Mode
- Exposed Graph Parameters to Tools UI
- Changing Graph results in update
- Added "per tools setting settings" in PCGEdModeSettings. Lets you specify default graph per settings class.

diff --git a/Engine/Plugins/PCG/Source/PCGEditor/Private/Details/PCGInteractiveToolSettingsDetails.cpp b/Engine/Plugins/PCG/Source/PCGEditor/Private/Details/PCGInteractiveToolSettingsDetails.cpp
new file mode 100644
--- /dev/null
+++ b/Engine/Plugins/PCG/Source/PCGEditor/Private/Details/PCGInteractiveToolSettingsDetails.cpp
@@ -0,0 +1,96 @@
+// Copyright Epic Games, Inc. All Rights Reserved.
+
+#include "Details/PCGInteractiveToolSettingsDetails.h"
+
+#include "PCGEditorModule.h"
+#include "PCGGraph.h"
+#include "DetailCategoryBuilder.h"
+#include "DetailLayoutBuilder.h"
+#include "DetailWidgetRow.h"
+#include "EditorMode/Tools/PCGInteractiveToolSettings.h"
+
+#define LOCTEXT_NAMESPACE "PCGGraphInstanceDetails"
+
+TSharedRef<IDetailCustomization> FPCGInteractiveToolSettingsDetails::MakeInstance()
+{
+	TSharedRef<FPCGInteractiveToolSettingsDetails> Customization = MakeShared<FPCGInteractiveToolSettingsDetails>();
+	Customization->RegisterDelegates();
+	return Customization;
+}
+
+FPCGInteractiveToolSettingsDetails::FPCGInteractiveToolSettingsDetails()
+{
+}
+
+FPCGInteractiveToolSettingsDetails::~FPCGInteractiveToolSettingsDetails()
+{
+	UnregisterDelegates();
+}
+
+void FPCGInteractiveToolSettingsDetails::RegisterDelegates()
+{
+	FCoreUObjectDelegates::OnPostObjectPropertyChanged.AddSP(this, &FPCGInteractiveToolSettingsDetails::Rebuild);
+}
+
+void FPCGInteractiveToolSettingsDetails::UnregisterDelegates()
+{
+	FCoreUObjectDelegates::OnPostObjectPropertyChanged.RemoveAll(this);
+}
+
+void FPCGInteractiveToolSettingsDetails::CustomizeDetails(IDetailLayoutBuilder& DetailBuilder)
+{	
+	TArray<TWeakObjectPtr<UObject>> ObjectsBeingCustomized;
+	DetailBuilder.GetObjectsBeingCustomized(ObjectsBeingCustomized);
+
+	for(const TWeakObjectPtr<UObject>& Object : ObjectsBeingCustomized)
+	{
+		UPCGInteractiveToolSettings* InToolSettings = Cast<UPCGInteractiveToolSettings>(Object.Get());
+		if(ensure(InToolSettings))
+		{
+			ToolSettings = InToolSettings;
+		}
+	}
+	
+	IDetailCategoryBuilder& CategoryBuilder = DetailBuilder.EditCategory("PCG");
+	if(ToolSettings.IsValid() && ToolSettings->GraphInstance != nullptr && ToolSettings->GraphInstance->GetUserParametersStruct()->GetNumPropertiesInBag() > 0)
+	{
+		TArray<TSharedRef<IPropertyHandle>> PropertyHandles;
+		CategoryBuilder.GetDefaultProperties(PropertyHandles);
+	
+		for(TSharedRef<IPropertyHandle> PropertyHandle : PropertyHandles)
+		{
+			CategoryBuilder.AddProperty(PropertyHandle);
+
+			// We add the parameters right behind the graph property
+			if(PropertyHandle->GetProperty()->GetFName() == GET_MEMBER_NAME_CHECKED(UPCGInteractiveToolSettings, TargetGraph))
+			{			
+				IDetailPropertyRow* PropertyRow = CategoryBuilder.AddExternalObjectProperty(
+				{ToolSettings->GraphInstance}, 
+				GET_MEMBER_NAME_CHECKED(UPCGGraphInstance, ParametersOverrides),
+				 EPropertyLocation::Default,
+				 FAddPropertyParams());
+		
+				PropertyRow->ShouldAutoExpand(true);
+			}
+		}		
+	}
+}
+
+void FPCGInteractiveToolSettingsDetails::CustomizeDetails(const TSharedPtr<IDetailLayoutBuilder>& DetailBuilder)
+{
+	DetailLayoutBuilder = DetailBuilder;
+	IDetailCustomization::CustomizeDetails(DetailBuilder);
+}
+
+void FPCGInteractiveToolSettingsDetails::Rebuild(UObject* Object, const FPropertyChangedChainEvent& PropertyChangedChainEvent)
+{
+	if(Object == ToolSettings)
+	{
+		if(PropertyChangedChainEvent.Property->GetFName() == GET_MEMBER_NAME_CHECKED(UPCGInteractiveToolSettings, TargetGraph))
+		{
+			DetailLayoutBuilder.Pin()->ForceRefreshDetails();
+		}
+	}
+}
+
+#undef LOCTEXT_NAMESPACE
diff --git a/Engine/Plugins/PCG/Source/PCGEditor/Private/EditorMode/Tools/Line/PCGDrawSplineTool.cpp b/Engine/Plugins/PCG/Source/PCGEditor/Private/EditorMode/Tools/Line/PCGDrawSplineTool.cpp
--- a/Engine/Plugins/PCG/Source/PCGEditor/Private/EditorMode/Tools/Line/PCGDrawSplineTool.cpp
+++ b/Engine/Plugins/PCG/Source/PCGEditor/Private/EditorMode/Tools/Line/PCGDrawSplineTool.cpp
@@ -211,112 +211,15 @@ FText UPCGDrawSplineToolSettings::GetToolTag() const
 	return PCGDrawSplineTool::Constants::ToolName;
 }
 
-bool UPCGDrawSplineToolSettings::ApplyTool(UWorld* World, const USplineComponent* ToolSplineComponent, FText& OutWarning) const
-{
-	using namespace PCG::EditorMode;
-
-	check(World);
-	check(ToolSplineComponent);
-
-	if (!TargetGraph)
-	{
-		OutWarning = LOCTEXT("InvalidGraphInstance", "[Spline Tool] Failed: Invalid Graph Instance");
-		return false;
-	}
-
-	AActor* Actor = nullptr;
-	UPCGComponent* PCGComponent = nullptr;
-
-	switch (Interaction)
-	{
-		case EPCGToolInteractionMode::CreateNewActor:
-		{
-			Actor = Actor::Spawn(World, NewActorName, FTransform::Identity);
-		}
-		break;
-		case EPCGToolInteractionMode::AddToSelectedActor:
-		{
-			if (!TargetActor.IsValid())
-			{
-				OutWarning = LOCTEXT("InvalidActorWarning", "Invalid Actor!");
-				return false;
-			}
-
-			Actor = TargetActor.Get();
-			PCGComponent = PCGComponent::Find(Actor);
-		}
-		break;
-		default:
-			checkNoEntry();
-	}
-
-	// @todo_pcg: Will overwrite the current spline. Could be an option.
-	USplineComponent* NewSpline = Spline::FindComponent(Actor);
-	if (!NewSpline)
-	{
-		NewSpline = Spline::CreateComponent(Actor, SplineName, ToolSplineComponent, /*bTransactional=*/true);
-		check(NewSpline);
-	}
-	else
-	{
-		Spline::CopySplinePoints(*ToolSplineComponent, *NewSpline, /*bTransactional=*/true);
-	}
-
-	if (TargetGraph)
-	{
-		if (!PCGComponent)
-		{
-			PCGComponent = PCGComponent::Create(Actor, NewPCGComponentName);
-			check(PCGComponent);
-		}
-
-		PCGComponent->SetGraph(TargetGraph);
-
-		if (bAutoGenerate)
-		{
-			PCGComponent->Generate();
-		}
-	}
-
-	return true;
-}
-
-FSoftObjectPath UPCGDrawSplineToolSettings::GetDefaultGraph() const
-{
-	// @todo_pcg: Once finalized, the default graph asset for this tool is returned here. Ex:
-	// return FSoftObjectPath(TEXT("/PCG/EdMode/ToolGraphs/DrawSpline/PCG_Spline.PCG_Spline"));
-	return FSoftObjectPath{};
-}
-
 void UPCGDrawSplineTool::Setup()
 {
 	using namespace PCGDrawSplineTool;
 
 	Super::Setup();
 
-	// @todo_pcg: These are set up to be re-used between edits. For now, they should be wiped every time.
-	check(!WorkingActor.IsValid())
-	check(!WorkingComponent.IsValid())
-	check(!WorkingSpline)
-
 	UPCGDrawSplineToolSettings* Settings = CastChecked<UPCGDrawSplineToolSettings>(ToolSettings);
 	CachedMinPointSpacingSquared = Settings->MinPointSpacing * Settings->MinPointSpacing;
 
-	WorkingActor = Actor::SpawnWorking(GetContextualWorld());
-	check(WorkingActor.IsValid());
-
-	// @todo_pcg: Maybe a helper to validate the graph?
-	if (!Settings->TargetGraph)
-	{
-		NotifyError(LOCTEXT("InvalidGraph", "Invalid or non-existent graph. Shutting down."));
-		GetToolManager()->ShutdownAfterError();
-		return;
-	}
-
-	WorkingComponent = PCGComponent::CreateWorking(WorkingActor.Get());
-	check(WorkingComponent.IsValid());
-	WorkingComponent->SetGraph(Settings->TargetGraph);
-
 	WorkingSpline = Spline::CreateWorkingComponent(WorkingActor.Get());
 	check(WorkingSpline);
 
@@ -351,7 +254,7 @@ void UPCGDrawSplineTool::Shutdown(const EToolShutdownType ShutdownType)
 
 	if (ShutdownType == EToolShutdownType::Accept && WorkingSpline && WorkingSpline->GetNumberOfSplinePoints() > 0)
 	{
-		ApplyToolInternal();
+		ApplyTool();
 	}
 
 	if (WorkingComponent.IsValid())
@@ -450,10 +353,80 @@ void UPCGDrawSplineTool::ApplyToolInternal() const
 	GetToolManager()->BeginUndoTransaction(LOCTEXT("DrawSplineTransactionName", "[PCG Tool] Draw Spline"));
 	const UPCGDrawSplineToolSettings* Settings = CastChecked<UPCGDrawSplineToolSettings>(ToolSettings);
 
-	FText Warnings;
-	if (!Settings->ApplyTool(GetContextualWorldChecked(), WorkingSpline.Get(), Warnings))
+	TOptional<FText> OutWarning;
+
+	UWorld* World = GetContextualWorldChecked();
+	USplineComponent* ToolSplineComponent = WorkingSpline.Get();
+
 	{
-		NotifyError(Warnings);
+		using namespace PCG::EditorMode;
+
+		check(World);
+		check(ToolSplineComponent);
+		
+		if (!Settings->TargetGraph)
+		{
+			OutWarning = LOCTEXT("InvalidGraphInstance", "[Spline Tool] Failed: Invalid Graph Instance");
+		}
+
+		AActor* Actor = nullptr;
+		UPCGComponent* PCGComponent = nullptr;
+
+		switch (Settings->Interaction)
+		{
+		case EPCGToolInteractionMode::CreateNewActor:
+			{
+				Actor = Actor::Spawn(World, Settings->NewActorName, FTransform::Identity);
+			}
+			break;
+		case EPCGToolInteractionMode::AddToSelectedActor:
+			{
+				if (!Settings->TargetActor.IsValid())
+				{
+					OutWarning = LOCTEXT("InvalidActorWarning", "Invalid Actor!");
+				}
+
+				Actor = Settings->TargetActor.Get();
+				PCGComponent = PCGComponent::Find(Actor);
+			}
+			break;
+		default:
+			checkNoEntry();
+		}
+
+		// @todo_pcg: Will overwrite the current spline. Could be an option.
+		USplineComponent* NewSpline = Spline::FindComponent(Actor);
+		if (!NewSpline)
+		{
+			NewSpline = Spline::CreateComponent(Actor, Settings->SplineName, ToolSplineComponent, /*bTransactional=*/true);
+			check(NewSpline);
+		}
+		else
+		{
+			Spline::CopySplinePoints(*ToolSplineComponent, *NewSpline, /*bTransactional=*/true);
+		}
+
+		if(Settings->TargetGraph)
+		{
+			if (!PCGComponent)
+			{
+				PCGComponent = PCGComponent::Create(Actor, Settings->NewPCGComponentName);
+				check(PCGComponent);
+			}
+
+			PCGComponent->SetGraph(Settings->TargetGraph);
+			PCGComponent->GetGraphInstance()->CopyParameterOverrides(WorkingComponent->GetGraphInstance());
+
+			if(Settings->bAutoGenerate)
+			{
+				PCGComponent->Generate();
+			}
+		}
+
+		if(OutWarning.IsSet())
+		{
+			NotifyError(OutWarning.GetValue());
+		}
 	}
 
 	GetToolManager()->EndUndoTransaction();
diff --git a/Engine/Plugins/PCG/Source/PCGEditor/Private/EditorMode/Tools/PCGInteractiveTool.cpp b/Engine/Plugins/PCG/Source/PCGEditor/Private/EditorMode/Tools/PCGInteractiveTool.cpp
--- a/Engine/Plugins/PCG/Source/PCGEditor/Private/EditorMode/Tools/PCGInteractiveTool.cpp
+++ b/Engine/Plugins/PCG/Source/PCGEditor/Private/EditorMode/Tools/PCGInteractiveTool.cpp
@@ -10,8 +10,11 @@
 
 #include UE_INLINE_GENERATED_CPP_BY_NAME(PCGInteractiveTool)
 
+#define LOCTEXT_NAMESPACE "PCGInteractiveTool"
+
 void UPCGInteractiveTool::Setup()
 {
+	using namespace PCG::EditorMode;
 	// `GetSettingsClass` must be overridden with a valid settings class for this tool.
 	if (const UClass* SettingsClass = GetSettingsClass())
 	{
@@ -29,6 +32,25 @@ void UPCGInteractiveTool::Setup()
 	check(!GetToolInfo().ToolDisplayMessage.IsEmpty())
 	check(GetToolInfo().ToolIcon)
 
+	WorkingActor = Actor::SpawnWorking(GetContextualWorld());
+	check(WorkingActor.IsValid());
+	
+	WorkingComponent = PCGComponent::CreateWorking(WorkingActor.Get());
+
+	// @todo_pcg: Maybe a helper to validate the graph?
+	if (!ToolSettings->TargetGraph)
+	{
+		NotifyError(LOCTEXT("InvalidGraph", "Invalid or non-existent graph. Shutting down."));
+		GetToolManager()->ShutdownAfterError();
+		return;
+	}
+
+	check(WorkingComponent.IsValid());
+	WorkingComponent->SetGraph(ToolSettings->TargetGraph);
+	
+	// We set the graph instance of the tool settings so that we can display its parameters
+	ToolSettings->GraphInstance = WorkingComponent->GetGraphInstance();
+	
 	Super::Setup();
 }
 
@@ -51,6 +73,17 @@ void UPCGInteractiveTool::OnTick(const float DeltaTime)
 	}
 }
 
+void UPCGInteractiveTool::OnPropertyModified(UObject* PropertySet, FProperty* Property)
+{
+	Super::OnPropertyModified(PropertySet, Property);
+
+	if(PropertySet->IsA<UPCGInteractiveToolSettings>() && Property->GetFName() == GET_MEMBER_NAME_CHECKED(UPCGInteractiveToolSettings, TargetGraph))
+	{
+		UPCGInteractiveToolSettings* Settings = Cast<UPCGInteractiveToolSettings>(PropertySet);
+		WorkingComponent->SetGraph(Settings->TargetGraph);
+	}
+}
+
 void UPCGInteractiveTool::OnRefresh()
 {
 	if (WorkingComponent.IsValid() && !WorkingComponent->IsCleaningUp() && !WorkingComponent->IsGenerating())
@@ -87,6 +120,11 @@ UWorld* UPCGInteractiveTool::GetContextualWorldChecked() const
 	return World;
 }
 
+void UPCGInteractiveTool::ApplyTool() const
+{
+	ApplyToolInternal();
+}
+
 void UPCGInteractiveTool::SetWarning(const FText& Warning) const
 {
 	WarningMessage = Warning;
@@ -110,3 +148,5 @@ void UPCGInteractiveTool::NotifyError(const FText& ErrorMessage) const
 
 	UE_LOG(LogPCGEditor, Error, TEXT("[%sTool]: %s"), *GetToolInfo().ToolDisplayName.ToString(), *ErrorMessage.ToString());
 }
+
+#undef LOCTEXT_NAMESPACE
diff --git a/Engine/Plugins/PCG/Source/PCGEditor/Private/EditorMode/Tools/PCGInteractiveToolSettings.cpp b/Engine/Plugins/PCG/Source/PCGEditor/Private/EditorMode/Tools/PCGInteractiveToolSettings.cpp
--- a/Engine/Plugins/PCG/Source/PCGEditor/Private/EditorMode/Tools/PCGInteractiveToolSettings.cpp
+++ b/Engine/Plugins/PCG/Source/PCGEditor/Private/EditorMode/Tools/PCGInteractiveToolSettings.cpp
@@ -3,9 +3,12 @@
 #include "EditorMode/Tools/PCGInteractiveToolSettings.h"
 
 #include "PCGEditorModule.h"
+#include "PCGEditorSettings.h"
 #include "EditorMode/PCGEdModeHelpers.h"
 #include "EditorMode/PCGEdModeSettings.h"
+#include "EditorMode/Tools/PCGInteractiveTool.h"
 #include "Helpers/PCGHelpers.h"
+#include "Logging/StructuredLog.h"
 
 #include UE_INLINE_GENERATED_CPP_BY_NAME(PCGInteractiveToolSettings)
 
@@ -19,6 +22,28 @@ void UPCGInteractiveToolSettings::TryLoadDefaultGraph()
 	TargetGraph = CastChecked<UPCGGraphInterface>(GetDefaultGraph().TryLoad(), ECastCheckedType::NullAllowed);
 }
 
+FSoftObjectPath UPCGInteractiveToolSettings::GetDefaultGraph() const
+{
+	const FPCGPerInteractiveToolSettingSettings* Settings = GetDefault<UPCGEditorModeSettings>()->InteractiveToolSettings.FindByPredicate([this](const FPCGPerInteractiveToolSettingSettings& SettingsCandidate)
+		{
+			return SettingsCandidate.SettingsClass == this->GetClass();
+		});
+
+	if(Settings != nullptr)
+	{
+		if(Settings->DefaultGraph.ToSoftObjectPath().IsAsset())
+		{
+			return Settings->DefaultGraph.ToSoftObjectPath();
+		}
+	
+		UE_LOGFMT(LogPCGEditor, Warning, "Could not get default graph for settings class {0} in tool {1}", this->GetClass()->GetDisplayNameText().ToString(), GetTypedOuter<UPCGInteractiveTool>()->GetClass()->GetDisplayNameText().ToString());
+		return FSoftObjectPath();
+	}
+
+	UE_LOGFMT(LogPCGEditor, Warning, "Could not find tool settings for settings class {0} in tool {1}. Please add them in the PCG Editor Settings.", this->GetClass()->GetDisplayNameText().ToString(), GetTypedOuter<UPCGInteractiveTool>()->GetClass()->GetDisplayNameText().ToString());
+	return FSoftObjectPath();
+}
+
 bool UPCGInteractiveToolSettings::GraphAssetFilter(const FAssetData& AssetData) const
 {
 	// This can be checked by subclasses by either finding a tag or more specifically checking for values.
diff --git a/Engine/Plugins/PCG/Source/PCGEditor/Private/PCGEditorModule.cpp b/Engine/Plugins/PCG/Source/PCGEditor/Private/PCGEditorModule.cpp
--- a/Engine/Plugins/PCG/Source/PCGEditor/Private/PCGEditorModule.cpp
+++ b/Engine/Plugins/PCG/Source/PCGEditor/Private/PCGEditorModule.cpp
@@ -53,6 +53,7 @@
 #include "Details/PCGGraphInstanceDetails.h"
 #include "Details/PCGInstancedPropertyBagOverrideDetails.h"
 #include "Details/PCGVolumeDetails.h"
+#include "Details/PCGInteractiveToolSettingsDetails.h"
 #include "Nodes/PCGEditorGraphNodeFactory.h"
 #include "Subsystems/PCGEditorSubsystem.h"
 
@@ -71,6 +72,7 @@
 #include "ToolMenus.h"
 #include "Editor/SceneOutliner/Public/ISceneOutliner.h"
 #include "EditorMode/PCGEdMode.h"
+#include "EditorMode/Tools/PCGInteractiveToolSettings.h"
 #include "Framework/MultiBox/MultiBoxBuilder.h"
 #include "Framework/Notifications/NotificationManager.h"
 #include "UObject/ObjectSaveContext.h"
@@ -303,6 +305,7 @@ void FPCGEditorModule::RegisterDetailsCustomizations()
 	PropertyEditor.RegisterCustomClassLayout("PCGUserParameterGetSettings", FOnGetDetailCustomizationInstance::CreateStatic(&FPCGEditableUserParameterDetails::MakeInstance));
 	PropertyEditor.RegisterCustomClassLayout("PCGCustomHLSLSettings", FOnGetDetailCustomizationInstance::CreateStatic(&FPCGCustomHLSLSettingsDetails::MakeInstance));
 	PropertyEditor.RegisterCustomClassLayout("PCGComputeSource", FOnGetDetailCustomizationInstance::CreateStatic(&FPCGComputeSourceDetails::MakeInstance));
+	PropertyEditor.RegisterCustomClassLayout("PCGInteractiveToolSettings", FOnGetDetailCustomizationInstance::CreateStatic(&FPCGInteractiveToolSettingsDetails::MakeInstance));
 
 	PropertyEditor.RegisterCustomPropertyTypeLayout("PCGAttributePropertySelector", FOnGetPropertyTypeCustomizationInstance::CreateStatic(&FPCGAttributePropertySelectorDetails::MakeInstance));
 	PropertyEditor.RegisterCustomPropertyTypeLayout("PCGAttributePropertyInputSelector", FOnGetPropertyTypeCustomizationInstance::CreateStatic(&FPCGAttributePropertySelectorDetails::MakeInstance));
@@ -319,6 +322,7 @@ void FPCGEditorModule::UnregisterDetailsCustomizations()
 {
 	if (FPropertyEditorModule* PropertyModule = FModuleManager::GetModulePtr<FPropertyEditorModule>("PropertyEditor"))
 	{
+		PropertyModule->UnregisterCustomClassLayout("PCGInteractiveToolSettings");
 		PropertyModule->UnregisterCustomClassLayout("PCGComputeSource");
 		PropertyModule->UnregisterCustomClassLayout("PCGCustomHLSLSettings");
 		PropertyModule->UnregisterCustomClassLayout("PCGUserParameterGetSettings");
diff --git a/Engine/Plugins/PCG/Source/PCGEditor/Public/Details/PCGInteractiveToolSettingsDetails.h b/Engine/Plugins/PCG/Source/PCGEditor/Public/Details/PCGInteractiveToolSettingsDetails.h
new file mode 100644
--- /dev/null
+++ b/Engine/Plugins/PCG/Source/PCGEditor/Public/Details/PCGInteractiveToolSettingsDetails.h
@@ -0,0 +1,31 @@
+// Copyright Epic Games, Inc. All Rights Reserved.
+
+#pragma once
+
+#include "IDetailCustomization.h"
+#include "Input/Reply.h" // IWYU pragma: keep
+#include "UObject/WeakObjectPtrTemplates.h"
+
+class UPCGInteractiveToolSettings;
+class UPCGGraphInstance;
+
+class FPCGInteractiveToolSettingsDetails : public IDetailCustomization
+{
+public:
+	static TSharedRef<IDetailCustomization> MakeInstance();
+
+	FPCGInteractiveToolSettingsDetails();
+	virtual ~FPCGInteractiveToolSettingsDetails() override;
+
+	void RegisterDelegates();
+	void UnregisterDelegates();
+	
+	virtual void CustomizeDetails(IDetailLayoutBuilder& DetailBuilder) override;
+	virtual void CustomizeDetails(const TSharedPtr<IDetailLayoutBuilder>& DetailBuilder) override;
+private:
+	void Rebuild(UObject* Object, const FPropertyChangedChainEvent& PropertyChangedChainEvent);
+
+private:
+	TWeakObjectPtr<UPCGInteractiveToolSettings> ToolSettings;
+	TWeakPtr<IDetailLayoutBuilder> DetailLayoutBuilder;
+};
diff --git a/Engine/Plugins/PCG/Source/PCGEditor/Public/EditorMode/PCGEdModeSettings.h b/Engine/Plugins/PCG/Source/PCGEditor/Public/EditorMode/PCGEdModeSettings.h
--- a/Engine/Plugins/PCG/Source/PCGEditor/Public/EditorMode/PCGEdModeSettings.h
+++ b/Engine/Plugins/PCG/Source/PCGEditor/Public/EditorMode/PCGEdModeSettings.h
@@ -6,6 +6,23 @@
 
 #include "PCGEdModeSettings.generated.h"
 
+class UPCGInteractiveToolSettings;
+class UPCGGraph;
+
+USTRUCT()
+struct FPCGPerInteractiveToolSettingSettings
+{
+	GENERATED_BODY()
+	
+	UPROPERTY(EditAnywhere, Config, Category = "Default")
+	TSoftClassPtr<UPCGInteractiveToolSettings> SettingsClass;
+
+	// @todo_pcg GetAssetFilter only works on UClasses due to need for UFUNCTION. Will reenable this via details customization instead.
+	//UPROPERTY(EditAnywhere, Config, Category = "Default", meta=(GetAssetFilter = "FilterBySettingsClass"))
+	UPROPERTY(EditAnywhere, Config, Category = "Default")
+	TSoftObjectPtr<UPCGGraph> DefaultGraph;
+};
+
 /** User preferences related to the PCG Editor Mode. */
 UCLASS(Config = EditorPerProjectUserSettings, DisplayName = "PCG Editor Mode Settings")
 class UPCGEditorModeSettings : public UDeveloperSettings
@@ -54,4 +71,7 @@ public:
 
 	UPROPERTY(EditAnywhere, Category = "Tools|Output|Spline", meta = (EditCondition = "bEnableEditorMode", EditConditionHides))
 	FName DefaultNewSplineComponentName = "PCG Tool Spline Component";
+
+	UPROPERTY(EditAnywhere, Config, Category = "Tools")
+	TArray<FPCGPerInteractiveToolSettingSettings> InteractiveToolSettings;
 };
diff --git a/Engine/Plugins/PCG/Source/PCGEditor/Public/EditorMode/Tools/Line/PCGDrawSplineTool.h b/Engine/Plugins/PCG/Source/PCGEditor/Public/EditorMode/Tools/Line/PCGDrawSplineTool.h
--- a/Engine/Plugins/PCG/Source/PCGEditor/Public/EditorMode/Tools/Line/PCGDrawSplineTool.h
+++ b/Engine/Plugins/PCG/Source/PCGEditor/Public/EditorMode/Tools/Line/PCGDrawSplineTool.h
@@ -33,9 +33,6 @@ public:
 
 	// ~Begin UPCGInteractiveToolSettings interface
 	virtual FText GetToolTag() const override;
-	/** Apply the results of the tool. */
-	bool ApplyTool(UWorld* World, const USplineComponent* ToolSplineComponent, FText& OutWarning) const;
-	virtual FSoftObjectPath GetDefaultGraph() const override;
 	// ~End UPCGInteractiveToolSettings interface
 
 	UPROPERTY(EditAnywhere, Category = "Spline")
@@ -94,10 +91,11 @@ public:
 
 	void SetIsSplineSurface(const bool bInIsSplineSurface) { bIsSplineSurface = bInIsSplineSurface; }
 
+protected:
+	virtual void ApplyToolInternal() const override;
 private:
 	void AddSplinePoint(const FVector& HitLocation, const FVector& HitNormal) const;
 	FVector GetNewPointNormal(const FVector& HitNormal) const;
-	void ApplyToolInternal() const;
 
 	UPROPERTY(Instanced)
 	TObjectPtr<USingleClickOrDragInputBehavior> ClickOrDragBehavior = nullptr;
diff --git a/Engine/Plugins/PCG/Source/PCGEditor/Public/EditorMode/Tools/PCGInteractiveTool.h b/Engine/Plugins/PCG/Source/PCGEditor/Public/EditorMode/Tools/PCGInteractiveTool.h
--- a/Engine/Plugins/PCG/Source/PCGEditor/Public/EditorMode/Tools/PCGInteractiveTool.h
+++ b/Engine/Plugins/PCG/Source/PCGEditor/Public/EditorMode/Tools/PCGInteractiveTool.h
@@ -22,18 +22,24 @@ public:
 	virtual void Setup() override;
 	virtual void Shutdown(EToolShutdownType ShutdownType) override;
 	virtual void OnTick(float DeltaTime) override;
+	virtual void OnPropertyModified(UObject* PropertySet, FProperty* Property) override;
 	// ~End UInteractiveTool interface
 
 	virtual void OnRefresh();
 	virtual void SetTargetActor(AActor* InTargetActor);
 	virtual TWeakObjectPtr<AActor> GetTargetActor() const;
-
+	
 	UWorld* GetContextualWorld() const;
 	UWorld* GetContextualWorldChecked() const;
 
 	OnPCGInteractiveToolWarningChanged OnToolWarningChanged;
 
+	void ApplyTool() const;
+
 protected:
+	/** Functionality to execute when the tool is applied. This must be defined in the subclass. */
+	virtual void ApplyToolInternal() const PURE_VIRTUAL(UPCGInteractiveTool::ApplyToolInternal, return;)
+	
 	/** The concrete class of the tool. This must be defined in the subclass. */
 	virtual TSubclassOf<UPCGInteractiveToolSettings> GetSettingsClass() const PURE_VIRTUAL(UPCGInteractiveTool::GetSettingsClass, return nullptr;)
 
diff --git a/Engine/Plugins/PCG/Source/PCGEditor/Public/EditorMode/Tools/PCGInteractiveToolSettings.h b/Engine/Plugins/PCG/Source/PCGEditor/Public/EditorMode/Tools/PCGInteractiveToolSettings.h
--- a/Engine/Plugins/PCG/Source/PCGEditor/Public/EditorMode/Tools/PCGInteractiveToolSettings.h
+++ b/Engine/Plugins/PCG/Source/PCGEditor/Public/EditorMode/Tools/PCGInteractiveToolSettings.h
@@ -19,22 +19,30 @@ class PCGEDITOR_API UPCGInteractiveToolSettings : public UInteractiveToolPropert
 public:
 	UPCGInteractiveToolSettings();
 
+	void Initialize(UPCGGraphInstance* InGraphInstance)
+	{
+		GraphInstance = InGraphInstance;
+	}
+	
 	/** Tool tag for use with filtering compatible graphs. */
 	virtual FText GetToolTag() const PURE_VIRTUAL(UPCGInteractiveToolSettings::GetToolTag, return FText::GetEmpty(););
 	/** A soft object path to the graph asset to load by default for this tool. */
-	virtual FSoftObjectPath GetDefaultGraph() const PURE_VIRTUAL(UPCGInteractiveToolSettings::GetDefaultGraph, return FSoftObjectPath{};); // @todo_pcg: Eventually, this could be a user preference per tool
+	virtual FSoftObjectPath GetDefaultGraph() const;
 
-protected:
 	/** Filter details view graph selection. */
 	UFUNCTION()
 	virtual bool GraphAssetFilter(const FAssetData& AssetData) const;
+protected:
+	
 	/** Try to load the default graph into TargetGraph. */
 	virtual void TryLoadDefaultGraph();
-
 public:
 	UPROPERTY(BlueprintReadOnly, EditAnywhere, Category = "PCG", meta = (NoResetToDefault, GetAssetFilter = "GraphAssetFilter"))
 	TObjectPtr<UPCGGraphInterface> TargetGraph = nullptr;
 
+	UPROPERTY(Transient)
+	TObjectPtr<UPCGGraphInstance> GraphInstance;
+	
 	/** Automatically generate the PCG component while using the tool. */
 	UPROPERTY(EditAnywhere, Category = "PCG")
 	bool bAutoGenerate = true;
@@ -49,10 +57,10 @@ public:
 	TWeakObjectPtr<AActor> TargetActor = nullptr;
 
 	/** The name to be issued to a newly created actor when the tool is applied. */
-	UPROPERTY(EditAnywhere, Category = "PCG", meta = (EditCondition = "Interaction == EPCGToolInteractionMode::CreateNewActor", EditConditionHides, TransientToolProperty))
+	UPROPERTY(EditAnywhere, DisplayName="Actor Label", Category = "PCG", meta = (EditCondition = "Interaction == EPCGToolInteractionMode::CreateNewActor", EditConditionHides, TransientToolProperty))
 	FName NewActorName;
 
 	/** The name to be issued to a newly created PCG component when the tool is applied. */
-	UPROPERTY(EditAnywhere, Category = "PCG", meta = (TransientToolProperty))
+	UPROPERTY(EditAnywhere, DisplayName="Component Name", Category = "PCG", meta = (TransientToolProperty))
 	FName NewPCGComponentName;
 };
diff --git a/Engine/Plugins/PCG/Source/PCGEditor/Public/EditorMode/Tools/Paint/PCGPaintTool.h b/Engine/Plugins/PCG/Source/PCGEditor/Public/EditorMode/Tools/Paint/PCGPaintTool.h
--- a/Engine/Plugins/PCG/Source/PCGEditor/Public/EditorMode/Tools/Paint/PCGPaintTool.h
+++ b/Engine/Plugins/PCG/Source/PCGEditor/Public/EditorMode/Tools/Paint/PCGPaintTool.h
@@ -39,6 +39,9 @@ public:
 	virtual TSubclassOf<UPCGInteractiveToolSettings> GetSettingsClass() const override;
 	virtual bool CanAccept() const override;
 	// ~End UPCGInteractiveToolSettings interface
+
+protected:
+	virtual void ApplyToolInternal() const override {}
 };
 
 UCLASS(Transient)
diff --git a/Engine/Plugins/PCG/Source/PCGEditor/Public/EditorMode/Tools/Volume/PCGCreateVolumeTool.h b/Engine/Plugins/PCG/Source/PCGEditor/Public/EditorMode/Tools/Volume/PCGCreateVolumeTool.h
--- a/Engine/Plugins/PCG/Source/PCGEditor/Public/EditorMode/Tools/Volume/PCGCreateVolumeTool.h
+++ b/Engine/Plugins/PCG/Source/PCGEditor/Public/EditorMode/Tools/Volume/PCGCreateVolumeTool.h
@@ -66,6 +66,8 @@ public:
 	virtual void OnClicked(const FInputDeviceRay& ClickPosition) override;
 	// ~End IClickBehaviorTarget interface
 
+protected:
+	virtual void ApplyToolInternal() const override {}
 private:
 	PCG::EditorMode::Scene::FViewRayParams CachedRayParams;
 };
