[PCG] Make another pass on FPCGDataTypeIdentifier constructors

* Variadic constructor was ambiguous and overlapped with the default constructor and copy constructor
* Move it into a static function and revert the previous constructors from other types.
* Make the custom subtype -1 by default as 0 can be a valid value
* Implement the move version of `operator|` and `operator|=`

#rnx

diff --git a/Engine/Plugins/PCG/Source/PCG/Private/Data/Registry/PCGDataTypeIdentifier.cpp b/Engine/Plugins/PCG/Source/PCG/Private/Data/Registry/PCGDataTypeIdentifier.cpp
--- a/Engine/Plugins/PCG/Source/PCG/Private/Data/Registry/PCGDataTypeIdentifier.cpp
+++ b/Engine/Plugins/PCG/Source/PCG/Private/Data/Registry/PCGDataTypeIdentifier.cpp
@@ -84,6 +84,29 @@ void FPCGDataTypeIdentifier::ComposeFromPCGDataSubclass(TSubclassOf<UPCGData> Su
 	}
 }
 
+void FPCGDataTypeIdentifier::ComposeFromBaseId(const FPCGDataTypeBaseId& BaseId)
+{
+	if (BaseId.IsValid())
+	{
+		Ids.AddUnique(BaseId);
+	}
+}
+
+FPCGDataTypeIdentifier::FPCGDataTypeIdentifier(EPCGDataType LegacyType)
+{
+	ComposeFromLegacyType(LegacyType);
+}
+
+FPCGDataTypeIdentifier::FPCGDataTypeIdentifier(const FPCGDataTypeBaseId& BaseId)
+{
+	ComposeFromBaseId(BaseId);
+}
+
+FPCGDataTypeIdentifier::FPCGDataTypeIdentifier(const TSubclassOf<UPCGData>& DataClass)
+{
+	ComposeFromPCGDataSubclass(DataClass);
+}
+
 FPCGDataTypeIdentifier FPCGDataTypeIdentifier::Construct(TConstArrayView<FPCGDataTypeIdentifier> InIds)
 {
 	FPCGDataTypeIdentifier Result;
@@ -133,11 +156,41 @@ FPCGDataTypeIdentifier& FPCGDataTypeIdentifier::operator|=(EPCGDataType OtherLeg
 
 FPCGDataTypeIdentifier& FPCGDataTypeIdentifier::operator|=(const FPCGDataTypeIdentifier& Other)
 {
-	for (const FPCGDataTypeBaseId& Id : Other.GetIds())
+	// If we are invalid, just copy Other
+	if (!IsValid())
 	{
-		if (Id.IsValid())
+		*this = Other;
+	}
+	else
+	{
+		for (const FPCGDataTypeBaseId& Id : Other.GetIds())
 		{
-			Ids.AddUnique(Id);
+			if (Id.IsValid())
+			{
+				Ids.AddUnique(Id);
+			}
+		}
+	}
+
+	return *this;
+}
+
+FPCGDataTypeIdentifier& FPCGDataTypeIdentifier::operator|=(FPCGDataTypeIdentifier&& Other)
+{
+	// If we are invalid, just swap with Other
+	if (!IsValid())
+	{
+		Swap(*this, Other);
+	}
+	else
+	{
+		for (const FPCGDataTypeBaseId& Id : Other.GetIds())
+		{
+			if (Id.IsValid())
+			{
+				// Safe to const_cast
+				Ids.AddUnique(MoveTemp(const_cast<FPCGDataTypeBaseId&>(Id)));
+			}
 		}
 	}
 	
@@ -156,6 +209,13 @@ FPCGDataTypeIdentifier FPCGDataTypeIdentifier::operator|(const FPCGDataTypeIdent
 	return Result;
 }
 
+FPCGDataTypeIdentifier FPCGDataTypeIdentifier::operator|(FPCGDataTypeIdentifier&& Other) const
+{
+	FPCGDataTypeIdentifier Result = *this;
+	Result |= MoveTemp(Other);
+	return Result;
+}
+
 FPCGDataTypeIdentifier& FPCGDataTypeIdentifier::operator&=(EPCGDataType OtherLegacyType)
 {
 	*this = operator&(OtherLegacyType);
diff --git a/Engine/Plugins/PCG/Source/PCG/Private/Tests/Registry/PCGDataTypeRegistryTests.cpp b/Engine/Plugins/PCG/Source/PCG/Private/Tests/Registry/PCGDataTypeRegistryTests.cpp
--- a/Engine/Plugins/PCG/Source/PCG/Private/Tests/Registry/PCGDataTypeRegistryTests.cpp
+++ b/Engine/Plugins/PCG/Source/PCG/Private/Tests/Registry/PCGDataTypeRegistryTests.cpp
@@ -452,13 +452,13 @@ bool FPCGUnicityTest::RunTest(const FString& Parameters)
 	const FPCGDataTypeIdentifier SameTypeWithClassTpl = FPCGDataTypeIdentifier::Construct<UPCGData, UPCGData>();
 	UTEST_EQUAL("Construct with same class in template gives just 1 ID", SameTypeWithClassTpl.GetIds().Num(), 1);
 
-	const FPCGDataTypeIdentifier SameTypeWithClass = FPCGDataTypeIdentifier(UPCGPointData::StaticClass(), UPCGPointData::StaticClass());
+	const FPCGDataTypeIdentifier SameTypeWithClass = FPCGDataTypeIdentifier::Construct(UPCGPointData::StaticClass(), UPCGPointData::StaticClass());
 	UTEST_EQUAL("Construct with same class gives just 1 ID", SameTypeWithClass.GetIds().Num(), 1);
 
-	const FPCGDataTypeIdentifier SameTypeWithBaseIds = FPCGDataTypeIdentifier(FPCGDataTypeBaseId::Construct<FPCGDataTypeInfoLandscape>(), FPCGDataTypeBaseId::Construct<FPCGDataTypeInfoLandscape>());
+	const FPCGDataTypeIdentifier SameTypeWithBaseIds = FPCGDataTypeIdentifier::Construct(FPCGDataTypeBaseId::Construct<FPCGDataTypeInfoLandscape>(), FPCGDataTypeBaseId::Construct<FPCGDataTypeInfoLandscape>());
 	UTEST_EQUAL("Construct with same base ids gives just 1 ID", SameTypeWithBaseIds.GetIds().Num(), 1);
 
-	const FPCGDataTypeIdentifier SameTypeWithIds = FPCGDataTypeIdentifier(SameTypeWithClassTpl, SameTypeWithClassTpl, SameTypeWithClassTpl);
+	const FPCGDataTypeIdentifier SameTypeWithIds = FPCGDataTypeIdentifier::Construct(SameTypeWithClassTpl, SameTypeWithClassTpl, SameTypeWithClassTpl);
 	UTEST_EQUAL("Construct with same ids gives just 1 ID", SameTypeWithIds.GetIds().Num(), 1);
 
 	const FPCGDataTypeIdentifier SameTypeWithDefaultTypeOr = SameTypeWithClass | EPCGDataType::Point;
@@ -548,8 +548,8 @@ bool FPCGIsChildOfTest::RunTest(const FString& Parameters)
 		{FPCGDataTypeInfo::AsId(), FPCGDataTypeInfoPoint::AsId(), false},
 		{FPCGDataTypeInfoSpline::AsId(), FPCGDataTypeInfoPolyline::AsId(), true},
 		{FPCGDataTypeInfoSpline::AsId(), FPCGDataTypeInfoLandscapeSpline::AsId(), false},
-		{FPCGDataTypeIdentifier{FPCGDataTypeInfoSpline::AsId(), FPCGDataTypeInfoLandscapeSpline::AsId()}, FPCGDataTypeInfoPolyline::AsId(), true},
-		{FPCGDataTypeIdentifier{FPCGDataTypeInfoParam::AsId(), FPCGDataTypeInfoPoint::AsId()}, FPCGDataTypeIdentifier{FPCGDataTypeInfoParam::AsId(), FPCGDataTypeInfoSpatial::AsId()}, true},
+		{FPCGDataTypeIdentifier::Construct(FPCGDataTypeInfoSpline::AsId(), FPCGDataTypeInfoLandscapeSpline::AsId()), FPCGDataTypeInfoPolyline::AsId(), true},
+		{FPCGDataTypeIdentifier::Construct(FPCGDataTypeInfoParam::AsId(), FPCGDataTypeInfoPoint::AsId()), FPCGDataTypeIdentifier::Construct(FPCGDataTypeInfoParam::AsId(), FPCGDataTypeInfoSpatial::AsId()), true},
 	};
 	
 	const FPCGDataTypeRegistry& Registry = FPCGModule::GetConstDataTypeRegistry();
diff --git a/Engine/Plugins/PCG/Source/PCG/Public/Data/Registry/PCGDataTypeIdentifier.h b/Engine/Plugins/PCG/Source/PCG/Public/Data/Registry/PCGDataTypeIdentifier.h
--- a/Engine/Plugins/PCG/Source/PCG/Public/Data/Registry/PCGDataTypeIdentifier.h
+++ b/Engine/Plugins/PCG/Source/PCG/Public/Data/Registry/PCGDataTypeIdentifier.h
@@ -30,20 +30,65 @@ struct FPCGDataTypeIdentifier
 	
 	GENERATED_BODY()
 
+	// Default constructors
 	FPCGDataTypeIdentifier() = default;
+	FPCGDataTypeIdentifier(const FPCGDataTypeIdentifier&) = default;
+	FPCGDataTypeIdentifier(FPCGDataTypeIdentifier&&) = default;
+	FPCGDataTypeIdentifier& operator=(const FPCGDataTypeIdentifier&) = default;
+	FPCGDataTypeIdentifier& operator=(FPCGDataTypeIdentifier&&) = default;
 
-	// Can be constructed from another FPCGDataTypeIdentifier, FPCGDataTypeBaseId, TSubclassOf<UPCGData> or EPCGDataType
+	/* Constructors from other types */
+
+	// Implicit to facilitate deprecation
+	PCG_API FPCGDataTypeIdentifier(EPCGDataType LegacyType);
+
+	// Implicit since it can be used interchangeably
+	PCG_API FPCGDataTypeIdentifier(const FPCGDataTypeBaseId& BaseId);
+
+	PCG_API explicit FPCGDataTypeIdentifier(const TSubclassOf<UPCGData>& DataClass);
+
+    /** Composition static constructors. */
+
+	// Can compose from multiple FPCGDataTypeIdentifier, FPCGDataTypeBaseId, TSubclassOf<UPCGData> or EPCGDataType
 	template <typename ...T> requires (CPCGDataTypeIdentifierConstructible<T> && ...)
-	FPCGDataTypeIdentifier(T const& ...InIds);
+	static FPCGDataTypeIdentifier Construct(T&& ...InIds)
+	{
+		FPCGDataTypeIdentifier Result{};
+
+		([&InIds, &Result]()
+		{
+			if constexpr (std::is_same_v<std::decay_t<T>, FPCGDataTypeIdentifier>)
+			{
+				Result |= std::forward<T>(InIds);
+			}
+			else if constexpr (std::is_same_v<std::decay_t<T>, FPCGDataTypeBaseId>)
+			{
+				Result.ComposeFromBaseId(std::forward<T>(InIds));
+			}
+			else if constexpr (std::is_convertible_v<std::decay_t<T>, TSubclassOf<UPCGData>>)
+			{
+				Result.ComposeFromPCGDataSubclass(std::forward<T>(InIds));
+			}
+			else if constexpr(std::is_same_v<std::decay_t<T>, EPCGDataType>)
+			{
+				Result.ComposeFromLegacyType(std::forward<T>(InIds));
+			}
+			else
+			{
+				static_assert(!std::is_same_v<T, T>, "Incompatible argument");
+			}
+		}(), ...);
+
+		return Result;
+	}
 
-    /** Composition constructors. Not default constructors because of ambiguity with the variadic constructor using initializer list. */
     PCG_API static FPCGDataTypeIdentifier Construct(TConstArrayView<FPCGDataTypeIdentifier> InIds);
     PCG_API static FPCGDataTypeIdentifier Construct(TConstArrayView<TSubclassOf<UPCGData>> Classes);
 
 	template <typename ...T> requires (std::is_base_of_v<UPCGData, T> && ...)
 	static FPCGDataTypeIdentifier Construct()
 	{
-		return FPCGDataTypeIdentifier(T::StaticClass()...);
+		return FPCGDataTypeIdentifier::Construct(T::StaticClass()...);
 	}
 
 	// ------- Begin deprecation with EPCGDataType --------
@@ -88,7 +133,10 @@ struct FPCGDataTypeIdentifier
 	}
 
 	PCG_API FPCGDataTypeIdentifier& operator|=(const FPCGDataTypeIdentifier& Other);
+	PCG_API FPCGDataTypeIdentifier& operator|=(FPCGDataTypeIdentifier&& Other);
 	PCG_API FPCGDataTypeIdentifier operator|(const FPCGDataTypeIdentifier& Other) const;
+	PCG_API FPCGDataTypeIdentifier operator|(FPCGDataTypeIdentifier&& Other) const;
+
 	PCG_API FPCGDataTypeIdentifier& operator&=(const FPCGDataTypeIdentifier& Other);
 	PCG_API FPCGDataTypeIdentifier operator&(const FPCGDataTypeIdentifier& Other) const;
 
@@ -136,6 +184,7 @@ struct FPCGDataTypeIdentifier
 private:
 	PCG_API void ComposeFromLegacyType(EPCGDataType LegacyType);
 	PCG_API void ComposeFromPCGDataSubclass(TSubclassOf<UPCGData> Subclass);
+	PCG_API void ComposeFromBaseId(const FPCGDataTypeBaseId& BaseId);
 
 private:
 	UPROPERTY(BlueprintReadOnly, VisibleAnywhere, Category = "PCG", meta = (AllowPrivateAccess))
@@ -143,7 +192,7 @@ private:
 
 public:
 	UPROPERTY(BlueprintReadWrite, EditAnywhere, Category = "PCG")
-	int32 CustomSubtype = 0;
+	int32 CustomSubtype = -1;
 };
 
 template<>
@@ -184,34 +233,3 @@ public:
 	UFUNCTION(BlueprintCallable, Category = "PCG|PCGDataTypeIdentifier", meta = (ScriptMethod))
 	static bool IsIdentical(UPARAM(Ref) const FPCGDataTypeIdentifier& This, UPARAM(Ref) const FPCGDataTypeIdentifier& Other) { return This.IsIdentical(Other); }
 };
-
-template <typename ...T> requires (CPCGDataTypeIdentifierConstructible<T> && ...)
-FPCGDataTypeIdentifier::FPCGDataTypeIdentifier(T const& ...InIds)
-{
-	([this, &InIds]()
-	{
-		if constexpr (std::is_same_v<T, FPCGDataTypeIdentifier>)
-		{
-			*this |= InIds;
-		}
-		else if constexpr (std::is_same_v<T, FPCGDataTypeBaseId>)
-		{
-			if (InIds.IsValid())
-			{
-				Ids.AddUnique(InIds);
-			}
-		}
-		else if constexpr (std::is_convertible_v<T, TSubclassOf<UPCGData>>)
-		{
-			ComposeFromPCGDataSubclass(InIds);
-		}
-		else if constexpr(std::is_same_v<T, EPCGDataType>)
-		{
-			ComposeFromLegacyType(InIds);
-		}
-		else
-		{
-			static_assert(!std::is_same_v<T, T>, "Incompatible argument");
-		}
-	}(), ...);
-}
\ No newline at end of file
