[FastGeo] Add a Procedural ISM Component.

* Only game worlds for now.

#rnx

diff --git a/Engine/Plugins/Experimental/FastGeoStreaming/Source/FastGeoStreaming/Internal/FastGeoComponentCluster.h b/Engine/Plugins/Experimental/FastGeoStreaming/Source/FastGeoStreaming/Internal/FastGeoComponentCluster.h
--- a/Engine/Plugins/Experimental/FastGeoStreaming/Source/FastGeoStreaming/Internal/FastGeoComponentCluster.h
+++ b/Engine/Plugins/Experimental/FastGeoStreaming/Source/FastGeoStreaming/Internal/FastGeoComponentCluster.h
@@ -8,6 +8,7 @@
 #include "FastGeoInstancedStaticMeshComponent.h"
 #include "FastGeoSkinnedMeshComponent.h"
 #include "FastGeoInstancedSkinnedMeshComponent.h"
+#include "FastGeoProceduralISMComponent.h"
 
 class ULevel;
 class UFastGeoContainer;
@@ -24,15 +25,11 @@ public:
 	typedef IFastGeoElement Super;
 
 	/** Static type identifier for this element class */
-	static const FFastGeoElementType Type;
+	FASTGEOSTREAMING_API static const FFastGeoElementType Type;
 
-	FFastGeoComponentCluster(UFastGeoContainer* InOwner = nullptr, FName InName = NAME_None, FFastGeoElementType InType = Type);
+	FASTGEOSTREAMING_API FFastGeoComponentCluster(UFastGeoContainer* InOwner = nullptr, FName InName = NAME_None, FFastGeoElementType InType = Type);
 
-#if WITH_EDITOR
 	FFastGeoComponentCluster(const FFastGeoComponentCluster& InOther);
-#else
-	FFastGeoComponentCluster(const FFastGeoComponentCluster& InOther) = delete;
-#endif
 
 	virtual ~FFastGeoComponentCluster() {}
 	virtual void OnRegister() {}
@@ -40,16 +37,17 @@ public:
 	virtual void Serialize(FArchive& Ar);
 	virtual bool IsVisible() const { return true; }
 
-	UFastGeoContainer* GetOwnerContainer() const;
+	FASTGEOSTREAMING_API UFastGeoContainer* GetOwnerContainer() const;
 	ULevel* GetLevel() const;
-	FFastGeoComponent* GetComponent(uint32 InComponentTypeID, int32 InComponentIndex);
+	FASTGEOSTREAMING_API FFastGeoComponent* GetComponent(uint32 InComponentTypeID, int32 InComponentIndex);
 	const FString& GetName() const { return Name; }
 	int32 GetComponentClusterIndex() const { return ComponentClusterIndex; }
 	void UpdateVisibility();
 
-#if WITH_EDITOR
-	bool HasComponents() const;
+	FASTGEOSTREAMING_API bool HasComponents() const;
 	virtual FFastGeoComponent& AddComponent(FFastGeoElementType InComponentType);
+
+#if WITH_EDITOR
 	virtual void PreSave(FObjectPreSaveContext ObjectSaveContext) {}
 #endif
 
@@ -89,6 +87,7 @@ protected:
 	TArray<FFastGeoInstancedStaticMeshComponent> InstancedStaticMeshComponents;
 	TArray<FFastGeoSkinnedMeshComponent> SkinnedMeshComponents;
 	TArray<FFastGeoInstancedSkinnedMeshComponent> InstancedSkinnedMeshComponents;
+	TArray<FFastGeoProceduralISMComponent> ProceduralISMComponents;
 
 	template <typename ThisType, typename TComponent, typename TFunc>
 	static void ForEachComponent(ThisType* Self, TFunc&& InFunc)
@@ -119,6 +118,7 @@ protected:
 		ForEachArray(Self->InstancedStaticMeshComponents);
 		ForEachArray(Self->SkinnedMeshComponents);
 		ForEachArray(Self->InstancedSkinnedMeshComponents);
+		ForEachArray(Self->ProceduralISMComponents);
 	}
 
 	template <typename ThisType, typename TComponent, typename TFunc>
@@ -157,7 +157,8 @@ protected:
 		return ForEachArray(Self->StaticMeshComponents) &&
 			   ForEachArray(Self->InstancedStaticMeshComponents) &&
 			   ForEachArray(Self->SkinnedMeshComponents) &&
-			   ForEachArray(Self->InstancedSkinnedMeshComponents);
+			   ForEachArray(Self->InstancedSkinnedMeshComponents) &&
+			   ForEachArray(Self->ProceduralISMComponents);
 	}
 
 private:
diff --git a/Engine/Plugins/Experimental/FastGeoStreaming/Source/FastGeoStreaming/Internal/FastGeoContainer.h b/Engine/Plugins/Experimental/FastGeoStreaming/Source/FastGeoStreaming/Internal/FastGeoContainer.h
--- a/Engine/Plugins/Experimental/FastGeoStreaming/Source/FastGeoStreaming/Internal/FastGeoContainer.h
+++ b/Engine/Plugins/Experimental/FastGeoStreaming/Source/FastGeoStreaming/Internal/FastGeoContainer.h
@@ -127,10 +127,8 @@ public:
 	virtual IPhysicsBodyInstanceOwner* ResolvePhysicsBodyInstanceOwner(Chaos::FConstPhysicsObjectHandle PhysicsObject) override;
 	//~End IPhysicsBodyInstanceOwnerResolver
 
-#if WITH_EDITOR
 	void OnCreated();
 	void AddComponentCluster(FFastGeoComponentCluster* ComponentCluster);
-#endif
 
 	void InitializeDynamicProperties(bool bInitForPlay = false);
 
@@ -191,8 +189,8 @@ private:
 
 #if WITH_EDITOR
 	void CreateEditorProxyObjects();
-	void CollectAssetReferences();
 #endif
+	void CollectAssetReferences();
 	void SerializeComponentClusters(FArchive& Ar);
 	void RegisterToNavigationSystem();
 	void UnregisterFromNavigationSystem();
diff --git a/Engine/Plugins/Experimental/FastGeoStreaming/Source/FastGeoStreaming/Internal/FastGeoInstancedStaticMeshComponent.h b/Engine/Plugins/Experimental/FastGeoStreaming/Source/FastGeoStreaming/Internal/FastGeoInstancedStaticMeshComponent.h
--- a/Engine/Plugins/Experimental/FastGeoStreaming/Source/FastGeoStreaming/Internal/FastGeoInstancedStaticMeshComponent.h
+++ b/Engine/Plugins/Experimental/FastGeoStreaming/Source/FastGeoStreaming/Internal/FastGeoInstancedStaticMeshComponent.h
@@ -77,10 +77,6 @@ private:
 	TArray<float> PerInstanceSMCustomData;
 	TArray<FInstancedStaticMeshRandomSeed> AdditionalRandomSeeds;
 	FBox NavigationBounds;
-	float InstanceLODDistanceScale = 1.f;
-	int32 InstanceStartCullDistance = 0;
-	int32 InstanceEndCullDistance = 0;
-	uint8 bUseGpuLodSelection : 1 = false;
 	FInstancedStaticMeshSceneProxyDesc SceneProxyDesc{};
 
 	// Transient data
diff --git a/Engine/Plugins/Experimental/FastGeoStreaming/Source/FastGeoStreaming/Internal/FastGeoPrimitiveComponent.h b/Engine/Plugins/Experimental/FastGeoStreaming/Source/FastGeoStreaming/Internal/FastGeoPrimitiveComponent.h
--- a/Engine/Plugins/Experimental/FastGeoStreaming/Source/FastGeoStreaming/Internal/FastGeoPrimitiveComponent.h
+++ b/Engine/Plugins/Experimental/FastGeoStreaming/Source/FastGeoStreaming/Internal/FastGeoPrimitiveComponent.h
@@ -76,6 +76,7 @@ public:
 
 	FFastGeoPrimitiveComponent(int32 InComponentIndex = INDEX_NONE, FFastGeoElementType InType = Type);
 	virtual ~FFastGeoPrimitiveComponent() = default;
+	FFastGeoPrimitiveComponent(const FFastGeoPrimitiveComponent& Other);
 
 	TArray<URuntimeVirtualTexture*> const& GetRuntimeVirtualTextures() const { return RuntimeVirtualTextures; }
 	bool IsFirstPersonRelevant() const;
diff --git a/Engine/Plugins/Experimental/FastGeoStreaming/Source/FastGeoStreaming/Internal/FastGeoProceduralISMComponent.h b/Engine/Plugins/Experimental/FastGeoStreaming/Source/FastGeoStreaming/Internal/FastGeoProceduralISMComponent.h
new file mode 100644
--- /dev/null
+++ b/Engine/Plugins/Experimental/FastGeoStreaming/Source/FastGeoStreaming/Internal/FastGeoProceduralISMComponent.h
@@ -0,0 +1,61 @@
+// Copyright Epic Games, Inc. All Rights Reserved.
+
+#pragma once
+
+#include "FastGeoStaticMeshComponent.h"
+#include "InstancedStaticMeshSceneProxyDesc.h"
+
+struct FProceduralISMComponentDescriptor;
+
+class FFastGeoProceduralISMComponent : public FFastGeoStaticMeshComponentBase
+{
+public:
+	typedef FFastGeoStaticMeshComponentBase Super;
+
+	/** Static type identifier for this element class */
+	FASTGEOSTREAMING_API static const FFastGeoElementType Type;
+
+	FFastGeoProceduralISMComponent(int32 InComponentIndex = INDEX_NONE, FFastGeoElementType InType = Type);
+	virtual ~FFastGeoProceduralISMComponent() = default;
+
+	FASTGEOSTREAMING_API void InitializeFromComponentDescriptor(const FProceduralISMComponentDescriptor& InDescriptor);
+
+	//~ Begin FFastGeoComponent Interface
+	virtual void Serialize(FArchive& Ar) override;
+	//~ End FFastGeoComponent Interface
+
+protected:
+	//~ Begin FFastGeoPrimitiveComponent interface
+	virtual FPrimitiveSceneProxyDesc& GetSceneProxyDesc() override { return SceneProxyDesc; }
+	virtual const FPrimitiveSceneProxyDesc& GetSceneProxyDesc() const override { return SceneProxyDesc; }
+	virtual bool IsNavigationRelevant() const override { return false; }
+#if WITH_EDITOR
+	virtual void ResetSceneProxyDescUnsupportedProperties() override;
+#endif
+	virtual void InitializeSceneProxyDescDynamicProperties() override;
+	virtual void ApplyWorldTransform(const FTransform& InTransform) override;
+	//~ End FFastGeoPrimitiveComponent interface
+
+	//~ Begin FFastGeoStaticMeshComponentBase interface
+	virtual FStaticMeshSceneProxyDesc& GetStaticMeshSceneProxyDesc() override { return SceneProxyDesc; }
+	virtual const FStaticMeshSceneProxyDesc& GetStaticMeshSceneProxyDesc() const override { return SceneProxyDesc; }
+	virtual FPrimitiveSceneProxy* CreateStaticMeshSceneProxy(const Nanite::FMaterialAudit& NaniteMaterials, bool bCreateNanite) override;
+#if !WITH_EDITOR && UE_WITH_PSO_PRECACHING
+	virtual void CollectPSOPrecacheData(const FPSOPrecacheParams& BasePrecachePSOParams, FMaterialInterfacePSOPrecacheParamsList& OutParams) override;
+#endif
+	//~ End FFastGeoStaticMeshComponentBase interface
+
+	TSharedPtr<FInstanceDataSceneProxy, ESPMode::ThreadSafe>& BuildInstanceData();
+
+private:
+	// Persistent Data
+	int32 NumInstances = 0;
+	int32 NumCustomDataFloats = 0;
+	FBox PrimitiveBoundsOverride = FBox(ForceInit);
+	FInstancedStaticMeshSceneProxyDesc SceneProxyDesc{};
+
+	// Transient data
+	TSharedPtr<FInstanceDataSceneProxy, ESPMode::ThreadSafe> DataProxy{};
+
+	friend class FInstancedStaticMeshComponentHelper;
+};
diff --git a/Engine/Plugins/Experimental/FastGeoStreaming/Source/FastGeoStreaming/Internal/FastGeoWeakElement.h b/Engine/Plugins/Experimental/FastGeoStreaming/Source/FastGeoStreaming/Internal/FastGeoWeakElement.h
--- a/Engine/Plugins/Experimental/FastGeoStreaming/Source/FastGeoStreaming/Internal/FastGeoWeakElement.h
+++ b/Engine/Plugins/Experimental/FastGeoStreaming/Source/FastGeoStreaming/Internal/FastGeoWeakElement.h
@@ -52,6 +52,8 @@ private:
 class FWeakFastGeoComponent
 {
 public:
+	FWeakFastGeoComponent() = default;
+
 	FWeakFastGeoComponent(FFastGeoComponent* InComponent)
 	{
 		if (FFastGeoComponentCluster* ComponentCluster = InComponent ? InComponent->GetOwnerComponentCluster() : nullptr)
@@ -60,6 +62,10 @@ public:
 			ComponentTypeID = InComponent->GetTypeID();
 			ComponentIndex = InComponent->GetComponentIndex();
 		}
+		else
+		{
+			ensure(!InComponent);
+		}
 	}
 
 	FORCEINLINE bool operator==(const FWeakFastGeoComponent& Other) const
diff --git a/Engine/Plugins/Experimental/FastGeoStreaming/Source/FastGeoStreaming/Private/FastGeoAsyncRenderStateJobQueue.cpp b/Engine/Plugins/Experimental/FastGeoStreaming/Source/FastGeoStreaming/Private/FastGeoAsyncRenderStateJobQueue.cpp
--- a/Engine/Plugins/Experimental/FastGeoStreaming/Source/FastGeoStreaming/Private/FastGeoAsyncRenderStateJobQueue.cpp
+++ b/Engine/Plugins/Experimental/FastGeoStreaming/Source/FastGeoStreaming/Private/FastGeoAsyncRenderStateJobQueue.cpp
@@ -123,7 +123,7 @@ void FFastGeoAsyncRenderStateJobQueue::Tick(bool bForceWaitCompletion)
 	Launch();
 
 	// Wait for tasks if the world is inside a blocking load
-	if (Scene->GetWorld()->GetSubsystem<UFastGeoWorldSubsystem>()->IsWaitingForCompletion())
+	if (bForceWaitCompletion || Scene->GetWorld()->GetSubsystem<UFastGeoWorldSubsystem>()->IsWaitingForCompletion())
 	{
 		WaitForAsyncTasksExecution();
 	}
diff --git a/Engine/Plugins/Experimental/FastGeoStreaming/Source/FastGeoStreaming/Private/FastGeoComponentCluster.cpp b/Engine/Plugins/Experimental/FastGeoStreaming/Source/FastGeoStreaming/Private/FastGeoComponentCluster.cpp
--- a/Engine/Plugins/Experimental/FastGeoStreaming/Source/FastGeoStreaming/Private/FastGeoComponentCluster.cpp
+++ b/Engine/Plugins/Experimental/FastGeoStreaming/Source/FastGeoStreaming/Private/FastGeoComponentCluster.cpp
@@ -16,7 +16,6 @@ FFastGeoComponentCluster::FFastGeoComponentCluster(UFastGeoContainer* InOwner, F
 {
 }
 
-#if WITH_EDITOR
 FFastGeoComponentCluster::FFastGeoComponentCluster(const FFastGeoComponentCluster& InOther)
 	: Super(InOther.ElementType)
 {
@@ -27,8 +26,8 @@ FFastGeoComponentCluster::FFastGeoComponentCluster(const FFastGeoComponentCluste
 	InstancedStaticMeshComponents = InOther.InstancedStaticMeshComponents;
 	SkinnedMeshComponents = InOther.SkinnedMeshComponents;
 	InstancedSkinnedMeshComponents = InOther.InstancedSkinnedMeshComponents;
+	ProceduralISMComponents = InOther.ProceduralISMComponents;
 }
-#endif
 
 FArchive& operator<<(FArchive& Ar, FFastGeoComponentCluster& ComponentCluster)
 {
@@ -44,6 +43,7 @@ void FFastGeoComponentCluster::Serialize(FArchive& Ar)
 	Ar << InstancedStaticMeshComponents;
 	Ar << SkinnedMeshComponents;
 	Ar << InstancedSkinnedMeshComponents;
+	Ar << ProceduralISMComponents;
 }
 
 void FFastGeoComponentCluster::InitializeDynamicProperties()
@@ -89,6 +89,10 @@ FFastGeoComponent* FFastGeoComponentCluster::GetComponent(uint32 InComponentType
 	{
 		return SkinnedMeshComponents.IsValidIndex(InComponentIndex) ? &SkinnedMeshComponents[InComponentIndex] : nullptr;
 	}
+	else if (FFastGeoProceduralISMComponent::Type.IsSameTypeID(InComponentTypeID))
+	{
+		return ProceduralISMComponents.IsValidIndex(InComponentIndex) ? &ProceduralISMComponents[InComponentIndex] : nullptr;
+	}
 	check(false);
 	return nullptr;
 }
@@ -209,7 +213,6 @@ void FFastGeoComponentCluster::UpdateVisibility_Internal(TArray<FFastGeoPrimitiv
 	}
 }
 
-#if WITH_EDITOR
 FFastGeoComponent& FFastGeoComponentCluster::AddComponent(FFastGeoElementType InComponentType)
 {
 	FFastGeoComponent* NewComponent = nullptr;
@@ -230,6 +233,10 @@ FFastGeoComponent& FFastGeoComponentCluster::AddComponent(FFastGeoElementType In
 	{
 		NewComponent = &InstancedSkinnedMeshComponents.Emplace_GetRef(InstancedSkinnedMeshComponents.Num());
 	}
+	else if (InComponentType == FFastGeoProceduralISMComponent::Type)
+	{
+		NewComponent = &ProceduralISMComponents.Emplace_GetRef(ProceduralISMComponents.Num());
+	}
 	else
 	{
 		unimplemented(); // Should never be reached, missing type handling
@@ -240,6 +247,5 @@ FFastGeoComponent& FFastGeoComponentCluster::AddComponent(FFastGeoElementType In
 
 bool FFastGeoComponentCluster::HasComponents() const
 {
-	return InstancedStaticMeshComponents.Num() > 0 || StaticMeshComponents.Num() > 0 || InstancedSkinnedMeshComponents.Num() > 0 || SkinnedMeshComponents.Num() > 0;
+	return InstancedStaticMeshComponents.Num() > 0 || StaticMeshComponents.Num() > 0 || InstancedSkinnedMeshComponents.Num() > 0 || SkinnedMeshComponents.Num() > 0 || ProceduralISMComponents.Num() > 0;
 }
-#endif
\ No newline at end of file
diff --git a/Engine/Plugins/Experimental/FastGeoStreaming/Source/FastGeoStreaming/Private/FastGeoContainer.cpp b/Engine/Plugins/Experimental/FastGeoStreaming/Source/FastGeoStreaming/Private/FastGeoContainer.cpp
--- a/Engine/Plugins/Experimental/FastGeoStreaming/Source/FastGeoStreaming/Private/FastGeoContainer.cpp
+++ b/Engine/Plugins/Experimental/FastGeoStreaming/Source/FastGeoStreaming/Private/FastGeoContainer.cpp
@@ -178,8 +178,8 @@ void UFastGeoContainer::Register()
 		InitializeDynamicProperties(true);
 #endif
 		UWorldPartition* WorldPartition = FWorldPartitionHelpers::GetWorldPartition(this);
-		const bool bApplyWorldTransform = WorldPartition->HasInstanceTransform();
-		const FTransform& Transform = WorldPartition->GetInstanceTransform();
+		const bool bApplyWorldTransform = WorldPartition && WorldPartition->HasInstanceTransform();
+		const FTransform& Transform = WorldPartition ? WorldPartition->GetInstanceTransform() : FTransform::Identity;
 		ForEachComponentCluster([bApplyWorldTransform , &Transform, this](FFastGeoComponentCluster& ComponentCluster)
 		{
 			ComponentCluster.OnRegister();
@@ -529,7 +529,6 @@ bool UFastGeoContainer::IsReadyForFinishDestroy()
 	return bResult;
 }
 
-#if WITH_EDITOR
 void UFastGeoContainer::AddComponentCluster(FFastGeoComponentCluster* ComponentCluster)
 {
 	FFastGeoComponentCluster* NewComponentCluster;
@@ -545,6 +544,7 @@ void UFastGeoContainer::AddComponentCluster(FFastGeoComponentCluster* ComponentC
 	}
 }
 
+#if WITH_EDITOR
 void UFastGeoContainer::CreateEditorProxyObjects()
 {
 	ForEachComponentCluster([this](FFastGeoComponentCluster& ComponentCluster)
@@ -571,6 +571,7 @@ void UFastGeoContainer::PreSave(FObjectPreSaveContext ObjectSaveContext)
 		ComponentCluster.PreSave(ObjectSaveContext);
 	});
 }
+#endif
 
 class FFastGeoGatherFastGeoContainerAssetRefsArchive : public FArchive
 {
@@ -641,14 +642,15 @@ void UFastGeoContainer::OnCreated()
 	// Initialize component clusters & components dynamic properties
 	InitializeDynamicProperties(GetLevel()->GetWorld()->IsGameWorld());
 
+#if WITH_EDITOR
 	// In editor, we need a UObject representation of the components for some operations
 	CreateEditorProxyObjects();
+#endif
 
 	// Collect references in order to avoid garbage collection of objects that may now be unreferenced
 	// The fast geo container will hold onto those objects if necessary.
 	CollectAssetReferences();
 }
-#endif
 
 void UFastGeoContainer::Serialize(FArchive& Ar)
 {
diff --git a/Engine/Plugins/Experimental/FastGeoStreaming/Source/FastGeoStreaming/Private/FastGeoPrimitiveComponent.cpp b/Engine/Plugins/Experimental/FastGeoStreaming/Source/FastGeoStreaming/Private/FastGeoPrimitiveComponent.cpp
--- a/Engine/Plugins/Experimental/FastGeoStreaming/Source/FastGeoStreaming/Private/FastGeoPrimitiveComponent.cpp
+++ b/Engine/Plugins/Experimental/FastGeoStreaming/Source/FastGeoStreaming/Private/FastGeoPrimitiveComponent.cpp
@@ -64,6 +64,40 @@ FFastGeoPrimitiveComponent::FFastGeoPrimitiveComponent(int32 InComponentIndex, F
 {
 }
 
+FFastGeoPrimitiveComponent::FFastGeoPrimitiveComponent(const FFastGeoPrimitiveComponent& Other)
+	: FFastGeoComponent(Other)
+	, LocalTransform(Other.LocalTransform)
+	, WorldTransform(Other.WorldTransform)
+	, LocalBounds(Other.LocalBounds)
+	, WorldBounds(Other.WorldBounds)
+	, bIsVisible(Other.bIsVisible)
+	, bStaticWhenNotMoveable(Other.bStaticWhenNotMoveable)
+	, bFillCollisionUnderneathForNavmesh(Other.bFillCollisionUnderneathForNavmesh)
+	, bRasterizeAsFilledConvexVolume(Other.bRasterizeAsFilledConvexVolume)
+	, bCanEverAffectNavigation(Other.bCanEverAffectNavigation)
+	, CustomPrimitiveData(Other.CustomPrimitiveData)
+	, DetailMode(Other.DetailMode)
+	, bHasCustomNavigableGeometry(Other.bHasCustomNavigableGeometry)
+	, BodyInstance(Other.BodyInstance)
+	, RuntimeVirtualTextures(Other.RuntimeVirtualTextures)
+	, BodyInstanceOwner(Other.BodyInstanceOwner)
+	, PrimitiveSceneData(Other.PrimitiveSceneData)
+	, AsyncTermBodyPayload(Other.AsyncTermBodyPayload)
+	, ProxyState(Other.ProxyState)
+	, bRenderStateDirty(Other.bRenderStateDirty)
+#if !WITH_EDITOR && UE_WITH_PSO_PRECACHING
+	, MaterialPSOPrecacheRequestIDs(Other.MaterialPSOPrecacheRequestIDs)
+	, LatestPSOPrecacheJobSet(Other.LatestPSOPrecacheJobSet)
+#endif
+{
+#if !WITH_EDITOR && UE_WITH_PSO_PRECACHING
+	LatestPSOPrecacheJobSetCompleted.store(Other.LatestPSOPrecacheJobSetCompleted.load());
+	bPSOPrecacheCalled.store(Other.bPSOPrecacheCalled.load());
+	bPSOPrecacheRequired.store(Other.bPSOPrecacheRequired.load());
+	PSOPrecacheRequestPriority.store(Other.PSOPrecacheRequestPriority.load());
+#endif
+}
+
 void FFastGeoPrimitiveComponent::Serialize(FArchive& Ar)
 {
 	Super::Serialize(Ar);
@@ -288,7 +322,7 @@ void FFastGeoPrimitiveComponent::InitializeSceneProxyDescDynamicProperties()
 	static FName NAME_FastGeoPrimitiveComponent(TEXT("FastGeoPrimitiveComponent"));
 	SceneProxyDesc.MeshDrawCommandStatsCategory = NAME_FastGeoPrimitiveComponent;
 #endif
-	check(SceneProxyDesc.bIsInstancedStaticMesh == this->IsA<FFastGeoInstancedStaticMeshComponent>());
+	check(SceneProxyDesc.bIsInstancedStaticMesh == (this->IsA<FFastGeoInstancedStaticMeshComponent>() || this->IsA<FFastGeoProceduralISMComponent>()));
 	UpdateVisibility();
 }
 
diff --git a/Engine/Plugins/Experimental/FastGeoStreaming/Source/FastGeoStreaming/Private/FastGeoProceduralISMComponent.cpp b/Engine/Plugins/Experimental/FastGeoStreaming/Source/FastGeoStreaming/Private/FastGeoProceduralISMComponent.cpp
new file mode 100644
--- /dev/null
+++ b/Engine/Plugins/Experimental/FastGeoStreaming/Source/FastGeoStreaming/Private/FastGeoProceduralISMComponent.cpp
@@ -0,0 +1,158 @@
+// Copyright Epic Games, Inc. All Rights Reserved.
+
+#include "FastGeoProceduralISMComponent.h"
+
+#include "InstancedStaticMeshComponentHelper.h"
+#include "NaniteSceneProxy.h"
+#include "InstancedStaticMesh/ISMInstanceDataSceneProxy.h"
+#include "ISMPartition/ProceduralISMComponentDescriptor.h"
+
+const FFastGeoElementType FFastGeoProceduralISMComponent::Type(&FFastGeoStaticMeshComponentBase::Type);
+
+FFastGeoProceduralISMComponent::FFastGeoProceduralISMComponent(int32 InComponentIndex, FFastGeoElementType InType)
+	: Super(InComponentIndex, InType)
+{
+}
+
+void FFastGeoProceduralISMComponent::InitializeFromComponentDescriptor(const FProceduralISMComponentDescriptor& InDescriptor)
+{
+	NumInstances = InDescriptor.NumInstances;
+	NumCustomDataFloats = InDescriptor.NumCustomFloats;
+	PrimitiveBoundsOverride = InDescriptor.WorldBounds;
+	LocalBounds = PrimitiveBoundsOverride.InverseTransformBy(WorldTransform);
+	WorldBounds = InDescriptor.WorldBounds;
+	OverrideMaterials = InDescriptor.OverrideMaterials;
+	RuntimeVirtualTextures = InDescriptor.RuntimeVirtualTextures;
+
+	SceneProxyDesc.StaticMesh = InDescriptor.StaticMesh.Get();
+	SceneProxyDesc.OverlayMaterial = InDescriptor.OverlayMaterial;
+	SceneProxyDesc.MinDrawDistance = InDescriptor.InstanceMinDrawDistance;
+	SceneProxyDesc.InstanceStartCullDistance = InDescriptor.InstanceStartCullDistance;
+	SceneProxyDesc.InstanceEndCullDistance = InDescriptor.InstanceEndCullDistance;
+	SceneProxyDesc.Mobility = InDescriptor.Mobility;
+	SceneProxyDesc.VirtualTextureRenderPassType = InDescriptor.VirtualTextureRenderPassType;
+	SceneProxyDesc.LightingChannels = InDescriptor.LightingChannels;
+	SceneProxyDesc.CustomDepthStencilWriteMask = InDescriptor.CustomDepthStencilWriteMask;
+	SceneProxyDesc.VirtualTextureCullMips = InDescriptor.VirtualTextureCullMips;
+	SceneProxyDesc.TranslucencySortPriority = InDescriptor.TranslucencySortPriority;
+	SceneProxyDesc.CustomDepthStencilValue = InDescriptor.CustomDepthStencilValue;
+	SceneProxyDesc.bVisibleInRayTracing = InDescriptor.bVisibleInRayTracing;
+	SceneProxyDesc.RayTracingGroupId = InDescriptor.RayTracingGroupId;
+	SceneProxyDesc.RayTracingGroupCullingPriority = InDescriptor.RayTracingGroupCullingPriority;
+	SceneProxyDesc.CastShadow = InDescriptor.bCastShadow;
+	SceneProxyDesc.bEmissiveLightSource = InDescriptor.bEmissiveLightSource;
+	SceneProxyDesc.bCastDynamicShadow = InDescriptor.bCastDynamicShadow;
+	SceneProxyDesc.bCastStaticShadow = InDescriptor.bCastStaticShadow;
+	SceneProxyDesc.bCastContactShadow = InDescriptor.bCastContactShadow;
+	SceneProxyDesc.bCastShadowAsTwoSided = InDescriptor.bCastShadowAsTwoSided;
+	SceneProxyDesc.bCastHiddenShadow = InDescriptor.bCastHiddenShadow;
+	SceneProxyDesc.bReceivesDecals = InDescriptor.bReceivesDecals;
+	SceneProxyDesc.bUseAsOccluder = InDescriptor.bUseAsOccluder;
+	SceneProxyDesc.bRenderCustomDepth = InDescriptor.bRenderCustomDepth;
+	SceneProxyDesc.bEvaluateWorldPositionOffset = InDescriptor.bEvaluateWorldPositionOffset;
+	SceneProxyDesc.bReverseCulling = InDescriptor.bReverseCulling;
+	SceneProxyDesc.WorldPositionOffsetDisableDistance = InDescriptor.WorldPositionOffsetDisableDistance;
+	SceneProxyDesc.ShadowCacheInvalidationBehavior = InDescriptor.ShadowCacheInvalidationBehavior;
+	SceneProxyDesc.bIsInstancedStaticMesh = true;
+}
+
+void FFastGeoProceduralISMComponent::Serialize(FArchive& Ar)
+{
+	Super::Serialize(Ar);
+
+	// Serialize persistent data from FFastGeoProceduralISMComponent
+	Ar << NumInstances;
+	Ar << NumCustomDataFloats;
+	Ar << PrimitiveBoundsOverride;
+
+	// Serialize persistent data from FInstancedStaticMeshSceneProxyDesc
+	Ar << SceneProxyDesc.InstanceLODDistanceScale;
+	Ar << SceneProxyDesc.InstanceMinDrawDistance;
+	Ar << SceneProxyDesc.InstanceStartCullDistance;
+	Ar << SceneProxyDesc.InstanceEndCullDistance;
+	FArchive_Serialize_BitfieldBool(Ar, SceneProxyDesc.bUseGpuLodSelection);
+}
+
+#if WITH_EDITOR
+void FFastGeoProceduralISMComponent::ResetSceneProxyDescUnsupportedProperties()
+{
+	Super::ResetSceneProxyDescUnsupportedProperties();
+
+	SceneProxyDesc.InstanceDataSceneProxy = nullptr;
+	SceneProxyDesc.bHasSelectedInstances = false;
+	SceneProxyDesc.bAffectDynamicIndirectLighting = false;
+	SceneProxyDesc.bAffectDistanceFieldLighting = false;
+}
+#endif
+
+void FFastGeoProceduralISMComponent::InitializeSceneProxyDescDynamicProperties()
+{
+	Super::InitializeSceneProxyDescDynamicProperties();
+
+	SceneProxyDesc.bAffectDynamicIndirectLighting = false;
+	SceneProxyDesc.bAffectDistanceFieldLighting = false;
+}
+
+void FFastGeoProceduralISMComponent::ApplyWorldTransform(const FTransform& InTransform)
+{
+	// Procedural component differs - bounds cannot be computed from instance data since it is not available on CPU,
+	// and bounds are provided by user in world space. Apply transform, then recompute local bounds from that.
+	FFastGeoPrimitiveComponent::ApplyWorldTransform(InTransform);
+
+	LocalBounds = PrimitiveBoundsOverride.InverseTransformBy(WorldTransform);
+}
+
+TSharedPtr<FInstanceDataSceneProxy, ESPMode::ThreadSafe>& FFastGeoProceduralISMComponent::BuildInstanceData()
+{
+	FInstanceSceneDataBuffers InstanceSceneDataBuffers(/*InbInstanceDataIsGPUOnly=*/true);
+	FInstanceSceneDataBuffers::FAccessTag AccessTag(PointerHash(this));
+	auto View = InstanceSceneDataBuffers.BeginWriteAccess(AccessTag);
+	
+	// PrimitiveLocalToWorld
+	InstanceSceneDataBuffers.SetPrimitiveLocalToWorld(GetRenderMatrix(), AccessTag);
+	
+	// InstanceLocalBounds
+	const FPrimitiveMaterialPropertyDescriptor PrimitiveMaterialDesc = GetUsedMaterialPropertyDesc(GetScene()->GetShaderPlatform());
+	const float LocalAbsMaxDisplacement = FMath::Max(-PrimitiveMaterialDesc.MinMaxMaterialDisplacement.X, PrimitiveMaterialDesc.MinMaxMaterialDisplacement.Y) + PrimitiveMaterialDesc.MaxWorldPositionOffsetDisplacement;
+	const FVector3f PadExtent = FISMCInstanceDataSceneProxy::GetLocalBoundsPadExtent(View.PrimitiveToRelativeWorld, LocalAbsMaxDisplacement);
+	FRenderBounds InstanceLocalBounds = GetStaticMesh()->GetBounds();
+	InstanceLocalBounds.Min -= PadExtent;
+	InstanceLocalBounds.Max += PadExtent;
+	check(!View.Flags.bHasPerInstanceLocalBounds);
+	View.InstanceLocalBounds.Add(InstanceLocalBounds);
+
+	View.NumInstancesGPUOnly = NumInstances;
+	View.NumCustomDataFloats = NumCustomDataFloats;
+	View.Flags.bHasPerInstanceCustomData = NumCustomDataFloats > 0;
+
+	InstanceSceneDataBuffers.EndWriteAccess(AccessTag);
+	InstanceSceneDataBuffers.ValidateData();
+
+	DataProxy = MakeShared<FInstanceDataSceneProxy, ESPMode::ThreadSafe>(MoveTemp(InstanceSceneDataBuffers));
+	return DataProxy;
+}
+
+FPrimitiveSceneProxy* FFastGeoProceduralISMComponent::CreateStaticMeshSceneProxy(const Nanite::FMaterialAudit& NaniteMaterials, bool bCreateNanite)
+{
+	check(GetWorld());
+	check(SceneProxyDesc.Scene);
+	check(NumInstances > 0);
+
+	SceneProxyDesc.InstanceDataSceneProxy = BuildInstanceData();
+	if (bCreateNanite)
+	{
+		PrimitiveSceneData.SceneProxy = ::new Nanite::FSceneProxy(NaniteMaterials, SceneProxyDesc);
+	}
+	else
+	{
+		PrimitiveSceneData.SceneProxy = ::new FInstancedStaticMeshSceneProxy(SceneProxyDesc, SceneProxyDesc.FeatureLevel);
+	}
+	return PrimitiveSceneData.SceneProxy;
+}
+
+#if !WITH_EDITOR && UE_WITH_PSO_PRECACHING
+void FFastGeoProceduralISMComponent::CollectPSOPrecacheData(const FPSOPrecacheParams& BasePrecachePSOParams, FMaterialInterfacePSOPrecacheParamsList& OutParams)
+{
+	return FInstancedStaticMeshComponentHelper::CollectPSOPrecacheData(*this, BasePrecachePSOParams, OutParams);
+}
+#endif
diff --git a/Engine/Plugins/PCG/Source/PCG/Private/Components/PCGProceduralISMComponent.cpp b/Engine/Plugins/PCG/Source/PCG/Private/Components/PCGProceduralISMComponent.cpp
--- a/Engine/Plugins/PCG/Source/PCG/Private/Components/PCGProceduralISMComponent.cpp
+++ b/Engine/Plugins/PCG/Source/PCG/Private/Components/PCGProceduralISMComponent.cpp
@@ -17,8 +17,6 @@
 #include "SceneInterface.h"
 #include "StaticMeshSceneProxyDesc.h"
 #include "Engine/StaticMesh.h"
-#include "Serialization/ArchiveCrc32.h"
-#include "VT/RuntimeVirtualTexture.h"
 
 #include UE_INLINE_GENERATED_CPP_BY_NAME(PCGProceduralISMComponent)
 
@@ -88,7 +86,7 @@ UPCGManagedProceduralISMComponent* PCGManagedProceduralISMComponent::GetOrCreate
 
 	const EObjectFlags ObjectFlags = RF_Transient;
 	UPCGProceduralISMComponent* ISMC = NewObject<UPCGProceduralISMComponent>(InTargetActor, UPCGProceduralISMComponent::StaticClass(), MakeUniqueObjectName(InTargetActor, UPCGProceduralISMComponent::StaticClass(), FName(ComponentName)), ObjectFlags);
-	Descriptor.InitComponent(ISMC);
+	ISMC->InitFromDescriptor(Descriptor);
 
 	ISMC->RegisterComponent();
 	InTargetActor->AddInstanceComponent(ISMC);
@@ -116,198 +114,6 @@ UPCGManagedProceduralISMComponent* PCGManagedProceduralISMComponent::GetOrCreate
 	return Resource;
 }
 
-FPCGProceduralISMComponentDescriptor::FPCGProceduralISMComponentDescriptor()
-	: NumInstances(0)
-	, NumCustomFloats(0)
-{
-	*this = FPCGSoftISMComponentDescriptor();
-}
-
-FPCGProceduralISMComponentDescriptor& FPCGProceduralISMComponentDescriptor::operator=(const FPCGSoftISMComponentDescriptor& Other)
-{
-	InstanceMinDrawDistance = Other.InstanceMinDrawDistance;
-	InstanceStartCullDistance = Other.InstanceStartCullDistance;
-	InstanceEndCullDistance = Other.InstanceEndCullDistance;
-	OverlayMaterial = Other.OverlayMaterial.LoadSynchronous();
-	StaticMesh = Other.StaticMesh;
-	ComponentTags = Other.ComponentTags;
-	Mobility = Other.Mobility;
-	VirtualTextureRenderPassType = Other.VirtualTextureRenderPassType;
-	LightingChannels = Other.LightingChannels;
-	CustomDepthStencilWriteMask = Other.CustomDepthStencilWriteMask;
-	VirtualTextureCullMips = Other.VirtualTextureCullMips;
-	TranslucencySortPriority = Other.TranslucencySortPriority;
-	CustomDepthStencilValue = Other.CustomDepthStencilValue;
-	bCastShadow = Other.bCastShadow;
-	bEmissiveLightSource = Other.bEmissiveLightSource;
-	bCastDynamicShadow = Other.bCastDynamicShadow;
-	bCastStaticShadow = Other.bCastStaticShadow;
-	bCastContactShadow = Other.bCastContactShadow;
-	bCastShadowAsTwoSided = Other.bCastShadowAsTwoSided;
-	bCastHiddenShadow = Other.bCastHiddenShadow;
-	bReceivesDecals = Other.bReceivesDecals;
-	bUseAsOccluder = Other.bUseAsOccluder;
-	bRenderCustomDepth = Other.bRenderCustomDepth;
-	bEvaluateWorldPositionOffset = Other.bEvaluateWorldPositionOffset;
-	bReverseCulling = Other.bReverseCulling;
-	WorldPositionOffsetDisableDistance = Other.WorldPositionOffsetDisableDistance;
-	NanitePixelProgrammableDistance = Other.NanitePixelProgrammableDistance;
-	ShadowCacheInvalidationBehavior = Other.ShadowCacheInvalidationBehavior;
-	DetailMode = Other.DetailMode;
-	bVisibleInRayTracing = Other.bVisibleInRayTracing;
-	RayTracingGroupId = Other.RayTracingGroupId;
-	RayTracingGroupCullingPriority = Other.RayTracingGroupCullingPriority;
-
-	Algo::Transform(Other.OverrideMaterials, OverrideMaterials, [](const TSoftObjectPtr<UMaterialInterface>& OverrideMaterial)
-	{
-		return OverrideMaterial.LoadSynchronous();
-	});
-
-	Algo::Transform(Other.RuntimeVirtualTextures, RuntimeVirtualTextures, [](const TSoftObjectPtr<URuntimeVirtualTexture>& RVT)
-	{
-		return RVT.LoadSynchronous();
-	});
-
-	return *this;
-}
-
-void FPCGProceduralISMComponentDescriptor::InitComponent(UPCGProceduralISMComponent* InProceduralISMComponent) const
-{
-	InProceduralISMComponent->SetStaticMesh(StaticMesh.Get());
-	InProceduralISMComponent->OverrideMaterials = OverrideMaterials;
-	InProceduralISMComponent->OverlayMaterial = OverlayMaterial;
-	InProceduralISMComponent->RuntimeVirtualTextures = RuntimeVirtualTextures;
-	InProceduralISMComponent->SetNumInstances(NumInstances);
-	InProceduralISMComponent->SetNumCustomDataFloats(NumCustomFloats);
-	InProceduralISMComponent->SetBounds(WorldBounds);
-	InProceduralISMComponent->SetMinDrawDistance(InstanceMinDrawDistance);
-	InProceduralISMComponent->SetCullDistances(InstanceStartCullDistance, InstanceEndCullDistance);
-	InProceduralISMComponent->ComponentTags = ComponentTags;
-	InProceduralISMComponent->Mobility = Mobility;
-	InProceduralISMComponent->VirtualTextureRenderPassType = VirtualTextureRenderPassType;
-	InProceduralISMComponent->LightingChannels = LightingChannels;
-	InProceduralISMComponent->CustomDepthStencilWriteMask = CustomDepthStencilWriteMask;
-	InProceduralISMComponent->VirtualTextureCullMips = VirtualTextureCullMips;
-	InProceduralISMComponent->TranslucencySortPriority = TranslucencySortPriority;
-	InProceduralISMComponent->CustomDepthStencilValue = CustomDepthStencilValue;
-	InProceduralISMComponent->CastShadow = bCastShadow;
-	InProceduralISMComponent->bEmissiveLightSource = bEmissiveLightSource;
-	InProceduralISMComponent->bCastDynamicShadow = bCastDynamicShadow;
-	InProceduralISMComponent->bCastStaticShadow = bCastStaticShadow;
-	InProceduralISMComponent->bCastContactShadow = bCastContactShadow;
-	InProceduralISMComponent->bCastShadowAsTwoSided = bCastShadowAsTwoSided;
-	InProceduralISMComponent->bCastHiddenShadow = bCastHiddenShadow;
-	InProceduralISMComponent->bReceivesDecals = bReceivesDecals;
-	InProceduralISMComponent->bUseAsOccluder = bUseAsOccluder;
-	InProceduralISMComponent->bRenderCustomDepth = bRenderCustomDepth;
-	InProceduralISMComponent->bEvaluateWorldPositionOffset = bEvaluateWorldPositionOffset;
-	InProceduralISMComponent->bReverseCulling = bReverseCulling;
-	InProceduralISMComponent->WorldPositionOffsetDisableDistance = WorldPositionOffsetDisableDistance;
-	InProceduralISMComponent->NanitePixelProgrammableDistance = NanitePixelProgrammableDistance;
-	InProceduralISMComponent->ShadowCacheInvalidationBehavior = ShadowCacheInvalidationBehavior;
-	InProceduralISMComponent->DetailMode = DetailMode;
-	InProceduralISMComponent->bVisibleInRayTracing = bVisibleInRayTracing;
-	InProceduralISMComponent->RayTracingGroupId = RayTracingGroupId;
-	InProceduralISMComponent->RayTracingGroupCullingPriority = RayTracingGroupCullingPriority;
-}
-
-void FPCGProceduralISMComponentDescriptor::InitFrom(UPCGProceduralISMComponent* InProceduralISMComponent)
-{
-	StaticMesh = InProceduralISMComponent->GetStaticMesh();
-	OverrideMaterials = InProceduralISMComponent->OverrideMaterials;
-	OverlayMaterial = InProceduralISMComponent->OverlayMaterial;
-	RuntimeVirtualTextures = InProceduralISMComponent->RuntimeVirtualTextures;
-	NumInstances = InProceduralISMComponent->GetNumInstances();
-	NumCustomFloats = InProceduralISMComponent->GetNumCustomDataFloats();
-	InProceduralISMComponent->GetBounds(WorldBounds);
-	InstanceMinDrawDistance = InProceduralISMComponent->MinDrawDistance;
-	InProceduralISMComponent->GetCullDistances(InstanceStartCullDistance, InstanceEndCullDistance);
-	ComponentTags = InProceduralISMComponent->ComponentTags;
-	Mobility = InProceduralISMComponent->Mobility;
-	VirtualTextureRenderPassType = InProceduralISMComponent->VirtualTextureRenderPassType;
-	LightingChannels = InProceduralISMComponent->LightingChannels;
-	CustomDepthStencilWriteMask = InProceduralISMComponent->CustomDepthStencilWriteMask;
-	VirtualTextureCullMips = InProceduralISMComponent->VirtualTextureCullMips;
-	TranslucencySortPriority = InProceduralISMComponent->TranslucencySortPriority;
-	CustomDepthStencilValue = InProceduralISMComponent->CustomDepthStencilValue;
-	bCastShadow = InProceduralISMComponent->CastShadow;
-	bEmissiveLightSource = InProceduralISMComponent->bEmissiveLightSource;
-	bCastDynamicShadow = InProceduralISMComponent->bCastDynamicShadow;
-	bCastStaticShadow = InProceduralISMComponent->bCastStaticShadow;
-	bCastContactShadow = InProceduralISMComponent->bCastContactShadow;
-	bCastShadowAsTwoSided = InProceduralISMComponent->bCastShadowAsTwoSided;
-	bCastHiddenShadow = InProceduralISMComponent->bCastHiddenShadow;
-	bReceivesDecals = InProceduralISMComponent->bReceivesDecals;
-	bUseAsOccluder = InProceduralISMComponent->bUseAsOccluder;
-	bRenderCustomDepth = InProceduralISMComponent->bRenderCustomDepth;
-	bEvaluateWorldPositionOffset = InProceduralISMComponent->bEvaluateWorldPositionOffset;
-	bReverseCulling = InProceduralISMComponent->bReverseCulling;
-	WorldPositionOffsetDisableDistance = InProceduralISMComponent->WorldPositionOffsetDisableDistance;
-	NanitePixelProgrammableDistance = InProceduralISMComponent->NanitePixelProgrammableDistance;
-	ShadowCacheInvalidationBehavior = InProceduralISMComponent->ShadowCacheInvalidationBehavior;
-	DetailMode = InProceduralISMComponent->DetailMode;
-	bVisibleInRayTracing = InProceduralISMComponent->bVisibleInRayTracing;
-	RayTracingGroupId = InProceduralISMComponent->RayTracingGroupId;
-	RayTracingGroupCullingPriority = InProceduralISMComponent->RayTracingGroupCullingPriority;
-}
-
-uint32 FPCGProceduralISMComponentDescriptor::ComputeHash() const
-{
-	FArchiveCrc32 CrcArchive;
-
-	Hash = 0; // we don't want the hash to impact the calculation
-	CrcArchive << *this;
-	Hash = CrcArchive.GetCrc();
-
-	return Hash;
-}
-
-bool FPCGProceduralISMComponentDescriptor::operator!=(const FPCGProceduralISMComponentDescriptor& Other) const
-{
-	return !(*this == Other);
-}
-
-bool FPCGProceduralISMComponentDescriptor::operator==(const FPCGProceduralISMComponentDescriptor& Other) const
-{
-	return StaticMesh == Other.StaticMesh
-		&& OverrideMaterials == Other.OverrideMaterials
-		&& OverlayMaterial == Other.OverlayMaterial
-		&& RuntimeVirtualTextures == Other.RuntimeVirtualTextures
-		&& NumInstances == Other.NumInstances
-		&& NumCustomFloats == Other.NumCustomFloats
-		&& WorldBounds == Other.WorldBounds
-		&& InstanceMinDrawDistance == Other.InstanceMinDrawDistance
-		&& InstanceStartCullDistance == Other.InstanceStartCullDistance
-		&& InstanceEndCullDistance == Other.InstanceEndCullDistance
-		&& ComponentTags == Other.ComponentTags
-		&& Mobility == Other.Mobility
-		&& VirtualTextureRenderPassType == Other.VirtualTextureRenderPassType
-		&& GetLightingChannelMaskForStruct(LightingChannels) == GetLightingChannelMaskForStruct(Other.LightingChannels)
-		&& CustomDepthStencilWriteMask == Other.CustomDepthStencilWriteMask
-		&& VirtualTextureCullMips == Other.VirtualTextureCullMips
-		&& TranslucencySortPriority == Other.TranslucencySortPriority
-		&& CustomDepthStencilValue == Other.CustomDepthStencilValue
-		&& bCastShadow == Other.bCastShadow
-		&& bEmissiveLightSource == Other.bEmissiveLightSource
-		&& bCastDynamicShadow == Other.bCastDynamicShadow
-		&& bCastStaticShadow == Other.bCastStaticShadow
-		&& bCastContactShadow == Other.bCastContactShadow
-		&& bCastShadowAsTwoSided == Other.bCastShadowAsTwoSided
-		&& bCastHiddenShadow == Other.bCastHiddenShadow
-		&& bReceivesDecals == Other.bReceivesDecals
-		&& bUseAsOccluder == Other.bUseAsOccluder
-		&& bRenderCustomDepth == Other.bRenderCustomDepth
-		&& bEvaluateWorldPositionOffset == Other.bEvaluateWorldPositionOffset
-		&& bReverseCulling == Other.bReverseCulling
-		&& WorldPositionOffsetDisableDistance == Other.WorldPositionOffsetDisableDistance
-		&& NanitePixelProgrammableDistance == Other.NanitePixelProgrammableDistance
-		&& ShadowCacheInvalidationBehavior == Other.ShadowCacheInvalidationBehavior
-		&& DetailMode == Other.DetailMode
-		&& bVisibleInRayTracing == Other.bVisibleInRayTracing
-		&& RayTracingGroupId == Other.RayTracingGroupId
-		&& RayTracingGroupCullingPriority == Other.RayTracingGroupCullingPriority;
-}
-
 UPCGProceduralISMComponent::UPCGProceduralISMComponent(const FObjectInitializer& ObjectInitializer)
 	: Super(ObjectInitializer)
 {
@@ -368,6 +174,45 @@ void UPCGProceduralISMComponent::OnUnregister()
 }
 #endif
 
+void UPCGProceduralISMComponent::InitFromDescriptor(const FPCGProceduralISMComponentDescriptor& InDescriptor)
+{
+	SetStaticMesh(InDescriptor.StaticMesh.Get());
+	OverrideMaterials = InDescriptor.OverrideMaterials;
+	OverlayMaterial = InDescriptor.OverlayMaterial;
+	RuntimeVirtualTextures = InDescriptor.RuntimeVirtualTextures;
+	SetNumInstances(InDescriptor.NumInstances);
+	SetNumCustomDataFloats(InDescriptor.NumCustomFloats);
+	SetBounds(InDescriptor.WorldBounds);
+	SetMinDrawDistance(InDescriptor.InstanceMinDrawDistance);
+	SetCullDistances(InDescriptor.InstanceStartCullDistance, InDescriptor.InstanceEndCullDistance);
+	ComponentTags = InDescriptor.ComponentTags;
+	Mobility = InDescriptor.Mobility;
+	VirtualTextureRenderPassType = InDescriptor.VirtualTextureRenderPassType;
+	LightingChannels = InDescriptor.LightingChannels;
+	CustomDepthStencilWriteMask = InDescriptor.CustomDepthStencilWriteMask;
+	VirtualTextureCullMips = InDescriptor.VirtualTextureCullMips;
+	TranslucencySortPriority = InDescriptor.TranslucencySortPriority;
+	CustomDepthStencilValue = InDescriptor.CustomDepthStencilValue;
+	CastShadow = InDescriptor.bCastShadow;
+	bEmissiveLightSource = InDescriptor.bEmissiveLightSource;
+	bCastDynamicShadow = InDescriptor.bCastDynamicShadow;
+	bCastStaticShadow = InDescriptor.bCastStaticShadow;
+	bCastContactShadow = InDescriptor.bCastContactShadow;
+	bCastShadowAsTwoSided = InDescriptor.bCastShadowAsTwoSided;
+	bCastHiddenShadow = InDescriptor.bCastHiddenShadow;
+	bReceivesDecals = InDescriptor.bReceivesDecals;
+	bUseAsOccluder = InDescriptor.bUseAsOccluder;
+	bRenderCustomDepth = InDescriptor.bRenderCustomDepth;
+	bEvaluateWorldPositionOffset = InDescriptor.bEvaluateWorldPositionOffset;
+	bReverseCulling = InDescriptor.bReverseCulling;
+	WorldPositionOffsetDisableDistance = InDescriptor.WorldPositionOffsetDisableDistance;
+	ShadowCacheInvalidationBehavior = InDescriptor.ShadowCacheInvalidationBehavior;
+	DetailMode = InDescriptor.DetailMode;
+	bVisibleInRayTracing = InDescriptor.bVisibleInRayTracing;
+	RayTracingGroupId = InDescriptor.RayTracingGroupId;
+	RayTracingGroupCullingPriority = InDescriptor.RayTracingGroupCullingPriority;
+}
+
 FPrimitiveSceneProxy* UPCGProceduralISMComponent::CreateStaticMeshSceneProxy(Nanite::FMaterialAudit& NaniteMaterials, bool bCreateNanite)
 {
 	LLM_SCOPE(ELLMTag::InstancedMesh);
diff --git a/Engine/Plugins/PCG/Source/PCG/Private/Components/PCGProceduralISMComponent.h b/Engine/Plugins/PCG/Source/PCG/Private/Components/PCGProceduralISMComponent.h
--- a/Engine/Plugins/PCG/Source/PCG/Private/Components/PCGProceduralISMComponent.h
+++ b/Engine/Plugins/PCG/Source/PCG/Private/Components/PCGProceduralISMComponent.h
@@ -5,6 +5,8 @@
 #include "Components/StaticMeshComponent.h"
 #include "PCGManagedResource.h"
 
+#include "Components/PCGProceduralISMComponentDescriptor.h"
+
 #include "CoreMinimal.h"
 #include "EngineDefines.h"
 #include "InstanceDataSceneProxy.h"
@@ -23,7 +25,6 @@ class UPCGComponent;
 class UPCGManagedProceduralISMComponent;
 struct FInstancedStaticMeshSceneProxyDesc;
 struct FPCGContext;
-struct FPCGSoftISMComponentDescriptor;
 struct FPrimitiveSceneDesc;
 struct FPrimitiveSceneProxyDesc;
 
@@ -45,6 +46,8 @@ public:
 	//~End UActorComponent Interface
 #endif
 
+	void InitFromDescriptor(const FPCGProceduralISMComponentDescriptor& InDescriptor);
+
 	UFUNCTION(BlueprintCallable, Category = "Components|PCGProceduralISMC")
 	int32 GetNumInstances() const { return NumInstances; }
 
@@ -177,169 +180,6 @@ protected:
 	void SetLocalBounds(const FBox& InWorldBounds) {}
 };
 
-/** Struct that holds properties that can be used to initialize Procedural ISM Components. */
-USTRUCT()
-struct FPCGProceduralISMComponentDescriptor
-{
-	GENERATED_BODY()
-
-public:
-	FPCGProceduralISMComponentDescriptor();
-
-	// Note: We need to explicitly disable warnings on these constructors/operators for clang to be happy with deprecated variables
-	PRAGMA_DISABLE_DEPRECATION_WARNINGS
-	virtual ~FPCGProceduralISMComponentDescriptor() = default;
-	FPCGProceduralISMComponentDescriptor(const FPCGProceduralISMComponentDescriptor&) = default;
-	FPCGProceduralISMComponentDescriptor(FPCGProceduralISMComponentDescriptor&&) = default;
-	FPCGProceduralISMComponentDescriptor& operator=(const FPCGProceduralISMComponentDescriptor&) = default;
-	FPCGProceduralISMComponentDescriptor& operator=(FPCGProceduralISMComponentDescriptor&&) = default;
-	PRAGMA_ENABLE_DEPRECATION_WARNINGS
-
-	FPCGProceduralISMComponentDescriptor& operator=(const FPCGSoftISMComponentDescriptor& Other);
-
-	virtual void InitComponent(UPCGProceduralISMComponent* InProceduralISMComponent) const;
-
-	virtual void InitFrom(UPCGProceduralISMComponent* InProceduralISMComponent);
-
-	virtual uint32 ComputeHash() const;
-
-	bool operator!=(const FPCGProceduralISMComponentDescriptor& Other) const;
-	bool operator==(const FPCGProceduralISMComponentDescriptor& Other) const;
-
-	friend inline uint32 GetTypeHash(const FPCGProceduralISMComponentDescriptor& Key)
-	{
-		return Key.GetTypeHash();
-	}
-
-	uint32 GetTypeHash() const
-	{
-		if (Hash == 0)
-		{
-			ComputeHash();
-		}
-		return Hash;
-	}
-
-public:
-	UPROPERTY()
-	mutable uint32 Hash = 0;
-
-	UPROPERTY()
-	TSoftObjectPtr<UStaticMesh> StaticMesh = nullptr;
-
-	UPROPERTY()
-	TArray<TObjectPtr<UMaterialInterface>> OverrideMaterials;
-
-	UPROPERTY()
-	TObjectPtr<UMaterialInterface> OverlayMaterial;
-
-	UPROPERTY()
-	TArray<TObjectPtr<URuntimeVirtualTexture>> RuntimeVirtualTextures;
-
-	UPROPERTY()
-	int32 NumInstances;
-
-	UPROPERTY()
-	int32 NumCustomFloats;
-
-	UPROPERTY()
-	FBox WorldBounds;
-
-	UPROPERTY()
-	int32 InstanceMinDrawDistance;
-
-	UPROPERTY()
-	int32 InstanceStartCullDistance;
-
-	UPROPERTY()
-	int32 InstanceEndCullDistance;
-
-	UPROPERTY()
-	TArray<FName> ComponentTags;
-
-	UPROPERTY()
-	TEnumAsByte<EComponentMobility::Type> Mobility;
-		
-	UPROPERTY()
-	ERuntimeVirtualTextureMainPassType VirtualTextureRenderPassType;
-	
-	UPROPERTY()
-	FLightingChannels LightingChannels;
-
-	UPROPERTY()
-	ERendererStencilMask CustomDepthStencilWriteMask;
-
-	UPROPERTY()
-	int32 VirtualTextureCullMips;
-
-	UPROPERTY()
-	int32 TranslucencySortPriority;
-
-	UPROPERTY()
-	int32 CustomDepthStencilValue;
-
-	UPROPERTY()
-	uint8 bVisibleInRayTracing : 1;
-
-	UPROPERTY()
-	int32 RayTracingGroupId;
-
-	UPROPERTY()
-	ERayTracingGroupCullingPriority RayTracingGroupCullingPriority;
-
-	UPROPERTY()
-	uint8 bCastShadow : 1;
-
-	UPROPERTY()
-	uint8 bEmissiveLightSource : 1;
-		
-	UPROPERTY()
-	uint8 bCastDynamicShadow : 1;
-
-	UPROPERTY()
-	uint8 bCastStaticShadow : 1;
-
-	UPROPERTY()
-	uint8 bCastContactShadow : 1;
-
-	UPROPERTY()
-	uint8 bCastShadowAsTwoSided : 1;
-
-	UPROPERTY()
-	uint8 bCastHiddenShadow : 1;
-
-	UPROPERTY()
-	uint8 bReceivesDecals : 1;
-
-	UPROPERTY()
-	uint8 bUseAsOccluder : 1;
-
-	UPROPERTY()
-	uint8 bRenderCustomDepth : 1;
-
-	UPROPERTY()
-	uint8 bEvaluateWorldPositionOffset : 1;
-
-	UPROPERTY()
-	uint8 bReverseCulling : 1;
-
-	UPROPERTY()
-	int32 WorldPositionOffsetDisableDistance;
-
-	UPROPERTY()
-	int32 NanitePixelProgrammableDistance;
-
-	UPROPERTY()
-	EShadowCacheInvalidationBehavior ShadowCacheInvalidationBehavior;
-
-	UPROPERTY()
-	TEnumAsByte<enum EDetailMode> DetailMode;
-
-	UE_DEPRECATED(5.6, "Use WorldBounds instead.")
-	UPROPERTY(meta = (DeprecatedProperty, DeprecationMessage = "Use WorldBounds instead."))
-	FBox LocalBounds;
-};
-
 struct FPCGProceduralISMCBuilderParameters
 {
 	friend inline uint32 GetTypeHash(const FPCGProceduralISMCBuilderParameters& Key)
diff --git a/Engine/Plugins/PCG/Source/PCG/Private/Components/PCGProceduralISMComponentDescriptor.cpp b/Engine/Plugins/PCG/Source/PCG/Private/Components/PCGProceduralISMComponentDescriptor.cpp
new file mode 100644
--- /dev/null
+++ b/Engine/Plugins/PCG/Source/PCG/Private/Components/PCGProceduralISMComponentDescriptor.cpp
@@ -0,0 +1,20 @@
+// Copyright Epic Games, Inc. All Rights Reserved.
+
+#include "Components/PCGProceduralISMComponentDescriptor.h"
+
+#include "MeshSelectors/PCGISMDescriptor.h"
+
+FPCGProceduralISMComponentDescriptor& FPCGProceduralISMComponentDescriptor::operator=(const FPCGSoftISMComponentDescriptor& Other)
+{
+	FProceduralISMComponentDescriptor::operator=(static_cast<const FSoftISMComponentDescriptor&>(Other));
+
+	ComponentTags = Other.ComponentTags;
+
+	return *this;
+}
+
+bool FPCGProceduralISMComponentDescriptor::operator==(const FPCGProceduralISMComponentDescriptor& Other) const
+{
+	return FProceduralISMComponentDescriptor::operator==(Other)
+		&& ComponentTags == Other.ComponentTags;
+}
diff --git a/Engine/Plugins/PCG/Source/PCG/Private/Components/PCGProceduralISMComponentDescriptor.h b/Engine/Plugins/PCG/Source/PCG/Private/Components/PCGProceduralISMComponentDescriptor.h
new file mode 100644
--- /dev/null
+++ b/Engine/Plugins/PCG/Source/PCG/Private/Components/PCGProceduralISMComponentDescriptor.h
@@ -0,0 +1,32 @@
+// Copyright Epic Games, Inc. All Rights Reserved.
+
+#pragma once
+
+#include "CoreMinimal.h"
+#include "ISMPartition/ProceduralISMComponentDescriptor.h"
+
+#include "PCGProceduralISMComponentDescriptor.generated.h"
+
+struct FPCGSoftISMComponentDescriptor;
+
+/** Struct that holds properties that can be used to initialize PCG Procedural ISM Components. */
+USTRUCT()
+struct FPCGProceduralISMComponentDescriptor : public FProceduralISMComponentDescriptor
+{
+	GENERATED_BODY()
+
+public:
+	PCG_API FPCGProceduralISMComponentDescriptor& operator=(const FPCGSoftISMComponentDescriptor& Other);
+
+	PCG_API bool operator==(const FPCGProceduralISMComponentDescriptor& Other) const;
+	PCG_API bool operator!=(const FPCGProceduralISMComponentDescriptor& Other) const { return !(*this == Other); }
+
+	friend inline uint32 GetTypeHash(const FPCGProceduralISMComponentDescriptor& Key)
+	{
+		return Key.GetTypeHash();
+	}
+
+public:
+	UPROPERTY()
+	TArray<FName> ComponentTags;
+};
diff --git a/Engine/Source/Runtime/Engine/Private/ISMPartition/ProceduralISMComponentDescriptor.cpp b/Engine/Source/Runtime/Engine/Private/ISMPartition/ProceduralISMComponentDescriptor.cpp
new file mode 100644
--- /dev/null
+++ b/Engine/Source/Runtime/Engine/Private/ISMPartition/ProceduralISMComponentDescriptor.cpp
@@ -0,0 +1,116 @@
+// Copyright Epic Games, Inc. All Rights Reserved.
+
+#include "ISMPartition/ProceduralISMComponentDescriptor.h"
+
+#include "ISMPartition/ISMComponentDescriptor.h"
+#include "Materials/MaterialInterface.h"
+#include "Serialization/ArchiveCrc32.h"
+#include "VT/RuntimeVirtualTexture.h"
+
+FProceduralISMComponentDescriptor::FProceduralISMComponentDescriptor()
+	: NumInstances(0)
+	, NumCustomFloats(0)
+{
+	*this = FSoftISMComponentDescriptor();
+}
+
+FProceduralISMComponentDescriptor& FProceduralISMComponentDescriptor::operator=(const FSoftISMComponentDescriptor& Other)
+{
+	InstanceMinDrawDistance = Other.InstanceMinDrawDistance;
+	InstanceStartCullDistance = Other.InstanceStartCullDistance;
+	InstanceEndCullDistance = Other.InstanceEndCullDistance;
+	OverlayMaterial = Other.OverlayMaterial.LoadSynchronous();
+	StaticMesh = Other.StaticMesh;
+	Mobility = Other.Mobility;
+	VirtualTextureRenderPassType = Other.VirtualTextureRenderPassType;
+	LightingChannels = Other.LightingChannels;
+	CustomDepthStencilWriteMask = Other.CustomDepthStencilWriteMask;
+	VirtualTextureCullMips = Other.VirtualTextureCullMips;
+	TranslucencySortPriority = Other.TranslucencySortPriority;
+	CustomDepthStencilValue = Other.CustomDepthStencilValue;
+	bCastShadow = Other.bCastShadow;
+	bEmissiveLightSource = Other.bEmissiveLightSource;
+	bCastDynamicShadow = Other.bCastDynamicShadow;
+	bCastStaticShadow = Other.bCastStaticShadow;
+	bCastContactShadow = Other.bCastContactShadow;
+	bCastShadowAsTwoSided = Other.bCastShadowAsTwoSided;
+	bCastHiddenShadow = Other.bCastHiddenShadow;
+	bReceivesDecals = Other.bReceivesDecals;
+	bUseAsOccluder = Other.bUseAsOccluder;
+	bRenderCustomDepth = Other.bRenderCustomDepth;
+	bEvaluateWorldPositionOffset = Other.bEvaluateWorldPositionOffset;
+	bReverseCulling = Other.bReverseCulling;
+	WorldPositionOffsetDisableDistance = Other.WorldPositionOffsetDisableDistance;
+	ShadowCacheInvalidationBehavior = Other.ShadowCacheInvalidationBehavior;
+	DetailMode = Other.DetailMode;
+	bVisibleInRayTracing = Other.bVisibleInRayTracing;
+	RayTracingGroupId = Other.RayTracingGroupId;
+	RayTracingGroupCullingPriority = Other.RayTracingGroupCullingPriority;
+
+	Algo::Transform(Other.OverrideMaterials, OverrideMaterials, [](const TSoftObjectPtr<UMaterialInterface>& OverrideMaterial)
+	{
+		return OverrideMaterial.LoadSynchronous();
+	});
+
+	Algo::Transform(Other.RuntimeVirtualTextures, RuntimeVirtualTextures, [](const TSoftObjectPtr<URuntimeVirtualTexture>& RVT)
+	{
+		return RVT.LoadSynchronous();
+	});
+
+	return *this;
+}
+
+uint32 FProceduralISMComponentDescriptor::ComputeHash() const
+{
+	FArchiveCrc32 CrcArchive;
+
+	Hash = 0; // we don't want the hash to impact the calculation
+	CrcArchive << *this;
+	Hash = CrcArchive.GetCrc();
+
+	return Hash;
+}
+
+bool FProceduralISMComponentDescriptor::operator!=(const FProceduralISMComponentDescriptor& Other) const
+{
+	return !(*this == Other);
+}
+
+bool FProceduralISMComponentDescriptor::operator==(const FProceduralISMComponentDescriptor& Other) const
+{
+	return StaticMesh == Other.StaticMesh
+		&& OverrideMaterials == Other.OverrideMaterials
+		&& OverlayMaterial == Other.OverlayMaterial
+		&& RuntimeVirtualTextures == Other.RuntimeVirtualTextures
+		&& NumInstances == Other.NumInstances
+		&& NumCustomFloats == Other.NumCustomFloats
+		&& WorldBounds == Other.WorldBounds
+		&& InstanceMinDrawDistance == Other.InstanceMinDrawDistance
+		&& InstanceStartCullDistance == Other.InstanceStartCullDistance
+		&& InstanceEndCullDistance == Other.InstanceEndCullDistance
+		&& Mobility == Other.Mobility
+		&& VirtualTextureRenderPassType == Other.VirtualTextureRenderPassType
+		&& GetLightingChannelMaskForStruct(LightingChannels) == GetLightingChannelMaskForStruct(Other.LightingChannels)
+		&& CustomDepthStencilWriteMask == Other.CustomDepthStencilWriteMask
+		&& VirtualTextureCullMips == Other.VirtualTextureCullMips
+		&& TranslucencySortPriority == Other.TranslucencySortPriority
+		&& CustomDepthStencilValue == Other.CustomDepthStencilValue
+		&& bCastShadow == Other.bCastShadow
+		&& bEmissiveLightSource == Other.bEmissiveLightSource
+		&& bCastDynamicShadow == Other.bCastDynamicShadow
+		&& bCastStaticShadow == Other.bCastStaticShadow
+		&& bCastContactShadow == Other.bCastContactShadow
+		&& bCastShadowAsTwoSided == Other.bCastShadowAsTwoSided
+		&& bCastHiddenShadow == Other.bCastHiddenShadow
+		&& bReceivesDecals == Other.bReceivesDecals
+		&& bUseAsOccluder == Other.bUseAsOccluder
+		&& bRenderCustomDepth == Other.bRenderCustomDepth
+		&& bEvaluateWorldPositionOffset == Other.bEvaluateWorldPositionOffset
+		&& bReverseCulling == Other.bReverseCulling
+		&& WorldPositionOffsetDisableDistance == Other.WorldPositionOffsetDisableDistance
+		&& ShadowCacheInvalidationBehavior == Other.ShadowCacheInvalidationBehavior
+		&& DetailMode == Other.DetailMode
+		&& bVisibleInRayTracing == Other.bVisibleInRayTracing
+		&& RayTracingGroupId == Other.RayTracingGroupId
+		&& RayTracingGroupCullingPriority == Other.RayTracingGroupCullingPriority;
+}
diff --git a/Engine/Source/Runtime/Engine/Public/ISMPartition/ProceduralISMComponentDescriptor.h b/Engine/Source/Runtime/Engine/Public/ISMPartition/ProceduralISMComponentDescriptor.h
new file mode 100644
--- /dev/null
+++ b/Engine/Source/Runtime/Engine/Public/ISMPartition/ProceduralISMComponentDescriptor.h
@@ -0,0 +1,168 @@
+// Copyright Epic Games, Inc. All Rights Reserved.
+
+#pragma once
+
+#include "CoreMinimal.h"
+#include "Engine/EngineTypes.h"
+
+#include "ProceduralISMComponentDescriptor.generated.h"
+
+class UStaticMesh;
+class UMaterialInterface;
+class URuntimeVirtualTexture;
+enum class ERayTracingGroupCullingPriority : uint8;
+enum class ERendererStencilMask : uint8;
+enum class ERuntimeVirtualTextureMainPassType : uint8;
+enum class EShadowCacheInvalidationBehavior : uint8;
+enum EDetailMode : int;
+namespace EComponentMobility { enum Type : int; }
+struct FSoftISMComponentDescriptor;
+
+/** Struct that holds properties that can be used to initialize Procedural ISM Components. */
+USTRUCT()
+struct FProceduralISMComponentDescriptor
+{
+	GENERATED_BODY()
+
+public:
+	ENGINE_API FProceduralISMComponentDescriptor();
+
+	// Note: We need to explicitly disable warnings on these constructors/operators for clang to be happy with deprecated variables
+	PRAGMA_DISABLE_DEPRECATION_WARNINGS
+	ENGINE_API virtual ~FProceduralISMComponentDescriptor() = default;
+	FProceduralISMComponentDescriptor(const FProceduralISMComponentDescriptor&) = default;
+	FProceduralISMComponentDescriptor(FProceduralISMComponentDescriptor&&) = default;
+	FProceduralISMComponentDescriptor& operator=(const FProceduralISMComponentDescriptor&) = default;
+	FProceduralISMComponentDescriptor& operator=(FProceduralISMComponentDescriptor&&) = default;
+	PRAGMA_ENABLE_DEPRECATION_WARNINGS
+
+	ENGINE_API FProceduralISMComponentDescriptor& operator=(const FSoftISMComponentDescriptor& Other);
+
+	ENGINE_API virtual uint32 ComputeHash() const;
+
+	ENGINE_API bool operator!=(const FProceduralISMComponentDescriptor& Other) const;
+	ENGINE_API bool operator==(const FProceduralISMComponentDescriptor& Other) const;
+
+	friend inline uint32 GetTypeHash(const FProceduralISMComponentDescriptor& Key)
+	{
+		return Key.GetTypeHash();
+	}
+
+	uint32 GetTypeHash() const
+	{
+		if (Hash == 0)
+		{
+			ComputeHash();
+		}
+		return Hash;
+	}
+
+public:
+	UPROPERTY()
+	mutable uint32 Hash = 0;
+
+	UPROPERTY()
+	TSoftObjectPtr<UStaticMesh> StaticMesh = nullptr;
+
+	UPROPERTY()
+	TArray<TObjectPtr<UMaterialInterface>> OverrideMaterials;
+
+	UPROPERTY()
+	TObjectPtr<UMaterialInterface> OverlayMaterial;
+
+	UPROPERTY()
+	TArray<TObjectPtr<URuntimeVirtualTexture>> RuntimeVirtualTextures;
+
+	UPROPERTY()
+	int32 NumInstances;
+
+	UPROPERTY()
+	int32 NumCustomFloats;
+
+	UPROPERTY()
+	FBox WorldBounds;
+
+	UPROPERTY()
+	int32 InstanceMinDrawDistance;
+
+	UPROPERTY()
+	int32 InstanceStartCullDistance;
+
+	UPROPERTY()
+	int32 InstanceEndCullDistance;
+
+	UPROPERTY()
+	TEnumAsByte<EComponentMobility::Type> Mobility;
+		
+	UPROPERTY()
+	ERuntimeVirtualTextureMainPassType VirtualTextureRenderPassType;
+	
+	UPROPERTY()
+	FLightingChannels LightingChannels;
+
+	UPROPERTY()
+	ERendererStencilMask CustomDepthStencilWriteMask;
+
+	UPROPERTY()
+	int32 VirtualTextureCullMips;
+
+	UPROPERTY()
+	int32 TranslucencySortPriority;
+
+	UPROPERTY()
+	int32 CustomDepthStencilValue;
+
+	UPROPERTY()
+	int32 RayTracingGroupId;
+
+	UPROPERTY()
+	ERayTracingGroupCullingPriority RayTracingGroupCullingPriority;
+
+	UPROPERTY()
+	uint8 bVisibleInRayTracing : 1;
+
+	UPROPERTY()
+	uint8 bCastShadow : 1;
+
+	UPROPERTY()
+	uint8 bEmissiveLightSource : 1;
+		
+	UPROPERTY()
+	uint8 bCastDynamicShadow : 1;
+
+	UPROPERTY()
+	uint8 bCastStaticShadow : 1;
+
+	UPROPERTY()
+	uint8 bCastContactShadow : 1;
+
+	UPROPERTY()
+	uint8 bCastShadowAsTwoSided : 1;
+
+	UPROPERTY()
+	uint8 bCastHiddenShadow : 1;
+
+	UPROPERTY()
+	uint8 bReceivesDecals : 1;
+
+	UPROPERTY()
+	uint8 bUseAsOccluder : 1;
+
+	UPROPERTY()
+	uint8 bRenderCustomDepth : 1;
+
+	UPROPERTY()
+	uint8 bEvaluateWorldPositionOffset : 1;
+
+	UPROPERTY()
+	uint8 bReverseCulling : 1;
+
+	UPROPERTY()
+	int32 WorldPositionOffsetDisableDistance;
+
+	UPROPERTY()
+	EShadowCacheInvalidationBehavior ShadowCacheInvalidationBehavior;
+
+	UPROPERTY()
+	TEnumAsByte<enum EDetailMode> DetailMode;
+};
