[PCG] Partially backout CL 42428247 which was causing in memory refs to not be properly updated

diff --git a/Engine/Plugins/PCG/Source/PCG/Private/PCGAssetExporterUtils.cpp b/Engine/Plugins/PCG/Source/PCG/Private/PCGAssetExporterUtils.cpp
--- a/Engine/Plugins/PCG/Source/PCG/Private/PCGAssetExporterUtils.cpp
+++ b/Engine/Plugins/PCG/Source/PCG/Private/PCGAssetExporterUtils.cpp
@@ -10,7 +10,6 @@
 #include "FileHelpers.h"
 #include "IContentBrowserSingleton.h"
 #include "AssetRegistry/AssetRegistryModule.h"
-#include "Engine/Engine.h"
 #endif // WITH_EDITOR
 
 #include UE_INLINE_GENERATED_CPP_BY_NAME(PCGAssetExporterUtils)
@@ -111,18 +110,21 @@ UPackage* UPCGAssetExporterUtils::CreateAsset(const UClass* AssetClass, const FP
 	UPackage* Package = FPackageName::DoesPackageExist(PackageName) ? LoadPackage(nullptr, *PackageName, LOAD_None) : FindPackage(nullptr, *PackageName);
 
 	UObject* Asset = nullptr;
-	UObject* PreviousAsset = nullptr;
 	bool NewAssetCreated = false;
 
 	if (Package)
 	{
-		PreviousAsset = FindObjectFast<UObject>(Package, *AssetName);
-		if (PreviousAsset)
+		UObject* Object = FindObjectFast<UObject>(Package, *AssetName);
+		if (Object && Object->GetClass() != AssetClass)
 		{
-			PreviousAsset->SetFlags(RF_Transient);
-			PreviousAsset->Rename(nullptr, GetTransientPackage(), REN_DontCreateRedirectors | REN_NonTransactional);
+			Object->SetFlags(RF_Transient);
+			Object->Rename(nullptr, GetTransientPackage(), REN_DontCreateRedirectors | REN_NonTransactional);
 			NewAssetCreated = true;
 		}
+		else
+		{
+			Asset = Object;
+		}
 	}
 	else
 	{
@@ -153,16 +155,15 @@ UPackage* UPCGAssetExporterUtils::CreateAsset(const UClass* AssetClass, const FP
 			FAssetRegistryModule::AssetCreated(Asset);
 		}
 
-		const bool bSuccess = ExportFunc(PackageName, Asset);
-
-		if (PreviousAsset)
+		if (ExportFunc(PackageName, Asset))
 		{
-			TMap<UObject*, UObject*> OldToNew;
-			OldToNew.Emplace(PreviousAsset, Asset);
-			GEngine->NotifyToolsOfObjectReplacement(OldToNew);
+			// Make sure everybody knows we changed the asset
+			if (!NewAssetCreated)
+			{
+				FCoreUObjectDelegates::BroadcastOnObjectModified(Asset);
+			}
 		}
-
-		if (!bSuccess)
+		else
 		{
 			return nullptr;
 		}
