[PCG] Level to Asset: Instead of checking if Material instance is in the transient package, check that it is an Asset (which is a public and can be referenced by our PCG Data Asset)
- IsAsset() still checks for the transient package but also excludes transient material instances that are not outered to the transient package

diff --git a/Engine/Plugins/PCG/Source/PCGEditor/Private/PCGLevelToAsset.cpp b/Engine/Plugins/PCG/Source/PCGEditor/Private/PCGLevelToAsset.cpp
--- a/Engine/Plugins/PCG/Source/PCGEditor/Private/PCGLevelToAsset.cpp
+++ b/Engine/Plugins/PCG/Source/PCGEditor/Private/PCGLevelToAsset.cpp
@@ -487,10 +487,10 @@ bool UPCGLevelToAsset::BP_ExportWorld_Implementation(UWorld* World, const FStrin
 
 			if (!MeshMaterials.IsEmpty())
 			{
-				// Avoid references to transient materials, in this case we reference the first non transient parent
+				// Avoid references to materials that are not public assets to avoid invalid references, in this case we reference the first parent asset we find.
 				UMaterialInterface* Material = MeshMaterials[0];
 				UMaterialInstance* MaterialInstance = Cast<UMaterialInstance>(Material);
-				while (MaterialInstance && MaterialInstance->GetPackage() == GetTransientPackage())
+				while (MaterialInstance && !MaterialInstance->IsAsset())
 				{
 					Material = MaterialInstance->Parent;
 					MaterialInstance = Cast<UMaterialInstance>(Material);
