[PCG] Builder: improve dirty package filtering to exclude WP umaps by default
- Always fallback on: UPCGWorldPartitionBuilderHelper::AddAssetToSubmitAllowList to override default behavior

diff --git a/Engine/Plugins/PCG/Source/PCGEditor/Private/WorldPartitionBuilder/PCGWorldPartitionBuilder.cpp b/Engine/Plugins/PCG/Source/PCGEditor/Private/WorldPartitionBuilder/PCGWorldPartitionBuilder.cpp
--- a/Engine/Plugins/PCG/Source/PCGEditor/Private/WorldPartitionBuilder/PCGWorldPartitionBuilder.cpp
+++ b/Engine/Plugins/PCG/Source/PCGEditor/Private/WorldPartitionBuilder/PCGWorldPartitionBuilder.cpp
@@ -82,7 +82,7 @@ struct FPCGWorldPartitionBuilder
 	static FSoftObjectPath GetActorPathFromComponentPtr(TSoftObjectPtr<UPCGComponent> InComponentPtr);
 
 	/** Get the list of Dirty packages */
-	static TArray<UPackage*> GetDirtyPackages(const TSet<FString>& AllowedAssetPackages);
+	static TArray<UPackage*> GetDirtyPackages(const UWorld* InWorld, const TSet<FString>& AllowedAssetPackages);
 
 	/** Set if Partition Actor creation should be disabled for world */
 	static void SetDisablePartitionActorCreationForWorld(UWorld* InWorld, bool bDisable);
@@ -667,7 +667,7 @@ bool UPCGWorldPartitionBuilder::RunEntireWorld(UWorld* World, FPackageSourceCont
 		// Wait for other executions that might have spawned from generating the components
 		FPCGWorldPartitionBuilder::WaitOnGraphExecutor(World);
 		
-		DirtyPackages = FPCGWorldPartitionBuilder::GetDirtyPackages(AllowedAssetPackages);
+		DirtyPackages = FPCGWorldPartitionBuilder::GetDirtyPackages(World, AllowedAssetPackages);
 	}
 	
 	// 5. Revert unchanged actors
@@ -1045,7 +1045,7 @@ bool UPCGWorldPartitionBuilder::PostRun(UWorld* World, FPackageSourceControlHelp
 
 bool UPCGWorldPartitionBuilder::SavePerCellPackages(UWorld* InWorld, FPackageSourceControlHelper& PackageHelper, TSet<FString> CellDeletedPackages)
 {
-	TArray<UPackage*> DirtyPackages = FPCGWorldPartitionBuilder::GetDirtyPackages(AllowedAssetPackages);
+	TArray<UPackage*> DirtyPackages = FPCGWorldPartitionBuilder::GetDirtyPackages(InWorld, AllowedAssetPackages);
 	TArray<UPackage*> PackagesToDelete;
 	UWorldPartition* WorldPartition = InWorld->GetWorldPartition();
 
@@ -1648,23 +1648,47 @@ void FPCGWorldPartitionBuilder::SetChainedDispatchToLocalComponents(UWorld* InWo
 	PCGSubsystem->SetChainedDispatchToLocalComponents(bInChainedDispatch);
 }
 
-TArray<UPackage*> FPCGWorldPartitionBuilder::GetDirtyPackages(const TSet<FString>& AllowedAssetPackages)
+TArray<UPackage*> FPCGWorldPartitionBuilder::GetDirtyPackages(const UWorld* InWorld, const TSet<FString>& AllowedAssetPackages)
 {
 	TArray<UPackage*> DirtyPackages;
-	FEditorFileUtils::GetDirtyWorldPackages(DirtyPackages);
+	TSet<UPackage*> ExternalPackages;
 
-	if (!AllowedAssetPackages.IsEmpty())
+	const UPackage* WorldPackage = InWorld->GetPackage();
+	const bool bIsPartitioned = InWorld->IsPartitionedWorld();
+
+	ExternalPackages.Append(WorldPackage->GetExternalPackages());
+
+	FEditorFileUtils::GetDirtyWorldPackages(DirtyPackages, /*ShouldIgnorePackageFunction=*/[bIsPartitioned, WorldPackage, &AllowedAssetPackages, &ExternalPackages](UPackage* DirtyPackage)
 	{
-		TArray<UPackage*> DirtyAssetPackages;
-		FEditorFileUtils::GetDirtyContentPackages(DirtyAssetPackages);
+		const bool bExplicitelyAllowedAsset = AllowedAssetPackages.Contains(DirtyPackage->GetName());
 
-		for (UPackage* DirtyAssetPackage : DirtyAssetPackages)
+		// Don't allow modification of World Partition umap unless allowed in AllowedAssetPackages
+		if (WorldPackage == DirtyPackage)
 		{
-			if (AllowedAssetPackages.Contains(DirtyAssetPackage->GetName()))
-			{
-				DirtyPackages.Add(DirtyAssetPackage);
-			}
+			return bIsPartitioned ? !bExplicitelyAllowedAsset : false;
 		}
+		
+		if (bIsPartitioned)
+		{
+			// WP world should only allow current world external packages (actors/folders) or explicitely allowed assets to be dirtied.
+			return !bExplicitelyAllowedAsset && !ExternalPackages.Contains(DirtyPackage);
+		}
+		else
+		{
+			// Non-WP world should only allow umap to be dirtied and explicitely allowed assets.
+			return !bExplicitelyAllowedAsset;
+		}
+	});
+	
+	if (!AllowedAssetPackages.IsEmpty())
+	{
+		TArray<UPackage*> DirtyAssetPackages;
+		FEditorFileUtils::GetDirtyContentPackages(DirtyAssetPackages, /*ShouldIgnorePackageFunction=*/[&AllowedAssetPackages](UPackage* DirtyPackage)
+		{
+			return !AllowedAssetPackages.Contains(DirtyPackage->GetName());
+		});
+
+		DirtyPackages.Append(DirtyAssetPackages);
 	}
 
 	DirtyPackages.Remove(nullptr);
