[PCG] Unify primitive factory across all primitive types, which will allow usages to be fully generalized.

#rnx

diff --git a/Engine/Plugins/PCG/Source/PCG/Private/Compute/DataInterfaces/Elements/PCGSkinnedMeshSpawnerDataInterface.cpp b/Engine/Plugins/PCG/Source/PCG/Private/Compute/DataInterfaces/Elements/PCGSkinnedMeshSpawnerDataInterface.cpp
--- a/Engine/Plugins/PCG/Source/PCG/Private/Compute/DataInterfaces/Elements/PCGSkinnedMeshSpawnerDataInterface.cpp
+++ b/Engine/Plugins/PCG/Source/PCG/Private/Compute/DataInterfaces/Elements/PCGSkinnedMeshSpawnerDataInterface.cpp
@@ -446,7 +446,6 @@ void UPCGSkinnedMeshSpawnerDataProvider::Reset()
 	PrimitiveDescriptors.Empty();
 	CustomFloatCount = 0;
 	bPrimitivesSetUp = false;
-	NumPrimitivesProcessed = 0;
 	NumPrimitivesSetup = 0;
 	bRegisteredPrimitives = false;
 	bStaticMeshesLoaded = false;
@@ -636,6 +635,11 @@ void UPCGSkinnedMeshSpawnerDataProvider::CreatePrimitiveDescriptors(FPCGContext*
 
 bool UPCGSkinnedMeshSpawnerDataProvider::SetupPrimitives(FPCGContext* InContext, UPCGDataBinding* InBinding)
 {
+	if (PrimitiveDescriptors.IsEmpty())
+	{
+		return true;
+	}
+
 	const UPCGSkinnedMeshSpawnerSettings* Settings = CastChecked<UPCGSkinnedMeshSpawnerSettings>(ProducerKernel->GetSettings());
 
 	AActor* TargetActor = Settings->TargetActor.Get() ? Settings->TargetActor.Get() : InContext->GetTargetActor(nullptr);
@@ -650,32 +654,41 @@ bool UPCGSkinnedMeshSpawnerDataProvider::SetupPrimitives(FPCGContext* InContext,
 		return true;
 	}
 
-	while (NumPrimitivesProcessed < PrimitiveDescriptors.Num())
-	{
-		const FPCGSoftSkinnedMeshComponentDescriptor& Desc = PrimitiveDescriptors[NumPrimitivesProcessed];
+	TSharedPtr<FPCGPrimitiveFactoryISKMC> Factory;
 
-		TSharedPtr<FPCGPrimitiveFactoryISKMC> Primitive = MakeShared<FPCGPrimitiveFactoryISKMC>();
+	if (!Primitives->PrimitiveFactory)
+	{
+		Factory = MakeShared<FPCGPrimitiveFactoryISKMC>();
 
-		if (Primitive->Create(InContext, TargetActor, Desc, Settings->UID, Settings->GetSettingsCrc()))
-		{
-			Primitives->PrimitiveFactories.Add(Primitive);
+		Primitives->PrimitiveFactory = Factory;
 
-			if (Settings->bApplyMeshBoundsToPoints)
-			{
-				PrimitiveMeshBounds.Add(Desc.SkinnedAsset->GetBounds().GetBox());
-			}
+		FPCGPrimitiveFactoryISKMC::FParameters Params;
+		Params.Descriptors = PrimitiveDescriptors;
+		Params.TargetActor = TargetActor;
 
-			++NumPrimitivesSetup;
-		}
+		Factory->Initialize(MoveTemp(Params));
+	}
+	else
+	{
+		Factory = StaticCastSharedPtr<FPCGPrimitiveFactoryISKMC>(Primitives->PrimitiveFactory);
+	}
 
-		++NumPrimitivesProcessed;
+	if (!Factory->Create(InContext))
+	{
+		// Not finished, continue next tick.
+		return false;
+	}
 
-		if (NumPrimitivesProcessed < PrimitiveDescriptors.Num() && InContext->AsyncState.ShouldStop())
+	if (Settings->bApplyMeshBoundsToPoints)
+	{
+		for (int32 PrimitiveIndex = 0; PrimitiveIndex < Factory->GetNumPrimitives(); ++PrimitiveIndex)
 		{
-			return false;
+			PrimitiveMeshBounds.Add(Factory->GetMeshBounds(PrimitiveIndex));
 		}
 	}
 
+	NumPrimitivesSetup = Factory->GetNumPrimitives();
+
 	return true;
 }
 
diff --git a/Engine/Plugins/PCG/Source/PCG/Private/Compute/DataInterfaces/Elements/PCGSkinnedMeshSpawnerDataInterface.h b/Engine/Plugins/PCG/Source/PCG/Private/Compute/DataInterfaces/Elements/PCGSkinnedMeshSpawnerDataInterface.h
--- a/Engine/Plugins/PCG/Source/PCG/Private/Compute/DataInterfaces/Elements/PCGSkinnedMeshSpawnerDataInterface.h
+++ b/Engine/Plugins/PCG/Source/PCG/Private/Compute/DataInterfaces/Elements/PCGSkinnedMeshSpawnerDataInterface.h
@@ -86,8 +86,6 @@ public:
 
 	bool bStaticMeshesLoaded = false;
 
-	int32 NumPrimitivesProcessed = 0;
-
 	bool bPrimitivesSetUp = false;
 
 	/** Number of primitives created during primitive setup. */
diff --git a/Engine/Plugins/PCG/Source/PCG/Private/Compute/DataInterfaces/Elements/PCGStaticMeshSpawnerDataInterface.cpp b/Engine/Plugins/PCG/Source/PCG/Private/Compute/DataInterfaces/Elements/PCGStaticMeshSpawnerDataInterface.cpp
--- a/Engine/Plugins/PCG/Source/PCG/Private/Compute/DataInterfaces/Elements/PCGStaticMeshSpawnerDataInterface.cpp
+++ b/Engine/Plugins/PCG/Source/PCG/Private/Compute/DataInterfaces/Elements/PCGStaticMeshSpawnerDataInterface.cpp
@@ -26,6 +26,7 @@
 #include "ShaderParameterMetadataBuilder.h"
 #include "SystemTextures.h"
 #include "Engine/StaticMesh.h"
+#include "Engine/World.h"
 #include "ComputeFramework/ShaderParamTypeDefinition.h"
 
 #include UE_INLINE_GENERATED_CPP_BY_NAME(PCGStaticMeshSpawnerDataInterface)
@@ -455,7 +456,6 @@ void UPCGStaticMeshSpawnerDataProvider::Reset()
 	PrimitiveDescriptors.Empty();
 	CustomFloatCount = 0;
 	bPrimitivesSetUp = false;
-	NumPrimitivesProcessed = 0;
 	NumPrimitivesSetup = 0;
 	bRegisteredPrimitives = false;
 	bStaticMeshesLoaded = false;
@@ -721,6 +721,11 @@ void UPCGStaticMeshSpawnerDataProvider::CreatePrimitiveDescriptors(FPCGContext*
 
 bool UPCGStaticMeshSpawnerDataProvider::SetupPrimitives(FPCGContext* InContext, UPCGDataBinding* InBinding)
 {
+	if (PrimitiveDescriptors.IsEmpty())
+	{
+		return true;
+	}
+
 	const UPCGStaticMeshSpawnerSettings* Settings = CastChecked<UPCGStaticMeshSpawnerSettings>(ProducerKernel->GetSettings());
 
 	AActor* TargetActor = Settings->TargetActor.Get() ? Settings->TargetActor.Get() : InContext->GetTargetActor(nullptr);
@@ -735,32 +740,50 @@ bool UPCGStaticMeshSpawnerDataProvider::SetupPrimitives(FPCGContext* InContext,
 		return true;
 	}
 
-	while (NumPrimitivesProcessed < PrimitiveDescriptors.Num())
+	// Ensure SMs are compiled, required to prevent ensure when partially built SM is accessed later in FastGeo code.
+	for (const FPCGProceduralISMComponentDescriptor& Desc : PrimitiveDescriptors)
 	{
-		const FPCGProceduralISMComponentDescriptor& Desc = PrimitiveDescriptors[NumPrimitivesProcessed];
-
-		TSharedPtr<FPCGPrimitiveFactoryPISMC> Factory = MakeShared<FPCGPrimitiveFactoryPISMC>();
-		
-		if (Factory->Create(InContext, TargetActor, Desc, Settings->UID))
+		if (Desc.StaticMesh->IsCompiling())
 		{
-			Primitives->PrimitiveFactories.Add(Factory);
+			return false;
+		}
+	}
 
-			if (Settings->bApplyMeshBoundsToPoints)
-			{
-				PrimitiveMeshBounds.Add(Desc.StaticMesh->GetBoundingBox());
-			}
+	TSharedPtr<IPCGPrimitiveFactoryISMBase> Factory;
 
-			++NumPrimitivesSetup;
-		}
+	if (!Primitives->PrimitiveFactory)
+	{
+		Factory = MakeShared<FPCGPrimitiveFactoryPISMC>();
 
-		++NumPrimitivesProcessed;
+		Primitives->PrimitiveFactory = Factory;
 
-		if (NumPrimitivesProcessed < PrimitiveDescriptors.Num() && InContext->AsyncState.ShouldStop())
+		IPCGPrimitiveFactoryISMBase::FParameters Params;
+		Params.Descriptors = PrimitiveDescriptors;
+		Params.TargetActor = TargetActor;
+
+		Factory->Initialize(MoveTemp(Params));
+	}
+	else
+	{
+		Factory = StaticCastSharedPtr<IPCGPrimitiveFactoryISMBase>(Primitives->PrimitiveFactory);
+	}
+
+	if (!Factory->Create(InContext))
+	{
+		// Not finished, continue next tick.
+		return false;
+	}
+
+	if (Settings->bApplyMeshBoundsToPoints)
+	{
+		for (int32 PrimitiveIndex = 0; PrimitiveIndex < Factory->GetNumPrimitives(); ++PrimitiveIndex)
 		{
-			return false;
+			PrimitiveMeshBounds.Add(Factory->GetMeshBounds(PrimitiveIndex));
 		}
 	}
 
+	NumPrimitivesSetup = Factory->GetNumPrimitives();
+
 	return true;
 }
 
diff --git a/Engine/Plugins/PCG/Source/PCG/Private/Compute/DataInterfaces/Elements/PCGStaticMeshSpawnerDataInterface.h b/Engine/Plugins/PCG/Source/PCG/Private/Compute/DataInterfaces/Elements/PCGStaticMeshSpawnerDataInterface.h
--- a/Engine/Plugins/PCG/Source/PCG/Private/Compute/DataInterfaces/Elements/PCGStaticMeshSpawnerDataInterface.h
+++ b/Engine/Plugins/PCG/Source/PCG/Private/Compute/DataInterfaces/Elements/PCGStaticMeshSpawnerDataInterface.h
@@ -48,7 +48,7 @@ public:
 
 	void CreatePrimitiveDescriptors(FPCGContext* InContext, UPCGDataBinding* InBinding);
 
-	/** Setup Procedural ISMCs. Returns true if at least one component was set up. */
+	/** Setup primitives. Returns true if at least one component was set up. */
 	virtual bool SetupPrimitives(FPCGContext* InContext, UPCGDataBinding* InBinding);
 
 public:
@@ -91,8 +91,6 @@ public:
 
 	bool bStaticMeshesLoaded = false;
 
-	int32 NumPrimitivesProcessed = 0;
-
 	bool bPrimitivesSetUp = false;
 
 	/** Number of primitives created during primitive setup. */
diff --git a/Engine/Plugins/PCG/Source/PCG/Private/Compute/DataInterfaces/PCGInstanceDataInterface.cpp b/Engine/Plugins/PCG/Source/PCG/Private/Compute/DataInterfaces/PCGInstanceDataInterface.cpp
--- a/Engine/Plugins/PCG/Source/PCG/Private/Compute/DataInterfaces/PCGInstanceDataInterface.cpp
+++ b/Engine/Plugins/PCG/Source/PCG/Private/Compute/DataInterfaces/PCGInstanceDataInterface.cpp
@@ -161,48 +161,38 @@ bool UPCGInstanceDataProvider::PrepareForExecute_GameThread(UPCGDataBinding* InB
 		return true;
 	}
 
-	if (!ensure(FoundPrimitives->PrimitiveFactories.Num() <= PCGComputeConstants::MAX_PRIMITIVE_COMPONENTS_PER_SPAWNER))
+	TSharedPtr<IPCGRuntimePrimitiveFactory>& Factory = FoundPrimitives->PrimitiveFactory;
+	if (!Factory)
 	{
-		// Last resort - should be clamped earlier during setup and not come to this.
-		FoundPrimitives->PrimitiveFactories.SetNum(PCGComputeConstants::MAX_PRIMITIVE_COMPONENTS_PER_SPAWNER);
-	}
-
-	bool bAllReady = true;
-
-	for (TSharedPtr<IPCGRuntimePrimitiveFactory>& Factory : FoundPrimitives->PrimitiveFactories)
-	{
-		if (ensure(Factory))
-		{
-			// This can tick the creation of the render proxy for some factories so call on all factories rather than early'ing out.
-			bAllReady &= Factory->IsRenderStateCreated();
-		}
+		return true;
 	}
 
-	if (!bAllReady)
+	if (!Factory->IsRenderStateCreated())
 	{
 		UE_LOG(LogPCG, Verbose, TEXT("FPCGInstanceDataProvider: One or more scene proxies were not ready. Will try again on the next tick."));
 		return false;
 	}
 
 	NumInstancesAllPrimitives = 0;
-	Primitives.Reserve(FoundPrimitives->PrimitiveFactories.Num());
 
-	for (TSharedPtr<IPCGRuntimePrimitiveFactory> Factory : FoundPrimitives->PrimitiveFactories)
+	int32 NumPrimitives = Factory->GetNumPrimitives();
+	if (NumPrimitives > PCGComputeConstants::MAX_PRIMITIVE_COMPONENTS_PER_SPAWNER)
 	{
-		if (!ensure(Factory))
-		{
-			continue;
-		}
+		UE_LOG(LogPCG, Error, TEXT("FPCGInstanceDataProvider: %d primitives provided which is more than the maximum limit (%d), some primitives will not be created."),
+			NumPrimitives, PCGComputeConstants::MAX_PRIMITIVE_COMPONENTS_PER_SPAWNER);
+		NumPrimitives = PCGComputeConstants::MAX_PRIMITIVE_COMPONENTS_PER_SPAWNER;
+	}
 
-		for (int32 PrimitiveIndex = 0; PrimitiveIndex < Factory->GetNumPrimitives(); ++PrimitiveIndex)
-		{
-			FPrimitiveSceneProxy* SceneProxy = Factory->GetSceneProxy(PrimitiveIndex);
-			check(SceneProxy);
+	Primitives.Reserve(NumPrimitives);
 
-			const int32 NumInstances = Factory->GetNumInstances(PrimitiveIndex);
-			Primitives.Add(FPCGPrimitiveAndNumInstances(SceneProxy, NumInstances));
-			NumInstancesAllPrimitives += NumInstances;
-		}
+	for (int32 PrimitiveIndex = 0; PrimitiveIndex < NumPrimitives; ++PrimitiveIndex)
+	{
+		FPrimitiveSceneProxy* SceneProxy = Factory->GetSceneProxy(PrimitiveIndex);
+		check(SceneProxy);
+
+		const int32 NumInstances = Factory->GetNumInstances(PrimitiveIndex);
+		Primitives.Add(FPCGPrimitiveAndNumInstances(SceneProxy, NumInstances));
+		NumInstancesAllPrimitives += NumInstances;
 	}
 
 	NumCustomFloatsPerInstance = FoundPrimitives->NumCustomFloats;
diff --git a/Engine/Plugins/PCG/Source/PCG/Private/Compute/PCGDataBinding.h b/Engine/Plugins/PCG/Source/PCG/Private/Compute/PCGDataBinding.h
--- a/Engine/Plugins/PCG/Source/PCG/Private/Compute/PCGDataBinding.h
+++ b/Engine/Plugins/PCG/Source/PCG/Private/Compute/PCGDataBinding.h
@@ -24,7 +24,7 @@ struct FPCGSpawnerPrimitives
 {
 	GENERATED_BODY()
 
-	TArray<TSharedPtr<IPCGRuntimePrimitiveFactory>> PrimitiveFactories;
+	TSharedPtr<IPCGRuntimePrimitiveFactory> PrimitiveFactory;
 
 	TArray<FBox> PrimitiveMeshBounds;
 
diff --git a/Engine/Plugins/PCG/Source/PCG/Private/Compute/PrimitiveFactories/PCGPrimitiveFactoryISKMC.cpp b/Engine/Plugins/PCG/Source/PCG/Private/Compute/PrimitiveFactories/PCGPrimitiveFactoryISKMC.cpp
--- a/Engine/Plugins/PCG/Source/PCG/Private/Compute/PrimitiveFactories/PCGPrimitiveFactoryISKMC.cpp
+++ b/Engine/Plugins/PCG/Source/PCG/Private/Compute/PrimitiveFactories/PCGPrimitiveFactoryISKMC.cpp
@@ -8,34 +8,89 @@
 #include "Components/InstancedSkinnedMeshComponent.h"
 #include "Helpers/PCGActorHelpers.h"
 
-bool FPCGPrimitiveFactoryISKMC::Create(FPCGContext* InContext, AActor* InTargetActor, const FPCGSoftSkinnedMeshComponentDescriptor& InDesc, uint64 InSettingsUID, const FPCGCrc& InSettingsCrc)
+#include "Components/PrimitiveComponent.h"
+#include "Engine/SkinnedAsset.h"
+
+FPrimitiveSceneProxy* FPCGPrimitiveFactoryISKMC::GetSceneProxy(int32 InPrimitiveIndex) const
+{
+	if (ensure(Components.IsValidIndex(InPrimitiveIndex)))
+	{
+		return Components[InPrimitiveIndex].IsValid() ? Components[InPrimitiveIndex]->GetSceneProxy() : nullptr;
+	}
+	else
+	{
+		return nullptr;
+	}
+}
+
+int32 FPCGPrimitiveFactoryISKMC::GetNumInstances(int32 InPrimitiveIndex) const
+{
+	return NumInstances.IsValidIndex(InPrimitiveIndex) ? NumInstances[InPrimitiveIndex] : -1;
+}
+
+bool FPCGPrimitiveFactoryISKMC::IsRenderStateCreated() const
+{
+	for (const TWeakObjectPtr<UPrimitiveComponent>& Component : Components)
+	{
+		if (Component.IsValid() && !Component->GetSceneProxy())
+		{
+			return false;
+		}
+	}
+
+	return true;
+}
+
+void FPCGPrimitiveFactoryISKMC::Initialize(FParameters&& InParameters)
+{
+	Descriptors = InParameters.Descriptors;
+	TargetActor = InParameters.TargetActor;
+}
+
+bool FPCGPrimitiveFactoryISKMC::Create(FPCGContext* InContext)
 {
 	UPCGComponent* SourceComponent = Cast<UPCGComponent>(InContext->ExecutionSource.Get());
 	check(SourceComponent);
 
 	FPCGSkinnedMeshComponentBuilderParams Params;
-	Params.Descriptor = InDesc;
 	Params.NumCustomDataFloats = 0; // This is the num custom floats allocated in CPU memory.
 
 	// If the root actor we're binding to is movable, then the component should be movable by default
-	if (USceneComponent* SceneComponent = InTargetActor->GetRootComponent())
+	if (USceneComponent* SceneComponent = TargetActor->GetRootComponent())
 	{
 		Params.Descriptor.Mobility = SceneComponent->Mobility;
 	}
 
-	ensure(InSettingsCrc.IsValid());
-	Params.SettingsCrc = InSettingsCrc;
-
-	// TODO - not sure this is a good managed resource object to use, could make procedural variant?
-	UPCGManagedISKMComponent* ManagedComponent = UPCGActorHelpers::GetOrCreateManagedABMC(InTargetActor, SourceComponent, Params, InContext);
+	bool bAnyCreated = false;
 
-	if (ManagedComponent)
+	while (NumPrimitivesProcessed < Descriptors.Num())
 	{
-		Component = ManagedComponent->GetComponent();
-		NumInstances = InDesc.NumInstancesGPUOnly;
+		const FPCGSoftSkinnedMeshComponentDescriptor& Desc = Descriptors[NumPrimitivesProcessed];
 
-		ManagedComponent->SetCrc(InContext->DependenciesCrc);
+		Params.Descriptor = Desc;
+
+		// TODO - not sure this is a good managed resource object to use, could make procedural variant?
+		UPCGManagedISKMComponent* ManagedComponent = UPCGActorHelpers::GetOrCreateManagedABMC(TargetActor, SourceComponent, Params, InContext);
+
+		if (ManagedComponent)
+		{
+			Components.Add(ManagedComponent->GetComponent());
+			NumInstances.Add(Desc.NumInstancesGPUOnly);
+			MeshBounds.Add(Desc.SkinnedAsset->GetBounds().GetBox());
+		}
+
+		++NumPrimitivesProcessed;
+
+		if (NumPrimitivesProcessed < Descriptors.Num() && InContext->AsyncState.ShouldStop())
+		{
+			return false;
+		}
 	}
 
-	return ManagedComponent != nullptr;
+	return NumPrimitivesProcessed == Descriptors.Num();
+}
+
+FBox FPCGPrimitiveFactoryISKMC::GetMeshBounds(int32 InPrimitiveIndex) const
+{
+	return ensure(MeshBounds.IsValidIndex(InPrimitiveIndex)) ? MeshBounds[InPrimitiveIndex] : FBox();
 }
diff --git a/Engine/Plugins/PCG/Source/PCG/Private/Compute/PrimitiveFactories/PCGPrimitiveFactoryISKMC.h b/Engine/Plugins/PCG/Source/PCG/Private/Compute/PrimitiveFactories/PCGPrimitiveFactoryISKMC.h
--- a/Engine/Plugins/PCG/Source/PCG/Private/Compute/PrimitiveFactories/PCGPrimitiveFactoryISKMC.h
+++ b/Engine/Plugins/PCG/Source/PCG/Private/Compute/PrimitiveFactories/PCGPrimitiveFactoryISKMC.h
@@ -2,16 +2,41 @@
 
 #pragma once
 
-#include "Compute/PrimitiveFactories/PCGPrimitiveFactoryPrimitiveComponentBase.h"
+#include "Compute/PrimitiveFactories/IPCGRuntimePrimitiveFactory.h"
 
 class AActor;
+class UPrimitiveComponent;
 struct FPCGContext;
 struct FPCGCrc;
 struct FPCGSoftSkinnedMeshComponentDescriptor;
 
-class FPCGPrimitiveFactoryISKMC : public FPCGPrimitiveFactoryPrimitiveComponentBase
+class FPCGPrimitiveFactoryISKMC : public IPCGRuntimePrimitiveFactory
 {
 public:
-	/** Creates a Instanced Skinned Mesh Component on the target actor, and registers it with the execution source. */
-	bool Create(FPCGContext* InContext, AActor* InTargetActor, const FPCGSoftSkinnedMeshComponentDescriptor& InDesc, uint64 InSettingsUID, const FPCGCrc& InSettingsCrc);
+	//~ Begin IPCGRuntimePrimitiveFactory interface
+	virtual int32 GetNumPrimitives() const override { return Components.Num(); }
+	virtual FPrimitiveSceneProxy* GetSceneProxy(int32 InPrimitiveIndex) const override;
+	virtual int32 GetNumInstances(int32 InPrimitiveIndex) const override;
+	virtual bool IsRenderStateCreated() const override;
+	//~ End IPCGRuntimePrimitiveFactory interface
+
+	struct FParameters
+	{
+		TArray<FPCGSoftSkinnedMeshComponentDescriptor> Descriptors;
+		AActor* TargetActor = nullptr;
+	};
+
+	virtual void Initialize(FParameters&& InParameters);
+	virtual bool Create(FPCGContext* InContext);
+	virtual FBox GetMeshBounds(int32 InPrimitiveIndex) const;
+
+protected:
+	TArray<FPCGSoftSkinnedMeshComponentDescriptor> Descriptors;
+	AActor* TargetActor = nullptr;
+
+	int32 NumPrimitivesProcessed = 0;
+
+	TArray<TWeakObjectPtr<UPrimitiveComponent>> Components;
+	TArray<int32> NumInstances;
+	TArray<FBox> MeshBounds;
 };
diff --git a/Engine/Plugins/PCG/Source/PCG/Private/Compute/PrimitiveFactories/PCGPrimitiveFactoryISMBase.h b/Engine/Plugins/PCG/Source/PCG/Private/Compute/PrimitiveFactories/PCGPrimitiveFactoryISMBase.h
new file mode 100644
--- /dev/null
+++ b/Engine/Plugins/PCG/Source/PCG/Private/Compute/PrimitiveFactories/PCGPrimitiveFactoryISMBase.h
@@ -0,0 +1,25 @@
+// Copyright Epic Games, Inc. All Rights Reserved.
+
+#pragma once
+
+#include "Compute/PrimitiveFactories/IPCGRuntimePrimitiveFactory.h"
+
+#include "UObject/WeakObjectPtr.h"
+
+class AActor;
+struct FPCGContext;
+struct FPCGProceduralISMComponentDescriptor;
+
+class IPCGPrimitiveFactoryISMBase : public IPCGRuntimePrimitiveFactory
+{
+public:
+	struct FParameters
+	{
+		TArray<FPCGProceduralISMComponentDescriptor> Descriptors;
+		AActor* TargetActor = nullptr; // Optional, unless creating actor primitive components.
+	};
+
+	virtual void Initialize(FParameters&& InParameters) = 0;
+	virtual bool Create(FPCGContext* InContext) = 0;
+	virtual FBox GetMeshBounds(int32 InPrimitiveIndex) const = 0;
+};
diff --git a/Engine/Plugins/PCG/Source/PCG/Private/Compute/PrimitiveFactories/PCGPrimitiveFactoryPISMC.cpp b/Engine/Plugins/PCG/Source/PCG/Private/Compute/PrimitiveFactories/PCGPrimitiveFactoryPISMC.cpp
--- a/Engine/Plugins/PCG/Source/PCG/Private/Compute/PrimitiveFactories/PCGPrimitiveFactoryPISMC.cpp
+++ b/Engine/Plugins/PCG/Source/PCG/Private/Compute/PrimitiveFactories/PCGPrimitiveFactoryPISMC.cpp
@@ -6,24 +6,80 @@
 #include "PCGContext.h"
 #include "Components/PCGProceduralISMComponent.h"
 
-bool FPCGPrimitiveFactoryPISMC::Create(FPCGContext* InContext, AActor* InTargetActor, const FPCGProceduralISMComponentDescriptor& InDesc, uint64 InSettingsUID)
+#include "Engine/StaticMesh.h"
+#include "Components/PrimitiveComponent.h"
+
+FPrimitiveSceneProxy* FPCGPrimitiveFactoryPISMC::GetSceneProxy(int32 InPrimitiveIndex) const
+{
+	if (ensure(Components.IsValidIndex(InPrimitiveIndex)))
+	{
+		return Components[InPrimitiveIndex].IsValid() ? Components[InPrimitiveIndex]->GetSceneProxy() : nullptr;
+	}
+	else
+	{
+		return nullptr;
+	}
+}
+
+int32 FPCGPrimitiveFactoryPISMC::GetNumInstances(int32 InPrimitiveIndex) const
+{
+	return NumInstances.IsValidIndex(InPrimitiveIndex) ? NumInstances[InPrimitiveIndex] : -1;
+}
+
+bool FPCGPrimitiveFactoryPISMC::IsRenderStateCreated() const
+{
+	for (const TWeakObjectPtr<UPrimitiveComponent>& Component : Components)
+	{
+		if (Component.IsValid() && !Component->GetSceneProxy())
+		{
+			return false;
+		}
+	}
+
+	return true;
+}
+
+void FPCGPrimitiveFactoryPISMC::Initialize(FParameters&& InParameters)
+{
+	Descriptors = InParameters.Descriptors;
+	TargetActor = InParameters.TargetActor;
+}
+
+bool FPCGPrimitiveFactoryPISMC::Create(FPCGContext* InContext)
 {
 	UPCGComponent* SourceComponent = Cast<UPCGComponent>(InContext->ExecutionSource.Get());
 	check(SourceComponent);
 
 	FPCGProceduralISMCBuilderParameters Params;
-	Params.Descriptor = InDesc;
 	Params.bAllowDescriptorChanges = false;
 
-	UPCGManagedProceduralISMComponent* ManagedComponent = PCGManagedProceduralISMComponent::GetOrCreateManagedProceduralISMC(InTargetActor, SourceComponent, InSettingsUID, Params);
-
-	if (ManagedComponent)
+	while (NumPrimitivesProcessed < Descriptors.Num())
 	{
-		Component = ManagedComponent->GetComponent();
-		NumInstances = InDesc.NumInstances;
+		const FPCGProceduralISMComponentDescriptor& Desc = Descriptors[NumPrimitivesProcessed];
 
-		ManagedComponent->SetCrc(InContext->DependenciesCrc);
+		Params.Descriptor = Desc;
+
+		UPCGManagedProceduralISMComponent* ManagedComponent = PCGManagedProceduralISMComponent::GetOrCreateManagedProceduralISMC(TargetActor, SourceComponent, /*InSettingsUID=*/0, Params);
+
+		if (ManagedComponent)
+		{
+			Components.Add(ManagedComponent->GetComponent());
+			NumInstances.Add(Desc.NumInstances);
+			MeshBounds.Add(Desc.StaticMesh->GetBoundingBox());
+		}
+
+		++NumPrimitivesProcessed;
+
+		if (NumPrimitivesProcessed < Descriptors.Num() && InContext->AsyncState.ShouldStop())
+		{
+			return false;
+		}
 	}
 
-	return ManagedComponent != nullptr;
+	return NumPrimitivesProcessed == Descriptors.Num();
+}
+
+FBox FPCGPrimitiveFactoryPISMC::GetMeshBounds(int32 InPrimitiveIndex) const
+{
+	return ensure(MeshBounds.IsValidIndex(InPrimitiveIndex)) ? MeshBounds[InPrimitiveIndex] : FBox();
 }
diff --git a/Engine/Plugins/PCG/Source/PCG/Private/Compute/PrimitiveFactories/PCGPrimitiveFactoryPISMC.h b/Engine/Plugins/PCG/Source/PCG/Private/Compute/PrimitiveFactories/PCGPrimitiveFactoryPISMC.h
--- a/Engine/Plugins/PCG/Source/PCG/Private/Compute/PrimitiveFactories/PCGPrimitiveFactoryPISMC.h
+++ b/Engine/Plugins/PCG/Source/PCG/Private/Compute/PrimitiveFactories/PCGPrimitiveFactoryPISMC.h
@@ -2,15 +2,36 @@
 
 #pragma once
 
-#include "Compute/PrimitiveFactories/PCGPrimitiveFactoryPrimitiveComponentBase.h"
+#include "Compute/PrimitiveFactories/PCGPrimitiveFactoryISMBase.h"
 
 class AActor;
+class UPrimitiveComponent;
 struct FPCGContext;
 struct FPCGProceduralISMComponentDescriptor;
 
-class FPCGPrimitiveFactoryPISMC : public FPCGPrimitiveFactoryPrimitiveComponentBase
+class FPCGPrimitiveFactoryPISMC : public IPCGPrimitiveFactoryISMBase
 {
 public:
-	/** Creates a Procedural Instanced Static Mesh Component on the target actor, and registers it with the execution source. */
-	bool Create(FPCGContext* InContext, AActor* InTargetActor, const FPCGProceduralISMComponentDescriptor& InDesc, uint64 InSettingsUID);
+	//~ Begin IPCGRuntimePrimitiveFactory interface
+	virtual int32 GetNumPrimitives() const override { return Components.Num(); }
+	virtual FPrimitiveSceneProxy* GetSceneProxy(int32 InPrimitiveIndex) const override;
+	virtual int32 GetNumInstances(int32 InPrimitiveIndex) const override;
+	virtual bool IsRenderStateCreated() const override;
+	//~ End IPCGRuntimePrimitiveFactory interface
+
+	//~ Begin IPCGPrimitiveFactoryISMBase interface
+	virtual void Initialize(FParameters&& InParameters) override;
+	virtual bool Create(FPCGContext* InContext) override;
+	virtual FBox GetMeshBounds(int32 InPrimitiveIndex) const override;
+	//~ End IPCGPrimitiveFactoryISMBase interface
+
+protected:
+	TArray<FPCGProceduralISMComponentDescriptor> Descriptors;
+	AActor* TargetActor = nullptr;
+
+	int32 NumPrimitivesProcessed = 0;
+
+	TArray<TWeakObjectPtr<UPrimitiveComponent>> Components;
+	TArray<int32> NumInstances;
+	TArray<FBox> MeshBounds;
 };
diff --git a/Engine/Plugins/PCG/Source/PCG/Private/Compute/PrimitiveFactories/PCGPrimitiveFactoryPrimitiveComponentBase.cpp b/Engine/Plugins/PCG/Source/PCG/Private/Compute/PrimitiveFactories/PCGPrimitiveFactoryPrimitiveComponentBase.cpp
deleted file mode 100644
diff --git a/Engine/Plugins/PCG/Source/PCG/Private/Compute/PrimitiveFactories/PCGPrimitiveFactoryPrimitiveComponentBase.h b/Engine/Plugins/PCG/Source/PCG/Private/Compute/PrimitiveFactories/PCGPrimitiveFactoryPrimitiveComponentBase.h
deleted file mode 100644
