[PCG] Refactor Number of Elements element, to simplify it and bring it up to current paradigm

* Combined into single element
- Deprecate base version

diff --git a/Engine/Plugins/PCG/Config/BasePCG.ini b/Engine/Plugins/PCG/Config/BasePCG.ini
--- a/Engine/Plugins/PCG/Config/BasePCG.ini
+++ b/Engine/Plugins/PCG/Config/BasePCG.ini
@@ -98,6 +98,8 @@
 +ClassRedirects=(OldName="PCGMetadataOperationSettings",NewName="/Script/PCG.PCGCopyAttributesSettings")
 +ClassRedirects=(OldName="PCGCreateSplineMeshSettings",NewName="/Script/PCG.PCGSpawnSplineMeshSettings")
 +ClassRedirects=(OldName="/Script/PCGGeometryScriptInterop.PCGDynamicMeshData",NewName="/Script/PCG.PCGDynamicMeshData")
++ClassRedirects = (OldName="/Script/PCG.PCGNumberOfEntriesSettings",NewName="/Script/PCG.PCGNumberOfElementsSettings")
++ClassRedirects = (OldName="/Script/PCG.PCGNumberOfPointsSettings",NewName="/Script/PCG.PCGNumberOfElementsSettings")
 
 +StructRedirects=(OldName="PCGPointFilterThresholdSettings",NewName="/Script/PCG.PCGAttributeFilterThresholdSettings")
 +StructRedirects=(OldName="PCGActorPropertyOverrideDescription",NewName="/Script/PCG.PCGObjectPropertyOverrideDescription")
diff --git a/Engine/Plugins/PCG/Source/PCG/Private/Elements/PCGNumberOfElements.cpp b/Engine/Plugins/PCG/Source/PCG/Private/Elements/PCGNumberOfElements.cpp
--- a/Engine/Plugins/PCG/Source/PCG/Private/Elements/PCGNumberOfElements.cpp
+++ b/Engine/Plugins/PCG/Source/PCG/Private/Elements/PCGNumberOfElements.cpp
@@ -6,17 +6,14 @@
 #include "PCGPin.h"
 #include "PCGParamData.h"
 #include "Data/PCGBasePointData.h"
+#include "Helpers/PCGHelpers.h"
 #include "Metadata/PCGMetadata.h"
 
 #include UE_INLINE_GENERATED_CPP_BY_NAME(PCGNumberOfElements)
 
 #define LOCTEXT_NAMESPACE "PCGNumberOfElementsSettings"
 
-////////////////////////////////////
-// UPCGNumberOfElementsBaseSettings
-////////////////////////////////////
-
-TArray<FPCGPinProperties> UPCGNumberOfElementsBaseSettings::OutputPinProperties() const
+TArray<FPCGPinProperties> UPCGNumberOfElementsSettings::OutputPinProperties() const
 {
 	TArray<FPCGPinProperties> PinProperties;
 	PinProperties.Emplace(PCGPinConstants::DefaultOutputLabel, EPCGDataType::Param);
@@ -24,53 +21,18 @@ TArray<FPCGPinProperties> UPCGNumberOfElementsBaseSettings::OutputPinProperties(
 	return PinProperties;
 }
 
-
-////////////////////////////////////
-// UPCGNumberOfPointsSettings
-////////////////////////////////////
-
-TArray<FPCGPinProperties> UPCGNumberOfPointsSettings::InputPinProperties() const
-{
-	return Super::DefaultPointInputPinProperties();
-}
-
-FPCGElementPtr UPCGNumberOfPointsSettings::CreateElement() const
+FPCGElementPtr UPCGNumberOfElementsSettings::CreateElement() const
 {
-	return MakeShared<FPCGNumberOfPointsElement>();
+	return MakeShared<FPCGNumberOfElementsElement>();
 }
 
-
-////////////////////////////////////
-// UPCGNumberOfEntriesSettings
-////////////////////////////////////
-
-TArray<FPCGPinProperties> UPCGNumberOfEntriesSettings::InputPinProperties() const
-{
-	TArray<FPCGPinProperties> PinProperties;
-	FPCGPinProperties& InputPinProperty = PinProperties.Emplace_GetRef(PCGPinConstants::DefaultInputLabel, EPCGDataType::Param, /*bInAllowMultipleConnections=*/true, /*bAllowMultipleData=*/true);
-	InputPinProperty.SetRequiredPin();
-
-	return PinProperties;
-}
-
-FPCGElementPtr UPCGNumberOfEntriesSettings::CreateElement() const
-{
-	return MakeShared<FPCGNumberOfEntriesElement>();
-}
-
-
-////////////////////////////////////
-// FPCGNumberOfElementsBaseElement
-////////////////////////////////////
-
-template <typename DataType>
-bool FPCGNumberOfElementsBaseElement<DataType>::ExecuteInternal(FPCGContext* Context) const
+bool FPCGNumberOfElementsElement::ExecuteInternal(FPCGContext* Context) const
 {
 	TRACE_CPUPROFILER_EVENT_SCOPE(FPCGNumberOfElementsBaseElement::Execute);
 
 	check(Context);
 
-	const UPCGNumberOfElementsBaseSettings* Settings = Context->GetInputSettings<UPCGNumberOfElementsBaseSettings>();
+	const UPCGNumberOfElementsSettings* Settings = Context->GetInputSettings<UPCGNumberOfElementsSettings>();
 	check(Settings);
 
 	TArray<FPCGTaggedData> Inputs = Context->InputData.GetInputsByPin(PCGPinConstants::DefaultInputLabel);
@@ -80,11 +42,11 @@ bool FPCGNumberOfElementsBaseElement<DataType>::ExecuteInternal(FPCGContext* Con
 
 	for (int32 i = 0; i < Inputs.Num(); ++i)
 	{
-		const DataType* InputData = Cast<DataType>(Inputs[i].Data);
+		const UPCGData* InputData = Inputs[i].Data;
 
 		if (!InputData)
 		{
-			PCGE_LOG(Error, GraphAndLog, FText::Format(LOCTEXT("InputNotRightType", "Input {0} is not of the right input type, discarded"), FText::AsNumber(i)));
+			PCGLog::InputOutput::LogInvalidInputDataError(Context);
 			continue;
 		}
 
@@ -95,14 +57,14 @@ bool FPCGNumberOfElementsBaseElement<DataType>::ExecuteInternal(FPCGContext* Con
 
 			if (!NewAttribute)
 			{
-				PCGE_LOG(Error, GraphAndLog, FText::Format(LOCTEXT("AttributeFailedToBeCreated", "New Attribute {0} failed to be created."), FText::FromName(Settings->OutputAttributeName)));
+				PCGLog::Metadata::LogFailToCreateAttributeError<int32>(Settings->OutputAttributeName, Context);
 				return true;
 			}
 		}
 
 		check(OutputParamData && NewAttribute);
 		// Implementation note: since we might have multiple entries, we must set the actual value to a new entry
-		NewAttribute->SetValue(OutputParamData->Metadata->AddEntry(), GetNum(InputData));
+		NewAttribute->SetValue(OutputParamData->Metadata->AddEntry(), PCGHelpers::GetNumberOfElements(InputData));
 	}
 
 	if (OutputParamData)
@@ -115,24 +77,4 @@ bool FPCGNumberOfElementsBaseElement<DataType>::ExecuteInternal(FPCGContext* Con
 	return true;
 }
 
-
-////////////////////////////////////
-// FPCGNumberOfPointsElement
-////////////////////////////////////
-
-int32 FPCGNumberOfPointsElement::GetNum(const UPCGBasePointData* InData) const
-{
-	return InData ? InData->GetNumPoints() : 0;
-}
-
-
-////////////////////////////////////
-// FPCGNumberOfEntriesElement
-////////////////////////////////////
-
-int32 FPCGNumberOfEntriesElement::GetNum(const UPCGParamData* InData) const
-{
-	return (InData && InData->Metadata) ? InData->Metadata->GetLocalItemCount() : 0;
-}
-
 #undef LOCTEXT_NAMESPACE
diff --git a/Engine/Plugins/PCG/Source/PCG/Public/Elements/PCGNumberOfElements.h b/Engine/Plugins/PCG/Source/PCG/Public/Elements/PCGNumberOfElements.h
--- a/Engine/Plugins/PCG/Source/PCG/Public/Elements/PCGNumberOfElements.h
+++ b/Engine/Plugins/PCG/Source/PCG/Public/Elements/PCGNumberOfElements.h
@@ -14,91 +14,34 @@ class UPCGBasePointData;
 * except for getting the number of elements, it is factorized in a base class.
 */
 
-// Base class for common elements
-UCLASS(Abstract)
-class UPCGNumberOfElementsBaseSettings : public UPCGSettings
-{
-	GENERATED_BODY()
-
-	//~Begin UPCGSettings interface
-	
-
-protected:
-	virtual TArray<FPCGPinProperties> OutputPinProperties() const override;
-	//~End UPCGSettings interface
-
-public:
-	UPROPERTY(BlueprintReadWrite, EditAnywhere, Category = Settings, meta = (PCG_Overridable))
-	FName OutputAttributeName = TEXT("NumEntries");
-};
-
-/**
-* Return the number of points in the input point data.
-*/
-
 UCLASS(BlueprintType, ClassGroup = (Procedural))
-class UPCGNumberOfPointsSettings : public UPCGNumberOfElementsBaseSettings
+class UPCGNumberOfElementsSettings : public UPCGSettings
 {
 	GENERATED_BODY()
 
 public:
 	//~Begin UPCGSettings interface
 #if WITH_EDITOR
-	virtual FName GetDefaultNodeName() const override { return FName(TEXT("GetPointsCount")); }
-	virtual FText GetDefaultNodeTitle() const override { return NSLOCTEXT("PCGNumberOfElementsSettings", "NodeTitlePoint", "Get Points Count"); }
-	virtual FText GetNodeTooltipText() const override { return NSLOCTEXT("PCGNumberOfElementsSettings", "NodeTooltipPoint", "Return the number of points in the input point data."); }
-	virtual EPCGSettingsType GetType() const override { return EPCGSettingsType::Spatial; }
+	virtual FName GetDefaultNodeName() const override { return FName(TEXT("GetElementsCount")); }
+	virtual FText GetDefaultNodeTitle() const override { return NSLOCTEXT("PCGNumberOfElementsSettings", "NodeTitleElement", "Get Element Count"); }
+	virtual FText GetNodeTooltipText() const override { return NSLOCTEXT("PCGNumberOfElementsSettings", "NodeTooltipElement", "Return the number of elements in the input data."); }
+	virtual EPCGSettingsType GetType() const override { return EPCGSettingsType::Metadata; }
+	virtual TArray<FText> GetNodeTitleAliases() const override { return {NSLOCTEXT("PCGNumberOfElementsElement", "NodeTitleAlias", "Get Points Count")}; }
 #endif
 
 protected:
-	virtual TArray<FPCGPinProperties> InputPinProperties() const override;
+	virtual TArray<FPCGPinProperties> OutputPinProperties() const override;
 	virtual FPCGElementPtr CreateElement() const override;
 	//~End UPCGSettings interface
-};
-
-/**
-* Return the number of entries in the input attribute sets.
-*/
-
-UCLASS(BlueprintType, ClassGroup = (Procedural))
-class UPCGNumberOfEntriesSettings : public UPCGNumberOfElementsBaseSettings
-{
-	GENERATED_BODY()
 
 public:
-	//~Begin UPCGSettings interface
-#if WITH_EDITOR
-	virtual FName GetDefaultNodeName() const override { return FName(TEXT("GetEntriesCount")); }
-	virtual FText GetDefaultNodeTitle() const override { return NSLOCTEXT("PCGNumberOfElementsSettings", "NodeTitleEntry", "Get Entries Count"); }
-	virtual FText GetNodeTooltipText() const override { return NSLOCTEXT("PCGNumberOfElementsSettings", "NodeTooltipEntry", "Return the number of entries in the input attribute sets."); }
-	virtual EPCGSettingsType GetType() const override { return EPCGSettingsType::Generic; }
-#endif
-
-protected:
-	virtual TArray<FPCGPinProperties> InputPinProperties() const override;
-	virtual FPCGElementPtr CreateElement() const override;
-	//~End UPCGSettings interface
+	UPROPERTY(BlueprintReadWrite, EditAnywhere, Category = Settings, meta = (PCG_Overridable))
+	FName OutputAttributeName = "NumEntries";
 };
 
-
-template <typename DataType>
-class FPCGNumberOfElementsBaseElement : public IPCGElement
+class FPCGNumberOfElementsElement : public IPCGElement
 {
 protected:
 	virtual bool ExecuteInternal(FPCGContext* Context) const override;
 	virtual bool SupportsBasePointDataInputs(FPCGContext* InContext) const override { return true; }
-
-	virtual int32 GetNum(const DataType* InData) const = 0;
-};
-
-class FPCGNumberOfPointsElement : public FPCGNumberOfElementsBaseElement<UPCGBasePointData>
-{
-protected:
-	virtual int32 GetNum(const UPCGBasePointData* InData) const override;
 };
-
-class FPCGNumberOfEntriesElement : public FPCGNumberOfElementsBaseElement<UPCGParamData>
-{
-protected:
-	virtual int32 GetNum(const UPCGParamData* InData) const override;
-};
\ No newline at end of file
