[PCG] Update Create Points Sphere tooltip

#rnx

diff --git a/Engine/Plugins/PCG/Source/PCG/Public/Elements/PCGCreatePointsSphere.h b/Engine/Plugins/PCG/Source/PCG/Public/Elements/PCGCreatePointsSphere.h
--- a/Engine/Plugins/PCG/Source/PCG/Public/Elements/PCGCreatePointsSphere.h
+++ b/Engine/Plugins/PCG/Source/PCG/Public/Elements/PCGCreatePointsSphere.h
@@ -27,8 +27,8 @@ enum class EPCGGeodesicSphereRepresentation : uint8
 UENUM()
 enum class EPCGSpherePointOrientation : uint8
 {
-	Radial UMETA(ToolTip = "Points are aligned outward from the origin (the normal of the sphere's surface)."),
-	Centric UMETA(ToolTip = "Points are aligned inward toward the origin (inverse of the normal)."),
+	Radial UMETA(ToolTip = "Points' Up vector is aligned outward from the origin (the normal of the sphere's surface)."),
+	Centric UMETA(ToolTip = "Points' Up vector is aligned inward toward the origin (inverse of the normal)."),
 	None UMETA(ToolTip = "Points will not be explicitly aligned.")
 };
 
