[PCG] Updated Loop nodes to assume their inherited Graph's color

diff --git a/Engine/Plugins/PCG/Source/PCG/Private/PCGSubgraph.cpp b/Engine/Plugins/PCG/Source/PCG/Private/PCGSubgraph.cpp
--- a/Engine/Plugins/PCG/Source/PCG/Private/PCGSubgraph.cpp
+++ b/Engine/Plugins/PCG/Source/PCG/Private/PCGSubgraph.cpp
@@ -421,7 +421,7 @@ void UPCGSubgraphSettings::PostEditChangeProperty(struct FPropertyChangedEvent&
 
 FLinearColor UPCGSubgraphSettings::GetNodeTitleColor() const
 {
-	if (!IsDynamicGraph())
+	if (!IsSubgraphOverridden())
 	{
 		TOptional<FLinearColor> OverrideColor = SubgraphInstance->GetColorOverride();
 		if (OverrideColor.IsSet())
@@ -469,6 +469,11 @@ UPCGGraphInterface* UPCGSubgraphSettings::GetSubgraphInterface() const
 }
 
 bool UPCGSubgraphSettings::IsDynamicGraph() const
+{
+	return IsSubgraphOverridden();
+}
+
+bool UPCGSubgraphSettings::IsSubgraphOverridden() const
 {
 	UPCGNode* Node = Cast<UPCGNode>(GetOuter());
 
@@ -865,4 +870,4 @@ bool FPCGInputForwardingElement::ExecuteInternal(FPCGContext* Context) const
 	return true;
 }
 
-#undef LOCTEXT_NAMESPACE
\ No newline at end of file
+#undef LOCTEXT_NAMESPACE
diff --git a/Engine/Plugins/PCG/Source/PCG/Public/PCGSubgraph.h b/Engine/Plugins/PCG/Source/PCG/Public/PCGSubgraph.h
--- a/Engine/Plugins/PCG/Source/PCG/Public/PCGSubgraph.h
+++ b/Engine/Plugins/PCG/Source/PCG/Public/PCGSubgraph.h
@@ -134,6 +134,8 @@ protected:
 	UE_API virtual void SetSubgraphInternal(UPCGGraphInterface* InGraph) override;
 	//~End UPCGBaseSubgraphSettings interface
 
+	bool IsSubgraphOverridden() const;
+
 public:
 	UPROPERTY(BlueprintReadOnly, VisibleAnywhere, Category = Properties, Instanced, meta = (NoResetToDefault))
 	TObjectPtr<UPCGGraphInstance> SubgraphInstance;
