[PCG] Be more resilient in case of invalid PCG User Parameter data

#rnx

diff --git a/Engine/Plugins/PCG/Source/PCG/Private/Elements/PCGUserParameterGet.cpp b/Engine/Plugins/PCG/Source/PCG/Private/Elements/PCGUserParameterGet.cpp
--- a/Engine/Plugins/PCG/Source/PCG/Private/Elements/PCGUserParameterGet.cpp
+++ b/Engine/Plugins/PCG/Source/PCG/Private/Elements/PCGUserParameterGet.cpp
@@ -50,12 +50,16 @@ namespace PCGUserParameterGet
 	#endif // WITH_EDITOR
 				if (const UPCGUserParametersData* OverrideParametersData = CastChecked<UPCGUserParametersData>(UserParameterData[0].Data, ECastCheckedType::NullAllowed))
 				{
-					Layouts.Emplace(OverrideParametersData->UserParameters);
-
-					const UPCGUserParametersData* UpstreamData = OverrideParametersData->UpstreamData.Get();
+					// Gathering all the user parameter data from upstream, starting with this current level data.
+					const UPCGUserParametersData* UpstreamData = OverrideParametersData;
 					while (UpstreamData)
 					{
-						Layouts.Emplace(UpstreamData->UserParameters);
+						// Make sure to check if the user parameters is valid, as we don't want to have invalid layouts
+						if (UpstreamData->UserParameters.IsValid())
+						{
+							Layouts.Emplace(UpstreamData->UserParameters);
+						}
+
 						UpstreamData = UpstreamData->UpstreamData.Get();
 					}
 				}
@@ -544,7 +548,14 @@ bool FPCGUserParameterGetElement::ExecuteInternal(FPCGContext* Context) const
 		}
 	}
 
-	for(int CurrentLayoutIndex = 0; CurrentLayoutIndex < ValidLayouts.Num(); ++CurrentLayoutIndex)
+	if (ValidLayouts.IsEmpty())
+	{
+		// In theory this should never happen, as the PCG graph should always have a user parameter struct, even if it is empty.
+		PCGLog::LogErrorOnGraph(LOCTEXT("NoValidLayout", "Failed to find a valid user parameters to extract from."));
+		return true;
+	}
+
+	for (int CurrentLayoutIndex = 0; CurrentLayoutIndex < ValidLayouts.Num(); ++CurrentLayoutIndex)
 	{
 		const bool bIsLastIteration = (CurrentLayoutIndex == ValidLayouts.Num() - 1);
 		FConstStructView Parameters = ValidLayouts[CurrentLayoutIndex];
