[PCG] Fix ensures when using invalid texture in PCGTextureDataInterface

#rnx

diff --git a/Engine/Plugins/PCG/Source/PCG/Private/Compute/DataInterfaces/PCGTextureDataInterface.cpp b/Engine/Plugins/PCG/Source/PCG/Private/Compute/DataInterfaces/PCGTextureDataInterface.cpp
--- a/Engine/Plugins/PCG/Source/PCG/Private/Compute/DataInterfaces/PCGTextureDataInterface.cpp
+++ b/Engine/Plugins/PCG/Source/PCG/Private/Compute/DataInterfaces/PCGTextureDataInterface.cpp
@@ -74,6 +74,10 @@ FPCGTextureBindingInfo::FPCGTextureBindingInfo(const UPCGBaseTextureData* InText
 		Size = Desc.Extent;
 		Dimension = Desc.Dimension;
 	}
+	else
+	{
+		UE_LOG(LogPCG, Error, TEXT("PCGTextureDataInterface: Texture data '%s' has an invalid texture."), InTextureData ? *InTextureData->GetName() : TEXT("NULL"));
+	}
 }
 
 FPCGTextureBindingInfo::FPCGTextureBindingInfo(const FPCGDataDesc& InDataDesc, const FTransform& InTransform)
@@ -96,6 +100,13 @@ bool FPCGTextureBindingInfo::operator==(const FPCGTextureBindingInfo& Other) con
 		&& TextureBounds == Other.TextureBounds;
 }
 
+bool FPCGTextureBindingInfo::IsValid() const
+{
+	const int32 MaxTextureDimension = GetMax2DTextureDimension();
+
+	return ResourceType != EPCGTextureResourceType::Invalid && Size.X > 0 && Size.Y > 0 && Size.X <= MaxTextureDimension && Size.Y <= MaxTextureDimension;
+}
+
 void UPCGTextureDataInterface::GetSupportedInputs(TArray<FShaderFunctionDefinition>& OutFunctions) const
 {
 	// TODO: Support texture3d, etc. UNiagaraDataInterfaceRenderTarget2D and related are useful references.
@@ -398,9 +409,9 @@ bool FPCGTextureDataProviderProxy::IsValid(FValidationData const& InValidationDa
 {
 	for (const FPCGTextureBindingInfo& BindingInfo : BindingInfos)
 	{
-		if (BindingInfo.ResourceType == EPCGTextureResourceType::Invalid)
+		if (!BindingInfo.IsValid())
 		{
-			UE_LOG(LogPCG, Warning, TEXT("FPCGTextureDataProviderProxy invalid due to invalid resource type."));
+			UE_LOG(LogPCG, Warning, TEXT("FPCGTextureDataProviderProxy invalid due to invalid texture binding."));
 			return false;
 		}
 	}
diff --git a/Engine/Plugins/PCG/Source/PCG/Private/Compute/DataInterfaces/PCGTextureDataInterface.h b/Engine/Plugins/PCG/Source/PCG/Private/Compute/DataInterfaces/PCGTextureDataInterface.h
--- a/Engine/Plugins/PCG/Source/PCG/Private/Compute/DataInterfaces/PCGTextureDataInterface.h
+++ b/Engine/Plugins/PCG/Source/PCG/Private/Compute/DataInterfaces/PCGTextureDataInterface.h
@@ -30,6 +30,7 @@ struct FPCGTextureBindingInfo
 	explicit FPCGTextureBindingInfo(const UPCGBaseTextureData* InTextureData);
 	FPCGTextureBindingInfo(const FPCGDataDesc& InDataDesc, const FTransform& InTransform);
 	bool operator==(const FPCGTextureBindingInfo& Other) const;
+	bool IsValid() const;
 
 	EPCGTextureResourceType ResourceType = EPCGTextureResourceType::TextureObject;
 	FTextureRHIRef Texture = nullptr;
