[PCG] Support conversion in both ways in PCGDataTypeInfo

* Before, when trying to check if InType is Compatible with OutType, we were doing InType->SupportComposition(OutType)
* This does not allow for OutType to specify if it is compatible with InType, only InType knows it is compatible with OutType.
* This is an issue for external plugins that would like to be able to convert points to their representation for example.
* Also add a potential custom message that can be set to give more information if it is not compatible.
* Also adapt AddConversionNode and DefaultConversionSettings to also accept InType and OutType.
* Moved the Concrete logic from FPCGDataTypeInfo to FPCGDataTypeInfoConcrete.

#rnx

diff --git a/Engine/Plugins/PCG/Source/PCG/Private/Data/PCGBasePointData.cpp b/Engine/Plugins/PCG/Source/PCG/Private/Data/PCGBasePointData.cpp
--- a/Engine/Plugins/PCG/Source/PCG/Private/Data/PCGBasePointData.cpp
+++ b/Engine/Plugins/PCG/Source/PCG/Private/Data/PCGBasePointData.cpp
@@ -426,21 +426,21 @@ namespace PCGBasePointData
 	}
 }
 
-bool FPCGDataTypeInfoPoint::SupportsConversion(const FPCGDataTypeIdentifier& Other) const
+bool FPCGDataTypeInfoPoint::SupportsConversionTo(const FPCGDataTypeIdentifier& ThisType, const FPCGDataTypeIdentifier& OutputType, TSubclassOf<UPCGSettings>* OptionalOutConversionSettings, FText* OptionalOutCompatibilityMessage) const
 {
 	// Point can convert to param
-	return Other == EPCGDataType::Param || FPCGDataTypeInfoConcrete::SupportsConversion(Other);
-}
-
-TSubclassOf<UPCGSettings> FPCGDataTypeInfoPoint::GetDefaultConversionSettings(const FPCGDataTypeIdentifier& Other) const
-{
-	if (Other == EPCGDataType::Param)
+	if (OutputType.IsSameType(EPCGDataType::Param))
 	{
-		return UPCGConvertToAttributeSetSettings::StaticClass();
+		if (OptionalOutConversionSettings)
+		{
+			*OptionalOutConversionSettings = UPCGConvertToAttributeSetSettings::StaticClass();
+		}
+
+		return true;
 	}
 	else
 	{
-		return FPCGDataTypeInfoConcrete::GetDefaultConversionSettings(Other);
+		return FPCGDataTypeInfoConcrete::SupportsConversionTo(ThisType, OutputType, OptionalOutConversionSettings, OptionalOutCompatibilityMessage);
 	}
 }
 
diff --git a/Engine/Plugins/PCG/Source/PCG/Private/Data/PCGPolygon2DData.cpp b/Engine/Plugins/PCG/Source/PCG/Private/Data/PCGPolygon2DData.cpp
--- a/Engine/Plugins/PCG/Source/PCG/Private/Data/PCGPolygon2DData.cpp
+++ b/Engine/Plugins/PCG/Source/PCG/Private/Data/PCGPolygon2DData.cpp
@@ -237,21 +237,21 @@ UPCGPolygon2DData::UPCGPolygon2DData(const FObjectInitializer& ObjectInitializer
 	Metadata->SetupDomain(PCGMetadataDomainID::Elements, /*bIsDefault=*/true);
 }
 
-bool FPCGDataTypeInfoPolygon2D::SupportsConversion(const FPCGDataTypeIdentifier& Other) const
+bool FPCGDataTypeInfoPolygon2D::SupportsConversionTo(const FPCGDataTypeIdentifier& ThisType, const FPCGDataTypeIdentifier& OutputType, TSubclassOf<UPCGSettings>* OptionalOutConversionSettings, FText* OptionalOutCompatibilityMessage) const
 {
 	// Polygon can convert to surface
-	return Other == EPCGDataType::Surface || FPCGDataTypeInfoPolyline::SupportsConversion(Other);
-}
-
-TSubclassOf<UPCGSettings> FPCGDataTypeInfoPolygon2D::GetDefaultConversionSettings(const FPCGDataTypeIdentifier& Other) const
-{
-	if (Other == EPCGDataType::Surface)
+	if (OutputType.IsSameType(EPCGDataType::Surface))
 	{
-		return UPCGCreateSurfaceFromPolygon2DSettings::StaticClass();
+		if (OptionalOutConversionSettings)
+		{
+			*OptionalOutConversionSettings = UPCGCreateSurfaceFromPolygon2DSettings::StaticClass();
+		}
+
+		return true;
 	}
 	else
 	{
-		return FPCGDataTypeInfoPolyline::GetDefaultConversionSettings(Other);
+		return FPCGDataTypeInfoPolyline::SupportsConversionTo(ThisType, OutputType, OptionalOutConversionSettings, OptionalOutCompatibilityMessage);
 	}
 }
 
diff --git a/Engine/Plugins/PCG/Source/PCG/Private/Data/PCGSpatialData.cpp b/Engine/Plugins/PCG/Source/PCG/Private/Data/PCGSpatialData.cpp
--- a/Engine/Plugins/PCG/Source/PCG/Private/Data/PCGSpatialData.cpp
+++ b/Engine/Plugins/PCG/Source/PCG/Private/Data/PCGSpatialData.cpp
@@ -11,6 +11,7 @@
 #include "Data/PCGUnionData.h"
 #include "Elements/PCGCollapseElement.h"
 #include "Elements/PCGExecuteBlueprint.h"
+#include "Elements/PCGMakeConcreteElement.h"
 
 #include "Serialization/ArchiveCrc32.h"
 
@@ -37,25 +38,44 @@ FPCGInitializeFromDataParams::FPCGInitializeFromDataParams(const UPCGSpatialData
 
 PCG_DEFINE_TYPE_INFO(FPCGDataTypeInfoSpatial, UPCGSpatialData)
 
-bool FPCGDataTypeInfoSpatial::SupportsConversion(const FPCGDataTypeIdentifier& Other) const
+bool FPCGDataTypeInfoSpatial::SupportsConversionTo(const FPCGDataTypeIdentifier& ThisType, const FPCGDataTypeIdentifier& OutputType, TSubclassOf<UPCGSettings>* OptionalOutConversionSettings, FText* OptionalOutCompatibilityMessage) const
 {
 	// Spatial can convert to points.
-	return Other == EPCGDataType::Point || FPCGDataTypeInfo::SupportsConversion(Other);
+	if (OutputType.IsSameType(EPCGDataType::Point))
+	{
+		if (OptionalOutConversionSettings)
+		{
+			*OptionalOutConversionSettings = UPCGCollapseSettings::StaticClass();
+		}
+
+		return true;
+	}
+	else
+	{
+		return FPCGDataTypeInfo::SupportsConversionTo(ThisType, OutputType, OptionalOutConversionSettings, OptionalOutCompatibilityMessage);
+	}
 }
 
-TSubclassOf<UPCGSettings> FPCGDataTypeInfoSpatial::GetDefaultConversionSettings(const FPCGDataTypeIdentifier& Other) const
+PCG_DEFINE_TYPE_INFO_WITHOUT_CLASS(FPCGDataTypeInfoConcrete)
+
+bool FPCGDataTypeInfoConcrete::SupportsConversionFrom(const FPCGDataTypeIdentifier& InputType, const FPCGDataTypeIdentifier& ThisType, TSubclassOf<UPCGSettings>* OptionalOutConversionSettings, FText* OptionalOutCompatibilityMessage) const
 {
-	if (Other == EPCGDataType::Point)
+	if (ThisType.IsSameType(EPCGDataType::Concrete) && (InputType.IsSameType(FPCGDataTypeInfo::AsId()) || InputType.IsSameType(FPCGDataTypeInfoSpatial::AsId()) || InputType.IsSameType(FPCGDataTypeInfoComposite::AsId())))
 	{
-		return UPCGCollapseSettings::StaticClass();
+		if (OptionalOutConversionSettings)
+		{
+			*OptionalOutConversionSettings = UPCGMakeConcreteSettings::StaticClass();
+		}
+
+		return true;
 	}
 	else
 	{
-		return FPCGDataTypeInfo::GetDefaultConversionSettings(Other);
+		return FPCGDataTypeInfoSpatial::SupportsConversionFrom(InputType, ThisType, OptionalOutConversionSettings, OptionalOutCompatibilityMessage);
 	}
 }
 
-PCG_DEFINE_TYPE_INFO_WITHOUT_CLASS(FPCGDataTypeInfoConcrete)
+
 PCG_DEFINE_TYPE_INFO_WITHOUT_CLASS(FPCGDataTypeInfoComposite)
 
 UPCGSpatialData::UPCGSpatialData(const FObjectInitializer& ObjectInitializer)
diff --git a/Engine/Plugins/PCG/Source/PCG/Private/Data/PCGSplineData.cpp b/Engine/Plugins/PCG/Source/PCG/Private/Data/PCGSplineData.cpp
--- a/Engine/Plugins/PCG/Source/PCG/Private/Data/PCGSplineData.cpp
+++ b/Engine/Plugins/PCG/Source/PCG/Private/Data/PCGSplineData.cpp
@@ -22,21 +22,21 @@
 
 PCG_DEFINE_TYPE_INFO(FPCGDataTypeInfoSpline, UPCGSplineData)
 
-bool FPCGDataTypeInfoSpline::SupportsConversion(const FPCGDataTypeIdentifier& Other) const
+bool FPCGDataTypeInfoSpline::SupportsConversionTo(const FPCGDataTypeIdentifier& ThisType, const FPCGDataTypeIdentifier& OutputType, TSubclassOf<UPCGSettings>* OptionalOutConversionSettings, FText* OptionalOutCompatibilityMessage) const
 {
 	// Spline can convert to surface
-	return Other == EPCGDataType::Surface || FPCGDataTypeInfoPolyline::SupportsConversion(Other);
-}
-
-TSubclassOf<UPCGSettings> FPCGDataTypeInfoSpline::GetDefaultConversionSettings(const FPCGDataTypeIdentifier& Other) const
-{
-	if (Other == EPCGDataType::Surface)
+	if (OutputType.IsSameType(EPCGDataType::Surface))
 	{
-		return UPCGCreateSurfaceFromSplineSettings::StaticClass();
+		if (OptionalOutConversionSettings)
+		{
+			*OptionalOutConversionSettings = UPCGCreateSurfaceFromSplineSettings::StaticClass();
+		}
+
+		return true;
 	}
 	else
 	{
-		return FPCGDataTypeInfoPolyline::GetDefaultConversionSettings(Other);
+		return FPCGDataTypeInfoPolyline::SupportsConversionTo(ThisType, OutputType, OptionalOutConversionSettings, OptionalOutCompatibilityMessage);
 	}
 }
 
diff --git a/Engine/Plugins/PCG/Source/PCG/Private/Data/Registry/PCGDataType.cpp b/Engine/Plugins/PCG/Source/PCG/Private/Data/Registry/PCGDataType.cpp
--- a/Engine/Plugins/PCG/Source/PCG/Private/Data/Registry/PCGDataType.cpp
+++ b/Engine/Plugins/PCG/Source/PCG/Private/Data/Registry/PCGDataType.cpp
@@ -61,38 +61,45 @@ FString FPCGDataTypeBaseId::ToString() const
 	}
 }
 
-bool FPCGDataTypeInfo::SupportsConversion(const FPCGDataTypeIdentifier& Other) const
+
+bool FPCGDataTypeInfo::SupportsConversionFrom(const FPCGDataTypeIdentifier& InputType, const FPCGDataTypeIdentifier& ThisType, TSubclassOf<UPCGSettings>* OptionalOutConversionSettings, FText* OptionalOutCompatibilityMessage) const
 {
-	// Any/Spatial/Composite can convert to concrete
-	if (Other == EPCGDataType::Concrete)
+	return false;
+}
+
+bool FPCGDataTypeInfo::SupportsConversionTo(const FPCGDataTypeIdentifier& ThisType, const FPCGDataTypeIdentifier& OutputType, TSubclassOf<UPCGSettings>* OptionalOutConversionSettings, FText* OptionalOutCompatibilityMessage) const
+{
+	return false;
+}
+
+TOptional<TArray<UPCGNode*>> FPCGDataTypeInfo::AddConversionNodesFrom(const FPCGDataTypeIdentifier& InputType, const FPCGDataTypeIdentifier& ThisType, UPCGGraph* InGraph, UPCGPin* InUpstreamPin, UPCGPin* InDownstreamPin) const
+{
+	TSubclassOf<UPCGSettings> SettingsClass;
+	if (SupportsConversionFrom(InputType, ThisType, &SettingsClass) && SettingsClass)
 	{
-		const EPCGDataType ThisType = GetAssociatedLegacyType();
-		return ThisType == EPCGDataType::Any || ThisType == EPCGDataType::Spatial || ThisType == EPCGDataType::Composite;
+		return TArray<UPCGNode*>{FPCGSubgraphHelpers::SpawnNodeAndConnect(InGraph, InUpstreamPin, InDownstreamPin, SettingsClass)};
 	}
 	else
 	{
-		return false;
+		return {};
 	}
 }
 
-UPCGNode* FPCGDataTypeInfo::AddConversionNode(const FPCGDataTypeIdentifier& Other, UPCGGraph* InGraph, UPCGPin* InUpstreamPin, UPCGPin* InDownstreamPin) const
+TOptional<TArray<UPCGNode*>> FPCGDataTypeInfo::AddConversionNodesTo(const FPCGDataTypeIdentifier& ThisType, const FPCGDataTypeIdentifier& OutputType, UPCGGraph* InGraph, UPCGPin* InUpstreamPin, UPCGPin* InDownstreamPin) const
 {
-	if (const TSubclassOf<UPCGSettings> SettingsClass = GetDefaultConversionSettings(Other))
+	TSubclassOf<UPCGSettings> SettingsClass;
+	if (SupportsConversionTo(ThisType, OutputType, &SettingsClass) && SettingsClass)
 	{
-		return FPCGSubgraphHelpers::SpawnNodeAndConnect(InGraph, InUpstreamPin, InDownstreamPin, SettingsClass);
+		return TArray<UPCGNode*>{FPCGSubgraphHelpers::SpawnNodeAndConnect(InGraph, InUpstreamPin, InDownstreamPin, SettingsClass)};
 	}
 	else
 	{
-		return nullptr;
+		return {};
 	}
 }
 
-TSubclassOf<UPCGSettings> FPCGDataTypeInfo::GetDefaultConversionSettings(const FPCGDataTypeIdentifier& Other) const
-{
-	return nullptr;
-}
-
 #if WITH_EDITOR
+
 FText FPCGDataTypeInfo::GetTooltip() const
 {
 	return FText{};
diff --git a/Engine/Plugins/PCG/Source/PCG/Private/Data/Registry/PCGDataTypeRegistry.cpp b/Engine/Plugins/PCG/Source/PCG/Private/Data/Registry/PCGDataTypeRegistry.cpp
--- a/Engine/Plugins/PCG/Source/PCG/Private/Data/Registry/PCGDataTypeRegistry.cpp
+++ b/Engine/Plugins/PCG/Source/PCG/Private/Data/Registry/PCGDataTypeRegistry.cpp
@@ -340,7 +340,7 @@ void FPCGDataTypeRegistry::RegisterEntry(const FPCGDataTypeBaseId& BaseID)
 
 FPCGDataTypeRegistry::~FPCGDataTypeRegistry() = default;
 
-EPCGDataTypeCompatibilityResult FPCGDataTypeRegistry::IsCompatible(const FPCGDataTypeIdentifier& InType, const FPCGDataTypeIdentifier& OutType) const
+EPCGDataTypeCompatibilityResult FPCGDataTypeRegistry::IsCompatible(const FPCGDataTypeIdentifier& InType, const FPCGDataTypeIdentifier& OutType, FText* OptionalOutCompatibilityMessage) const
 {
 	if (!InType.IsValid() || !OutType.IsValid())
 	{
@@ -375,7 +375,10 @@ EPCGDataTypeCompatibilityResult FPCGDataTypeRegistry::IsCompatible(const FPCGDat
 				return EPCGDataTypeCompatibilityResult::UnknownType;
 			}
 
-			return InDefIt->Get().SupportsConversion(OutType) ? EPCGDataTypeCompatibilityResult::RequireConversion : EPCGDataTypeCompatibilityResult::RequireFilter;
+			const bool SupportConversion = InDefIt->Get().SupportsConversionTo(InType, OutType, /*OptionalOutConversionSettings=*/nullptr, OptionalOutCompatibilityMessage)
+				|| OutDefIt->Get().SupportsConversionFrom(InType, OutType, /*OptionalOutConversionSettings=*/nullptr, OptionalOutCompatibilityMessage);
+
+			return SupportConversion ? EPCGDataTypeCompatibilityResult::RequireConversion : EPCGDataTypeCompatibilityResult::RequireFilter;
 		}
 		else
 		{
@@ -399,7 +402,10 @@ EPCGDataTypeCompatibilityResult FPCGDataTypeRegistry::IsCompatible(const FPCGDat
 		return EPCGDataTypeCompatibilityResult::UnknownType;
 	}
 
-	return InDefIt->Get().SupportsConversion(OutType) ? EPCGDataTypeCompatibilityResult::RequireConversion : EPCGDataTypeCompatibilityResult::NotCompatible;
+	const bool SupportConversion = InDefIt->Get().SupportsConversionTo(InType, OutType, /*OptionalOutConversionSettings=*/nullptr, OptionalOutCompatibilityMessage)
+		|| OutDefIt->Get().SupportsConversionFrom(InType, OutType, /*OptionalOutConversionSettings=*/nullptr, OptionalOutCompatibilityMessage);
+	
+	return SupportConversion ? EPCGDataTypeCompatibilityResult::RequireConversion : EPCGDataTypeCompatibilityResult::NotCompatible;
 }
 
 const FPCGDataTypeInfo* FPCGDataTypeRegistry::GetTypeInfo(const FPCGDataTypeBaseId& ID) const
diff --git a/Engine/Plugins/PCG/Source/PCG/Private/PCGParamData.cpp b/Engine/Plugins/PCG/Source/PCG/Private/PCGParamData.cpp
--- a/Engine/Plugins/PCG/Source/PCG/Private/PCGParamData.cpp
+++ b/Engine/Plugins/PCG/Source/PCG/Private/PCGParamData.cpp
@@ -13,23 +13,25 @@
 
 #include UE_INLINE_GENERATED_CPP_BY_NAME(PCGParamData)
 
+#define LOCTEXT_NAMESPACE "PCGParamData"
+
 PCG_DEFINE_TYPE_INFO(FPCGDataTypeInfoParam, UPCGParamData)
 
-bool FPCGDataTypeInfoParam::SupportsConversion(const FPCGDataTypeIdentifier& Other) const
+bool FPCGDataTypeInfoParam::SupportsConversionTo(const FPCGDataTypeIdentifier& ThisType, const FPCGDataTypeIdentifier& OutputType, TSubclassOf<UPCGSettings>* OptionalOutConversionSettings, FText* OptionalOutCompatibilityMessage) const
 {
 	// Param can convert to points
-	return Other == EPCGDataType::Point || FPCGDataTypeInfo::SupportsConversion(Other);
-}
-
-TSubclassOf<UPCGSettings> FPCGDataTypeInfoParam::GetDefaultConversionSettings(const FPCGDataTypeIdentifier& Other) const
-{
-	if (Other == EPCGDataType::Point)
+	if (OutputType.IsSameType(EPCGDataType::Point))
 	{
-		return UPCGConvertToPointDataSettings::StaticClass();
+		if (OptionalOutConversionSettings)
+		{
+			*OptionalOutConversionSettings = UPCGConvertToPointDataSettings::StaticClass();
+		}
+
+		return true;
 	}
 	else
 	{
-		return FPCGDataTypeInfo::GetDefaultConversionSettings(Other);
+		return FPCGDataTypeInfo::SupportsConversionTo(ThisType, OutputType, OptionalOutConversionSettings, OptionalOutCompatibilityMessage);
 	}
 }
 
@@ -208,3 +210,5 @@ bool UPCGParamData::SetDomainFromDomainID(const FPCGMetadataDomainID& InDomainID
 		return Super::SetDomainFromDomainID(InDomainID, InOutSelector);
 	}
 }
+
+#undef LOCTEXT_NAMESPACE
\ No newline at end of file
diff --git a/Engine/Plugins/PCG/Source/PCG/Private/PCGPin.cpp b/Engine/Plugins/PCG/Source/PCG/Private/PCGPin.cpp
--- a/Engine/Plugins/PCG/Source/PCG/Private/PCGPin.cpp
+++ b/Engine/Plugins/PCG/Source/PCG/Private/PCGPin.cpp
@@ -371,7 +371,7 @@ bool UPCGPin::IsCompatible(const UPCGPin* OtherPin) const
 	return Compatibility != EPCGDataTypeCompatibilityResult::NotCompatible && Compatibility != EPCGDataTypeCompatibilityResult::UnknownType;
 }
 
-EPCGDataTypeCompatibilityResult UPCGPin::GetCompatibilityWithOtherPin(const UPCGPin* OtherPin) const
+EPCGDataTypeCompatibilityResult UPCGPin::GetCompatibilityWithOtherPin(const UPCGPin* OtherPin, FText* OptionalOutCompatibilityMessage) const
 {
 	const bool bThisIsOutputPin = IsOutputPin();
 	if (!OtherPin || bThisIsOutputPin == OtherPin->IsOutputPin())
@@ -382,7 +382,7 @@ EPCGDataTypeCompatibilityResult UPCGPin::GetCompatibilityWithOtherPin(const UPCG
 	const UPCGPin* UpstreamPin = bThisIsOutputPin ? this : OtherPin;
 	const UPCGPin* DownstreamPin = bThisIsOutputPin ? OtherPin : this;
 
-	return FPCGModule::GetConstDataTypeRegistry().IsCompatible(UpstreamPin->GetCurrentTypesID(), DownstreamPin->Properties.AllowedTypes);
+	return FPCGModule::GetConstDataTypeRegistry().IsCompatible(UpstreamPin->GetCurrentTypesID(), DownstreamPin->Properties.AllowedTypes, OptionalOutCompatibilityMessage);
 }
 
 bool UPCGPin::IsDownstreamPinTypeCompatible(const EPCGDataType UpstreamTypes) const
diff --git a/Engine/Plugins/PCG/Source/PCG/Public/Data/PCGBasePointData.h b/Engine/Plugins/PCG/Source/PCG/Public/Data/PCGBasePointData.h
--- a/Engine/Plugins/PCG/Source/PCG/Public/Data/PCGBasePointData.h
+++ b/Engine/Plugins/PCG/Source/PCG/Public/Data/PCGBasePointData.h
@@ -25,8 +25,7 @@ struct FPCGDataTypeInfoPoint : public FPCGDataTypeInfoConcrete
 
 	PCG_DECLARE_LEGACY_TYPE_INFO(EPCGDataType::Point)
 
-	PCG_API virtual bool SupportsConversion(const FPCGDataTypeIdentifier& Other) const override;
-	PCG_API virtual TSubclassOf<UPCGSettings> GetDefaultConversionSettings(const FPCGDataTypeIdentifier& Other) const override;
+	PCG_API virtual bool SupportsConversionTo(const FPCGDataTypeIdentifier& ThisType, const FPCGDataTypeIdentifier& OutputType, TSubclassOf<UPCGSettings>* OptionalOutConversionSettings, FText* OptionalOutCompatibilityMessage) const override;
 };
 
 UCLASS(MinimalAPI, Abstract, BlueprintType, ClassGroup = (Procedural))
diff --git a/Engine/Plugins/PCG/Source/PCG/Public/Data/PCGPolygon2DData.h b/Engine/Plugins/PCG/Source/PCG/Public/Data/PCGPolygon2DData.h
--- a/Engine/Plugins/PCG/Source/PCG/Public/Data/PCGPolygon2DData.h
+++ b/Engine/Plugins/PCG/Source/PCG/Public/Data/PCGPolygon2DData.h
@@ -62,8 +62,7 @@ struct FPCGDataTypeInfoPolygon2D : public FPCGDataTypeInfoPolyline
 
 	PCG_DECLARE_LEGACY_TYPE_INFO(EPCGDataType::Polygon2D);
 
-	PCG_API virtual bool SupportsConversion(const FPCGDataTypeIdentifier& Other) const override;
-	PCG_API virtual TSubclassOf<UPCGSettings> GetDefaultConversionSettings(const FPCGDataTypeIdentifier& Other) const override;
+	PCG_API virtual bool SupportsConversionTo(const FPCGDataTypeIdentifier& ThisType, const FPCGDataTypeIdentifier& OutputType, TSubclassOf<UPCGSettings>* OptionalOutConversionSettings, FText* OptionalOutCompatibilityMessage) const override;
 };
 
 /** Data representing a single 2D polygon with a 3D transform (for spatial operations). */
diff --git a/Engine/Plugins/PCG/Source/PCG/Public/Data/PCGSpatialData.h b/Engine/Plugins/PCG/Source/PCG/Public/Data/PCGSpatialData.h
--- a/Engine/Plugins/PCG/Source/PCG/Public/Data/PCGSpatialData.h
+++ b/Engine/Plugins/PCG/Source/PCG/Public/Data/PCGSpatialData.h
@@ -68,8 +68,7 @@ struct FPCGDataTypeInfoSpatial : public FPCGDataTypeInfo
 
 	PCG_DECLARE_LEGACY_TYPE_INFO(EPCGDataType::Spatial)
 
-	PCG_API virtual bool SupportsConversion(const FPCGDataTypeIdentifier& Other) const override;
-	PCG_API virtual TSubclassOf<UPCGSettings> GetDefaultConversionSettings(const FPCGDataTypeIdentifier& Other) const override;
+	PCG_API virtual bool SupportsConversionTo(const FPCGDataTypeIdentifier& ThisType, const FPCGDataTypeIdentifier& OutputType, TSubclassOf<UPCGSettings>* OptionalOutConversionSettings, FText* OptionalOutCompatibilityMessage) const override;
 };
 
 // Subclasses of UPCGSpatial
@@ -79,6 +78,8 @@ struct FPCGDataTypeInfoConcrete : public FPCGDataTypeInfoSpatial
 	GENERATED_BODY()
 
 	PCG_DECLARE_LEGACY_TYPE_INFO(EPCGDataType::Concrete)
+
+	PCG_API virtual bool SupportsConversionFrom(const FPCGDataTypeIdentifier& InputType, const FPCGDataTypeIdentifier& ThisType, TSubclassOf<UPCGSettings>* OptionalOutConversionSettings, FText* OptionalOutCompatibilityMessage) const override;
 };
 
 USTRUCT()
diff --git a/Engine/Plugins/PCG/Source/PCG/Public/Data/PCGSplineData.h b/Engine/Plugins/PCG/Source/PCG/Public/Data/PCGSplineData.h
--- a/Engine/Plugins/PCG/Source/PCG/Public/Data/PCGSplineData.h
+++ b/Engine/Plugins/PCG/Source/PCG/Public/Data/PCGSplineData.h
@@ -29,8 +29,7 @@ struct FPCGDataTypeInfoSpline : public FPCGDataTypeInfoPolyline
 
 	PCG_DECLARE_LEGACY_TYPE_INFO(EPCGDataType::Spline)
 
-	PCG_API virtual bool SupportsConversion(const FPCGDataTypeIdentifier& Other) const override;
-	PCG_API virtual TSubclassOf<UPCGSettings> GetDefaultConversionSettings(const FPCGDataTypeIdentifier& Other) const override;
+	PCG_API virtual bool SupportsConversionTo(const FPCGDataTypeIdentifier& ThisType, const FPCGDataTypeIdentifier& OutputType, TSubclassOf<UPCGSettings>* OptionalOutConversionSettings, FText* OptionalOutCompatibilityMessage) const override;
 };
 
 UCLASS(MinimalAPI, BlueprintType, ClassGroup = (Procedural))
diff --git a/Engine/Plugins/PCG/Source/PCG/Public/Data/Registry/PCGDataType.h b/Engine/Plugins/PCG/Source/PCG/Public/Data/Registry/PCGDataType.h
--- a/Engine/Plugins/PCG/Source/PCG/Public/Data/Registry/PCGDataType.h
+++ b/Engine/Plugins/PCG/Source/PCG/Public/Data/Registry/PCGDataType.h
@@ -71,10 +71,28 @@ struct FPCGDataTypeInfo
 
 	virtual EPCGDataType GetAssociatedLegacyType() const { return EPCGDataType::Any; }
 	static EPCGDataType GetStaticAssociatedLegacyType() { return EPCGDataType::Any; }
+
+	/**
+	 * Can be overriden by a type info to indicate if we can convert from InputType to ThisType.
+	 * Since we have a hierarchy of type, ThisType is by construction a child of the current type info ; But it can be a child and not the same type.
+	 * Can also return an optional simple conversion settings if the conversion is valid.
+	 * Can also provide a custom message to provide the user more information.
+	 */
+	PCG_API virtual bool SupportsConversionFrom(const FPCGDataTypeIdentifier& InputType, const FPCGDataTypeIdentifier& ThisType, TSubclassOf<UPCGSettings>* OptionalOutConversionSettings = nullptr, FText* OptionalOutCompatibilityMessage = nullptr) const;
+	
+	/**
+	 * Can be overriden by a type info to indicate if we can convert from ThisType to OutputType.
+	 * Since we have a hierarchy of type, ThisType is by construction a child of the current type info ; But it can be a child and not the same type.
+	 * Can also return an optional simple conversion settings if the conversion is valid.
+	 * Can also provide a custom message to provide the user more information.
+	 */
+	PCG_API virtual bool SupportsConversionTo(const FPCGDataTypeIdentifier& ThisType, const FPCGDataTypeIdentifier& OutputType, TSubclassOf<UPCGSettings>* OptionalOutConversionSettings = nullptr, FText* OptionalOutCompatibilityMessage = nullptr) const;
+
+	/** If SupportsConversionFrom returns true, can provide a function to add conversion nodes. Returns an optional array of added nodes. */
+	PCG_API virtual TOptional<TArray<UPCGNode*>> AddConversionNodesFrom(const FPCGDataTypeIdentifier& InputType, const FPCGDataTypeIdentifier& ThisType, UPCGGraph* InGraph, UPCGPin* InUpstreamPin, UPCGPin* InDownstreamPin) const;
 	
-	PCG_API virtual bool SupportsConversion(const FPCGDataTypeIdentifier& Other) const;
-	PCG_API virtual UPCGNode* AddConversionNode(const FPCGDataTypeIdentifier& Other, UPCGGraph* InGraph, UPCGPin* InUpstreamPin, UPCGPin* InDownstreamPin) const;
-	PCG_API virtual TSubclassOf<UPCGSettings> GetDefaultConversionSettings(const FPCGDataTypeIdentifier& Other) const;
+	/** If SupportsConversionTo returns true, can provide a function to add conversion nodes. Returns an optional array of added nodes. */
+	PCG_API virtual TOptional<TArray<UPCGNode*>> AddConversionNodesTo(const FPCGDataTypeIdentifier& ThisType, const FPCGDataTypeIdentifier& OutputType, UPCGGraph* InGraph, UPCGPin* InUpstreamPin, UPCGPin* InDownstreamPin) const;
 
 #if WITH_EDITOR
 	PCG_API virtual FText GetTooltip() const;
diff --git a/Engine/Plugins/PCG/Source/PCG/Public/Data/Registry/PCGDataTypeRegistry.h b/Engine/Plugins/PCG/Source/PCG/Public/Data/Registry/PCGDataTypeRegistry.h
--- a/Engine/Plugins/PCG/Source/PCG/Public/Data/Registry/PCGDataTypeRegistry.h
+++ b/Engine/Plugins/PCG/Source/PCG/Public/Data/Registry/PCGDataTypeRegistry.h
@@ -3,9 +3,8 @@
 #pragma once
 
 #include "PCGCommon.h"
-#include "Containers/AdderRef.h"
-
 #include "Data/Registry/PCGDataTypeIdentifier.h"
+
 #include "StructUtils/InstancedStruct.h"
 
 struct FPCGPinProperties;
@@ -61,8 +60,8 @@ public:
 	
 	~FPCGDataTypeRegistry();
 
-	/** Is it possible to connect the InType into a OutType. */
-	PCG_API EPCGDataTypeCompatibilityResult IsCompatible(const FPCGDataTypeIdentifier& InType, const FPCGDataTypeIdentifier& OutType) const;
+	/** Is it possible to connect the InType into a OutType. Can also provide an optional FText for a custom message. */
+	PCG_API EPCGDataTypeCompatibilityResult IsCompatible(const FPCGDataTypeIdentifier& InType, const FPCGDataTypeIdentifier& OutType, FText* OptionalOutCompatibilityMessage = nullptr) const;
 
 	/** Return the type info, to allow to call virtual functions. */
 	PCG_API const FPCGDataTypeInfo* GetTypeInfo(const FPCGDataTypeBaseId& ID) const;
diff --git a/Engine/Plugins/PCG/Source/PCG/Public/PCGParamData.h b/Engine/Plugins/PCG/Source/PCG/Public/PCGParamData.h
--- a/Engine/Plugins/PCG/Source/PCG/Public/PCGParamData.h
+++ b/Engine/Plugins/PCG/Source/PCG/Public/PCGParamData.h
@@ -21,8 +21,7 @@ struct FPCGDataTypeInfoParam : public FPCGDataTypeInfo
 
 	PCG_DECLARE_LEGACY_TYPE_INFO(EPCGDataType::Param)
 
-	PCG_API virtual bool SupportsConversion(const FPCGDataTypeIdentifier& Other) const override;
-	PCG_API virtual TSubclassOf<UPCGSettings> GetDefaultConversionSettings(const FPCGDataTypeIdentifier& Other) const override;
+	PCG_API virtual bool SupportsConversionTo(const FPCGDataTypeIdentifier& ThisType, const FPCGDataTypeIdentifier& OutputType, TSubclassOf<UPCGSettings>* OptionalOutConversionSettings, FText* OptionalOutCompatibilityMessage) const override;
 };
 
 /**
diff --git a/Engine/Plugins/PCG/Source/PCG/Public/PCGPin.h b/Engine/Plugins/PCG/Source/PCG/Public/PCGPin.h
--- a/Engine/Plugins/PCG/Source/PCG/Public/PCGPin.h
+++ b/Engine/Plugins/PCG/Source/PCG/Public/PCGPin.h
@@ -278,7 +278,9 @@ public:
 	UE_API bool AllowsMultipleConnections() const;
 	UE_API bool AllowsMultipleData() const;
 	UE_API bool IsCompatible(const UPCGPin* OtherPin) const;
-	UE_API EPCGDataTypeCompatibilityResult GetCompatibilityWithOtherPin(const UPCGPin* OtherPin) const;
+
+	/** Query the data type registry to check if this pin is compatible with the other pin. Can also provide an optional text to have custom compatibility messages. */
+	UE_API EPCGDataTypeCompatibilityResult GetCompatibilityWithOtherPin(const UPCGPin* OtherPin, FText* OptionalOutCompatibilityMessage = nullptr) const;
 
 	/** For checking when the new node doesn't exist yet. */
 	UE_API bool IsDownstreamPinTypeCompatible(EPCGDataType UpstreamTypes) const;
diff --git a/Engine/Plugins/PCG/Source/PCGEditor/Private/PCGEditorGraphSchema.cpp b/Engine/Plugins/PCG/Source/PCGEditor/Private/PCGEditorGraphSchema.cpp
--- a/Engine/Plugins/PCG/Source/PCGEditor/Private/PCGEditorGraphSchema.cpp
+++ b/Engine/Plugins/PCG/Source/PCGEditor/Private/PCGEditorGraphSchema.cpp
@@ -157,17 +157,18 @@ const FPinConnectionResponse UPCGEditorGraphSchema::CanCreateConnection(const UE
 		return FPinConnectionResponse(CONNECT_RESPONSE_DISALLOW, LOCTEXT("ConnectionFailed", "Unable to verify pins"));
 	}
 
-	const EPCGDataTypeCompatibilityResult Compatibility = InputPin->GetCompatibilityWithOtherPin(OutputPin);
+	FText CustomCompatiblityMessage{};
+	const EPCGDataTypeCompatibilityResult Compatibility = InputPin->GetCompatibilityWithOtherPin(OutputPin, &CustomCompatiblityMessage);
 
 	switch (Compatibility)
 	{
 	case EPCGDataTypeCompatibilityResult::RequireFilter:
-		return FPinConnectionResponse(CONNECT_RESPONSE_MAKE_WITH_CONVERSION_NODE, LOCTEXT("ConnectionUsingFilter", "Filter data to match type"));
+		return FPinConnectionResponse(CONNECT_RESPONSE_MAKE_WITH_CONVERSION_NODE, CustomCompatiblityMessage.IsEmpty() ? LOCTEXT("ConnectionUsingFilter", "Filter data to match type") : CustomCompatiblityMessage);
 	case EPCGDataTypeCompatibilityResult::RequireConversion:
-		return FPinConnectionResponse(CONNECT_RESPONSE_MAKE_WITH_CONVERSION_NODE, LOCTEXT("ConnectionConversion", "Convert data to match type"));
+		return FPinConnectionResponse(CONNECT_RESPONSE_MAKE_WITH_CONVERSION_NODE, CustomCompatiblityMessage.IsEmpty() ? LOCTEXT("ConnectionConversion", "Convert data to match type") : CustomCompatiblityMessage);
 	case EPCGDataTypeCompatibilityResult::NotCompatible: // fall-through
 	case EPCGDataTypeCompatibilityResult::UnknownType:
-		return FPinConnectionResponse(CONNECT_RESPONSE_DISALLOW, LOCTEXT("ConnectionTypesIncompatible", "Pins are incompatible"));
+		return FPinConnectionResponse(CONNECT_RESPONSE_DISALLOW, CustomCompatiblityMessage.IsEmpty() ? LOCTEXT("ConnectionTypesIncompatible", "Pins are incompatible") : CustomCompatiblityMessage);
 	case EPCGDataTypeCompatibilityResult::Compatible:// fall-through
 	default:
 		break;
@@ -256,13 +257,40 @@ bool UPCGEditorGraphSchema::TryCreateConnectionInternal(UEdGraphPin* InA, UEdGra
 	else if (Compatibility == EPCGDataTypeCompatibilityResult::RequireConversion && bAddConversionNodeIfNeeded)
 	{
 		const FPCGDataTypeRegistry& Registry = FPCGModule::GetConstDataTypeRegistry();
-		for (const FPCGDataTypeBaseId& ID : PCGPinA->Properties.AllowedTypes.GetIds())
+
+		auto AddConversionNodes = [&Registry, PCGPinA, PCGPinB, PCGGraph, &bShouldReconstruct](const FPCGDataTypeBaseId& Id, bool bIsInput) -> bool
+		{
+			const FPCGDataTypeInfo* TypeInfo = Registry.GetTypeInfo(Id);
+
+			if (bIsInput && TypeInfo && TypeInfo->AddConversionNodesTo(PCGPinA->Properties.AllowedTypes, PCGPinB->Properties.AllowedTypes, PCGGraph, PCGPinA, PCGPinB).IsSet())
+			{
+				bShouldReconstruct = true;
+				return true;
+			}
+			else if (!bIsInput && TypeInfo && TypeInfo->AddConversionNodesFrom(PCGPinA->Properties.AllowedTypes, PCGPinB->Properties.AllowedTypes, PCGGraph, PCGPinA, PCGPinB).IsSet())
+			{
+				bShouldReconstruct = true;
+				return true;
+			}
+
+			return false;
+		};
+
+		// Try for both pins
+		for (const FPCGDataTypeBaseId& Id : PCGPinA->Properties.AllowedTypes.GetIds())
+		{
+			if (AddConversionNodes(Id, /*bIsInput=*/ true))
+			{
+				break;
+			}
+		}
+
+		if (!bShouldReconstruct)
 		{
-			if (const FPCGDataTypeInfo* TypeInfo = Registry.GetTypeInfo(ID); TypeInfo && TypeInfo->SupportsConversion(PCGPinB->Properties.AllowedTypes))
+			for (const FPCGDataTypeBaseId& Id : PCGPinB->Properties.AllowedTypes.GetIds())
 			{
-				if (TypeInfo->AddConversionNode(PCGPinB->Properties.AllowedTypes, PCGGraph, PCGPinA, PCGPinB) != nullptr)
+				if (AddConversionNodes(Id, /*bIsInput=*/ false))
 				{
-					bShouldReconstruct = true;
 					break;
 				}
 			}
