[PCG] Add PCGFastGeoInterop plugin to enable usage of FastGeo components (experimental) from within PCG.

#rnx

diff --git a/Engine/Plugins/Experimental/PCGInterops/PCGFastGeoInterop/PCGFastGeoInterop.uplugin b/Engine/Plugins/Experimental/PCGInterops/PCGFastGeoInterop/PCGFastGeoInterop.uplugin
new file mode 100644
--- /dev/null
+++ b/Engine/Plugins/Experimental/PCGInterops/PCGFastGeoInterop/PCGFastGeoInterop.uplugin
@@ -0,0 +1,35 @@
+{
+	"FileVersion": 3,
+	"Version": 1,
+	"VersionName": "0.1",
+	"FriendlyName": "Procedural Content Generation Framework (PCG) FastGeo Interop",
+	"Description": "Extra plugin for Procedural Content Generation Framework which enables runtime spawning of primitives using FastGeo components.",
+	"Category": "Editor",
+	"CreatedBy": "Epic Games, Inc.",
+	"CreatedByURL": "https://epicgames.com",
+	"DocsURL": "https://docs.unrealengine.com/latest/en-US/procedural-content-generation--framework-in-unreal-engine/",
+	"MarketplaceURL": "",
+	"SupportURL": "",
+	"EnabledByDefault": false,
+	"CanContainContent": false,
+	"IsBetaVersion": false,
+	"IsExperimentalVersion": true,
+	"Installed": false,
+	"Modules": [
+		{
+			"Name": "PCGFastGeoInterop",
+			"Type": "Runtime",
+			"LoadingPhase": "Default"
+		}
+	],
+	"Plugins": [
+		{
+			"Name": "PCG",
+			"Enabled": true
+		},
+		{
+			"Name": "FastGeoStreaming",
+			"Enabled": true
+		}
+	]
+}
diff --git a/Engine/Plugins/Experimental/PCGInterops/PCGFastGeoInterop/Source/PCGFastGeoInterop/PCGFastGeoInterop.Build.cs b/Engine/Plugins/Experimental/PCGInterops/PCGFastGeoInterop/Source/PCGFastGeoInterop/PCGFastGeoInterop.Build.cs
new file mode 100644
--- /dev/null
+++ b/Engine/Plugins/Experimental/PCGInterops/PCGFastGeoInterop/Source/PCGFastGeoInterop/PCGFastGeoInterop.Build.cs
@@ -0,0 +1,27 @@
+// Copyright Epic Games, Inc. All Rights Reserved.
+
+namespace UnrealBuildTool.Rules
+{
+	public class PCGFastGeoInterop : ModuleRules
+	{
+		public PCGFastGeoInterop(ReadOnlyTargetRules Target) : base(Target)
+		{
+			PCHUsage = PCHUsageMode.UseExplicitOrSharedPCHs;
+
+			PublicDependencyModuleNames.AddRange(
+				new string[] {
+					"Core",
+					"CoreUObject",
+					"Engine",
+				}
+			);
+
+			PrivateDependencyModuleNames.AddRange(
+				new string[]
+				{
+					"FastGeoStreaming",
+					"PCG",
+				});
+		}
+	}
+}
diff --git a/Engine/Plugins/Experimental/PCGInterops/PCGFastGeoInterop/Source/PCGFastGeoInterop/Private/Components/PCGManagedFastGeoContainer.cpp b/Engine/Plugins/Experimental/PCGInterops/PCGFastGeoInterop/Source/PCGFastGeoInterop/Private/Components/PCGManagedFastGeoContainer.cpp
new file mode 100644
--- /dev/null
+++ b/Engine/Plugins/Experimental/PCGInterops/PCGFastGeoInterop/Source/PCGFastGeoInterop/Private/Components/PCGManagedFastGeoContainer.cpp
@@ -0,0 +1,23 @@
+// Copyright Epic Games, Inc. All Rights Reserved.
+
+#include "PCGManagedFastGeoContainer.h"
+
+#include UE_INLINE_GENERATED_CPP_BY_NAME(PCGManagedFastGeoContainer)
+
+bool UPCGManagedFastGeoContainer::Release(bool bHardRelease, TSet<TSoftObjectPtr<AActor>>& OutActorsToDelete)
+{
+	TRACE_CPUPROFILER_EVENT_SCOPE(UPCGManagedFastGeoContainer::Release);
+
+	if (FastGeo)
+	{
+		FastGeo->Unregister();
+
+		// Note: This may be costly. Monitor performance and possibly build a way to keep the FastGeo container alive until async work is done.
+		FastGeo->Tick(/*bWaitForCompletion=*/true);
+
+		FastGeo = nullptr;
+	}
+
+	// Can always remove from PCG Component.
+	return true;
+}
diff --git a/Engine/Plugins/Experimental/PCGInterops/PCGFastGeoInterop/Source/PCGFastGeoInterop/Private/Components/PCGManagedFastGeoContainer.h b/Engine/Plugins/Experimental/PCGInterops/PCGFastGeoInterop/Source/PCGFastGeoInterop/Private/Components/PCGManagedFastGeoContainer.h
new file mode 100644
--- /dev/null
+++ b/Engine/Plugins/Experimental/PCGInterops/PCGFastGeoInterop/Source/PCGFastGeoInterop/Private/Components/PCGManagedFastGeoContainer.h
@@ -0,0 +1,24 @@
+// Copyright Epic Games, Inc. All Rights Reserved.
+
+#pragma once
+
+#include "PCGManagedResource.h"
+
+#include "FastGeoContainer.h"
+
+#include "PCGManagedFastGeoContainer.generated.h"
+
+UCLASS()
+class UPCGManagedFastGeoContainer : public UPCGManagedResource
+{
+	GENERATED_BODY()
+
+public:
+	//~Begin UPCGManagedResource interface
+	virtual bool Release(bool bHardRelease, TSet<TSoftObjectPtr<AActor>>& OutActorsToDelete) override;
+	virtual bool ReleaseOnTeardown() const override { return true; }
+	//~End UPCGManagedResource interface
+
+	UPROPERTY(Transient)
+	TObjectPtr<UFastGeoContainer> FastGeo;
+};
diff --git a/Engine/Plugins/Experimental/PCGInterops/PCGFastGeoInterop/Source/PCGFastGeoInterop/Private/Compute/PrimitiveFactories/PCGPrimitiveFactoryFastGeoPISMC.cpp b/Engine/Plugins/Experimental/PCGInterops/PCGFastGeoInterop/Source/PCGFastGeoInterop/Private/Compute/PrimitiveFactories/PCGPrimitiveFactoryFastGeoPISMC.cpp
new file mode 100644
--- /dev/null
+++ b/Engine/Plugins/Experimental/PCGInterops/PCGFastGeoInterop/Source/PCGFastGeoInterop/Private/Compute/PrimitiveFactories/PCGPrimitiveFactoryFastGeoPISMC.cpp
@@ -0,0 +1,152 @@
+// Copyright Epic Games, Inc. All Rights Reserved.
+
+#include "Compute/PrimitiveFactories/PCGPrimitiveFactoryFastGeoPISMC.h"
+
+#include "Components/PCGManagedFastGeoContainer.h"
+
+#include "PCGComponent.h"
+#include "PCGContext.h"
+#include "PCGModule.h"
+#include "Components/PCGProceduralISMComponentDescriptor.h"
+#include "Engine/World.h"
+
+bool FPCGPrimitiveFactoryFastGeoPISMC::IsRenderStateCreated() const
+{
+	if (Components.IsEmpty())
+	{
+		return false;
+	}
+
+	// Use first component to get to container.
+	FFastGeoComponent* Component = Components[0].Get();
+	UFastGeoContainer* Container = Component ? Component->GetOwnerContainer() : nullptr;
+
+	if (!Container)
+	{
+		return false;
+	}
+
+	if (Container->HasAnyPendingTasks())
+	{
+		Container->Tick();
+
+		if (Container->HasAnyPendingTasks())
+		{
+			return false;
+		}
+	}
+
+	for (int32 PrimitiveIndex = 0; PrimitiveIndex < GetNumPrimitives(); ++PrimitiveIndex)
+	{
+		if (GetSceneProxy(PrimitiveIndex) == nullptr)
+		{
+			return false;
+		}
+	}
+
+	return true;
+}
+
+FPrimitiveSceneProxy* FPCGPrimitiveFactoryFastGeoPISMC::GetSceneProxy(int32 InPrimitiveIndex) const
+{
+	if (ensure(Components.IsValidIndex(InPrimitiveIndex)) && Components[InPrimitiveIndex].Get())
+	{
+		return static_cast<FFastGeoPrimitiveComponent*>(Components[InPrimitiveIndex].Get())->GetSceneProxy();
+	}
+	else
+	{
+		return nullptr;
+	}
+}
+
+int32 FPCGPrimitiveFactoryFastGeoPISMC::GetNumInstances(int32 InPrimitiveIndex) const
+{
+	return ensure(InstanceCounts.IsValidIndex(InPrimitiveIndex)) ? InstanceCounts[InPrimitiveIndex] : 0;
+}
+
+void FPCGPrimitiveFactoryFastGeoPISMC::Initialize(FParameters&& InParameters)
+{
+	Descriptors = InParameters.Descriptors;
+}
+
+bool FPCGPrimitiveFactoryFastGeoPISMC::Create(FPCGContext* InContext)
+{
+	// todo_pcg: Move to execution source.
+	UPCGComponent* Component = Cast<UPCGComponent>(InContext->ExecutionSource.Get());
+	check(Component);
+
+	ULevel* Level = Component->GetOwner() ? Component->GetOwner()->GetLevel() : nullptr;
+	if (!Level)
+	{
+		UE_LOG(LogPCG, Error, TEXT("No persistent level found, geometry will not be created."));
+		return true;
+	}
+
+	// Currently only support transient editor world.
+	if (!ensure(InContext->ExecutionSource->GetExecutionState().GetWorld()->IsGameWorld()))
+	{
+		return true;
+	}
+
+	// Must be outer'd to level - so pick same level as pcg component.
+	UFastGeoContainer* FastGeo = FPCGContext::NewObject_AnyThread<UFastGeoContainer>(InContext, Level);
+
+	const FString DebugName = Component->GetExecutionState().GetDebugName();
+	TUniquePtr<FFastGeoComponentCluster> ComponentCluster = MakeUnique<FFastGeoComponentCluster>(FastGeo, *FString::Printf(TEXT("FastGeoComponentCluster_%s"), *DebugName));
+
+	MeshBounds.SetNumUninitialized(Descriptors.Num());
+
+	for (int32 DescIndex = 0; DescIndex < Descriptors.Num(); ++DescIndex)
+	{
+		const FPCGProceduralISMComponentDescriptor& Desc = Descriptors[DescIndex];
+
+		FFastGeoProceduralISMComponent& FastGeoComponent = static_cast<FFastGeoProceduralISMComponent&>(ComponentCluster->AddComponent(FFastGeoProceduralISMComponent::Type));
+		FastGeoComponent.InitializeFromComponentDescriptor(Desc);
+
+		MeshBounds[DescIndex] = Desc.StaticMesh->GetBoundingBox();
+	}
+
+	ensure(ComponentCluster->HasComponents());
+
+	FastGeo->AddComponentCluster(ComponentCluster.Get());
+
+	FastGeo->OnCreated();
+
+	FastGeo->PrecachePSOs();
+
+	FastGeo->Register();
+
+	int32 ComponentIndex = 0;
+	FastGeo->ForEachComponentCluster([This = this, &ComponentIndex, &Descriptors = Descriptors](FFastGeoComponentCluster& InComponentCluster)
+		{
+			check(ComponentIndex == 0);
+
+			InComponentCluster.ForEachComponent([This, &ComponentIndex, &Descriptors](FFastGeoComponent& InComponent)
+				{
+					check(Descriptors.IsValidIndex(ComponentIndex));
+					const FPCGProceduralISMComponentDescriptor& Desc = Descriptors[ComponentIndex];
+
+					// If not true, weakptr assignment below will fail.
+					check(InComponent.GetOwnerComponentCluster());
+
+					This->Components.Add(&InComponent);
+					This->InstanceCounts.Add(Desc.NumInstances);
+
+					++ComponentIndex;
+				});
+
+			check(ComponentIndex == Descriptors.Num());
+		});
+
+	UPCGManagedFastGeoContainer* ManagedPrimitives = FPCGContext::NewObject_AnyThread<UPCGManagedFastGeoContainer>(InContext, Component);
+	ManagedPrimitives->FastGeo = FastGeo;
+
+	Component->AddToManagedResources(ManagedPrimitives);
+
+	return true;
+}
+
+FBox FPCGPrimitiveFactoryFastGeoPISMC::GetMeshBounds(int32 InPrimitiveIndex) const
+{
+	return ensure(MeshBounds.IsValidIndex(InPrimitiveIndex)) ? MeshBounds[InPrimitiveIndex] : FBox();
+}
diff --git a/Engine/Plugins/Experimental/PCGInterops/PCGFastGeoInterop/Source/PCGFastGeoInterop/Private/Compute/PrimitiveFactories/PCGPrimitiveFactoryFastGeoPISMC.h b/Engine/Plugins/Experimental/PCGInterops/PCGFastGeoInterop/Source/PCGFastGeoInterop/Private/Compute/PrimitiveFactories/PCGPrimitiveFactoryFastGeoPISMC.h
new file mode 100644
--- /dev/null
+++ b/Engine/Plugins/Experimental/PCGInterops/PCGFastGeoInterop/Source/PCGFastGeoInterop/Private/Compute/PrimitiveFactories/PCGPrimitiveFactoryFastGeoPISMC.h
@@ -0,0 +1,37 @@
+// Copyright Epic Games, Inc. All Rights Reserved.
+
+#pragma once
+
+#include "Compute/PrimitiveFactories/PCGPrimitiveFactoryISMBase.h"
+
+#include "FastGeoWeakElement.h"
+#include "UObject/WeakObjectPtr.h"
+
+struct FPCGContext;
+class FPrimitiveSceneProxy;
+class UPrimitiveComponent;
+struct FPCGProceduralISMComponentDescriptor;
+
+class FPCGPrimitiveFactoryFastGeoPISMC : public IPCGPrimitiveFactoryISMBase
+{
+public:
+	//~ Begin IPCGRuntimePrimitiveFactory interface
+	virtual bool IsRenderStateCreated() const override;
+	virtual int32 GetNumPrimitives() const override { return Components.Num(); }
+	virtual FPrimitiveSceneProxy* GetSceneProxy(int32 InPrimitiveIndex) const override;
+	virtual int32 GetNumInstances(int32 InPrimitiveIndex) const override;
+	//~ End IPCGRuntimePrimitiveFactory interface
+
+	//~ Begin IPCGPrimitiveFactoryISMBase interface
+	virtual void Initialize(FParameters&& InParameters) override;
+	virtual bool Create(FPCGContext* InContext) override;
+	virtual FBox GetMeshBounds(int32 InPrimitiveIndex) const override;
+	//~ End IPCGPrimitiveFactoryISMBase interface
+
+protected:
+	TArray<FPCGProceduralISMComponentDescriptor> Descriptors;
+
+	TArray<FWeakFastGeoComponent> Components;
+	TArray<int32> InstanceCounts;
+	TArray<FBox> MeshBounds;
+};
diff --git a/Engine/Plugins/Experimental/PCGInterops/PCGFastGeoInterop/Source/PCGFastGeoInterop/Private/PCGFastGeoInteropModule.cpp b/Engine/Plugins/Experimental/PCGInterops/PCGFastGeoInterop/Source/PCGFastGeoInterop/Private/PCGFastGeoInteropModule.cpp
new file mode 100644
--- /dev/null
+++ b/Engine/Plugins/Experimental/PCGInterops/PCGFastGeoInterop/Source/PCGFastGeoInterop/Private/PCGFastGeoInteropModule.cpp
@@ -0,0 +1,23 @@
+// Copyright Epic Games, Inc. All Rights Reserved.
+
+#include "PCGFastGeoInteropModule.h"
+
+#include "Compute/PrimitiveFactories/PCGPrimitiveFactoryFastGeoPISMC.h"
+
+#include "Compute/DataInterfaces/Elements/PCGStaticMeshSpawnerDataInterface.h"
+#include "Modules/ModuleManager.h"
+
+IMPLEMENT_MODULE(FPCGFastGeoInteropModule, PCGFastGeoInterop);
+
+void FPCGFastGeoInteropModule::StartupModule()
+{
+	PCGStaticMeshSpawnerDataInterface::Private::SetupFastGeoPrimitiveFactory([]()
+	{
+		return MakeShared<FPCGPrimitiveFactoryFastGeoPISMC>();
+	});
+}
+
+void FPCGFastGeoInteropModule::ShutdownModule()
+{
+	PCGStaticMeshSpawnerDataInterface::Private::ResetFastGeoPrimitiveFactory();
+}
diff --git a/Engine/Plugins/Experimental/PCGInterops/PCGFastGeoInterop/Source/PCGFastGeoInterop/Public/PCGFastGeoInteropModule.h b/Engine/Plugins/Experimental/PCGInterops/PCGFastGeoInterop/Source/PCGFastGeoInterop/Public/PCGFastGeoInteropModule.h
new file mode 100644
--- /dev/null
+++ b/Engine/Plugins/Experimental/PCGInterops/PCGFastGeoInterop/Source/PCGFastGeoInterop/Public/PCGFastGeoInteropModule.h
@@ -0,0 +1,17 @@
+// Copyright Epic Games, Inc. All Rights Reserved.
+
+#pragma once
+
+#include "Modules/ModuleInterface.h"
+#include "Stats/Stats.h"
+
+/** Module for experimental plugin that provides access to FastGeo components. To be removed when FastGeo exits experimental. */
+class FPCGFastGeoInteropModule final : public IModuleInterface
+{
+public:
+	//~ IModuleInterface implementation
+	virtual void StartupModule() override;
+	virtual void ShutdownModule() override;
+	virtual bool SupportsDynamicReloading() override { return true; }
+	//~ End IModuleInterface implementation
+};
diff --git a/Engine/Plugins/PCG/Source/PCG/Private/Components/PCGProceduralISMComponentDescriptor.h b/Engine/Plugins/PCG/Source/PCG/Private/Components/PCGProceduralISMComponentDescriptor.h
deleted file mode 100644
diff --git a/Engine/Plugins/PCG/Source/PCG/Private/Compute/DataInterfaces/Elements/PCGStaticMeshSpawnerDataInterface.cpp b/Engine/Plugins/PCG/Source/PCG/Private/Compute/DataInterfaces/Elements/PCGStaticMeshSpawnerDataInterface.cpp
--- a/Engine/Plugins/PCG/Source/PCG/Private/Compute/DataInterfaces/Elements/PCGStaticMeshSpawnerDataInterface.cpp
+++ b/Engine/Plugins/PCG/Source/PCG/Private/Compute/DataInterfaces/Elements/PCGStaticMeshSpawnerDataInterface.cpp
@@ -1,6 +1,6 @@
 // Copyright Epic Games, Inc. All Rights Reserved.
 
-#include "PCGStaticMeshSpawnerDataInterface.h"
+#include "Compute/DataInterfaces/Elements/PCGStaticMeshSpawnerDataInterface.h"
 
 #include "PCGComponent.h"
 #include "PCGContext.h"
@@ -21,6 +21,7 @@
 #include "RenderGraphBuilder.h"
 #include "RenderGraphResources.h"
 #include "RHIResources.h"
+#include "SceneDefinitions.h"
 #include "ShaderCompilerCore.h"
 #include "ShaderCore.h"
 #include "ShaderParameterMetadataBuilder.h"
@@ -38,6 +39,28 @@ namespace PCGStaticMeshSpawnerDataInterface
 	static const FText CouldNotLoadStaticMeshFormat = LOCTEXT("CouldNotLoadStaticMesh", "Could not load static mesh from path '{0}'.");
 	static const FText TooManyPrimitivesFormat = LOCTEXT("TooManyPrimitives", "Attempted to emit too many primitive components, terminated after creating '{0}'.");
 	static const FText NoMeshEntriesFormat = LOCTEXT("NoMeshEntries", "No mesh entries provided in weighted mesh selector.");
+
+	static TAutoConsoleVariable<bool> CVarCreatePrimitivesComponentless(
+		TEXT("pcg.RuntimeGeneration.ComponentlessPrimitives"),
+		false,
+		TEXT("Uses a component-less path for creating primitives in game worlds (PIE and standalone)."));
+
+	static TFunction<TSharedPtr<IPCGPrimitiveFactoryISMBase>()> FastGeoPrimitiveFactoryGetter;
+
+	namespace Private
+	{
+		void SetupFastGeoPrimitiveFactory(TFunction<TSharedPtr<IPCGPrimitiveFactoryISMBase>()>&& Getter)
+		{
+			ensure(!FastGeoPrimitiveFactoryGetter); // Trivial set and reset patterns expected.
+			FastGeoPrimitiveFactoryGetter = MoveTemp(Getter);
+		}
+
+		void ResetFastGeoPrimitiveFactory()
+		{
+			ensure(FastGeoPrimitiveFactoryGetter); // Trivial set and reset patterns expected.
+			FastGeoPrimitiveFactoryGetter.Reset();
+		}
+	}
 }
 
 void UPCGStaticMeshSpawnerDataInterface::GetSupportedInputs(TArray<FShaderFunctionDefinition>& OutFunctions) const
@@ -753,7 +776,24 @@ bool UPCGStaticMeshSpawnerDataProvider::SetupPrimitives(FPCGContext* InContext,
 
 	if (!Primitives->PrimitiveFactory)
 	{
-		Factory = MakeShared<FPCGPrimitiveFactoryPISMC>();
+		bool bCreateComponentless = false;
+
+		if (PCGStaticMeshSpawnerDataInterface::CVarCreatePrimitivesComponentless.GetValueOnGameThread() && PCGStaticMeshSpawnerDataInterface::FastGeoPrimitiveFactoryGetter)
+		{
+			if (UWorld* World = InContext->ExecutionSource.IsValid() ? InContext->ExecutionSource->GetExecutionState().GetWorld() : nullptr)
+			{
+				bCreateComponentless = World->IsGameWorld();
+			}
+		}
+
+		if (bCreateComponentless)
+		{
+			Factory = PCGStaticMeshSpawnerDataInterface::FastGeoPrimitiveFactoryGetter();
+		}
+		else
+		{
+			Factory = MakeShared<FPCGPrimitiveFactoryPISMC>();
+		}
 
 		Primitives->PrimitiveFactory = Factory;
 
diff --git a/Engine/Plugins/PCG/Source/PCG/Private/Compute/DataInterfaces/Elements/PCGStaticMeshSpawnerDataInterface.h b/Engine/Plugins/PCG/Source/PCG/Private/Compute/DataInterfaces/Elements/PCGStaticMeshSpawnerDataInterface.h
deleted file mode 100644
diff --git a/Engine/Plugins/PCG/Source/PCG/Private/Compute/DataInterfaces/PCGComputeDataInterface.h b/Engine/Plugins/PCG/Source/PCG/Private/Compute/DataInterfaces/PCGComputeDataInterface.h
deleted file mode 100644
diff --git a/Engine/Plugins/PCG/Source/PCG/Private/Compute/PrimitiveFactories/IPCGRuntimePrimitiveFactory.h b/Engine/Plugins/PCG/Source/PCG/Private/Compute/PrimitiveFactories/IPCGRuntimePrimitiveFactory.h
deleted file mode 100644
diff --git a/Engine/Plugins/PCG/Source/PCG/Private/Compute/PrimitiveFactories/PCGPrimitiveFactoryISMBase.h b/Engine/Plugins/PCG/Source/PCG/Private/Compute/PrimitiveFactories/PCGPrimitiveFactoryISMBase.h
deleted file mode 100644
diff --git a/Engine/Plugins/PCG/Source/PCG/Private/PCGComponent.cpp b/Engine/Plugins/PCG/Source/PCG/Private/PCGComponent.cpp
--- a/Engine/Plugins/PCG/Source/PCG/Private/PCGComponent.cpp
+++ b/Engine/Plugins/PCG/Source/PCG/Private/PCGComponent.cpp
@@ -1892,6 +1892,33 @@ void UPCGComponent::EndPlay(const EEndPlayReason::Type EndPlayReason)
 		Subsystem->UnregisterPCGComponent(this);
 	}
 
+	if (ensure(!GeneratedResourcesInaccessible))
+	{
+		UE::TScopeLock ResourcesLock(GeneratedResourcesLock);
+
+		for (int32 ResourceIndex = GeneratedResources.Num() - 1; ResourceIndex >= 0; --ResourceIndex)
+		{
+			if (!GeneratedResources[ResourceIndex])
+			{
+				// Remove null entries.
+				GeneratedResources.RemoveAtSwap(ResourceIndex);
+			}
+			else if (GeneratedResources[ResourceIndex]->ReleaseOnTeardown())
+			{
+				if (GeneratedResources[ResourceIndex])
+				{
+					TSet<TSoftObjectPtr<AActor>> ActorsToDelete;
+					GeneratedResources[ResourceIndex]->Release(/*bHardRelease=*/true, ActorsToDelete);
+
+					// Don't support deleting actors during teardown.
+					ensure(ActorsToDelete.IsEmpty());
+				}
+
+				GeneratedResources.RemoveAtSwap(ResourceIndex);
+			}
+		}
+	}
+
 	Super::EndPlay(EndPlayReason);
 }
 
diff --git a/Engine/Plugins/PCG/Source/PCG/Public/Components/PCGProceduralISMComponentDescriptor.h b/Engine/Plugins/PCG/Source/PCG/Public/Components/PCGProceduralISMComponentDescriptor.h
new file mode 100644
--- /dev/null
+++ b/Engine/Plugins/PCG/Source/PCG/Public/Components/PCGProceduralISMComponentDescriptor.h
@@ -0,0 +1,32 @@
+// Copyright Epic Games, Inc. All Rights Reserved.
+
+#pragma once
+
+#include "CoreMinimal.h"
+#include "ISMPartition/ProceduralISMComponentDescriptor.h"
+
+#include "PCGProceduralISMComponentDescriptor.generated.h"
+
+struct FPCGSoftISMComponentDescriptor;
+
+/** Struct that holds properties that can be used to initialize PCG Procedural ISM Components. */
+USTRUCT()
+struct FPCGProceduralISMComponentDescriptor : public FProceduralISMComponentDescriptor
+{
+	GENERATED_BODY()
+
+public:
+	PCG_API FPCGProceduralISMComponentDescriptor& operator=(const FPCGSoftISMComponentDescriptor& Other);
+
+	PCG_API bool operator==(const FPCGProceduralISMComponentDescriptor& Other) const;
+	bool operator!=(const FPCGProceduralISMComponentDescriptor& Other) const { return !(*this == Other); }
+
+	friend inline uint32 GetTypeHash(const FPCGProceduralISMComponentDescriptor& Key)
+	{
+		return Key.GetTypeHash();
+	}
+
+public:
+	UPROPERTY()
+	TArray<FName> ComponentTags;
+};
diff --git a/Engine/Plugins/PCG/Source/PCG/Public/Compute/DataInterfaces/Elements/PCGStaticMeshSpawnerDataInterface.h b/Engine/Plugins/PCG/Source/PCG/Public/Compute/DataInterfaces/Elements/PCGStaticMeshSpawnerDataInterface.h
new file mode 100644
--- /dev/null
+++ b/Engine/Plugins/PCG/Source/PCG/Public/Compute/DataInterfaces/Elements/PCGStaticMeshSpawnerDataInterface.h
@@ -0,0 +1,148 @@
+// Copyright Epic Games, Inc. All Rights Reserved.
+
+#pragma once
+
+#include "Compute/DataInterfaces/PCGComputeDataInterface.h"
+#include "ComputeFramework/ComputeDataProvider.h"
+
+#include "Components/PCGProceduralISMComponentDescriptor.h"
+
+#include "PCGStaticMeshSpawnerDataInterface.generated.h"
+
+class FPCGStaticMeshSpawnerDataInterfaceParameters;
+class IPCGPrimitiveFactoryISMBase;
+class UPCGStaticMeshSpawnerSettings;
+class UPCGSettings;
+class UPCGStaticMeshSpawnerKernel;
+struct FPCGContext;
+
+namespace PCGStaticMeshSpawnerDataInterface
+{
+	// todo_pcg: Remove and integrate FastGeo component spawning directly into PCG once FastGeo exits experimental.
+	namespace Private
+	{
+		UE_EXPERIMENTAL(5.7, "This API is internal and not intended to be used.")
+		PCG_API void SetupFastGeoPrimitiveFactory(TFunction<TSharedPtr<IPCGPrimitiveFactoryISMBase>()>&& Getter);
+		UE_EXPERIMENTAL(5.7, "This API is internal and not intended to be used.")
+		PCG_API void ResetFastGeoPrimitiveFactory();
+	}
+}
+
+/** Data Interface to marshal Static Mesh Spawner settings to the GPU. */
+UCLASS(ClassGroup = (Procedural))
+class UPCGStaticMeshSpawnerDataInterface : public UPCGComputeDataInterface
+{
+	GENERATED_BODY()
+
+public:
+	//~ Begin UComputeDataInterface Interface
+	TCHAR const* GetClassName() const override { return TEXT("PCGStaticMeshSpawner"); }
+	void GetSupportedInputs(TArray<FShaderFunctionDefinition>& OutFunctions) const override;
+	void GetShaderParameters(TCHAR const* UID, FShaderParametersMetadataBuilder& InOutBuilder, FShaderParametersMetadataAllocations& InOutAllocations) const override;
+	void GetHLSL(FString& OutHLSL, FString const& InDataInterfaceName) const override;
+	UComputeDataProvider* CreateDataProvider() const override;
+	//~ End UComputeDataInterface Interface
+
+	static constexpr uint32 MAX_ATTRIBUTES = 64;
+};
+
+UCLASS()
+class UPCGStaticMeshSpawnerDataProvider : public UPCGComputeDataProvider
+{
+	GENERATED_BODY()
+
+public:
+	//~ Begin UComputeDataProvider Interface
+	virtual bool PerformPreExecuteReadbacks_GameThread(UPCGDataBinding* InBinding) override;
+	virtual bool PrepareForExecute_GameThread(UPCGDataBinding* InBinding) override;
+	FComputeDataProviderRenderProxy* GetRenderProxy() override;
+	virtual void Reset() override;
+	//~ End UComputeDataProvider Interface
+
+	void CreatePrimitiveDescriptors(FPCGContext* InContext, UPCGDataBinding* InBinding);
+
+	/** Setup primitives. Returns true if at least one component was set up. */
+	virtual bool SetupPrimitives(FPCGContext* InContext, UPCGDataBinding* InBinding);
+
+public:
+	/** Attributes to use for writing per-instance custom floats. */
+	UPROPERTY()
+	TArray<FUintVector4> AttributeIdOffsetStrides;
+
+	UPROPERTY()
+	TArray<int32> PrimitiveStringKeys;
+
+	UPROPERTY()
+	TArray<FBox> PrimitiveMeshBounds;
+
+	UPROPERTY()
+	TArray<float> PrimitiveSelectionCDF;
+
+	/** Attribute Id for mesh selector. */
+	UPROPERTY()
+	int32 SelectorAttributeId = INDEX_NONE;
+
+	UPROPERTY()
+	int32 NumInputPoints = 0;
+
+	/** Output attribute Id for the output selected mesh. */
+	UPROPERTY()
+	int32 SelectedMeshAttributeId = INDEX_NONE;
+
+	/** The number of instances per string key value, used for by-attribute spawning. */
+	TMap<int32, uint32> StringKeyToInstanceCount;
+
+	int32 AnalysisDataIndex = INDEX_NONE;
+
+	bool bPrimitiveDescriptorsCreated = false;
+
+	TArray<FPCGProceduralISMComponentDescriptor> PrimitiveDescriptors;
+
+	uint32 CustomFloatCount = 0;
+
+	bool bRegisteredPrimitives = false;
+
+	bool bStaticMeshesLoaded = false;
+
+	bool bPrimitivesSetUp = false;
+
+	/** Number of primitives created during primitive setup. */
+	int32 NumPrimitivesSetup = 0;
+};
+
+class FPCGStaticMeshSpawnerDataProviderProxy : public FComputeDataProviderRenderProxy
+{
+public:
+	FPCGStaticMeshSpawnerDataProviderProxy(
+		const TArray<FUintVector4>& InAttributeIdOffsetStrides,
+		int32 InSelectorAttributeId,
+		const TArray<int32>& InPrimitiveStringKeys,
+		TArray<float> InSelectionCDF,
+		int32 InSelectedMeshAttributeId,
+		const TArray<FBox>& InPrimitiveMeshBounds);
+
+	//~ Begin FComputeDataProviderRenderProxy Interface
+	bool IsValid(FValidationData const& InValidationData) const override;
+	void AllocateResources(FRDGBuilder& GraphBuilder, FAllocationData const& InAllocationData) override;
+	void GatherDispatchData(FDispatchData const& InDispatchData) override;
+	//~ End FComputeDataProviderRenderProxy Interface
+
+protected:
+	using FParameters = FPCGStaticMeshSpawnerDataInterfaceParameters;
+
+	TArray<FUintVector4> AttributeIdOffsetStrides;
+	TArray<float> SelectionCDF;
+
+	int32 SelectorAttributeId = INDEX_NONE;
+
+	TArray<int32> PrimitiveStringKeys;
+	FRDGBufferSRVRef PrimitiveStringKeysBufferSRV = nullptr;
+
+	TArray<FVector4f> PrimitiveMeshBoundsMin;
+	FRDGBufferSRVRef PrimitiveMeshBoundsMinBufferSRV = nullptr;
+
+	TArray<FVector4f> PrimitiveMeshBoundsMax;
+	FRDGBufferSRVRef PrimitiveMeshBoundsMaxBufferSRV = nullptr;
+
+	int32 SelectedMeshAttributeId = INDEX_NONE;
+};
diff --git a/Engine/Plugins/PCG/Source/PCG/Public/Compute/DataInterfaces/PCGComputeDataInterface.h b/Engine/Plugins/PCG/Source/PCG/Public/Compute/DataInterfaces/PCGComputeDataInterface.h
new file mode 100644
--- /dev/null
+++ b/Engine/Plugins/PCG/Source/PCG/Public/Compute/DataInterfaces/PCGComputeDataInterface.h
@@ -0,0 +1,184 @@
+// Copyright Epic Games, Inc. All Rights Reserved.
+
+#pragma once
+
+#include "Compute/PCGComputeCommon.h"
+
+#include "ComputeFramework/ComputeDataInterface.h"
+#include "ComputeFramework/ComputeDataProvider.h"
+
+#include "PCGComputeDataInterface.generated.h"
+
+class UPCGComputeKernel;
+class UPCGDataBinding;
+
+UCLASS(Abstract, ClassGroup = (Procedural))
+class UPCGComputeDataInterface : public UComputeDataInterface
+{
+	GENERATED_BODY()
+
+public:
+	/** Register a downstream pin (and optionally compute graph element pin alias). */
+	virtual void AddDownstreamInputPin(FName InInputPinLabel, const FName* InOptionalInputPinLabelAlias = nullptr);
+
+	/** Set the output pin label and label alias that this data interface is associated with. */
+	virtual void SetOutputPin(FName InOutputPinLabel, const FName* InOptionalOutputPinLabelAlias = nullptr);
+
+	void SetProducerSettings(const UPCGSettings* InProducerSettings);
+	const UPCGSettings* GetProducerSettings() const;
+
+	static void AddReferencedObjects(UObject* InThis, FReferenceCollector& Collector);
+
+	/** Set whether data originates from the CPU and should be uploaded to the GPU. */
+	void SetProducedByCPU(bool bInProducedByCPU) { bProducedByCPU = bInProducedByCPU; }
+
+public:
+	UPROPERTY()
+	int32 GraphBindingIndex = INDEX_NONE;
+
+	UPROPERTY()
+	TObjectPtr<const UPCGComputeKernel> ProducerKernel;
+
+	/** Label of output pin that this data interface is associated with. */
+	UPROPERTY()
+	FName OutputPinLabel;
+
+	UPROPERTY()
+	FName OutputPinLabelAlias;
+
+	/** Generated PCG data will be assigned these labels. */
+	UPROPERTY()
+	TArray<FName> DownstreamInputPinLabelAliases;
+
+	UPROPERTY()
+	bool bProducedByCPU = false;
+
+protected:
+	UPROPERTY()
+	TSoftObjectPtr<const UPCGSettings> ProducerSettings;
+
+	mutable TObjectPtr<const UPCGSettings> ResolvedProducerSettings = nullptr;
+};
+
+UCLASS(Abstract, ClassGroup (Procedural))
+class UPCGComputeDataProvider : public UComputeDataProvider
+{
+	GENERATED_BODY()
+
+public:
+	virtual void Initialize(const UComputeDataInterface* InDataInterface, UObject* InBinding, uint64 InInputMask, uint64 InOutputMask) override;
+	virtual void Reset() override;
+
+	/** Called before PrepareForExecute to allow data providers to do any readbacks and finalize data descriptions. */
+	virtual bool PerformPreExecuteReadbacks_GameThread(UPCGDataBinding* InBinding) { return true; }
+
+	/* Called prior to GetRenderProxy and execution. Any processing of data descriptions should be done here (after PerformPreExecuteReadbacks has been called). */
+	virtual bool PrepareForExecute_GameThread(UPCGDataBinding* InBinding) { return true; }
+
+	/** Return true when all done. */
+	virtual bool PostExecute(UPCGDataBinding* InBinding) { return true; }
+
+	/** Release any handles to transient resources like GPU buffers. */
+	virtual void ReleaseTransientResources(const TCHAR* InReason = nullptr) {}
+
+	const UPCGSettings* GetProducerSettings() const;
+	void SetProducerSettings(const UPCGSettings* InSettings);
+
+	static void AddReferencedObjects(UObject* InThis, FReferenceCollector& Collector);
+
+protected:
+#if WITH_EDITOR
+	void NotifyProducerUploadedData(UPCGDataBinding* InBinding);
+#endif
+
+public:
+	int32 GraphBindingIndex = INDEX_NONE;
+	FName OutputPinLabel;
+	FName OutputPinLabelAlias;
+	TArray<FName> DownstreamInputPinLabelAliases;
+	bool bProducedByCPU = false;
+
+	/** Bumped each time the data provider is initialized or reset, so that async callbacks can detect if they originated from a previous usage of the data provider and no-op. */
+	std::atomic<uint64> GenerationCounter = 0;
+
+protected:
+	UPROPERTY()
+	TObjectPtr<const UPCGComputeKernel> ProducerKernel;
+
+private:
+	/** Settings of node that produces this data, normally the upstream node. */
+	UPROPERTY()
+	TSoftObjectPtr<const UPCGSettings> ProducerSettings;
+
+	mutable TObjectPtr<const UPCGSettings> ResolvedProducerSettings = nullptr;
+};
+
+/** Interface for data that can be exported/downloaded from the GPU to the CPU. */
+UCLASS(Abstract, ClassGroup = (Procedural))
+class UPCGExportableDataInterface : public UPCGComputeDataInterface
+{
+	GENERATED_BODY()
+
+public:
+	void SetRequiresExport(bool bInRequiresExport) { bRequiresExport = bInRequiresExport; }
+	bool GetRequiresExport() const { return bRequiresExport; }
+
+protected:
+	/** Whether this data is passed to downstream tasks outside of this compute graph. */
+	UPROPERTY()
+	bool bRequiresExport = false;
+};
+
+UCLASS(Abstract, ClassGroup (Procedural))
+class UPCGExportableDataProvider : public UPCGComputeDataProvider
+{
+	GENERATED_BODY()
+
+public:
+	//~ Begin UComputeDataInterface Interface
+	virtual void Initialize(const UComputeDataInterface* InDataInterface, UObject* InBinding, uint64 InInputMask, uint64 InOutputMask) override;
+	virtual void Reset() override;
+	//~ End UComputeDataInterface Interface
+
+	//~ Begin UPCGComputeDataProvider Interface
+	virtual bool PrepareForExecute_GameThread(UPCGDataBinding* InBinding) override;
+	//~ End UPCGComputeDataProvider Interface
+
+	/** Returns true if the generated buffer should be exported to a persistent buffer to be passed to downstream nodes. */
+	bool IsExportRequired() const { return ExportMode != EPCGExportMode::NoExport; }
+
+	EPCGExportMode GetExportMode() const { return ExportMode; }
+	UPCGDataBinding* GetDataBinding() const { return Binding.Get(); }
+	const FPCGDataCollectionDesc* GetPinDescription() const { return PinDesc; }
+
+	DECLARE_EVENT(UPCGExportableDataProvider, FOnDataExported);
+	FOnDataExported& OnDataExported_GameThread() { return OnDataExported; }
+
+protected:
+	EPCGExportMode ExportMode = EPCGExportMode::NoExport;
+	FOnDataExported OnDataExported;
+
+	TWeakObjectPtr<UPCGDataBinding> Binding;
+	const FPCGDataCollectionDesc* PinDesc = nullptr;
+};
+
+/** Interface for marshalling settings/params to the GPU. */
+UCLASS(Abstract, ClassGroup = (Procedural))
+class UPCGKernelParamsDataInterface : public UPCGComputeDataInterface
+{
+	GENERATED_BODY()
+
+public:
+	//~ Begin UComputeDataInterface Interface
+	void GetSupportedInputs(TArray<FShaderFunctionDefinition>& OutFunctions) const override;
+	void GetHLSL(FString& OutHLSL, FString const& InDataInterfaceName) const override;
+	//~ End UComputeDataInterface Interface
+};
+
+UCLASS(Abstract, ClassGroup (Procedural))
+class UPCGKernelParamsDataProvider : public UPCGComputeDataProvider
+{
+	GENERATED_BODY()
+
+	// Nothing implemented, but provided to avoid breaking changes in the future in case it becomes necessary.
+};
diff --git a/Engine/Plugins/PCG/Source/PCG/Public/Compute/PrimitiveFactories/IPCGRuntimePrimitiveFactory.h b/Engine/Plugins/PCG/Source/PCG/Public/Compute/PrimitiveFactories/IPCGRuntimePrimitiveFactory.h
new file mode 100644
--- /dev/null
+++ b/Engine/Plugins/PCG/Source/PCG/Public/Compute/PrimitiveFactories/IPCGRuntimePrimitiveFactory.h
@@ -0,0 +1,19 @@
+// Copyright Epic Games, Inc. All Rights Reserved.
+
+#pragma once
+
+#include "CoreMinimal.h"
+
+class FPrimitiveSceneProxy;
+
+/** Interface for helper object that creates and adds primitives to the scene. */
+class IPCGRuntimePrimitiveFactory
+{
+public:
+	virtual ~IPCGRuntimePrimitiveFactory() = default;
+
+	virtual int32 GetNumPrimitives() const = 0;
+	virtual FPrimitiveSceneProxy* GetSceneProxy(int32 InPrimitiveIndex) const = 0;
+	virtual int32 GetNumInstances(int32 InPrimitiveIndex) const = 0;
+	virtual bool IsRenderStateCreated() const = 0;
+};
diff --git a/Engine/Plugins/PCG/Source/PCG/Public/Compute/PrimitiveFactories/PCGPrimitiveFactoryISMBase.h b/Engine/Plugins/PCG/Source/PCG/Public/Compute/PrimitiveFactories/PCGPrimitiveFactoryISMBase.h
new file mode 100644
--- /dev/null
+++ b/Engine/Plugins/PCG/Source/PCG/Public/Compute/PrimitiveFactories/PCGPrimitiveFactoryISMBase.h
@@ -0,0 +1,25 @@
+// Copyright Epic Games, Inc. All Rights Reserved.
+
+#pragma once
+
+#include "Compute/PrimitiveFactories/IPCGRuntimePrimitiveFactory.h"
+
+#include "UObject/WeakObjectPtr.h"
+
+class AActor;
+struct FPCGContext;
+struct FPCGProceduralISMComponentDescriptor;
+
+class IPCGPrimitiveFactoryISMBase : public IPCGRuntimePrimitiveFactory
+{
+public:
+	struct FParameters
+	{
+		TArray<FPCGProceduralISMComponentDescriptor> Descriptors;
+		AActor* TargetActor = nullptr; // Optional, unless creating actor primitive components.
+	};
+
+	virtual void Initialize(FParameters&& InParameters) = 0;
+	virtual bool Create(FPCGContext* InContext) = 0;
+	virtual FBox GetMeshBounds(int32 InPrimitiveIndex) const = 0;
+};
diff --git a/Engine/Plugins/PCG/Source/PCG/Public/PCGManagedResource.h b/Engine/Plugins/PCG/Source/PCG/Public/PCGManagedResource.h
--- a/Engine/Plugins/PCG/Source/PCG/Public/PCGManagedResource.h
+++ b/Engine/Plugins/PCG/Source/PCG/Public/PCGManagedResource.h
@@ -42,6 +42,9 @@ public:
 	/** Returns whether a resource can be used - generally true except for resources marked as transient (from loading) */
 	UE_API virtual bool CanBeUsed() const;
 
+	/** Whether resource is transient and should be released on tear down. */
+	UE_API virtual bool ReleaseOnTeardown() const { return false; }
+
 	/** Marks the resources as being kept and changed through generation */
 	virtual void MarkAsUsed() { ensure(CanBeUsed()); bIsMarkedUnused = false; }
 	/** Marks the resource as being reused as-is during the generation */
