[PCG] Fix grid cells missing from runtime generation.

Change detection and MaxNumComponentsGenerating were interacting badly. This fix makes sure we don't lose components that did not fit within current generation capacity.

#rnx

diff --git a/Engine/Plugins/PCG/Source/PCG/Private/RuntimeGen/PCGRuntimeGenScheduler.cpp b/Engine/Plugins/PCG/Source/PCG/Private/RuntimeGen/PCGRuntimeGenScheduler.cpp
--- a/Engine/Plugins/PCG/Source/PCG/Private/RuntimeGen/PCGRuntimeGenScheduler.cpp
+++ b/Engine/Plugins/PCG/Source/PCG/Private/RuntimeGen/PCGRuntimeGenScheduler.cpp
@@ -274,9 +274,6 @@ void FPCGRuntimeGenScheduler::Tick(APCGWorldActor* InPCGWorldActor, double InEnd
 
 	// 1. Queue nearby components for generation.
 	
-	// Mapping of component + coordinates to priorities - needed to compute max priority over all gen sources.
-	ComponentsToGenerate.Empty(128);
-
 	if (!GenSources.IsEmpty())
 	{
 		FTickQueueComponentsForGenerationInputs Inputs;
@@ -933,7 +930,7 @@ void FPCGRuntimeGenScheduler::TickCleanup(const TSet<IPCGGenSourceBase*>& InGenS
 	}
 }
 
-void FPCGRuntimeGenScheduler::TickScheduleGeneration(TMap<FGridGenerationKey, double>& InComponentsToGenerate)
+void FPCGRuntimeGenScheduler::TickScheduleGeneration(TMap<FGridGenerationKey, double>& InOutComponentsToGenerate)
 {
 	TRACE_CPUPROFILER_EVENT_SCOPE(FPCGRuntimeGenScheduler::TickScheduleGeneration);
 
@@ -961,24 +958,41 @@ void FPCGRuntimeGenScheduler::TickScheduleGeneration(TMap<FGridGenerationKey, do
 	}
 
 	const int MaxNumGenerating = FMath::Max(1, PCGRuntimeGenSchedulerHelpers::CVarNumGeneratingComponentsAtSameTime.GetValueOnAnyThread());
+	const int GeneratingCapacity = FMath::Max(0, MaxNumGenerating - NumGenerating);
+
+	if (GeneratingCapacity == 0)
+	{
+		// No capacity to start generating new components, try next frame.
+		return;
+	}
 
-	GeneratedComponents.Reserve(GeneratedComponents.Num() + InComponentsToGenerate.Num());
+	GeneratedComponents.Reserve(GeneratedComponents.Num() + InOutComponentsToGenerate.Num());
 
-	for (TPair<FGridGenerationKey, double>& GenerateEntry : InComponentsToGenerate)
+	// Components that we couldn't generate due to missed capacity. Generation will be attempted again next frame.
+	TMap<FGridGenerationKey, double> MissedComponents;
+	MissedComponents.Reserve(FMath::Max(0, InOutComponentsToGenerate.Num() - GeneratingCapacity));
+
+	for (auto It = InOutComponentsToGenerate.CreateConstIterator(); It; ++It)
 	{
 		if (NumGenerating >= MaxNumGenerating)
 		{
+			// We've filled generation capacity, record the remaining components to try again next frame.
+			for (; It; ++It)
+			{
+				MissedComponents.Add(*It);
+			}
+
 			break;
 		}
 
-		FGridGenerationKey& Key = GenerateEntry.Key;
-		const double Priority = GenerateEntry.Value;
+		const FGridGenerationKey& Key = It.Key();
+		const double Priority = It.Value();
 
 		const FPCGGridDescriptor GridDescriptor = Key.GetGridDescriptor();
 		const EPCGHiGenGrid Grid = GridDescriptor.GetHiGenGrid();
 
-		const FIntVector GridCoords = GenerateEntry.Key.GetGridCoords();
-		UPCGComponent* OriginalComponent = GenerateEntry.Key.GetOriginalComponent();
+		const FIntVector GridCoords = Key.GetGridCoords();
+		UPCGComponent* OriginalComponent = Key.GetOriginalComponent();
 		check(OriginalComponent);
 
 		// If the Grid is unbounded, we have a non-partitioned or unbounded component.
@@ -1105,6 +1119,9 @@ void FPCGRuntimeGenScheduler::TickScheduleGeneration(TMap<FGridGenerationKey, do
 		++NumGenerating;
 	}
 
+	// Set any missed components to retry next frame (or clear the pending components if no components were missed).
+	InOutComponentsToGenerate = MoveTemp(MissedComponents);
+
 #if PCG_RGS_ONSCREENDEBUGMESSAGES
 	PCGRuntimeGenSchedulerHelpers::Stats.NumGeneratingComponents = NumGenerating;
 #endif
diff --git a/Engine/Plugins/PCG/Source/PCG/Public/RuntimeGen/PCGRuntimeGenScheduler.h b/Engine/Plugins/PCG/Source/PCG/Public/RuntimeGen/PCGRuntimeGenScheduler.h
--- a/Engine/Plugins/PCG/Source/PCG/Public/RuntimeGen/PCGRuntimeGenScheduler.h
+++ b/Engine/Plugins/PCG/Source/PCG/Public/RuntimeGen/PCGRuntimeGenScheduler.h
@@ -116,7 +116,7 @@ protected:
 		UPCGComponent* GetOriginalComponent() const { return OriginalComponent.ResolveObjectPtr(); }
 
 		UPCGComponent* GetCachedLocalComponent() const { return CachedLocalComponent.Get(); }
-		void SetCachedLocalComponent(UPCGComponent* InComponent) { CachedLocalComponent = InComponent; }
+		void SetCachedLocalComponent(UPCGComponent* InComponent) const { CachedLocalComponent = InComponent; }
 
 		FPCGGridDescriptor GetGridDescriptor() const;
 
@@ -126,7 +126,7 @@ protected:
 		TObjectKey<UPCGComponent> OriginalComponent;
 
 		// Optional/opportunistic cached local component if one has been created.
-		TWeakObjectPtr<UPCGComponent> CachedLocalComponent;
+		mutable TWeakObjectPtr<UPCGComponent> CachedLocalComponent;
 
 		friend uint32 GetTypeHash(const FGridGenerationKey& In)
 		{
@@ -156,7 +156,7 @@ protected:
 	void TickCleanup(const TSet<IPCGGenSourceBase*>& GenSources, const APCGWorldActor* InPCGWorldActor, double InEndTime);
 
 	/** Schedule generation on components in priority order. */
-	void TickScheduleGeneration(TMap<FGridGenerationKey, double>& ComponentsToGenerate);
+	void TickScheduleGeneration(TMap<FGridGenerationKey, double>& InOutComponentsToGenerate);
 
 	/** Request any required virtual textures to be primed within the necessary generation radius. */
 	void TickRequestVirtualTexturePriming(const TSet<IPCGGenSourceBase*>& InGenSources);
@@ -211,9 +211,11 @@ private:
 	*/
 	TSet<FGridGenerationKey> GeneratedComponentsToRemove;
 
+	/** Mapping of component + coordinates to priorities - needed to compute max priority over all gen sources. */
+	TMap<FGridGenerationKey, double> ComponentsToGenerate;
+
 	// Local to member functions but hoisted for efficiency.
 	TSet<IPCGGenSourceBase*> GenSources;
-	TMap<FGridGenerationKey, double> ComponentsToGenerate;
 
 	/** Pool of RuntimeGen PartitionActors used for hierarchical generation. */
 	TArray<TObjectPtr<APCGPartitionActor>> PartitionActorPool;
