[PCG] Be more resilient to invalid parent in metadata initialize

#rnx

diff --git a/Engine/Plugins/PCG/Source/PCG/Private/Metadata/PCGMetadata.cpp b/Engine/Plugins/PCG/Source/PCG/Private/Metadata/PCGMetadata.cpp
--- a/Engine/Plugins/PCG/Source/PCG/Private/Metadata/PCGMetadata.cpp
+++ b/Engine/Plugins/PCG/Source/PCG/Private/Metadata/PCGMetadata.cpp
@@ -243,10 +243,14 @@ FPCGMetadataDomain* UPCGMetadata::CreateMetadataDomain_Unsafe(const FPCGMetadata
 FPCGMetadataInitializeParams::FPCGMetadataInitializeParams(const UPCGMetadata* InParent, const TArray<PCGMetadataEntryKey>* InOptionalEntriesToCopy)
 	: Parent(InParent)
 {
-	FPCGMetadataDomainInitializeParams& DefaultDomainParams = DomainInitializeParams.Emplace_GetRef(PCGMetadataDomainID::Default, FPCGMetadataDomainInitializeParams(InParent->GetConstDefaultMetadataDomain())).Get<1>();
-	if (InOptionalEntriesToCopy)
+	if (IsValid(InParent))
 	{
-		DefaultDomainParams.OptionalEntriesToCopy.Emplace(PCGValueRangeHelpers::MakeConstValueRange(*InOptionalEntriesToCopy));
+		FPCGMetadataDomainInitializeParams& DefaultDomainParams = DomainInitializeParams.Emplace_GetRef(PCGMetadataDomainID::Default, FPCGMetadataDomainInitializeParams(InParent->GetConstDefaultMetadataDomain())).Get<1>();
+	
+		if (InOptionalEntriesToCopy)
+		{
+			DefaultDomainParams.OptionalEntriesToCopy.Emplace(PCGValueRangeHelpers::MakeConstValueRange(*InOptionalEntriesToCopy));
+		}
 	}
 }
 
@@ -403,7 +407,12 @@ void UPCGMetadata::Serialize(FArchive& InArchive)
 
 void UPCGMetadata::K2_InitializeAsCopy(const UPCGMetadata* InMetadataToCopy, const TArray<int64>& InOptionalEntriesToCopy)
 {
-	return InitializeAsCopy(FPCGMetadataInitializeParams(InMetadataToCopy, !InOptionalEntriesToCopy.IsEmpty() ? &InOptionalEntriesToCopy : nullptr));
+	if (!IsValid(InMetadataToCopy))
+	{
+		return;
+	}
+
+	InitializeAsCopy(FPCGMetadataInitializeParams(InMetadataToCopy, !InOptionalEntriesToCopy.IsEmpty() ? &InOptionalEntriesToCopy : nullptr));
 }
 
 void UPCGMetadata::K2_InitializeAsCopyWithAttributeFilter(const UPCGMetadata* InMetadataToCopy,
@@ -412,7 +421,12 @@ void UPCGMetadata::K2_InitializeAsCopyWithAttributeFilter(const UPCGMetadata* In
 	EPCGMetadataFilterMode InFilterMode,
 	EPCGStringMatchingOperator InMatchOperator)
 {
-	return InitializeAsCopy(FPCGMetadataInitializeParams(InMetadataToCopy, InFilteredAttributes, InFilterMode, InMatchOperator, !InOptionalEntriesToCopy.IsEmpty() ? &InOptionalEntriesToCopy : nullptr));
+	if (!IsValid(InMetadataToCopy))
+	{
+		return;
+	}
+
+	InitializeAsCopy(FPCGMetadataInitializeParams(InMetadataToCopy, InFilteredAttributes, InFilterMode, InMatchOperator, !InOptionalEntriesToCopy.IsEmpty() ? &InOptionalEntriesToCopy : nullptr));
 }
 
 void UPCGMetadata::AddToCrc(FArchiveCrc32& Ar, bool bFullDataCrc) const
@@ -445,12 +459,17 @@ void UPCGMetadata::AddToCrc(FArchiveCrc32& Ar, bool bFullDataCrc) const
 
 void UPCGMetadata::Initialize(const UPCGMetadata* InParent)
 {
+	if (!IsValid(InParent))
+	{
+		return;
+	}
+
 	Initialize(FPCGMetadataInitializeParams(InParent));
 }
 
 void UPCGMetadata::Initialize(const UPCGMetadata* InParent, bool bAddAttributesFromParent)
 {
-	if (!InParent)
+	if (!IsValid(InParent))
 	{
 		return;
 	}
@@ -463,6 +482,11 @@ void UPCGMetadata::Initialize(const UPCGMetadata* InParent, bool bAddAttributesF
 
 void UPCGMetadata::InitializeWithAttributeFilter(const UPCGMetadata* InParent, const TSet<FName>& InFilteredAttributes, EPCGMetadataFilterMode InFilterMode, EPCGStringMatchingOperator InMatchOperator)
 {
+	if (!IsValid(InParent))
+	{
+		return;
+	}
+
 	Initialize(FPCGMetadataInitializeParams(InParent, InFilteredAttributes, InFilterMode, InMatchOperator));
 }
 
@@ -495,11 +519,21 @@ void UPCGMetadata::Initialize(const FPCGMetadataInitializeParams& InParams)
 
 void UPCGMetadata::InitializeAsCopy(const UPCGMetadata* InMetadataToCopy, const TArray<PCGMetadataEntryKey>* EntriesToCopy)
 {
+	if (!IsValid(InMetadataToCopy))
+	{
+		return;
+	}
+
 	InitializeAsCopy(FPCGMetadataInitializeParams(InMetadataToCopy, EntriesToCopy));
 }
 
 void UPCGMetadata::InitializeAsCopyWithAttributeFilter(const UPCGMetadata* InMetadataToCopy, const TSet<FName>& InFilteredAttributes, EPCGMetadataFilterMode InFilterMode, const TArray<PCGMetadataEntryKey>* EntriesToCopy, EPCGStringMatchingOperator InMatchOperator)
 {
+	if (!IsValid(InMetadataToCopy))
+	{
+		return;
+	}
+
 	InitializeAsCopy(FPCGMetadataInitializeParams(InMetadataToCopy, InFilteredAttributes, InFilterMode, InMatchOperator, EntriesToCopy));
 }
 
@@ -547,11 +581,21 @@ void UPCGMetadata::InitializeAsCopy(const FPCGMetadataInitializeParams& InParams
 
 void UPCGMetadata::AddAttributes(const UPCGMetadata* InOther)
 {
+	if (!IsValid(InOther))
+	{
+		return;
+	}
+
 	AddAttributes(FPCGMetadataInitializeParams(InOther));
 }
 
 void UPCGMetadata::AddAttributesFiltered(const UPCGMetadata* InOther, const TSet<FName>& InFilteredAttributes, EPCGMetadataFilterMode InFilterMode, EPCGStringMatchingOperator InMatchOperator)
 {
+	if (!IsValid(InOther))
+	{
+		return;
+	}
+
 	AddAttributes(FPCGMetadataInitializeParams(InOther, InFilteredAttributes, InFilterMode, InMatchOperator));
 }
 
