[PCG] Refactor PCG Blueprint Element classes to introduce a base class (PCGBlueprintBaseElement)
- Moved PCGBlueprintElement to PCGBlueprintDeprecatedElement.h/.cpp
- Ground work to introduce range based elements

diff --git a/Engine/Plugins/PCG/Source/PCG/Private/Blueprint/PCGBlueprintDeprecatedElement.cpp b/Engine/Plugins/PCG/Source/PCG/Private/Blueprint/PCGBlueprintDeprecatedElement.cpp
new file mode 100644
--- /dev/null
+++ b/Engine/Plugins/PCG/Source/PCG/Private/Blueprint/PCGBlueprintDeprecatedElement.cpp
@@ -0,0 +1,177 @@
+// Copyright Epic Games, Inc. All Rights Reserved.
+
+#include "Elements/Blueprint/PCGBlueprintDeprecatedElement.h"
+
+#include "Helpers/PCGAsync.h"
+
+#include UE_INLINE_GENERATED_CPP_BY_NAME(PCGBlueprintDeprecatedElement)
+
+#define LOCTEXT_NAMESPACE "PCGBlueprintElement"
+
+namespace PCGBlueprintConstants
+{
+	constexpr int32 RunawayResetFrequency = 1024;
+}
+
+void UPCGBlueprintElement::ExecuteWithContext_Implementation(FPCGContext& InContext, const FPCGDataCollection& Input, FPCGDataCollection& Output)
+{
+	Execute(Input, Output);
+}
+
+void UPCGBlueprintElement::PointLoop(FPCGContext& InContext, const UPCGPointData* InData, UPCGPointData*& OutData, UPCGPointData* OptionalOutData) const
+{
+	if (!InData)
+	{
+		PCGE_LOG_C(Error, GraphAndLog, &InContext, LOCTEXT("InvalidInputDataPointLoop", "Invalid input data in PointLoop"));
+		return;
+	}
+
+	if (OptionalOutData)
+	{
+		OutData = OptionalOutData;
+	}
+	else
+	{
+		OutData = FPCGContext::NewObject_AnyThread<UPCGPointData>(&InContext);
+		OutData->InitializeFromData(InData);
+	}
+
+	const TArray<FPCGPoint>& InPoints = InData->GetPoints();
+	TArray<FPCGPoint>& OutPoints = OutData->GetMutablePoints();
+
+	bool bPreviousBPStateValue = true;
+	std::swap(bPreviousBPStateValue, InContext.AsyncState.bIsCallingBlueprint);
+
+	FPCGAsync::AsyncPointProcessing(&InContext, InPoints.Num(), OutPoints, [this, &InContext, InData, OutData, &InPoints](int32 Index, FPCGPoint& OutPoint)
+	{
+		if (Index % PCGBlueprintConstants::RunawayResetFrequency == 0)
+		{
+			GInitRunaway(); // Reset periodically the iteration count, because we know we're in a fixed size loop.
+		}
+
+		return PointLoopBody(InContext, InData, InPoints[Index], OutPoint, OutData->Metadata, Index);
+	});
+
+	std::swap(bPreviousBPStateValue, InContext.AsyncState.bIsCallingBlueprint);
+}
+
+void UPCGBlueprintElement::VariableLoop(FPCGContext& InContext, const UPCGPointData* InData, UPCGPointData*& OutData, UPCGPointData* OptionalOutData) const
+{
+	if (!InData)
+	{
+		PCGE_LOG_C(Error, GraphAndLog, &InContext, LOCTEXT("InvalidInputDataVariableLoop", "Invalid input data in VariableLoop"));
+		return;
+	}
+
+	if (OptionalOutData)
+	{
+		OutData = OptionalOutData;
+	}
+	else
+	{
+		OutData = FPCGContext::NewObject_AnyThread<UPCGPointData>(&InContext);
+		OutData->InitializeFromData(InData);
+	}
+
+	const TArray<FPCGPoint>& InPoints = InData->GetPoints();
+	TArray<FPCGPoint>& OutPoints = OutData->GetMutablePoints();
+
+	bool bPreviousBPStateValue = true;
+	std::swap(bPreviousBPStateValue, InContext.AsyncState.bIsCallingBlueprint);
+
+	FPCGAsync::AsyncMultiPointProcessing(&InContext, InPoints.Num(), OutPoints, [this, &InContext, InData, OutData, &InPoints](int32 Index)
+	{
+		if (Index % PCGBlueprintConstants::RunawayResetFrequency == 0)
+		{
+			GInitRunaway(); // Reset periodically the iteration count, because we know we're in a fixed size loop.
+		}
+
+		return VariableLoopBody(InContext, InData, InPoints[Index], OutData->Metadata, Index);
+	});
+
+	std::swap(bPreviousBPStateValue, InContext.AsyncState.bIsCallingBlueprint);
+}
+
+void UPCGBlueprintElement::NestedLoop(FPCGContext& InContext, const UPCGPointData* InOuterData, const UPCGPointData* InInnerData, UPCGPointData*& OutData, UPCGPointData* OptionalOutData) const
+{
+	if (!InOuterData || !InInnerData)
+	{
+		PCGE_LOG_C(Error, GraphAndLog, &InContext, LOCTEXT("InvalidInputDataNestedLoop", "Invalid input data in NestedLoop"));
+		return;
+	}
+
+	if (OptionalOutData)
+	{
+		OutData = OptionalOutData;
+	}
+	else
+	{
+		OutData = FPCGContext::NewObject_AnyThread<UPCGPointData>(&InContext);
+		OutData->InitializeFromData(InOuterData);
+		OutData->Metadata->AddAttributes(InInnerData->Metadata);
+	}
+
+	const TArray<FPCGPoint>& InOuterPoints = InOuterData->GetPoints();
+	const TArray<FPCGPoint>& InInnerPoints = InInnerData->GetPoints();
+	TArray<FPCGPoint>& OutPoints = OutData->GetMutablePoints();
+
+	bool bPreviousBPStateValue = true;
+	std::swap(bPreviousBPStateValue, InContext.AsyncState.bIsCallingBlueprint);
+
+	FPCGAsync::AsyncPointProcessing(&InContext, InOuterPoints.Num() * InInnerPoints.Num(), OutPoints, [this, &InContext, InOuterData, InInnerData, OutData, &InOuterPoints, &InInnerPoints](int32 Index, FPCGPoint& OutPoint)
+	{
+		if (Index % PCGBlueprintConstants::RunawayResetFrequency == 0)
+		{
+			GInitRunaway(); // Reset periodically the iteration count, because we know we're in a fixed size loop.
+		}
+
+		const int32 OuterIndex = Index / InInnerPoints.Num();
+		const int32 InnerIndex = Index % InInnerPoints.Num();
+		return NestedLoopBody(InContext, InOuterData, InInnerData, InOuterPoints[OuterIndex], InInnerPoints[InnerIndex], OutPoint, OutData->Metadata, OuterIndex, InnerIndex);
+	});
+
+	std::swap(bPreviousBPStateValue, InContext.AsyncState.bIsCallingBlueprint);
+}
+
+void UPCGBlueprintElement::IterationLoop(FPCGContext& InContext, int64 NumIterations, UPCGPointData*& OutData, const UPCGSpatialData* InA, const UPCGSpatialData* InB, UPCGPointData* OptionalOutData) const
+{
+	if (NumIterations < 0)
+	{
+		PCGE_LOG_C(Error, GraphAndLog, &InContext, FText::Format(LOCTEXT("InvalidIterationCount", "Invalid number of iterations ({0})"), NumIterations));
+		return;
+	}
+
+	if (OptionalOutData)
+	{
+		OutData = OptionalOutData;
+	}
+	else
+	{
+		const UPCGSpatialData* Owner = (InA ? InA : InB);
+		OutData = FPCGContext::NewObject_AnyThread<UPCGPointData>(&InContext);
+
+		if (Owner)
+		{
+			OutData->InitializeFromData(Owner);
+		}
+	}
+
+	TArray<FPCGPoint>& OutPoints = OutData->GetMutablePoints();
+
+	bool bPreviousBPStateValue = true;
+	std::swap(bPreviousBPStateValue, InContext.AsyncState.bIsCallingBlueprint);
+
+	FPCGAsync::AsyncPointProcessing(&InContext, NumIterations, OutPoints, [this, &InContext, InA, InB, OutData](int32 Index, FPCGPoint& OutPoint)
+	{
+		if (Index % PCGBlueprintConstants::RunawayResetFrequency == 0)
+		{
+			GInitRunaway(); // Reset periodically the iteration count, because we know we're in a fixed size loop.
+		}
+
+		return IterationLoopBody(InContext, Index, InA, InB, OutPoint, OutData->Metadata);
+	});
+
+	std::swap(bPreviousBPStateValue, InContext.AsyncState.bIsCallingBlueprint);
+}
+
+#undef LOCTEXT_NAMESPACE
diff --git a/Engine/Plugins/PCG/Source/PCG/Private/Data/PCGDifferenceData.cpp b/Engine/Plugins/PCG/Source/PCG/Private/Data/PCGDifferenceData.cpp
--- a/Engine/Plugins/PCG/Source/PCG/Private/Data/PCGDifferenceData.cpp
+++ b/Engine/Plugins/PCG/Source/PCG/Private/Data/PCGDifferenceData.cpp
@@ -47,7 +47,7 @@ void UPCGDifferenceData::Initialize(const UPCGSpatialData* InData)
 
 void UPCGDifferenceData::K2_AddDifference(const UPCGSpatialData* InDifference)
 {
-	return AddDifference(UPCGBlueprintElement::ResolveContext(), InDifference);
+	return AddDifference(UPCGBlueprintBaseElement::ResolveContext(), InDifference);
 }
 
 void UPCGDifferenceData::AddDifference(FPCGContext* InContext, const UPCGSpatialData* InDifference)
diff --git a/Engine/Plugins/PCG/Source/PCG/Private/Data/PCGSpatialData.cpp b/Engine/Plugins/PCG/Source/PCG/Private/Data/PCGSpatialData.cpp
--- a/Engine/Plugins/PCG/Source/PCG/Private/Data/PCGSpatialData.cpp
+++ b/Engine/Plugins/PCG/Source/PCG/Private/Data/PCGSpatialData.cpp
@@ -299,7 +299,7 @@ void UPCGSpatialData::ProjectPoints(const TArrayView<const TPair<FTransform, FBo
 
 UPCGIntersectionData* UPCGSpatialData::K2_IntersectWith(const UPCGSpatialData* InOther) const
 {
-	return IntersectWith(UPCGBlueprintElement::ResolveContext(), InOther);
+	return IntersectWith(UPCGBlueprintBaseElement::ResolveContext(), InOther);
 }
 
 UPCGIntersectionData* UPCGSpatialData::IntersectWith(FPCGContext* InContext, const UPCGSpatialData* InOther) const
@@ -312,7 +312,7 @@ UPCGIntersectionData* UPCGSpatialData::IntersectWith(FPCGContext* InContext, con
 
 UPCGSpatialData* UPCGSpatialData::K2_ProjectOn(const UPCGSpatialData* InOther, const FPCGProjectionParams& InParams) const
 {
-	return ProjectOn(UPCGBlueprintElement::ResolveContext(), InOther, InParams);
+	return ProjectOn(UPCGBlueprintBaseElement::ResolveContext(), InOther, InParams);
 }
 
 UPCGSpatialData* UPCGSpatialData::ProjectOn(FPCGContext* InContext, const UPCGSpatialData* InOther, const FPCGProjectionParams& InParams) const
@@ -345,7 +345,7 @@ UPCGSpatialData* UPCGSpatialData::ProjectOn(FPCGContext* InContext, const UPCGSp
 
 UPCGUnionData* UPCGSpatialData::K2_UnionWith(const UPCGSpatialData* InOther) const
 {
-	return UnionWith(UPCGBlueprintElement::ResolveContext(), InOther);
+	return UnionWith(UPCGBlueprintBaseElement::ResolveContext(), InOther);
 }
 
 UPCGUnionData* UPCGSpatialData::UnionWith(FPCGContext* InContext, const UPCGSpatialData* InOther) const
@@ -358,7 +358,7 @@ UPCGUnionData* UPCGSpatialData::UnionWith(FPCGContext* InContext, const UPCGSpat
 
 UPCGDifferenceData* UPCGSpatialData::K2_Subtract(const UPCGSpatialData* InOther) const
 {
-	return Subtract(UPCGBlueprintElement::ResolveContext(), InOther);
+	return Subtract(UPCGBlueprintBaseElement::ResolveContext(), InOther);
 }
 
 UPCGDifferenceData* UPCGSpatialData::Subtract(FPCGContext* InContext, const UPCGSpatialData* InOther) const
diff --git a/Engine/Plugins/PCG/Source/PCG/Private/Elements/Blueprint/PCGBlueprintBaseElement.cpp b/Engine/Plugins/PCG/Source/PCG/Private/Elements/Blueprint/PCGBlueprintBaseElement.cpp
new file mode 100644
--- /dev/null
+++ b/Engine/Plugins/PCG/Source/PCG/Private/Elements/Blueprint/PCGBlueprintBaseElement.cpp
@@ -0,0 +1,292 @@
+// Copyright Epic Games, Inc. All Rights Reserved.
+
+#include "Elements/Blueprint/PCGBlueprintBaseElement.h"
+
+#include "PCGPin.h"
+#include "Elements/PCGExecuteBlueprint.h"
+
+#include "Algo/Transform.h"
+#include "Engine/World.h"
+#include "Helpers/PCGHelpers.h"
+#include "Math/RandomStream.h"
+
+#include UE_INLINE_GENERATED_CPP_BY_NAME(PCGBlueprintBaseElement)
+
+#define LOCTEXT_NAMESPACE "PCGBlueprintBaseElement"
+
+#if WITH_EDITOR
+namespace PCGBlueprintHelper
+{
+	TSet<TWeakObjectPtr<UObject>> GetDataDependencies(UPCGBlueprintBaseElement* InElement, int32 MaxDepth)
+	{
+		check(InElement && InElement->GetClass());
+		UClass* BPClass = InElement->GetClass();
+
+		TSet<TObjectPtr<UObject>> Dependencies;
+		PCGHelpers::GatherDependencies(InElement, Dependencies, MaxDepth);
+		TSet<TWeakObjectPtr<UObject>> WeakDependencies;
+
+		Algo::Transform(Dependencies, WeakDependencies, [](const TObjectPtr<UObject>& InObjectPtr) { return TWeakObjectPtr<UObject>(InObjectPtr); });
+
+		return WeakDependencies;
+	}
+}
+#endif // WITH_EDITOR
+
+UWorld* UPCGBlueprintBaseElement::GetWorld() const
+{
+#if WITH_EDITOR
+	return GWorld;
+#else
+	return InstanceWorld ? InstanceWorld : Super::GetWorld();
+#endif
+}
+
+void UPCGBlueprintBaseElement::PostLoad()
+{
+	Super::PostLoad();
+	Initialize();
+}
+
+void UPCGBlueprintBaseElement::BeginDestroy()
+{
+#if WITH_EDITOR
+	if (!DataDependencies.IsEmpty())
+	{
+		FCoreUObjectDelegates::OnObjectPropertyChanged.RemoveAll(this);
+		DataDependencies.Empty();
+	}
+#endif
+
+	Super::BeginDestroy();
+}
+
+void UPCGBlueprintBaseElement::Execute_Implementation(const FPCGDataCollection& Input, FPCGDataCollection& Output)
+{
+}
+
+void UPCGBlueprintBaseElement::Initialize()
+{
+#if WITH_EDITOR
+	UpdateDependencies();
+#endif
+}
+
+FPCGContext& UPCGBlueprintBaseElement::GetContext() const
+{
+	checkf(CurrentContext, TEXT("Execution context is not ready - do not call the GetContext method inside of non-execution methods"));
+	return *CurrentContext;
+}
+
+void UPCGBlueprintBaseElement::SetCurrentContext(FPCGContext* InCurrentContext)
+{
+	ensure(CurrentContext == nullptr || InCurrentContext == nullptr || CurrentContext == InCurrentContext);
+	CurrentContext = InCurrentContext;
+}
+
+FPCGContext* UPCGBlueprintBaseElement::ResolveContext()
+{
+	if (FFrame::GetThreadLocalTopStackFrame() && FFrame::GetThreadLocalTopStackFrame()->Object)
+	{
+		if (UPCGBlueprintBaseElement* Caller = Cast<UPCGBlueprintBaseElement>(FFrame::GetThreadLocalTopStackFrame()->Object))
+		{
+			return &Caller->GetContext();
+		}
+	}
+
+	return nullptr;
+}
+
+#if WITH_EDITOR
+void UPCGBlueprintBaseElement::PostEditChangeProperty(FPropertyChangedEvent& PropertyChangedEvent)
+{
+	Super::PostEditChangeProperty(PropertyChangedEvent);
+
+	UpdateDependencies();
+
+	OnBlueprintElementChangedDelegate.Broadcast(this);
+}
+
+void UPCGBlueprintBaseElement::UpdateDependencies()
+{
+	// Avoid calculating dependencies for graph execution element
+	if (!GetOuter()->IsA<UPCGBlueprintSettings>())
+	{
+		return;
+	}
+
+	// Backup to know if we need to unregister from the Delegate or not
+	const bool bHadDependencies = DataDependencies.Num() > 0;
+
+	// Since we don't really know what changed, let's just rebuild our data dependencies
+	DataDependencies = PCGBlueprintHelper::GetDataDependencies(this, DependencyParsingDepth);
+
+	// Only Bind to event if we do have dependencies
+	if (!DataDependencies.IsEmpty())
+	{
+		if (!bHadDependencies)
+		{
+			FCoreUObjectDelegates::OnObjectPropertyChanged.AddUObject(this, &UPCGBlueprintBaseElement::OnDependencyChanged);
+		}
+	}
+	else if (bHadDependencies)
+	{
+		FCoreUObjectDelegates::OnObjectPropertyChanged.RemoveAll(this);
+	}
+}
+
+void UPCGBlueprintBaseElement::OnDependencyChanged(UObject* Object, FPropertyChangedEvent& PropertyChangedEvent)
+{
+	if (PropertyChangedEvent.ChangeType == EPropertyChangeType::Interactive)
+	{
+		return;
+	}
+
+	// There are many engine notifications that aren't needed for us, esp. wrt to compilation
+	if (PropertyChangedEvent.Property == nullptr && PropertyChangedEvent.ChangeType == EPropertyChangeType::Unspecified)
+	{
+		return;
+	}
+
+	if (!DataDependencies.Contains(Object))
+	{
+		return;
+	}
+
+	OnBlueprintElementChangedDelegate.Broadcast(this);
+}
+
+FString UPCGBlueprintBaseElement::GetParentClassName()
+{
+	return FObjectPropertyBase::GetExportPath(UPCGBlueprintBaseElement::StaticClass());
+}
+#endif // WITH_EDITOR
+
+FName UPCGBlueprintBaseElement::NodeTitleOverride_Implementation() const
+{
+	return NAME_None;
+}
+
+FLinearColor UPCGBlueprintBaseElement::NodeColorOverride_Implementation() const
+{
+	return FLinearColor::White;
+}
+
+EPCGSettingsType UPCGBlueprintBaseElement::NodeTypeOverride_Implementation() const
+{
+	return EPCGSettingsType::Blueprint;
+}
+
+bool UPCGBlueprintBaseElement::IsCacheableOverride_Implementation() const
+{
+	return bIsCacheable;
+}
+
+int32 UPCGBlueprintBaseElement::DynamicPinTypesOverride_Implementation(const UPCGSettings* InSettings, const UPCGPin* InPin) const
+{
+	PRAGMA_DISABLE_DEPRECATION_WARNINGS
+	return static_cast<int32>(InSettings->UPCGSettings::GetCurrentPinTypesID(InPin));
+	PRAGMA_ENABLE_DEPRECATION_WARNINGS
+}
+
+TSet<FName> UPCGBlueprintBaseElement::CustomInputLabels() const
+{
+	TSet<FName> Labels;
+	for (const FPCGPinProperties& PinProperty : CustomInputPins)
+	{
+		Labels.Emplace(PinProperty.Label);
+	}
+
+	return Labels;
+}
+
+TSet<FName> UPCGBlueprintBaseElement::CustomOutputLabels() const
+{
+	TSet<FName> Labels;
+	for (const FPCGPinProperties& PinProperty : CustomOutputPins)
+	{
+		Labels.Emplace(PinProperty.Label);
+	}
+
+	return Labels;
+}
+
+TArray<FPCGPinProperties> UPCGBlueprintBaseElement::GetInputPins() const
+{
+	if (CurrentContext)
+	{
+		if (const UPCGBlueprintSettings* Settings = CurrentContext->GetInputSettings<UPCGBlueprintSettings>())
+		{
+			return Settings->InputPinProperties();
+		}
+	}
+	else if (const UPCGBlueprintSettings* OriginalSettings = Cast<UPCGBlueprintSettings>(GetOuter()))
+	{
+		return OriginalSettings->InputPinProperties();
+	}
+
+	// Can't retrieve settings - return only custom pins then
+	return CustomInputPins;
+}
+
+TArray<FPCGPinProperties> UPCGBlueprintBaseElement::GetOutputPins() const
+{
+	if (CurrentContext)
+	{
+		if (const UPCGBlueprintSettings* Settings = CurrentContext->GetInputSettings<UPCGBlueprintSettings>())
+		{
+			return Settings->OutputPinProperties();
+		}
+	}
+	else if (const UPCGBlueprintSettings* OriginalSettings = Cast<UPCGBlueprintSettings>(GetOuter()))
+	{
+		return OriginalSettings->OutputPinProperties();
+	}
+
+	// Can't retrieve settings - return only custom pins then
+	return CustomOutputPins;
+}
+
+bool UPCGBlueprintBaseElement::GetInputPinByLabel(FName InPinLabel, FPCGPinProperties& OutFoundPin) const
+{
+	TArray<FPCGPinProperties> InputPinProperties = GetInputPins();
+	for (const FPCGPinProperties& InputPin : InputPinProperties)
+	{
+		if (InputPin.Label == InPinLabel)
+		{
+			OutFoundPin = InputPin;
+			return true;
+		}
+	}
+
+	OutFoundPin = FPCGPinProperties();
+	return false;;
+}
+
+bool UPCGBlueprintBaseElement::GetOutputPinByLabel(FName InPinLabel, FPCGPinProperties& OutFoundPin) const
+{
+	TArray<FPCGPinProperties> OutputPinProperties = GetOutputPins();
+	for (const FPCGPinProperties& OutputPin : OutputPinProperties)
+	{
+		if (OutputPin.Label == InPinLabel)
+		{
+			OutFoundPin = OutputPin;
+			return true;
+		}
+	}
+
+	OutFoundPin = FPCGPinProperties();
+	return false;
+}
+
+int UPCGBlueprintBaseElement::GetSeed(FPCGContext& InContext) const
+{
+	return InContext.GetSeed();
+}
+
+FRandomStream UPCGBlueprintBaseElement::GetRandomStream(FPCGContext& InContext) const
+{
+	return FRandomStream(GetSeed(InContext));
+}
+
+#undef LOCTEXT_NAMESPACE
\ No newline at end of file
diff --git a/Engine/Plugins/PCG/Source/PCG/Private/Elements/PCGExecuteBlueprint.cpp b/Engine/Plugins/PCG/Source/PCG/Private/Elements/PCGExecuteBlueprint.cpp
--- a/Engine/Plugins/PCG/Source/PCG/Private/Elements/PCGExecuteBlueprint.cpp
+++ b/Engine/Plugins/PCG/Source/PCG/Private/Elements/PCGExecuteBlueprint.cpp
@@ -2,34 +2,21 @@
 
 #include "Elements/PCGExecuteBlueprint.h"
 
-#include "Algo/Transform.h"
-#include "Engine/Blueprint.h"
-#include "Math/RandomStream.h"
-#include "PCGComponent.h"
 #include "Data/PCGPointArrayData.h"
 #include "Data/PCGPointData.h"
-#include "Helpers/PCGAsync.h"
+#include "Elements/Blueprint/PCGBlueprintDeprecatedElement.h"
 #include "Helpers/PCGHelpers.h"
 #include "Helpers/PCGSettingsHelpers.h"
 
-#include "Engine/World.h"
-#include "PCGPin.h"
+#include "Algo/Transform.h"
+#include "Engine/Blueprint.h"
 #include "UObject/Package.h"
 #include "UObject/Stack.h"
 #include "UObject/UObjectThreadContext.h"
 
-#if WITH_EDITOR
-#include "Algo/Transform.h"
-#endif
-
 #include UE_INLINE_GENERATED_CPP_BY_NAME(PCGExecuteBlueprint)
 
-#define LOCTEXT_NAMESPACE "PCGBlueprintElement"
-
-namespace PCGBlueprintConstants
-{
-	constexpr int32 RunawayResetFrequency = 1024;
-}
+#define LOCTEXT_NAMESPACE "PCGExecuteBlueprint"
 
 FPCGContextBlueprintScope::FPCGContextBlueprintScope(FPCGContext* InContext)
 	: Context(InContext)
@@ -55,322 +42,6 @@ FPCGContextBlueprintScope::~FPCGContextBlueprintScope()
 	}
 }
 
-#if WITH_EDITOR
-namespace PCGBlueprintHelper
-{
-	TSet<TWeakObjectPtr<UObject>> GetDataDependencies(UPCGBlueprintElement* InElement, int32 MaxDepth)
-	{
-		check(InElement && InElement->GetClass());
-		UClass* BPClass = InElement->GetClass();
-
-		TSet<TObjectPtr<UObject>> Dependencies;
-		PCGHelpers::GatherDependencies(InElement, Dependencies, MaxDepth);
-		TSet<TWeakObjectPtr<UObject>> WeakDependencies;
-
-		Algo::Transform(Dependencies, WeakDependencies, [](const TObjectPtr<UObject>& InObjectPtr) { return TWeakObjectPtr<UObject>(InObjectPtr); });
-
-		return WeakDependencies;
-	}
-}
-#endif // WITH_EDITOR
-
-UWorld* UPCGBlueprintElement::GetWorld() const
-{
-#if WITH_EDITOR
-	return GWorld;
-#else
-	return InstanceWorld ? InstanceWorld : Super::GetWorld();
-#endif
-}
-
-void UPCGBlueprintElement::PostLoad()
-{
-	Super::PostLoad();
-	Initialize();
-
-#if WITH_EDITOR
-	if (!InputPinLabels_DEPRECATED.IsEmpty())
-	{
-		for (const FName& Label : InputPinLabels_DEPRECATED)
-		{
-			CustomInputPins.Emplace(Label);
-		}
-
-		InputPinLabels_DEPRECATED.Reset();
-	}
-
-	if (!OutputPinLabels_DEPRECATED.IsEmpty())
-	{
-		for (const FName& Label : OutputPinLabels_DEPRECATED)
-		{
-			CustomOutputPins.Emplace(Label);
-		}
-
-		OutputPinLabels_DEPRECATED.Reset();
-	}
-
-	// Go through the user-defined custom input pins and remove any Param pins labelled 'Params' or 'Param'. Such pins should not be
-	// added manually, the params pin is created dynamically from code based on presence of overrides.
-	for (int32 i = CustomInputPins.Num() - 1; i >= 0; --i)
-	{
-		FPCGPinProperties& Properties = CustomInputPins[i];
-		if (Properties.AllowedTypes == EPCGDataType::Param && (Properties.Label == FName(TEXT("Params")) || Properties.Label == FName(TEXT("Param"))))
-		{
-			// Non-swap version to preserve order
-			CustomInputPins.RemoveAt(i);
-		}
-	}
-
-	if (bCanBeMultithreaded_DEPRECATED)
-	{
-		bRequiresGameThread = false;
-	}
-	bCanBeMultithreaded_DEPRECATED = false;
-#endif
-}
-
-void UPCGBlueprintElement::BeginDestroy()
-{
-#if WITH_EDITOR
-	if (!DataDependencies.IsEmpty())
-	{
-		FCoreUObjectDelegates::OnObjectPropertyChanged.RemoveAll(this);
-		DataDependencies.Empty();
-	}
-#endif
-
-	Super::BeginDestroy();
-}
-
-void UPCGBlueprintElement::ExecuteWithContext_Implementation(FPCGContext& InContext, const FPCGDataCollection& Input, FPCGDataCollection& Output)
-{
-	Execute(Input, Output);
-}
-
-void UPCGBlueprintElement::Initialize()
-{
-#if WITH_EDITOR
-	UpdateDependencies();
-#endif
-}
-
-FPCGContext& UPCGBlueprintElement::GetContext() const
-{
-	checkf(CurrentContext, TEXT("Execution context is not ready - do not call the GetContext method inside of non-execution methods"));
-	return *CurrentContext;
-}
-
-void UPCGBlueprintElement::SetCurrentContext(FPCGContext* InCurrentContext)
-{
-	ensure(CurrentContext == nullptr || InCurrentContext == nullptr || CurrentContext == InCurrentContext);
-	CurrentContext = InCurrentContext;
-}
-
-FPCGContext* UPCGBlueprintElement::ResolveContext()
-{
-	if (FFrame::GetThreadLocalTopStackFrame() && FFrame::GetThreadLocalTopStackFrame()->Object)
-	{
-		if (UPCGBlueprintElement* Caller = Cast<UPCGBlueprintElement>(FFrame::GetThreadLocalTopStackFrame()->Object))
-		{
-			return &Caller->GetContext();
-		}
-	}
-
-	return nullptr;
-}
-
-#if WITH_EDITOR
-void UPCGBlueprintElement::PostEditChangeProperty(FPropertyChangedEvent& PropertyChangedEvent)
-{
-	Super::PostEditChangeProperty(PropertyChangedEvent);
-
-	UpdateDependencies();
-
-	OnBlueprintChangedDelegate.Broadcast(this);
-}
-
-void UPCGBlueprintElement::UpdateDependencies()
-{
-	// Avoid calculating dependencies for graph execution element
-	if (!GetOuter()->IsA<UPCGBlueprintSettings>())
-	{
-		return;
-	}
-
-	// Backup to know if we need to unregister from the Delegate or not
-	const bool bHadDependencies = DataDependencies.Num() > 0;
-
-	// Since we don't really know what changed, let's just rebuild our data dependencies
-	DataDependencies = PCGBlueprintHelper::GetDataDependencies(this, DependencyParsingDepth);
-
-	// Only Bind to event if we do have dependencies
-	if (!DataDependencies.IsEmpty())
-	{
-		if (!bHadDependencies)
-		{
-			FCoreUObjectDelegates::OnObjectPropertyChanged.AddUObject(this, &UPCGBlueprintElement::OnDependencyChanged);
-		}
-	}
-	else if(bHadDependencies)
-	{
-		FCoreUObjectDelegates::OnObjectPropertyChanged.RemoveAll(this);
-	}
-}
-
-void UPCGBlueprintElement::OnDependencyChanged(UObject* Object, FPropertyChangedEvent& PropertyChangedEvent)
-{
-	if (PropertyChangedEvent.ChangeType == EPropertyChangeType::Interactive)
-	{
-		return;
-	}
-
-	// There are many engine notifications that aren't needed for us, esp. wrt to compilation
-	if (PropertyChangedEvent.Property == nullptr && PropertyChangedEvent.ChangeType == EPropertyChangeType::Unspecified)
-	{
-		return;
-	}
-
-	if (!DataDependencies.Contains(Object))
-	{
-		return;
-	}
-
-	OnBlueprintChangedDelegate.Broadcast(this);
-}
-
-FString UPCGBlueprintElement::GetParentClassName()
-{
-	return FObjectPropertyBase::GetExportPath(UPCGBlueprintElement::StaticClass());
-}
-#endif // WITH_EDITOR
-
-FName UPCGBlueprintElement::NodeTitleOverride_Implementation() const
-{
-	return NAME_None;
-}
-
-FLinearColor UPCGBlueprintElement::NodeColorOverride_Implementation() const
-{
-	return FLinearColor::White;
-}
-
-EPCGSettingsType UPCGBlueprintElement::NodeTypeOverride_Implementation() const
-{
-	return EPCGSettingsType::Blueprint;
-}
-
-bool UPCGBlueprintElement::IsCacheableOverride_Implementation() const
-{
-	return bIsCacheable;
-}
-
-int32 UPCGBlueprintElement::DynamicPinTypesOverride_Implementation(const UPCGSettings* InSettings, const UPCGPin* InPin) const
-{
-PRAGMA_DISABLE_DEPRECATION_WARNINGS
-	return static_cast<int32>(InSettings->UPCGSettings::GetCurrentPinTypesID(InPin));
-PRAGMA_ENABLE_DEPRECATION_WARNINGS
-}
-
-TSet<FName> UPCGBlueprintElement::CustomInputLabels() const
-{
-	TSet<FName> Labels;
-	for (const FPCGPinProperties& PinProperty : CustomInputPins)
-	{
-		Labels.Emplace(PinProperty.Label);
-	}
-
-	return Labels;
-}
-
-TSet<FName> UPCGBlueprintElement::CustomOutputLabels() const
-{
-	TSet<FName> Labels;
-	for (const FPCGPinProperties& PinProperty : CustomOutputPins)
-	{
-		Labels.Emplace(PinProperty.Label);
-	}
-
-	return Labels;
-}
-
-TArray<FPCGPinProperties> UPCGBlueprintElement::GetInputPins() const
-{
-	if (CurrentContext)
-	{
-		if (const UPCGBlueprintSettings* Settings = CurrentContext->GetInputSettings<UPCGBlueprintSettings>())
-		{
-			return Settings->InputPinProperties();
-		}
-	}
-	else if (const UPCGBlueprintSettings* OriginalSettings = Cast<UPCGBlueprintSettings>(GetOuter()))
-	{
-		return OriginalSettings->InputPinProperties();
-	}
-	
-	// Can't retrieve settings - return only custom pins then
-	return CustomInputPins;
-}
-
-TArray<FPCGPinProperties> UPCGBlueprintElement::GetOutputPins() const
-{
-	if (CurrentContext)
-	{
-		if (const UPCGBlueprintSettings* Settings = CurrentContext->GetInputSettings<UPCGBlueprintSettings>())
-		{
-			return Settings->OutputPinProperties();
-		}
-	}
-	else if (const UPCGBlueprintSettings* OriginalSettings = Cast<UPCGBlueprintSettings>(GetOuter()))
-	{
-		return OriginalSettings->OutputPinProperties();
-	}
-
-	// Can't retrieve settings - return only custom pins then
-	return CustomOutputPins;
-}
-
-bool UPCGBlueprintElement::GetInputPinByLabel(FName InPinLabel, FPCGPinProperties& OutFoundPin) const
-{
-	TArray<FPCGPinProperties> InputPinProperties = GetInputPins();
-	for (const FPCGPinProperties& InputPin : InputPinProperties)
-	{
-		if (InputPin.Label == InPinLabel)
-		{
-			OutFoundPin = InputPin;
-			return true;
-		}
-	}
-
-	OutFoundPin = FPCGPinProperties();
-	return false;;
-}
-
-bool UPCGBlueprintElement::GetOutputPinByLabel(FName InPinLabel, FPCGPinProperties& OutFoundPin) const
-{
-	TArray<FPCGPinProperties> OutputPinProperties = GetOutputPins();
-	for (const FPCGPinProperties& OutputPin : OutputPinProperties)
-	{
-		if (OutputPin.Label == InPinLabel)
-		{
-			OutFoundPin = OutputPin;
-			return true;
-		}
-	}
-
-	OutFoundPin = FPCGPinProperties();
-	return false;
-}
-
-int UPCGBlueprintElement::GetSeed(FPCGContext& InContext) const 
-{
-	return InContext.GetSeed();
-}
-
-FRandomStream UPCGBlueprintElement::GetRandomStream(FPCGContext& InContext) const
-{
-	return FRandomStream(GetSeed(InContext));
-}
-
 UPCGBlueprintSettings::UPCGBlueprintSettings()
 {
 #if WITH_EDITORONLY_DATA
@@ -403,7 +74,7 @@ void UPCGBlueprintSettings::SetupBlueprintElementEvent()
 #if WITH_EDITOR
 	if (BlueprintElementInstance)
 	{
-		BlueprintElementInstance->OnBlueprintChangedDelegate.AddUObject(this, &UPCGBlueprintSettings::OnBlueprintElementChanged);
+		BlueprintElementInstance->OnBlueprintElementChangedDelegate.AddUObject(this, &UPCGBlueprintSettings::OnBlueprintElementChanged);
 	}
 #endif
 }
@@ -413,7 +84,7 @@ void UPCGBlueprintSettings::TeardownBlueprintElementEvent()
 #if WITH_EDITOR
 	if (BlueprintElementInstance)
 	{
-		BlueprintElementInstance->OnBlueprintChangedDelegate.RemoveAll(this);
+		BlueprintElementInstance->OnBlueprintElementChangedDelegate.RemoveAll(this);
 	}
 #endif
 }
@@ -422,14 +93,6 @@ void UPCGBlueprintSettings::PostLoad()
 {
 	Super::PostLoad();
 
-#if WITH_EDITOR
-	if (BlueprintElement_DEPRECATED && !BlueprintElementType)
-	{
-		BlueprintElementType = BlueprintElement_DEPRECATED;
-		BlueprintElement_DEPRECATED = nullptr;
-	}
-#endif
-
 	SetupBlueprintEvent();
 
 	if (!BlueprintElementInstance)
@@ -445,17 +108,7 @@ void UPCGBlueprintSettings::PostLoad()
 	{
 		BlueprintElementInstance->ConditionalPostLoad();
 		BlueprintElementInstance->SetFlags(RF_Transactional);
-#if WITH_EDITOR
-		if (bCanBeMultithreaded_DEPRECATED)
-		{
-			BlueprintElementInstance->bRequiresGameThread = false;
-		}
-#endif
 	}
-
-#if WITH_EDITOR
-	bCanBeMultithreaded_DEPRECATED = false;
-#endif
 }
 
 void UPCGBlueprintSettings::BeginDestroy()
@@ -509,9 +162,9 @@ void UPCGBlueprintSettings::OnObjectsReplaced(const TMap<UObject*, UObject*>& In
 	{
 		// When the blueprint changes, the element gets recreated, so we must rewire it here.
 		TeardownBlueprintElementEvent();
-		BlueprintElementInstance = Cast<UPCGBlueprintElement>(NewObject);
+		BlueprintElementInstance = Cast<UPCGBlueprintBaseElement>(NewObject);
 		SetupBlueprintElementEvent();
-		
+
 		DirtyCache();
 
 		if (BlueprintElementInstance)
@@ -526,7 +179,7 @@ void UPCGBlueprintSettings::OnObjectsReplaced(const TMap<UObject*, UObject*>& In
 	}
 }
 
-void UPCGBlueprintSettings::OnBlueprintElementChanged(UPCGBlueprintElement* InElement)
+void UPCGBlueprintSettings::OnBlueprintElementChanged(UPCGBlueprintBaseElement* InElement)
 {
 	if (InElement == BlueprintElementInstance)
 	{
@@ -538,6 +191,7 @@ void UPCGBlueprintSettings::OnBlueprintElementChanged(UPCGBlueprintElement* InEl
 }
 #endif // WITH_EDITOR
 
+// deprecated
 void UPCGBlueprintSettings::SetElementType(TSubclassOf<UPCGBlueprintElement> InElementType, UPCGBlueprintElement*& ElementInstance)
 {
 	if (!BlueprintElementInstance || InElementType != BlueprintElementType)
@@ -548,7 +202,24 @@ void UPCGBlueprintSettings::SetElementType(TSubclassOf<UPCGBlueprintElement> InE
 			BlueprintElementType = InElementType;
 			SetupBlueprintEvent();
 		}
-		
+
+		RefreshBlueprintElement();
+	}
+
+	ElementInstance = Cast<UPCGBlueprintElement>(BlueprintElementInstance);
+}
+
+void UPCGBlueprintSettings::SetBlueprintElementType(TSubclassOf<UPCGBlueprintBaseElement> InElementType, UPCGBlueprintBaseElement*& ElementInstance)
+{
+	if (!BlueprintElementInstance || InElementType != BlueprintElementType)
+	{
+		if (InElementType != BlueprintElementType)
+		{
+			TeardownBlueprintEvent();
+			BlueprintElementType = InElementType;
+			SetupBlueprintEvent();
+		}
+
 		RefreshBlueprintElement();
 	}
 
@@ -561,7 +232,7 @@ void UPCGBlueprintSettings::RefreshBlueprintElement()
 
 	if (BlueprintElementType)
 	{
-		BlueprintElementInstance = NewObject<UPCGBlueprintElement>(this, BlueprintElementType, NAME_None, RF_Transactional);
+		BlueprintElementInstance = NewObject<UPCGBlueprintBaseElement>(this, BlueprintElementType, NAME_None, RF_Transactional);
 		BlueprintElementInstance->Initialize();
 		SetupBlueprintElementEvent();
 	}
@@ -586,7 +257,7 @@ FPCGDataTypeIdentifier UPCGBlueprintSettings::GetCurrentPinTypesID(const UPCGPin
 	// We can't call a BP function while we are PostLoading, in that case (or if we don't have an instance or it's not dynamic) just return the pin allowed types.
 	if (BlueprintElementInstance && BlueprintElementInstance->bHasDynamicPins && !FUObjectThreadContext::Get().IsRoutingPostLoad)
 	{
-		return FPCGDataTypeIdentifier{static_cast<EPCGDataType>(BlueprintElementInstance->DynamicPinTypesOverride(this, InPin))};
+		return FPCGDataTypeIdentifier{ static_cast<EPCGDataType>(BlueprintElementInstance->DynamicPinTypesOverride(this, InPin)) };
 	}
 	else
 	{
@@ -633,7 +304,7 @@ UObject* UPCGBlueprintSettings::GetJumpTargetForDoubleClick() const
 	{
 		return BlueprintElementType->ClassGeneratedBy;
 	}
-	
+
 	return Super::GetJumpTargetForDoubleClick();
 }
 
@@ -644,7 +315,7 @@ void UPCGBlueprintSettings::ApplyDeprecationBeforeUpdatePins(UPCGNode* InOutNode
 	// Rename first found 'Param' or 'Params' pin to 'Overrides' which helps to ensure legacy params pins will retain incident edges.
 	for (TObjectPtr<UPCGPin>& InputPin : InputPins)
 	{
-		if (InputPin && InputPin->Properties.AllowedTypes == EPCGDataType::Param && 
+		if (InputPin && InputPin->Properties.AllowedTypes == EPCGDataType::Param &&
 			(InputPin->Properties.Label == FName(TEXT("Params")) || InputPin->Properties.Label == FName(TEXT("Param"))))
 		{
 			InputPin->Properties.Label = PCGPinConstants::DefaultParamsLabel;
@@ -776,7 +447,7 @@ void UPCGBlueprintSettings::FixingOverridableParamPropertyClass(FPCGSettingsOver
 		}
 	}
 
-	if(!bFound)
+	if (!bFound)
 	{
 		Super::FixingOverridableParamPropertyClass(Param);
 	}
@@ -810,29 +481,47 @@ bool FPCGExecuteBlueprintElement::ExecuteInternal(FPCGContext* InContext) const
 		for (int32 InputIndex = 0; InputIndex < Context->InputData.TaggedData.Num(); ++InputIndex)
 		{
 			const FPCGTaggedData& Input = Context->InputData.TaggedData[InputIndex];
-			if (const UPCGPointData* PointData = Cast<UPCGPointData>(Input.Data))
+			if (const UPCGBasePointData* PointData = Cast<UPCGBasePointData>(Input.Data))
 			{
-				PCGE_LOG(Verbose, LogOnly, FText::Format(LOCTEXT("InputPointInfo", "Input {0} has {1} points"), InputIndex, PointData->GetPoints().Num()));
+				PCGE_LOG(Verbose, LogOnly, FText::Format(LOCTEXT("InputPointInfo", "Input {0} has {1} points"), InputIndex, PointData->GetNumPoints()));
 			}
 		}
 
 		// Note that the context is actually sliced so there should never be any members in the BP element context that are visible/accessible from blueprint
 		/** Finally, execute the actual blueprint */
 		Context->BlueprintElementInstance->SetCurrentContext(Context);
-		
+
 		{
-			FPCGContextBlueprintScope BlueprintScope(Context);
+			// Support deprecated ExecuteWithContext
+			UPCGBlueprintElement* DeprecatedElement = Cast<UPCGBlueprintElement>(Context->BlueprintElementInstance);
 
 			// When running outside of main thread make sure GC can't run (BP Nodes might create objects which can't happen while GC runs)
 			if (!IsInGameThread())
 			{
 				ensure(!Context->AsyncState.bIsRunningOnMainThread);
 				FGCScopeGuard Scope;
-				Context->BlueprintElementInstance->ExecuteWithContext(*Context, Context->InputData, Context->OutputData);
+
+				if (DeprecatedElement)
+				{
+					FPCGContextBlueprintScope BlueprintScope(Context);
+					DeprecatedElement->ExecuteWithContext(*Context, Context->InputData, Context->OutputData);
+				}
+				else
+				{
+					Context->BlueprintElementInstance->Execute(Context->InputData, Context->OutputData);
+				}
 			}
 			else
 			{
-				Context->BlueprintElementInstance->ExecuteWithContext(*Context, Context->InputData, Context->OutputData);
+				if (DeprecatedElement)
+				{
+					FPCGContextBlueprintScope BlueprintScope(Context);
+					DeprecatedElement->ExecuteWithContext(*Context, Context->InputData, Context->OutputData);
+				}
+				else
+				{
+					Context->BlueprintElementInstance->Execute(Context->InputData, Context->OutputData);
+				}
 			}
 		}
 
@@ -843,7 +532,7 @@ bool FPCGExecuteBlueprintElement::ExecuteInternal(FPCGContext* InContext) const
 		// Nothing to do but forward data
 		Context->OutputData = Context->InputData;
 	}
-	
+
 	return true;
 }
 
@@ -858,15 +547,15 @@ void FPCGExecuteBlueprintElement::PostExecuteInternal(FPCGContext* InContext) co
 
 		// Build a list of InputObjects so that we don't remove Async flags from Outputs that come from the Input as it isn't our responsability
 		TSet<TObjectPtr<const UObject>> InputObjects;
-		Algo::TransformIf(Context->InputData.TaggedData, InputObjects, 
-		[](const FPCGTaggedData& TaggedData)
-		{
-			return TaggedData.Data != nullptr;
-		},
-		[](const FPCGTaggedData& TaggedData)
-		{
-			return TaggedData.Data;
-		});
+		Algo::TransformIf(Context->InputData.TaggedData, InputObjects,
+			[](const FPCGTaggedData& TaggedData)
+			{
+				return TaggedData.Data != nullptr;
+			},
+			[](const FPCGTaggedData& TaggedData)
+			{
+				return TaggedData.Data;
+			});
 
 		// Log info on outputs
 		for (int32 OutputIndex = 0; OutputIndex < Context->OutputData.TaggedData.Num(); ++OutputIndex)
@@ -892,20 +581,20 @@ void FPCGExecuteBlueprintElement::PostExecuteInternal(FPCGContext* InContext) co
 				}
 
 				auto ReOuterToTransientPackageIfCreatedFromThis = [Context](UObject* InObject)
-				{
-					bool bHasInstanceAsOuter = false;
-					UObject* CurrentObject = InObject;
-					while (CurrentObject && !bHasInstanceAsOuter)
 					{
-						bHasInstanceAsOuter = (CurrentObject->GetOuter() == Context->BlueprintElementInstance);
-						CurrentObject = CurrentObject->GetOuter();
-					}
-
-					if (bHasInstanceAsOuter)
-					{
-						InObject->Rename(nullptr, GetTransientPackage(), REN_DoNotDirty | REN_DontCreateRedirectors | REN_NonTransactional);
-					}
-				};
+						bool bHasInstanceAsOuter = false;
+						UObject* CurrentObject = InObject;
+						while (CurrentObject && !bHasInstanceAsOuter)
+						{
+							bHasInstanceAsOuter = (CurrentObject->GetOuter() == Context->BlueprintElementInstance);
+							CurrentObject = CurrentObject->GetOuter();
+						}
+
+						if (bHasInstanceAsOuter)
+						{
+							InObject->Rename(nullptr, GetTransientPackage(), REN_DoNotDirty | REN_DontCreateRedirectors | REN_NonTransactional);
+						}
+					};
 
 				UObject* ThisData = const_cast<UPCGData*>(Output.Data.Get());
 				ReOuterToTransientPackageIfCreatedFromThis(ThisData);
@@ -931,8 +620,8 @@ void FPCGExecuteBlueprintElement::PostExecuteInternal(FPCGContext* InContext) co
 					}
 				}
 
-				// @todo_pcg: expose PointArrayData to the Blueprint API but for now if we have PointArrayData enabled convert any PointData output to PointArrayData
-				if (CVarPCGEnablePointArrayData.GetValueOnAnyThread())
+				// If current element (deprecated elements only) does not support base point data then convert its output back to Point Array Data
+				if (CVarPCGEnablePointArrayData.GetValueOnAnyThread() && !SupportsBasePointDataInputs(Context))
 				{
 					if (const UPCGPointData* PointData = Cast<UPCGPointData>(Output.Data))
 					{
@@ -944,162 +633,6 @@ void FPCGExecuteBlueprintElement::PostExecuteInternal(FPCGContext* InContext) co
 	}
 }
 
-void UPCGBlueprintElement::PointLoop(FPCGContext& InContext, const UPCGPointData* InData, UPCGPointData*& OutData, UPCGPointData* OptionalOutData) const
-{
-	if (!InData)
-	{
-		PCGE_LOG_C(Error, GraphAndLog, &InContext, LOCTEXT("InvalidInputDataPointLoop", "Invalid input data in PointLoop"));
-		return;
-	}
-
-	if (OptionalOutData)
-	{
-		OutData = OptionalOutData;
-	}
-	else
-	{
-		OutData = FPCGContext::NewObject_AnyThread<UPCGPointData>(&InContext);
-		OutData->InitializeFromData(InData);
-	}
-
-	const TArray<FPCGPoint>& InPoints = InData->GetPoints();
-	TArray<FPCGPoint>& OutPoints = OutData->GetMutablePoints();
-
-	bool bPreviousBPStateValue = true;
-	std::swap(bPreviousBPStateValue, InContext.AsyncState.bIsCallingBlueprint);
-
-	FPCGAsync::AsyncPointProcessing(&InContext, InPoints.Num(), OutPoints, [this, &InContext, InData, OutData, &InPoints](int32 Index, FPCGPoint& OutPoint)
-	{
-		if (Index % PCGBlueprintConstants::RunawayResetFrequency == 0)
-		{
-			GInitRunaway(); // Reset periodically the iteration count, because we know we're in a fixed size loop.
-		}
-
-		return PointLoopBody(InContext, InData, InPoints[Index], OutPoint, OutData->Metadata, Index);
-	});
-
-	std::swap(bPreviousBPStateValue, InContext.AsyncState.bIsCallingBlueprint);
-}
-
-void UPCGBlueprintElement::VariableLoop(FPCGContext& InContext, const UPCGPointData* InData, UPCGPointData*& OutData, UPCGPointData* OptionalOutData) const
-{
-	if (!InData)
-	{
-		PCGE_LOG_C(Error, GraphAndLog, &InContext, LOCTEXT("InvalidInputDataVariableLoop", "Invalid input data in VariableLoop"));
-		return;
-	}
-
-	if (OptionalOutData)
-	{
-		OutData = OptionalOutData;
-	}
-	else
-	{
-		OutData = FPCGContext::NewObject_AnyThread<UPCGPointData>(&InContext);
-		OutData->InitializeFromData(InData);
-	}
-
-	const TArray<FPCGPoint>& InPoints = InData->GetPoints();
-	TArray<FPCGPoint>& OutPoints = OutData->GetMutablePoints();
-
-	bool bPreviousBPStateValue = true;
-	std::swap(bPreviousBPStateValue, InContext.AsyncState.bIsCallingBlueprint);
-
-	FPCGAsync::AsyncMultiPointProcessing(&InContext, InPoints.Num(), OutPoints, [this, &InContext, InData, OutData, &InPoints](int32 Index)
-	{
-		if (Index % PCGBlueprintConstants::RunawayResetFrequency == 0)
-		{
-			GInitRunaway(); // Reset periodically the iteration count, because we know we're in a fixed size loop.
-		}
-
-		return VariableLoopBody(InContext, InData, InPoints[Index], OutData->Metadata, Index);
-	});
-
-	std::swap(bPreviousBPStateValue, InContext.AsyncState.bIsCallingBlueprint);
-}
-
-void UPCGBlueprintElement::NestedLoop(FPCGContext& InContext, const UPCGPointData* InOuterData, const UPCGPointData* InInnerData, UPCGPointData*& OutData, UPCGPointData* OptionalOutData) const
-{
-	if (!InOuterData || !InInnerData)
-	{
-		PCGE_LOG_C(Error, GraphAndLog, &InContext, LOCTEXT("InvalidInputDataNestedLoop", "Invalid input data in NestedLoop"));
-		return;
-	}
-
-	if (OptionalOutData)
-	{
-		OutData = OptionalOutData;
-	}
-	else
-	{
-		OutData = FPCGContext::NewObject_AnyThread<UPCGPointData>(&InContext);
-		OutData->InitializeFromData(InOuterData);
-		OutData->Metadata->AddAttributes(InInnerData->Metadata);
-	}
-
-	const TArray<FPCGPoint>& InOuterPoints = InOuterData->GetPoints();
-	const TArray<FPCGPoint>& InInnerPoints = InInnerData->GetPoints();
-	TArray<FPCGPoint>& OutPoints = OutData->GetMutablePoints();
-
-	bool bPreviousBPStateValue = true;
-	std::swap(bPreviousBPStateValue, InContext.AsyncState.bIsCallingBlueprint);
-
-	FPCGAsync::AsyncPointProcessing(&InContext, InOuterPoints.Num() * InInnerPoints.Num(), OutPoints, [this, &InContext, InOuterData, InInnerData, OutData, &InOuterPoints, &InInnerPoints](int32 Index, FPCGPoint& OutPoint)
-	{
-		if (Index % PCGBlueprintConstants::RunawayResetFrequency == 0)
-		{
-			GInitRunaway(); // Reset periodically the iteration count, because we know we're in a fixed size loop.
-		}
-
-		const int32 OuterIndex = Index / InInnerPoints.Num();
-		const int32 InnerIndex = Index % InInnerPoints.Num();
-		return NestedLoopBody(InContext, InOuterData, InInnerData, InOuterPoints[OuterIndex], InInnerPoints[InnerIndex], OutPoint, OutData->Metadata, OuterIndex, InnerIndex);
-	});
-
-	std::swap(bPreviousBPStateValue, InContext.AsyncState.bIsCallingBlueprint);
-}
-
-void UPCGBlueprintElement::IterationLoop(FPCGContext& InContext, int64 NumIterations, UPCGPointData*& OutData, const UPCGSpatialData* InA, const UPCGSpatialData* InB, UPCGPointData* OptionalOutData) const
-{
-	if (NumIterations < 0)
-	{
-		PCGE_LOG_C(Error, GraphAndLog, &InContext, FText::Format(LOCTEXT("InvalidIterationCount", "Invalid number of iterations ({0})"), NumIterations));
-		return;
-	}
-
-	if (OptionalOutData)
-	{
-		OutData = OptionalOutData;
-	}
-	else
-	{
-		const UPCGSpatialData* Owner = (InA ? InA : InB);
-		OutData = FPCGContext::NewObject_AnyThread<UPCGPointData>(&InContext);
-
-		if (Owner)
-		{
-			OutData->InitializeFromData(Owner);
-		}
-	}
-
-	TArray<FPCGPoint>& OutPoints = OutData->GetMutablePoints();
-
-	bool bPreviousBPStateValue = true;
-	std::swap(bPreviousBPStateValue, InContext.AsyncState.bIsCallingBlueprint);
-
-	FPCGAsync::AsyncPointProcessing(&InContext, NumIterations, OutPoints, [this, &InContext, InA, InB, OutData](int32 Index, FPCGPoint& OutPoint)
-	{
-		if (Index % PCGBlueprintConstants::RunawayResetFrequency == 0)
-		{
-			GInitRunaway(); // Reset periodically the iteration count, because we know we're in a fixed size loop.
-		}
-
-		return IterationLoopBody(InContext, Index, InA, InB, OutPoint, OutData->Metadata);
-	});
-
-	std::swap(bPreviousBPStateValue, InContext.AsyncState.bIsCallingBlueprint);
-}
-
 void FPCGBlueprintExecutionContext::AddExtraStructReferencedObjects(FReferenceCollector& Collector)
 {
 	if (BlueprintElementInstance)
@@ -1117,7 +650,7 @@ FPCGContext* FPCGExecuteBlueprintElement::Initialize(const FPCGInitializeElement
 	const UPCGBlueprintSettings* Settings = Context->GetInputSettings<UPCGBlueprintSettings>();
 	if (Settings && Settings->BlueprintElementInstance)
 	{
-		Context->BlueprintElementInstance = CastChecked<UPCGBlueprintElement>(StaticDuplicateObject(Settings->BlueprintElementInstance, GetTransientPackage(), FName()));
+		Context->BlueprintElementInstance = CastChecked<UPCGBlueprintBaseElement>(StaticDuplicateObject(Settings->BlueprintElementInstance, GetTransientPackage(), FName()));
 
 #if !WITH_EDITOR
 		if (Context->ExecutionSource.IsValid())
@@ -1187,4 +720,17 @@ bool FPCGExecuteBlueprintElement::CanExecuteOnlyOnMainThread(FPCGContext* Contex
 	}
 }
 
+bool FPCGExecuteBlueprintElement::SupportsBasePointDataInputs(FPCGContext* Context) const
+{
+	if (!Context)
+	{
+		return true;
+	}
+
+	// Consider all BP Elements to support base point data unless they are child classes of the deprecated UPCGBlueprintElement
+	FPCGBlueprintExecutionContext* BPContext = static_cast<FPCGBlueprintExecutionContext*>(Context);
+	const UPCGBlueprintSettings* BPSettings = Context->GetInputSettings<UPCGBlueprintSettings>();
+	return !BPSettings->BlueprintElementType->IsChildOf<UPCGBlueprintElement>();
+}
+
 #undef LOCTEXT_NAMESPACE
diff --git a/Engine/Plugins/PCG/Source/PCG/Private/Elements/PCGIndirectionElement.cpp b/Engine/Plugins/PCG/Source/PCG/Private/Elements/PCGIndirectionElement.cpp
--- a/Engine/Plugins/PCG/Source/PCG/Private/Elements/PCGIndirectionElement.cpp
+++ b/Engine/Plugins/PCG/Source/PCG/Private/Elements/PCGIndirectionElement.cpp
@@ -104,7 +104,7 @@ TArray<FPCGPinProperties> UPCGIndirectionSettings::InputPinProperties() const
 		case EPCGProxyInterfaceMode::ByBlueprintElement:
 			if (IsValid(BlueprintElementClass))
 			{
-				const UPCGBlueprintElement* BlueprintSettingsDefaultObject = CastChecked<UPCGBlueprintElement>(BlueprintElementClass->GetDefaultObject());
+				const UPCGBlueprintBaseElement* BlueprintSettingsDefaultObject = CastChecked<UPCGBlueprintBaseElement>(BlueprintElementClass->GetDefaultObject());
 				InputProperties = BlueprintSettingsDefaultObject->GetInputPins();
 				bSetProperties = true;
 			}
@@ -154,7 +154,7 @@ TArray<FPCGPinProperties> UPCGIndirectionSettings::OutputPinProperties() const
 		case EPCGProxyInterfaceMode::ByBlueprintElement:
 			if (IsValid(BlueprintElementClass))
 			{
-				const UPCGBlueprintElement* BlueprintSettingsDefaultObject = CastChecked<UPCGBlueprintElement>(BlueprintElementClass->GetDefaultObject());
+				const UPCGBlueprintBaseElement* BlueprintSettingsDefaultObject = CastChecked<UPCGBlueprintBaseElement>(BlueprintElementClass->GetDefaultObject());
 				OutputProperties = BlueprintSettingsDefaultObject->GetOutputPins();
 			}
 			break;
diff --git a/Engine/Plugins/PCG/Source/PCG/Private/PCGParamData.cpp b/Engine/Plugins/PCG/Source/PCG/Private/PCGParamData.cpp
--- a/Engine/Plugins/PCG/Source/PCG/Private/PCGParamData.cpp
+++ b/Engine/Plugins/PCG/Source/PCG/Private/PCGParamData.cpp
@@ -87,7 +87,7 @@ int64 UPCGParamData::FindOrAddMetadataKey(const FName& InName)
 
 UPCGParamData* UPCGParamData::K2_FilterParamsByName(const FName& InName) const
 {
-	return FilterParamsByName(UPCGBlueprintElement::ResolveContext(), InName);
+	return FilterParamsByName(UPCGBlueprintBaseElement::ResolveContext(), InName);
 }
 
 UPCGParamData* UPCGParamData::FilterParamsByName(FPCGContext* Context, const FName& InName) const
@@ -106,7 +106,7 @@ UPCGParamData* UPCGParamData::FilterParamsByName(FPCGContext* Context, const FNa
 
 UPCGParamData* UPCGParamData::K2_FilterParamsByKey(int64 InKey) const
 {
-	return FilterParamsByKey(UPCGBlueprintElement::ResolveContext(), InKey);
+	return FilterParamsByKey(UPCGBlueprintBaseElement::ResolveContext(), InKey);
 }
 
 UPCGParamData* UPCGParamData::FilterParamsByKey(FPCGContext* Context, int64 InKey) const
diff --git a/Engine/Plugins/PCG/Source/PCG/Public/Elements/Blueprint/PCGBlueprintBaseElement.h b/Engine/Plugins/PCG/Source/PCG/Public/Elements/Blueprint/PCGBlueprintBaseElement.h
new file mode 100644
--- /dev/null
+++ b/Engine/Plugins/PCG/Source/PCG/Public/Elements/Blueprint/PCGBlueprintBaseElement.h
@@ -0,0 +1,200 @@
+// Copyright Epic Games, Inc. All Rights Reserved.
+
+#pragma once
+
+#include "PCGContext.h"
+#include "PCGSettings.h"
+
+#include "PCGBlueprintBaseElement.generated.h"
+
+class UWorld;
+class UPCGBlueprintElement;
+class UPCGBlueprintBaseElement;
+
+#if WITH_EDITOR
+DECLARE_MULTICAST_DELEGATE_OneParam(FOnPCGBlueprintElementChanged, UPCGBlueprintBaseElement*);
+
+namespace PCGBlueprintHelper
+{
+	TSet<TObjectPtr<UObject>> GetDataDependencies(UPCGBlueprintBaseElement* InElement);
+}
+#endif // WITH_EDITOR
+
+UCLASS(MinimalAPI, Abstract, BlueprintType, Blueprintable, hidecategories = (Object))
+class UPCGBlueprintBaseElement : public UObject
+{
+	GENERATED_BODY()
+
+public:
+	// ~Begin UObject interface
+	PCG_API virtual void PostLoad() override;
+	PCG_API virtual void BeginDestroy() override;
+	// ~End UObject interface
+
+	/**
+	 * Main execution function that will contain the logic for this PCG Element. Use GetContext to have access to the context.
+	 * @param Input  - Input collection containing all the data passed as input to the node.
+	 * @param Output - Data collection that will be passed as the output of the node, with pins matching the ones provided during the execution.
+	 */
+	UFUNCTION(BlueprintNativeEvent, BlueprintCallable, Category = "PCG|Execution")
+	PCG_API void Execute(const FPCGDataCollection& Input, FPCGDataCollection& Output);
+
+	/** Override for the default node name */
+	UFUNCTION(BlueprintNativeEvent, Category = "PCG|Node Customization")
+	PCG_API FName NodeTitleOverride() const;
+
+	/** Override for the default node color. */
+	UFUNCTION(BlueprintNativeEvent, Category = "PCG|Node Customization")
+	PCG_API FLinearColor NodeColorOverride() const;
+
+	/** Override to change the node type. */
+	UFUNCTION(BlueprintNativeEvent, Category = "PCG|Node Customization")
+	PCG_API EPCGSettingsType NodeTypeOverride() const;
+
+	/** If Dynamic Pins is enabled in the BP settings, override this function to provide the type for the given pin. You can use "GetTypeUnionOfIncidentEdges" from the settings to get the union of input types on a given pin. Use the bitwise OR to combine multiple types together.*/
+	UFUNCTION(BlueprintNativeEvent, Category = "PCG|Node Customization")
+	PCG_API int32 DynamicPinTypesOverride(const UPCGSettings* InSettings, const UPCGPin* InPin) const;
+
+	/** Override for the IsCacheable node property when it depends on the settings in your node. If true, the node will be cached, if not it will always be executed. */
+	UFUNCTION(BlueprintNativeEvent, Category = "PCG|Execution")
+	PCG_API bool IsCacheableOverride() const;
+
+	/** Apply the preconfigured settings specified in the class default. Used to create nodes that are configured with pre-defined settings. Use InPreconfigureInfo index to know which settings it is. */
+	UFUNCTION(BlueprintImplementableEvent, Category = "PCG|Preconfigure Settings", meta = (ForceAsFunction))
+	PCG_API void ApplyPreconfiguredSettings(UPARAM(ref) const FPCGPreConfiguredSettingsInfo& InPreconfigureInfo);
+
+	// Returns the labels of custom input pins only
+	UFUNCTION(BlueprintCallable, Category = "PCG|Input & Output")
+	PCG_API TSet<FName> CustomInputLabels() const;
+
+	// Returns the labels of custom output pins only
+	UFUNCTION(BlueprintCallable, Category = "PCG|Input & Output")
+	PCG_API TSet<FName> CustomOutputLabels() const;
+
+	UFUNCTION(BlueprintCallable, Category = "PCG|Input & Output", meta = (HideSelfPin = "true"))
+	PCG_API TArray<FPCGPinProperties> GetInputPins() const;
+
+	UFUNCTION(BlueprintCallable, Category = "PCG|Input & Output", meta = (HideSelfPin = "true"))
+	PCG_API TArray<FPCGPinProperties> GetOutputPins() const;
+
+	/** Returns true if there is an input pin with the matching label. If found, will copy the pin properties in OutFoundPin */
+	UFUNCTION(BlueprintCallable, Category = "PCG|Input & Output", meta = (HideSelfPin = "true"))
+	PCG_API bool GetInputPinByLabel(FName InPinLabel, FPCGPinProperties& OutFoundPin) const;
+
+	/** Returns true if there is an output pin with the matching label. If found, will copy the pin properties in OutFoundPin */
+	UFUNCTION(BlueprintCallable, Category = "PCG|Input & Output", meta = (HideSelfPin = "true"))
+	PCG_API bool GetOutputPinByLabel(FName InPinLabel, FPCGPinProperties& OutFoundPin) const;
+
+	/** Gets the seed from the associated settings & source component */
+	UFUNCTION(BlueprintCallable, Category = "PCG|Random")
+	PCG_API int GetSeed(UPARAM(ref) FPCGContext& InContext) const;
+
+	/** Creates a random stream from the settings & source component */
+	UFUNCTION(BlueprintCallable, BlueprintPure = false, Category = "PCG|Random")
+	PCG_API FRandomStream GetRandomStream(UPARAM(ref) FPCGContext& InContext) const;
+
+	/** Called after object creation to setup the object callbacks */
+	PCG_API void Initialize();
+
+	/** Retrieves the execution context - note that this will not be valid outside of the Execute functions */
+	UFUNCTION(BlueprintCallable, Category = "PCG|Advanced", meta = (HideSelfPin = "true"))
+	PCG_API FPCGContext& GetContext() const;
+
+	/** Called after the element duplication during execution to be able to get the context easily - internal call only */
+	PCG_API void SetCurrentContext(FPCGContext* InCurrentContext);
+
+	/** Tries to resolve current Context from thread local BP stack */
+	static PCG_API FPCGContext* ResolveContext();
+
+#if WITH_EDITOR
+	// ~Begin UObject interface
+	PCG_API virtual void PostEditChangeProperty(FPropertyChangedEvent& PropertyChangedEvent) override;
+	// ~End UObject interface
+
+	/** Used for filtering */
+	static PCG_API FString GetParentClassName();
+#endif
+
+	/** Needed to be able to call certain blueprint functions */
+	PCG_API virtual UWorld* GetWorld() const override;
+
+#if !WITH_EDITOR
+	void SetInstanceWorld(UWorld* World) { InstanceWorld = World; }
+#endif
+
+#if WITH_EDITOR
+	FOnPCGBlueprintElementChanged OnBlueprintElementChangedDelegate;
+#endif
+
+	/** Controls whether results can be cached so we can bypass execution if the inputs & settings are the same in a subsequent execution.
+	* If you have implemented the IsCacheableOverride function, then this value is ignored.
+	* Note that if your node relies on data that is not directly tracked by PCG or creates any kind of artifact (adds components, creates actors, etc.) then it should not be cacheable.
+	*/
+	UPROPERTY(BlueprintReadWrite, EditDefaultsOnly, AdvancedDisplay, Category = Settings)
+	bool bIsCacheable = false;
+
+	/** In cases where your node is non-cacheable but is likely to yield the same results on subsequent executions, this controls whether we will do a deep & computationally intensive CRC computation (true),
+	* which will allow cache usage in downstream nodes in your graph, or, by default (false), a shallow but quick crc computation which will not be cache-friendly. */
+	UPROPERTY(BlueprintReadWrite, EditDefaultsOnly, AdvancedDisplay, Category = Settings)
+	bool bComputeFullDataCrc = false;
+
+	/** Controls whether this node execution can be run from a non-game thread. This is not related to the Loop functions provided/implemented in this class, which should always run on any thread. */
+	UPROPERTY(BlueprintReadWrite, EditDefaultsOnly, Category = Settings)
+	bool bRequiresGameThread = true;
+
+	UPROPERTY(BlueprintReadWrite, EditDefaultsOnly, Category = "Settings|Input & Output")
+	TArray<FPCGPinProperties> CustomInputPins;
+
+	UPROPERTY(BlueprintReadWrite, EditDefaultsOnly, Category = "Settings|Input & Output")
+	TArray<FPCGPinProperties> CustomOutputPins;
+
+	UPROPERTY(BlueprintReadWrite, EditDefaultsOnly, Category = "Settings|Input & Output")
+	bool bHasDefaultInPin = true;
+
+	UPROPERTY(BlueprintReadWrite, EditDefaultsOnly, Category = "Settings|Input & Output")
+	bool bHasDefaultOutPin = true;
+
+	/** If enabled, by default, the Out pin type will have the union of In pin types. Default only works if the pins are In and Out. For custom behavior, implement DynamicPinTypesOverride. */
+	UPROPERTY(BlueprintReadWrite, EditDefaultsOnly, Category = "Settings|Input & Output")
+	bool bHasDynamicPins = false;
+
+#if WITH_EDITORONLY_DATA
+	UPROPERTY(BlueprintReadWrite, EditDefaultsOnly, Category = AssetInfo, AssetRegistrySearchable)
+	bool bExposeToLibrary = false;
+
+	UPROPERTY(BlueprintReadWrite, EditDefaultsOnly, Category = "AssetInfo|Preconfigured Settings", AssetRegistrySearchable)
+	bool bEnablePreconfiguredSettings = false;
+
+	UPROPERTY(BlueprintReadWrite, EditDefaultsOnly, Category = "AssetInfo|Preconfigured Settings", AssetRegistrySearchable, meta = (EditCondition = bEnablePreconfiguredSettings, EditConditionHides))
+	bool bOnlyExposePreconfiguredSettings = false;
+
+	UPROPERTY(BlueprintReadWrite, EditDefaultsOnly, Category = "AssetInfo|Preconfigured Settings", meta = (EditCondition = bEnablePreconfiguredSettings, EditConditionHides))
+	TArray<FPCGPreConfiguredSettingsInfo> PreconfiguredInfo;
+
+	UPROPERTY(BlueprintReadWrite, EditDefaultsOnly, Category = AssetInfo, AssetRegistrySearchable)
+	FText Category;
+
+	UPROPERTY(BlueprintReadWrite, EditDefaultsOnly, Category = AssetInfo, AssetRegistrySearchable)
+	FText Description;
+
+	UPROPERTY(BlueprintReadWrite, EditDefaultsOnly, AdvancedDisplay, Category = "Settings")
+	int32 DependencyParsingDepth = 1;
+#endif
+
+protected:
+#if WITH_EDITOR
+	PCG_API void OnDependencyChanged(UObject* Object, FPropertyChangedEvent& PropertyChangedEvent);
+	PCG_API void UpdateDependencies();
+	TSet<TWeakObjectPtr<UObject>> DataDependencies;
+#endif
+
+#if !WITH_EDITORONLY_DATA
+	UWorld* InstanceWorld = nullptr;
+#endif
+
+	// Since we duplicate the blueprint elements prior to execution, they will be unique
+	// and have a 1:1 match with their context, which allows us to store it here
+	FPCGContext* CurrentContext = nullptr;
+};
+
+
diff --git a/Engine/Plugins/PCG/Source/PCG/Public/Elements/Blueprint/PCGBlueprintDeprecatedElement.h b/Engine/Plugins/PCG/Source/PCG/Public/Elements/Blueprint/PCGBlueprintDeprecatedElement.h
new file mode 100644
--- /dev/null
+++ b/Engine/Plugins/PCG/Source/PCG/Public/Elements/Blueprint/PCGBlueprintDeprecatedElement.h
@@ -0,0 +1,120 @@
+// Copyright Epic Games, Inc. All Rights Reserved.
+
+#pragma once
+
+#include "PCGBlueprintBaseElement.h"
+#include "PCGContext.h"
+#include "PCGSettings.h"
+#include "Delegates/Delegate.h"
+
+#include "PCGBlueprintDeprecatedElement.generated.h"
+
+class UPCGBlueprintElement;
+class UPCGMetadata;
+class UPCGPointData;
+class UPCGSpatialData;
+struct FPCGPoint;
+
+class UWorld;
+
+#if WITH_EDITOR
+
+// Deprecated
+DECLARE_MULTICAST_DELEGATE_OneParam(FOnPCGBlueprintChanged, UPCGBlueprintElement*);
+
+#endif // WITH_EDITOR
+
+UCLASS(MinimalAPI, Abstract, BlueprintType, Blueprintable, hidecategories = (Object))
+class UPCGBlueprintElement : public UPCGBlueprintBaseElement
+{
+	GENERATED_BODY()
+
+public:
+
+#if WITH_EDITOR
+	UE_DEPRECATED(5.7, "Use OnBlueprintElementChangedDelegate instead")
+	FOnPCGBlueprintChanged OnBlueprintChangedDelegate;
+#endif
+
+	/**
+	 * Main execution function that will contain the logic for this PCG Element, with the context as parameter.
+	 * @param InContext - Context of the execution
+	 * @param Input     - Input collection containing all the data passed as input to the node.
+	 * @param Output    - Data collection that will be passed as the output of the node, with pins matching the ones provided during the execution.
+	 */
+	UFUNCTION(BlueprintNativeEvent, Category = "PCG|Execution")
+	PCG_API void ExecuteWithContext(UPARAM(ref)FPCGContext& InContext, const FPCGDataCollection& Input, FPCGDataCollection& Output);
+
+	/**
+	 * Multi-threaded loop that will iterate on all points in InData. All points will be added in the same order than in input. Will be called by Point Loop function.
+	 * @param InContext   - Context of the execution
+	 * @param InData      - Input point data. Constant, must not be modified.
+	 * @param InPoint     - Point for the current iteration. Constant, must not be modified.
+	 * @param OutPoint    - Point that will be added to the output data. Can be modified.
+	 * @param OutMetadata - Output metadata to write attribute to. Can be modified.
+	 * @param Iteration   - Index of the current point. Must only be used to access input data, as this call is multi-threaded. It is not safe to access output data.
+	 * @returns True if the point should be kept, False if not.
+	 */
+	UFUNCTION(BlueprintImplementableEvent, Category = "PCG|Flow Control")
+	PCG_API bool PointLoopBody(const FPCGContext& InContext, const UPCGPointData* InData, const FPCGPoint& InPoint, FPCGPoint& OutPoint, UPCGMetadata* OutMetadata, int64 Iteration) const;
+
+	/**
+	 * Multi-threaded loop that will be called on all points in InData. Can return a variable number of output points.
+	 * All points will be added in the same order than in input. Will be called by Variable Loop function.
+	 * @param InContext   - Context of the execution
+	 * @param InData      - Input point data. Constant, must not be modified.
+	 * @param InPoint     - Point for the current iteration. Constant, must not be modified.
+	 * @param OutMetadata - Output metadata to write attribute to. Can be modified.
+	 * @param Iteration   - Index of the current point. Must only be used to access input data, as this call is multi-threaded. It is not safe to access output data.
+	 * @returns Array of new points that will be added to the output point data.
+	 */
+	UFUNCTION(BlueprintImplementableEvent, Category = "PCG|Flow Control")
+	PCG_API TArray<FPCGPoint> VariableLoopBody(const FPCGContext& InContext, const UPCGPointData* InData, const FPCGPoint& InPoint, UPCGMetadata* OutMetadata, int64 Iteration) const;
+
+	/**
+	 * Multi-threaded loop that will iterate on all nested loop pairs (e.g. (o, i) for all o in Outer, i in Inner).
+	 * All points will be added in the same order than in input (e.g: (0,0), (0,1), (0,2), ...). Will be called by Nested Loop function.
+	 * @param InContext      - Context of the execution
+	 * @param InOuterData    - Outer point data. Constant, must not be modified.
+	 * @param InInnerData    - Inner point data. Constant, must not be modified.
+	 * @param InOuterPoint   - Outer Point for the current iteration. Constant, must not be modified.
+	 * @param InInnerPoint   - Inner Point for the current iteration. Constant, must not be modified.
+	 * @param OutPoint       - Point that will be added to the output data. Can be modified.
+	 * @param OutMetadata    - Output metadata to write attribute to. Can be modified.
+	 * @param OuterIteration - Index of the current point in outer data. Must only be used to access input data, as this call is multi-threaded. It is not safe to access output data.
+	 * @param InnerIteration - Index of the current point in inner data. Must only be used to access input data, as this call is multi-threaded. It is not safe to access output data.
+	 * @returns True if the point should be kept, False if not.
+	 */
+	UFUNCTION(BlueprintImplementableEvent, Category = "PCG|Flow Control")
+	PCG_API bool NestedLoopBody(const FPCGContext& InContext, const UPCGPointData* InOuterData, const UPCGPointData* InInnerData, const FPCGPoint& InOuterPoint, const FPCGPoint& InInnerPoint, FPCGPoint& OutPoint, UPCGMetadata* OutMetadata, int64 OuterIteration, int64 InnerIteration) const;
+
+	/**
+	 * Multi-threaded loop that will be called N number of times (defined by Iteration Loop parameter NumIterations).
+	 * All points will be added in order (iteration 0 will be before iteration 1 in the final array).
+	 * @param InContext   - Context of the execution
+	 * @param Iteration   - Index of the current iteration. Must only be used to access input data, as this call is multi-threaded. It is not safe to access output data.
+	 * @param InA         - Optional input spatial data, can be null. Constant, must not be modified.
+	 * @param InB         - Optional input spatial data, can be null. Constant, must not be modified.
+	 * @param OutPoint    - Point that will be added to the output data. Can be modified.
+	 * @param OutMetadata - Output metadata to write attribute to. Can be modified.
+	 * @returns True if the point should be kept, False if not.
+	 */
+	UFUNCTION(BlueprintImplementableEvent, Category = "PCG|Flow Control")
+	PCG_API bool IterationLoopBody(const FPCGContext& InContext, int64 Iteration, const UPCGSpatialData* InA, const UPCGSpatialData* InB, FPCGPoint& OutPoint, UPCGMetadata* OutMetadata) const;
+
+	/** Calls the PointLoopBody function on all points */
+	UFUNCTION(BlueprintCallable, BlueprintPure = false, Category = "PCG|Flow Control", meta = (HideSelfPin = "true"))
+	PCG_API void PointLoop(UPARAM(ref) FPCGContext& InContext, const UPCGPointData* InData, UPCGPointData*& OutData, UPCGPointData* OptionalOutData = nullptr) const;
+
+	/** Calls the VariableLoopBody function on all points, each call can return a variable number of points */
+	UFUNCTION(BlueprintCallable, BlueprintPure = false, Category = "PCG|Flow Control", meta = (HideSelfPin = "true"))
+	PCG_API void VariableLoop(UPARAM(ref) FPCGContext& InContext, const UPCGPointData* InData, UPCGPointData*& OutData, UPCGPointData* OptionalOutData = nullptr) const;
+
+	/** Calls the NestedLoopBody function on all nested loop pairs (e.g. (o, i) for all o in Outer, i in Inner) */
+	UFUNCTION(BlueprintCallable, BlueprintPure = false, Category = "PCG|FLow Control", meta = (HideSelfPin = "true"))
+	PCG_API void NestedLoop(UPARAM(ref) FPCGContext& InContext, const UPCGPointData* InOuterData, const UPCGPointData* InInnerData, UPCGPointData*& OutData, UPCGPointData* OptionalOutData = nullptr) const;
+
+	/** Calls the IterationLoopBody a fixed number of times, optional parameters are used to potentially initialized the Out Data, but otherwise are used to remove the need to have variables */
+	UFUNCTION(BlueprintCallable, BlueprintPure = false, Category = "PCG|Flow Control", meta = (HideSelfPin = "true"))
+	PCG_API void IterationLoop(UPARAM(ref) FPCGContext& InContext, int64 NumIterations, UPCGPointData*& OutData, const UPCGSpatialData* OptionalA = nullptr, const UPCGSpatialData* OptionalB = nullptr, UPCGPointData* OptionalOutData = nullptr) const;
+};
diff --git a/Engine/Plugins/PCG/Source/PCG/Public/Elements/PCGExecuteBlueprint.h b/Engine/Plugins/PCG/Source/PCG/Public/Elements/PCGExecuteBlueprint.h
--- a/Engine/Plugins/PCG/Source/PCG/Public/Elements/PCGExecuteBlueprint.h
+++ b/Engine/Plugins/PCG/Source/PCG/Public/Elements/PCGExecuteBlueprint.h
@@ -5,39 +5,21 @@
 #include "PCGContext.h"
 #include "PCGSettings.h"
 
-
+#include "Blueprint/PCGBlueprintBaseElement.h"
+#include "Blueprint/PCGBlueprintDeprecatedElement.h"
 
 #include "PCGExecuteBlueprint.generated.h"
 
-#define UE_API PCG_API
-
-class UPCGBlueprintElement;
-class UPCGMetadata;
-class UPCGPointData;
-class UPCGSpatialData;
-struct FPCGPoint;
-
-class UWorld;
-
-#if WITH_EDITOR
-DECLARE_MULTICAST_DELEGATE_OneParam(FOnPCGBlueprintChanged, UPCGBlueprintElement*);
-
-namespace PCGBlueprintHelper
-{
-	TSet<TObjectPtr<UObject>> GetDataDependencies(UPCGBlueprintElement* InElement);
-}
-#endif // WITH_EDITOR
-
 /**
  * Helper class which can be used before calling a Blueprint function with a FPCGContext& parameter.
- * 
+ *
  * Since BP calls will copy the incoming FPCGContext reference we need to clear out a couple of fields before the BP copy is made to prevent issues when the copy is destroyed.
  */
 class FPCGContextBlueprintScope
 {
 public:
-	UE_API explicit FPCGContextBlueprintScope(FPCGContext* InContext);
-	UE_API ~FPCGContextBlueprintScope();
+	PCG_API explicit FPCGContextBlueprintScope(FPCGContext* InContext);
+	PCG_API ~FPCGContextBlueprintScope();
 
 private:
 	FPCGContext* Context = nullptr;
@@ -45,276 +27,6 @@ private:
 	TWeakPtr<FPCGGraphExecutor> GraphExecutor = nullptr;
 };
 
-UCLASS(MinimalAPI, Abstract, BlueprintType, Blueprintable, hidecategories = (Object))
-class UPCGBlueprintElement : public UObject
-{
-	GENERATED_BODY()
-
-public:
-	// ~Begin UObject interface
-	UE_API virtual void PostLoad() override;
-	UE_API virtual void BeginDestroy() override;
-	// ~End UObject interface
-
-	/**
-	 * Main execution function that will contain the logic for this PCG Element, with the context as parameter.
-	 * @param InContext - Context of the execution
-	 * @param Input     - Input collection containing all the data passed as input to the node.
-	 * @param Output    - Data collection that will be passed as the output of the node, with pins matching the ones provided during the execution.
-	 */
-	UFUNCTION(BlueprintNativeEvent, Category = "PCG|Execution")
-	UE_API void ExecuteWithContext(UPARAM(ref)FPCGContext& InContext, const FPCGDataCollection& Input, FPCGDataCollection& Output);
-
-	/**
-	 * Main execution function that will contain the logic for this PCG Element. Use GetContext to have access to the context.
-	 * @param Input  - Input collection containing all the data passed as input to the node.
-	 * @param Output - Data collection that will be passed as the output of the node, with pins matching the ones provided during the execution.
-	 */
-	UFUNCTION(BlueprintImplementableEvent, BlueprintCallable, Category = "PCG|Execution")
-	UE_API void Execute(const FPCGDataCollection& Input, FPCGDataCollection& Output);
-
-	/**
-	 * Multi-threaded loop that will iterate on all points in InData. All points will be added in the same order than in input. Will be called by Point Loop function.
-	 * @param InContext   - Context of the execution
-	 * @param InData      - Input point data. Constant, must not be modified.
-	 * @param InPoint     - Point for the current iteration. Constant, must not be modified.
-	 * @param OutPoint    - Point that will be added to the output data. Can be modified.
-	 * @param OutMetadata - Output metadata to write attribute to. Can be modified.
-	 * @param Iteration   - Index of the current point. Must only be used to access input data, as this call is multi-threaded. It is not safe to access output data.
-	 * @returns True if the point should be kept, False if not.
-	 */
-	UFUNCTION(BlueprintImplementableEvent, Category = "PCG|Flow Control")
-	UE_API bool PointLoopBody(const FPCGContext& InContext, const UPCGPointData* InData, const FPCGPoint& InPoint, FPCGPoint& OutPoint, UPCGMetadata* OutMetadata, int64 Iteration) const;
-
-	/**
-	 * Multi-threaded loop that will be called on all points in InData. Can return a variable number of output points. 
-	 * All points will be added in the same order than in input. Will be called by Variable Loop function. 
-	 * @param InContext   - Context of the execution
-	 * @param InData      - Input point data. Constant, must not be modified.
-	 * @param InPoint     - Point for the current iteration. Constant, must not be modified.
-	 * @param OutMetadata - Output metadata to write attribute to. Can be modified.
-	 * @param Iteration   - Index of the current point. Must only be used to access input data, as this call is multi-threaded. It is not safe to access output data.
-	 * @returns Array of new points that will be added to the output point data.
-	 */
-	UFUNCTION(BlueprintImplementableEvent, Category = "PCG|Flow Control")
-	UE_API TArray<FPCGPoint> VariableLoopBody(const FPCGContext& InContext, const UPCGPointData* InData, const FPCGPoint& InPoint, UPCGMetadata* OutMetadata, int64 Iteration) const;
-
-	/**
-	 * Multi-threaded loop that will iterate on all nested loop pairs (e.g. (o, i) for all o in Outer, i in Inner).
-	 * All points will be added in the same order than in input (e.g: (0,0), (0,1), (0,2), ...). Will be called by Nested Loop function.
-	 * @param InContext      - Context of the execution
-	 * @param InOuterData    - Outer point data. Constant, must not be modified.
-	 * @param InInnerData    - Inner point data. Constant, must not be modified.
-	 * @param InOuterPoint   - Outer Point for the current iteration. Constant, must not be modified.
-	 * @param InInnerPoint   - Inner Point for the current iteration. Constant, must not be modified.
-	 * @param OutPoint       - Point that will be added to the output data. Can be modified.
-	 * @param OutMetadata    - Output metadata to write attribute to. Can be modified.
-	 * @param OuterIteration - Index of the current point in outer data. Must only be used to access input data, as this call is multi-threaded. It is not safe to access output data.
-	 * @param InnerIteration - Index of the current point in inner data. Must only be used to access input data, as this call is multi-threaded. It is not safe to access output data.
-	 * @returns True if the point should be kept, False if not.
-	 */
-	UFUNCTION(BlueprintImplementableEvent, Category = "PCG|Flow Control")
-	UE_API bool NestedLoopBody(const FPCGContext& InContext, const UPCGPointData* InOuterData, const UPCGPointData* InInnerData, const FPCGPoint& InOuterPoint, const FPCGPoint& InInnerPoint, FPCGPoint& OutPoint, UPCGMetadata* OutMetadata, int64 OuterIteration, int64 InnerIteration) const;
-
-	/**
-	 * Multi-threaded loop that will be called N number of times (defined by Iteration Loop parameter NumIterations). 
-	 * All points will be added in order (iteration 0 will be before iteration 1 in the final array).
-	 * @param InContext   - Context of the execution
-	 * @param Iteration   - Index of the current iteration. Must only be used to access input data, as this call is multi-threaded. It is not safe to access output data.
-	 * @param InA         - Optional input spatial data, can be null. Constant, must not be modified.
-	 * @param InB         - Optional input spatial data, can be null. Constant, must not be modified.
-	 * @param OutPoint    - Point that will be added to the output data. Can be modified.
-	 * @param OutMetadata - Output metadata to write attribute to. Can be modified.
-	 * @returns True if the point should be kept, False if not.
-	 */
-	UFUNCTION(BlueprintImplementableEvent, Category = "PCG|Flow Control")
-	UE_API bool IterationLoopBody(const FPCGContext& InContext, int64 Iteration, const UPCGSpatialData* InA, const UPCGSpatialData* InB, FPCGPoint& OutPoint, UPCGMetadata* OutMetadata) const;
-
-	/** Calls the PointLoopBody function on all points */
-	UFUNCTION(BlueprintCallable, BlueprintPure = false, Category = "PCG|Flow Control", meta = (HideSelfPin = "true"))
-	UE_API void PointLoop(UPARAM(ref) FPCGContext& InContext, const UPCGPointData* InData, UPCGPointData*& OutData, UPCGPointData* OptionalOutData = nullptr) const;
-
-	/** Calls the VariableLoopBody function on all points, each call can return a variable number of points */
-	UFUNCTION(BlueprintCallable, BlueprintPure = false, Category = "PCG|Flow Control", meta = (HideSelfPin = "true"))
-	UE_API void VariableLoop(UPARAM(ref) FPCGContext& InContext, const UPCGPointData* InData, UPCGPointData*& OutData, UPCGPointData* OptionalOutData = nullptr) const;
-
-	/** Calls the NestedLoopBody function on all nested loop pairs (e.g. (o, i) for all o in Outer, i in Inner) */
-	UFUNCTION(BlueprintCallable, BlueprintPure = false, Category = "PCG|FLow Control", meta = (HideSelfPin = "true"))
-	UE_API void NestedLoop(UPARAM(ref) FPCGContext& InContext, const UPCGPointData* InOuterData, const UPCGPointData* InInnerData, UPCGPointData*& OutData, UPCGPointData* OptionalOutData = nullptr) const;
-	
-	/** Calls the IterationLoopBody a fixed number of times, optional parameters are used to potentially initialized the Out Data, but otherwise are used to remove the need to have variables */
-	UFUNCTION(BlueprintCallable, BlueprintPure = false, Category = "PCG|Flow Control", meta = (HideSelfPin = "true"))
-	UE_API void IterationLoop(UPARAM(ref) FPCGContext& InContext, int64 NumIterations, UPCGPointData*& OutData, const UPCGSpatialData* OptionalA = nullptr, const UPCGSpatialData* OptionalB = nullptr, UPCGPointData* OptionalOutData = nullptr) const;
-
-	/** Override for the default node name */
-	UFUNCTION(BlueprintNativeEvent, Category = "PCG|Node Customization")
-	UE_API FName NodeTitleOverride() const;
-
-	/** Override for the default node color. */
-	UFUNCTION(BlueprintNativeEvent, Category = "PCG|Node Customization")
-	UE_API FLinearColor NodeColorOverride() const;
-
-	/** Override to change the node type. */
-	UFUNCTION(BlueprintNativeEvent, Category = "PCG|Node Customization")
-	UE_API EPCGSettingsType NodeTypeOverride() const;
-
-	/** If Dynamic Pins is enabled in the BP settings, override this function to provide the type for the given pin. You can use "GetTypeUnionOfIncidentEdges" from the settings to get the union of input types on a given pin. Use the bitwise OR to combine multiple types together.*/
-	UFUNCTION(BlueprintNativeEvent, Category = "PCG|Node Customization")
-	UE_API int32 DynamicPinTypesOverride(const UPCGSettings* InSettings, const UPCGPin* InPin) const;
-
-	/** Override for the IsCacheable node property when it depends on the settings in your node. If true, the node will be cached, if not it will always be executed. */
-	UFUNCTION(BlueprintNativeEvent, Category = "PCG|Execution")
-	UE_API bool IsCacheableOverride() const;
-
-	/** Apply the preconfigured settings specified in the class default. Used to create nodes that are configured with pre-defined settings. Use InPreconfigureInfo index to know which settings it is. */
-	UFUNCTION(BlueprintImplementableEvent, Category = "PCG|Preconfigure Settings", meta = (ForceAsFunction))
-	UE_API void ApplyPreconfiguredSettings(UPARAM(ref) const FPCGPreConfiguredSettingsInfo& InPreconfigureInfo);
-
-	// Returns the labels of custom input pins only
-	UFUNCTION(BlueprintCallable, Category = "PCG|Input & Output")
-	UE_API TSet<FName> CustomInputLabels() const;
-
-	// Returns the labels of custom output pins only
-	UFUNCTION(BlueprintCallable, Category = "PCG|Input & Output")
-	UE_API TSet<FName> CustomOutputLabels() const;
-
-	UFUNCTION(BlueprintCallable, Category = "PCG|Input & Output", meta = (HideSelfPin = "true"))
-	UE_API TArray<FPCGPinProperties> GetInputPins() const;
-
-	UFUNCTION(BlueprintCallable, Category = "PCG|Input & Output", meta = (HideSelfPin = "true"))
-	UE_API TArray<FPCGPinProperties> GetOutputPins() const;
-
-	/** Returns true if there is an input pin with the matching label. If found, will copy the pin properties in OutFoundPin */
-	UFUNCTION(BlueprintCallable, Category = "PCG|Input & Output", meta = (HideSelfPin = "true"))
-	UE_API bool GetInputPinByLabel(FName InPinLabel, FPCGPinProperties& OutFoundPin) const;
-
-	/** Returns true if there is an output pin with the matching label. If found, will copy the pin properties in OutFoundPin */
-	UFUNCTION(BlueprintCallable, Category = "PCG|Input & Output", meta = (HideSelfPin = "true"))
-	UE_API bool GetOutputPinByLabel(FName InPinLabel, FPCGPinProperties& OutFoundPin) const;
-
-	/** Gets the seed from the associated settings & source component */
-	UFUNCTION(BlueprintCallable, Category = "PCG|Random")
-	UE_API int GetSeed(UPARAM(ref) FPCGContext& InContext) const;
-
-	/** Creates a random stream from the settings & source component */
-	UFUNCTION(BlueprintCallable, BlueprintPure = false, Category = "PCG|Random")
-	UE_API FRandomStream GetRandomStream(UPARAM(ref) FPCGContext& InContext) const;
-
-	/** Called after object creation to setup the object callbacks */
-	UE_API void Initialize();
-
-	/** Retrieves the execution context - note that this will not be valid outside of the Execute functions */
-	UFUNCTION(BlueprintCallable, Category = "PCG|Advanced", meta = (HideSelfPin = "true"))
-	UE_API FPCGContext& GetContext() const;
-
-	/** Called after the element duplication during execution to be able to get the context easily - internal call only */
-	UE_API void SetCurrentContext(FPCGContext* InCurrentContext);
-
-	/** Tries to resolve current Context from thread local BP stack */
-	static UE_API FPCGContext* ResolveContext();
-
-#if WITH_EDITOR
-	// ~Begin UObject interface
-	UE_API virtual void PostEditChangeProperty(FPropertyChangedEvent& PropertyChangedEvent) override;
-	// ~End UObject interface
-
-	/** Used for filtering */
-	static UE_API FString GetParentClassName();
-#endif
-
-	/** Needed to be able to call certain blueprint functions */
-	UE_API virtual UWorld* GetWorld() const override;
-
-#if !WITH_EDITOR
-	void SetInstanceWorld(UWorld* World) { InstanceWorld = World; }
-#endif
-
-#if WITH_EDITOR
-	FOnPCGBlueprintChanged OnBlueprintChangedDelegate;
-#endif
-
-	/** Controls whether results can be cached so we can bypass execution if the inputs & settings are the same in a subsequent execution.
-	* If you have implemented the IsCacheableOverride function, then this value is ignored.
-	* Note that if your node relies on data that is not directly tracked by PCG or creates any kind of artifact (adds components, creates actors, etc.) then it should not be cacheable.
-	*/
-	UPROPERTY(BlueprintReadWrite, EditDefaultsOnly, AdvancedDisplay, Category = Settings)
-	bool bIsCacheable = false;
-
-	/** In cases where your node is non-cacheable but is likely to yield the same results on subsequent executions, this controls whether we will do a deep & computationally intensive CRC computation (true), 
-	* which will allow cache usage in downstream nodes in your graph, or, by default (false), a shallow but quick crc computation which will not be cache-friendly. */
-	UPROPERTY(BlueprintReadWrite, EditDefaultsOnly, AdvancedDisplay, Category = Settings)
-	bool bComputeFullDataCrc = false;
-
-	/** Controls whether this node execution can be run from a non-game thread. This is not related to the Loop functions provided/implemented in this class, which should always run on any thread. */
-	UPROPERTY(BlueprintReadWrite, EditDefaultsOnly, Category = Settings)
-	bool bRequiresGameThread = true;
-
-#if WITH_EDITORONLY_DATA
-	/** This property is deprecated, please use 'Requires Game Thread' instead. */
-	UPROPERTY()
-	bool bCanBeMultithreaded_DEPRECATED = false;
-
-	UPROPERTY(BlueprintGetter=CustomInputLabels, Category = "Settings|Input & Output", meta = (DeprecatedProperty, DeprecatedMessage = "Input Pin Labels are deprecated - use Input Labels instead."))
-	TSet<FName> InputPinLabels_DEPRECATED;
-
-	UPROPERTY(BlueprintGetter=CustomOutputLabels, Category = "Settings|Input & Output")
-	TSet<FName> OutputPinLabels_DEPRECATED;
-#endif
-
-	UPROPERTY(BlueprintReadWrite, EditDefaultsOnly, Category = "Settings|Input & Output")
-	TArray<FPCGPinProperties> CustomInputPins;
-
-	UPROPERTY(BlueprintReadWrite, EditDefaultsOnly, Category = "Settings|Input & Output")
-	TArray<FPCGPinProperties> CustomOutputPins;
-
-	UPROPERTY(BlueprintReadWrite, EditDefaultsOnly, Category = "Settings|Input & Output")
-	bool bHasDefaultInPin = true;
-
-	UPROPERTY(BlueprintReadWrite, EditDefaultsOnly, Category = "Settings|Input & Output")
-	bool bHasDefaultOutPin = true;
-
-	/** If enabled, by default, the Out pin type will have the union of In pin types. Default only works if the pins are In and Out. For custom behavior, implement DynamicPinTypesOverride. */
-	UPROPERTY(BlueprintReadWrite, EditDefaultsOnly, Category = "Settings|Input & Output")
-	bool bHasDynamicPins = false;
-
-#if WITH_EDITORONLY_DATA
-	UPROPERTY(BlueprintReadWrite, EditDefaultsOnly, Category = AssetInfo, AssetRegistrySearchable)
-	bool bExposeToLibrary = false;
-
-	UPROPERTY(BlueprintReadWrite, EditDefaultsOnly, Category = "AssetInfo|Preconfigured Settings", AssetRegistrySearchable)
-	bool bEnablePreconfiguredSettings = false;
-
-	UPROPERTY(BlueprintReadWrite, EditDefaultsOnly, Category = "AssetInfo|Preconfigured Settings", AssetRegistrySearchable, meta = (EditCondition = bEnablePreconfiguredSettings, EditConditionHides))
-	bool bOnlyExposePreconfiguredSettings = false;
-
-	UPROPERTY(BlueprintReadWrite, EditDefaultsOnly, Category = "AssetInfo|Preconfigured Settings", meta = (EditCondition = bEnablePreconfiguredSettings, EditConditionHides))
-	TArray<FPCGPreConfiguredSettingsInfo> PreconfiguredInfo;
-
-	UPROPERTY(BlueprintReadWrite, EditDefaultsOnly, Category = AssetInfo, AssetRegistrySearchable)
-	FText Category;
-
-	UPROPERTY(BlueprintReadWrite, EditDefaultsOnly, Category = AssetInfo, AssetRegistrySearchable)
-	FText Description;
-
-	UPROPERTY(BlueprintReadWrite, EditDefaultsOnly, AdvancedDisplay, Category = "Settings")
-	int32 DependencyParsingDepth = 1;
-#endif
-
-protected:
-#if WITH_EDITOR
-	UE_API void OnDependencyChanged(UObject* Object, FPropertyChangedEvent& PropertyChangedEvent);
-	UE_API void UpdateDependencies();
-	TSet<TWeakObjectPtr<UObject>> DataDependencies;
-#endif
-
-#if !WITH_EDITORONLY_DATA
-	UWorld* InstanceWorld = nullptr;
-#endif
-
-	// Since we duplicate the blueprint elements prior to execution, they will be unique
-	// and have a 1:1 match with their context, which allows us to store it here
-	FPCGContext* CurrentContext = nullptr;
-};
 
 UCLASS(MinimalAPI, BlueprintType, ClassGroup = (Procedural))
 class UPCGBlueprintSettings : public UPCGSettings
@@ -322,7 +34,7 @@ class UPCGBlueprintSettings : public UPCGSettings
 	GENERATED_BODY()
 
 public:
-	UE_API UPCGBlueprintSettings();
+	PCG_API UPCGBlueprintSettings();
 
 	friend class FPCGExecuteBlueprintElement;
 
@@ -330,68 +42,70 @@ public:
 #if WITH_EDITOR
 	virtual FName GetDefaultNodeName() const override { return FName(TEXT("ExecuteBlueprint")); }
 	virtual FText GetDefaultNodeTitle() const override { return NSLOCTEXT("PCGBlueprintSettings", "NodeTitle", "Execute Blueprint"); }
-	UE_API virtual FLinearColor GetNodeTitleColor() const override;
-	UE_API virtual EPCGSettingsType GetType() const override;
-	UE_API virtual void GetStaticTrackedKeys(FPCGSelectionKeyToSettingsMap& OutKeysToSettings, TArray<TObjectPtr<const UPCGGraph>>& OutVisitedGraphs) const override;
-	UE_API virtual UObject* GetJumpTargetForDoubleClick() const override;
-	UE_API virtual void ApplyDeprecationBeforeUpdatePins(UPCGNode* InOutNode, TArray<TObjectPtr<UPCGPin>>& InputPins, TArray<TObjectPtr<UPCGPin>>& OutputPins) override;
-	UE_API virtual TArray<FPCGPreConfiguredSettingsInfo> GetPreconfiguredInfo() const override;
-	UE_API virtual bool OnlyExposePreconfiguredSettings() const override;
+	PCG_API virtual FLinearColor GetNodeTitleColor() const override;
+	PCG_API virtual EPCGSettingsType GetType() const override;
+	PCG_API virtual void GetStaticTrackedKeys(FPCGSelectionKeyToSettingsMap& OutKeysToSettings, TArray<TObjectPtr<const UPCGGraph>>& OutVisitedGraphs) const override;
+	PCG_API virtual UObject* GetJumpTargetForDoubleClick() const override;
+	PCG_API virtual void ApplyDeprecationBeforeUpdatePins(UPCGNode* InOutNode, TArray<TObjectPtr<UPCGPin>>& InputPins, TArray<TObjectPtr<UPCGPin>>& OutputPins) override;
+	PCG_API virtual TArray<FPCGPreConfiguredSettingsInfo> GetPreconfiguredInfo() const override;
+	PCG_API virtual bool OnlyExposePreconfiguredSettings() const override;
 #endif
 	virtual bool UseSeed() const override { return true; }
-	UE_API virtual bool HasDynamicPins() const override;
-	UE_API virtual FPCGDataTypeIdentifier GetCurrentPinTypesID(const UPCGPin* InPin) const override;
+	PCG_API virtual bool HasDynamicPins() const override;
+	PCG_API virtual FPCGDataTypeIdentifier GetCurrentPinTypesID(const UPCGPin* InPin) const override;
 
-	UE_API virtual void ApplyPreconfiguredSettings(const FPCGPreConfiguredSettingsInfo& InPreconfiguredsInfo) override;
-	UE_API virtual FString GetAdditionalTitleInformation() const override;
+	PCG_API virtual void ApplyPreconfiguredSettings(const FPCGPreConfiguredSettingsInfo& InPreconfiguredsInfo) override;
+	PCG_API virtual FString GetAdditionalTitleInformation() const override;
 	virtual bool HasFlippedTitleLines() const override { return true; }
 	// This node may have side effects, don't assume we can cull even when unwired.
 	virtual bool CanCullTaskIfUnwired() const { return false; }
-	UE_API virtual TArray<FPCGPinProperties> InputPinProperties() const override;
-	UE_API virtual TArray<FPCGPinProperties> OutputPinProperties() const override;
+	PCG_API virtual TArray<FPCGPinProperties> InputPinProperties() const override;
+	PCG_API virtual TArray<FPCGPinProperties> OutputPinProperties() const override;
 
 	/** To be removed when we support automatic override of BP params. For now always return true to force params pin. */
 	virtual bool HasOverridableParams() const override { return true; }
 protected:
-	UE_API virtual FPCGElementPtr CreateElement() const override;
+	PCG_API virtual FPCGElementPtr CreateElement() const override;
 #if WITH_EDITOR
 	virtual EPCGChangeType GetChangeTypeForProperty(const FName& InPropertyName) const override { return Super::GetChangeTypeForProperty(InPropertyName) | EPCGChangeType::Cosmetic; }
-	UE_API virtual TArray<FPCGSettingsOverridableParam> GatherOverridableParams() const override;
+	PCG_API virtual TArray<FPCGSettingsOverridableParam> GatherOverridableParams() const override;
 #endif // WITH_EDITOR
-	UE_API virtual void FixingOverridableParamPropertyClass(FPCGSettingsOverridableParam& Param) const override;
+	PCG_API virtual void FixingOverridableParamPropertyClass(FPCGSettingsOverridableParam& Param) const override;
 	// ~End UPCGSettings interface
 
 public:
 	// ~Begin UObject interface
-	UE_API virtual void PostLoad() override;
-	UE_API virtual void BeginDestroy() override;
+	PCG_API virtual void PostLoad() override;
+	PCG_API virtual void BeginDestroy() override;
 #if WITH_EDITOR
-	UE_API virtual void PreEditChange(FProperty* PropertyAboutToChange) override;
-	UE_API virtual void PostEditChangeProperty(FPropertyChangedEvent& PropertyChangedEvent) override;
+	PCG_API virtual void PreEditChange(FProperty* PropertyAboutToChange) override;
+	PCG_API virtual void PostEditChangeProperty(FPropertyChangedEvent& PropertyChangedEvent) override;
 	// ~End UObject interface
 #endif
 
-	UFUNCTION(BlueprintCallable, Category = "Settings|Template", meta=(DeterminesOutputType="InElementType", DynamicOutputParam = "ElementInstance"))
-	UE_API void SetElementType(TSubclassOf<UPCGBlueprintElement> InElementType, UPCGBlueprintElement*& ElementInstance);
+	/** Deprecated: use SetBlueprintElementType instead. */
+	UFUNCTION(BlueprintCallable, Category = "Settings|Template", meta = (DeterminesOutputType = "InElementType", DynamicOutputParam = "ElementInstance"))
+	PCG_API void SetElementType(TSubclassOf<UPCGBlueprintElement> InElementType, UPCGBlueprintElement*& ElementInstance);
+
+	UFUNCTION(BlueprintCallable, Category = "Settings|Template", meta = (DeterminesOutputType = "InElementType", DynamicOutputParam = "ElementInstance"))
+	PCG_API void SetBlueprintElementType(TSubclassOf<UPCGBlueprintBaseElement> InElementType, UPCGBlueprintBaseElement*& ElementInstance);
 
 	UFUNCTION(BlueprintCallable, Category = "Settings|Template")
-	TSubclassOf<UPCGBlueprintElement> GetElementType() const { return BlueprintElementType; }
+	TSubclassOf<UPCGBlueprintBaseElement> GetElementType() const { return BlueprintElementType; }
 
 #if WITH_EDITOR
-	TObjectPtr<UPCGBlueprintElement> GetElementInstance() const { return BlueprintElementInstance; }
-#endif
+	UE_DEPRECATED(5.7, "Use GetElementObject instead")
+	TObjectPtr<UPCGBlueprintBaseElement> GetElementInstance() const { return nullptr; }
 
-protected:
-#if WITH_EDITORONLY_DATA
-	UPROPERTY()
-	TSubclassOf<UPCGBlueprintElement> BlueprintElement_DEPRECATED;
+	TObjectPtr<UPCGBlueprintBaseElement> GetElementObject() const { return BlueprintElementInstance; }
 #endif
 
+protected:
 	UPROPERTY(BlueprintReadOnly, EditAnywhere, Category = Template)
-	TSubclassOf<UPCGBlueprintElement> BlueprintElementType;
+	TSubclassOf<UPCGBlueprintBaseElement> BlueprintElementType;
 
 	UPROPERTY(BlueprintReadOnly, VisibleAnywhere, Instanced, Category = "Instance", meta = (ShowOnlyInnerProperties))
-	TObjectPtr<UPCGBlueprintElement> BlueprintElementInstance;
+	TObjectPtr<UPCGBlueprintBaseElement> BlueprintElementInstance;
 
 #if WITH_EDITORONLY_DATA
 	UPROPERTY(BlueprintReadWrite, EditAnywhere, Category = Settings)
@@ -400,29 +114,27 @@ protected:
 	/** If this is checked, found actors that are outside component bounds will not trigger a refresh. Only works for tags for now in editor. */
 	UPROPERTY(BlueprintReadOnly, EditAnywhere, Category = Settings)
 	bool bTrackActorsOnlyWithinBounds = false;
-
-	UPROPERTY()
-	bool bCanBeMultithreaded_DEPRECATED = false;
 #endif
 
 protected:
 #if WITH_EDITOR
 	UE_DEPRECATED(5.6, "No longer supported")
 	void OnBlueprintChanged(UBlueprint* InBlueprint) {}
-	UE_API void OnBlueprintElementChanged(UPCGBlueprintElement* InElement);
-	UE_API void OnObjectsReplaced(const TMap<UObject*, UObject*>& InOldToNewInstances);
+
+	PCG_API void OnBlueprintElementChanged(UPCGBlueprintBaseElement* InElement);
+	PCG_API void OnObjectsReplaced(const TMap<UObject*, UObject*>& InOldToNewInstances);
 #endif
 
-	UE_API void RefreshBlueprintElement();
-	UE_API void SetupBlueprintEvent();
-	UE_API void TeardownBlueprintEvent();
-	UE_API void SetupBlueprintElementEvent();
-	UE_API void TeardownBlueprintElementEvent();
+	PCG_API void RefreshBlueprintElement();
+	PCG_API void SetupBlueprintEvent();
+	PCG_API void TeardownBlueprintEvent();
+	PCG_API void SetupBlueprintElementEvent();
+	PCG_API void TeardownBlueprintElementEvent();
 };
 
 struct FPCGBlueprintExecutionContext : public FPCGContext
 {
-	TObjectPtr<UPCGBlueprintElement> BlueprintElementInstance = nullptr;
+	TObjectPtr<UPCGBlueprintBaseElement> BlueprintElementInstance = nullptr;
 
 protected:
 	virtual void AddExtraStructReferencedObjects(FReferenceCollector& Collector) override;
@@ -437,11 +149,10 @@ public:
 	virtual bool ShouldComputeFullOutputDataCrc(FPCGContext* Context) const override;
 	/** Set it to true by default, if there is a performance concern, we can expose a bool in the element class. */
 	virtual bool ShouldVerifyIfOutputsAreUsedMultipleTimes(const UPCGSettings* InSettings) const override { return true; }
+	virtual bool SupportsBasePointDataInputs(FPCGContext* Context) const override;
 
 protected:
 	virtual bool ExecuteInternal(FPCGContext* Context) const override;
 	virtual void PostExecuteInternal(FPCGContext* Context) const override;
-	virtual FPCGContext* Initialize(const FPCGInitializeElementParams& InParams) override;	
-};
-
-#undef UE_API
+	virtual FPCGContext* Initialize(const FPCGInitializeElementParams& InParams) override;
+};
\ No newline at end of file
diff --git a/Engine/Plugins/PCG/Source/PCG/Public/Elements/PCGIndirectionElement.h b/Engine/Plugins/PCG/Source/PCG/Public/Elements/PCGIndirectionElement.h
--- a/Engine/Plugins/PCG/Source/PCG/Public/Elements/PCGIndirectionElement.h
+++ b/Engine/Plugins/PCG/Source/PCG/Public/Elements/PCGIndirectionElement.h
@@ -7,7 +7,7 @@
 
 #include "PCGIndirectionElement.generated.h"
 
-class UPCGBlueprintElement;
+class UPCGBlueprintBaseElement;
 
 UENUM()
 enum class EPCGProxyInterfaceMode : uint8
@@ -58,7 +58,7 @@ public:
 
 	/** The blueprint element class used to define the pin interface for this node instance */
 	UPROPERTY(EditAnywhere, Category = Settings, meta = (EditCondition = "ProxyInterfaceMode == EPCGProxyInterfaceMode::ByBlueprintElement", EditConditionHides))
-    TSubclassOf<UPCGBlueprintElement> BlueprintElementClass;
+    TSubclassOf<UPCGBlueprintBaseElement> BlueprintElementClass;
 
 	/** The element settings, which can be overriden, that will be used during the proxy execution */
 	UPROPERTY(EditAnywhere, BlueprintReadWrite, Category = Settings, meta = (PCG_Overridable))
diff --git a/Engine/Plugins/PCG/Source/PCGEditor/Private/Details/PCGBlueprintSettingsDetails.cpp b/Engine/Plugins/PCG/Source/PCGEditor/Private/Details/PCGBlueprintSettingsDetails.cpp
--- a/Engine/Plugins/PCG/Source/PCGEditor/Private/Details/PCGBlueprintSettingsDetails.cpp
+++ b/Engine/Plugins/PCG/Source/PCGEditor/Private/Details/PCGBlueprintSettingsDetails.cpp
@@ -45,10 +45,10 @@ void FPCGBlueprintSettingsDetails::CustomizeDetails(IDetailLayoutBuilder& Detail
 
 	// Go through all the properties, if we find properties that have CallInEditor,
 	// create buttons for them to be executed
-	UPCGBlueprintElement* BlueprintElement = BlueprintSettings->GetElementInstance();
+	UPCGBlueprintBaseElement* BlueprintElement = BlueprintSettings->GetElementObject();
 	if (BlueprintElement)
 	{
-		SelectedObject = TWeakObjectPtr<UPCGBlueprintElement>(BlueprintElement);
+		SelectedObject = TWeakObjectPtr<UPCGBlueprintBaseElement>(BlueprintElement);
 		AddCallInEditorMethods(DetailBuilder);
 		//FObjectDetailsTools::AddCallInEditorMethods(
 		//	DetailBuilder,
diff --git a/Engine/Plugins/PCG/Source/PCGEditor/Private/PCGEditorGraphSchema.cpp b/Engine/Plugins/PCG/Source/PCGEditor/Private/PCGEditorGraphSchema.cpp
--- a/Engine/Plugins/PCG/Source/PCGEditor/Private/PCGEditorGraphSchema.cpp
+++ b/Engine/Plugins/PCG/Source/PCGEditor/Private/PCGEditorGraphSchema.cpp
@@ -457,14 +457,14 @@ void UPCGEditorGraphSchema::GetBlueprintElementActions(FGraphActionMenuBuilder&
 {
 	PCGEditorUtils::ForEachPCGBlueprintAssetData([&ActionMenuBuilder, &ActionFilter](const FAssetData& AssetData)
 	{
-		const bool bExposeToLibrary = AssetData.GetTagValueRef<bool>(GET_MEMBER_NAME_CHECKED(UPCGBlueprintElement, bExposeToLibrary));
-		const bool bOnlyExposePreconfiguredSettings = AssetData.GetTagValueRef<bool>(GET_MEMBER_NAME_CHECKED(UPCGBlueprintElement, bOnlyExposePreconfiguredSettings));
+		const bool bExposeToLibrary = AssetData.GetTagValueRef<bool>(GET_MEMBER_NAME_CHECKED(UPCGBlueprintBaseElement, bExposeToLibrary));
+		const bool bOnlyExposePreconfiguredSettings = AssetData.GetTagValueRef<bool>(GET_MEMBER_NAME_CHECKED(UPCGBlueprintBaseElement, bOnlyExposePreconfiguredSettings));
 
 		if (bExposeToLibrary)
 		{
 			const FText MenuDesc = FText::FromString(FName::NameToDisplayString(AssetData.AssetName.ToString(), false));
-			const FText Category = AssetData.GetTagValueRef<FText>(GET_MEMBER_NAME_CHECKED(UPCGBlueprintElement, Category));
-			const FText Description = AssetData.GetTagValueRef<FText>(GET_MEMBER_NAME_CHECKED(UPCGBlueprintElement, Description));
+			const FText Category = AssetData.GetTagValueRef<FText>(GET_MEMBER_NAME_CHECKED(UPCGBlueprintBaseElement, Category));
+			const FText Description = AssetData.GetTagValueRef<FText>(GET_MEMBER_NAME_CHECKED(UPCGBlueprintBaseElement, Description));
 
 			if (!ActionFilter.Accepts(Category))
 			{
@@ -475,10 +475,10 @@ void UPCGEditorGraphSchema::GetBlueprintElementActions(FGraphActionMenuBuilder&
 
 			// Only load the class if we have enabled preconfigured settings.
 			TArray<FPCGPreConfiguredSettingsInfo> AllPreconfiguredInfo;
-			if (AssetData.GetTagValueRef<bool>(GET_MEMBER_NAME_CHECKED(UPCGBlueprintElement, bEnablePreconfiguredSettings)))
+			if (AssetData.GetTagValueRef<bool>(GET_MEMBER_NAME_CHECKED(UPCGBlueprintBaseElement, bEnablePreconfiguredSettings)))
 			{
-				TSubclassOf<UPCGBlueprintElement> BlueprintClass = GeneratedClass.TryLoadClass<UPCGBlueprintElement>();
-				const UPCGBlueprintElement* BlueprintElement = BlueprintClass ? Cast<UPCGBlueprintElement>(BlueprintClass->GetDefaultObject()) : nullptr;
+				TSubclassOf<UPCGBlueprintBaseElement> BlueprintClass = GeneratedClass.TryLoadClass<UPCGBlueprintBaseElement>();
+				const UPCGBlueprintBaseElement* BlueprintElement = BlueprintClass ? Cast<UPCGBlueprintBaseElement>(BlueprintClass->GetDefaultObject()) : nullptr;
 
 				AllPreconfiguredInfo = BlueprintElement ? BlueprintElement->PreconfiguredInfo : TArray<FPCGPreConfiguredSettingsInfo>{};
 			}
@@ -792,7 +792,7 @@ void UPCGEditorGraphSchema::GetAssetsGraphHoverMessage(const TArray<FAssetData>&
 			}
 			else if(AssetClass->IsChildOf(UBlueprint::StaticClass()))
 			{
-				OutTooltipText = LOCTEXT("PCGEditorDropAssetInvalidBP", "Blueprint does not derive from UPCGBlueprintElement").ToString();
+				OutTooltipText = LOCTEXT("PCGEditorDropAssetInvalidBP", "Blueprint does not derive from UPCGBlueprintBaseElement").ToString();
 				OutOkIcon = false;
 				return;
 			}
diff --git a/Engine/Plugins/PCG/Source/PCGEditor/Private/PCGEditorGraphSchemaActions.cpp b/Engine/Plugins/PCG/Source/PCGEditor/Private/PCGEditorGraphSchemaActions.cpp
--- a/Engine/Plugins/PCG/Source/PCGEditor/Private/PCGEditorGraphSchemaActions.cpp
+++ b/Engine/Plugins/PCG/Source/PCGEditor/Private/PCGEditorGraphSchemaActions.cpp
@@ -472,9 +472,9 @@ UEdGraphNode* FPCGEditorGraphSchemaAction_NewBlueprintElement::PerformAction(UEd
 	UPCGNode* NewPCGNode = PCGGraph->AddNodeOfType(UPCGBlueprintSettings::StaticClass(), DefaultNodeSettings);
 	UPCGBlueprintSettings* DefaultBlueprintSettings = CastChecked<UPCGBlueprintSettings>(DefaultNodeSettings);
 
-	UPCGBlueprintElement* ElementInstance = nullptr;
-	TSubclassOf<UPCGBlueprintElement> BlueprintClass = BlueprintClassPath.TryLoadClass<UPCGBlueprintElement>();
-	DefaultBlueprintSettings->SetElementType(BlueprintClass, ElementInstance);
+	UPCGBlueprintBaseElement* ElementInstance = nullptr;
+	TSubclassOf<UPCGBlueprintBaseElement> BlueprintClass = BlueprintClassPath.TryLoadClass<UPCGBlueprintBaseElement>();
+	DefaultBlueprintSettings->SetBlueprintElementType(BlueprintClass, ElementInstance);
 	DefaultBlueprintSettings->ApplyPreconfiguredSettings(PreconfiguredInfo);
 
 	NewPCGNode->UpdateAfterSettingsChangeDuringCreation();
diff --git a/Engine/Plugins/PCG/Source/PCGEditor/Private/PCGEditorUtils.cpp b/Engine/Plugins/PCG/Source/PCGEditor/Private/PCGEditorUtils.cpp
--- a/Engine/Plugins/PCG/Source/PCGEditor/Private/PCGEditorUtils.cpp
+++ b/Engine/Plugins/PCG/Source/PCGEditor/Private/PCGEditorUtils.cpp
@@ -44,7 +44,7 @@ bool PCGEditorUtils::IsAssetPCGBlueprint(const FAssetData& InAssetData)
 	const FString InNativeParentClassName = InAssetData.GetTagValueRef<FString>(FBlueprintTags::NativeParentClassPath);
 	const UClass* NativeParentClass = !InNativeParentClassName.IsEmpty() ? FSoftClassPath(InNativeParentClassName).ResolveClass() : nullptr;
 
-	return NativeParentClass && NativeParentClass->IsChildOf<UPCGBlueprintElement>();
+	return NativeParentClass && NativeParentClass->IsChildOf<UPCGBlueprintBaseElement>();
 }
 
 void PCGEditorUtils::GetParentPackagePathAndUniqueName(const UObject* OriginalObject, const FString& NewAssetTentativeName, FString& OutPackagePath, FString& OutUniqueName)
@@ -91,7 +91,7 @@ void PCGEditorUtils::ForEachPCGBlueprintAssetData(TFunctionRef<bool(const FAsset
 	FARFilter Filter;
 	Filter.ClassPaths.Add(UBlueprint::StaticClass()->GetClassPathName());
 	Filter.bRecursiveClasses = true;
-	Filter.TagsAndValues.Add(FBlueprintTags::NativeParentClassPath, UPCGBlueprintElement::GetParentClassName());
+	Filter.TagsAndValues.Add(FBlueprintTags::NativeParentClassPath, UPCGBlueprintBaseElement::GetParentClassName());
 
 	ForEachAssetData(Filter, InFunc);
 }
@@ -129,7 +129,7 @@ void PCGEditorUtils::ForcePCGBlueprintVariableVisibility()
 	{
 		const FString GeneratedClass = AssetData.GetTagValueRef<FString>(FBlueprintTags::GeneratedClassPath);
 		FSoftClassPath BlueprintClassPath = FSoftClassPath(GeneratedClass);
-		TSubclassOf<UPCGBlueprintElement> BlueprintClass = BlueprintClassPath.TryLoadClass<UPCGBlueprintElement>();
+		TSubclassOf<UPCGBlueprintBaseElement> BlueprintClass = BlueprintClassPath.TryLoadClass<UPCGBlueprintBaseElement>();
 		if (BlueprintClass)
 		{
 			if (UBlueprint* Blueprint = Cast<UBlueprint>(BlueprintClass->ClassGeneratedBy))
diff --git a/Engine/Plugins/PCG/Source/PCGEditor/Public/Details/PCGBlueprintSettingsDetails.h b/Engine/Plugins/PCG/Source/PCGEditor/Public/Details/PCGBlueprintSettingsDetails.h
--- a/Engine/Plugins/PCG/Source/PCGEditor/Public/Details/PCGBlueprintSettingsDetails.h
+++ b/Engine/Plugins/PCG/Source/PCGEditor/Public/Details/PCGBlueprintSettingsDetails.h
@@ -10,7 +10,7 @@ class FReply;
 
 class IDetailLayoutBuilder;
 class UFunction;
-class UPCGBlueprintElement;
+class UPCGBlueprintBaseElement;
 
 class FPCGBlueprintSettingsDetails : public IDetailCustomization
 {
@@ -29,5 +29,5 @@ private:
 	FReply OnExecuteCallInEditorFunction(TWeakObjectPtr<UFunction> WeakFunctionPtr);
 
 private:
-	TWeakObjectPtr<UPCGBlueprintElement> SelectedObject;
+	TWeakObjectPtr<UPCGBlueprintBaseElement> SelectedObject;
 };
diff --git a/Engine/Plugins/PCG/Source/PCGEditor/Public/Schema/PCGEditorGraphSchemaActions.h b/Engine/Plugins/PCG/Source/PCGEditor/Public/Schema/PCGEditorGraphSchemaActions.h
--- a/Engine/Plugins/PCG/Source/PCGEditor/Public/Schema/PCGEditorGraphSchemaActions.h
+++ b/Engine/Plugins/PCG/Source/PCGEditor/Public/Schema/PCGEditorGraphSchemaActions.h
@@ -12,7 +12,7 @@
 
 class UEdGraph;
 class UEdGraphPin;
-class UPCGBlueprintElement;
+class UPCGBlueprintBaseElement;
 class UPCGNode;
 class UPCGEditorGraph;
 struct FSlateBrush;
diff --git a/Engine/Plugins/PCGInterops/PCGGeometryScriptInterop/Source/PCGGeometryScriptInterop/Private/Elements/PCGGeometryBlueprintElement.cpp b/Engine/Plugins/PCGInterops/PCGGeometryScriptInterop/Source/PCGGeometryScriptInterop/Private/Elements/PCGGeometryBlueprintElement.cpp
--- a/Engine/Plugins/PCGInterops/PCGGeometryScriptInterop/Source/PCGGeometryScriptInterop/Private/Elements/PCGGeometryBlueprintElement.cpp
+++ b/Engine/Plugins/PCGInterops/PCGGeometryScriptInterop/Source/PCGGeometryScriptInterop/Private/Elements/PCGGeometryBlueprintElement.cpp
@@ -13,7 +13,7 @@
 #define LOCTEXT_NAMESPACE "PCGGeometryBlueprintElement"
 
 UPCGGeometryBlueprintElement::UPCGGeometryBlueprintElement()
-	: UPCGBlueprintElement()
+	: UPCGBlueprintBaseElement()
 {
 	// Setup the element for basic dynamic mesh processing
 	bIsCacheable = false;
@@ -24,38 +24,19 @@ UPCGGeometryBlueprintElement::UPCGGeometryBlueprintElement()
 	CustomOutputPins.Emplace(PCGPinConstants::DefaultOutputLabel, EPCGDataType::DynamicMesh);
 }
 
-void UPCGGeometryBlueprintElement::ExecuteWithContext_Implementation(FPCGContext& InContext, const FPCGDataCollection& Input, FPCGDataCollection& Output)
+void UPCGGeometryBlueprintElement::Execute_Implementation(const FPCGDataCollection& Input, FPCGDataCollection& Output)
 {
 	// Verify that we are in the right setup
 	const bool bHasASingleInputPin = (bHasDefaultInPin && CustomInputPins.IsEmpty()) || (!bHasDefaultInPin && CustomInputPins.Num() == 1);
 	const bool bHasASingleOutputPin = (bHasDefaultOutPin && CustomOutputPins.IsEmpty()) || (!bHasDefaultOutPin && CustomOutputPins.Num() == 1);
 
-#if WITH_EDITOR
-	const UBlueprint* Blueprint = Cast<UBlueprint>(GetClass()->ClassGeneratedBy);
-	auto IsFunctionOverridden = [Blueprint](const FName FunctionName)
-	{
-		return Blueprint && Blueprint->FunctionGraphs.ContainsByPredicate([FunctionName](const TObjectPtr<UEdGraph>& Graph) { return Graph->GetFName() == FunctionName;});
-	};
-	
-	const bool bIsExecuteOverridden = IsFunctionOverridden(GET_FUNCTION_NAME_CHECKED(UPCGBlueprintElement, Execute));
-	const bool bIsProcessDynMeshOverridden = IsFunctionOverridden(GET_FUNCTION_NAME_CHECKED(UPCGGeometryBlueprintElement, ProcessDynamicMesh));
-#else
-	// We can't know in runtime build as the functions are compiled.
-	const bool bIsExecuteOverridden = false;
-	const bool bIsProcessDynMeshOverridden = false;
-#endif // WITH_EDITOR
+	const bool bIsProcessDynMeshOverridden =  GetClass()->IsFunctionImplementedInScript(GET_FUNCTION_NAME_CHECKED(UPCGGeometryBlueprintElement, ProcessDynamicMesh));
 	
-	if (!bHasASingleInputPin || !bHasASingleOutputPin || bIsExecuteOverridden)
+	if ((!bHasASingleInputPin || !bHasASingleOutputPin) && bIsProcessDynMeshOverridden)
 	{
 		// Make sure to throw a warning if we are in this case and Process Dynamic Mesh is overridden to warn the user their function won't be called
-		if (bIsProcessDynMeshOverridden)
-		{
-			PCGLog::LogWarningOnGraph(LOCTEXT("DynMeshOverridenButNotCalled", "Process Dynamic Mesh was overridden, but we don't have the expected setup (single input and output pin)"
+		PCGLog::LogWarningOnGraph(LOCTEXT("DynMeshOverridenButNotCalled", "Process Dynamic Mesh was overridden, but we don't have the expected setup (single input and output pin)"
 				" or Execute is also overriden. Process Dynamic Mesh won't be called."), &GetContext());
-		}
-		
-		// If we have not the right number of pins, or Execute is overriden, just call the parent function.
-		UPCGBlueprintElement::ExecuteWithContext_Implementation(InContext, Input, Output);
 		return;
 	}
 
diff --git a/Engine/Plugins/PCGInterops/PCGGeometryScriptInterop/Source/PCGGeometryScriptInterop/Public/Elements/PCGGeometryBlueprintElement.h b/Engine/Plugins/PCGInterops/PCGGeometryScriptInterop/Source/PCGGeometryScriptInterop/Public/Elements/PCGGeometryBlueprintElement.h
--- a/Engine/Plugins/PCGInterops/PCGGeometryScriptInterop/Source/PCGGeometryScriptInterop/Public/Elements/PCGGeometryBlueprintElement.h
+++ b/Engine/Plugins/PCGInterops/PCGGeometryScriptInterop/Source/PCGGeometryScriptInterop/Public/Elements/PCGGeometryBlueprintElement.h
@@ -6,27 +6,28 @@
 
 #include "PCGGeometryBlueprintElement.generated.h"
 
-#define UE_API PCGGEOMETRYSCRIPTINTEROP_API
-
 class UDynamicMesh;
 class UPCGDynamicMeshData;
 
 /**
-* Subclass of PCG Blueprint Element, it comes with pre-configured pins as input and output for Dynamic meshes and force to be non-cacheable.
+* Subclass of PCG Blueprint Base Element, it comes with pre-configured pins as input and output for Dynamic meshes and force to be non-cacheable.
 * The function CopyOrStealInputData is a helper to either steal (efficient) or copy the input data (less efficient) so work can be done in place on the Dynamic Mesh.
 * More importantly, a user deriving from this class will want to implement ProcessDynamicMesh, the only thing needed to streamline the process and
 * removes all the boilerplate when in a simple input->output case.
 */
 UCLASS(MinimalAPI, Abstract, BlueprintType, Blueprintable, hidecategories = (Object))
-class UPCGGeometryBlueprintElement : public UPCGBlueprintElement
+class UPCGGeometryBlueprintElement : public UPCGBlueprintBaseElement
 {
 	GENERATED_BODY()
 	
 public:
-	UE_API UPCGGeometryBlueprintElement();
+	PCGGEOMETRYSCRIPTINTEROP_API UPCGGeometryBlueprintElement();
+
+	UE_DEPRECATED(5.7, "Use Execute_Implementation instead")
+	PCGGEOMETRYSCRIPTINTEROP_API virtual void ExecuteWithContext_Implementation(FPCGContext& InContext, const FPCGDataCollection& Input, FPCGDataCollection& Output) {}
 
 	/** Virtual implementation to streamline the creation of a Geometry Script node processing. Do not override this function or Execute if you want to use the streamlined version. */
-	UE_API virtual void ExecuteWithContext_Implementation(FPCGContext& InContext, const FPCGDataCollection& Input, FPCGDataCollection& Output) override;
+	PCGGEOMETRYSCRIPTINTEROP_API virtual void Execute_Implementation(const FPCGDataCollection& Input, FPCGDataCollection& Output) override;
 
 	/**
      * Streamlined version of the Execute function, to only process the dynamic meshes.
@@ -35,11 +36,9 @@ public:
      * @param OutTags  - Optional tags to add to the output. By default, it will inherit the tags of the input.
      */
     UFUNCTION(BlueprintImplementableEvent, BlueprintCallable, Category = "PCG|Execution", meta = (AutoCreateRefTerm = "InTags"))
-    UE_API void ProcessDynamicMesh(UDynamicMesh* InDynMesh, TArray<FString>& OutTags);
+	PCGGEOMETRYSCRIPTINTEROP_API void ProcessDynamicMesh(UDynamicMesh* InDynMesh, TArray<FString>& OutTags);
 
 	/** Allows to steal the data and work in place if the data is not used elsewhere. If this element is cacheable, it will automatically copy. */
 	UFUNCTION(BlueprintCallable, Category = "PCG|DynamicMesh", meta = (HidePin = "Self"))
-	UE_API UPCGDynamicMeshData* CopyOrStealInputData(const FPCGTaggedData& InTaggedData) const;
+	PCGGEOMETRYSCRIPTINTEROP_API UPCGDynamicMeshData* CopyOrStealInputData(const FPCGTaggedData& InTaggedData) const;
 };
-
-#undef UE_API
