[PCG] Precursor update for Data View changes

#rnx

diff --git a/Engine/Plugins/PCG/Source/PCG/Private/Elements/IO/PCGExportSelectedAttributes.cpp b/Engine/Plugins/PCG/Source/PCG/Private/Elements/IO/PCGExportSelectedAttributes.cpp
--- a/Engine/Plugins/PCG/Source/PCG/Private/Elements/IO/PCGExportSelectedAttributes.cpp
+++ b/Engine/Plugins/PCG/Source/PCG/Private/Elements/IO/PCGExportSelectedAttributes.cpp
@@ -4,22 +4,25 @@
 
 #include "Data/PCGBasePointData.h"
 #include "Data/PCGSplineData.h"
-#include "Elements/IO/PCGJsonHelpers.h"
+#include "Helpers/IO/PCGJsonHelpers.h"
 #include "Metadata/PCGMetadata.h"
 #include "Metadata/PCGMetadataAttributeTraits.h"
 #include "Metadata/Accessors/PCGAttributeAccessorHelpers.h"
 #include "Utils/PCGLogErrors.h"
 
 #include "DesktopPlatformModule.h"
+#include "IDesktopPlatform.h"
+#include "Data/PCGSplineStruct.h"
 #include "HAL/FileManager.h"
 #include "HAL/FileManagerGeneric.h"
+#include "Misc/FileHelper.h"
+#include "Serialization/JsonSerializer.h"
+#include "Serialization/JsonWriter.h"
+#include "Helpers/IO/PCGIOHelpers.h"
 
 #if WITH_EDITOR
 #include "IContentBrowserSingleton.h"
 #endif // WITH_EDITOR
-#include "IDesktopPlatform.h"
-#include "Data/PCGSplineStruct.h"
-#include "Misc/FileHelper.h"
 
 #include UE_INLINE_GENERATED_CPP_BY_NAME(PCGExportSelectedAttributes)
 
@@ -80,27 +83,6 @@ namespace PCGExportSelectedAttributes
 			}
 		}
 
-		template <typename T>
-		bool GetAttributeValues(const FAccessorData& Data, TArray<T>& OutValues, EPCGAttributeAccessorFlags Flags, const FPCGContext* Context)
-		{
-			if constexpr (std::is_trivially_copyable_v<T>)
-			{
-				OutValues.SetNumUninitialized(Data.Keys->GetNum());
-			}
-			else
-			{
-				OutValues.SetNum(Data.Keys->GetNum());
-			}
-
-			if (!Data.Accessor->GetRange<T>(OutValues, 0, *Data.Keys, Flags))
-			{
-				PCGLog::Metadata::LogFailToGetAttributeError(Data.Selector, Context);
-				return false;
-			}
-
-			return true;
-		}
-
 		namespace String
 		{
 			FString MergeFileName(const FString& Path, const FString& File, const FString& Ext)
@@ -204,6 +186,7 @@ bool FPCGExportSelectedAttributesElement::ExecuteInternal(FPCGContext* Context)
 	FArchive* Archive = nullptr;
 	TSharedPtr<FJsonObject> RootJsonObject = nullptr;
 
+	using namespace PCG::IO::Constants::Attribute;
 	// Initialize the serialization object needed.
 	if (Settings->Format == EPCGExportAttributesFormat::Binary)
 	{
@@ -216,7 +199,6 @@ bool FPCGExportSelectedAttributesElement::ExecuteInternal(FPCGContext* Context)
 		}
 
 		// Set up the custom export version
-		using namespace PCG::Json::Constants;
 		Archive->SetCustomVersion(FCustomExportVersion::GUID, FCustomExportVersion::LatestVersion, FCustomExportVersion::GetFriendlyName());
 
 		// Set up the user's custom data version
@@ -231,7 +213,8 @@ bool FPCGExportSelectedAttributesElement::ExecuteInternal(FPCGContext* Context)
 		RootJsonObject = MakeShared<FJsonObject>();
 
 		// Set up the Json header to include custom versions
-		TSharedPtr<FJsonObject> JsonHeaderObject = PCG::Json::Helpers::CreateJsonHeaderObject(PCG::Json::Constants::FCustomExportVersion::LatestVersion, Settings->bAddCustomDataVersion ? Settings->CustomVersion : -1);
+		TSharedPtr<FJsonObject> JsonHeaderObject = MakeShared<FJsonObject>();
+		PCG::IO::Json::Helpers::AppendHeader(JsonHeaderObject, FCustomExportVersion::LatestVersion, Settings->bAddCustomDataVersion ? Settings->CustomVersion : -1);
 		RootJsonObject->SetObjectField(LOCTEXT("JsonHeaderField", "Header").ToString(), std::move(JsonHeaderObject));
 	}
 
@@ -306,7 +289,7 @@ bool FPCGExportSelectedAttributesElement::ExecuteInternal(FPCGContext* Context)
 			PCGMetadataAttribute::CallbackWithRightType(Data.Accessor->GetUnderlyingType(), [Archive, &Data, Context]<typename T>(const T&)
 			{
 				TArray<T> Values;
-				if (Helpers::GetAttributeValues(Data, Values, EPCGAttributeAccessorFlags::StrictType, Context))
+				if (PCGAttributeAccessorHelpers::ExtractAllValues(Data.Accessor.Get(), Data.Keys.Get(), Values, EPCGAttributeAccessorFlags::StrictType))
 				{
 					*Archive << Values;
 				}
@@ -332,24 +315,24 @@ bool FPCGExportSelectedAttributesElement::ExecuteInternal(FPCGContext* Context)
 
 			TArray<TSharedPtr<FJsonValue>>& SourceJsonValues = AllSourceJsonValues[SourceIndex].Get<1>();
 
-			PCGMetadataAttribute::CallbackWithRightType(Data.Accessor->GetUnderlyingType(), [&Data, &SourceJsonValues, Context]<typename T>(const T&)
+			PCGMetadataAttribute::CallbackWithRightType(Data.Accessor->GetUnderlyingType(), [&Data, &SourceJsonValues]<typename T>(const T&)
 			{
-				// For floating point values and containers (FVector, FRotator, FTransform, etc), use the Json numerical to maintain precision.
+				// For floating point values and containers (FVector, FRotator, FTransform, etc.), use the Json numerical to maintain precision.
 				if constexpr (PCG::Private::MetadataTraits<T>::IsFloatingPoint)
 				{
 					TArray<T> Values;
-					Helpers::GetAttributeValues(Data, Values, EPCGAttributeAccessorFlags::StrictType, Context);
+					PCGAttributeAccessorHelpers::ExtractAllValues(Data.Accessor.Get(), Data.Keys.Get(), Values, EPCGAttributeAccessorFlags::StrictType);
 
 					SourceJsonValues.Reserve(Values.Num());
 					Algo::Transform(Values, SourceJsonValues, [&Data](const T& Value)
 					{
-						return PCG::Json::Helpers::ConvertFloatingPointTypeToJson(Value);
+						return PCG::IO::Json::Helpers::ConvertFloatingPointType(Value);
 					});
 				}
 				else // For everything else, automatically broadcast it to string
 				{
 					TArray<FString> Values;
-					Helpers::GetAttributeValues(Data, Values, EPCGAttributeAccessorFlags::AllowBroadcast, Context);
+					PCGAttributeAccessorHelpers::ExtractAllValues(Data.Accessor.Get(), Data.Keys.Get(), Values, EPCGAttributeAccessorFlags::AllowBroadcast);
 
 					SourceJsonValues.Reserve(Values.Num());
 					Algo::Transform(Values, SourceJsonValues, [](const FString& Value)
diff --git a/Engine/Plugins/PCG/Source/PCG/Private/Elements/IO/PCGJsonHelpers.cpp b/Engine/Plugins/PCG/Source/PCG/Private/Elements/IO/PCGJsonHelpers.cpp
deleted file mode 100644
diff --git a/Engine/Plugins/PCG/Source/PCG/Private/Elements/IO/PCGJsonHelpers.h b/Engine/Plugins/PCG/Source/PCG/Private/Elements/IO/PCGJsonHelpers.h
deleted file mode 100644
diff --git a/Engine/Plugins/PCG/Source/PCG/Private/Helpers/IO/PCGIOHelpers.cpp b/Engine/Plugins/PCG/Source/PCG/Private/Helpers/IO/PCGIOHelpers.cpp
new file mode 100644
--- /dev/null
+++ b/Engine/Plugins/PCG/Source/PCG/Private/Helpers/IO/PCGIOHelpers.cpp
@@ -0,0 +1,30 @@
+// Copyright Epic Games, Inc. All Rights Reserved.
+
+#include "Helpers/IO/PCGIOHelpers.h"
+
+#include "Misc/Paths.h"
+#include "Serialization/CustomVersion.h"
+
+#define LOCTEXT_NAMESPACE "PCGIOHelpers"
+
+namespace PCG::IO
+{
+	namespace Constants
+	{
+		namespace Attribute
+		{
+			const FGuid FCustomExportVersion::GUID = FGuid(0x04E74488, 0x4BAC8717, 0xBBB18694, 0x39F8F3CE);
+			FCustomVersionRegistration GRegisterPCGExportSelectedAttributesCustomVersion(FCustomExportVersion::GUID, FCustomExportVersion::LatestVersion, TEXT("PCGExportSelectedAttributes"));
+		}
+	}
+
+	namespace Helpers::String
+	{
+		void ToPrecisionString(const double Value, const int32 Precision, FString& OutString)
+		{
+			OutString = FString::SanitizeFloat(Value, Precision);
+		}
+	} // namespace Helpers::String
+} // namespace PCG::IO
+
+#undef LOCTEXT_NAMESPACE
diff --git a/Engine/Plugins/PCG/Source/PCG/Private/Helpers/IO/PCGJsonHelpers.cpp b/Engine/Plugins/PCG/Source/PCG/Private/Helpers/IO/PCGJsonHelpers.cpp
new file mode 100644
--- /dev/null
+++ b/Engine/Plugins/PCG/Source/PCG/Private/Helpers/IO/PCGJsonHelpers.cpp
@@ -0,0 +1,218 @@
+// Copyright Epic Games, Inc. All Rights Reserved.
+
+#include "Helpers/IO/PCGJsonHelpers.h"
+
+#include "Helpers/IO/PCGIOHelpers.h"
+#include "Metadata/Accessors/PCGAttributeAccessorHelpers.h"
+
+#include "Serialization/JsonSerializer.h"
+#include "Serialization/JsonWriter.h"
+
+#define LOCTEXT_NAMESPACE "PCGJson"
+
+namespace PCG::IO::Json
+{
+	using JsonKeyValuesMap = TArray<TTuple<FString, TArray<TSharedPtr<FJsonValue>>>>;
+
+	namespace Private
+	{
+		TSharedPtr<FJsonValueArray> DoubleArrayToJsonArray(const TConstArrayView<double> ArrayView)
+		{
+			TArray<TSharedPtr<FJsonValue>> JsonArray;
+			for (double Value : ArrayView)
+			{
+				JsonArray.Emplace(MakeShared<FJsonValueNumber>(Value));
+			}
+
+			return MakeShared<FJsonValueArray>(MoveTemp(JsonArray));
+		}
+
+		Accessor::FCache CreateAccessorCache(const UPCGData* Data, const TArray<FPCGAttributePropertySelector>& Selectors, const FPCGContext* Context = nullptr)
+		{
+			check(Data);
+			Accessor::FCache AccessorCache;
+
+			for (const FPCGAttributePropertySelector& Selector : Selectors)
+			{
+				Accessor::FCacheEntry Entry;
+				Entry.Selector = Selector;
+				Entry.Accessor = PCGAttributeAccessorHelpers::CreateConstAccessor(Data, Entry.Selector);
+				Entry.Keys = PCGAttributeAccessorHelpers::CreateConstKeys(Data, Entry.Selector);
+				if (!Entry.Accessor || !Entry.Keys)
+				{
+					PCGLog::Metadata::LogFailToCreateAccessorError(Entry.Selector, Context);
+					continue;
+				}
+
+				AccessorCache.Emplace(MoveTemp(Entry));
+			}
+
+			return AccessorCache;
+		}
+
+		TArray<TSharedPtr<FJsonValue>> AttributeToJsonValues(const Accessor::FCacheEntry& Entry)
+		{
+			check(Entry.Accessor.IsValid());
+			return PCGMetadataAttribute::CallbackWithRightType(Entry.Accessor->GetUnderlyingType(), [&Entry]<typename T>(const T&)
+			{
+				TArray<TSharedPtr<FJsonValue>> JsonValues;
+
+				// For floating point values and containers (FVector, FRotator, FTransform, etc.), use the Json numerical to maintain precision.
+				if constexpr (PCG::Private::MetadataTraits<T>::IsFloatingPoint)
+				{
+					TArray<T> Values;
+					PCGAttributeAccessorHelpers::ExtractAllValues(Entry.Accessor.Get(), Entry.Keys.Get(), Values, EPCGAttributeAccessorFlags::StrictType);
+
+					JsonValues.Reserve(Values.Num());
+					Algo::Transform(Values, JsonValues, [&Entry](const T& Value)
+					{
+						return Helpers::ConvertFloatingPointType(Value);
+					});
+				}
+				else // For everything else, automatically broadcast it to string
+				{
+					TArray<FString> Values;
+					PCGAttributeAccessorHelpers::ExtractAllValues(Entry.Accessor.Get(), Entry.Keys.Get(), Values, EPCGAttributeAccessorFlags::StrictType);
+
+					JsonValues.Reserve(Values.Num());
+					Algo::Transform(Values, JsonValues, [](const FString& OutValue)
+					{
+						return MakeShared<FJsonValueString>(OutValue);
+					});
+				}
+
+				return JsonValues;
+			});
+		}
+
+		void BuildKeyValuesMap(JsonKeyValuesMap& OutKeyValuesMap, const Accessor::FCache& AccessorCache)
+		{
+			OutKeyValuesMap.Reset();
+			OutKeyValuesMap.SetNum(AccessorCache.Num());
+			// Pre-process the source data into Json values.
+			for (int32 CacheIndex = 0; CacheIndex < AccessorCache.Num(); ++CacheIndex)
+			{
+				const Accessor::FCacheEntry& Entry = AccessorCache[CacheIndex];
+				FString& SourceName = OutKeyValuesMap[CacheIndex].Get<0>();
+				SourceName = Entry.Selector.ToString();
+
+				TArray<TSharedPtr<FJsonValue>>& SourceJsonValues = OutKeyValuesMap[CacheIndex].Get<1>();
+				SourceJsonValues = AttributeToJsonValues(Entry);
+			}
+		}
+
+		void AppendSelectionByAttribute(const TSharedPtr<FJsonObject>& InOutJsonObject, JsonKeyValuesMap& InOutKeyValueMap)
+		{
+			for (auto& KeyValuePair : InOutKeyValueMap)
+			{
+				const FString& Key = KeyValuePair.Get<0>();
+				TArray<TSharedPtr<FJsonValue>>& SourceValues = KeyValuePair.Get<1>();
+				InOutJsonObject->SetArrayField(Key, MoveTemp(SourceValues));
+			}
+		}
+
+		void AppendSelectionByElement(const TSharedPtr<FJsonObject>& InOutJsonObject, JsonKeyValuesMap& InOutKeyValueMap, const int32 NumElements)
+		{
+			if (!ensure(NumElements > 0))
+			{
+				return;
+			}
+
+			for (int32 ElementIndex = 0; ElementIndex < NumElements; ++ElementIndex)
+			{
+				TSharedPtr<FJsonObject> ElementJsonObject = MakeShared<FJsonObject>();
+
+				for (auto& KeyValuePair : InOutKeyValueMap)
+				{
+					const FString& Key = KeyValuePair.Get<0>();
+					TArray<TSharedPtr<FJsonValue>>& JsonValues = KeyValuePair.Get<1>();
+
+					// Currently only support the following Json types:
+					if (!JsonValues.IsValidIndex(ElementIndex)
+						|| !(JsonValues[ElementIndex]->Type == EJson::Array
+							 || JsonValues[ElementIndex]->Type == EJson::Number
+							 || JsonValues[ElementIndex]->Type == EJson::String
+							 || JsonValues[ElementIndex]->Type == EJson::Object))
+					{
+						continue;
+					}
+
+					// Move the values directly
+					ElementJsonObject->SetField(Key, MoveTemp(JsonValues[ElementIndex]));
+				}
+
+				InOutJsonObject->SetObjectField(FString::FromInt(ElementIndex), MoveTemp(ElementJsonObject));
+			}
+
+			// The map is no longer valid, so clean it.
+			InOutKeyValueMap.Empty();
+		}
+	}
+
+	namespace Helpers
+	{
+		TSharedPtr<FJsonValue> ConvertFloatingPointType(const FVector2D& Value)
+		{
+			return Private::DoubleArrayToJsonArray({Value.X, Value.Y});
+		}
+
+		TSharedPtr<FJsonValue> ConvertFloatingPointType(const FVector& Value)
+		{
+			return Private::DoubleArrayToJsonArray({Value.X, Value.Y, Value.Z});
+		}
+
+		TSharedPtr<FJsonValue> ConvertFloatingPointType(const FVector4& Value)
+		{
+			return Private::DoubleArrayToJsonArray({Value.X, Value.Y, Value.Z, Value.W});
+		}
+
+		TSharedPtr<FJsonValue> ConvertFloatingPointType(const FQuat& Value)
+		{
+			return Private::DoubleArrayToJsonArray({Value.X, Value.Y, Value.Z, Value.W});
+		}
+
+		TSharedPtr<FJsonValue> ConvertFloatingPointType(const FRotator& Value)
+		{
+			return Private::DoubleArrayToJsonArray({Value.Pitch, Value.Yaw, Value.Roll});
+		}
+
+		TSharedPtr<FJsonValue> ConvertFloatingPointType(const FTransform& Value)
+		{
+			TSharedPtr<FJsonObject> JsonTransformObject = MakeShared<FJsonObject>();
+			JsonTransformObject->SetArrayField(TEXT("translation"), ConvertFloatingPointType(Value.GetTranslation())->AsArray());
+			JsonTransformObject->SetArrayField(TEXT("rotation"), ConvertFloatingPointType(Value.GetRotation())->AsArray());
+			JsonTransformObject->SetArrayField(TEXT("scale"), ConvertFloatingPointType(Value.GetScale3D())->AsArray());
+
+			return MakeShared<FJsonValueObject>(JsonTransformObject);
+		}
+
+		void AppendHeader(const TSharedPtr<FJsonObject>& InOutJsonObject, const int32 DataVersion, const int32 CustomDataVersion)
+		{
+			if (ensure(InOutJsonObject.IsValid()))
+			{
+				InOutJsonObject->SetNumberField(TEXT("data_version"), DataVersion);
+
+				if (CustomDataVersion >= 0)
+				{
+					InOutJsonObject->SetNumberField(TEXT("custom_data_version"), CustomDataVersion);
+				}
+			}
+		}
+
+		FString ToJsonString(const TSharedPtr<FJsonObject>& JsonObject)
+		{
+			FString OutputString;
+			// @todo_pcg: Condensed by default. Add 'pretty json' as an option.
+			const TSharedRef<TJsonWriter<>> JsonWriter = TJsonWriterFactory<>::Create(&OutputString);
+			if (!FJsonSerializer::Serialize(JsonObject.ToSharedRef(), JsonWriter))
+			{
+				PCGLog::LogErrorOnGraph(LOCTEXT("JsonSerializationError", "Serialization of Json data failed."));
+				return {};
+			}
+
+			return OutputString;
+		}
+	} // namespace Helpers
+}  // namespace PCG::IO::Json
+
+#undef LOCTEXT_NAMESPACE
diff --git a/Engine/Plugins/PCG/Source/PCG/Private/PCGGraph.cpp b/Engine/Plugins/PCG/Source/PCG/Private/PCGGraph.cpp
--- a/Engine/Plugins/PCG/Source/PCG/Private/PCGGraph.cpp
+++ b/Engine/Plugins/PCG/Source/PCG/Private/PCGGraph.cpp
@@ -16,7 +16,6 @@
 #include "Elements/PCGUserParameterGet.h"
 #include "Elements/ControlFlow/PCGQualityBranch.h"
 #include "Elements/ControlFlow/PCGQualitySelect.h"
-#include "Elements/IO/PCGJsonHelpers.h"
 #include "Graph/PCGGraphCompilationData.h"
 #include "Graph/PCGGraphCompiler.h"
 #include "Graph/PCGGraphExecutor.h"
@@ -30,7 +29,11 @@
 #include "Dialogs/Dialogs.h"
 #include "EdGraph/EdGraphPin.h"
 #include "GameFramework/Actor.h"
+#include "Serialization/JsonSerializer.h"
+#include "Serialization/JsonWriter.h"
 #include "UObject/ObjectSaveContext.h"
+#include "Helpers/IO/PCGJsonHelpers.h"
+#include "Helpers/IO/PCGIOHelpers.h"
 #else
 #include "UObject/Package.h"
 #endif
@@ -2236,11 +2239,13 @@ EPropertyBagAlterationResult UPCGGraph::AddUserParameters(const TArray<FProperty
 #if WITH_EDITOR
 FString UPCGGraph::GetGraphDescriptionJson() const
 {
+	using namespace PCG::IO::Json;
 	TSharedPtr<FJsonObject> RootJsonObject = MakeShared<FJsonObject>();
 
 	// Set up the Json header to include custom versions
-	TSharedPtr<FJsonObject> JsonHeaderObject = PCG::Json::Helpers::CreateJsonHeaderObject(PCG::Json::Constants::FCustomExportVersion::LatestVersion);
-	RootJsonObject->SetObjectField(LOCTEXT("JsonHeaderField", "Header").ToString(), std::move(JsonHeaderObject));
+	TSharedPtr<FJsonObject> JsonHeaderObject = MakeShared<FJsonObject>();
+	Helpers::AppendHeader(JsonHeaderObject, PCG::IO::Constants::Attribute::FCustomExportVersion::LatestVersion);
+	RootJsonObject->SetObjectField(LOCTEXT("JsonHeaderField", "Header").ToString(), MoveTemp(JsonHeaderObject));
 
 	// 0. Graph information
 	// - Name
@@ -2303,7 +2308,7 @@ FString UPCGGraph::GetGraphDescriptionJson() const
 			// - Value
 			auto SetJson = [&JsonObject]<typename T>(const T& Value) -> bool
 			{
-				PCG::Json::Helpers::SetJsonValue<T>(JsonObject, TEXT("Value"), Value);
+				Helpers::SetValue<T>(JsonObject, TEXT("value"), Value);
 				return true;
 			};
 
diff --git a/Engine/Plugins/PCG/Source/PCG/Public/Helpers/IO/PCGIOHelpers.h b/Engine/Plugins/PCG/Source/PCG/Public/Helpers/IO/PCGIOHelpers.h
new file mode 100644
--- /dev/null
+++ b/Engine/Plugins/PCG/Source/PCG/Public/Helpers/IO/PCGIOHelpers.h
@@ -0,0 +1,51 @@
+// Copyright Epic Games, Inc. All Rights Reserved.
+
+#pragma once
+
+#include "Metadata/PCGAttributePropertySelector.h"
+#include "Metadata/Accessors/PCGAttributeAccessor.h"
+#include "Metadata/Accessors/PCGAttributeAccessorKeys.h"
+
+namespace PCG::IO
+{
+	namespace Constants
+	{
+		namespace Attribute
+		{
+			// Custom PCG version for attribute export
+			struct FCustomExportVersion
+			{
+				enum Type
+				{
+					InitialVersion = 0,
+
+					// New versions can be added above this line
+					VersionPlusOne,
+					LatestVersion = VersionPlusOne - 1
+				};
+
+				static FName GetFriendlyName() { return FName(TEXT("PCG::IO::AttributeExport")); }
+				const static FGuid GUID;
+			};
+		} // namespace Attribute
+
+		static constexpr int32 InlineAllocationCount = 16;
+	} // namespace Constants
+
+	namespace Helpers::String
+	{
+		void ToPrecisionString(const double Value, const int32 Precision, FString& OutString);
+	}
+
+	namespace Accessor
+	{
+		struct FCacheEntry
+		{
+			FPCGAttributePropertySelector Selector;
+			TUniquePtr<const IPCGAttributeAccessor> Accessor = nullptr;
+			TUniquePtr<const IPCGAttributeAccessorKeys> Keys = nullptr;
+		};
+
+		using FCache = TArray<FCacheEntry, TInlineAllocator<Constants::InlineAllocationCount>>;
+	} // namespace Accessor
+} // namespace PCG::IO
diff --git a/Engine/Plugins/PCG/Source/PCG/Public/Helpers/IO/PCGJsonHelpers.h b/Engine/Plugins/PCG/Source/PCG/Public/Helpers/IO/PCGJsonHelpers.h
new file mode 100644
--- /dev/null
+++ b/Engine/Plugins/PCG/Source/PCG/Public/Helpers/IO/PCGJsonHelpers.h
@@ -0,0 +1,79 @@
+// Copyright Epic Games, Inc. All Rights Reserved.
+
+#pragma once
+
+#include "Metadata/PCGMetadataAttributeTraits.h"
+
+#include "Dom/JsonValue.h"
+#include "Dom/JsonObject.h"
+
+class FJsonValue;
+class FJsonValueArray;
+class UPCGDataViewData;
+struct FPCGContext;
+struct FPCGAttributePropertySelector;
+
+namespace PCG::IO::Json
+{
+	template <typename T>
+	concept CFloatingPoint = std::is_floating_point_v<std::remove_cv_t<T>>;
+
+	namespace Helpers
+	{
+		// Converts a floating point (float, double, FVector, FRotator, etc) type to a Json value.
+		template <CFloatingPoint T>
+		TSharedPtr<FJsonValue> ConvertFloatingPointType(const T& Value)
+		{
+			return MakeShared<FJsonValueNumber>(static_cast<double>(Value));
+		}
+
+		PCG_API TSharedPtr<FJsonValue> ConvertFloatingPointType(const FVector2D& Value);
+		PCG_API TSharedPtr<FJsonValue> ConvertFloatingPointType(const FVector& Value);
+		PCG_API TSharedPtr<FJsonValue> ConvertFloatingPointType(const FVector4& Value);
+		PCG_API TSharedPtr<FJsonValue> ConvertFloatingPointType(const FQuat& Value);
+		PCG_API TSharedPtr<FJsonValue> ConvertFloatingPointType(const FRotator& Value);
+		PCG_API TSharedPtr<FJsonValue> ConvertFloatingPointType(const FTransform& Value);
+
+		template <typename T>
+		TSharedPtr<FJsonValue> ConvertFloatingPointType(const T&)
+		{
+			static_assert(!std::is_same_v<T, T>, "Unsupported type");
+			return nullptr;
+		}
+
+		// Pass a JsonObject and set a field inside the object directly
+		template <typename T>
+		void SetValue(TSharedPtr<FJsonObject>& InOutJsonObject, FString&& ValueName, const T& Value)
+		{
+			check(InOutJsonObject);
+			if constexpr (std::is_same_v<T, bool>)
+			{
+				InOutJsonObject->SetBoolField(std::forward<FString>(ValueName), Value);
+			}
+			else if constexpr (std::is_unsigned_v<std::remove_cv_t<T>> || std::is_signed_v<std::remove_cv_t<T>>)
+			{
+				InOutJsonObject->SetNumberField(std::forward<FString>(ValueName), Value);
+			}
+			else if constexpr (Private::MetadataTraits<T>::IsFloatingPoint)
+			{
+				InOutJsonObject->SetField(std::forward<FString>(ValueName), ConvertFloatingPointType(Value));
+			}
+			else if constexpr (std::is_same_v<T, FString>)
+			{
+				InOutJsonObject->SetStringField(std::forward<FString>(ValueName), Value);
+			}
+			else
+			{
+				InOutJsonObject->SetStringField(std::forward<FString>(ValueName), Private::MetadataTraits<T>::ToString(Value));
+			}
+		}
+
+		// Append PCG Json data version header object to a Json object
+		void AppendHeader(const TSharedPtr<FJsonObject>& InOutJsonObject, int32 DataVersion, int32 CustomDataVersion = -1);
+
+		// Helper to convert a Json object into a string
+		FString ToJsonString(const TSharedPtr<FJsonObject>& JsonObject);
+
+		// @todo_pcg: Add a fast track for AppendAllAttributes
+	}
+}
diff --git a/Engine/Plugins/PCG/Source/PCG/Public/Metadata/Accessors/PCGAttributeAccessorHelpers.h b/Engine/Plugins/PCG/Source/PCG/Public/Metadata/Accessors/PCGAttributeAccessorHelpers.h
--- a/Engine/Plugins/PCG/Source/PCG/Public/Metadata/Accessors/PCGAttributeAccessorHelpers.h
+++ b/Engine/Plugins/PCG/Source/PCG/Public/Metadata/Accessors/PCGAttributeAccessorHelpers.h
@@ -117,6 +117,40 @@ namespace PCGAttributeAccessorHelpers
 	PCG_API TUniquePtr<const IPCGAttributeAccessorKeys> CreateConstKeys(const UPCGData* InData, const FPCGAttributePropertySelector& InSelector);
 	PCG_API TUniquePtr<IPCGAttributeAccessorKeys> CreateKeys(UPCGData* InData, const FPCGAttributePropertySelector& InSelector);
 
+	/** 
+	 * Utility function to extract all the values through the accessor for all the keys.
+	 *
+	 * @param InAccessor  Attribute accessor from which to extract the values.
+	 * @param InKeys      Attribute accessor keys for the values.
+	 * @param OutArray    Output array where the values will be written.
+	 * @param GetFlags    Flags for the get accessor. By default, is broadcast and constructible.
+	 * @return true if it succeeded.
+	 */
+	template <typename T, typename AllocatorType>
+	bool ExtractAllValues(const IPCGAttributeAccessor* InAccessor, const IPCGAttributeAccessorKeys* InKeys, TArray<T, AllocatorType>& OutArray, EPCGAttributeAccessorFlags GetFlags = EPCGAttributeAccessorFlags::AllowBroadcastAndConstructible)
+	{
+		if (!InAccessor || !InKeys)
+		{
+			return false;
+		}
+
+		if (InKeys->GetNum() < 1)
+		{
+			return true;
+		}
+
+		if constexpr (std::is_trivially_copyable_v<T>)
+		{
+			OutArray.SetNumUninitialized(InKeys->GetNum());
+		}
+		else
+		{
+			OutArray.SetNum(InKeys->GetNum());
+		}
+
+		return InAccessor->GetRange<T>(OutArray, 0, *InKeys, GetFlags);
+	}
+
 	/** 
 	 * Utility function to create the accessor and keys and extract all the values for all the keys. Also, automatically resolve the selector for @Last.
 	 *
@@ -135,6 +169,7 @@ namespace PCGAttributeAccessorHelpers
 		const FPCGAttributePropertyInputSelector Selector = InSelector.CopyAndFixLast(InData);
 
 		TUniquePtr<const IPCGAttributeAccessor> Accessor = CreateConstAccessor(InData, Selector, bQuiet);
+		TUniquePtr<const IPCGAttributeAccessorKeys> Keys = CreateConstKeys(InData, Selector);
 		if (!Accessor.IsValid())
 		{
 			if (!bQuiet)
@@ -144,34 +179,20 @@ namespace PCGAttributeAccessorHelpers
 			return false;
 		}
 
-		TUniquePtr<const IPCGAttributeAccessorKeys> Keys = CreateConstKeys(InData, Selector);
 		if (!Keys.IsValid())
 		{
 			PCGLog::Metadata::LogFailToCreateAccessorError(Selector, Context);
 			return false;
 		}
 
-		if (Keys->GetNum() < 1)
-		{
-			return true;
-		}
-			
-		if constexpr (std::is_trivially_copyable_v<T>)
-		{
-			OutArray.SetNumUninitialized(Keys->GetNum());
-		}
-		else
-		{
-			OutArray.SetNum(Keys->GetNum());
-		}
+		const bool bResult = ExtractAllValues<T, AllocatorType>(Accessor.Get(), Keys.Get(), OutArray, GetFlags);
 
-		if (!Accessor->GetRange<T>(OutArray, 0, *Keys, GetFlags))
+		if (!bResult && !bQuiet)
 		{
-			PCGLog::Metadata::LogFailToGetAttributeError<T>(Selector, Accessor.Get(), Context);
-			return false;
+			PCGLog::Metadata::LogFailToGetAttributeError(Selector, Context);
 		}
 
-		return true;
+		return bResult;
 	}
 	
 	/** 
