[PCG][Bughawk] Fix ForEachBlueprintAssetData only finding child classes of a specific native class

diff --git a/Engine/Plugins/PCG/Source/PCGEditor/Private/PCGEditorUtils.cpp b/Engine/Plugins/PCG/Source/PCGEditor/Private/PCGEditorUtils.cpp
--- a/Engine/Plugins/PCG/Source/PCGEditor/Private/PCGEditorUtils.cpp
+++ b/Engine/Plugins/PCG/Source/PCGEditor/Private/PCGEditorUtils.cpp
@@ -17,6 +17,7 @@
 #include "Algo/Transform.h"
 #include "AssetRegistry/AssetData.h"
 #include "AssetRegistry/ARFilter.h"
+#include "AssetRegistry/AssetRegistryHelpers.h"
 #include "AssetRegistry/AssetRegistryModule.h"
 #include "Blueprint/BlueprintSupport.h"
 #include "Filters/SBasicFilterBar.h"
@@ -91,9 +92,24 @@ void PCGEditorUtils::ForEachPCGBlueprintAssetData(TFunctionRef<bool(const FAsset
 	FARFilter Filter;
 	Filter.ClassPaths.Add(UBlueprint::StaticClass()->GetClassPathName());
 	Filter.bRecursiveClasses = true;
-	Filter.TagsAndValues.Add(FBlueprintTags::NativeParentClassPath, UPCGBlueprintBaseElement::GetParentClassName());
+	
+	TSet<FTopLevelAssetPath> NativeBaseClasses;
+	
+	const FAssetRegistryModule& AssetRegistryModule = FModuleManager::LoadModuleChecked<FAssetRegistryModule>(TEXT("AssetRegistry"));
+	
+	// Get all sub-classes of UPCGBlueprintBaseElement (Native & BP)
+	AssetRegistryModule.Get().GetDerivedClassNames({ UPCGBlueprintBaseElement::StaticClass()->GetClassPathName() }, TSet<FTopLevelAssetPath>(), NativeBaseClasses);
+	NativeBaseClasses.Add(UPCGBlueprintBaseElement::StaticClass()->GetClassPathName());
 
-	ForEachAssetData(Filter, InFunc);
+	// Filter based on if asset is subclassing one of those sub-classes
+	AssetRegistryModule.Get().EnumerateAssets(Filter, [&InFunc, &NativeBaseClasses](const FAssetData& InAssetData)
+	{
+		if(UAssetRegistryHelpers::IsAssetDataBlueprintOfClassSet(InAssetData, NativeBaseClasses))
+		{
+			InFunc(InAssetData);
+		}
+		return true;
+	});
 }
 
 void PCGEditorUtils::ForEachPCGSettingsAssetData(TFunctionRef<bool(const FAssetData&)> InFunc)
