[PCG] Indirection Element: Fix ConvertInputsIfNeeded not being called on proper context

diff --git a/Engine/Plugins/PCG/Source/PCG/Private/Elements/PCGIndirectionElement.cpp b/Engine/Plugins/PCG/Source/PCG/Private/Elements/PCGIndirectionElement.cpp
--- a/Engine/Plugins/PCG/Source/PCG/Private/Elements/PCGIndirectionElement.cpp
+++ b/Engine/Plugins/PCG/Source/PCG/Private/Elements/PCGIndirectionElement.cpp
@@ -282,13 +282,16 @@ bool FPCGIndirectionElement::PrepareDataInternal(FPCGContext* InContext) const
 
 		Context->InnerElement = InnerSettings->GetElement();
 		check(Context->InnerElement);
-				
-		if (!Context->InnerElement->ConvertInputsIfNeeded(Context))
+
+		if (!Context->InnerContext)
 		{
-			return false;
+			Context->SetupInnerContext();
 		}
 
-		Context->SetupInnerContext();
+		if (!Context->InnerElement->ConvertInputsIfNeeded(Context->InnerContext))
+		{
+			return false;
+		}
 
 #if WITH_EDITOR
 		// If we have an override, register for dynamic tracking.
