[PCG] Decoupled title and color overriding on graphs from asset exposure to library

#rnx

diff --git a/Engine/Plugins/PCG/Source/PCG/Public/PCGGraph.h b/Engine/Plugins/PCG/Source/PCG/Public/PCGGraph.h
--- a/Engine/Plugins/PCG/Source/PCG/Public/PCGGraph.h
+++ b/Engine/Plugins/PCG/Source/PCG/Public/PCGGraph.h
@@ -156,25 +156,25 @@ public:
 
 #if WITH_EDITORONLY_DATA
 	UPROPERTY(BlueprintReadWrite, EditAnywhere, Category = AssetInfo, AssetRegistrySearchable, meta = (EditCondition = "IsExportToLibraryEnabled", EditConditionHides, PCGNoHash))
-	bool bExposeToLibrary = false;
-
-	UPROPERTY(BlueprintReadWrite, EditAnywhere, Category = AssetInfo, AssetRegistrySearchable, meta = (EditCondition = "AreOverridesEnabled", EditConditionHides, PCGNoHash))
 	bool bOverrideTitle = false;
 
-	UPROPERTY(BlueprintReadWrite, EditAnywhere, Category = AssetInfo, AssetRegistrySearchable, meta = (EditCondition = "IsTemplatePropertyEnabled", EditConditionHides, PCGNoHash))
-	bool bIsTemplate = false;
-
 	/** Override of the title for the subgraph node for this graph. */
-	UPROPERTY(BlueprintReadWrite, EditAnywhere, Category = AssetInfo, AssetRegistrySearchable, meta = (EditCondition = "AreOverridesEnabled && bOverrideTitle", EditConditionHides, PCGNoHash))
+	UPROPERTY(BlueprintReadWrite, EditAnywhere, Category = AssetInfo, AssetRegistrySearchable, meta = (EditCondition = "IsExportToLibraryEnabled && bOverrideTitle", EditConditionHides, PCGNoHash))
 	FText Title;
 
-	UPROPERTY(BlueprintReadWrite, EditAnywhere, Category = AssetInfo, meta = (EditCondition = "AreOverridesEnabled", EditConditionHides, PCGNoHash))
+	UPROPERTY(BlueprintReadWrite, EditAnywhere, Category = AssetInfo, meta = (EditCondition = "IsExportToLibraryEnabled", EditConditionHides, PCGNoHash))
 	bool bOverrideColor = false;
 
 	/** Override of the color for the subgraph node for this graph. */
-	UPROPERTY(BlueprintReadWrite, EditAnywhere, Category = AssetInfo, meta = (EditCondition = "AreOverridesEnabled && bOverrideColor", EditConditionHides, PCGNoHash))
+	UPROPERTY(BlueprintReadWrite, EditAnywhere, Category = AssetInfo, meta = (EditCondition = "IsExportToLibraryEnabled && bOverrideColor", EditConditionHides, PCGNoHash))
 	FLinearColor Color = FLinearColor::White;
 
+	UPROPERTY(BlueprintReadWrite, EditAnywhere, Category = AssetInfo, AssetRegistrySearchable, meta = (EditCondition = "IsTemplatePropertyEnabled", EditConditionHides, PCGNoHash))
+	bool bIsTemplate = false;
+
+	UPROPERTY(BlueprintReadWrite, EditAnywhere, Category = AssetInfo, AssetRegistrySearchable, meta = (EditCondition = "IsExportToLibraryEnabled", EditConditionHides, PCGNoHash))
+	bool bExposeToLibrary = false;
+
 	/** Signifies the graph has compatibility with the annotated tags. Multiple tags may be comma separated. Ex. "PCGTool1, PCGTool2, PCGTool3" */
 	UPROPERTY(BlueprintReadWrite, EditAnywhere, Category = AssetInfo, AssetRegistrySearchable, meta=(EditCondition = "bExposeToLibrary", EditConditionHides, PCGNoHash))
 	FString CompatibleToolTags;
