[PCG] Fix Rotate Around Axis node defining input source 3 as a vector by default, when its unsupported

diff --git a/Engine/Plugins/PCG/Source/PCG/Private/Elements/Metadata/PCGMetadataVectorOpElement.cpp b/Engine/Plugins/PCG/Source/PCG/Private/Elements/Metadata/PCGMetadataVectorOpElement.cpp
--- a/Engine/Plugins/PCG/Source/PCG/Private/Elements/Metadata/PCGMetadataVectorOpElement.cpp
+++ b/Engine/Plugins/PCG/Source/PCG/Private/Elements/Metadata/PCGMetadataVectorOpElement.cpp
@@ -356,6 +356,14 @@ bool UPCGMetadataVectorSettings::IsPinDefaultValueEnabled(const FName PinLabel)
 	return Super::IsPinDefaultValueEnabled(PinLabel);
 }
 
+EPCGMetadataTypes UPCGMetadataVectorSettings::GetPinInitialDefaultValueType(const FName PinLabel) const
+{
+	// Only the angle pin on Rotate Around Axis should be double, otherwise vector.
+	return Operation == EPCGMetadataVectorOperation::RotateAroundAxis && PinLabel == PCGMetadataVectorConstants::AngleLabel
+		? EPCGMetadataTypes::Double
+		: EPCGMetadataTypes::Vector;
+}
+
 FPCGElementPtr UPCGMetadataVectorSettings::CreateElement() const
 {
 	return MakeShared<FPCGMetadataVectorElement>();
diff --git a/Engine/Plugins/PCG/Source/PCG/Public/Elements/Metadata/PCGMetadataVectorOpElement.h b/Engine/Plugins/PCG/Source/PCG/Public/Elements/Metadata/PCGMetadataVectorOpElement.h
--- a/Engine/Plugins/PCG/Source/PCG/Public/Elements/Metadata/PCGMetadataVectorOpElement.h
+++ b/Engine/Plugins/PCG/Source/PCG/Public/Elements/Metadata/PCGMetadataVectorOpElement.h
@@ -65,7 +65,7 @@ protected:
 
 	//~Begin IPCGSettingsDefaultValueProvider interface
 	virtual bool IsPinDefaultValueEnabled(FName PinLabel) const override;
-	virtual EPCGMetadataTypes GetPinInitialDefaultValueType(FName PinLabel) const override { return EPCGMetadataTypes::Vector; }
+	virtual EPCGMetadataTypes GetPinInitialDefaultValueType(FName PinLabel) const override;
 	//~End IPCGSettingsDefaultValueProvider interface
 
 public:
