[PCG] Fix instancing data interface dangerously holding onto scene proxy pointers across GT ticks.

#rnx

diff --git a/Engine/Plugins/Experimental/PCGInterops/PCGFastGeoInterop/Source/PCGFastGeoInterop/Private/Compute/PrimitiveFactories/PCGPrimitiveFactoryFastGeoPISMC.cpp b/Engine/Plugins/Experimental/PCGInterops/PCGFastGeoInterop/Source/PCGFastGeoInterop/Private/Compute/PrimitiveFactories/PCGPrimitiveFactoryFastGeoPISMC.cpp
--- a/Engine/Plugins/Experimental/PCGInterops/PCGFastGeoInterop/Source/PCGFastGeoInterop/Private/Compute/PrimitiveFactories/PCGPrimitiveFactoryFastGeoPISMC.cpp
+++ b/Engine/Plugins/Experimental/PCGInterops/PCGFastGeoInterop/Source/PCGFastGeoInterop/Private/Compute/PrimitiveFactories/PCGPrimitiveFactoryFastGeoPISMC.cpp
@@ -8,6 +8,7 @@
 #include "PCGContext.h"
 #include "PCGModule.h"
 #include "Components/PCGProceduralISMComponentDescriptor.h"
+#include "Compute/PCGComputeCommon.h"
 #include "Engine/World.h"
 
 bool FPCGPrimitiveFactoryFastGeoPISMC::IsRenderStateCreated() const
@@ -94,9 +95,11 @@ bool FPCGPrimitiveFactoryFastGeoPISMC::Create(FPCGContext* InContext)
 	const FString DebugName = Component->GetExecutionState().GetDebugName();
 	TUniquePtr<FFastGeoComponentCluster> ComponentCluster = MakeUnique<FFastGeoComponentCluster>(FastGeo, *FString::Printf(TEXT("FastGeoComponentCluster_%s"), *DebugName));
 
-	MeshBounds.SetNumUninitialized(Descriptors.Num());
+	const int32 NumToCreate = FMath::Min(Descriptors.Num(), PCGComputeConstants::MAX_PRIMITIVE_COMPONENTS_PER_SPAWNER);
 
-	for (int32 DescIndex = 0; DescIndex < Descriptors.Num(); ++DescIndex)
+	MeshBounds.SetNumUninitialized(NumToCreate);
+
+	for (int32 DescIndex = 0; DescIndex < NumToCreate; ++DescIndex)
 	{
 		const FPCGProceduralISMComponentDescriptor& Desc = Descriptors[DescIndex];
 
@@ -117,26 +120,22 @@ bool FPCGPrimitiveFactoryFastGeoPISMC::Create(FPCGContext* InContext)
 	FastGeo->Register();
 
 	int32 ComponentIndex = 0;
-	FastGeo->ForEachComponentCluster([This = this, &ComponentIndex, &Descriptors = Descriptors](FFastGeoComponentCluster& InComponentCluster)
+	FastGeo->ForEachComponentCluster([This=this, &ComponentIndex](FFastGeoComponentCluster& InComponentCluster)
+	{
+		InComponentCluster.ForEachComponent([This, &ComponentIndex](FFastGeoComponent& InComponent)
 		{
-			check(ComponentIndex == 0);
-
-			InComponentCluster.ForEachComponent([This, &ComponentIndex, &Descriptors](FFastGeoComponent& InComponent)
-				{
-					check(Descriptors.IsValidIndex(ComponentIndex));
-					const FPCGProceduralISMComponentDescriptor& Desc = Descriptors[ComponentIndex];
-
-					// If not true, weakptr assignment below will fail.
-					check(InComponent.GetOwnerComponentCluster());
+			check(This->Descriptors.IsValidIndex(ComponentIndex));
+			const FPCGProceduralISMComponentDescriptor& Desc = This->Descriptors[ComponentIndex];
 
-					This->Components.Add(&InComponent);
-					This->InstanceCounts.Add(Desc.NumInstances);
+			// If not true, weakptr assignment below will fail.
+			check(InComponent.GetOwnerComponentCluster());
 
-					++ComponentIndex;
-				});
+			This->Components.Add(&InComponent);
+			This->InstanceCounts.Add(Desc.NumInstances);
 
-			check(ComponentIndex == Descriptors.Num());
+			++ComponentIndex;
 		});
+	});
 
 	UPCGManagedFastGeoContainer* ManagedPrimitives = FPCGContext::NewObject_AnyThread<UPCGManagedFastGeoContainer>(InContext, Component);
 	ManagedPrimitives->FastGeo = FastGeo;
diff --git a/Engine/Plugins/PCG/Source/PCG/Internal/Compute/PCGComputeCommon.h b/Engine/Plugins/PCG/Source/PCG/Internal/Compute/PCGComputeCommon.h
--- a/Engine/Plugins/PCG/Source/PCG/Internal/Compute/PCGComputeCommon.h
+++ b/Engine/Plugins/PCG/Source/PCG/Internal/Compute/PCGComputeCommon.h
@@ -90,7 +90,7 @@ namespace PCGComputeConstants
 	constexpr uint32 AttributeStrideMask = 0xFF;
 	constexpr uint32 AttributeAllocatedMask = 1 << 31;
 
-	constexpr uint32 MAX_PRIMITIVE_COMPONENTS_PER_SPAWNER = 256;
+	constexpr int32 MAX_PRIMITIVE_COMPONENTS_PER_SPAWNER = 256;
 
 	constexpr uint32 THREAD_GROUP_SIZE = 64;
 
diff --git a/Engine/Plugins/PCG/Source/PCG/Private/Compute/DataInterfaces/PCGInstanceDataInterface.cpp b/Engine/Plugins/PCG/Source/PCG/Private/Compute/DataInterfaces/PCGInstanceDataInterface.cpp
--- a/Engine/Plugins/PCG/Source/PCG/Private/Compute/DataInterfaces/PCGInstanceDataInterface.cpp
+++ b/Engine/Plugins/PCG/Source/PCG/Private/Compute/DataInterfaces/PCGInstanceDataInterface.cpp
@@ -155,44 +155,37 @@ bool UPCGInstanceDataProvider::PrepareForExecute_GameThread(UPCGDataBinding* InB
 	}
 
 	FPCGSpawnerPrimitives* FoundPrimitives = Binding->MeshSpawnersToPrimitives.Find(ProducerKernel);
-
 	if (!FoundPrimitives)
 	{
 		return true;
 	}
 
-	TSharedPtr<IPCGRuntimePrimitiveFactory>& Factory = FoundPrimitives->PrimitiveFactory;
-	if (!Factory)
+	PrimitiveFactory = FoundPrimitives->PrimitiveFactory;
+	if (!PrimitiveFactory)
 	{
+		// Factory can be null if no primitives or no instances.
 		return true;
 	}
 
-	if (!Factory->IsRenderStateCreated())
+	// The factories enforce this limit, to simplify the logic here.
+	if (!ensure(PrimitiveFactory->GetNumPrimitives() <= PCGComputeConstants::MAX_PRIMITIVE_COMPONENTS_PER_SPAWNER))
 	{
-		UE_LOG(LogPCG, Verbose, TEXT("FPCGInstanceDataProvider: One or more scene proxies were not ready. Will try again on the next tick."));
-		return false;
+		UE_LOG(LogPCG, Error, TEXT("FPCGInstanceDataProvider: %d primitives provided which is more than the maximum limit (%d), instances will not be created."),
+			PrimitiveFactory->GetNumPrimitives(), PCGComputeConstants::MAX_PRIMITIVE_COMPONENTS_PER_SPAWNER);
+		return true;
 	}
 
-	NumInstancesAllPrimitives = 0;
-
-	int32 NumPrimitives = Factory->GetNumPrimitives();
-	if (NumPrimitives > PCGComputeConstants::MAX_PRIMITIVE_COMPONENTS_PER_SPAWNER)
+	if (!PrimitiveFactory->IsRenderStateCreated())
 	{
-		UE_LOG(LogPCG, Error, TEXT("FPCGInstanceDataProvider: %d primitives provided which is more than the maximum limit (%d), some primitives will not be created."),
-			NumPrimitives, PCGComputeConstants::MAX_PRIMITIVE_COMPONENTS_PER_SPAWNER);
-		NumPrimitives = PCGComputeConstants::MAX_PRIMITIVE_COMPONENTS_PER_SPAWNER;
+		UE_LOG(LogPCG, Verbose, TEXT("FPCGInstanceDataProvider: One or more scene proxies were not ready. Will try again on the next tick."));
+		return false;
 	}
 
-	Primitives.Reserve(NumPrimitives);
+	NumInstancesAllPrimitives = 0;
 
-	for (int32 PrimitiveIndex = 0; PrimitiveIndex < NumPrimitives; ++PrimitiveIndex)
+	for (int32 Index = 0; Index < PrimitiveFactory->GetNumPrimitives(); ++Index)
 	{
-		FPrimitiveSceneProxy* SceneProxy = Factory->GetSceneProxy(PrimitiveIndex);
-		check(SceneProxy);
-
-		const int32 NumInstances = Factory->GetNumInstances(PrimitiveIndex);
-		Primitives.Add(FPCGPrimitiveAndNumInstances(SceneProxy, NumInstances));
-		NumInstancesAllPrimitives += NumInstances;
+		NumInstancesAllPrimitives += PrimitiveFactory->GetNumInstances(Index);
 	}
 
 	NumCustomFloatsPerInstance = FoundPrimitives->NumCustomFloats;
@@ -216,36 +209,45 @@ FComputeDataProviderRenderProxy* UPCGInstanceDataProvider::GetRenderProxy()
 {
 	TRACE_CPUPROFILER_EVENT_SCOPE(UPCGInstanceDataProvider::GetRenderProxy);
 
-	return new FPCGInstanceDataProviderProxy(Primitives, NumInstancesAllPrimitives, NumCustomFloatsPerInstance, MakeWeakObjectPtr(this), ContextHandle);
+	FPCGInstanceDataProviderProxy::FParams Params =
+	{
+		.PrimitiveFactory = PrimitiveFactory,
+		.NumInstancesAllPrimitives = NumInstancesAllPrimitives,
+		.NumCustomFloatsPerInstance = NumCustomFloatsPerInstance,
+		.DataProvider = MakeWeakObjectPtr(this),
+		.ContextHandle = ContextHandle,
+	};
+
+	return new FPCGInstanceDataProviderProxy(Params);
 }
 
 void UPCGInstanceDataProvider::Reset()
 {
 	Super::Reset();
 
-	Primitives.Empty(Primitives.Num());
+	PrimitiveFactory.Reset();
 	NumInstancesAllPrimitives = 0;
 	NumCustomFloatsPerInstance = 0;
 	bWroteInstances = false;
 	ContextHandle.Reset();
 }
 
-FPCGInstanceDataProviderProxy::FPCGInstanceDataProviderProxy(TArray<FPCGPrimitiveAndNumInstances> InPrimitives, uint32 InNumInstancesAllPrimitives, uint32 InNumCustomFloatsPerInstance, TWeakObjectPtr<UPCGInstanceDataProvider> InDataProvider, TWeakPtr<FPCGContextHandle> InContextHandle)
-	: Primitives(MoveTemp(InPrimitives))
-	, NumInstancesAllPrimitives(InNumInstancesAllPrimitives)
-	, NumCustomFloatsPerInstance(InNumCustomFloatsPerInstance)
-	, DataProvider(InDataProvider)
-	, ContextHandle(InContextHandle)
+FPCGInstanceDataProviderProxy::FPCGInstanceDataProviderProxy(const FParams& InParams)
+	: PrimitiveFactory(InParams.PrimitiveFactory)
+	, NumInstancesAllPrimitives(InParams.NumInstancesAllPrimitives)
+	, NumCustomFloatsPerInstance(InParams.NumCustomFloatsPerInstance)
+	, DataProvider(InParams.DataProvider)
+	, ContextHandle(InParams.ContextHandle)
 {
 	bIsValid = true;
 
-	if (ensure(InDataProvider.IsValid()))
+	if (ensure(DataProvider.IsValid()))
 	{
-		DataProviderGeneration = InDataProvider->GenerationCounter;
+		DataProviderGeneration = DataProvider->GenerationCounter;
 
 		if (NumInstancesAllPrimitives == 0)
 		{
-			InDataProvider->bWroteInstances = true;
+			DataProvider->bWroteInstances = true;
 		}
 	}
 	else
@@ -254,11 +256,22 @@ FPCGInstanceDataProviderProxy::FPCGInstanceDataProviderProxy(TArray<FPCGPrimitiv
 		bIsValid = false;
 	}
 
-	if (Algo::AnyOf(Primitives, [](const FPCGPrimitiveAndNumInstances& InPrimitive) { return !InPrimitive.Get<0>(); }))
+	if (!PrimitiveFactory)
 	{
-		UE_LOG(LogPCG, Warning, TEXT("FPCGInstanceDataProviderProxy: One or more components were invalid, proxy is invalid, compute graph will not execute."));
 		bIsValid = false;
 	}
+	else
+	{
+		for (int32 Index = 0; Index < PrimitiveFactory->GetNumPrimitives(); ++Index)
+		{
+			if (!PrimitiveFactory->GetSceneProxy(Index))
+			{
+				UE_LOG(LogPCG, Warning, TEXT("FPCGInstanceDataProviderProxy: One or more component proxies were invalid, compute graph will not execute."));
+				bIsValid = false;
+				break;
+			}
+		}
+	}
 }
 
 bool FPCGInstanceDataProviderProxy::IsValid(FValidationData const& InValidationData) const
@@ -278,7 +291,7 @@ void FPCGInstanceDataProviderProxy::AllocateResources(FRDGBuilder& GraphBuilder,
 	InstanceCustomFloatDataSRV = GraphBuilder.CreateSRV(InstanceCustomFloatData);
 	InstanceCustomFloatDataUAV = GraphBuilder.CreateUAV(InstanceCustomFloatData);
 
-	const int32 NumCountersRequired = FMath::Max(1, Primitives.Num());
+	const int32 NumCountersRequired = FMath::Max(1, PrimitiveFactory->GetNumPrimitives());
 	WriteCounters = GraphBuilder.CreateBuffer(FRDGBufferDesc::CreateStructuredDesc(sizeof(uint32), NumCountersRequired), TEXT("PCGWriteCounters"));
 	WriteCountersSRV = GraphBuilder.CreateSRV(WriteCounters);
 	WriteCountersUAV = GraphBuilder.CreateUAV(WriteCounters);
@@ -293,7 +306,7 @@ void FPCGInstanceDataProviderProxy::GatherDispatchData(FDispatchData const& InDi
 	for (int32 InvocationIndex = 0; InvocationIndex < ParameterArray.Num(); ++InvocationIndex)
 	{
 		FParameters& Parameters = ParameterArray[InvocationIndex];
-		Parameters.NumPrimitives = Primitives.Num();
+		Parameters.NumPrimitives = PrimitiveFactory->GetNumPrimitives();
 		Parameters.NumInstancesAllPrimitives = NumInstancesAllPrimitives;
 		Parameters.NumCustomFloatsPerInstance = NumCustomFloatsPerInstance;
 		Parameters.InstanceData = InstanceDataUAV;
@@ -301,9 +314,9 @@ void FPCGInstanceDataProviderProxy::GatherDispatchData(FDispatchData const& InDi
 		Parameters.WriteCounters = WriteCountersUAV;
 
 		uint32 CumulativeInstanceCount = 0;
-		for (int32 Index = 0; Index < Primitives.Num(); ++Index)
+		for (int32 Index = 0; Index < PrimitiveFactory->GetNumPrimitives(); ++Index)
 		{
-			const uint32 PrimitiveInstanceCount = (Primitives[Index].Get<1>() > 0) ? static_cast<uint32>(Primitives[Index].Get<1>()) : 0u;
+			const uint32 PrimitiveInstanceCount = PrimitiveFactory->GetNumInstances(Index);
 
 			Parameters.PrimitiveParameters[Index] = FUintVector4(
 				PrimitiveInstanceCount,
@@ -359,18 +372,18 @@ void FPCGInstanceDataProviderProxy::PostSubmit(FComputeContext& Context) const
 			return false;
 		}
 
-		for (int PrimitiveIndex = 0; PrimitiveIndex < DataProviderInner->Primitives.Num(); ++PrimitiveIndex)
-		{
-			FPCGPrimitiveAndNumInstances& PrimitiveAndInstanceCount = DataProviderInner->Primitives[PrimitiveIndex];
+		TSharedPtr<IPCGRuntimePrimitiveFactory> Factory = DataProviderInner->PrimitiveFactory;
 
-			const int32 PrimitiveNumInstances = PrimitiveAndInstanceCount.Get<1>();
+		for (int32 PrimitiveIndex = 0; PrimitiveIndex < Factory->GetNumPrimitives(); ++PrimitiveIndex)
+		{
+			const int32 PrimitiveNumInstances = Factory->GetNumInstances(PrimitiveIndex);
 			if (PrimitiveNumInstances <= 0)
 			{
 				UE_LOG(LogPCG, Warning, TEXT("Primitive with 0 instances encountered during instancing, should not happen."));
 				continue;
 			}
 
-			FPrimitiveSceneProxy* SceneProxy = PrimitiveAndInstanceCount.Get<0>();
+			FPrimitiveSceneProxy* SceneProxy = Factory->GetSceneProxy(PrimitiveIndex);
 			check(SceneProxy);
 
 			if (SceneProxy->GetPrimitiveSceneInfo()->GetInstanceSceneDataOffset() == -1)
diff --git a/Engine/Plugins/PCG/Source/PCG/Private/Compute/DataInterfaces/PCGInstanceDataInterface.h b/Engine/Plugins/PCG/Source/PCG/Private/Compute/DataInterfaces/PCGInstanceDataInterface.h
--- a/Engine/Plugins/PCG/Source/PCG/Private/Compute/DataInterfaces/PCGInstanceDataInterface.h
+++ b/Engine/Plugins/PCG/Source/PCG/Private/Compute/DataInterfaces/PCGInstanceDataInterface.h
@@ -10,12 +10,11 @@
 
 class FPCGInstanceDataInterfaceParameters;
 class FPrimitiveSceneProxy;
+class IPCGRuntimePrimitiveFactory;
 class UPCGSettings;
 class UComputeKernel;
 struct FPCGContextHandle;
 
-using FPCGPrimitiveAndNumInstances = TPair<FPrimitiveSceneProxy*, int32>;
-
 /** Data Interface for writing instance data (transform and custom floats) to an intermediate buffer, which is then injected into the GPU Scene. */
 UCLASS(ClassGroup = (Procedural))
 class UPCGInstanceDataInterface : public UPCGComputeDataInterface
@@ -57,7 +56,7 @@ public:
 	virtual void Reset() override;
 	//~ End UComputeDataProvider Interface
 
-	TArray<FPCGPrimitiveAndNumInstances> Primitives;
+	TSharedPtr<IPCGRuntimePrimitiveFactory> PrimitiveFactory;
 
 	uint32 NumInstancesAllPrimitives = 0;
 
@@ -72,7 +71,16 @@ public:
 class FPCGInstanceDataProviderProxy : public FComputeDataProviderRenderProxy
 {
 public:
-	FPCGInstanceDataProviderProxy(TArray<FPCGPrimitiveAndNumInstances> InPrimitives, uint32 InNumInstancesAllPrimitives, uint32 InNumCustomFloatsPerInstance, TWeakObjectPtr<UPCGInstanceDataProvider> InDataProvider, TWeakPtr<FPCGContextHandle> InContextHandle);
+	struct FParams
+	{
+		TSharedPtr<IPCGRuntimePrimitiveFactory> PrimitiveFactory;
+		uint32 NumInstancesAllPrimitives = 0;
+		uint32 NumCustomFloatsPerInstance = 0;
+		TWeakObjectPtr<UPCGInstanceDataProvider> DataProvider;
+		TWeakPtr<FPCGContextHandle> ContextHandle;
+	};
+
+	FPCGInstanceDataProviderProxy(const FParams& InParams);
 
 	//~ Begin FComputeDataProviderRenderProxy Interface
 	bool IsValid(FValidationData const& InValidationData) const override;
@@ -86,7 +94,7 @@ protected:
 
 	bool bIsValid = false;
 
-	TArray<FPCGPrimitiveAndNumInstances> Primitives;
+	TSharedPtr<IPCGRuntimePrimitiveFactory> PrimitiveFactory;
 
 	/** Instance transforms of all instances across all primitives. Three float4s per instance. */
 	FRDGBufferRef InstanceData = nullptr;
diff --git a/Engine/Plugins/PCG/Source/PCG/Private/Compute/PrimitiveFactories/PCGPrimitiveFactoryISKMC.cpp b/Engine/Plugins/PCG/Source/PCG/Private/Compute/PrimitiveFactories/PCGPrimitiveFactoryISKMC.cpp
--- a/Engine/Plugins/PCG/Source/PCG/Private/Compute/PrimitiveFactories/PCGPrimitiveFactoryISKMC.cpp
+++ b/Engine/Plugins/PCG/Source/PCG/Private/Compute/PrimitiveFactories/PCGPrimitiveFactoryISKMC.cpp
@@ -6,6 +6,7 @@
 #include "PCGContext.h"
 #include "PCGManagedResource.h"
 #include "Components/InstancedSkinnedMeshComponent.h"
+#include "Compute/PCGComputeCommon.h"
 #include "Helpers/PCGActorHelpers.h"
 
 #include "Components/PrimitiveComponent.h"
@@ -61,9 +62,11 @@ bool FPCGPrimitiveFactoryISKMC::Create(FPCGContext* InContext)
 		Params.Descriptor.Mobility = SceneComponent->Mobility;
 	}
 
+	const int32 NumToCreate = FMath::Min(Descriptors.Num(), PCGComputeConstants::MAX_PRIMITIVE_COMPONENTS_PER_SPAWNER);
+
 	bool bAnyCreated = false;
 
-	while (NumPrimitivesProcessed < Descriptors.Num())
+	while (NumPrimitivesProcessed < NumToCreate)
 	{
 		const FPCGSoftSkinnedMeshComponentDescriptor& Desc = Descriptors[NumPrimitivesProcessed];
 
@@ -81,13 +84,13 @@ bool FPCGPrimitiveFactoryISKMC::Create(FPCGContext* InContext)
 
 		++NumPrimitivesProcessed;
 
-		if (NumPrimitivesProcessed < Descriptors.Num() && InContext->AsyncState.ShouldStop())
+		if (NumPrimitivesProcessed < NumToCreate && InContext->AsyncState.ShouldStop())
 		{
 			return false;
 		}
 	}
 
-	return NumPrimitivesProcessed == Descriptors.Num();
+	return NumPrimitivesProcessed == NumToCreate;
 }
 
 FBox FPCGPrimitiveFactoryISKMC::GetMeshBounds(int32 InPrimitiveIndex) const
diff --git a/Engine/Plugins/PCG/Source/PCG/Private/Compute/PrimitiveFactories/PCGPrimitiveFactoryPISMC.cpp b/Engine/Plugins/PCG/Source/PCG/Private/Compute/PrimitiveFactories/PCGPrimitiveFactoryPISMC.cpp
--- a/Engine/Plugins/PCG/Source/PCG/Private/Compute/PrimitiveFactories/PCGPrimitiveFactoryPISMC.cpp
+++ b/Engine/Plugins/PCG/Source/PCG/Private/Compute/PrimitiveFactories/PCGPrimitiveFactoryPISMC.cpp
@@ -5,6 +5,7 @@
 #include "PCGComponent.h"
 #include "PCGContext.h"
 #include "Components/PCGProceduralISMComponent.h"
+#include "Compute/PCGComputeCommon.h"
 
 #include "Engine/StaticMesh.h"
 #include "Components/PrimitiveComponent.h"
@@ -53,7 +54,9 @@ bool FPCGPrimitiveFactoryPISMC::Create(FPCGContext* InContext)
 	FPCGProceduralISMCBuilderParameters Params;
 	Params.bAllowDescriptorChanges = false;
 
-	while (NumPrimitivesProcessed < Descriptors.Num())
+	const int32 NumToCreate = FMath::Min(Descriptors.Num(), PCGComputeConstants::MAX_PRIMITIVE_COMPONENTS_PER_SPAWNER);
+
+	while (NumPrimitivesProcessed < NumToCreate)
 	{
 		const FPCGProceduralISMComponentDescriptor& Desc = Descriptors[NumPrimitivesProcessed];
 
@@ -70,13 +73,13 @@ bool FPCGPrimitiveFactoryPISMC::Create(FPCGContext* InContext)
 
 		++NumPrimitivesProcessed;
 
-		if (NumPrimitivesProcessed < Descriptors.Num() && InContext->AsyncState.ShouldStop())
+		if (NumPrimitivesProcessed < NumToCreate && InContext->AsyncState.ShouldStop())
 		{
 			return false;
 		}
 	}
 
-	return NumPrimitivesProcessed == Descriptors.Num();
+	return NumPrimitivesProcessed == NumToCreate;
 }
 
 FBox FPCGPrimitiveFactoryPISMC::GetMeshBounds(int32 InPrimitiveIndex) const
