[PCG] Enable runtime gen scheduler generation source change detection by default.

#rnx

diff --git a/Engine/Plugins/PCG/Source/PCG/Private/RuntimeGen/PCGRuntimeGenScheduler.cpp b/Engine/Plugins/PCG/Source/PCG/Private/RuntimeGen/PCGRuntimeGenScheduler.cpp
--- a/Engine/Plugins/PCG/Source/PCG/Private/RuntimeGen/PCGRuntimeGenScheduler.cpp
+++ b/Engine/Plugins/PCG/Source/PCG/Private/RuntimeGen/PCGRuntimeGenScheduler.cpp
@@ -54,7 +54,7 @@ namespace PCGRuntimeGenSchedulerHelpers
 
 	static TAutoConsoleVariable<bool> CVarRuntimeGenerationEnableChangeDetection(
 		TEXT("pcg.RuntimeGeneration.EnableChangeDetection"),
-		false,
+		true,
 		TEXT("Skips execution of scheduling loop if generation sources and other world state have not changed."));
 
 	static TAutoConsoleVariable<int32> CVarNumGeneratingComponentsAtSameTime(
