[PCG] Fix "Primitive not allocated in GPU scene" errors when editor launched in background/unfocused.

#rnx

diff --git a/Engine/Plugins/Experimental/PCGInterops/PCGFastGeoInterop/Source/PCGFastGeoInterop/Private/Compute/PrimitiveFactories/PCGPrimitiveFactoryFastGeoPISMC.cpp b/Engine/Plugins/Experimental/PCGInterops/PCGFastGeoInterop/Source/PCGFastGeoInterop/Private/Compute/PrimitiveFactories/PCGPrimitiveFactoryFastGeoPISMC.cpp
--- a/Engine/Plugins/Experimental/PCGInterops/PCGFastGeoInterop/Source/PCGFastGeoInterop/Private/Compute/PrimitiveFactories/PCGPrimitiveFactoryFastGeoPISMC.cpp
+++ b/Engine/Plugins/Experimental/PCGInterops/PCGFastGeoInterop/Source/PCGFastGeoInterop/Private/Compute/PrimitiveFactories/PCGPrimitiveFactoryFastGeoPISMC.cpp
@@ -7,6 +7,7 @@
 #include "PCGComponent.h"
 #include "PCGContext.h"
 #include "PCGModule.h"
+#include "PrimitiveSceneInfo.h"
 #include "Components/PCGProceduralISMComponentDescriptor.h"
 #include "Compute/PCGComputeCommon.h"
 #include "Engine/World.h"
@@ -39,7 +40,8 @@ bool FPCGPrimitiveFactoryFastGeoPISMC::IsRenderStateCreated() const
 
 	for (int32 PrimitiveIndex = 0; PrimitiveIndex < GetNumPrimitives(); ++PrimitiveIndex)
 	{
-		if (GetSceneProxy(PrimitiveIndex) == nullptr)
+		FPrimitiveSceneProxy* SceneProxy = GetSceneProxy(PrimitiveIndex);
+		if (!SceneProxy || !SceneProxy->GetPrimitiveSceneInfo() || SceneProxy->GetPrimitiveSceneInfo()->GetInstanceSceneDataOffset() == -1)
 		{
 			return false;
 		}
diff --git a/Engine/Plugins/PCG/Source/PCG/Private/Compute/DataInterfaces/PCGInstanceDataInterface.cpp b/Engine/Plugins/PCG/Source/PCG/Private/Compute/DataInterfaces/PCGInstanceDataInterface.cpp
--- a/Engine/Plugins/PCG/Source/PCG/Private/Compute/DataInterfaces/PCGInstanceDataInterface.cpp
+++ b/Engine/Plugins/PCG/Source/PCG/Private/Compute/DataInterfaces/PCGInstanceDataInterface.cpp
@@ -374,6 +374,22 @@ void FPCGInstanceDataProviderProxy::PostSubmit(FComputeContext& Context) const
 
 		TSharedPtr<IPCGRuntimePrimitiveFactory> Factory = DataProviderInner->PrimitiveFactory;
 
+#if WITH_EDITOR
+		// When launching Editor in background/unfocused, it appears the render setup does not make progress and can be not ready here.
+		// Does not seem to repro in standalone build.
+		for (int32 PrimitiveIndex = 0; PrimitiveIndex < Factory->GetNumPrimitives(); ++PrimitiveIndex)
+		{
+			FPrimitiveSceneProxy* SceneProxy = Factory->GetSceneProxy(PrimitiveIndex);
+			check(SceneProxy);
+
+			if (SceneProxy->GetPrimitiveSceneInfo()->GetInstanceSceneDataOffset() == -1)
+			{
+				UE_LOG(LogPCG, Verbose, TEXT("Primitive instances not allocated in GPU scene, retrying next frame."));
+				return true;
+			}
+		}
+#endif // WITH_EDITOR
+
 		for (int32 PrimitiveIndex = 0; PrimitiveIndex < Factory->GetNumPrimitives(); ++PrimitiveIndex)
 		{
 			const int32 PrimitiveNumInstances = Factory->GetNumInstances(PrimitiveIndex);
@@ -385,12 +401,7 @@ void FPCGInstanceDataProviderProxy::PostSubmit(FComputeContext& Context) const
 
 			FPrimitiveSceneProxy* SceneProxy = Factory->GetSceneProxy(PrimitiveIndex);
 			check(SceneProxy);
-
-			if (SceneProxy->GetPrimitiveSceneInfo()->GetInstanceSceneDataOffset() == -1)
-			{
-				UE_LOG(LogPCG, Error, TEXT("Primitive not allocated in GPU scene, should not happen."));
-				continue;
-			}
+			check(SceneProxy->GetPrimitiveSceneInfo()->GetInstanceSceneDataOffset() != -1);
 
 #if UE_BUILD_SHIPPING
 			const bool bTriggerCapture = false;
diff --git a/Engine/Plugins/PCG/Source/PCG/Private/Compute/PrimitiveFactories/PCGPrimitiveFactoryISKMC.cpp b/Engine/Plugins/PCG/Source/PCG/Private/Compute/PrimitiveFactories/PCGPrimitiveFactoryISKMC.cpp
--- a/Engine/Plugins/PCG/Source/PCG/Private/Compute/PrimitiveFactories/PCGPrimitiveFactoryISKMC.cpp
+++ b/Engine/Plugins/PCG/Source/PCG/Private/Compute/PrimitiveFactories/PCGPrimitiveFactoryISKMC.cpp
@@ -9,6 +9,7 @@
 #include "Compute/PCGComputeCommon.h"
 #include "Helpers/PCGActorHelpers.h"
 
+#include "PrimitiveSceneInfo.h"
 #include "Components/PrimitiveComponent.h"
 #include "Engine/SkinnedAsset.h"
 
@@ -31,9 +32,10 @@ int32 FPCGPrimitiveFactoryISKMC::GetNumInstances(int32 InPrimitiveIndex) const
 
 bool FPCGPrimitiveFactoryISKMC::IsRenderStateCreated() const
 {
-	for (const TWeakObjectPtr<UPrimitiveComponent>& Component : Components)
+	for (int32 Index = 0; Index < GetNumPrimitives(); ++Index)
 	{
-		if (Component.IsValid() && !Component->GetSceneProxy())
+		FPrimitiveSceneProxy* SceneProxy = GetSceneProxy(Index);
+		if (!SceneProxy || !SceneProxy->GetPrimitiveSceneInfo() || SceneProxy->GetPrimitiveSceneInfo()->GetInstanceSceneDataOffset() == -1)
 		{
 			return false;
 		}
diff --git a/Engine/Plugins/PCG/Source/PCG/Private/Compute/PrimitiveFactories/PCGPrimitiveFactoryPISMC.cpp b/Engine/Plugins/PCG/Source/PCG/Private/Compute/PrimitiveFactories/PCGPrimitiveFactoryPISMC.cpp
--- a/Engine/Plugins/PCG/Source/PCG/Private/Compute/PrimitiveFactories/PCGPrimitiveFactoryPISMC.cpp
+++ b/Engine/Plugins/PCG/Source/PCG/Private/Compute/PrimitiveFactories/PCGPrimitiveFactoryPISMC.cpp
@@ -7,6 +7,7 @@
 #include "Components/PCGProceduralISMComponent.h"
 #include "Compute/PCGComputeCommon.h"
 
+#include "PrimitiveSceneInfo.h"
 #include "Engine/StaticMesh.h"
 #include "Components/PrimitiveComponent.h"
 
@@ -29,9 +30,10 @@ int32 FPCGPrimitiveFactoryPISMC::GetNumInstances(int32 InPrimitiveIndex) const
 
 bool FPCGPrimitiveFactoryPISMC::IsRenderStateCreated() const
 {
-	for (const TWeakObjectPtr<UPrimitiveComponent>& Component : Components)
+	for (int32 Index = 0; Index < GetNumPrimitives(); ++Index)
 	{
-		if (Component.IsValid() && !Component->GetSceneProxy())
+		FPrimitiveSceneProxy* SceneProxy = GetSceneProxy(Index);
+		if (!SceneProxy || !SceneProxy->GetPrimitiveSceneInfo() || SceneProxy->GetPrimitiveSceneInfo()->GetInstanceSceneDataOffset() == -1)
 		{
 			return false;
 		}
