[PCG] Create spline in linear mode shouldn't try to create attributes/warn on attributes not found.

#rnx

diff --git a/Engine/Plugins/PCG/Source/PCG/Private/Elements/PCGCreateSpline.cpp b/Engine/Plugins/PCG/Source/PCG/Private/Elements/PCGCreateSpline.cpp
--- a/Engine/Plugins/PCG/Source/PCG/Private/Elements/PCGCreateSpline.cpp
+++ b/Engine/Plugins/PCG/Source/PCG/Private/Elements/PCGCreateSpline.cpp
@@ -105,7 +105,7 @@ bool FPCGCreateSplineElement::ExecuteInternal(FPCGContext* Context) const
 		const FPCGMetadataAttribute<FVector>* ArriveTangentAttribute = nullptr;
 		const FPCGMetadataAttribute<FVector>* LeaveTangentAttribute = nullptr;
 
-		if (Settings->bApplyCustomTangents)
+		if (Settings->bApplyCustomTangents && !Settings->bLinear)
 		{
 			const UPCGMetadata* PointMetadata = PointData->ConstMetadata();
 			check(PointMetadata);
