[PCG] Save 1.3mb of memory when graph cache is disabled.

#rnx

diff --git a/Engine/Plugins/PCG/Source/PCG/Private/Graph/PCGGraphCache.cpp b/Engine/Plugins/PCG/Source/PCG/Private/Graph/PCGGraphCache.cpp
--- a/Engine/Plugins/PCG/Source/PCG/Private/Graph/PCGGraphCache.cpp
+++ b/Engine/Plugins/PCG/Source/PCG/Private/Graph/PCGGraphCache.cpp
@@ -43,7 +43,7 @@ static TAutoConsoleVariable<bool> CVarValidateElementToCacheEntryKeys(
 static const int32 GPCGGraphCacheInitialCapacity = 65536;
 
 FPCGGraphCache::FPCGGraphCache()
-	: CacheData(GPCGGraphCacheInitialCapacity)
+	: CacheData(CVarCacheEnabled.GetValueOnGameThread() ? GPCGGraphCacheInitialCapacity : 0)
 {
 }
 
@@ -310,9 +310,10 @@ void FPCGGraphCache::GrowCache_Unsafe()
 	}
 
 	// Resize and flush
-	ClearCacheInternal(CacheData.Num() * 2, /*bClearMemory=*/false);
+	const int32 NewSize = (CacheData.Num() > 0) ? (2 * CacheData.Num()) : GPCGGraphCacheInitialCapacity;
+	ClearCacheInternal(NewSize, /*bClearMemory=*/false);
 
-	UE_LOG(LogPCG, Log, TEXT("Graph cache doubled in capacity to %d entries."), CacheData.Max());
+	UE_LOG(LogPCG, Log, TEXT("Graph cache capacity increased to %d entries."), CacheData.Max());
 
 	// Copy back. Restore temporal order.
 	for (TLruCache<FPCGCacheEntryKey, FPCGDataCollection>::TIterator It(CacheDataCopy); It; ++It)
