[PCG] Demote output pin limit from check() to warning.

#rnx

diff --git a/Engine/Plugins/PCG/Source/PCG/Private/Graph/PCGGraphCompiler.cpp b/Engine/Plugins/PCG/Source/PCG/Private/Graph/PCGGraphCompiler.cpp
--- a/Engine/Plugins/PCG/Source/PCG/Private/Graph/PCGGraphCompiler.cpp
+++ b/Engine/Plugins/PCG/Source/PCG/Private/Graph/PCGGraphCompiler.cpp
@@ -1110,6 +1110,13 @@ void FPCGGraphCompiler::CalculateDynamicActivePinDependencies(FPCGTaskId InTaskI
 		return;
 	}
 
+	auto LogTooManyOutputPinsWarning = [](const UPCGNode* InUpstreamNode)
+	{
+		UE_LOG(LogPCG, Warning, TEXT("Number of output pins (%d) on node %s exceeded maximum limit (%d) for dynamic node culling. "
+			"Dynamic culling will not function for this node's output pins."),
+			InUpstreamNode->GetOutputPins().Num(), *InUpstreamNode->GetName(), PCGPinIdHelpers::MaxOutputPins);
+	};
+
 	const UPCGNode* Node = InOutCompiledTasks[InTaskId].Node;
 
 	// First if there are any required pins, build the pin dependencies from these. We have an array of inputs, each one can
@@ -1137,6 +1144,12 @@ void FPCGGraphCompiler::CalculateDynamicActivePinDependencies(FPCGTaskId InTaskI
 			continue;
 		}
 
+		if (UpstreamNode->GetOutputPins().Num() > PCGPinIdHelpers::MaxOutputPins)
+		{
+			LogTooManyOutputPinsWarning(UpstreamNode);
+			continue;
+		}
+
 		const int PinIndex = UpstreamNode->GetOutputPins().IndexOfByPredicate([&Input](const UPCGPin* InPin)
 		{
 			return InPin->Properties == Input.UpstreamPin;
@@ -1189,6 +1202,12 @@ void FPCGGraphCompiler::CalculateDynamicActivePinDependencies(FPCGTaskId InTaskI
 					continue;
 				}
 
+				if (UpstreamNode->GetOutputPins().Num() > PCGPinIdHelpers::MaxOutputPins)
+				{
+					LogTooManyOutputPinsWarning(UpstreamNode);
+					continue;
+				}
+
 				const int PinIndex = UpstreamNode->GetOutputPins().IndexOfByPredicate([&Input](const UPCGPin* InPin)
 				{
 					return InPin->Properties == Input.UpstreamPin.GetValue();
