[PCG] Add PCG LLM tag and add scopes to various entry points into PCG code.

GT entry points (subsystem events), tickers, render commands, async tasks (including ParallelFor's). ParallelFor should be applied with caution - thousands of tag scopes running on a per-element basis is inefficient. I omitted the scope from one place with a comment, for the others they are per-attribute at most and should be ok.

Could not find any significant untagged PCG allocations with this instrumentation.

[RN] [PCG] Added LLM tags for tracking memory allocations

diff --git a/Engine/Plugins/PCG/Source/PCG/Private/Compute/DataInterfaces/Elements/PCGCopyPointsDataInterface.cpp b/Engine/Plugins/PCG/Source/PCG/Private/Compute/DataInterfaces/Elements/PCGCopyPointsDataInterface.cpp
--- a/Engine/Plugins/PCG/Source/PCG/Private/Compute/DataInterfaces/Elements/PCGCopyPointsDataInterface.cpp
+++ b/Engine/Plugins/PCG/Source/PCG/Private/Compute/DataInterfaces/Elements/PCGCopyPointsDataInterface.cpp
@@ -297,6 +297,7 @@ void FPCGCopyPointsDataProviderProxy::GatherPermutations(FPermutationData& InOut
 void FPCGCopyPointsDataProviderProxy::AllocateResources(FRDGBuilder& GraphBuilder, FAllocationData const& InAllocationData)
 {
 	TRACE_CPUPROFILER_EVENT_SCOPE(FPCGCopyPointsDataProviderProxy::AllocateResources);
+	LLM_SCOPE_BYTAG(PCG);
 
 	if (!Data.SourceAndTargetDataIndices.IsEmpty())
 	{
diff --git a/Engine/Plugins/PCG/Source/PCG/Private/Compute/DataInterfaces/Elements/PCGMetadataPartitionDataInterface.cpp b/Engine/Plugins/PCG/Source/PCG/Private/Compute/DataInterfaces/Elements/PCGMetadataPartitionDataInterface.cpp
--- a/Engine/Plugins/PCG/Source/PCG/Private/Compute/DataInterfaces/Elements/PCGMetadataPartitionDataInterface.cpp
+++ b/Engine/Plugins/PCG/Source/PCG/Private/Compute/DataInterfaces/Elements/PCGMetadataPartitionDataInterface.cpp
@@ -299,6 +299,8 @@ bool FPCGMetaDataPartitionProviderProxy::IsValid(FValidationData const& InValida
 
 void FPCGMetaDataPartitionProviderProxy::AllocateResources(FRDGBuilder& GraphBuilder, FAllocationData const& InAllocationData)
 {
+	LLM_SCOPE_BYTAG(PCG);
+
 	// For each data, precompute a target partition index for each attribute value.
 	{
 		const int32 NumIndicesToAllocate = FMath::Max(1, NumInputData * NumPartitions);
diff --git a/Engine/Plugins/PCG/Source/PCG/Private/Compute/DataInterfaces/Elements/PCGSkinnedMeshSpawnerDataInterface.cpp b/Engine/Plugins/PCG/Source/PCG/Private/Compute/DataInterfaces/Elements/PCGSkinnedMeshSpawnerDataInterface.cpp
--- a/Engine/Plugins/PCG/Source/PCG/Private/Compute/DataInterfaces/Elements/PCGSkinnedMeshSpawnerDataInterface.cpp
+++ b/Engine/Plugins/PCG/Source/PCG/Private/Compute/DataInterfaces/Elements/PCGSkinnedMeshSpawnerDataInterface.cpp
@@ -720,6 +720,8 @@ bool FPCGSkinnedMeshSpawnerDataProviderProxy::IsValid(FValidationData const& InV
 
 void FPCGSkinnedMeshSpawnerDataProviderProxy::AllocateResources(FRDGBuilder& GraphBuilder, FAllocationData const& InAllocationData)
 {
+	LLM_SCOPE_BYTAG(PCG);
+
 	if (!PrimitiveStringKeys.IsEmpty())
 	{
 		const FRDGBufferDesc BufferDesc = FRDGBufferDesc::CreateStructuredDesc(PrimitiveStringKeys.GetTypeSize(), PrimitiveStringKeys.Num());
diff --git a/Engine/Plugins/PCG/Source/PCG/Private/Compute/DataInterfaces/PCGDataCollectionDataInterface.cpp b/Engine/Plugins/PCG/Source/PCG/Private/Compute/DataInterfaces/PCGDataCollectionDataInterface.cpp
--- a/Engine/Plugins/PCG/Source/PCG/Private/Compute/DataInterfaces/PCGDataCollectionDataInterface.cpp
+++ b/Engine/Plugins/PCG/Source/PCG/Private/Compute/DataInterfaces/PCGDataCollectionDataInterface.cpp
@@ -725,6 +725,7 @@ void FPCGDataCollectionDataProviderProxy::GatherDispatchData(FDispatchData const
 void FPCGDataCollectionDataProviderProxy::AllocateResources(FRDGBuilder& GraphBuilder, FAllocationData const& InAllocationData)
 {
 	TRACE_CPUPROFILER_EVENT_SCOPE(FPCGDataCollectionDataProviderProxy::AllocateResources);
+	LLM_SCOPE_BYTAG(PCG);
 
 	if (!bProducedByCPU)
 	{
@@ -788,6 +789,8 @@ void FPCGDataCollectionDataProviderProxy::AllocateResources(FRDGBuilder& GraphBu
 
 			ExecuteOnGameThread(UE_SOURCE_LOCATION, [ExportedPooledBuffer, DataProviderWeakPtr = DataProviderWeakPtr, PinDesc = PinDesc, OutputPinLabel = OutputPinLabel, SizeBytes = SizeBytes, ExportMode = ExportMode, OutputPinLabelAlias = OutputPinLabelAlias, GenerationCount = OriginatingGenerationCount]()
 			{
+				LLM_SCOPE_BYTAG(PCG);
+
 				// Obtain objects. No ensures added because a graph cancellation could feasibly destroy some or all of these.
 				UPCGDataCollectionDataProvider* DataProvider = DataProviderWeakPtr.Get();
 				if (!DataProvider)
diff --git a/Engine/Plugins/PCG/Source/PCG/Private/Compute/DataInterfaces/PCGDataLabelResolverDataInterface.cpp b/Engine/Plugins/PCG/Source/PCG/Private/Compute/DataInterfaces/PCGDataLabelResolverDataInterface.cpp
--- a/Engine/Plugins/PCG/Source/PCG/Private/Compute/DataInterfaces/PCGDataLabelResolverDataInterface.cpp
+++ b/Engine/Plugins/PCG/Source/PCG/Private/Compute/DataInterfaces/PCGDataLabelResolverDataInterface.cpp
@@ -177,6 +177,8 @@ void FPCGDataLabelResolverDataProviderProxy::GatherDispatchData(FDispatchData co
 
 void FPCGDataLabelResolverDataProviderProxy::AllocateResources(FRDGBuilder& GraphBuilder, FAllocationData const& InAllocationData)
 {
+	LLM_SCOPE_BYTAG(PCG);
+
 	if (!DataIdToDataIndexMap.IsEmpty())
 	{
 		const FRDGBufferDesc BufferDesc = FRDGBufferDesc::CreateStructuredDesc(DataIdToDataIndexMap.GetTypeSize(), DataIdToDataIndexMap.Num());
diff --git a/Engine/Plugins/PCG/Source/PCG/Private/Compute/DataInterfaces/PCGDebugDataInterface.cpp b/Engine/Plugins/PCG/Source/PCG/Private/Compute/DataInterfaces/PCGDebugDataInterface.cpp
--- a/Engine/Plugins/PCG/Source/PCG/Private/Compute/DataInterfaces/PCGDebugDataInterface.cpp
+++ b/Engine/Plugins/PCG/Source/PCG/Private/Compute/DataInterfaces/PCGDebugDataInterface.cpp
@@ -129,6 +129,8 @@ void FPCGDebugDataProviderProxy::GatherDispatchData(FDispatchData const& InDispa
 
 void FPCGDebugDataProviderProxy::AllocateResources(FRDGBuilder& GraphBuilder, FAllocationData const& InAllocationData)
 {
+	LLM_SCOPE_BYTAG(PCG);
+
 	const uint32 NumFloats = FMath::Max(DebugBufferSizeFloats, 1u);
 
 	FRDGBufferDesc Desc = FRDGBufferDesc::CreateByteAddressDesc(sizeof(float) * NumFloats);
diff --git a/Engine/Plugins/PCG/Source/PCG/Private/Compute/DataInterfaces/PCGInstanceDataInterface.cpp b/Engine/Plugins/PCG/Source/PCG/Private/Compute/DataInterfaces/PCGInstanceDataInterface.cpp
--- a/Engine/Plugins/PCG/Source/PCG/Private/Compute/DataInterfaces/PCGInstanceDataInterface.cpp
+++ b/Engine/Plugins/PCG/Source/PCG/Private/Compute/DataInterfaces/PCGInstanceDataInterface.cpp
@@ -281,6 +281,8 @@ bool FPCGInstanceDataProviderProxy::IsValid(FValidationData const& InValidationD
 
 void FPCGInstanceDataProviderProxy::AllocateResources(FRDGBuilder& GraphBuilder, FAllocationData const& InAllocationData)
 {
+	LLM_SCOPE_BYTAG(PCG);
+
 	const uint32 StrideUint4s = 3;
 	InstanceData = GraphBuilder.CreateBuffer(FRDGBufferDesc::CreateStructuredDesc(sizeof(FUintVector4), FMath::Max(1u, NumInstancesAllPrimitives) * StrideUint4s), TEXT("PCGInstanceDataBuffer"));
 	InstanceDataSRV = GraphBuilder.CreateSRV(InstanceData);
diff --git a/Engine/Plugins/PCG/Source/PCG/Private/Compute/DataInterfaces/PCGLandscapeDataInterface.cpp b/Engine/Plugins/PCG/Source/PCG/Private/Compute/DataInterfaces/PCGLandscapeDataInterface.cpp
--- a/Engine/Plugins/PCG/Source/PCG/Private/Compute/DataInterfaces/PCGLandscapeDataInterface.cpp
+++ b/Engine/Plugins/PCG/Source/PCG/Private/Compute/DataInterfaces/PCGLandscapeDataInterface.cpp
@@ -613,6 +613,8 @@ void FPCGLandscapeDataProviderProxy::GatherDispatchData(FDispatchData const& InD
 
 void FPCGLandscapeDataProviderProxy::AllocateResources(FRDGBuilder& GraphBuilder, FAllocationData const& InAllocationData)
 {
+	LLM_SCOPE_BYTAG(PCG);
+
 	if (TStrongObjectPtr<const URuntimeVirtualTexture> BaseColorVT = BaseColorVirtualTexture.Pin())
 	{
 		BaseColorVirtualLayer.Initialize(BaseColorVT.Get(), /*PageTableIndex=*/0, bBaseColorSRGB);
diff --git a/Engine/Plugins/PCG/Source/PCG/Private/Compute/DataInterfaces/PCGStaticMeshDataInterface.cpp b/Engine/Plugins/PCG/Source/PCG/Private/Compute/DataInterfaces/PCGStaticMeshDataInterface.cpp
--- a/Engine/Plugins/PCG/Source/PCG/Private/Compute/DataInterfaces/PCGStaticMeshDataInterface.cpp
+++ b/Engine/Plugins/PCG/Source/PCG/Private/Compute/DataInterfaces/PCGStaticMeshDataInterface.cpp
@@ -269,6 +269,8 @@ void FPCGStaticMeshDataProviderProxy::GatherDispatchData(FDispatchData const& In
 
 void FPCGStaticMeshDataProviderProxy::AllocateResources(FRDGBuilder& GraphBuilder, FAllocationData const& InAllocationData)
 {
+	LLM_SCOPE_BYTAG(PCG);
+
 	FBufferRHIRef IndexBufferRHIRef = LODResources ? LODResources->IndexBuffer.IndexBufferRHI : nullptr;
 	const bool bAllowCPUAccess = LODResources? LODResources->IndexBuffer.GetAllowCPUAccess() : false;
 	const bool bCanCreateIndexSRV =
diff --git a/Engine/Plugins/PCG/Source/PCG/Private/Compute/DataInterfaces/PCGTextureDataInterface.cpp b/Engine/Plugins/PCG/Source/PCG/Private/Compute/DataInterfaces/PCGTextureDataInterface.cpp
--- a/Engine/Plugins/PCG/Source/PCG/Private/Compute/DataInterfaces/PCGTextureDataInterface.cpp
+++ b/Engine/Plugins/PCG/Source/PCG/Private/Compute/DataInterfaces/PCGTextureDataInterface.cpp
@@ -489,6 +489,8 @@ void FPCGTextureDataProviderProxy::GatherDispatchData(FDispatchData const& InDis
 
 void FPCGTextureDataProviderProxy::AllocateResources(FRDGBuilder& GraphBuilder, FAllocationData const& InAllocationData)
 {
+	LLM_SCOPE_BYTAG(PCG);
+
 	CreateDefaultTextures(GraphBuilder);
 	CreateTextures(GraphBuilder);
 	PackTextureInfos(GraphBuilder);
@@ -597,9 +599,13 @@ void FPCGTextureDataProviderProxy::CreateTextures(FRDGBuilder& GraphBuilder)
 
 void FPCGTextureDataProviderProxy::ExportTextureUAVs(const TArray<TRefCountPtr<IPooledRenderTarget>>& ExportedTextures)
 {
+	LLM_SCOPE_BYTAG(PCG);
+
 	// Export textures and pass a reference back to the main thread where it will be picked up by the compute graph element.
 	ExecuteOnGameThread(UE_SOURCE_LOCATION, [ExportedTextures=ExportedTextures, DataProviderWeakPtr=DataProviderWeakPtr_GT, PinDesc=PinDesc, OutputPinLabel=OutputPinLabel, ExportMode=ExportMode, OutputPinLabelAlias=OutputPinLabelAlias, GenerationCount=OriginatingGenerationCount]()
 	{
+		LLM_SCOPE_BYTAG(PCG);
+
 		// Obtain objects. No ensures added because a graph cancellation could feasibly destroy some or all of these.
 		UPCGTextureDataProvider* DataProvider = DataProviderWeakPtr.Get();
 		if (!DataProvider)
diff --git a/Engine/Plugins/PCG/Source/PCG/Private/Compute/Elements/PCGComputeGraphElement.cpp b/Engine/Plugins/PCG/Source/PCG/Private/Compute/Elements/PCGComputeGraphElement.cpp
--- a/Engine/Plugins/PCG/Source/PCG/Private/Compute/Elements/PCGComputeGraphElement.cpp
+++ b/Engine/Plugins/PCG/Source/PCG/Private/Compute/Elements/PCGComputeGraphElement.cpp
@@ -55,6 +55,8 @@ FPCGComputeGraphContext::~FPCGComputeGraphContext()
 	// Will execute immediately if on game thread, otherwise on next game thread frame.
 	PCGHelpers::ExecuteOnGameThread(TEXT("~FPCGComputeGraphContext"), [DataBinding=DataBinding, GraphInstance=ComputeGraphInstance, bInstanceInitialized=bComputeGraphInstanceInitialized, GraphWeak, GridSize=GenerationGridSize, GraphIndex=ComputeGraphIndex]
 	{
+		LLM_SCOPE_BYTAG(PCG);
+
 		if (DataBinding)
 		{
 			DataBinding->ReleaseTransientResources();
@@ -210,6 +212,7 @@ bool FPCGComputeGraphElement::ExecuteInternal(FPCGContext* InContext) const
 		// Start data binding full initialization task which can run concurrently with initializing data providers.
 		UE::Tasks::TTask<void> InitializeBindingTask = UE::Tasks::Launch(UE_SOURCE_LOCATION, [ContextHandle = Context->GetOrCreateHandle()]()
 		{
+			LLM_SCOPE_BYTAG(PCG);
 			FPCGContext::FSharedContext<FPCGComputeGraphContext> SharedContext(ContextHandle);
 			FPCGComputeGraphContext* Context = SharedContext.Get();
 			if (!Context)
@@ -367,6 +370,7 @@ bool FPCGComputeGraphElement::ExecuteInternal(FPCGContext* InContext) const
 			Context->DataDescrPrimeAndValidateTask = UE::Tasks::Launch(UE_SOURCE_LOCATION, [ContextHandle = Context->GetOrCreateHandle()]()
 			{
 				TRACE_CPUPROFILER_EVENT_SCOPE(DataDescrPrimeAndValidateTask);
+				LLM_SCOPE_BYTAG(PCG);
 
 				FPCGContext::FSharedContext<FPCGComputeGraphContext> SharedContext(ContextHandle);
 				FPCGComputeGraphContext* Context = SharedContext.Get();
diff --git a/Engine/Plugins/PCG/Source/PCG/Private/Compute/PCGDataForGPU.cpp b/Engine/Plugins/PCG/Source/PCG/Private/Compute/PCGDataForGPU.cpp
--- a/Engine/Plugins/PCG/Source/PCG/Private/Compute/PCGDataForGPU.cpp
+++ b/Engine/Plugins/PCG/Source/PCG/Private/Compute/PCGDataForGPU.cpp
@@ -1949,6 +1949,7 @@ EPCGUnpackDataCollectionResult FPCGDataCollectionDesc::UnpackDataCollection(FPCG
 
 				ParallelFor(NumElements, [&](int32 ElementIndex)
 				{
+					//LLM_SCOPE_BYTAG(PCG); // AddEntryPlaceholder() does not seem to allocate but consider chunking this loop and enabling this LLM tag if this changes.
 					AllMetadataEntries[ElementIndex] = MakeTuple(Metadata->AddEntryPlaceholder(), PCGInvalidEntryKey);
 				});
 
@@ -2142,6 +2143,8 @@ bool FPCGProxyForGPUDataCollection::GetCPUData(FPCGContext* InContext, int32 InD
 
 		ENQUEUE_RENDER_COMMAND(ReadbackDataCollection)([WeakThis=AsWeak(), ReadbackRequest=ReadbackRequest, Buffer=Buffer, BufferSizeBytes=BufferSizeBytes](FRHICommandListImmediate& RHICmdList)
 		{
+			LLM_SCOPE_BYTAG(PCG);
+
 			FRDGPooledBuffer* RDGBuffer = Buffer.GetReference();
 
 			ReadbackRequest->EnqueueCopy(RHICmdList, Buffer->GetRHI());
@@ -2171,6 +2174,8 @@ bool FPCGProxyForGPUDataCollection::GetCPUData(FPCGContext* InContext, int32 InD
 
 			auto RunnerFunc = [WeakThis, ReadbackRequest, BufferSizeBytes, ExecuteAsync](auto&& RunnerFunc) -> void
 			{
+				LLM_SCOPE_BYTAG(PCG);
+
 				if (TSharedPtr<FPCGProxyForGPUDataCollection> This = WeakThis.Pin())
 				{
 					if (ReadbackRequest->IsReady())
diff --git a/Engine/Plugins/PCG/Source/PCG/Private/Data/PCGDataPtrWrapper.cpp b/Engine/Plugins/PCG/Source/PCG/Private/Data/PCGDataPtrWrapper.cpp
--- a/Engine/Plugins/PCG/Source/PCG/Private/Data/PCGDataPtrWrapper.cpp
+++ b/Engine/Plugins/PCG/Source/PCG/Private/Data/PCGDataPtrWrapper.cpp
@@ -2,6 +2,7 @@
 
 #include "Data/PCGDataPtrWrapper.h"
 #include "PCGData.h"
+#include "PCGModule.h"
 
 #include "Containers/Ticker.h"
 
@@ -97,6 +98,8 @@ void FPCGDataPtrWrapper::DecRefCount(const TWeakObjectPtr<const UPCGData>& InWea
 	{
 		ExecuteOnGameThread(UE_SOURCE_LOCATION, [InWeakData]()
 		{
+			LLM_SCOPE_BYTAG(PCG);
+
 			if (const UPCGData* InData = InWeakData.Get())
 			{
 				InData->DecCollectionRefCount();
diff --git a/Engine/Plugins/PCG/Source/PCG/Private/Elements/PCGCopyPoints.cpp b/Engine/Plugins/PCG/Source/PCG/Private/Elements/PCGCopyPoints.cpp
--- a/Engine/Plugins/PCG/Source/PCG/Private/Elements/PCGCopyPoints.cpp
+++ b/Engine/Plugins/PCG/Source/PCG/Private/Elements/PCGCopyPoints.cpp
@@ -476,6 +476,7 @@ bool FPCGCopyPointsElement::ExecuteInternal(FPCGContext* Context) const
 					const int32 AttributeCountInCurrentDispatch = FMath::Min(AttributePerDispatch, AttributesToSet.Num() - AttributeOffset);
 					ParallelFor(AttributeCountInCurrentDispatch, [AttributeOffset, &AttributesToSet, &AttributeValuesToSet](int32 WorkerIndex)
 					{
+						LLM_SCOPE_BYTAG(PCG);
 						FPCGMetadataAttributeBase* Attribute = AttributesToSet[AttributeOffset + WorkerIndex];
 						const TArray<TTuple<PCGMetadataEntryKey, PCGMetadataValueKey>>& Values = AttributeValuesToSet[AttributeOffset + WorkerIndex];
 						check(Attribute);
diff --git a/Engine/Plugins/PCG/Source/PCG/Private/Elements/PCGGenerateGrassMaps.cpp b/Engine/Plugins/PCG/Source/PCG/Private/Elements/PCGGenerateGrassMaps.cpp
--- a/Engine/Plugins/PCG/Source/PCG/Private/Elements/PCGGenerateGrassMaps.cpp
+++ b/Engine/Plugins/PCG/Source/PCG/Private/Elements/PCGGenerateGrassMaps.cpp
@@ -516,6 +516,8 @@ bool FPCGGenerateGrassMapsElement::ExecuteInternal(FPCGContext* InContext) const
 
 		ENQUEUE_RENDER_COMMAND(GenerateGrassMaps)([ContextHandle, LandscapeComponentExtent=Context->LandscapeComponentExtent, LandscapeTileCoords=MoveTemp(LandscapeTileCoords), ComponentSizeQuads=LandscapeProxy->ComponentSizeQuads](FRHICommandListImmediate& RHICmdList)
 		{
+			LLM_SCOPE_BYTAG(PCG);
+
 			FPCGContext::FSharedContext<FPCGGenerateGrassMapsContext> SharedContext(ContextHandle);
 			FPCGGenerateGrassMapsContext* Context = SharedContext.Get();
 			if (!Context)
@@ -596,6 +598,8 @@ bool FPCGGenerateGrassMapsElement::ExecuteInternal(FPCGContext* InContext) const
 			// Pass exported buffer back to game thread and wake up this element.
 			ExecuteOnGameThread(UE_SOURCE_LOCATION, [ContextHandle, GrassMapExported]()
 			{
+				LLM_SCOPE_BYTAG(PCG);
+
 				FPCGContext::FSharedContext<FPCGGenerateGrassMapsContext> SharedContext(ContextHandle);
 				if (FPCGGenerateGrassMapsContext* Context = SharedContext.Get())
 				{
diff --git a/Engine/Plugins/PCG/Source/PCG/Private/Elements/PCGSplineSampler.cpp b/Engine/Plugins/PCG/Source/PCG/Private/Elements/PCGSplineSampler.cpp
--- a/Engine/Plugins/PCG/Source/PCG/Private/Elements/PCGSplineSampler.cpp
+++ b/Engine/Plugins/PCG/Source/PCG/Private/Elements/PCGSplineSampler.cpp
@@ -1244,6 +1244,7 @@ namespace PCGSplineSamplerHelpers
 
 		ParallelFor(NumDispatch, [&](int32 DispatchIndex)
 		{
+			LLM_SCOPE_BYTAG(PCG);
 			const bool bIsLastIteration = (DispatchIndex == (NumDispatch - 1));
 			const int32 StartIterationIndex = DispatchIndex * NumIterationsPerDispatch;
 			const int32 EndIterationIndex = bIsLastIteration ? NumIterations : (StartIterationIndex + NumIterationsPerDispatch);
diff --git a/Engine/Plugins/PCG/Source/PCG/Private/Graph/PCGGraphExecutor.cpp b/Engine/Plugins/PCG/Source/PCG/Private/Graph/PCGGraphExecutor.cpp
--- a/Engine/Plugins/PCG/Source/PCG/Private/Graph/PCGGraphExecutor.cpp
+++ b/Engine/Plugins/PCG/Source/PCG/Private/Graph/PCGGraphExecutor.cpp
@@ -1241,6 +1241,8 @@ void FPCGGraphExecutor::PostTaskExecute(TSharedPtr<FPCGGraphActiveTask> ActiveTa
 
 			PCGHelpers::ExecuteOnGameThread(UE_SOURCE_LOCATION, [WeakHandler = GameThreadHandler->AsWeak(), ActiveTaskPtr]()
 			{
+				LLM_SCOPE_BYTAG(PCG);
+
 				DebugDisplayTask(WeakHandler, ActiveTaskPtr);
 			});
 		}
@@ -1772,6 +1774,7 @@ bool FPCGGraphExecutor::ExecuteScheduling(double EndTime, TSharedPtr<FPCGGraphAc
 						ActiveTask->ExecutingTask = UE::Tasks::Launch(UE_SOURCE_LOCATION, [this, ActiveTask, SchedulingBudgetInSeconds]() -> bool
 						{
 							TRACE_CPUPROFILER_EVENT_SCOPE(FPCGGraphExecutor::ExecuteAsyncTask);
+							LLM_SCOPE_BYTAG(PCG);
 
 							// Set Out Of tick budget time (minimum 1ms)
 							const float OutOfTickBudgetInSeconds = FMath::Max(1.0f, FPCGAsync::ConsoleVar::CVarAsyncOutOfTickBudgetInMilliseconds.GetValueOnAnyThread()) / 1000.f;
diff --git a/Engine/Plugins/PCG/Source/PCG/Private/Graph/PCGGraphTask.cpp b/Engine/Plugins/PCG/Source/PCG/Private/Graph/PCGGraphTask.cpp
--- a/Engine/Plugins/PCG/Source/PCG/Private/Graph/PCGGraphTask.cpp
+++ b/Engine/Plugins/PCG/Source/PCG/Private/Graph/PCGGraphTask.cpp
@@ -3,6 +3,7 @@
 #include "PCGGraphTask.h"
 
 #include "PCGComponent.h"
+#include "PCGModule.h"
 #include "PCGSubsystem.h"
 #include "Compute/Elements/PCGComputeGraphElement.h"
 #include "Editor/IPCGEditorModule.h"
@@ -116,6 +117,8 @@ FPCGGraphActiveTask::~FPCGGraphActiveTask()
 	{
 		PCGHelpers::ExecuteOnGameThread(UE_SOURCE_LOCATION, [ContextPtr = Context.Release()]()
 		{
+			LLM_SCOPE_BYTAG(PCG);
+
 			PCGGraphExecutor::ClearAsyncFlags(ContextPtr->AsyncObjects);
 
 			FPCGContext::Release(ContextPtr);
diff --git a/Engine/Plugins/PCG/Source/PCG/Private/Grid/PCGLandscapeCache.cpp b/Engine/Plugins/PCG/Source/PCG/Private/Grid/PCGLandscapeCache.cpp
--- a/Engine/Plugins/PCG/Source/PCG/Private/Grid/PCGLandscapeCache.cpp
+++ b/Engine/Plugins/PCG/Source/PCG/Private/Grid/PCGLandscapeCache.cpp
@@ -709,6 +709,7 @@ void UPCGLandscapeCache::PrimeCache()
 
 	ParallelFor(CacheEntriesToBuild.Num(), [this, &CacheEntriesToBuild, &NewEntries](int32 EntryIndex)
 	{
+		LLM_SCOPE_BYTAG(PCG);
 		FTaskTagScope Scope(ETaskTag::EParallelGameThread);
 		const TTuple<ULandscapeInfo*, ULandscapeComponent*, CacheMapKey>& CacheEntryInfo = CacheEntriesToBuild[EntryIndex];
 		NewEntries[EntryIndex] = FPCGLandscapeCacheEntry::CreateCacheEntry(CacheEntryInfo.Get<0>(), CacheEntryInfo.Get<1>());
diff --git a/Engine/Plugins/PCG/Source/PCG/Private/Helpers/PCGAsync.cpp b/Engine/Plugins/PCG/Source/PCG/Private/Helpers/PCGAsync.cpp
--- a/Engine/Plugins/PCG/Source/PCG/Private/Helpers/PCGAsync.cpp
+++ b/Engine/Plugins/PCG/Source/PCG/Private/Helpers/PCGAsync.cpp
@@ -3,6 +3,7 @@
 #include "Helpers/PCGAsync.h"
 
 #include "PCGContext.h"
+#include "PCGModule.h"
 #include "PCGPoint.h"
 
 #include "Tasks/Task.h"
@@ -80,6 +81,8 @@ namespace FPCGAsync
 
 			AsyncTasks.Emplace(UE::Tasks::Launch(UE_SOURCE_LOCATION, [&IterationInnerLoop, StartIndex, EndIndex, bInitBPContext]()
 			{
+				LLM_SCOPE_BYTAG(PCG);
+
 				if (bInitBPContext)
 				{
 					GInitRunaway(); // Reset counter as threads can run multiple workloads and might not be properly reset.
@@ -187,6 +190,8 @@ namespace FPCGAsync
 
 			AsyncTasks.Emplace(UE::Tasks::Launch(UE_SOURCE_LOCATION, [&IterationInnerLoop, StartIndex, EndIndex, bInitBPContext]() -> TArray<FPCGPoint>
 			{
+				LLM_SCOPE_BYTAG(PCG);
+
 				if (bInitBPContext)
 				{
 					GInitRunaway(); // Reset counter as threads can run multiple workloads and might not be properly reset.
@@ -265,6 +270,8 @@ namespace FPCGAsync
 
 			AsyncTasks.Emplace(UE::Tasks::Launch(UE_SOURCE_LOCATION, [&IterationInnerLoop, StartIndex, EndIndex, bInitBPContext]() -> TPair<int32, int32>
 			{
+				LLM_SCOPE_BYTAG(PCG);
+
 				if (bInitBPContext)
 				{
 					GInitRunaway(); // Reset counter as threads can run multiple workloads and might not be properly reset.
@@ -588,6 +595,8 @@ namespace FPCGAsync
 
 				AsyncTasks.Emplace(UE::Tasks::Launch(UE_SOURCE_LOCATION, [&JobTask, &SynchroStruct, &StartEnd = AsyncState.TasksChunksStartEnd[TaskIndex]]() -> void
 				{
+					LLM_SCOPE_BYTAG(PCG);
+
 					// Do at least one run, as we could fall into an infinite loop if we always have to stop before doing anything.
 					// Quit if we have no more work to do or we were told to stop.
 					do
diff --git a/Engine/Plugins/PCG/Source/PCG/Private/Helpers/PCGHelpers.cpp b/Engine/Plugins/PCG/Source/PCG/Private/Helpers/PCGHelpers.cpp
--- a/Engine/Plugins/PCG/Source/PCG/Private/Helpers/PCGHelpers.cpp
+++ b/Engine/Plugins/PCG/Source/PCG/Private/Helpers/PCGHelpers.cpp
@@ -815,6 +815,8 @@ namespace PCGHelpers
 
 		PCGHelpers::ExecuteOnGameThread(TEXT("Execution debug"), [ExecutionSource = InContext->ExecutionSource, Color]()
 		{
+			LLM_SCOPE_BYTAG(PCG);
+
 			UWorld* World = ExecutionSource.IsValid() ? ExecutionSource->GetExecutionState().GetWorld() : nullptr;
 			if (!World)
 			{
diff --git a/Engine/Plugins/PCG/Source/PCG/Private/Metadata/PCGMetadataDomain.cpp b/Engine/Plugins/PCG/Source/PCG/Private/Metadata/PCGMetadataDomain.cpp
--- a/Engine/Plugins/PCG/Source/PCG/Private/Metadata/PCGMetadataDomain.cpp
+++ b/Engine/Plugins/PCG/Source/PCG/Private/Metadata/PCGMetadataDomain.cpp
@@ -1446,6 +1446,8 @@ void FPCGMetadataDomain::SetAttributes(const TArrayView<const PCGMetadataEntryKe
 		const int32 AttributeCountInCurrentDispatch = FMath::Min(AttributesPerDispatch, Attributes.Num() - AttributeOffset);
 		ParallelFor(AttributeCountInCurrentDispatch, [this, OptionalContext, AttributeOffset, InMetadata, &AttributeNames, &InKeys, &OutKeys](int32 WorkerIndex)
 		{
+			LLM_SCOPE_BYTAG(PCG);
+
 			const FName AttributeName = AttributeNames[AttributeOffset + WorkerIndex];
 			FPCGMetadataAttributeBase* Attribute = Attributes[AttributeName];
 
diff --git a/Engine/Plugins/PCG/Source/PCG/Private/PCGComponent.cpp b/Engine/Plugins/PCG/Source/PCG/Private/PCGComponent.cpp
--- a/Engine/Plugins/PCG/Source/PCG/Private/PCGComponent.cpp
+++ b/Engine/Plugins/PCG/Source/PCG/Private/PCGComponent.cpp
@@ -1960,12 +1960,14 @@ void UPCGComponent::OnUnregister()
 		{
 			// Dispatch delete of actors to next tick
 			::ExecuteOnGameThread(UE_SOURCE_LOCATION, [WeakWorld = TWeakObjectPtr<UWorld>(GetWorld()), InActorsToDelete = MoveTemp(ActorsToDelete)]()
+			{
+				LLM_SCOPE_BYTAG(PCG);
+
+				if (UWorld* World = WeakWorld.Get())
 				{
-					if (UWorld* World = WeakWorld.Get())
-					{
-						UPCGActorHelpers::DeleteActors(World, InActorsToDelete.Array());
-					}
-				});
+					UPCGActorHelpers::DeleteActors(World, InActorsToDelete.Array());
+				}
+			});
 		}
 	}
 #endif // WITH_EDITOR
diff --git a/Engine/Plugins/PCG/Source/PCG/Private/PCGModule.cpp b/Engine/Plugins/PCG/Source/PCG/Private/PCGModule.cpp
--- a/Engine/Plugins/PCG/Source/PCG/Private/PCGModule.cpp
+++ b/Engine/Plugins/PCG/Source/PCG/Private/PCGModule.cpp
@@ -29,6 +29,8 @@
 #include "Misc/Paths.h"
 #endif
 
+LLM_DEFINE_TAG(PCG);
+
 #define LOCTEXT_NAMESPACE "FPCGModule"
 
 #if WITH_EDITOR
@@ -68,6 +70,8 @@ bool FPCGModule::IsPCGModuleLoaded()
 
 void FPCGModule::StartupModule()
 {
+	LLM_SCOPE_BYTAG(PCG);
+
 	// Cache for fast access
 	check(!PCGModule::PCGModulePtr);
 	PCGModule::PCGModulePtr = this;
@@ -134,6 +138,7 @@ void FPCGModule::PreExit()
 
 void FPCGModule::OnPostInitEngine()
 {
+	LLM_SCOPE_BYTAG(PCG);
 	DataTypeRegistry.RegisterKnownTypes();
 }
 
@@ -145,6 +150,8 @@ void FPCGModule::ExecuteNextTick(TFunction<void()> TickFunction)
 
 bool FPCGModule::Tick(float DeltaTime)
 {
+	LLM_SCOPE_BYTAG(PCG);
+
 	TArray<TFunction<void()>> LocalExecuteNextTicks;
 
 	{
diff --git a/Engine/Plugins/PCG/Source/PCG/Private/RuntimeGen/PCGRuntimeGenScheduler.cpp b/Engine/Plugins/PCG/Source/PCG/Private/RuntimeGen/PCGRuntimeGenScheduler.cpp
--- a/Engine/Plugins/PCG/Source/PCG/Private/RuntimeGen/PCGRuntimeGenScheduler.cpp
+++ b/Engine/Plugins/PCG/Source/PCG/Private/RuntimeGen/PCGRuntimeGenScheduler.cpp
@@ -837,6 +837,7 @@ void FPCGRuntimeGenScheduler::TickCleanup(const TSet<IPCGGenSourceBase*>& InGenS
 	ParallelFor(NumGeneratedComponents, [&GeneratedComponentsArray, &ComponentsToClean, &InvalidKeys, &ComponentsToCleanCounter, &InvalidKeysCounter, ActorAndComponentMapping=ActorAndComponentMapping, &CheckIfAllGenSourcesWantToCleanup](int32 Index)
 	{
 		TRACE_CPUPROFILER_EVENT_SCOPE(SelectComponentForCleanup);
+		LLM_SCOPE_BYTAG(PCG);
 
 		const FGridGenerationKey& GenerationKey = GeneratedComponentsArray[Index];
 
diff --git a/Engine/Plugins/PCG/Source/PCG/Private/Subsystems/PCGSubsystem.cpp b/Engine/Plugins/PCG/Source/PCG/Private/Subsystems/PCGSubsystem.cpp
--- a/Engine/Plugins/PCG/Source/PCG/Private/Subsystems/PCGSubsystem.cpp
+++ b/Engine/Plugins/PCG/Source/PCG/Private/Subsystems/PCGSubsystem.cpp
@@ -216,6 +216,8 @@ void UPCGSubsystem::PostInitialize()
 {
 	Super::PostInitialize();
 
+	LLM_SCOPE_BYTAG(PCG);
+
 	// Gather world pcg actor if it exists
 	if (!PCGWorldActor)
 	{
@@ -301,6 +303,7 @@ bool UPCGSubsystem::RemoveAndCopyConstructionScriptSourceComponent(AActor* InCom
 void UPCGSubsystem::Tick(float DeltaSeconds)
 {
 	TRACE_CPUPROFILER_EVENT_SCOPE(UPCGSubsystem::Tick);
+	LLM_SCOPE_BYTAG(PCG);
 
 	Super::Tick(DeltaSeconds);
 
diff --git a/Engine/Plugins/PCG/Source/PCG/Public/PCGModule.h b/Engine/Plugins/PCG/Source/PCG/Public/PCGModule.h
--- a/Engine/Plugins/PCG/Source/PCG/Public/PCGModule.h
+++ b/Engine/Plugins/PCG/Source/PCG/Public/PCGModule.h
@@ -27,7 +27,7 @@ namespace PCGEngineShowFlags
 // Stats
 DECLARE_STATS_GROUP(TEXT("PCG"), STATGROUP_PCG, STATCAT_Advanced);
 
-// CVars
+LLM_DECLARE_TAG(PCG);
 
 // Delegates
 DECLARE_MULTICAST_DELEGATE_TwoParams(FPCGGraphChangedDelegate, UPCGGraph* /*InGraph*/, EPCGChangeType /*ChangeType*/);
diff --git a/Engine/Plugins/PCG/Source/PCGCompute/Private/PCGComputeModule.cpp b/Engine/Plugins/PCG/Source/PCGCompute/Private/PCGComputeModule.cpp
--- a/Engine/Plugins/PCG/Source/PCGCompute/Private/PCGComputeModule.cpp
+++ b/Engine/Plugins/PCG/Source/PCGCompute/Private/PCGComputeModule.cpp
@@ -35,6 +35,8 @@ void FPCGComputeModule::ShutdownModule()
 
 IMPLEMENT_MODULE(FPCGComputeModule, PCGCompute);
 
-PCGCOMPUTE_API DEFINE_LOG_CATEGORY(LogPCGCompute);
+DEFINE_LOG_CATEGORY(LogPCGCompute);
+
+LLM_DEFINE_TAG(PCGCompute);
 
 #undef LOCTEXT_NAMESPACE
diff --git a/Engine/Plugins/PCG/Source/PCGCompute/Private/PCGTextureReadback.cpp b/Engine/Plugins/PCG/Source/PCGCompute/Private/PCGTextureReadback.cpp
--- a/Engine/Plugins/PCG/Source/PCGCompute/Private/PCGTextureReadback.cpp
+++ b/Engine/Plugins/PCG/Source/PCGCompute/Private/PCGTextureReadback.cpp
@@ -2,6 +2,8 @@
 
 #include "PCGTextureReadback.h"
 
+#include "PCGComputeModule.h"
+
 #include "GlobalShader.h"
 #include "RHIGPUReadback.h"
 #include "RenderGraphUtils.h"
@@ -40,6 +42,7 @@ IMPLEMENT_GLOBAL_SHADER(FPCGTextureReadbackCS, "/PCGComputeShaders/PCGTextureRea
 
 void FPCGTextureReadbackInterface::Dispatch_RenderThread(FRHICommandListImmediate& RHICmdList, const FPCGTextureReadbackDispatchParams& Params, const TFunction<void(void* OutBuffer, int32 ReadbackWidth, int32 ReadbackHeight)>& AsyncCallback)
 {
+	LLM_SCOPE_BYTAG(PCGCompute);
 	check(Params.SourceTexture && Params.SourceSampler);
 
 	const bool bIsTextureArray = Params.SourceTexture->GetDesc().Dimension == ETextureDimension::Texture2DArray;
@@ -138,6 +141,8 @@ void FPCGTextureReadbackInterface::Dispatch_RenderThread(FRHICommandListImmediat
 
 	auto RunnerFunc = [GPUTextureReadBackPtr = GPUTextureReadback.Release(), AsyncCallback, ExecuteAsync](auto&& RunnerFunc) -> void
 	{
+		LLM_SCOPE_BYTAG(PCGCompute);
+
 		if (GPUTextureReadBackPtr->IsReady())
 		{
 			int32 ReadbackWidth = 0, ReadbackHeight = 0;
diff --git a/Engine/Plugins/PCG/Source/PCGCompute/Public/PCGComputeModule.h b/Engine/Plugins/PCG/Source/PCGCompute/Public/PCGComputeModule.h
--- a/Engine/Plugins/PCG/Source/PCGCompute/Public/PCGComputeModule.h
+++ b/Engine/Plugins/PCG/Source/PCGCompute/Public/PCGComputeModule.h
@@ -10,5 +10,4 @@ PCGCOMPUTE_API DECLARE_LOG_CATEGORY_EXTERN(LogPCGCompute, Log, All);
 // Stats
 DECLARE_STATS_GROUP(TEXT("PCGCompute"), STATGROUP_PCGCompute, STATCAT_Advanced)
 
-// CVars
-
+LLM_DECLARE_TAG(PCGCompute);
