[PCG] Rename cvar for componentless primitive spawning.

Enable FastGeo spawning by default in PCGSample to increase test mileage.

#rnx

diff --git a/Engine/Plugins/PCG/Source/PCG/Private/Compute/DataInterfaces/Elements/PCGStaticMeshSpawnerDataInterface.cpp b/Engine/Plugins/PCG/Source/PCG/Private/Compute/DataInterfaces/Elements/PCGStaticMeshSpawnerDataInterface.cpp
--- a/Engine/Plugins/PCG/Source/PCG/Private/Compute/DataInterfaces/Elements/PCGStaticMeshSpawnerDataInterface.cpp
+++ b/Engine/Plugins/PCG/Source/PCG/Private/Compute/DataInterfaces/Elements/PCGStaticMeshSpawnerDataInterface.cpp
@@ -41,7 +41,7 @@ namespace PCGStaticMeshSpawnerDataInterface
 	static const FText NoMeshEntriesFormat = LOCTEXT("NoMeshEntries", "No mesh entries provided in weighted mesh selector.");
 
 	static TAutoConsoleVariable<bool> CVarCreatePrimitivesComponentless(
-		TEXT("pcg.RuntimeGeneration.ComponentlessPrimitives"),
+		TEXT("pcg.RuntimeGeneration.ISM.ComponentlessPrimitives"),
 		false,
 		TEXT("Uses a component-less path for creating primitives in game worlds (PIE and standalone)."));
 
