[PCG] Add node to save texture data to asset

diff --git a/Engine/Plugins/PCG/Source/PCG/Private/Elements/PCGSaveTextureToAsset.cpp b/Engine/Plugins/PCG/Source/PCG/Private/Elements/PCGSaveTextureToAsset.cpp
new file mode 100644
--- /dev/null
+++ b/Engine/Plugins/PCG/Source/PCG/Private/Elements/PCGSaveTextureToAsset.cpp
@@ -0,0 +1,233 @@
+// Copyright Epic Games, Inc. All Rights Reserved.
+
+#include "Elements/PCGSaveTextureToAsset.h"
+
+#include "PCGAssetExporterUtils.h"
+#include "PCGModule.h"
+#include "PCGParamData.h"
+#include "PCGTextureReadback.h"
+#include "Data/PCGTextureData.h"
+#include "Metadata/PCGMetadata.h"
+
+#include "RHIStaticStates.h"
+#include "Engine/Texture2D.h"
+
+#include UE_INLINE_GENERATED_CPP_BY_NAME(PCGSaveTextureToAsset)
+
+#define LOCTEXT_NAMESPACE "PCGSaveTextureToAsset"
+
+namespace PCGSaveTextureToAssetConstants
+{
+	const FName OutAssetPathLabel = TEXT("OutAssetPath");
+}
+
+TArray<FPCGPinProperties> UPCGSaveTextureToAssetSettings::InputPinProperties() const
+{
+	TArray<FPCGPinProperties> Properties;
+	FPCGPinProperties& InputTexturePinProperties = Properties.Emplace_GetRef(PCGPinConstants::DefaultInputLabel, EPCGDataType::BaseTexture, /*bInAllowMultipleConnections=*/false, /*bAllowMultipleData=*/false);
+	InputTexturePinProperties.SetRequiredPin();
+
+	return Properties;
+}
+
+TArray<FPCGPinProperties> UPCGSaveTextureToAssetSettings::OutputPinProperties() const
+{
+	TArray<FPCGPinProperties> Properties;
+	Properties.Emplace(PCGSaveTextureToAssetConstants::OutAssetPathLabel, EPCGDataType::Param, /*bInAllowMultipleConnections=*/true, /*bAllowMultipleData=*/false);
+
+	return Properties;
+}
+
+FPCGElementPtr UPCGSaveTextureToAssetSettings::CreateElement() const
+{
+	return MakeShared<FPCGSaveTextureToAssetElement>();
+}
+
+FPCGSaveTextureToAssetContext::~FPCGSaveTextureToAssetContext()
+{
+	FMemory::Free(RawReadbackData);
+}
+
+void FPCGSaveTextureToAssetContext::AddExtraStructReferencedObjects(FReferenceCollector& Collector)
+{
+	Collector.AddReferencedObject(ExportedTexture);
+}
+
+bool FPCGSaveTextureToAssetElement::ExecuteInternal(FPCGContext* InContext) const
+{
+	TRACE_CPUPROFILER_EVENT_SCOPE(FPCGSaveTextureToAssetElement::Execute);
+	check(InContext);
+
+#if WITH_EDITOR
+	FPCGSaveTextureToAssetContext* Context = static_cast<FPCGSaveTextureToAssetContext*>(InContext);
+
+	const UPCGSaveTextureToAssetSettings* Settings = Context->GetInputSettings<UPCGSaveTextureToAssetSettings>();
+	check(Settings);
+
+	auto SleepUntilNextFrame = [Context]()
+	{
+		Context->bIsPaused = true;
+		FPCGModule::GetPCGModuleChecked().ExecuteNextTick([ContextHandle = Context->GetOrCreateHandle()]()
+		{
+			if (TSharedPtr<FPCGContextHandle> SharedHandle = ContextHandle.Pin())
+			{
+				if (FPCGContext* ContextPtr = SharedHandle->GetContext())
+				{
+					ContextPtr->bIsPaused = false;
+				}
+			}
+		});
+	};
+
+	if (!Context->InputTextureData)
+	{
+		const TArray<FPCGTaggedData> Inputs = InContext->InputData.GetInputsByPin(PCGPinConstants::DefaultInputLabel);
+
+		if (Inputs.IsEmpty())
+		{
+			return true;
+		}
+
+		if (Inputs.Num() > 1)
+		{
+			PCGLog::InputOutput::LogFirstInputOnlyWarning(PCGPinConstants::DefaultInputLabel, InContext);
+		}
+
+		Context->InputTextureData = Cast<const UPCGBaseTextureData>(Inputs[0].Data);
+
+		if (!Context->InputTextureData)
+		{
+			PCGLog::InputOutput::LogTypedDataNotFoundWarning(EPCGDataType::BaseTexture, PCGPinConstants::DefaultInputLabel, InContext);
+			return true;
+		}
+	}
+
+	// Poll readback of input texture. Sleep each tick until readback is complete.
+	if (!ReadbackInputTexture(Context))
+	{
+		SleepUntilNextFrame();
+		return false;
+	}
+
+	if (!Context->ExportedTexture)
+	{
+		// Create the texture asset.
+		UTexture2D* ExportedTexture = nullptr;
+
+		UPCGAssetExporterUtils::CreateAsset(
+			UTexture2D::StaticClass(),
+			Settings->ExporterParams,
+			[&ExportedTexture](const FString& PackagePath, UObject* Asset) -> bool
+			{
+				ExportedTexture = Cast<UTexture2D>(Asset);
+				return ExportedTexture != nullptr;
+			},
+			Context);
+
+		// Initialize exported texture with the readback data.
+		ExportedTexture->Source.Init(Context->ReadbackWidth, Context->ReadbackHeight, /*NewNumSlices=*/1, /*NewNumMips=*/1, TSF_BGRA8, Context->RawReadbackData);
+		ExportedTexture->UpdateResource();
+
+		Context->ExportedTexture = ExportedTexture;
+	}
+
+	if (Context->ExportedTexture->HasPendingInitOrStreaming(/*bWaitForLODTransition=*/true))
+	{
+		SleepUntilNextFrame();
+		return false;
+	}
+
+	// Create an attribute set output holding the exported texture asset path.
+	UPCGParamData* OutputParamData = FPCGContext::NewObject_AnyThread<UPCGParamData>(Context);
+	OutputParamData->Metadata->CreateAttribute<FSoftObjectPath>(TEXT("AssetPath"), FSoftObjectPath(Context->ExportedTexture), /*bAllowInterpolation=*/false, /*bOverrideParent=*/false);
+	OutputParamData->Metadata->AddEntry();
+
+	Context->OutputData.TaggedData.Emplace_GetRef().Data = OutputParamData;
+#else
+	PCGLog::LogErrorOnGraph(LOCTEXT("CannotExportInNonEditor", "Texture cannot be saved to an asset in non-editor builds"), InContext);
+#endif // WITH_EDITOR
+
+	return true;
+}
+
+bool FPCGSaveTextureToAssetElement::ReadbackInputTexture(FPCGSaveTextureToAssetContext* InContext) const
+{
+	TRACE_CPUPROFILER_EVENT_SCOPE(FPCGSaveTextureToAssetElement::ReadbackInputTexture);
+
+	// Only dispatch readback if it hasn't already been submitted.
+	if (!InContext->bReadbackDispatched)
+	{
+		if (!InContext->InputTextureData)
+		{
+			return true;
+		}
+
+		if (InContext->InputTextureData->GetTextureResourceType() == EPCGTextureResourceType::TextureObject)
+		{
+			UTexture* InputTexture = InContext->InputTextureData->GetTexture();
+
+			if (!InputTexture)
+			{
+				return true;
+			}
+
+			if (!InContext->bUpdatedReadbackTextureResource)
+			{
+				InputTexture->UpdateResource();
+
+				InContext->bUpdatedReadbackTextureResource = true;
+			}
+
+			if (InputTexture->HasPendingInitOrStreaming(/*bWaitForLODTransition=*/true))
+			{
+				return false;
+			}
+		}
+
+		FTextureRHIRef TextureRHI = InContext->InputTextureData->GetTextureRHI();
+
+		if (!ensure(TextureRHI))
+		{
+			// If the texture could not be acquired, readback cannot continue.
+			return true;
+		}
+
+		const FIntVector TextureSize = TextureRHI->GetDesc().GetSize();
+
+		FPCGTextureReadbackDispatchParams Params;
+		Params.SourceTexture = TextureRHI;
+		Params.SourceSampler = TStaticSamplerState<SF_Point, AM_Clamp, AM_Clamp, AM_Clamp>::GetRHI();
+		Params.SourceTextureIndex = InContext->InputTextureData->GetTextureSlice();
+		Params.SourceDimensions = FIntPoint(TextureSize.X, TextureSize.Y);
+
+		InContext->bReadbackDispatched = true;
+
+		FPCGTextureReadbackInterface::Dispatch(Params, [ContextHandle = InContext->GetOrCreateHandle()](void* OutBuffer, int32 ReadbackWidth, int32 ReadbackHeight)
+		{
+			TRACE_CPUPROFILER_EVENT_SCOPE(FPCGSaveTextureToAssetElement::ReadbackInputTexture::DispatchCallback);
+
+			TSharedPtr<FPCGContextHandle> SharedHandle = ContextHandle.Pin();
+			FPCGSaveTextureToAssetContext* Context = SharedHandle ? static_cast<FPCGSaveTextureToAssetContext*>(SharedHandle->GetContext()) : nullptr;
+
+			if (!Context)
+			{
+				return;
+			}
+
+			Context->ReadbackWidth = ReadbackWidth;
+			Context->ReadbackHeight = ReadbackHeight;
+
+			// PCGTextureReadback always uses PF_B8G8R8A8. @todo_pcg: Maybe the read back format could be configurable.
+			const uint32 BufferSizeBytes = ReadbackWidth * ReadbackHeight * GPixelFormats[PF_B8G8R8A8].BlockBytes;
+			Context->RawReadbackData = static_cast<uint8*>(FMemory::Malloc(BufferSizeBytes));
+			FMemory::Memcpy(Context->RawReadbackData, OutBuffer, BufferSizeBytes);
+
+			Context->bReadbackComplete = true;
+			Context->bIsPaused = false;
+		});
+	}
+
+	return InContext->bReadbackComplete;
+}
+
+#undef LOCTEXT_NAMESPACE
diff --git a/Engine/Plugins/PCG/Source/PCG/Private/Utils/PCGLogErrors.cpp b/Engine/Plugins/PCG/Source/PCG/Private/Utils/PCGLogErrors.cpp
--- a/Engine/Plugins/PCG/Source/PCG/Private/Utils/PCGLogErrors.cpp
+++ b/Engine/Plugins/PCG/Source/PCG/Private/Utils/PCGLogErrors.cpp
@@ -23,7 +23,7 @@ namespace PCGLog::InputOutput
 	{
 		const UEnum* PCGDataTypeEnum = StaticEnum<EPCGDataType>();
 		const FText TypeText = PCGDataTypeEnum ? PCGDataTypeEnum->GetDisplayNameTextByValue(static_cast<int64>(DataType)) : LOCTEXT("UnknownDataType", "Unknown");
-		LogWarningOnGraph(FText::Format(Format::TypedInputNotFound, FText::FromName(PinLabel), TypeText), InContext);
+		LogWarningOnGraph(FText::Format(Format::TypedInputNotFound, TypeText, FText::FromName(PinLabel)), InContext);
 	}
 
 	void LogFirstInputOnlyWarning(const FName PinLabel, const FPCGContext* InContext)
diff --git a/Engine/Plugins/PCG/Source/PCG/Public/Elements/PCGSaveTextureToAsset.h b/Engine/Plugins/PCG/Source/PCG/Public/Elements/PCGSaveTextureToAsset.h
new file mode 100644
--- /dev/null
+++ b/Engine/Plugins/PCG/Source/PCG/Public/Elements/PCGSaveTextureToAsset.h
@@ -0,0 +1,69 @@
+// Copyright Epic Games, Inc. All Rights Reserved.
+
+#pragma once
+
+#include "PCGAssetExporter.h"
+#include "PCGContext.h"
+#include "PCGSettings.h"
+
+#include "PCGSaveTextureToAsset.generated.h"
+
+class UPCGBaseTextureData;
+class UTexture2D;
+
+/** Save the input texture to a UTexture2D asset (format is always BGRA8). Outputs a soft object path attribute for the saved texture. Editor only. */
+UCLASS(MinimalAPI, BlueprintType, ClassGroup = (Procedural))
+class UPCGSaveTextureToAssetSettings : public UPCGSettings
+{
+	GENERATED_BODY()
+
+public:
+	//~Begin UPCGSettings interface
+#if WITH_EDITOR
+	virtual FName GetDefaultNodeName() const override { return FName(TEXT("SaveTextureToAsset")); }
+	virtual FText GetDefaultNodeTitle() const override { return NSLOCTEXT("PCGSaveTextureToAssetElement", "NodeTitle", "Save Texture to Asset"); }
+	virtual FText GetNodeTooltipText() const override { return NSLOCTEXT("PCGSaveTextureToAssetElement", "NodeTooltip", "Save the input texture to a UTexture2D asset (format is always BGRA8). Outputs a soft object path attribute for the saved texture. Note: This node is editor only."); }
+	virtual EPCGSettingsType GetType() const override { return EPCGSettingsType::InputOutput; }
+#endif
+
+protected:
+	virtual TArray<FPCGPinProperties> InputPinProperties() const override;
+	virtual TArray<FPCGPinProperties> OutputPinProperties() const override;
+	virtual FPCGElementPtr CreateElement() const override;
+	//~End UPCGSettings interface
+
+public:
+	UPROPERTY(BlueprintReadWrite, EditAnywhere, Category = Data, meta = (ShowOnlyInnerProperties, PCG_Overridable))
+	FPCGAssetExporterParameters ExporterParams;
+};
+
+struct FPCGSaveTextureToAssetContext : public FPCGContext
+{
+public:
+	virtual ~FPCGSaveTextureToAssetContext();
+
+protected:
+	virtual void AddExtraStructReferencedObjects(FReferenceCollector& Collector) override;
+
+public:
+	TObjectPtr<const UPCGBaseTextureData> InputTextureData = nullptr;
+	TObjectPtr<UTexture2D> ExportedTexture = nullptr;
+	uint8* RawReadbackData = nullptr;
+	int32 ReadbackWidth = INDEX_NONE;
+	int32 ReadbackHeight = INDEX_NONE;
+	bool bUpdatedReadbackTextureResource = false;
+	bool bReadbackDispatched = false;
+	bool bReadbackComplete = false;
+};
+
+class FPCGSaveTextureToAssetElement : public IPCGElementWithCustomContext<FPCGSaveTextureToAssetContext>
+{
+public:
+	/** Creating asset and initializing texture resource must occur on main thread. */
+	virtual bool CanExecuteOnlyOnMainThread(FPCGContext* Context) const override { return true; }
+
+protected:
+	virtual bool ExecuteInternal(FPCGContext* InContext) const override;
+
+	bool ReadbackInputTexture(FPCGSaveTextureToAssetContext* InContext) const;
+};
