[PCG] Add SceneCapture node

diff --git a/Engine/Plugins/PCG/Source/PCG/Private/Elements/PCGSceneCapture.cpp b/Engine/Plugins/PCG/Source/PCG/Private/Elements/PCGSceneCapture.cpp
new file mode 100644
--- /dev/null
+++ b/Engine/Plugins/PCG/Source/PCG/Private/Elements/PCGSceneCapture.cpp
@@ -0,0 +1,187 @@
+// Copyright Epic Games, Inc. All Rights Reserved.
+
+#include "Elements/PCGSceneCapture.h"
+
+#include "Helpers/PCGSettingsHelpers.h"
+
+#include "Data/PCGRenderTargetData.h"
+
+#include UE_INLINE_GENERATED_CPP_BY_NAME(PCGSceneCapture)
+
+#define LOCTEXT_NAMESPACE "PCGSceneCaptureElement"
+
+namespace PCGSceneCaptureConstants
+{
+	const FName BoundingShapeLabel = TEXT("BoundingShape");
+}
+
+TArray<FPCGPinProperties> UPCGSceneCaptureSettings::InputPinProperties() const
+{
+	TArray<FPCGPinProperties> Properties;
+
+	FPCGPinProperties& BoundsPin = Properties.Emplace_GetRef(
+		PCGSceneCaptureConstants::BoundingShapeLabel,
+		EPCGDataType::Spatial,
+		/*bInAllowMultipleConnections=*/true,
+		/*bAllowMultipleData=*/false,
+		LOCTEXT("BoundingShapePinTooltip", "Optional bounds to use instead of the PCG generation bounds. Note, bound overrides will always be axis aligned and top down."));
+	BoundsPin.SetAdvancedPin();
+
+	return Properties;
+}
+
+TArray<FPCGPinProperties> UPCGSceneCaptureSettings::OutputPinProperties() const
+{
+	TArray<FPCGPinProperties> Properties;
+	Properties.Emplace(PCGPinConstants::DefaultInputLabel, EPCGDataType::RenderTarget, /*bInAllowMultipleConnections=*/true, /*bAllowMultipleData=*/false);
+
+	return Properties;
+}
+
+FPCGElementPtr UPCGSceneCaptureSettings::CreateElement() const
+{
+	return MakeShared<FPCGSceneCaptureElement>();
+}
+
+void FPCGSceneCaptureContext::AddExtraStructReferencedObjects(FReferenceCollector& Collector)
+{
+	Collector.AddReferencedObject(SceneCaptureComponent);
+	Collector.AddReferencedObject(RenderTarget);
+}
+
+bool FPCGSceneCaptureElement::ExecuteInternal(FPCGContext* InContext) const
+{
+	TRACE_CPUPROFILER_EVENT_SCOPE(FPCGSceneCaptureElement::Execute);
+	check(InContext);
+
+	FPCGSceneCaptureContext* Context = static_cast<FPCGSceneCaptureContext*>(InContext);
+
+	const UPCGSceneCaptureSettings* Settings = InContext->GetInputSettings<UPCGSceneCaptureSettings>();
+	check(Settings);
+
+	if (!Context->bSubmittedSceneCapture)
+	{
+		check(InContext->ExecutionSource.Get());
+
+		const IPCGGraphExecutionState& ExecutionState = InContext->ExecutionSource->GetExecutionState();
+		UWorld* World = ExecutionState.GetWorld();
+
+		if (!World)
+		{
+			PCGLog::LogErrorOnGraph(LOCTEXT("InvalidWorld", "Invalid world. Cannot create USceneCaptureComponent."), InContext);
+			return true;
+		}
+
+		if (Settings->TexelSize <= UE_DOUBLE_SMALL_NUMBER)
+		{
+			PCGLog::LogErrorOnGraph(FText::Format(LOCTEXT("InvalidTexelSize", "Invalid texel size {0}. Must be greater than zero."), Settings->TexelSize), InContext);
+			return true;
+		}
+
+		FTransform CaptureTransform = FTransform::Identity;
+		FBox CaptureBounds(EForceInit::ForceInit);
+		bool bUnionWasCreated = false;
+
+		if (const UPCGSpatialData* BoundingShape = Context->InputData.GetSpatialUnionOfInputsByPin(Context, PCGSceneCaptureConstants::BoundingShapeLabel, bUnionWasCreated))
+		{
+			// Spatial data does not have a transform, so scene captures through an overridden bounds must be top down and axis aligned.
+			CaptureBounds = BoundingShape->GetBounds();
+		}
+		else
+		{
+			CaptureTransform = ExecutionState.GetTransform();
+			CaptureBounds = ExecutionState.GetBounds();
+		}
+
+		if (!CaptureBounds.IsValid || CaptureBounds.GetVolume() <= UE_DOUBLE_SMALL_NUMBER)
+		{
+			PCGLog::LogErrorOnGraph(LOCTEXT("InvalidCaptureBounds", "Invalid capture bounds. Volume must be greater than zero."), InContext);
+			return true;
+		}
+
+		const FVector CaptureBoundsLocation = CaptureBounds.GetCenter();
+		const FQuat CaptureBoundsRotation = CaptureTransform.GetRotation().GetNormalized(); // Zero quat will return identity when normalized.
+		const FVector CaptureBoundsSize = CaptureBounds.GetSize();
+		const FVector SceneCaptureLocation = CaptureBoundsLocation + CaptureBoundsRotation.GetUpVector() * (CaptureBoundsSize.Z / 2.0);
+		const FQuat SceneCaptureRotation = CaptureBoundsRotation * FQuat(CaptureBoundsRotation.GetRightVector(), FMath::DegreesToRadians(90.0f));
+		const float OrthoWidth = FMath::Max(CaptureBoundsSize.X, CaptureBoundsSize.Y);
+
+		Context->RenderTargetTransform = FTransform(
+			CaptureBoundsRotation * FQuat(CaptureBoundsRotation.GetUpVector(), FMath::DegreesToRadians(90.0f)),
+			CaptureBoundsLocation,
+			FVector(OrthoWidth / 2.0f));
+
+		const int32 MaxTextureDimension = GetMax2DTextureDimension();
+		const int32 RenderTargetDimensions = FMath::RoundToInt(OrthoWidth / Settings->TexelSize);
+
+		if (RenderTargetDimensions <= 0 || RenderTargetDimensions > MaxTextureDimension)
+		{
+			PCGLog::LogErrorOnGraph(
+				FText::Format(LOCTEXT("InvalidDimensions", "Invalid render target dimensions ({0}, {1}), must be between (1, 1) and ({2}, {2})."),
+					RenderTargetDimensions,
+					RenderTargetDimensions,
+					MaxTextureDimension),
+				InContext);
+
+			return true;
+		}
+
+		// Initialize render target.
+		Context->RenderTarget = NewObject<UTextureRenderTarget2D>();
+		Context->RenderTarget->SizeX = RenderTargetDimensions;
+		Context->RenderTarget->SizeY = RenderTargetDimensions;
+		Context->RenderTarget->RenderTargetFormat = Settings->PixelFormat;
+		Context->RenderTarget->ClearColor = FLinearColor::Black;
+		Context->RenderTarget->UpdateResource();
+
+		// Setup scene capture.
+		// @todo_pcg: To improve perf we could avoid relying on a USceneCaptureComponent2D, and do the scene capture ourselves.
+		Context->SceneCaptureComponent = NewObject<USceneCaptureComponent2D>(GetTransientPackage(), NAME_None, RF_Transient);
+		Context->SceneCaptureComponent->bTickInEditor = false;
+		Context->SceneCaptureComponent->SetComponentTickEnabled(false);
+		Context->SceneCaptureComponent->SetVisibility(true);
+		Context->SceneCaptureComponent->bCaptureEveryFrame = false;
+		Context->SceneCaptureComponent->bCaptureOnMovement = false;
+		Context->SceneCaptureComponent->TextureTarget = Context->RenderTarget;
+		Context->SceneCaptureComponent->CaptureSource = Settings->CaptureSource;
+		Context->SceneCaptureComponent->RegisterComponentWithWorld(World);
+		Context->SceneCaptureComponent->ProjectionType = ECameraProjectionMode::Orthographic;
+		Context->SceneCaptureComponent->OrthoWidth = OrthoWidth;
+		Context->SceneCaptureComponent->SetWorldLocation(SceneCaptureLocation);
+		Context->SceneCaptureComponent->SetWorldRotation(SceneCaptureRotation);
+
+		// Perform scene capture.
+		Context->SceneCaptureComponent->CaptureSceneDeferred();
+		Context->bSubmittedSceneCapture = true;
+		Context->bIsPaused = true;
+
+		// Scene capture will be processed alongside the next render frame, so pause for one tick.
+		FPCGModule::GetPCGModuleChecked().ExecuteNextTick([ContextHandle = Context->GetOrCreateHandle()]()
+		{
+			if (TSharedPtr<FPCGContextHandle> SharedHandle = ContextHandle.Pin())
+			{
+				if (FPCGContext* Context = SharedHandle->GetContext())
+				{
+					Context->bIsPaused = false;
+				}
+			}
+		});
+
+		return false;
+	}
+
+	// Cleanup scene capture.
+	Context->SceneCaptureComponent->TextureTarget = nullptr;
+	Context->SceneCaptureComponent->UnregisterComponent();
+
+	// Initialize render target data.
+	UPCGRenderTargetData* RenderTargetData = FPCGContext::NewObject_AnyThread<UPCGRenderTargetData>(InContext);
+	RenderTargetData->Initialize(Context->RenderTarget, Context->RenderTargetTransform, Settings->bSkipReadbackToCPU, /*bInTakeOwnershipOfRenderTarget=*/true);
+
+	FPCGTaggedData& OutputData = InContext->OutputData.TaggedData.Emplace_GetRef();
+	OutputData.Data = RenderTargetData;
+
+	return true;
+}
+
+#undef LOCTEXT_NAMESPACE
diff --git a/Engine/Plugins/PCG/Source/PCG/Public/Elements/PCGSceneCapture.h b/Engine/Plugins/PCG/Source/PCG/Public/Elements/PCGSceneCapture.h
new file mode 100644
--- /dev/null
+++ b/Engine/Plugins/PCG/Source/PCG/Public/Elements/PCGSceneCapture.h
@@ -0,0 +1,72 @@
+// Copyright Epic Games, Inc. All Rights Reserved.
+
+#pragma once
+
+#include "PCGContext.h"
+#include "PCGSettings.h"
+
+#include "Components/SceneCaptureComponent2D.h"
+#include "Engine/TextureRenderTarget2D.h"
+
+#include "PCGSceneCapture.generated.h"
+
+/** Perform a 2D orthographic scene capture and write the result to a render target data. Can be costly, use with caution with runtime generation. */
+UCLASS(MinimalAPI, BlueprintType, ClassGroup = (Procedural))
+class UPCGSceneCaptureSettings : public UPCGSettings
+{
+	GENERATED_BODY()
+
+public:
+	//~Begin UPCGSettings interface
+#if WITH_EDITOR
+	virtual FName GetDefaultNodeName() const override { return FName(TEXT("SceneCapture")); }
+	virtual FText GetDefaultNodeTitle() const override { return NSLOCTEXT("PCGSceneCaptureElement", "NodeTitle", "Scene Capture"); }
+	virtual FText GetNodeTooltipText() const override { return NSLOCTEXT("PCGSceneCaptureElement", "NodeTooltip", "Perform a 2D orthographic scene capture and write the result to a render target data. Can be costly, use with caution with runtime generation."); }
+	virtual EPCGSettingsType GetType() const override { return EPCGSettingsType::Sampler; }
+#endif
+
+protected:
+	virtual TArray<FPCGPinProperties> InputPinProperties() const override;
+	virtual TArray<FPCGPinProperties> OutputPinProperties() const override;
+	virtual FPCGElementPtr CreateElement() const override;
+	//~End UPCGSettings interface
+
+public:
+	/** Subset of EPixelFormat exposed to UTextureRenderTarget2D. */
+	UPROPERTY(BlueprintReadWrite, EditAnywhere, Category = Settings, meta=(PCG_Overridable))
+	TEnumAsByte<enum ETextureRenderTargetFormat> PixelFormat = ETextureRenderTargetFormat::RTF_RGBA16f;
+
+	/** Specifies which component of the scene rendering should be output to the render target. */
+	UPROPERTY(BlueprintReadWrite, EditAnywhere, Category = Settings, meta=(PCG_Overridable))
+	TEnumAsByte<enum ESceneCaptureSource> CaptureSource = ESceneCaptureSource::SCS_SceneColorHDR;
+
+	/** Size of a texel in the render target in world units (cm). */
+	UPROPERTY(BlueprintReadWrite, EditAnywhere, Category = Settings, meta=(UIMin = "1.0", ClampMin = "1.0", PCG_Overridable))
+	float TexelSize = 50.0f;
+
+	/** Skip CPU readback during initialization of the render target data. */
+	UPROPERTY(BlueprintReadWrite, EditAnywhere, Category = Settings)
+	bool bSkipReadbackToCPU = false;
+};
+
+struct FPCGSceneCaptureContext : public FPCGContext
+{
+protected:
+	virtual void AddExtraStructReferencedObjects(FReferenceCollector& Collector) override;
+
+public:
+	TObjectPtr<USceneCaptureComponent2D> SceneCaptureComponent = nullptr;
+	TObjectPtr<UTextureRenderTarget2D> RenderTarget = nullptr;
+	FTransform RenderTargetTransform = FTransform::Identity;
+	bool bSubmittedSceneCapture = false;
+};
+
+class FPCGSceneCaptureElement : public IPCGElementWithCustomContext<FPCGSceneCaptureContext>
+{
+public:
+	/** Required to create the scene capture component. */
+	virtual bool CanExecuteOnlyOnMainThread(FPCGContext* Context) const override { return true; }
+
+protected:
+	virtual bool ExecuteInternal(FPCGContext* InContext) const override;
+};
