[PCG] Decouple PCGComponent from HiGen grid links

diff --git a/Engine/Plugins/PCG/Source/PCG/Private/Editor/PCGEditorDefaultExecutionSource.cpp b/Engine/Plugins/PCG/Source/PCG/Private/Editor/PCGEditorDefaultExecutionSource.cpp
--- a/Engine/Plugins/PCG/Source/PCG/Private/Editor/PCGEditorDefaultExecutionSource.cpp
+++ b/Engine/Plugins/PCG/Source/PCG/Private/Editor/PCGEditorDefaultExecutionSource.cpp
@@ -68,6 +68,12 @@ bool FPCGEditorDefaultExecutionState::IsGenerating() const
 	return Source->CurrentGenerationTask != InvalidPCGTaskId;
 }
 
+IPCGGraphExecutionSource* FPCGEditorDefaultExecutionState::GetOriginalSource() const
+{
+	// Partitioned generation not supported.
+	return Source;
+}
+
 #if WITH_EDITOR
 const PCGUtils::FExtraCapture& FPCGEditorDefaultExecutionState::GetExtraCapture() const
 {
diff --git a/Engine/Plugins/PCG/Source/PCG/Private/Graph/PCGGraphExecutor.cpp b/Engine/Plugins/PCG/Source/PCG/Private/Graph/PCGGraphExecutor.cpp
--- a/Engine/Plugins/PCG/Source/PCG/Private/Graph/PCGGraphExecutor.cpp
+++ b/Engine/Plugins/PCG/Source/PCG/Private/Graph/PCGGraphExecutor.cpp
@@ -3138,15 +3138,17 @@ namespace PCGGraphExecutor
 		check(InContext);
 		TRACE_CPUPROFILER_EVENT_SCOPE(PCGGraphExecutor::ExecuteGridLinkage);
 
-		UPCGComponent* SourceComponent = Cast<UPCGComponent>(InContext->ExecutionSource.Get());
-		if (!SourceComponent)
+		const IPCGGraphExecutionSource* ExecutionSource = InContext->ExecutionSource.Get();
+		if (!ExecutionSource)
 		{
 			// Source no longer exists, nothing to be done.
 			return true;
 		}
 
+		const IPCGGraphExecutionState& ExecutionState = InContext->ExecutionSource->GetExecutionState();
+
 		// Non-hierarchical generation - no linkage required - data should just pass through.
-		if (!InContext->ExecutionSource->GetExecutionState().GetGraph()->IsHierarchicalGenerationEnabled()
+		if (!ExecutionState.GetGraph()->IsHierarchicalGenerationEnabled()
 			|| !ensure(PCGHiGenGrid::IsValidGrid(InFromGrid) || InFromGrid == EPCGHiGenGrid::Unbounded))
 		{
 			InContext->OutputData = InContext->InputData;
@@ -3199,58 +3201,49 @@ namespace PCGGraphExecutor
 
 			PCGGraphExecutionLogging::LogGridLinkageTaskExecuteStore(InContext, InGenerationGrid, FromGridSize, ToGridSize, InResourceKey, Data.TaggedData.Num());
 			
-			SourceComponent->StoreOutputDataForPin(InResourceKey, Data);
+			ExecutionState.StoreOutputDataForPin(InResourceKey, Data);
 		}
 		else if (InToGrid == InGenerationGrid && FromGridSize != ToGridSize)
 		{
 			PCGGraphExecutionLogging::LogGridLinkageTaskExecuteRetrieve(InContext, InGenerationGrid, FromGridSize, ToGridSize, InResourceKey);
 
-			UPCGSubsystem* Subsystem = UPCGSubsystem::GetInstance(SourceComponent->GetWorld());
-			if (!ensure(Subsystem))
-			{
-				return false;
-			}
+			IPCGGraphExecutionSource* ExecSourceWithData = nullptr;
+			bool bExecSourceWithDataIsOriginal = false;
 
-			UPCGComponent* ComponentWithData = nullptr;
-			bool bComponentWithDataIsOriginalComponent = false;
 			if (FromGridSize == PCGHiGenGrid::UnboundedGridSize())
 			{
-				ComponentWithData = SourceComponent->GetOriginalComponent();
-				bComponentWithDataIsOriginalComponent = true;
+				ExecSourceWithData = ExecutionState.GetOriginalSource();
+				bExecSourceWithDataIsOriginal = true;
 			}
 			else
 			{
-				const AActor* ComponentActor = SourceComponent->GetOwner();
-				if (ComponentActor)
+				const IPCGGraphExecutionSource* OriginalSource = ExecutionState.GetOriginalSource();
+				if (!ensure(OriginalSource))
 				{
-					const UPCGComponent* OriginalComponent = SourceComponent->GetOriginalComponent();
-					if (!ensure(OriginalComponent))
-					{
-						return false;
-					}
+					return false;
+				}
 
-					const FPCGGridDescriptor GridDescriptor = OriginalComponent->GetGridDescriptor(FromGridSize);
+				const FPCGGridDescriptor GridDescriptor = OriginalSource->GetExecutionState().GetGridDescriptor(FromGridSize);
 
-					// Get grid coords using the parent grid (FromGridSize).
-					const FIntVector CellCoords = UPCGActorHelpers::GetCellCoord(ComponentActor->GetActorLocation(), GridDescriptor.GetGridSize(), GridDescriptor.Is2DGrid());
+				// Get grid coords using the parent grid (FromGridSize).
+				const FIntVector CellCoords = UPCGActorHelpers::GetCellCoord(ExecutionState.GetTransform().GetLocation(), GridDescriptor.GetGridSize(), GridDescriptor.Is2DGrid());
 
-					ComponentWithData = Subsystem->GetLocalComponent(GridDescriptor, CellCoords, OriginalComponent);
-				}
+				ExecSourceWithData = ExecutionState.GetLocalSource(GridDescriptor, CellCoords);
 			}
 
-			if (!ComponentWithData)
+			if (!ExecSourceWithData)
 			{
-				// Nothing we can do currently if PCG component not present. One idea is to schedule an artifact-less execution but that
+				// Nothing we can do currently if execution source is not present. One idea is to schedule an artifact-less execution but that
 				// comes with complications - artifacts/side effects are an integral part of execution. Most likely we'll do a cleanup
-				// pass of any unwanted artifacts/local-components later.
+				// pass of any unwanted artifacts/local-execution sources later.
 				PCGGraphExecutionLogging::LogGridLinkageTaskExecuteRetrieveNoLocalSource(InContext, InResourceKey);
 				return true;
 			}
 
-			// Once we've found our component, try to retrieve the data.
-			if (const FPCGDataCollection* Data = ComponentWithData->RetrieveOutputDataForPin(InResourceKey))
+			// Once we've found our execution source, try to retrieve the data.
+			if (const FPCGDataCollection* Data = ExecSourceWithData->GetExecutionState().RetrieveOutputDataForPin(InResourceKey))
 			{
-				PCGGraphExecutionLogging::LogGridLinkageTaskExecuteRetrieveSuccess(InContext, ComponentWithData, InResourceKey, Data->TaggedData.Num());
+				PCGGraphExecutionLogging::LogGridLinkageTaskExecuteRetrieveSuccess(InContext, ExecSourceWithData, InResourceKey, Data->TaggedData.Num());
 				InContext->OutputData = *Data;
 
 				// Grid Linkages have a single virtual output pin which forwards the same label as the upstream output pin. Apply the label now.
@@ -3265,39 +3258,36 @@ namespace PCGGraphExecutor
 			}
 
 			// At this point we could not get to the data, so we'll try executing the graph if we did not do that already.
-			// If we need data from a local component but the local component is still generating, then we'll wait for it. On the other hand
-			// if we need data from the original component we assume the generation has already happened because it is always scheduled before
-			// the local components.
-			const bool bWaitForGeneration = bComponentWithDataIsOriginalComponent ? false : ComponentWithData->IsGenerating();
+			// If we need data from a local execution source but the local execution source is still generating, then we'll wait for it. On the other hand
+			// if we need data from the original execution source we assume the generation has already happened because it is always scheduled before
+			// the local execution source.
+			const bool bWaitForGeneration = bExecSourceWithDataIsOriginal ? false : ExecSourceWithData->GetExecutionState().IsGenerating();
 			if (bWaitForGeneration)
 			{
-				PCGGraphExecutionLogging::LogGridLinkageTaskExecuteRetrieveWaitOnScheduledGraph(InContext, ComponentWithData, InResourceKey);
+				PCGGraphExecutionLogging::LogGridLinkageTaskExecuteRetrieveWaitOnScheduledGraph(InContext, ExecSourceWithData, InResourceKey);
 
-				// The component was already generating, but we were not paused. Not really clear what's happening here,
+				// The execution source was already generating, but we were not paused. Not really clear what's happening here,
 				// but in any case go to paused and unpause when it's done.
 				InContext->bIsPaused = true;
 
 				// Wake up this task after graph has generated.
-				FPCGTaskId GenerationTask = ComponentWithData->GetGenerationTaskId();
+				FPCGTaskId GenerationTask = ExecSourceWithData->GetExecutionState().GetGenerationTaskId();
 				if (ensure(GenerationTask != InvalidPCGTaskId))
 				{
-					InContext->DynamicDependencies.Add(ComponentWithData->GetGenerationTaskId());
+					InContext->DynamicDependencies.Add(GenerationTask);
 				}
 
 				return false;
 			}
 
 			// Graph is not currently generating. If we have not already tried generating, try it once now.
-			// But don't do this for the original component as this will always be scheduled before the local components.
-			if (!InContext->bScheduledGraph && !bComponentWithDataIsOriginalComponent)
+			// But don't do this for the original execution source as this will always be scheduled before the local execution source.
+			if (!InContext->bScheduledGraph && !bExecSourceWithDataIsOriginal)
 			{
-				PCGGraphExecutionLogging::LogGridLinkageTaskExecuteRetrieveScheduleGraph(InContext, ComponentWithData, InResourceKey);
-
-				EPCGComponentGenerationTrigger GenTrigger = PCGHelpers::IsRuntimeGeneration(InContext->ExecutionSource.Get()) ?
-					EPCGComponentGenerationTrigger::GenerateAtRuntime : EPCGComponentGenerationTrigger::GenerateOnDemand;
+				PCGGraphExecutionLogging::LogGridLinkageTaskExecuteRetrieveScheduleGraph(InContext, ExecSourceWithData, InResourceKey);
 
 				// Wake up this task after graph has generated.
-				const FPCGTaskId GraphTaskId = ComponentWithData->GenerateLocalGetTaskId(GenTrigger, /*bForce=*/true);
+				const FPCGTaskId GraphTaskId = ExecSourceWithData->GetExecutionState().GenerateLocalGetTaskId();
 
 				// Update state and go to paused.
 				InContext->DynamicDependencies.Add(GraphTaskId);
@@ -3308,7 +3298,7 @@ namespace PCGGraphExecutor
 			else
 			{
 				// We tried generating but no luck, at this point give up.
-				PCGGraphExecutionLogging::LogGridLinkageTaskExecuteRetrieveNoData(InContext, ComponentWithData, InResourceKey);
+				PCGGraphExecutionLogging::LogGridLinkageTaskExecuteRetrieveNoData(InContext, ExecSourceWithData, InResourceKey);
 				return true;
 			}
 		}
diff --git a/Engine/Plugins/PCG/Source/PCG/Private/Helpers/PCGHelpers.cpp b/Engine/Plugins/PCG/Source/PCG/Private/Helpers/PCGHelpers.cpp
--- a/Engine/Plugins/PCG/Source/PCG/Private/Helpers/PCGHelpers.cpp
+++ b/Engine/Plugins/PCG/Source/PCG/Private/Helpers/PCGHelpers.cpp
@@ -567,26 +567,12 @@ namespace PCGHelpers
 
 	int32 GetGenerationGridSize(const IPCGGraphExecutionSource* InExecutionSource)
 	{
-		if (const UPCGComponent* PCGComponent = Cast<UPCGComponent>(InExecutionSource))
-		{
-			return PCGComponent->GetGenerationGridSize();
-		}
-		else
-		{
-			return PCGHiGenGrid::UninitializedGridSize();
-		}
+		return InExecutionSource ? InExecutionSource->GetExecutionState().GetGenerationGridSize() : PCGHiGenGrid::UninitializedGridSize();
 	}
 
 	bool IsRuntimeGeneration(const IPCGGraphExecutionSource* InExecutionSource)
 	{
-		if (const UPCGComponent* PCGComponent = Cast<UPCGComponent>(InExecutionSource))
-		{
-			return PCGComponent->IsManagedByRuntimeGenSystem();
-		}
-		else
-		{
-			return false;
-		}
+		return InExecutionSource && InExecutionSource->GetExecutionState().IsManagedByRuntimeGenSystem();
 	}
 
 #if WITH_EDITOR
diff --git a/Engine/Plugins/PCG/Source/PCG/Private/PCGComponent.cpp b/Engine/Plugins/PCG/Source/PCG/Private/PCGComponent.cpp
--- a/Engine/Plugins/PCG/Source/PCG/Private/PCGComponent.cpp
+++ b/Engine/Plugins/PCG/Source/PCG/Private/PCGComponent.cpp
@@ -4284,6 +4284,12 @@ FBox FPCGComponentExecutionState::GetBounds() const
 	return Component->GetGridBounds();
 }
 
+bool FPCGComponentExecutionState::Use2DGrid() const
+{
+	check(Component);
+	return Component->Use2DGrid();
+}
+
 UPCGGraph* FPCGComponentExecutionState::GetGraph() const
 {
 	check(Component);
@@ -4314,6 +4320,68 @@ bool FPCGComponentExecutionState::IsGenerating() const
 	return Component->IsGenerating();
 }
 
+bool FPCGComponentExecutionState::IsManagedByRuntimeGenSystem() const
+{
+	check(Component);
+	return Component->IsManagedByRuntimeGenSystem();
+}
+
+IPCGGraphExecutionSource* FPCGComponentExecutionState::GetOriginalSource() const
+{
+	check(Component);
+	return Component->GetOriginalComponent();
+}
+
+IPCGGraphExecutionSource* FPCGComponentExecutionState::GetLocalSource(const FPCGGridDescriptor& GridDescriptor, const FIntVector& CellCoords) const
+{
+	check(Component);
+
+	if (UPCGSubsystem* Subsystem = Cast<UPCGSubsystem>(GetSubsystem()))
+	{
+		return Subsystem->GetLocalComponent(GridDescriptor, CellCoords, CastChecked<UPCGComponent>(GetOriginalSource()));
+	}
+
+	return nullptr;
+}
+
+FPCGGridDescriptor FPCGComponentExecutionState::GetGridDescriptor(uint32 InGridSize) const
+{
+	check(Component);
+	return Component->GetGridDescriptor(InGridSize);
+}
+
+int32 FPCGComponentExecutionState::GetGenerationGridSize() const
+{
+	check(Component);
+	return Component->GetGenerationGridSize();
+}
+
+FPCGTaskId FPCGComponentExecutionState::GetGenerationTaskId() const
+{
+	check(Component);
+	return Component->GetGenerationTaskId();
+}
+
+void FPCGComponentExecutionState::StoreOutputDataForPin(const FString& InResourceKey, const FPCGDataCollection& InData) const
+{
+	check(Component);
+	Component->StoreOutputDataForPin(InResourceKey, InData);
+}
+
+const FPCGDataCollection* FPCGComponentExecutionState::RetrieveOutputDataForPin(const FString& InResourceKey) const
+{
+	check(Component);
+	return Component->RetrieveOutputDataForPin(InResourceKey);
+}
+
+FPCGTaskId FPCGComponentExecutionState::GenerateLocalGetTaskId(EPCGHiGenGrid Grid)
+{
+	check(Component)
+
+	const EPCGComponentGenerationTrigger GenTrigger = IsManagedByRuntimeGenSystem() ? EPCGComponentGenerationTrigger::GenerateAtRuntime : EPCGComponentGenerationTrigger::GenerateOnDemand;
+	return Component->GenerateLocalGetTaskId(GenTrigger, /*bForce=*/true, Grid);
+}
+
 #if WITH_EDITOR
 const PCGUtils::FExtraCapture& FPCGComponentExecutionState::GetExtraCapture() const
 {
diff --git a/Engine/Plugins/PCG/Source/PCG/Private/PCGGraphExecutionStateInterface.cpp b/Engine/Plugins/PCG/Source/PCG/Private/PCGGraphExecutionStateInterface.cpp
--- a/Engine/Plugins/PCG/Source/PCG/Private/PCGGraphExecutionStateInterface.cpp
+++ b/Engine/Plugins/PCG/Source/PCG/Private/PCGGraphExecutionStateInterface.cpp
@@ -25,6 +25,14 @@ IPCGBaseSubsystem* IPCGGraphExecutionState::GetSubsystem() const
 	return Subsystem;
 }
 
+FPCGGridDescriptor IPCGGraphExecutionState::GetGridDescriptor(uint32 InGridSize) const
+{
+	return FPCGGridDescriptor()
+		.SetGridSize(InGridSize)
+		.SetIs2DGrid(Use2DGrid())
+		.SetIsRuntime(IsManagedByRuntimeGenSystem());
+}
+
 #if WITH_EDITOR
 
 FPCGDynamicTrackingPriority IPCGGraphExecutionState::GetDynamicTrackingPriority() const
diff --git a/Engine/Plugins/PCG/Source/PCG/Public/Editor/PCGEditorDefaultExecutionSource.h b/Engine/Plugins/PCG/Source/PCG/Public/Editor/PCGEditorDefaultExecutionSource.h
--- a/Engine/Plugins/PCG/Source/PCG/Public/Editor/PCGEditorDefaultExecutionSource.h
+++ b/Engine/Plugins/PCG/Source/PCG/Public/Editor/PCGEditorDefaultExecutionSource.h
@@ -29,6 +29,7 @@ public:
 	PCG_API virtual void OnGraphExecutionAborted(bool bQuiet, bool bCleanupUnusedResources) override;
 	PCG_API virtual void Cancel() override;
 	PCG_API virtual bool IsGenerating() const override;
+	PCG_API virtual IPCGGraphExecutionSource* GetOriginalSource() const override;
 
 #if WITH_EDITOR
 	PCG_API virtual const PCGUtils::FExtraCapture& GetExtraCapture() const override;
diff --git a/Engine/Plugins/PCG/Source/PCG/Public/PCGComponent.h b/Engine/Plugins/PCG/Source/PCG/Public/PCGComponent.h
--- a/Engine/Plugins/PCG/Source/PCG/Public/PCGComponent.h
+++ b/Engine/Plugins/PCG/Source/PCG/Public/PCGComponent.h
@@ -103,11 +103,21 @@ public:
 	virtual UWorld* GetWorld() const override;
 	virtual bool HasAuthority() const override;
 	virtual FBox GetBounds() const override;
+	virtual bool Use2DGrid() const override;
 	virtual UPCGGraph* GetGraph() const override;
 	virtual UPCGGraphInstance* GetGraphInstance() const override;
 	virtual void OnGraphExecutionAborted(bool bQuiet, bool bCleanupUnusedResources) override;
 	virtual void Cancel() override;
 	virtual bool IsGenerating() const override;
+	virtual bool IsManagedByRuntimeGenSystem() const override;
+	virtual IPCGGraphExecutionSource* GetOriginalSource() const override;
+	virtual IPCGGraphExecutionSource* GetLocalSource(const FPCGGridDescriptor& GridDescriptor, const FIntVector& CellCoords) const override;
+	virtual FPCGGridDescriptor GetGridDescriptor(uint32 InGridSize) const override;
+	virtual int32 GetGenerationGridSize() const override;
+	virtual FPCGTaskId GetGenerationTaskId() const override;
+	virtual void StoreOutputDataForPin(const FString& InResourceKey, const FPCGDataCollection& InData) const override;
+	virtual const FPCGDataCollection* RetrieveOutputDataForPin(const FString& InResourceKey) const override;
+	virtual FPCGTaskId GenerateLocalGetTaskId(EPCGHiGenGrid Grid = EPCGHiGenGrid::Uninitialized) override;
 
 #if WITH_EDITOR
 	virtual const PCGUtils::FExtraCapture& GetExtraCapture() const override;
@@ -186,7 +196,7 @@ public:
 	UE_API UPCGData* GetLandscapeHeightPCGData() const;
 	UE_API UPCGData* GetOriginalActorPCGData() const;
 
-	/** If this is a local component returns self, otherwise returns the original component. */
+	/** If this is a local component returns the original component, otherwise returns self. */
 	UE_API UPCGComponent* GetOriginalComponent() const;
 	UE_API const UPCGComponent* GetConstOriginalComponent() const;
 
diff --git a/Engine/Plugins/PCG/Source/PCG/Public/PCGGraphExecutionStateInterface.h b/Engine/Plugins/PCG/Source/PCG/Public/PCGGraphExecutionStateInterface.h
--- a/Engine/Plugins/PCG/Source/PCG/Public/PCGGraphExecutionStateInterface.h
+++ b/Engine/Plugins/PCG/Source/PCG/Public/PCGGraphExecutionStateInterface.h
@@ -17,6 +17,7 @@ class IPCGGraphExecutionSource;
 class UPCGData;
 class UPCGGraph;
 class UPCGGraphInstance;
+struct FPCGGridDescriptor;
 
 namespace PCGUtils
 {
@@ -63,12 +64,24 @@ public:
 	/** Cancel execution of this ExecutionState. */
 	virtual void Cancel() = 0;
 
-	/** Notify ExecutionState that its execution is being aborted */
+	/** Notify ExecutionState that its execution is being aborted. */
 	virtual void OnGraphExecutionAborted(bool bQuiet = false, bool bCleanupUnusedResources = true) = 0;
 
-	/** Indicates if the current source is generating */
+	/** Indicates if the current source is generating. */
 	virtual bool IsGenerating() const = 0;
 
+	/** Gets the task ID for the current generation task if it exists. */
+	virtual FPCGTaskId GetGenerationTaskId() const { return InvalidPCGTaskId; }
+
+	/** Requests the execution source to generate only on the specified grid level (all grid levels if EPCGHiGenGrid::Uninitialized). */
+	virtual FPCGTaskId GenerateLocalGetTaskId(EPCGHiGenGrid Grid = EPCGHiGenGrid::Uninitialized) { return InvalidPCGTaskId; }
+
+	/** Store data with a resource key that identifies the pin. */
+	virtual void StoreOutputDataForPin(const FString& InResourceKey, const FPCGDataCollection& InData) const {}
+
+	/** Lookup data using a resource key that identifies the pin. */
+	virtual const FPCGDataCollection* RetrieveOutputDataForPin(const FString& InResourceKey) const { return nullptr; }
+
 	/** Returns the subsystem associated with this source. By default, it will return the UPCGSubsystem if the source has a world, otherwise, the EditorSubsystem (in editor only, obviously). */
 	PCG_API virtual IPCGBaseSubsystem* GetSubsystem() const;
 
@@ -97,6 +110,25 @@ public:
 	virtual bool IsRefreshInProgress() const = 0;
 
 #endif
+
+	/** If this is a local execution source returns the original execution source, otherwise returns self. */
+	virtual IPCGGraphExecutionSource* GetOriginalSource() const = 0;
+
+	/** Retrieves a local execution source using grid descriptor and grid coordinates, returns nullptr if no such execution source is found. */
+	virtual IPCGGraphExecutionSource* GetLocalSource(const FPCGGridDescriptor& GridDescriptor, const FIntVector& CellCoords) const { return nullptr; }
+
+	// @todo_pcg: The execution source should probably just have an execution domain instead of directly querying runtime gen vs etc.
+	/** Returns true if the execution source is managed by the runtime generation system. Nothing else should generate or cleanup this execution source. */
+	virtual bool IsManagedByRuntimeGenSystem() const { return false; }
+
+	/** Returns true if execution source should output on a 2D grid. */
+	virtual bool Use2DGrid() const { return false; }
+
+	/** Returns a GridDescriptor based on this execution source for the specified grid size. */
+	virtual FPCGGridDescriptor GetGridDescriptor(uint32 InGridSize) const;
+
+	/** Get the size for the HiGen grid this execution source executes on. */
+	virtual int32 GetGenerationGridSize() const { return PCGHiGenGrid::UninitializedGridSize(); }
 };
 
 UINTERFACE(BlueprintType, MinimalAPI)
