[PCG] Minor code cleanup, removed spaces in labels, cleaned up access specifiers.

#rnx

diff --git a/Engine/Plugins/PCG/Source/PCG/Private/Elements/Polygon/PCGClipPaths.cpp b/Engine/Plugins/PCG/Source/PCG/Private/Elements/Polygon/PCGClipPaths.cpp
--- a/Engine/Plugins/PCG/Source/PCG/Private/Elements/Polygon/PCGClipPaths.cpp
+++ b/Engine/Plugins/PCG/Source/PCG/Private/Elements/Polygon/PCGClipPaths.cpp
@@ -3,6 +3,7 @@
 #include "Elements/Polygon/PCGClipPaths.h"
 
 #include "PCGContext.h"
+#include "PCGCustomVersion.h"
 #include "Data/PCGBasePointData.h"
 #include "Data/PCGSplineData.h"
 #include "Data/PCGPolygon2DData.h"
@@ -27,6 +28,16 @@ FText UPCGClipPathsSettings::GetNodeTooltipText() const
 {
 	return LOCTEXT("NodeTooltip", "Clips paths against the polygons, to have the paths either inscribed in the polygons (intersection) or outside the polygons (difference).");
 }
+
+void UPCGClipPathsSettings::ApplyDeprecationBeforeUpdatePins(UPCGNode* InOutNode, TArray<TObjectPtr<UPCGPin>>& InputPins, TArray<TObjectPtr<UPCGPin>>& OutputPins)
+{
+	Super::ApplyDeprecationBeforeUpdatePins(InOutNode, InputPins, OutputPins);
+
+	if (DataVersion < FPCGCustomVersion::RemovedSpacesInPolygonPinLabels)
+	{
+		InOutNode->RenameInputPin(PCGPolygon2DUtils::Constants::Deprecated::OldClipPolysLabel, PCGPolygon2DUtils::Constants::ClipPolysLabel);
+	}
+}
 #endif // WITH_EDITOR
 
 FString UPCGClipPathsSettings::GetAdditionalTitleInformation() const
diff --git a/Engine/Plugins/PCG/Source/PCG/Private/Elements/Polygon/PCGPolygon2DOperation.cpp b/Engine/Plugins/PCG/Source/PCG/Private/Elements/Polygon/PCGPolygon2DOperation.cpp
--- a/Engine/Plugins/PCG/Source/PCG/Private/Elements/Polygon/PCGPolygon2DOperation.cpp
+++ b/Engine/Plugins/PCG/Source/PCG/Private/Elements/Polygon/PCGPolygon2DOperation.cpp
@@ -3,6 +3,7 @@
 #include "Elements/Polygon/PCGPolygon2DOperation.h"
 
 #include "PCGContext.h"
+#include "PCGCustomVersion.h"
 #include "Data/PCGPolygon2DData.h"
 #include "Helpers/PCGPolygon2DProcessingHelpers.h"
 
@@ -34,6 +35,17 @@ EPCGChangeType UPCGPolygon2DOperationSettings::GetChangeTypeForProperty(FPropert
 
 	return ChangeType;
 }
+
+void UPCGPolygon2DOperationSettings::ApplyDeprecationBeforeUpdatePins(UPCGNode* InOutNode, TArray<TObjectPtr<UPCGPin>>& InputPins, TArray<TObjectPtr<UPCGPin>>& OutputPins)
+{
+	Super::ApplyDeprecationBeforeUpdatePins(InOutNode, InputPins, OutputPins);
+
+	if (DataVersion < FPCGCustomVersion::RemovedSpacesInPolygonPinLabels)
+	{
+		InOutNode->RenameInputPin(PCGPolygon2DUtils::Constants::Deprecated::OldClipPolysLabel, PCGPolygon2DUtils::Constants::ClipPolysLabel);
+		InOutNode->RenameInputPin(PCGPolygon2DUtils::Constants::Deprecated::OldClipPathsLabel, PCGPolygon2DUtils::Constants::ClipPathsLabel);
+	}
+}
 #endif // WITH_EDITOR
 
 FString UPCGPolygon2DOperationSettings::GetAdditionalTitleInformation() const
diff --git a/Engine/Plugins/PCG/Source/PCG/Public/Elements/Polygon/PCGClipPaths.h b/Engine/Plugins/PCG/Source/PCG/Public/Elements/Polygon/PCGClipPaths.h
--- a/Engine/Plugins/PCG/Source/PCG/Public/Elements/Polygon/PCGClipPaths.h
+++ b/Engine/Plugins/PCG/Source/PCG/Public/Elements/Polygon/PCGClipPaths.h
@@ -36,6 +36,7 @@ public:
 	virtual FText GetDefaultNodeTitle() const override;
 	virtual FText GetNodeTooltipText() const override;
 	virtual EPCGSettingsType GetType() const override { return EPCGSettingsType::Spatial; }
+	virtual void ApplyDeprecationBeforeUpdatePins(UPCGNode* InOutNode, TArray<TObjectPtr<UPCGPin>>& InputPins, TArray<TObjectPtr<UPCGPin>>& OutputPins) override;
 #endif
 	virtual FString GetAdditionalTitleInformation() const override;
 
diff --git a/Engine/Plugins/PCG/Source/PCG/Public/Elements/Polygon/PCGCreatePolygon2D.h b/Engine/Plugins/PCG/Source/PCG/Public/Elements/Polygon/PCGCreatePolygon2D.h
--- a/Engine/Plugins/PCG/Source/PCG/Public/Elements/Polygon/PCGCreatePolygon2D.h
+++ b/Engine/Plugins/PCG/Source/PCG/Public/Elements/Polygon/PCGCreatePolygon2D.h
@@ -17,12 +17,10 @@ enum class EPCGCreatePolygonInputType : uint8
 
 class FPCGCreatePolygonElement : public IPCGElement
 {
-public:
-	virtual bool ExecuteInternal(FPCGContext* InContext) const override;
-	virtual EPCGElementExecutionLoopMode ExecutionLoopMode(const UPCGSettings* Settings) const override { return EPCGElementExecutionLoopMode::SinglePrimaryPin; }
-
 protected:
+	virtual bool ExecuteInternal(FPCGContext* InContext) const override;
 	virtual bool SupportsBasePointDataInputs(FPCGContext* InContext) const { return true; }
+	virtual EPCGElementExecutionLoopMode ExecutionLoopMode(const UPCGSettings* Settings) const override { return EPCGElementExecutionLoopMode::SinglePrimaryPin; }
 };
 
 /**
diff --git a/Engine/Plugins/PCG/Source/PCG/Public/Elements/Polygon/PCGPolygon2DOperation.h b/Engine/Plugins/PCG/Source/PCG/Public/Elements/Polygon/PCGPolygon2DOperation.h
--- a/Engine/Plugins/PCG/Source/PCG/Public/Elements/Polygon/PCGPolygon2DOperation.h
+++ b/Engine/Plugins/PCG/Source/PCG/Public/Elements/Polygon/PCGPolygon2DOperation.h
@@ -29,7 +29,7 @@ enum class EPCGPolygonOperationMetadataMode : uint8
 
 class FPCGPolygon2DOperationElement : public IPCGElement
 {
-public:
+protected:
 	virtual bool ExecuteInternal(FPCGContext* InContext) const override;
 	virtual EPCGElementExecutionLoopMode ExecutionLoopMode(const UPCGSettings* Settings) const override;
 };
@@ -51,6 +51,7 @@ public:
 	virtual EPCGSettingsType GetType() const override { return EPCGSettingsType::Spatial; }
 	// @todo_pcg : presets per operation
 	virtual EPCGChangeType GetChangeTypeForProperty(FPropertyChangedEvent& PropertyChangedEvent) const override;
+	virtual void ApplyDeprecationBeforeUpdatePins(UPCGNode* InOutNode, TArray<TObjectPtr<UPCGPin>>& InputPins, TArray<TObjectPtr<UPCGPin>>& OutputPins) override;
 #endif
 	virtual FString GetAdditionalTitleInformation() const override;
 
diff --git a/Engine/Plugins/PCG/Source/PCG/Public/Elements/Polygon/PCGPolygon2DUtils.h b/Engine/Plugins/PCG/Source/PCG/Public/Elements/Polygon/PCGPolygon2DUtils.h
--- a/Engine/Plugins/PCG/Source/PCG/Public/Elements/Polygon/PCGPolygon2DUtils.h
+++ b/Engine/Plugins/PCG/Source/PCG/Public/Elements/Polygon/PCGPolygon2DUtils.h
@@ -8,8 +8,15 @@ namespace PCGPolygon2DUtils
 {
 	namespace Constants
 	{
-		const FLazyName ClipPolysLabel = TEXT("Clip Polygons");
-		const FLazyName ClipPathsLabel = TEXT("Clip Paths");
+		const FLazyName ClipPolysLabel = TEXT("ClipPolygons");
+		const FLazyName ClipPathsLabel = TEXT("ClipPaths");
+
+		// Since these were introduced in 5.7, they will be removed as of 5.8.
+		namespace Deprecated
+		{
+			const FLazyName OldClipPolysLabel = TEXT("Clip Polygons");
+			const FLazyName OldClipPathsLabel = TEXT("Clip Paths");
+		}
 	}
 
 	TArray<FPCGPinProperties> DefaultPolygonInputPinProperties();
diff --git a/Engine/Plugins/PCG/Source/PCG/Public/PCGCustomVersion.h b/Engine/Plugins/PCG/Source/PCG/Public/PCGCustomVersion.h
--- a/Engine/Plugins/PCG/Source/PCG/Public/PCGCustomVersion.h
+++ b/Engine/Plugins/PCG/Source/PCG/Public/PCGCustomVersion.h
@@ -89,6 +89,9 @@ struct FPCGCustomVersion
 		// Changed Attribute rename to support selectors
 		AttributeRenameSupportToSelectors = 26,
 
+		// Removed spaces from some polygon 2d node pin labels
+		RemovedSpacesInPolygonPinLabels = 27,
+
 		// -----<new versions can be added above this line>-------------------------------------------------
 		VersionPlusOne,
 		LatestVersion = VersionPlusOne - 1
