[PCG] New simple point processor blueprint element
- Implement Initialize (optional)
- Implement PointLoopBody to process input to output
- Allows inheritance and selection of allocated properties on output
- New BP Structs FPCGPointOutputRange / FPCGPointInputRange allow getting / setting range values on point properties
- PCGBasePointData Set/Get API : SetTransformValuesFromRange / GetTransformValuesFromRange (Density, BoundsMin, BoundsMax, Color, Steepness, Seed, MetadataEntry)
- Move GetContext/GetSeed/GetRandomStream to deprecated BP Element
- Introduce FPCGBlueprintContextHandle and GetContextHandle/GetSeedWithContext/GetRandomStreamWithContext in the base element class

diff --git a/Engine/Plugins/PCG/Source/PCG/Private/Blueprint/PCGBlueprintDeprecatedElement.cpp b/Engine/Plugins/PCG/Source/PCG/Private/Blueprint/PCGBlueprintDeprecatedElement.cpp
--- a/Engine/Plugins/PCG/Source/PCG/Private/Blueprint/PCGBlueprintDeprecatedElement.cpp
+++ b/Engine/Plugins/PCG/Source/PCG/Private/Blueprint/PCGBlueprintDeprecatedElement.cpp
@@ -174,4 +174,21 @@ void UPCGBlueprintElement::IterationLoop(FPCGContext& InContext, int64 NumIterat
 	std::swap(bPreviousBPStateValue, InContext.AsyncState.bIsCallingBlueprint);
 }
 
+FPCGContext& UPCGBlueprintElement::GetContext() const
+{
+	checkf(CurrentContext, TEXT("Execution context is not ready - do not call the GetContext method inside of non-execution methods"));
+	return *CurrentContext;
+}
+
+int UPCGBlueprintElement::GetSeed(FPCGContext& InContext) const
+{
+	return InContext.GetSeed();
+}
+
+FRandomStream UPCGBlueprintElement::GetRandomStream(FPCGContext& InContext) const
+{
+	return FRandomStream(GetSeed(InContext));
+}
+
+
 #undef LOCTEXT_NAMESPACE
diff --git a/Engine/Plugins/PCG/Source/PCG/Private/Data/PCGBasePointData.cpp b/Engine/Plugins/PCG/Source/PCG/Private/Data/PCGBasePointData.cpp
--- a/Engine/Plugins/PCG/Source/PCG/Private/Data/PCGBasePointData.cpp
+++ b/Engine/Plugins/PCG/Source/PCG/Private/Data/PCGBasePointData.cpp
@@ -424,6 +424,81 @@ namespace PCGBasePointData
 
 		return Accessor;
 	}
+
+	template<EPCGPointNativeProperties PointProperty>
+	bool IsValidRange(const FPCGPointInputRange& Range, const typename TPCGPointNativeProperty<PointProperty>::ConstValueRange& ValueRange)
+	{
+		return Range.RangeSize > 0 && ValueRange.IsValidIndex(Range.RangeStartIndex) && ValueRange.IsValidIndex(Range.RangeStartIndex + Range.RangeSize - 1);
+	}
+
+	template<EPCGPointNativeProperties PointProperty>
+	TArray<typename TPCGPointNativeProperty<PointProperty>::Type> GetInputRangeValues(const FPCGPointInputRange& Range)
+	{
+		if (!Range.PointData)
+		{
+			UE_LOG(LogPCG, Error, TEXT("GetInputRangeValues '%s' called with invalid point data."), *StaticEnum<EPCGPointNativeProperties>()->GetNameStringByValue((int64)PointProperty));
+			return {};
+		}
+
+		const typename TPCGPointNativeProperty<PointProperty>::ConstValueRange ValueRange = Range.PointData->GetConstValueRange<PointProperty>();
+
+		if (!PCGBasePointData::IsValidRange<PointProperty>(Range, ValueRange))
+		{
+			UE_LOG(LogPCG, Error, TEXT("GetInputRangeValues '%s' called with input range [%d, %d[ when valid range is [0, %d[."), *StaticEnum<EPCGPointNativeProperties>()->GetNameStringByValue((int64)PointProperty), Range.RangeStartIndex, Range.RangeSize, Range.PointData->GetNumPoints());
+			return {};
+		}
+
+		TArray<typename TPCGPointNativeProperty<PointProperty>::Type> Values;
+		Values.Reserve(Range.RangeSize);
+
+		for (int32 Index = 0; Index < Range.RangeSize; ++Index)
+		{
+			Values.Add(ValueRange[Range.RangeStartIndex + Index]);
+		}
+
+		return Values;
+	}
+
+	template<EPCGPointNativeProperties PointProperty>
+	bool IsValidRange(const FPCGPointOutputRange& Range, const typename TPCGPointNativeProperty<PointProperty>::ValueRange& ValueRange)
+	{
+		return Range.RangeSize > 0 && ValueRange.IsValidIndex(Range.RangeStartIndex) && ValueRange.IsValidIndex(Range.RangeStartIndex + Range.RangeSize - 1);
+	}
+
+	template<EPCGPointNativeProperties PointProperty>
+	void SetOutputRangeValues(FPCGPointOutputRange& Range, const TArray<typename TPCGPointNativeProperty<PointProperty>::Type>& Values)
+	{
+		if (!IsValid(Range.PointData))
+		{
+			UE_LOG(LogPCG, Error, TEXT("SetOutputRangeValues '%s' called with invalid point data."), *StaticEnum<EPCGPointNativeProperties>()->GetNameStringByValue((int64)PointProperty));
+			return;
+		}
+
+		if (!EnumHasAnyFlags(Range.PointData->GetAllocatedProperties(/*bWithInheritance=*/false), PointProperty))
+		{
+			UE_LOG(LogPCG, Error, TEXT("SetOutputRangeValues '%s' called on non-allocated property."), *StaticEnum<EPCGPointNativeProperties>()->GetNameStringByValue((int64)PointProperty));
+			return;
+		}
+
+		typename TPCGPointNativeProperty<PointProperty>::ValueRange ValueRange = Range.PointData->GetValueRange<PointProperty>(/*bAllocate=*/false);
+
+		if (!PCGBasePointData::IsValidRange<PointProperty>(Range, ValueRange))
+		{
+			UE_LOG(LogPCG, Error, TEXT("SetOutputRangeValues '%s' called with output range [%d, %d[ when valid range is [0, %d[."), *StaticEnum<EPCGPointNativeProperties>()->GetNameStringByValue((int64)PointProperty), Range.RangeStartIndex, Range.RangeSize, Range.PointData->GetNumPoints());
+			return;
+		}
+
+		if (Values.Num() > Range.RangeSize)
+		{
+			UE_LOG(LogPCG, Error, TEXT("SetOutputRangeValues '%s' called with invalid values array size '%d' when max array size '%d'."), *StaticEnum<EPCGPointNativeProperties>()->GetNameStringByValue((int64)PointProperty), Values.Num(), Range.RangeSize);
+			return;
+		}
+
+		for (int32 Index = 0; Index < Values.Num(); ++Index)
+		{
+			ValueRange[Range.RangeStartIndex + Index] = Values[Index];
+		}
+	}
 }
 
 bool FPCGDataTypeInfoPoint::SupportsConversionTo(const FPCGDataTypeIdentifier& ThisType, const FPCGDataTypeIdentifier& OutputType, TSubclassOf<UPCGSettings>* OptionalOutConversionSettings, FText* OptionalOutCompatibilityMessage) const
@@ -1407,4 +1482,84 @@ void UPCGBasePointData::AddToCrc(FArchiveCrc32& Ar, bool bFullDataCrc) const
 	}
 }
 
+TArray<FTransform> UPCGBasePointData::GetTransformValuesFromRange(const FPCGPointInputRange& InputRange)
+{
+	return PCGBasePointData::GetInputRangeValues<EPCGPointNativeProperties::Transform>(InputRange);
+}
+
+TArray<float> UPCGBasePointData::GetDensityValuesFromRange(const FPCGPointInputRange& InputRange)
+{
+	return PCGBasePointData::GetInputRangeValues<EPCGPointNativeProperties::Density>(InputRange);
+}
+
+TArray<FVector> UPCGBasePointData::GetBoundsMinValuesFromRange(const FPCGPointInputRange& InputRange)
+{
+	return PCGBasePointData::GetInputRangeValues<EPCGPointNativeProperties::BoundsMin>(InputRange);
+}
+
+TArray<FVector> UPCGBasePointData::GetBoundsMaxValuesFromRange(const FPCGPointInputRange& InputRange)
+{
+	return PCGBasePointData::GetInputRangeValues<EPCGPointNativeProperties::BoundsMax>(InputRange);
+}
+
+TArray<FVector4> UPCGBasePointData::GetColorValuesFromRange(const FPCGPointInputRange& InputRange)
+{
+	return PCGBasePointData::GetInputRangeValues<EPCGPointNativeProperties::Color>(InputRange);
+}
+
+TArray<float> UPCGBasePointData::GetSteepnessValuesFromRange(const FPCGPointInputRange& InputRange)
+{
+	return PCGBasePointData::GetInputRangeValues<EPCGPointNativeProperties::Steepness>(InputRange);
+}
+
+TArray<int32> UPCGBasePointData::GetSeedValuesFromRange(const FPCGPointInputRange& InputRange)
+{
+	return PCGBasePointData::GetInputRangeValues<EPCGPointNativeProperties::Seed>(InputRange);
+}
+
+TArray<int64> UPCGBasePointData::GetMetadataEntryValuesFromRange(const FPCGPointInputRange& InputRange)
+{
+	return PCGBasePointData::GetInputRangeValues<EPCGPointNativeProperties::MetadataEntry>(InputRange);
+}
+
+void UPCGBasePointData::SetTransformValuesFromRange(FPCGPointOutputRange& OutputRange, const TArray<FTransform>& Values)
+{
+	PCGBasePointData::SetOutputRangeValues<EPCGPointNativeProperties::Transform>(OutputRange, Values);
+}
+
+void UPCGBasePointData::SetDensityValuesFromRange(UPARAM(ref) FPCGPointOutputRange& OutputRange, const TArray<float>& Values)
+{
+	PCGBasePointData::SetOutputRangeValues<EPCGPointNativeProperties::Density>(OutputRange, Values);
+}
+
+void UPCGBasePointData::SetBoundsMinValuesFromRange(UPARAM(ref) FPCGPointOutputRange& OutputRange, const TArray<FVector>& Values)
+{
+	PCGBasePointData::SetOutputRangeValues<EPCGPointNativeProperties::BoundsMin>(OutputRange, Values);
+}
+
+void UPCGBasePointData::SetBoundsMaxValuesFromRange(UPARAM(ref) FPCGPointOutputRange& OutputRange, const TArray<FVector>& Values)
+{
+	PCGBasePointData::SetOutputRangeValues<EPCGPointNativeProperties::BoundsMax>(OutputRange, Values);
+}
+
+void UPCGBasePointData::SetColorValuesFromRange(UPARAM(ref) FPCGPointOutputRange& OutputRange, const TArray<FVector4>& Values)
+{
+	PCGBasePointData::SetOutputRangeValues<EPCGPointNativeProperties::Color>(OutputRange, Values);
+}
+
+void UPCGBasePointData::SetSteepnessValuesFromRange(UPARAM(ref) FPCGPointOutputRange& OutputRange, const TArray<float>& Values)
+{
+	PCGBasePointData::SetOutputRangeValues<EPCGPointNativeProperties::Steepness>(OutputRange, Values);
+}
+
+void UPCGBasePointData::SetSeedValuesFromRange(FPCGPointOutputRange& OutputRange, const TArray<int32>& Values)
+{
+	PCGBasePointData::SetOutputRangeValues<EPCGPointNativeProperties::Seed>(OutputRange, Values);
+}
+
+void UPCGBasePointData::SetMetadataEntryValuesFromRange(UPARAM(ref) FPCGPointOutputRange& OutputRange, const TArray<int64>& Values)
+{
+	PCGBasePointData::SetOutputRangeValues<EPCGPointNativeProperties::MetadataEntry>(OutputRange, Values);
+}
+
 #undef LOCTEXT_NAMESPACE
diff --git a/Engine/Plugins/PCG/Source/PCG/Private/Elements/Blueprint/PCGBlueprintBaseElement.cpp b/Engine/Plugins/PCG/Source/PCG/Private/Elements/Blueprint/PCGBlueprintBaseElement.cpp
--- a/Engine/Plugins/PCG/Source/PCG/Private/Elements/Blueprint/PCGBlueprintBaseElement.cpp
+++ b/Engine/Plugins/PCG/Source/PCG/Private/Elements/Blueprint/PCGBlueprintBaseElement.cpp
@@ -72,16 +72,41 @@ void UPCGBlueprintBaseElement::Initialize()
 #endif
 }
 
-FPCGContext& UPCGBlueprintBaseElement::GetContext() const
+FPCGBlueprintContextHandle UPCGBlueprintBaseElement::GetContextHandle() const
 {
-	checkf(CurrentContext, TEXT("Execution context is not ready - do not call the GetContext method inside of non-execution methods"));
-	return *CurrentContext;
+	FPCGBlueprintContextHandle BlueprintHandle;
+	BlueprintHandle.Handle = CurrentContext ? CurrentContext->GetOrCreateHandle() : nullptr;
+	return BlueprintHandle;
+}
+
+int32 UPCGBlueprintBaseElement::GetSeedWithContext(const FPCGBlueprintContextHandle& InContextHandle) const
+{
+	if (TSharedPtr<FPCGContextHandle> NativeContextHandle = InContextHandle.Handle.Pin())
+	{
+		if (const FPCGContext* Context = NativeContextHandle->GetContext())
+		{
+			return Context->GetSeed();
+		}
+	}
+
+	return 42;
+}
+
+FRandomStream UPCGBlueprintBaseElement::GetRandomStreamWithContext(const FPCGBlueprintContextHandle& InContextHandle) const
+{
+	return FRandomStream(GetSeedWithContext(InContextHandle));
 }
 
 void UPCGBlueprintBaseElement::SetCurrentContext(FPCGContext* InCurrentContext)
 {
 	ensure(CurrentContext == nullptr || InCurrentContext == nullptr || CurrentContext == InCurrentContext);
 	CurrentContext = InCurrentContext;
+
+	// Make sure to create handle so that multi-threaded blueprint calls don't have to worry about concurrent creation
+	if (CurrentContext)
+	{
+		CurrentContext->GetOrCreateHandle();
+	}
 }
 
 FPCGContext* UPCGBlueprintBaseElement::ResolveContext()
@@ -90,7 +115,7 @@ FPCGContext* UPCGBlueprintBaseElement::ResolveContext()
 	{
 		if (UPCGBlueprintBaseElement* Caller = Cast<UPCGBlueprintBaseElement>(FFrame::GetThreadLocalTopStackFrame()->Object))
 		{
-			return &Caller->GetContext();
+			return Caller->CurrentContext;
 		}
 	}
 
@@ -279,14 +304,4 @@ bool UPCGBlueprintBaseElement::GetOutputPinByLabel(FName InPinLabel, FPCGPinProp
 	return false;
 }
 
-int UPCGBlueprintBaseElement::GetSeed(FPCGContext& InContext) const
-{
-	return InContext.GetSeed();
-}
-
-FRandomStream UPCGBlueprintBaseElement::GetRandomStream(FPCGContext& InContext) const
-{
-	return FRandomStream(GetSeed(InContext));
-}
-
 #undef LOCTEXT_NAMESPACE
\ No newline at end of file
diff --git a/Engine/Plugins/PCG/Source/PCG/Private/Elements/Blueprint/PCGBlueprintPointProcessorSimpleElement.cpp b/Engine/Plugins/PCG/Source/PCG/Private/Elements/Blueprint/PCGBlueprintPointProcessorSimpleElement.cpp
new file mode 100644
--- /dev/null
+++ b/Engine/Plugins/PCG/Source/PCG/Private/Elements/Blueprint/PCGBlueprintPointProcessorSimpleElement.cpp
@@ -0,0 +1,122 @@
+// Copyright Epic Games, Inc. All Rights Reserved.
+
+#include "Elements/Blueprint/PCGBlueprintPointProcessorSimpleElement.h"
+#include "Helpers/PCGAsync.h"
+
+namespace PCGBlueprintPointProcessorSimpleElement
+{
+	constexpr int32 RunawayResetFrequency = 1024;
+}
+
+UPCGBlueprintPointProcessorSimpleElement::UPCGBlueprintPointProcessorSimpleElement()
+	: UPCGBlueprintBaseElement()
+{
+	bHasDefaultInPin = false;
+	bHasDefaultOutPin = false;
+	CustomInputPins.Emplace_GetRef(PCGPinConstants::DefaultInputLabel, EPCGDataType::Point).SetRequiredPin();
+	CustomOutputPins.Emplace(PCGPinConstants::DefaultOutputLabel, EPCGDataType::Point);
+}
+
+void UPCGBlueprintPointProcessorSimpleElement::Execute_Implementation(const FPCGDataCollection& Input, FPCGDataCollection& Output)
+{
+	for (const FPCGTaggedData& InputData : Input.GetInputsByPin(PCGPinConstants::DefaultInputLabel))
+	{
+		const UPCGBasePointData* InputPointData = Cast<const UPCGBasePointData>(InputData.Data);
+		if (!InputPointData)
+		{
+			continue;
+		}
+
+		UPCGBasePointData* OutputPointData = FPCGContext::NewPointData_AnyThread(CurrentContext);
+		
+		FPCGInitializeFromDataParams InitializeFromDataParams(InputPointData);
+		InitializeFromDataParams.bInheritMetadata = bInheritMetadata;
+		InitializeFromDataParams.bInheritSpatialData = bInheritSpatialData;
+		OutputPointData->InitializeFromDataWithParams(InitializeFromDataParams);
+		OutputPointData->SetNumPoints(InputPointData->GetNumPoints());
+		
+		// If we are not inheriting the metadata but we are inheriting the spatial data, make sure to reset the entry keys
+		if (!bInheritMetadata && bInheritSpatialData)
+		{
+			OutputPointData->SetMetadataEntry(PCGInvalidEntryKey);
+		}
+		
+		OutputPointData->AllocateProperties((EPCGPointNativeProperties)PropertiesToAllocate);
+
+		Initialize(InputPointData, OutputPointData);
+
+		TFunction<int32(int32, int32, int32)> RangeProcessFunc = [this, InputPointData, OutputPointData](int32 StartReadIndex, int32 StartWriteIndex, int32 Count)
+		{
+			const FConstPCGPointValueRanges InputValueRanges(InputPointData);
+			FPCGPointValueRanges OutputValueRanges(OutputPointData, /*bAllocate=*/false);
+
+			FPCGPointInputRange InputRange{ InputPointData, StartReadIndex, Count };
+			FPCGPointOutputRange OutputRange{ OutputPointData, StartWriteIndex, Count };
+
+			return RangePointLoopBody(InputRange, OutputRange);
+		};
+
+		TFunction<int32(int32, int32, int32)> PerPointProcessFunc = [this, InputPointData, OutputPointData](int32 StartReadIndex, int32 StartWriteIndex, int32 Count)
+		{
+			const FConstPCGPointValueRanges InputValueRanges(InputPointData);
+			FPCGPointValueRanges OutputValueRanges(OutputPointData, /*bAllocate=*/false);
+
+			int32 NumWritten = 0;
+
+			for (int32 ReadIndex = StartReadIndex; ReadIndex < StartReadIndex + Count; ++ReadIndex)
+			{
+				if (ReadIndex % PCGBlueprintPointProcessorSimpleElement::RunawayResetFrequency == 0)
+				{
+					GInitRunaway(); // Reset periodically the iteration count, because we know we're in a fixed size loop.
+				}
+
+				FPCGPoint InPoint = InputValueRanges.GetPoint(ReadIndex);
+				
+				FPCGPoint OutPoint;
+
+				if (bInheritSpatialData)
+				{
+					OutPoint = InPoint;
+				}
+
+				if (PointLoopBody(InputPointData, InPoint, OutputPointData, OutPoint, ReadIndex))
+				{
+					OutputValueRanges.SetFromPoint(NumWritten, OutPoint);
+					++NumWritten;
+				}
+			}
+
+			return NumWritten;
+		};
+
+		TFunction<int32(int32, int32, int32)> ProcessFunc = GetClass()->IsFunctionImplementedInScript(GET_FUNCTION_NAME_CHECKED(UPCGBlueprintPointProcessorSimpleElement, RangePointLoopBody)) ? RangeProcessFunc : PerPointProcessFunc;
+
+		auto MoveDataRangeFunc = [OutputPointData](int32 RangeStartIndex, int32 MoveToIndex, int32 NumElements)
+		{
+			OutputPointData->MoveRange(RangeStartIndex, MoveToIndex, NumElements);
+		};
+
+		auto FinishedFunc = [OutputPointData](int32 NumWritten)
+		{
+			OutputPointData->SetNumPoints(NumWritten);
+		};
+
+		{
+			TGuardValue<bool> IsCallingBlueprint(CurrentContext->AsyncState.bIsCallingBlueprint, true);
+
+			FPCGAsync::AsyncProcessingRangeEx(
+				&CurrentContext->AsyncState,
+				InputPointData->GetNumPoints(),
+				/*InitializeFunc=*/[]() {},
+				ProcessFunc,
+				MoveDataRangeFunc,
+				FinishedFunc,
+				/*bEnableTimeSlicing=*/false
+			);
+		}
+		
+		FPCGTaggedData& OutputData = Output.TaggedData.Emplace_GetRef(InputData);
+		OutputData.Data = OutputPointData;
+		OutputData.Pin = PCGPinConstants::DefaultOutputLabel;
+	}
+}
\ No newline at end of file
diff --git a/Engine/Plugins/PCG/Source/PCG/Private/Helpers/PCGAsync.cpp b/Engine/Plugins/PCG/Source/PCG/Private/Helpers/PCGAsync.cpp
--- a/Engine/Plugins/PCG/Source/PCG/Private/Helpers/PCGAsync.cpp
+++ b/Engine/Plugins/PCG/Source/PCG/Private/Helpers/PCGAsync.cpp
@@ -561,9 +561,10 @@ namespace FPCGAsync
 			// Futures are not returning anything.
 			TArray<UE::Tasks::TTask<void>> AsyncTasks;
 			AsyncTasks.Reserve(NumFutures);
+			const bool bInitBPContext = AsyncState.bIsCallingBlueprint;
 
 			// A job task will pick up the chunk to process, do the work on it, and push in the queue info that it processed the chunk and how many points it has written.
-			auto JobTask = [&IterationInnerLoop, &SynchroStruct, ChunkSize, NumIterations](const int32 CurrentChunkToProcess)
+			auto JobTask = [&IterationInnerLoop, &SynchroStruct, ChunkSize, NumIterations, bInitBPContext](const int32 CurrentChunkToProcess)
 			{
 				const int32 StartReadIndex = CurrentChunkToProcess * ChunkSize;
 				// We write in the same "memory space" with the async task (in preallocated output array) and will be collapsed by the main thread.
@@ -573,6 +574,11 @@ namespace FPCGAsync
 				check(Count > 0);
 				int32 NumElementsWritten;
 
+				if (bInitBPContext)
+				{
+					GInitRunaway(); // Reset counter as threads can run multiple workloads and might not be properly reset.
+				}
+
 				{
 					TRACE_CPUPROFILER_EVENT_SCOPE(FPCGAsync::Private::AsyncProcessing::InnerLoop);
 
diff --git a/Engine/Plugins/PCG/Source/PCG/Public/Data/PCGBasePointData.h b/Engine/Plugins/PCG/Source/PCG/Public/Data/PCGBasePointData.h
--- a/Engine/Plugins/PCG/Source/PCG/Public/Data/PCGBasePointData.h
+++ b/Engine/Plugins/PCG/Source/PCG/Public/Data/PCGBasePointData.h
@@ -18,6 +18,36 @@ namespace PCGPointDataConstants
 	const FName ElementsDomainName = "Points";
 }
 
+USTRUCT(BlueprintType)
+struct FPCGPointInputRange
+{
+	GENERATED_BODY()
+
+	UPROPERTY(BlueprintReadWrite, EditAnywhere, Category = Default)
+	TObjectPtr<const UPCGBasePointData> PointData = nullptr;
+
+	UPROPERTY(BlueprintReadWrite, EditAnywhere, Category = Default)
+	int32 RangeStartIndex = INDEX_NONE;
+
+	UPROPERTY(BlueprintReadWrite, EditAnywhere, Category = Default)
+	int32 RangeSize = 0;
+};
+
+USTRUCT(BlueprintType)
+struct FPCGPointOutputRange
+{
+	GENERATED_BODY()
+
+	UPROPERTY(BlueprintReadWrite, EditAnywhere, Category = Default)
+	TObjectPtr<UPCGBasePointData> PointData = nullptr;
+
+	UPROPERTY(BlueprintReadWrite, EditAnywhere, Category = Default)
+	int32 RangeStartIndex = INDEX_NONE;
+
+	UPROPERTY(BlueprintReadWrite, EditAnywhere, Category = Default)
+	int32 RangeSize = 0;
+};
+
 USTRUCT()
 struct FPCGDataTypeInfoPoint : public FPCGDataTypeInfoConcrete
 {
@@ -74,10 +104,18 @@ public:
 
 	virtual bool IsValidRef(const PCGPointOctree::FPointRef& InPointRef) const PURE_VIRTUAL(UPCGBasePointData::IsValidRef, return false;);
 
+	UFUNCTION(BlueprintCallable, Category = SpatialData)
 	virtual int32 GetNumPoints() const PURE_VIRTUAL(UPCGBasePointData::GetNumPoints, return 0;);
+
+	UFUNCTION(BlueprintCallable, Category = SpatialData)
 	bool IsEmpty() const { return GetNumPoints() == 0; }
+
+	UFUNCTION(BlueprintCallable, Category = SpatialData)
 	virtual void SetNumPoints(int32 InNumPoints, bool bInitializeValues = true) PURE_VIRTUAL(UPCGBasePointData::SetNumPoints, return;);
 
+	UFUNCTION(BlueprintCallable, Category = SpatialData, meta=(DisplayName="Allocate Properties"))
+	void BP_AllocateProperties(UPARAM(meta = (BitMask, BitmaskEnum = "/Script/PCG.EPCGPointNativeProperties")) int32 Properties) { AllocateProperties((EPCGPointNativeProperties)Properties); }
+
 	virtual void AllocateProperties(EPCGPointNativeProperties Properties) {}
 	virtual void FreeProperties(EPCGPointNativeProperties Properties) {}
 	virtual void MoveRange(int32 RangeStartIndex, int32 MoveToIndex, int32 NumElements) PURE_VIRTUAL(UPCGBasePointData::MoveRange, return;);
@@ -139,37 +177,94 @@ public:
 	const typename TPCGPointNativeProperty<Property>::Type& GetPointPropertyValue(int32 InPointIndex) const
 	{
 		const typename TPCGPointNativeProperty<Property>::ConstValueRange ConstPropertyValueRange = GetConstValueRange<TPCGPointNativeProperty<Property>::EnumValue>();
-		return ConstPropertyValueRange[InPointIndex];
+		if (ConstPropertyValueRange.IsValidIndex(InPointIndex))
+		{
+			return ConstPropertyValueRange[InPointIndex];
+		}
+		else
+		{
+			UE_LOG(LogPCG, Error, TEXT("Invalid point index : %d"), InPointIndex);
+			static typename TPCGPointNativeProperty<Property>::Type Invalid;
+			return Invalid;
+		}
 	}
 
-	UE_API virtual void SetTransform(const FPCGPointTransform::Type& InTransform);
-	UE_API virtual void SetDensity(FPCGPointDensity::Type InDensity);
-	UE_API virtual void SetBoundsMin(const FPCGPointBoundsMin::Type& InBoundsMin);
-	UE_API virtual void SetBoundsMax(const FPCGPointBoundsMax::Type& InBoundsMax);
-	UE_API virtual void SetColor(const FPCGPointColor::Type& InColor);
-	UE_API virtual void SetSteepness(FPCGPointSteepness::Type InSteepness);
-	UE_API virtual void SetSeed(FPCGPointSeed::Type InSeed);
-	UE_API virtual void SetMetadataEntry(FPCGPointMetadataEntry::Type InMetadataEntry);
+	UFUNCTION(BlueprintCallable, Category = SpatialData, meta = (DisplayName = "Set All Transforms", ToolTip = "Set all transforms to same value"))
+	UE_API virtual void SetTransform(const FTransform& InTransform);
+	
+	UFUNCTION(BlueprintCallable, Category = SpatialData, meta = (DisplayName = "Set All Densities", ToolTip = "Set all densities to same value"))
+	UE_API virtual void SetDensity(float InDensity);
+	
+	UFUNCTION(BlueprintCallable, Category = SpatialData, meta = (DisplayName = "Set All Bounds Min", ToolTip = "Set all bounds min to same value"))
+	UE_API virtual void SetBoundsMin(const FVector& InBoundsMin);
+	
+	UFUNCTION(BlueprintCallable, Category = SpatialData, meta = (DisplayName = "Set All Bounds Max", ToolTip = "Set all bounds max to same value"))
+	UE_API virtual void SetBoundsMax(const FVector& InBoundsMax);
+	
+	UFUNCTION(BlueprintCallable, Category = SpatialData, meta = (DisplayName = "Set All Colors", ToolTip = "Set all colors to same value"))
+	UE_API virtual void SetColor(const FVector4& InColor);
+	
+	UFUNCTION(BlueprintCallable, Category = SpatialData, meta = (DisplayName = "Set All Steepness", ToolTip = "Set all steepness to same value"))
+	UE_API virtual void SetSteepness(float InSteepness);
+	
+	UFUNCTION(BlueprintCallable, Category = SpatialData, meta = (DisplayName = "Set All Seeds", ToolTip = "Set all seeds to same value"))
+	UE_API virtual void SetSeed(int32 InSeed);
+	
+	UFUNCTION(BlueprintCallable, Category = SpatialData, meta = (DisplayName = "Set All Metadata Entries", ToolTip = "Set all metadata entries to same value"))
+	UE_API virtual void SetMetadataEntry(int64 InMetadataEntry);
 
+	UFUNCTION(BlueprintCallable, Category = SpatialData, meta = (DisplayName = "Set All Extents", ToolTip = "Set all extents to same value"))
 	UE_API virtual void SetExtents(const FVector& InExtents);
+	
+	UFUNCTION(BlueprintCallable, Category = SpatialData, meta = (DisplayName = "Set All Local Centers", ToolTip = "Set all local centers to same value"))
 	UE_API virtual void SetLocalCenter(const FVector& InLocalCenter);
 
-	UE_API virtual const FPCGPointTransform::Type& GetTransform(int32 InPointIndex) const;
-	UE_API virtual const FPCGPointDensity::Type& GetDensity(int32 InPointIndex) const;
-	UE_API virtual const FPCGPointBoundsMin::Type& GetBoundsMin(int32 InPointIndex) const;
-	UE_API virtual const FPCGPointBoundsMax::Type& GetBoundsMax(int32 InPointIndex) const;
-	UE_API virtual const FPCGPointColor::Type& GetColor(int32 InPointIndex) const;
-	UE_API virtual const FPCGPointSteepness::Type& GetSteepness(int32 InPointIndex) const;
-	UE_API virtual const FPCGPointSeed::Type& GetSeed(int32 InPointIndex) const;
-	UE_API virtual const FPCGPointMetadataEntry::Type& GetMetadataEntry(int32 InPointIndex) const;
+	UFUNCTION(BlueprintCallable, Category = SpatialData)
+	UE_API virtual const FTransform& GetTransform(int32 InPointIndex) const;
+	
+	UFUNCTION(BlueprintCallable, Category = SpatialData)
+	UE_API virtual const float& GetDensity(int32 InPointIndex) const;
+	
+	UFUNCTION(BlueprintCallable, Category = SpatialData)
+	UE_API virtual const FVector& GetBoundsMin(int32 InPointIndex) const;
+	
+	UFUNCTION(BlueprintCallable, Category = SpatialData)
+	UE_API virtual const FVector& GetBoundsMax(int32 InPointIndex) const;
+	
+	UFUNCTION(BlueprintCallable, Category = SpatialData)
+	UE_API virtual const FVector4& GetColor(int32 InPointIndex) const;
+	
+	UFUNCTION(BlueprintCallable, Category = SpatialData)
+	UE_API virtual const float& GetSteepness(int32 InPointIndex) const;
 
+	UFUNCTION(BlueprintCallable, Category = SpatialData)
+	UE_API virtual const int32& GetSeed(int32 InPointIndex) const;
+
+	UFUNCTION(BlueprintCallable, Category = SpatialData)
+	UE_API virtual const int64& GetMetadataEntry(int32 InPointIndex) const;
+
+	UFUNCTION(BlueprintCallable, Category = SpatialData)
 	UE_API virtual FBoxSphereBounds GetDensityBounds(int32 InPointIndex) const;
+
+	UFUNCTION(BlueprintCallable, Category = SpatialData)
 	UE_API virtual FBox GetLocalDensityBounds(int32 InPointIndex) const;
+
+	UFUNCTION(BlueprintCallable, Category = SpatialData)
 	UE_API virtual FBox GetLocalBounds(int32 InPointIndex) const;
+
+	UFUNCTION(BlueprintCallable, Category = SpatialData)
 	UE_API virtual FVector GetLocalCenter(int32 InPointIndex) const;
+
+	UFUNCTION(BlueprintCallable, Category = SpatialData)
 	UE_API virtual FVector GetExtents(int32 InPointIndex) const;
+
+	UFUNCTION(BlueprintCallable, Category = SpatialData)
 	UE_API virtual FVector GetScaledExtents(int32 InPointIndex) const;
+
+	UFUNCTION(BlueprintCallable, Category = SpatialData)
 	UE_API virtual FVector GetLocalSize(int32 InPointIndex) const;
+
+	UFUNCTION(BlueprintCallable, Category = SpatialData)
 	UE_API virtual FVector GetScaledLocalSize(int32 InPointIndex) const;
 		
 	/** Get the dirty status of the Octree. Note that the Point Octree can be rebuilt from another thread, so this info can be invalidated at anytime. */
@@ -242,6 +337,54 @@ public:
 		return PropertiesToAllocate;
 	}
 
+	UFUNCTION(BlueprintCallable, Category = SpatialData)
+	static TArray<FTransform> GetTransformValuesFromRange(const FPCGPointInputRange& InputRange);
+
+	UFUNCTION(BlueprintCallable, Category = SpatialData)
+	static TArray<float> GetDensityValuesFromRange(const FPCGPointInputRange& InputRange);
+
+	UFUNCTION(BlueprintCallable, Category = SpatialData)
+	static TArray<FVector> GetBoundsMinValuesFromRange(const FPCGPointInputRange& InputRange);
+
+	UFUNCTION(BlueprintCallable, Category = SpatialData)
+	static TArray<FVector> GetBoundsMaxValuesFromRange(const FPCGPointInputRange& InputRange);
+
+	UFUNCTION(BlueprintCallable, Category = SpatialData)
+	static TArray<FVector4> GetColorValuesFromRange(const FPCGPointInputRange& InputRange);
+
+	UFUNCTION(BlueprintCallable, Category = SpatialData)
+	static TArray<float> GetSteepnessValuesFromRange(const FPCGPointInputRange& InputRange);
+
+	UFUNCTION(BlueprintCallable, Category = SpatialData)
+	static TArray<int32> GetSeedValuesFromRange(const FPCGPointInputRange& InputRange);
+
+	UFUNCTION(BlueprintCallable, Category = SpatialData)
+	static TArray<int64> GetMetadataEntryValuesFromRange(const FPCGPointInputRange& InputRange);
+
+	UFUNCTION(BlueprintCallable, Category = SpatialData)
+	static void SetTransformValuesFromRange(UPARAM(ref) FPCGPointOutputRange& OutputRange, const TArray<FTransform>& Values);
+		
+	UFUNCTION(BlueprintCallable, Category = SpatialData)
+	static void SetDensityValuesFromRange(UPARAM(ref) FPCGPointOutputRange& OutputRange, const TArray<float>& Values);
+
+	UFUNCTION(BlueprintCallable, Category = SpatialData)
+	static void SetBoundsMinValuesFromRange(UPARAM(ref) FPCGPointOutputRange& OutputRange, const TArray<FVector>& Values);
+
+	UFUNCTION(BlueprintCallable, Category = SpatialData)
+	static void SetBoundsMaxValuesFromRange(UPARAM(ref) FPCGPointOutputRange& OutputRange, const TArray<FVector>& Values);
+
+	UFUNCTION(BlueprintCallable, Category = SpatialData)
+	static void SetColorValuesFromRange(UPARAM(ref) FPCGPointOutputRange& OutputRange, const TArray<FVector4>& Values);
+
+	UFUNCTION(BlueprintCallable, Category = SpatialData)
+	static void SetSteepnessValuesFromRange(UPARAM(ref) FPCGPointOutputRange& OutputRange, const TArray<float>& Values);
+
+	UFUNCTION(BlueprintCallable, Category = SpatialData)
+	static void SetSeedValuesFromRange(UPARAM(ref) FPCGPointOutputRange& OutputRange, const TArray<int32>& Values);
+
+	UFUNCTION(BlueprintCallable, Category = SpatialData)
+	static void SetMetadataEntryValuesFromRange(UPARAM(ref) FPCGPointOutputRange& OutputRange, const TArray<int64>& Values);
+
 protected:
 
 	void RebuildOctreeIfNeeded() const
diff --git a/Engine/Plugins/PCG/Source/PCG/Public/Data/PCGPointData.h b/Engine/Plugins/PCG/Source/PCG/Public/Data/PCGPointData.h
--- a/Engine/Plugins/PCG/Source/PCG/Public/Data/PCGPointData.h
+++ b/Engine/Plugins/PCG/Source/PCG/Public/Data/PCGPointData.h
@@ -129,14 +129,10 @@ public:
 	UFUNCTION(BlueprintCallable, BlueprintPure = false, Category = SpatialData)
 	TArray<FPCGPoint> GetPointsCopy() const { return Points; }
 
-	UFUNCTION(BlueprintCallable, Category = SpatialData)
 	virtual int32 GetNumPoints() const override { return Points.Num(); }
 
 	UE_API virtual void SetNumPoints(int32 InNumPoints, bool bInitializeValues = true) override;
 
-	UFUNCTION(BlueprintCallable, Category = SpatialData)
-	bool IsEmpty() const { return Points.IsEmpty(); }
-
 	UFUNCTION(BlueprintCallable, Category = SpatialData)
 	UE_API FPCGPoint GetPoint(int32 Index) const;
 
diff --git a/Engine/Plugins/PCG/Source/PCG/Public/Elements/Blueprint/PCGBlueprintBaseElement.h b/Engine/Plugins/PCG/Source/PCG/Public/Elements/Blueprint/PCGBlueprintBaseElement.h
--- a/Engine/Plugins/PCG/Source/PCG/Public/Elements/Blueprint/PCGBlueprintBaseElement.h
+++ b/Engine/Plugins/PCG/Source/PCG/Public/Elements/Blueprint/PCGBlueprintBaseElement.h
@@ -87,18 +87,17 @@ public:
 
 	/** Gets the seed from the associated settings & source component */
 	UFUNCTION(BlueprintCallable, Category = "PCG|Random")
-	PCG_API int GetSeed(UPARAM(ref) FPCGContext& InContext) const;
+	PCG_API int32 GetSeedWithContext(const FPCGBlueprintContextHandle& InContextHandle) const;
 
 	/** Creates a random stream from the settings & source component */
-	UFUNCTION(BlueprintCallable, BlueprintPure = false, Category = "PCG|Random")
-	PCG_API FRandomStream GetRandomStream(UPARAM(ref) FPCGContext& InContext) const;
+	UFUNCTION(BlueprintCallable, Category = "PCG|Random")
+	PCG_API FRandomStream GetRandomStreamWithContext(const FPCGBlueprintContextHandle& InContextHandle) const;
 
 	/** Called after object creation to setup the object callbacks */
 	PCG_API void Initialize();
 
-	/** Retrieves the execution context - note that this will not be valid outside of the Execute functions */
 	UFUNCTION(BlueprintCallable, Category = "PCG|Advanced", meta = (HideSelfPin = "true"))
-	PCG_API FPCGContext& GetContext() const;
+	PCG_API FPCGBlueprintContextHandle GetContextHandle() const;
 
 	/** Called after the element duplication during execution to be able to get the context easily - internal call only */
 	PCG_API void SetCurrentContext(FPCGContext* InCurrentContext);
diff --git a/Engine/Plugins/PCG/Source/PCG/Public/Elements/Blueprint/PCGBlueprintDeprecatedElement.h b/Engine/Plugins/PCG/Source/PCG/Public/Elements/Blueprint/PCGBlueprintDeprecatedElement.h
--- a/Engine/Plugins/PCG/Source/PCG/Public/Elements/Blueprint/PCGBlueprintDeprecatedElement.h
+++ b/Engine/Plugins/PCG/Source/PCG/Public/Elements/Blueprint/PCGBlueprintDeprecatedElement.h
@@ -117,4 +117,16 @@ public:
 	/** Calls the IterationLoopBody a fixed number of times, optional parameters are used to potentially initialized the Out Data, but otherwise are used to remove the need to have variables */
 	UFUNCTION(BlueprintCallable, BlueprintPure = false, Category = "PCG|Flow Control", meta = (HideSelfPin = "true"))
 	PCG_API void IterationLoop(UPARAM(ref) FPCGContext& InContext, int64 NumIterations, UPCGPointData*& OutData, const UPCGSpatialData* OptionalA = nullptr, const UPCGSpatialData* OptionalB = nullptr, UPCGPointData* OptionalOutData = nullptr) const;
+
+	/** Creates a random stream from the settings & source component */
+	UFUNCTION(BlueprintCallable, BlueprintPure = false, Category = "PCG|Random")
+	PCG_API FRandomStream GetRandomStream(UPARAM(ref) FPCGContext& InContext) const;
+
+	/** Gets the seed from the associated settings & source component */
+	UFUNCTION(BlueprintCallable, Category = "PCG|Random")
+	PCG_API int GetSeed(UPARAM(ref) FPCGContext& InContext) const;
+
+	/** Retrieves the execution context - note that this will not be valid outside of the Execute functions */
+	UFUNCTION(BlueprintCallable, Category = "PCG|Advanced", meta = (HideSelfPin = "true"))
+	PCG_API FPCGContext& GetContext() const;
 };
diff --git a/Engine/Plugins/PCG/Source/PCG/Public/Elements/Blueprint/PCGBlueprintPointProcessorSimpleElement.h b/Engine/Plugins/PCG/Source/PCG/Public/Elements/Blueprint/PCGBlueprintPointProcessorSimpleElement.h
new file mode 100644
--- /dev/null
+++ b/Engine/Plugins/PCG/Source/PCG/Public/Elements/Blueprint/PCGBlueprintPointProcessorSimpleElement.h
@@ -0,0 +1,63 @@
+// Copyright Epic Games, Inc. All Rights Reserved.
+
+#pragma once
+
+#include "PCGBlueprintBaseElement.h"
+#include "Data/PCGBasePointData.h"
+
+#include "PCGBlueprintPointProcessorSimpleElement.generated.h"
+
+/**
+ * Simple point processor element that supports the following two overrides:
+ * - RangePointLoopBody is the more efficient but more complex loop operating on a point range per call.
+ * - PointLoopBody is the less efficient but simplest loop operating on a single point per call.
+ */
+UCLASS(MinimalAPI, Abstract, BlueprintType, Blueprintable, hidecategories = (Object))
+class UPCGBlueprintPointProcessorSimpleElement : public UPCGBlueprintBaseElement
+{
+	GENERATED_BODY()
+
+public:
+	UPCGBlueprintPointProcessorSimpleElement();
+
+	/** Execute implementation will call the proper loop body based on which is implemented */
+	virtual void Execute_Implementation(const FPCGDataCollection & Input, FPCGDataCollection & Output) override;
+
+	/** Allows initialization of the output point data before the loop body calls. Called once per pin input - output pair. */
+	UFUNCTION(BlueprintImplementableEvent, Category = "PCG|Flow Control", meta = (ForceAsFunction))
+	void Initialize(const UPCGBasePointData* InputPointData, UPCGBasePointData* OutputPointData);
+
+	/**
+	 * Multi-threaded loop that will iterate on all points in InPointData. All points will be added in the same order than in input.
+	 * @param InPointData  - Input point data. Constant, must not be modified.
+	 * @param InPoint      - Point for the current iteration. Constant, must not be modified.
+	 * @param OutPointData - Output point data. Can be modified.
+	 * @param OutPoint     - Point that will be added to the output data. Can be modified.
+	 * @param PointIndex   - Index of the current point. Must only be used to access input data, as this call is multi-threaded. It is not safe to access output data.
+	 * @returns True if the point should be kept, False if not.
+	 */
+	UFUNCTION(BlueprintImplementableEvent, Category = "PCG|Flow Control")
+	bool PointLoopBody(const UPCGBasePointData* InPointData, const FPCGPoint& InPoint, UPCGBasePointData* OutPointData, FPCGPoint& OutPoint, int64 PointIndex);
+
+	/**
+	 * Multi-threaded loop that will iterate on all points. All points will be added in the same order than in input.
+	 * Each call will work on a subset of the data (ranges).
+	 * @param InputRange  - Input Range. Constant, must not be modified. Can be used to get property value ranges on the Input point data.
+	 * @param OutputRange - Output Range. Can be modified. Can be used to set property value ranges on the Output point data.
+	 * @returns Number of points that were written.
+	 */
+	UFUNCTION(BlueprintImplementableEvent, Category = "PCG|Flow Control")
+	int32 RangePointLoopBody(const FPCGPointInputRange& InputRange, const FPCGPointOutputRange& OutputRange);
+	
+	/** If true, output will inherit metadata from inputs */
+	UPROPERTY(BlueprintReadWrite, EditDefaultsOnly, Category = Settings)
+	bool bInheritMetadata = true;
+
+	/** If true, output will inherit spatial data from inputs */
+	UPROPERTY(BlueprintReadWrite, EditDefaultsOnly, Category = Settings)
+	bool bInheritSpatialData = true;
+
+	/** Specify which properties to allocate on the outputs */
+	UPROPERTY(BlueprintReadWrite, EditDefaultsOnly, Category = Settings, meta = (Bitmask, BitmaskEnum = "/Script/PCG.EPCGPointNativeProperties"))
+	int32 PropertiesToAllocate = 0;
+};
\ No newline at end of file
diff --git a/Engine/Plugins/PCG/Source/PCG/Public/PCGContext.h b/Engine/Plugins/PCG/Source/PCG/Public/PCGContext.h
--- a/Engine/Plugins/PCG/Source/PCG/Public/PCGContext.h
+++ b/Engine/Plugins/PCG/Source/PCG/Public/PCGContext.h
@@ -74,6 +74,20 @@ private:
 	FPCGContext* Context = nullptr;
 };
 
+/**
+ * Blueprint specific FPCGContext handle.
+ *
+ * This handle is safe to copy and pass around in BP.
+ * Prior API was using FPCGContext struct references. BP struct by reference can do actual copies causing issues specifically with multi-threaded access of FPCGContext.
+ */
+USTRUCT(BlueprintType)
+struct FPCGBlueprintContextHandle
+{
+	GENERATED_BODY()
+
+	TWeakPtr<FPCGContextHandle> Handle;
+};
+
 USTRUCT(BlueprintType)
 struct FPCGContext
 {
diff --git a/Engine/Plugins/PCGInterops/PCGGeometryScriptInterop/Source/PCGGeometryScriptInterop/Private/Elements/PCGGeometryBlueprintElement.cpp b/Engine/Plugins/PCGInterops/PCGGeometryScriptInterop/Source/PCGGeometryScriptInterop/Private/Elements/PCGGeometryBlueprintElement.cpp
--- a/Engine/Plugins/PCGInterops/PCGGeometryScriptInterop/Source/PCGGeometryScriptInterop/Private/Elements/PCGGeometryBlueprintElement.cpp
+++ b/Engine/Plugins/PCGInterops/PCGGeometryScriptInterop/Source/PCGGeometryScriptInterop/Private/Elements/PCGGeometryBlueprintElement.cpp
@@ -39,7 +39,7 @@ void UPCGGeometryBlueprintElement::Execute_Implementation(const FPCGDataCollecti
 	if (!bHasASingleInputPin || !bHasASingleOutputPin)
 	{
 		// Make sure to throw a warning if we are in this case and Process Dynamic Mesh is overridden to warn the user their function won't be called
-		PCGLog::LogWarningOnGraph(LOCTEXT("DynMeshOverridenButNotCalled", "Process Dynamic Mesh was overridden, but we don't have the expected setup (single input and output pin). Process Dynamic Mesh won't be called."), &GetContext());
+		PCGLog::LogWarningOnGraph(LOCTEXT("DynMeshOverridenButNotCalled", "Process Dynamic Mesh was overridden, but we don't have the expected setup (single input and output pin). Process Dynamic Mesh won't be called."), CurrentContext);
 		return;
 	}
 
@@ -69,12 +69,12 @@ UPCGDynamicMeshData* UPCGGeometryBlueprintElement::CopyOrStealInputData(const FP
 	{
 		// Verification that the user didn't change default settings
 		PCGLog::LogWarningOnGraph(LOCTEXT("SettingsDifferent", "In PCG Geometry Blueprint Element, the default settings were changed (not cacheable and not verifying outputs used multiple times)."
-			"Use the normal BP element if you want this behavior. Will always copy and never steal."), &GetContext());
+			"Use the normal BP element if you want this behavior. Will always copy and never steal."), CurrentContext);
 		const UPCGDynamicMeshData* InData = Cast<const UPCGDynamicMeshData>(InTaggedData.Data);
-		return InData ? CastChecked<UPCGDynamicMeshData>(InData->DuplicateData(&GetContext())) : nullptr;
+		return InData ? CastChecked<UPCGDynamicMeshData>(InData->DuplicateData(CurrentContext)) : nullptr;
 	}
 	
-	return IPCGDynamicMeshBaseElement::CopyOrSteal(InTaggedData, &GetContext());
+	return IPCGDynamicMeshBaseElement::CopyOrSteal(InTaggedData, CurrentContext);
 }
 
 #undef LOCTEXT_NAMESPACE
