[PCG] [Bughawk] Fixes to 45482868

#rnx

diff --git a/Engine/Plugins/PCG/Source/PCG/Private/Elements/Blueprint/PCGBlueprintPointProcessorSimpleElement.cpp b/Engine/Plugins/PCG/Source/PCG/Private/Elements/Blueprint/PCGBlueprintPointProcessorSimpleElement.cpp
--- a/Engine/Plugins/PCG/Source/PCG/Private/Elements/Blueprint/PCGBlueprintPointProcessorSimpleElement.cpp
+++ b/Engine/Plugins/PCG/Source/PCG/Private/Elements/Blueprint/PCGBlueprintPointProcessorSimpleElement.cpp
@@ -81,7 +81,7 @@ void UPCGBlueprintPointProcessorSimpleElement::Execute_Implementation(const FPCG
 
 				if (PointLoopBody(InputPointData, InPoint, OutputPointData, OutPoint, ReadIndex))
 				{
-					OutputValueRanges.SetFromPoint(NumWritten, OutPoint);
+					OutputValueRanges.SetFromPoint(StartWriteIndex + NumWritten, OutPoint);
 					++NumWritten;
 				}
 			}
diff --git a/Engine/Plugins/PCG/Source/PCG/Public/Elements/Blueprint/PCGBlueprintBaseElement.h b/Engine/Plugins/PCG/Source/PCG/Public/Elements/Blueprint/PCGBlueprintBaseElement.h
--- a/Engine/Plugins/PCG/Source/PCG/Public/Elements/Blueprint/PCGBlueprintBaseElement.h
+++ b/Engine/Plugins/PCG/Source/PCG/Public/Elements/Blueprint/PCGBlueprintBaseElement.h
@@ -32,7 +32,7 @@ public:
 	// ~End UObject interface
 
 	/**
-	 * Main execution function that will contain the logic for this PCG Element. Use GetContext to have access to the context.
+	 * Main execution function that will contain the logic for this PCG Element. Use GetContextHandle to have access to the context.
 	 * @param Input  - Input collection containing all the data passed as input to the node.
 	 * @param Output - Data collection that will be passed as the output of the node, with pins matching the ones provided during the execution.
 	 */
diff --git a/Engine/Plugins/PCG/Source/PCG/Public/Elements/Blueprint/PCGBlueprintPointProcessorSimpleElement.h b/Engine/Plugins/PCG/Source/PCG/Public/Elements/Blueprint/PCGBlueprintPointProcessorSimpleElement.h
--- a/Engine/Plugins/PCG/Source/PCG/Public/Elements/Blueprint/PCGBlueprintPointProcessorSimpleElement.h
+++ b/Engine/Plugins/PCG/Source/PCG/Public/Elements/Blueprint/PCGBlueprintPointProcessorSimpleElement.h
@@ -47,7 +47,7 @@ public:
 	 * @returns Number of points that were written.
 	 */
 	UFUNCTION(BlueprintImplementableEvent, Category = "PCG|Flow Control")
-	int32 RangePointLoopBody(const FPCGPointInputRange& InputRange, const FPCGPointOutputRange& OutputRange);
+	int32 RangePointLoopBody(const FPCGPointInputRange& InputRange, UPARAM(Ref) FPCGPointOutputRange& OutputRange);
 	
 	/** If true, output will inherit metadata from inputs */
 	UPROPERTY(BlueprintReadWrite, EditDefaultsOnly, Category = Settings)
