[PCG] Add overrides to CustomHLSL node

diff --git a/Engine/Plugins/PCG/Source/PCG/Private/Compute/DataInterfaces/Elements/PCGCustomHLSLDataInterface.cpp b/Engine/Plugins/PCG/Source/PCG/Private/Compute/DataInterfaces/Elements/PCGCustomHLSLDataInterface.cpp
new file mode 100644
--- /dev/null
+++ b/Engine/Plugins/PCG/Source/PCG/Private/Compute/DataInterfaces/Elements/PCGCustomHLSLDataInterface.cpp
@@ -0,0 +1,89 @@
+// Copyright Epic Games, Inc. All Rights Reserved.
+
+#include "PCGCustomHLSLDataInterface.h"
+
+#include "PCGContext.h"
+#include "Compute/PCGComputeGraph.h"
+#include "Compute/PCGDataBinding.h"
+#include "Compute/Elements/PCGCustomHLSL.h"
+#include "Compute/Elements/PCGCustomHLSLKernel.h"
+
+#include "ShaderCompilerCore.h"
+#include "ShaderCore.h"
+#include "ShaderParameterMetadataBuilder.h"
+#include "ComputeFramework/ShaderParamTypeDefinition.h"
+
+#include UE_INLINE_GENERATED_CPP_BY_NAME(PCGCustomHLSLDataInterface)
+
+BEGIN_SHADER_PARAMETER_STRUCT(FPCGCustomHLSLDataInterfaceParameters,)
+	SHADER_PARAMETER(int32, NumElements)
+	//SHADER_PARAMETER(FIntVector2, NumElements2D)
+	SHADER_PARAMETER(int32, ThreadCountMultiplier)
+	SHADER_PARAMETER(int32, FixedThreadCount)
+END_SHADER_PARAMETER_STRUCT()
+
+void UPCGCustomHLSLDataInterface::GetShaderParameters(TCHAR const* UID, FShaderParametersMetadataBuilder& InOutBuilder, FShaderParametersMetadataAllocations& InOutAllocations) const
+{
+	InOutBuilder.AddNestedStruct<FPCGCustomHLSLDataInterfaceParameters>(UID);
+}
+
+UComputeDataProvider* UPCGCustomHLSLDataInterface::CreateDataProvider() const
+{
+	return NewObject<UPCGCustomHLSLDataProvider>();
+}
+
+FComputeDataProviderRenderProxy* UPCGCustomHLSLDataProvider::GetRenderProxy()
+{
+	FPCGCustomHLSLDataProviderProxy::FCustomHLSLData ProxyData;
+
+	if (ensure(KernelParams))
+	{
+		ProxyData.NumElements = KernelParams->GetValueInt(GET_MEMBER_NAME_CHECKED(UPCGCustomHLSLSettings, NumElements));
+		//ProxyData.NumElements2D = KernelParams->GetValueInt2(GET_MEMBER_NAME_CHECKED(UPCGCustomHLSLSettings, NumElements2D));
+		ProxyData.ThreadCountMultiplier = KernelParams->GetValueInt(GET_MEMBER_NAME_CHECKED(UPCGCustomHLSLSettings, ThreadCountMultiplier));
+		ProxyData.FixedThreadCount = KernelParams->GetValueInt(GET_MEMBER_NAME_CHECKED(UPCGCustomHLSLSettings, FixedThreadCount));
+	}
+
+	return new FPCGCustomHLSLDataProviderProxy(MoveTemp(ProxyData));
+}
+
+void UPCGCustomHLSLDataProvider::Reset()
+{
+	Super::Reset();
+
+	KernelParams = nullptr;
+}
+
+bool UPCGCustomHLSLDataProvider::PrepareForExecute_GameThread(UPCGDataBinding* InBinding)
+{
+	TRACE_CPUPROFILER_EVENT_SCOPE(UPCGCustomHLSLDataProvider::PrepareForExecute_GameThread);
+	check(InBinding);
+
+	if (!Super::PrepareForExecute_GameThread(InBinding))
+	{
+		return false;
+	}
+
+	KernelParams = InBinding->GetCachedKernelParams(ProducerKernel);
+
+	return true;
+}
+
+bool FPCGCustomHLSLDataProviderProxy::IsValid(FValidationData const& InValidationData) const
+{
+	return InValidationData.ParameterStructSize == sizeof(FParameters);
+}
+
+void FPCGCustomHLSLDataProviderProxy::GatherDispatchData(FDispatchData const& InDispatchData)
+{
+	const TStridedView<FParameters> ParameterArray = MakeStridedParameterView<FParameters>(InDispatchData);
+
+	for (int32 InvocationIndex = 0; InvocationIndex < ParameterArray.Num(); ++InvocationIndex)
+	{
+		FParameters& Parameters = ParameterArray[InvocationIndex];
+		Parameters.NumElements = Data.NumElements;
+		//Parameters.NumElements2D = static_cast<FVector2f>(Data.NumElements2D);
+		Parameters.ThreadCountMultiplier = Data.ThreadCountMultiplier;
+		Parameters.FixedThreadCount = Data.FixedThreadCount;
+	}
+}
diff --git a/Engine/Plugins/PCG/Source/PCG/Private/Compute/DataInterfaces/Elements/PCGCustomHLSLDataInterface.h b/Engine/Plugins/PCG/Source/PCG/Private/Compute/DataInterfaces/Elements/PCGCustomHLSLDataInterface.h
new file mode 100644
--- /dev/null
+++ b/Engine/Plugins/PCG/Source/PCG/Private/Compute/DataInterfaces/Elements/PCGCustomHLSLDataInterface.h
@@ -0,0 +1,69 @@
+// Copyright Epic Games, Inc. All Rights Reserved.
+
+#pragma once
+
+#include "Compute/DataInterfaces/PCGComputeDataInterface.h"
+
+#include "PCGCustomHLSLDataInterface.generated.h"
+
+class FPCGCustomHLSLDataInterfaceParameters;
+struct FPCGKernelParams;
+
+/** Data Interface to marshal UPCGCustomHLSLSettings settings to the GPU. */
+UCLASS(ClassGroup = (Procedural))
+class UPCGCustomHLSLDataInterface : public UPCGKernelParamsDataInterface
+{
+	GENERATED_BODY()
+
+public:
+	//~ Begin UComputeDataInterface Interface
+	virtual TCHAR const* GetClassName() const override { return TEXT("PCGCustomHLSL"); }
+	virtual void GetShaderParameters(TCHAR const* UID, FShaderParametersMetadataBuilder& InOutBuilder, FShaderParametersMetadataAllocations& InOutAllocations) const override;
+	virtual UComputeDataProvider* CreateDataProvider() const override;
+	//~ End UComputeDataInterface Interface
+};
+
+UCLASS()
+class UPCGCustomHLSLDataProvider : public UPCGKernelParamsDataProvider
+{
+	GENERATED_BODY()
+
+public:
+	//~ Begin UComputeDataProvider Interface
+	virtual FComputeDataProviderRenderProxy* GetRenderProxy() override;
+	virtual void Reset() override;
+	//~ End UComputeDataProvider Interface
+
+	//~ Begin UPCGComputeDataProvider Interface
+	virtual bool PrepareForExecute_GameThread(UPCGDataBinding* InBinding) override;
+	//~ End UPCGComputeDataProvider Interface
+
+	/** The kernel params are owned by the KernelParamsCache, so we are not responsible for its lifetime. */
+	const FPCGKernelParams* KernelParams = nullptr;
+};
+
+class FPCGCustomHLSLDataProviderProxy : public FComputeDataProviderRenderProxy
+{
+public:
+	struct FCustomHLSLData
+	{
+		int32 NumElements = 0;
+		// @todo_pcg: Support param override for NumElements2D. Requires struct support for GPU overridable params.
+		//FIntPoint NumElements2D = FIntPoint::ZeroValue;
+		int32 ThreadCountMultiplier = 0;
+		int32 FixedThreadCount = 0;
+	};
+
+	FPCGCustomHLSLDataProviderProxy(FCustomHLSLData InData)
+		: Data(MoveTemp(InData)) {}
+
+	//~ Begin FComputeDataProviderRenderProxy Interface
+	virtual bool IsValid(FValidationData const& InValidationData) const override;
+	virtual void GatherDispatchData(FDispatchData const& InDispatchData) override;
+	//~ End FComputeDataProviderRenderProxy Interface
+
+protected:
+	using FParameters = FPCGCustomHLSLDataInterfaceParameters;
+
+	FCustomHLSLData Data;
+};
diff --git a/Engine/Plugins/PCG/Source/PCG/Private/Compute/Elements/PCGCustomHLSL.cpp b/Engine/Plugins/PCG/Source/PCG/Private/Compute/Elements/PCGCustomHLSL.cpp
--- a/Engine/Plugins/PCG/Source/PCG/Private/Compute/Elements/PCGCustomHLSL.cpp
+++ b/Engine/Plugins/PCG/Source/PCG/Private/Compute/Elements/PCGCustomHLSL.cpp
@@ -50,6 +50,14 @@ void UPCGCustomHLSLSettings::PostLoad()
 {
 	Super::PostLoad();
 
+#if WITH_EDITOR
+	if (PointCount_DEPRECATED > 0)
+	{
+		NumElements = PointCount_DEPRECATED;
+		PointCount_DEPRECATED = 0;
+	}
+#endif
+
 	// Note: We update here so that Custom HLSL nodes will have the correct pin settings & declarations on load.
 	UpdatePinSettings();
 	UpdateAttributeKeys();
@@ -305,6 +313,8 @@ void UPCGCustomHLSLSettings::CreateKernels(FPCGGPUCompilationContext& InOutConte
 	CreateParams.NodeOutputPinsToWire.Empty(OutputPins.Num());
 	Algo::Transform(OutputPins, CreateParams.NodeOutputPinsToWire, [](const FPCGPinProperties& InProps) { return InProps.Label; });
 
+	CreateParams.bRequiresOverridableParams = true;
+
 	PCGKernelHelpers::CreateKernel<UPCGCustomHLSLKernel>(InOutContext, CreateParams, OutKernels, OutEdges);
 }
 #endif // WITH_EDITOR
@@ -341,15 +351,11 @@ void UPCGCustomHLSLSettings::UpdateInputDeclarations()
 
 		if (KernelType == EPCGKernelType::PointGenerator)
 		{
-			InputDeclarations += FString::Format(TEXT(
-				"const uint NumElements = {0};\n"
-				"const uint NumPoints = {0}; // Deprecated, use NumElements instead.\n" // Deprecated 5.6
-				"\n"
-			), { PointCount });
+			InputDeclarations += TEXT("const uint NumElements;\n\n");
 		}
 		else if (KernelType == EPCGKernelType::TextureGenerator)
 		{
-			InputDeclarations += FString::Format(TEXT("const uint2 NumElements = uint2({0}, {1});\n\n"), { NumElements2D.X, NumElements2D.Y });
+			InputDeclarations += TEXT("const uint2 NumElements;\n\n");
 		}
 
 		InputDeclarations += TEXT("/*** INPUT PER-THREAD CONSTANTS ***/\n\n");
diff --git a/Engine/Plugins/PCG/Source/PCG/Private/Compute/Elements/PCGCustomHLSL.h b/Engine/Plugins/PCG/Source/PCG/Private/Compute/Elements/PCGCustomHLSL.h
--- a/Engine/Plugins/PCG/Source/PCG/Private/Compute/Elements/PCGCustomHLSL.h
+++ b/Engine/Plugins/PCG/Source/PCG/Private/Compute/Elements/PCGCustomHLSL.h
@@ -58,7 +58,6 @@ public:
 	//~Begin UPCGSettings interface
 	virtual TArray<FPCGPinProperties> InputPinProperties() const override { return InputPins; }
 	virtual TArray<FPCGPinProperties> OutputPinProperties() const override;
-	virtual bool HasOverridableParams() const override { return false; }
 	virtual bool IsInputPinRequiredByExecution(const UPCGPin* InPin) const override { return true; }
 	virtual bool UseSeed() const override { return true; }
 #if WITH_EDITOR
@@ -96,9 +95,6 @@ public:
 	//~End IPCGNodeSourceTextProvider interface
 #endif
 
-	/** Get user multiplier value applied to the dispatch thread count. */
-	uint32 GetThreadCountMultiplier() const { return IsThreadCountMultiplierInUse() ? static_cast<uint32>(ThreadCountMultiplier) : 1u; }
-
 #if WITH_EDITOR
 	UFUNCTION()
 	static bool SupportsComposition() { return false; }
@@ -141,12 +137,12 @@ protected:
 	bool IsTextureKernel() const { return KernelType == EPCGKernelType::TextureProcessor || KernelType == EPCGKernelType::TextureGenerator; }
 	bool IsPointKernel() const { return KernelType == EPCGKernelType::PointProcessor || KernelType == EPCGKernelType::PointGenerator; }
 
-protected:
+public:
 	UPROPERTY(EditAnywhere, Category = "Settings")
 	EPCGKernelType KernelType = EPCGKernelType::PointProcessor;
 
-	UPROPERTY(EditAnywhere, DisplayName = "Num Elements", Category = "Settings", meta = (EditCondition = "KernelType == EPCGKernelType::PointGenerator", EditConditionHides))
-	int PointCount = 256;
+	UPROPERTY(EditAnywhere, Category = "Settings", meta = (EditCondition = "KernelType == EPCGKernelType::PointGenerator", EditConditionHides, PCG_OverridableCPUAndGPUWithReadback))
+	int NumElements = 256;
 
 	UPROPERTY(EditAnywhere, DisplayName = "Num Elements", Category = "Settings", meta = (EditCondition = "KernelType == EPCGKernelType::TextureGenerator", EditConditionHides))
 	FIntPoint NumElements2D = FIntPoint(64, 64);
@@ -154,16 +150,15 @@ protected:
 	UPROPERTY(EditAnywhere, Category = "Settings|Thread Count", meta = (EditCondition = "KernelType == EPCGKernelType::Custom", EditConditionHides))
 	EPCGDispatchThreadCount DispatchThreadCount = EPCGDispatchThreadCount::FromFirstOutputPin;
 
-	UPROPERTY(EditAnywhere, Category = "Settings|Thread Count", meta = (EditCondition = "KernelType == EPCGKernelType::Custom && DispatchThreadCount != EPCGDispatchThreadCount::Fixed", EditConditionHides))
+	UPROPERTY(EditAnywhere, Category = "Settings|Thread Count", meta = (EditCondition = "KernelType == EPCGKernelType::Custom && DispatchThreadCount != EPCGDispatchThreadCount::Fixed", EditConditionHides, PCG_OverridableCPUAndGPUWithReadback))
 	int ThreadCountMultiplier = 1;
 
-	UPROPERTY(EditAnywhere, Category = "Settings|Thread Count", meta = (EditCondition = "KernelType == EPCGKernelType::Custom && DispatchThreadCount == EPCGDispatchThreadCount::Fixed", EditConditionHides))
+	UPROPERTY(EditAnywhere, Category = "Settings|Thread Count", meta = (EditCondition = "KernelType == EPCGKernelType::Custom && DispatchThreadCount == EPCGDispatchThreadCount::Fixed", EditConditionHides, PCG_OverridableCPUAndGPUWithReadback))
 	int FixedThreadCount = 1;
 
 	UPROPERTY(EditAnywhere, DisplayName = "Input Pins", Category = "Settings|Thread Count", meta = (EditCondition = "KernelType == EPCGKernelType::Custom && DispatchThreadCount == EPCGDispatchThreadCount::FromProductOfInputPins", EditConditionHides, GetOptions = "GetInputPinNames"))
 	TArray<FName> ThreadCountInputPinLabels;
 
-public:
 	UPROPERTY(EditAnywhere, Category = "Settings")
 	TArray<FPCGPinProperties> InputPins = Super::DefaultPointInputPinProperties();
 
@@ -212,6 +207,9 @@ protected:
 	/** Helper data and functions that can be used from the shader code. Intended to be viewed using the HLSL Source Editor window. */
 	UPROPERTY(Transient)
 	FString HelperDeclarations;
+
+	UPROPERTY()
+	int PointCount_DEPRECATED = 0;
 #endif
 
 	friend class UPCGCustomHLSLKernel;
diff --git a/Engine/Plugins/PCG/Source/PCG/Private/Compute/Elements/PCGCustomHLSLKernel.cpp b/Engine/Plugins/PCG/Source/PCG/Private/Compute/Elements/PCGCustomHLSLKernel.cpp
--- a/Engine/Plugins/PCG/Source/PCG/Private/Compute/Elements/PCGCustomHLSLKernel.cpp
+++ b/Engine/Plugins/PCG/Source/PCG/Private/Compute/Elements/PCGCustomHLSLKernel.cpp
@@ -8,8 +8,9 @@
 #include "PCGPointPropertiesTraits.h"
 #include "Compute/PCGComputeGraph.h"
 #include "Compute/PCGDataBinding.h"
-#include "Graph/PCGGPUGraphCompilationContext.h"
 #include "Compute/Elements/PCGCustomHLSL.h"
+#include "Compute/DataInterfaces/Elements/PCGCustomHLSLDataInterface.h"
+#include "Graph/PCGGPUGraphCompilationContext.h"
 
 #include "ComputeFramework/ComputeSource.h"
 #include "Containers/StaticArray.h"
@@ -257,8 +258,15 @@ bool UPCGCustomHLSLKernel::ComputeOutputBindingDataDesc(const UPCGComputeGraph*
 	}
 	else if (InOutputPinLabel == FirstOutputPinLabel && CustomHLSLSettings->KernelType == EPCGKernelType::PointGenerator)
 	{
+		const FPCGKernelParams* KernelParams = InBinding->GetCachedKernelParams(this);
+
+		if (!ensure(KernelParams))
+		{
+			return false;
+		}
+
 		// Generators always produce a single point data with known point count.
-		OutputDataDesc.DataDescs.Emplace(EPCGDataType::Point, CustomHLSLSettings->PointCount);
+		OutputDataDesc.DataDescs.Emplace(EPCGDataType::Point, KernelParams->GetValueInt(GET_MEMBER_NAME_CHECKED(UPCGCustomHLSLSettings, NumElements)));
 	}
 	else if (InOutputPinLabel == FirstOutputPinLabel && CustomHLSLSettings->KernelType == EPCGKernelType::TextureGenerator)
 	{
@@ -472,11 +480,10 @@ FString UPCGCustomHLSLKernel::GetCookedSource(FPCGGPUCompilationContext& InOutCo
 	}
 	else if (CustomHLSLSettings->KernelType == EPCGKernelType::PointGenerator)
 	{
-		KernelSpecificPreamble += FString::Format(TEXT(
-			"    const uint NumElements = {0};\n"
+		KernelSpecificPreamble += TEXT(
+			"    const uint NumElements = NumElements_GetOverridableValue();\n"
 			"    // NumPoints is deprecated.\n"
-			"    const uint NumPoints = NumElements;\n"
-		), { CustomHLSLSettings->PointCount });
+			"    const uint NumPoints = NumElements;\n");
 
 		if (const FPCGPinPropertiesGPU* OutputPin = GetFirstOutputPin())
 		{
@@ -586,6 +593,16 @@ FString UPCGCustomHLSLKernel::GetCookedSource(FPCGGPUCompilationContext& InOutCo
 
 	return Result;
 }
+
+void UPCGCustomHLSLKernel::CreateAdditionalInputDataInterfaces(FPCGGPUCompilationContext& InOutContext, UObject* InObjectOuter, TArray<TObjectPtr<UComputeDataInterface>>& OutDataInterfaces) const
+{
+	Super::CreateAdditionalInputDataInterfaces(InOutContext, InObjectOuter, OutDataInterfaces);
+
+	TObjectPtr<UPCGCustomHLSLDataInterface> NodeDI = InOutContext.NewObject_AnyThread<UPCGCustomHLSLDataInterface>(InObjectOuter);
+	NodeDI->ProducerKernel = this;
+
+	OutDataInterfaces.Add(NodeDI);
+}
 #endif // WITH_EDITOR
 
 int UPCGCustomHLSLKernel::ComputeThreadCount(const UPCGDataBinding* InBinding) const
@@ -593,10 +610,17 @@ int UPCGCustomHLSLKernel::ComputeThreadCount(const UPCGDataBinding* InBinding) c
 	const UPCGCustomHLSLSettings* CustomHLSLSettings = CastChecked<UPCGCustomHLSLSettings>(GetSettings());
 	int ThreadCount = 0;
 
+	const FPCGKernelParams* KernelParams = InBinding->GetCachedKernelParams(this);
+
+	if (!ensure(KernelParams))
+	{
+		return ThreadCount;
+	}
+
 	if (CustomHLSLSettings->KernelType == EPCGKernelType::PointGenerator)
 	{
 		// Point generator has fixed thread count.
-		ThreadCount = CustomHLSLSettings->PointCount;
+		ThreadCount = KernelParams->GetValueInt(GET_MEMBER_NAME_CHECKED(UPCGCustomHLSLSettings, NumElements));
 	}
 	else if (CustomHLSLSettings->KernelType == EPCGKernelType::TextureGenerator)
 	{
@@ -633,7 +657,7 @@ int UPCGCustomHLSLKernel::ComputeThreadCount(const UPCGDataBinding* InBinding) c
 		}
 		else if (CustomHLSLSettings->DispatchThreadCount == EPCGDispatchThreadCount::Fixed)
 		{
-			ThreadCount = CustomHLSLSettings->FixedThreadCount;
+			ThreadCount = KernelParams->GetValueInt(GET_MEMBER_NAME_CHECKED(UPCGCustomHLSLSettings, FixedThreadCount));
 		}
 	}
 	else
@@ -643,7 +667,7 @@ int UPCGCustomHLSLKernel::ComputeThreadCount(const UPCGDataBinding* InBinding) c
 
 	if (IsThreadCountMultiplierInUse())
 	{
-		ThreadCount *= CustomHLSLSettings->ThreadCountMultiplier;
+		ThreadCount *= KernelParams->GetValueInt(GET_MEMBER_NAME_CHECKED(UPCGCustomHLSLSettings, ThreadCountMultiplier));
 	}
 
 	return ThreadCount;
diff --git a/Engine/Plugins/PCG/Source/PCG/Private/Compute/Elements/PCGCustomHLSLKernel.h b/Engine/Plugins/PCG/Source/PCG/Private/Compute/Elements/PCGCustomHLSLKernel.h
--- a/Engine/Plugins/PCG/Source/PCG/Private/Compute/Elements/PCGCustomHLSLKernel.h
+++ b/Engine/Plugins/PCG/Source/PCG/Private/Compute/Elements/PCGCustomHLSLKernel.h
@@ -119,6 +119,7 @@ public:
 #if WITH_EDITOR
 	virtual FString GetCookedSource(FPCGGPUCompilationContext& InOutContext) const override;
 	virtual FString GetEntryPoint() const override { return EntryPoint; }
+	virtual void CreateAdditionalInputDataInterfaces(FPCGGPUCompilationContext& InOutContext, UObject* InObjectOuter, TArray<TObjectPtr<UComputeDataInterface>>& OutDataInterfaces) const override;
 #endif
 	virtual int ComputeThreadCount(const UPCGDataBinding* Binding) const override;
 	virtual void GetDataLabels(FName InPinLabel, TArray<FString>& OutDataLabels) const override;
diff --git a/Engine/Plugins/PCG/Source/PCG/Private/Compute/PCGKernelParamsCache.cpp b/Engine/Plugins/PCG/Source/PCG/Private/Compute/PCGKernelParamsCache.cpp
--- a/Engine/Plugins/PCG/Source/PCG/Private/Compute/PCGKernelParamsCache.cpp
+++ b/Engine/Plugins/PCG/Source/PCG/Private/Compute/PCGKernelParamsCache.cpp
@@ -145,6 +145,25 @@ bool FPCGKernelParams::GetValueBool(FName InParamName) const
 	return ensure(Property && ValuePtr) ? Property->GetPropertyValue(ValuePtr) : false;
 }
 
+int FPCGKernelParams::GetValueInt(FName InParamName) const
+{
+	const FIntProperty* Property = nullptr;
+	const void* ValuePtr = nullptr;
+
+	if (const FPropertyAndAddress* PropertyAndAddress = OverriddenProperties.Find(InParamName))
+	{
+		Property = CastField<FIntProperty>(PropertyAndAddress->Key);
+		ValuePtr = PropertyAndAddress->Value;
+	}
+	else if (const UPCGSettings* Settings = WeakSettings.Get())
+	{
+		Property = CastField<FIntProperty>(Settings->GetClass()->FindPropertyByName(InParamName));
+		ValuePtr = Property ? Property->ContainerPtrToValuePtr<void>(Settings) : nullptr;
+	}
+
+	return ensure(Property && ValuePtr) ? Property->GetPropertyValue(ValuePtr) : 0;
+}
+
 FName FPCGKernelParams::GetValueName(FName InParamName) const
 {
 	const FNameProperty* Property = nullptr;
diff --git a/Engine/Plugins/PCG/Source/PCG/Private/Compute/PCGKernelParamsCache.h b/Engine/Plugins/PCG/Source/PCG/Private/Compute/PCGKernelParamsCache.h
--- a/Engine/Plugins/PCG/Source/PCG/Private/Compute/PCGKernelParamsCache.h
+++ b/Engine/Plugins/PCG/Source/PCG/Private/Compute/PCGKernelParamsCache.h
@@ -15,6 +15,7 @@ struct FPCGKernelParams
 {
 public:
 	bool GetValueBool(FName InParamName) const;
+	int GetValueInt(FName InParamName) const;
 	FName GetValueName(FName InParamName) const;
 
 	template<typename EnumValueType>
