[PCG] Adjusted wording, robustness in GetSegment node.

#rnx

diff --git a/Engine/Plugins/PCG/Source/PCG/Private/Data/PCGPolygon2DData.cpp b/Engine/Plugins/PCG/Source/PCG/Private/Data/PCGPolygon2DData.cpp
--- a/Engine/Plugins/PCG/Source/PCG/Private/Data/PCGPolygon2DData.cpp
+++ b/Engine/Plugins/PCG/Source/PCG/Private/Data/PCGPolygon2DData.cpp
@@ -605,7 +605,11 @@ FTransform UPCGPolygon2DData::GetTransformAtAlpha(float InAlpha) const
 
 	FTransform OutTransform;
 	OutTransform.SetLocation(FVector(PointOnSegment, 0.0));
-	OutTransform.SetRotation(FRotationMatrix::MakeFromXZ(FVector(Segment.Direction, 0.0), FVector::UpVector).ToQuat()); // unclear z sign here, should depend on winding(?)
+	
+	// @todo_pcg: check if Z sign has to be consequent with the polygon winding (likely)
+	const FQuat Rotation = Segment.Direction.IsNearlyZero() ? FQuat::Identity : FRotationMatrix::MakeFromXZ(FVector(Segment.Direction, 0.0), FVector::UpVector).ToQuat();
+	OutTransform.SetRotation(Rotation);
+	
 	OutTransform.SetScale3D(FVector::OneVector);
 
 	return OutTransform * Transform;
@@ -632,8 +636,11 @@ FTransform UPCGPolygon2DData::GetTransformAtDistance(int SegmentIndex, double Di
 	const FVector2D PointOnSegment = Segment.PointAt(Distance - Segment.Extent);
 
 	OutTransform.SetLocation(FVector(PointOnSegment, 0.0));
+
 	// @todo_pcg: check if Z sign has to be consequent with the polygon winding (likely)
-	OutTransform.SetRotation(FRotationMatrix::MakeFromXZ(FVector(Segment.Direction, 0.0), FVector::UpVector).ToQuat()); // unclear z sign here, should depend on winding(?)
+	const FQuat Rotation = Segment.Direction.IsNearlyZero() ? FQuat::Identity : FRotationMatrix::MakeFromXZ(FVector(Segment.Direction, 0.0), FVector::UpVector).ToQuat();
+	OutTransform.SetRotation(Rotation);
+
 	OutTransform.SetScale3D(FVector::OneVector);
 
 	if (bWorldSpace)
diff --git a/Engine/Plugins/PCG/Source/PCG/Public/Elements/Polygon/PCGGetSegment.h b/Engine/Plugins/PCG/Source/PCG/Public/Elements/Polygon/PCGGetSegment.h
--- a/Engine/Plugins/PCG/Source/PCG/Public/Elements/Polygon/PCGGetSegment.h
+++ b/Engine/Plugins/PCG/Source/PCG/Public/Elements/Polygon/PCGGetSegment.h
@@ -49,7 +49,7 @@ public:
 	UPROPERTY(EditAnywhere, BlueprintReadWrite, Category = Settings)
 	bool bOutputSplineData = false;
 
-	/** Specifies the segment index attribute, if using input segment data. */
+	/** Specifies the attribute from which to extract the segment index value. */
 	UPROPERTY(EditAnywhere, BlueprintReadWrite, Category = Settings, meta = (EditCondition = "bUseInputSegmentData", EditConditionHides, PCG_Overridable))
 	FPCGAttributePropertyInputSelector SegmentIndexAttribute;
 
@@ -57,7 +57,7 @@ public:
 	UPROPERTY(EditAnywhere, BlueprintReadWrite, Category = Settings, meta = (EditCondition = "!bUseInputSegmentData", EditConditionHides, PCG_Overridable))
 	int32 SegmentIndex = 0;
 
-	/** Specifies the hole index attribute, if using input segment data and polygon 2d data. Will be ignored if same value as the segment index attribute. */
+	/** Specifies the attribute from which to extract the hole index value. */
 	UPROPERTY(EditAnywhere, BlueprintReadWrite, Category = Settings, meta = (EditCondition = "bUseInputHoleData", EditConditionHides, PCG_Overridable))
 	FPCGAttributePropertyInputSelector HoleIndexAttribute;
 
