[PCG] Add color and static validation on Overrides pins on nodes.

* Use the same colors as BP for types
* Types are now shown in the subtypes category for overrides
* Also implemented for Graph Parameters (with matching node color) and Create Constant
* Static validation when trying to plug an typed param into an incompatible type param

diff --git a/Engine/Plugins/PCG/Source/PCG/Private/Data/Registry/PCGDataType.cpp b/Engine/Plugins/PCG/Source/PCG/Private/Data/Registry/PCGDataType.cpp
--- a/Engine/Plugins/PCG/Source/PCG/Private/Data/Registry/PCGDataType.cpp
+++ b/Engine/Plugins/PCG/Source/PCG/Private/Data/Registry/PCGDataType.cpp
@@ -100,9 +100,14 @@ TOptional<TArray<UPCGNode*>> FPCGDataTypeInfo::AddConversionNodesTo(const FPCGDa
 
 #if WITH_EDITOR
 
-FText FPCGDataTypeInfo::GetTooltip() const
+TOptional<FText> FPCGDataTypeInfo::GetSubtypeTooltip(const FPCGDataTypeIdentifier& ThisType) const
 {
-	return FText{};
+	return {};
+}
+
+TOptional<FText> FPCGDataTypeInfo::GetExtraTooltip(const FPCGDataTypeIdentifier& ThisType) const
+{
+	return {};
 }
 
 bool FPCGDataTypeInfo::Hidden() const
diff --git a/Engine/Plugins/PCG/Source/PCG/Private/Data/Registry/PCGDataTypeIdentifier.cpp b/Engine/Plugins/PCG/Source/PCG/Private/Data/Registry/PCGDataTypeIdentifier.cpp
--- a/Engine/Plugins/PCG/Source/PCG/Private/Data/Registry/PCGDataTypeIdentifier.cpp
+++ b/Engine/Plugins/PCG/Source/PCG/Private/Data/Registry/PCGDataTypeIdentifier.cpp
@@ -427,6 +427,40 @@ FText FPCGDataTypeIdentifier::ToDisplayText() const
 	return FText::FromString(ToString());
 }
 
+#if WITH_EDITOR
+FText FPCGDataTypeIdentifier::GetSubtypeTooltip() const
+{
+	if (IsValid() && !IsComposition())
+	{
+		if (const FPCGDataTypeInfo* Info = FPCGModule::GetConstDataTypeRegistry().GetTypeInfo(GetId()))
+		{
+			if (TOptional<FText> Tooltip = Info->GetSubtypeTooltip(*this); Tooltip.IsSet())
+			{
+				return MoveTemp(Tooltip.GetValue());
+			}
+		}
+	}
+
+	return ToDisplayText();
+}
+
+FText FPCGDataTypeIdentifier::GetExtraTooltip() const
+{
+	if (IsValid() && !IsComposition())
+	{
+		if (const FPCGDataTypeInfo* Info = FPCGModule::GetConstDataTypeRegistry().GetTypeInfo(GetId()))
+		{
+			if (TOptional<FText> Tooltip = Info->GetExtraTooltip(*this); Tooltip.IsSet())
+			{
+				return MoveTemp(Tooltip.GetValue());
+			}
+		}
+	}
+
+	return {};
+}
+#endif // WITH_EDITOR
+
 FPCGDataTypeIdentifier UPCGDataTypeIdentifierHelpers::GetIdentifierFromLegacyType(EPCGExclusiveDataType LegacyDataType)
 {
 	const UEnum* DataTypeEnum = StaticEnum<EPCGDataType>();
diff --git a/Engine/Plugins/PCG/Source/PCG/Private/Data/Registry/PCGDataTypeRegistry.cpp b/Engine/Plugins/PCG/Source/PCG/Private/Data/Registry/PCGDataTypeRegistry.cpp
--- a/Engine/Plugins/PCG/Source/PCG/Private/Data/Registry/PCGDataTypeRegistry.cpp
+++ b/Engine/Plugins/PCG/Source/PCG/Private/Data/Registry/PCGDataTypeRegistry.cpp
@@ -347,9 +347,22 @@ EPCGDataTypeCompatibilityResult FPCGDataTypeRegistry::IsCompatible(const FPCGDat
 		return EPCGDataTypeCompatibilityResult::UnknownType;
 	}
 
+	auto IsCompatibleForSubtype = [this, &InType, &OutType, &OptionalOutCompatibilityMessage]()
+	{
+		if (!InType.IsComposition())
+		{
+			const TInstancedStruct<FPCGDataTypeInfo>* InDefIt = Mapping.Find(InType.GetId());
+			return InDefIt ? InDefIt->Get().IsCompatibleForSubtype(InType, OutType, OptionalOutCompatibilityMessage) : EPCGDataTypeCompatibilityResult::UnknownType;
+		}
+		else
+		{
+			return EPCGDataTypeCompatibilityResult::Compatible;
+		}
+	};
+
 	if (InType.IsSameType(OutType))
 	{
-		return EPCGDataTypeCompatibilityResult::Compatible;
+		return IsCompatibleForSubtype();
 	}
 
 	// Run the intersection between the two.
@@ -358,7 +371,7 @@ EPCGDataTypeCompatibilityResult FPCGDataTypeRegistry::IsCompatible(const FPCGDat
 	// If the intersection gives InType, it means OutType accept all InType, so it is compatible
 	if (Intersection.IsSameType(InType))
 	{
-		return EPCGDataTypeCompatibilityResult::Compatible;
+		return IsCompatibleForSubtype();
 	}
 	
 	// If the intersection gives OutType, or something else that is valid, it means that InType is wider than OutType, so we might need a filter or a conversion.
diff --git a/Engine/Plugins/PCG/Source/PCG/Private/Elements/PCGCreateAttribute.cpp b/Engine/Plugins/PCG/Source/PCG/Private/Elements/PCGCreateAttribute.cpp
--- a/Engine/Plugins/PCG/Source/PCG/Private/Elements/PCGCreateAttribute.cpp
+++ b/Engine/Plugins/PCG/Source/PCG/Private/Elements/PCGCreateAttribute.cpp
@@ -269,7 +269,9 @@ TArray<FPCGPinProperties> UPCGCreateAttributeSetSettings::InputPinProperties() c
 TArray<FPCGPinProperties> UPCGCreateAttributeSetSettings::OutputPinProperties() const
 {
 	TArray<FPCGPinProperties> PinProperties;
-	PinProperties.Emplace(PCGPinConstants::DefaultOutputLabel, EPCGDataType::Param);
+	FPCGPinProperties& Pin = PinProperties.Emplace_GetRef(PCGPinConstants::DefaultOutputLabel, EPCGDataType::Param);
+
+	Pin.AllowedTypes.CustomSubtype = static_cast<int32>(AttributeTypes.Type);
 
 	return PinProperties;
 }
diff --git a/Engine/Plugins/PCG/Source/PCG/Private/Elements/PCGReroute.cpp b/Engine/Plugins/PCG/Source/PCG/Private/Elements/PCGReroute.cpp
--- a/Engine/Plugins/PCG/Source/PCG/Private/Elements/PCGReroute.cpp
+++ b/Engine/Plugins/PCG/Source/PCG/Private/Elements/PCGReroute.cpp
@@ -3,12 +3,14 @@
 #include "Elements/PCGReroute.h"
 
 #include "PCGContext.h"
+#include "PCGEdge.h"
 #include "PCGGraph.h"
 #include "PCGModule.h"
 #include "PCGNode.h"
 #include "PCGPin.h"
 #include "Compute/PCGDataForGPU.h"
 #include "Helpers/PCGConversion.h"
+#include "Helpers/PCGSettingsHelpers.h"
 
 #include UE_INLINE_GENERATED_CPP_BY_NAME(PCGReroute)
 
@@ -46,6 +48,25 @@ UPCGRerouteSettings::UPCGRerouteSettings()
 #endif
 }
 
+FPCGDataTypeIdentifier UPCGRerouteSettings::GetCurrentPinTypesID(const UPCGPin* InPin) const
+{
+	// Since we can have only a single input pin, just forward the type (with its potential subtype)
+	if (InPin && InPin->IsOutputPin())
+	{		
+		const UPCGNode* Node = Cast<const UPCGNode>(GetOuter());
+		const UPCGPin* FirstPin = Node && !Node->GetInputPins().IsEmpty() ? Node->GetInputPins()[0] : nullptr;
+
+		if (const UPCGPin* OtherPin = (FirstPin && !FirstPin->Edges.IsEmpty() && FirstPin->Edges[0]) ? FirstPin->Edges[0]->InputPin : nullptr)
+		{
+			// Use the scope to avoid recomputing it multiple times
+			PCGSettingsHelpers::FPinTypeScopeHelper PinTypeHelper;
+			return PinTypeHelper.GetCurrentPinTypeID(OtherPin);
+		}
+	}
+	
+	return Super::GetCurrentPinTypesID(InPin);
+}
+
 TOptional<FName> UPCGRerouteSettings::GetCollisionFreeNodeName(const UPCGGraph* InGraph, FName BaseName)
 {
 	FName NamedRerouteTitle(BaseName);
diff --git a/Engine/Plugins/PCG/Source/PCG/Private/Elements/PCGUserParameterGet.cpp b/Engine/Plugins/PCG/Source/PCG/Private/Elements/PCGUserParameterGet.cpp
--- a/Engine/Plugins/PCG/Source/PCG/Private/Elements/PCGUserParameterGet.cpp
+++ b/Engine/Plugins/PCG/Source/PCG/Private/Elements/PCGUserParameterGet.cpp
@@ -179,7 +179,33 @@ namespace PCGUserParameterGet
 TArray<FPCGPinProperties> UPCGUserParameterGetSettings::OutputPinProperties() const
 {
 	TArray<FPCGPinProperties> PinProperties;
-	PinProperties.Emplace(PropertyName, EPCGDataType::Param);
+	FPCGPinProperties& Pin = PinProperties.Emplace_GetRef(PropertyName, EPCGDataType::Param);
+
+	if (PropertyName != NAME_None)
+	{
+		// Find the property associated with this param
+		const UPCGNode* Node = Cast<UPCGNode>(GetOuter());
+		const UPCGGraph* Graph = Node ? Node->GetGraph() : nullptr;
+		const FPropertyBagPropertyDesc* PropertyDesc = Graph ? Graph->GetUserParametersStruct()->FindPropertyDescByName(PropertyName) : nullptr;
+		if (PropertyDesc && PropertyDesc->CachedProperty)
+		{
+			// If the property is a container, check its inner
+			const FProperty* Property = PropertyDesc->CachedProperty;
+			if (const FArrayProperty* ArrayProperty = CastField<FArrayProperty>(Property))
+			{
+				Property = ArrayProperty->Inner;
+			}
+			else if (const FSetProperty* SetProperty = CastField<FSetProperty>(Property))
+			{
+				Property = SetProperty->ElementProp;
+			}
+
+			if (Property)
+			{
+				Pin.AllowedTypes.CustomSubtype = static_cast<int32>(PCGAttributeAccessorHelpers::GetMetadataTypeForProperty(Property));
+			}
+		}
+	}
 
 	return PinProperties;
 }
diff --git a/Engine/Plugins/PCG/Source/PCG/Private/Helpers/PCGPropertyHelpers.cpp b/Engine/Plugins/PCG/Source/PCG/Private/Helpers/PCGPropertyHelpers.cpp
--- a/Engine/Plugins/PCG/Source/PCG/Private/Helpers/PCGPropertyHelpers.cpp
+++ b/Engine/Plugins/PCG/Source/PCG/Private/Helpers/PCGPropertyHelpers.cpp
@@ -215,14 +215,7 @@ namespace PCGPropertyHelpers
 
 EPCGMetadataTypes PCGPropertyHelpers::GetMetadataTypeFromProperty(const FProperty* InProperty)
 {
-	if (!InProperty)
-	{
-		return EPCGMetadataTypes::Unknown;
-	}
-
-	TUniquePtr<IPCGAttributeAccessor> PropertyAccessor = PCGAttributeAccessorHelpers::CreatePropertyAccessor(InProperty);
-
-	return PropertyAccessor.IsValid() ? EPCGMetadataTypes(PropertyAccessor->GetUnderlyingType()) : EPCGMetadataTypes::Unknown;
+	return PCGAttributeAccessorHelpers::GetMetadataTypeForProperty(InProperty);
 }
 
 FPropertyBagPropertyDesc PCGPropertyHelpers::CreatePropertyBagDescWithMetadataType(const FName InPropertyName, const EPCGMetadataTypes Type)
diff --git a/Engine/Plugins/PCG/Source/PCG/Private/Helpers/PCGSettingsHelpers.cpp b/Engine/Plugins/PCG/Source/PCG/Private/Helpers/PCGSettingsHelpers.cpp
--- a/Engine/Plugins/PCG/Source/PCG/Private/Helpers/PCGSettingsHelpers.cpp
+++ b/Engine/Plugins/PCG/Source/PCG/Private/Helpers/PCGSettingsHelpers.cpp
@@ -426,6 +426,10 @@ namespace PCGSettingsHelpers
 					Param.Properties.Add(Property);
 					Param.Properties.Append(std::move(ChildParam.Properties));
 					Param.PropertyClass = InClass;
+#if WITH_EDITOR
+					Param.UnderlyingType = ChildParam.UnderlyingType;
+#endif //WITH_EDITOR
+					
 					if (!ChildParam.bHasNameClash && !bHasNameClash)
 					{
 						Param.Label = Label;
@@ -493,6 +497,7 @@ namespace PCGSettingsHelpers
 				Param.bSupportsGPU = Property->HasMetaData(PCGObjectMetadata::OverridableCPUAndGPU) || Property->HasMetaData(PCGObjectMetadata::OverridableCPUAndGPUWithReadback);
 				Param.bRequiresGPUReadback = Property->HasMetaData(PCGObjectMetadata::OverridableCPUAndGPUWithReadback);
 				GatherAliases(Property, Param);
+				Param.UnderlyingType = PCGAttributeAccessorHelpers::GetMetadataTypeForProperty(Property);
 #endif // WITH_EDITOR
 			}
 			else if (const FStructProperty* StructProperty = CastField<FStructProperty>(PropertyToCheck))
diff --git a/Engine/Plugins/PCG/Source/PCG/Private/Metadata/Accessors/PCGAttributeAccessorHelpers.cpp b/Engine/Plugins/PCG/Source/PCG/Private/Metadata/Accessors/PCGAttributeAccessorHelpers.cpp
--- a/Engine/Plugins/PCG/Source/PCG/Private/Metadata/Accessors/PCGAttributeAccessorHelpers.cpp
+++ b/Engine/Plugins/PCG/Source/PCG/Private/Metadata/Accessors/PCGAttributeAccessorHelpers.cpp
@@ -351,6 +351,19 @@ bool PCGAttributeAccessorHelpers::IsPropertyAccessorChainSupported(const TArray<
 		PCGAttributeAccessorHelpers::IsPropertyAccessorSupported(PropertyChain.Last());
 }
 
+EPCGMetadataTypes PCGAttributeAccessorHelpers::GetMetadataTypeForProperty(const FProperty* InProperty)
+{
+	bool bFound = false;
+	const EPCGMetadataTypes Result = PCGAttributeAccessorHelpers::DispatchPropertyTypes(InProperty, [&bFound](auto SignatureDummy, const auto* TypedProperty) -> EPCGMetadataTypes
+	{
+		bFound = true;
+		using TypedAccessor = typename decltype(SignatureDummy)::Type;
+		return static_cast<EPCGMetadataTypes>(PCG::Private::MetadataTypes<typename TypedAccessor::Type>::Id);
+	});
+
+	return bFound ? Result : EPCGMetadataTypes::Unknown;
+}
+
 TUniquePtr<IPCGAttributeAccessor> PCGAttributeAccessorHelpers::CreateExtraAccessor(EPCGExtraProperties InExtraProperties)
 {
 	switch (InExtraProperties)
diff --git a/Engine/Plugins/PCG/Source/PCG/Private/PCGContext.cpp b/Engine/Plugins/PCG/Source/PCG/Private/PCGContext.cpp
--- a/Engine/Plugins/PCG/Source/PCG/Private/PCGContext.cpp
+++ b/Engine/Plugins/PCG/Source/PCG/Private/PCGContext.cpp
@@ -416,7 +416,12 @@ void FPCGContext::OverrideSettings()
 			PropertyType Value{};
 			if (!PCGAttributeAccessorHelpers::GetOverrideParamValue(*AttributeAccessor, Value))
 			{
-				PCGE_LOG_C(Warning, GraphAndLog, this, FText::Format(LOCTEXT("ConversionFailed", "Parameter '{0}' cannot be converted from attribute '{1}'"), FText::FromName(Param.Label), FText::FromName(AttributeName)));
+				PCGE_LOG_C(Warning, GraphAndLog, this, FText::Format(LOCTEXT("ConversionFailed", "Parameter '{0}' ({1}) cannot be converted from attribute '{2}' ({3})"),
+					FText::FromName(Param.Label), 
+					PCG::Private::GetTypeNameText(PCG::Private::MetadataTypes<PropertyType>::Id),
+					FText::FromName(AttributeName),
+					PCG::Private::GetTypeNameText(AttributeAccessor->GetUnderlyingType())));
+				
 				return false;
 			}
 
diff --git a/Engine/Plugins/PCG/Source/PCG/Private/PCGGraph.cpp b/Engine/Plugins/PCG/Source/PCG/Private/PCGGraph.cpp
--- a/Engine/Plugins/PCG/Source/PCG/Private/PCGGraph.cpp
+++ b/Engine/Plugins/PCG/Source/PCG/Private/PCGGraph.cpp
@@ -2049,7 +2049,10 @@ bool UPCGGraph::UserParametersIsPinTypeAccepted(FEdGraphPinType InPinType, bool
 
 void UPCGGraph::OnGraphParametersChanged(EPCGGraphParameterEvent InChangeType, FName InChangedPropertyName)
 {
-	if (InChangeType == EPCGGraphParameterEvent::RemovedUsed || InChangeType == EPCGGraphParameterEvent::RemovedUnused || InChangeType == EPCGGraphParameterEvent::PropertyRenamed)
+	if (InChangeType == EPCGGraphParameterEvent::RemovedUsed
+		|| InChangeType == EPCGGraphParameterEvent::RemovedUnused
+		|| InChangeType == EPCGGraphParameterEvent::PropertyRenamed
+		|| InChangeType == EPCGGraphParameterEvent::PropertyTypeModified)
 	{
 		// Look for all the Get Parameter nodes and make sure to delete all nodes that doesn't exist anymore
 		TArray<UPCGNode*> NodesToRemove;
@@ -2076,6 +2079,12 @@ void UPCGGraph::OnGraphParametersChanged(EPCGGraphParameterEvent InChangeType, F
 					// We make sure to keep the edges connected, by renaming the pin label
 					Node->RenameOutputPin(OldName, PropertyDesc->Name);
 				}
+
+				if (InChangeType == EPCGGraphParameterEvent::PropertyTypeModified)
+				{
+					// Also force a refresh on the pin output for the subtype.
+					Node->UpdatePins();
+				}
 			}
 		}
 
diff --git a/Engine/Plugins/PCG/Source/PCG/Private/PCGParamData.cpp b/Engine/Plugins/PCG/Source/PCG/Private/PCGParamData.cpp
--- a/Engine/Plugins/PCG/Source/PCG/Private/PCGParamData.cpp
+++ b/Engine/Plugins/PCG/Source/PCG/Private/PCGParamData.cpp
@@ -17,6 +17,31 @@
 
 PCG_DEFINE_TYPE_INFO(FPCGDataTypeInfoParam, UPCGParamData)
 
+EPCGDataTypeCompatibilityResult FPCGDataTypeInfoParam::IsCompatibleForSubtype(const FPCGDataTypeIdentifier& InType, const FPCGDataTypeIdentifier& OutType, FText* OptionalOutCompatibilityMessage) const
+{	
+	// We can only compare subtypes if they are all params with valid custom subtype
+	if (HasValidSubtype(InType) && HasValidSubtype(OutType))
+	{
+		if (!PCG::Private::IsBroadcastableOrConstructible(static_cast<uint16>(InType.CustomSubtype), static_cast<uint16>(OutType.CustomSubtype)))
+		{
+			if (OptionalOutCompatibilityMessage)
+			{
+				*OptionalOutCompatibilityMessage = FText::Format(LOCTEXT("IncompatibleSubtypes", "Input pin type '{0}' is not a compatible with type '{1}'"), PCG::Private::GetTypeNameText(static_cast<uint16>(InType.CustomSubtype)), PCG::Private::GetTypeNameText(static_cast<uint16>(OutType.CustomSubtype)));
+			}
+
+			return EPCGDataTypeCompatibilityResult::TypeCompatibleSubtypeNotCompatible;
+		}
+		else
+		{
+			return EPCGDataTypeCompatibilityResult::Compatible;
+		}
+	}
+	else
+	{
+		return EPCGDataTypeCompatibilityResult::Compatible;
+	}
+}
+
 bool FPCGDataTypeInfoParam::SupportsConversionTo(const FPCGDataTypeIdentifier& ThisType, const FPCGDataTypeIdentifier& OutputType, TSubclassOf<UPCGSettings>* OptionalOutConversionSettings, FText* OptionalOutCompatibilityMessage) const
 {
 	// Param can convert to points
@@ -35,6 +60,25 @@ bool FPCGDataTypeInfoParam::SupportsConversionTo(const FPCGDataTypeIdentifier& T
 	}
 }
 
+#if WITH_EDITOR
+TOptional<FText> FPCGDataTypeInfoParam::GetSubtypeTooltip(const FPCGDataTypeIdentifier& ThisType) const
+{
+	if (HasValidSubtype(ThisType))
+	{
+		return PCG::Private::GetTypeNameText(static_cast<uint16>(ThisType.CustomSubtype));
+	}
+	else
+	{
+		return LOCTEXT("NoTypeParam", "Any");
+	}
+}
+#endif // WITH_EDITOR
+
+bool FPCGDataTypeInfoParam::HasValidSubtype(const FPCGDataTypeIdentifier& Id)
+{
+	return !Id.IsComposition() && Id.GetId() == AsId() && Id.CustomSubtype >= 0 && Id.CustomSubtype < static_cast<int32>(EPCGMetadataTypes::Count);
+}
+
 UPCGParamData::UPCGParamData(const FObjectInitializer& ObjectInitializer)
 	: Super(ObjectInitializer)
 {
diff --git a/Engine/Plugins/PCG/Source/PCG/Private/PCGPin.cpp b/Engine/Plugins/PCG/Source/PCG/Private/PCGPin.cpp
--- a/Engine/Plugins/PCG/Source/PCG/Private/PCGPin.cpp
+++ b/Engine/Plugins/PCG/Source/PCG/Private/PCGPin.cpp
@@ -285,7 +285,14 @@ bool UPCGPin::BreakAllIncompatibleEdges(TSet<UPCGNode*>* InTouchedNodes/*= nullp
 		UPCGEdge* Edge = Edges[EdgeIndex];
 		UPCGPin* OtherPin = Edge->GetOtherPin(this);
 
-		bool bRemoveEdge = !IsCompatible(OtherPin) || (!AllowsMultipleConnections() && bHasAValidEdge);
+		bool bRemoveEdge = (!AllowsMultipleConnections() && bHasAValidEdge);
+
+		// If the edge is not compatible because of a subtype mismatch, we do not want to break it for deprecation purposes.
+		const EPCGDataTypeCompatibilityResult Compatibility = GetCompatibilityWithOtherPin(OtherPin);
+		if (!PCGDataTypeCompatibilityResult::IsValid(Compatibility) && Compatibility != EPCGDataTypeCompatibilityResult::TypeCompatibleSubtypeNotCompatible)
+		{
+			bRemoveEdge = true;
+		}
 
 		if (bRemoveEdge)
 		{
@@ -368,7 +375,7 @@ FPCGDataTypeIdentifier UPCGPin::GetCurrentTypesID() const
 bool UPCGPin::IsCompatible(const UPCGPin* OtherPin) const
 {
 	const EPCGDataTypeCompatibilityResult Compatibility = GetCompatibilityWithOtherPin(OtherPin);
-	return Compatibility != EPCGDataTypeCompatibilityResult::NotCompatible && Compatibility != EPCGDataTypeCompatibilityResult::UnknownType;
+	return PCGDataTypeCompatibilityResult::IsValid(Compatibility);
 }
 
 EPCGDataTypeCompatibilityResult UPCGPin::GetCompatibilityWithOtherPin(const UPCGPin* OtherPin, FText* OptionalOutCompatibilityMessage) const
@@ -388,7 +395,7 @@ EPCGDataTypeCompatibilityResult UPCGPin::GetCompatibilityWithOtherPin(const UPCG
 bool UPCGPin::IsDownstreamPinTypeCompatible(const EPCGDataType UpstreamTypes) const
 {
 	const EPCGDataTypeCompatibilityResult Compatibility = FPCGModule::GetConstDataTypeRegistry().IsCompatible(FPCGDataTypeIdentifier{UpstreamTypes}, Properties.AllowedTypes);
-	return Compatibility != EPCGDataTypeCompatibilityResult::NotCompatible && Compatibility != EPCGDataTypeCompatibilityResult::UnknownType;
+	return PCGDataTypeCompatibilityResult::IsValid(Compatibility);
 }
 
 bool UPCGPin::AllowsMultipleConnections() const
diff --git a/Engine/Plugins/PCG/Source/PCG/Private/PCGSettings.cpp b/Engine/Plugins/PCG/Source/PCG/Private/PCGSettings.cpp
--- a/Engine/Plugins/PCG/Source/PCG/Private/PCGSettings.cpp
+++ b/Engine/Plugins/PCG/Source/PCG/Private/PCGSettings.cpp
@@ -495,6 +495,7 @@ void UPCGSettings::FillOverridableParamsPins(TArray<FPCGPinProperties>& OutPins)
 	else
 	{
 		FPCGPinProperties& ParamPin = OutPins.Emplace_GetRef(PCGPinConstants::DefaultParamsLabel, EPCGDataType::Param, /*bInAllowMultipleConnections=*/ true, /*bAllowMultipleData=*/ true);
+		ParamPin.AllowedTypes.CustomSubtype = static_cast<int32>(EPCGMetadataTypes::Unknown);
 		ParamPin.SetOverrideOrUserParamPin();
 
 #if WITH_EDITOR
@@ -523,6 +524,7 @@ void UPCGSettings::FillOverridableParamsPins(TArray<FPCGPinProperties>& OutPins)
 		FPCGPinProperties& ParamPin = OutPins.Emplace_GetRef(OverridableParam.Label, EPCGDataType::Param, /*bInAllowMultipleConnections=*/ false, /*bAllowMultipleData=*/ false);
 		ParamPin.SetOverrideOrUserParamPin();
 #if WITH_EDITOR
+		ParamPin.AllowedTypes.CustomSubtype = static_cast<int32>(OverridableParam.UnderlyingType);
 
 		if (!OverridableParam.Properties.IsEmpty())
 		{
diff --git a/Engine/Plugins/PCG/Source/PCG/Public/Data/Registry/PCGDataType.h b/Engine/Plugins/PCG/Source/PCG/Public/Data/Registry/PCGDataType.h
--- a/Engine/Plugins/PCG/Source/PCG/Public/Data/Registry/PCGDataType.h
+++ b/Engine/Plugins/PCG/Source/PCG/Public/Data/Registry/PCGDataType.h
@@ -3,6 +3,7 @@
 #pragma once
 
 #include "PCGCommon.h"
+#include "Data/Registry/PCGDataTypeCommon.h"
 
 #include "PCGDataType.generated.h"
 
@@ -72,6 +73,13 @@ struct FPCGDataTypeInfo
 	virtual EPCGDataType GetAssociatedLegacyType() const { return EPCGDataType::Any; }
 	static EPCGDataType GetStaticAssociatedLegacyType() { return EPCGDataType::Any; }
 
+	/**
+	 * Can be overriden to have more restrictive compatibility based on the subtype. InType and OutType are expected to be already compatible without their subtype.
+	 * InType is expected to be the type associated with this type info.
+	 * Can also provide a custom compatibility message, instead of the generic "Not compatible"
+	 */
+	virtual EPCGDataTypeCompatibilityResult IsCompatibleForSubtype(const FPCGDataTypeIdentifier& InType, const FPCGDataTypeIdentifier& OutType, FText* OptionalOutCompatibilityMessage = nullptr) const { return EPCGDataTypeCompatibilityResult::Compatible; }
+
 	/**
 	 * Can be overriden by a type info to indicate if we can convert from InputType to ThisType.
 	 * Since we have a hierarchy of type, ThisType is by construction a child of the current type info ; But it can be a child and not the same type.
@@ -95,7 +103,11 @@ struct FPCGDataTypeInfo
 	PCG_API virtual TOptional<TArray<UPCGNode*>> AddConversionNodesTo(const FPCGDataTypeIdentifier& ThisType, const FPCGDataTypeIdentifier& OutputType, UPCGGraph* InGraph, UPCGPin* InUpstreamPin, UPCGPin* InDownstreamPin) const;
 
 #if WITH_EDITOR
-	PCG_API virtual FText GetTooltip() const;
+	/** Can provide a custom tooltip for subtype that will be shown when hovering a pin of this type. */
+	PCG_API virtual TOptional<FText> GetSubtypeTooltip(const FPCGDataTypeIdentifier& ThisType) const;
+
+	/** Can provide a custom extra tooltip that will be shown when hovering a pin of this type. */
+	PCG_API virtual TOptional<FText> GetExtraTooltip(const FPCGDataTypeIdentifier& ThisType) const;
 
 	// If the type should be hidden to the user.
 	PCG_API virtual bool Hidden() const;
diff --git a/Engine/Plugins/PCG/Source/PCG/Public/Data/Registry/PCGDataTypeCommon.h b/Engine/Plugins/PCG/Source/PCG/Public/Data/Registry/PCGDataTypeCommon.h
new file mode 100644
--- /dev/null
+++ b/Engine/Plugins/PCG/Source/PCG/Public/Data/Registry/PCGDataTypeCommon.h
@@ -0,0 +1,38 @@
+﻿// Copyright Epic Games, Inc. All Rights Reserved.
+
+#pragma once
+
+enum class EPCGDataTypeCompatibilityResult
+{
+	// Types are not compatible, they can't be connected together.
+	NotCompatible = 0,
+
+	// Types are fully compatible, they can be connected as-is.
+	Compatible = 1,
+
+	// Types are compatible, but requires a filter to be added for the connection.
+	RequireFilter = 2,
+
+	// Types are compatible, but require conversion nodes to be added for the connection.
+	RequireConversion = 3,
+
+	// Error if the identifier is not found
+	UnknownType = 4,
+
+	// For deprecation purposes, we do not want to break overridable pins that were mismatched in types, as it will
+	// make the dynamic error go away, and "fail" silently.
+	// This will allow to catch when types are compatible but not their subtype.
+	TypeCompatibleSubtypeNotCompatible = 5,
+
+	Count,
+};
+
+namespace PCGDataTypeCompatibilityResult
+{
+	inline bool IsValid(const EPCGDataTypeCompatibilityResult Result)
+	{
+		return Result != EPCGDataTypeCompatibilityResult::NotCompatible
+		&& Result != EPCGDataTypeCompatibilityResult::UnknownType
+		&& Result != EPCGDataTypeCompatibilityResult::TypeCompatibleSubtypeNotCompatible;
+	}
+}
diff --git a/Engine/Plugins/PCG/Source/PCG/Public/Data/Registry/PCGDataTypeIdentifier.h b/Engine/Plugins/PCG/Source/PCG/Public/Data/Registry/PCGDataTypeIdentifier.h
--- a/Engine/Plugins/PCG/Source/PCG/Public/Data/Registry/PCGDataTypeIdentifier.h
+++ b/Engine/Plugins/PCG/Source/PCG/Public/Data/Registry/PCGDataTypeIdentifier.h
@@ -181,6 +181,24 @@ struct FPCGDataTypeIdentifier
 	PCG_API FString ToString() const;
 	PCG_API FText ToDisplayText() const;
 
+#if WITH_EDITOR
+	/**
+	 * Return a subtype tooltip associated with this type.
+	 * If it is a combination it will just return "ToDisplayText()".
+	 * If it is not a combination, it will query the Data Registry for the underlying FPCGDataTypeInfo, if it has anything.
+	 * If the FPCGDataTypeInfo doesn't provide a custom tooltip, it will just return "ToDisplayText()"
+	 */
+	PCG_API FText GetSubtypeTooltip() const;
+
+	/**
+	 * Return an extra tooltip associated with this type.
+	 * If it is a combination it will return an empty text.
+	 * If it is not a combination, it will query the Data Registry for the underlying FPCGDataTypeInfo, if it has anything.
+	 * If the FPCGDataTypeInfo doesn't provide a custom tooltip, it will just return empty.
+	 */
+	PCG_API FText GetExtraTooltip() const;
+#endif // WITH_EDITOR
+
 private:
 	PCG_API void ComposeFromLegacyType(EPCGDataType LegacyType);
 	PCG_API void ComposeFromPCGDataSubclass(TSubclassOf<UPCGData> Subclass);
diff --git a/Engine/Plugins/PCG/Source/PCG/Public/Data/Registry/PCGDataTypeRegistry.h b/Engine/Plugins/PCG/Source/PCG/Public/Data/Registry/PCGDataTypeRegistry.h
--- a/Engine/Plugins/PCG/Source/PCG/Public/Data/Registry/PCGDataTypeRegistry.h
+++ b/Engine/Plugins/PCG/Source/PCG/Public/Data/Registry/PCGDataTypeRegistry.h
@@ -3,6 +3,7 @@
 #pragma once
 
 #include "PCGCommon.h"
+#include "Data/Registry/PCGDataTypeCommon.h"
 #include "Data/Registry/PCGDataTypeIdentifier.h"
 
 #include "StructUtils/InstancedStruct.h"
@@ -19,32 +20,6 @@ namespace PCGDataTypeHierarchy
 	struct FTree;
 }
 
-enum class EPCGDataTypeCompatibilityResult
-{
-	// Types are not compatible, they can't be connected together.
-	NotCompatible = 0,
-
-	// Types are fully compatible, they can be connected as-is.
-	Compatible = 1,
-
-	// Types are compatible, but requires a filter to be added for the connection.
-	RequireFilter = 2,
-
-	// Types are compatible, but require conversion nodes to be added for the connection.
-	RequireConversion = 3,
-
-	// Error if the identifier is not found
-	UnknownType = 4
-};
-
-namespace PCGDataTypeCompatibilityResult
-{
-	inline bool IsValid(const EPCGDataTypeCompatibilityResult Result)
-	{
-		return Result != EPCGDataTypeCompatibilityResult::NotCompatible && Result != EPCGDataTypeCompatibilityResult::UnknownType;
-	}
-}
-
 /**
  * Registry to hold PCG data types and the hierarchy.
  */
diff --git a/Engine/Plugins/PCG/Source/PCG/Public/Elements/PCGReroute.h b/Engine/Plugins/PCG/Source/PCG/Public/Elements/PCGReroute.h
--- a/Engine/Plugins/PCG/Source/PCG/Public/Elements/PCGReroute.h
+++ b/Engine/Plugins/PCG/Source/PCG/Public/Elements/PCGReroute.h
@@ -28,6 +28,7 @@ public:
 
 	//~Begin UPCGSettings interface
 	virtual bool HasDynamicPins() const override { return true; }
+	virtual FPCGDataTypeIdentifier GetCurrentPinTypesID(const UPCGPin* InPin) const override;
 
 #if WITH_EDITOR
 	virtual FName GetDefaultNodeName() const override { return FName("Reroute"); }
diff --git a/Engine/Plugins/PCG/Source/PCG/Public/Metadata/Accessors/PCGAttributeAccessorHelpers.h b/Engine/Plugins/PCG/Source/PCG/Public/Metadata/Accessors/PCGAttributeAccessorHelpers.h
--- a/Engine/Plugins/PCG/Source/PCG/Public/Metadata/Accessors/PCGAttributeAccessorHelpers.h
+++ b/Engine/Plugins/PCG/Source/PCG/Public/Metadata/Accessors/PCGAttributeAccessorHelpers.h
@@ -35,6 +35,9 @@ namespace PCGAttributeAccessorHelpers
 	/** Returns true if the property chain exists,  and the last property is supported by PCG types and has a conversion to a Metadata type. */
 	PCG_API bool IsPropertyAccessorChainSupported(const TArray<FName>& InPropertyNames, const UStruct* InStruct);
 
+	/** Returns the metadata type associated with this property. */
+	PCG_API EPCGMetadataTypes GetMetadataTypeForProperty(const FProperty* InProperty);
+
 	/** Create an accessor for the given property. Property needs to be supported by PCG (cf. IsPropertyAccessorSupported). */
 	PCG_API TUniquePtr<IPCGAttributeAccessor> CreatePropertyAccessor(const FProperty* InProperty);
 	
diff --git a/Engine/Plugins/PCG/Source/PCG/Public/PCGParamData.h b/Engine/Plugins/PCG/Source/PCG/Public/PCGParamData.h
--- a/Engine/Plugins/PCG/Source/PCG/Public/PCGParamData.h
+++ b/Engine/Plugins/PCG/Source/PCG/Public/PCGParamData.h
@@ -21,7 +21,14 @@ struct FPCGDataTypeInfoParam : public FPCGDataTypeInfo
 
 	PCG_DECLARE_LEGACY_TYPE_INFO(EPCGDataType::Param)
 
+	PCG_API virtual EPCGDataTypeCompatibilityResult IsCompatibleForSubtype(const FPCGDataTypeIdentifier& InType, const FPCGDataTypeIdentifier& OutType, FText* OptionalOutCompatibilityMessage) const override;
 	PCG_API virtual bool SupportsConversionTo(const FPCGDataTypeIdentifier& ThisType, const FPCGDataTypeIdentifier& OutputType, TSubclassOf<UPCGSettings>* OptionalOutConversionSettings, FText* OptionalOutCompatibilityMessage) const override;
+
+#if WITH_EDITOR
+	PCG_API virtual TOptional<FText> GetSubtypeTooltip(const FPCGDataTypeIdentifier& ThisType) const override;
+#endif // WITH_EDITOR
+
+	PCG_API static bool HasValidSubtype(const FPCGDataTypeIdentifier& Id);
 };
 
 /**
diff --git a/Engine/Plugins/PCG/Source/PCG/Public/PCGSettings.h b/Engine/Plugins/PCG/Source/PCG/Public/PCGSettings.h
--- a/Engine/Plugins/PCG/Source/PCG/Public/PCGSettings.h
+++ b/Engine/Plugins/PCG/Source/PCG/Public/PCGSettings.h
@@ -123,6 +123,11 @@ struct FPCGSettingsOverridableParam
 	UPROPERTY()
 	bool bRequiresGPUReadback = false;
 
+#if WITH_EDITORONLY_DATA
+	UPROPERTY()
+	EPCGMetadataTypes UnderlyingType = EPCGMetadataTypes::Unknown;
+#endif // WITH_EDITORONLY_DATA
+
 	bool HasAliases() const { return !MapOfAliases.IsEmpty(); }
 
 	// Transient
diff --git a/Engine/Plugins/PCG/Source/PCGEditor/Private/Details/PCGDataTypeIdentifierDetails.cpp b/Engine/Plugins/PCG/Source/PCGEditor/Private/Details/PCGDataTypeIdentifierDetails.cpp
--- a/Engine/Plugins/PCG/Source/PCGEditor/Private/Details/PCGDataTypeIdentifierDetails.cpp
+++ b/Engine/Plugins/PCG/Source/PCGEditor/Private/Details/PCGDataTypeIdentifierDetails.cpp
@@ -128,7 +128,7 @@ TSharedRef<SWidget> FPCGDataTypeIdentifierDetails::OnGetMenuContent()
 		
 		MenuBuilder.AddMenuEntry(
 			FText::FromString(DataType.ToString()),
-			TypeInfo->GetTooltip(),
+			FText{},
 			FSlateIcon(),
 			FUIAction
 			(
diff --git a/Engine/Plugins/PCG/Source/PCGEditor/Private/Nodes/PCGEditorGraphNodeBase.cpp b/Engine/Plugins/PCG/Source/PCGEditor/Private/Nodes/PCGEditorGraphNodeBase.cpp
--- a/Engine/Plugins/PCG/Source/PCGEditor/Private/Nodes/PCGEditorGraphNodeBase.cpp
+++ b/Engine/Plugins/PCG/Source/PCGEditor/Private/Nodes/PCGEditorGraphNodeBase.cpp
@@ -1272,15 +1272,15 @@ void UPCGEditorGraphNodeBase::GetPinHoverText(const UEdGraphPin& Pin, FString& H
 {
 	const bool bIsInputPin = (Pin.Direction == EGPD_Input);
 	UPCGPin* MatchingPin = (PCGNode ? (bIsInputPin ? PCGNode->GetInputPin(Pin.PinName) : PCGNode->GetOutputPin(Pin.PinName)) : nullptr);
+	const FPCGDataTypeIdentifier Type = MatchingPin ? MatchingPin->GetCurrentTypesID() : FPCGDataTypeIdentifier{};
 
-	auto PinTypeToText = [](const FName& Category, UPCGPin* MatchingPin)
+	auto PinTypeToText = [&Type](const FName& Category)
 	{
 		if (Category != NAME_None)
 		{
 			return FText::FromName(Category);
 		}
-			
-		const FPCGDataTypeIdentifier Type = MatchingPin ? MatchingPin->GetCurrentTypesID() : FPCGDataTypeIdentifier{};
+
 		if (Type == EPCGDataType::Any)
 		{
 			return FText::FromName(FName("Any"));
@@ -1295,8 +1295,8 @@ void UPCGEditorGraphNodeBase::GetPinHoverText(const UEdGraphPin& Pin, FString& H
 		}
 	};
 
-	const FText DataTypeText = PinTypeToText(Pin.PinType.PinCategory, MatchingPin);
-	const FText DataSubtypeText = PinTypeToText(Pin.PinType.PinSubCategory, MatchingPin);
+	const FText DataTypeText = PinTypeToText(Pin.PinType.PinCategory);
+	const FText DataSubtypeText = Type.GetSubtypeTooltip();
 
 	FText Description;
 	if (MatchingPin)
@@ -1334,6 +1334,12 @@ void UPCGEditorGraphNodeBase::GetPinHoverText(const UEdGraphPin& Pin, FString& H
 		Required,
 		MultiDataSupport,
 		MultiConnectionSupport).ToString();
+
+	const FText ExtraTooltip = Type.GetExtraTooltip();
+	if (!ExtraTooltip.IsEmpty())
+	{
+		HoverTextOut = FText::Format(INVTEXT("{0}\n{1}"), FText::FromString(HoverTextOut), ExtraTooltip).ToString();
+	}
 }
 
 UObject* UPCGEditorGraphNodeBase::GetJumpTargetForDoubleClick() const
diff --git a/Engine/Plugins/PCG/Source/PCGEditor/Private/PCGEditorGraphSchema.cpp b/Engine/Plugins/PCG/Source/PCGEditor/Private/PCGEditorGraphSchema.cpp
--- a/Engine/Plugins/PCG/Source/PCGEditor/Private/PCGEditorGraphSchema.cpp
+++ b/Engine/Plugins/PCG/Source/PCGEditor/Private/PCGEditorGraphSchema.cpp
@@ -105,11 +105,6 @@ void UPCGEditorGraphSchema::GetGraphActions(FGraphActionMenuBuilder& ActionMenuB
 	}
 }
 
-FLinearColor UPCGEditorGraphSchema::GetPinTypeColor(const FEdGraphPinType& PinType) const
-{
-	return GetDefault<UPCGEditorSettings>()->GetPinColor(PinType);
-}
-
 FConnectionDrawingPolicy* UPCGEditorGraphSchema::CreateConnectionDrawingPolicy(int32 InBackLayerID, int32 InFrontLayerID, float InZoomFactor, const FSlateRect& InClippingRect, class FSlateWindowElementList& InDrawElements, class UEdGraph* InGraphObj) const
 {
 	return new FPCGEditorConnectionDrawingPolicy(InBackLayerID, InFrontLayerID, InZoomFactor, InClippingRect, InDrawElements, InGraphObj);
@@ -166,6 +161,7 @@ const FPinConnectionResponse UPCGEditorGraphSchema::CanCreateConnection(const UE
 		return FPinConnectionResponse(CONNECT_RESPONSE_MAKE_WITH_CONVERSION_NODE, CustomCompatiblityMessage.IsEmpty() ? LOCTEXT("ConnectionUsingFilter", "Filter data to match type") : CustomCompatiblityMessage);
 	case EPCGDataTypeCompatibilityResult::RequireConversion:
 		return FPinConnectionResponse(CONNECT_RESPONSE_MAKE_WITH_CONVERSION_NODE, CustomCompatiblityMessage.IsEmpty() ? LOCTEXT("ConnectionConversion", "Convert data to match type") : CustomCompatiblityMessage);
+	case EPCGDataTypeCompatibilityResult::TypeCompatibleSubtypeNotCompatible: // fall-through
 	case EPCGDataTypeCompatibilityResult::NotCompatible: // fall-through
 	case EPCGDataTypeCompatibilityResult::UnknownType:
 		return FPinConnectionResponse(CONNECT_RESPONSE_DISALLOW, CustomCompatiblityMessage.IsEmpty() ? LOCTEXT("ConnectionTypesIncompatible", "Pins are incompatible") : CustomCompatiblityMessage);
@@ -890,7 +886,6 @@ FLinearColor UPCGEditorGraphSchema::GetMetadataTypeColor(const EPCGMetadataTypes
 			return GraphSettings->SoftClassPinTypeColor;
 		default: // fall-through
 		case EPCGMetadataTypes::Unknown:
-			ensure(false);
 			return FLinearColor::White;
 	}
 }
@@ -975,6 +970,29 @@ FLinearColor UPCGEditorGraphSchema::GetPropertyBagTypeColor(const FPropertyBagPr
 	return FLinearColor::White;
 }
 
+FLinearColor UPCGEditorGraphSchema::GetPinColor(const UEdGraphPin* InPin) const
+{
+	if (!InPin)
+	{
+		return GetDefault<UPCGEditorSettings>()->DefaultPinColor;
+	}
+	
+	const UPCGEditorGraphNodeBase* EditorNode = CastChecked<const UPCGEditorGraphNodeBase>(InPin->GetOwningNode());
+	const UPCGNode* PCGNode = EditorNode ? EditorNode->GetPCGNode() : nullptr;
+	const UPCGPin* PCGPin = nullptr;
+	
+	if (InPin->Direction == EGPD_Input)
+	{
+		PCGPin = PCGNode ? PCGNode->GetInputPin(InPin->GetFName()) : nullptr;
+	}
+	else
+	{
+		PCGPin = PCGNode ? PCGNode->GetOutputPin(InPin->GetFName()) : nullptr;
+	}
+	
+	return PCGPin ? FPCGModule::GetConstDataTypeRegistry().GetPinColor(PCGPin->GetCurrentTypesID()) : GetDefault<UPCGEditorSettings>()->DefaultPinColor;
+}
+
 FPCGEditorConnectionDrawingPolicy::FPCGEditorConnectionDrawingPolicy(int32 InBackLayerID, int32 InFrontLayerID, float InZoomFactor, const FSlateRect& InClippingRect, FSlateWindowElementList& InDrawElements, UEdGraph* InGraph)
 	: FConnectionDrawingPolicy(InBackLayerID, InFrontLayerID, InZoomFactor, InClippingRect, InDrawElements)
 	, Graph(CastChecked<UPCGEditorGraph>(InGraph))
@@ -1054,9 +1072,12 @@ void FPCGEditorConnectionDrawingPolicy::DetermineWiringStyle(UEdGraphPin* Output
 	}
 
 	// Base the color of the wire on the color of the output pin
-	if (OutputPin)
+	if (OutputPin && Graph->GetSchema())
 	{
-		Params.WireColor = GetDefault<UPCGEditorSettings>()->GetPinColor(OutputPin->PinType);
+		if (const UPCGEditorGraphSchema* Schema = Cast<UPCGEditorGraphSchema>(Graph->GetSchema()))
+		{
+			Params.WireColor = Schema->GetPinColor(OutputPin);
+		}
 	}
 
 	// Desaturate connection if downstream node is disabled or if the data on this wire won't be used
diff --git a/Engine/Plugins/PCG/Source/PCGEditor/Private/PCGEditorModule.cpp b/Engine/Plugins/PCG/Source/PCGEditor/Private/PCGEditorModule.cpp
--- a/Engine/Plugins/PCG/Source/PCGEditor/Private/PCGEditorModule.cpp
+++ b/Engine/Plugins/PCG/Source/PCGEditor/Private/PCGEditorModule.cpp
@@ -55,6 +55,7 @@
 #include "Details/PCGVolumeDetails.h"
 #include "Details/PCGInteractiveToolSettingsDetails.h"
 #include "Nodes/PCGEditorGraphNodeFactory.h"
+#include "Schema/PCGEditorGraphSchema.h"
 #include "Subsystems/PCGEditorSubsystem.h"
 
 #include "ContentBrowserMenuContexts.h"
@@ -80,6 +81,15 @@
 
 #define LOCTEXT_NAMESPACE "FPCGEditorModule"
 
+namespace PCGEditorGraphSwitches
+{
+	TAutoConsoleVariable<bool> CVarEnableColorsOnParamPins{
+		TEXT("pcg.Editor.EnableColorsOnParamPins"),
+		true,
+		TEXT("Enable coloring of param pins based on their underlying types.")
+	};
+}
+
 namespace PCGEditorModule
 {
 	static const FName PCGBuildType(TEXT("PCG"));
@@ -846,8 +856,19 @@ void FPCGEditorModule::RegisterPinColorAndIcons()
 	FPCGDataTypeRegistry& InRegistry = FPCGModule::GetMutableDataTypeRegistry();
 
 	// Pin Colors
+	InRegistry.RegisterPinColorFunction(FPCGDataTypeInfoParam::AsId(), [](const FPCGDataTypeIdentifier& InType)
+	{
+		if (PCGEditorGraphSwitches::CVarEnableColorsOnParamPins.GetValueOnAnyThread())
+		{
+			return GetDefault<UPCGEditorGraphSchema>()->GetMetadataTypeColor(static_cast<EPCGMetadataTypes>(InType.CustomSubtype));
+		}
+		else
+		{
+			return GetDefault<UPCGEditorSettings>()->ParamDataPinColor;
+		}
+	});
+
 	InRegistry.RegisterPinColorFunction(FPCGDataTypeInfoPoint::AsId(), [](const FPCGDataTypeIdentifier&) { return GetDefault<UPCGEditorSettings>()->PointDataPinColor; });
-	InRegistry.RegisterPinColorFunction(FPCGDataTypeInfoParam::AsId(), [](const FPCGDataTypeIdentifier&) { return GetDefault<UPCGEditorSettings>()->ParamDataPinColor; });
 	InRegistry.RegisterPinColorFunction(FPCGDataTypeInfoConcrete::AsId(), [](const FPCGDataTypeIdentifier&) { return GetDefault<UPCGEditorSettings>()->ConcreteDataPinColor; });
 	InRegistry.RegisterPinColorFunction(FPCGDataTypeInfoPolyline::AsId(), [](const FPCGDataTypeIdentifier&) { return GetDefault<UPCGEditorSettings>()->PolyLineDataPinColor; });
 	InRegistry.RegisterPinColorFunction(FPCGDataTypeInfoLandscape::AsId(), [](const FPCGDataTypeIdentifier&) { return GetDefault<UPCGEditorSettings>()->LandscapeDataPinColor; });
@@ -941,8 +962,8 @@ void FPCGEditorModule::UnregisterPinColorAndIcons()
 	InRegistry.UnregisterPinColorFunction(FPCGDataTypeInfoLandscape::AsId());
 	InRegistry.UnregisterPinColorFunction(FPCGDataTypeInfoPolyline::AsId());
 	InRegistry.UnregisterPinColorFunction(FPCGDataTypeInfoConcrete::AsId());
-	InRegistry.UnregisterPinColorFunction(FPCGDataTypeInfoParam::AsId());
 	InRegistry.UnregisterPinColorFunction(FPCGDataTypeInfoPoint::AsId());
+	InRegistry.UnregisterPinColorFunction(FPCGDataTypeInfoParam::AsId());
 
 	// PinIcons. Just have catch all as most of it will be the same icons.
 	InRegistry.UnregisterPinIconsFunction(FPCGDataTypeInfo::AsId());
diff --git a/Engine/Plugins/PCG/Source/PCGEditor/Private/PCGEditorSettings.cpp b/Engine/Plugins/PCG/Source/PCGEditor/Private/PCGEditorSettings.cpp
--- a/Engine/Plugins/PCG/Source/PCGEditor/Private/PCGEditorSettings.cpp
+++ b/Engine/Plugins/PCG/Source/PCGEditor/Private/PCGEditorSettings.cpp
@@ -2,6 +2,7 @@
 
 #include "PCGEditorSettings.h"
 #include "PCGEditorCommon.h"
+#include "PCGModule.h"
 
 #include "EdGraph/EdGraphPin.h"
 #include "PCGSettings.h"
@@ -102,7 +103,11 @@ FLinearColor UPCGEditorSettings::GetColor(UPCGSettings* Settings) const
 		case EPCGSettingsType::PointOps:
 			return PointOpsNodeColor;
 		case EPCGSettingsType::GraphParameters:
-			return GraphParametersNodeColor;
+		{
+			// For graph parameters, use the output pin
+			const TArray<FPCGPinProperties> Properties = Settings->OutputPinProperties();
+			return !Properties.IsEmpty() ? FPCGModule::GetConstDataTypeRegistry().GetPinColor(Properties[0].AllowedTypes) : GraphParametersNodeColor;
+		}
 		case EPCGSettingsType::Reroute:
 			return RerouteNodeColor;
 		case EPCGSettingsType::DynamicMesh:
diff --git a/Engine/Plugins/PCG/Source/PCGEditor/Private/Pins/SPCGEditorGraphNodePin.cpp b/Engine/Plugins/PCG/Source/PCGEditor/Private/Pins/SPCGEditorGraphNodePin.cpp
--- a/Engine/Plugins/PCG/Source/PCGEditor/Private/Pins/SPCGEditorGraphNodePin.cpp
+++ b/Engine/Plugins/PCG/Source/PCGEditor/Private/Pins/SPCGEditorGraphNodePin.cpp
@@ -12,6 +12,7 @@
 #include "SLevelOfDetailBranchNode.h"
 #include "SPinTypeSelector.h"
 #include "Algo/AnyOf.h"
+#include "Schema/PCGEditorGraphSchema.h"
 #include "Widgets/SToolTip.h"
 #include "Widgets/Input/SButton.h"
 
@@ -309,17 +310,7 @@ void SPCGEditorGraphNodePin::ApplyUnusedPinStyle(FSlateColor& InOutColor) const
 
 FSlateColor SPCGEditorGraphNodePin::GetPinColor() const
 {
-	const UPCGPin* PCGPin = nullptr;
-	const UPCGNode* PCGNode = nullptr;
-
-	GetPCGNodeAndPin(PCGNode, PCGPin);
-
-	if (!PCGPin)
-	{
-		return SGraphPin::GetPinColor();
-	}
-	
-	FSlateColor Color{FPCGModule::GetConstDataTypeRegistry().GetPinColor(PCGPin->GetCurrentTypesID())};
+	FSlateColor Color = SGraphPin::GetPinColor();
 
 	ApplyUnusedPinStyle(Color);
 
diff --git a/Engine/Plugins/PCG/Source/PCGEditor/Private/Widgets/SPCGEditorGraphFind.cpp b/Engine/Plugins/PCG/Source/PCGEditor/Private/Widgets/SPCGEditorGraphFind.cpp
--- a/Engine/Plugins/PCG/Source/PCGEditor/Private/Widgets/SPCGEditorGraphFind.cpp
+++ b/Engine/Plugins/PCG/Source/PCGEditor/Private/Widgets/SPCGEditorGraphFind.cpp
@@ -174,7 +174,7 @@ TSharedRef<SWidget> FPCGEditorGraphFindResult::CreateIcon() const
 		// TODO get pin icon from nodebase?
 		Brush = FAppStyle::GetBrush(TEXT("GraphEditor.PinIcon"));
 		const UEdGraphSchema* Schema = ResolvedPin->GetSchema();
-		IconColor = Schema->GetPinTypeColor(ResolvedPin->PinType);
+		IconColor = Schema->GetPinColor(ResolvedPin);
 	}
 	else if (GraphNode.IsValid())
 	{
diff --git a/Engine/Plugins/PCG/Source/PCGEditor/Public/PCGEditorSettings.h b/Engine/Plugins/PCG/Source/PCGEditor/Public/PCGEditorSettings.h
--- a/Engine/Plugins/PCG/Source/PCGEditor/Public/PCGEditorSettings.h
+++ b/Engine/Plugins/PCG/Source/PCGEditor/Public/PCGEditorSettings.h
@@ -77,6 +77,8 @@ public:
 	PCGEDITOR_API UPCGEditorSettings(const FObjectInitializer& ObjectInitializer);
 
 	PCGEDITOR_API FLinearColor GetColor(UPCGSettings* InSettings) const;
+
+	UE_DEPRECATED(5.7, "Not used anymore, Pin Colors are queried from the PCG Data Type Registry.")
 	PCGEDITOR_API FLinearColor GetPinColor(const FEdGraphPinType& InPinType) const;
 
 	/** Default node color */
diff --git a/Engine/Plugins/PCG/Source/PCGEditor/Public/Schema/PCGEditorGraphSchema.h b/Engine/Plugins/PCG/Source/PCGEditor/Public/Schema/PCGEditorGraphSchema.h
--- a/Engine/Plugins/PCG/Source/PCGEditor/Public/Schema/PCGEditorGraphSchema.h
+++ b/Engine/Plugins/PCG/Source/PCGEditor/Public/Schema/PCGEditorGraphSchema.h
@@ -43,7 +43,7 @@ public:
 
 	//~ Begin EdGraphSchema Interface
 	PCGEDITOR_API virtual void GetGraphContextActions(FGraphContextMenuBuilder& ContextMenuBuilder) const override;
-	PCGEDITOR_API virtual FLinearColor GetPinTypeColor(const FEdGraphPinType& PinType) const override;
+	PCGEDITOR_API virtual FLinearColor GetPinColor(const UEdGraphPin* InPin) const override;
 	PCGEDITOR_API virtual const FPinConnectionResponse CanCreateConnection(const UEdGraphPin* A, const UEdGraphPin* B) const override;
 	PCGEDITOR_API virtual bool TryCreateConnection(UEdGraphPin* A, UEdGraphPin* B) const override;
 	PCGEDITOR_API virtual void BreakPinLinks(UEdGraphPin& TargetPin, bool bSendsNodeNotification) const override;
@@ -59,7 +59,6 @@ public:
 
 	PCGEDITOR_API virtual const FSlateBrush* GetPropertyBagTypeSlateBrush(EPropertyBagContainerType ContainerType) const;
 	PCGEDITOR_API virtual FLinearColor GetPropertyBagTypeColor(const FPropertyBagPropertyDesc& Desc) const;
-
 protected:
 	PCGEDITOR_API virtual void GetGraphActions(FGraphActionMenuBuilder& ActionMenuBuilder, const FPCGActionsFilter& ActionFilter, bool bIsContextual) const;
 	PCGEDITOR_API virtual void GetNativeElementActions(FGraphActionMenuBuilder& ActionMenuBuilder, const FPCGActionsFilter& ActionFilter) const;
