[PCG] Attachment for Spline Mesh when re-using resource

* Make sure that mobility spline mesh component and descriptor are the same
* Try to re-attach if there is a difference between mobility or if the resource is not attached.

#rnx

diff --git a/Engine/Plugins/PCG/Source/PCG/Private/Helpers/PCGActorHelpers.cpp b/Engine/Plugins/PCG/Source/PCG/Private/Helpers/PCGActorHelpers.cpp
--- a/Engine/Plugins/PCG/Source/PCG/Private/Helpers/PCGActorHelpers.cpp
+++ b/Engine/Plugins/PCG/Source/PCG/Private/Helpers/PCGActorHelpers.cpp
@@ -410,6 +410,28 @@ UPCGManagedSplineMeshComponent* UPCGActorHelpers::GetOrCreateManagedSplineMeshCo
 		return nullptr;
 	}
 
+	auto AttachToRoot = [InTargetActor, OptionalContext](USceneComponent* Component)
+	{
+		check(Component);
+		USceneComponent* RootComponent = InTargetActor->GetRootComponent();
+
+		// Implementation note: since the data passed to the params here is in world space,
+		// We need the transform on the spline mesh component to be the identity - in world space, because unlike the ISMs
+		// where we can set the instances and specify the data is in world space, we can't do that here.
+		FAttachmentTransformRules AttachmentRules(EAttachmentRule::KeepWorld, EAttachmentRule::KeepWorld, EAttachmentRule::KeepWorld, /*bInWeldSimulatedBodies=*/false);
+			
+		// If already attached, detach first.
+		if (Component->GetAttachParent())
+		{
+			Component->DetachFromComponent(FDetachmentTransformRules(AttachmentRules, /*bInCallModify=*/false));
+		}
+
+		if (!Component->AttachToComponent(RootComponent, AttachmentRules))
+		{
+			PCGLog::Component::LogComponentAttachmentFailedWarning(OptionalContext);
+		}
+	};
+
 	TRACE_CPUPROFILER_EVENT_SCOPE(UPCGActorHelpers::GetOrCreateManagedSplineMeshComponent);
 
 	{
@@ -456,6 +478,17 @@ UPCGManagedSplineMeshComponent* UPCGActorHelpers::GetOrCreateManagedSplineMeshCo
 			if (ensure(SplineMeshComponent))
 			{
 				SplineMeshComponent->Modify(!InSourceComponent->IsInPreviewMode());
+
+				const USceneComponent* RootComponent = InTargetActor->GetRootComponent();
+				
+				// If mobility differs between Root Component and Spline Mesh Component, or it is not attached, we need to re-attach it
+				if (SplineMeshComponent->GetAttachParent() != RootComponent 
+					|| (RootComponent && SplineMeshComponent->Mobility != RootComponent->Mobility)
+					|| (SplineMeshComponent->Mobility != InParams.Descriptor.Mobility))
+				{
+					SplineMeshComponent->Mobility = InParams.Descriptor.Mobility;
+					AttachToRoot(SplineMeshComponent);
+				}
 			}
 
 			return MatchingResource;
@@ -490,14 +523,8 @@ UPCGManagedSplineMeshComponent* UPCGActorHelpers::GetOrCreateManagedSplineMeshCo
 	SplineMeshComponent->RegisterComponent();
 	InTargetActor->AddInstanceComponent(SplineMeshComponent);
 
-	// Implementation note: since the data passed to the params here is in world space,
-	// We need the transform on the spline mesh component to be the identity - in world space, because unlike the ISMs
-	// where we can set the instances and specify the data is in world space, we can't do that here.
-	if (!SplineMeshComponent->AttachToComponent(InTargetActor->GetRootComponent(), FAttachmentTransformRules(EAttachmentRule::KeepWorld, EAttachmentRule::KeepWorld, EAttachmentRule::KeepWorld, false)))
-	{
-		PCGLog::Component::LogComponentAttachmentFailedWarning(OptionalContext);
-	}
-	
+	AttachToRoot(SplineMeshComponent);
+
 	SplineMeshComponent->ComponentTags.Add(InSourceComponent->GetFName());
 	SplineMeshComponent->ComponentTags.Add(PCGHelpers::DefaultPCGTag);
 
