[PCG] Trivial null check

#rnx

diff --git a/Engine/Plugins/PCG/Source/PCG/Private/Elements/PCGUserParameterGet.cpp b/Engine/Plugins/PCG/Source/PCG/Private/Elements/PCGUserParameterGet.cpp
--- a/Engine/Plugins/PCG/Source/PCG/Private/Elements/PCGUserParameterGet.cpp
+++ b/Engine/Plugins/PCG/Source/PCG/Private/Elements/PCGUserParameterGet.cpp
@@ -186,7 +186,8 @@ TArray<FPCGPinProperties> UPCGUserParameterGetSettings::OutputPinProperties() co
 		// Find the property associated with this param
 		const UPCGNode* Node = Cast<UPCGNode>(GetOuter());
 		const UPCGGraph* Graph = Node ? Node->GetGraph() : nullptr;
-		const FPropertyBagPropertyDesc* PropertyDesc = Graph ? Graph->GetUserParametersStruct()->FindPropertyDescByName(PropertyName) : nullptr;
+		const FInstancedPropertyBag* PropertyBag = Graph ? Graph->GetUserParametersStruct() : nullptr;
+		const FPropertyBagPropertyDesc* PropertyDesc = PropertyBag ? PropertyBag->FindPropertyDescByName(PropertyName) : nullptr;
 		if (PropertyDesc && PropertyDesc->CachedProperty)
 		{
 			// If the property is a container, check its inner
