[PCG] Prevent crash when no actor selected + added object modified notification so it can be picked up by other PCG components.
Also made the component reuse contingent on having being created by a tool (e.g. with a component tag).

#rnx

diff --git a/Engine/Plugins/PCG/Source/PCGEditor/Private/EditorMode/Tools/Helpers/PCGEdModeActorHelpers.cpp b/Engine/Plugins/PCG/Source/PCGEditor/Private/EditorMode/Tools/Helpers/PCGEdModeActorHelpers.cpp
--- a/Engine/Plugins/PCG/Source/PCGEditor/Private/EditorMode/Tools/Helpers/PCGEdModeActorHelpers.cpp
+++ b/Engine/Plugins/PCG/Source/PCGEditor/Private/EditorMode/Tools/Helpers/PCGEdModeActorHelpers.cpp
@@ -53,7 +53,8 @@ namespace PCG::EditorMode
 				return nullptr;
 			}
 
-			return Owner->FindComponentByClass<UPCGComponent>();
+			TArray<UActorComponent*> ToolComponents = Owner->GetComponentsByTag(UPCGComponent::StaticClass(), Tags::ToolGeneratedTag);
+			return ToolComponents.IsEmpty() ? nullptr : Cast<UPCGComponent>(ToolComponents[0]);
 		}
 
 		UPCGComponent* Create(AActor* Owner, const FName ComponentName)
@@ -64,6 +65,7 @@ namespace PCG::EditorMode
 			}
 
 			UPCGComponent* PCGComponent = NewObject<UPCGComponent>(Owner, ComponentName);
+			PCGComponent->ComponentTags.Add(Tags::ToolGeneratedTag);
 			PCGComponent->RegisterComponent();
 			Owner->AddInstanceComponent(PCGComponent);
 			return PCGComponent;
diff --git a/Engine/Plugins/PCG/Source/PCGEditor/Private/EditorMode/Tools/Line/PCGDrawSplineTool.cpp b/Engine/Plugins/PCG/Source/PCGEditor/Private/EditorMode/Tools/Line/PCGDrawSplineTool.cpp
--- a/Engine/Plugins/PCG/Source/PCGEditor/Private/EditorMode/Tools/Line/PCGDrawSplineTool.cpp
+++ b/Engine/Plugins/PCG/Source/PCGEditor/Private/EditorMode/Tools/Line/PCGDrawSplineTool.cpp
@@ -387,7 +387,7 @@ void UPCGDrawSplineTool::ApplyToolInternal() const
 				}
 
 				Actor = Settings->TargetActor.Get();
-				PCGComponent = PCGComponent::Find(Actor);
+				PCGComponent = Actor ? PCGComponent::Find(Actor) : nullptr;
 			}
 			break;
 		default:
@@ -395,31 +395,38 @@ void UPCGDrawSplineTool::ApplyToolInternal() const
 		}
 
 		// @todo_pcg: Will overwrite the current spline. Could be an option.
-		USplineComponent* NewSpline = Spline::FindComponent(Actor);
-		if (!NewSpline)
+		if (Actor)
 		{
-			NewSpline = Spline::CreateComponent(Actor, Settings->SplineName, ToolSplineComponent, /*bTransactional=*/true);
-			check(NewSpline);
-		}
-		else
-		{
-			Spline::CopySplinePoints(*ToolSplineComponent, *NewSpline, /*bTransactional=*/true);
-		}
-
-		if(Settings->TargetGraph)
-		{
-			if (!PCGComponent)
+			USplineComponent* Spline = Spline::FindComponent(Actor);
+			if (!Spline)
 			{
-				PCGComponent = PCGComponent::Create(Actor, Settings->NewPCGComponentName);
-				check(PCGComponent);
+				Spline = Spline::CreateComponent(Actor, Settings->SplineName, ToolSplineComponent, /*bTransactional=*/true);
+				check(Spline);
 			}
+			else
+			{
+				Spline::CopySplinePoints(*ToolSplineComponent, *Spline, /*bTransactional=*/true);
 
-			PCGComponent->SetGraph(Settings->TargetGraph);
-			PCGComponent->GetGraphInstance()->CopyParameterOverrides(WorkingComponent->GetGraphInstance());
+				// The copy points doesn't trigger a object changed event; we need to trigger one here manually so it gets picked up
+				FPropertyChangedEvent EmptyEvent{ nullptr };
+				FCoreUObjectDelegates::OnObjectPropertyChanged.Broadcast(Spline, EmptyEvent);
+			}
 
-			if(Settings->bAutoGenerate)
+			if(Settings->TargetGraph)
 			{
-				PCGComponent->Generate();
+				if (!PCGComponent)
+				{
+					PCGComponent = PCGComponent::Create(Actor, Settings->NewPCGComponentName);
+					check(PCGComponent);
+				}
+
+				PCGComponent->SetGraph(Settings->TargetGraph);
+				PCGComponent->GetGraphInstance()->CopyParameterOverrides(WorkingComponent->GetGraphInstance());
+
+				if(Settings->bAutoGenerate)
+				{
+					PCGComponent->Generate();
+				}
 			}
 		}
 
diff --git a/Engine/Plugins/PCG/Source/PCGEditor/Public/EditorMode/Tools/Helpers/PCGEdModeActorHelpers.h b/Engine/Plugins/PCG/Source/PCGEditor/Public/EditorMode/Tools/Helpers/PCGEdModeActorHelpers.h
--- a/Engine/Plugins/PCG/Source/PCGEditor/Public/EditorMode/Tools/Helpers/PCGEdModeActorHelpers.h
+++ b/Engine/Plugins/PCG/Source/PCGEditor/Public/EditorMode/Tools/Helpers/PCGEdModeActorHelpers.h
@@ -6,6 +6,11 @@
 
 namespace PCG::EditorMode
 {
+	namespace Tags
+	{
+		const FLazyName ToolGeneratedTag = "PCG Tool";
+	}
+
 	namespace Actor
 	{
 		/** Spawn an actor for the toolkit with spawn parameters. */
