[PCG] Decouple PCGComponent from GetExecutionContext node

#rnx

diff --git a/Engine/Plugins/PCG/Source/PCG/Private/Elements/PCGGetExecutionContext.cpp b/Engine/Plugins/PCG/Source/PCG/Private/Elements/PCGGetExecutionContext.cpp
--- a/Engine/Plugins/PCG/Source/PCG/Private/Elements/PCGGetExecutionContext.cpp
+++ b/Engine/Plugins/PCG/Source/PCG/Private/Elements/PCGGetExecutionContext.cpp
@@ -2,7 +2,6 @@
 
 #include "Elements/PCGGetExecutionContext.h"
 
-#include "PCGComponent.h"
 #include "PCGParamData.h"
 #include "PCGPin.h"
 #include "Helpers/PCGHelpers.h"
@@ -74,7 +73,6 @@ bool FPCGGetExecutionContextElement::ExecuteInternal(FPCGContext* Context) const
 	OutputData.Data = ParamData;
 
 	const IPCGGraphExecutionSource* ExecutionSource = Context->ExecutionSource.Get();
-	const UPCGComponent* SourceComponent = Cast<UPCGComponent>(ExecutionSource);
 
 	auto IsBoolOutput = [](EPCGGetExecutionContextMode Mode)
 	{
@@ -93,11 +91,11 @@ bool FPCGGetExecutionContextElement::ExecuteInternal(FPCGContext* Context) const
 		}
 		else if (Mode == EPCGGetExecutionContextMode::IsOriginal || Mode == EPCGGetExecutionContextMode::IsLocal)
 		{
-			Value = SourceComponent && (SourceComponent->IsLocalComponent() == (Mode == EPCGGetExecutionContextMode::IsLocal));
+			Value = ExecutionSource && (ExecutionSource->GetExecutionState().IsLocalSource() == (Mode == EPCGGetExecutionContextMode::IsLocal));
 		}
 		else if (Mode == EPCGGetExecutionContextMode::IsPartitioned)
 		{
-			Value = SourceComponent && SourceComponent->IsPartitioned();
+			Value = ExecutionSource && ExecutionSource->GetExecutionState().IsPartitioned();
 		}
 		else if (Mode == EPCGGetExecutionContextMode::IsRuntimeGeneration)
 		{
diff --git a/Engine/Plugins/PCG/Source/PCG/Private/PCGComponent.cpp b/Engine/Plugins/PCG/Source/PCG/Private/PCGComponent.cpp
--- a/Engine/Plugins/PCG/Source/PCG/Private/PCGComponent.cpp
+++ b/Engine/Plugins/PCG/Source/PCG/Private/PCGComponent.cpp
@@ -4344,6 +4344,18 @@ IPCGGraphExecutionSource* FPCGComponentExecutionState::GetLocalSource(const FPCG
 	return nullptr;
 }
 
+bool FPCGComponentExecutionState::IsPartitioned() const
+{
+	check(Component);
+	return Component->IsPartitioned();
+}
+
+bool FPCGComponentExecutionState::IsLocalSource() const
+{
+	check(Component);
+	return Component->IsLocalComponent();
+}
+
 FPCGGridDescriptor FPCGComponentExecutionState::GetGridDescriptor(uint32 InGridSize) const
 {
 	check(Component);
diff --git a/Engine/Plugins/PCG/Source/PCG/Public/PCGComponent.h b/Engine/Plugins/PCG/Source/PCG/Public/PCGComponent.h
--- a/Engine/Plugins/PCG/Source/PCG/Public/PCGComponent.h
+++ b/Engine/Plugins/PCG/Source/PCG/Public/PCGComponent.h
@@ -112,6 +112,8 @@ public:
 	virtual bool IsManagedByRuntimeGenSystem() const override;
 	virtual IPCGGraphExecutionSource* GetOriginalSource() const override;
 	virtual IPCGGraphExecutionSource* GetLocalSource(const FPCGGridDescriptor& GridDescriptor, const FIntVector& CellCoords) const override;
+	virtual bool IsPartitioned() const override;
+	virtual bool IsLocalSource() const override;
 	virtual FPCGGridDescriptor GetGridDescriptor(uint32 InGridSize) const override;
 	virtual int32 GetGenerationGridSize() const override;
 	virtual FPCGTaskId GetGenerationTaskId() const override;
diff --git a/Engine/Plugins/PCG/Source/PCG/Public/PCGGraphExecutionStateInterface.h b/Engine/Plugins/PCG/Source/PCG/Public/PCGGraphExecutionStateInterface.h
--- a/Engine/Plugins/PCG/Source/PCG/Public/PCGGraphExecutionStateInterface.h
+++ b/Engine/Plugins/PCG/Source/PCG/Public/PCGGraphExecutionStateInterface.h
@@ -117,6 +117,12 @@ public:
 	/** Retrieves a local execution source using grid descriptor and grid coordinates, returns nullptr if no such execution source is found. */
 	virtual IPCGGraphExecutionSource* GetLocalSource(const FPCGGridDescriptor& GridDescriptor, const FIntVector& CellCoords) const { return nullptr; }
 
+	/** Returns true if this execution source uses partitioned generation. */
+	virtual bool IsPartitioned() const { return false; }
+
+	/** Returns true if this execution source is a local source of a partitioned execution source. */
+	virtual bool IsLocalSource() const { return false; }
+
 	// @todo_pcg: The execution source should probably just have an execution domain instead of directly querying runtime gen vs etc.
 	/** Returns true if the execution source is managed by the runtime generation system. Nothing else should generate or cleanup this execution source. */
 	virtual bool IsManagedByRuntimeGenSystem() const { return false; }
