[PCG] Support standalone PCG Graph execution
- New property on UPCGGraphInterface: bIsStandaloneGraph
- Content Browser right-click "Generate Standalone Graph(s)" on selected graphs that have this property enabled
- PCG Graph Editor will show a "Standalone Graph" execution source in the debug panel to generate graph without having to create an instance
- Move more of the APIs to IPCGGraphExecutionSource from UPCGComponent (fixes profiling)
- EditorDefaultExecutionSource now supports being initialized with a PCG Graph Instance

diff --git a/Engine/Plugins/Experimental/ProceduralVegetationEditor/Source/ProceduralVegetationEditor/Private/PVEditor.cpp b/Engine/Plugins/Experimental/ProceduralVegetationEditor/Source/ProceduralVegetationEditor/Private/PVEditor.cpp
--- a/Engine/Plugins/Experimental/ProceduralVegetationEditor/Source/ProceduralVegetationEditor/Private/PVEditor.cpp
+++ b/Engine/Plugins/Experimental/ProceduralVegetationEditor/Source/ProceduralVegetationEditor/Private/PVEditor.cpp
@@ -58,7 +58,7 @@ void FPVEditor::Initialize(const EToolkitMode::Type InMode, const TSharedPtr<cla
 	
 	const FPCGEditorDefaultExecutionSourceParams Params
 	{
-		.Graph = ProceduralVegetationGraph
+		.GraphInterface = ProceduralVegetationGraph
 	};
 	
 	ExecutionSource = Subsystem->CreateEditorExecutionSource(Params);
diff --git a/Engine/Plugins/PCG/Source/PCG/Private/Editor/PCGEditorDefaultExecutionSource.cpp b/Engine/Plugins/PCG/Source/PCG/Private/Editor/PCGEditorDefaultExecutionSource.cpp
--- a/Engine/Plugins/PCG/Source/PCG/Private/Editor/PCGEditorDefaultExecutionSource.cpp
+++ b/Engine/Plugins/PCG/Source/PCG/Private/Editor/PCGEditorDefaultExecutionSource.cpp
@@ -46,12 +46,13 @@ FBox FPCGEditorDefaultExecutionState::GetBounds() const
 UPCGGraph* FPCGEditorDefaultExecutionState::GetGraph() const
 {
 	check(Source);
-	return Source->Graph;
+	return Source->GetGraph();
 }
 
 UPCGGraphInstance* FPCGEditorDefaultExecutionState::GetGraphInstance() const
 {
-	return nullptr;
+	check(Source);
+	return Source->GetGraphInstance();
 }
 
 void FPCGEditorDefaultExecutionState::OnGraphExecutionAborted(bool bQuiet, bool bCleanupUnusedResources)
@@ -100,34 +101,44 @@ FPCGGraphExecutionInspection& FPCGEditorDefaultExecutionState::GetInspection()
 }
 #endif // WITH_EDITOR
 
-void UPCGEditorDefaultExecutionSource::SetGraph(UPCGGraph* InGraph)
+void UPCGEditorDefaultExecutionSource::SetGraphInterface(UPCGGraphInterface* InGraphInterface)
 {
-	if (InGraph == Graph)
+	if (InGraphInterface == GraphInterface)
 	{
 		return;
 	}
 
 #if WITH_EDITOR
-	if (Graph)
+	if (GraphInterface)
 	{
-		Graph->OnGraphChangedDelegate.RemoveAll(this);
+		GraphInterface->OnGraphChangedDelegate.RemoveAll(this);
 	}
 #endif // WITH_EDITOR
 	
-	Graph = InGraph;
+	GraphInterface = InGraphInterface;
 
 #if WITH_EDITOR
-	if (Graph)
+	if (GraphInterface)
 	{
-		Graph->OnGraphChangedDelegate.AddUObject(this, &UPCGEditorDefaultExecutionSource::OnGraphChanged);
+		GraphInterface->OnGraphChangedDelegate.AddUObject(this, &UPCGEditorDefaultExecutionSource::OnGraphChanged);
 	}
 #endif // WITH_EDITOR
 }
 
+UPCGGraphInstance* UPCGEditorDefaultExecutionSource::GetGraphInstance()
+{
+	return Cast<UPCGGraphInstance>(GraphInterface);
+}
+
+UPCGGraph* UPCGEditorDefaultExecutionSource::GetGraph()
+{
+	return GraphInterface ? GraphInterface->GetMutablePCGGraph() : nullptr;
+}
+
 #if WITH_EDITOR
-void UPCGEditorDefaultExecutionSource::OnGraphChanged(UPCGGraphInterface* InGraph, EPCGChangeType ChangeType)
+void UPCGEditorDefaultExecutionSource::OnGraphChanged(UPCGGraphInterface* InGraphInterface, EPCGChangeType ChangeType)
 {
-	if (InGraph != Graph)
+	if (InGraphInterface != GraphInterface)
 	{
 		return;
 	}
@@ -140,7 +151,7 @@ void UPCGEditorDefaultExecutionSource::OnGraphChanged(UPCGGraphInterface* InGrap
 		return;
 	}
 
-	if (!InGraph || !InGraph->GetGraph())
+	if (!InGraphInterface || !InGraphInterface->GetGraph())
 	{
 		return;
 	}
@@ -177,7 +188,7 @@ void UPCGEditorDefaultExecutionSource::Generate()
 					if (WeakSubsystem.IsValid())
 					{
 #if WITH_EDITOR
-						WeakSubsystem->OnPCGSourceGenerationDone.Broadcast(WeakSubsystem.Get(), this, EPCGGenerationStatus::Completed);
+						WeakSubsystem->OnPCGSourceGenerationDone(this, EPCGGenerationStatus::Completed);
 #endif // WITH_EDITOR
 					}
 					
@@ -189,7 +200,7 @@ void UPCGEditorDefaultExecutionSource::Generate()
 					if (WeakSubsystem.IsValid())
 					{
 #if WITH_EDITOR
-						WeakSubsystem->OnPCGSourceGenerationDone.Broadcast(WeakSubsystem.Get(), this, EPCGGenerationStatus::Aborted);
+						WeakSubsystem->OnPCGSourceGenerationDone(this, EPCGGenerationStatus::Aborted);
 #endif // WITH_EDITOR
 					}
 					
@@ -207,10 +218,10 @@ void UPCGEditorDefaultExecutionSource::Generate()
 void UPCGEditorDefaultExecutionSource::Sunset()
 {
 #if WITH_EDITOR
-	if (IsValid(Graph))
+	if (IsValid(GraphInterface))
 	{
-		check(Graph);
-		Graph->OnGraphChangedDelegate.RemoveAll(this);
+		check(GraphInterface);
+		GraphInterface->OnGraphChangedDelegate.RemoveAll(this);
 	}
 #endif // WITH_EDITOR
 }
@@ -227,7 +238,7 @@ UPCGEditorDefaultExecutionSource::~UPCGEditorDefaultExecutionSource()
 
 void UPCGEditorDefaultExecutionSource::Initialize(const FPCGEditorDefaultExecutionSourceParams& InParams)
 {
-	SetGraph(InParams.Graph);
+	SetGraphInterface(InParams.GraphInterface);
 	SetSeed(InParams.Seed);
 }
 
diff --git a/Engine/Plugins/PCG/Source/PCG/Private/Elements/PCGWorldQuery.cpp b/Engine/Plugins/PCG/Source/PCG/Private/Elements/PCGWorldQuery.cpp
--- a/Engine/Plugins/PCG/Source/PCG/Private/Elements/PCGWorldQuery.cpp
+++ b/Engine/Plugins/PCG/Source/PCG/Private/Elements/PCGWorldQuery.cpp
@@ -177,7 +177,13 @@ bool FPCGWorldRayHitQueryElement::ExecuteInternal(FPCGContext* Context) const
 	// if we are partitioned or non-partitioned. But with 3D partitioning we might want something different.
 	// Since it is not yet widely used, we'll stick with same behavior for all cases.
 	UPCGComponent* SourceComponent = Cast<UPCGComponent>(Context->ExecutionSource.Get());
-	UPCGComponent* OriginalComponent = SourceComponent ? SourceComponent->GetOriginalComponent() : nullptr;
+	if (!SourceComponent)
+	{
+		PCGLog::LogErrorOnGraph(LOCTEXT("NoSourceComponent", "Execution source is not a PCG Component."), Context);
+		return true;
+	}
+
+	UPCGComponent* OriginalComponent = SourceComponent->GetOriginalComponent();
 	AActor* Owner = OriginalComponent ? OriginalComponent->GetOwner() : nullptr;
 	FTransform Transform = Owner ? Owner->GetTransform() : FTransform::Identity;
 	FBox LocalBounds = Owner ? PCGHelpers::GetActorLocalBounds(Owner) : FBox(EForceInit::ForceInit);
@@ -209,7 +215,6 @@ bool FPCGWorldRayHitQueryElement::ExecuteInternal(FPCGContext* Context) const
 		}
 	}
 
-	check(SourceComponent);
 	UWorld* World = SourceComponent->GetWorld();
 
 	UPCGWorldRayHitData* Data = FPCGContext::NewObject_AnyThread<UPCGWorldRayHitData>(Context);
diff --git a/Engine/Plugins/PCG/Source/PCG/Private/Elements/PCGWorldRaycast.cpp b/Engine/Plugins/PCG/Source/PCG/Private/Elements/PCGWorldRaycast.cpp
--- a/Engine/Plugins/PCG/Source/PCG/Private/Elements/PCGWorldRaycast.cpp
+++ b/Engine/Plugins/PCG/Source/PCG/Private/Elements/PCGWorldRaycast.cpp
@@ -402,7 +402,11 @@ bool FPCGWorldRaycastElement::ExecuteInternal(FPCGContext* InContext) const
 	}
 
 	UWorld* World = Context->ExecutionSource.IsValid() ? Context->ExecutionSource->GetExecutionState().GetWorld() : nullptr;
-	check(World);
+	if (!World)
+	{
+		PCGLog::LogErrorOnGraph(LOCTEXT("NoWorldFound", "Execution source does not belong to a world."), Context);
+		return true;
+	}
 
 	return ExecuteSlice(Context, [World](ContextType* Context, const ExecStateType& ExecState, IterStateType& IterState, const uint32 IterIndex)
 	{
diff --git a/Engine/Plugins/PCG/Source/PCG/Private/Graph/PCGGraphExecutor.cpp b/Engine/Plugins/PCG/Source/PCG/Private/Graph/PCGGraphExecutor.cpp
--- a/Engine/Plugins/PCG/Source/PCG/Private/Graph/PCGGraphExecutor.cpp
+++ b/Engine/Plugins/PCG/Source/PCG/Private/Graph/PCGGraphExecutor.cpp
@@ -2623,10 +2623,10 @@ void FPCGGraphExecutor::CullInactiveDownstreamNodes(FPCGTaskId InCompletedTaskId
 #if WITH_EDITOR
 void FPCGGraphExecutor::SendInactivePinNotification(const UPCGNode* InNode, const FPCGStack* InStack, uint64 InactiveOutputPinBitmask)
 {
-	const UPCGComponent* Component = InStack ? InStack->GetRootComponent() : nullptr;
-	if (Component && InNode)
+	const IPCGGraphExecutionSource* ExecutionSource = InStack ? InStack->GetRootSource() : nullptr;
+	if (ExecutionSource && InNode)
 	{
-		Component->NotifyNodeDynamicInactivePins(InNode, InStack, InactiveOutputPinBitmask);
+		ExecutionSource->GetExecutionState().GetInspection().NotifyNodeDynamicInactivePins(InNode, InStack, InactiveOutputPinBitmask);
 	}
 }
 #endif
diff --git a/Engine/Plugins/PCG/Source/PCG/Private/PCGComponent.cpp b/Engine/Plugins/PCG/Source/PCG/Private/PCGComponent.cpp
--- a/Engine/Plugins/PCG/Source/PCG/Private/PCGComponent.cpp
+++ b/Engine/Plugins/PCG/Source/PCG/Private/PCGComponent.cpp
@@ -864,7 +864,7 @@ void UPCGComponent::PostProcessGraph(const FBox& InNewBounds, bool bInGenerated,
 
 	if (Subsystem)
 	{
-		Subsystem->OnPCGSourceGenerationDone.Broadcast(Subsystem, IsValidChecked(this) ? this : nullptr, EPCGGenerationStatus::Completed);
+		Subsystem->OnPCGSourceGenerationDone(IsValidChecked(this) ? this : nullptr, EPCGGenerationStatus::Completed);
 	}
 #endif
 }
@@ -1007,7 +1007,7 @@ void UPCGComponent::OnProcessGraphAborted(bool bQuiet, bool bCleanupUnusedResour
 #if WITH_EDITOR
 	if (Subsystem)
 	{
-		Subsystem->OnPCGSourceGenerationDone.Broadcast(Subsystem, this, EPCGGenerationStatus::Aborted);
+		Subsystem->OnPCGSourceGenerationDone(this, EPCGGenerationStatus::Aborted);
 	}
 #endif
 
diff --git a/Engine/Plugins/PCG/Source/PCG/Private/PCGGraph.cpp b/Engine/Plugins/PCG/Source/PCG/Private/PCGGraph.cpp
--- a/Engine/Plugins/PCG/Source/PCG/Private/PCGGraph.cpp
+++ b/Engine/Plugins/PCG/Source/PCG/Private/PCGGraph.cpp
@@ -1600,8 +1600,6 @@ bool UPCGGraph::PrimeGraphCompilationCache()
 		return false;
 	}
 
-	const UPCGComponent* InspectedComponent = InspectedStack.GetRootComponent();
-
 	FPCGStackContext StackContext;
 	GraphCompiler->GetCompiledTasks(this, PCGHiGenGrid::UninitializedGridSize(), StackContext, /*bIsTopGraph=*/true);
 
@@ -1940,7 +1938,8 @@ void UPCGGraph::PostEditChangeChainProperty(FPropertyChangedChainEvent& Property
 	else if (PropertyName == GET_MEMBER_NAME_CHECKED(UPCGGraphInterface, Title)
 		|| PropertyName == GET_MEMBER_NAME_CHECKED(UPCGGraphInterface, bOverrideTitle)
 		|| PropertyName == GET_MEMBER_NAME_CHECKED(UPCGGraphInterface, Color)
-		|| PropertyName == GET_MEMBER_NAME_CHECKED(UPCGGraphInterface, bOverrideColor))
+		|| PropertyName == GET_MEMBER_NAME_CHECKED(UPCGGraphInterface, bOverrideColor)
+		|| PropertyName == GET_MEMBER_NAME_CHECKED(UPCGGraphInterface, bIsStandaloneGraph))
 	{
 		NotifyGraphChanged(EPCGChangeType::Cosmetic);
 	}
diff --git a/Engine/Plugins/PCG/Source/PCG/Private/Subsystems/IPCGBaseSubsystem.cpp b/Engine/Plugins/PCG/Source/PCG/Private/Subsystems/IPCGBaseSubsystem.cpp
--- a/Engine/Plugins/PCG/Source/PCG/Private/Subsystems/IPCGBaseSubsystem.cpp
+++ b/Engine/Plugins/PCG/Source/PCG/Private/Subsystems/IPCGBaseSubsystem.cpp
@@ -263,6 +263,10 @@ uint32 IPCGBaseSubsystem::GetGraphCacheEntryCount(IPCGElement* InElement) const
 	return GraphExecutor ? GraphExecutor->GetGraphCacheEntryCount(InElement) : 0;
 }
 
+void IPCGBaseSubsystem::OnPCGSourceGenerationDone(IPCGGraphExecutionSource* InExecutionSource, EPCGGenerationStatus InStatus)
+{
+	OnPCGSourceGenerationDoneDelegate.Broadcast(this, InExecutionSource, InStatus);
+}
 
 void IPCGBaseSubsystem::SetDisableClearResults(bool bInDisableClearResults)
 {
diff --git a/Engine/Plugins/PCG/Source/PCG/Private/Subsystems/PCGSubsystem.cpp b/Engine/Plugins/PCG/Source/PCG/Private/Subsystems/PCGSubsystem.cpp
--- a/Engine/Plugins/PCG/Source/PCG/Private/Subsystems/PCGSubsystem.cpp
+++ b/Engine/Plugins/PCG/Source/PCG/Private/Subsystems/PCGSubsystem.cpp
@@ -1771,7 +1771,7 @@ void UPCGSubsystem::ClearExecutionMetadata(const FPCGStack& BaseStack)
 {
 	if (IPCGEditorModule* PCGEditorModule = IPCGEditorModule::Get())
 	{
-		if (UPCGComponent* Component = const_cast<UPCGComponent*>(BaseStack.GetRootComponent()))
+		if (UPCGComponent* Component = const_cast<UPCGComponent*>(Cast<UPCGComponent>(BaseStack.GetRootSource())))
 		{
 			PCGEditorModule->ClearExecutionMetadata(Component);
 		}
diff --git a/Engine/Plugins/PCG/Source/PCG/Private/Utils/PCGNodeVisualLogs.cpp b/Engine/Plugins/PCG/Source/PCG/Private/Utils/PCGNodeVisualLogs.cpp
--- a/Engine/Plugins/PCG/Source/PCG/Private/Utils/PCGNodeVisualLogs.cpp
+++ b/Engine/Plugins/PCG/Source/PCG/Private/Utils/PCGNodeVisualLogs.cpp
@@ -160,10 +160,10 @@ void FPCGNodeVisualLogs::ForAllMatchingLogs(const FPCGStack& InPCGStack, TFuncti
 	}
 }
 
-void FPCGNodeVisualLogs::GetLogs(const UPCGNode* InNode, FPCGPerNodeVisualLogs& OutLogs, TArray<const UPCGComponent*>& OutComponents) const
+void FPCGNodeVisualLogs::GetLogsAndSources(const UPCGNode* InNode, FPCGPerNodeVisualLogs& OutLogs, TArray<const IPCGGraphExecutionSource*>& OutExecutionSources) const
 {
 	OutLogs.Reset();
-	OutComponents.Reset();
+	OutExecutionSources.Reset();
 
 	FReadScopeLock ScopedReadLock(LogsLock);
 	for (const TPair<FPCGStack, FPCGPerNodeVisualLogs>& Entry : StackToLogs)
@@ -172,18 +172,18 @@ void FPCGNodeVisualLogs::GetLogs(const UPCGNode* InNode, FPCGPerNodeVisualLogs&
 		if (Stack.HasObject(InNode))
 		{
 			OutLogs.Append(Entry.Value);
-			const UPCGComponent* RootComponent = Stack.GetRootComponent();
+			const IPCGGraphExecutionSource* RootSource = Stack.GetRootSource();
 			for (int LogIndex = 0; LogIndex < Entry.Value.Num(); ++LogIndex)
 			{
-				OutComponents.Add(RootComponent);
+				OutExecutionSources.Add(RootSource);
 			}
 		}
 	}
 }
 
-FText FPCGNodeVisualLogs::GetSummaryText(const FPCGPerNodeVisualLogs& InLogs, const TArray<const UPCGComponent*>* InComponents, ELogVerbosity::Type* OutMinimumVerbosity)
+FText FPCGNodeVisualLogs::GetSummaryTextWithSources(const FPCGPerNodeVisualLogs& InLogs, const TArray<const IPCGGraphExecutionSource*>* InExecutionSources, ELogVerbosity::Type* OutMinimumVerbosity)
 {
-	check(!InComponents || InLogs.Num() == InComponents->Num());
+	check(!InExecutionSources || InLogs.Num() == InExecutionSources->Num());
 
 	FText Summary = FText::GetEmpty();
 	ELogVerbosity::Type MinVerbosity = ELogVerbosity::All;
@@ -195,18 +195,33 @@ FText FPCGNodeVisualLogs::GetSummaryText(const FPCGPerNodeVisualLogs& InLogs, co
 		const FPCGNodeLogEntry& LogEntry = InLogs[LogIndex];
 
 		const FText& MessageVerbosity = (LogEntry.Verbosity == ELogVerbosity::Warning ? PCGNodeVisualLogsConstants::Warning : PCGNodeVisualLogsConstants::Error);
-		if (InComponents)
+		if (InExecutionSources)
 		{
-			const UPCGComponent* Component = (*InComponents)[LogIndex];
-			const FText ActorName = (Component && Component->GetOwner()) ? FText::FromString(Component->GetOwner()->GetActorLabel()) : LOCTEXT("PCGLogMissingComponent", "MissingComponent");
+			const IPCGGraphExecutionSource* ExecutionSource = (*InExecutionSources)[LogIndex];
+			const UPCGComponent* Component = Cast<UPCGComponent>(ExecutionSource);
+
+			FText SourceName;
+			// Backward compat
+			if (Component && Component->GetOwner())
+			{
+				SourceName = FText::FromString(Component->GetOwner()->GetActorLabel());
+			}
+			else if(ExecutionSource)
+			{ 
+				SourceName = FText::FromString(ExecutionSource->GetExecutionState().GetDebugName());
+			}
+			else
+			{
+				SourceName = LOCTEXT("PCGLogMissingSource", "MissingSource");
+			}
 
 			if (Summary.IsEmpty())
 			{
-				Summary = FText::Format(LOCTEXT("NodeTooltipLogWithActorEmpty", "[{0}] {1}: {2}"), ActorName, MessageVerbosity, LogEntry.Message);
+				Summary = FText::Format(LOCTEXT("NodeTooltipLogWithActorEmpty", "[{0}] {1}: {2}"), SourceName, MessageVerbosity, LogEntry.Message);
 			}
 			else
 			{
-				Summary = FText::Format(LOCTEXT("NodeTooltipLogWithActor", "{0}\n[{1}] {2}: {3}"), Summary, ActorName, MessageVerbosity, LogEntry.Message);
+				Summary = FText::Format(LOCTEXT("NodeTooltipLogWithActor", "{0}\n[{1}] {2}: {3}"), Summary, SourceName, MessageVerbosity, LogEntry.Message);
 			}
 		}
 		else
@@ -249,15 +264,15 @@ FText FPCGNodeVisualLogs::GetSummaryText(const FPCGPerNodeVisualLogs& InLogs, co
 FText FPCGNodeVisualLogs::GetLogsSummaryText(const UPCGNode* InNode, ELogVerbosity::Type& OutMinimumVerbosity) const
 {
 	FPCGPerNodeVisualLogs Logs;
-	TArray<const UPCGComponent*> Components;
-	GetLogs(InNode, Logs, Components);
+	TArray<const IPCGGraphExecutionSource*> ExecutionSources;
+	GetLogsAndSources(InNode, Logs, ExecutionSources);
 
-	return GetSummaryText(Logs, &Components, &OutMinimumVerbosity);
+	return GetSummaryTextWithSources(Logs, &ExecutionSources, &OutMinimumVerbosity);
 }
 
 FText FPCGNodeVisualLogs::GetLogsSummaryText(const FPCGStack& InBaseStack, ELogVerbosity::Type* OutMinimumVerbosity) const
 {
-	return GetSummaryText(GetLogs(InBaseStack), nullptr, OutMinimumVerbosity);
+	return GetSummaryTextWithSources(GetLogs(InBaseStack), nullptr, OutMinimumVerbosity);
 }
 
 void FPCGNodeVisualLogs::ClearLogs(const FPCGStack& InPCGStack)
@@ -271,14 +286,14 @@ void FPCGNodeVisualLogs::ClearLogs(const FPCGStack& InPCGStack)
 		TArray<FPCGStack> StacksToRemove;
 		for (const TPair<FPCGStack, FPCGPerNodeVisualLogs>& Entry : StackToLogs)
 		{
-			// Always take every opportunity to flush messages logged against invalid/dead components.
-			const bool bComponentValid = IsValid(Entry.Key.GetRootComponent());
+			// Always take every opportunity to flush messages logged against invalid/dead sources.
+			const bool bSourceValid = Entry.Key.GetRootSource() && IsValid(Cast<UObject>(Entry.Key.GetRootSource()));
 
-			if (!bComponentValid || Entry.Key.BeginsWith(InPCGStack))
+			if (!bSourceValid || Entry.Key.BeginsWith(InPCGStack))
 			{
-				if (!bComponentValid)
+				if (!bSourceValid)
 				{
-					UE_LOG(LogPCG, Verbose, TEXT("Cleared out logs for null component."));
+					UE_LOG(LogPCG, Verbose, TEXT("Cleared out logs for null source."));
 				}
 
 				StacksToRemove.Add(Entry.Key);
diff --git a/Engine/Plugins/PCG/Source/PCG/Public/Editor/PCGEditorDefaultExecutionSource.h b/Engine/Plugins/PCG/Source/PCG/Public/Editor/PCGEditorDefaultExecutionSource.h
--- a/Engine/Plugins/PCG/Source/PCG/Public/Editor/PCGEditorDefaultExecutionSource.h
+++ b/Engine/Plugins/PCG/Source/PCG/Public/Editor/PCGEditorDefaultExecutionSource.h
@@ -9,6 +9,7 @@
 
 class UPCGEditorDefaultExecutionSource;
 class UPCGGraph;
+class UPCGGraphInstance;
 class UPCGGraphInterface;
 
 class FPCGEditorDefaultExecutionState : public IPCGGraphExecutionState
@@ -53,7 +54,7 @@ private:
 
 struct FPCGEditorDefaultExecutionSourceParams
 {
-	UPCGGraph* Graph = nullptr;
+	UPCGGraphInterface* GraphInterface = nullptr;
 	int32 Seed = 42;
 };
 
@@ -75,13 +76,16 @@ public:
 	virtual IPCGGraphExecutionState& GetExecutionState() override { return State; }
 	virtual const IPCGGraphExecutionState& GetExecutionState() const override { return State; }
 
-	PCG_API void SetGraph(UPCGGraph* InGraph);
-	UPCGGraph* GetGraph() const { return Graph.Get(); }
+	PCG_API void SetGraphInterface(UPCGGraphInterface* InGraphInterface);
+
+	UPCGGraphInterface* GetGraphInterface() { return GraphInterface.Get(); }
+	PCG_API UPCGGraphInstance* GetGraphInstance();
+	PCG_API UPCGGraph* GetGraph();
 
 	PCG_API void SetSeed(int32 InSeed);
 
 #if WITH_EDITOR
-	PCG_API void OnGraphChanged(UPCGGraphInterface* Graph, EPCGChangeType ChangeType);
+	PCG_API void OnGraphChanged(UPCGGraphInterface* GraphInterface, EPCGChangeType ChangeType);
 #endif // WITH_EDITOR
 	
 	PCG_API void Generate();
@@ -92,7 +96,7 @@ private:
 	FPCGTaskId CurrentGenerationTask = InvalidPCGTaskId;
 	
 	UPROPERTY()
-	TObjectPtr<UPCGGraph> Graph;
+	TObjectPtr<UPCGGraphInterface> GraphInterface;
 	
 	FPCGEditorDefaultExecutionState State;
 	int32 Seed = 42;
diff --git a/Engine/Plugins/PCG/Source/PCG/Public/Graph/PCGStackContext.h b/Engine/Plugins/PCG/Source/PCG/Public/Graph/PCGStackContext.h
--- a/Engine/Plugins/PCG/Source/PCG/Public/Graph/PCGStackContext.h
+++ b/Engine/Plugins/PCG/Source/PCG/Public/Graph/PCGStackContext.h
@@ -153,6 +153,7 @@ public:
 	TArray<FPCGStackFrame>& GetStackFramesMutable() { return StackFrames; }
 
 	/** Component given by first stack frame. */
+	UE_DEPRECATED(5.7, "Use GetRootSource instead")
 	UE_API const UPCGComponent* GetRootComponent() const;
 
 	/** Execution source given by first stack frame. */
diff --git a/Engine/Plugins/PCG/Source/PCG/Public/PCGGraph.h b/Engine/Plugins/PCG/Source/PCG/Public/PCGGraph.h
--- a/Engine/Plugins/PCG/Source/PCG/Public/PCGGraph.h
+++ b/Engine/Plugins/PCG/Source/PCG/Public/PCGGraph.h
@@ -178,6 +178,10 @@ public:
 	/** Signifies the graph has compatibility with the annotated tags. Multiple tags may be comma separated. Ex. "PCGTool1, PCGTool2, PCGTool3" */
 	UPROPERTY(BlueprintReadWrite, EditAnywhere, Category = AssetInfo, AssetRegistrySearchable, meta=(EditCondition = "bExposeToLibrary", EditConditionHides, PCGNoHash))
 	FString CompatibleToolTags;
+
+	/** When enabled, this graph can be executed outside of the world using an editor execution source */
+	UPROPERTY(BlueprintReadWrite, EditAnywhere, Category = AssetInfo, AssetRegistrySearchable, meta = (PCGNoHash))
+	bool bIsStandaloneGraph = false;
 #endif
 
 	template <typename T>
diff --git a/Engine/Plugins/PCG/Source/PCG/Public/Subsystems/IPCGBaseSubsystem.h b/Engine/Plugins/PCG/Source/PCG/Public/Subsystems/IPCGBaseSubsystem.h
--- a/Engine/Plugins/PCG/Source/PCG/Public/Subsystems/IPCGBaseSubsystem.h
+++ b/Engine/Plugins/PCG/Source/PCG/Public/Subsystems/IPCGBaseSubsystem.h
@@ -112,7 +112,9 @@ public:
 	/** Returns how many times InElement is present in the cache. */
 	PCG_API uint32 GetGraphCacheEntryCount(IPCGElement* InElement) const;
 	
-	FPCGOnPCGSourceGenerationDone OnPCGSourceGenerationDone;
+	PCG_API virtual void OnPCGSourceGenerationDone(IPCGGraphExecutionSource* InExecutionSource, EPCGGenerationStatus InStatus);
+
+	FPCGOnPCGSourceGenerationDone& GetOnPCGSourceGenerationDone() { return OnPCGSourceGenerationDoneDelegate; }
 #endif // WITH_EDITOR
 	
 protected:
@@ -130,5 +132,7 @@ protected:
 
 #if WITH_EDITOR
 	PCG_API void SetDisableClearResults(bool bInDisableClearResults);
+
+	FPCGOnPCGSourceGenerationDone OnPCGSourceGenerationDoneDelegate;
 #endif
 };
diff --git a/Engine/Plugins/PCG/Source/PCG/Public/Utils/PCGNodeVisualLogs.h b/Engine/Plugins/PCG/Source/PCG/Public/Utils/PCGNodeVisualLogs.h
--- a/Engine/Plugins/PCG/Source/PCG/Public/Utils/PCGNodeVisualLogs.h
+++ b/Engine/Plugins/PCG/Source/PCG/Public/Utils/PCGNodeVisualLogs.h
@@ -51,8 +51,11 @@ public:
 	/** Returns all logs from a given stack */
 	UE_API FPCGPerNodeVisualLogs GetLogs(const FPCGStack& InPCGStack) const;
 
+	UE_DEPRECATED(5.7, "Use GetLogsAndSources instead")
+	UE_API void GetLogs(const UPCGNode* InNode, FPCGPerNodeVisualLogs& OutLogs, TArray<const UPCGComponent*>& OutComponents) const {}
+
 	/** Returns all logs for a given node */
-	UE_API void GetLogs(const UPCGNode* InNode, FPCGPerNodeVisualLogs& OutLogs, TArray<const UPCGComponent*>& OutComponents) const;
+	UE_API void GetLogsAndSources(const UPCGNode* InNode, FPCGPerNodeVisualLogs& OutLogs, TArray<const IPCGGraphExecutionSource*>& OutComponents) const;
 
 	/** Calls a function for all stacks that match in the logs or until the function returns false. */
 	UE_API void ForAllMatchingLogs(const FPCGStack& InPCGStack, TFunctionRef<bool(const FPCGStack&, const FPCGPerNodeVisualLogs&)> InFunc) const;
@@ -72,9 +75,11 @@ public:
 	/** Clear all errors and warnings corresponding to the given execution source. */
 	UE_API void ClearLogs(const IPCGGraphExecutionSource* InSource);
 
-	/** Generates summary text from log list, with optional component pointers to show actor label. */
-	static UE_API FText GetSummaryText(const FPCGPerNodeVisualLogs& InLogs, const TArray<const UPCGComponent*>* InComponents, ELogVerbosity::Type* OutMinimumVerbosity = nullptr);
+	UE_DEPRECATED(5.7, "Use GetSummaryTextWithSources instead")
+	static UE_API FText GetSummaryText(const FPCGPerNodeVisualLogs& InLogs, const TArray<const IPCGGraphExecutionSource*>* InComponents, ELogVerbosity::Type* OutMinimumVerbosity = nullptr) { return FText(); }
 
+	/** Generates summary text from log list, with optional source pointers to show actor label. */
+	static UE_API FText GetSummaryTextWithSources(const FPCGPerNodeVisualLogs& InLogs, const TArray<const IPCGGraphExecutionSource*>* InSources, ELogVerbosity::Type* OutMinimumVerbosity = nullptr);
 private:
 	TMap<FPCGStack, FPCGPerNodeVisualLogs> StackToLogs;
 	mutable FRWLock LogsLock;
diff --git a/Engine/Plugins/PCG/Source/PCGEditor/Private/AssetDefinitions/AssetDefinition_PCGGraphInterface.cpp b/Engine/Plugins/PCG/Source/PCGEditor/Private/AssetDefinitions/AssetDefinition_PCGGraphInterface.cpp
--- a/Engine/Plugins/PCG/Source/PCGEditor/Private/AssetDefinitions/AssetDefinition_PCGGraphInterface.cpp
+++ b/Engine/Plugins/PCG/Source/PCGEditor/Private/AssetDefinitions/AssetDefinition_PCGGraphInterface.cpp
@@ -7,8 +7,11 @@
 #include "PCGEditorCommon.h"
 #include "PCGEditorStyle.h"
 #include "PCGGraphFactory.h"
+#include "Editor/PCGEditorDefaultExecutionSource.h"
+#include "Subsystems/PCGEditorSubsystem.h"
 
 #include "ContentBrowserMenuContexts.h"
+#include "Editor.h"
 #include "IAssetTools.h"
 #include "ToolMenu.h"
 #include "ToolMenuEntry.h"
@@ -77,6 +80,54 @@ namespace MenuExtension_PCGGraphInterface
 					InSection.AddMenuEntry("PCGGraph_NewInstance", Label, ToolTip, FSlateIcon(FPCGEditorStyle::Get().GetStyleSetName(), "ClassIcon.PCGGraphInstance"), UIAction);
 				}
 			}));
+
+			Section.AddDynamicEntry(NAME_None, FNewToolMenuSectionDelegate::CreateLambda([](FToolMenuSection& InSection)
+			{
+				const UContentBrowserAssetContextMenuContext* CBContext = UContentBrowserAssetContextMenuContext::FindContextWithAssets(InSection);
+
+				const bool bStandaloneGraphSelected = CBContext->SelectedAssets.ContainsByPredicate([](const FAssetData& InAsset)
+				{
+					if (!InAsset.IsInstanceOf<UPCGGraphInterface>())
+					{
+						return false;
+					}
+
+					bool bTagValue = false;
+					return InAsset.GetTagValue(GET_MEMBER_NAME_CHECKED(UPCGGraphInterface, bIsStandaloneGraph), bTagValue) && bTagValue;
+				});
+
+				if (bStandaloneGraphSelected)
+				{
+					const TAttribute<FText> Label = LOCTEXT("GenerateStandaloneGraphsMenuLabel", "Generate Standalone Graph(s)");
+					const TAttribute<FText> ToolTip = LOCTEXT("GenerateStandaloneGraphsMenuToolTip", "Will generate every standalone graph in the selection.");
+					const FSlateIcon Icon = FSlateIcon(FPCGEditorStyle::Get().GetStyleSetName(), "ClassIcon.PCGGraph");
+
+					FToolUIAction UIAction;
+					UIAction.ExecuteAction.BindLambda([](const FToolMenuContext& InContext)
+					{
+						const UContentBrowserAssetContextMenuContext* CBContext = UContentBrowserAssetContextMenuContext::FindContextWithAssets(InContext);
+
+						if (UPCGEditorSubsystem* EditorSubsystem = GEditor->GetEditorSubsystem<UPCGEditorSubsystem>())
+						{
+							for (const FAssetData& AssetData : CBContext->SelectedAssets)
+							{
+								bool bTagValue = false;
+								const bool bIsStandaloneGraph = AssetData.IsInstanceOf<UPCGGraphInterface>() && AssetData.GetTagValue(GET_MEMBER_NAME_CHECKED(UPCGGraphInterface, bIsStandaloneGraph), bTagValue) && bTagValue;
+								if (!bIsStandaloneGraph)
+								{
+									continue;
+								}
+
+								if (UPCGGraphInterface* GraphInterface = Cast<UPCGGraphInterface>(AssetData.GetAsset()))
+								{
+									EditorSubsystem->GenerateGraph({ GraphInterface });
+								}
+							}
+						}
+					});
+					InSection.AddMenuEntry("GenerateStandaloneGraphsMenu", Label, ToolTip, Icon, UIAction);
+				}
+			}));
 		}));
 	});
 }
diff --git a/Engine/Plugins/PCG/Source/PCGEditor/Private/Nodes/PCGEditorGraphNodeBase.cpp b/Engine/Plugins/PCG/Source/PCGEditor/Private/Nodes/PCGEditorGraphNodeBase.cpp
--- a/Engine/Plugins/PCG/Source/PCGEditor/Private/Nodes/PCGEditorGraphNodeBase.cpp
+++ b/Engine/Plugins/PCG/Source/PCGEditor/Private/Nodes/PCGEditorGraphNodeBase.cpp
@@ -170,8 +170,8 @@ void UPCGEditorGraphNodeBase::Construct(UPCGNode* InPCGNode)
 	// Update to current graph/inspection state.
 	const UPCGEditorGraph* Graph = Cast<UPCGEditorGraph>(GetOuter());
 	const FPCGEditor* Editor = Graph ? Graph->GetEditor().Pin().Get() : nullptr;
-	UpdateStructuralVisualization(Editor ? Editor->GetPCGComponentBeingInspected() : nullptr, Editor ? Editor->GetStackBeingInspected() : nullptr, /*bNewlyPlaced=*/true);
-	UpdateGPUVisualization(Editor ? Editor->GetPCGComponentBeingInspected() : nullptr, Editor ? Editor->GetStackBeingInspected() : nullptr);
+	UpdateStructuralVisualization(Editor ? Editor->GetPCGSourceBeingInspected() : nullptr, Editor ? Editor->GetStackBeingInspected() : nullptr, /*bNewlyPlaced=*/true);
+	UpdateGPUVisualization(Editor ? Editor->GetPCGSourceBeingInspected() : nullptr, Editor ? Editor->GetStackBeingInspected() : nullptr);
 }
 
 void UPCGEditorGraphNodeBase::BeginDestroy()
@@ -524,7 +524,7 @@ void UPCGEditorGraphNodeBase::SetInspected(bool bInIsInspecting)
 			const FPCGEditor* Editor = EditorGraph ? EditorGraph->GetEditor().Pin().Get() : nullptr;
 
 			// Force refresh/regenerate.
-			if (UPCGComponent* Component = Editor ? Editor->GetPCGComponentBeingInspected() : nullptr)
+			if (UPCGComponent* Component = Editor ? Cast<UPCGComponent>(Editor->GetPCGSourceBeingInspected()) : nullptr)
 			{
 				if (Component->IsManagedByRuntimeGenSystem())
 				{
diff --git a/Engine/Plugins/PCG/Source/PCGEditor/Private/PCGEditor.cpp b/Engine/Plugins/PCG/Source/PCGEditor/Private/PCGEditor.cpp
--- a/Engine/Plugins/PCG/Source/PCGEditor/Private/PCGEditor.cpp
+++ b/Engine/Plugins/PCG/Source/PCGEditor/Private/PCGEditor.cpp
@@ -10,6 +10,7 @@
 #include "PCGPin.h"
 #include "PCGSubsystem.h"
 #include "Editor/IPCGEditorModule.h"
+#include "Editor/PCGEditorDefaultExecutionSource.h"
 #include "Elements/PCGReroute.h"
 #include "Helpers/PCGSubgraphHelpers.h"
 #include "Rendering/SlateRenderer.h"
@@ -153,6 +154,8 @@ void FPCGEditor::Initialize(const EToolkitMode::Type InMode, const TSharedPtr<cl
 {
 	PCGGraphBeingEdited = InPCGGraph;
 
+	UpdateDefaultExecutionSource();
+
 	// Initializes the UPCGEditorGraph if needed
 	const TSubclassOf<UPCGEditorGraphSchema> SchemaClass = GetSchemaClass();
 	const bool bShouldUpdateSchema = !InPCGGraph->PCGEditorGraph || InPCGGraph->PCGEditorGraph->Schema != SchemaClass;
@@ -230,7 +233,7 @@ void FPCGEditor::Initialize(const EToolkitMode::Type InMode, const TSharedPtr<cl
 
 		if (UPCGEditorSubsystem* EditorSubsystem = GEditor->GetEditorSubsystem<UPCGEditorSubsystem>())
 		{
-			EditorSubsystem->OnPCGSourceGenerationDone.AddRaw(this, &FPCGEditor::OnSourceGenerationDone);
+			EditorSubsystem->GetOnPCGSourceGenerationDone().AddRaw(this, &FPCGEditor::OnSourceGenerationDone);
 		}
 	}
 
@@ -265,17 +268,17 @@ void FPCGEditor::SetStackBeingInspected(const FPCGStack& FullStack)
 		return;
 	}
 
-	IPCGGraphExecutionSource* LastComponent = LastValidPCGSourceBeingInspected.Get();
-	IPCGGraphExecutionSource* NewComponent = const_cast<IPCGGraphExecutionSource*>(FullStack.GetRootSource());
+	IPCGGraphExecutionSource* LastSource = LastValidPCGSourceBeingInspected.Get();
+	IPCGGraphExecutionSource* NewSource = const_cast<IPCGGraphExecutionSource*>(FullStack.GetRootSource());
 
-	if (NewComponent && NewComponent != LastComponent)
+	if (NewSource && NewSource != LastSource)
 	{
-		if (LastComponent && LastComponent->GetExecutionState().GetInspection().IsInspecting())
+		if (LastSource && LastSource->GetExecutionState().GetInspection().IsInspecting())
 		{
-			LastComponent->GetExecutionState().GetInspection().DisableInspection();
+			LastSource->GetExecutionState().GetInspection().DisableInspection();
 		}
 
-		LastValidPCGSourceBeingInspected = NewComponent;
+		LastValidPCGSourceBeingInspected = NewSource;
 	}
 
 	if (PCGGraphBeingEdited)
@@ -288,7 +291,7 @@ void FPCGEditor::SetStackBeingInspected(const FPCGStack& FullStack)
 		PCGGraphBeingEdited->EnableInspection(FullStack);
 	}
 
-	PCGSourceBeingInspected = NewComponent;
+	PCGSourceBeingInspected = NewSource;
 
 	StackBeingInspected = FullStack;
 	OnInspectedStackChangedDelegate.Broadcast(StackBeingInspected);
@@ -370,6 +373,11 @@ void FPCGEditor::ClearStackBeingInspected()
 	}
 }
 
+UPCGComponent* FPCGEditor::GetPCGComponentBeingInspected() const
+{
+	return const_cast<UPCGComponent*>(Cast<UPCGComponent>(GetPCGSourceBeingInspected()));
+}
+
 void FPCGEditor::UpdateDebugAfterComponentSelection(UPCGComponent* InOldComponent, UPCGComponent* InNewComponent, bool bInNewComponentStartedInspecting)
 {
 	if (!ensure(PCGGraphBeingEdited))
@@ -667,6 +675,7 @@ void FPCGEditor::UnregisterTabSpawners(const TSharedRef<class FTabManager>& InTa
 void FPCGEditor::AddReferencedObjects(FReferenceCollector& Collector)
 {
 	Collector.AddReferencedObject(PCGGraphBeingEdited);
+	Collector.AddReferencedObject(PCGDefaultExecutionSource);
 
 	for (TSharedPtr<SPCGEditorGraphAttributeListView> ALV : AttributesWidgets)
 	{
@@ -2872,7 +2881,7 @@ void FPCGEditor::OnClose()
 	{
 		if (UPCGEditorSubsystem* EditorSubsystem = GEditor->GetEditorSubsystem<UPCGEditorSubsystem>())
 		{
-			EditorSubsystem->OnPCGSourceGenerationDone.RemoveAll(this);
+			EditorSubsystem->GetOnPCGSourceGenerationDone().RemoveAll(this);
 		}
 
 		UnregisterDelegatesForWorld(GEditor->GetEditorWorldContext().World());
@@ -3157,7 +3166,7 @@ void FPCGEditor::RegisterDelegatesForWorld(UWorld* World)
 	if (UPCGSubsystem* Subsystem = UPCGSubsystem::GetInstance(World))
 	{
 		Subsystem->OnPCGComponentUnregistered.AddRaw(this, &FPCGEditor::OnComponentUnregistered);
-		Subsystem->OnPCGSourceGenerationDone.AddRaw(this, &FPCGEditor::OnSourceGenerationDone);
+		Subsystem->GetOnPCGSourceGenerationDone().AddRaw(this, &FPCGEditor::OnSourceGenerationDone);
 	}
 }
 
@@ -3166,12 +3175,14 @@ void FPCGEditor::UnregisterDelegatesForWorld(UWorld* World)
 	if (UPCGSubsystem* Subsystem = UPCGSubsystem::GetInstance(World))
 	{
 		Subsystem->OnPCGComponentUnregistered.RemoveAll(this);
-		Subsystem->OnPCGSourceGenerationDone.RemoveAll(this);
+		Subsystem->GetOnPCGSourceGenerationDone().RemoveAll(this);
 	}
 }
 
 void FPCGEditor::OnGraphChanged(UPCGGraphInterface* InGraph, EPCGChangeType ChangeType)
 {
+	UpdateDefaultExecutionSource();
+
 	if (!!(ChangeType & EPCGChangeType::ShaderSource))
 	{
 		// Flush the shader file cache in case we are editing engine or data interface shaders.
@@ -3215,7 +3226,7 @@ void FPCGEditor::OnMapChanged(UWorld* InWorld, EMapChangeType InMapChangedType)
 {
 	if (InMapChangedType != EMapChangeType::SaveMap)
 	{
-		RefreshViewsOnLevelChange();
+		RefreshViews();
 
 		// Subsystem has been torn down and rebuilt.
 		if (GEditor)
@@ -3245,7 +3256,7 @@ void FPCGEditor::OnLevelActorDeleted(AActor* InActor)
 	}
 }
 
-void FPCGEditor::RefreshViewsOnLevelChange()
+void FPCGEditor::RefreshViews()
 {
 	if (DebugObjectTreeWidget.IsValid())
 	{
@@ -3261,6 +3272,28 @@ void FPCGEditor::RefreshViewsOnLevelChange()
 	}
 }
 
+void FPCGEditor::UpdateDefaultExecutionSource()
+{
+	if (PCGGraphBeingEdited->bIsStandaloneGraph && !PCGDefaultExecutionSource)
+	{
+		if (UPCGEditorSubsystem* EditorSubsystem = GEditor->GetEditorSubsystem<UPCGEditorSubsystem>())
+		{
+			PCGDefaultExecutionSource = EditorSubsystem->CreateEditorExecutionSource({ PCGGraphBeingEdited });
+			RefreshViews();
+		}
+	}
+	else if (!PCGGraphBeingEdited->bIsStandaloneGraph && PCGDefaultExecutionSource)
+	{
+		if (IPCGEditorModule* PCGEditorModule = IPCGEditorModule::Get())
+		{
+			PCGEditorModule->ClearExecutionMetadata(PCGDefaultExecutionSource);
+		}
+		PCGDefaultExecutionSource = nullptr;
+		CollectGarbage(GARBAGE_COLLECTION_KEEPFLAGS);
+		RefreshViews();
+	}
+}
+
 TSharedRef<FTabManager::FLayout> FPCGEditor::GetDefaultLayout() const
 {
 	return FTabManager::NewLayout("Standalone_PCGGraphEditor_DefaultLayout_v1.0")
diff --git a/Engine/Plugins/PCG/Source/PCGEditor/Private/Subsystems/PCGEditorSubsystem.cpp b/Engine/Plugins/PCG/Source/PCGEditor/Private/Subsystems/PCGEditorSubsystem.cpp
--- a/Engine/Plugins/PCG/Source/PCGEditor/Private/Subsystems/PCGEditorSubsystem.cpp
+++ b/Engine/Plugins/PCG/Source/PCGEditor/Private/Subsystems/PCGEditorSubsystem.cpp
@@ -23,7 +23,7 @@ void UPCGEditorSubsystem::Deinitialize()
 
 UPCGEditorDefaultExecutionSource* UPCGEditorSubsystem::CreateEditorExecutionSource(const FPCGEditorDefaultExecutionSourceParams& InParams)
 {
-	if (!InParams.Graph)
+	if (!InParams.GraphInterface)
 	{
 		return nullptr;
 	}
@@ -58,4 +58,31 @@ void UPCGEditorSubsystem::NotifyGraphChanged(UPCGGraph* InGraph, EPCGChangeType
 	IPCGBaseSubsystem::NotifyGraphChanged(InGraph, ChangeType);
 }
 
+void UPCGEditorSubsystem::OnPCGSourceGenerationDone(IPCGGraphExecutionSource* InExecutionSource, EPCGGenerationStatus InStatus)
+{
+	IPCGBaseSubsystem::OnPCGSourceGenerationDone(InExecutionSource, InStatus);
+
+	UPCGEditorDefaultExecutionSource* EditorExecutionSource = Cast<UPCGEditorDefaultExecutionSource>(InExecutionSource);
+	if (EditorExecutionSource && ExecutionSources.Remove(EditorExecutionSource) != 0)
+	{
+		if (IPCGEditorModule* PCGEditorModule = IPCGEditorModule::Get())
+		{
+			PCGEditorModule->ClearExecutionMetadata(EditorExecutionSource);
+		}
+
+		// Make sure our execution sources get GCed
+		CollectGarbage(GARBAGE_COLLECTION_KEEPFLAGS);
+	}
+}
+
+void UPCGEditorSubsystem::GenerateGraph(const FPCGEditorDefaultExecutionSourceParams& InParams)
+{
+	TStrongObjectPtr<UPCGEditorDefaultExecutionSource> ExecutionSource(CreateEditorExecutionSource(InParams));
+	if (ExecutionSource.IsValid())
+	{
+		ExecutionSources.Add(ExecutionSource.Get());
+		ExecutionSource->Generate();
+	}
+}
+
 #endif // WITH_EDITOR
diff --git a/Engine/Plugins/PCG/Source/PCGEditor/Private/Widgets/SPCGEditorGraphDebugObjectTree.cpp b/Engine/Plugins/PCG/Source/PCGEditor/Private/Widgets/SPCGEditorGraphDebugObjectTree.cpp
--- a/Engine/Plugins/PCG/Source/PCGEditor/Private/Widgets/SPCGEditorGraphDebugObjectTree.cpp
+++ b/Engine/Plugins/PCG/Source/PCGEditor/Private/Widgets/SPCGEditorGraphDebugObjectTree.cpp
@@ -6,6 +6,7 @@
 #include "PCGGraph.h"
 #include "PCGSubsystem.h"
 #include "Editor/IPCGEditorModule.h"
+#include "Editor/PCGEditorDefaultExecutionSource.h"
 #include "Grid/PCGPartitionActor.h"
 #include "Helpers/PCGHelpers.h"
 
@@ -127,7 +128,7 @@ namespace PCGEditorGraphDebugObjectTree
 
 		if (IPCGEditorModule* PCGEditorModule = IPCGEditorModule::Get())
 		{
-			OutIconTooltipText = PCGEditorModule->GetNodeVisualLogs().GetSummaryText(Logs, nullptr);
+			OutIconTooltipText = PCGEditorModule->GetNodeVisualLogs().GetSummaryTextWithSources(Logs, nullptr);
 		}
 		
 		OutIconVisibility = OutIconTooltipText.IsEmpty() ? EVisibility::Hidden : EVisibility::Visible;
@@ -251,6 +252,11 @@ const FSlateBrush* FPCGEditorGraphDebugObjectItem_Actor::GetIcon() const
 
 FString FPCGEditorGraphDebugObjectItem_PCGSource::GetLabel(bool bForSorting) const
 {
+	if (LabelOverride.IsSet())
+	{
+		return LabelOverride.GetValue();
+	}
+
 	//@todo_pcg: Should be part of the interface?
 	const UObject* Source = PCGSource.GetObject();
 	if (Source && PCGGraph.IsValid())
@@ -632,7 +638,7 @@ void SPCGEditorGraphDebugObjectTree::SetNodeBeingInspected(const UPCGNode* InPCG
 
 FReply SPCGEditorGraphDebugObjectTree::FocusOnDebugObject_OnClicked() const
 {
-	if (UPCGComponent* PCGComponent = PCGEditor.Pin()->GetPCGComponentBeingInspected())
+	if (UPCGComponent* PCGComponent = Cast<UPCGComponent>(PCGEditor.Pin()->GetPCGSourceBeingInspected()))
 	{
 		AActor* Actor = PCGComponent->GetOwner();
 		if (Actor && GEditor && GUnrealEd && GEditor->CanSelectActor(Actor, /*bInSelected=*/true))
@@ -649,7 +655,7 @@ FReply SPCGEditorGraphDebugObjectTree::FocusOnDebugObject_OnClicked() const
 
 bool SPCGEditorGraphDebugObjectTree::IsFocusOnDebugObjectButtonEnabled() const
 {
-	return PCGEditor.IsValid() && (PCGEditor.Pin()->GetPCGComponentBeingInspected() != nullptr);
+	return PCGEditor.IsValid() && (Cast<UPCGComponent>(PCGEditor.Pin()->GetPCGSourceBeingInspected()) != nullptr);
 }
 
 void SPCGEditorGraphDebugObjectTree::SetDebugObjectFromSelection_OnClicked()
@@ -1080,9 +1086,15 @@ void SPCGEditorGraphDebugObjectTree::AddStackToTree(const FPCGStack& Stack,
 					}
 					else
 					{
+						TSharedPtr<FPCGEditorGraphDebugObjectItem_PCGSource> Item = MakeShared<FPCGEditorGraphDebugObjectItem_PCGSource>(PCGSource, StackGraph, GraphStack, bIsDebuggable, bDisplayGrayedOut);
+						if (UObject* DefaultExecutionSource = PCGEditor.Pin()->GetDefaultExecutionSource(); PCGSource == Cast<IPCGGraphExecutionSource>(DefaultExecutionSource))
+						{
+							Item->SetLabel(TEXT("Standalone Graph"));
+						}
+
 						FPCGEditorGraphDebugObjectItemPtr TopGraphItem = InOutStackToItem.Emplace(
 							GraphStack,
-							MakeShared<FPCGEditorGraphDebugObjectItem_PCGSource>(PCGSource, StackGraph, GraphStack, bIsDebuggable, bDisplayGrayedOut));
+							MoveTemp(Item));
 
 						AllGraphItems.Add(TopGraphItem);
 
@@ -1168,6 +1180,11 @@ void SPCGEditorGraphDebugObjectTree::RefreshTree()
 	TArray<UObject*> PCGSources;
 	GetObjectsOfClass(UPCGComponent::StaticClass(), PCGSources, /*bIncludeDerivedClasses=*/ true);
 
+	if (UObject* DefaultExecutionSource = PCGEditor.Pin()->GetDefaultExecutionSource())
+	{
+		PCGSources.Add(DefaultExecutionSource);
+	}
+
 	if (!PCGSources.IsEmpty())
 	{
 		TMap<UObject*, TSharedPtr<FPCGEditorGraphDebugObjectItem>> OwnerItems;
@@ -1200,13 +1217,14 @@ void SPCGEditorGraphDebugObjectTree::RefreshTree()
 			}
 
 			const bool bIsPartitioned = PCGComponent && PCGComponent->IsPartitioned();
+			const uint32 GridSize = (bIsPartitioned || !PCGComponent) ? PCGHiGenGrid::UnboundedGridSize() : PCGComponent->GetGenerationGridSize();
 
 			// Process static stacks that can be read from the compiled graph.
 			{
 				//TRACE_CPUPROFILER_EVENT_SCOPE(SPCGEditorGraphDebugObjectTree::RefreshTree::StaticStacks);
 				// Considering it's likely for multiple components to have the same graph, we'll do a per-graph check first prior to building the component + graph stack.
 				UPCGGraph* ComponentGraph = PCGSource->GetExecutionState().GetGraph();
-				FGraphContextMapKey MapKey(ComponentGraph, bIsPartitioned ? PCGHiGenGrid::UnboundedGridSize() : PCGComponent->GetGenerationGridSize(), bIsPartitioned);
+				FGraphContextMapKey MapKey(ComponentGraph, GridSize, bIsPartitioned);
 				FGraphContextMapValue* MapValue = GraphToContexts.Find(MapKey);
 
 				if (!MapValue)
@@ -1224,10 +1242,10 @@ void SPCGEditorGraphDebugObjectTree::RefreshTree()
 				check(MapValue);
 				if (MapValue->Get<1>())
 				{
-					FPCGStack ComponentStack;
-					ComponentStack.PushFrame(PCGSourceObject);
-					FPCGStackContext ComponentContext(MapValue->Get<0>(), ComponentStack);
-					AddStacksToTree(ComponentContext.GetStacks(), OwnerItems, StackToItem);
+					FPCGStack SourceStack;
+					SourceStack.PushFrame(PCGSourceObject);
+					FPCGStackContext SourceContext(MapValue->Get<0>(), SourceStack);
+					AddStacksToTree(SourceContext.GetStacks(), OwnerItems, StackToItem);
 				}
 			}
 		}
diff --git a/Engine/Plugins/PCG/Source/PCGEditor/Private/Widgets/SPCGEditorGraphDebugObjectTree.h b/Engine/Plugins/PCG/Source/PCGEditor/Private/Widgets/SPCGEditorGraphDebugObjectTree.h
--- a/Engine/Plugins/PCG/Source/PCGEditor/Private/Widgets/SPCGEditorGraphDebugObjectTree.h
+++ b/Engine/Plugins/PCG/Source/PCGEditor/Private/Widgets/SPCGEditorGraphDebugObjectTree.h
@@ -115,6 +115,7 @@ public:
 	{}
 
 	virtual FString GetLabel(bool bForSorting = false) const override;
+	void SetLabel(const FString& InLabel) { LabelOverride = InLabel; }
 	virtual const FPCGStack* GetPCGStack() const override { return &PCGStack; }
 	virtual bool IsDebuggable() const override { return bIsDebuggable; }
 	virtual const UPCGGraph* GetPCGGraph() const override { return PCGGraph.Get(); }
@@ -127,6 +128,7 @@ protected:
 
 	FPCGStack PCGStack;
 	bool bIsDebuggable = false;
+	TOptional<FString> LabelOverride;
 };
 
 enum class EPCGEditorSubgraphNodeType : uint8
diff --git a/Engine/Plugins/PCG/Source/PCGEditor/Private/Widgets/SPCGEditorGraphLogView.cpp b/Engine/Plugins/PCG/Source/PCGEditor/Private/Widgets/SPCGEditorGraphLogView.cpp
--- a/Engine/Plugins/PCG/Source/PCGEditor/Private/Widgets/SPCGEditorGraphLogView.cpp
+++ b/Engine/Plugins/PCG/Source/PCGEditor/Private/Widgets/SPCGEditorGraphLogView.cpp
@@ -407,7 +407,7 @@ void SPCGEditorGraphLogView::CreateAndAddItem(const UPCGEditorGraphNode* InPCGEd
 
 IPCGGraphExecutionSource* SPCGEditorGraphLogView::GetExecutionSource() const
 {
-	return PCGEditorPtr.IsValid() ? PCGEditorPtr.Pin()->GetPCGComponentBeingInspected() : nullptr;
+	return PCGEditorPtr.IsValid() ? PCGEditorPtr.Pin()->GetPCGSourceBeingInspected() : nullptr;
 }
 
 void SPCGEditorGraphLogView::OnDebugStackChanged(const FPCGStack& InPCGStack)
diff --git a/Engine/Plugins/PCG/Source/PCGEditor/Private/Widgets/SPCGEditorGraphProfilingView.cpp b/Engine/Plugins/PCG/Source/PCGEditor/Private/Widgets/SPCGEditorGraphProfilingView.cpp
--- a/Engine/Plugins/PCG/Source/PCGEditor/Private/Widgets/SPCGEditorGraphProfilingView.cpp
+++ b/Engine/Plugins/PCG/Source/PCGEditor/Private/Widgets/SPCGEditorGraphProfilingView.cpp
@@ -801,7 +801,7 @@ void SPCGEditorGraphProfilingView::OnDebugStackChanged(const FPCGStack& InPCGSta
 
 IPCGGraphExecutionSource* SPCGEditorGraphProfilingView::GetExecutionSource() const
 {
-	return PCGEditorPtr.IsValid() ? PCGEditorPtr.Pin()->GetPCGComponentBeingInspected() : nullptr;
+	return PCGEditorPtr.IsValid() ? PCGEditorPtr.Pin()->GetPCGSourceBeingInspected() : nullptr;
 }
 
 TSharedRef<ITableRow> SPCGEditorGraphProfilingView::OnGenerateRow(PCGProfilingListViewItemPtr Item, const TSharedRef<STableViewBase>& OwnerTable) const
diff --git a/Engine/Plugins/PCG/Source/PCGEditor/Public/PCGEditor.h b/Engine/Plugins/PCG/Source/PCGEditor/Public/PCGEditor.h
--- a/Engine/Plugins/PCG/Source/PCGEditor/Public/PCGEditor.h
+++ b/Engine/Plugins/PCG/Source/PCGEditor/Public/PCGEditor.h
@@ -30,6 +30,7 @@ class SPCGEditorGraphUserParametersView;
 class SPCGEditorViewport;
 class UEdGraphNode;
 class UPCGComponent;
+class UPCGEditorDefaultExecutionSource;
 class UPCGEditorGraph;
 class UPCGEditorGraphNodeBase;
 class UPCGGraph;
@@ -91,7 +92,9 @@ public:
 	PCGEDITOR_API void ClearStackBeingInspected();
 
 	/** Gets the PCG source we are debugging */
-	UPCGComponent* GetPCGComponentBeingInspected() const { return const_cast<UPCGComponent*>(StackBeingInspected.GetRootComponent()); }
+	UE_DEPRECATED(5.7, "Use GetPCGSourceBeingInspected instead")
+	PCGEDITOR_API UPCGComponent* GetPCGComponentBeingInspected() const;
+
 	IPCGGraphExecutionSource* GetPCGSourceBeingInspected() const { return const_cast<IPCGGraphExecutionSource*>(StackBeingInspected.GetRootSource()); }
 	
 	/** Gets the PCG stack we are inspecting */
@@ -183,6 +186,8 @@ public:
 	// Can override the schema used for this editor. By default, it's UPCGEditorSchema.
 	PCGEDITOR_API virtual TSubclassOf<UPCGEditorGraphSchema> GetSchemaClass() const;
 
+	UPCGEditorDefaultExecutionSource* GetDefaultExecutionSource() const { return PCGDefaultExecutionSource; }
+
 protected:
 	PCGEDITOR_API virtual TAttribute<FGraphAppearanceInfo> GetAppearanceInfo() const;
 
@@ -441,7 +446,9 @@ private:
 	void OnEndPIE(bool bIsSimulating);
 	void OnLevelActorDeleted(AActor* InActor);
 
-	void RefreshViewsOnLevelChange();
+	void RefreshViews();
+
+	void UpdateDefaultExecutionSource();
 
 	TSharedRef<SDockTab> SpawnTab_GraphEditor(const FSpawnTabArgs& Args);
 	TSharedRef<SDockTab> SpawnTab_PropertyDetails(const FSpawnTabArgs& Args, int PropertyDetailsIndex);
@@ -478,6 +485,7 @@ private:
 	TSharedPtr<FUICommandList> GraphEditorCommands;
 
 	TObjectPtr<UPCGGraph> PCGGraphBeingEdited = nullptr;
+	TObjectPtr<UPCGEditorDefaultExecutionSource> PCGDefaultExecutionSource = nullptr;
 	UPCGEditorGraph* PCGEditorGraph = nullptr;
 
 	FPCGSoftGraphExecutionSource PCGSourceBeingInspected;
diff --git a/Engine/Plugins/PCG/Source/PCGEditor/Public/Subsystems/PCGEditorSubsystem.h b/Engine/Plugins/PCG/Source/PCGEditor/Public/Subsystems/PCGEditorSubsystem.h
--- a/Engine/Plugins/PCG/Source/PCGEditor/Public/Subsystems/PCGEditorSubsystem.h
+++ b/Engine/Plugins/PCG/Source/PCGEditor/Public/Subsystems/PCGEditorSubsystem.h
@@ -23,12 +23,18 @@ public:
 
 	PCGEDITOR_API UPCGEditorDefaultExecutionSource* CreateEditorExecutionSource(const FPCGEditorDefaultExecutionSourceParams& InParams);
 
+	PCGEDITOR_API void GenerateGraph(const FPCGEditorDefaultExecutionSourceParams& InParams);
+
 	PCGEDITOR_API virtual void NotifyGraphChanged(UPCGGraph* InGraph, EPCGChangeType ChangeType) override;
 
+	PCGEDITOR_API virtual void OnPCGSourceGenerationDone(IPCGGraphExecutionSource* InExecutionSource, EPCGGenerationStatus InStatus) override;
 protected:
 	//** FTickableEditorObject Interface
 	PCGEDITOR_API virtual void Tick(float DeltaTime) override;
 	PCGEDITOR_API virtual ETickableTickType GetTickableTickType() const override ;
 	PCGEDITOR_API virtual TStatId GetStatId() const override;
 	//~ END FTickableEditorObject Interface
+
+	UPROPERTY()
+	TArray<TObjectPtr<UPCGEditorDefaultExecutionSource>> ExecutionSources;
 };
\ No newline at end of file
