[PCG] Added hole support in the create polygon node (for point data cases) and configurable per-data polygon width.
Similar feature added to polygon offset.
Fixed a potential issue when getting a transform from the polygon if the segment is trivial.

#rnx

diff --git a/Engine/Plugins/PCG/Source/PCG/Private/Data/PCGPolygon2DData.cpp b/Engine/Plugins/PCG/Source/PCG/Private/Data/PCGPolygon2DData.cpp
--- a/Engine/Plugins/PCG/Source/PCG/Private/Data/PCGPolygon2DData.cpp
+++ b/Engine/Plugins/PCG/Source/PCG/Private/Data/PCGPolygon2DData.cpp
@@ -592,6 +592,11 @@ FVector UPCGPolygon2DData::GetLocationAtAlpha(float Alpha) const
 
 FTransform UPCGPolygon2DData::GetTransformAtAlpha(float InAlpha) const
 {
+	if (SegmentCount == 0)
+	{
+		return FTransform::Identity;
+	}
+
 	const double Alpha = FMath::Clamp(InAlpha, 0.0, 1.0 - UE_SMALL_NUMBER);
 
 	// Note: alpha in splines is a value from 0 to 1, with each segment having the same relative value.
diff --git a/Engine/Plugins/PCG/Source/PCG/Private/Elements/Polygon/PCGCreatePolygon2D.cpp b/Engine/Plugins/PCG/Source/PCG/Private/Elements/Polygon/PCGCreatePolygon2D.cpp
--- a/Engine/Plugins/PCG/Source/PCG/Private/Elements/Polygon/PCGCreatePolygon2D.cpp
+++ b/Engine/Plugins/PCG/Source/PCG/Private/Elements/Polygon/PCGCreatePolygon2D.cpp
@@ -8,17 +8,25 @@
 #include "Data/PCGPolygon2DData.h"
 #include "Utils/PCGLogErrors.h"
 #include "Helpers/PCGPolygon2DProcessingHelpers.h"
+#include "Metadata/Accessors/PCGAttributeAccessorHelpers.h"
 
 #include "CompGeom/PolygonTriangulation.h"
 #include "Curve/PolygonOffsetUtils.h"
 
 #include UE_INLINE_GENERATED_CPP_BY_NAME(PCGCreatePolygon2D)
 
+#define LOCTEXT_NAMESPACE "PCGCreatePolygonElement"
+
+UPCGCreatePolygon2DSettings::UPCGCreatePolygon2DSettings()
+{
+	HoleIndexAttribute.Update(StaticEnum<EPCGPolygon2DProperties>()->GetNameStringByValue(static_cast<int64>(EPCGPolygon2DProperties::HoleIndex)));
+}
+
 TArray<FPCGPinProperties> UPCGCreatePolygon2DSettings::InputPinProperties() const
 {
 	TArray<FPCGPinProperties> PinProperties;
 	// @todo_pcg: add landscape spline support?
-	PinProperties.Emplace(PCGPinConstants::DefaultInputLabel, EPCGDataType::Point | EPCGDataType::Spline);
+	PinProperties.Emplace_GetRef(PCGPinConstants::DefaultInputLabel, EPCGDataType::Point | EPCGDataType::Spline).SetRequiredPin();
 
 	return PinProperties;
 }
@@ -33,6 +41,7 @@ bool FPCGCreatePolygonElement::ExecuteInternal(FPCGContext* Context) const
 	TArray<FPCGTaggedData> Inputs = Context->InputData.GetInputsByPin(PCGPinConstants::DefaultInputLabel);
 
 	TArray<FVector> Positions;
+	TArray<int> HoleIndices;
 	TArray<PCGMetadataEntryKey> EntryKeys;
 
 	for (const FPCGTaggedData& Input : Inputs)
@@ -45,6 +54,7 @@ bool FPCGCreatePolygonElement::ExecuteInternal(FPCGContext* Context) const
 		bool bInputHasMetadata = false;
 
 		Positions.Reset();
+		HoleIndices.Reset();
 		EntryKeys.Reset();
 
 		if (const UPCGSplineData* SplineData = Cast<UPCGSplineData>(Input.Data))
@@ -95,6 +105,13 @@ bool FPCGCreatePolygonElement::ExecuteInternal(FPCGContext* Context) const
 						EntryKeys.Add(EK);
 					}
 				}
+
+				// Get hole index information if relevant
+				if (Settings->bUseHoleAttribute &&
+					!PCGAttributeAccessorHelpers::ExtractAllValues(Input.Data, Settings->HoleIndexAttribute, HoleIndices, Context))
+				{
+					continue;
+				}
 			}
 		}
 		else
@@ -104,14 +121,16 @@ bool FPCGCreatePolygonElement::ExecuteInternal(FPCGContext* Context) const
 			continue;
 		}
 
-		// @todo_pcg: Eventually, add metadata.
-		auto CreatePolygonData = [Context, SpatialData, &Input]()
+		double OpenPolygonWidth = Settings->OpenPolygonWidth;
+
+		if (Settings->bUsePolygonWidthAttribute)
 		{
-			UPCGPolygon2DData* Polygon = FPCGContext::NewObject_AnyThread<UPCGPolygon2DData>(Context);
-			Polygon->InitializeFromData(SpatialData);
-			Context->OutputData.TaggedData.Add_GetRef(Input).Data = Polygon;
-			return Polygon;
-		};
+			double PolyWidthFromData = 0.0;
+			if (PCGAttributeAccessorHelpers::ExtractParamValue(Input.Data, Settings->PolygonWidthAttribute, PolyWidthFromData, Context))
+			{
+				OpenPolygonWidth = PolyWidthFromData;
+			}
+		}
 
 		// Project the positions so they lie on a plane. Note that we'll assume at this point that the polygon is always closed, regardless of whether that's true or not, without loss of generality.
 		FVector PlaneNormal = FVector::UpVector;
@@ -138,83 +157,138 @@ bool FPCGCreatePolygonElement::ExecuteInternal(FPCGContext* Context) const
 
 		const bool bCreateOpenPolygon = (Settings->InputType == EPCGCreatePolygonInputType::ForceOpen || (bInputIsOpen && Settings->InputType == EPCGCreatePolygonInputType::Automatic));
 
-		if (bCreateOpenPolygon)
+		// Create starting polygon, which will be used as-is if this isn't an open polygon.
+		// 1. Split positions into outer/holes...
+		UE::Geometry::FGeneralPolygon2d SourcePolygon;
+
+		if (HoleIndices.IsEmpty())
 		{
 			// Special case: if we had a closed polygon but forced this here to be considered as an opened polygon, we'll still respect the loop.
-			if (bInputFromClosedSpline && !Positions2D.IsEmpty())
+			if (bCreateOpenPolygon && bInputFromClosedSpline && !Positions2D.IsEmpty())
 			{
 				const FVector2D FirstPosition2D = Positions2D[0];
 				Positions2D.Add(FirstPosition2D);
 			}
 
-			//Note: same values as FGeometryScriptOpenPathOffsetOptions, but could be parametrized in advanced settings
-			UE::Geometry::FOffsetPolygon2d PolygonOffset;
-			PolygonOffset.EndType = UE::Geometry::EPolygonOffsetEndType::Square;
-			PolygonOffset.JoinType = UE::Geometry::EPolygonOffsetJoinType::Square;
-			PolygonOffset.MiterLimit = 2.0;
-			PolygonOffset.MaxStepsPerRadian = 10.0;
-			PolygonOffset.DefaultStepsPerRadianScale = 1.0;
-			PolygonOffset.Offset = Settings->OpenPolygonWidth;
-			PolygonOffset.Polygons.Add(MoveTemp(Positions2D));
-
-			bool bOperationSuccessful = PolygonOffset.ComputeResult();
-			if (bOperationSuccessful)
+			UE::Geometry::FPolygon2d Polygon(Positions2D);
+			SourcePolygon.SetOuter(MoveTemp(Polygon));
+		}
+		else if (HoleIndices.Num() != Positions2D.Num())
+		{
+			PCGLog::LogErrorOnGraph(LOCTEXT("PositionsAndHoleAttributeMismatch", "There needs to be the same number of hole indices as there are points in the data."), Context);
+			continue;
+		}
+		else if(!Positions2D.IsEmpty())
+		{
+			// Some assumptions here:
+			// Same hole vertices should be continuous.
+			// Increasing hole values, assumed to be starting from -1.
+			int CurrentIndex = 0;
+			while (CurrentIndex != Positions2D.Num())
 			{
-				for (int PolyIndex = 0; PolyIndex < PolygonOffset.Result.Num(); ++PolyIndex)
+				UE::Geometry::FPolygon2d Polygon;
+				bool bIsOuter = (CurrentIndex == 0);
+
+				const int StartIndex = CurrentIndex;
+				while (CurrentIndex < Positions2D.Num() && HoleIndices[CurrentIndex] == HoleIndices[StartIndex])
 				{
-					UE::Geometry::FGeneralPolygon2d& Poly = PolygonOffset.Result[PolyIndex];
+					Polygon.AppendVertex(Positions2D[CurrentIndex++]);
+				}
 
-					if (Poly.OuterIsClockwise())
-					{
-						Poly.Reverse();
-					}
+				// Similarly to the closed spline case, when we're forcing open a representation that is supposed to be closed, we need to duplicate the start/end point
+				if (bCreateOpenPolygon)
+				{
+					Polygon.AppendVertex(Positions2D[StartIndex]);
+				}
 
-					UPCGPolygon2DData* Polygon = CreatePolygonData();
+				if (bIsOuter)
+				{
+					SourcePolygon.SetOuter(MoveTemp(Polygon));
+				}
+				else
+				{
+					SourcePolygon.AddHole(MoveTemp(Polygon));
+				}
+			}
+		}
 
-					if (bInputHasMetadata)
-					{
-						EntryKeys = PCGPolygon2DProcessingHelpers::ComputeEntryKeysByNearestSegmentSampling(SpatialData, Settings->OpenPolygonWidth, Poly, PlaneTransform, Polygon->MutableMetadata());
-					}
+		// Cleanup
+		if (SourcePolygon.OuterIsClockwise())
+		{
+			SourcePolygon.Reverse();
+		}
 
-					TConstArrayView<PCGMetadataEntryKey> EntryKeysView(EntryKeys);
+		TArray<UE::Geometry::FGeneralPolygon2d> OutPolygons;
 
-					Polygon->SetPolygon(MoveTemp(Poly), EntryKeysView.IsEmpty() ? nullptr : &EntryKeysView);
-					Polygon->SetTransform(PlaneTransform);
-				}
-			}
-			else
+		if (bCreateOpenPolygon)
+		{
+			TConstArrayView<UE::Geometry::FGeneralPolygon2d> PolyView = MakeConstArrayView(&SourcePolygon, 1);
+			TArray<UE::Geometry::FGeneralPolygon2d> OffsetPolygons;
+
+			//Note: same values as FGeometryScriptOpenPathOffsetOptions, but could be parametrized in advanced settings
+			const UE::Geometry::EPolygonOffsetEndType EndType = UE::Geometry::EPolygonOffsetEndType::Square;
+			const UE::Geometry::EPolygonOffsetJoinType JoinType = UE::Geometry::EPolygonOffsetJoinType::Square;
+			const double MiterLimit = 2.0;
+			const double MaxStepsPerRadian = 10.0;
+			const double DefaultStepsPerRadianScale = 1.0;
+			const double Offset = OpenPolygonWidth;
+
+			bool bOperationSuccessful = UE::Geometry::PolygonsOffset(
+				Offset,
+				PolyView,
+				OffsetPolygons,
+				/*bCopyOnFailure=*/false,
+				MiterLimit,
+				JoinType,
+				EndType,
+				MaxStepsPerRadian,
+				DefaultStepsPerRadianScale);
+
+			if (!bOperationSuccessful)
 			{
 				// The offset failed, not something we can do a lot about at this point.
 				PCGLog::LogWarningOnGraph(NSLOCTEXT("PCGCreatePolygon2DElement", "ErrorInPolygonOffset", "Unable to create extruded polygon from open polygon."), Context);
 				continue;
 			}
+			else
+			{
+				OutPolygons = MoveTemp(OffsetPolygons);
+			}
 		}
 		else
 		{
-			// Create the outer polygon
-			UE::Geometry::FPolygon2d OuterPolygon(Positions2D);
-			UE::Geometry::FGeneralPolygon2d NewPolygon(OuterPolygon);
+			// Move source polygon into out polygons
+			OutPolygons.Add(MoveTemp(SourcePolygon));
+		}
 
-			if (NewPolygon.OuterIsClockwise())
+		for (UE::Geometry::FGeneralPolygon2d& OutPoly : OutPolygons)
+		{
+			if (OutPoly.OuterIsClockwise())
 			{
-				NewPolygon.Reverse();
+				OutPoly.Reverse();
 			}
 
-			UPCGPolygon2DData* Polygon = CreatePolygonData();
+			UPCGPolygon2DData* Polygon = FPCGContext::NewObject_AnyThread<UPCGPolygon2DData>(Context);
+			Polygon->InitializeFromData(SpatialData);
+			Context->OutputData.TaggedData.Add_GetRef(Input).Data = Polygon;
 
 			// Once we've created the polygon object we can assign new metadata entries as needed.
-			if (bBuildEntryKeysBySampling)
+			if (bCreateOpenPolygon && bInputHasMetadata)
+			{
+				EntryKeys = PCGPolygon2DProcessingHelpers::ComputeEntryKeysByNearestSegmentSampling(SpatialData, OpenPolygonWidth, OutPoly, PlaneTransform, Polygon->MutableMetadata());
+			}
+			else if (bBuildEntryKeysBySampling)
 			{
 				EntryKeys = PCGPolygon2DProcessingHelpers::ComputeEntryKeysBySampling(SpatialData, Positions, Polygon->MutableMetadata());
 			}
 
 			TConstArrayView<PCGMetadataEntryKey> EntryKeysView(EntryKeys);
-			Polygon->SetPolygon(MoveTemp(NewPolygon), EntryKeysView.IsEmpty() ? nullptr : &EntryKeysView);
+			Polygon->SetPolygon(MoveTemp(OutPoly), EntryKeysView.IsEmpty() ? nullptr : &EntryKeysView);
 			Polygon->SetTransform(PlaneTransform);
-
-			// @todo_pcg: Create hole polygons / multi polygons...
 		}
 	}
 
 	return true;
 }
+
+#undef LOCTEXT_NAMESPACE
\ No newline at end of file
diff --git a/Engine/Plugins/PCG/Source/PCG/Private/Elements/Polygon/PCGGetSegment.cpp b/Engine/Plugins/PCG/Source/PCG/Private/Elements/Polygon/PCGGetSegment.cpp
--- a/Engine/Plugins/PCG/Source/PCG/Private/Elements/Polygon/PCGGetSegment.cpp
+++ b/Engine/Plugins/PCG/Source/PCG/Private/Elements/Polygon/PCGGetSegment.cpp
@@ -126,13 +126,10 @@ bool FPCGGetSegmentElement::ExecuteInternal(FPCGContext* Context) const
 			}
 			
 			// Retrieve segment value
-			TArray<int> SegmentIndexFromParam;
-			const FPCGAttributePropertyInputSelector SegmentSelector = Settings->SegmentIndexAttribute.CopyAndFixLast(ParamData);
-
-			if (PCGAttributeAccessorHelpers::ExtractAllValues(ParamData, SegmentSelector, SegmentIndexFromParam, Context) && 
-				!SegmentIndexFromParam.IsEmpty())
+			int SegmentIndexFromParam = -1;
+			if (PCGAttributeAccessorHelpers::ExtractParamValue(ParamData, Settings->SegmentIndexAttribute, SegmentIndexFromParam, Context))
 			{
-				SegmentIndex = SegmentIndexFromParam[0];
+				SegmentIndex = SegmentIndexFromParam;
 			}
 		}
 
@@ -147,13 +144,10 @@ bool FPCGGetSegmentElement::ExecuteInternal(FPCGContext* Context) const
 			}
 
 			// Retrieve hole value only the input supports holes, e.g. is a Polygon2D.
-			TArray<int> HoleIndexFromParam;
-			const FPCGAttributePropertyInputSelector HoleSelector = Settings->HoleIndexAttribute.CopyAndFixLast(ParamData);
-			if (Cast<UPCGPolygon2DData>(Input.Data) &&
-				PCGAttributeAccessorHelpers::ExtractAllValues(ParamData, HoleSelector, HoleIndexFromParam, Context) &&
-				!HoleIndexFromParam.IsEmpty())
+			int HoleIndexFromParam = -1;
+			if (Cast<UPCGPolygon2DData>(Input.Data) && PCGAttributeAccessorHelpers::ExtractParamValue(ParamData, Settings->HoleIndexAttribute, HoleIndexFromParam, Context))
 			{
-				HoleIndex = HoleIndexFromParam[0];
+				HoleIndex = HoleIndexFromParam;
 			}
 			else
 			{
diff --git a/Engine/Plugins/PCG/Source/PCG/Private/Elements/Polygon/PCGOffsetPolygon2D.cpp b/Engine/Plugins/PCG/Source/PCG/Private/Elements/Polygon/PCGOffsetPolygon2D.cpp
--- a/Engine/Plugins/PCG/Source/PCG/Private/Elements/Polygon/PCGOffsetPolygon2D.cpp
+++ b/Engine/Plugins/PCG/Source/PCG/Private/Elements/Polygon/PCGOffsetPolygon2D.cpp
@@ -5,11 +5,43 @@
 #include "PCGContext.h"
 #include "Data/PCGPolygon2DData.h"
 #include "Helpers/PCGPolygon2DProcessingHelpers.h"
+#include "Metadata/Accessors/PCGAttributeAccessorHelpers.h"
 
 #include "Curve/PolygonOffsetUtils.h"
 
 #include UE_INLINE_GENERATED_CPP_BY_NAME(PCGOffsetPolygon2D)
 
+namespace PCGOffsetPolygon2D
+{
+	const FLazyName OffsetPinLabel = "Offset";
+}
+
+#if WITH_EDITOR
+EPCGChangeType UPCGOffsetPolygon2DSettings::GetChangeTypeForProperty(FPropertyChangedEvent& PropertyChangedEvent) const
+{
+	EPCGChangeType ChangeType = Super::GetChangeTypeForProperty(PropertyChangedEvent);
+
+	if (PropertyChangedEvent.GetPropertyName() == GET_MEMBER_NAME_CHECKED(UPCGOffsetPolygon2DSettings, bUseOffsetFromInput))
+	{
+		ChangeType |= EPCGChangeType::Structural;
+	}
+
+	return ChangeType;
+}
+#endif // WITH_EDITOR
+
+TArray<FPCGPinProperties> UPCGOffsetPolygon2DSettings::InputPinProperties() const
+{
+	TArray<FPCGPinProperties> PinProperties = PCGPolygon2DUtils::DefaultPolygonInputPinProperties();
+
+	if (bUseOffsetFromInput)
+	{
+		PinProperties.Emplace_GetRef(PCGOffsetPolygon2D::OffsetPinLabel, EPCGDataType::Param | EPCGDataType::Polygon2D).SetRequiredPin();
+	}
+
+	return PinProperties;
+}
+
 bool FPCGOffsetPolygonElement::ExecuteInternal(FPCGContext* Context) const
 {
 	TRACE_CPUPROFILER_EVENT_SCOPE(FPCGOffsetPolygonElement::Execute);
@@ -20,8 +52,22 @@ bool FPCGOffsetPolygonElement::ExecuteInternal(FPCGContext* Context) const
 	TArray<FPCGTaggedData> Inputs = Context->InputData.GetInputsByPin(PCGPinConstants::DefaultInputLabel);
 	TArray<FPCGTaggedData>& Outputs = Context->OutputData.TaggedData;
 
-	for (const FPCGTaggedData& Input : Inputs)
+	TArray<FPCGTaggedData> OffsetInputs;
+	if (Settings->bUseOffsetFromInput)
+	{
+		OffsetInputs = Context->InputData.GetInputsByPin(PCGOffsetPolygon2D::OffsetPinLabel);
+
+		// Verify that cardinality matches
+		if (OffsetInputs.Num() != 1 && OffsetInputs.Num() != Inputs.Num())
+		{
+			PCGLog::InputOutput::LogInvalidCardinalityError(PCGPinConstants::DefaultInputLabel, PCGOffsetPolygon2D::OffsetPinLabel, Context);
+			return true;
+		}
+	}
+
+	for(int InputIndex = 0; InputIndex < Inputs.Num(); ++InputIndex)
 	{
+		const FPCGTaggedData& Input = Inputs[InputIndex];
 		const UPCGPolygon2DData* PolygonData = Cast<UPCGPolygon2DData>(Input.Data);
 
 		if (!PolygonData)
@@ -36,6 +82,18 @@ bool FPCGOffsetPolygonElement::ExecuteInternal(FPCGContext* Context) const
 
 		bool bOperationSuccess = false;
 
+		double Offset = Settings->Offset;
+
+		if (Settings->bUseOffsetFromInput)
+		{
+			const UPCGData* OffsetData = OffsetInputs[InputIndex % OffsetInputs.Num()].Data;
+			double OffsetFromData = 0.0;
+			if (PCGAttributeAccessorHelpers::ExtractParamValue(OffsetData, Settings->OffsetAttribute, OffsetFromData, Context))
+			{
+				Offset = OffsetFromData;
+			}
+		}
+
 		const double MiterLimit = 2.0;
 		const UE::Geometry::EPolygonOffsetJoinType JoinType = UE::Geometry::EPolygonOffsetJoinType::Square;
 		const UE::Geometry::EPolygonOffsetEndType EndType = UE::Geometry::EPolygonOffsetEndType::Polygon;
@@ -43,7 +101,7 @@ bool FPCGOffsetPolygonElement::ExecuteInternal(FPCGContext* Context) const
 		if (Settings->Operation == EPCGPolygonOffsetOperation::Offset)
 		{
 			bOperationSuccess = UE::Geometry::PolygonsOffset(
-				Settings->Offset,
+				Offset,
 				PolyView,
 				OutPolys,
 				/*bCopyOnFailure=*/false,
@@ -53,7 +111,7 @@ bool FPCGOffsetPolygonElement::ExecuteInternal(FPCGContext* Context) const
 		}
 		else
 		{
-			const double FirstOffset = (Settings->Operation == EPCGPolygonOffsetOperation::Open) ? -Settings->Offset : +Settings->Offset;
+			const double FirstOffset = (Settings->Operation == EPCGPolygonOffsetOperation::Open) ? -Offset : +Offset;
 			const double SecondOffset = -FirstOffset;
 
 			bOperationSuccess = UE::Geometry::PolygonsOffsets(
diff --git a/Engine/Plugins/PCG/Source/PCG/Public/Elements/Polygon/PCGCreatePolygon2D.h b/Engine/Plugins/PCG/Source/PCG/Public/Elements/Polygon/PCGCreatePolygon2D.h
--- a/Engine/Plugins/PCG/Source/PCG/Public/Elements/Polygon/PCGCreatePolygon2D.h
+++ b/Engine/Plugins/PCG/Source/PCG/Public/Elements/Polygon/PCGCreatePolygon2D.h
@@ -32,6 +32,8 @@ class UPCGCreatePolygon2DSettings : public UPCGSettings
 	GENERATED_BODY()
 
 public:
+	UPCGCreatePolygon2DSettings();
+
 	//~Begin UPCGSettings interface
 #if WITH_EDITOR
 	virtual FName GetDefaultNodeName() const override { return FName(TEXT("CreatePolygon2D")); }
@@ -47,14 +49,28 @@ protected:
 	//~End UPCGSettings interface
 
 public:
-	//@todo_pcg: add hole support & multi polygon
-
 	/** Controls the way input data will be considered - either as closed or open paths. */
 	UPROPERTY(EditAnywhere, BlueprintReadWrite, Category = Settings, meta = (PCG_Overridable))
 	EPCGCreatePolygonInputType InputType = EPCGCreatePolygonInputType::Automatic;
 
+	/** Controls whether we will use the hole attribute to separate outer from hole segments. Note that the attribute will be used only when the input is a point data. */
+	UPROPERTY(EditAnywhere, BlueprintReadWrite, Category = Settings, meta = (PCG_Overridable, InlineEditConditionToggle))
+	bool bUseHoleAttribute = false;
+
+	/** Attribute that describes the hole index of the points, where -1 is the outer polygon, and other values are for holes. Expects values on inputs to be sequential and increasing only. Note that this will be used only on point data. */
+	UPROPERTY(EditAnywhere, BlueprintReadWrite, Category = Settings, meta = (PCG_Overridable, EditCondition = "bUseHoleAttribute"))
+	FPCGAttributePropertyInputSelector HoleIndexAttribute;
+
+	/** Controls whether target polygon width (when open/forced open) will be driven by an attribute. */
+	UPROPERTY(EditAnywhere, BlueprintReadWrite, Category = Settings, meta = (PCG_Overridable, InlineEditConditionToggle))
+	bool bUsePolygonWidthAttribute = false;
+
+	/** Attribute from which to get the width value */
+	UPROPERTY(EditAnywhere, BlueprintReadWrite, Category = Settings, meta = (PCG_Overridable, EditCondition = "bUsePolygonWidthAttribute"))
+	FPCGAttributePropertyInputSelector PolygonWidthAttribute;
+
 	/** Target width of the polygon generated from an open path. */
-	UPROPERTY(EditAnywhere, BlueprintReadWrite, Category = Settings, meta = (EditCondition = "InputType != EPCGCreatePolygonInputType::ForceClosed", PCG_Overridable))
+	UPROPERTY(EditAnywhere, BlueprintReadWrite, Category = Settings, meta = (EditCondition = "!bUsePolygonWidthAttribute && InputType != EPCGCreatePolygonInputType::ForceClosed", PCG_Overridable))
 	double OpenPolygonWidth = 100.0;
 
 	/** Maximum squared distance before we need to subdivide a segment again as part of the spline discretization to a polygon. */
diff --git a/Engine/Plugins/PCG/Source/PCG/Public/Elements/Polygon/PCGOffsetPolygon2D.h b/Engine/Plugins/PCG/Source/PCG/Public/Elements/Polygon/PCGOffsetPolygon2D.h
--- a/Engine/Plugins/PCG/Source/PCG/Public/Elements/Polygon/PCGOffsetPolygon2D.h
+++ b/Engine/Plugins/PCG/Source/PCG/Public/Elements/Polygon/PCGOffsetPolygon2D.h
@@ -19,7 +19,7 @@ class FPCGOffsetPolygonElement : public IPCGElement
 {
 public:
 	virtual bool ExecuteInternal(FPCGContext* InContext) const override;
-	virtual EPCGElementExecutionLoopMode ExecutionLoopMode(const UPCGSettings* Settings) const override { return EPCGElementExecutionLoopMode::SinglePrimaryPin; }
+	virtual EPCGElementExecutionLoopMode ExecutionLoopMode(const UPCGSettings* Settings) const override { return EPCGElementExecutionLoopMode::PrimaryPinAndBroadcastablePins; }
 };
 
 /**
@@ -38,10 +38,11 @@ public:
 	virtual FText GetNodeTooltipText() const override { return NSLOCTEXT("PCGOffsetPolygon2D", "NodeTooltip", "Offsets polygon to either make it larger or smaller, or open/close holes based on the offset quantity."); }
 	virtual EPCGSettingsType GetType() const override { return EPCGSettingsType::Spatial; }
 	// @todo_pcg : presets for offset/open/close
+	virtual EPCGChangeType GetChangeTypeForProperty(FPropertyChangedEvent& PropertyChangedEvent) const override;
 #endif
 
 protected:
-	virtual TArray<FPCGPinProperties> InputPinProperties() const override { return PCGPolygon2DUtils::DefaultPolygonInputPinProperties(); }
+	virtual TArray<FPCGPinProperties> InputPinProperties() const override;
 	virtual TArray<FPCGPinProperties> OutputPinProperties() const override { return PCGPolygon2DUtils::DefaultPolygonOutputPinProperties(); }
 	virtual FPCGElementPtr CreateElement() const override { return MakeShared<FPCGOffsetPolygonElement>(); }
 	//~End UPCGSettings interface
@@ -53,8 +54,12 @@ public:
 	UPROPERTY(BlueprintReadWrite, EditAnywhere, Category = Settings)
 	bool bInheritMetadata = true;
 
-	UPROPERTY(BlueprintReadWrite, EditAnywhere, Category = Settings, meta = (PCG_Overridable))
+	UPROPERTY(BlueprintReadWrite, EditAnywhere, Category = Settings)
+	bool bUseOffsetFromInput = false;
+
+	UPROPERTY(BlueprintReadWrite, EditAnywhere, Category = Settings, meta = (EditCondition = "bUseOffsetFromInput", EditConditionHides, PCG_Overridable))
+	FPCGAttributePropertyInputSelector OffsetAttribute;
+
+	UPROPERTY(BlueprintReadWrite, EditAnywhere, Category = Settings, meta = (EditCondition = "!bUseOffsetFromInput", EditConditionHides, PCG_Overridable))
 	double Offset = 100.0;
-	
-	// Note: uses UE::Geometry::PolygonsOffset, and geometry script does the same thing.
 };
\ No newline at end of file
diff --git a/Engine/Plugins/PCG/Source/PCG/Public/Metadata/Accessors/PCGAttributeAccessorHelpers.h b/Engine/Plugins/PCG/Source/PCG/Public/Metadata/Accessors/PCGAttributeAccessorHelpers.h
--- a/Engine/Plugins/PCG/Source/PCG/Public/Metadata/Accessors/PCGAttributeAccessorHelpers.h
+++ b/Engine/Plugins/PCG/Source/PCG/Public/Metadata/Accessors/PCGAttributeAccessorHelpers.h
@@ -197,6 +197,32 @@ namespace PCGAttributeAccessorHelpers
 
 		return bResult;
 	}
+
+	template <typename T>
+	bool ExtractParamValue(const UPCGData* InData, const FPCGAttributePropertyInputSelector& InSelector, T& OutValue, FPCGContext* Context = nullptr, EPCGAttributeAccessorFlags GetFlags = EPCGAttributeAccessorFlags::AllowBroadcastAndConstructible, bool bQuiet = false)
+	{
+		const FPCGAttributePropertyInputSelector Selector = InSelector.CopyAndFixLast(InData);
+
+		TUniquePtr<const IPCGAttributeAccessor> Accessor = CreateConstAccessor(InData, Selector, bQuiet);
+		if (!Accessor.IsValid())
+		{
+			if (!bQuiet)
+			{
+				PCGLog::Metadata::LogFailToCreateAccessorError(Selector, Context);
+			}
+			return false;
+		}
+
+		FPCGAttributeAccessorKeysEntries FirstEntry(PCGFirstEntryKey);
+		bool bResult = Accessor->Get<T>(OutValue, FirstEntry, GetFlags);
+
+		if (!bResult && !bQuiet)
+		{
+			PCGLog::Metadata::LogFailToGetAttributeError(Selector, Context);
+		}
+
+		return bResult;
+	}
 	
 	/** 
 	 * Utility function to create the accessor and keys and write all the values for all the keys coming from InValues. Also, automatically resolve the selector for @Last.
