[PCG] Fixed issue where the metadata interpolation would fail for the last segment of polygons causing wrong values or checks.

#rnx

diff --git a/Engine/Plugins/PCG/Source/PCG/Private/Data/PCGPolygon2DData.cpp b/Engine/Plugins/PCG/Source/PCG/Private/Data/PCGPolygon2DData.cpp
--- a/Engine/Plugins/PCG/Source/PCG/Private/Data/PCGPolygon2DData.cpp
+++ b/Engine/Plugins/PCG/Source/PCG/Private/Data/PCGPolygon2DData.cpp
@@ -24,6 +24,11 @@ PCG_DEFINE_TYPE_INFO(FPCGDataTypeInfoPolygon2D, UPCGPolygon2DData)
 
 namespace PCGPolygon2DData
 {
+	namespace Constants
+	{
+		const float KeyEpsilon = 1.0e-6f;
+	}
+
 	// @todo_pcg : could suggest to move this to the polygon implementation
 	void SerializePolygon(UE::Geometry::FGeneralPolygon2d& Polygon, FArchive& Archive)
 	{
@@ -355,7 +360,7 @@ bool UPCGPolygon2DData::SamplePoint(const FTransform& InTransform, const FBox& I
 		const double SegmentRatio = Polygon.Segment(SegmentIndex, HoleIndex).ConvertToUnitRange(DistanceParameter);
 
 		const int GlobalSegmentIndex = SegmentAndHoleIndicesToSegmentIndex[TPair<int, int>(SegmentIndex, HoleIndex)];
-		const double NearestKey = GlobalSegmentIndex + SegmentRatio;
+		const float NearestKey = GlobalSegmentIndex + FMath::Min(SegmentRatio, 1.0f - PCGPolygon2DData::Constants::KeyEpsilon);
 
 		WriteMetadataToPoint(NearestKey, OutPoint, OutMetadata);
 	}
@@ -405,7 +410,7 @@ bool UPCGPolygon2DData::ProjectPoint(const FTransform& InTransform, const FBox&
 		const double SegmentRatio = Polygon.Segment(SegmentIndex, HoleIndex).ConvertToUnitRange(DistanceParameter);
 
 		const int GlobalSegmentIndex = SegmentAndHoleIndicesToSegmentIndex[TPair<int, int>(SegmentIndex, HoleIndex)];
-		const double NearestKey = GlobalSegmentIndex + SegmentRatio;
+		const float NearestKey = GlobalSegmentIndex + FMath::Min(SegmentRatio, 1.0f - PCGPolygon2DData::Constants::KeyEpsilon);
 
 		WriteMetadataToPoint(NearestKey, OutPoint, OutMetadata);
 	}
@@ -434,18 +439,9 @@ void UPCGPolygon2DData::WriteMetadataToEntry(float InputKey, PCGMetadataEntryKey
 		return;
 	}
 
-	const double InterpolationRatio = FMath::Frac(InputKey);
+	const float InterpolationRatio = FMath::Frac(InputKey);
 	const int CurrentVertexIndex = FMath::FloorToInt32(InputKey);
 
-	check(VertexEntryKeys.IsValidIndex(CurrentVertexIndex));
-	OutMetadataDomain->InitializeOnSet(OutEntryKey, VertexEntryKeys[CurrentVertexIndex], MetadataDomain);
-
-	// Implementation note: there's no need to compute property interpolation if we're directly on a vertex.
-	if (InterpolationRatio <= UE_DOUBLE_SMALL_NUMBER)
-	{
-		return;
-	}
-
 	int NextVertexIndex = INDEX_NONE;
 	int CurrentSegmentIndex = INDEX_NONE;
 	int CurrentHoleIndex = INDEX_NONE;
@@ -461,6 +457,16 @@ void UPCGPolygon2DData::WriteMetadataToEntry(float InputKey, PCGMetadataEntryKey
 		return;
 	}
 
+	// If this fails, it should have failed the pair search, hence the check here.
+	check(VertexEntryKeys.IsValidIndex(CurrentVertexIndex));
+
+	// Implementation note: there's no need to compute property interpolation if we're directly on the first vertex of the segment.
+	if (InterpolationRatio <= PCGPolygon2DData::Constants::KeyEpsilon)
+	{
+		OutMetadataDomain->InitializeOnSet(OutEntryKey, VertexEntryKeys[CurrentVertexIndex], MetadataDomain);
+		return;
+	}
+
 	if (const int* TentativeNextVertex = SegmentAndHoleIndicesToSegmentIndex.Find(TPair<int, int>(CurrentSegmentIndex + 1, CurrentHoleIndex)))
 	{
 		NextVertexIndex = *TentativeNextVertex;
@@ -472,13 +478,25 @@ void UPCGPolygon2DData::WriteMetadataToEntry(float InputKey, PCGMetadataEntryKey
 	else
 	{
 		// Invalid index
+		OutMetadataDomain->InitializeOnSet(OutEntryKey, VertexEntryKeys[CurrentVertexIndex], MetadataDomain);
 		return;
 	}
 
+	// If this fails, it should have failed the pair search, hence the check here.
 	check(VertexEntryKeys.IsValidIndex(NextVertexIndex));
 
+	// Implementation note: there's no need to compute property interpolation if we're directly on the second vertex of the segment.
+	if (InterpolationRatio >= 1.0f - PCGPolygon2DData::Constants::KeyEpsilon)
+	{
+		OutMetadataDomain->InitializeOnSet(OutEntryKey, VertexEntryKeys[NextVertexIndex], MetadataDomain);
+		return;
+	}
+
+	// Finally, perform full interpolation if it's needed.
+	OutMetadataDomain->InitializeOnSet(OutEntryKey, VertexEntryKeys[CurrentVertexIndex], MetadataDomain);
+
 	TStaticArray<TPair<PCGMetadataEntryKey, float>, 2> Coefficients;
-	Coefficients[0] = { VertexEntryKeys[CurrentVertexIndex], 1.0 - InterpolationRatio };
+	Coefficients[0] = { VertexEntryKeys[CurrentVertexIndex], 1.0f - InterpolationRatio };
 	Coefficients[1] = { VertexEntryKeys[NextVertexIndex], InterpolationRatio };
 	OutMetadataDomain->ComputeWeightedAttribute(OutEntryKey, Coefficients, MetadataDomain);
 }
