[PCG] Clarify 2D textures in new type names before release.

#rnx

diff --git a/Engine/Plugins/PCG/Source/PCG/Private/Data/PCGRenderTargetData.cpp b/Engine/Plugins/PCG/Source/PCG/Private/Data/PCGRenderTargetData.cpp
--- a/Engine/Plugins/PCG/Source/PCG/Private/Data/PCGRenderTargetData.cpp
+++ b/Engine/Plugins/PCG/Source/PCG/Private/Data/PCGRenderTargetData.cpp
@@ -8,7 +8,7 @@
 
 #include UE_INLINE_GENERATED_CPP_BY_NAME(PCGRenderTargetData)
 
-PCG_DEFINE_TYPE_INFO(FPCGDataTypeInfoRenderTarget, UPCGRenderTargetData)
+PCG_DEFINE_TYPE_INFO(FPCGDataTypeInfoRenderTarget2D, UPCGRenderTargetData)
 
 void UPCGRenderTargetData::Initialize(UTextureRenderTarget2D* InRenderTarget, const FTransform& InTransform, bool bInSkipReadbackToCPU, bool bInTakeOwnershipOfRenderTarget)
 {
diff --git a/Engine/Plugins/PCG/Source/PCG/Private/Data/PCGTextureData.cpp b/Engine/Plugins/PCG/Source/PCG/Private/Data/PCGTextureData.cpp
--- a/Engine/Plugins/PCG/Source/PCG/Private/Data/PCGTextureData.cpp
+++ b/Engine/Plugins/PCG/Source/PCG/Private/Data/PCGTextureData.cpp
@@ -17,8 +17,8 @@
 
 #include UE_INLINE_GENERATED_CPP_BY_NAME(PCGTextureData)
 
-PCG_DEFINE_TYPE_INFO(FPCGDataTypeInfoBaseTexture, UPCGBaseTextureData)
-PCG_DEFINE_TYPE_INFO(FPCGDataTypeInfoTexture, UPCGTextureData)
+PCG_DEFINE_TYPE_INFO(FPCGDataTypeInfoBaseTexture2D, UPCGBaseTextureData)
+PCG_DEFINE_TYPE_INFO(FPCGDataTypeInfoTexture2D, UPCGTextureData)
 
 namespace PCGTextureSamplingHelpers
 {
diff --git a/Engine/Plugins/PCG/Source/PCG/Private/Data/Registry/PCGDataType.cpp b/Engine/Plugins/PCG/Source/PCG/Private/Data/Registry/PCGDataType.cpp
--- a/Engine/Plugins/PCG/Source/PCG/Private/Data/Registry/PCGDataType.cpp
+++ b/Engine/Plugins/PCG/Source/PCG/Private/Data/Registry/PCGDataType.cpp
@@ -135,13 +135,13 @@ FPCGDataTypeBaseId FPCGDataTypeBaseId::MakeFromLegacyType(EPCGDataType DataType)
 	case EPCGDataType::Landscape:
 		return Construct<FPCGDataTypeInfoLandscape>();
 	case EPCGDataType::Texture:
-		return Construct<FPCGDataTypeInfoTexture>();
+		return Construct<FPCGDataTypeInfoTexture2D>();
 	case EPCGDataType::RenderTarget:
-		return Construct<FPCGDataTypeInfoRenderTarget>();
+		return Construct<FPCGDataTypeInfoRenderTarget2D>();
 	case EPCGDataType::VirtualTexture:
 		return Construct<FPCGDataTypeInfoVirtualTexture>();
 	case EPCGDataType::BaseTexture:
-		return Construct<FPCGDataTypeInfoBaseTexture>();
+		return Construct<FPCGDataTypeInfoBaseTexture2D>();
 	case EPCGDataType::Surface:
 		return Construct<FPCGDataTypeInfoSurface>();
 	case EPCGDataType::Volume:
@@ -187,10 +187,10 @@ EPCGDataType FPCGDataTypeBaseId::AsLegacyType() const
 	PCG_DATA_TYPE_TEST(FPCGDataTypeInfoPolyline)
 	PCG_DATA_TYPE_TEST(FPCGDataTypeInfoPolygon2D)
 	PCG_DATA_TYPE_TEST(FPCGDataTypeInfoLandscape)
-	PCG_DATA_TYPE_TEST(FPCGDataTypeInfoTexture)
-	PCG_DATA_TYPE_TEST(FPCGDataTypeInfoRenderTarget)
+	PCG_DATA_TYPE_TEST(FPCGDataTypeInfoTexture2D)
+	PCG_DATA_TYPE_TEST(FPCGDataTypeInfoRenderTarget2D)
 	PCG_DATA_TYPE_TEST(FPCGDataTypeInfoVirtualTexture)
-	PCG_DATA_TYPE_TEST(FPCGDataTypeInfoBaseTexture)
+	PCG_DATA_TYPE_TEST(FPCGDataTypeInfoBaseTexture2D)
 	PCG_DATA_TYPE_TEST(FPCGDataTypeInfoSurface)
 	PCG_DATA_TYPE_TEST(FPCGDataTypeInfoVolume)
 	PCG_DATA_TYPE_TEST(FPCGDataTypeInfoPrimitive)
diff --git a/Engine/Plugins/PCG/Source/PCG/Public/Data/PCGRenderTargetData.h b/Engine/Plugins/PCG/Source/PCG/Public/Data/PCGRenderTargetData.h
--- a/Engine/Plugins/PCG/Source/PCG/Public/Data/PCGRenderTargetData.h
+++ b/Engine/Plugins/PCG/Source/PCG/Public/Data/PCGRenderTargetData.h
@@ -11,9 +11,8 @@
 class UTexture;
 class UTextureRenderTarget2D;
 
-
 USTRUCT()
-struct FPCGDataTypeInfoRenderTarget : public FPCGDataTypeInfoBaseTexture
+struct FPCGDataTypeInfoRenderTarget2D : public FPCGDataTypeInfoBaseTexture2D
 {
 	GENERATED_BODY()
 
@@ -29,7 +28,7 @@ class UPCGRenderTargetData : public UPCGBaseTextureData
 
 public:
 	// ~Begin UPCGData interface
-	PCG_ASSIGN_DEFAULT_TYPE_INFO(FPCGDataTypeInfoRenderTarget);
+	PCG_ASSIGN_DEFAULT_TYPE_INFO(FPCGDataTypeInfoRenderTarget2D);
 	virtual void AddToCrc(FArchiveCrc32& Ar, bool bFullDataCrc) const override;
 	virtual bool HoldsTransientResources() const override { return bOwnsRenderTarget; }
 	virtual bool IsCacheable() const { return Super::IsCacheable() && !bOwnsRenderTarget; }
diff --git a/Engine/Plugins/PCG/Source/PCG/Public/Data/PCGTextureData.h b/Engine/Plugins/PCG/Source/PCG/Public/Data/PCGTextureData.h
--- a/Engine/Plugins/PCG/Source/PCG/Public/Data/PCGTextureData.h
+++ b/Engine/Plugins/PCG/Source/PCG/Public/Data/PCGTextureData.h
@@ -63,8 +63,9 @@ namespace PCGTextureSamplingHelpers
 	TOptional<bool> CanGPUTextureBeCPUAccessed(UTexture2D* Texture);
 }
 
+/** Base type of 2D textures/render targets. */
 USTRUCT()
-struct FPCGDataTypeInfoBaseTexture : public FPCGDataTypeInfoSurface
+struct FPCGDataTypeInfoBaseTexture2D : public FPCGDataTypeInfoSurface
 {
 	GENERATED_BODY()
 
@@ -72,13 +73,14 @@ struct FPCGDataTypeInfoBaseTexture : public FPCGDataTypeInfoSurface
 };
 
 USTRUCT()
-struct FPCGDataTypeInfoTexture : public FPCGDataTypeInfoBaseTexture
+struct FPCGDataTypeInfoTexture2D : public FPCGDataTypeInfoBaseTexture2D
 {
 	GENERATED_BODY()
 
 	PCG_DECLARE_LEGACY_TYPE_INFO(EPCGDataType::Texture)
 };
 
+/** Base class for a 2D texture or render target. */
 UCLASS(MinimalAPI, Abstract)
 class UPCGBaseTextureData : public UPCGSurfaceData
 {
@@ -90,7 +92,7 @@ public:
 	// ~End UObject interface
 
 	// ~Begin UPCGData interface
-	PCG_ASSIGN_DEFAULT_TYPE_INFO(FPCGDataTypeInfoBaseTexture)
+	PCG_ASSIGN_DEFAULT_TYPE_INFO(FPCGDataTypeInfoBaseTexture2D)
 	// ~End UPCGData interface
 
 	//~Begin UPCGSpatialData interface
@@ -193,7 +195,7 @@ protected:
 };
 
 UCLASS(MinimalAPI, BlueprintType, ClassGroup=(Procedural))
-class UPCGTextureData : public UPCGBaseTextureData
+class UPCGTextureData : public UPCGBaseTextureData // Texture2D
 {
 	GENERATED_BODY()
 
@@ -211,7 +213,7 @@ public:
 	EPCGTextureResourceType GetTextureResourceType() const override { return ResourceType; }
 
 	//~Begin UPCGData interface
-	PCG_ASSIGN_DEFAULT_TYPE_INFO(FPCGDataTypeInfoTexture)
+	PCG_ASSIGN_DEFAULT_TYPE_INFO(FPCGDataTypeInfoTexture2D)
 	virtual void AddToCrc(FArchiveCrc32& Ar, bool bFullDataCrc) const override;
 	// Don't hold onto exported buffers currently as graphics memory usage (and lifetimes) may cause issues.
 	virtual bool HoldsTransientResources() const override { return ResourceType == EPCGTextureResourceType::ExportedTexture; }
diff --git a/Engine/Plugins/PCG/Source/PCGEditor/Private/PCGEditorModule.cpp b/Engine/Plugins/PCG/Source/PCGEditor/Private/PCGEditorModule.cpp
--- a/Engine/Plugins/PCG/Source/PCGEditor/Private/PCGEditorModule.cpp
+++ b/Engine/Plugins/PCG/Source/PCGEditor/Private/PCGEditorModule.cpp
@@ -872,9 +872,9 @@ void FPCGEditorModule::RegisterPinColorAndIcons()
 	InRegistry.RegisterPinColorFunction(FPCGDataTypeInfoConcrete::AsId(), [](const FPCGDataTypeIdentifier&) { return GetDefault<UPCGEditorSettings>()->ConcreteDataPinColor; });
 	InRegistry.RegisterPinColorFunction(FPCGDataTypeInfoPolyline::AsId(), [](const FPCGDataTypeIdentifier&) { return GetDefault<UPCGEditorSettings>()->PolyLineDataPinColor; });
 	InRegistry.RegisterPinColorFunction(FPCGDataTypeInfoLandscape::AsId(), [](const FPCGDataTypeIdentifier&) { return GetDefault<UPCGEditorSettings>()->LandscapeDataPinColor; });
-	InRegistry.RegisterPinColorFunction(FPCGDataTypeInfoBaseTexture::AsId(), [](const FPCGDataTypeIdentifier&) { return GetDefault<UPCGEditorSettings>()->BaseTextureDataPinColor; });
-	InRegistry.RegisterPinColorFunction(FPCGDataTypeInfoTexture::AsId(), [](const FPCGDataTypeIdentifier&) { return GetDefault<UPCGEditorSettings>()->TextureDataPinColor; });
-	InRegistry.RegisterPinColorFunction(FPCGDataTypeInfoRenderTarget::AsId(), [](const FPCGDataTypeIdentifier&) { return GetDefault<UPCGEditorSettings>()->RenderTargetDataPinColor; });
+	InRegistry.RegisterPinColorFunction(FPCGDataTypeInfoBaseTexture2D::AsId(), [](const FPCGDataTypeIdentifier&) { return GetDefault<UPCGEditorSettings>()->BaseTextureDataPinColor; });
+	InRegistry.RegisterPinColorFunction(FPCGDataTypeInfoTexture2D::AsId(), [](const FPCGDataTypeIdentifier&) { return GetDefault<UPCGEditorSettings>()->TextureDataPinColor; });
+	InRegistry.RegisterPinColorFunction(FPCGDataTypeInfoRenderTarget2D::AsId(), [](const FPCGDataTypeIdentifier&) { return GetDefault<UPCGEditorSettings>()->RenderTargetDataPinColor; });
 	InRegistry.RegisterPinColorFunction(FPCGDataTypeInfoVirtualTexture::AsId(), [](const FPCGDataTypeIdentifier&) { return GetDefault<UPCGEditorSettings>()->VirtualTextureDataPinColor; });
 	InRegistry.RegisterPinColorFunction(FPCGDataTypeInfoSurface::AsId(), [](const FPCGDataTypeIdentifier&) { return GetDefault<UPCGEditorSettings>()->SurfaceDataPinColor; });
 	InRegistry.RegisterPinColorFunction(FPCGDataTypeInfoVolume::AsId(), [](const FPCGDataTypeIdentifier&) { return GetDefault<UPCGEditorSettings>()->VolumeDataPinColor; });
@@ -956,9 +956,9 @@ void FPCGEditorModule::UnregisterPinColorAndIcons()
 	InRegistry.UnregisterPinColorFunction(FPCGDataTypeInfoVolume::AsId());
 	InRegistry.UnregisterPinColorFunction(FPCGDataTypeInfoSurface::AsId());
 	InRegistry.UnregisterPinColorFunction(FPCGDataTypeInfoVirtualTexture::AsId());
-	InRegistry.UnregisterPinColorFunction(FPCGDataTypeInfoRenderTarget::AsId());
-	InRegistry.UnregisterPinColorFunction(FPCGDataTypeInfoTexture::AsId());
-	InRegistry.UnregisterPinColorFunction(FPCGDataTypeInfoBaseTexture::AsId());
+	InRegistry.UnregisterPinColorFunction(FPCGDataTypeInfoRenderTarget2D::AsId());
+	InRegistry.UnregisterPinColorFunction(FPCGDataTypeInfoTexture2D::AsId());
+	InRegistry.UnregisterPinColorFunction(FPCGDataTypeInfoBaseTexture2D::AsId());
 	InRegistry.UnregisterPinColorFunction(FPCGDataTypeInfoLandscape::AsId());
 	InRegistry.UnregisterPinColorFunction(FPCGDataTypeInfoPolyline::AsId());
 	InRegistry.UnregisterPinColorFunction(FPCGDataTypeInfoConcrete::AsId());
