[PCG] Unhide StaticMeshResource data type

#rnx

diff --git a/Engine/Plugins/PCG/Source/PCG/Public/Data/PCGStaticMeshResourceData.h b/Engine/Plugins/PCG/Source/PCG/Public/Data/PCGStaticMeshResourceData.h
--- a/Engine/Plugins/PCG/Source/PCG/Public/Data/PCGStaticMeshResourceData.h
+++ b/Engine/Plugins/PCG/Source/PCG/Public/Data/PCGStaticMeshResourceData.h
@@ -16,6 +16,10 @@ struct FPCGDataTypeInfoStaticMeshResource : public FPCGDataTypeInfoResource
 	GENERATED_BODY()
 
 	PCG_DECLARE_LEGACY_TYPE_INFO(EPCGDataType::StaticMeshResource)
+
+#if WITH_EDITOR
+	virtual bool Hidden() const override { return false; }
+#endif // WITH_EDITOR
 };
 
 /** Data that wraps a Static Mesh soft object path. */
