[PCG] Massage a few things to move towards execution sources for runtime generation of primitives (using FastGeo path).

#rnx

diff --git a/Engine/Plugins/Experimental/PCGInterops/PCGFastGeoInterop/Source/PCGFastGeoInterop/Private/Compute/PrimitiveFactories/PCGPrimitiveFactoryFastGeoPISMC.cpp b/Engine/Plugins/Experimental/PCGInterops/PCGFastGeoInterop/Source/PCGFastGeoInterop/Private/Compute/PrimitiveFactories/PCGPrimitiveFactoryFastGeoPISMC.cpp
--- a/Engine/Plugins/Experimental/PCGInterops/PCGFastGeoInterop/Source/PCGFastGeoInterop/Private/Compute/PrimitiveFactories/PCGPrimitiveFactoryFastGeoPISMC.cpp
+++ b/Engine/Plugins/Experimental/PCGInterops/PCGFastGeoInterop/Source/PCGFastGeoInterop/Private/Compute/PrimitiveFactories/PCGPrimitiveFactoryFastGeoPISMC.cpp
@@ -74,11 +74,15 @@ void FPCGPrimitiveFactoryFastGeoPISMC::Initialize(FParameters&& InParameters)
 
 bool FPCGPrimitiveFactoryFastGeoPISMC::Create(FPCGContext* InContext)
 {
-	// todo_pcg: Move to execution source.
-	UPCGComponent* Component = Cast<UPCGComponent>(InContext->ExecutionSource.Get());
-	check(Component);
+	IPCGGraphExecutionSource* ExecutionSource = InContext->ExecutionSource.Get();
+	if (!ExecutionSource)
+	{
+		return true;
+	}
+
+	UWorld* World = InContext->ExecutionSource->GetExecutionState().GetWorld();
+	ULevel* Level = World ? World->GetCurrentLevel() : nullptr;
 
-	ULevel* Level = Component->GetOwner() ? Component->GetOwner()->GetLevel() : nullptr;
 	if (!Level)
 	{
 		UE_LOG(LogPCG, Error, TEXT("No persistent level found, geometry will not be created."));
@@ -94,7 +98,7 @@ bool FPCGPrimitiveFactoryFastGeoPISMC::Create(FPCGContext* InContext)
 	// Must be outer'd to level - so pick same level as pcg component.
 	UFastGeoContainer* FastGeo = FPCGContext::NewObject_AnyThread<UFastGeoContainer>(InContext, Level);
 
-	const FString DebugName = Component->GetExecutionState().GetDebugName();
+	const FString DebugName = ExecutionSource->GetExecutionState().GetDebugName();
 	TUniquePtr<FFastGeoComponentCluster> ComponentCluster = MakeUnique<FFastGeoComponentCluster>(FastGeo, *FString::Printf(TEXT("FastGeoComponentCluster_%s"), *DebugName));
 
 	const int32 NumToCreate = FMath::Min(Descriptors.Num(), PCGComputeConstants::MAX_PRIMITIVE_COMPONENTS_PER_SPAWNER);
@@ -139,10 +143,10 @@ bool FPCGPrimitiveFactoryFastGeoPISMC::Create(FPCGContext* InContext)
 		});
 	});
 
-	UPCGManagedFastGeoContainer* ManagedPrimitives = FPCGContext::NewObject_AnyThread<UPCGManagedFastGeoContainer>(InContext, Component);
+	UPCGManagedFastGeoContainer* ManagedPrimitives = FPCGContext::NewObject_AnyThread<UPCGManagedFastGeoContainer>(InContext, Cast<UObject>(ExecutionSource));
 	ManagedPrimitives->FastGeo = FastGeo;
 
-	Component->AddToManagedResources(ManagedPrimitives);
+	ExecutionSource->GetExecutionState().AddToManagedResources(ManagedPrimitives);
 
 	return true;
 }
diff --git a/Engine/Plugins/PCG/Source/PCG/Private/Compute/DataInterfaces/Elements/PCGStaticMeshSpawnerDataInterface.cpp b/Engine/Plugins/PCG/Source/PCG/Private/Compute/DataInterfaces/Elements/PCGStaticMeshSpawnerDataInterface.cpp
--- a/Engine/Plugins/PCG/Source/PCG/Private/Compute/DataInterfaces/Elements/PCGStaticMeshSpawnerDataInterface.cpp
+++ b/Engine/Plugins/PCG/Source/PCG/Private/Compute/DataInterfaces/Elements/PCGStaticMeshSpawnerDataInterface.cpp
@@ -323,8 +323,7 @@ bool UPCGStaticMeshSpawnerDataProvider::PrepareForExecute_GameThread(UPCGDataBin
 		return true;
 	}
 
-	UPCGComponent* SourceComponent = Cast<UPCGComponent>(Context->ExecutionSource.Get());
-	if (!ensure(SourceComponent))
+	if (!ensure(Context->ExecutionSource.IsValid()))
 	{
 		InBinding->CompletedMeshSpawners.Add(ProducerKernel);
 		return true;
@@ -433,7 +432,10 @@ bool UPCGStaticMeshSpawnerDataProvider::PrepareForExecute_GameThread(UPCGDataBin
 
 		if (NumPrimitivesSetup > 0)
 		{
-			SourceComponent->NotifyProceduralInstancesInUse();
+			if (UPCGComponent* SourceComponent = Cast<UPCGComponent>(Context->ExecutionSource.Get()))
+			{
+				SourceComponent->NotifyProceduralInstancesInUse();
+			}
 		}
 		else
 		{
@@ -490,8 +492,7 @@ void UPCGStaticMeshSpawnerDataProvider::CreatePrimitiveDescriptors(FPCGContext*
 
 	const UPCGStaticMeshSpawnerSettings* Settings = CastChecked<UPCGStaticMeshSpawnerSettings>(ProducerKernel->GetSettings());
 
-	AActor* TargetActor = Settings->TargetActor.Get() ? Settings->TargetActor.Get() : InContext->GetTargetActor(nullptr);
-	if (!ensure(Settings->MeshSelectorParameters) || !ensure(TargetActor))
+	if (!ensure(Settings->MeshSelectorParameters))
 	{
 		return;
 	}
@@ -509,8 +510,7 @@ void UPCGStaticMeshSpawnerDataProvider::CreatePrimitiveDescriptors(FPCGContext*
 		return;
 	}
 
-	UPCGComponent* SourceComponent = Cast<UPCGComponent>(InContext->ExecutionSource.Get());
-	if (!SourceComponent)
+	if (!InContext->ExecutionSource.IsValid())
 	{
 		return;
 	}
@@ -616,7 +616,7 @@ void UPCGStaticMeshSpawnerDataProvider::CreatePrimitiveDescriptors(FPCGContext*
 			FPCGProceduralISMComponentDescriptor Descriptor;
 			Descriptor = SelectorByAttribute->TemplateDescriptor;
 			Descriptor.NumInstances = StringKeyAndInstanceCount.Value;
-			Descriptor.WorldBounds = SourceComponent->GetGridBounds();
+			Descriptor.WorldBounds = InContext->ExecutionSource->GetExecutionState().GetBounds();
 			Descriptor.NumCustomFloats = CustomFloatCount;
 			Descriptor.StaticMesh = FSoftObjectPath(MeshPathString);
 
@@ -681,7 +681,7 @@ void UPCGStaticMeshSpawnerDataProvider::CreatePrimitiveDescriptors(FPCGContext*
 
 				FPCGProceduralISMComponentDescriptor Descriptor;
 				Descriptor = Entry.Descriptor;
-				Descriptor.WorldBounds = SourceComponent->GetGridBounds();
+				Descriptor.WorldBounds = InContext->ExecutionSource->GetExecutionState().GetBounds();
 				Descriptor.NumCustomFloats = CustomFloatCount;
 				Descriptor.StaticMesh = Entry.Descriptor.StaticMesh;
 
@@ -751,12 +751,6 @@ bool UPCGStaticMeshSpawnerDataProvider::SetupPrimitives(FPCGContext* InContext,
 
 	const UPCGStaticMeshSpawnerSettings* Settings = CastChecked<UPCGStaticMeshSpawnerSettings>(ProducerKernel->GetSettings());
 
-	AActor* TargetActor = Settings->TargetActor.Get() ? Settings->TargetActor.Get() : InContext->GetTargetActor(nullptr);
-	if (!ensure(TargetActor))
-	{
-		return true;
-	}
-
 	FPCGSpawnerPrimitives* Primitives = InBinding->MeshSpawnersToPrimitives.Find(ProducerKernel);
 	if (!Primitives)
 	{
@@ -799,7 +793,7 @@ bool UPCGStaticMeshSpawnerDataProvider::SetupPrimitives(FPCGContext* InContext,
 
 		IPCGPrimitiveFactoryISMBase::FParameters Params;
 		Params.Descriptors = PrimitiveDescriptors;
-		Params.TargetActor = TargetActor;
+		Params.TargetActor = Settings->TargetActor.Get() ? Settings->TargetActor.Get() : InContext->GetTargetActor(nullptr);
 
 		Factory->Initialize(MoveTemp(Params));
 	}
diff --git a/Engine/Plugins/PCG/Source/PCG/Private/Compute/Elements/PCGComputeGraphElement.cpp b/Engine/Plugins/PCG/Source/PCG/Private/Compute/Elements/PCGComputeGraphElement.cpp
--- a/Engine/Plugins/PCG/Source/PCG/Private/Compute/Elements/PCGComputeGraphElement.cpp
+++ b/Engine/Plugins/PCG/Source/PCG/Private/Compute/Elements/PCGComputeGraphElement.cpp
@@ -5,7 +5,6 @@
 #include "PCGComponent.h"
 #include "PCGGraph.h"
 #include "PCGModule.h"
-#include "PCGSubsystem.h"
 #include "Components/PCGProceduralISMComponent.h"
 #include "Compute/PCGComputeCommon.h"
 #include "Data/PCGTextureData.h"
@@ -15,6 +14,7 @@
 #include "Editor/IPCGEditorModule.h"
 #include "Graph/PCGGraphCompiler.h"
 #include "Helpers/PCGHelpers.h"
+#include "Subsystems/IPCGBaseSubsystem.h"
 
 #include "ComputeWorkerInterface.h"
 #include "ComputeFramework/ComputeFramework.h"
@@ -164,16 +164,10 @@ bool FPCGComputeGraphElement::ExecuteInternal(FPCGContext* InContext) const
 
 		if (TopGraph->IsHierarchicalGenerationEnabled() && !DynamicSubgraph)
 		{
-			if (const UPCGComponent* Component = Cast<UPCGComponent>(Source))
-			{
-				if (Component->IsLocalComponent() || Component->IsPartitioned())
-				{
-					Context->GenerationGridSize = Component->GetGenerationGridSize();
-				}
-			}
+			Context->GenerationGridSize = Source->GetExecutionState().GetGenerationGridSize();
 		}
-
-		if (UPCGSubsystem* Subsystem = UWorld::GetSubsystem<UPCGSubsystem>(Source->GetExecutionState().GetWorld()))
+		
+		if (IPCGBaseSubsystem* Subsystem = Source->GetExecutionState().GetSubsystem())
 		{
 			Context->ComputeGraph = Subsystem->GetComputeGraph(DynamicSubgraph ? DynamicSubgraph : TopGraph, Context->GenerationGridSize, ComputeGraphIndex);
 		}
diff --git a/Engine/Plugins/PCG/Source/PCG/Private/Compute/PrimitiveFactories/PCGPrimitiveFactoryPISMC.cpp b/Engine/Plugins/PCG/Source/PCG/Private/Compute/PrimitiveFactories/PCGPrimitiveFactoryPISMC.cpp
--- a/Engine/Plugins/PCG/Source/PCG/Private/Compute/PrimitiveFactories/PCGPrimitiveFactoryPISMC.cpp
+++ b/Engine/Plugins/PCG/Source/PCG/Private/Compute/PrimitiveFactories/PCGPrimitiveFactoryPISMC.cpp
@@ -4,6 +4,7 @@
 
 #include "PCGComponent.h"
 #include "PCGContext.h"
+#include "PCGModule.h"
 #include "Components/PCGProceduralISMComponent.h"
 #include "Compute/PCGComputeCommon.h"
 
@@ -51,7 +52,15 @@ void FPCGPrimitiveFactoryPISMC::Initialize(FParameters&& InParameters)
 bool FPCGPrimitiveFactoryPISMC::Create(FPCGContext* InContext)
 {
 	UPCGComponent* SourceComponent = Cast<UPCGComponent>(InContext->ExecutionSource.Get());
-	check(SourceComponent);
+	if (!SourceComponent)
+	{
+		if (InContext->ExecutionSource.IsValid())
+		{
+			UE_LOG(LogPCG, Error, TEXT("FPCGPrimitiveFactoryPISMC: This primitive factory currently requires a PCG component execution source."));
+		}
+
+		return true;
+	}
 
 	FPCGProceduralISMCBuilderParameters Params;
 	Params.bAllowDescriptorChanges = false;
diff --git a/Engine/Plugins/PCG/Source/PCG/Private/PCGComponent.cpp b/Engine/Plugins/PCG/Source/PCG/Private/PCGComponent.cpp
--- a/Engine/Plugins/PCG/Source/PCG/Private/PCGComponent.cpp
+++ b/Engine/Plugins/PCG/Source/PCG/Private/PCGComponent.cpp
@@ -4362,7 +4362,7 @@ FPCGGridDescriptor FPCGComponentExecutionState::GetGridDescriptor(uint32 InGridS
 	return Component->GetGridDescriptor(InGridSize);
 }
 
-int32 FPCGComponentExecutionState::GetGenerationGridSize() const
+uint32 FPCGComponentExecutionState::GetGenerationGridSize() const
 {
 	check(Component);
 	return Component->GetGenerationGridSize();
@@ -4386,6 +4386,12 @@ const FPCGDataCollection* FPCGComponentExecutionState::RetrieveOutputDataForPin(
 	return Component->RetrieveOutputDataForPin(InResourceKey);
 }
 
+void FPCGComponentExecutionState::AddToManagedResources(UPCGManagedResource* InResource)
+{
+	check(Component);
+	return Component->AddToManagedResources(InResource);
+}
+
 FPCGTaskId FPCGComponentExecutionState::GenerateLocalGetTaskId(EPCGHiGenGrid Grid)
 {
 	check(Component)
diff --git a/Engine/Plugins/PCG/Source/PCG/Private/PCGGraphExecutionStateInterface.cpp b/Engine/Plugins/PCG/Source/PCG/Private/PCGGraphExecutionStateInterface.cpp
--- a/Engine/Plugins/PCG/Source/PCG/Private/PCGGraphExecutionStateInterface.cpp
+++ b/Engine/Plugins/PCG/Source/PCG/Private/PCGGraphExecutionStateInterface.cpp
@@ -2,11 +2,17 @@
 
 #include "PCGGraphExecutionStateInterface.h"
 
+#include "PCGModule.h"
 #include "Editor/IPCGEditorModule.h"
 #include "Subsystems/PCGSubsystem.h"
 
 #include UE_INLINE_GENERATED_CPP_BY_NAME(PCGGraphExecutionStateInterface)
 
+void IPCGGraphExecutionState::AddToManagedResources(UPCGManagedResource* InResource)
+{
+	UE_LOG(LogPCG, Error, TEXT("This execution source (%s) does not support managed resources."), *GetDebugName());
+}
+
 IPCGBaseSubsystem* IPCGGraphExecutionState::GetSubsystem() const
 {
 	IPCGBaseSubsystem* Subsystem = nullptr;
diff --git a/Engine/Plugins/PCG/Source/PCG/Public/PCGComponent.h b/Engine/Plugins/PCG/Source/PCG/Public/PCGComponent.h
--- a/Engine/Plugins/PCG/Source/PCG/Public/PCGComponent.h
+++ b/Engine/Plugins/PCG/Source/PCG/Public/PCGComponent.h
@@ -115,10 +115,11 @@ public:
 	virtual bool IsPartitioned() const override;
 	virtual bool IsLocalSource() const override;
 	virtual FPCGGridDescriptor GetGridDescriptor(uint32 InGridSize) const override;
-	virtual int32 GetGenerationGridSize() const override;
+	virtual uint32 GetGenerationGridSize() const override;
 	virtual FPCGTaskId GetGenerationTaskId() const override;
 	virtual void StoreOutputDataForPin(const FString& InResourceKey, const FPCGDataCollection& InData) const override;
 	virtual const FPCGDataCollection* RetrieveOutputDataForPin(const FString& InResourceKey) const override;
+	virtual void AddToManagedResources(UPCGManagedResource* InResource) override;
 	virtual FPCGTaskId GenerateLocalGetTaskId(EPCGHiGenGrid Grid = EPCGHiGenGrid::Uninitialized) override;
 
 #if WITH_EDITOR
diff --git a/Engine/Plugins/PCG/Source/PCG/Public/PCGGraphExecutionStateInterface.h b/Engine/Plugins/PCG/Source/PCG/Public/PCGGraphExecutionStateInterface.h
--- a/Engine/Plugins/PCG/Source/PCG/Public/PCGGraphExecutionStateInterface.h
+++ b/Engine/Plugins/PCG/Source/PCG/Public/PCGGraphExecutionStateInterface.h
@@ -17,6 +17,7 @@ class IPCGGraphExecutionSource;
 class UPCGData;
 class UPCGGraph;
 class UPCGGraphInstance;
+class UPCGManagedResource;
 struct FPCGGridDescriptor;
 
 namespace PCGUtils
@@ -82,6 +83,8 @@ public:
 	/** Lookup data using a resource key that identifies the pin. */
 	virtual const FPCGDataCollection* RetrieveOutputDataForPin(const FString& InResourceKey) const { return nullptr; }
 
+	virtual void AddToManagedResources(UPCGManagedResource* InResource);
+
 	/** Returns the subsystem associated with this source. By default, it will return the UPCGSubsystem if the source has a world, otherwise, the EditorSubsystem (in editor only, obviously). */
 	PCG_API virtual IPCGBaseSubsystem* GetSubsystem() const;
 
@@ -134,7 +137,7 @@ public:
 	virtual FPCGGridDescriptor GetGridDescriptor(uint32 InGridSize) const;
 
 	/** Get the size for the HiGen grid this execution source executes on. */
-	virtual int32 GetGenerationGridSize() const { return PCGHiGenGrid::UninitializedGridSize(); }
+	virtual uint32 GetGenerationGridSize() const { return PCGHiGenGrid::UninitializedGridSize(); }
 };
 
 UINTERFACE(BlueprintType, MinimalAPI)
