[PCG] Rework PCG runtime generation cell priority calculation to make it behave consistently across grid sizes, and make direction weight much easier to tune (1.0 gives well balanced behavior).

[RN] [PCG] Improved priority calculation in distance and direction runtime generation policy to behave more uniformly across grid sizes and to be easier to tweak. May require retweaking if non-default weights were in use.

diff --git a/Engine/Plugins/PCG/Source/PCG/Private/RuntimeGen/SchedulingPolicies/PCGSchedulingPolicyDistanceAndDirection.cpp b/Engine/Plugins/PCG/Source/PCG/Private/RuntimeGen/SchedulingPolicies/PCGSchedulingPolicyDistanceAndDirection.cpp
--- a/Engine/Plugins/PCG/Source/PCG/Private/RuntimeGen/SchedulingPolicies/PCGSchedulingPolicyDistanceAndDirection.cpp
+++ b/Engine/Plugins/PCG/Source/PCG/Private/RuntimeGen/SchedulingPolicies/PCGSchedulingPolicyDistanceAndDirection.cpp
@@ -14,6 +14,7 @@ double UPCGSchedulingPolicyDistanceAndDirection::CalculatePriority(const IPCGGen
 	check(InGenSource);
 
 	double Priority = 0.0;
+	double MaxPriority = 0.0;
 
 	if (bUseDistance || bUseDirection)
 	{
@@ -31,9 +32,14 @@ double UPCGSchedulingPolicyDistanceAndDirection::CalculatePriority(const IPCGGen
 
 			if (bUseDistance)
 			{
-				const double MaxRadius = 2.0 * PCGHiGenGrid::GridToGridSize(EPCGHiGenGrid::GridMax);
-				const double Distance = 1.0 - FMath::Clamp(GenSourceToComponent.Length() / MaxRadius, 0.0, 1.0);
-				Priority += Distance * FMath::Clamp(DistanceWeight, 0.0f, 1.0f);
+				// Normalize via grid size so that priority calculations are well behaved across different grid sizes.
+				const double GridSize = FMath::Max(GenerationBounds.GetSize().X, GenerationBounds.GetSize().Y);
+				const double DistInCells = GenSourceToComponent.Length() / GridSize;
+				// No major significance of reciprocal choice here, just a weight curve that is monotonic and always in [0, 1].
+				const double DistPriority = 1.0 / FMath::Max(DistInCells, 1.0);
+
+				Priority += DistPriority;
+				MaxPriority += 1.0;
 			}
 
 			if (bUseDirection)
@@ -42,15 +48,21 @@ double UPCGSchedulingPolicyDistanceAndDirection::CalculatePriority(const IPCGGen
 
 				if ((GenSourceDirectionOptional.IsSet()))
 				{
+					// Directional weight is simply dot product scaled to [0, 1], then weighted by user.
 					const double DirectionDotProd = FVector::DotProduct(GenSourceToComponent.GetSafeNormal(), GenSourceDirectionOptional.GetValue());
-					Priority += DirectionDotProd * FMath::Clamp(DirectionWeight, 0.0f, 1.0f);
+					Priority += DirectionWeight * FMath::Lerp(0.5, 1.0, DirectionDotProd);
+					MaxPriority += DirectionWeight;
 				}
 			}
 		}
 	}
 
-	// Remap range from [-1, 2] to [0, 1].
-	Priority = (Priority + 1.0) / 3.0;
+	// Normalize.
+	if (MaxPriority > 0.0)
+	{
+		Priority /= MaxPriority;
+	}
+
 	ensure(Priority >= 0.0 && Priority <= 1.0);
 
 	return Priority;
@@ -133,7 +145,6 @@ bool UPCGSchedulingPolicyDistanceAndDirection::IsEquivalent(const UPCGScheduling
 	if (const UPCGSchedulingPolicyDistanceAndDirection* Other = Cast<UPCGSchedulingPolicyDistanceAndDirection>(OtherSchedulingPolicy))
 	{
 		return bUseDistance == Other->bUseDistance
-			&& DistanceWeight == Other->DistanceWeight
 			&& bUseDirection == Other->bUseDirection
 			&& DirectionWeight == Other->DirectionWeight
 			&& bUseFrustumCulling == Other->bUseFrustumCulling
diff --git a/Engine/Plugins/PCG/Source/PCG/Public/RuntimeGen/SchedulingPolicies/PCGSchedulingPolicyDistanceAndDirection.h b/Engine/Plugins/PCG/Source/PCG/Public/RuntimeGen/SchedulingPolicies/PCGSchedulingPolicyDistanceAndDirection.h
--- a/Engine/Plugins/PCG/Source/PCG/Public/RuntimeGen/SchedulingPolicies/PCGSchedulingPolicyDistanceAndDirection.h
+++ b/Engine/Plugins/PCG/Source/PCG/Public/RuntimeGen/SchedulingPolicies/PCGSchedulingPolicyDistanceAndDirection.h
@@ -41,17 +41,13 @@ public:
 	UPROPERTY(BlueprintReadWrite, EditAnywhere, Category = "RuntimeGeneration|Scheduling Policy Parameters", meta = (EditCondition = "bShouldDisplayProperties", EditConditionHides, HideEditConditionToggle))
 	bool bUseDistance = true;
 
-	/** Scalar value used to increase/decrease the impact of distance in the scheduling priority. */
-	UPROPERTY(BlueprintReadWrite, EditAnywhere, Category = "RuntimeGeneration|Scheduling Policy Parameters", meta = (ClampMin = 0.0f, ClampMax = 1.0f, EditCondition = "bShouldDisplayProperties && bUseDistance", EditConditionHides))
-	float DistanceWeight = 1.0f;
-
 	/** Toggle whether or not direction is used to calculate the scheduling priority. */
 	UPROPERTY(BlueprintReadWrite, EditAnywhere, Category = "RuntimeGeneration|Scheduling Policy Parameters", meta = (EditCondition = "bShouldDisplayProperties", EditConditionHides, HideEditConditionToggle))
 	bool bUseDirection = true;
 
 	/** Scalar value used to increase/decrease the impact of direction in the scheduling priority. */
-	UPROPERTY(BlueprintReadWrite, EditAnywhere, Category = "RuntimeGeneration|Scheduling Policy Parameters", meta = (ClampMin = 0.0f, ClampMax = 1.0f, EditCondition = "bShouldDisplayProperties && bUseDirection", EditConditionHides))
-	float DirectionWeight = 0.0025f;
+	UPROPERTY(BlueprintReadWrite, EditAnywhere, Category = "RuntimeGeneration|Scheduling Policy Parameters", meta = (ClampMin = 0.0f, UIMax = 1.0f, EditCondition = "bShouldDisplayProperties && bUseDirection", EditConditionHides))
+	float DirectionWeight = 1.0f;
 
 	/** With frustum culling enabled, only components whose bounds overlap the view frustum will be generated. */
 	UPROPERTY(BlueprintReadWrite, EditAnywhere, Category = "RuntimeGeneration|Scheduling Policy Parameters", meta = (EditCondition = "bShouldDisplayProperties", EditConditionHides, HideEditConditionToggle))
@@ -68,4 +64,16 @@ public:
 	/** Client/Server policy mode will determine which generation sources to consider. */
 	UPROPERTY(BlueprintReadWrite, EditAnywhere, Category = "RuntimeGeneration|Scheduling Policy Parameters")
 	EPCGSchedulingPolicyNetworkMode NetworkMode = EPCGSchedulingPolicyNetworkMode::All;
+
+public:
+	UE_DEPRECATED(5.7, "Distance weight no longer in use, use DirectionWeight instead.")
+	UPROPERTY(BlueprintReadWrite, Category = "RuntimeGeneration|Scheduling Policy Parameters", meta = (DeprecatedProperty), BlueprintGetter = "GetDistanceWeight_DEPRECATED", BlueprintSetter = "SetDistanceWeight_DEPRECATED")
+	float DistanceWeight = 1.0f;
+
+private:
+	UFUNCTION(BlueprintPure, Category = "RuntimeGeneration|Scheduling Policy Parameters", meta = (DeprecatedFunction, DeprecationMessage = "Use Direction Weight instead."))
+	float GetDistanceWeight_DEPRECATED() const { return 0.0f; }
+
+	UFUNCTION(BlueprintCallable, Category = "RuntimeGeneration|Scheduling Policy Parameters", meta = (DeprecatedFunction, DeprecationMessage = "Use Direction Weight instead."))
+	void SetDistanceWeight_DEPRECATED(float Value) {}
 };
