[PCG] Remove inspection data from attribute list view + more expositions of PCG editor functionalities

* Attribute list view is central in inspection, makeing it necessary for the viewport to work
* To give more flexibility and reduce the logic on one widget, a new InspectionData manager is introduced
* It is owned by the PCG Editor and will notify widgets when something changes, instead of the widget handling the logic
* Also add support to choose which toolbar button to have
* Also add support to toggle inspection on a specific node.

#rnx

diff --git a/Engine/Plugins/Experimental/ProceduralVegetationEditor/Source/ProceduralVegetationEditor/Private/PVEditor.cpp b/Engine/Plugins/Experimental/ProceduralVegetationEditor/Source/ProceduralVegetationEditor/Private/PVEditor.cpp
--- a/Engine/Plugins/Experimental/ProceduralVegetationEditor/Source/ProceduralVegetationEditor/Private/PVEditor.cpp
+++ b/Engine/Plugins/Experimental/ProceduralVegetationEditor/Source/ProceduralVegetationEditor/Private/PVEditor.cpp
@@ -230,17 +230,10 @@ TSharedRef<FTabManager::FLayout> FPVEditor::GetDefaultLayout() const
 	);
 }
 
-void FPVEditor::RegisterToolbar() const
+void FPVEditor::RegisterToolbarInternal(FToolMenuSection& PCGSection) const
 {
-	UToolMenus* ToolMenus = UToolMenus::Get();
-	FName ParentName;
-	const FName ToolbarName = GetToolMenuToolbarName(ParentName);
-
-	UToolMenu* const AssetToolbar = UToolMenus::Get()->ExtendMenu(ToolbarName);
-	FToolMenuSection& Section = AssetToolbar->FindOrAddSection("ProceduralVegetationTools");
-
 	const FPVEditorCommands& Commands = FPVEditorCommands::Get();
-	Section.AddEntry(FToolMenuEntry::InitToolBarButton(
+	PCGSection.AddEntry(FToolMenuEntry::InitToolBarButton(
 			Commands.Export,
 			TAttribute<FText>(),
 			TAttribute<FText>(),
diff --git a/Engine/Plugins/Experimental/ProceduralVegetationEditor/Source/ProceduralVegetationEditor/Private/PVEditor.h b/Engine/Plugins/Experimental/ProceduralVegetationEditor/Source/ProceduralVegetationEditor/Private/PVEditor.h
--- a/Engine/Plugins/Experimental/ProceduralVegetationEditor/Source/ProceduralVegetationEditor/Private/PVEditor.h
+++ b/Engine/Plugins/Experimental/ProceduralVegetationEditor/Source/ProceduralVegetationEditor/Private/PVEditor.h
@@ -39,7 +39,7 @@ protected:
 	// ~Begin FPCGEditor interface
 	virtual TAttribute<FGraphAppearanceInfo> GetAppearanceInfo() const override;
 	virtual TSharedRef<FTabManager::FLayout> GetDefaultLayout() const override;
-	virtual void RegisterToolbar() const override;
+	virtual void RegisterToolbarInternal(FToolMenuSection& PCGSection) const override;
 	virtual void BindCommands() override;
 	virtual void OnSelectedNodesChanged(const TSet<UObject*>& InNewSelection) override;
 	// ~End FPCGEditor interface
diff --git a/Engine/Plugins/PCG/Source/PCG/Private/Tests/Determinism/PCGDeterminismTestsCommon.cpp b/Engine/Plugins/PCG/Source/PCG/Private/Tests/Determinism/PCGDeterminismTestsCommon.cpp
--- a/Engine/Plugins/PCG/Source/PCG/Private/Tests/Determinism/PCGDeterminismTestsCommon.cpp
+++ b/Engine/Plugins/PCG/Source/PCG/Private/Tests/Determinism/PCGDeterminismTestsCommon.cpp
@@ -69,17 +69,17 @@ namespace PCGDeterminismTests
 	}
 
 #if WITH_EDITOR
-	void RunDeterminismTest(const UPCGGraph* InPCGGraph, UPCGComponent* InPCGComponent, FDeterminismTestResult& OutResult)
+	void RunDeterminismTest(const UPCGGraph* InPCGGraph, const UPCGComponent* InPCGComponent, FDeterminismTestResult& OutResult)
 	{
-		RunDeterminismTest(InPCGGraph, static_cast<IPCGGraphExecutionSource*>(InPCGComponent), OutResult);
+		RunDeterminismTest(InPCGGraph, static_cast<const IPCGGraphExecutionSource*>(InPCGComponent), OutResult);
 	}
 	
-	void RunDeterminismTest(const UPCGGraph* InPCGGraph, IPCGGraphExecutionSource* InSource, FDeterminismTestResult& OutResult)
+	void RunDeterminismTest(const UPCGGraph* InPCGGraph, const IPCGGraphExecutionSource* InSource, FDeterminismTestResult& OutResult)
 	{
 		TRACE_CPUPROFILER_EVENT_SCOPE(PCGDeterminismTests::RunDeterminismTest::Graph);
 
 		// @todo_pcg: Support Determinism test on all sources.  
-		UPCGComponent* InPCGComponent = Cast<UPCGComponent>(InSource);
+		const UPCGComponent* InPCGComponent = Cast<const UPCGComponent>(InSource);
 		
 		if (!InPCGGraph || !InPCGComponent)
 		{
diff --git a/Engine/Plugins/PCG/Source/PCG/Public/Tests/Determinism/PCGDeterminismTestsCommon.h b/Engine/Plugins/PCG/Source/PCG/Public/Tests/Determinism/PCGDeterminismTestsCommon.h
--- a/Engine/Plugins/PCG/Source/PCG/Public/Tests/Determinism/PCGDeterminismTestsCommon.h
+++ b/Engine/Plugins/PCG/Source/PCG/Public/Tests/Determinism/PCGDeterminismTestsCommon.h
@@ -142,7 +142,7 @@ namespace PCGDeterminismTests
 
 #if WITH_EDITOR
 	/** Validates if a PCGGraph is deterministic */
-	PCG_API void RunDeterminismTest(const UPCGGraph* InPCGGraph, IPCGGraphExecutionSource* InSource, FDeterminismTestResult& OutResult);
+	PCG_API void RunDeterminismTest(const UPCGGraph* InPCGGraph, const IPCGGraphExecutionSource* InSource, FDeterminismTestResult& OutResult);
 #endif
 
 	/** Validates all the generic determinism tests for any given node */
diff --git a/Engine/Plugins/PCG/Source/PCGEditor/Private/Managers/PCGEditorInspectionDataManager.cpp b/Engine/Plugins/PCG/Source/PCGEditor/Private/Managers/PCGEditorInspectionDataManager.cpp
new file mode 100644
--- /dev/null
+++ b/Engine/Plugins/PCG/Source/PCGEditor/Private/Managers/PCGEditorInspectionDataManager.cpp
@@ -0,0 +1,154 @@
+// Copyright Epic Games, Inc. All Rights Reserved.
+
+#include "Managers/PCGEditorInspectionDataManager.h"
+
+#include "PCGEditorGraph.h"
+#include "PCGGraphExecutionInspection.h"
+#include "PCGGraphExecutionStateInterface.h"
+#include "Nodes/PCGEditorGraphNodeBase.h"
+
+void FPCGEditorInspectionDataEntry::Clear()
+{
+	InspectionData = nullptr;
+	Node.Reset();
+	PinName = NAME_None;
+	bIsOutputPin = true;
+	Stack = FPCGStack{};
+}
+
+void FPCGEditorInspectionDataManager::Cleanup()
+{
+	if (PCGSourceBeingInspected.IsValid())
+	{
+		if (PCGSourceBeingInspected->GetExecutionState().GetInspection().IsInspecting())
+		{
+			PCGSourceBeingInspected->GetExecutionState().GetInspection().DisableInspection();
+		}
+	}
+
+	if (LastValidPCGSourceBeingInspected.IsValid())
+	{
+		if (LastValidPCGSourceBeingInspected->GetExecutionState().GetInspection().IsInspecting())
+		{
+			LastValidPCGSourceBeingInspected->GetExecutionState().GetInspection().DisableInspection();
+		}
+	}
+
+	OnInspectedStackChangedDelegate.Clear();
+}
+
+void FPCGEditorInspectionDataManager::SetStackBeingInspected(const FPCGStack& FullStack)
+{
+	if (FullStack == StackBeingInspected)
+	{
+		// No-op if we're already inspecting this stack.
+		return;
+	}
+
+	IPCGGraphExecutionSource* LastSource = LastValidPCGSourceBeingInspected.Get();
+	IPCGGraphExecutionSource* NewSource = const_cast<IPCGGraphExecutionSource*>(FullStack.GetRootSource());
+
+	if (NewSource && NewSource != LastSource)
+	{
+		if (LastSource && LastSource->GetExecutionState().GetInspection().IsInspecting())
+		{
+			LastSource->GetExecutionState().GetInspection().DisableInspection();
+		}
+
+		LastValidPCGSourceBeingInspected = NewSource;
+	}
+
+	PCGSourceBeingInspected = NewSource;
+
+	StackBeingInspected = FullStack;
+	OnInspectedStackChangedDelegate.Broadcast(StackBeingInspected);
+}
+
+void FPCGEditorInspectionDataManager::OnSourceGenerated(IPCGGraphExecutionSource* InSource) const
+{
+	if (InSource == GetPCGSourceBeingInspected())
+	{
+		TRACE_CPUPROFILER_EVENT_SCOPE(FPCGEditor::SetStackBeingInspected::BroadcastStackBeingInspected);
+		OnInspectedStackChangedDelegate.Broadcast(StackBeingInspected);
+	}
+}
+
+const FPCGDataCollection* FPCGEditorInspectionDataManager::GetInspectionData(const FPCGEditorInspectionDataEntrySetupParams& InParams) const
+{
+	const FPCGEditorInspectionDataEntry NewEntry = ConstructEntry(InParams);
+	return NewEntry.IsValid() ? NewEntry.InspectionData : nullptr;
+}
+
+void FPCGEditorInspectionDataManager::SetupInspectionEntry(const FPCGEditorInspectionDataEntrySetupParams& InParams)
+{
+	InspectionEntries[InParams.EntryIndex] = ConstructEntry(InParams);
+	OnInspectionEntryChangedDelegate.Broadcast(InspectionEntries[InParams.EntryIndex], InParams.EntryIndex);
+}
+
+const FPCGEditorInspectionDataEntry& FPCGEditorInspectionDataManager::GetInspectionEntry(int32 EntryIndex)
+{
+	check(EntryIndex >= 0 && EntryIndex < NumberOfEntries);
+	return InspectionEntries[EntryIndex].IsValid() ? InspectionEntries[EntryIndex] : InspectionEntries[0];
+}
+
+FPCGEditorInspectionDataEntry FPCGEditorInspectionDataManager::ConstructEntry(const FPCGEditorInspectionDataEntrySetupParams& InParams) const
+{
+	FPCGEditorInspectionDataEntry Result{};
+
+	if (InParams.EntryIndex < 0
+		|| InParams.EntryIndex >= InspectionEntries.Num()
+		|| !InParams.Node
+		|| !PCGSourceBeingInspected.IsValid()
+		|| StackBeingInspected.GetStackFrames().IsEmpty())
+	{
+		return Result;
+	}
+
+	const UPCGNode* PCGNode = InParams.Node->GetPCGNode();
+	if (!PCGNode)
+	{
+		return Result;
+	}
+
+	const TArray<TObjectPtr<UPCGPin>>& Pins = InParams.bIsOutputPin ? PCGNode->GetOutputPins() : PCGNode->GetInputPins();
+	int32 PinIndex = InParams.PinIndex;
+	if (PinIndex == INDEX_NONE)
+	{
+		PinIndex = InParams.PinName == NAME_None ? 0 : Pins.IndexOfByPredicate([PinName = InParams.PinName](const UPCGPin* InPin) { return InPin && InPin->Properties.Label == PinName; });
+	}
+	
+	if (!Pins.IsValidIndex(PinIndex))
+	{
+		return Result;
+	}
+
+	const UPCGPin* Pin = Pins[PinIndex];
+
+	PCGEditorGraphUtils::GetInspectablePin(PCGNode, Pin, PCGNode, Pin);
+
+	if (!PCGNode || !Pin)
+	{
+		return Result;
+	}
+
+	// Create a temporary stack with Node+Pin to query the exact DataCollection we are inspecting
+	FPCGStack Stack = StackBeingInspected;
+	TArray<FPCGStackFrame>& StackFrames = Stack.GetStackFramesMutable();
+	StackFrames.Reserve(StackFrames.Num() + 2);
+	StackFrames.Emplace(PCGNode);
+	StackFrames.Emplace(Pin);
+
+	const FPCGDataCollection* DataCollection = PCGSourceBeingInspected.Get()->GetExecutionState().GetInspection().GetInspectionData(Stack);
+	if (!DataCollection)
+	{
+		return Result;
+	}
+
+	Result.InspectionData = DataCollection;
+	Result.Node = InParams.Node;
+	Result.PinName = Pin->Properties.Label;
+	Result.bIsOutputPin = InParams.bIsOutputPin;
+	Result.Stack = MoveTemp(Stack);
+
+	return Result;
+}
diff --git a/Engine/Plugins/PCG/Source/PCGEditor/Private/Nodes/PCGEditorGraphNodeBase.cpp b/Engine/Plugins/PCG/Source/PCGEditor/Private/Nodes/PCGEditorGraphNodeBase.cpp
--- a/Engine/Plugins/PCG/Source/PCGEditor/Private/Nodes/PCGEditorGraphNodeBase.cpp
+++ b/Engine/Plugins/PCG/Source/PCGEditor/Private/Nodes/PCGEditorGraphNodeBase.cpp
@@ -694,7 +694,7 @@ EPCGChangeType UPCGEditorGraphNodeBase::UpdateErrorsAndWarnings()
 	return bStateChanged ? EPCGChangeType::Cosmetic : EPCGChangeType::None;
 }
 
-EPCGChangeType UPCGEditorGraphNodeBase::UpdateStructuralVisualization(IPCGGraphExecutionSource* InSourceBeingDebugged, const FPCGStack* InStackBeingInspected, bool bNewlyPlaced)
+EPCGChangeType UPCGEditorGraphNodeBase::UpdateStructuralVisualization(const IPCGGraphExecutionSource* InSourceBeingDebugged, const FPCGStack* InStackBeingInspected, bool bNewlyPlaced)
 {
 	const UPCGGraph* Graph = PCGNode ? PCGNode->GetGraph() : nullptr;
 	if (!Graph || !InSourceBeingDebugged)
@@ -811,7 +811,7 @@ EPCGChangeType UPCGEditorGraphNodeBase::UpdateStructuralVisualization(IPCGGraphE
 	return ChangeType;
 }
 
-EPCGChangeType UPCGEditorGraphNodeBase::UpdateGPUVisualization(IPCGGraphExecutionSource* InSourceBeingDebugged, const FPCGStack* InStackBeingInspected)
+EPCGChangeType UPCGEditorGraphNodeBase::UpdateGPUVisualization(const IPCGGraphExecutionSource* InSourceBeingDebugged, const FPCGStack* InStackBeingInspected)
 {
 	EPCGChangeType ChangeType = EPCGChangeType::None;
 
diff --git a/Engine/Plugins/PCG/Source/PCGEditor/Private/Nodes/PCGEditorGraphNodeBase.h b/Engine/Plugins/PCG/Source/PCGEditor/Private/Nodes/PCGEditorGraphNodeBase.h
--- a/Engine/Plugins/PCG/Source/PCGEditor/Private/Nodes/PCGEditorGraphNodeBase.h
+++ b/Engine/Plugins/PCG/Source/PCGEditor/Private/Nodes/PCGEditorGraphNodeBase.h
@@ -90,9 +90,9 @@ public:
 	EPCGChangeType UpdateErrorsAndWarnings();
 
 	/** If the currently inspected grid size is smaller than the grid size of this node, display transparent. */
-	EPCGChangeType UpdateStructuralVisualization(IPCGGraphExecutionSource* InSourceBeingDebugged, const FPCGStack* InStackBeingInspected, bool bNewlyPlaced = false);
+	EPCGChangeType UpdateStructuralVisualization(const IPCGGraphExecutionSource* InSourceBeingDebugged, const FPCGStack* InStackBeingInspected, bool bNewlyPlaced = false);
 
-	EPCGChangeType UpdateGPUVisualization(IPCGGraphExecutionSource* InSourceBeingDebugged, const FPCGStack* InStackBeingInspected);
+	EPCGChangeType UpdateGPUVisualization(const IPCGGraphExecutionSource* InSourceBeingDebugged, const FPCGStack* InStackBeingInspected);
 
 	bool CanUserAddRemoveDynamicInputPins() const;
 
diff --git a/Engine/Plugins/PCG/Source/PCGEditor/Private/PCGEditor.cpp b/Engine/Plugins/PCG/Source/PCGEditor/Private/PCGEditor.cpp
--- a/Engine/Plugins/PCG/Source/PCGEditor/Private/PCGEditor.cpp
+++ b/Engine/Plugins/PCG/Source/PCGEditor/Private/PCGEditor.cpp
@@ -24,6 +24,7 @@
 #include "PCGEditorSettings.h"
 #include "PCGEditorStyle.h"
 #include "PCGEditorUtils.h"
+#include "Managers/PCGEditorInspectionDataManager.h"
 #include "Nodes/PCGEditorGraphNode.h"
 #include "Nodes/PCGEditorGraphNodeInput.h"
 #include "Nodes/PCGEditorGraphNodeOutput.h"
@@ -152,6 +153,8 @@ UPCGEditorGraph* FPCGEditor::GetPCGEditorGraph(const UPCGSettings* InSettings)
 
 void FPCGEditor::Initialize(const EToolkitMode::Type InMode, const TSharedPtr<class IToolkitHost>& InToolkitHost, UPCGGraph* InPCGGraph)
 {
+	InspectionDataManager.OnInspectedStackChangedDelegate.AddSP(SharedThis(this), &FPCGEditor::UpdateAfterInspectedStackChanged);
+
 	PCGGraphBeingEdited = InPCGGraph;
 
 	UpdateDefaultExecutionSource();
@@ -186,11 +189,11 @@ void FPCGEditor::Initialize(const EToolkitMode::Type InMode, const TSharedPtr<cl
 	DebugObjectTreeWidget = CreateDebugObjectTreeWidget();
 	FindWidget = CreateFindWidget();
 
-	for (int AttributesIndex = 0; AttributesIndex < 4; ++AttributesIndex)
+	for (int AttributesIndex = 0; AttributesIndex < FPCGEditorInspectionDataManager::NumberOfEntries; ++AttributesIndex)
 	{
 		const int ViewportEditorPanel = static_cast<int>(EPCGEditorPanel::Viewport1) + AttributesIndex;
 
-		AttributesWidgets.Add(CreateAttributesWidget());
+		AttributesWidgets.Add(CreateAttributesWidget(AttributesIndex));
 		AttributesWidgets[AttributesIndex]->SetViewportWidget(CreateViewportWidget(), static_cast<EPCGEditorPanel>(ViewportEditorPanel));
 	}
 	
@@ -262,41 +265,7 @@ void FPCGEditor::SetStackBeingInspectedFromAnotherEditor(const FPCGStack& FullSt
 
 void FPCGEditor::SetStackBeingInspected(const FPCGStack& FullStack)
 {
-	if (FullStack == StackBeingInspected)
-	{
-		// No-op if we're already inspecting this stack.
-		return;
-	}
-
-	IPCGGraphExecutionSource* LastSource = LastValidPCGSourceBeingInspected.Get();
-	IPCGGraphExecutionSource* NewSource = const_cast<IPCGGraphExecutionSource*>(FullStack.GetRootSource());
-
-	if (NewSource && NewSource != LastSource)
-	{
-		if (LastSource && LastSource->GetExecutionState().GetInspection().IsInspecting())
-		{
-			LastSource->GetExecutionState().GetInspection().DisableInspection();
-		}
-
-		LastValidPCGSourceBeingInspected = NewSource;
-	}
-
-	if (PCGGraphBeingEdited)
-	{
-		if (PCGGraphBeingEdited->IsInspecting())
-		{
-			PCGGraphBeingEdited->DisableInspection();
-		}
-
-		PCGGraphBeingEdited->EnableInspection(FullStack);
-	}
-
-	PCGSourceBeingInspected = NewSource;
-
-	StackBeingInspected = FullStack;
-	OnInspectedStackChangedDelegate.Broadcast(StackBeingInspected);
-
-	UpdateAfterInspectedStackChanged();
+	InspectionDataManager.SetStackBeingInspected(FullStack);
 }
 
 void FPCGEditor::OnSourceGenerated(IPCGGraphExecutionSource* InSource)
@@ -306,11 +275,7 @@ void FPCGEditor::OnSourceGenerated(IPCGGraphExecutionSource* InSource)
 		DebugObjectTreeWidget->RequestRefresh();
 	}
 
-	if(InSource == GetPCGSourceBeingInspected())
-	{
-		TRACE_CPUPROFILER_EVENT_SCOPE(FPCGEditor::SetStackBeingInspected::BroadcastStackBeingInspected);
-		OnInspectedStackChangedDelegate.Broadcast(StackBeingInspected);
-	}
+	InspectionDataManager.OnSourceGenerated(InSource);
 }
 
 bool FPCGEditor::OnValidateNodeTitle(const FText& NewName, UEdGraphNode* GraphNode, FText& OutErrorMessage)
@@ -327,24 +292,39 @@ bool FPCGEditor::OnValidateNodeTitle(const FText& NewName, UEdGraphNode* GraphNo
 	return false;
 }
 
-void FPCGEditor::UpdateAfterInspectedStackChanged()
+void FPCGEditor::UpdateAfterInspectedStackChanged(const FPCGStack& FullStack)
 {
-	IPCGGraphExecutionSource* Source = PCGSourceBeingInspected.Get();
-	UPCGComponent* Component = Cast<UPCGComponent>(PCGSourceBeingInspected.Get());
+	const bool bValidStack = !FullStack.GetStackFrames().IsEmpty();
+	
+	IPCGGraphExecutionSource* Source = InspectionDataManager.GetPCGSourceBeingInspected();
 
-	if (Source)
+	if (bValidStack)
 	{
-		// Implementation note: if we're inspecting and have not pre-run the graph, then it probably makes sense to enable inspection by default. 
-		// TODO This could be selected with a cvar though.
-		const bool bHasBeenGeneratedThisSession = Component && Component->bGenerated && Component->WasGeneratedThisSession();
-		const bool bWasInspecting = Source->GetExecutionState().GetInspection().IsInspecting();
-		const bool bNeedsInspection = Algo::AnyOf(AttributesWidgets, [](const TSharedPtr<SPCGEditorGraphAttributeListView>& ALV) { return ALV->GetNodeBeingInspected() != nullptr; });
-
-		if (!bHasBeenGeneratedThisSession || (bNeedsInspection && !bWasInspecting))
+		if (PCGGraphBeingEdited)
 		{
-			Source->GetExecutionState().GetInspection().EnableInspection();
-			UpdateDebugAfterComponentSelection(Component, Component, true);
+			PCGGraphBeingEdited->EnableInspection(FullStack);
 		}
+		
+		if (Source)
+		{
+			UPCGComponent* Component = Cast<UPCGComponent>(Source);
+		
+			// Implementation note: if we're inspecting and have not pre-run the graph, then it probably makes sense to enable inspection by default. 
+			// TODO This could be selected with a cvar though.
+			const bool bHasBeenGeneratedThisSession = Component && Component->bGenerated && Component->WasGeneratedThisSession();
+			const bool bWasInspecting = Source->GetExecutionState().GetInspection().IsInspecting();
+			const bool bNeedsInspection = Algo::AnyOf(AttributesWidgets, [](const TSharedPtr<SPCGEditorGraphAttributeListView>& ALV) { return ALV->GetNodeBeingInspected() != nullptr; });
+
+			if (!bHasBeenGeneratedThisSession || (bNeedsInspection && !bWasInspecting))
+			{
+				Source->GetExecutionState().GetInspection().EnableInspection();
+				UpdateDebugAfterComponentSelection(Component, Component, true);
+			}
+		}
+	}
+	else if (PCGGraphBeingEdited)
+	{
+		PCGGraphBeingEdited->DisableInspection();
 	}
 
 	check(PCGEditorGraph);
@@ -354,8 +334,8 @@ void FPCGEditor::UpdateAfterInspectedStackChanged()
 		{
 			// Update now that component has changed. Will fire OnNodeChanged if necessary.
 			EPCGChangeType ChangeType = PCGNode->UpdateErrorsAndWarnings();
-			ChangeType |= PCGNode->UpdateStructuralVisualization(Source, &StackBeingInspected);
-			ChangeType |= PCGNode->UpdateGPUVisualization(Source, &StackBeingInspected);
+			ChangeType |= PCGNode->UpdateStructuralVisualization(Source, &FullStack);
+			ChangeType |= PCGNode->UpdateGPUVisualization(Source, &FullStack);
 
 			if (ChangeType != EPCGChangeType::None)
 			{
@@ -378,6 +358,11 @@ UPCGComponent* FPCGEditor::GetPCGComponentBeingInspected() const
 	return const_cast<UPCGComponent*>(Cast<UPCGComponent>(GetPCGSourceBeingInspected()));
 }
 
+IPCGGraphExecutionSource* FPCGEditor::GetPCGSourceBeingInspected() const
+{
+	return InspectionDataManager.GetPCGSourceBeingInspected();
+}
+
 void FPCGEditor::UpdateDebugAfterComponentSelection(UPCGComponent* InOldComponent, UPCGComponent* InNewComponent, bool bInNewComponentStartedInspecting)
 {
 	if (!ensure(PCGGraphBeingEdited))
@@ -453,6 +438,7 @@ void FPCGEditor::UpdateDebugAfterComponentSelection(UPCGComponent* InOldComponen
 
 const FPCGStack* FPCGEditor::GetStackBeingInspected() const
 {
+	const FPCGStack& StackBeingInspected = InspectionDataManager.GetStackBeingInspected();
 	return StackBeingInspected.GetStackFrames().IsEmpty() ? nullptr : &StackBeingInspected;
 }
 
@@ -758,63 +744,112 @@ void FPCGEditor::RegisterToolbar() const
 	if (!ToolMenus->IsMenuRegistered(ToolbarName))
 	{
 		UToolMenu* ToolBar = ToolMenus->RegisterMenu(ToolbarName, ParentName, EMultiBoxType::ToolBar);
-
-		const FPCGEditorCommands& PCGEditorCommands = FPCGEditorCommands::Get();
 		FToolMenuSection& Section = ToolBar->AddSection("PCGToolbar", TAttribute<FText>());
 
-		Section.AddEntry(FToolMenuEntry::InitToolBarButton(
-			PCGEditorCommands.Find,
-			TAttribute<FText>(),
-			TAttribute<FText>(),
-			FSlateIcon(FPCGEditorStyle::Get().GetStyleSetName(), "PCG.Command.Find")));
-
-		Section.AddEntry(FToolMenuEntry::InitToolBarButton(
-			FPCGEditorCommands::Get().PauseAutoRegeneration,
-			TAttribute<FText>(),
-			TAttribute<FText>(),
-			FSlateIcon(FPCGEditorStyle::Get().GetStyleSetName(), "PCG.Command.PauseRegen")));
-
-		Section.AddEntry(FToolMenuEntry::InitToolBarButton(
-			FPCGEditorCommands::Get().ForceGraphRegeneration,
-			TAttribute<FText>(),
-			TAttribute<FText>(),
-			TAttribute<FSlateIcon>::CreateLambda([]()
-			{
-				static const FSlateIcon ForceRegen = FSlateIcon(FPCGEditorStyle::Get().GetStyleSetName(), "PCG.Command.ForceRegen");
-				static const FSlateIcon ForceRegenClearCache = FSlateIcon(FPCGEditorStyle::Get().GetStyleSetName(), "PCG.Command.ForceRegenClearCache");
-					
-				FModifierKeysState ModifierKeys = FSlateApplication::Get().GetModifierKeys();
-				return ModifierKeys.IsControlDown() ? ForceRegenClearCache : ForceRegen;
-			})));
+		RegisterToolbarInternal(Section);
+	}
+}
 
-		Section.AddEntry(FToolMenuEntry::InitToolBarButton(
-			FPCGEditorCommands::Get().CancelExecution,
-			TAttribute<FText>(),
-			TAttribute<FText>(),
-			FSlateIcon(FPCGEditorStyle::Get().GetStyleSetName(), "PCG.Command.StopRegen")));
+void FPCGEditor::RegisterToolbarInternal(FToolMenuSection& PCGSection) const
+{
+	RegisterToolbarButton(PCGSection, EPCGToolbarButtons::Find);
+	RegisterToolbarButton(PCGSection, EPCGToolbarButtons::PauseRegen);
+	RegisterToolbarButton(PCGSection, EPCGToolbarButtons::ForceRegen);
+	RegisterToolbarButton(PCGSection, EPCGToolbarButtons::CancelExecution);
+	RegisterToolbarButton(PCGSection, EPCGToolbarButtons::OpenDebugObjectTreeTab);
+
+	PCGSection.AddSeparator(NAME_None);
 
-		if (IsPanelAvailable(FPCGEditor_private::DebugObjectID))
+	RegisterToolbarButton(PCGSection, EPCGToolbarButtons::GraphParams);
+	RegisterToolbarButton(PCGSection, EPCGToolbarButtons::GraphSettings);
+}
+
+void FPCGEditor::RegisterToolbarButton(FToolMenuSection& Section, EPCGToolbarButtons Button) const
+{
+	const FPCGEditorCommands& PCGEditorCommands = FPCGEditorCommands::Get();
+	
+	switch (Button)
+	{
+	case EPCGToolbarButtons::Find:
 		{
 			Section.AddEntry(FToolMenuEntry::InitToolBarButton(
-				FPCGEditorCommands::Get().OpenDebugObjectTreeTab,
+				PCGEditorCommands.Find,
 				TAttribute<FText>(),
 				TAttribute<FText>(),
-				FSlateIcon(FPCGEditorStyle::Get().GetStyleSetName(), "PCG.Command.OpenDebugTreeTab")));
+				FSlateIcon(FPCGEditorStyle::Get().GetStyleSetName(), "PCG.Command.Find")));
+
+			break;
 		}
+	case EPCGToolbarButtons::PauseRegen:
+		{
+			Section.AddEntry(FToolMenuEntry::InitToolBarButton(
+				FPCGEditorCommands::Get().PauseAutoRegeneration,
+				TAttribute<FText>(),
+				TAttribute<FText>(),
+				FSlateIcon(FPCGEditorStyle::Get().GetStyleSetName(), "PCG.Command.PauseRegen")));
+
+			break;
+		}
+	case EPCGToolbarButtons::ForceRegen:
+		{
+			Section.AddEntry(FToolMenuEntry::InitToolBarButton(
+				FPCGEditorCommands::Get().ForceGraphRegeneration,
+				TAttribute<FText>(),
+				TAttribute<FText>(),
+				TAttribute<FSlateIcon>::CreateLambda([]()
+				{
+					static const FSlateIcon ForceRegen = FSlateIcon(FPCGEditorStyle::Get().GetStyleSetName(), "PCG.Command.ForceRegen");
+					static const FSlateIcon ForceRegenClearCache = FSlateIcon(FPCGEditorStyle::Get().GetStyleSetName(), "PCG.Command.ForceRegenClearCache");
+								
+					FModifierKeysState ModifierKeys = FSlateApplication::Get().GetModifierKeys();
+					return ModifierKeys.IsControlDown() ? ForceRegenClearCache : ForceRegen;
+				})));
+
+			break;
+		}
+	case EPCGToolbarButtons::CancelExecution:
+		{
+			Section.AddEntry(FToolMenuEntry::InitToolBarButton(
+				FPCGEditorCommands::Get().CancelExecution,
+				TAttribute<FText>(),
+				TAttribute<FText>(),
+				FSlateIcon(FPCGEditorStyle::Get().GetStyleSetName(), "PCG.Command.StopRegen")));
 
-		Section.AddSeparator(NAME_None);
+			break;
+		}
+	case EPCGToolbarButtons::OpenDebugObjectTreeTab:
+		{
+			if (IsPanelAvailable(FPCGEditor_private::DebugObjectID))
+			{
+				Section.AddEntry(FToolMenuEntry::InitToolBarButton(
+					FPCGEditorCommands::Get().OpenDebugObjectTreeTab,
+					TAttribute<FText>(),
+					TAttribute<FText>(),
+					FSlateIcon(FPCGEditorStyle::Get().GetStyleSetName(), "PCG.Command.OpenDebugTreeTab")));
+			}
+
+			break;
+		}
+	case EPCGToolbarButtons::GraphParams:
+		{
+			Section.AddEntry(FToolMenuEntry::InitToolBarButton(
+				 PCGEditorCommands.ToggleGraphParams,
+				 TAttribute<FText>(),
+				 TAttribute<FText>(),
+				 FSlateIcon(FPCGEditorStyle::Get().GetStyleSetName(), "PCG.Command.OpenGraphParams")));
 
-		Section.AddEntry(FToolMenuEntry::InitToolBarButton(
-			 PCGEditorCommands.ToggleGraphParams,
-			 TAttribute<FText>(),
-			 TAttribute<FText>(),
-			 FSlateIcon(FPCGEditorStyle::Get().GetStyleSetName(), "PCG.Command.OpenGraphParams")));
+			break;
+		}
+	case EPCGToolbarButtons::GraphSettings:
+		{
+			Section.AddEntry(FToolMenuEntry::InitToolBarButton(
+				 PCGEditorCommands.EditGraphSettings,
+				 TAttribute<FText>(),
+				 TAttribute<FText>(),
+				 FSlateIcon(FPCGEditorStyle::Get().GetStyleSetName(), "PCG.Command.GraphSettings")));
 
-		Section.AddEntry(FToolMenuEntry::InitToolBarButton(
-			 PCGEditorCommands.EditGraphSettings,
-			 TAttribute<FText>(),
-			 TAttribute<FText>(),
-			 FSlateIcon(FPCGEditorStyle::Get().GetStyleSetName(), "PCG.Command.GraphSettings")));
+			break;
+		}
 	}
 }
 
@@ -1237,14 +1272,16 @@ void FPCGEditor::OnDeterminismNodeTest() const
 
 bool FPCGEditor::CanRunDeterminismGraphTest() const
 {
-	return PCGEditorGraph && PCGSourceBeingInspected.IsValid();
+	return PCGEditorGraph && InspectionDataManager.GetPCGSourceBeingInspected();
 }
 
 void FPCGEditor::OnDeterminismGraphTest() const
 {
 	check(GraphEditorWidget.IsValid());
 
-	if (!DeterminismWidget.IsValid() || !DeterminismWidget->WidgetIsConstructed() || !PCGGraphBeingEdited || !PCGSourceBeingInspected.IsValid())
+	const IPCGGraphExecutionSource* PCGSourceBeingInspected = InspectionDataManager.GetPCGSourceBeingInspected();
+
+	if (!DeterminismWidget.IsValid() || !DeterminismWidget->WidgetIsConstructed() || !PCGGraphBeingEdited || !PCGSourceBeingInspected)
 	{
 		return;
 	}
@@ -1267,7 +1304,7 @@ void FPCGEditor::OnDeterminismGraphTest() const
 	TestResult->TestResultName = PCGGraphBeingEdited->GetName();
 	TestResult->Seed = PCGSourceBeingInspected->GetExecutionState().GetSeed();
 
-	PCGDeterminismTests::RunDeterminismTest(PCGGraphBeingEdited, PCGSourceBeingInspected.Get(), *TestResult);
+	PCGDeterminismTests::RunDeterminismTest(PCGGraphBeingEdited, PCGSourceBeingInspected, *TestResult);
 
 	DeterminismWidget->AddItem(TestResult);
 	DeterminismWidget->AddDetailsColumn();
@@ -1902,6 +1939,21 @@ void FPCGEditor::OnToggleInspected()
 	const UPCGNode* PCGNode = PCGGraphNodeBase ? PCGGraphNodeBase->GetPCGNode() : nullptr;
 	const UPCGSettingsInterface* PCGSettingsInterface = PCGNode ? PCGNode->GetSettingsInterface() : nullptr;
 
+	// Switch value.
+	SetNodeInspected(PCGGraphNodeBase, /*bValue=*/ PCGGraphNodeBase ? !PCGGraphNodeBase->GetInspected() : false);
+}
+
+void FPCGEditor::SetNodeInspected(UPCGEditorGraphNodeBase* InspectedNode, bool bValue)
+{
+	if (InspectedNode && InspectedNode->GetInspected() == bValue)
+	{
+		// Nothing to do.
+		return;
+	}
+	
+	const UPCGNode* PCGNode = InspectedNode ? InspectedNode->GetPCGNode() : nullptr;
+	const UPCGSettingsInterface* PCGSettingsInterface = PCGNode ? PCGNode->GetSettingsInterface() : nullptr;
+
 	if (PCGSettingsInterface && !PCGSettingsInterface->CanBeDebugged())
 	{
 		return;
@@ -1916,13 +1968,13 @@ void FPCGEditor::OnToggleInspected()
 	bool bIsInspecting = false;
 
 	// If the selected node was previously inspected, stop inspecting it, and unselect it from the attribute list views
-	if (InspectedNodesBefore.Contains(PCGGraphNodeBase))
+	if (InspectedNodesBefore.Contains(InspectedNode))
 	{
-		PCGGraphNodeBase->SetInspected(false);
+		InspectedNode->SetInspected(false);
 
 		for (TSharedPtr<SPCGEditorGraphAttributeListView> AttributeListView : AttributesWidgets)
 		{
-			if (AttributeListView->GetNodeBeingInspected() == PCGGraphNodeBase)
+			if (AttributeListView->GetNodeBeingInspected() == InspectedNode)
 			{
 				AttributeListView->SetNodeBeingInspected(nullptr);
 			}
@@ -1936,15 +1988,15 @@ void FPCGEditor::OnToggleInspected()
 		{
 			if (!AttributeListView->IsLocked())
 			{
-				AttributeListView->SetNodeBeingInspected(PCGGraphNodeBase);
+				AttributeListView->SetNodeBeingInspected(InspectedNode);
 			}
 
 			InspectedNodesAfter.Add(AttributeListView->GetNodeBeingInspected());
 		}
 
-		if (InspectedNodesAfter.Contains(PCGGraphNodeBase))
+		if (InspectedNode && InspectedNodesAfter.Contains(InspectedNode))
 		{
-			PCGGraphNodeBase->SetInspected(true);
+			InspectedNode->SetInspected(true);
 			bIsInspecting = true;
 		}
 
@@ -1960,17 +2012,20 @@ void FPCGEditor::OnToggleInspected()
 	if (bIsInspecting)
 	{
 		// Summon the first attribute list view that is inspecting this node
-		auto InvokeFirstTab = [this, PCGGraphNodeBase](bool bVisibleOnly) -> bool
+		auto InvokeFirstTab = [this, InspectedNode](bool bVisibleOnly) -> bool
 		{
 			for (int AttributeListViewIndex = 0; AttributeListViewIndex < AttributesWidgets.Num(); ++AttributeListViewIndex)
 			{
 				TSharedPtr<SPCGEditorGraphAttributeListView> AttributeListView = AttributesWidgets[AttributeListViewIndex];
-				if (AttributeListView->GetNodeBeingInspected() == PCGGraphNodeBase)
+				if (AttributeListView->GetNodeBeingInspected() == InspectedNode)
 				{
 					if (!bVisibleOnly || TabManager->FindExistingLiveTab(FPCGEditor_private::AttributesID[AttributeListViewIndex]))
 					{
-						GetTabManager()->TryInvokeTab(FPCGEditor_private::AttributesID[AttributeListViewIndex]);
-						return true;
+						if (IsPanelAvailable(FPCGEditor_private::AttributesID[AttributeListViewIndex]))
+						{
+							GetTabManager()->TryInvokeTab(FPCGEditor_private::AttributesID[AttributeListViewIndex]);
+							return true;
+						}
 					}
 				}
 			}
@@ -1981,7 +2036,7 @@ void FPCGEditor::OnToggleInspected()
 		const bool bTabSummoned = (InvokeFirstTab(true) || InvokeFirstTab(false));
 
 		// Default to first if they are all locked
-		if (!bTabSummoned)
+		if (!bTabSummoned && IsPanelAvailable(FPCGEditor_private::AttributesID[0]))
 		{
 			GetTabManager()->TryInvokeTab(FPCGEditor_private::AttributesID[0]);
 		}
@@ -1994,7 +2049,7 @@ void FPCGEditor::OnToggleInspected()
 	}
 
 	// Turn on "inspecting" on graph if we now have at least one inspected node and had none before
-	UpdateAfterInspectedStackChanged();
+	UpdateAfterInspectedStackChanged(InspectionDataManager.GetStackBeingInspected());
 }
 
 bool FPCGEditor::CanToggleInspected() const
@@ -2844,22 +2899,8 @@ void FPCGEditor::OnClose()
 	FEditorDelegates::EndPIE.RemoveAll(this);
 
 	FAssetEditorToolkit::OnClose();
-
-	if (PCGSourceBeingInspected.IsValid())
-	{
-		if (PCGSourceBeingInspected->GetExecutionState().GetInspection().IsInspecting())
-		{
-			PCGSourceBeingInspected->GetExecutionState().GetInspection().DisableInspection();
-		}
-	}
-
-	if (LastValidPCGSourceBeingInspected.IsValid())
-	{
-		if (LastValidPCGSourceBeingInspected->GetExecutionState().GetInspection().IsInspecting())
-		{
-			LastValidPCGSourceBeingInspected->GetExecutionState().GetInspection().DisableInspection();
-		}
-	}
+	
+	InspectionDataManager.Cleanup();
 
 	if (PCGGraphBeingEdited)
 	{
@@ -2913,9 +2954,10 @@ TSharedRef<SPCGEditorGraphFind> FPCGEditor::CreateFindWidget()
 	return SNew(SPCGEditorGraphFind, SharedThis(this));
 }
 
-TSharedRef<SPCGEditorGraphAttributeListView> FPCGEditor::CreateAttributesWidget()
+TSharedRef<SPCGEditorGraphAttributeListView> FPCGEditor::CreateAttributesWidget(int32 Index)
 {
-	return SNew(SPCGEditorGraphAttributeListView, SharedThis(this));
+	return SNew(SPCGEditorGraphAttributeListView, SharedThis(this))
+		.WidgetEntryNumber(Index);
 }
 
 TSharedRef<SPCGEditorGraphDeterminismListView> FPCGEditor::CreateDeterminismWidget()
@@ -3105,6 +3147,9 @@ void FPCGEditor::OnComponentUnregistered(UPCGComponent* Component)
 
 void FPCGEditor::OnSourceGenerationDone(IPCGBaseSubsystem* Subsystem, IPCGGraphExecutionSource* Source, EPCGGenerationStatus Status)
 {
+	const FPCGStack& StackBeingInspected = InspectionDataManager.GetStackBeingInspected();
+	const IPCGGraphExecutionSource* PCGSourceBeingInspected = InspectionDataManager.GetPCGSourceBeingInspected();
+	
 	// We want to refresh if the component that is done generating has generated the current graph being edited,
 	// or if it is the root of the current stack being inspected (for subgraphs to also be refreshed).
 	// If we don't have a component, we refresh nonetheless.
@@ -3112,7 +3157,7 @@ void FPCGEditor::OnSourceGenerationDone(IPCGBaseSubsystem* Subsystem, IPCGGraphE
 
 	// Additionally, if we are not inspecting but the component that's done executing contains this graph, then we should also update.
 	IPCGEditorModule* PCGEditorModule = IPCGEditorModule::Get();
-	if (!bShouldRefresh && !PCGSourceBeingInspected.IsValid() && PCGGraphBeingEdited && PCGEditorModule)
+	if (!bShouldRefresh && !PCGSourceBeingInspected && PCGGraphBeingEdited && PCGEditorModule)
 	{
 		TArray<FPCGStackSharedPtr> ExecutedStacksContainingThisGraph = PCGEditorModule->GetExecutedStacksPtrs(Source, PCGGraphBeingEdited);
 		bShouldRefresh |= !ExecutedStacksContainingThisGraph.IsEmpty();
@@ -3137,8 +3182,8 @@ void FPCGEditor::OnSourceGenerationDone(IPCGBaseSubsystem* Subsystem, IPCGGraphE
 			EPCGChangeType ChangeType = CacheDebuggingEnabled ? EPCGChangeType::Cosmetic : EPCGChangeType::None;
 
 			ChangeType |= PCGEditorGraphNode->UpdateErrorsAndWarnings();
-			ChangeType |= PCGEditorGraphNode->UpdateStructuralVisualization(PCGSourceBeingInspected.Get(), &StackBeingInspected);
-			ChangeType |= PCGEditorGraphNode->UpdateGPUVisualization(PCGSourceBeingInspected.Get(), &StackBeingInspected);
+			ChangeType |= PCGEditorGraphNode->UpdateStructuralVisualization(PCGSourceBeingInspected, &StackBeingInspected);
+			ChangeType |= PCGEditorGraphNode->UpdateGPUVisualization(PCGSourceBeingInspected, &StackBeingInspected);
 
 			if (ChangeType != EPCGChangeType::None)
 			{
diff --git a/Engine/Plugins/PCG/Source/PCGEditor/Private/Widgets/SPCGEditorGraphAttributeListView.cpp b/Engine/Plugins/PCG/Source/PCGEditor/Private/Widgets/SPCGEditorGraphAttributeListView.cpp
--- a/Engine/Plugins/PCG/Source/PCGEditor/Private/Widgets/SPCGEditorGraphAttributeListView.cpp
+++ b/Engine/Plugins/PCG/Source/PCGEditor/Private/Widgets/SPCGEditorGraphAttributeListView.cpp
@@ -24,6 +24,7 @@
 #include "PCGEditorStyle.h"
 #include "PCGEditorUtils.h"
 #include "DataVisualizations/PCGDataVisualizationHelpers.h"
+#include "Managers/PCGEditorInspectionDataManager.h"
 #include "Nodes/PCGEditorGraphNodeBase.h"
 #include "Widgets/AssetEditorViewport/SPCGEditorViewport.h"
 
@@ -331,7 +332,7 @@ SPCGEditorGraphAttributeListView::~SPCGEditorGraphAttributeListView()
 {
 	if (PCGEditorPtr.IsValid())
 	{
-		PCGEditorPtr.Pin()->OnInspectedStackChangedDelegate.RemoveAll(this);
+		PCGEditorPtr.Pin()->GetInspectionDataManager().OnInspectedStackChangedDelegate.RemoveAll(this);
 	}
 
 	DataPtr = nullptr;
@@ -339,8 +340,12 @@ SPCGEditorGraphAttributeListView::~SPCGEditorGraphAttributeListView()
 
 void SPCGEditorGraphAttributeListView::Construct(const FArguments& InArgs, TSharedPtr<FPCGEditor> InPCGEditor)
 {
+	WidgetEntryNumber = InArgs._WidgetEntryNumber;
+
+	check(InPCGEditor.IsValid());
+	
 	PCGEditorPtr = InPCGEditor;
-	PCGEditorPtr.Pin()->OnInspectedStackChangedDelegate.AddSP(this, &SPCGEditorGraphAttributeListView::OnInspectedStackChanged);
+	PCGEditorPtr.Pin()->GetInspectionDataManager().OnInspectedStackChangedDelegate.AddSP(this, &SPCGEditorGraphAttributeListView::OnInspectedStackChanged);
 
 	TextFilter = MakeShareable(new FTextFilterExpressionEvaluator(ETextFilterExpressionEvaluatorMode::Complex));
 
@@ -782,60 +787,25 @@ void SPCGEditorGraphAttributeListView::AddReferencedObjects(FReferenceCollector&
 
 const FPCGDataCollection* SPCGEditorGraphAttributeListView::GetInspectionData(const TSharedPtr<FPinComboBoxItem>& EditorPin) const
 {
-	TWeakInterfacePtr<IPCGGraphExecutionSource> ExecutionSource = GetExecutionSource();
-	if (!ExecutionSource.IsValid())
-	{
-		return nullptr;
-	}
-
-	const UPCGNode* PCGNode = PCGEditorGraphNode.IsValid() ? PCGEditorGraphNode->GetPCGNode() : nullptr;
-	if (!PCGNode)
-	{
-		return nullptr;
-	}
-
-	const UPCGPin* Pin = nullptr;
-	if (EditorPin)
-	{
-		const TArray<TObjectPtr<UPCGPin>>& Pins = EditorPin->bIsOutputPin ? PCGNode->GetOutputPins() : PCGNode->GetInputPins();
-		if (Pins.IsValidIndex(EditorPin->PinIndex))
-		{
-			Pin = Pins[EditorPin->PinIndex];
-		}
-	}
-
-	if (!Pin)
+	if (!EditorPin)
 	{
 		return nullptr;
 	}
 
-	PCGEditorGraphUtils::GetInspectablePin(PCGNode, Pin, PCGNode, Pin);
+	const FPCGEditorInspectionDataEntrySetupParams Params{WidgetEntryNumber, PCGEditorGraphNode.Get(), EditorPin->PinIndex, EditorPin->bIsOutputPin};
 
-	if (!PCGNode || !Pin)
-	{
-		return nullptr;
-	}
+	return PCGEditorPtr.IsValid() ? PCGEditorPtr.Pin()->GetInspectionDataManager().GetInspectionData(Params) : nullptr;
+}
 
-	const TSharedPtr<FPCGEditor> PCGEditor = PCGEditorPtr.Pin();
-	const FPCGStack* PCGStack = PCGEditor->GetStackBeingInspected();
-	if (!PCGStack)
+const FPCGDataCollection* SPCGEditorGraphAttributeListView::GetInspectionData() const
+{
+	if (!PCGEditorPtr.IsValid())
 	{
 		return nullptr;
 	}
 
-	// Create a temporary stack with Node+Pin to query the exact DataCollection we are inspecting
-	FPCGStack Stack = *PCGStack;
-	TArray<FPCGStackFrame>& StackFrames = Stack.GetStackFramesMutable();
-	StackFrames.Reserve(StackFrames.Num() + 2);
-	StackFrames.Emplace(PCGNode);
-	StackFrames.Emplace(Pin);
-
-	return ExecutionSource->GetExecutionState().GetInspection().GetInspectionData(Stack);
-}
-
-const FPCGDataCollection* SPCGEditorGraphAttributeListView::GetInspectionData() const
-{
-	return GetInspectionData(PinComboBox->GetSelectedItem());
+	const FPCGEditorInspectionDataEntry& Entry = PCGEditorPtr.Pin()->GetInspectionDataManager().GetInspectionEntry(WidgetEntryNumber);
+	return Entry.InspectionData;
 }
 
 void SPCGEditorGraphAttributeListView::RefreshAttributeList()
@@ -1454,6 +1424,16 @@ FText SPCGEditorGraphAttributeListView::OnGenerateSelectedPinText() const
 
 void SPCGEditorGraphAttributeListView::OnSelectionChangedPin(TSharedPtr<FPinComboBoxItem> InItem, ESelectInfo::Type InSelectInfo)
 {
+	TSharedPtr<FPinComboBoxItem> SelectedPin = PinComboBox->GetSelectedItem();
+
+	// If the pin is invalid, it will still refresh to clean up.
+	const FPCGEditorInspectionDataEntrySetupParams Params{WidgetEntryNumber, PCGEditorGraphNode.Get(), SelectedPin ? SelectedPin->PinIndex : INDEX_NONE, SelectedPin ? SelectedPin->bIsOutputPin : false};
+
+	if (PCGEditorPtr.IsValid())
+	{
+		PCGEditorPtr.Pin()->GetInspectionDataManager().SetupInspectionEntry(Params);
+	}
+
 	RefreshDataComboBox(/*bKeepSelection=*/InSelectInfo == ESelectInfo::Direct);
 
 	if (InSelectInfo != ESelectInfo::Direct)
diff --git a/Engine/Plugins/PCG/Source/PCGEditor/Private/Widgets/SPCGEditorGraphAttributeListView.h b/Engine/Plugins/PCG/Source/PCGEditor/Private/Widgets/SPCGEditorGraphAttributeListView.h
--- a/Engine/Plugins/PCG/Source/PCGEditor/Private/Widgets/SPCGEditorGraphAttributeListView.h
+++ b/Engine/Plugins/PCG/Source/PCGEditor/Private/Widgets/SPCGEditorGraphAttributeListView.h
@@ -15,12 +15,15 @@
 #include "Widgets/Views/SListView.h"
 #include "Widgets/Views/STableRow.h"
 
-class STableViewBase;
-class STextBlock;
-template <typename OptionType> class SComboBox;
-
+enum class EPCGMetadataTypes : uint8;
+enum class EPCGEditorPanel;
+struct FPCGDataCollection;
 class FPCGEditor;
+struct FPCGEditorInspectionDataEntry;
+class FPCGEditorInspectionDataManager;
 class FPCGMetadataAttributeBase;
+struct FPCGPoint;
+struct FPCGTableVisualizerColumnInfo;
 class FUICommandList;
 class SPCGEditorGraphAttributeListView;
 class SPCGEditorViewport;
@@ -29,17 +32,15 @@ class UPCGEditorGraphNodeBase;
 class UPCGMetadata;
 class UPCGParamData;
 class UPCGPointData;
-struct FPCGTableVisualizerColumnInfo;
-struct FPCGDataCollection;
-struct FPCGPoint;
-enum class EPCGMetadataTypes : uint8;
-enum class EPCGEditorPanel;
 
+template <typename OptionType> class SComboBox;
+struct FSlateBrush;
+struct FStreamableHandle;
 class SComboButton;
 class SHeaderRow;
 class SSearchBox;
-struct FSlateBrush;
-struct FStreamableHandle;
+class STableViewBase;
+class STextBlock;
 
 namespace PCGEditorGraphAttributeListView
 {
@@ -143,6 +144,7 @@ class SPCGEditorGraphAttributeListView : public SCompoundWidget
 
 public:
 	SLATE_BEGIN_ARGS(SPCGEditorGraphAttributeListView) {}
+		SLATE_ARGUMENT(int32, WidgetEntryNumber)
 	SLATE_END_ARGS()
 
 	virtual ~SPCGEditorGraphAttributeListView();
@@ -265,6 +267,9 @@ private:
 	/** Returns the current PCG component */
 	TWeakInterfacePtr<IPCGGraphExecutionSource> GetExecutionSource() const;
 
+	/** Index of the entry of this widget, to know which entry to query. */
+	int32 WidgetEntryNumber = 0;
+
 	/** Pointer back to the PCG editor that owns us */
 	TWeakPtr<FPCGEditor> PCGEditorPtr;
 
diff --git a/Engine/Plugins/PCG/Source/PCGEditor/Private/Widgets/SPCGEditorGraphFind.cpp b/Engine/Plugins/PCG/Source/PCGEditor/Private/Widgets/SPCGEditorGraphFind.cpp
--- a/Engine/Plugins/PCG/Source/PCGEditor/Private/Widgets/SPCGEditorGraphFind.cpp
+++ b/Engine/Plugins/PCG/Source/PCGEditor/Private/Widgets/SPCGEditorGraphFind.cpp
@@ -11,6 +11,7 @@
 #include "PCGSubsystem.h"
 #include "Editor/IPCGEditorModule.h"
 #include "Elements/PCGUserParameterGet.h"
+#include "Managers/PCGEditorInspectionDataManager.h"
 #include "Nodes/PCGEditorGraphNodeBase.h"
 
 #include "Algo/Contains.h"
@@ -202,7 +203,7 @@ SPCGEditorGraphFind::~SPCGEditorGraphFind()
 {
 	if (PCGEditorPtr.IsValid())
 	{
-		PCGEditorPtr.Pin()->OnInspectedStackChangedDelegate.RemoveAll(this);
+		PCGEditorPtr.Pin()->GetInspectionDataManager().OnInspectedStackChangedDelegate.RemoveAll(this);
 	}
 }
 
@@ -214,7 +215,7 @@ void SPCGEditorGraphFind::Construct(const FArguments& InArgs, TSharedPtr<FPCGEdi
 
 	if (PCGEditor)
 	{
-		PCGEditor->OnInspectedStackChangedDelegate.AddSP(this, &SPCGEditorGraphFind::OnInspectedStackChanged);
+		PCGEditor->GetInspectionDataManager().OnInspectedStackChangedDelegate.AddSP(this, &SPCGEditorGraphFind::OnInspectedStackChanged);
 	}
 
 	this->ChildSlot
diff --git a/Engine/Plugins/PCG/Source/PCGEditor/Private/Widgets/SPCGEditorGraphLogView.cpp b/Engine/Plugins/PCG/Source/PCGEditor/Private/Widgets/SPCGEditorGraphLogView.cpp
--- a/Engine/Plugins/PCG/Source/PCGEditor/Private/Widgets/SPCGEditorGraphLogView.cpp
+++ b/Engine/Plugins/PCG/Source/PCGEditor/Private/Widgets/SPCGEditorGraphLogView.cpp
@@ -8,6 +8,7 @@
 
 #include "PCGEditor.h"
 #include "PCGEditorGraph.h"
+#include "Managers/PCGEditorInspectionDataManager.h"
 #include "Nodes/PCGEditorGraphNode.h"
 
 #include "Styling/StyleColors.h"
@@ -106,7 +107,7 @@ SPCGEditorGraphLogView::~SPCGEditorGraphLogView()
 {
 	if (PCGEditorPtr.IsValid())
 	{
-		PCGEditorPtr.Pin()->OnInspectedStackChangedDelegate.RemoveAll(this);
+		PCGEditorPtr.Pin()->GetInspectionDataManager().OnInspectedStackChangedDelegate.RemoveAll(this);
 	}
 }
 
@@ -118,7 +119,7 @@ void SPCGEditorGraphLogView::Construct(const FArguments& InArgs, TSharedPtr<FPCG
 	if (PCGEditor)
 	{
 		PCGEditorGraph = PCGEditor->GetPCGEditorGraph();
-		PCGEditor->OnInspectedStackChangedDelegate.AddSP(this, &SPCGEditorGraphLogView::OnDebugStackChanged);
+		PCGEditor->GetInspectionDataManager().OnInspectedStackChangedDelegate.AddSP(this, &SPCGEditorGraphLogView::OnDebugStackChanged);
 	}
 
 	ListViewHeader = CreateHeaderRowWidget();
diff --git a/Engine/Plugins/PCG/Source/PCGEditor/Private/Widgets/SPCGEditorGraphProfilingView.cpp b/Engine/Plugins/PCG/Source/PCGEditor/Private/Widgets/SPCGEditorGraphProfilingView.cpp
--- a/Engine/Plugins/PCG/Source/PCGEditor/Private/Widgets/SPCGEditorGraphProfilingView.cpp
+++ b/Engine/Plugins/PCG/Source/PCGEditor/Private/Widgets/SPCGEditorGraphProfilingView.cpp
@@ -7,6 +7,7 @@
 
 #include "PCGEditor.h"
 #include "PCGEditorGraph.h"
+#include "Managers/PCGEditorInspectionDataManager.h"
 #include "Nodes/PCGEditorGraphNode.h"
 
 #include "HAL/PlatformApplicationMisc.h"
@@ -217,7 +218,7 @@ SPCGEditorGraphProfilingView::~SPCGEditorGraphProfilingView()
 {
 	if (PCGEditorPtr.IsValid())
 	{
-		PCGEditorPtr.Pin()->OnInspectedStackChangedDelegate.RemoveAll(this);
+		PCGEditorPtr.Pin()->GetInspectionDataManager().OnInspectedStackChangedDelegate.RemoveAll(this);
 	}
 }
 
@@ -229,7 +230,7 @@ void SPCGEditorGraphProfilingView::Construct(const FArguments& InArgs, TSharedPt
 	if (PCGEditor)
 	{
 		PCGEditorGraph = PCGEditor->GetPCGEditorGraph();
-		PCGEditor->OnInspectedStackChangedDelegate.AddSP(this, &SPCGEditorGraphProfilingView::OnDebugStackChanged);
+		PCGEditor->GetInspectionDataManager().OnInspectedStackChangedDelegate.AddSP(this, &SPCGEditorGraphProfilingView::OnDebugStackChanged);
 	}
 
 	SortingColumn = PCGEditorGraphProfilingView::NAME_TotalTime;
diff --git a/Engine/Plugins/PCG/Source/PCGEditor/Public/Managers/PCGEditorInspectionDataManager.h b/Engine/Plugins/PCG/Source/PCGEditor/Public/Managers/PCGEditorInspectionDataManager.h
new file mode 100644
--- /dev/null
+++ b/Engine/Plugins/PCG/Source/PCGEditor/Public/Managers/PCGEditorInspectionDataManager.h
@@ -0,0 +1,90 @@
+// Copyright Epic Games, Inc. All Rights Reserved.
+
+#pragma once
+
+#include "PCGGraphExecutionStateInterface.h"
+#include "Graph/PCGStackContext.h"
+
+class UPCGEditorGraphNodeBase;
+
+struct FPCGEditorInspectionDataEntry
+{
+	const FPCGDataCollection* InspectionData = nullptr;
+	TWeakObjectPtr<const UPCGEditorGraphNodeBase> Node;
+	FName PinName = NAME_None;
+	bool bIsOutputPin = true;
+	FPCGStack Stack;
+
+	bool IsValid() const { return InspectionData != nullptr; }
+	PCGEDITOR_API void Clear();
+};
+
+struct FPCGEditorInspectionDataEntrySetupParams
+{
+	FPCGEditorInspectionDataEntrySetupParams(int32 InEntryIndex, const UPCGEditorGraphNodeBase* InNode, FName InPinName = NAME_None, bool bInIsOutputPin = true)
+		: EntryIndex(InEntryIndex)
+		, Node(InNode)
+		, PinName(InPinName)
+		, bIsOutputPin(bInIsOutputPin)
+	{}
+
+	FPCGEditorInspectionDataEntrySetupParams(int32 InEntryIndex, const UPCGEditorGraphNodeBase* InNode, int32 InPinIndex = INDEX_NONE, bool bInIsOutputPin = true)
+		: EntryIndex(InEntryIndex)
+		, Node(InNode)
+		, PinIndex(InPinIndex)
+		, bIsOutputPin(bInIsOutputPin)
+	{}
+	
+	int32 EntryIndex = INDEX_NONE;
+	const UPCGEditorGraphNodeBase* Node = nullptr;
+
+	// If left to NAME_None and PinIndex is INDEX_NONE, by default it will take the first output pin.
+	FName PinName = NAME_None;
+	int32 PinIndex = INDEX_NONE;
+	bool bIsOutputPin = true;
+};
+
+DECLARE_MULTICAST_DELEGATE_OneParam(FOnInspectedStackChanged, const FPCGStack&);
+DECLARE_MULTICAST_DELEGATE_TwoParams(FOnInspectionEntryChanged, const FPCGEditorInspectionDataEntry&, int32);
+
+/**
+* Central class that will handle everything related to Inspection Data.
+* Other widgets can ask the manager to get the inspection data for a given node,
+* or change which Execution Source is currently inspected.
+*/
+class FPCGEditorInspectionDataManager : public TSharedFromThis<FPCGEditorInspectionDataManager>
+{
+public:
+	// Hard-coded to 4, as it is how many attribute list view widgets we can have at maximum.
+	static constexpr int32 NumberOfEntries = 4;
+	
+	PCGEDITOR_API void Cleanup();
+	
+	const FPCGStack& GetStackBeingInspected() const { return StackBeingInspected; }
+	PCGEDITOR_API void SetStackBeingInspected(const FPCGStack& FullStack);
+
+	IPCGGraphExecutionSource* GetPCGSourceBeingInspected() const { return const_cast<IPCGGraphExecutionSource*>(StackBeingInspected.GetRootSource()); }
+
+	PCGEDITOR_API void OnSourceGenerated(IPCGGraphExecutionSource* InSource) const;
+
+	PCGEDITOR_API const FPCGDataCollection* GetInspectionData(const FPCGEditorInspectionDataEntrySetupParams& InParams) const;
+	
+	PCGEDITOR_API void SetupInspectionEntry(const FPCGEditorInspectionDataEntrySetupParams& InParams);
+
+	// Return the entry at EntryIndex. Must be within 0 and NumberOfEntries - 1.
+	// If the entry is invalid at this index, return the first one (so multiple widget would show the same thing).
+	PCGEDITOR_API const FPCGEditorInspectionDataEntry& GetInspectionEntry(int32 EntryIndex);
+
+	FOnInspectedStackChanged OnInspectedStackChangedDelegate;
+	FOnInspectionEntryChanged OnInspectionEntryChangedDelegate;
+
+private:
+	FPCGEditorInspectionDataEntry ConstructEntry(const FPCGEditorInspectionDataEntrySetupParams& InParams) const;
+	
+	FPCGSoftGraphExecutionSource PCGSourceBeingInspected;
+	// Implementation note: we'll keep the last valid source inspected so we don't un-inspect on spurious selection changes
+	FPCGSoftGraphExecutionSource LastValidPCGSourceBeingInspected;
+	FPCGStack StackBeingInspected;
+	
+	TStaticArray<FPCGEditorInspectionDataEntry, NumberOfEntries> InspectionEntries;
+};
\ No newline at end of file
diff --git a/Engine/Plugins/PCG/Source/PCGEditor/Public/PCGEditor.h b/Engine/Plugins/PCG/Source/PCGEditor/Public/PCGEditor.h
--- a/Engine/Plugins/PCG/Source/PCGEditor/Public/PCGEditor.h
+++ b/Engine/Plugins/PCG/Source/PCGEditor/Public/PCGEditor.h
@@ -4,12 +4,14 @@
 
 #include "PCGContext.h"
 #include "Graph/PCGStackContext.h"
+#include "Managers/PCGEditorInspectionDataManager.h"
 
 #include "EditorUndoClient.h"
 #include "GraphEditor.h"
 #include "Schema/PCGEditorGraphSchema.h"
 #include "Toolkits/AssetEditorToolkit.h"
 
+class FPCGEditorInspectionDataManager;
 enum class ECheckBoxState : uint8;
 namespace ETextCommit { enum Type : int; }
 
@@ -37,7 +39,6 @@ class UPCGGraph;
 class UPCGSubsystem;
 struct FPCGCompilerDiagnostics;
 
-DECLARE_MULTICAST_DELEGATE_OneParam(FOnInspectedStackChanged, const FPCGStack&);
 DECLARE_MULTICAST_DELEGATE_OneParam(FOnInspectedNodeChanged, UPCGEditorGraphNodeBase*);
 
 enum class EPCGEditorPanel
@@ -65,6 +66,17 @@ enum class EPCGEditorPanel
 	Viewport4
 };
 
+enum class EPCGToolbarButtons
+{
+	Find,
+	PauseRegen,
+	ForceRegen,
+	CancelExecution,
+	OpenDebugObjectTreeTab,
+	GraphParams,
+	GraphSettings
+};
+
 class FPCGEditor : public FAssetEditorToolkit, public FGCObject, public FSelfRegisteringEditorUndoClient
 {
 public:
@@ -95,11 +107,14 @@ public:
 	UE_DEPRECATED(5.7, "Use GetPCGSourceBeingInspected instead")
 	PCGEDITOR_API UPCGComponent* GetPCGComponentBeingInspected() const;
 
-	IPCGGraphExecutionSource* GetPCGSourceBeingInspected() const { return const_cast<IPCGGraphExecutionSource*>(StackBeingInspected.GetRootSource()); }
+	PCGEDITOR_API IPCGGraphExecutionSource* GetPCGSourceBeingInspected() const;
 	
 	/** Gets the PCG stack we are inspecting */
 	PCGEDITOR_API const FPCGStack* GetStackBeingInspected() const;
 
+	FPCGEditorInspectionDataManager& GetInspectionDataManager() { return InspectionDataManager; }
+	const FPCGEditorInspectionDataManager& GetInspectionDataManager() const { return InspectionDataManager; }
+
 	PCGEDITOR_API void SetSourceEditorTargetObject(UObject* InObject);
 
 	/** Focus the graph view on a specific node */
@@ -176,8 +191,6 @@ public:
 
 	FActionMenuContent OnCreateActionMenuContent(UEdGraph* InGraph, const FVector2f& Location, const TArray<UEdGraphPin*>& InDraggedPins, bool bAutoExpand, SGraphEditor::FActionMenuClosed OnMenuClosed);
 
-	FOnInspectedStackChanged OnInspectedStackChangedDelegate;
-
 	/** Can determinism be tested on the current graph */
 	bool CanRunDeterminismGraphTest() const;
 	/** Run the determinism test on the current graph */
@@ -192,7 +205,10 @@ protected:
 	PCGEDITOR_API virtual TAttribute<FGraphAppearanceInfo> GetAppearanceInfo() const;
 
 	/** Register PCG specific toolbar for the editor */
-	PCGEDITOR_API virtual void RegisterToolbar() const;
+	PCGEDITOR_API virtual void RegisterToolbarInternal(FToolMenuSection& PCGSection) const;
+
+	/** Register PCG default editor toolbar button */
+	PCGEDITOR_API void RegisterToolbarButton(FToolMenuSection& Section, EPCGToolbarButtons Button) const;
 
 	/** Bind commands to delegates */
 	PCGEDITOR_API virtual void BindCommands();
@@ -209,7 +225,25 @@ protected:
 	 */
 	PCGEDITOR_API virtual void OnNodeDoubleClicked(UEdGraphNode* Node);
 
+	/** Toggle node inspection state for selected nodes. */
+	PCGEDITOR_API void OnToggleInspected();
+
+	/** Set inspection state for node. Can be null, in that case it will clear inspection. */
+	PCGEDITOR_API void SetNodeInspected(UPCGEditorGraphNodeBase* InspectedNode, bool bValue);
+
+	/** Whether we can toggle inspection of selected nodes */
+	PCGEDITOR_API virtual bool CanToggleInspected() const;
+
+	/** Whether we can toggle debug state of selected nodes */
+	PCGEDITOR_API virtual bool CanToggleDebug() const;
+
+	/** Whether we can toggle enabled state of selected nodes */
+	PCGEDITOR_API virtual bool CanToggleEnabled() const;
+
 private:
+	/** Register PCG specific toolbar for the editor */
+	void RegisterToolbar() const;
+	
 	/** Bring up the find tab */
 	void OnFind();
 
@@ -233,26 +267,18 @@ private:
 	/** Force a regeneration by invoking the graph notifications  */
 	void OnForceGraphRegeneration_Clicked();
 
-	/** Toggle node inspection state for selected nodes */
-	void OnToggleInspected();
-	/** Whether we can toggle inspection of selected nodes */
-	bool CanToggleInspected() const;
 	/** Whether selected nodes are inspected or not */
 	ECheckBoxState GetInspectedCheckState() const;
 
-	void UpdateAfterInspectedStackChanged();
+	void UpdateAfterInspectedStackChanged(const FPCGStack& FullStack);
 	
 	/** Toggle node enabled state for selected nodes */
 	void OnToggleEnabled();
-	/** Whether we can toggle enabled state of selected nodes */
-	bool CanToggleEnabled() const;
 	/** Whether selected nodes are enabled or not */
 	ECheckBoxState GetEnabledCheckState() const;
 	
 	/** Toggle node debug state for selected nodes */
 	void OnToggleDebug();
-	/** Whether we can toggle debug state of selected nodes */
-	bool CanToggleDebug() const;
 	/** Whether selected nodes are being debugged or not */
 	ECheckBoxState GetDebugCheckState() const;
 
@@ -391,7 +417,7 @@ private:
 	TSharedRef<SPCGEditorGraphFind> CreateFindWidget();
 
 	/** Create new attributes widget */
-	TSharedRef<SPCGEditorGraphAttributeListView> CreateAttributesWidget();
+	TSharedRef<SPCGEditorGraphAttributeListView> CreateAttributesWidget(int32 Index);
 
 	/** Create a new determinism tab widget */
 	TSharedRef<SPCGEditorGraphDeterminismListView> CreateDeterminismWidget();
@@ -488,8 +514,5 @@ private:
 	TObjectPtr<UPCGEditorDefaultExecutionSource> PCGDefaultExecutionSource = nullptr;
 	UPCGEditorGraph* PCGEditorGraph = nullptr;
 
-	FPCGSoftGraphExecutionSource PCGSourceBeingInspected;
-	// Implementation note: we'll keep the last valid source inspected so we don't un-inspect on spurious selection changes
-	FPCGSoftGraphExecutionSource LastValidPCGSourceBeingInspected;
-	FPCGStack StackBeingInspected;
+	FPCGEditorInspectionDataManager InspectionDataManager;
 };
