[PCG] Add options to PCG GraphInterface/PCGGraph to hide some settings

* Also hide Input/Output and Nodes as they should not appear in the details view

#rnx

diff --git a/Engine/Plugins/PCG/Source/PCG/Public/PCGGraph.h b/Engine/Plugins/PCG/Source/PCG/Public/PCGGraph.h
--- a/Engine/Plugins/PCG/Source/PCG/Public/PCGGraph.h
+++ b/Engine/Plugins/PCG/Source/PCG/Public/PCGGraph.h
@@ -148,6 +148,9 @@ protected:
 	UFUNCTION()
 	virtual bool IsTemplatePropertyEnabled() const { return false; }
 
+	UFUNCTION()
+	virtual bool CanToggleStandaloneGraph() const { return true; }
+
 	UE_API bool VerifyAndUpdateIfGraphParameterValueChanged(const FPropertyChangedChainEvent& PropertyChangedEvent);
 	UE_API bool VerifyIfGraphCustomizationChanged(const FPropertyChangedChainEvent& PropertyChangedEvent);
 #endif // WITH_EDITOR
@@ -179,9 +182,13 @@ public:
 	UPROPERTY(BlueprintReadWrite, EditAnywhere, Category = AssetInfo, AssetRegistrySearchable, meta=(EditCondition = "bExposeToLibrary", EditConditionHides, PCGNoHash))
 	FString CompatibleToolTags;
 
-	/** When enabled, this graph can be executed outside of the world using an editor execution source */
-	UPROPERTY(BlueprintReadWrite, EditAnywhere, Category = AssetInfo, AssetRegistrySearchable, meta = (PCGNoHash))
+	/** When enabled, this graph can be executed outside of the world using an editor execution source. */
+	UPROPERTY(BlueprintReadWrite, EditAnywhere, Category = AssetInfo, AssetRegistrySearchable, meta = (PCGNoHash, EditCondition = "CanToggleStandaloneGraph()", EditConditionHides))
 	bool bIsStandaloneGraph = false;
+
+	/** When enabled and standalone, it adds the action in the asset explorer to trigger a generation. Invisible, to be set in the CDO of child class to disable. */
+	UPROPERTY(AssetRegistrySearchable, meta = (PCGNoHash))
+	bool bExposeGenerationInAssetExplorer = true;
 #endif
 
 	template <typename T>
@@ -320,14 +327,14 @@ public:
 	FText Description;
 
 	/** Marks the graph to be not refreshed automatically when the landscape changes, even if it is used. */
-	UPROPERTY(BlueprintReadWrite, EditAnywhere, Category = "Settings|Advanced", meta=(PCGNoHash))
+	UPROPERTY(BlueprintReadWrite, EditAnywhere, Category = "Settings|Advanced", meta=(PCGNoHash, EditCondition="!bIsStandaloneGraph", EditConditionHides))
 	bool bIgnoreLandscapeTracking = false;
 
-	UPROPERTY(EditAnywhere, Category = Customization, meta=(PCGNoHash))
+	UPROPERTY(EditAnywhere, Category = Customization, meta=(PCGNoHash, EditCondition="ShowGraphCustomization()", EditConditionHides))
 	FPCGGraphEditorCustomization GraphCustomization;
 #endif
 
-	UPROPERTY(EditAnywhere, Category = "Settings|Advanced")
+	UPROPERTY(EditAnywhere, Category = "Settings|Advanced", meta=(EditCondition="!bIsStandaloneGraph", EditConditionHides))
 	bool bLandscapeUsesMetadata = true;
 
 	/** Creates a node using the given settings interface. Does not manage ownership - done outside of this method. */
@@ -476,13 +483,19 @@ public:
 
 protected:
 #if WITH_EDITOR
+	UFUNCTION()
+	virtual bool SupportHierarchicalGeneration() const { return !bIsStandaloneGraph; }
+
+	UFUNCTION()
+	virtual bool ShowGraphCustomization() const { return true; }
+
 	virtual bool IsTemplatePropertyEnabled() const override { return IsAsset(); }
 
 	/** Mark the input node hidden/not hidden. */
 	PCG_API void SetHiddenFlagInputNode(bool bHidden);
 	
 	/** Mark the output node hidden/not hidden. */
-	PCG_API void SetHiddenFlagOutputNode(bool bHidden); 
+	PCG_API void SetHiddenFlagOutputNode(bool bHidden);
 #endif
 
 	/** Internal function to react to add/remove nodes. bNotify can be set to false to not notify the world. */
@@ -506,14 +519,14 @@ protected:
 	/** Calculates all parent grid sizes that a given node depends on. Not thread safe, must be called within write lock. */
 	UE_API PCGHiGenGrid::FSizeArray CalculateNodeGridSizesRecursiveNoLock(const UPCGNode* InNode, uint32 InDefaultGridSize) const;
 
-	UPROPERTY(BlueprintReadOnly, VisibleAnywhere, Category = Graph, meta = (NoResetToDefault))
+	UPROPERTY(BlueprintReadOnly, Category = Graph, meta = (NoResetToDefault))
 	TArray<TObjectPtr<UPCGNode>> Nodes;
 
 	// Add input/output nodes
-	UPROPERTY(BlueprintReadOnly, VisibleAnywhere, Category = Graph, meta = (NoResetToDefault, PCGNoHash))
+	UPROPERTY(BlueprintReadOnly, Category = Graph, meta = (NoResetToDefault, PCGNoHash))
 	TObjectPtr<UPCGNode> InputNode;
 
-	UPROPERTY(BlueprintReadOnly, VisibleAnywhere, Category = Graph, meta = (NoResetToDefault, PCGNoHash))
+	UPROPERTY(BlueprintReadOnly, Category = Graph, meta = (NoResetToDefault, PCGNoHash))
 	TObjectPtr<UPCGNode> OutputNode;
 
 #if WITH_EDITORONLY_DATA
@@ -543,21 +556,21 @@ protected:
 
 	UE_API virtual FInstancedPropertyBag* GetMutableUserParametersStruct() override;
 
-	UPROPERTY(EditAnywhere, Category = Settings)
+	UPROPERTY(EditAnywhere, Category = Settings, meta=(EditCondition="SupportHierarchicalGeneration", EditConditionHides))
 	bool bUseHierarchicalGeneration = false;
 
-	UPROPERTY(EditAnywhere, Category = Settings, meta = (DisplayName = "HiGen Default Grid Size", EditCondition = "bUseHierarchicalGeneration"))
+	UPROPERTY(EditAnywhere, Category = Settings, meta = (DisplayName = "HiGen Default Grid Size", EditCondition = "SupportHierarchicalGeneration && bUseHierarchicalGeneration", EditConditionHides))
 	EPCGHiGenGrid HiGenGridSize = EPCGHiGenGrid::Grid256;
 
 	/** Shifts the grid sizes upwards based on the value, which allows to use larger grids. A value of 1 will effectively use the graph's Grid-400 values x 2 for the actual Grid-800 sizes and so on. */
-	UPROPERTY(EditAnywhere, Category = Settings, meta = (ClampMax = "10", DisplayName = "HiGen Grid Size Exponential", EditCondition = "bUseHierarchicalGeneration"))
+	UPROPERTY(EditAnywhere, Category = Settings, meta = (ClampMax = "10", DisplayName = "HiGen Grid Size Exponent", EditCondition = "SupportHierarchicalGeneration && bUseHierarchicalGeneration", EditConditionHides))
 	uint32 HiGenExponential = 0;
 
-	UPROPERTY(EditAnywhere, Category = Settings, meta = (DisplayName = "2D Grid"))
+	UPROPERTY(EditAnywhere, Category = Settings, meta = (DisplayName = "2D Grid", EditCondition = "SupportHierarchicalGeneration && bUseHierarchicalGeneration", EditConditionHides))
 	bool bUse2DGrid = true;
 
 	/** Controls whether the default pins (In, ...) are provided data on top-level graphs automatically. Note this will eventually change to be turned off by default. */
-	UPROPERTY(BlueprintReadWrite, EditAnywhere, Category = Settings)
+	UPROPERTY(BlueprintReadWrite, EditAnywhere, Category = Settings, meta=(EditCondition="!bIsStandaloneGraph", EditConditionHides))
 	bool bHasDefaultConstructedInputs = true;
 
 	/** Execution grid size for nodes. */
@@ -580,14 +593,14 @@ protected:
 	mutable FCriticalSection CachedGridInfoLock;
 
 	/** Sets whether this graph is marked as editor-only; note that the IsEditorOnly call depends on the local graph value and the value in all subgraphs, recursively. */
-	UPROPERTY(EditAnywhere, BlueprintReadWrite, Category = Cooking)
+	UPROPERTY(EditAnywhere, BlueprintReadWrite, Category = Cooking, meta=(EditCondition="!bIsStandaloneGraph", EditConditionHides))
 	bool bIsEditorOnly = false;
 
 #if WITH_EDITORONLY_DATA
 	/** When true the Debug flag in the graph editor will display debug information contextually for the selected debug object. Otherwise
 	* debug information is displayed for all components using a graph (requires regenerate).
 	*/
-	UPROPERTY(EditAnywhere, Category = Debug)
+	UPROPERTY(EditAnywhere, Category = Debug, meta=(EditCondition="!bIsStandaloneGraph", EditConditionHides))
 	bool bDebugFlagAppliesToIndividualComponents = true;
 #endif // WITH_EDITORONLY_DATA
 
@@ -607,7 +620,7 @@ public:
 	// Best used on a brand new PCG Graph.
 	UE_API EPropertyBagAlterationResult AddUserParameters(const TArray<FPropertyBagPropertyDesc>& InDescs, const UPCGGraph* InOptionalOriginalGraph = nullptr);
 
-	UPROPERTY(BlueprintReadWrite, EditAnywhere, Category = "Runtime Generation")
+	UPROPERTY(BlueprintReadWrite, EditAnywhere, Category = "Runtime Generation", meta=(EditCondition = "!bIsStandaloneGraph", EditConditionHides))
 	FPCGRuntimeGenerationRadii GenerationRadii;
 
 	UE_API virtual void OnGraphParametersChanged(EPCGGraphParameterEvent InChangeType, FName InChangedPropertyName) override;
diff --git a/Engine/Plugins/PCG/Source/PCGEditor/Private/AssetDefinitions/AssetDefinition_PCGGraphInterface.cpp b/Engine/Plugins/PCG/Source/PCGEditor/Private/AssetDefinitions/AssetDefinition_PCGGraphInterface.cpp
--- a/Engine/Plugins/PCG/Source/PCGEditor/Private/AssetDefinitions/AssetDefinition_PCGGraphInterface.cpp
+++ b/Engine/Plugins/PCG/Source/PCGEditor/Private/AssetDefinitions/AssetDefinition_PCGGraphInterface.cpp
@@ -113,7 +113,8 @@ namespace MenuExtension_PCGGraphInterface
 							{
 								bool bTagValue = false;
 								const bool bIsStandaloneGraph = AssetData.IsInstanceOf<UPCGGraphInterface>() && AssetData.GetTagValue(GET_MEMBER_NAME_CHECKED(UPCGGraphInterface, bIsStandaloneGraph), bTagValue) && bTagValue;
-								if (!bIsStandaloneGraph)
+								const bool bExportAssetActionAvailable = AssetData.IsInstanceOf<UPCGGraphInterface>() && AssetData.GetTagValue(GET_MEMBER_NAME_CHECKED(UPCGGraphInterface, bExposeGenerationInAssetExplorer), bTagValue) && bTagValue;
+								if (!bIsStandaloneGraph || !bExportAssetActionAvailable)
 								{
 									continue;
 								}
diff --git a/Engine/Plugins/PCG/Source/PCGEditor/Private/Details/PCGGraphDetails.cpp b/Engine/Plugins/PCG/Source/PCGEditor/Private/Details/PCGGraphDetails.cpp
--- a/Engine/Plugins/PCG/Source/PCGEditor/Private/Details/PCGGraphDetails.cpp
+++ b/Engine/Plugins/PCG/Source/PCGEditor/Private/Details/PCGGraphDetails.cpp
@@ -10,6 +10,7 @@
 #include "DetailLayoutBuilder.h"
 #include "DetailWidgetRow.h"
 #include "PropertyHandle.h"
+#include "Algo/AnyOf.h"
 
 #define LOCTEXT_NAMESPACE "PCGGraphDetails"
 
@@ -23,6 +24,8 @@ void FPCGGraphDetails::CustomizeDetails(IDetailLayoutBuilder& DetailBuilder)
 	TArray<TWeakObjectPtr<UObject>> ObjectsBeingCustomized;
 	DetailBuilder.GetObjectsBeingCustomized(ObjectsBeingCustomized);
 
+	SelectedGraphs.Empty();
+
 	for (TWeakObjectPtr<UObject>& Object : ObjectsBeingCustomized)
 	{
 		UPCGGraph* Graph = Cast<UPCGGraph>(Object.Get());
@@ -92,48 +95,51 @@ void FPCGGraphDetails::CustomizeDetails(IDetailLayoutBuilder& DetailBuilder)
 				]
 			];
 
-		// Add a "Run Graph Determinism Test" button in the debug category
-		IDetailCategoryBuilder& DebugCategory = DetailBuilder.EditCategory(TEXT("Debug"));
+		// Add a "Run Graph Determinism Test" button in the debug category, if we're not a standalone graph
+		if (!Algo::AnyOf(SelectedGraphs, [](const TWeakObjectPtr<UPCGGraph> Graph) { return !Graph.IsValid() || Graph->bIsStandaloneGraph; }))
+		{
+			IDetailCategoryBuilder& DebugCategory = DetailBuilder.EditCategory(TEXT("Debug"));
 
-		DebugCategory.AddCustomRow(FText::GetEmpty())
-			.ValueContent()
-			.MaxDesiredWidth(120.f)
-			[
-				SNew(SHorizontalBox)
-				+ SHorizontalBox::Slot()
-				.AutoWidth()
-				.Padding(2.0f, 0.0f)
-				.VAlign(VAlign_Fill)
+			DebugCategory.AddCustomRow(FText::GetEmpty())
+				.ValueContent()
+				.MaxDesiredWidth(120.f)
 				[
-					SNew(SButton)
-					.OnClicked_Lambda(
-					[PCGEditorWeakPtr]() -> FReply
-					{
-						if (const TSharedPtr<const FPCGEditor> PCGEditor = PCGEditorWeakPtr.Pin())
+					SNew(SHorizontalBox)
+					+ SHorizontalBox::Slot()
+					.AutoWidth()
+					.Padding(2.0f, 0.0f)
+					.VAlign(VAlign_Fill)
+					[
+						SNew(SButton)
+						.OnClicked_Lambda(
+						[PCGEditorWeakPtr]() -> FReply
 						{
-							PCGEditor->OnDeterminismGraphTest();
-						}
-
-						return FReply::Handled();
-					})
-					.IsEnabled_Lambda(
-					[PCGEditorWeakPtr]() -> bool
-					{
-						if (const TSharedPtr<const FPCGEditor> PCGEditor = PCGEditorWeakPtr.Pin())
+							if (const TSharedPtr<const FPCGEditor> PCGEditor = PCGEditorWeakPtr.Pin())
+							{
+								PCGEditor->OnDeterminismGraphTest();
+							}
+
+							return FReply::Handled();
+						})
+						.IsEnabled_Lambda(
+						[PCGEditorWeakPtr]() -> bool
 						{
-							return PCGEditor->CanRunDeterminismGraphTest();
-						}
-
-						return false;
-					})
-					.ToolTipText(LOCTEXT("RunGraphDeterminismTestTooltip", "Runs the graph-level determinism test on the currently selected debug object."))
-					[
-						SNew(STextBlock)
-						.Font(IDetailLayoutBuilder::GetDetailFont())
-						.Text(LOCTEXT("RunGraphDeterminism", "Run Determinism Test"))
+							if (const TSharedPtr<const FPCGEditor> PCGEditor = PCGEditorWeakPtr.Pin())
+							{
+								return PCGEditor->CanRunDeterminismGraphTest();
+							}
+
+							return false;
+						})
+						.ToolTipText(LOCTEXT("RunGraphDeterminismTestTooltip", "Runs the graph-level determinism test on the currently selected debug object."))
+						[
+							SNew(STextBlock)
+							.Font(IDetailLayoutBuilder::GetDetailFont())
+							.Text(LOCTEXT("RunGraphDeterminism", "Run Determinism Test"))
+						]
 					]
-				]
-			];
+				];
+		}
 	}
 
 	const FName PCGCategoryName("PCG");
