[PCG] Add option to not show debug/inspection brush on nodes

#rnx

diff --git a/Engine/Plugins/PCG/Source/PCGEditor/Private/Nodes/SPCGEditorGraphNode.cpp b/Engine/Plugins/PCG/Source/PCGEditor/Private/Nodes/SPCGEditorGraphNode.cpp
--- a/Engine/Plugins/PCG/Source/PCGEditor/Private/Nodes/SPCGEditorGraphNode.cpp
+++ b/Engine/Plugins/PCG/Source/PCGEditor/Private/Nodes/SPCGEditorGraphNode.cpp
@@ -15,6 +15,7 @@
 #include "Pins/SPCGEditorGraphPinNumSlider.h"
 #include "Pins/SPCGEditorGraphPinString.h"
 #include "Pins/SPCGEditorGraphPinVectorSlider.h"
+#include "Schema/PCGEditorGraphSchema.h"
 
 #include "GraphEditorSettings.h"
 #include "IDocumentation.h"
@@ -306,13 +307,13 @@ void SPCGEditorGraphNode::GetOverlayBrushes(bool bSelected, const FVector2f& Wid
 
 	if (const UPCGNode* PCGNode = PCGEditorGraphNode->GetPCGNode())
 	{
-		if (PCGNode->GetSettingsInterface() && PCGNode->GetSettingsInterface()->bDebug)
+		if (PCGNode->GetSettingsInterface() && PCGNode->GetSettingsInterface()->bDebug && UsesDebugBrush())
 		{
 			AddOverlayBrush(TEXT("PCG.NodeOverlay.Debug"));
 		}
 	}
 
-	if (PCGEditorGraphNode->GetInspected())
+	if (PCGEditorGraphNode->GetInspected() && UsesInspectBrush())
 	{
 		AddOverlayBrush(TEXT("PCG.NodeOverlay.Inspect"));
 	}
@@ -343,6 +344,20 @@ bool SPCGEditorGraphNode::UsesGPUOverlay() const
 	return PCGEditorGraphNode->GetSettings() && PCGEditorGraphNode->GetSettings()->ShouldExecuteOnGPU();
 }
 
+bool SPCGEditorGraphNode::UsesInspectBrush() const
+{
+	const UEdGraph* Graph = PCGEditorGraphNode->GetGraph();
+	const UPCGEditorGraphSchema* Schema = Graph ? Cast<UPCGEditorGraphSchema>(Graph->GetSchema()) : nullptr;
+	return Schema ? Schema->ShouldAddInspectBrushOnNode() : true;
+}
+
+bool SPCGEditorGraphNode::UsesDebugBrush() const
+{
+	const UEdGraph* Graph = PCGEditorGraphNode->GetGraph();
+	const UPCGEditorGraphSchema* Schema = Graph ? Cast<UPCGEditorGraphSchema>(Graph->GetSchema()) : nullptr;
+	return Schema ? Schema->ShouldAddDebugBrushOnNode() : true;
+}
+
 FLinearColor SPCGEditorGraphNode::GetGridLabelColor(EPCGHiGenGrid NodeGrid)
 {
 	// All colours hand tweaked to give a kind of "temperature scale" for the hierarchy.
diff --git a/Engine/Plugins/PCG/Source/PCGEditor/Private/Nodes/SPCGEditorGraphNode.h b/Engine/Plugins/PCG/Source/PCGEditor/Private/Nodes/SPCGEditorGraphNode.h
--- a/Engine/Plugins/PCG/Source/PCGEditor/Private/Nodes/SPCGEditorGraphNode.h
+++ b/Engine/Plugins/PCG/Source/PCGEditor/Private/Nodes/SPCGEditorGraphNode.h
@@ -42,6 +42,10 @@ protected:
 	virtual bool UsesHiGenOverlay() const;
 	/** Will add the GPU icon overlay to the node. */
 	virtual bool UsesGPUOverlay() const;
+	/** Will add the inspect brush to the node. */
+	virtual bool UsesInspectBrush() const;
+	/** Will add the debug brush to the node. */
+	virtual bool UsesDebugBrush() const;
 
 private:
 	static FLinearColor GetGridLabelColor(EPCGHiGenGrid NodeGrid);
diff --git a/Engine/Plugins/PCG/Source/PCGEditor/Public/Schema/PCGEditorGraphSchema.h b/Engine/Plugins/PCG/Source/PCGEditor/Public/Schema/PCGEditorGraphSchema.h
--- a/Engine/Plugins/PCG/Source/PCGEditor/Public/Schema/PCGEditorGraphSchema.h
+++ b/Engine/Plugins/PCG/Source/PCGEditor/Public/Schema/PCGEditorGraphSchema.h
@@ -59,6 +59,11 @@ public:
 
 	PCGEDITOR_API virtual const FSlateBrush* GetPropertyBagTypeSlateBrush(EPropertyBagContainerType ContainerType) const;
 	PCGEDITOR_API virtual FLinearColor GetPropertyBagTypeColor(const FPropertyBagPropertyDesc& Desc) const;
+
+	// To customize if we want to have special brushes for Inspect and Debug on the nodes.
+	virtual bool ShouldAddInspectBrushOnNode() const { return true; }
+	virtual bool ShouldAddDebugBrushOnNode() const { return true; }
+
 protected:
 	PCGEDITOR_API virtual void GetGraphActions(FGraphActionMenuBuilder& ActionMenuBuilder, const FPCGActionsFilter& ActionFilter, bool bIsContextual) const;
 	PCGEDITOR_API virtual void GetNativeElementActions(FGraphActionMenuBuilder& ActionMenuBuilder, const FPCGActionsFilter& ActionFilter) const;
