[PCG] Update custom PinIcons API registration to not be tied to PCG Editor Style.

* Instead of returning the FName, returns the Brush directly

#rnx

diff --git a/Engine/Plugins/PCG/Source/PCG/Private/Data/Registry/PCGDataTypeRegistry.cpp b/Engine/Plugins/PCG/Source/PCG/Private/Data/Registry/PCGDataTypeRegistry.cpp
--- a/Engine/Plugins/PCG/Source/PCG/Private/Data/Registry/PCGDataTypeRegistry.cpp
+++ b/Engine/Plugins/PCG/Source/PCG/Private/Data/Registry/PCGDataTypeRegistry.cpp
@@ -824,7 +824,7 @@ void FPCGDataTypeRegistry::UnregisterPinIconsFunction(const FPCGDataTypeIdentifi
 	PCGDataTypeRegistry::RemoveFromContainer(InId, InId.IsComposition() ? CombinationPinIcons : SingleTypePinIcons);
 }
 
-TTuple<FName, FName> FPCGDataTypeRegistry::GetPinIcons(const FPCGDataTypeIdentifier& InId, const FPCGPinProperties& InProperties, bool bIsInput) const
+TTuple<const FSlateBrush*, const FSlateBrush*> FPCGDataTypeRegistry::GetPinIcons(const FPCGDataTypeIdentifier& InId, const FPCGPinProperties& InProperties, bool bIsInput) const
 {
 	return PCGDataTypeRegistry::CallFunctionWithComposition(*this, InId, CombinationPinIcons, SingleTypePinIcons, InProperties, bIsInput);
 }
diff --git a/Engine/Plugins/PCG/Source/PCG/Public/Data/Registry/PCGDataTypeRegistry.h b/Engine/Plugins/PCG/Source/PCG/Public/Data/Registry/PCGDataTypeRegistry.h
--- a/Engine/Plugins/PCG/Source/PCG/Public/Data/Registry/PCGDataTypeRegistry.h
+++ b/Engine/Plugins/PCG/Source/PCG/Public/Data/Registry/PCGDataTypeRegistry.h
@@ -15,6 +15,11 @@ class FPCGDataTypeRegistryTestBase;
 class FPCGHierarchyTest;
 class FPCGModule;
 
+#if WITH_EDITOR
+// For custom icons
+struct FSlateBrush;
+#endif // WITH_EDITOR
+
 namespace PCGDataTypeHierarchy
 {
 	struct FTree;
@@ -83,12 +88,12 @@ public:
 
 	PCG_API FLinearColor GetPinColor(const FPCGDataTypeIdentifier& InId) const;
 
-	using FPCGPinIconsQueryFunction = TFunction<TTuple<FName, FName>(const FPCGDataTypeIdentifier& InId, const FPCGPinProperties& InProperties, bool bIsInput)>;
-	/** Can register a custom pin icon for a given Id. Can be a composition or a single type. Custom Types will be ignored. Callback should return Slate Brush names for connected and disconnected pins. */
+	using FPCGPinIconsQueryFunction = TFunction<TTuple<const FSlateBrush*, const FSlateBrush*>(const FPCGDataTypeIdentifier& InId, const FPCGPinProperties& InProperties, bool bIsInput)>;
+	/** Can register a custom pin icon for a given Id. Can be a composition or a single type. Custom Types will be ignored. Callback should return Slate Brush for connected and disconnected pins. */
 	PCG_API void RegisterPinIconsFunction(const FPCGDataTypeIdentifier& InId, FPCGPinIconsQueryFunction InFunction);
 	PCG_API void UnregisterPinIconsFunction(const FPCGDataTypeIdentifier& InId);
 	
-	PCG_API TTuple<FName, FName> GetPinIcons(const FPCGDataTypeIdentifier& InId, const FPCGPinProperties& InProperties, bool bIsInput) const;
+	PCG_API TTuple<const FSlateBrush*, const FSlateBrush*> GetPinIcons(const FPCGDataTypeIdentifier& InId, const FPCGPinProperties& InProperties, bool bIsInput) const;
 #endif // WITH_EDITOR
 	
 private:
diff --git a/Engine/Plugins/PCG/Source/PCGEditor/Private/Nodes/SPCGEditorGraphNode.cpp b/Engine/Plugins/PCG/Source/PCGEditor/Private/Nodes/SPCGEditorGraphNode.cpp
--- a/Engine/Plugins/PCG/Source/PCGEditor/Private/Nodes/SPCGEditorGraphNode.cpp
+++ b/Engine/Plugins/PCG/Source/PCGEditor/Private/Nodes/SPCGEditorGraphNode.cpp
@@ -216,12 +216,15 @@ void SPCGEditorGraphNode::AddPin(const TSharedRef<SGraphPin>& PinToAdd)
 			}
 			
 			const FPCGDataTypeIdentifier PinType = Pin->GetCurrentTypesID();
-			const TTuple<FName, FName> PinBrushes = FPCGModule::GetConstDataTypeRegistry().GetPinIcons(PinType, Pin->Properties, bIsInPin);
+			TTuple<const FSlateBrush*, const FSlateBrush*> PinBrushes = FPCGModule::GetConstDataTypeRegistry().GetPinIcons(PinType, Pin->Properties, bIsInPin);
 
-			const FSlateBrush* ConnectedBrush = FPCGEditorStyle::Get().GetBrush(PinBrushes.Get<0>());
-			const FSlateBrush* DisconnectedBrush = FPCGEditorStyle::Get().GetBrush(PinBrushes.Get<1>());
+			// If any of the pin brushes are null, error and fallback on default ones
+			if (PinBrushes.Get<0>() == nullptr || PinBrushes.Get<1>() == nullptr)
+			{
+				PinBrushes = FPCGModule::GetConstDataTypeRegistry().GetPinIcons(FPCGDataTypeInfo::AsId(), Pin->Properties, bIsInPin);
+			}
 
-			PinToAdd->SetCustomPinIcon(ConnectedBrush, DisconnectedBrush);
+			PinToAdd->SetCustomPinIcon(PinBrushes.Get<0>(), PinBrushes.Get<1>());
 		}
 	}
 
diff --git a/Engine/Plugins/PCG/Source/PCGEditor/Private/PCGEditorModule.cpp b/Engine/Plugins/PCG/Source/PCGEditor/Private/PCGEditorModule.cpp
--- a/Engine/Plugins/PCG/Source/PCGEditor/Private/PCGEditorModule.cpp
+++ b/Engine/Plugins/PCG/Source/PCGEditor/Private/PCGEditorModule.cpp
@@ -886,7 +886,7 @@ void FPCGEditorModule::RegisterPinColorAndIcons()
 	InRegistry.RegisterPinColorFunction(FPCGDataTypeInfo::AsId(), [](const FPCGDataTypeIdentifier&) { return GetDefault<UPCGEditorSettings>()->DefaultPinColor; });
 
 	// PinIcons. Just have catch all as most of it will be the same icons.
-	InRegistry.RegisterPinIconsFunction(FPCGDataTypeInfo::AsId(), [](const FPCGDataTypeIdentifier& InId, const FPCGPinProperties& InProperties, const bool bIsInput) -> TTuple<FName, FName>
+	InRegistry.RegisterPinIconsFunction(FPCGDataTypeInfo::AsId(), [](const FPCGDataTypeIdentifier& InId, const FPCGPinProperties& InProperties, const bool bIsInput) -> TTuple<const FSlateBrush*, const FSlateBrush*>
 	{
 		const bool bIsMultiData = InProperties.bAllowMultipleData;
 		const bool bIsMultiConnections = InProperties.AllowsMultipleConnections();
@@ -939,7 +939,9 @@ void FPCGEditorModule::RegisterPinColorAndIcons()
 			DisconnectedBrushName = *PinBrushes[DisconnectedIndex];
 		}
 
-		return {ConnectedBrushName, DisconnectedBrushName};
+		const FPCGEditorStyle& EditorStyle = FPCGEditorStyle::Get();
+
+		return {EditorStyle.GetBrush(ConnectedBrushName), EditorStyle.GetBrush(DisconnectedBrushName)};
 	});
 }
 
