[PCG] Fix unwired pins by adding missing redirects after texture type name tweaks.

#rnx

diff --git a/Engine/Plugins/PCG/Config/BasePCG.ini b/Engine/Plugins/PCG/Config/BasePCG.ini
--- a/Engine/Plugins/PCG/Config/BasePCG.ini
+++ b/Engine/Plugins/PCG/Config/BasePCG.ini
@@ -103,6 +103,9 @@
 
 +StructRedirects=(OldName="PCGPointFilterThresholdSettings",NewName="/Script/PCG.PCGAttributeFilterThresholdSettings")
 +StructRedirects=(OldName="PCGActorPropertyOverrideDescription",NewName="/Script/PCG.PCGObjectPropertyOverrideDescription")
++StructRedirects=(OldName="PCGDataTypeInfoBaseTexture",NewName="/Script/PCG.PCGDataTypeInfoBaseTexture2D")
++StructRedirects=(OldName="PCGDataTypeInfoTexture",NewName="/Script/PCG.PCGDataTypeInfoTexture2D")
++StructRedirects=(OldName="PCGDataTypeInfoRenderTarget",NewName="/Script/PCG.PCGDataTypeInfoRenderTarget2D")
 
 +EnumRedirects=(OldName="EPCGDensityNoiseMode",NewName="EPCGAttributeNoiseMode")
 +EnumRedirects=(OldName="EPCGPointFilterOperator",NewName="EPCGAttributeFilterOperator")
