[PCG] Move PCGEditorSubsystem (unreleased, no deprecation needed) to PCGEngineSubsystem to allow headless standalone graph execution
- PCGEngineSubsystem::GenerateGraph now taks FPCGGenerateGraphParams which has an optional completion callback
- Implement a UBlueprintAsyncActionBase to expose PCGEngineSubsystem::GenerateGraph to Blueprint
- GraphInstance now inherit their standalone flag from chosen Graph
- Allow standalone executor to show notifications
- Fix Cancellation of Default Execution Source

diff --git a/Engine/Plugins/Experimental/ProceduralVegetationEditor/Source/ProceduralVegetationEditor/Private/PVEditor.cpp b/Engine/Plugins/Experimental/ProceduralVegetationEditor/Source/ProceduralVegetationEditor/Private/PVEditor.cpp
--- a/Engine/Plugins/Experimental/ProceduralVegetationEditor/Source/ProceduralVegetationEditor/Private/PVEditor.cpp
+++ b/Engine/Plugins/Experimental/ProceduralVegetationEditor/Source/ProceduralVegetationEditor/Private/PVEditor.cpp
@@ -3,8 +3,8 @@
 #include "PVEditor.h"
 #include "EngineAnalytics.h"
 #include "PVEditorCommands.h"
-#include "Editor/PCGEditorDefaultExecutionSource.h"
-#include "Subsystems/PCGEditorSubsystem.h"
+#include "PCGDefaultExecutionSource.h"
+#include "Subsystems/PCGEngineSubsystem.h"
 #include "PVEditorSchema.h"
 #include "ProceduralVegetation.h"
 #include "PVPCGData.h"
@@ -50,18 +50,18 @@ void FPVEditor::Initialize(const EToolkitMode::Type InMode, const TSharedPtr<cla
 
 	AddEditingObject(InProceduralVegetation);
 
-	check(GEditor);
+	check(GEngine);
 	
-	UPCGEditorSubsystem* Subsystem = GEditor->GetEditorSubsystem<UPCGEditorSubsystem>();
+	UPCGEngineSubsystem* Subsystem = GEngine->GetEngineSubsystem<UPCGEngineSubsystem>();
 
 	check(Subsystem);
 	
-	const FPCGEditorDefaultExecutionSourceParams Params
+	const FPCGDefaultExecutionSourceParams Params
 	{
 		.GraphInterface = ProceduralVegetationGraph
 	};
 	
-	ExecutionSource = Subsystem->CreateEditorExecutionSource(Params);
+	ExecutionSource = Subsystem->CreateExecutionSource(Params);
 
 	// Set the stack
 	FPCGStack Stack;
@@ -163,7 +163,7 @@ void FPVEditor::UnregisterTabSpawners(const TSharedRef<class FTabManager>& InTab
 
 IPCGBaseSubsystem* FPVEditor::GetSubsystem() const
 {
-	return GEditor ? GEditor->GetEditorSubsystem<UPCGEditorSubsystem>() : nullptr;
+	return UPCGEngineSubsystem::Get();
 }
 
 TAttribute<FGraphAppearanceInfo> FPVEditor::GetAppearanceInfo() const
diff --git a/Engine/Plugins/Experimental/ProceduralVegetationEditor/Source/ProceduralVegetationEditor/Private/PVEditor.h b/Engine/Plugins/Experimental/ProceduralVegetationEditor/Source/ProceduralVegetationEditor/Private/PVEditor.h
--- a/Engine/Plugins/Experimental/ProceduralVegetationEditor/Source/ProceduralVegetationEditor/Private/PVEditor.h
+++ b/Engine/Plugins/Experimental/ProceduralVegetationEditor/Source/ProceduralVegetationEditor/Private/PVEditor.h
@@ -6,7 +6,7 @@
 
 class SCollectionSpreadSheetWidget;
 struct FManagedArrayCollection;
-class UPCGEditorDefaultExecutionSource;
+class UPCGDefaultExecutionSource;
 class UProceduralVegetation;
 
 #define PROCEDURAL_VEGETATION_DEBUG_SHOW_COLLECTION_TAB !UE_BUILD_SHIPPING_WITH_EDITOR 
@@ -45,7 +45,7 @@ protected:
 	// ~End FPCGEditor interface
 
 	bool bGenerationInProgress = false;
-	TObjectPtr<UPCGEditorDefaultExecutionSource> ExecutionSource;
+	TObjectPtr<UPCGDefaultExecutionSource> ExecutionSource;
 
 private:
 	TObjectPtr<UProceduralVegetation> ProceduralVegetationBeingEdited = nullptr;
diff --git a/Engine/Plugins/PCG/Source/PCG/Private/Editor/PCGEditorDefaultExecutionSource.cpp b/Engine/Plugins/PCG/Source/PCG/Private/Editor/PCGEditorDefaultExecutionSource.cpp
deleted file mode 100644
diff --git a/Engine/Plugins/PCG/Source/PCG/Private/Graph/PCGGraphExecutor.cpp b/Engine/Plugins/PCG/Source/PCG/Private/Graph/PCGGraphExecutor.cpp
--- a/Engine/Plugins/PCG/Source/PCG/Private/Graph/PCGGraphExecutor.cpp
+++ b/Engine/Plugins/PCG/Source/PCG/Private/Graph/PCGGraphExecutor.cpp
@@ -18,6 +18,7 @@
 #include "Helpers/PCGAsync.h"
 #include "Helpers/PCGHelpers.h"
 #include "Metadata/PCGMetadata.h"
+#include "Subsystems/PCGEngineSubsystem.h"
 #include "Utils/PCGGraphExecutionLogging.h"
 
 #include "Algo/AnyOf.h"
@@ -698,15 +699,9 @@ IPCGBaseSubsystem* FPCGGraphExecutor::GetSubsystem() const
 	{
 		return UPCGSubsystem::GetInstance(World);
 	}
-#if WITH_EDITOR
-	else if(IPCGEditorModule* EditorModule = IPCGEditorModule::Get())
-	{
-		return EditorModule->GetEditorSubsystem();
-	}
-#endif
 	else
 	{
-		return nullptr;
+		return UPCGEngineSubsystem::Get();
 	}
 }
 
@@ -2834,8 +2829,14 @@ void FPCGGraphExecutor::NotifyGraphChanged(UPCGGraph* InGraph, EPCGChangeType Ch
 
 void FPCGGraphExecutor::UpdateGenerationNotification()
 {
-	// Avoid Notifications for test executors / commandlets / non simulating game world executors
-	if (!World || IsRunningCommandlet() || (World->IsGameWorld() && GEditor && !GEditor->IsSimulateInEditorInProgress()))
+	// Avoid notifications for commandlets
+	if (IsRunningCommandlet())
+	{
+		return;
+	}
+
+	// Avoid notifications for non world / non simulating game world executors unless they are standalone executors
+	if (!bIsStandaloneGraphExecutor && (!World || (World->IsGameWorld() && GEditor && !GEditor->IsSimulateInEditorInProgress())))
 	{
 		return;
 	}
diff --git a/Engine/Plugins/PCG/Source/PCG/Private/Graph/PCGGraphExecutor.h b/Engine/Plugins/PCG/Source/PCG/Private/Graph/PCGGraphExecutor.h
--- a/Engine/Plugins/PCG/Source/PCG/Private/Graph/PCGGraphExecutor.h
+++ b/Engine/Plugins/PCG/Source/PCG/Private/Graph/PCGGraphExecutor.h
@@ -155,8 +155,9 @@ public:
 	void RemoveReleasedContextForGC(FPCGContext* InContext);
 
 private:
-	friend UPCGSubsystem;
-	friend IPCGBaseSubsystem;
+	friend class UPCGSubsystem;
+	friend class UPCGEngineSubsystem;
+	friend class IPCGBaseSubsystem;
 
 	IPCGBaseSubsystem* GetSubsystem() const;
 
@@ -377,6 +378,7 @@ private:
 	TWeakPtr<IPCGEditorProgressNotification> GenerationProgressNotification;
 	double GenerationProgressNotificationStartTime = 0.0;
 	int32 GenerationProgressLastTaskNum = 0;
+	bool bIsStandaloneGraphExecutor = false;
 
 	/** Set true when a non-runtime-gen component is scheduled and reset to false when all tasks are done executing. */
 	std::atomic<bool> bAnyNonRuntimeGenComponentScheduled = false;
diff --git a/Engine/Plugins/PCG/Source/PCG/Private/Helpers/PCGBlueprintHelpers.cpp b/Engine/Plugins/PCG/Source/PCG/Private/Helpers/PCGBlueprintHelpers.cpp
--- a/Engine/Plugins/PCG/Source/PCG/Private/Helpers/PCGBlueprintHelpers.cpp
+++ b/Engine/Plugins/PCG/Source/PCG/Private/Helpers/PCGBlueprintHelpers.cpp
@@ -6,10 +6,11 @@
 #include "PCGContext.h"
 #include "PCGModule.h"
 #include "PCGPoint.h"
-#include "PCGSubsystem.h"
 #include "Grid/PCGPartitionActor.h"
 #include "Grid/PCGLandscapeCache.h"
 #include "Helpers/PCGHelpers.h"
+#include "Subsystems/PCGEngineSubsystem.h"
+#include "Subsystems/PCGSubsystem.h"
 
 #include "Blueprint/BlueprintExceptionInfo.h"
 #include "Engine/World.h"
@@ -18,6 +19,43 @@
 
 #include UE_INLINE_GENERATED_CPP_BY_NAME(PCGBlueprintHelpers)
 
+UPCGGenerateGraphAsync* UPCGGenerateGraphAsync::GenerateGraphAsync(UPCGGraphInterface* InGraph, int32 InSeed)
+{
+	UPCGGenerateGraphAsync* AsyncCall = NewObject<UPCGGenerateGraphAsync>();
+
+#if WITH_EDITOR
+	AsyncCall->Graph = InGraph;
+	AsyncCall->Seed = InSeed;
+#endif
+
+	return AsyncCall;
+}
+
+void UPCGGenerateGraphAsync::Activate()
+{
+#if WITH_EDITOR
+	FPCGGenerateGraphParams GenerateParams;
+	GenerateParams.Graph = Graph;
+	GenerateParams.Seed = Seed;
+	GenerateParams.GenerationCallback.BindUObject(this, &UPCGGenerateGraphAsync::OnGenerationDone);
+
+	if (UPCGEngineSubsystem* EngineSubsystem = UPCGEngineSubsystem::Get())
+	{
+		EngineSubsystem->GenerateGraph(GenerateParams);
+	}
+#else
+	UE_LOG(LogPCG, Error, TEXT("GenerateGraphAsync not supported in non editor builds."))
+	OnGenerationDone(nullptr, EPCGGenerationStatus::Aborted);
+#endif // WITH_EDITOR
+}
+
+void UPCGGenerateGraphAsync::OnGenerationDone(IPCGGraphExecutionSource* InExecutionSource, EPCGGenerationStatus InStatus)
+{
+	OnCompleted.Broadcast(InStatus);
+
+	SetReadyToDestroy();
+}
+
 void UPCGBlueprintHelpers::ThrowBlueprintException(const FText& ErrorMessage)
 {
 	if (FFrame::GetThreadLocalTopStackFrame() && FFrame::GetThreadLocalTopStackFrame()->Object)
diff --git a/Engine/Plugins/PCG/Source/PCG/Private/PCGDefaultExecutionSource.cpp b/Engine/Plugins/PCG/Source/PCG/Private/PCGDefaultExecutionSource.cpp
new file mode 100644
--- /dev/null
+++ b/Engine/Plugins/PCG/Source/PCG/Private/PCGDefaultExecutionSource.cpp
@@ -0,0 +1,265 @@
+// Copyright Epic Games, Inc. All Rights Reserved.
+
+#include "PCGDefaultExecutionSource.h"
+
+#include "PCGSubgraph.h"
+#include "Subsystems/IPCGBaseSubsystem.h"
+
+#include UE_INLINE_GENERATED_CPP_BY_NAME(PCGDefaultExecutionSource)
+
+UPCGData* FPCGDefaultExecutionState::GetSelfData() const
+{
+	return nullptr;
+}
+
+int32 FPCGDefaultExecutionState::GetSeed() const
+{
+	check(Source);
+	return Source->Seed;
+}
+
+FString FPCGDefaultExecutionState::GetDebugName() const
+{
+	return TEXT("PCGDefaultExecutionSource");
+}
+
+FTransform FPCGDefaultExecutionState::GetTransform() const
+{
+	return FTransform::Identity;
+}
+
+UWorld* FPCGDefaultExecutionState::GetWorld() const
+{
+	return nullptr;
+}
+
+bool FPCGDefaultExecutionState::HasAuthority() const
+{
+	return false;
+}
+
+FBox FPCGDefaultExecutionState::GetBounds() const
+{
+	return FBox(-FVector::One(), FVector::One());
+}
+
+UPCGGraph* FPCGDefaultExecutionState::GetGraph() const
+{
+	check(Source);
+	return Source->GetGraph();
+}
+
+UPCGGraphInstance* FPCGDefaultExecutionState::GetGraphInstance() const
+{
+	check(Source);
+	return Source->GetGraphInstance();
+}
+
+void FPCGDefaultExecutionState::OnGraphExecutionAborted(bool bQuiet, bool bCleanupUnusedResources)
+{
+	Source->CurrentGenerationTask = InvalidPCGTaskId;
+#if WITH_EDITOR
+	if (IPCGBaseSubsystem* Subsystem = GetSubsystem())
+	{
+		Subsystem->OnPCGSourceGenerationDone(Source, EPCGGenerationStatus::Aborted);
+	}
+#endif
+}
+
+void FPCGDefaultExecutionState::Cancel()
+{
+	if (Source->CurrentGenerationTask == InvalidPCGTaskId)
+	{
+		return;
+	}
+
+	if (IPCGBaseSubsystem* Subsystem = GetSubsystem())
+	{
+		Subsystem->CancelGeneration(Source);
+	}
+}
+
+bool FPCGDefaultExecutionState::IsGenerating() const
+{
+	check(Source);
+	return Source->CurrentGenerationTask != InvalidPCGTaskId;
+}
+
+IPCGGraphExecutionSource* FPCGDefaultExecutionState::GetOriginalSource() const
+{
+	// Partitioned generation not supported.
+	return Source;
+}
+
+#if WITH_EDITOR
+const PCGUtils::FExtraCapture& FPCGDefaultExecutionState::GetExtraCapture() const
+{
+	check(Source);
+	return Source->ExtraCapture;
+}
+
+PCGUtils::FExtraCapture& FPCGDefaultExecutionState::GetExtraCapture()
+{
+	check(Source);
+	return Source->ExtraCapture;
+}
+
+const FPCGGraphExecutionInspection& FPCGDefaultExecutionState::GetInspection() const
+{
+	check(Source);
+	return Source->Inspection;
+}
+
+FPCGGraphExecutionInspection& FPCGDefaultExecutionState::GetInspection()
+{
+	check(Source);
+	return Source->Inspection;
+}
+#endif // WITH_EDITOR
+
+void UPCGDefaultExecutionSource::SetGraphInterface(UPCGGraphInterface* InGraphInterface)
+{
+	if (InGraphInterface == GraphInterface)
+	{
+		return;
+	}
+
+#if WITH_EDITOR
+	if (GraphInterface)
+	{
+		GraphInterface->OnGraphChangedDelegate.RemoveAll(this);
+	}
+#endif // WITH_EDITOR
+	
+	GraphInterface = InGraphInterface;
+
+#if WITH_EDITOR
+	if (GraphInterface)
+	{
+		GraphInterface->OnGraphChangedDelegate.AddUObject(this, &UPCGDefaultExecutionSource::OnGraphChanged);
+	}
+#endif // WITH_EDITOR
+}
+
+UPCGGraphInstance* UPCGDefaultExecutionSource::GetGraphInstance()
+{
+	return Cast<UPCGGraphInstance>(GraphInterface);
+}
+
+UPCGGraph* UPCGDefaultExecutionSource::GetGraph()
+{
+	return GraphInterface ? GraphInterface->GetMutablePCGGraph() : nullptr;
+}
+
+#if WITH_EDITOR
+void UPCGDefaultExecutionSource::OnGraphChanged(UPCGGraphInterface* InGraphInterface, EPCGChangeType ChangeType)
+{
+	if (InGraphInterface != GraphInterface)
+	{
+		return;
+	}
+
+	if (ChangeType == EPCGChangeType::Cosmetic ||
+		ChangeType == EPCGChangeType::GraphCustomization ||
+		ChangeType == EPCGChangeType::None)
+	{
+		// If it is a cosmetic change (or no change), nothing to do
+		return;
+	}
+
+	if (!InGraphInterface || !InGraphInterface->GetGraph())
+	{
+		return;
+	}
+
+	Generate();
+}
+#endif // WITH_EDITOR
+
+void UPCGDefaultExecutionSource::Generate()
+{
+	if (CurrentGenerationTask != InvalidPCGTaskId)
+	{
+		return;
+	}
+
+	if (IPCGBaseSubsystem* Subsystem = State.GetSubsystem())
+	{
+		FPCGScheduleGraphParams GraphParams{
+			/*InGraph=*/ GetGraph(),
+			/*InExecutionSource=*/ this,
+			/*InPreGraphElement=*/ MakeShared<FPCGInputForwardingElement>(FPCGDataCollection{}),
+			/*InInputElement=*/ MakeShared<FPCGInputForwardingElement>(FPCGDataCollection{}),
+			/*InExternalDependencies*/ {},
+			/*InFromStack*/ nullptr,
+			/*bInAllowHierarchicalGeneration*/ false};
+		
+		const FPCGTaskId GenerationTaskId = Subsystem->ScheduleGraph(GraphParams);
+		if (GenerationTaskId != InvalidPCGTaskId)
+		{
+			FPCGScheduleGenericParams Params{
+				/*InOperation=*/ [this, WeakSubsystem = TWeakInterfacePtr(Subsystem)](FPCGContext*) -> bool
+				{
+					CurrentGenerationTask = InvalidPCGTaskId;
+					if (WeakSubsystem.IsValid())
+					{
+#if WITH_EDITOR
+						WeakSubsystem->OnPCGSourceGenerationDone(this, EPCGGenerationStatus::Completed);
+#endif // WITH_EDITOR
+					}
+					
+					return true;
+				},
+				/*InExecutionSource=*/this,
+				/*InExecutionDependencies=*/{GenerationTaskId},
+				/*InDataDependencies=*/{},
+				/*bSupportBasePointDataInput=*/true
+			};
+			
+			CurrentGenerationTask = Subsystem->ScheduleGeneric(Params);
+		}
+	}
+}
+
+void UPCGDefaultExecutionSource::Sunset()
+{
+#if WITH_EDITOR
+	if (IsValid(GraphInterface))
+	{
+		check(GraphInterface);
+		GraphInterface->OnGraphChangedDelegate.RemoveAll(this);
+	}
+#endif // WITH_EDITOR
+}
+
+UPCGDefaultExecutionSource::UPCGDefaultExecutionSource()
+	: State(this)
+{
+}
+
+UPCGDefaultExecutionSource::~UPCGDefaultExecutionSource()
+{
+	Sunset();
+}
+
+void UPCGDefaultExecutionSource::Initialize(const FPCGDefaultExecutionSourceParams& InParams)
+{
+	SetGraphInterface(InParams.GraphInterface);
+	SetSeed(InParams.Seed);
+}
+
+void UPCGDefaultExecutionSource::SetSeed(int32 InSeed)
+{
+	Seed = InSeed;
+}
+
+void UPCGDefaultExecutionSource::AddReferencedObjects(UObject* InThis, FReferenceCollector& Collector)
+{
+	UPCGDefaultExecutionSource* This = CastChecked<UPCGDefaultExecutionSource>(InThis);
+	
+#if WITH_EDITOR
+	This->Inspection.AddReferencedObjects(Collector);
+#endif // WITH_EDITOR
+
+	Super::AddReferencedObjects(This, Collector);
+}
diff --git a/Engine/Plugins/PCG/Source/PCG/Private/PCGGraph.cpp b/Engine/Plugins/PCG/Source/PCG/Private/PCGGraph.cpp
--- a/Engine/Plugins/PCG/Source/PCG/Private/PCGGraph.cpp
+++ b/Engine/Plugins/PCG/Source/PCG/Private/PCGGraph.cpp
@@ -2587,6 +2587,8 @@ void UPCGGraphInstance::SetGraph(UPCGGraphInterface* InGraph)
 	Graph = InGraph;
 
 #if WITH_EDITOR
+	bIsStandaloneGraph = InGraph->bIsStandaloneGraph;
+
 	SetupCallbacks();
 #endif // WITH_EDITOR
 
diff --git a/Engine/Plugins/PCG/Source/PCG/Private/PCGGraphExecutionStateInterface.cpp b/Engine/Plugins/PCG/Source/PCG/Private/PCGGraphExecutionStateInterface.cpp
--- a/Engine/Plugins/PCG/Source/PCG/Private/PCGGraphExecutionStateInterface.cpp
+++ b/Engine/Plugins/PCG/Source/PCG/Private/PCGGraphExecutionStateInterface.cpp
@@ -5,6 +5,7 @@
 #include "PCGModule.h"
 #include "Editor/IPCGEditorModule.h"
 #include "Subsystems/PCGSubsystem.h"
+#include "Subsystems/PCGEngineSubsystem.h"
 
 #include UE_INLINE_GENERATED_CPP_BY_NAME(PCGGraphExecutionStateInterface)
 
@@ -21,12 +22,10 @@ IPCGBaseSubsystem* IPCGGraphExecutionState::GetSubsystem() const
 	{
 		Subsystem = UPCGSubsystem::GetInstance(World);
 	}
-#if WITH_EDITOR
-	else if (const IPCGEditorModule* EditorModule = IPCGEditorModule::Get())
+	else
 	{
-		Subsystem = EditorModule->GetEditorSubsystem();
+		Subsystem = UPCGEngineSubsystem::Get();
 	}
-#endif // WITH_EDITOR
 
 	return Subsystem;
 }
diff --git a/Engine/Plugins/PCG/Source/PCG/Private/Subsystems/PCGEngineSubsystem.cpp b/Engine/Plugins/PCG/Source/PCG/Private/Subsystems/PCGEngineSubsystem.cpp
new file mode 100644
--- /dev/null
+++ b/Engine/Plugins/PCG/Source/PCG/Private/Subsystems/PCGEngineSubsystem.cpp
@@ -0,0 +1,148 @@
+// Copyright Epic Games, Inc. All Rights Reserved.
+
+#include "Subsystems/PCGEngineSubsystem.h"
+
+#include "PCGDefaultExecutionSource.h"
+#include "PCGGraph.h"
+#include "PCGModule.h"
+#include "Editor/IPCGEditorModule.h"
+#include "Graph/PCGGraphExecutor.h"
+
+#include "Engine/Engine.h"
+
+#include UE_INLINE_GENERATED_CPP_BY_NAME(PCGEngineSubsystem)
+
+UPCGEngineSubsystem* UPCGEngineSubsystem::Get()
+{
+	return GEngine ? GEngine->GetEngineSubsystem<UPCGEngineSubsystem>() : nullptr;
+}
+
+#if WITH_EDITOR
+
+void UPCGEngineSubsystem::AddReferencedObjects(UObject* InThis, FReferenceCollector& Collector)
+{
+	Super::AddReferencedObjects(InThis, Collector);
+
+	UPCGEngineSubsystem* This = Cast<UPCGEngineSubsystem>(InThis);
+	check(This);
+
+	for (FGraphExecution& GraphExecution : This->GraphExecutions)
+	{
+		Collector.AddReferencedObject(GraphExecution.ExecutionSource);
+	}
+}
+
+void UPCGEngineSubsystem::Initialize(FSubsystemCollectionBase& Collection)
+{
+	Super::Initialize(Collection);
+	IPCGBaseSubsystem::InitializeBaseSubsystem();
+
+	// Show notifications even if executor has no world 
+	// @todo_pcg: for 5.8 have a better GraphExecutor initialization (this new member is private so can be easily removed)
+	if (GraphExecutor)
+	{
+		GraphExecutor->bIsStandaloneGraphExecutor = true;
+	}
+}
+
+void UPCGEngineSubsystem::Deinitialize()
+{
+	IPCGBaseSubsystem::DeinitializeBaseSubsystem();
+	Super::Deinitialize();
+}
+
+UPCGDefaultExecutionSource* UPCGEngineSubsystem::CreateExecutionSource(const FPCGDefaultExecutionSourceParams& InParams)
+{
+	if (!InParams.GraphInterface)
+	{
+		return nullptr;
+	}
+
+	UPCGDefaultExecutionSource* NewSource = NewObject<UPCGDefaultExecutionSource>(this);
+	NewSource->Initialize(InParams);
+	return NewSource;
+}
+
+void UPCGEngineSubsystem::NotifyGraphChanged(UPCGGraph* InGraph, EPCGChangeType ChangeType)
+{
+	if (IPCGEditorModule* PCGEditorModule = IPCGEditorModule::Get())
+	{
+		PCGEditorModule->NotifyGraphChanged(InGraph, ChangeType);
+	}
+
+	IPCGBaseSubsystem::NotifyGraphChanged(InGraph, ChangeType);
+}
+
+void UPCGEngineSubsystem::OnPCGSourceGenerationDone(IPCGGraphExecutionSource* InExecutionSource, EPCGGenerationStatus InStatus)
+{
+	IPCGBaseSubsystem::OnPCGSourceGenerationDone(InExecutionSource, InStatus);
+
+	UPCGDefaultExecutionSource* ExecutionSource = Cast<UPCGDefaultExecutionSource>(InExecutionSource);
+	if (!ExecutionSource)
+	{
+		return;
+	}
+
+	const int32 Index = GraphExecutions.IndexOfByPredicate([ExecutionSource](const FGraphExecution& GraphExecution) { return ExecutionSource == GraphExecution.ExecutionSource.Get(); });
+
+	if (Index != INDEX_NONE)
+	{
+		FGraphExecution& GraphExecution = GraphExecutions[Index];
+		GraphExecution.GenerationParams.GenerationCallback.ExecuteIfBound(ExecutionSource, InStatus);
+		GraphExecutions.RemoveAtSwap(Index);
+
+		if (IPCGEditorModule* PCGEditorModule = IPCGEditorModule::Get())
+		{
+			PCGEditorModule->ClearExecutionMetadata(ExecutionSource);
+		}
+
+		// Mark source as invalid, request GC (non-blocking)
+		ExecutionSource->MarkAsGarbage();
+		GEngine->ForceGarbageCollection(true);
+	}
+}
+
+void UPCGEngineSubsystem::GenerateGraph(const FPCGGenerateGraphParams& InParams)
+{
+	if (!InParams.Graph)
+	{
+		UE_LOG(LogPCG, Error, TEXT("GenerateGraph called without valid graph"));
+		InParams.GenerationCallback.ExecuteIfBound(nullptr, EPCGGenerationStatus::Aborted);
+		return;
+	}
+
+	if (!InParams.Graph->bIsStandaloneGraph)
+	{
+		UE_LOG(LogPCG, Error, TEXT("GenerateGraph called with non standalone graph"));
+		InParams.GenerationCallback.ExecuteIfBound(nullptr, EPCGGenerationStatus::Aborted);
+		return;
+	}
+
+	UPCGDefaultExecutionSource* ExecutionSource = CreateExecutionSource({ InParams.Graph, InParams.Seed });
+	if (ExecutionSource)
+	{
+		GraphExecutions.Add({ ExecutionSource, InParams });
+		ExecutionSource->Generate();
+	}
+}
+
+#endif // WITH_EDITOR
+
+void UPCGEngineSubsystem::Tick(float DeltaTime)
+{
+	IPCGBaseSubsystem::Tick();
+}
+
+ETickableTickType UPCGEngineSubsystem::GetTickableTickType() const
+{
+#if WITH_EDITOR
+	return IsTemplate() ? ETickableTickType::Never : ETickableTickType::Always;
+#else
+	return ETickableTickType::Never;
+#endif
+}
+
+TStatId UPCGEngineSubsystem::GetStatId() const
+{
+	RETURN_QUICK_DECLARE_CYCLE_STAT(FPCGEngineSubsystemTickFunction, STATGROUP_Tickables);
+}
diff --git a/Engine/Plugins/PCG/Source/PCG/Public/Editor/IPCGEditorModule.h b/Engine/Plugins/PCG/Source/PCG/Public/Editor/IPCGEditorModule.h
--- a/Engine/Plugins/PCG/Source/PCG/Public/Editor/IPCGEditorModule.h
+++ b/Engine/Plugins/PCG/Source/PCG/Public/Editor/IPCGEditorModule.h
@@ -26,7 +26,6 @@ public:
 	virtual const FPCGNodeVisualLogs& GetNodeVisualLogs() const = 0;
 	virtual FPCGNodeVisualLogs& GetNodeVisualLogsMutable() = 0;
 	virtual bool CanSelectPartitionActors() const = 0;
-	virtual IPCGBaseSubsystem* GetEditorSubsystem() const = 0;
 #endif
 
 protected:
@@ -36,7 +35,7 @@ protected:
 	friend class SPCGEditorGraphFind;
 	friend class UPCGGraph;
 	friend class UPCGSubsystem;
-	friend class UPCGEditorSubsystem;
+	friend class UPCGEngineSubsystem;
 
 	virtual void OnScheduleGraph(const FPCGStackContext& StackContext) = 0;
 	virtual void OnGraphPreSave(UPCGGraph* Graph, FObjectPreSaveContext ObjectSaveContext) = 0;
diff --git a/Engine/Plugins/PCG/Source/PCG/Public/Editor/PCGEditorDefaultExecutionSource.h b/Engine/Plugins/PCG/Source/PCG/Public/Editor/PCGEditorDefaultExecutionSource.h
deleted file mode 100644
diff --git a/Engine/Plugins/PCG/Source/PCG/Public/Helpers/PCGBlueprintHelpers.h b/Engine/Plugins/PCG/Source/PCG/Public/Helpers/PCGBlueprintHelpers.h
--- a/Engine/Plugins/PCG/Source/PCG/Public/Helpers/PCGBlueprintHelpers.h
+++ b/Engine/Plugins/PCG/Source/PCG/Public/Helpers/PCGBlueprintHelpers.h
@@ -2,20 +2,59 @@
 
 #pragma once
 
+#include "Delegates/Delegate.h"
+#include "Kismet/BlueprintAsyncActionBase.h"
 #include "Kismet/BlueprintFunctionLibrary.h"
 #include "Math/Box.h"
 
 #include "PCGBlueprintHelpers.generated.h"
 
-#define UE_API PCG_API
+class IPCGGraphExecutionSource;
 
 class UPCGComponent;
+class UPCGData;
+class UPCGGraphInterface;
+class UPCGSettings;
+
 struct FPCGContext;
 struct FPCGLandscapeLayerWeight;
 struct FPCGPoint;
 
-class UPCGSettings;
-class UPCGData;
+enum class EPCGGenerationStatus : uint8;
+
+DECLARE_DYNAMIC_MULTICAST_DELEGATE_OneParam(FPCGOnGenerateGraphAsyncCompleted, EPCGGenerationStatus, Status);
+
+UCLASS(MinimalAPI)
+class UPCGGenerateGraphAsync : public UBlueprintAsyncActionBase
+{
+	GENERATED_BODY()
+
+	/**
+	 * Triggers async generation of a PCG Standalone Graph.
+	 * 
+	 * @param Graph Graph to generate.
+	 * @param Seed Seed used to generate graph.
+	 */
+	UFUNCTION(BlueprintCallable, Category="PCG|Helpers", meta = (BlueprintInternalUseOnly = "true", DisplayName = "Generate Graph Async"))
+	static UPCGGenerateGraphAsync* GenerateGraphAsync(UPCGGraphInterface* Graph, int32 Seed = 42);
+
+	/**
+	 * Called when graph generation completes or aborts. 
+	 */
+	UPROPERTY(BlueprintAssignable)
+	FPCGOnGenerateGraphAsyncCompleted OnCompleted;
+
+	virtual void Activate() override;
+
+private:
+	void OnGenerationDone(IPCGGraphExecutionSource* InExecutionSource, EPCGGenerationStatus InStatus);
+
+	UPROPERTY()
+	TObjectPtr<UPCGGraphInterface> Graph = nullptr;
+
+	UPROPERTY()
+	int32 Seed = 42;
+};
 
 UCLASS(MinimalAPI)
 class UPCGBlueprintHelpers : public UBlueprintFunctionLibrary
@@ -23,82 +62,80 @@ class UPCGBlueprintHelpers : public UBlueprintFunctionLibrary
 	GENERATED_BODY()
 
 public:
-	static UE_API void ThrowBlueprintException(const FText& ErrorMessage);
+	static PCG_API void ThrowBlueprintException(const FText& ErrorMessage);
 
 	UFUNCTION(BlueprintCallable, Category = "PCG|Helpers")
-	static UE_API int ComputeSeedFromPosition(const FVector& InPosition);
+	static PCG_API int ComputeSeedFromPosition(const FVector& InPosition);
 
 	UFUNCTION(BlueprintCallable, Category = "PCG|Helpers")
-	static UE_API void SetSeedFromPosition(UPARAM(ref) FPCGPoint& InPoint);
+	static PCG_API void SetSeedFromPosition(UPARAM(ref) FPCGPoint& InPoint);
 
 	/** Creates a random stream from a point's seed and settings/component's seed (optional) */
 	UFUNCTION(BlueprintCallable, Category = "PCG|Helpers", meta = (ScriptMethod))
-	static UE_API FRandomStream GetRandomStreamFromPoint(const FPCGPoint& InPoint, const UPCGSettings* OptionalSettings = nullptr, const UPCGComponent* OptionalComponent = nullptr);
+	static PCG_API FRandomStream GetRandomStreamFromPoint(const FPCGPoint& InPoint, const UPCGSettings* OptionalSettings = nullptr, const UPCGComponent* OptionalComponent = nullptr);
 
 	/** Creates a random stream from using the random seeds from two points, as well as settings/component's seed (optional) */
 	UFUNCTION(BlueprintCallable, Category = "PCG|Helpers", meta = (ScriptMethod))
-	static UE_API FRandomStream GetRandomStreamFromTwoPoints(const FPCGPoint& InPointA, const FPCGPoint& InPointB, const UPCGSettings* OptionalSettings = nullptr, const UPCGComponent* OptionalComponent = nullptr);
+	static PCG_API FRandomStream GetRandomStreamFromTwoPoints(const FPCGPoint& InPointA, const FPCGPoint& InPointB, const UPCGSettings* OptionalSettings = nullptr, const UPCGComponent* OptionalComponent = nullptr);
 
 	UFUNCTION(BlueprintCallable, Category = "PCG|Helpers", meta = (ScriptMethod))
-	static UE_API const UPCGSettings* GetSettings(UPARAM(ref) FPCGContext& Context);
+	static PCG_API const UPCGSettings* GetSettings(UPARAM(ref) FPCGContext& Context);
 
 	UFUNCTION(BlueprintCallable, Category = "PCG|Temporary", meta = (ScriptMethod))
-	static UE_API UPCGData* GetActorData(UPARAM(ref) FPCGContext& Context);
+	static PCG_API UPCGData* GetActorData(UPARAM(ref) FPCGContext& Context);
 
 	UFUNCTION(BlueprintCallable, Category = "PCG|Temporary", meta = (ScriptMethod))
-	static UE_API UPCGData* GetInputData(UPARAM(ref) FPCGContext& Context);
+	static PCG_API UPCGData* GetInputData(UPARAM(ref) FPCGContext& Context);
 
 	UFUNCTION(BlueprintCallable, Category = "PCG|Temporary", meta = (ScriptMethod))
-	static UE_API UPCGComponent* GetComponent(UPARAM(ref) FPCGContext& Context);
+	static PCG_API UPCGComponent* GetComponent(UPARAM(ref) FPCGContext& Context);
 
 	UFUNCTION(BlueprintCallable, Category = "PCG|Temporary", meta = (ScriptMethod))
-	static UE_API UPCGComponent* GetOriginalComponent(UPARAM(ref) FPCGContext& Context);
+	static PCG_API UPCGComponent* GetOriginalComponent(UPARAM(ref) FPCGContext& Context);
 
 	UFUNCTION(BlueprintCallable, Category = "PCG", meta = (ScriptMethod))
-	static UE_API AActor* GetTargetActor(UPARAM(ref) FPCGContext& Context, UPCGSpatialData* SpatialData);
+	static PCG_API AActor* GetTargetActor(UPARAM(ref) FPCGContext& Context, UPCGSpatialData* SpatialData);
 
 	UFUNCTION(BlueprintCallable, Category = "PCG|Points", meta = (ScriptMethod))
-	static UE_API void SetExtents(UPARAM(ref) FPCGPoint& InPoint, const FVector& InExtents);
+	static PCG_API void SetExtents(UPARAM(ref) FPCGPoint& InPoint, const FVector& InExtents);
 
 	UFUNCTION(BlueprintCallable, Category = "PCG|Points", meta = (ScriptMethod))
-	static UE_API FVector GetExtents(const FPCGPoint& InPoint);
+	static PCG_API FVector GetExtents(const FPCGPoint& InPoint);
 
 	UFUNCTION(BlueprintCallable, Category = "PCG|Points", meta = (ScriptMethod))
-	static UE_API void SetLocalCenter(UPARAM(ref) FPCGPoint& InPoint, const FVector& InLocalCenter);
+	static PCG_API void SetLocalCenter(UPARAM(ref) FPCGPoint& InPoint, const FVector& InLocalCenter);
 
 	UFUNCTION(BlueprintCallable, Category = "PCG|Points", meta = (ScriptMethod))
-	static UE_API FVector GetLocalCenter(const FPCGPoint& InPoint);
+	static PCG_API FVector GetLocalCenter(const FPCGPoint& InPoint);
 
 	UFUNCTION(BlueprintCallable, Category = "PCG|Points", meta = (ScriptMethod))
-	static UE_API FBox GetTransformedBounds(const FPCGPoint& InPoint);
+	static PCG_API FBox GetTransformedBounds(const FPCGPoint& InPoint);
 
 	UFUNCTION(BlueprintCallable, Category = "PCG|Helpers", meta = (ScriptMethod))
-	static UE_API FBox GetActorBoundsPCG(AActor* InActor, bool bIgnorePCGCreatedComponents = true);
+	static PCG_API FBox GetActorBoundsPCG(AActor* InActor, bool bIgnorePCGCreatedComponents = true);
 
 	UFUNCTION(BlueprintCallable, Category = "PCG|Helpers", meta = (ScriptMethod))
-	static UE_API FBox GetActorLocalBoundsPCG(AActor* InActor, bool bIgnorePCGCreatedComponents = true);
+	static PCG_API FBox GetActorLocalBoundsPCG(AActor* InActor, bool bIgnorePCGCreatedComponents = true);
 
 	UFUNCTION(BlueprintCallable, Category = "PCG|Helpers", meta = (ScriptMethod))
-	static UE_API UPCGData* CreatePCGDataFromActor(AActor* InActor, bool bParseActor = true);
+	static PCG_API UPCGData* CreatePCGDataFromActor(AActor* InActor, bool bParseActor = true);
 
 	UFUNCTION(BlueprintCallable, Category = "PCG|Helpers", meta = (WorldContext="WorldContextObject"))
-	static UE_API TArray<FPCGLandscapeLayerWeight> GetInterpolatedPCGLandscapeLayerWeights(UObject* WorldContextObject, const FVector& Location);
+	static PCG_API TArray<FPCGLandscapeLayerWeight> GetInterpolatedPCGLandscapeLayerWeights(UObject* WorldContextObject, const FVector& Location);
 
 	UFUNCTION(BLueprintCallable, Category = "PCG|Helpers", meta = (ScriptMethod))
-	static UE_API int64 GetTaskId(UPARAM(ref) FPCGContext& Context);
+	static PCG_API int64 GetTaskId(UPARAM(ref) FPCGContext& Context);
 
 	/** Flush the cache, to be used if you have changed something PCG depends on at runtime. Same as `pcg.FlushCache` command. Returns true if it succeeded. */
 	UFUNCTION(BlueprintCallable, Category = "PCG", meta=(DisplayName = "Flush PCG Cache"))
-	static UE_API bool FlushPCGCache();
+	static PCG_API bool FlushPCGCache();
 
 	/** Refresh a component set to Generate At Runtime, if some parameters changed. Can also flush the cache. */
 	UFUNCTION(BlueprintCallable, Category = "PCG|Runtime", meta = (ScriptMethod, DisplayName = "Refresh PCG Runtime Component"))
-	static UE_API void RefreshPCGRuntimeComponent(UPCGComponent* InComponent, const bool bFlushCache = false);
+	static PCG_API void RefreshPCGRuntimeComponent(UPCGComponent* InComponent, const bool bFlushCache = false);
 
 	// Implementation note: Needs to be done outside of UPCGData because of circular dependency between PCGContext.h and PCGData.h
 	/** Return a copy of the data, with Metadata inheritance for spatial data. */
 	UFUNCTION(BlueprintCallable, Category="PCG|Data", meta = (ScriptMethod))
-	static UE_API UPCGData* DuplicateData(const UPCGData* InData, UPARAM(ref) FPCGContext& Context, bool bInitializeMetadata = true);
+	static PCG_API UPCGData* DuplicateData(const UPCGData* InData, UPARAM(ref) FPCGContext& Context, bool bInitializeMetadata = true);
 };
-
-#undef UE_API
diff --git a/Engine/Plugins/PCG/Source/PCG/Public/PCGDefaultExecutionSource.h b/Engine/Plugins/PCG/Source/PCG/Public/PCGDefaultExecutionSource.h
new file mode 100644
--- /dev/null
+++ b/Engine/Plugins/PCG/Source/PCG/Public/PCGDefaultExecutionSource.h
@@ -0,0 +1,108 @@
+// Copyright Epic Games, Inc. All Rights Reserved.
+
+#pragma once
+
+#include "PCGGraphExecutionInspection.h"
+#include "PCGGraphExecutionStateInterface.h"
+
+#include "PCGDefaultExecutionSource.generated.h"
+
+class UPCGDefaultExecutionSource;
+class UPCGGraph;
+class UPCGGraphInstance;
+class UPCGGraphInterface;
+
+class FPCGDefaultExecutionState : public IPCGGraphExecutionState
+{
+public:
+	FPCGDefaultExecutionState() = default;
+	explicit FPCGDefaultExecutionState(UPCGDefaultExecutionSource* InSource) : Source(InSource){}
+	
+	PCG_API virtual UPCGData* GetSelfData() const override;
+	PCG_API virtual int32 GetSeed() const override;
+	PCG_API virtual FString GetDebugName() const override;
+	PCG_API virtual FTransform GetTransform() const override;
+	PCG_API virtual UWorld* GetWorld() const override;
+	PCG_API virtual bool HasAuthority() const override;
+	PCG_API virtual FBox GetBounds() const override;
+	PCG_API virtual UPCGGraph* GetGraph() const override;
+	PCG_API virtual UPCGGraphInstance* GetGraphInstance() const override;
+	PCG_API virtual void OnGraphExecutionAborted(bool bQuiet, bool bCleanupUnusedResources) override;
+	PCG_API virtual void Cancel() override;
+	PCG_API virtual bool IsGenerating() const override;
+	PCG_API virtual IPCGGraphExecutionSource* GetOriginalSource() const override;
+
+#if WITH_EDITOR
+	PCG_API virtual const PCGUtils::FExtraCapture& GetExtraCapture() const override;
+	PCG_API virtual PCGUtils::FExtraCapture& GetExtraCapture() override;
+
+	PCG_API virtual const FPCGGraphExecutionInspection& GetInspection() const override;
+	PCG_API virtual FPCGGraphExecutionInspection& GetInspection() override;
+
+	// No dynamic tracking for editor resources (at the moment)
+	virtual void RegisterDynamicTracking(const UPCGSettings* InSettings, const TArrayView<TPair<FPCGSelectionKey, bool>>& InDynamicKeysAndCulling) override {};
+	virtual void RegisterDynamicTracking(const FPCGSelectionKeyToSettingsMap& InKeysToSettings) override {};
+	
+	virtual bool IsRefreshInProgress() const override {return IsGenerating();}
+
+	virtual FPCGDynamicTrackingPriority GetDynamicTrackingPriority() const override { return FPCGDynamicTrackingPriority(); }
+#endif // WITH_EDITOR
+
+private:
+	UPCGDefaultExecutionSource* Source = nullptr;
+};
+
+struct FPCGDefaultExecutionSourceParams
+{
+	UPCGGraphInterface* GraphInterface = nullptr;
+	int32 Seed = 42;
+};
+
+UCLASS(MinimalAPI)
+class UPCGDefaultExecutionSource : public UObject, public IPCGGraphExecutionSource
+{
+	friend FPCGDefaultExecutionState;
+	
+	GENERATED_BODY()
+	
+public:
+	PCG_API UPCGDefaultExecutionSource();
+	PCG_API ~UPCGDefaultExecutionSource();
+
+	PCG_API void Initialize(const FPCGDefaultExecutionSourceParams& InParams);
+
+	PCG_API static void AddReferencedObjects(UObject* InThis, FReferenceCollector& Collector);
+	
+	virtual IPCGGraphExecutionState& GetExecutionState() override { return State; }
+	virtual const IPCGGraphExecutionState& GetExecutionState() const override { return State; }
+
+	PCG_API void SetGraphInterface(UPCGGraphInterface* InGraphInterface);
+
+	UPCGGraphInterface* GetGraphInterface() { return GraphInterface.Get(); }
+	PCG_API UPCGGraphInstance* GetGraphInstance();
+	PCG_API UPCGGraph* GetGraph();
+
+	PCG_API void SetSeed(int32 InSeed);
+
+#if WITH_EDITOR
+	PCG_API void OnGraphChanged(UPCGGraphInterface* GraphInterface, EPCGChangeType ChangeType);
+#endif // WITH_EDITOR
+	
+	PCG_API void Generate();
+
+	PCG_API void Sunset();
+
+private:
+	FPCGTaskId CurrentGenerationTask = InvalidPCGTaskId;
+	
+	UPROPERTY()
+	TObjectPtr<UPCGGraphInterface> GraphInterface;
+	
+	FPCGDefaultExecutionState State;
+	int32 Seed = 42;
+
+#if WITH_EDITOR
+	PCGUtils::FExtraCapture ExtraCapture;
+	FPCGGraphExecutionInspection Inspection;
+#endif // WITH_EDITOR
+};
\ No newline at end of file
diff --git a/Engine/Plugins/PCG/Source/PCG/Public/Subsystems/PCGEngineSubsystem.h b/Engine/Plugins/PCG/Source/PCG/Public/Subsystems/PCGEngineSubsystem.h
new file mode 100644
--- /dev/null
+++ b/Engine/Plugins/PCG/Source/PCG/Public/Subsystems/PCGEngineSubsystem.h
@@ -0,0 +1,74 @@
+// Copyright Epic Games, Inc. All Rights Reserved.
+
+#pragma once
+
+#include "Tickable.h"
+
+#include "Subsystems/EngineSubsystem.h"
+#include "Subsystems/IPCGBaseSubsystem.h"
+
+#include "PCGEngineSubsystem.generated.h"
+
+struct FPCGDefaultExecutionSourceParams;
+class UPCGDefaultExecutionSource;
+class UPCGGraphInterface;
+
+DECLARE_DELEGATE_TwoParams(FPCGOnEditorGenerationDone, IPCGGraphExecutionSource*, EPCGGenerationStatus);
+
+struct FPCGGenerateGraphParams
+{
+	/** Graph to generate. */
+	TObjectPtr<UPCGGraphInterface> Graph = nullptr;
+
+	/** Generation seed. */
+	int32 Seed = 42;
+
+	/** Optional callback when generation completes. */
+	FPCGOnEditorGenerationDone GenerationCallback;
+};
+
+/**
+ * Engine subsystem that allows running graph generations in Editor and headless configurations. 
+ */
+UCLASS(MinimalAPI)
+class UPCGEngineSubsystem : public UEngineSubsystem, public IPCGBaseSubsystem, public FTickableGameObject
+{
+	GENERATED_BODY()
+	
+public:
+
+	PCG_API static UPCGEngineSubsystem* Get();
+
+#if WITH_EDITOR
+	PCG_API static void AddReferencedObjects(UObject* InThis, FReferenceCollector& Collector);
+
+	PCG_API virtual void Initialize(FSubsystemCollectionBase& Collection) override;
+	PCG_API virtual void Deinitialize() override;
+
+	PCG_API UPCGDefaultExecutionSource* CreateExecutionSource(const FPCGDefaultExecutionSourceParams& InParams);
+
+	PCG_API void GenerateGraph(const FPCGGenerateGraphParams& InParams);
+
+	PCG_API virtual void NotifyGraphChanged(UPCGGraph* InGraph, EPCGChangeType ChangeType) override;
+
+	PCG_API virtual void OnPCGSourceGenerationDone(IPCGGraphExecutionSource* InExecutionSource, EPCGGenerationStatus InStatus) override;
+#endif // WITH_EDITOR
+
+protected:
+	//** FTickableGameObject Interface
+	PCG_API virtual void Tick(float DeltaTime) override;
+	PCG_API virtual ETickableTickType GetTickableTickType() const override;
+	PCG_API virtual TStatId GetStatId() const override;
+	virtual bool IsTickableInEditor() const { return true; }
+	//~ END FTickableGameObject Interface
+
+#if WITH_EDITOR
+	struct FGraphExecution
+	{
+		TObjectPtr<UPCGDefaultExecutionSource> ExecutionSource;
+		FPCGGenerateGraphParams GenerationParams;
+	};
+
+	TArray<FGraphExecution> GraphExecutions;
+#endif // WITH_EDITOR
+};
\ No newline at end of file
diff --git a/Engine/Plugins/PCG/Source/PCGEditor/Private/AssetDefinitions/AssetDefinition_PCGGraphInterface.cpp b/Engine/Plugins/PCG/Source/PCGEditor/Private/AssetDefinitions/AssetDefinition_PCGGraphInterface.cpp
--- a/Engine/Plugins/PCG/Source/PCGEditor/Private/AssetDefinitions/AssetDefinition_PCGGraphInterface.cpp
+++ b/Engine/Plugins/PCG/Source/PCGEditor/Private/AssetDefinitions/AssetDefinition_PCGGraphInterface.cpp
@@ -4,11 +4,11 @@
 
 #include "PCGGraph.h"
 
+#include "PCGDefaultExecutionSource.h"
 #include "PCGEditorCommon.h"
 #include "PCGEditorStyle.h"
 #include "PCGGraphFactory.h"
-#include "Editor/PCGEditorDefaultExecutionSource.h"
-#include "Subsystems/PCGEditorSubsystem.h"
+#include "Subsystems/PCGEngineSubsystem.h"
 
 #include "ContentBrowserMenuContexts.h"
 #include "Editor.h"
@@ -107,7 +107,7 @@ namespace MenuExtension_PCGGraphInterface
 					{
 						const UContentBrowserAssetContextMenuContext* CBContext = UContentBrowserAssetContextMenuContext::FindContextWithAssets(InContext);
 
-						if (UPCGEditorSubsystem* EditorSubsystem = GEditor->GetEditorSubsystem<UPCGEditorSubsystem>())
+						if (UPCGEngineSubsystem* EngineSubsystem = UPCGEngineSubsystem::Get())
 						{
 							for (const FAssetData& AssetData : CBContext->SelectedAssets)
 							{
@@ -121,7 +121,7 @@ namespace MenuExtension_PCGGraphInterface
 
 								if (UPCGGraphInterface* GraphInterface = Cast<UPCGGraphInterface>(AssetData.GetAsset()))
 								{
-									EditorSubsystem->GenerateGraph({ GraphInterface });
+									EngineSubsystem->GenerateGraph({ GraphInterface });
 								}
 							}
 						}
diff --git a/Engine/Plugins/PCG/Source/PCGEditor/Private/PCGEditor.cpp b/Engine/Plugins/PCG/Source/PCGEditor/Private/PCGEditor.cpp
--- a/Engine/Plugins/PCG/Source/PCGEditor/Private/PCGEditor.cpp
+++ b/Engine/Plugins/PCG/Source/PCGEditor/Private/PCGEditor.cpp
@@ -3,6 +3,7 @@
 #include "PCGEditor.h"
 
 #include "PCGComponent.h"
+#include "PCGDefaultExecutionSource.h"
 #include "PCGEdge.h"
 #include "PCGEditorModule.h"
 #include "PCGGraph.h"
@@ -10,10 +11,10 @@
 #include "PCGPin.h"
 #include "PCGSubsystem.h"
 #include "Editor/IPCGEditorModule.h"
-#include "Editor/PCGEditorDefaultExecutionSource.h"
 #include "Elements/PCGReroute.h"
 #include "Helpers/PCGSubgraphHelpers.h"
 #include "Rendering/SlateRenderer.h"
+#include "Subsystems/PCGEngineSubsystem.h"
 #include "Tests/Determinism/PCGDeterminismNativeTests.h"
 #include "Tests/Determinism/PCGDeterminismTestBlueprintBase.h"
 #include "Tests/Determinism/PCGDeterminismTestsCommon.h"
@@ -72,7 +73,6 @@
 #include "Misc/MessageDialog.h"
 #include "Misc/TransactionObjectEvent.h"
 #include "Preferences/UnrealEdOptions.h"
-#include "Subsystems/PCGEditorSubsystem.h"
 #include "Widgets/Docking/SDockTab.h"
 #include "Widgets/Input/SEditableTextBox.h"
 #include "Widgets/Notifications/SNotificationList.h"
@@ -227,17 +227,17 @@ void FPCGEditor::Initialize(const EToolkitMode::Type InMode, const TSharedPtr<cl
 	FEditorDelegates::PostPIEStarted.AddRaw(this, &FPCGEditor::OnPostPIEStarted);
 	FEditorDelegates::EndPIE.AddRaw(this, &FPCGEditor::OnEndPIE);
 
+	if (UPCGEngineSubsystem* EngineSubsystem = UPCGEngineSubsystem::Get())
+	{
+		EngineSubsystem->GetOnPCGSourceGenerationDone().AddRaw(this, &FPCGEditor::OnSourceGenerationDone);
+	}
+
 	if (GEditor)
 	{
 		RegisterDelegatesForWorld(GEditor->GetEditorWorldContext().World());
 
 		// In case the editor is opened while in PIE, we should try setting up callbacks for the PIE world.
 		RegisterDelegatesForWorld(GEditor->PlayWorld.Get());
-
-		if (UPCGEditorSubsystem* EditorSubsystem = GEditor->GetEditorSubsystem<UPCGEditorSubsystem>())
-		{
-			EditorSubsystem->GetOnPCGSourceGenerationDone().AddRaw(this, &FPCGEditor::OnSourceGenerationDone);
-		}
 	}
 
 	// Clear inspection flag on all nodes.
@@ -2918,13 +2918,13 @@ void FPCGEditor::OnClose()
 		}
 	}
 
-	if (GEditor)
+	if (UPCGEngineSubsystem* EngineSubsystem = UPCGEngineSubsystem::Get())
 	{
-		if (UPCGEditorSubsystem* EditorSubsystem = GEditor->GetEditorSubsystem<UPCGEditorSubsystem>())
-		{
-			EditorSubsystem->GetOnPCGSourceGenerationDone().RemoveAll(this);
-		}
+		EngineSubsystem->GetOnPCGSourceGenerationDone().RemoveAll(this);
+	}
 
+	if (GEditor)
+	{
 		UnregisterDelegatesForWorld(GEditor->GetEditorWorldContext().World());
 		UnregisterDelegatesForWorld(GEditor->PlayWorld.Get());
 	}
@@ -3321,9 +3321,9 @@ void FPCGEditor::UpdateDefaultExecutionSource()
 {
 	if (PCGGraphBeingEdited->bIsStandaloneGraph && !PCGDefaultExecutionSource)
 	{
-		if (UPCGEditorSubsystem* EditorSubsystem = GEditor->GetEditorSubsystem<UPCGEditorSubsystem>())
+		if (UPCGEngineSubsystem* EngineSubsystem = UPCGEngineSubsystem::Get())
 		{
-			PCGDefaultExecutionSource = EditorSubsystem->CreateEditorExecutionSource({ PCGGraphBeingEdited });
+			PCGDefaultExecutionSource = EngineSubsystem->CreateExecutionSource({ PCGGraphBeingEdited });
 			RefreshViews();
 		}
 	}
diff --git a/Engine/Plugins/PCG/Source/PCGEditor/Private/PCGEditorModule.cpp b/Engine/Plugins/PCG/Source/PCGEditor/Private/PCGEditorModule.cpp
--- a/Engine/Plugins/PCG/Source/PCGEditor/Private/PCGEditorModule.cpp
+++ b/Engine/Plugins/PCG/Source/PCGEditor/Private/PCGEditorModule.cpp
@@ -56,7 +56,6 @@
 #include "Details/PCGInteractiveToolSettingsDetails.h"
 #include "Nodes/PCGEditorGraphNodeFactory.h"
 #include "Schema/PCGEditorGraphSchema.h"
-#include "Subsystems/PCGEditorSubsystem.h"
 
 #include "ContentBrowserMenuContexts.h"
 #include "ContentBrowserModule.h"
@@ -289,11 +288,6 @@ bool FPCGEditorModule::CanSelectPartitionActors() const
 	return GetDefault<UPCGEditorSettings>()->bCanSelectPartitionActors;
 }
 
-IPCGBaseSubsystem* FPCGEditorModule::GetEditorSubsystem() const
-{
-	return GEditor ? GEditor->GetEditorSubsystem<UPCGEditorSubsystem>() : nullptr;
-}
-
 void FPCGEditorModule::RegisterOnEditorModeChange()
 {
 	// Have a callback that catches changes in the Editor modes, to catch when we exit the landscape edit mode.
diff --git a/Engine/Plugins/PCG/Source/PCGEditor/Private/Subsystems/PCGEditorSubsystem.cpp b/Engine/Plugins/PCG/Source/PCGEditor/Private/Subsystems/PCGEditorSubsystem.cpp
deleted file mode 100644
diff --git a/Engine/Plugins/PCG/Source/PCGEditor/Private/Widgets/SPCGEditorGraphDebugObjectTree.cpp b/Engine/Plugins/PCG/Source/PCGEditor/Private/Widgets/SPCGEditorGraphDebugObjectTree.cpp
--- a/Engine/Plugins/PCG/Source/PCGEditor/Private/Widgets/SPCGEditorGraphDebugObjectTree.cpp
+++ b/Engine/Plugins/PCG/Source/PCGEditor/Private/Widgets/SPCGEditorGraphDebugObjectTree.cpp
@@ -3,10 +3,10 @@
 #include "SPCGEditorGraphDebugObjectTree.h"
 
 #include "PCGComponent.h"
+#include "PCGDefaultExecutionSource.h"
 #include "PCGGraph.h"
 #include "PCGSubsystem.h"
 #include "Editor/IPCGEditorModule.h"
-#include "Editor/PCGEditorDefaultExecutionSource.h"
 #include "Grid/PCGPartitionActor.h"
 #include "Helpers/PCGHelpers.h"
 
diff --git a/Engine/Plugins/PCG/Source/PCGEditor/Public/PCGEditor.h b/Engine/Plugins/PCG/Source/PCGEditor/Public/PCGEditor.h
--- a/Engine/Plugins/PCG/Source/PCGEditor/Public/PCGEditor.h
+++ b/Engine/Plugins/PCG/Source/PCGEditor/Public/PCGEditor.h
@@ -32,7 +32,7 @@ class SPCGEditorGraphUserParametersView;
 class SPCGEditorViewport;
 class UEdGraphNode;
 class UPCGComponent;
-class UPCGEditorDefaultExecutionSource;
+class UPCGDefaultExecutionSource;
 class UPCGEditorGraph;
 class UPCGEditorGraphNodeBase;
 class UPCGGraph;
@@ -199,7 +199,7 @@ public:
 	// Can override the schema used for this editor. By default, it's UPCGEditorSchema.
 	PCGEDITOR_API virtual TSubclassOf<UPCGEditorGraphSchema> GetSchemaClass() const;
 
-	UPCGEditorDefaultExecutionSource* GetDefaultExecutionSource() const { return PCGDefaultExecutionSource; }
+	UPCGDefaultExecutionSource* GetDefaultExecutionSource() const { return PCGDefaultExecutionSource; }
 
 protected:
 	PCGEDITOR_API virtual TAttribute<FGraphAppearanceInfo> GetAppearanceInfo() const;
@@ -511,7 +511,7 @@ private:
 	TSharedPtr<FUICommandList> GraphEditorCommands;
 
 	TObjectPtr<UPCGGraph> PCGGraphBeingEdited = nullptr;
-	TObjectPtr<UPCGEditorDefaultExecutionSource> PCGDefaultExecutionSource = nullptr;
+	TObjectPtr<UPCGDefaultExecutionSource> PCGDefaultExecutionSource = nullptr;
 	UPCGEditorGraph* PCGEditorGraph = nullptr;
 
 	FPCGEditorInspectionDataManager InspectionDataManager;
diff --git a/Engine/Plugins/PCG/Source/PCGEditor/Public/PCGEditorModule.h b/Engine/Plugins/PCG/Source/PCGEditor/Public/PCGEditorModule.h
--- a/Engine/Plugins/PCG/Source/PCGEditor/Public/PCGEditorModule.h
+++ b/Engine/Plugins/PCG/Source/PCGEditor/Public/PCGEditorModule.h
@@ -37,7 +37,6 @@ public:
 	virtual const FPCGNodeVisualLogs& GetNodeVisualLogs() const override { return NodeVisualLogs; }
 	virtual FPCGNodeVisualLogs& GetNodeVisualLogsMutable() override { return NodeVisualLogs; }
 	virtual bool CanSelectPartitionActors() const override;
-	virtual IPCGBaseSubsystem* GetEditorSubsystem() const override;
 
 protected:
 	virtual void OnScheduleGraph(const FPCGStackContext& StackContext) override;
diff --git a/Engine/Plugins/PCG/Source/PCGEditor/Public/Subsystems/PCGEditorSubsystem.h b/Engine/Plugins/PCG/Source/PCGEditor/Public/Subsystems/PCGEditorSubsystem.h
deleted file mode 100644
