[PCG] Add Generate Landscape Textures node that can generate both grass maps and a height texture.

Refactor existing GGM node to be special case of this.

This new landscape height data path avoids the game thread cost of the CPU path, and avoids the residency/latency issues of the VT path.

[RN] Added Generate Landscape Textures node that can generate a height map as well as grass maps from the landscape.

diff --git a/Engine/Plugins/PCG/Config/BasePCG.ini b/Engine/Plugins/PCG/Config/BasePCG.ini
--- a/Engine/Plugins/PCG/Config/BasePCG.ini
+++ b/Engine/Plugins/PCG/Config/BasePCG.ini
@@ -97,9 +97,10 @@
 +ClassRedirects=(OldName="PCGSortPointsSettings",NewName="/Script/PCG.PCGSortAttributesSettings")
 +ClassRedirects=(OldName="PCGMetadataOperationSettings",NewName="/Script/PCG.PCGCopyAttributesSettings")
 +ClassRedirects=(OldName="PCGCreateSplineMeshSettings",NewName="/Script/PCG.PCGSpawnSplineMeshSettings")
++ClassRedirects=(OldName="PCGGenerateGrassMapsSettings",NewName="/Script/PCG.PCGGenerateLandscapeTexturesSettings")
 +ClassRedirects=(OldName="/Script/PCGGeometryScriptInterop.PCGDynamicMeshData",NewName="/Script/PCG.PCGDynamicMeshData")
-+ClassRedirects = (OldName="/Script/PCG.PCGNumberOfEntriesSettings",NewName="/Script/PCG.PCGNumberOfElementsSettings")
-+ClassRedirects = (OldName="/Script/PCG.PCGNumberOfPointsSettings",NewName="/Script/PCG.PCGNumberOfElementsSettings")
++ClassRedirects=(OldName="/Script/PCG.PCGNumberOfEntriesSettings",NewName="/Script/PCG.PCGNumberOfElementsSettings")
++ClassRedirects=(OldName="/Script/PCG.PCGNumberOfPointsSettings",NewName="/Script/PCG.PCGNumberOfElementsSettings")
 
 +StructRedirects=(OldName="PCGPointFilterThresholdSettings",NewName="/Script/PCG.PCGAttributeFilterThresholdSettings")
 +StructRedirects=(OldName="PCGActorPropertyOverrideDescription",NewName="/Script/PCG.PCGObjectPropertyOverrideDescription")
diff --git a/Engine/Plugins/PCG/Shaders/Private/PCGGrassMapUnpackerCS.usf b/Engine/Plugins/PCG/Shaders/Private/PCGGrassMapUnpackerCS.usf
--- a/Engine/Plugins/PCG/Shaders/Private/PCGGrassMapUnpackerCS.usf
+++ b/Engine/Plugins/PCG/Shaders/Private/PCGGrassMapUnpackerCS.usf
@@ -2,6 +2,7 @@
 
 #include "/Engine/Private/Common.ush"
 #include "/Engine/Public/Platform.ush"
+#include "/Engine/Private/Landscape/LandscapeCommon.ush"
 
 // Input:
 // First 2 channels of first texture always reserved for heightmap (see constructor of FLandscapeGrassWeightExporter_RenderThread::FComponentInfo).
@@ -17,28 +18,34 @@
 //  -----------------------------------------
 Texture2D<float4> InPackedGrassMaps;
 
-// Output:
-// Each slice is a single grass combined grass map. Dimension is (NumComponentsX * ResX, NumComponentsY * ResY, NumGrassMaps).
+// Grass maps output. Each slice is a single grass combined grass map.
+// Dimension is (NumComponentsX * InLandscapeComponentResolution, NumComponentsY * InLandscapeComponentResolution, NumGrassMaps).
+//
+//                                    -------------
+//                                  -------------  |
+//                                 |             | |
+//        Num Comps Y *            |             | |
+//  InLandscapeComponentResolution |             | |
+//                                 |             | |
+//                                 |             |-  NumGrassMaps
+//                                  -------------
+//                                  Num Comps X *
+//                         InLandscapeComponentResolution
 //
-//                -------------
-//              -------------  |
-//             |             | |
-// Num Comps Y |             | |
-//   * ResY    |             | |
-//             |             | |
-//             |             |-  NumGrassMaps
-//              -------------
-//               Num Comps X
-//                 * ResX
 RWTexture2DArray<float> OutUnpackedGrassMaps;
 
-int4 InLinearTileIndexToComponentIndex[PCG_MAX_NUM_LANDSCAPE_COMPONENTS]; // 16 byte alignment on arrays
+// Height output. Dimension is (NumComponentsX * InLandscapeComponentResolution, NumComponentsY * InLandscapeComponentResolution, NumGrassMaps).
+RWTexture2D<float> OutUnpackedHeight;
+
+int4 InLinearTileIndexToComponentIndex[PCG_MAX_NUM_LANDSCAPE_COMPONENTS]; // Only x component used, 16 byte alignment on arrays
 uint InNumTilesX;
 uint InLandscapeComponentResolution;
 uint InNumGrassMapPasses;
+float3 InLandscapeGridScale; // x == LS Actor DrawScale.X, y == LS Actor DrawScale.y, z == LS Actor DrawScale.z / 128.0f (ZSCALE)
 uint2 InOutputResolution;
+uint InOutputHeight;
 
-// Dispatch group count is (NumComponentsX * ResX / 8, NumComponentsY * ResY / 8, NumGrassMaps).
+// Dispatch group count is (InOutputResolution.xy / 8, NumGrassMaps).
 [numthreads(THREADGROUPSIZE_X, THREADGROUPSIZE_Y, THREADGROUPSIZE_Z)]
 void PCGGrassMapUnpacker_CS(uint3 DispatchThreadId : SV_DispatchThreadID)
 {
@@ -65,8 +72,16 @@ void PCGGrassMapUnpacker_CS(uint3 DispatchThreadId : SV_DispatchThreadID)
 	const uint ComponentXOffset = (ComponentIndex * InNumGrassMapPasses + PassIndex) * InLandscapeComponentResolution;
 	const uint2 InputCoords = uint2(DispatchThreadId.x % InLandscapeComponentResolution + ComponentXOffset, DispatchThreadId.y % InLandscapeComponentResolution);
 
-	// Source channel (see packing above).
+	// Heights. The first layer threads also deal with height.
+	if (InOutputHeight == 1 && DispatchThreadId.z == 0)
+	{
+		const float Height = (UnpackHeight(InPackedGrassMaps[InputCoords].xy) - LANDSCAPE_MID_VALUE) * InLandscapeGridScale.z * TERRAIN_ZSCALE;
+		OutUnpackedHeight[DispatchThreadId.xy] = Height;
+	}
+
+	// Grass map source channel (see packing above).
 	const uint InputChannelIndex = (GrassMapIndex + 2) % 4;
 
+	// Grass maps
 	OutUnpackedGrassMaps[DispatchThreadId] = InPackedGrassMaps[InputCoords][InputChannelIndex];
 }
diff --git a/Engine/Plugins/PCG/Source/PCG/Private/Elements/PCGGenerateGrassMaps.cpp b/Engine/Plugins/PCG/Source/PCG/Private/Elements/PCGGenerateGrassMaps.cpp
deleted file mode 100644
diff --git a/Engine/Plugins/PCG/Source/PCG/Private/Elements/PCGGenerateLandscapeTextures.cpp b/Engine/Plugins/PCG/Source/PCG/Private/Elements/PCGGenerateLandscapeTextures.cpp
new file mode 100644
--- /dev/null
+++ b/Engine/Plugins/PCG/Source/PCG/Private/Elements/PCGGenerateLandscapeTextures.cpp
@@ -0,0 +1,801 @@
+// Copyright Epic Games, Inc. All Rights Reserved.
+
+#include "Elements/PCGGenerateLandscapeTextures.h"
+
+#include "PCGComponent.h"
+#include "PCGGrassMapUnpackerCS.h"
+#include "PCGModule.h"
+#include "PCGSubsystem.h"
+#include "Compute/PCGComputeCommon.h"
+#include "Data/PCGLandscapeData.h"
+#include "Data/PCGTextureData.h"
+#include "Helpers/PCGHelpers.h"
+#include "Metadata/Accessors/PCGAttributeAccessorHelpers.h"
+
+#include "DrawDebugHelpers.h"
+#include "EngineDefines.h" // For UE_ENABLE_DEBUG_DRAWING
+#include "Landscape.h"
+#include "LandscapeComponent.h"
+#include "LandscapeGrassType.h"
+#include "LandscapeGrassWeightExporter.h"
+#include "LandscapeProxy.h"
+#include "RenderCaptureInterface.h"
+#include "RenderGraphBuilder.h"
+#include "RenderGraphUtils.h"
+#include "SystemTextures.h"
+#include "Engine/World.h"
+
+#include UE_INLINE_GENERATED_CPP_BY_NAME(PCGGenerateLandscapeTextures)
+
+#define LOCTEXT_NAMESPACE "PCGGenerateLanscapeTexturesElement"
+
+namespace PCGGenerateLandscapeTextures
+{
+	const FName InputPinLabel = TEXT("Landscape");
+	const FName GrassTypeOverridesPinLabel = TEXT("Grass Types");
+	const FName OutputHeightPinLabel = TEXT("Height");
+	const FName OutputGrassMapsPinLabel = TEXT("GrassMaps");
+
+	const FText LandscapeComponentLostError = LOCTEXT("LandscapeComponentLost", "Reference to one or more landscape components lost, grass maps will not be generated.");
+
+	static int32 GTriggerGPUCaptureDispatches = 0;
+	static FAutoConsoleVariableRef CVarTriggerGPUCapture(
+		TEXT("pcg.GPU.TriggerRenderCaptures.GrassMapGeneration"),
+		GTriggerGPUCaptureDispatches,
+		TEXT("Trigger GPU captures for this many of the subsequent grass generations."));
+
+#if WITH_EDITOR
+	TAutoConsoleVariable<bool> CVarDebugDrawGeneratedComponents(
+		TEXT("pcg.Grass.DebugDrawGeneratedComponents"),
+		false,
+		TEXT("Draws debug boxes around landscapes for which grass maps are generated, colored by the current task ID."));
+#endif
+
+	static bool IsTextureFullyStreamedIn(UTexture* InTexture)
+	{
+		const bool bCheckForLODTransition = true;
+		return InTexture &&
+#if WITH_EDITOR
+			!InTexture->IsDefaultTexture() &&
+#endif // WITH_EDITOR
+			!InTexture->HasPendingInitOrStreaming(bCheckForLODTransition) && InTexture->IsFullyStreamedIn();
+	}
+}
+
+UPCGGenerateLandscapeTexturesSettings::UPCGGenerateLandscapeTexturesSettings()
+{
+	if (PCGHelpers::IsNewObjectAndNotDefault(this))
+	{
+		bGenerateHeightMap = true;
+	}
+}
+
+#if WITH_EDITOR
+TArray<FPCGPreConfiguredSettingsInfo> UPCGGenerateLandscapeTexturesSettings::GetPreconfiguredInfo() const
+{
+	TArray<FPCGPreConfiguredSettingsInfo> PreconfiguredInfo;
+	PreconfiguredInfo.Emplace(0, GetDefaultNodeTitle());
+	PreconfiguredInfo.Emplace(1, LOCTEXT("GenerateGrassMapsNodeTitle", "Generate Grass Maps"));
+
+	return PreconfiguredInfo;
+}
+#endif
+
+void UPCGGenerateLandscapeTexturesSettings::ApplyPreconfiguredSettings(const FPCGPreConfiguredSettingsInfo& PreconfigureInfo)
+{
+	if (PreconfigureInfo.PreconfiguredIndex == 1)
+	{
+		bGenerateHeightMap = false;
+	}
+}
+
+TArray<FPCGPinProperties> UPCGGenerateLandscapeTexturesSettings::InputPinProperties() const
+{
+	TArray<FPCGPinProperties> PinProperties;
+
+	// Single landscape only for now. In the future we should iterate over all landscapes and generate grass maps for each.
+	FPCGPinProperties& PinProp = PinProperties.Emplace_GetRef(
+		PCGGenerateLandscapeTextures::InputPinLabel,
+		EPCGDataType::Landscape,
+		/*bAllowMultipleConnections=*/true,
+		/*bInAllowMultipleData=*/false);
+	PinProp.SetRequiredPin();
+
+	if (bOverrideFromInput)
+	{
+		FPCGPinProperties& GrassTypeOverrides = PinProperties.Emplace_GetRef(PCGGenerateLandscapeTextures::GrassTypeOverridesPinLabel, EPCGDataType::Param);
+		GrassTypeOverrides.SetRequiredPin();
+	}
+
+	return PinProperties;
+}
+
+TArray<FPCGPinProperties> UPCGGenerateLandscapeTexturesSettings::OutputPinProperties() const
+{
+	TArray<FPCGPinProperties> PinProperties;
+
+	if (bGenerateHeightMap)
+	{
+		PinProperties.Emplace(PCGGenerateLandscapeTextures::OutputHeightPinLabel, EPCGDataType::Texture, /*bAllowMultipleConnections=*/true, /*bInAllowMultipleData=*/false);
+	}
+
+	PinProperties.Emplace(PCGGenerateLandscapeTextures::OutputGrassMapsPinLabel, EPCGDataType::Texture, /*bAllowMultipleConnections=*/true, /*bInAllowMultipleData=*/true);
+
+	return PinProperties;
+}
+
+FPCGElementPtr UPCGGenerateLandscapeTexturesSettings::CreateElement() const
+{
+	return MakeShared<FPCGGenerateLandscapeTexturesElement>();
+}
+
+#if WITH_EDITOR
+EPCGChangeType UPCGGenerateLandscapeTexturesSettings::GetChangeTypeForProperty(const FName& InPropertyName) const
+{
+	EPCGChangeType ChangeType = Super::GetChangeTypeForProperty(InPropertyName);
+
+	if (InPropertyName == GET_MEMBER_NAME_CHECKED(UPCGGenerateLandscapeTexturesSettings, bOverrideFromInput)
+		|| InPropertyName == GET_MEMBER_NAME_CHECKED(UPCGGenerateLandscapeTexturesSettings, bGenerateHeightMap))
+	{
+		ChangeType |= EPCGChangeType::Structural;
+	}
+
+	return ChangeType;
+}
+#endif // WITH_EDITOR
+
+FPCGGenerateLandscapeTexturesContext::~FPCGGenerateLandscapeTexturesContext()
+{
+	delete LandscapeGrassWeightExporter;
+}
+
+void FPCGGenerateLandscapeTexturesContext::AddExtraStructReferencedObjects(FReferenceCollector& Collector)
+{
+	if (HeightTextureData)
+	{
+		Collector.AddReferencedObject(HeightTextureData);
+	}
+
+	for (TObjectPtr<UPCGTextureData>& Data : GrassMapTextureDatas)
+	{
+		if (Data)
+		{
+			Collector.AddReferencedObject(Data);
+		}
+	}
+
+	for (TObjectPtr<UTexture>& Texture : TexturesToStream)
+	{
+		if (Texture)
+		{
+			Collector.AddReferencedObject(Texture);
+		}
+	}
+}
+
+void FPCGGenerateLandscapeTexturesElement::GetDependenciesCrc(const FPCGGetDependenciesCrcParams& InParams, FPCGCrc& OutCrc) const
+{
+	FPCGCrc Crc;
+	IPCGElement::GetDependenciesCrc(InParams, Crc);
+
+	// todo_pcg: Technically this could be fancier and hash the set of landscape components that overlap with the generation volume.
+	if (const UPCGData* Data = InParams.ExecutionSource ? InParams.ExecutionSource->GetExecutionState().GetSelfData() : nullptr)
+	{
+		Crc.Combine(Data->GetOrComputeCrc(/*bFullDataCrc=*/false));
+	}
+
+	OutCrc = Crc;
+}
+
+FPCGContext* FPCGGenerateLandscapeTexturesElement::CreateContext()
+{
+	return new FPCGGenerateLandscapeTexturesContext();
+}
+
+bool FPCGGenerateLandscapeTexturesElement::ExecuteInternal(FPCGContext* InContext) const
+{
+	TRACE_CPUPROFILER_EVENT_SCOPE(FPCGGenerateLandscapeTexturesElement::ExecuteInternal);
+
+	FPCGGenerateLandscapeTexturesContext* Context = static_cast<FPCGGenerateLandscapeTexturesContext*>(InContext);
+	check(Context);
+
+	const UPCGGenerateLandscapeTexturesSettings* Settings = Context->GetInputSettings<UPCGGenerateLandscapeTexturesSettings>();
+	check(Settings);
+
+	IPCGGraphExecutionSource* ExecutionSource = Context->ExecutionSource.Get();
+	if (!ExecutionSource)
+	{
+		return true;
+	}
+
+	// 1. Select landscape components that overlap the given bounds.
+	if (!Context->bLandscapeComponentsFiltered)
+	{
+		const UPCGLandscapeData* LandscapeData = nullptr;
+		for (const FPCGTaggedData& Data : Context->InputData.TaggedData)
+		{
+			if (const UPCGLandscapeData* InputLandscapeData = Cast<UPCGLandscapeData>(Data.Data))
+			{
+				if (!LandscapeData)
+				{
+					LandscapeData = InputLandscapeData;
+				}
+				else
+				{
+					PCGLog::LogWarningOnGraph(LOCTEXT("MultipleLandscapesNotSpported", "Multiple landscape data inputs not currently supported, only first will be used."), Context);
+				}
+			}
+		}
+
+		if (!LandscapeData)
+		{
+			// No input landscape, done.
+			return true;
+		}
+
+		if (!UE::Landscape::IsRuntimeGrassMapGenerationSupported() && (!ExecutionSource->GetExecutionState().GetWorld() || ExecutionSource->GetExecutionState().GetWorld()->IsGameWorld()))
+		{
+			UE_LOG(LogPCG, Warning, TEXT("Grass map generation is disabled outside of editor worlds. Try enabling the CVar 'grass.GrassMap.AlwaysBuildRuntimeGenerationResources'."));
+			return true;
+		}
+
+		TArray<FString> SelectedGrassTypes;
+
+		if (!Settings->bOverrideFromInput)
+		{
+			SelectedGrassTypes = Settings->SelectedGrassTypes;
+		}
+		else
+		{
+			const TArray<FPCGTaggedData> OverrideTaggedDatas = InContext->InputData.GetInputsByPin(PCGGenerateLandscapeTextures::GrassTypeOverridesPinLabel);
+			TArray<FString> GrassTypeOverrides;
+
+			for (const FPCGTaggedData& OverrideTaggedData : OverrideTaggedDatas)
+			{
+				if (const UPCGData* OverrideData = OverrideTaggedData.Data)
+				{
+					const FPCGAttributePropertyInputSelector Selector = Settings->GrassTypesAttribute.CopyAndFixLast(OverrideData);
+
+					if (PCGAttributeAccessorHelpers::ExtractAllValues(OverrideData, Selector, GrassTypeOverrides, InContext))
+					{
+						SelectedGrassTypes = std::move(GrassTypeOverrides);
+					}
+					else
+					{
+						PCGLog::LogWarningOnGraph(LOCTEXT("FailExtractGrassTypeOverrides", "Failed to extract grass type overrides."), InContext);
+					}
+				}
+			}
+		}
+
+		const FBox GenerationBounds = ExecutionSource->GetExecutionState().GetBounds();
+		TArray<TObjectPtr<ULandscapeGrassType>> GrassTypes;
+
+		for (TSoftObjectPtr<ALandscapeProxy> LandscapeProxyPtr : LandscapeData->Landscapes)
+		{
+			ALandscapeProxy* LandscapeProxy = LandscapeProxyPtr.Get();
+
+			if (!LandscapeProxy || !LandscapeProxy->GetLandscapeMaterial())
+			{
+				continue;
+			}
+
+			if (!Context->LandscapeProxy.Get())
+			{
+				Context->LandscapeProxy = LandscapeProxy;
+			}
+
+			for (ULandscapeComponent* LandscapeComponent : LandscapeProxy->LandscapeComponents)
+			{
+				if (!ensure(LandscapeComponent))
+				{
+					continue;
+				}
+
+				// Only generate grass map if there is meaningful overlap with our domain of interest.
+				const FBox LandscapeComponentBounds = LandscapeComponent->Bounds.GetBox();
+				if (LandscapeComponentBounds.Overlap(GenerationBounds).GetVolume() > UE_KINDA_SMALL_NUMBER)
+				{
+					if (Context->LandscapeComponents.Num() >= FPCGGrassMapUnpackerCS::MaxNumLandscapeComponents)
+					{
+						PCGLog::LogWarningOnGraph(LOCTEXT("MaxLandscapeComponentsExceeded", "Too many landscape components overlap the generation domain. Consider partitioning the component onto a smaller grid size."), Context);
+						break;
+					}
+
+					LandscapeComponent->UpdateGrassTypes();
+
+					Context->LandscapeComponents.Add(LandscapeComponent);
+
+					if (Context->LandscapeComponents.Num() == 1)
+					{
+						GrassTypes = LandscapeComponent->GetGrassTypes();
+						Context->NumGrassTypes = GrassTypes.Num();
+
+						FString GrassTypeName;
+						GrassTypeName.Reserve(256);
+
+						for (int GrassTypeIndex = 0; GrassTypeIndex < GrassTypes.Num(); ++GrassTypeIndex)
+						{
+							if (TObjectPtr<ULandscapeGrassType> GrassType = GrassTypes[GrassTypeIndex])
+							{
+								GrassType->GetName(GrassTypeName);
+
+								const bool bIsSelectedLayer = SelectedGrassTypes.Contains(GrassTypeName);
+
+								if (bIsSelectedLayer != Settings->bExcludeSelectedGrassTypes)
+								{
+									Context->SelectedGrassTypes.Emplace(GrassType, GrassTypeIndex);
+								}
+							}
+						}
+
+						Context->GrassMapBounds = LandscapeComponentBounds;
+
+						Context->LandscapeComponentExtent = LandscapeComponent->Bounds.BoxExtent.X * 2.0;
+
+						if (!ensure(FMath::IsNearlyEqual(LandscapeComponent->Bounds.BoxExtent.X, LandscapeComponent->Bounds.BoxExtent.Y)))
+						{
+							return true;
+						}
+
+						if (!ensure(Context->LandscapeComponentExtent > 0.0))
+						{
+							return true;
+						}
+					}
+					else
+					{
+						Context->GrassMapBounds += LandscapeComponentBounds;
+
+						const UMaterialInterface* LandscapeProxyMaterial = Context->LandscapeProxy->GetLandscapeMaterial();
+						const UMaterialInterface* LandscapeMaterial = LandscapeComponent->GetLandscapeMaterial();
+
+						// We expect all landscape components for a single landscape to have the same grass types.
+						// TODO: In order to support distinct grass types we need to do a dispatch of one landscape grass exporter per LS proxy.
+						if (LandscapeMaterial != LandscapeProxyMaterial && GrassTypes != LandscapeComponent->GetGrassTypes())
+						{
+#if !(UE_BUILD_SHIPPING || UE_BUILD_TEST) || USE_LOGGING_IN_SHIPPING
+							const FString LandscapeProxyMaterialName = LandscapeProxyMaterial ? LandscapeProxyMaterial->GetName() : TEXT("MISSINGMATERIAL");
+							const FString LandscapeMaterialName = LandscapeMaterial ? LandscapeMaterial->GetName() : TEXT("MISSINGMATERIAL");
+							
+							PCGLog::LogErrorOnGraph(
+								FText::Format(LOCTEXT("LandscapeComponentGrassTypeMismatch",
+									"Encountered landscape component '{0}' with material '{1}' that does not match the grass types on material '{2}'. All landscape materials must have the same grass types at this time."),
+									FText::FromString(LandscapeComponent->GetName()),
+									FText::FromString(LandscapeMaterialName),
+									FText::FromString(LandscapeProxyMaterialName)),
+								Context);
+#endif
+							return true;
+						}
+
+						// Currently assuming all landscape have the same extents.
+						ensure(FMath::IsNearlyEqual(Context->LandscapeComponentExtent, LandscapeComponent->Bounds.BoxExtent.X * 2.0));
+					}
+
+#if WITH_EDITOR
+					if (PCGGenerateLandscapeTextures::CVarDebugDrawGeneratedComponents.GetValueOnGameThread())
+					{
+						DrawDebugBox(ExecutionSource->GetExecutionState().GetWorld(),
+							LandscapeComponentBounds.GetCenter(),
+							LandscapeComponentBounds.GetExtent(),
+							FColor::MakeRandomSeededColor(Context->TaskId),
+							/*bPersistentLines=*/false,
+							/*LifeTime=*/3.0f);
+					}
+#endif
+				}
+			}
+		}
+
+		if (!Context->LandscapeProxy.Get() || Context->LandscapeComponents.IsEmpty() || Context->SelectedGrassTypes.IsEmpty())
+		{
+			return true;
+		}
+
+		Context->bLandscapeComponentsFiltered = true;
+	}
+
+	// 2. Wait for landscape components to be ready for grass map rendering.
+	if (!Context->bTextureStreamingRequested)
+	{
+		if (const UWorld* World = ExecutionSource->GetExecutionState().GetWorld())
+		{
+			for (TWeakObjectPtr<ULandscapeComponent> LandscapeComponentWeak : Context->LandscapeComponents)
+			{
+				if (ULandscapeComponent* LandscapeComponent = LandscapeComponentWeak.Get())
+				{
+					// Make list of textures to stream before generating.
+					if (UTexture* HeightMap = LandscapeComponent->GetHeightmap())
+					{
+						Context->TexturesToStream.Add(HeightMap);
+					}
+
+					const ERHIFeatureLevel::Type FeatureLevel = World->GetFeatureLevel();
+					for (UTexture2D* WeightmapTexture : LandscapeComponent->GetRenderedWeightmapTexturesForFeatureLevel(FeatureLevel))
+					{
+						if (WeightmapTexture)
+						{
+							Context->TexturesToStream.Add(WeightmapTexture);
+						}
+					}
+				}
+			}
+
+			for (UTexture* Texture : Context->TexturesToStream)
+			{
+				Texture->bForceMiplevelsToBeResident = true;
+			}
+		}
+
+		Context->bTextureStreamingRequested = true;
+	}
+
+	if (!Context->bReadyToRender)
+	{
+		bool bAllReady = true;
+
+		for (TWeakObjectPtr<ULandscapeComponent> LandscapeComponentWeak : Context->LandscapeComponents)
+		{
+			ULandscapeComponent* LandscapeComponent = LandscapeComponentWeak.Get();
+			if (!LandscapeComponent)
+			{
+				PCGLog::LogErrorOnGraph(PCGGenerateLandscapeTextures::LandscapeComponentLostError, InContext);
+				return true;
+			}
+
+			bAllReady &= UE::Landscape::CanRenderGrassMap(LandscapeComponent);
+		}
+
+		if (bAllReady)
+		{
+			for (UTexture* Texture : Context->TexturesToStream)
+			{
+				const bool bStreamedIn = PCGGenerateLandscapeTextures::IsTextureFullyStreamedIn(Texture);
+				bAllReady &= bStreamedIn;
+
+				if (!bStreamedIn)
+				{
+					UE_LOG(LogPCG, Verbose, TEXT("Waiting for landscape texture '%s' to stream in."), *Texture->GetName());
+					break;
+				}
+			}
+		}
+		
+		if (!bAllReady)
+		{
+#if UE_ENABLE_DEBUG_DRAWING
+			if (PCGSystemSwitches::CVarPCGDebugDrawGeneratedCells.GetValueOnGameThread())
+			{
+				FColor DebugColor = FColor::Yellow;
+				PCGHelpers::DebugDrawGenerationVolume(InContext, &DebugColor);
+			}
+#endif
+
+			// Sleep until next frame, no use spinning on this.
+			Context->bIsPaused = true;
+			FPCGModule::GetPCGModuleChecked().ExecuteNextTick([ContextHandle = InContext->GetOrCreateHandle()]()
+			{
+				if (TSharedPtr<FPCGContextHandle> SharedHandle = ContextHandle.Pin())
+				{
+					if (FPCGContext* ContextPtr = SharedHandle->GetContext())
+					{
+						ContextPtr->bIsPaused = false;
+					}
+				}
+			});
+
+			return false;
+		}
+
+		Context->bReadyToRender = true;
+	}
+
+	// 3. Schedule grass map generation.
+	if (!Context->bGenerationScheduled)
+	{
+		ALandscapeProxy* LandscapeProxy = Context->LandscapeProxy.IsValid() ? Context->LandscapeProxy.Get() : nullptr;
+		if (!LandscapeProxy)
+		{
+			PCGLog::LogErrorOnGraph(LOCTEXT("LandscapeProxyLost", "Reference to landscape proxy actor lost, grass maps will not be generated."), InContext);
+			return true;
+		}
+
+		const FVector LandscapeGridScale = LandscapeProxy->GetRootComponent()->GetRelativeScale3D();
+
+		TArray<ULandscapeComponent*> LandscapeComponents;
+		TArray<FIntVector2> LandscapeTileCoords;
+		LandscapeComponents.Reserve(Context->LandscapeComponents.Num());
+		LandscapeTileCoords.Reserve(Context->LandscapeComponents.Num());
+
+		for (TWeakObjectPtr<ULandscapeComponent> LandscapeComponentPtr : Context->LandscapeComponents)
+		{
+			ULandscapeComponent* LandscapeComponent = LandscapeComponentPtr.IsValid() ? LandscapeComponentPtr.Get() : nullptr;
+
+			if (!LandscapeComponent)
+			{
+				PCGLog::LogErrorOnGraph(PCGGenerateLandscapeTextures::LandscapeComponentLostError, InContext);
+				return true;
+			}
+
+			LandscapeComponents.Add(LandscapeComponent);
+
+			// Landscape components are not ordered, so store a 2d index of each component within the grass map.
+			// TODO could likely sort the LandscapeComponents array so the order is known and no indices need looking up.
+			LandscapeTileCoords.Add(FIntVector2(
+				(LandscapeComponent->Bounds.Origin.X - Context->GrassMapBounds.Min.X) / Context->LandscapeComponentExtent,
+				(LandscapeComponent->Bounds.Origin.Y - Context->GrassMapBounds.Min.Y) / Context->LandscapeComponentExtent));
+		}
+
+		Context->LandscapeGrassWeightExporter = new FLandscapeGrassWeightExporter(
+			LandscapeProxy,
+			LandscapeComponents,
+			/*bInNeedsGrassmap=*/true,
+			Settings->bGenerateHeightMap,
+			/*HeightMips=*/{},
+			/*bExport=*/false,
+			/*bShouldReadback=*/false);
+
+		const FIntPoint& GrassMapTextureTargetSize = Context->LandscapeGrassWeightExporter->GetTargetSize();
+		const int32 Max2DTextureDimension = GetMax2DTextureDimension();
+
+		if (GrassMapTextureTargetSize.X < 1 || GrassMapTextureTargetSize.X > Max2DTextureDimension || GrassMapTextureTargetSize.Y < 1 || GrassMapTextureTargetSize.Y > Max2DTextureDimension)
+		{
+			PCGLog::LogErrorOnGraph(FText::Format(LOCTEXT("InvalidTextureSize", "Invalid grass map texture size ({0}, {1}), cannot exceed ({2}, {2}). Consider partitioning the component onto a smaller grid size."), GrassMapTextureTargetSize.X, GrassMapTextureTargetSize.Y, Max2DTextureDimension), InContext);
+			return true;
+		}
+
+		UE::RenderCommandPipe::FSyncScope SyncScope;
+
+		RenderCaptureInterface::FScopedCapture RenderCapture(PCGGenerateLandscapeTextures::GTriggerGPUCaptureDispatches > 0, TEXT("PCGLandscapeGrassmapCapture"));
+		PCGGenerateLandscapeTextures::GTriggerGPUCaptureDispatches = FMath::Max(PCGGenerateLandscapeTextures::GTriggerGPUCaptureDispatches - 1, 0);
+		
+		TWeakPtr<FPCGContextHandle> ContextHandle = Context->GetOrCreateHandle();
+
+		ENQUEUE_RENDER_COMMAND(GenerateGrassMaps)([ContextHandle, LandscapeComponentExtent=Context->LandscapeComponentExtent, LandscapeTileCoords=MoveTemp(LandscapeTileCoords), ComponentSizeQuads=LandscapeProxy->ComponentSizeQuads, LandscapeGridScale, bGenerateHeightMap=Settings->bGenerateHeightMap](FRHICommandListImmediate& RHICmdList)
+		{
+			LLM_SCOPE_BYTAG(PCG);
+
+			FPCGContext::FSharedContext<FPCGGenerateLandscapeTexturesContext> SharedContext(ContextHandle);
+			FPCGGenerateLandscapeTexturesContext* Context = SharedContext.Get();
+			if (!Context)
+			{
+				return;
+			}
+
+			FRDGBuilder GraphBuilder(RHICmdList);
+
+			FRDGTextureDesc GrassMapTextureDesc = FRDGTextureDesc::Create2D(
+				Context->LandscapeGrassWeightExporter->GetTargetSize(),
+				PF_B8G8R8A8,
+				FClearValueBinding(),
+				ETextureCreateFlags::RenderTargetable | ETextureCreateFlags::ShaderResource);
+
+			FRDGTextureRef GrassMapTexture = GraphBuilder.CreateTexture(GrassMapTextureDesc, TEXT("PCGLandscapeGrassMapRenderTarget"));
+			FRDGTextureSRVRef GrassMapTextureSRV = GraphBuilder.CreateSRV(GrassMapTexture);
+
+			// Generate grass maps. All will be generated to a single texture.
+			Context->LandscapeGrassWeightExporter->RenderLandscapeComponentToTexture_RenderThread(GraphBuilder, GrassMapTexture);
+
+			const FVector GrassMapExtent = Context->GrassMapBounds.GetExtent() * 2;
+			const uint32 NumTilesX = FMath::RoundToInt(GrassMapExtent.X / LandscapeComponentExtent);
+			const uint32 NumTilesY = FMath::RoundToInt(GrassMapExtent.Y / LandscapeComponentExtent);
+
+			// Each corner of a quad in the landscape corresponds to one texel in the grass map.
+			const uint32 LandscapeComponentResolution = ComponentSizeQuads + 1;
+			const FIntPoint GrassMapResolution(LandscapeComponentResolution * NumTilesX, LandscapeComponentResolution * NumTilesY);
+			const uint32 NumGrassTypes = Context->NumGrassTypes;
+
+			FRDGTextureDesc GrassHeightMapDesc = FRDGTextureDesc::Create2D(
+				bGenerateHeightMap ? GrassMapResolution : FIntPoint(1, 1), // No default system texture usable for UAV, so always create but with 1x1 resolution
+				PF_R32_FLOAT,
+				FClearValueBinding(),
+				ETextureCreateFlags::ShaderResource | ETextureCreateFlags::UAV);
+
+			FRDGTextureRef GrassHeightMap = GraphBuilder.CreateTexture(GrassHeightMapDesc, TEXT("PCGLandscapeHeightMapUnpacked"));
+			FRDGTextureUAVRef GrassHeightMapUAV = GraphBuilder.CreateUAV(GrassHeightMap);
+
+			// Output texture is array of textures, one per grass map.
+			FRDGTextureDesc GrassMapDesc = FRDGTextureDesc::Create2DArray(
+				GrassMapResolution,
+				PF_G8,
+				FClearValueBinding(),
+				ETextureCreateFlags::ShaderResource | ETextureCreateFlags::UAV,
+				NumGrassTypes);
+
+			FRDGTextureRef GrassMap = GraphBuilder.CreateTexture(GrassMapDesc, TEXT("PCGLandscapeGrassMapUnpacked"));
+			FRDGTextureUAVRef GrassMapUAV = GraphBuilder.CreateUAV(GrassMap);
+
+			// Unpack the generated results to simple world aligned textures.
+			FPCGGrassMapUnpackerCS::FParameters* Parameters = GraphBuilder.AllocParameters<FPCGGrassMapUnpackerCS::FParameters>();
+			Parameters->InPackedGrassMaps = GrassMapTextureSRV;
+			Parameters->OutUnpackedGrassMaps = GrassMapUAV;
+			Parameters->OutUnpackedHeight = GrassHeightMapUAV;
+			Parameters->InNumTilesX = NumTilesX;
+			Parameters->InLandscapeGridScale = static_cast<FVector3f>(LandscapeGridScale);
+			Parameters->InLandscapeComponentResolution = LandscapeComponentResolution;
+			Parameters->InOutputResolution = FUintVector2(GrassMapResolution.X, GrassMapResolution.Y);
+			Parameters->InOutputHeight = bGenerateHeightMap ? 1u : 0u;
+
+			// The first 2 channels are reserved for height data. See illustration of packing in PCGGrassMapUnpackerCS.usf.
+			Parameters->InNumGrassMapPasses = FMath::DivideAndRoundUp<int>(NumGrassTypes + 2, 4);
+
+			// Initialize to invalid component indices.
+			for (int Index = 0; Index < FPCGGrassMapUnpackerCS::MaxNumLandscapeComponents; ++Index)
+			{
+				Parameters->InLinearTileIndexToComponentIndex[Index].X = -1;
+			}
+
+			// Now write component mapping.
+			for (int Index = 0; Index < LandscapeTileCoords.Num(); ++Index)
+			{
+				Parameters->InLinearTileIndexToComponentIndex[LandscapeTileCoords[Index].Y * NumTilesX + LandscapeTileCoords[Index].X].X = Index;
+			}
+
+			TShaderMapRef<FPCGGrassMapUnpackerCS> Shader(GetGlobalShaderMap(GMaxRHIFeatureLevel));
+			const int GroupCountX = FMath::DivideAndRoundUp<int>(GrassMapResolution.X, FPCGGrassMapUnpackerCS::ThreadGroupDim);
+			const int GroupCountY = FMath::DivideAndRoundUp<int>(GrassMapResolution.Y, FPCGGrassMapUnpackerCS::ThreadGroupDim);
+			const int GroupCountZ = NumGrassTypes;
+
+			FComputeShaderUtils::AddPass(GraphBuilder, RDG_EVENT_NAME("PCGUnpackGrassMap"), ERDGPassFlags::Compute, Shader, Parameters, FIntVector(GroupCountX, GroupCountY, GroupCountZ));
+
+			// Export the output textures so they can be used downstream.
+			TRefCountPtr<IPooledRenderTarget> HeightsExported;
+			if (bGenerateHeightMap)
+			{
+				HeightsExported = GraphBuilder.ConvertToExternalTexture(GrassHeightMap);
+				GraphBuilder.SetTextureAccessFinal(GrassHeightMap, ERHIAccess::SRVCompute);
+			}
+
+			TRefCountPtr<IPooledRenderTarget> GrassMapExported = GraphBuilder.ConvertToExternalTexture(GrassMap);
+			GraphBuilder.SetTextureAccessFinal(GrassMap, ERHIAccess::SRVCompute);
+
+			GraphBuilder.Execute();
+
+			// Pass exported buffer back to game thread and wake up this element.
+			ExecuteOnGameThread(UE_SOURCE_LOCATION, [ContextHandle, HeightsExported, GrassMapExported]()
+			{
+				LLM_SCOPE_BYTAG(PCG);
+
+				FPCGContext::FSharedContext<FPCGGenerateLandscapeTexturesContext> SharedContext(ContextHandle);
+				if (FPCGGenerateLandscapeTexturesContext* Context = SharedContext.Get())
+				{
+					Context->HeightHandle = HeightsExported;
+					Context->GrassMapHandle = GrassMapExported;
+					Context->bIsPaused = false;
+				}
+			});
+		});
+
+		Context->bGenerationScheduled = true;
+
+		// Render command will wake this task up after completing.
+		Context->bIsPaused = true;
+
+		return false;
+	}
+
+	// 4. Initialize texture data objects.
+
+	if (Settings->bGenerateHeightMap && !Context->HeightTextureData)
+	{
+		Context->HeightTextureData = FPCGContext::NewObject_AnyThread<UPCGTextureData>(Context);
+	}
+
+	if (Context->GrassMapTextureDatas.IsEmpty())
+	{
+		// Create the texture data objects if they haven't been created already. There should be one per selected grass type.
+		for (int DataIndex = 0; DataIndex < Context->SelectedGrassTypes.Num(); ++DataIndex)
+		{
+			UPCGTextureData* TextureData = FPCGContext::NewObject_AnyThread<UPCGTextureData>(Context);
+			Context->GrassMapTextureDatas.Add(TextureData);
+		}
+	}
+
+	if (!Context->GrassMapHandle.IsValid() && !Context->HeightHandle.IsValid())
+	{
+		return true;
+	}
+
+	const FTransform OutputTextureTransform(FQuat::Identity, Context->GrassMapBounds.GetCenter(), Context->GrassMapBounds.GetExtent());
+	bool bAllTexturesInitialized = true;
+
+	if (Context->HeightHandle)
+	{
+		check(Context->HeightTextureData);
+		bAllTexturesInitialized &= Context->HeightTextureData->Initialize(Context->HeightHandle, /*TextureIndex=*/0, OutputTextureTransform, Settings->bSkipReadbackToCPU);
+	}
+
+	if (Context->GrassMapHandle)
+	{
+		for (int DataIndex = 0; DataIndex < Context->SelectedGrassTypes.Num(); ++DataIndex)
+		{
+			check(Context->GrassMapTextureDatas.IsValidIndex(DataIndex));
+			UPCGTextureData* TextureData = Context->GrassMapTextureDatas[DataIndex];
+			check(TextureData);
+
+			// Poll initialize (fine to be called even when initialization was already complete).
+			bAllTexturesInitialized &= TextureData->Initialize(Context->GrassMapHandle, /*TextureIndex=*/Context->SelectedGrassTypes[DataIndex].Get<1>(), OutputTextureTransform, Settings->bSkipReadbackToCPU);
+		}
+	}
+
+	if (!bAllTexturesInitialized)
+	{
+		// Initialization not complete. Could be waiting on async texture processing or for GPU readback. Sleep until next frame.
+		// TODO: Ideally we do lazy readback on texture data in the future so we don't have to read it back to CPU unless it's needed.
+		Context->bIsPaused = true;
+		FPCGModule::GetPCGModuleChecked().ExecuteNextTick([ContextHandle = Context->GetOrCreateHandle()]()
+		{
+			FPCGContext::FSharedContext<FPCGGenerateLandscapeTexturesContext> SharedContext(ContextHandle);
+			if (FPCGGenerateLandscapeTexturesContext* ContextPtr = SharedContext.Get())
+			{
+				ContextPtr->bIsPaused = false;
+			}
+		});
+
+		return false;
+	}
+
+#if WITH_EDITOR
+	if ((Context->HeightTextureData || !Context->GrassMapTextureDatas.IsEmpty()) && !Settings->bSkipReadbackToCPU)
+	{
+		if (!Settings->bSkipReadbackToCPU && Context->Node && Context->GetStack() && Context->ExecutionSource.IsValid())
+		{
+			Context->ExecutionSource->GetExecutionState().GetInspection().NotifyGPUToCPUReadback(Context->Node, Context->GetStack());
+		}
+	}
+#endif // WITH_EDITOR
+
+	// 5. Emit texture data objects.
+	if (Context->HeightTextureData)
+	{
+		if (Context->HeightTextureData->IsSuccessfullyInitialized())
+		{
+			FPCGTaggedData& OutTaggedData = Context->OutputData.TaggedData.Emplace_GetRef();
+			OutTaggedData.Data = Context->HeightTextureData;
+			OutTaggedData.Pin = PCGGenerateLandscapeTextures::OutputHeightPinLabel;
+		}
+		else
+		{
+			PCGE_LOG(Warning, LogOnly, LOCTEXT("TextureInitFailed", "Data could not be retrieved for height texture, initialization failed."));
+		}
+	}
+
+	for (int DataIndex = 0; DataIndex < Context->SelectedGrassTypes.Num(); ++DataIndex)
+	{
+		check(Context->GrassMapTextureDatas.IsValidIndex(DataIndex));
+		UPCGTextureData* TextureData = Context->GrassMapTextureDatas[DataIndex];
+		check(TextureData);
+
+		if (!TextureData->IsSuccessfullyInitialized())
+		{
+			PCGE_LOG(Warning, LogOnly, LOCTEXT("TextureInitFailed", "Data could not be retrieved for this texture, initialization failed."));
+			continue;
+		}
+
+		FPCGTaggedData& OutTaggedData = Context->OutputData.TaggedData.Emplace_GetRef();
+		OutTaggedData.Pin = PCGGenerateLandscapeTextures::OutputGrassMapsPinLabel;
+		OutTaggedData.Data = TextureData;
+		
+		FString GrassTypeName = TEXT("");
+
+		if (ULandscapeGrassType* GrassType = Context->SelectedGrassTypes[DataIndex].Get<0>().Get())
+		{
+			GrassType->GetName(GrassTypeName);
+		}
+
+		if (!GrassTypeName.IsEmpty())
+		{
+			OutTaggedData.Tags.Add(PCGComputeHelpers::GetPrefixedDataLabel(GrassTypeName));
+		}
+		else
+		{
+			PCGE_LOG(Warning, LogOnly, LOCTEXT("MissingGrassName",
+				"Grass type name was missing, data could not be labeled. Make sure all GrassTypes in your landscape material have an asset associated."));
+		}
+	}
+
+	return true;
+}
+
+#undef LOCTEXT_NAMESPACE
diff --git a/Engine/Plugins/PCG/Source/PCG/Public/Elements/PCGGenerateGrassMaps.h b/Engine/Plugins/PCG/Source/PCG/Public/Elements/PCGGenerateGrassMaps.h
--- a/Engine/Plugins/PCG/Source/PCG/Public/Elements/PCGGenerateGrassMaps.h
+++ b/Engine/Plugins/PCG/Source/PCG/Public/Elements/PCGGenerateGrassMaps.h
@@ -2,124 +2,12 @@
 
 #pragma once
 
-#include "PCGContext.h"
-#include "PCGSettings.h"
-
-#include "RendererInterface.h"
+#include "PCGGenerateLandscapeTextures.h"
 
 #include "PCGGenerateGrassMaps.generated.h"
 
-class ALandscapeProxy;
-class FLandscapeGrassWeightExporter;
-class ULandscapeComponent;
-class ULandscapeGrassType;
-class UPCGTextureData;
-class UTexture;
-struct IPooledRenderTarget;
-
-UCLASS(MinimalAPI, BlueprintType, ClassGroup = (Procedural))
-class UPCGGenerateGrassMapsSettings : public UPCGSettings
+UCLASS(MinimalAPI, BlueprintType, Deprecated, meta = (DeprecationMessage = "Deprecated 5.7. Use UPCGGenerateLandscapeTexturesSettings instead."))
+class UDEPRECATED_PCGGenerateGrassMapsSettings : public UPCGGenerateLandscapeTexturesSettings
 {
 	GENERATED_BODY()
-
-public:
-	//~Begin UPCGSettings interface
-#if WITH_EDITOR
-	virtual FName GetDefaultNodeName() const override { return FName(TEXT("GenerateGrassMaps")); }
-	virtual FText GetDefaultNodeTitle() const override { return NSLOCTEXT("PCGGenerateGrassMapsElement", "NodeTitle", "Generate Grass Maps"); }
-	virtual FText GetNodeTooltipText() const override { return NSLOCTEXT("PCGGenerateGrassMapsElement", "NodeTooltip", "Generates landscape grass maps on the GPU."); }
-	virtual EPCGSettingsType GetType() const override { return EPCGSettingsType::GPU; }
-#endif
-
-protected:
-	virtual TArray<FPCGPinProperties> InputPinProperties() const override;
-	virtual TArray<FPCGPinProperties> OutputPinProperties() const override;
-	virtual bool IsInputPinRequiredByExecution(const UPCGPin* InPin) const override { return true; }
-	virtual FPCGElementPtr CreateElement() const override;
-#if WITH_EDITOR
-	virtual EPCGChangeType GetChangeTypeForProperty(const FName& InPropertyName) const override;
-#endif
-	//~End UPCGSettings interface
-
-public:
-	/** Select which grass types to generate. */
-	UPROPERTY(BlueprintReadWrite, EditAnywhere, Category = Settings, meta = (EditCondition = "!bOverrideFromInput"))
-	TArray<FString> SelectedGrassTypes;
-
-	/** Override grass types from input. */
-	UPROPERTY(BlueprintReadWrite, EditAnywhere, Category = Settings)
-	bool bOverrideFromInput = false;
-
-	/** Input attribute to pull grass type strings from. */
-	UPROPERTY(BlueprintReadWrite, EditAnywhere, Category = Settings, meta = (EditCondition = "bOverrideFromInput", PCG_DiscardPropertySelection, PCG_DiscardExtraSelection))
-	FPCGAttributePropertyInputSelector GrassTypesAttribute;
-
-	/** If toggled, will only generate grass types which are not selected. */
-	UPROPERTY(BlueprintReadWrite, EditAnywhere, Category = Settings, meta = (PCG_Overridable))
-	bool bExcludeSelectedGrassTypes = true;
-
-	/** Skip CPU readback of emitted textures during initialization of the texture datas. */
-	UPROPERTY(BlueprintReadWrite, EditAnywhere, Category = Settings)
-	bool bSkipReadbackToCPU = false;
-};
-
-struct FPCGGenerateGrassMapsContext : public FPCGContext
-{
-public:
-	~FPCGGenerateGrassMapsContext();
-
-	FLandscapeGrassWeightExporter* LandscapeGrassWeightExporter = nullptr;
-
-	/** Output texture data objects. */
-	TArray<TObjectPtr<UPCGTextureData>> TextureDatas;
-
-	/** Exported result texture array. */
-	TRefCountPtr<IPooledRenderTarget> GrassMapHandle = nullptr;
-
-	/** List of the grass types selected for generation. We have to hold their texture index as well, since
-	 * this array could be sparse, but the actual texture array we produce will always have all the grass types.
-	 */
-	TArray<TTuple<TWeakObjectPtr<ULandscapeGrassType>, /*TextureIndex=*/int32>> SelectedGrassTypes;
-
-	/** Total number of grass types used by the landscape component. Includes grass types which were not selected for generation. */
-	int32 NumGrassTypes = 0;
-
-	TWeakObjectPtr<ALandscapeProxy> LandscapeProxy = nullptr;
-	TArray<TWeakObjectPtr<ULandscapeComponent>> LandscapeComponents;
-
-	/** World-space bounds containing all of the landscape components given to the grass weight exporter. */
-	FBox GrassMapBounds = FBox(EForceInit::ForceInit);
-
-	/** Extent (side length) of each landscape component. */
-	double LandscapeComponentExtent = 0.0;
-
-	/** True when we have filtered all of the incoming landscape components down to the ones which overlap the given bounds. */
-	bool bLandscapeComponentsFiltered = false;
-
-	/** True when the landscape components are ready for rendering. */
-	bool bReadyToRender = false;
-
-	/** Textures that we wait to be streamed before generating the grass maps. */
-	TArray<TObjectPtr<UTexture>> TexturesToStream;
-
-	/** True when streaming has been requested on landscape textures. */
-	bool bTextureStreamingRequested = false;
-
-	/** True when grass map generation has been scheduled on the render thread. */
-	bool bGenerationScheduled = false;
-
-protected:
-	virtual void AddExtraStructReferencedObjects(FReferenceCollector& Collector) override;
-};
-
-class FPCGGenerateGrassMapsElement : public IPCGElement
-{
-public:
-	virtual void GetDependenciesCrc(const FPCGGetDependenciesCrcParams& InParams, FPCGCrc& OutCrc) const override;
-	/** FLandscapeGrassWeightExporter expects to exist in scope only on the game thread. */
-	virtual bool CanExecuteOnlyOnMainThread(FPCGContext* Context) const { return true; }
-
-protected:
-	virtual FPCGContext* CreateContext() override;
-	virtual bool ExecuteInternal(FPCGContext* InContext) const override;
 };
diff --git a/Engine/Plugins/PCG/Source/PCG/Public/Elements/PCGGenerateLandscapeTextures.h b/Engine/Plugins/PCG/Source/PCG/Public/Elements/PCGGenerateLandscapeTextures.h
new file mode 100644
--- /dev/null
+++ b/Engine/Plugins/PCG/Source/PCG/Public/Elements/PCGGenerateLandscapeTextures.h
@@ -0,0 +1,137 @@
+// Copyright Epic Games, Inc. All Rights Reserved.
+
+#pragma once
+
+#include "PCGContext.h"
+#include "PCGSettings.h"
+
+#include "RendererInterface.h"
+
+#include "PCGGenerateLandscapeTextures.generated.h"
+
+class ALandscapeProxy;
+class FLandscapeGrassWeightExporter;
+class ULandscapeComponent;
+class ULandscapeGrassType;
+class UPCGTextureData;
+class UTexture;
+struct IPooledRenderTarget;
+
+UCLASS(MinimalAPI, BlueprintType, ClassGroup = (Procedural))
+class UPCGGenerateLandscapeTexturesSettings : public UPCGSettings
+{
+	GENERATED_BODY()
+
+public:
+	UPCGGenerateLandscapeTexturesSettings();
+
+	//~Begin UPCGSettings interface
+#if WITH_EDITOR
+	virtual FName GetDefaultNodeName() const override { return FName(TEXT("GenerateLandscapeTextures")); }
+	virtual FText GetDefaultNodeTitle() const override { return NSLOCTEXT("PCGGenerateLandscapeTexturesElement", "NodeTitle", "Generate Landscape Textures"); }
+	virtual FText GetNodeTooltipText() const override { return NSLOCTEXT("PCGGenerateLandscapeTexturesElement", "NodeTooltip", "Generates landscape height texture and grass maps on the GPU."); }
+	virtual EPCGSettingsType GetType() const override { return EPCGSettingsType::GPU; }
+	virtual TArray<FPCGPreConfiguredSettingsInfo> GetPreconfiguredInfo() const override;
+	virtual bool OnlyExposePreconfiguredSettings() const override { return true; }
+	virtual bool GroupPreconfiguredSettings() const override { return false; }
+#endif
+	virtual void ApplyPreconfiguredSettings(const FPCGPreConfiguredSettingsInfo& PreconfigureInfo) override;
+
+protected:
+	virtual TArray<FPCGPinProperties> InputPinProperties() const override;
+	virtual TArray<FPCGPinProperties> OutputPinProperties() const override;
+	virtual bool IsInputPinRequiredByExecution(const UPCGPin* InPin) const override { return true; }
+	virtual FPCGElementPtr CreateElement() const override;
+#if WITH_EDITOR
+	virtual EPCGChangeType GetChangeTypeForProperty(const FName& InPropertyName) const override;
+#endif
+	//~End UPCGSettings interface
+
+public:
+	/** Select which grass types to generate. */
+	UPROPERTY(BlueprintReadWrite, EditAnywhere, Category = Settings, meta = (EditCondition = "!bOverrideFromInput"))
+	TArray<FString> SelectedGrassTypes;
+
+	/** Override grass types from input. */
+	UPROPERTY(BlueprintReadWrite, EditAnywhere, Category = Settings)
+	bool bOverrideFromInput = false;
+
+	/** Input attribute to pull grass type strings from. */
+	UPROPERTY(BlueprintReadWrite, EditAnywhere, Category = Settings, meta = (EditCondition = "bOverrideFromInput", PCG_DiscardPropertySelection, PCG_DiscardExtraSelection))
+	FPCGAttributePropertyInputSelector GrassTypesAttribute;
+
+	/** If toggled, will only generate grass types which are not selected. */
+	UPROPERTY(BlueprintReadWrite, EditAnywhere, Category = Settings, meta = (PCG_Overridable))
+	bool bExcludeSelectedGrassTypes = true;
+
+	/** Skip CPU readback of emitted textures during initialization of the texture datas. */
+	UPROPERTY(BlueprintReadWrite, EditAnywhere, Category = Settings)
+	bool bSkipReadbackToCPU = false;
+
+	/** Generate a landscape height texture. */
+	UPROPERTY(BlueprintReadWrite, EditAnywhere, Category = Settings)
+	bool bGenerateHeightMap = false;
+};
+
+struct FPCGGenerateLandscapeTexturesContext : public FPCGContext
+{
+public:
+	~FPCGGenerateLandscapeTexturesContext();
+
+	FLandscapeGrassWeightExporter* LandscapeGrassWeightExporter = nullptr;
+
+	/** Output texture data objects. */
+	TObjectPtr<UPCGTextureData> HeightTextureData;
+	TArray<TObjectPtr<UPCGTextureData>> GrassMapTextureDatas;
+
+	/** Exported result textures. */
+	TRefCountPtr<IPooledRenderTarget> HeightHandle = nullptr;
+	TRefCountPtr<IPooledRenderTarget> GrassMapHandle = nullptr;
+
+	/** List of the grass types selected for generation. We have to hold their texture index as well, since
+	 * this array could be sparse, but the actual texture array we produce will always have all the grass types.
+	 */
+	TArray<TTuple<TWeakObjectPtr<ULandscapeGrassType>, /*TextureIndex=*/int32>> SelectedGrassTypes;
+
+	/** Total number of grass types used by the landscape component. Includes grass types which were not selected for generation. */
+	int32 NumGrassTypes = 0;
+
+	TWeakObjectPtr<ALandscapeProxy> LandscapeProxy = nullptr;
+	TArray<TWeakObjectPtr<ULandscapeComponent>> LandscapeComponents;
+
+	/** World-space bounds containing all of the landscape components given to the grass weight exporter. */
+	FBox GrassMapBounds = FBox(EForceInit::ForceInit);
+
+	/** Extent (side length) of each landscape component. */
+	double LandscapeComponentExtent = 0.0;
+
+	/** True when we have filtered all of the incoming landscape components down to the ones which overlap the given bounds. */
+	bool bLandscapeComponentsFiltered = false;
+
+	/** True when the landscape components are ready for rendering. */
+	bool bReadyToRender = false;
+
+	/** Textures that we wait to be streamed before generating the grass maps. */
+	TArray<TObjectPtr<UTexture>> TexturesToStream;
+
+	/** True when streaming has been requested on landscape textures. */
+	bool bTextureStreamingRequested = false;
+
+	/** True when grass map generation has been scheduled on the render thread. */
+	bool bGenerationScheduled = false;
+
+protected:
+	virtual void AddExtraStructReferencedObjects(FReferenceCollector& Collector) override;
+};
+
+class FPCGGenerateLandscapeTexturesElement : public IPCGElement
+{
+public:
+	virtual void GetDependenciesCrc(const FPCGGetDependenciesCrcParams& InParams, FPCGCrc& OutCrc) const override;
+	/** FLandscapeGrassWeightExporter expects to exist in scope only on the game thread. */
+	virtual bool CanExecuteOnlyOnMainThread(FPCGContext* Context) const { return true; }
+
+protected:
+	virtual FPCGContext* CreateContext() override;
+	virtual bool ExecuteInternal(FPCGContext* InContext) const override;
+};
diff --git a/Engine/Plugins/PCG/Source/PCGCompute/Internal/PCGGrassMapUnpackerCS.h b/Engine/Plugins/PCG/Source/PCGCompute/Internal/PCGGrassMapUnpackerCS.h
--- a/Engine/Plugins/PCG/Source/PCGCompute/Internal/PCGGrassMapUnpackerCS.h
+++ b/Engine/Plugins/PCG/Source/PCGCompute/Internal/PCGGrassMapUnpackerCS.h
@@ -20,12 +20,15 @@ public:
 
 	BEGIN_SHADER_PARAMETER_STRUCT(FParameters, PCGCOMPUTE_API)
 		SHADER_PARAMETER_RDG_TEXTURE_SRV(Texture2D<float4>, InPackedGrassMaps)
+		SHADER_PARAMETER_RDG_TEXTURE_UAV(RWTexture2D<float>, OutUnpackedHeight)
 		SHADER_PARAMETER_RDG_TEXTURE_UAV(RWTexture2DArray<float>, OutUnpackedGrassMaps)
 		SHADER_PARAMETER_ARRAY(FIntVector4, InLinearTileIndexToComponentIndex, [MaxNumLandscapeComponents])
 		SHADER_PARAMETER(uint32, InNumTilesX)
 		SHADER_PARAMETER(uint32, InLandscapeComponentResolution)
 		SHADER_PARAMETER(uint32, InNumGrassMapPasses)
+		SHADER_PARAMETER(FVector3f, InLandscapeGridScale)
 		SHADER_PARAMETER(FUintVector2, InOutputResolution)
+		SHADER_PARAMETER(uint32, InOutputHeight)
 	END_SHADER_PARAMETER_STRUCT()
 
 public:
