[PCG] Add landscape Z offset to height texture to fix missing offset from heights.

#rnx

diff --git a/Engine/Plugins/PCG/Shaders/Private/PCGGrassMapUnpackerCS.usf b/Engine/Plugins/PCG/Shaders/Private/PCGGrassMapUnpackerCS.usf
--- a/Engine/Plugins/PCG/Shaders/Private/PCGGrassMapUnpackerCS.usf
+++ b/Engine/Plugins/PCG/Shaders/Private/PCGGrassMapUnpackerCS.usf
@@ -42,6 +42,7 @@ uint InNumTilesX;
 uint InLandscapeComponentResolution;
 uint InNumGrassMapPasses;
 float3 InLandscapeGridScale; // x == LS Actor DrawScale.X, y == LS Actor DrawScale.y, z == LS Actor DrawScale.z / 128.0f (ZSCALE)
+float InLandscapeLocationZ;
 uint2 InOutputResolution;
 uint InOutputHeight;
 
@@ -75,8 +76,8 @@ void PCGGrassMapUnpacker_CS(uint3 DispatchThreadId : SV_DispatchThreadID)
 	// Heights. The first layer threads also deal with height.
 	if (InOutputHeight == 1 && DispatchThreadId.z == 0)
 	{
-		const float Height = (UnpackHeight(InPackedGrassMaps[InputCoords].xy) - LANDSCAPE_MID_VALUE) * InLandscapeGridScale.z * TERRAIN_ZSCALE;
-		OutUnpackedHeight[DispatchThreadId.xy] = Height;
+		const float RelativeHeight = (UnpackHeight(InPackedGrassMaps[InputCoords].xy) - LANDSCAPE_MID_VALUE) * InLandscapeGridScale.z * TERRAIN_ZSCALE;
+		OutUnpackedHeight[DispatchThreadId.xy] = RelativeHeight + InLandscapeLocationZ;
 	}
 
 	// Grass map source channel (see packing above).
diff --git a/Engine/Plugins/PCG/Source/PCG/Private/Elements/PCGGenerateLandscapeTextures.cpp b/Engine/Plugins/PCG/Source/PCG/Private/Elements/PCGGenerateLandscapeTextures.cpp
--- a/Engine/Plugins/PCG/Source/PCG/Private/Elements/PCGGenerateLandscapeTextures.cpp
+++ b/Engine/Plugins/PCG/Source/PCG/Private/Elements/PCGGenerateLandscapeTextures.cpp
@@ -502,6 +502,7 @@ bool FPCGGenerateLandscapeTexturesElement::ExecuteInternal(FPCGContext* InContex
 		}
 
 		const FVector LandscapeGridScale = LandscapeProxy->GetRootComponent()->GetRelativeScale3D();
+		const double LandscapeZ = LandscapeProxy->GetRootComponent()->GetRelativeLocation().Z;
 
 		TArray<ULandscapeComponent*> LandscapeComponents;
 		TArray<FIntVector2> LandscapeTileCoords;
@@ -552,7 +553,7 @@ bool FPCGGenerateLandscapeTexturesElement::ExecuteInternal(FPCGContext* InContex
 		
 		TWeakPtr<FPCGContextHandle> ContextHandle = Context->GetOrCreateHandle();
 
-		ENQUEUE_RENDER_COMMAND(GenerateGrassMaps)([ContextHandle, LandscapeComponentExtent=Context->LandscapeComponentExtent, LandscapeTileCoords=MoveTemp(LandscapeTileCoords), ComponentSizeQuads=LandscapeProxy->ComponentSizeQuads, LandscapeGridScale, bGenerateHeightMap=Settings->bGenerateHeightMap](FRHICommandListImmediate& RHICmdList)
+		ENQUEUE_RENDER_COMMAND(GenerateGrassMaps)([ContextHandle, LandscapeComponentExtent=Context->LandscapeComponentExtent, LandscapeTileCoords=MoveTemp(LandscapeTileCoords), ComponentSizeQuads=LandscapeProxy->ComponentSizeQuads, LandscapeGridScale, LandscapeZ, bGenerateHeightMap=Settings->bGenerateHeightMap](FRHICommandListImmediate& RHICmdList)
 		{
 			LLM_SCOPE_BYTAG(PCG);
 
@@ -613,6 +614,7 @@ bool FPCGGenerateLandscapeTexturesElement::ExecuteInternal(FPCGContext* InContex
 			Parameters->OutUnpackedHeight = GrassHeightMapUAV;
 			Parameters->InNumTilesX = NumTilesX;
 			Parameters->InLandscapeGridScale = static_cast<FVector3f>(LandscapeGridScale);
+			Parameters->InLandscapeLocationZ = static_cast<float>(LandscapeZ);
 			Parameters->InLandscapeComponentResolution = LandscapeComponentResolution;
 			Parameters->InOutputResolution = FUintVector2(GrassMapResolution.X, GrassMapResolution.Y);
 			Parameters->InOutputHeight = bGenerateHeightMap ? 1u : 0u;
diff --git a/Engine/Plugins/PCG/Source/PCGCompute/Internal/PCGGrassMapUnpackerCS.h b/Engine/Plugins/PCG/Source/PCGCompute/Internal/PCGGrassMapUnpackerCS.h
--- a/Engine/Plugins/PCG/Source/PCGCompute/Internal/PCGGrassMapUnpackerCS.h
+++ b/Engine/Plugins/PCG/Source/PCGCompute/Internal/PCGGrassMapUnpackerCS.h
@@ -27,6 +27,7 @@ public:
 		SHADER_PARAMETER(uint32, InLandscapeComponentResolution)
 		SHADER_PARAMETER(uint32, InNumGrassMapPasses)
 		SHADER_PARAMETER(FVector3f, InLandscapeGridScale)
+		SHADER_PARAMETER(float, InLandscapeLocationZ)
 		SHADER_PARAMETER(FUintVector2, InOutputResolution)
 		SHADER_PARAMETER(uint32, InOutputHeight)
 	END_SHADER_PARAMETER_STRUCT()
