[PCG] Made the polygon metadata support a bit more robust

#rnx

diff --git a/Engine/Plugins/PCG/Source/PCG/Private/Data/PCGPolygon2DData.cpp b/Engine/Plugins/PCG/Source/PCG/Private/Data/PCGPolygon2DData.cpp
--- a/Engine/Plugins/PCG/Source/PCG/Private/Data/PCGPolygon2DData.cpp
+++ b/Engine/Plugins/PCG/Source/PCG/Private/Data/PCGPolygon2DData.cpp
@@ -269,6 +269,11 @@ void UPCGPolygon2DData::Serialize(FArchive& InArchive)
 	PCGPolygon2DData::SerializePolygon(Polygon, InArchive);
 	InArchive << Transform;
 	InArchive << VertexEntryKeys;
+
+	if (InArchive.IsLoading())
+	{
+		UpdateMappings();
+	}
 }
 
 void UPCGPolygon2DData::AddToCrc(FArchiveCrc32& Ar, bool bFullDataCrc) const
@@ -454,16 +459,18 @@ void UPCGPolygon2DData::WriteMetadataToEntry(float InputKey, PCGMetadataEntryKey
 	else
 	{
 		// invalid index
+		OutEntryKey = PCGInvalidEntryKey;
 		return;
 	}
 
 	// If this fails, it should have failed the pair search, hence the check here.
 	check(VertexEntryKeys.IsValidIndex(CurrentVertexIndex));
+	PCGMetadataEntryKey CurrentVertexEntryKey = VertexEntryKeys[CurrentVertexIndex];
 
 	// Implementation note: there's no need to compute property interpolation if we're directly on the first vertex of the segment.
 	if (InterpolationRatio <= PCGPolygon2DData::Constants::KeyEpsilon)
 	{
-		OutMetadataDomain->InitializeOnSet(OutEntryKey, VertexEntryKeys[CurrentVertexIndex], MetadataDomain);
+		OutMetadataDomain->InitializeOnSet(OutEntryKey, CurrentVertexEntryKey, MetadataDomain);
 		return;
 	}
 
@@ -478,27 +485,31 @@ void UPCGPolygon2DData::WriteMetadataToEntry(float InputKey, PCGMetadataEntryKey
 	else
 	{
 		// Invalid index
-		OutMetadataDomain->InitializeOnSet(OutEntryKey, VertexEntryKeys[CurrentVertexIndex], MetadataDomain);
+		OutMetadataDomain->InitializeOnSet(OutEntryKey, CurrentVertexEntryKey, MetadataDomain);
 		return;
 	}
 
 	// If this fails, it should have failed the pair search, hence the check here.
 	check(VertexEntryKeys.IsValidIndex(NextVertexIndex));
+	PCGMetadataEntryKey NextVertexEntryKey = VertexEntryKeys[NextVertexIndex];
 
 	// Implementation note: there's no need to compute property interpolation if we're directly on the second vertex of the segment.
-	if (InterpolationRatio >= 1.0f - PCGPolygon2DData::Constants::KeyEpsilon)
+	if (InterpolationRatio >= 1.0f - PCGPolygon2DData::Constants::KeyEpsilon || CurrentVertexEntryKey == PCGInvalidEntryKey)
 	{
-		OutMetadataDomain->InitializeOnSet(OutEntryKey, VertexEntryKeys[NextVertexIndex], MetadataDomain);
+		OutMetadataDomain->InitializeOnSet(OutEntryKey, NextVertexEntryKey, MetadataDomain);
 		return;
 	}
 
-	// Finally, perform full interpolation if it's needed.
-	OutMetadataDomain->InitializeOnSet(OutEntryKey, VertexEntryKeys[CurrentVertexIndex], MetadataDomain);
+	// Finally, perform full interpolation if it's needed and valid
+	OutMetadataDomain->InitializeOnSet(OutEntryKey, CurrentVertexEntryKey, MetadataDomain);
 
-	TStaticArray<TPair<PCGMetadataEntryKey, float>, 2> Coefficients;
-	Coefficients[0] = { VertexEntryKeys[CurrentVertexIndex], 1.0f - InterpolationRatio };
-	Coefficients[1] = { VertexEntryKeys[NextVertexIndex], InterpolationRatio };
-	OutMetadataDomain->ComputeWeightedAttribute(OutEntryKey, Coefficients, MetadataDomain);
+	if (CurrentVertexEntryKey != PCGInvalidEntryKey && NextVertexEntryKey != PCGInvalidEntryKey)
+	{
+		TStaticArray<TPair<PCGMetadataEntryKey, float>, 2> Coefficients;
+		Coefficients[0] = { CurrentVertexEntryKey, 1.0f - InterpolationRatio };
+		Coefficients[1] = { NextVertexEntryKey, InterpolationRatio };
+		OutMetadataDomain->ComputeWeightedAttribute(OutEntryKey, Coefficients, MetadataDomain);
+	}
 }
 
 UPCGSpatialData* UPCGPolygon2DData::CopyInternal(FPCGContext* Context) const
diff --git a/Engine/Plugins/PCG/Source/PCG/Private/Helpers/PCGPolygon2DProcessingHelpers.cpp b/Engine/Plugins/PCG/Source/PCG/Private/Helpers/PCGPolygon2DProcessingHelpers.cpp
--- a/Engine/Plugins/PCG/Source/PCG/Private/Helpers/PCGPolygon2DProcessingHelpers.cpp
+++ b/Engine/Plugins/PCG/Source/PCG/Private/Helpers/PCGPolygon2DProcessingHelpers.cpp
@@ -393,12 +393,12 @@ TArray<PCGMetadataEntryKey> ComputeEntryKeysAndInitializeAttributes(UPCGMetadata
 		PCGMetadataEntryKey& NewKey = EntryKeys.Add_GetRef(PCGInvalidEntryKey);
 
 		const UE::Geometry::FGeneralPolygon2d* SourcePolygon = VertexMapping.Poly;
-		if (!SourcePolygon)
+		if (!SourcePolygon || !InInfoMap.Contains(SourcePolygon))
 		{
 			continue;
 		}
 
-		const auto [SourceMetadata, SourceEntryKeys] = InInfoMap[SourcePolygon];
+		const auto& [SourceMetadata, SourceEntryKeys] = InInfoMap[SourcePolygon];
 
 		// Propagate attributes from source metadata to out metadata
 		if (!VisitedMetadata.Contains(SourceMetadata))
@@ -412,12 +412,32 @@ TArray<PCGMetadataEntryKey> ComputeEntryKeysAndInitializeAttributes(UPCGMetadata
 			continue;
 		}
 
-		OutMetadata->InitializeOnSet(NewKey);
+		if (!SourceEntryKeys->IsValidIndex(VertexMapping.StartVertexIndex) ||
+			!SourceEntryKeys->IsValidIndex(VertexMapping.EndVertexIndex))
+		{
+			continue;
+		}
+
+		const auto& StartVertexEntryKey = (*SourceEntryKeys)[VertexMapping.StartVertexIndex];
+		const auto& EndVertexEntryKey = (*SourceEntryKeys)[VertexMapping.EndVertexIndex];
+
+		if (StartVertexEntryKey != PCGInvalidEntryKey && EndVertexEntryKey != PCGInvalidEntryKey)
+		{
+			OutMetadata->InitializeOnSet(NewKey);
 
-		TStaticArray<TPair<PCGMetadataEntryKey, float>, 2> Coefficients;
-		Coefficients[0] = { (*SourceEntryKeys)[VertexMapping.StartVertexIndex], 1.0f - VertexMapping.Ratio };
-		Coefficients[1] = { (*SourceEntryKeys)[VertexMapping.EndVertexIndex], VertexMapping.Ratio };
-		OutMetadata->ComputeWeightedAttribute(NewKey, Coefficients, SourceMetadata);
+			TStaticArray<TPair<PCGMetadataEntryKey, float>, 2> Coefficients;
+			Coefficients[0] = { StartVertexEntryKey, 1.0f - VertexMapping.Ratio };
+			Coefficients[1] = { EndVertexEntryKey, VertexMapping.Ratio };
+			OutMetadata->ComputeWeightedAttribute(NewKey, Coefficients, SourceMetadata);
+		}
+		else if (StartVertexEntryKey != PCGInvalidEntryKey)
+		{
+			OutMetadata->InitializeOnSet(NewKey, StartVertexEntryKey, SourceMetadata);
+		}
+		else if (EndVertexEntryKey != PCGInvalidEntryKey)
+		{
+			OutMetadata->InitializeOnSet(NewKey, EndVertexEntryKey, SourceMetadata);
+		}
 	}
 
 	return EntryKeys;
