[PCG] Extend coverage of LLM PCG tag scopes.

Missed previously as we were not enabling memory trace immediately from startup which misses allocations and also makes memory traces unreliable.

#rnx

diff --git a/Engine/Plugins/PCG/Source/PCG/Private/Data/PCGDataPtrWrapper.cpp b/Engine/Plugins/PCG/Source/PCG/Private/Data/PCGDataPtrWrapper.cpp
--- a/Engine/Plugins/PCG/Source/PCG/Private/Data/PCGDataPtrWrapper.cpp
+++ b/Engine/Plugins/PCG/Source/PCG/Private/Data/PCGDataPtrWrapper.cpp
@@ -115,6 +115,8 @@ void FPCGDataPtrWrapper::DecRefCount()
 
 bool FPCGDataPtrWrapper::Serialize(FArchive& Ar)
 {
+	LLM_SCOPE_BYTAG(PCG);
+
 	TWeakObjectPtr<const UPCGData> WeakDataBeforeSerialization = WeakData;
 	Ar << Data;
 
diff --git a/Engine/Plugins/PCG/Source/PCG/Private/Data/PCGPolygon2DData.cpp b/Engine/Plugins/PCG/Source/PCG/Private/Data/PCGPolygon2DData.cpp
--- a/Engine/Plugins/PCG/Source/PCG/Private/Data/PCGPolygon2DData.cpp
+++ b/Engine/Plugins/PCG/Source/PCG/Private/Data/PCGPolygon2DData.cpp
@@ -262,6 +262,7 @@ bool FPCGDataTypeInfoPolygon2D::SupportsConversionTo(const FPCGDataTypeIdentifie
 
 void UPCGPolygon2DData::Serialize(FArchive& InArchive)
 {
+	LLM_SCOPE_BYTAG(PCG);
 	Super::Serialize(InArchive);
 
 	InArchive.UsingCustomVersion(FPCGCustomVersion::GUID);
diff --git a/Engine/Plugins/PCG/Source/PCG/Private/Grid/PCGLandscapeCache.cpp b/Engine/Plugins/PCG/Source/PCG/Private/Grid/PCGLandscapeCache.cpp
--- a/Engine/Plugins/PCG/Source/PCG/Private/Grid/PCGLandscapeCache.cpp
+++ b/Engine/Plugins/PCG/Source/PCG/Private/Grid/PCGLandscapeCache.cpp
@@ -448,6 +448,8 @@ bool FPCGLandscapeCacheEntry::SerializeFromBulkData() const
 
 void FPCGLandscapeCacheEntry::Serialize(FArchive& Archive, UObject* Owner, int32 Index, EPCGLandscapeCacheSerializationContents SerializeContents)
 {
+	LLM_SCOPE_BYTAG(PCG);
+
 	// If the current bulk output data does not match the desired save data, make sure to load from the bulk data and update the bulked save data
 	if(Archive.IsCooking() && SerializeContents != EPCGLandscapeCacheSerializationContents::SerializeAll && !bDataLoaded)
 	{
@@ -499,6 +501,7 @@ void UPCGLandscapeCache::PostEditChangeProperty(FPropertyChangedEvent& PropertyC
 
 void UPCGLandscapeCache::Serialize(FArchive& Archive)
 {
+	LLM_SCOPE_BYTAG(PCG);
 	Archive.UsingCustomVersion(FFortniteMainBranchObjectVersion::GUID);
 
 	if (Archive.IsLoading() && Archive.CustomVer(FFortniteMainBranchObjectVersion::GUID) < FFortniteMainBranchObjectVersion::PCGLandscapeCacheDefaultSerializationChanged)
diff --git a/Engine/Plugins/PCG/Source/PCG/Private/Grid/PCGPartitionActor.cpp b/Engine/Plugins/PCG/Source/PCG/Private/Grid/PCGPartitionActor.cpp
--- a/Engine/Plugins/PCG/Source/PCG/Private/Grid/PCGPartitionActor.cpp
+++ b/Engine/Plugins/PCG/Source/PCG/Private/Grid/PCGPartitionActor.cpp
@@ -99,6 +99,8 @@ void APCGPartitionActor::BeginDestroy()
 
 void APCGPartitionActor::Serialize(FArchive& Ar)
 {
+	LLM_SCOPE_BYTAG(PCG);
+
 #if WITH_EDITOR
 	TMap<TObjectPtr<UPCGComponent>, TSoftObjectPtr<UPCGComponent>> LocalToOriginalCopy;
 	TMap<TObjectPtr<UPCGComponent>, TSoftObjectPtr<UPCGComponent>> Transients;
diff --git a/Engine/Plugins/PCG/Source/PCG/Private/Metadata/PCGMetadata.cpp b/Engine/Plugins/PCG/Source/PCG/Private/Metadata/PCGMetadata.cpp
--- a/Engine/Plugins/PCG/Source/PCG/Private/Metadata/PCGMetadata.cpp
+++ b/Engine/Plugins/PCG/Source/PCG/Private/Metadata/PCGMetadata.cpp
@@ -335,6 +335,7 @@ void UPCGMetadata::SetupDomainsFromPCGDataType(const TSubclassOf<UPCGData>& PCGD
 
 void UPCGMetadata::Serialize(FArchive& InArchive)
 {
+	LLM_SCOPE_BYTAG(PCG);
 	Super::Serialize(InArchive);
 
 	InArchive.UsingCustomVersion(FPCGCustomVersion::GUID);
diff --git a/Engine/Plugins/PCG/Source/PCG/Private/Metadata/PCGMetadataDomain.cpp b/Engine/Plugins/PCG/Source/PCG/Private/Metadata/PCGMetadataDomain.cpp
--- a/Engine/Plugins/PCG/Source/PCG/Private/Metadata/PCGMetadataDomain.cpp
+++ b/Engine/Plugins/PCG/Source/PCG/Private/Metadata/PCGMetadataDomain.cpp
@@ -147,6 +147,8 @@ namespace PCGMetadata
 
 void FPCGMetadataDomain::Serialize(FArchive& InArchive)
 {
+	LLM_SCOPE_BYTAG(PCG);
+
 	int32 NumAttributes = (InArchive.IsLoading() ? 0 : Attributes.Num());
 	// We need to keep track of the max attribute Id, since it won't necessary be equal to the number of attributes + 1.
 	int64 MaxAttributeId = -1;
diff --git a/Engine/Plugins/PCG/Source/PCG/Private/PCGComponent.cpp b/Engine/Plugins/PCG/Source/PCG/Private/PCGComponent.cpp
--- a/Engine/Plugins/PCG/Source/PCG/Private/PCGComponent.cpp
+++ b/Engine/Plugins/PCG/Source/PCG/Private/PCGComponent.cpp
@@ -1849,6 +1849,7 @@ void UPCGComponent::ClearGraphGeneratedOutput()
 
 void UPCGComponent::BeginPlay()
 {
+	LLM_SCOPE_BYTAG(PCG);
 	Super::BeginPlay();
 
 #if WITH_EDITOR
@@ -2064,6 +2065,8 @@ void UPCGComponent::AddReferencedObjects(UObject* InThis, FReferenceCollector& C
 
 void UPCGComponent::Serialize(FArchive& Ar)
 {
+	LLM_SCOPE_BYTAG(PCG);
+
 #if WITH_EDITOR
 	TArray<TObjectPtr<UPCGManagedResource>> GeneratedResourcesCopy;
 
@@ -2116,6 +2119,7 @@ void UPCGComponent::Serialize(FArchive& Ar)
 
 void UPCGComponent::PostLoad()
 {
+	LLM_SCOPE_BYTAG(PCG);
 	Super::PostLoad();
 
 #if WITH_EDITOR
@@ -2272,6 +2276,8 @@ bool UPCGComponent::IsEditorOnly() const
 
 void UPCGComponent::PostInitProperties()
 {
+	LLM_SCOPE_BYTAG(PCG);
+
 #if WITH_EDITOR
 	GraphInstance->OnGraphChangedDelegate.AddUObject(this, &UPCGComponent::OnGraphChanged);
 #endif // WITH_EDITOR
@@ -2281,6 +2287,7 @@ void UPCGComponent::PostInitProperties()
 
 void UPCGComponent::OnRegister()
 {
+	LLM_SCOPE_BYTAG(PCG);
 	Super::OnRegister();
 
 #if WITH_EDITOR
@@ -2430,6 +2437,8 @@ void UPCGComponent::PreEditChange(FProperty* PropertyAboutToChange)
 
 void UPCGComponent::PostEditChangeProperty(FPropertyChangedEvent& PropertyChangedEvent)
 {
+	LLM_SCOPE_BYTAG(PCG);
+
 	if (!PropertyChangedEvent.Property || !IsValidChecked(this))
 	{
 		Super::PostEditChangeProperty(PropertyChangedEvent);
@@ -2582,6 +2591,7 @@ void UPCGComponent::PostEditChangeProperty(FPropertyChangedEvent& PropertyChange
 
 void UPCGComponent::PostEditImport()
 {
+	LLM_SCOPE_BYTAG(PCG);
 	Super::PostEditImport();
 
 	SetupCallbacksOnCreation();
diff --git a/Engine/Plugins/PCG/Source/PCG/Private/PCGPointArray.cpp b/Engine/Plugins/PCG/Source/PCG/Private/PCGPointArray.cpp
--- a/Engine/Plugins/PCG/Source/PCG/Private/PCGPointArray.cpp
+++ b/Engine/Plugins/PCG/Source/PCG/Private/PCGPointArray.cpp
@@ -2,6 +2,8 @@
 
 #include "PCGPointArray.h"
 
+#include "PCGModule.h"
+
 #include UE_INLINE_GENERATED_CPP_BY_NAME(PCGPointArray)
 
 FPCGPointArray::FPCGPointArray()
@@ -141,6 +143,8 @@ TArray<FTransform> FPCGPointArray::GetTransformCopy() const
 
 bool FPCGPointArray::Serialize(FArchive& Ar)
 {
+	LLM_SCOPE_BYTAG(PCG);
+
 	Ar << NumPoints;
 	Transform.Serialize(Ar);
 	Density.Serialize(Ar);
diff --git a/Engine/Plugins/PCG/Source/PCG/Private/PCGSettings.cpp b/Engine/Plugins/PCG/Source/PCG/Private/PCGSettings.cpp
--- a/Engine/Plugins/PCG/Source/PCG/Private/PCGSettings.cpp
+++ b/Engine/Plugins/PCG/Source/PCG/Private/PCGSettings.cpp
@@ -344,6 +344,7 @@ void UPCGSettings::PostInitProperties()
 
 void UPCGSettings::Serialize(FArchive& Ar)
 {
+	LLM_SCOPE_BYTAG(PCG);
 	Super::Serialize(Ar);
 
 	Ar.UsingCustomVersion(FPCGCustomVersion::GUID);
