[PCG] Add Raw Buffer Data which is simply an array of uint32s.

Use it to avoid constructing UPCGParamData for simple analysis data to save chunk of time during primitive instancing.

Also fix bug where internal kernel pins could be assigned an incorrect type before data interfaces are created.

Disabled by default on CVar.

#rnx

diff --git a/Engine/Plugins/PCG/Shaders/Private/BuiltInKernels/PCGCountUniqueAttributeValues.usf b/Engine/Plugins/PCG/Shaders/Private/BuiltInKernels/PCGCountUniqueAttributeValues.usf
--- a/Engine/Plugins/PCG/Shaders/Private/BuiltInKernels/PCGCountUniqueAttributeValues.usf
+++ b/Engine/Plugins/PCG/Shaders/Private/BuiltInKernels/PCGCountUniqueAttributeValues.usf
@@ -2,32 +2,56 @@
 
 // Counts how many times each unique value of an attribute is encountered in the input data elements.
 
+#define PCG_OUTPUT_TO_RAW_BUFFER {RawDataOutputEnabled}
+
 // Input: Input elements with attribute.
-// Output: Attribute set with a ValueCount attribute which has the count of each value. Number of elements is set to the number of unique values (known CPU-side).
+// Output (PCG_OUTPUT_TO_RAW_BUFFER): Uint array consisting of (string key value, value instance count) pairs. Number of pairs is the number of unique values (known CPU-side).
+// Output (!PCG_OUTPUT_TO_RAW_BUFFER): Attribute set with a ValueCount attribute which has the count of each value. Number of elements is the number of unique values (known CPU-side).
+
+void InitializeValueKey(uint InDataIndex, uint InValueKey)
+{
+#if PCG_OUTPUT_TO_RAW_BUFFER
+	const uint FirstElementIndex = InValueKey * 2u;
+	if (FirstElementIndex < Out_GetNumElements(0))
+	{
+		// We currently have an output element per possible string key value. This may be compacted in the future.
+		Out_Store(InDataIndex, FirstElementIndex, InValueKey);
+	}
+#else
+	// Write value of attribute. We currently have an output element per possible string key value.
+ 	const int OutputValueAttributeId = CountUniqueValues_GetOutputValueAttributeId();
+ 	Out_SetInt(InDataIndex, InValueKey, OutputValueAttributeId, InValueKey);
+#endif
+}
 
 [numthreads(64, 1, 1)]
 void PCGCountUniqueAttributeValuesCS(uint3 GroupId : SV_GroupID, uint GroupIndex : SV_GroupIndex)
 {
+#if !PCG_OUTPUT_TO_RAW_BUFFER
 	// Mark the kernel as having executed. Must run before we early out via thread index, because the kernel is still 'executed' even if the number of
 	// threads to iterate on is zero. Even if GetNumThreads() returns 0, the kernel will still have been dispatched on a single thread to set this flag.
 	if (all(GroupId == 0) && GroupIndex == 0)
 	{
 		Out_SetAsExecutedInternal();
 	}
+#endif
 
 	const uint ThreadIndex = GetUnWrappedDispatchThreadId(GroupId, GroupIndex, 64);
 
 	// Pre-job - if this thread maps to any element in the output, initialize the value key. The dispatch thread count will be large enough to ensure
 	// all output elements are initialized.
 	{
+#if PCG_OUTPUT_TO_RAW_BUFFER
+		// If outputting to raw buffer, always single data.
+		InitializeValueKey(/*DataIndex=*/0, ThreadIndex);
+#else
 		uint Out_DataIndex;
 		uint Out_ElementIndex;
 		if (Out_GetThreadData(ThreadIndex, Out_DataIndex, Out_ElementIndex))
 		{
-			// Write value of attribute. We currently have an output element per possible string key value.
- 			const int OutputValueAttributeId = CountUniqueValues_GetOutputValueAttributeId();
- 			Out_SetInt(Out_DataIndex, Out_ElementIndex, OutputValueAttributeId, Out_ElementIndex);
+			InitializeValueKey(Out_DataIndex, Out_ElementIndex);
 		}
+#endif
 	}
 
 	// Main job - one thread per input data element. Read the string key value and atomic increment the output counter.
@@ -56,6 +80,10 @@ void PCGCountUniqueAttributeValuesCS(uint3 GroupId : SV_GroupID, uint GroupIndex
 		// If not outputting per-data counts, then all counters are aggregated on a single data.
 		const uint Out_DataIndex = CountUniqueValues_GetEmitPerDataCounts() ? In_DataIndex : 0;
 
+#if PCG_OUTPUT_TO_RAW_BUFFER
+		const uint ElementIndex = AttributeValue * 2u + 1;
+		Out_AtomicAdd(/*DataIndex=*/0, ElementIndex, 1u);
+#else
 		// Always one bucket per attribute value.
 		const uint Out_ElementIndex = AttributeValue;
 
@@ -63,5 +91,8 @@ void PCGCountUniqueAttributeValuesCS(uint3 GroupId : SV_GroupID, uint GroupIndex
 
 		// Increment count of this attribute.
 		Out_AtomicAddInt(Out_DataIndex, Out_ElementIndex, OutputCountAttributeId, 1);
+#endif
 	}
 }
+
+#undef PCG_OUTPUT_TO_RAW_BUFFER
diff --git a/Engine/Plugins/PCG/Shaders/Private/PCGRawBufferDataInterface.ush b/Engine/Plugins/PCG/Shaders/Private/PCGRawBufferDataInterface.ush
new file mode 100644
--- /dev/null
+++ b/Engine/Plugins/PCG/Shaders/Private/PCGRawBufferDataInterface.ush
@@ -0,0 +1,60 @@
+// Copyright Epic Games, Inc. All Rights Reserved.
+
+/**
+ * Raw buffer: Flat list of raw bytes, read/written in uints.
+ * 
+ * No multi-data support currently. Would need software-bindless, in order to keep each buffer pure rather than adding headers.
+ */
+
+// @todo_pcg: Support having a non-RW byte address buffer for input data.
+RWByteAddressBuffer {DataInterfaceName}_Data;
+
+// Number of bytes in buffer.
+uint {DataInterfaceName}_SizeBytes;
+
+// #################### META INFO ##########################
+
+uint GetNumData_{DataInterfaceName}()
+{
+	// Multi-data not currently supported. Could be done with multiple buffers (each individual buffer should be a pure raw buffer).
+	return 1u;
+}
+
+uint GetNumElements_{DataInterfaceName}(uint InDataIndex)
+{
+	// InDataIndex unused for now as multi-data not currently supported.
+	return {DataInterfaceName}_SizeBytes >> 2;
+}
+
+// #################### ACCESSORS ##########################
+
+uint Load_{DataInterfaceName}(uint InDataIndex, uint InElementIndex)
+{
+	return {DataInterfaceName}_Data.Load(InElementIndex << 2);
+}
+
+uint4 Load4_{DataInterfaceName}(uint InDataIndex, uint InFirstElementIndex)
+{
+	return {DataInterfaceName}_Data.Load4(InFirstElementIndex << 2);
+}
+
+void Store_{DataInterfaceName}(uint InDataIndex, uint InElementIndex, uint InValue)
+{
+	{DataInterfaceName}_Data.Store(InElementIndex << 2, InValue);
+}
+
+void Store4_{DataInterfaceName}(uint InDataIndex, uint InFirstElementIndex, uint4 InValue)
+{
+	{DataInterfaceName}_Data.Store4(InFirstElementIndex << 2, InValue);
+}
+
+// #################### ATOMICS ##########################
+
+// Returns value of attribute before incrementing.
+uint AtomicAdd_{DataInterfaceName}(uint InDataIndex, uint InElementIndex, uint InValueToAdd)
+{
+	uint OriginalValue;
+	{DataInterfaceName}_Data.InterlockedAdd(InElementIndex << 2, (uint)InValueToAdd, OriginalValue);
+
+	return OriginalValue;
+}
diff --git a/Engine/Plugins/PCG/Source/PCG/Internal/Compute/PCGDataForGPU.h b/Engine/Plugins/PCG/Source/PCG/Internal/Compute/PCGDataForGPU.h
--- a/Engine/Plugins/PCG/Source/PCG/Internal/Compute/PCGDataForGPU.h
+++ b/Engine/Plugins/PCG/Source/PCG/Internal/Compute/PCGDataForGPU.h
@@ -164,7 +164,7 @@ public:
 	FPCGDataDesc(FPCGDataTypeIdentifier InType, FIntVector4 InElementCount);
 	FPCGDataDesc(const FPCGTaggedData& InTaggedData, const UPCGDataBinding* InBinding);
 
-	uint64 ComputePackedSize() const;
+	uint64 ComputePackedSizeBytes() const;
 
 	bool HasElementsMetadataDomainAttributes() const;
 
diff --git a/Engine/Plugins/PCG/Source/PCG/Private/Compute/BuiltInKernels/PCGCountUniqueAttributeValuesKernel.cpp b/Engine/Plugins/PCG/Source/PCG/Private/Compute/BuiltInKernels/PCGCountUniqueAttributeValuesKernel.cpp
--- a/Engine/Plugins/PCG/Source/PCG/Private/Compute/BuiltInKernels/PCGCountUniqueAttributeValuesKernel.cpp
+++ b/Engine/Plugins/PCG/Source/PCG/Private/Compute/BuiltInKernels/PCGCountUniqueAttributeValuesKernel.cpp
@@ -4,10 +4,12 @@
 
 #include "PCGContext.h"
 #include "PCGModule.h"
+#include "PCGParamData.h"
 #include "Compute/Elements/PCGComputeGraphElement.h"
 #include "Compute/PCGComputeGraph.h"
 #include "Compute/PCGDataBinding.h"
 #include "Compute/PCGPinPropertiesGPU.h"
+#include "Compute/Data/PCGRawBufferData.h"
 #include "Compute/DataInterfaces/BuiltInKernels/PCGCountUniqueAttributeValuesDataInterface.h"
 #include "Graph/PCGGPUGraphCompilationContext.h"
 
@@ -119,9 +121,21 @@ bool UPCGCountUniqueAttributeValuesKernel::ComputeOutputBindingDataDesc(const UP
 		const int32 ElementCount = (MaxStringKeyValue ? *MaxStringKeyValue : 0) + 1;
 
 		// Base data description object, copied N times below to populate the final description.
-		FPCGDataDesc CountDataDesc(EPCGDataType::Param, ElementCount);
-		CountDataDesc.AddAttribute(FPCGKernelAttributeKey(PCGCountUniqueAttributeValuesConstants::ValueAttributeName, EPCGKernelAttributeType::Int), InBinding);
-		CountDataDesc.AddAttribute(FPCGKernelAttributeKey(PCGCountUniqueAttributeValuesConstants::ValueCountAttributeName, EPCGKernelAttributeType::Int), InBinding);
+		FPCGDataDesc CountDataDesc;
+		
+		if (bOutputRawBuffer)
+		{
+			CountDataDesc = FPCGDataDesc(FPCGDataTypeInfoRawBuffer::AsId(), ElementCount * 2);
+		}
+		else
+		{
+			CountDataDesc = FPCGDataDesc(EPCGDataType::Param, ElementCount);
+			CountDataDesc.AddAttribute(FPCGKernelAttributeKey(PCGCountUniqueAttributeValuesConstants::ValueAttributeName, EPCGKernelAttributeType::Int), InBinding);
+			CountDataDesc.AddAttribute(FPCGKernelAttributeKey(PCGCountUniqueAttributeValuesConstants::ValueCountAttributeName, EPCGKernelAttributeType::Int), InBinding);
+		}
+
+		// Raw buffer output not supported for multiple data currently.
+		ensure(!bEmitPerDataCounts || !bOutputRawBuffer);
 
 		if (bEmitPerDataCounts)
 		{
@@ -175,14 +189,30 @@ bool UPCGCountUniqueAttributeValuesKernel::DoesOutputPinRequireZeroInitializatio
 }
 
 #if WITH_EDITOR
+FString UPCGCountUniqueAttributeValuesKernel::GetCookedSource(FPCGGPUCompilationContext& InOutContext) const
+{
+	FString SourceFile;
+	ensure(LoadShaderSourceFile(GetSourceFilePath(), EShaderPlatform::SP_PCD3D_SM5, &SourceFile, nullptr));
+	
+	TMap<FString, FStringFormatArg> TemplateArgs =
+	{
+		{ TEXT("RawDataOutputEnabled"), bOutputRawBuffer ? TEXT("1") : TEXT("0") },
+	};
+
+	SourceFile = FString::Format(*SourceFile, TemplateArgs);
+
+	return SourceFile;
+}
+
 void UPCGCountUniqueAttributeValuesKernel::CreateAdditionalInputDataInterfaces(FPCGGPUCompilationContext& InOutContext, UObject* InObjectOuter, TArray<TObjectPtr<UComputeDataInterface>>& OutDataInterfaces) const
 {
 	Super::CreateAdditionalInputDataInterfaces(InOutContext, InObjectOuter, OutDataInterfaces);
 
 	TObjectPtr<UPCGCountUniqueAttributeValuesDataInterface> KernelDI = InOutContext.NewObject_AnyThread<UPCGCountUniqueAttributeValuesDataInterface>(InObjectOuter);
 	KernelDI->ProducerKernel = this;
-	KernelDI->AttributeToCountName = AttributeName;
-	KernelDI->bEmitPerDataCounts = bEmitPerDataCounts;
+	KernelDI->SetAttributeToCountName(AttributeName);
+	KernelDI->SetEmitPerDataCounts(bEmitPerDataCounts);
+	KernelDI->SetOutputRawBuffer(bOutputRawBuffer);
 
 	OutDataInterfaces.Add(KernelDI);
 }
@@ -191,8 +221,11 @@ void UPCGCountUniqueAttributeValuesKernel::CreateAdditionalInputDataInterfaces(F
 void UPCGCountUniqueAttributeValuesKernel::GetKernelAttributeKeys(TArray<FPCGKernelAttributeKey>& OutKeys) const
 {
 	// Register the attribute this node creates.
-	OutKeys.AddUnique(FPCGKernelAttributeKey(PCGCountUniqueAttributeValuesConstants::ValueAttributeName, EPCGKernelAttributeType::Int));
-	OutKeys.AddUnique(FPCGKernelAttributeKey(PCGCountUniqueAttributeValuesConstants::ValueCountAttributeName, EPCGKernelAttributeType::Int));
+	if (!bOutputRawBuffer)
+	{
+		OutKeys.AddUnique(FPCGKernelAttributeKey(PCGCountUniqueAttributeValuesConstants::ValueAttributeName, EPCGKernelAttributeType::Int));
+		OutKeys.AddUnique(FPCGKernelAttributeKey(PCGCountUniqueAttributeValuesConstants::ValueCountAttributeName, EPCGKernelAttributeType::Int));
+	}
 }
 
 void UPCGCountUniqueAttributeValuesKernel::GetInputPins(TArray<FPCGPinProperties>& OutPins) const
@@ -203,7 +236,7 @@ void UPCGCountUniqueAttributeValuesKernel::GetInputPins(TArray<FPCGPinProperties
 void UPCGCountUniqueAttributeValuesKernel::GetOutputPins(TArray<FPCGPinPropertiesGPU>& OutPins) const
 {
 	// Attribute set with a value count attribute, element count equal to number of unique values of the counted attribute.
-	OutPins.Emplace(PCGPinConstants::DefaultOutputLabel, EPCGDataType::Param);
+	OutPins.Emplace(PCGPinConstants::DefaultOutputLabel, bOutputRawBuffer ? FPCGDataTypeInfoRawBuffer::AsId() : FPCGDataTypeInfoParam::AsId());
 }
 
 #undef LOCTEXT_NAMESPACE
diff --git a/Engine/Plugins/PCG/Source/PCG/Private/Compute/BuiltInKernels/PCGCountUniqueAttributeValuesKernel.h b/Engine/Plugins/PCG/Source/PCG/Private/Compute/BuiltInKernels/PCGCountUniqueAttributeValuesKernel.h
--- a/Engine/Plugins/PCG/Source/PCG/Private/Compute/BuiltInKernels/PCGCountUniqueAttributeValuesKernel.h
+++ b/Engine/Plugins/PCG/Source/PCG/Private/Compute/BuiltInKernels/PCGCountUniqueAttributeValuesKernel.h
@@ -25,13 +25,17 @@ public:
 	void SetAttributeName(FName InAttributeName) { AttributeName = InAttributeName; }
 
 	/** Whether to produce one set of element counts that count across all input data, rather than producing counters per input data. */
-	virtual void SetEmitPerDataCounts(bool bInEmitPerDataCounts) { bEmitPerDataCounts = bInEmitPerDataCounts; }
+	void SetEmitPerDataCounts(bool bInEmitPerDataCounts) { bEmitPerDataCounts = bInEmitPerDataCounts; }
+
+	/** Whether to output a raw array of (string key value, instance count) values from analysis instead of constructing an attribute set. */
+	void SetOutputRawBuffer(bool bInOutputRawBuffer) { bOutputRawBuffer = bInOutputRawBuffer; }
 
 	virtual bool IsKernelDataValid(FPCGContext* InContext) const override;
 	virtual bool ComputeOutputBindingDataDesc(const UPCGComputeGraph* InGraph, FName InOutputPinLabel, UPCGDataBinding* InBinding, FPCGDataCollectionDesc& OutDataDesc) const override;
 	virtual int ComputeThreadCount(const UPCGDataBinding* Binding) const override;
 	virtual bool DoesOutputPinRequireZeroInitialization(FName InOutputPinLabel) const override;
 #if WITH_EDITOR
+	virtual FString GetCookedSource(FPCGGPUCompilationContext& InOutContext) const override;
 	virtual const TCHAR* GetSourceFilePath() const override { return TEXT("/Plugin/PCG/Private/BuiltInKernels/PCGCountUniqueAttributeValues.usf"); }
 	virtual FString GetEntryPoint() const override { return TEXT("PCGCountUniqueAttributeValuesCS"); }
 	void CreateAdditionalInputDataInterfaces(FPCGGPUCompilationContext& InOutContext, UObject* InObjectOuter, TArray<TObjectPtr<UComputeDataInterface>>& OutDataInterfaces) const override;
@@ -43,11 +47,14 @@ public:
 	virtual void GetInputPins(TArray<FPCGPinProperties>& OutPins) const override;
 	virtual void GetOutputPins(TArray<FPCGPinPropertiesGPU>& OutPins) const override;
 
-public:
+protected:
 	UPROPERTY()
 	FName AttributeName;
 
 	/** Produce one set of element counts that count across all input data, rather than producing counters per input data. */
 	UPROPERTY()
 	bool bEmitPerDataCounts = true;
+
+	UPROPERTY()
+	bool bOutputRawBuffer = false;
 };
diff --git a/Engine/Plugins/PCG/Source/PCG/Private/Compute/Data/PCGRawBufferData.cpp b/Engine/Plugins/PCG/Source/PCG/Private/Compute/Data/PCGRawBufferData.cpp
new file mode 100644
--- /dev/null
+++ b/Engine/Plugins/PCG/Source/PCG/Private/Compute/Data/PCGRawBufferData.cpp
@@ -0,0 +1,17 @@
+// Copyright Epic Games, Inc. All Rights Reserved.
+
+#include "Compute/Data/PCGRawBufferData.h"
+
+#include "PCGModule.h"
+
+#include UE_INLINE_GENERATED_CPP_BY_NAME(PCGRawBufferData)
+
+PCG_DEFINE_TYPE_INFO(FPCGDataTypeInfoRawBuffer, UPCGRawBufferData)
+
+void UPCGRawBufferData::AddToCrc(FArchiveCrc32& Ar, bool bFullDataCrc) const
+{
+	Super::AddToCrc(Ar, bFullDataCrc);
+
+	// This data does not have a bespoke CRC implementation so just use the unique object instance UID.
+	AddUIDToCrc(Ar);
+}
diff --git a/Engine/Plugins/PCG/Source/PCG/Private/Compute/Data/PCGRawBufferData.h b/Engine/Plugins/PCG/Source/PCG/Private/Compute/Data/PCGRawBufferData.h
new file mode 100644
--- /dev/null
+++ b/Engine/Plugins/PCG/Source/PCG/Private/Compute/Data/PCGRawBufferData.h
@@ -0,0 +1,43 @@
+// Copyright Epic Games, Inc. All Rights Reserved.
+
+#pragma once
+
+#include "PCGData.h"
+
+#include "PCGRawBufferData.generated.h"
+
+USTRUCT()
+struct FPCGDataTypeInfoRawBuffer : public FPCGDataTypeInfo
+{
+	GENERATED_BODY()
+
+	PCG_DECLARE_TYPE_INFO()
+
+#if WITH_EDITOR
+	virtual bool Hidden() const override { return true; }
+#endif // WITH_EDITOR
+};
+
+/** A proxy for data residing on the GPU with functionality to read the data back to the CPU. */
+UCLASS(MinimalAPI, ClassGroup = (Procedural), DisplayName = "Raw Buffer")
+class UPCGRawBufferData : public UPCGData
+{
+	GENERATED_BODY()
+
+public:
+	//~Begin UPCGData interface
+	PCG_ASSIGN_TYPE_INFO(FPCGDataTypeInfoRawBuffer);
+
+	PCG_API virtual void AddToCrc(FArchiveCrc32& Ar, bool bFullDataCrc) const override;
+	//~End UPCGData interface
+
+	void SetData(TArray<uint32>&& InData) { Data = MoveTemp(InData); }
+	int32 GetNumUint32s() const { return Data.Num(); }
+
+	TArray<uint32>& GetData() { return Data; }
+	const TArray<uint32>& GetConstData() const { return Data; }
+
+protected:
+	UPROPERTY()
+	TArray<uint32> Data;
+};
diff --git a/Engine/Plugins/PCG/Source/PCG/Private/Compute/DataInterfaces/BuiltInKernels/PCGCountUniqueAttributeValuesDataInterface.cpp b/Engine/Plugins/PCG/Source/PCG/Private/Compute/DataInterfaces/BuiltInKernels/PCGCountUniqueAttributeValuesDataInterface.cpp
--- a/Engine/Plugins/PCG/Source/PCG/Private/Compute/DataInterfaces/BuiltInKernels/PCGCountUniqueAttributeValuesDataInterface.cpp
+++ b/Engine/Plugins/PCG/Source/PCG/Private/Compute/DataInterfaces/BuiltInKernels/PCGCountUniqueAttributeValuesDataInterface.cpp
@@ -99,8 +99,9 @@ void UPCGCountUniqueAttributeValuesDataProvider::Initialize(const UComputeDataIn
 	const UPCGCountUniqueAttributeValuesDataInterface* DataInterface = CastChecked<UPCGCountUniqueAttributeValuesDataInterface>(InDataInterface);
 	Kernel = DataInterface->ProducerKernel;
 
-	AttributeToCountName = DataInterface->AttributeToCountName;
-	bEmitPerDataCounts = DataInterface->bEmitPerDataCounts;
+	AttributeToCountName = DataInterface->GetAttributeToCountName();
+	bEmitPerDataCounts = DataInterface->GetEmitPerDataCounts();
+	bOutputRawBuffer = DataInterface->GetOutputRawBuffer();
 }
 
 bool UPCGCountUniqueAttributeValuesDataProvider::PrepareForExecute_GameThread(UPCGDataBinding* InBinding)
@@ -113,10 +114,13 @@ bool UPCGCountUniqueAttributeValuesDataProvider::PrepareForExecute_GameThread(UP
 		return false;
 	}
 
-	OutputValueAttributeId = InBinding->GetAttributeId(PCGCountUniqueAttributeValuesConstants::ValueAttributeName, EPCGKernelAttributeType::Int);
-	check(OutputValueAttributeId != -1);
-	OutputCountAttributeId = InBinding->GetAttributeId(PCGCountUniqueAttributeValuesConstants::ValueCountAttributeName, EPCGKernelAttributeType::Int);
-	check(OutputCountAttributeId != -1);
+	if (!bOutputRawBuffer)
+	{
+		OutputValueAttributeId = InBinding->GetAttributeId(PCGCountUniqueAttributeValuesConstants::ValueAttributeName, EPCGKernelAttributeType::Int);
+		check(OutputValueAttributeId != -1);
+		OutputCountAttributeId = InBinding->GetAttributeId(PCGCountUniqueAttributeValuesConstants::ValueCountAttributeName, EPCGKernelAttributeType::Int);
+		check(OutputCountAttributeId != -1);
+	}
 
 	const FPCGDataCollectionDesc* InputDataDesc = InBinding->GetCachedKernelPinDataDesc(Kernel, PCGPinConstants::DefaultInputLabel, /*bIsInput*/true);
 
@@ -181,6 +185,7 @@ void UPCGCountUniqueAttributeValuesDataProvider::Reset()
 	OutputValueAttributeId = INDEX_NONE;
 	OutputCountAttributeId = INDEX_NONE;
 	bEmitPerDataCounts = true;
+	bOutputRawBuffer = false;
 
 	Super::Reset();
 }
diff --git a/Engine/Plugins/PCG/Source/PCG/Private/Compute/DataInterfaces/BuiltInKernels/PCGCountUniqueAttributeValuesDataInterface.h b/Engine/Plugins/PCG/Source/PCG/Private/Compute/DataInterfaces/BuiltInKernels/PCGCountUniqueAttributeValuesDataInterface.h
--- a/Engine/Plugins/PCG/Source/PCG/Private/Compute/DataInterfaces/BuiltInKernels/PCGCountUniqueAttributeValuesDataInterface.h
+++ b/Engine/Plugins/PCG/Source/PCG/Private/Compute/DataInterfaces/BuiltInKernels/PCGCountUniqueAttributeValuesDataInterface.h
@@ -24,12 +24,24 @@ public:
 	UComputeDataProvider* CreateDataProvider() const override;
 	//~ End UComputeDataInterface Interface
 
-public:
+	FName GetAttributeToCountName() const { return AttributeToCountName; }
+	void SetAttributeToCountName(FName InAttributeToCountName) { AttributeToCountName = InAttributeToCountName; }
+
+	bool GetEmitPerDataCounts() const { return bEmitPerDataCounts; }
+	void SetEmitPerDataCounts(bool bInEmitPerDataCounts) { bEmitPerDataCounts = bInEmitPerDataCounts; }
+
+	bool GetOutputRawBuffer() const { return bOutputRawBuffer; }
+	void SetOutputRawBuffer(bool bInOutputRawBuffer) { bOutputRawBuffer = bInOutputRawBuffer; }
+
+protected:
 	UPROPERTY()
 	FName AttributeToCountName;
 
 	UPROPERTY()
 	bool bEmitPerDataCounts = true;
+
+	UPROPERTY()
+	bool bOutputRawBuffer = false;
 };
 
 UCLASS()
@@ -61,6 +73,8 @@ public:
 	int32 OutputCountAttributeId = INDEX_NONE;
 
 	bool bEmitPerDataCounts = true;
+
+	bool bOutputRawBuffer = false;
 };
 
 class FPCGCountUniqueAttributeValuesProviderProxy : public FComputeDataProviderRenderProxy
diff --git a/Engine/Plugins/PCG/Source/PCG/Private/Compute/DataInterfaces/Elements/PCGStaticMeshSpawnerDataInterface.cpp b/Engine/Plugins/PCG/Source/PCG/Private/Compute/DataInterfaces/Elements/PCGStaticMeshSpawnerDataInterface.cpp
--- a/Engine/Plugins/PCG/Source/PCG/Private/Compute/DataInterfaces/Elements/PCGStaticMeshSpawnerDataInterface.cpp
+++ b/Engine/Plugins/PCG/Source/PCG/Private/Compute/DataInterfaces/Elements/PCGStaticMeshSpawnerDataInterface.cpp
@@ -10,6 +10,7 @@
 #include "Compute/PCGDataBinding.h"
 #include "Compute/BuiltInKernels/PCGCountUniqueAttributeValuesKernel.h"
 #include "Compute/Data/PCGProxyForGPUData.h"
+#include "Compute/Data/PCGRawBufferData.h"
 #include "Compute/PrimitiveFactories/PCGPrimitiveFactoryPISMC.h"
 #include "Elements/PCGStaticMeshSpawner.h"
 #include "Elements/PCGStaticMeshSpawnerKernel.h"
@@ -261,38 +262,53 @@ bool UPCGStaticMeshSpawnerDataProvider::PerformPreExecuteReadbacks_GameThread(UP
 		return false;
 	}
 
-	const UPCGParamData* AnalysisResultsData = Cast<UPCGParamData>(InBinding->InputDataCollection.TaggedData[AnalysisDataIndex].Data);
-	const UPCGMetadata* AnalysisMetadata = AnalysisResultsData ? AnalysisResultsData->ConstMetadata() : nullptr;
+	if (const UPCGParamData* AnalysisResultsData = Cast<UPCGParamData>(InBinding->InputDataCollection.TaggedData[AnalysisDataIndex].Data))
+	{
+		const UPCGMetadata* AnalysisMetadata = AnalysisResultsData->ConstMetadata();
+		check(AnalysisMetadata);
 
-	const FPCGMetadataAttributeBase* AnalysisValueAttributeBase = AnalysisMetadata ? AnalysisMetadata->GetConstAttribute(PCGCountUniqueAttributeValuesConstants::ValueAttributeName) : nullptr;
-	const FPCGMetadataAttributeBase* AnalysisCountAttributeBase = AnalysisMetadata ? AnalysisMetadata->GetConstAttribute(PCGCountUniqueAttributeValuesConstants::ValueCountAttributeName) : nullptr;
+		const FPCGMetadataAttributeBase* AnalysisValueAttributeBase = AnalysisMetadata->GetConstAttribute(PCGCountUniqueAttributeValuesConstants::ValueAttributeName);
+		const FPCGMetadataAttributeBase* AnalysisCountAttributeBase = AnalysisMetadata->GetConstAttribute(PCGCountUniqueAttributeValuesConstants::ValueCountAttributeName);
 
-	if (AnalysisValueAttributeBase && AnalysisValueAttributeBase->GetTypeId() == PCG::Private::MetadataTypes<int32>::Id &&
-		AnalysisCountAttributeBase && AnalysisCountAttributeBase->GetTypeId() == PCG::Private::MetadataTypes<int32>::Id)
-	{
-		const FPCGMetadataAttribute<int32>* ValueAttribute = static_cast<const FPCGMetadataAttribute<int32>*>(AnalysisValueAttributeBase);
-		const FPCGMetadataAttribute<int32>* CountAttribute = static_cast<const FPCGMetadataAttribute<int32>*>(AnalysisCountAttributeBase);
+		if (AnalysisValueAttributeBase && AnalysisValueAttributeBase->GetTypeId() == PCG::Private::MetadataTypes<int32>::Id &&
+			AnalysisCountAttributeBase && AnalysisCountAttributeBase->GetTypeId() == PCG::Private::MetadataTypes<int32>::Id)
+		{
+			const FPCGMetadataAttribute<int32>* ValueAttribute = static_cast<const FPCGMetadataAttribute<int32>*>(AnalysisValueAttributeBase);
+			const FPCGMetadataAttribute<int32>* CountAttribute = static_cast<const FPCGMetadataAttribute<int32>*>(AnalysisCountAttributeBase);
 
-		const int32 NumElements = AnalysisMetadata->GetItemCountForChild();
+			const int32 NumElements = AnalysisMetadata->GetItemCountForChild();
 
-		StringKeyToInstanceCount.Reserve(NumElements);
+			StringKeyToInstanceCount.Reserve(NumElements);
 
-		// TODO: Range based get would scale better.
-		for (int64 MetadataKey = 0; MetadataKey < NumElements; ++MetadataKey)
+			// TODO: Range based get would scale better.
+			for (int64 MetadataKey = 0; MetadataKey < NumElements; ++MetadataKey)
+			{
+				StringKeyToInstanceCount.Add(ValueAttribute->GetValue(MetadataKey), CountAttribute->GetValue(MetadataKey));
+			}
+		}
+		else
 		{
-			StringKeyToInstanceCount.Add(ValueAttribute->GetValue(MetadataKey), CountAttribute->GetValue(MetadataKey));
+			UE_LOG(LogPCG, Warning, TEXT("No analysis data received by static mesh spawner kernel, worst case instance allocations will be made."));
+
+			if (InBinding->InputDataCollection.TaggedData[AnalysisDataIndex].Data->IsA<UPCGProxyForGPUData>())
+			{
+				UE_LOG(LogPCG, Error, TEXT("Data was not read back."));
+			}
+
+			return true;
 		}
 	}
-	else
+	else if (const UPCGRawBufferData* RawAnalysisData = Cast<UPCGRawBufferData>(InBinding->InputDataCollection.TaggedData[AnalysisDataIndex].Data))
 	{
-		UE_LOG(LogPCG, Warning, TEXT("No analysis data received by static mesh spawner kernel, worst case instance allocations will be made."));
+		// Hardcoded data format - uint array containing pairs of (string key value, instance count) pairs.
+		const int32 NumUints = RawAnalysisData->GetNumUint32s();
+		ensure((NumUints % 2) == 0);
 
-		if (InBinding->InputDataCollection.TaggedData[AnalysisDataIndex].Data->IsA<UPCGProxyForGPUData>())
+		const uint32* RawData = RawAnalysisData->GetConstData().GetData();
+		for (const uint32* It = RawData; It < RawData + NumUints; It += 2)
 		{
-			UE_LOG(LogPCG, Error, TEXT("Data was not read back."));
+			StringKeyToInstanceCount.Add(static_cast<int32>(It[0]), It[1]);
 		}
-
-		return true;
 	}
 
 	return true;
diff --git a/Engine/Plugins/PCG/Source/PCG/Private/Compute/DataInterfaces/PCGDataCollectionDataInterface.cpp b/Engine/Plugins/PCG/Source/PCG/Private/Compute/DataInterfaces/PCGDataCollectionDataInterface.cpp
--- a/Engine/Plugins/PCG/Source/PCG/Private/Compute/DataInterfaces/PCGDataCollectionDataInterface.cpp
+++ b/Engine/Plugins/PCG/Source/PCG/Private/Compute/DataInterfaces/PCGDataCollectionDataInterface.cpp
@@ -692,7 +692,7 @@ bool FPCGDataCollectionDataProviderProxy::IsValid(FValidationData const& InValid
 
 	if (SizeBytes <= 0)
 	{
-		UE_LOG(LogPCG, Error, TEXT("Proxy invalid due to invalid buffer size."));
+		UE_LOG(LogPCG, Error, TEXT("FPCGDataCollectionDataProviderProxy invalid due to invalid buffer size."));
 		return false;
 	}
 
diff --git a/Engine/Plugins/PCG/Source/PCG/Private/Compute/DataInterfaces/PCGDataCollectionDataInterface.h b/Engine/Plugins/PCG/Source/PCG/Private/Compute/DataInterfaces/PCGDataCollectionDataInterface.h
--- a/Engine/Plugins/PCG/Source/PCG/Private/Compute/DataInterfaces/PCGDataCollectionDataInterface.h
+++ b/Engine/Plugins/PCG/Source/PCG/Private/Compute/DataInterfaces/PCGDataCollectionDataInterface.h
@@ -15,7 +15,13 @@ class FPCGDataCollectionDataInterfaceParameters;
 class FPCGDataUploadAdaptor;
 class UPCGDataBinding;
 
-/** Compute Framework Data Interface for reading PCG data. */
+/**
+* Compute Framework Data Interface for PCG data types that support being packed into data collections, which means they are composed of:
+* - One or more data items
+* - One or more elements
+* - Elements can have system attributes (like point transform, bounds, etc)
+* - Elements can have user-added attributes
+*/
 UCLASS(ClassGroup = (Procedural))
 class UPCGDataCollectionDataInterface : public UPCGExportableDataInterface
 {
diff --git a/Engine/Plugins/PCG/Source/PCG/Private/Compute/DataInterfaces/PCGRawBufferDataInterface.cpp b/Engine/Plugins/PCG/Source/PCG/Private/Compute/DataInterfaces/PCGRawBufferDataInterface.cpp
new file mode 100644
--- /dev/null
+++ b/Engine/Plugins/PCG/Source/PCG/Private/Compute/DataInterfaces/PCGRawBufferDataInterface.cpp
@@ -0,0 +1,310 @@
+// Copyright Epic Games, Inc. All Rights Reserved.
+
+#include "PCGRawBufferDataInterface.h"
+
+#include "PCGCommon.h"
+#include "PCGModule.h"
+#include "Compute/PCGDataBinding.h"
+#include "Compute/Data/PCGRawBufferData.h"
+
+#include "RenderGraphBuilder.h"
+#include "ShaderCompilerCore.h"
+#include "ShaderParameterMetadataBuilder.h"
+#include "ShaderParameterStruct.h"
+#include "ComputeFramework/ShaderParamTypeDefinition.h"
+
+#include UE_INLINE_GENERATED_CPP_BY_NAME(PCGRawBufferDataInterface)
+
+namespace PCGRawBufferDataInterface
+{
+	static TCHAR const* TemplateFilePath = TEXT("/Plugin/PCG/Private/PCGRawBufferDataInterface.ush");
+}
+
+void UPCGRawBufferDataInterface::GetSupportedOutputs(TArray<FShaderFunctionDefinition>& OutFunctions) const
+{
+	OutFunctions.AddDefaulted_GetRef()
+		.SetName(TEXT("GetNumData"))
+		.AddReturnType(EShaderFundamentalType::Uint); // Num data
+
+	OutFunctions.AddDefaulted_GetRef()
+		.SetName(TEXT("GetNumElements"))
+		.AddParam(EShaderFundamentalType::Uint) // InDataIndex
+		.AddReturnType(EShaderFundamentalType::Uint); // Size in bytes
+
+	OutFunctions.AddDefaulted_GetRef()
+		.SetName(TEXT("Load"))
+		.AddParam(EShaderFundamentalType::Uint) // InDataIndex
+		.AddParam(EShaderFundamentalType::Uint) // InElementIndex
+		.AddReturnType(EShaderFundamentalType::Uint);
+
+	OutFunctions.AddDefaulted_GetRef()
+		.SetName(TEXT("Load4"))
+		.AddParam(EShaderFundamentalType::Uint) // InDataIndex
+		.AddParam(EShaderFundamentalType::Uint) // InFirstElementIndex
+		.AddReturnType(EShaderFundamentalType::Uint, 4);
+
+	OutFunctions.AddDefaulted_GetRef()
+		.SetName(TEXT("Store"))
+		.AddParam(EShaderFundamentalType::Uint) // InDataIndex
+		.AddParam(EShaderFundamentalType::Uint) // InElementIndex
+		.AddParam(EShaderFundamentalType::Uint); // InValue
+
+	OutFunctions.AddDefaulted_GetRef()
+		.SetName(TEXT("Store4"))
+		.AddParam(EShaderFundamentalType::Uint) // InDataIndex
+		.AddParam(EShaderFundamentalType::Uint) // InFirstElementIndex
+		.AddParam(EShaderFundamentalType::Uint, 4); // InValue
+
+	OutFunctions.AddDefaulted_GetRef()
+		.SetName(TEXT("AtomicAdd"))
+		.AddParam(EShaderFundamentalType::Uint) // InDataIndex
+		.AddParam(EShaderFundamentalType::Uint) // InElementIndex
+		.AddParam(EShaderFundamentalType::Uint) // InValueToAdd
+		.AddReturnType(EShaderFundamentalType::Uint); // Value before add
+}
+
+BEGIN_SHADER_PARAMETER_STRUCT(FPCGRawBufferDataInterfaceParameters,)
+	SHADER_PARAMETER_RDG_BUFFER_UAV(RWByteAddressBuffer, Data)
+	SHADER_PARAMETER(uint32, SizeBytes)
+END_SHADER_PARAMETER_STRUCT()
+
+void UPCGRawBufferDataInterface::GetShaderParameters(TCHAR const* UID, FShaderParametersMetadataBuilder& InOutBuilder, FShaderParametersMetadataAllocations& InOutAllocations) const
+{
+	InOutBuilder.AddNestedStruct<FPCGRawBufferDataInterfaceParameters>(UID);
+}
+
+TCHAR const* UPCGRawBufferDataInterface::GetShaderVirtualPath() const
+{
+	return PCGRawBufferDataInterface::TemplateFilePath;
+}
+
+void UPCGRawBufferDataInterface::GetShaderHash(FString& InOutKey) const
+{
+	GetShaderFileHash(PCGRawBufferDataInterface::TemplateFilePath, EShaderPlatform::SP_PCD3D_SM5).AppendString(InOutKey);
+}
+
+void UPCGRawBufferDataInterface::GetHLSL(FString& OutHLSL, FString const& InDataInterfaceName) const
+{
+	TMap<FString, FStringFormatArg> TemplateArgs =
+	{
+		{ TEXT("DataInterfaceName"), InDataInterfaceName },
+	};
+
+	FString TemplateFile;
+	if (ensure(LoadShaderSourceFile(PCGRawBufferDataInterface::TemplateFilePath, EShaderPlatform::SP_PCD3D_SM5, &TemplateFile, nullptr)))
+	{
+		OutHLSL += FString::Format(*TemplateFile, TemplateArgs);
+	}
+}
+
+UComputeDataProvider* UPCGRawBufferDataInterface::CreateDataProvider() const
+{
+	return NewObject<UPCGRawBufferDataProvider>();
+}
+
+void UPCGRawBufferDataProvider::Initialize(const UComputeDataInterface* InDataInterface, UObject* InBinding, uint64 InInputMask, uint64 InOutputMask)
+{
+	Super::Initialize(InDataInterface, InBinding, InInputMask, InOutputMask);
+
+	const UPCGRawBufferDataInterface* DataInterface = CastChecked<UPCGRawBufferDataInterface>(InDataInterface);
+	bZeroInitialize = DataInterface->GetRequiresZeroInitialization();
+}
+
+bool UPCGRawBufferDataProvider::PrepareForExecute_GameThread(UPCGDataBinding* InBinding)
+{
+	if (!Super::PrepareForExecute_GameThread(InBinding))
+	{
+		return false;
+	}
+
+	if (const FPCGDataCollectionDesc* ProducerPinDesc = InBinding->GetCachedKernelPinDataDesc(GraphBindingIndex); ensure(ProducerPinDesc))
+	{
+		const uint64 PackedSizeBytes = ProducerPinDesc->ComputePackedSizeBytes();
+
+		if (!PCGComputeHelpers::IsBufferSizeTooLarge(PackedSizeBytes))
+		{
+			SizeBytes = static_cast<uint32>(PackedSizeBytes);
+		}
+	}
+
+	if (bProducedByCPU)
+	{
+		check(!DownstreamInputPinLabelAliases.IsEmpty());
+
+		const TArray<FPCGTaggedData> TaggedDatas = InBinding->InputDataCollection.GetInputsByPin(DownstreamInputPinLabelAliases[0]);
+
+		if (!TaggedDatas.IsEmpty())
+		{
+			ensureMsgf(TaggedDatas.Num() == 1, TEXT("UPCGRawBufferDataProvider: Multi-data is not currently supported (received %d data items)."), TaggedDatas.Num());
+
+			DataToUpload = Cast<UPCGRawBufferData>(TaggedDatas[0].Data);
+		}
+	}
+
+	ReadbackState = MakeShared<FReadbackState>();
+
+	return true;
+}
+
+FComputeDataProviderRenderProxy* UPCGRawBufferDataProvider::GetRenderProxy()
+{
+	TRACE_CPUPROFILER_EVENT_SCOPE(UPCGRawBufferDataProvider::GetRenderProxy);
+
+	TWeakObjectPtr<UPCGRawBufferDataProvider> ThisWeak = MakeWeakObjectPtr(this);
+
+	FPCGRawBufferDataProviderProxy::FParams Params =
+	{
+		.SizeBytes = SizeBytes,
+		.bZeroInitialize = bZeroInitialize,
+		.DataToUpload = DataToUpload,
+		.ExportMode = ExportMode,
+		.OutputPinLabel = OutputPinLabel,
+		.OutputPinLabelAlias = OutputPinLabelAlias,
+		.DataProviderWeakPtr = ThisWeak,
+	};
+
+	ensure(ReadbackState);
+
+	Params.AsyncReadbackCallback_RenderThread = [ReadbackState=ReadbackState, DataProviderWeakPtr=ThisWeak, GenerationCount=GenerationCounter.load(), ExportMode=ExportMode, OutputPinLabel=OutputPinLabel, OutputPinLabelAlias=OutputPinLabelAlias, BindingWeakPtr=Binding](const void* InData, int InNumBytes)
+	{
+		if (!InData || !ReadbackState)
+		{
+			return;
+		}
+
+		ReadbackState->Data.Append(static_cast<const uint32*>(InData), InNumBytes / sizeof(uint32));
+
+		ExecuteOnGameThread(UE_SOURCE_LOCATION, [ReadbackState, DataProviderWeakPtr, GenerationCount, ExportMode, OutputPinLabel, OutputPinLabelAlias, BindingWeakPtr]()
+		{
+			UPCGDataBinding* Binding = BindingWeakPtr.Get();
+			if (!Binding)
+			{
+				return;
+			}
+
+			UPCGRawBufferDataProvider* DataProvider = DataProviderWeakPtr.Get();
+			if (!DataProvider)
+			{
+				UE_LOG(LogPCG, Error, TEXT("Could not resolve UPCGRawBufferDataProvider pointer to pass back buffer handle."));
+				return;
+			}
+
+			if (DataProvider->GenerationCounter != GenerationCount)
+			{
+				// Data provider was returned to the pool, possibly due to cancellation. Nothing to do here.
+				return;
+			}
+
+			UPCGRawBufferData* RawBufferData = NewObject<UPCGRawBufferData>();
+			RawBufferData->SetData(MoveTemp(ReadbackState->Data));
+
+			Binding->ReceiveDataFromGPU_GameThread(RawBufferData, DataProvider->GetProducerSettings(), ExportMode, OutputPinLabel, OutputPinLabelAlias);
+
+			DataProvider->OnDataExported_GameThread().Broadcast();
+		});
+	};
+
+	return new FPCGRawBufferDataProviderProxy(Params);
+}
+
+void UPCGRawBufferDataProvider::Reset()
+{
+	Super::Reset();
+
+	SizeBytes = INDEX_NONE;
+	bZeroInitialize = false;
+	DataToUpload = nullptr;
+	ReadbackState.Reset();
+}
+
+FPCGRawBufferDataProviderProxy::FPCGRawBufferDataProviderProxy(const FParams& InParams)
+	: SizeBytes(InParams.SizeBytes)
+	, bZeroInitialize(InParams.bZeroInitialize)
+	, ExportMode(InParams.ExportMode)
+	, OutputPinLabel(InParams.OutputPinLabel)
+	, OutputPinLabelAlias(InParams.OutputPinLabelAlias)
+	, DataProviderWeakPtr(InParams.DataProviderWeakPtr)
+	, AsyncReadbackCallback_RenderThread(InParams.AsyncReadbackCallback_RenderThread)
+{
+	if (InParams.DataToUpload)
+	{
+		DataToUpload.Append(InParams.DataToUpload->GetConstData());
+	}
+}
+
+bool FPCGRawBufferDataProviderProxy::IsValid(FValidationData const& InValidationData) const
+{
+	if (SizeBytes <= 0 || (SizeBytes % sizeof(uint32)) != 0)
+	{
+		UE_LOG(LogPCG, Warning, TEXT("FPCGRawBufferDataProviderProxy not valid due to invalid size: %d bytes."), SizeBytes);
+		return false;
+	}
+
+	return InValidationData.ParameterStructSize == sizeof(FParameters);
+}
+
+void FPCGRawBufferDataProviderProxy::AllocateResources(FRDGBuilder& GraphBuilder, FAllocationData const& InAllocationData)
+{
+	check(SizeBytes > 0);
+
+	FRDGBufferDesc Desc = FRDGBufferDesc::CreateByteAddressDesc(SizeBytes);
+
+	if (ExportMode != EPCGExportMode::NoExport)
+	{
+		// We don't know for sure whether buffer will be read back or not, so need to flag the possibility if the buffer will be passed downstream.
+		Desc.Usage |= BUF_SourceCopy;
+	}
+
+	Data = GraphBuilder.CreateBuffer(Desc, TEXT("PCGRawBuffer"));
+	DataUAV = GraphBuilder.CreateUAV(Data);
+
+	if (!DataToUpload.IsEmpty())
+	{
+		TRACE_CPUPROFILER_EVENT_SCOPE(FPCGRawBufferDataProviderProxy::AllocateResources::UploadData);
+		GraphBuilder.QueueBufferUpload(Data, MakeConstArrayView(DataToUpload));
+	}
+	else if (bZeroInitialize)
+	{
+		TRACE_CPUPROFILER_EVENT_SCOPE(FPCGRawBufferDataProviderProxy::AllocateResources::UploadZeroes);
+		DataToUpload.SetNumZeroed(SizeBytes / sizeof(uint32));
+		GraphBuilder.QueueBufferUpload(Data, MakeConstArrayView(DataToUpload));
+	}
+#if !UE_BUILD_SHIPPING
+	else if (PCGSystemSwitches::CVarFuzzGPUMemory.GetValueOnAnyThread())
+	{
+		DataToUpload.SetNumUninitialized(SizeBytes / sizeof(uint32));
+
+		FRandomStream Rand(GFrameCounter);
+		for (uint32& Value : DataToUpload)
+		{
+			Value = Rand.GetUnsignedInt();
+		}
+
+		GraphBuilder.QueueBufferUpload(Data, MakeConstArrayView(DataToUpload));
+	}
+#endif
+}
+
+void FPCGRawBufferDataProviderProxy::GatherDispatchData(FDispatchData const& InDispatchData)
+{
+	const TStridedView<FParameters> ParameterArray = MakeStridedParameterView<FParameters>(InDispatchData);
+	for (int32 InvocationIndex = 0; InvocationIndex < ParameterArray.Num(); ++InvocationIndex)
+	{
+		FParameters& Parameters = ParameterArray[InvocationIndex];
+		Parameters.Data = DataUAV;
+		Parameters.SizeBytes = SizeBytes;
+	}
+}
+
+void FPCGRawBufferDataProviderProxy::GetReadbackData(TArray<FReadbackData>& OutReadbackData) const
+{
+	if (ExportMode != EPCGExportMode::NoExport)
+	{
+		FReadbackData ReadbackData;
+		ReadbackData.Buffer = Data;
+		ReadbackData.NumBytes = SizeBytes;
+		ReadbackData.ReadbackCallback_RenderThread = &AsyncReadbackCallback_RenderThread;
+
+		OutReadbackData.Add(MoveTemp(ReadbackData));
+	}
+}
diff --git a/Engine/Plugins/PCG/Source/PCG/Private/Compute/DataInterfaces/PCGRawBufferDataInterface.h b/Engine/Plugins/PCG/Source/PCG/Private/Compute/DataInterfaces/PCGRawBufferDataInterface.h
new file mode 100644
--- /dev/null
+++ b/Engine/Plugins/PCG/Source/PCG/Private/Compute/DataInterfaces/PCGRawBufferDataInterface.h
@@ -0,0 +1,115 @@
+// Copyright Epic Games, Inc. All Rights Reserved.
+
+#pragma once
+
+#include "Compute/DataInterfaces/PCGComputeDataInterface.h"
+
+#include "RHIResources.h"
+
+#include "PCGRawBufferDataInterface.generated.h"
+
+class FPCGRawBufferDataInterfaceParameters;
+class UPCGRawBufferData;
+
+/** A data interface for a simple array of uint values. No data format header or attributes, just raw array access. */
+UCLASS(ClassGroup = (Procedural))
+class UPCGRawBufferDataInterface : public UPCGExportableDataInterface
+{
+	GENERATED_BODY()
+
+public:
+	//~ Begin UComputeDataInterface Interface
+	TCHAR const* GetClassName() const override { return TEXT("PCGRawBuffer"); }
+	void GetSupportedInputs(TArray<FShaderFunctionDefinition>& OutFunctions) const override { return GetSupportedOutputs(OutFunctions); }
+	void GetSupportedOutputs(TArray<FShaderFunctionDefinition>& OutFunctions) const override;
+	void GetShaderParameters(TCHAR const* UID, FShaderParametersMetadataBuilder& InOutBuilder, FShaderParametersMetadataAllocations& InOutAllocations) const override;
+	TCHAR const* GetShaderVirtualPath() const override;
+	void GetShaderHash(FString& InOutKey) const override;
+	void GetHLSL(FString& OutHLSL, FString const& InDataInterfaceName) const override;
+	virtual bool GetRequiresReadback() const override { return true; }
+	UComputeDataProvider* CreateDataProvider() const override;
+	//~ End UComputeDataInterface Interface
+
+	bool GetRequiresZeroInitialization() const { return bRequiresZeroInitialization; }
+	void SetRequiresZeroInitialization(bool bInZeroInit) { bRequiresZeroInitialization = bInZeroInit; }
+
+protected:
+	/** Whether to perform full 0-initialization of the buffer. */
+	UPROPERTY()
+	bool bRequiresZeroInitialization = false;
+};
+
+UCLASS()
+class UPCGRawBufferDataProvider : public UPCGExportableDataProvider
+{
+	GENERATED_BODY()
+
+public:
+	//~ Begin UComputeDataProvider Interface
+	virtual void Initialize(const UComputeDataInterface* InDataInterface, UObject* InBinding, uint64 InInputMask, uint64 InOutputMask) override;
+	virtual bool PrepareForExecute_GameThread(UPCGDataBinding* InBinding) override;
+	virtual FComputeDataProviderRenderProxy* GetRenderProxy() override;
+	virtual void Reset() override;
+	//~ End UComputeDataProvider Interface
+
+protected:
+	int32 SizeBytes = INDEX_NONE;
+
+	bool bZeroInitialize = false;
+
+	UPROPERTY()
+	TObjectPtr<const UPCGRawBufferData> DataToUpload;
+
+	struct FReadbackState
+	{
+		TArray<uint32> Data;
+	};
+	TSharedPtr<FReadbackState, ESPMode::ThreadSafe> ReadbackState;
+};
+
+class FPCGRawBufferDataProviderProxy : public FComputeDataProviderRenderProxy
+{
+public:
+	struct FParams
+	{
+		int32 SizeBytes = 0;
+		bool bZeroInitialize = false;
+		const UPCGRawBufferData* DataToUpload = nullptr;
+		EPCGExportMode ExportMode = EPCGExportMode::NoExport;
+		FName OutputPinLabel;
+		FName OutputPinLabelAlias;
+		TWeakObjectPtr<UPCGRawBufferDataProvider> DataProviderWeakPtr;
+		FReadbackCallback AsyncReadbackCallback_RenderThread;
+	};
+
+	FPCGRawBufferDataProviderProxy(const FParams& InParams);
+
+	//~ Begin FComputeDataProviderRenderProxy Interface
+	bool IsValid(FValidationData const& InValidationData) const override;
+	void AllocateResources(FRDGBuilder& GraphBuilder, FAllocationData const& InAllocationData);
+	void GatherDispatchData(FDispatchData const& InDispatchData) override;
+	void GetReadbackData(TArray<FReadbackData>& OutReadbackData) const override;
+	//~ End FComputeDataProviderRenderProxy Interface
+
+protected:
+	using FParameters = FPCGRawBufferDataInterfaceParameters;
+
+	int32 SizeBytes = INDEX_NONE;
+
+	bool bZeroInitialize = false;
+
+	EPCGExportMode ExportMode = EPCGExportMode::NoExport;
+
+	FName OutputPinLabel;
+	FName OutputPinLabelAlias;
+
+	FRDGBufferRef Data = nullptr;
+	FRDGBufferUAVRef DataUAV = nullptr;
+
+	TArray<uint32> DataToUpload;
+
+	TWeakObjectPtr<UPCGRawBufferDataProvider> DataProviderWeakPtr;
+
+	/** Called from render thread when readback from GPU is complete. */
+	FReadbackCallback AsyncReadbackCallback_RenderThread;
+};
diff --git a/Engine/Plugins/PCG/Source/PCG/Private/Compute/PCGComputeCommon.cpp b/Engine/Plugins/PCG/Source/PCG/Private/Compute/PCGComputeCommon.cpp
--- a/Engine/Plugins/PCG/Source/PCG/Private/Compute/PCGComputeCommon.cpp
+++ b/Engine/Plugins/PCG/Source/PCG/Private/Compute/PCGComputeCommon.cpp
@@ -11,8 +11,10 @@
 #include "Compute/PCGDataBinding.h"
 #include "Compute/PCGPinPropertiesGPU.h"
 #include "Compute/Data/PCGProxyForGPUData.h"
+#include "Compute/Data/PCGRawBufferData.h"
 #include "Compute/DataInterfaces/PCGDataCollectionDataInterface.h"
 #include "Compute/DataInterfaces/PCGLandscapeDataInterface.h"
+#include "Compute/DataInterfaces/PCGRawBufferDataInterface.h"
 #include "Compute/DataInterfaces/PCGStaticMeshDataInterface.h"
 #include "Compute/DataInterfaces/PCGTextureDataInterface.h"
 #include "Compute/DataInterfaces/PCGVirtualTextureDataInterface.h"
@@ -32,15 +34,18 @@ namespace PCGComputeConstants
 	/** PCG data types supported in GPU node inputs. */
 	inline const FPCGDataTypeIdentifier& GetAllowedInputTypes()
 	{
-		static FPCGDataTypeIdentifier Result{
+		static FPCGDataTypeIdentifier Result = FPCGDataTypeIdentifier::Construct(
 			EPCGDataType::Point
-			| EPCGDataType::Param
-			| EPCGDataType::Landscape
-			| EPCGDataType::BaseTexture
-			| EPCGDataType::VirtualTexture
-			| EPCGDataType::Resource
-			| EPCGDataType::ProxyForGPU
-		};
+			, EPCGDataType::Param
+			, EPCGDataType::Landscape
+			, EPCGDataType::Param
+			, EPCGDataType::Landscape
+			, EPCGDataType::BaseTexture
+			, EPCGDataType::VirtualTexture
+			, EPCGDataType::Resource
+			, EPCGDataType::ProxyForGPU
+			, FPCGDataTypeInfoRawBuffer::AsId()
+		);
 
 		return Result;
 	}
@@ -48,13 +53,21 @@ namespace PCGComputeConstants
 	/** PCG data types supported in GPU node outputs. */
 	inline const FPCGDataTypeIdentifier& GetAllowedOutputTypes()
 	{
-		static FPCGDataTypeIdentifier Result{ EPCGDataType::Point | EPCGDataType::Param | EPCGDataType::ProxyForGPU | EPCGDataType::BaseTexture };
+		static FPCGDataTypeIdentifier Result = FPCGDataTypeIdentifier::Construct(
+			EPCGDataType::Point
+			, EPCGDataType::Param
+			, EPCGDataType::ProxyForGPU
+			, EPCGDataType::BaseTexture
+			, FPCGDataTypeInfoRawBuffer::AsId()
+		);
+
 		return Result;
 	}
 
 	/** PCG data types supported in GPU data collections. */
 	inline const FPCGDataTypeIdentifier& GetAllowedDataCollectionTypes()
 	{
+		// Intentionally excludes raw buffer data which cannot be mixed with other data types like point.
 		static FPCGDataTypeIdentifier Result{ EPCGDataType::Point | EPCGDataType::Param | EPCGDataType::ProxyForGPU };
 		return Result;
 	}
@@ -92,6 +105,10 @@ namespace PCGComputeHelpers
 			ElementCount.X = TextureSize.X;
 			ElementCount.Y = TextureSize.Y;
 		}
+		else if (const UPCGRawBufferData* RawBufferData = Cast<UPCGRawBufferData>(InData))
+		{
+			ElementCount.X = RawBufferData->GetNumUint32s();
+		}
 
 		return ElementCount;
 	}
@@ -232,7 +249,7 @@ namespace PCGComputeHelpers
 
 		UPCGComputeDataInterface* DataInterface = nullptr;
 
-		if (PinProperties.AllowedTypes == EPCGDataType::Point || PinProperties.AllowedTypes == EPCGDataType::Param)
+		if (PCGComputeHelpers::IsTypeAllowedInDataCollection(PinProperties.AllowedTypes))
 		{
 			UPCGDataCollectionDataInterface* DataInterfacePCGData = Context.NewObject_AnyThread<UPCGDataCollectionDataInterface>(InParams.ObjectOuter);
 
@@ -266,6 +283,17 @@ namespace PCGComputeHelpers
 		{
 			DataInterface = Context.NewObject_AnyThread<UPCGStaticMeshDataInterface>(InParams.ObjectOuter);
 		}
+		else if (PinProperties.AllowedTypes == FPCGDataTypeInfoRawBuffer::AsId())
+		{
+			UPCGRawBufferDataInterface* DataInterfaceRawBuffer = Context.NewObject_AnyThread<UPCGRawBufferDataInterface>(InParams.ObjectOuter);
+
+			DataInterfaceRawBuffer->SetRequiresExport(InParams.bRequiresExport);
+			DataInterfaceRawBuffer->SetRequiresZeroInitialization(InParams.ProducerKernel && InParams.ProducerKernel->DoesOutputPinRequireZeroInitialization(PinProperties.Label));
+
+			// todo_pcg: We should support SetElementCountMultiplier also.
+
+			DataInterface = DataInterfaceRawBuffer;
+		}
 		else
 		{
 			UE_LOG(LogPCG, Error, TEXT("Unsupported connected upstream pin '%s' on node '%s' with type %s. Consider adding a conversion to a supported type such as Point."),
diff --git a/Engine/Plugins/PCG/Source/PCG/Private/Compute/PCGDataForGPU.cpp b/Engine/Plugins/PCG/Source/PCG/Private/Compute/PCGDataForGPU.cpp
--- a/Engine/Plugins/PCG/Source/PCG/Private/Compute/PCGDataForGPU.cpp
+++ b/Engine/Plugins/PCG/Source/PCG/Private/Compute/PCGDataForGPU.cpp
@@ -16,6 +16,7 @@
 #include "Compute/PCGDataBinding.h"
 #include "Compute/PCGPinPropertiesGPU.h"
 #include "Compute/Data/PCGProxyForGPUData.h"
+#include "Compute/Data/PCGRawBufferData.h"
 #include "Data/PCGBasePointData.h"
 #include "Data/PCGPointData.h"
 #include "Elements/Metadata/PCGMetadataElementCommon.h"
@@ -791,20 +792,27 @@ FPCGDataDesc::FPCGDataDesc(const FPCGTaggedData& InTaggedData, const UPCGDataBin
 	InitializeAttributeDescs(InTaggedData.Data, InBinding);
 }
 
-uint64 FPCGDataDesc::ComputePackedSize() const
+uint64 FPCGDataDesc::ComputePackedSizeBytes() const
 {
-	check(PCGComputeHelpers::IsTypeAllowedInDataCollection(Type));
 	ensure(ElementDimension == EPCGElementDimension::One);
 
-	uint64 DataSizeBytes = DATA_HEADER_SIZE_BYTES;
-
-	for (const FPCGKernelAttributeDesc& AttributeDesc : AttributeDescs)
+	if (Type == FPCGDataTypeInfoRawBuffer::AsId())
 	{
-		const uint64 NumValues = static_cast<uint64>(GetElementCountForAttribute(AttributeDesc).X);
-		DataSizeBytes += static_cast<uint64>(PCGDataForGPUHelpers::GetAttributeTypeStrideBytes(AttributeDesc.AttributeKey.Type)) * NumValues;
+		return ElementCount.X * sizeof(uint32);
 	}
+	else
+	{
+		check(PCGComputeHelpers::IsTypeAllowedInDataCollection(Type));
+		uint64 DataSizeBytes = DATA_HEADER_SIZE_BYTES;
+
+		for (const FPCGKernelAttributeDesc& AttributeDesc : AttributeDescs)
+		{
+			const uint64 NumValues = static_cast<uint64>(GetElementCountForAttribute(AttributeDesc).X);
+			DataSizeBytes += static_cast<uint64>(PCGDataForGPUHelpers::GetAttributeTypeStrideBytes(AttributeDesc.AttributeKey.Type)) * NumValues;
+		}
 
-	return DataSizeBytes;
+		return DataSizeBytes;
+	}
 }
 
 bool FPCGDataDesc::HasElementsMetadataDomainAttributes() const
@@ -1186,16 +1194,24 @@ uint32 FPCGDataCollectionDesc::ComputePackedHeaderSizeBytes() const
 
 uint64 FPCGDataCollectionDesc::ComputePackedSizeBytes() const
 {
-	TRACE_CPUPROFILER_EVENT_SCOPE(FPCGDataCollectionDesc::ComputePackedSize);
+	TRACE_CPUPROFILER_EVENT_SCOPE(FPCGDataCollectionDesc::ComputePackedSizeBytes);
 
-	uint64 TotalCollectionSizeBytes = ComputePackedHeaderSizeBytes();
+	bool bRequiresHeader = false;
+	uint64 TotalDataSizeBytes = 0;
 
 	for (const FPCGDataDesc& DataDesc : DataDescs)
 	{
-		TotalCollectionSizeBytes += DataDesc.ComputePackedSize();
+		bRequiresHeader |= PCGComputeHelpers::IsTypeAllowedInDataCollection(DataDesc.GetType());
+		TotalDataSizeBytes += DataDesc.ComputePackedSizeBytes();
 	}
 
-	return TotalCollectionSizeBytes;
+	// Assumption: If no data items then pack as empty data collection, so include header.
+	if (bRequiresHeader || DataDescs.IsEmpty())
+	{
+		TotalDataSizeBytes += ComputePackedHeaderSizeBytes();
+	}
+
+	return TotalDataSizeBytes;
 }
 
 void FPCGDataCollectionDesc::WriteHeader(TArray<uint32>& OutPackedDataCollectionHeader) const
@@ -1288,6 +1304,18 @@ void FPCGDataCollectionDesc::PackDataCollection(const FPCGDataCollection& InData
 
 	const TArray<FPCGTaggedData> InputData = InDataCollection.GetInputsByPin(InPin);
 
+	for (const FPCGTaggedData& TaggedData : InputData)
+	{
+		if (TaggedData.Data && !PCGComputeHelpers::IsTypeAllowedInDataCollection(TaggedData.Data->GetDataType()))
+		{
+			UE_LOG(LogPCG, Error,
+				TEXT("PackDataCollection: Encountered data '%s' which does not support packing into data collections. ")
+				TEXT("Ensure the correct data interface has been created for the pin ('%s'). Packing will fail."),
+				*TaggedData.Data->GetName(), *InPin.ToString());
+			return;
+		}
+	}
+
 	const uint32 PackedDataCollectionSizeBytes = ComputePackedSizeBytes();
 	check(PackedDataCollectionSizeBytes >= 0);
 
diff --git a/Engine/Plugins/PCG/Source/PCG/Private/Elements/PCGStaticMeshSpawner.cpp b/Engine/Plugins/PCG/Source/PCG/Private/Elements/PCGStaticMeshSpawner.cpp
--- a/Engine/Plugins/PCG/Source/PCG/Private/Elements/PCGStaticMeshSpawner.cpp
+++ b/Engine/Plugins/PCG/Source/PCG/Private/Elements/PCGStaticMeshSpawner.cpp
@@ -35,6 +35,13 @@ static TAutoConsoleVariable<bool> CVarAllowISMReuse(
 	true,
 	TEXT("Controls whether ISMs can be reused and skipped when re-executing"));
 
+static TAutoConsoleVariable<bool> CVarUseRawDataForAnalysis(
+	TEXT("pcg.RuntimeGeneration.ISM.UseRawDataForGPUAnalysis"),
+	false,
+	TEXT("Read back raw data in GPU primitive instancing path rather than constructing an attribute set, to save CPU time. Read-only, must be set from ini."),
+	ECVF_ReadOnly // Structural change, could flush graph compilation cache, but trickier to flush compute graph instance pool in UPCGGraph currently.
+	);
+
 UPCGStaticMeshSpawnerSettings::UPCGStaticMeshSpawnerSettings(const FObjectInitializer &ObjectInitializer)
 {
 	MeshSelectorType = UPCGMeshSelectorWeighted::StaticClass();
@@ -52,7 +59,10 @@ void UPCGStaticMeshSpawnerSettings::CreateKernels(FPCGGPUCompilationContext& InO
 {
 	PCGKernelHelpers::FCreateKernelParams CreateParams(InObjectOuter, this);
 
+	const bool bUseRawBuffer = CVarUseRawDataForAnalysis.GetValueOnAnyThread();
+
 	UPCGStaticMeshSpawnerKernel* SpawnerKernel = PCGKernelHelpers::CreateKernel<UPCGStaticMeshSpawnerKernel>(InOutContext, CreateParams, OutKernels, OutEdges);
+	SpawnerKernel->SetUseRawInstanceCountsBuffer(bUseRawBuffer);
 
 	// If doing by-attribute selection, add analysis kernel that will count how many instances of each mesh are present.
 	if (const UPCGMeshSelectorByAttribute* Selector = Cast<UPCGMeshSelectorByAttribute>(MeshSelectorParameters))
@@ -64,6 +74,7 @@ void UPCGStaticMeshSpawnerSettings::CreateKernels(FPCGGPUCompilationContext& InO
 		CountKernel->SetAttributeName(Selector->AttributeName);
 		// We operate across all input data rather than spawning for each input data separately.
 		CountKernel->SetEmitPerDataCounts(false);
+		CountKernel->SetOutputRawBuffer(bUseRawBuffer);
 
 		OutEdges.Emplace(FPCGPinReference(CountKernel, PCGPinConstants::DefaultOutputLabel), FPCGPinReference(SpawnerKernel, PCGStaticMeshSpawnerConstants::InstanceCountsPinLabel));
 	}
diff --git a/Engine/Plugins/PCG/Source/PCG/Private/Elements/PCGStaticMeshSpawnerKernel.cpp b/Engine/Plugins/PCG/Source/PCG/Private/Elements/PCGStaticMeshSpawnerKernel.cpp
--- a/Engine/Plugins/PCG/Source/PCG/Private/Elements/PCGStaticMeshSpawnerKernel.cpp
+++ b/Engine/Plugins/PCG/Source/PCG/Private/Elements/PCGStaticMeshSpawnerKernel.cpp
@@ -12,6 +12,7 @@
 #include "Compute/PCGPinPropertiesGPU.h"
 #include "Compute/BuiltInKernels/PCGCountUniqueAttributeValuesKernel.h"
 #include "Compute/Data/PCGProxyForGPUData.h"
+#include "Compute/Data/PCGRawBufferData.h"
 #include "Compute/DataInterfaces/PCGInstanceDataInterface.h"
 #include "Compute/DataInterfaces/Elements/PCGStaticMeshSpawnerDataInterface.h"
 #include "Elements/PCGStaticMeshSpawner.h"
@@ -164,7 +165,7 @@ void UPCGStaticMeshSpawnerKernel::GetKernelAttributeKeys(TArray<FPCGKernelAttrib
 void UPCGStaticMeshSpawnerKernel::GetInputPins(TArray<FPCGPinProperties>& OutPins) const
 {
 	OutPins.Emplace(PCGPinConstants::DefaultInputLabel, EPCGDataType::Point);
-	OutPins.Emplace(PCGStaticMeshSpawnerConstants::InstanceCountsPinLabel, EPCGDataType::Param);
+	OutPins.Emplace(PCGStaticMeshSpawnerConstants::InstanceCountsPinLabel, bUseRawInstanceCountsBuffer ? FPCGDataTypeInfoRawBuffer::AsId() : FPCGDataTypeInfoParam::AsId());
 }
 
 void UPCGStaticMeshSpawnerKernel::GetOutputPins(TArray<FPCGPinPropertiesGPU>& OutPins) const
diff --git a/Engine/Plugins/PCG/Source/PCG/Private/Graph/PCGGraphCompilerGPU.cpp b/Engine/Plugins/PCG/Source/PCG/Private/Graph/PCGGraphCompilerGPU.cpp
--- a/Engine/Plugins/PCG/Source/PCG/Private/Graph/PCGGraphCompilerGPU.cpp
+++ b/Engine/Plugins/PCG/Source/PCG/Private/Graph/PCGGraphCompilerGPU.cpp
@@ -433,7 +433,7 @@ void FPCGGraphCompilerGPU::ExpandGPUNodeKernelsToTasks(
 
 					// Look for output pin of upstream kernel in order to set the input upstream pin properties.
 					UpstreamKernelOutputs.Reset();
-					Kernel->GetOutputPins(UpstreamKernelOutputs);
+					UpstreamKernel->GetOutputPins(UpstreamKernelOutputs);
 
 					const FPCGPinPropertiesGPU* UpstreamOutput = UpstreamKernelOutputs.FindByPredicate([Edge](const FPCGPinPropertiesGPU& InUpstreamOutput)
 					{
diff --git a/Engine/Plugins/PCG/Source/PCG/Public/Elements/PCGStaticMeshSpawnerKernel.h b/Engine/Plugins/PCG/Source/PCG/Public/Elements/PCGStaticMeshSpawnerKernel.h
--- a/Engine/Plugins/PCG/Source/PCG/Public/Elements/PCGStaticMeshSpawnerKernel.h
+++ b/Engine/Plugins/PCG/Source/PCG/Public/Elements/PCGStaticMeshSpawnerKernel.h
@@ -33,8 +33,15 @@ public:
 	virtual void GetInputPins(TArray<FPCGPinProperties>& OutPins) const override;
 	virtual void GetOutputPins(TArray<FPCGPinPropertiesGPU>& OutPins) const override;
 
+	void SetUseRawInstanceCountsBuffer(bool bInUseRawInstanceCountsBuffer) { bUseRawInstanceCountsBuffer = bInUseRawInstanceCountsBuffer; }
+
 protected:
 #if WITH_EDITOR
 	virtual bool PerformStaticValidation() override;
 #endif
+
+protected:
+	/** Consume the instance counts as a raw buffer instead of an attribute set. */
+	UPROPERTY()
+	bool bUseRawInstanceCountsBuffer = false;
 };
