[PCG] Fix compute graph submission failed warnings when generating cells that have 0 instances.

Instance data interface needs to gracefully no-op in this case.

#rnx

diff --git a/Engine/Plugins/PCG/Source/PCG/Private/Compute/DataInterfaces/PCGInstanceDataInterface.cpp b/Engine/Plugins/PCG/Source/PCG/Private/Compute/DataInterfaces/PCGInstanceDataInterface.cpp
--- a/Engine/Plugins/PCG/Source/PCG/Private/Compute/DataInterfaces/PCGInstanceDataInterface.cpp
+++ b/Engine/Plugins/PCG/Source/PCG/Private/Compute/DataInterfaces/PCGInstanceDataInterface.cpp
@@ -256,11 +256,7 @@ FPCGInstanceDataProviderProxy::FPCGInstanceDataProviderProxy(const FParams& InPa
 		bIsValid = false;
 	}
 
-	if (!PrimitiveFactory)
-	{
-		bIsValid = false;
-	}
-	else
+	if (PrimitiveFactory)
 	{
 		for (int32 Index = 0; Index < PrimitiveFactory->GetNumPrimitives(); ++Index)
 		{
@@ -293,13 +289,14 @@ void FPCGInstanceDataProviderProxy::AllocateResources(FRDGBuilder& GraphBuilder,
 	InstanceCustomFloatDataSRV = GraphBuilder.CreateSRV(InstanceCustomFloatData);
 	InstanceCustomFloatDataUAV = GraphBuilder.CreateUAV(InstanceCustomFloatData);
 
-	const int32 NumCountersRequired = FMath::Max(1, PrimitiveFactory->GetNumPrimitives());
+	const int32 NumCountersRequired = FMath::Max(1u, PrimitiveFactory ? PrimitiveFactory->GetNumPrimitives() : 0u);
 	WriteCounters = GraphBuilder.CreateBuffer(FRDGBufferDesc::CreateStructuredDesc(sizeof(uint32), NumCountersRequired), TEXT("PCGWriteCounters"));
 	WriteCountersSRV = GraphBuilder.CreateSRV(WriteCounters);
 	WriteCountersUAV = GraphBuilder.CreateUAV(WriteCounters);
-	TArray<uint32> Zeros;
+
+	TArray<uint32, TInlineAllocator<32>> Zeros;
 	Zeros.SetNumZeroed(NumCountersRequired);
-	GraphBuilder.QueueBufferUpload(WriteCounters, Zeros.GetData(), Zeros.Num() * Zeros.GetTypeSize());
+	GraphBuilder.QueueBufferUpload(WriteCounters, MakeConstArrayView(Zeros));
 }
 
 void FPCGInstanceDataProviderProxy::GatherDispatchData(FDispatchData const& InDispatchData)
@@ -308,24 +305,28 @@ void FPCGInstanceDataProviderProxy::GatherDispatchData(FDispatchData const& InDi
 	for (int32 InvocationIndex = 0; InvocationIndex < ParameterArray.Num(); ++InvocationIndex)
 	{
 		FParameters& Parameters = ParameterArray[InvocationIndex];
-		Parameters.NumPrimitives = PrimitiveFactory->GetNumPrimitives();
+		Parameters.NumPrimitives = PrimitiveFactory ? PrimitiveFactory->GetNumPrimitives() : 0;
 		Parameters.NumInstancesAllPrimitives = NumInstancesAllPrimitives;
 		Parameters.NumCustomFloatsPerInstance = NumCustomFloatsPerInstance;
 		Parameters.InstanceData = InstanceDataUAV;
 		Parameters.InstanceCustomFloatData = InstanceCustomFloatDataUAV;
 		Parameters.WriteCounters = WriteCountersUAV;
 
-		uint32 CumulativeInstanceCount = 0;
-		for (int32 Index = 0; Index < PrimitiveFactory->GetNumPrimitives(); ++Index)
+		if (PrimitiveFactory)
 		{
-			const uint32 PrimitiveInstanceCount = PrimitiveFactory->GetNumInstances(Index);
+			uint32 CumulativeInstanceCount = 0;
 
-			Parameters.PrimitiveParameters[Index] = FUintVector4(
-				PrimitiveInstanceCount,
-				CumulativeInstanceCount,
-				0, 0); // Unused
+			for (int32 Index = 0; Index < PrimitiveFactory->GetNumPrimitives(); ++Index)
+			{
+				const uint32 PrimitiveInstanceCount = PrimitiveFactory->GetNumInstances(Index);
 
-			CumulativeInstanceCount += PrimitiveInstanceCount;
+				Parameters.PrimitiveParameters[Index] = FUintVector4(
+					PrimitiveInstanceCount,
+					CumulativeInstanceCount,
+					0, 0); // Unused
+
+				CumulativeInstanceCount += PrimitiveInstanceCount;
+			}
 		}
 	}
 }
