[PCG] Fixed soft path/soft class getters/settters API in pcg graph parameters

#rnx

diff --git a/Engine/Plugins/PCG/Source/PCG/Private/Helpers/PCGGraphParameterExtension.cpp b/Engine/Plugins/PCG/Source/PCG/Private/Helpers/PCGGraphParameterExtension.cpp
--- a/Engine/Plugins/PCG/Source/PCG/Private/Helpers/PCGGraphParameterExtension.cpp
+++ b/Engine/Plugins/PCG/Source/PCG/Private/Helpers/PCGGraphParameterExtension.cpp
@@ -39,13 +39,21 @@ template<> TValueOrError<FString, EPropertyBagResult> FPCGGraphParameterExtensio
 {
 	return PropertyBag.GetValueString(PropertyName);
 }
-template<> TValueOrError<FSoftObjectPath*, EPropertyBagResult> FPCGGraphParameterExtension::GetGraphParameter(const FInstancedPropertyBag& PropertyBag, const FName PropertyName)
+template<> TValueOrError<FSoftObjectPath, EPropertyBagResult> FPCGGraphParameterExtension::GetGraphParameter(const FInstancedPropertyBag& PropertyBag, const FName PropertyName)
 {
-	return PropertyBag.GetValueStruct<FSoftObjectPath>(PropertyName);
+	return PropertyBag.GetValueSoftPath(PropertyName);
 }
-template<> TValueOrError<FSoftClassPath*, EPropertyBagResult> FPCGGraphParameterExtension::GetGraphParameter(const FInstancedPropertyBag& PropertyBag, const FName PropertyName)
+template<> TValueOrError<FSoftClassPath, EPropertyBagResult> FPCGGraphParameterExtension::GetGraphParameter(const FInstancedPropertyBag& PropertyBag, const FName PropertyName)
 {
-	return PropertyBag.GetValueStruct<FSoftClassPath>(PropertyName);
+	TValueOrError<FSoftObjectPath, EPropertyBagResult> InitialResult = PropertyBag.GetValueSoftPath(PropertyName);
+	if (InitialResult.HasError())
+	{
+		return MakeError(InitialResult.GetError());
+	}
+	else
+	{
+		return MakeValue(InitialResult.GetValue().ToString());
+	}
 }
 template<> TValueOrError<TSoftObjectPtr<UObject>, EPropertyBagResult> FPCGGraphParameterExtension::GetGraphParameter(const FInstancedPropertyBag& PropertyBag, const FName PropertyName)
 {
@@ -149,11 +157,11 @@ template<> EPropertyBagResult FPCGGraphParameterExtension::SetGraphParameter(FIn
 }
 template<> EPropertyBagResult FPCGGraphParameterExtension::SetGraphParameter(FInstancedPropertyBag& PropertyBag, const FName PropertyName, const FSoftObjectPath& Value)
 {
-	return PropertyBag.SetValueStruct(PropertyName, Value);
+	return PropertyBag.SetValueSoftPath(PropertyName, Value);
 }
 template<> EPropertyBagResult FPCGGraphParameterExtension::SetGraphParameter(FInstancedPropertyBag& PropertyBag, const FName PropertyName, const FSoftClassPath& Value)
 {
-	return PropertyBag.SetValueStruct(PropertyName, Value);
+	return PropertyBag.SetValueSoftPath(PropertyName, Value);
 }
 template<> EPropertyBagResult FPCGGraphParameterExtension::SetGraphParameter(FInstancedPropertyBag& PropertyBag, const FName PropertyName, const TSoftObjectPtr<UObject>& Value)
 {
@@ -200,4 +208,15 @@ EPropertyBagResult FPCGGraphParameterExtension::SetGraphParameter(FInstancedProp
 {
 	// TODO: To stay congruent with BP Enums and Property Bag, Value converted to uint8. Update this if/when this is improved externally
 	return Enum ? PropertyBag.SetValueEnum(PropertyName, static_cast<uint8>(Value), Enum) : PropertyBag.SetValueByte(PropertyName, static_cast<uint8>(Value));
+}
+
+// DEPRECATED 5.7
+template<> TValueOrError<FSoftObjectPath*, EPropertyBagResult> FPCGGraphParameterExtension::GetGraphParameter(const FInstancedPropertyBag& PropertyBag, const FName PropertyName)
+{
+	return PropertyBag.GetValueStruct<FSoftObjectPath>(PropertyName);
+}
+// DEPRECATED 5.7
+template<> TValueOrError<FSoftClassPath*, EPropertyBagResult> FPCGGraphParameterExtension::GetGraphParameter(const FInstancedPropertyBag& PropertyBag, const FName PropertyName)
+{
+	return PropertyBag.GetValueStruct<FSoftClassPath>(PropertyName);
 }
\ No newline at end of file
diff --git a/Engine/Plugins/PCG/Source/PCG/Private/Helpers/PCGGraphParametersHelpers.cpp b/Engine/Plugins/PCG/Source/PCG/Private/Helpers/PCGGraphParametersHelpers.cpp
--- a/Engine/Plugins/PCG/Source/PCG/Private/Helpers/PCGGraphParametersHelpers.cpp
+++ b/Engine/Plugins/PCG/Source/PCG/Private/Helpers/PCGGraphParametersHelpers.cpp
@@ -234,7 +234,7 @@ uint8 UPCGGraphParametersHelpers::GetEnumParameter(const UPCGGraphInterface* Gra
 
 FSoftObjectPath UPCGGraphParametersHelpers::GetSoftObjectPathParameter(const UPCGGraphInterface* GraphInterface, const FName Name)
 {
-	return PCGGraphParametersHelpersPrivate::ValidateAndReturnResult<FSoftObjectPath, FSoftObjectPath*>(GraphInterface, Name);
+	return PCGGraphParametersHelpersPrivate::ValidateAndReturnResult<FSoftObjectPath>(GraphInterface, Name);
 }
 
 TSoftObjectPtr<UObject> UPCGGraphParametersHelpers::GetSoftObjectParameter(const UPCGGraphInterface* GraphInterface, const FName Name)
diff --git a/Engine/Plugins/PCG/Source/PCG/Public/Helpers/PCGGraphParameterExtension.h b/Engine/Plugins/PCG/Source/PCG/Public/Helpers/PCGGraphParameterExtension.h
--- a/Engine/Plugins/PCG/Source/PCG/Public/Helpers/PCGGraphParameterExtension.h
+++ b/Engine/Plugins/PCG/Source/PCG/Public/Helpers/PCGGraphParameterExtension.h
@@ -32,8 +32,8 @@ template<> PCG_API TValueOrError<int32, EPropertyBagResult> FPCGGraphParameterEx
 template<> PCG_API TValueOrError<int64, EPropertyBagResult> FPCGGraphParameterExtension::GetGraphParameter(const FInstancedPropertyBag& PropertyBag, const FName PropertyName);
 template<> PCG_API TValueOrError<FName, EPropertyBagResult> FPCGGraphParameterExtension::GetGraphParameter(const FInstancedPropertyBag& PropertyBag, const FName PropertyName);
 template<> PCG_API TValueOrError<FString, EPropertyBagResult> FPCGGraphParameterExtension::GetGraphParameter(const FInstancedPropertyBag& PropertyBag, const FName PropertyName);
-template<> PCG_API TValueOrError<FSoftObjectPath*, EPropertyBagResult> FPCGGraphParameterExtension::GetGraphParameter(const FInstancedPropertyBag& PropertyBag, const FName PropertyName);
-template<> PCG_API TValueOrError<FSoftClassPath*, EPropertyBagResult> FPCGGraphParameterExtension::GetGraphParameter(const FInstancedPropertyBag& PropertyBag, const FName PropertyName);
+template<> PCG_API TValueOrError<FSoftObjectPath, EPropertyBagResult> FPCGGraphParameterExtension::GetGraphParameter(const FInstancedPropertyBag& PropertyBag, const FName PropertyName);
+template<> PCG_API TValueOrError<FSoftClassPath, EPropertyBagResult> FPCGGraphParameterExtension::GetGraphParameter(const FInstancedPropertyBag& PropertyBag, const FName PropertyName);
 template<> PCG_API TValueOrError<TSoftObjectPtr<UObject>, EPropertyBagResult> FPCGGraphParameterExtension::GetGraphParameter(const FInstancedPropertyBag& PropertyBag, const FName PropertyName);
 template<> PCG_API TValueOrError<TSoftClassPtr<UObject>, EPropertyBagResult> FPCGGraphParameterExtension::GetGraphParameter(const FInstancedPropertyBag& PropertyBag, const FName PropertyName);
 template<> PCG_API TValueOrError<UObject*, EPropertyBagResult> FPCGGraphParameterExtension::GetGraphParameter(const FInstancedPropertyBag& PropertyBag, const FName PropertyName);
@@ -66,4 +66,13 @@ template<> PCG_API EPropertyBagResult FPCGGraphParameterExtension::SetGraphParam
 template<> PCG_API EPropertyBagResult FPCGGraphParameterExtension::SetGraphParameter(FInstancedPropertyBag& PropertyBag, const FName PropertyName, const FVector2D& Value);
 template<> PCG_API EPropertyBagResult FPCGGraphParameterExtension::SetGraphParameter(FInstancedPropertyBag& PropertyBag, const FName PropertyName, const FQuat& Value);
 
+
+template<> 
+UE_DEPRECATED(5.7, "This function is not working properly and the non-pointer version should be used instead.")
+PCG_API TValueOrError<FSoftObjectPath*, EPropertyBagResult> FPCGGraphParameterExtension::GetGraphParameter(const FInstancedPropertyBag& PropertyBag, const FName PropertyName);
+
+template<> 
+UE_DEPRECATED(5.7, "This function is not working properly and the non-pointer version should be used instead.")
+PCG_API TValueOrError<FSoftClassPath*, EPropertyBagResult> FPCGGraphParameterExtension::GetGraphParameter(const FInstancedPropertyBag& PropertyBag, const FName PropertyName);
+
 #undef UE_API
