[PCG] Bughawk fixes CL 45626515

#rnx

diff --git a/Engine/Plugins/PCG/Source/PCG/Private/Helpers/PCGBlueprintHelpers.cpp b/Engine/Plugins/PCG/Source/PCG/Private/Helpers/PCGBlueprintHelpers.cpp
--- a/Engine/Plugins/PCG/Source/PCG/Private/Helpers/PCGBlueprintHelpers.cpp
+++ b/Engine/Plugins/PCG/Source/PCG/Private/Helpers/PCGBlueprintHelpers.cpp
@@ -39,7 +39,7 @@ void UPCGGenerateGraphAsync::Activate()
 	GenerateParams.Seed = Seed;
 	GenerateParams.GenerationCallback.BindUObject(this, &UPCGGenerateGraphAsync::OnGenerationDone);
 
-	if (UPCGEngineSubsystem* EngineSubsystem = UPCGEngineSubsystem::Get())
+	if (UPCGEngineSubsystem* EngineSubsystem = UPCGEngineSubsystem::Get(); ensure(EngineSubsystem))
 	{
 		EngineSubsystem->GenerateGraph(GenerateParams);
 	}
diff --git a/Engine/Plugins/PCG/Source/PCG/Private/PCGDefaultExecutionSource.cpp b/Engine/Plugins/PCG/Source/PCG/Private/PCGDefaultExecutionSource.cpp
--- a/Engine/Plugins/PCG/Source/PCG/Private/PCGDefaultExecutionSource.cpp
+++ b/Engine/Plugins/PCG/Source/PCG/Private/PCGDefaultExecutionSource.cpp
@@ -218,6 +218,12 @@ void UPCGDefaultExecutionSource::Generate()
 			
 			CurrentGenerationTask = Subsystem->ScheduleGeneric(Params);
 		}
+#if WITH_EDITOR
+		else
+		{
+			Subsystem->OnPCGSourceGenerationDone(this, EPCGGenerationStatus::Aborted);
+		}
+#endif
 	}
 }
 
diff --git a/Engine/Plugins/PCG/Source/PCG/Private/PCGGraph.cpp b/Engine/Plugins/PCG/Source/PCG/Private/PCGGraph.cpp
--- a/Engine/Plugins/PCG/Source/PCG/Private/PCGGraph.cpp
+++ b/Engine/Plugins/PCG/Source/PCG/Private/PCGGraph.cpp
@@ -2587,7 +2587,7 @@ void UPCGGraphInstance::SetGraph(UPCGGraphInterface* InGraph)
 	Graph = InGraph;
 
 #if WITH_EDITOR
-	bIsStandaloneGraph = InGraph->bIsStandaloneGraph;
+	bIsStandaloneGraph = InGraph && InGraph->bIsStandaloneGraph;
 
 	SetupCallbacks();
 #endif // WITH_EDITOR
diff --git a/Engine/Plugins/PCG/Source/PCG/Public/Helpers/PCGBlueprintHelpers.h b/Engine/Plugins/PCG/Source/PCG/Public/Helpers/PCGBlueprintHelpers.h
--- a/Engine/Plugins/PCG/Source/PCG/Public/Helpers/PCGBlueprintHelpers.h
+++ b/Engine/Plugins/PCG/Source/PCG/Public/Helpers/PCGBlueprintHelpers.h
@@ -123,7 +123,7 @@ public:
 	UFUNCTION(BlueprintCallable, Category = "PCG|Helpers", meta = (WorldContext="WorldContextObject"))
 	static PCG_API TArray<FPCGLandscapeLayerWeight> GetInterpolatedPCGLandscapeLayerWeights(UObject* WorldContextObject, const FVector& Location);
 
-	UFUNCTION(BLueprintCallable, Category = "PCG|Helpers", meta = (ScriptMethod))
+	UFUNCTION(BlueprintCallable, Category = "PCG|Helpers", meta = (ScriptMethod))
 	static PCG_API int64 GetTaskId(UPARAM(ref) FPCGContext& Context);
 
 	/** Flush the cache, to be used if you have changed something PCG depends on at runtime. Same as `pcg.FlushCache` command. Returns true if it succeeded. */
