[PCG] Fix crash when adding to UPCGSchedulingPolicyBase::GridsDependentOnWorldStreaming
- There are still issues with undoing changes to the policy which need to be investigated further

diff --git a/Engine/Plugins/PCG/Source/PCG/Private/PCGComponent.cpp b/Engine/Plugins/PCG/Source/PCG/Private/PCGComponent.cpp
--- a/Engine/Plugins/PCG/Source/PCG/Private/PCGComponent.cpp
+++ b/Engine/Plugins/PCG/Source/PCG/Private/PCGComponent.cpp
@@ -1984,7 +1984,9 @@ void UPCGComponent::MarkSubObjectsAsGarbage()
 
 	auto ModifyAndMarkForGarbage = [](UObject* SubObject)
 	{
-		if (!SubObject)
+		// @todo_pcg: revisit this but for now this fixes a crash when adding elements to GridsDependentOnWorldStreaming on a Blueprint PCG Component
+		// ModifyAndMarkForGarbage main use is for UPCGData that would have references on the PCG Component's owning actor.
+		if (!IsValid(SubObject) || SubObject->IsA<UPCGSchedulingPolicyBase>())
 		{
 			return;
 		}
