[PCG] Enable raw data readback optimization for SM Spawner by default.

Has been enabled in our test project for some time.

#rnx

diff --git a/Engine/Plugins/PCG/Source/PCG/Private/Elements/PCGStaticMeshSpawner.cpp b/Engine/Plugins/PCG/Source/PCG/Private/Elements/PCGStaticMeshSpawner.cpp
--- a/Engine/Plugins/PCG/Source/PCG/Private/Elements/PCGStaticMeshSpawner.cpp
+++ b/Engine/Plugins/PCG/Source/PCG/Private/Elements/PCGStaticMeshSpawner.cpp
@@ -37,7 +37,7 @@ static TAutoConsoleVariable<bool> CVarAllowISMReuse(
 
 static TAutoConsoleVariable<bool> CVarUseRawDataForAnalysis(
 	TEXT("pcg.RuntimeGeneration.ISM.UseRawDataForGPUAnalysis"),
-	false,
+	true,
 	TEXT("Read back raw data in GPU primitive instancing path rather than constructing an attribute set, to save CPU time. Read-only, must be set from ini."),
 	ECVF_ReadOnly // Structural change, could flush graph compilation cache, but trickier to flush compute graph instance pool in UPCGGraph currently.
 	);
