[PCG] Fixed missing clear of data description cache which saves 1.5-4.4mb of persistent CPU mem in our test levels.

#rnx

diff --git a/Engine/Plugins/PCG/Source/PCG/Private/Compute/PCGDataBinding.cpp b/Engine/Plugins/PCG/Source/PCG/Private/Compute/PCGDataBinding.cpp
--- a/Engine/Plugins/PCG/Source/PCG/Private/Compute/PCGDataBinding.cpp
+++ b/Engine/Plugins/PCG/Source/PCG/Private/Compute/PCGDataBinding.cpp
@@ -73,6 +73,7 @@ void UPCGDataBinding::ReleaseTransientResources()
 	DataToDebug.Reset();
 	DataToInspect.Reset();
 	ContextHandle = nullptr;
+	DataDescriptionCache.Reset();
 }
 
 void UPCGDataBinding::PrimeDataDescriptionCache()
