[PCG] Removed extraneous point on spline discretization on linear segments

#rnx

diff --git a/Engine/Plugins/PCG/Source/PCG/Private/Data/PCGSplineStruct.cpp b/Engine/Plugins/PCG/Source/PCG/Private/Data/PCGSplineStruct.cpp
--- a/Engine/Plugins/PCG/Source/PCG/Private/Data/PCGSplineStruct.cpp
+++ b/Engine/Plugins/PCG/Source/PCG/Private/Data/PCGSplineStruct.cpp
@@ -743,6 +743,15 @@ FVector FPCGSplineStruct::GetTangentAtSplineInputKey(float InKey, ESplineCoordin
 bool FPCGSplineStruct::ConvertSplineSegmentToPolyLine(int32 SplinePointStartIndex, ESplineCoordinateSpace::Type CoordinateSpace, const float MaxSquareDistanceFromSpline, TArray<FVector>& OutPoints) const
 {
 	OutPoints.Empty();
+
+	// Early out - if the segment is a linear segment we can just add the start/end and not add any additional points.
+	if (SplinePointStartIndex < GetNumberOfPoints() && GetSplinePointType(SplinePointStartIndex) == EInterpCurveMode::CIM_Linear &&
+		(SplinePointStartIndex < GetNumberOfPoints() - 1 || IsClosedLoop()))
+	{
+		OutPoints.Add(GetLocation(SplinePointStartIndex));
+		OutPoints.Add(GetLocation((SplinePointStartIndex + 1) % GetNumberOfPoints()));
+		return true;
+	}
 	
 	TArray<double> DummyDistances;
 
