[PCG] Filter out unsupported types in CustomHLSL settings dropdowns

#rnx

diff --git a/Engine/Plugins/PCG/Source/PCG/Internal/Compute/PCGComputeCommon.h b/Engine/Plugins/PCG/Source/PCG/Internal/Compute/PCGComputeCommon.h
--- a/Engine/Plugins/PCG/Source/PCG/Internal/Compute/PCGComputeCommon.h
+++ b/Engine/Plugins/PCG/Source/PCG/Internal/Compute/PCGComputeCommon.h
@@ -110,6 +110,12 @@ namespace PCGComputeHelpers
 	EPCGElementDimension GetElementDimension(const UPCGData* InData);
 	EPCGElementDimension GetElementDimension(const FPCGDataTypeIdentifier& InDataType);
 
+	/** PCG data types supported in GPU node inputs. */
+	const TArray<FPCGDataTypeIdentifier>& GetAllowedInputTypesList();
+
+	/** PCG data types supported in GPU node outputs. */
+	const TArray<FPCGDataTypeIdentifier>& GetAllowedOutputTypesList();
+
 	/** True if 'Type' is valid on a GPU input pin. */
 	bool IsTypeAllowedAsInput(const FPCGDataTypeIdentifier& Type);
 
diff --git a/Engine/Plugins/PCG/Source/PCG/Private/Compute/Elements/PCGCustomHLSL.cpp b/Engine/Plugins/PCG/Source/PCG/Private/Compute/Elements/PCGCustomHLSL.cpp
--- a/Engine/Plugins/PCG/Source/PCG/Private/Compute/Elements/PCGCustomHLSL.cpp
+++ b/Engine/Plugins/PCG/Source/PCG/Private/Compute/Elements/PCGCustomHLSL.cpp
@@ -317,6 +317,16 @@ void UPCGCustomHLSLSettings::CreateKernels(FPCGGPUCompilationContext& InOutConte
 
 	PCGKernelHelpers::CreateKernel<UPCGCustomHLSLKernel>(InOutContext, CreateParams, OutKernels, OutEdges);
 }
+
+TArray<FPCGDataTypeIdentifier> UPCGCustomHLSLSettings::GetAllowedInputTypes()
+{
+	return PCGComputeHelpers::GetAllowedInputTypesList();
+}
+
+TArray<FPCGDataTypeIdentifier> UPCGCustomHLSLSettings::GetAllowedOutputTypes()
+{
+	return PCGComputeHelpers::GetAllowedOutputTypesList();
+}
 #endif // WITH_EDITOR
 
 const FPCGPinPropertiesGPU* UPCGCustomHLSLSettings::GetOutputPinPropertiesGPU(const FName& InPinLabel) const
diff --git a/Engine/Plugins/PCG/Source/PCG/Private/Compute/Elements/PCGCustomHLSL.h b/Engine/Plugins/PCG/Source/PCG/Private/Compute/Elements/PCGCustomHLSL.h
--- a/Engine/Plugins/PCG/Source/PCG/Private/Compute/Elements/PCGCustomHLSL.h
+++ b/Engine/Plugins/PCG/Source/PCG/Private/Compute/Elements/PCGCustomHLSL.h
@@ -98,6 +98,12 @@ public:
 #if WITH_EDITOR
 	UFUNCTION()
 	static bool SupportsComposition() { return false; }
+
+	UFUNCTION()
+	static TArray<FPCGDataTypeIdentifier> GetAllowedInputTypes();
+
+	UFUNCTION()
+	static TArray<FPCGDataTypeIdentifier> GetAllowedOutputTypes();
 #endif // WITH_EDITOR
 
 protected:
@@ -159,10 +165,10 @@ public:
 	UPROPERTY(EditAnywhere, DisplayName = "Input Pins", Category = "Settings|Thread Count", meta = (EditCondition = "KernelType == EPCGKernelType::Custom && DispatchThreadCount == EPCGDispatchThreadCount::FromProductOfInputPins", EditConditionHides, GetOptions = "GetInputPinNames"))
 	TArray<FName> ThreadCountInputPinLabels;
 
-	UPROPERTY(EditAnywhere, Category = "Settings")
+	UPROPERTY(EditAnywhere, Category = "Settings", meta = (PCG_TypeFilter = "GetAllowedInputTypes()"))
 	TArray<FPCGPinProperties> InputPins = Super::DefaultPointInputPinProperties();
 
-	UPROPERTY(EditAnywhere, Category = "Settings")
+	UPROPERTY(EditAnywhere, Category = "Settings", meta = (PCG_TypeFilter = "GetAllowedOutputTypes()"))
 	TArray<FPCGPinPropertiesGPU> OutputPins = { FPCGPinPropertiesGPU(PCGPinConstants::DefaultOutputLabel, FPCGDataTypeIdentifier{EPCGDataType::Point}) };
 
 #if WITH_EDITOR
diff --git a/Engine/Plugins/PCG/Source/PCG/Private/Compute/PCGComputeCommon.cpp b/Engine/Plugins/PCG/Source/PCG/Private/Compute/PCGComputeCommon.cpp
--- a/Engine/Plugins/PCG/Source/PCG/Private/Compute/PCGComputeCommon.cpp
+++ b/Engine/Plugins/PCG/Source/PCG/Private/Compute/PCGComputeCommon.cpp
@@ -29,50 +29,6 @@
 
 #include UE_INLINE_GENERATED_CPP_BY_NAME(PCGComputeCommon)
 
-namespace PCGComputeConstants
-{
-	/** PCG data types supported in GPU node inputs. */
-	inline const FPCGDataTypeIdentifier& GetAllowedInputTypes()
-	{
-		static FPCGDataTypeIdentifier Result = FPCGDataTypeIdentifier::Construct(
-			EPCGDataType::Point
-			, EPCGDataType::Param
-			, EPCGDataType::Landscape
-			, EPCGDataType::Param
-			, EPCGDataType::Landscape
-			, EPCGDataType::BaseTexture
-			, EPCGDataType::VirtualTexture
-			, EPCGDataType::Resource
-			, EPCGDataType::ProxyForGPU
-			, FPCGDataTypeInfoRawBuffer::AsId()
-		);
-
-		return Result;
-	}
-
-	/** PCG data types supported in GPU node outputs. */
-	inline const FPCGDataTypeIdentifier& GetAllowedOutputTypes()
-	{
-		static FPCGDataTypeIdentifier Result = FPCGDataTypeIdentifier::Construct(
-			EPCGDataType::Point
-			, EPCGDataType::Param
-			, EPCGDataType::ProxyForGPU
-			, EPCGDataType::BaseTexture
-			, FPCGDataTypeInfoRawBuffer::AsId()
-		);
-
-		return Result;
-	}
-
-	/** PCG data types supported in GPU data collections. */
-	inline const FPCGDataTypeIdentifier& GetAllowedDataCollectionTypes()
-	{
-		// Intentionally excludes raw buffer data which cannot be mixed with other data types like point.
-		static FPCGDataTypeIdentifier Result{ EPCGDataType::Point | EPCGDataType::Param | EPCGDataType::ProxyForGPU };
-		return Result;
-	}
-}
-
 namespace PCGComputeHelpers
 {
 	static TAutoConsoleVariable<float> CVarMaxGPUBufferSizeProportion(
@@ -123,19 +79,70 @@ namespace PCGComputeHelpers
 		return !!(InDataType & EPCGDataType::BaseTexture) ? EPCGElementDimension::Two : EPCGElementDimension::One;
 	}
 
+	const TArray<FPCGDataTypeIdentifier>& GetAllowedInputTypesList()
+	{
+		static TArray<FPCGDataTypeIdentifier> AllowedTypes =
+		{
+			EPCGDataType::Point,
+			EPCGDataType::Param,
+			EPCGDataType::Landscape,
+			EPCGDataType::BaseTexture,
+			EPCGDataType::VirtualTexture,
+			EPCGDataType::StaticMeshResource,
+			EPCGDataType::ProxyForGPU,
+			FPCGDataTypeInfoRawBuffer::AsId()
+		};
+
+		return AllowedTypes;
+	}
+
+	const TArray<FPCGDataTypeIdentifier>& GetAllowedOutputTypesList()
+	{
+		static TArray<FPCGDataTypeIdentifier> AllowedTypes =
+		{
+			EPCGDataType::Point,
+			EPCGDataType::Param,
+			EPCGDataType::ProxyForGPU,
+			EPCGDataType::BaseTexture,
+			FPCGDataTypeInfoRawBuffer::AsId()
+		};
+
+		return AllowedTypes;
+	}
+
+	inline const FPCGDataTypeIdentifier& GetAllowedInputTypes()
+	{
+		static FPCGDataTypeIdentifier Result = FPCGDataTypeIdentifier::Construct(GetAllowedInputTypesList());
+		return Result;
+	}
+
+	inline const FPCGDataTypeIdentifier& GetAllowedOutputTypes()
+	{
+		static FPCGDataTypeIdentifier Result = FPCGDataTypeIdentifier::Construct(GetAllowedOutputTypesList());
+		return Result;
+	}
+
+	/** PCG data types supported in GPU data collections. */
+	inline const FPCGDataTypeIdentifier& GetAllowedDataCollectionTypes()
+	{
+		// Intentionally excludes raw buffer data which cannot be mixed with other data types like point.
+		static FPCGDataTypeIdentifier Result{ EPCGDataType::Point | EPCGDataType::Param | EPCGDataType::ProxyForGPU };
+		return Result;
+	}
+
 	bool IsTypeAllowedAsInput(const FPCGDataTypeIdentifier& Type)
 	{
-		return PCGComputeConstants::GetAllowedInputTypes().Intersects(Type);
+		return GetAllowedInputTypes().Intersects(Type);
 	}
 
 	bool IsTypeAllowedAsOutput(const FPCGDataTypeIdentifier& Type)
 	{
-		return PCGComputeConstants::GetAllowedOutputTypes().Intersects(Type);
+		return GetAllowedOutputTypes().Intersects(Type);
 	}
 
 	bool IsTypeAllowedInDataCollection(const FPCGDataTypeIdentifier& Type)
 	{
-		return PCGComputeConstants::GetAllowedDataCollectionTypes().Intersects(Type);
+		return GetAllowedDataCollectionTypes().Intersects(Type);
 	}
 
 	bool ShouldImportAttributesFromData(const UPCGData* InData)
diff --git a/Engine/Plugins/PCG/Source/PCG/Public/PCGCommon.h b/Engine/Plugins/PCG/Source/PCG/Public/PCGCommon.h
--- a/Engine/Plugins/PCG/Source/PCG/Public/PCGCommon.h
+++ b/Engine/Plugins/PCG/Source/PCG/Public/PCGCommon.h
@@ -249,6 +249,7 @@ namespace PCGObjectMetadata
 	const FLazyName EnumMetadataDomain = TEXT("PCG_MetadataDomain");
 	const FLazyName PropertyReadOnly = TEXT("PCG_PropertyReadOnly");
 	const FLazyName DataTypeIdentifierSupportsComposition = TEXT("PCG_SupportsComposition");
+	const FLazyName DataTypeIdentifierFilter = TEXT("PCG_TypeFilter");
 
 	// Metadata usable in UPROPERTY for customizing the behavior when displaying the property in a property panel or graph node
 	enum
@@ -286,6 +287,9 @@ namespace PCGObjectMetadata
 
 		/// [PropertyMetadata] For FPCGDataTypeIdentifier, will enable selecting types like BitFlag, to compose types.
 		PCG_SupportsComposition,
+
+		/// [PropertyMetadata] For FPCGDataTypeIdentifier, will enable selecting valid types to show in the dropdown.
+		PCG_TypeFilter,
 	};
 
 	// Metadata usable in UENUM for customizing the entry in the attribute property selector
diff --git a/Engine/Plugins/PCG/Source/PCGEditor/Private/Details/PCGDataTypeIdentifierDetails.cpp b/Engine/Plugins/PCG/Source/PCGEditor/Private/Details/PCGDataTypeIdentifierDetails.cpp
--- a/Engine/Plugins/PCG/Source/PCGEditor/Private/Details/PCGDataTypeIdentifierDetails.cpp
+++ b/Engine/Plugins/PCG/Source/PCGEditor/Private/Details/PCGDataTypeIdentifierDetails.cpp
@@ -14,6 +14,28 @@
 
 namespace PCGDataTypeIdentifierDetails
 {
+	/** Walks property parent hierarchy until a property is found that belongs to a UObject. */
+	TSharedPtr<IPropertyHandle> GetFirstUObjectParentProperty(TSharedPtr<IPropertyHandle> InPropertyHandle)
+	{
+		while (InPropertyHandle.IsValid())
+		{
+			if (const FProperty* Property = InPropertyHandle->GetProperty())
+			{
+				const UStruct* OwnerStruct = Property->GetOwnerStruct();
+
+				// OwnerStruct will be a UClass if this property is declared on a UObject.
+				if (OwnerStruct && OwnerStruct->IsA<UClass>())
+				{
+					return InPropertyHandle;
+				}
+			}
+
+			InPropertyHandle = InPropertyHandle->GetParentHandle();
+		}
+
+		return nullptr;
+	}
+
 	// Check if the identifier is a superset of other.
 	bool IsSuperSet(const FPCGDataTypeIdentifier& Identifier, const FPCGDataTypeBaseId& Other)
 	{
@@ -59,6 +81,51 @@ namespace PCGDataTypeIdentifierDetails
 			return false;
 		}
 	}
+
+	/**
+	 * Look for a function on the outer object of the property (or its parent properties) to check if the widget should filter out any types.
+	 */
+	TOptional<TArray<FPCGDataTypeIdentifier>> GetFilter(TSharedRef<IPropertyHandle> InPropertyHandle)
+	{
+		TSharedPtr<IPropertyHandle> PropertyHandle = GetFirstUObjectParentProperty(InPropertyHandle.ToSharedPtr());
+
+		if (!PropertyHandle)
+		{
+			return {};
+		}
+
+		const FString& GetFilterFunctionName = PropertyHandle->GetMetaData(PCGObjectMetadata::DataTypeIdentifierFilter);
+
+		if (!GetFilterFunctionName.IsEmpty())
+		{
+			TArray<UObject*> OutObjects;
+			PropertyHandle->GetOuterObjects(OutObjects);
+
+			if (OutObjects.IsEmpty() || !OutObjects[0])
+			{
+				return {};
+			}
+
+			UObject* Object = OutObjects[0];
+			// Strip `()` that could be at the end of the function name
+			const FString FunctionName = GetFilterFunctionName.Replace(TEXT("()"), TEXT(""));
+			
+			if (!Object->GetClass() || !Object->GetClass()->FindFunctionByName(*FunctionName))
+			{
+				return {};
+			}
+
+			DECLARE_DELEGATE_RetVal(TArray<FPCGDataTypeIdentifier>, FSupportComposition);
+			auto Func = FSupportComposition::CreateUFunction(Object, *FunctionName);
+			
+			if (Func.IsBound())
+			{
+				return Func.Execute();
+			}
+		}
+
+		return {};
+	}
 }
 
 void FPCGDataTypeIdentifierDetails::CustomizeHeader(TSharedRef<IPropertyHandle> InPropertyHandle, FDetailWidgetRow& HeaderRow,
@@ -69,6 +136,7 @@ void FPCGDataTypeIdentifierDetails::CustomizeHeader(TSharedRef<IPropertyHandle>
 
 	TSharedRef<SWidget> PropertyNameWidget = PropertyHandle->CreatePropertyNameWidget();
 	
+	const TOptional<TArray<FPCGDataTypeIdentifier>> Filter = PCGDataTypeIdentifierDetails::GetFilter(InPropertyHandle);
 	bSupportComposition = PCGDataTypeIdentifierDetails::SupportComposition(InPropertyHandle);
 
 	const FPCGDataTypeRegistry& Registry = FPCGModule::GetConstDataTypeRegistry();
@@ -77,13 +145,15 @@ void FPCGDataTypeIdentifierDetails::CustomizeHeader(TSharedRef<IPropertyHandle>
 	VisibleTypes.Empty();
 
 	// Do a depth search so types are ordered by category.
-	Registry.HierarchyDepthSearch([this, &Registry](const FPCGDataTypeBaseId& DataType, int32 Depth) -> FPCGDataTypeRegistry::ESearchCommand
+	Registry.HierarchyDepthSearch([this, &Filter, &Registry](const FPCGDataTypeBaseId& DataType, int32 Depth) -> FPCGDataTypeRegistry::ESearchCommand
 	{
 		const FPCGDataTypeInfo* TypeInfo = Registry.GetTypeInfo(DataType);
 		check(TypeInfo);
 
-		// Remove hidden types
-		if (TypeInfo->Hidden())
+		const bool bInFilter = !Filter.IsSet() || Filter.GetValue().Contains(DataType);
+
+		// Remove hidden or filtered types
+		if (TypeInfo->Hidden() || !bInFilter)
 		{
 			HiddenTypes.Add(DataType);
 		}
