[PCG] Move to storing shared ptrs in the DataDescriptionCache

#rnx

diff --git a/Engine/Plugins/PCG/Source/PCG/Internal/Compute/PCGDataForGPU.h b/Engine/Plugins/PCG/Source/PCG/Internal/Compute/PCGDataForGPU.h
--- a/Engine/Plugins/PCG/Source/PCG/Internal/Compute/PCGDataForGPU.h
+++ b/Engine/Plugins/PCG/Source/PCG/Internal/Compute/PCGDataForGPU.h
@@ -216,9 +216,17 @@ private:
 /** Data description for a data collection (FPCGDataCollection). */
 struct FPCGDataCollectionDesc
 {
-	static FPCGDataCollectionDesc BuildFromDataCollection(
-		const FPCGDataCollection& InDataCollection,
-		const UPCGDataBinding* InBinding);
+	/** FPCGDataCollectionDesc should not be created directly, use MakeShared(). */
+	FPCGDataCollectionDesc(const FPCGDataCollectionDesc&) = delete;
+	FPCGDataCollectionDesc& operator=(const FPCGDataCollectionDesc&) = delete;
+	FPCGDataCollectionDesc(FPCGDataCollectionDesc&&) = delete;
+	FPCGDataCollectionDesc& operator=(FPCGDataCollectionDesc&&) = delete;
+
+	/** Create an FPCGDataCollectionDesc. */
+	static TSharedPtr<FPCGDataCollectionDesc> MakeShared();
+
+	/** Create an FPCGDataCollectionDesc from another FPCGDataCollectionDesc. */
+	static TSharedPtr<FPCGDataCollectionDesc> MakeSharedFrom(const TSharedPtr<const FPCGDataCollectionDesc> InOtherDataDesc);
 
 	/** Computes the size (in bytes) of the header portion of the packed data collection buffer. */
 	uint32 ComputePackedHeaderSizeBytes() const;
@@ -240,9 +248,6 @@ struct FPCGDataCollectionDesc
 	/** Compute total number of processing elements of the given type. */
 	uint32 ComputeTotalElementCount(FPCGDataTypeIdentifier InDataType) const;
 
-	/** Aggregate another data description. */
-	void Combine(const FPCGDataCollectionDesc& Other);
-
 	/** Get description of first attribute with matching identifier in input data. Returns true if such an attribute found and also signals whether multiple matching attributes
 	* with conflicting names are present. */
 	bool GetAttributeDesc(FPCGAttributeIdentifier InAttributeIdentifier, FPCGKernelAttributeDesc& OutAttributeDesc, bool& bOutConflictingTypesFound, bool& bOutPresentOnAllData) const;
@@ -260,13 +265,16 @@ struct FPCGDataCollectionDesc
 
 	/** Description of each data in this data collection. */
 	TArray<FPCGDataDesc> DataDescs;
+
+private:
+	FPCGDataCollectionDesc() = default;
 };
 
 /** A proxy for a data collection residing in a GPU buffer along with functionality to retrieve the data on the CPU. Holds onto GPU memory. */
 struct FPCGProxyForGPUDataCollection : public TSharedFromThis<FPCGProxyForGPUDataCollection>
 {
 public:
-	FPCGProxyForGPUDataCollection(TRefCountPtr<FRDGPooledBuffer> InBuffer, uint32 InBufferSizeBytes, const FPCGDataCollectionDesc& InDescription, const TArray<FString>& InStringTable);
+	FPCGProxyForGPUDataCollection(TRefCountPtr<FRDGPooledBuffer> InBuffer, uint32 InBufferSizeBytes, TSharedPtr<const FPCGDataCollectionDesc> InDescription, const TArray<FString>& InStringTable);
 
 	/** Populates a CPU data object representing the GPU data for the given index, performing a readback from GPU->CPU if required. */
 	bool GetCPUData(FPCGContext* InContext, int32 InDataIndex, FPCGTaggedData& OutData);
@@ -275,7 +283,7 @@ public:
 
 	uint32 GetBufferSizeBytes() const { return BufferSizeBytes; }
 
-	const FPCGDataCollectionDesc& GetDescription() const { return Description; }
+	TSharedPtr<const FPCGDataCollectionDesc> GetDescription() const { return Description; }
 
 	const TArray<FString>& GetStringTable() const { return StringTable; }
 
@@ -285,7 +293,7 @@ protected:
 
 	uint32 BufferSizeBytes = 0;
 
-	FPCGDataCollectionDesc Description;
+	TSharedPtr<const FPCGDataCollectionDesc> Description = nullptr;
 
 	/** Used to comprehend string IDs in buffer. */
 	TArray<FString> StringTable;
@@ -316,7 +324,7 @@ namespace PCGDataForGPUHelpers
 		const UPCGSettings* InSettings,
 		TArray<FName>& InAttributeNames,
 		const UPCGDataBinding* InBinding,
-		const FPCGDataCollectionDesc* InDataCollectionDescription,
+		const TSharedPtr<const FPCGDataCollectionDesc> InDataCollectionDescription,
 		uint32& OutCustomFloatCount,
 		TArray<FUint32Vector4>& OutAttributeIdOffsetStrides);
 }
diff --git a/Engine/Plugins/PCG/Source/PCG/Private/Compute/BuiltInKernels/PCGCountUniqueAttributeValuesKernel.cpp b/Engine/Plugins/PCG/Source/PCG/Private/Compute/BuiltInKernels/PCGCountUniqueAttributeValuesKernel.cpp
--- a/Engine/Plugins/PCG/Source/PCG/Private/Compute/BuiltInKernels/PCGCountUniqueAttributeValuesKernel.cpp
+++ b/Engine/Plugins/PCG/Source/PCG/Private/Compute/BuiltInKernels/PCGCountUniqueAttributeValuesKernel.cpp
@@ -38,7 +38,7 @@ bool UPCGCountUniqueAttributeValuesKernel::IsKernelDataValid(FPCGContext* InCont
 	const UPCGDataBinding* DataBinding = Context->DataBinding.Get();
 	if (DataBinding)
 	{
-		const FPCGDataCollectionDesc* InputDataDesc = DataBinding->GetCachedKernelPinDataDesc(this, PCGPinConstants::DefaultInputLabel, /*bIsInputPin=*/true);
+		const TSharedPtr<const FPCGDataCollectionDesc> InputDataDesc = DataBinding->GetCachedKernelPinDataDesc(this, PCGPinConstants::DefaultInputLabel, /*bIsInputPin=*/true);
 
 		if (!ensure(InputDataDesc))
 		{
@@ -88,7 +88,7 @@ bool UPCGCountUniqueAttributeValuesKernel::IsKernelDataValid(FPCGContext* InCont
 	return true;
 }
 
-bool UPCGCountUniqueAttributeValuesKernel::ComputeOutputBindingDataDesc(const UPCGComputeGraph* InGraph, FName InOutputPinLabel, UPCGDataBinding* InBinding, FPCGDataCollectionDesc& OutDataDesc) const
+TSharedPtr<const FPCGDataCollectionDesc> UPCGCountUniqueAttributeValuesKernel::ComputeOutputBindingDataDesc(const UPCGComputeGraph* InGraph, FName InOutputPinLabel, UPCGDataBinding* InBinding) const
 {
 	check(InGraph);
 	check(InBinding);
@@ -96,25 +96,26 @@ bool UPCGCountUniqueAttributeValuesKernel::ComputeOutputBindingDataDesc(const UP
 	// Code assumes single output pin.
 	if (!ensure(InOutputPinLabel == PCGPinConstants::DefaultOutputLabel))
 	{
-		return false;
+		return nullptr;
 	}
 
-	FPCGDataCollectionDesc InputDesc;
 	const FPCGKernelPin InputKernelPin(KernelIndex, PCGPinConstants::DefaultInputLabel, /*bIsInput=*/true);
-	ensure(InBinding->ComputeKernelPinDataDesc(InputKernelPin, &InputDesc));
+	const TSharedPtr<const FPCGDataCollectionDesc> InputDesc = InBinding->ComputeKernelPinDataDesc(InputKernelPin);
 
-	if (InputDesc.DataDescs.IsEmpty())
+	if (!ensure(InputDesc) || InputDesc->DataDescs.IsEmpty())
 	{
-		return true;
+		return nullptr;
 	}
 
+	TSharedPtr<FPCGDataCollectionDesc> OutDataDesc = FPCGDataCollectionDesc::MakeShared();
+
 	const int32 AttributeId = InBinding->GetAttributeId(AttributeName, EPCGKernelAttributeType::StringKey);
 
 	if (AttributeId != INDEX_NONE)
 	{
 		// Compute a set of (value, count) pairs across all data.
 		TArray<int32> UniqueStringKeyValues;
-		InputDesc.GetUniqueStringKeyValues(AttributeId, UniqueStringKeyValues);
+		InputDesc->GetUniqueStringKeyValues(AttributeId, UniqueStringKeyValues);
 
 		// If the highest string key value is 3, we'll allocate 4 counters. String key value 0 is reserved for empty/null string.
 		const int32* MaxStringKeyValue = Algo::MaxElement(UniqueStringKeyValues);
@@ -140,33 +141,33 @@ bool UPCGCountUniqueAttributeValuesKernel::ComputeOutputBindingDataDesc(const UP
 		if (bEmitPerDataCounts)
 		{
 			// Output data is array of (value, count) pairs for each data.
-			for (int InputIndex = 0; InputIndex < InputDesc.DataDescs.Num(); ++InputIndex)
+			for (int InputIndex = 0; InputIndex < InputDesc->DataDescs.Num(); ++InputIndex)
 			{
-				const FPCGDataDesc& DataDesc = InputDesc.DataDescs[InputIndex];
+				const FPCGDataDesc& DataDesc = InputDesc->DataDescs[InputIndex];
 
 				// Data descriptions have multiple collection objects so worth to move last element for efficiency, otherwise copy.
-				if (InputIndex < InputDesc.DataDescs.Num() - 1)
+				if (InputIndex < InputDesc->DataDescs.Num() - 1)
 				{
-					OutDataDesc.DataDescs.Add(CountDataDesc);
+					OutDataDesc->DataDescs.Add(CountDataDesc);
 				}
 				else
 				{
-					OutDataDesc.DataDescs.Add(MoveTemp(CountDataDesc));
+					OutDataDesc->DataDescs.Add(MoveTemp(CountDataDesc));
 				}
 			}
 		}
 		else
 		{
-			OutDataDesc.DataDescs.Add(MoveTemp(CountDataDesc));
+			OutDataDesc->DataDescs.Add(MoveTemp(CountDataDesc));
 		}
 	}
 
-	return true;
+	return OutDataDesc;
 }
 
 int UPCGCountUniqueAttributeValuesKernel::ComputeThreadCount(const UPCGDataBinding* InBinding) const
 {
-	if (const FPCGDataCollectionDesc* InputPinDesc = InBinding->GetCachedKernelPinDataDesc(this, PCGPinConstants::DefaultInputLabel, /*bIsInput=*/true); ensure(InputPinDesc))
+	if (const TSharedPtr<const FPCGDataCollectionDesc> InputPinDesc = InBinding->GetCachedKernelPinDataDesc(this, PCGPinConstants::DefaultInputLabel, /*bIsInput=*/true); ensure(InputPinDesc))
 	{
 		const int ElementCount = InputPinDesc->ComputeTotalElementCount();
 
diff --git a/Engine/Plugins/PCG/Source/PCG/Private/Compute/BuiltInKernels/PCGCountUniqueAttributeValuesKernel.h b/Engine/Plugins/PCG/Source/PCG/Private/Compute/BuiltInKernels/PCGCountUniqueAttributeValuesKernel.h
--- a/Engine/Plugins/PCG/Source/PCG/Private/Compute/BuiltInKernels/PCGCountUniqueAttributeValuesKernel.h
+++ b/Engine/Plugins/PCG/Source/PCG/Private/Compute/BuiltInKernels/PCGCountUniqueAttributeValuesKernel.h
@@ -31,7 +31,7 @@ public:
 	void SetOutputRawBuffer(bool bInOutputRawBuffer) { bOutputRawBuffer = bInOutputRawBuffer; }
 
 	virtual bool IsKernelDataValid(FPCGContext* InContext) const override;
-	virtual bool ComputeOutputBindingDataDesc(const UPCGComputeGraph* InGraph, FName InOutputPinLabel, UPCGDataBinding* InBinding, FPCGDataCollectionDesc& OutDataDesc) const override;
+	virtual TSharedPtr<const FPCGDataCollectionDesc> ComputeOutputBindingDataDesc(const UPCGComputeGraph* InGraph, FName InOutputPinLabel, UPCGDataBinding* InBinding) const override;
 	virtual int ComputeThreadCount(const UPCGDataBinding* Binding) const override;
 	virtual bool DoesOutputPinRequireZeroInitialization(FName InOutputPinLabel) const override;
 #if WITH_EDITOR
diff --git a/Engine/Plugins/PCG/Source/PCG/Private/Compute/BuiltInKernels/PCGDataNumKernel.cpp b/Engine/Plugins/PCG/Source/PCG/Private/Compute/BuiltInKernels/PCGDataNumKernel.cpp
--- a/Engine/Plugins/PCG/Source/PCG/Private/Compute/BuiltInKernels/PCGDataNumKernel.cpp
+++ b/Engine/Plugins/PCG/Source/PCG/Private/Compute/BuiltInKernels/PCGDataNumKernel.cpp
@@ -14,7 +14,7 @@
 
 #define LOCTEXT_NAMESPACE "PCGDataNumKernel"
 
-bool UPCGDataNumKernel::ComputeOutputBindingDataDesc(const UPCGComputeGraph* InGraph, const FName InOutputPinLabel, UPCGDataBinding* InBinding, FPCGDataCollectionDesc& OutDataDesc) const
+TSharedPtr<const FPCGDataCollectionDesc> UPCGDataNumKernel::ComputeOutputBindingDataDesc(const UPCGComputeGraph* InGraph, const FName InOutputPinLabel, UPCGDataBinding* InBinding) const
 {
 	check(InGraph);
 	check(InBinding);
@@ -22,13 +22,15 @@ bool UPCGDataNumKernel::ComputeOutputBindingDataDesc(const UPCGComputeGraph* InG
 	// Only single output pin
 	if (!ensure(InOutputPinLabel == PCGPinConstants::DefaultOutputLabel))
 	{
-		return false;
+		return nullptr;
 	}
 
-	FPCGDataDesc& OutputDataDesc = OutDataDesc.DataDescs.Emplace_GetRef(EPCGDataType::Param, 1);
+	TSharedPtr<FPCGDataCollectionDesc> OutDataDesc = FPCGDataCollectionDesc::MakeShared();
+
+	FPCGDataDesc& OutputDataDesc = OutDataDesc->DataDescs.Emplace_GetRef(EPCGDataType::Param, 1);
 	OutputDataDesc.AddAttribute(OutCountAttributeKey, InBinding);
 
-	return true;
+	return OutDataDesc;
 }
 
 int UPCGDataNumKernel::ComputeThreadCount(const UPCGDataBinding* InBinding) const
diff --git a/Engine/Plugins/PCG/Source/PCG/Private/Compute/BuiltInKernels/PCGDataNumKernel.h b/Engine/Plugins/PCG/Source/PCG/Private/Compute/BuiltInKernels/PCGDataNumKernel.h
--- a/Engine/Plugins/PCG/Source/PCG/Private/Compute/BuiltInKernels/PCGDataNumKernel.h
+++ b/Engine/Plugins/PCG/Source/PCG/Private/Compute/BuiltInKernels/PCGDataNumKernel.h
@@ -14,7 +14,7 @@ class UPCGDataNumKernel : public UPCGComputeKernel
 	GENERATED_BODY()
 
 public:
-	virtual bool ComputeOutputBindingDataDesc(const UPCGComputeGraph* InGraph, FName InOutputPinLabel, UPCGDataBinding* InBinding, FPCGDataCollectionDesc& OutDataDesc) const override;
+	virtual TSharedPtr<const FPCGDataCollectionDesc> ComputeOutputBindingDataDesc(const UPCGComputeGraph* InGraph, FName InOutputPinLabel, UPCGDataBinding* InBinding) const override;
 	virtual int ComputeThreadCount(const UPCGDataBinding* Binding) const override;
 
 #if WITH_EDITOR
diff --git a/Engine/Plugins/PCG/Source/PCG/Private/Compute/Data/PCGProxyForGPUData.cpp b/Engine/Plugins/PCG/Source/PCG/Private/Compute/Data/PCGProxyForGPUData.cpp
--- a/Engine/Plugins/PCG/Source/PCG/Private/Compute/Data/PCGProxyForGPUData.cpp
+++ b/Engine/Plugins/PCG/Source/PCG/Private/Compute/Data/PCGProxyForGPUData.cpp
@@ -40,7 +40,10 @@ FPCGDataTypeBaseId UPCGProxyForGPUData::GetUnderlyingDataTypeId() const
 {
 	if (TSharedPtr<const FPCGProxyForGPUDataCollection> DataCollectionGPU = GetGPUInfo())
 	{
-		return DataCollectionGPU->GetDescription().DataDescs[DataIndexInCollection].GetType().GetId();
+		if (const TSharedPtr<const FPCGDataCollectionDesc> DataCollectionDesc = DataCollectionGPU->GetDescription())
+		{
+			return DataCollectionDesc->DataDescs[DataIndexInCollection].GetType().GetId();
+		}
 	}
 
 	return FPCGDataTypeBaseId{};
@@ -55,7 +58,10 @@ FIntVector4 UPCGProxyForGPUData::GetElementCount() const
 {
 	if (TSharedPtr<const FPCGProxyForGPUDataCollection> DataCollectionGPU = GetGPUInfo())
 	{
-		return DataCollectionGPU->GetDescription().DataDescs[DataIndexInCollection].GetElementCount();
+		if (const TSharedPtr<const FPCGDataCollectionDesc> DataCollectionDesc = DataCollectionGPU->GetDescription())
+		{
+			return DataCollectionDesc->DataDescs[DataIndexInCollection].GetElementCount();
+		}
 	}
 
 	return FIntVector4::ZeroValue;
@@ -65,8 +71,11 @@ bool UPCGProxyForGPUData::GetDescription(FPCGDataDesc& OutDescription) const
 {
 	if (TSharedPtr<const FPCGProxyForGPUDataCollection> DataCollectionGPU = GetGPUInfo())
 	{
-		OutDescription = DataCollectionGPU->GetDescription().DataDescs[DataIndexInCollection];
-		return true;
+		if (const TSharedPtr<const FPCGDataCollectionDesc> DataCollectionDesc = DataCollectionGPU->GetDescription())
+		{
+			OutDescription = DataCollectionDesc->DataDescs[DataIndexInCollection];
+			return true;
+		}
 	}
 
 	return false;
@@ -98,14 +107,14 @@ TSharedPtr<const FPCGProxyForGPUDataCollection> UPCGProxyForGPUData::GetGPUInfo(
 
 TSharedPtr<FPCGProxyForGPUDataCollection> UPCGProxyForGPUData::GetGPUInfoMutable() const
 {
-	if (!DataCollectionOnGPU)
+	if (!DataCollectionOnGPU || !DataCollectionOnGPU->GetDescription())
 	{
 		ensureMsgf(false, TEXT("Data %s: GPU data collection lost. Enabling the define 'PCG_DATA_USAGE_LOGGING' may help to identify when resource was released."), *GetName());
 		return nullptr;
 	}
-	else if (!DataCollectionOnGPU->GetDescription().DataDescs.IsValidIndex(DataIndexInCollection))
+	else if (!DataCollectionOnGPU->GetDescription()->DataDescs.IsValidIndex(DataIndexInCollection))
 	{
-		ensureMsgf(false, TEXT("Data %s: DataIndexInCollection (%d) was out of range [0, %d)."), *GetName(), DataIndexInCollection, DataCollectionOnGPU->GetDescription().DataDescs.Num());
+		ensureMsgf(false, TEXT("Data %s: DataIndexInCollection (%d) was out of range [0, %d)."), *GetName(), DataIndexInCollection, DataCollectionOnGPU->GetDescription()->DataDescs.Num());
 		return nullptr;
 	}
 
diff --git a/Engine/Plugins/PCG/Source/PCG/Private/Compute/DataInterfaces/BuiltInKernels/PCGCountUniqueAttributeValuesDataInterface.cpp b/Engine/Plugins/PCG/Source/PCG/Private/Compute/DataInterfaces/BuiltInKernels/PCGCountUniqueAttributeValuesDataInterface.cpp
--- a/Engine/Plugins/PCG/Source/PCG/Private/Compute/DataInterfaces/BuiltInKernels/PCGCountUniqueAttributeValuesDataInterface.cpp
+++ b/Engine/Plugins/PCG/Source/PCG/Private/Compute/DataInterfaces/BuiltInKernels/PCGCountUniqueAttributeValuesDataInterface.cpp
@@ -122,7 +122,7 @@ bool UPCGCountUniqueAttributeValuesDataProvider::PrepareForExecute_GameThread(UP
 		check(OutputCountAttributeId != -1);
 	}
 
-	const FPCGDataCollectionDesc* InputDataDesc = InBinding->GetCachedKernelPinDataDesc(Kernel, PCGPinConstants::DefaultInputLabel, /*bIsInput*/true);
+	const TSharedPtr<const FPCGDataCollectionDesc> InputDataDesc = InBinding->GetCachedKernelPinDataDesc(Kernel, PCGPinConstants::DefaultInputLabel, /*bIsInput*/true);
 
 	if (!ensure(InputDataDesc))
 	{
diff --git a/Engine/Plugins/PCG/Source/PCG/Private/Compute/DataInterfaces/Elements/PCGCopyPointsDataInterface.cpp b/Engine/Plugins/PCG/Source/PCG/Private/Compute/DataInterfaces/Elements/PCGCopyPointsDataInterface.cpp
--- a/Engine/Plugins/PCG/Source/PCG/Private/Compute/DataInterfaces/Elements/PCGCopyPointsDataInterface.cpp
+++ b/Engine/Plugins/PCG/Source/PCG/Private/Compute/DataInterfaces/Elements/PCGCopyPointsDataInterface.cpp
@@ -174,8 +174,8 @@ bool UPCGCopyPointsDataProvider::PrepareForExecute_GameThread(UPCGDataBinding* I
 			return true;
 		}
 
-		const FPCGDataCollectionDesc* SourcePinDesc = InBinding->GetCachedKernelPinDataDesc(ProducerKernel, PCGCopyPointsConstants::SourcePointsLabel, /*bIsInput=*/true);
-		const FPCGDataCollectionDesc* TargetPinDesc = InBinding->GetCachedKernelPinDataDesc(ProducerKernel, PCGCopyPointsConstants::TargetPointsLabel, /*bIsInput=*/true);
+		const TSharedPtr<const FPCGDataCollectionDesc> SourcePinDesc = InBinding->GetCachedKernelPinDataDesc(ProducerKernel, PCGCopyPointsConstants::SourcePointsLabel, /*bIsInput=*/true);
+		const TSharedPtr<const FPCGDataCollectionDesc> TargetPinDesc = InBinding->GetCachedKernelPinDataDesc(ProducerKernel, PCGCopyPointsConstants::TargetPointsLabel, /*bIsInput=*/true);
 
 		if (!ensure(SourcePinDesc) || !ensure(TargetPinDesc))
 		{
diff --git a/Engine/Plugins/PCG/Source/PCG/Private/Compute/DataInterfaces/Elements/PCGMetadataPartitionDataInterface.cpp b/Engine/Plugins/PCG/Source/PCG/Private/Compute/DataInterfaces/Elements/PCGMetadataPartitionDataInterface.cpp
--- a/Engine/Plugins/PCG/Source/PCG/Private/Compute/DataInterfaces/Elements/PCGMetadataPartitionDataInterface.cpp
+++ b/Engine/Plugins/PCG/Source/PCG/Private/Compute/DataInterfaces/Elements/PCGMetadataPartitionDataInterface.cpp
@@ -149,7 +149,7 @@ bool UPCGMetaDataPartitionDataProvider::PrepareForExecute_GameThread(UPCGDataBin
 		return true;
 	}
 
-	const FPCGDataCollectionDesc* InputDesc = InBinding->GetCachedKernelPinDataDesc(ProducerKernel, PCGPinConstants::DefaultInputLabel, /*bIsInput=*/true);
+	const TSharedPtr<const FPCGDataCollectionDesc> InputDesc = InBinding->GetCachedKernelPinDataDesc(ProducerKernel, PCGPinConstants::DefaultInputLabel, /*bIsInput=*/true);
 	if (!InputDesc)
 	{
 		UE_LOG(LogPCG, Error, TEXT("UPCGMetaDataPartitionDataProvider: Could not retrieve input data description."));
@@ -221,7 +221,7 @@ bool UPCGMetaDataPartitionDataProvider::PrepareForExecute_GameThread(UPCGDataBin
 
 	if (ensure(Settings->PartitionAttributeSelectors.Num() == 1 && Settings->PartitionAttributeSelectors[0].IsBasicAttribute()))
 	{
-		const FPCGDataCollectionDesc* InputDataDesc = InBinding->GetCachedKernelPinDataDesc(ProducerKernel, PCGPinConstants::DefaultInputLabel, /*bIsInput*/true);
+		const TSharedPtr<const FPCGDataCollectionDesc> InputDataDesc = InBinding->GetCachedKernelPinDataDesc(ProducerKernel, PCGPinConstants::DefaultInputLabel, /*bIsInput*/true);
 
 		if (!ensure(InputDataDesc))
 		{
diff --git a/Engine/Plugins/PCG/Source/PCG/Private/Compute/DataInterfaces/Elements/PCGSkinnedMeshSpawnerDataInterface.cpp b/Engine/Plugins/PCG/Source/PCG/Private/Compute/DataInterfaces/Elements/PCGSkinnedMeshSpawnerDataInterface.cpp
--- a/Engine/Plugins/PCG/Source/PCG/Private/Compute/DataInterfaces/Elements/PCGSkinnedMeshSpawnerDataInterface.cpp
+++ b/Engine/Plugins/PCG/Source/PCG/Private/Compute/DataInterfaces/Elements/PCGSkinnedMeshSpawnerDataInterface.cpp
@@ -309,7 +309,7 @@ bool UPCGSkinnedMeshSpawnerDataProvider::PrepareForExecute_GameThread(UPCGDataBi
 	if (SelectorAttributeId == INDEX_NONE && ensure(Settings->MeshSelectorParameters))
 	{
 		const FName SelectorName = Settings->MeshSelectorParameters->MeshAttribute.GetName();
-		const FPCGDataCollectionDesc* InputDataDesc = InBinding->GetCachedKernelPinDataDesc(ProducerKernel, PCGPinConstants::DefaultInputLabel, /*bIsInput=*/true);
+		const TSharedPtr<const FPCGDataCollectionDesc> InputDataDesc = InBinding->GetCachedKernelPinDataDesc(ProducerKernel, PCGPinConstants::DefaultInputLabel, /*bIsInput=*/true);
 
 		if (!ensure(InputDataDesc))
 		{
@@ -463,7 +463,7 @@ void UPCGSkinnedMeshSpawnerDataProvider::CreatePrimitiveDescriptors(FPCGContext*
 		return;
 	}
 
-	const FPCGDataCollectionDesc* InputDataDesc = InBinding->GetCachedKernelPinDataDesc(ProducerKernel, PCGPinConstants::DefaultInputLabel, /*bIsInput=*/true);
+	const TSharedPtr<const FPCGDataCollectionDesc> InputDataDesc = InBinding->GetCachedKernelPinDataDesc(ProducerKernel, PCGPinConstants::DefaultInputLabel, /*bIsInput=*/true);
 
 	if (!ensure(InputDataDesc))
 	{
diff --git a/Engine/Plugins/PCG/Source/PCG/Private/Compute/DataInterfaces/Elements/PCGStaticMeshSpawnerDataInterface.cpp b/Engine/Plugins/PCG/Source/PCG/Private/Compute/DataInterfaces/Elements/PCGStaticMeshSpawnerDataInterface.cpp
--- a/Engine/Plugins/PCG/Source/PCG/Private/Compute/DataInterfaces/Elements/PCGStaticMeshSpawnerDataInterface.cpp
+++ b/Engine/Plugins/PCG/Source/PCG/Private/Compute/DataInterfaces/Elements/PCGStaticMeshSpawnerDataInterface.cpp
@@ -349,7 +349,7 @@ bool UPCGStaticMeshSpawnerDataProvider::PrepareForExecute_GameThread(UPCGDataBin
 	if (SelectorAttributeId == INDEX_NONE && SelectorByAttribute)
 	{
 		const FName SelectorName = SelectorByAttribute->AttributeName;
-		const FPCGDataCollectionDesc* InputDataDesc = InBinding->GetCachedKernelPinDataDesc(ProducerKernel, PCGPinConstants::DefaultInputLabel, /*bIsInput=*/true);
+		const TSharedPtr<const FPCGDataCollectionDesc> InputDataDesc = InBinding->GetCachedKernelPinDataDesc(ProducerKernel, PCGPinConstants::DefaultInputLabel, /*bIsInput=*/true);
 
 		if (!ensure(InputDataDesc))
 		{
@@ -513,7 +513,7 @@ void UPCGStaticMeshSpawnerDataProvider::CreatePrimitiveDescriptors(FPCGContext*
 		return;
 	}
 
-	const FPCGDataCollectionDesc* InputDataDesc = InBinding->GetCachedKernelPinDataDesc(ProducerKernel, PCGPinConstants::DefaultInputLabel, /*bIsInput=*/true);
+	const TSharedPtr<const FPCGDataCollectionDesc> InputDataDesc = InBinding->GetCachedKernelPinDataDesc(ProducerKernel, PCGPinConstants::DefaultInputLabel, /*bIsInput=*/true);
 
 	if (!ensure(InputDataDesc))
 	{
diff --git a/Engine/Plugins/PCG/Source/PCG/Private/Compute/DataInterfaces/Elements/PCGTransformPointsDataInterface.cpp b/Engine/Plugins/PCG/Source/PCG/Private/Compute/DataInterfaces/Elements/PCGTransformPointsDataInterface.cpp
--- a/Engine/Plugins/PCG/Source/PCG/Private/Compute/DataInterfaces/Elements/PCGTransformPointsDataInterface.cpp
+++ b/Engine/Plugins/PCG/Source/PCG/Private/Compute/DataInterfaces/Elements/PCGTransformPointsDataInterface.cpp
@@ -189,7 +189,7 @@ bool UPCGTransformPointsDataProvider::PrepareForExecute_GameThread(UPCGDataBindi
 
 	if (Settings->bApplyToAttribute && AttributeId == INDEX_NONE)
 	{
-		const FPCGDataCollectionDesc* InputDataDesc = InBinding->GetCachedKernelPinDataDesc(ProducerKernel, PCGPinConstants::DefaultInputLabel, /*bIsInput=*/true);
+		const TSharedPtr<const FPCGDataCollectionDesc> InputDataDesc = InBinding->GetCachedKernelPinDataDesc(ProducerKernel, PCGPinConstants::DefaultInputLabel, /*bIsInput=*/true);
 
 		if (!ensure(InputDataDesc))
 		{
diff --git a/Engine/Plugins/PCG/Source/PCG/Private/Compute/DataInterfaces/PCGDataCollectionDataInterface.cpp b/Engine/Plugins/PCG/Source/PCG/Private/Compute/DataInterfaces/PCGDataCollectionDataInterface.cpp
--- a/Engine/Plugins/PCG/Source/PCG/Private/Compute/DataInterfaces/PCGDataCollectionDataInterface.cpp
+++ b/Engine/Plugins/PCG/Source/PCG/Private/Compute/DataInterfaces/PCGDataCollectionDataInterface.cpp
@@ -664,23 +664,20 @@ FPCGDataCollectionDataProviderProxy::FPCGDataCollectionDataProviderProxy(const F
 {
 	OriginatingGenerationCount = DataProviderWeakPtr->GenerationCounter;
 
-	if (InParams.PinDesc)
+	if (ensure(InParams.PinDesc))
 	{
-		TRACE_CPUPROFILER_EVENT_SCOPE(CopyDataDescription);
+		PinDesc = InParams.PinDesc;
 
-		// @todo_pcg: Avoid costly copy of data description.
-		PinDesc = *InParams.PinDesc;
-
-		for (const FPCGDataDesc& DataDesc : PinDesc.DataDescs)
+		for (const FPCGDataDesc& DataDesc : PinDesc->DataDescs)
 		{
 			for (const FPCGKernelAttributeDesc& AttrDesc : DataDesc.AttributeDescs)
 			{
 				LargestAttributeIndex = FMath::Max(LargestAttributeIndex, PCGComputeHelpers::GetMetadataAttributeIndexFromAttributeId(AttrDesc.AttributeId));
 			}
 		}
-	}
 
-	SizeBytes = PinDesc.ComputePackedSizeBytes();
+		SizeBytes = PinDesc->ComputePackedSizeBytes();
+	}
 }
 
 bool FPCGDataCollectionDataProviderProxy::IsValid(FValidationData const& InValidationData) const
@@ -750,12 +747,8 @@ void FPCGDataCollectionDataProviderProxy::AllocateResources(FRDGBuilder& GraphBu
 		// Initialize with an empty data collection. The kernel may not run, for example if indirect dispatch args end up being 0. This ensures
 		// there is something meaningful to readback.
 		TArray<uint32> PackedDataCollection;
-		PackedDataCollection.Reserve(bZeroInitialize ? (SizeBytes / PackedDataCollection.GetTypeSize()) : PinDesc.ComputePackedHeaderSizeBytes());
-
-		{
-			TRACE_CPUPROFILER_EVENT_SCOPE(FPCGDataCollectionDataProviderProxy::AllocateResources::WriteHeader);
-			PinDesc.WriteHeader(PackedDataCollection);
-		}
+		PackedDataCollection.Reserve(bZeroInitialize ? (SizeBytes / PackedDataCollection.GetTypeSize()) : PinDesc->ComputePackedHeaderSizeBytes());
+		PinDesc->WriteHeader(PackedDataCollection);
 
 		if (bZeroInitialize)
 		{
@@ -804,6 +797,11 @@ void FPCGDataCollectionDataProviderProxy::AllocateResources(FRDGBuilder& GraphBu
 					return;
 				}
 
+				if (!ensure(PinDesc))
+				{
+					return;
+				}
+
 				UPCGDataBinding* Binding = DataProvider ? DataProvider->GetDataBinding() : nullptr;
 				UPCGSubsystem* Subsystem = UPCGSubsystem::GetSubsystemForCurrentWorld();
 
@@ -811,7 +809,7 @@ void FPCGDataCollectionDataProviderProxy::AllocateResources(FRDGBuilder& GraphBu
 				{
 					TSharedPtr<FPCGProxyForGPUDataCollection> DataCollectionOnGPU = MakeShared<FPCGProxyForGPUDataCollection>(ExportedPooledBuffer, SizeBytes, PinDesc, Binding->GetStringTable());
 
-					for (int DataIndex = 0; DataIndex < PinDesc.DataDescs.Num(); ++DataIndex)
+					for (int DataIndex = 0; DataIndex < PinDesc->DataDescs.Num(); ++DataIndex)
 					{
 						UPCGProxyForGPUData* Proxy = NewObject<UPCGProxyForGPUData>();
 						Proxy->Initialize(DataCollectionOnGPU, DataIndex);
diff --git a/Engine/Plugins/PCG/Source/PCG/Private/Compute/DataInterfaces/PCGDataCollectionDataInterface.h b/Engine/Plugins/PCG/Source/PCG/Private/Compute/DataInterfaces/PCGDataCollectionDataInterface.h
--- a/Engine/Plugins/PCG/Source/PCG/Private/Compute/DataInterfaces/PCGDataCollectionDataInterface.h
+++ b/Engine/Plugins/PCG/Source/PCG/Private/Compute/DataInterfaces/PCGDataCollectionDataInterface.h
@@ -99,7 +99,7 @@ public:
 	{
 		// Weak pointer useful for passing back buffer handles. Do not access this directly from the render thread.
 		TWeakObjectPtr<UPCGDataCollectionDataProvider> DataProvider;
-		const FPCGDataCollectionDesc* PinDesc = nullptr;
+		TSharedPtr<const FPCGDataCollectionDesc> PinDesc = nullptr;
 		EPCGExportMode ExportMode = EPCGExportMode::NoExport;
 		bool bZeroInitialize = false;
 		uint32 ElementCountMultiplier = 0;
@@ -137,7 +137,7 @@ protected:
 
 	using FParameters = FPCGDataCollectionDataInterfaceParameters;
 
-	FPCGDataCollectionDesc PinDesc;
+	TSharedPtr<const FPCGDataCollectionDesc> PinDesc = nullptr;
 
 	FRDGBufferRef Buffer = nullptr;
 	FRDGBufferUAVRef BufferUAV = nullptr;
diff --git a/Engine/Plugins/PCG/Source/PCG/Private/Compute/DataInterfaces/PCGDataLabelResolverDataInterface.cpp b/Engine/Plugins/PCG/Source/PCG/Private/Compute/DataInterfaces/PCGDataLabelResolverDataInterface.cpp
--- a/Engine/Plugins/PCG/Source/PCG/Private/Compute/DataInterfaces/PCGDataLabelResolverDataInterface.cpp
+++ b/Engine/Plugins/PCG/Source/PCG/Private/Compute/DataInterfaces/PCGDataLabelResolverDataInterface.cpp
@@ -97,7 +97,7 @@ bool UPCGDataLabelResolverDataProvider::PrepareForExecute_GameThread(UPCGDataBin
 
 	const TArray<FString>& StringTable = InBinding->GetStringTable();
 
-	const FPCGDataCollectionDesc* PinDataDesc = InBinding->GetCachedKernelPinDataDesc(Kernel, PinLabel, bIsInput);
+	const TSharedPtr<const FPCGDataCollectionDesc> PinDataDesc = InBinding->GetCachedKernelPinDataDesc(Kernel, PinLabel, bIsInput);
 
 	if (!ensure(PinDataDesc))
 	{
diff --git a/Engine/Plugins/PCG/Source/PCG/Private/Compute/DataInterfaces/PCGRawBufferDataInterface.cpp b/Engine/Plugins/PCG/Source/PCG/Private/Compute/DataInterfaces/PCGRawBufferDataInterface.cpp
--- a/Engine/Plugins/PCG/Source/PCG/Private/Compute/DataInterfaces/PCGRawBufferDataInterface.cpp
+++ b/Engine/Plugins/PCG/Source/PCG/Private/Compute/DataInterfaces/PCGRawBufferDataInterface.cpp
@@ -117,7 +117,7 @@ bool UPCGRawBufferDataProvider::PrepareForExecute_GameThread(UPCGDataBinding* In
 		return false;
 	}
 
-	if (const FPCGDataCollectionDesc* ProducerPinDesc = InBinding->GetCachedKernelPinDataDesc(GraphBindingIndex); ensure(ProducerPinDesc))
+	if (const TSharedPtr<const FPCGDataCollectionDesc> ProducerPinDesc = InBinding->GetCachedKernelPinDataDesc(GraphBindingIndex); ensure(ProducerPinDesc))
 	{
 		const uint64 PackedSizeBytes = ProducerPinDesc->ComputePackedSizeBytes();
 
diff --git a/Engine/Plugins/PCG/Source/PCG/Private/Compute/DataInterfaces/PCGTextureDataInterface.cpp b/Engine/Plugins/PCG/Source/PCG/Private/Compute/DataInterfaces/PCGTextureDataInterface.cpp
--- a/Engine/Plugins/PCG/Source/PCG/Private/Compute/DataInterfaces/PCGTextureDataInterface.cpp
+++ b/Engine/Plugins/PCG/Source/PCG/Private/Compute/DataInterfaces/PCGTextureDataInterface.cpp
@@ -395,14 +395,7 @@ FPCGTextureDataProviderProxy::FPCGTextureDataProviderProxy(TWeakObjectPtr<UPCGTe
 	OutputPinLabelAlias = InDataProvider->OutputPinLabelAlias;
 	OriginatingGenerationCount = InDataProvider->GenerationCounter;
 	DataProviderWeakPtr_GT = InDataProvider;
-
-	if (InDataProvider->GetPinDescription())
-	{
-		TRACE_CPUPROFILER_EVENT_SCOPE(CopyDataDescription);
-
-		// @todo_pcg: Avoid costly copy of data description.
-		PinDesc = *InDataProvider->GetPinDescription();
-	}
+	PinDesc = InDataProvider->GetPinDescription();
 }
 
 bool FPCGTextureDataProviderProxy::IsValid(FValidationData const& InValidationData) const
@@ -625,6 +618,11 @@ void FPCGTextureDataProviderProxy::ExportTextureUAVs(const TArray<TRefCountPtr<I
 			return;
 		}
 
+		if (!ensure(PinDesc))
+		{
+			return;
+		}
+
 		if (UPCGDataBinding* Binding = DataProvider->GetDataBinding())
 		{
 			const TArray<FPCGTextureInfo>& TextureInfos_GT = DataProvider->GetTextureInfos();
@@ -635,7 +633,7 @@ void FPCGTextureDataProviderProxy::ExportTextureUAVs(const TArray<TRefCountPtr<I
 			{
 				const FPCGTextureInfo& TextureInfo = TextureInfos_GT[TextureDataIndex];
 				check(ExportedTextures.IsValidIndex(TextureDataIndex));
-				check(PinDesc.DataDescs.IsValidIndex(TextureDataIndex))
+				check(PinDesc->DataDescs.IsValidIndex(TextureDataIndex))
 				check(BindingInfos_GT.IsValidIndex(TextureInfo.BindingIndex));
 
 				const FPCGTextureBindingInfo& BindingInfo = BindingInfos_GT[TextureInfo.BindingIndex];
@@ -649,7 +647,7 @@ void FPCGTextureDataProviderProxy::ExportTextureUAVs(const TArray<TRefCountPtr<I
 				ExportedData->Initialize(ExportedTextures[TextureDataIndex], TextureInfo.SliceIndex, TextureTransform, /*bInSkipReadbackToCPU=*/true);
 
 				TSet<FString> Tags;
-				for (const int32 TagStringKey : PinDesc.DataDescs[TextureDataIndex].TagStringKeys)
+				for (const int32 TagStringKey : PinDesc->DataDescs[TextureDataIndex].TagStringKeys)
 				{
 					if (ensure(StringTable.IsValidIndex(TagStringKey)))
 					{
diff --git a/Engine/Plugins/PCG/Source/PCG/Private/Compute/DataInterfaces/PCGTextureDataInterface.h b/Engine/Plugins/PCG/Source/PCG/Private/Compute/DataInterfaces/PCGTextureDataInterface.h
--- a/Engine/Plugins/PCG/Source/PCG/Private/Compute/DataInterfaces/PCGTextureDataInterface.h
+++ b/Engine/Plugins/PCG/Source/PCG/Private/Compute/DataInterfaces/PCGTextureDataInterface.h
@@ -136,7 +136,7 @@ protected:
 
 	// For export
 	EPCGExportMode ExportMode = EPCGExportMode::NoExport;
-	FPCGDataCollectionDesc PinDesc;
+	TSharedPtr<const FPCGDataCollectionDesc> PinDesc = nullptr;
 	FName OutputPinLabel;
 	FName OutputPinLabelAlias;
 
diff --git a/Engine/Plugins/PCG/Source/PCG/Private/Compute/Elements/PCGCullPointsOutsideActorBoundsKernel.cpp b/Engine/Plugins/PCG/Source/PCG/Private/Compute/Elements/PCGCullPointsOutsideActorBoundsKernel.cpp
--- a/Engine/Plugins/PCG/Source/PCG/Private/Compute/Elements/PCGCullPointsOutsideActorBoundsKernel.cpp
+++ b/Engine/Plugins/PCG/Source/PCG/Private/Compute/Elements/PCGCullPointsOutsideActorBoundsKernel.cpp
@@ -14,26 +14,23 @@
 
 #define LOCTEXT_NAMESPACE "PCGCullPointsOutsideActorBoundsKernel"
 
-bool UPCGCullPointsOutsideActorBoundsKernel::ComputeOutputBindingDataDesc(const UPCGComputeGraph* InGraph, FName InOutputPinLabel, UPCGDataBinding* InBinding, FPCGDataCollectionDesc& OutDataDesc) const
+TSharedPtr<const FPCGDataCollectionDesc> UPCGCullPointsOutsideActorBoundsKernel::ComputeOutputBindingDataDesc(const UPCGComputeGraph* InGraph, FName InOutputPinLabel, UPCGDataBinding* InBinding) const
 {
 	check(InBinding);
 
 	// Code assumes single output pin.
 	if (!ensure(InOutputPinLabel == PCGPinConstants::DefaultOutputLabel))
 	{
-		return false;
+		return nullptr;
 	}
 
 	const FPCGKernelPin SourceKernelPin(KernelIndex, PCGPinConstants::DefaultInputLabel, /*bIsInput=*/true);
-
-	ensure(InBinding->ComputeKernelPinDataDesc(SourceKernelPin, &OutDataDesc));
-
-	return true;
+	return InBinding->ComputeKernelPinDataDesc(SourceKernelPin);
 }
 
 int UPCGCullPointsOutsideActorBoundsKernel::ComputeThreadCount(const UPCGDataBinding* InBinding) const
 {
-	const FPCGDataCollectionDesc* OutputPinDesc = InBinding->GetCachedKernelPinDataDesc(this, PCGPinConstants::DefaultOutputLabel, /*bIsInput=*/false);
+	const TSharedPtr<const FPCGDataCollectionDesc> OutputPinDesc = InBinding->GetCachedKernelPinDataDesc(this, PCGPinConstants::DefaultOutputLabel, /*bIsInput=*/false);
 	return ensure(OutputPinDesc) ? OutputPinDesc->ComputeTotalElementCount() : 0;
 }
 
diff --git a/Engine/Plugins/PCG/Source/PCG/Private/Compute/Elements/PCGCullPointsOutsideActorBoundsKernel.h b/Engine/Plugins/PCG/Source/PCG/Private/Compute/Elements/PCGCullPointsOutsideActorBoundsKernel.h
--- a/Engine/Plugins/PCG/Source/PCG/Private/Compute/Elements/PCGCullPointsOutsideActorBoundsKernel.h
+++ b/Engine/Plugins/PCG/Source/PCG/Private/Compute/Elements/PCGCullPointsOutsideActorBoundsKernel.h
@@ -15,7 +15,7 @@ class UPCGCullPointsOutsideActorBoundsKernel : public UPCGComputeKernel
 
 public:
 	//~Begin UPCGComputeKernel interface
-	virtual bool ComputeOutputBindingDataDesc(const UPCGComputeGraph* InGraph, FName InOutputPinLabel, UPCGDataBinding* InBinding, FPCGDataCollectionDesc& OutDataDesc) const override;
+	virtual TSharedPtr<const FPCGDataCollectionDesc> ComputeOutputBindingDataDesc(const UPCGComputeGraph* InGraph, FName InOutputPinLabel, UPCGDataBinding* InBinding) const override;
 	virtual int ComputeThreadCount(const UPCGDataBinding* Binding) const override;
 #if WITH_EDITOR
 	virtual const TCHAR* GetSourceFilePath() const override { return TEXT("/Plugin/PCG/Private/Elements/PCGCullPointsOutsideActorBounds.usf"); }
diff --git a/Engine/Plugins/PCG/Source/PCG/Private/Compute/Elements/PCGCustomHLSLKernel.cpp b/Engine/Plugins/PCG/Source/PCG/Private/Compute/Elements/PCGCustomHLSLKernel.cpp
--- a/Engine/Plugins/PCG/Source/PCG/Private/Compute/Elements/PCGCustomHLSLKernel.cpp
+++ b/Engine/Plugins/PCG/Source/PCG/Private/Compute/Elements/PCGCustomHLSLKernel.cpp
@@ -216,45 +216,44 @@ bool UPCGCustomHLSLKernel::IsKernelDataValid(FPCGContext* InContext) const
 	return true;
 }
 
-bool UPCGCustomHLSLKernel::ComputeOutputBindingDataDesc(const UPCGComputeGraph* InGraph, FName InOutputPinLabel, UPCGDataBinding* InBinding, FPCGDataCollectionDesc& OutDataDesc) const
+TSharedPtr<const FPCGDataCollectionDesc> UPCGCustomHLSLKernel::ComputeOutputBindingDataDesc(const UPCGComputeGraph* InGraph, FName InOutputPinLabel, UPCGDataBinding* InBinding) const
 {
 	check(InGraph);
 	check(InBinding);
 	
 	const UPCGCustomHLSLSettings* CustomHLSLSettings = CastChecked<UPCGCustomHLSLSettings>(GetSettings());
 
-	FPCGDataCollectionDesc OutputDataDesc;
-
-	const FPCGPinPropertiesGPU* OutputPinProperties = CustomHLSLSettings->OutputPins.FindByPredicate([InOutputPinLabel](const FPCGPinPropertiesGPU& InProps) {return InProps.Label == InOutputPinLabel; });
+	const FPCGPinPropertiesGPU* OutputPinProperties = CustomHLSLSettings->OutputPins.FindByPredicate([InOutputPinLabel](const FPCGPinPropertiesGPU& InProps) { return InProps.Label == InOutputPinLabel; });
 	if (!OutputPinProperties)
 	{
-		OutDataDesc = {};
-		return true;
+		return nullptr;
 	}
 
+	TSharedPtr<FPCGDataCollectionDesc> OutDataDesc = nullptr;
 	const FName FirstOutputPinLabel = CustomHLSLSettings->OutputPins[0].Label;
 
 	// The primary output pin follows any rules prescribed by kernel type.
 	if (InOutputPinLabel == FirstOutputPinLabel && CustomHLSLSettings->IsProcessorKernel())
 	{
-		FPCGDataCollectionDesc InputDesc;
-
 		if (const FPCGPinProperties* FirstInputPinProps = GetFirstInputPin())
 		{
 			const FPCGKernelPin FirstKernelPin(KernelIndex, FirstInputPinProps->Label, /*bIsInput=*/true);
+			const TSharedPtr<const FPCGDataCollectionDesc> InputDataDesc = InBinding->ComputeKernelPinDataDesc(FirstKernelPin);
 
-			if (!InBinding->ComputeKernelPinDataDesc(FirstKernelPin, &InputDesc))
+			if (ensure(InputDataDesc))
+			{
+				OutDataDesc = FPCGDataCollectionDesc::MakeSharedFrom(InputDataDesc);
+			}
+			else
 			{
 				UE_LOG(LogPCG, Warning, TEXT("Kernel pin not registered in compute graph. May be due to unsupported pin type. KernelIndex=%d, PinLabel='%s', Input=%d"),
 					FirstKernelPin.KernelIndex,
 					*FirstKernelPin.PinLabel.ToString(),
 					FirstKernelPin.bIsInput);
 
-				return false;
+				return nullptr;
 			}
 		}
-
-		OutputDataDesc = MoveTemp(InputDesc);
 	}
 	else if (InOutputPinLabel == FirstOutputPinLabel && CustomHLSLSettings->KernelType == EPCGKernelType::PointGenerator)
 	{
@@ -262,26 +261,37 @@ bool UPCGCustomHLSLKernel::ComputeOutputBindingDataDesc(const UPCGComputeGraph*
 
 		if (!ensure(KernelParams))
 		{
-			return false;
+			return nullptr;
 		}
 
+		OutDataDesc = FPCGDataCollectionDesc::MakeShared();
+
 		// Generators always produce a single point data with known point count.
-		OutputDataDesc.DataDescs.Emplace(EPCGDataType::Point, KernelParams->GetValueInt(GET_MEMBER_NAME_CHECKED(UPCGCustomHLSLSettings, NumElements)));
+		OutDataDesc->DataDescs.Emplace(EPCGDataType::Point, KernelParams->GetValueInt(GET_MEMBER_NAME_CHECKED(UPCGCustomHLSLSettings, NumElements)));
 	}
 	else if (InOutputPinLabel == FirstOutputPinLabel && CustomHLSLSettings->KernelType == EPCGKernelType::TextureGenerator)
 	{
+		OutDataDesc = FPCGDataCollectionDesc::MakeShared();
+
 		// Generators always produce a single texture data with known size.
-		OutputDataDesc.DataDescs.Emplace(EPCGDataType::BaseTexture, FIntVector4(CustomHLSLSettings->NumElements2D.X, CustomHLSLSettings->NumElements2D.Y, 0, 0));
+		OutDataDesc->DataDescs.Emplace(EPCGDataType::BaseTexture, FIntVector4(CustomHLSLSettings->NumElements2D.X, CustomHLSLSettings->NumElements2D.Y, 0, 0));
 	}
 	else if (ensure(OutputPinProperties))
 	{
-		ComputeDataDescFromPinProperties(*OutputPinProperties, MakeArrayView(CustomHLSLSettings->InputPins), InBinding, OutputDataDesc);
+		OutDataDesc = FPCGDataCollectionDesc::MakeShared();
+
+		ComputeDataDescFromPinProperties(*OutputPinProperties, MakeArrayView(CustomHLSLSettings->InputPins), InBinding, OutDataDesc);
+	}
+
+	if (!ensure(OutDataDesc))
+	{
+		return nullptr;
 	}
 
 	// Add attributes that will be created for this pin on the GPU. This will stomp any existing attributes if they collide!
 	for (const FPCGKernelAttributeKey& CreatedKey : OutputPinProperties->PropertiesGPU.CreatedKernelAttributeKeys)
 	{
-		OutputDataDesc.AddAttributeToAllData(CreatedKey, InBinding);
+		OutDataDesc->AddAttributeToAllData(CreatedKey, InBinding);
 	}
 
 	// Try to propagate string keys across node. Not trivial because there could be one or more string key attributes on input pins and on output pins,
@@ -289,12 +299,12 @@ bool UPCGCustomHLSLKernel::ComputeOutputBindingDataDesc(const UPCGComputeGraph*
 	// from matching attribute names (across all input pins), and then fall back to collecting keys from all string key attributes across all inputs.
 	bool bOutputHasStringKeys = false;
 
-	for (const FPCGDataDesc& DataDesc : OutputDataDesc.DataDescs)
+	for (const FPCGDataDesc& DataDesc : OutDataDesc->DataDescs)
 	{
 		auto HasStringKey = [](const FPCGKernelAttributeDesc& InAttributeDesc)
-			{
-				return InAttributeDesc.AttributeKey.Type == EPCGKernelAttributeType::StringKey || InAttributeDesc.AttributeKey.Type == EPCGKernelAttributeType::Name;
-			};
+		{
+			return InAttributeDesc.AttributeKey.Type == EPCGKernelAttributeType::StringKey || InAttributeDesc.AttributeKey.Type == EPCGKernelAttributeType::Name;
+		};
 
 		if (DataDesc.AttributeDescs.FindByPredicate(HasStringKey))
 		{
@@ -305,18 +315,22 @@ bool UPCGCustomHLSLKernel::ComputeOutputBindingDataDesc(const UPCGComputeGraph*
 
 	if (bOutputHasStringKeys)
 	{
-		TArray<FPCGDataCollectionDesc> RelevantInputDataDescs;
+		TArray<const TSharedPtr<const FPCGDataCollectionDesc>> RelevantInputDataDescs;
 
 		// Collect descriptions of input data items that have string key attributes.
 		for(const FPCGPinProperties& PinProps : CustomHLSLSettings->InputPins)
 		{
-			FPCGDataCollectionDesc InputPinDesc;
 			const FPCGKernelPin InputKernelPin(KernelIndex, PinProps.Label, /*bIsInput=*/true);
-			ensure(InBinding->ComputeKernelPinDataDesc(InputKernelPin, &InputPinDesc));
+			const TSharedPtr<const FPCGDataCollectionDesc> InputPinDesc = InBinding->ComputeKernelPinDataDesc(InputKernelPin);
+
+			if (!ensure(InputPinDesc))
+			{
+				continue;
+			}
 
 			bool bFoundStringKeyAttribute = false;
 
-			for (FPCGDataDesc& DataDesc : OutputDataDesc.DataDescs)
+			for (FPCGDataDesc& DataDesc : OutDataDesc->DataDescs)
 			{
 				for (FPCGKernelAttributeDesc& AttributeDesc : DataDesc.AttributeDescs)
 				{
@@ -332,7 +346,7 @@ bool UPCGCustomHLSLKernel::ComputeOutputBindingDataDesc(const UPCGComputeGraph*
 
 		if (!RelevantInputDataDescs.IsEmpty())
 		{
-			for (FPCGDataDesc& DataDesc : OutputDataDesc.DataDescs)
+			for (FPCGDataDesc& DataDesc : OutDataDesc->DataDescs)
 			{
 				for (FPCGKernelAttributeDesc& AttributeDesc : DataDesc.AttributeDescs)
 				{
@@ -343,11 +357,11 @@ bool UPCGCustomHLSLKernel::ComputeOutputBindingDataDesc(const UPCGComputeGraph*
 
 					bool bFoundMatchingAttribute = false;
 
-					for (const FPCGDataCollectionDesc& InputPinDataDesc : RelevantInputDataDescs)
+					for (const TSharedPtr<const FPCGDataCollectionDesc>& InputPinDataDesc : RelevantInputDataDescs)
 					{
 						// Try to find string keys for matching attributes on inputs. E.g. if we are processing an output attribute named 'MeshPath',
 						// look at data on all input pins for an attribute named MeshPath and assume we could use any of its values - copy the string keys.
-						for (const FPCGDataDesc& InputDataDesc : InputPinDataDesc.DataDescs)
+						for (const FPCGDataDesc& InputDataDesc : InputPinDataDesc->DataDescs)
 						{
 							for (const FPCGKernelAttributeDesc& InputAttributeDesc : InputDataDesc.AttributeDescs)
 							{
@@ -366,11 +380,11 @@ bool UPCGCustomHLSLKernel::ComputeOutputBindingDataDesc(const UPCGComputeGraph*
 					{
 						// We didn't find an exact attribute. Fall back to finding any and all string keys. This is concerning and perhaps we can
 						// have additional hinting mechanisms in the kernel source or in the node UI.
-						for (const FPCGDataCollectionDesc& InputPinDataDesc : RelevantInputDataDescs)
+						for (const TSharedPtr<const FPCGDataCollectionDesc>& InputPinDataDesc : RelevantInputDataDescs)
 						{
 							// Try to find string keys for matching attributes on inputs. E.g. if we are processing an output attribute named 'MeshPath',
 							// look at data on all input pins for an attribute named MeshPath and assume we could use any of its values - copy the string keys.
-							for (const FPCGDataDesc& InputDataDesc : InputPinDataDesc.DataDescs)
+							for (const FPCGDataDesc& InputDataDesc : InputPinDataDesc->DataDescs)
 							{
 								for (const FPCGKernelAttributeDesc& InputAttributeDesc : InputDataDesc.AttributeDescs)
 								{
@@ -396,11 +410,10 @@ bool UPCGCustomHLSLKernel::ComputeOutputBindingDataDesc(const UPCGComputeGraph*
 	// Allocate any properties that were modified by the Custom HLSL source. If this data was initialized from another, it will have inherited its allocated properties already too.
 	if (const int32* AllocatedProperties = PinAllocatedProperties.Find(OutputPinProperties->Label))
 	{
-		OutputDataDesc.AllocatePropertiesForAllData(static_cast<EPCGPointNativeProperties>(*AllocatedProperties));
+		OutDataDesc->AllocatePropertiesForAllData(static_cast<EPCGPointNativeProperties>(*AllocatedProperties));
 	}
 
-	OutDataDesc = MoveTemp(OutputDataDesc);
-	return true;
+	return OutDataDesc;
 }
 
 #if WITH_EDITOR
@@ -641,7 +654,7 @@ int UPCGCustomHLSLKernel::ComputeThreadCount(const UPCGDataBinding* InBinding) c
 		{
 			if (const FPCGPinPropertiesGPU* OutputPin = CustomHLSLSettings->OutputPins.IsEmpty() ? nullptr : &CustomHLSLSettings->OutputPins[0])
 			{
-				const FPCGDataCollectionDesc* Desc = InBinding->GetCachedKernelPinDataDesc(this, OutputPin->Label, /*bIsInput=*/false);
+				const TSharedPtr<const FPCGDataCollectionDesc> Desc = InBinding->GetCachedKernelPinDataDesc(this, OutputPin->Label, /*bIsInput=*/false);
 				ThreadCount = Desc ? OutputPin->GetElementCountMultiplier() * Desc->ComputeTotalElementCount(OutputPin->AllowedTypes) : 0;
 			}
 		}
@@ -1741,15 +1754,15 @@ bool UPCGCustomHLSLKernel::AreAttributesValid(FPCGContext* InContext, FText* Out
 
 	if (DataBinding)
 	{
-		TMap<FName, const FPCGDataCollectionDesc*> InputPinDescs;
-		TMap<FName, const FPCGDataCollectionDesc*> OutputPinDescs;
+		TMap<FName, const TSharedPtr<const FPCGDataCollectionDesc>> InputPinDescs;
+		TMap<FName, const TSharedPtr<const FPCGDataCollectionDesc>> OutputPinDescs;
 
 		{
 			TRACE_CPUPROFILER_EVENT_SCOPE(UPCGCustomHLSLKernel::AreAttributesValid::GetCachedPinDataDescs);
 
 			for (const FPCGPinProperties& Pin : CustomHLSLSettings->InputPins)
 			{
-				if (const FPCGDataCollectionDesc* DataDesc = DataBinding->GetCachedKernelPinDataDesc(this, Pin.Label, /*bIsInputPin=*/true))
+				if (const TSharedPtr<const FPCGDataCollectionDesc> DataDesc = DataBinding->GetCachedKernelPinDataDesc(this, Pin.Label, /*bIsInputPin=*/true))
 				{
 					InputPinDescs.Add(Pin.Label, DataDesc);
 				}
@@ -1757,7 +1770,7 @@ bool UPCGCustomHLSLKernel::AreAttributesValid(FPCGContext* InContext, FText* Out
 
 			for (const FPCGPinProperties& Pin : CustomHLSLSettings->OutputPins)
 			{
-				if (const FPCGDataCollectionDesc* DataDesc = DataBinding->GetCachedKernelPinDataDesc(this, Pin.Label, /*bIsInputPin=*/false))
+				if (const TSharedPtr<const FPCGDataCollectionDesc> DataDesc = DataBinding->GetCachedKernelPinDataDesc(this, Pin.Label, /*bIsInputPin=*/false))
 				{
 					OutputPinDescs.Add(Pin.Label, DataDesc);
 				}
@@ -1777,7 +1790,7 @@ bool UPCGCustomHLSLKernel::AreAttributesValid(FPCGContext* InContext, FText* Out
 				const FString TypeStr = AttributeTypeEnum->GetNameStringByValue(ParsedFunction.AttributeType);
 
 				const FName PinLabel = FName(*PinLabelStr);
-				const FPCGDataCollectionDesc* PinDesc = nullptr;
+				TSharedPtr<const FPCGDataCollectionDesc> PinDesc = nullptr;
 
 				auto ConstructFunctionText = [&ParsedFunction, &TypeStr]()
 				{
@@ -1786,7 +1799,7 @@ bool UPCGCustomHLSLKernel::AreAttributesValid(FPCGContext* InContext, FText* Out
 
 				if (FunctionName == PCGCustomHLSLKernel::AttributeFunctionSetKeyword)
 				{
-					if (const FPCGDataCollectionDesc** PinDescPtr = OutputPinDescs.Find(PinLabel))
+					if (const TSharedPtr<const FPCGDataCollectionDesc>* PinDescPtr = OutputPinDescs.Find(PinLabel))
 					{
 						PinDesc = *PinDescPtr;
 					}
@@ -1808,7 +1821,7 @@ bool UPCGCustomHLSLKernel::AreAttributesValid(FPCGContext* InContext, FText* Out
 				}
 				else if (ensure(FunctionName == PCGCustomHLSLKernel::AttributeFunctionGetKeyword))
 				{
-					if (const FPCGDataCollectionDesc** PinDescPtr = InputPinDescs.Find(PinLabel))
+					if (const TSharedPtr<const FPCGDataCollectionDesc>* PinDescPtr = InputPinDescs.Find(PinLabel))
 					{
 						PinDesc = *PinDescPtr;
 					}
@@ -1927,7 +1940,7 @@ int UPCGCustomHLSLKernel::GetElementCountForInputPin(const FPCGPinProperties& In
 	check(InBinding);
 	check(InBinding->Graph);
 
-	const FPCGDataCollectionDesc* InputDesc = InBinding->GetCachedKernelPinDataDesc(this, InInputPinProps.Label, /*bIsInputPin=*/true);
+	const TSharedPtr<const FPCGDataCollectionDesc> InputDesc = InBinding->GetCachedKernelPinDataDesc(this, InInputPinProps.Label, /*bIsInputPin=*/true);
 
 	return InputDesc ? InputDesc->ComputeTotalElementCount(InInputPinProps.AllowedTypes) : 0;
 }
diff --git a/Engine/Plugins/PCG/Source/PCG/Private/Compute/Elements/PCGCustomHLSLKernel.h b/Engine/Plugins/PCG/Source/PCG/Private/Compute/Elements/PCGCustomHLSLKernel.h
--- a/Engine/Plugins/PCG/Source/PCG/Private/Compute/Elements/PCGCustomHLSLKernel.h
+++ b/Engine/Plugins/PCG/Source/PCG/Private/Compute/Elements/PCGCustomHLSLKernel.h
@@ -115,7 +115,7 @@ class UPCGCustomHLSLKernel : public UPCGComputeKernel
 	// ~Begin UPCGComputeKernel interface
 public:
 	virtual bool IsKernelDataValid(FPCGContext* InContext) const override;
-	virtual bool ComputeOutputBindingDataDesc(const UPCGComputeGraph* InGraph, FName InOutputPinLabel, UPCGDataBinding* InBinding, FPCGDataCollectionDesc& OutDataDesc) const override;
+	virtual TSharedPtr<const FPCGDataCollectionDesc> ComputeOutputBindingDataDesc(const UPCGComputeGraph* InGraph, FName InOutputPinLabel, UPCGDataBinding* InBinding) const override;
 #if WITH_EDITOR
 	virtual FString GetCookedSource(FPCGGPUCompilationContext& InOutContext) const override;
 	virtual FString GetEntryPoint() const override { return EntryPoint; }
diff --git a/Engine/Plugins/PCG/Source/PCG/Private/Compute/Elements/PCGNormalToDensityKernel.cpp b/Engine/Plugins/PCG/Source/PCG/Private/Compute/Elements/PCGNormalToDensityKernel.cpp
--- a/Engine/Plugins/PCG/Source/PCG/Private/Compute/Elements/PCGNormalToDensityKernel.cpp
+++ b/Engine/Plugins/PCG/Source/PCG/Private/Compute/Elements/PCGNormalToDensityKernel.cpp
@@ -13,7 +13,7 @@
 
 #define LOCTEXT_NAMESPACE "PCGTemplateKernel"
 
-bool UPCGNormalToDensityKernel::ComputeOutputBindingDataDesc(const UPCGComputeGraph* InGraph, FName InOutputPinLabel, UPCGDataBinding* InBinding, FPCGDataCollectionDesc& OutDataDesc) const
+TSharedPtr<const FPCGDataCollectionDesc> UPCGNormalToDensityKernel::ComputeOutputBindingDataDesc(const UPCGComputeGraph* InGraph, FName InOutputPinLabel, UPCGDataBinding* InBinding) const
 {
 	check(InGraph);
 	check(InBinding);
@@ -21,15 +21,21 @@ bool UPCGNormalToDensityKernel::ComputeOutputBindingDataDesc(const UPCGComputeGr
 	// Code assumes single output pin.
 	if (!ensure(InOutputPinLabel == PCGPinConstants::DefaultOutputLabel))
 	{
-		return false;
+		return nullptr;
 	}
 
 	const FPCGKernelPin SourceKernelPin(KernelIndex, PCGPinConstants::DefaultInputLabel, /*bIsInput=*/true);
+	const TSharedPtr<const FPCGDataCollectionDesc> InDataDesc = InBinding->ComputeKernelPinDataDesc(SourceKernelPin);
 
-	ensure(InBinding->ComputeKernelPinDataDesc(SourceKernelPin, &OutDataDesc));
+	if (!ensure(InDataDesc))
+	{
+		return nullptr;
+	}
+
+	TSharedPtr<FPCGDataCollectionDesc> OutDataDesc = FPCGDataCollectionDesc::MakeSharedFrom(InDataDesc);
 
 	// Allocate density property if the transform is also allocated. Otherwise the densities will all be the same value and full-width allocation is unnecessary.
-	for (FPCGDataDesc& DataDesc : OutDataDesc.DataDescs)
+	for (FPCGDataDesc& DataDesc : OutDataDesc->DataDescs)
 	{
 		if (!!(DataDesc.GetAllocatedProperties() & EPCGPointNativeProperties::Transform))
 		{
@@ -37,12 +43,12 @@ bool UPCGNormalToDensityKernel::ComputeOutputBindingDataDesc(const UPCGComputeGr
 		}
 	}
 
-	return true;
+	return OutDataDesc;
 }
 
 int UPCGNormalToDensityKernel::ComputeThreadCount(const UPCGDataBinding* InBinding) const
 {
-	const FPCGDataCollectionDesc* OutputPinDesc = InBinding->GetCachedKernelPinDataDesc(this, PCGPinConstants::DefaultOutputLabel, /*bIsInput=*/false);
+	const TSharedPtr<const FPCGDataCollectionDesc> OutputPinDesc = InBinding->GetCachedKernelPinDataDesc(this, PCGPinConstants::DefaultOutputLabel, /*bIsInput=*/false);
 	return ensure(OutputPinDesc) ? OutputPinDesc->ComputeTotalElementCount() : 0;
 }
 
diff --git a/Engine/Plugins/PCG/Source/PCG/Private/Compute/Elements/PCGNormalToDensityKernel.h b/Engine/Plugins/PCG/Source/PCG/Private/Compute/Elements/PCGNormalToDensityKernel.h
--- a/Engine/Plugins/PCG/Source/PCG/Private/Compute/Elements/PCGNormalToDensityKernel.h
+++ b/Engine/Plugins/PCG/Source/PCG/Private/Compute/Elements/PCGNormalToDensityKernel.h
@@ -15,7 +15,7 @@ class UPCGNormalToDensityKernel : public UPCGComputeKernel
 	GENERATED_BODY()
 
 public:
-	virtual bool ComputeOutputBindingDataDesc(const UPCGComputeGraph* InGraph, FName InOutputPinLabel, UPCGDataBinding* InBinding, FPCGDataCollectionDesc& OutDataDesc) const override;
+	virtual TSharedPtr<const FPCGDataCollectionDesc> ComputeOutputBindingDataDesc(const UPCGComputeGraph* InGraph, FName InOutputPinLabel, UPCGDataBinding* InBinding) const override;
 	virtual int ComputeThreadCount(const UPCGDataBinding* Binding) const override;
 
 #if WITH_EDITOR
diff --git a/Engine/Plugins/PCG/Source/PCG/Private/Compute/Elements/PCGTransformPointsKernel.cpp b/Engine/Plugins/PCG/Source/PCG/Private/Compute/Elements/PCGTransformPointsKernel.cpp
--- a/Engine/Plugins/PCG/Source/PCG/Private/Compute/Elements/PCGTransformPointsKernel.cpp
+++ b/Engine/Plugins/PCG/Source/PCG/Private/Compute/Elements/PCGTransformPointsKernel.cpp
@@ -14,7 +14,7 @@
 
 #define LOCTEXT_NAMESPACE "PCGTransformPointsKernel"
 
-bool UPCGTransformPointsKernel::ComputeOutputBindingDataDesc(const UPCGComputeGraph* InGraph, FName InOutputPinLabel, UPCGDataBinding* InBinding, FPCGDataCollectionDesc& OutDataDesc) const
+TSharedPtr<const FPCGDataCollectionDesc> UPCGTransformPointsKernel::ComputeOutputBindingDataDesc(const UPCGComputeGraph* InGraph, FName InOutputPinLabel, UPCGDataBinding* InBinding) const
 {
 	check(InGraph);
 	check(InBinding);
@@ -22,30 +22,37 @@ bool UPCGTransformPointsKernel::ComputeOutputBindingDataDesc(const UPCGComputeGr
 	// Code assumes single output pin.
 	if (!ensure(InOutputPinLabel == PCGPinConstants::DefaultOutputLabel))
 	{
-		return false;
+		return nullptr;
 	}
 
 	const FPCGKernelPin SourceKernelPin(KernelIndex, PCGPinConstants::DefaultInputLabel, /*bIsInput=*/true);
-	ensure(InBinding->ComputeKernelPinDataDesc(SourceKernelPin, &OutDataDesc));
+	const TSharedPtr<const FPCGDataCollectionDesc> InDataDesc = InBinding->ComputeKernelPinDataDesc(SourceKernelPin);
+
+	if (!ensure(InDataDesc))
+	{
+		return nullptr;
+	}
+
+	TSharedPtr<FPCGDataCollectionDesc> OutDataDesc = FPCGDataCollectionDesc::MakeSharedFrom(InDataDesc);
 	
 	const UPCGTransformPointsSettings* Settings = CastChecked<UPCGTransformPointsSettings>(GetSettings());
 	if (!ensure(Settings))
 	{
-		return false;
+		return nullptr;
 	}
 
 	// Allocate properties if needed
-	for (FPCGDataDesc& DataDesc : OutDataDesc.DataDescs)
+	for (FPCGDataDesc& DataDesc : OutDataDesc->DataDescs)
 	{
 		DataDesc.AllocateProperties(Settings->GetPropertiesToAllocate());
 	}
 
-	return true;
+	return OutDataDesc;
 }
 
 int UPCGTransformPointsKernel::ComputeThreadCount(const UPCGDataBinding* InBinding) const
 {
-	const FPCGDataCollectionDesc* OutputPinDesc = InBinding->GetCachedKernelPinDataDesc(this, PCGPinConstants::DefaultOutputLabel, /*bIsInput=*/false);
+	const TSharedPtr<const FPCGDataCollectionDesc> OutputPinDesc = InBinding->GetCachedKernelPinDataDesc(this, PCGPinConstants::DefaultOutputLabel, /*bIsInput=*/false);
 	return ensure(OutputPinDesc) ? OutputPinDesc->ComputeTotalElementCount() : 0;
 }
 
diff --git a/Engine/Plugins/PCG/Source/PCG/Private/Compute/Elements/PCGTransformPointsKernel.h b/Engine/Plugins/PCG/Source/PCG/Private/Compute/Elements/PCGTransformPointsKernel.h
--- a/Engine/Plugins/PCG/Source/PCG/Private/Compute/Elements/PCGTransformPointsKernel.h
+++ b/Engine/Plugins/PCG/Source/PCG/Private/Compute/Elements/PCGTransformPointsKernel.h
@@ -15,7 +15,7 @@ class UPCGTransformPointsKernel : public UPCGComputeKernel
 
 public:
 	//~Begin UPCGComputeKernel interface
-	virtual bool ComputeOutputBindingDataDesc(const UPCGComputeGraph* InGraph, FName InOutputPinLabel, UPCGDataBinding* InBinding, FPCGDataCollectionDesc& OutDataDesc) const override;
+	virtual TSharedPtr<const FPCGDataCollectionDesc> ComputeOutputBindingDataDesc(const UPCGComputeGraph* InGraph, FName InOutputPinLabel, UPCGDataBinding* InBinding) const override;
 	virtual int ComputeThreadCount(const UPCGDataBinding* Binding) const override;
 #if WITH_EDITOR
 	virtual const TCHAR* GetSourceFilePath() const override { return TEXT("/Plugin/PCG/Private/Elements/PCGTransformPoints.usf"); }
diff --git a/Engine/Plugins/PCG/Source/PCG/Private/Compute/PCGComputeGraph.cpp b/Engine/Plugins/PCG/Source/PCG/Private/Compute/PCGComputeGraph.cpp
--- a/Engine/Plugins/PCG/Source/PCG/Private/Compute/PCGComputeGraph.cpp
+++ b/Engine/Plugins/PCG/Source/PCG/Private/Compute/PCGComputeGraph.cpp
@@ -98,7 +98,7 @@ void UPCGComputeGraph::OnKernelCompilationComplete(int32 InKernelIndex, FCompute
 	}
 }
 
-bool UPCGComputeGraph::ComputeKernelBindingDataDesc(int32 InBindingIndex, UPCGDataBinding* InBinding, FPCGDataCollectionDesc& OutDataDesc) const
+TSharedPtr<const FPCGDataCollectionDesc> UPCGComputeGraph::ComputeKernelBindingDataDesc(int32 InBindingIndex, UPCGDataBinding* InBinding) const
 {
 	TRACE_CPUPROFILER_EVENT_SCOPE(UPCGComputeGraph::ComputeKernelBindingDataDesc);
 
@@ -117,14 +117,12 @@ bool UPCGComputeGraph::ComputeKernelBindingDataDesc(int32 InBindingIndex, UPCGDa
 	if (const FName* VirtualLabel = CPUDataBindingToVirtualPinLabel.Find(InBindingIndex))
 	{
 		// Create description from compute graph element input data.
-		ComputeExternalPinDesc(*VirtualLabel, InBinding, OutDataDesc);
-
-		return true;
+		return ComputeExternalPinDesc(*VirtualLabel, InBinding);
 	}
 	else if (GraphEdge.bKernelInput)
 	{
 		const int32* UpstreamBindingIndex = DownstreamToUpstreamBinding.Find(InBindingIndex);
-		return ensure(UpstreamBindingIndex) && InBinding->ComputeKernelBindingDataDesc(*UpstreamBindingIndex, &OutDataDesc);
+		return ensure(UpstreamBindingIndex) ? InBinding->ComputeKernelBindingDataDesc(*UpstreamBindingIndex) : nullptr;
 	}
 	else
 	{
@@ -134,7 +132,7 @@ bool UPCGComputeGraph::ComputeKernelBindingDataDesc(int32 InBindingIndex, UPCGDa
 		const UPCGComputeKernel* Kernel = CastChecked<UPCGComputeKernel>(KernelInvocations[GraphEdge.KernelIndex]);
 		const FName* OutputPinLabel = KernelBindingToPinLabel.Find(InBindingIndex);
 
-		return ensure(OutputPinLabel) && Kernel->ComputeOutputBindingDataDesc(this, *OutputPinLabel, InBinding, OutDataDesc);
+		return ensure(OutputPinLabel) ? Kernel->ComputeOutputBindingDataDesc(this, *OutputPinLabel, InBinding) : nullptr;
 	}
 }
 
@@ -208,7 +206,7 @@ void UPCGComputeGraph::DebugLogDataDescriptions(const UPCGDataBinding* InBinding
 	const UEnum* PCGKernelAttributeTypeEnum = StaticEnum<EPCGKernelAttributeType>();
 	check(PCGKernelAttributeTypeEnum);
 
-	auto LogDataDescription = [this, PCGKernelAttributeTypeEnum, InBinding](const FPCGDataCollectionDesc* InDataDesc, const FPCGKernelPin& InKernelPin)
+	auto LogDataDescription = [this, PCGKernelAttributeTypeEnum, InBinding](const TSharedPtr<const FPCGDataCollectionDesc> InDataDesc, const FPCGKernelPin& InKernelPin)
 	{
 		if (!InDataDesc)
 		{
@@ -319,7 +317,7 @@ void UPCGComputeGraph::DebugLogDataDescriptions(const UPCGDataBinding* InBinding
 			{
 				if (Pin.Label != NAME_None)
 				{
-					if (const FPCGDataCollectionDesc* DataDesc = InBinding->GetCachedKernelPinDataDesc(PCGKernel, Pin.Label, true))
+					if (const TSharedPtr<const FPCGDataCollectionDesc> DataDesc = InBinding->GetCachedKernelPinDataDesc(PCGKernel, Pin.Label, true))
 					{
 						UE_LOG(LogPCG, Display, TEXT("\tInput Pin: %s (%d data)"), *Pin.Label.ToString(), DataDesc->DataDescs.Num());
 
@@ -334,7 +332,7 @@ void UPCGComputeGraph::DebugLogDataDescriptions(const UPCGDataBinding* InBinding
 			{
 				if (Pin.Label != NAME_None)
 				{
-					if (const FPCGDataCollectionDesc* DataDesc = InBinding->GetCachedKernelPinDataDesc(PCGKernel, Pin.Label, false))
+					if (const TSharedPtr<const FPCGDataCollectionDesc> DataDesc = InBinding->GetCachedKernelPinDataDesc(PCGKernel, Pin.Label, false))
 					{
 						UE_LOG(LogPCG, Display, TEXT("\tOutput Pin: %s (%d data)"), *Pin.Label.ToString(), DataDesc->DataDescs.Num());
 
@@ -369,10 +367,12 @@ bool UPCGComputeGraph::ShouldRepeatDispatch() const
 }
 #endif // PCG_GPU_KERNEL_PROFILING
 
-void UPCGComputeGraph::ComputeExternalPinDesc(FName InVirtualLabel, UPCGDataBinding* InBinding, FPCGDataCollectionDesc& OutDataDesc) const
+TSharedPtr<const FPCGDataCollectionDesc> UPCGComputeGraph::ComputeExternalPinDesc(FName InVirtualLabel, UPCGDataBinding* InBinding) const
 {
 	check(InBinding);
 
+	TSharedPtr<FPCGDataCollectionDesc> OutDataDesc = FPCGDataCollectionDesc::MakeShared();
+
 	const TArray<FPCGTaggedData> TaggedDatas = InBinding->InputDataCollection.GetInputsByPin(InVirtualLabel);
 
 	for (const FPCGTaggedData& TaggedData : TaggedDatas)
@@ -385,9 +385,11 @@ void UPCGComputeGraph::ComputeExternalPinDesc(FName InVirtualLabel, UPCGDataBind
 		FPCGDataDesc DataDesc;
 		if (ensure(ComputeTaggedDataPinDesc(TaggedData, InBinding, DataDesc)))
 		{
-			OutDataDesc.DataDescs.Add(MoveTemp(DataDesc));
+			OutDataDesc->DataDescs.Add(MoveTemp(DataDesc));
 		}
 	}
+
+	return OutDataDesc;
 }
 
 bool UPCGComputeGraph::ComputeTaggedDataPinDesc(const FPCGTaggedData& InTaggedData, const UPCGDataBinding* InBinding, FPCGDataDesc& OutDescription) const
diff --git a/Engine/Plugins/PCG/Source/PCG/Private/Compute/PCGComputeGraph.h b/Engine/Plugins/PCG/Source/PCG/Private/Compute/PCGComputeGraph.h
--- a/Engine/Plugins/PCG/Source/PCG/Private/Compute/PCGComputeGraph.h
+++ b/Engine/Plugins/PCG/Source/PCG/Private/Compute/PCGComputeGraph.h
@@ -77,7 +77,7 @@ public:
 	 *
 	 * A pin will have many bindings, but we arbitrarily use the first binding index to uniquely represent a pin.
 	 */
-	bool ComputeKernelBindingDataDesc(int32 InBindingIndex, UPCGDataBinding* InBinding, FPCGDataCollectionDesc& OutDataDesc) const;
+	TSharedPtr<const FPCGDataCollectionDesc> ComputeKernelBindingDataDesc(int32 InBindingIndex, UPCGDataBinding* InBinding) const;
 
 	/** Get the first binding index, which uniquely identifies a kernel pin to the compute graph. */
 	int32 GetBindingIndex(const FPCGKernelPin& InKernelPin) const;
@@ -109,7 +109,7 @@ public:
 
 private:
 	/** Compute the data description on a pin external to the compute graph. */
-	void ComputeExternalPinDesc(FName InVirtualLabel, UPCGDataBinding* InBinding, FPCGDataCollectionDesc& OutDataDesc) const;
+	TSharedPtr<const FPCGDataCollectionDesc> ComputeExternalPinDesc(FName InVirtualLabel, UPCGDataBinding* InBinding) const;
 
 	/** Compute the data description of a tagged data. */
 	bool ComputeTaggedDataPinDesc(const FPCGTaggedData& InTaggedData, const UPCGDataBinding* InBinding, FPCGDataDesc& OutDescription) const;
diff --git a/Engine/Plugins/PCG/Source/PCG/Private/Compute/PCGComputeKernel.cpp b/Engine/Plugins/PCG/Source/PCG/Private/Compute/PCGComputeKernel.cpp
--- a/Engine/Plugins/PCG/Source/PCG/Private/Compute/PCGComputeKernel.cpp
+++ b/Engine/Plugins/PCG/Source/PCG/Private/Compute/PCGComputeKernel.cpp
@@ -169,7 +169,7 @@ bool UPCGComputeKernel::IsKernelDataValid(FPCGContext* InContext) const
 					// when looking up the pin data description, so only consider InitPins which actually exist.
 					if (Node->GetInputPin(InitPin))
 					{
-						if (const FPCGDataCollectionDesc* InputDesc = DataBinding->GetCachedKernelPinDataDesc(this, InitPin, /*bIsInputPin=*/true))
+						if (const TSharedPtr<const FPCGDataCollectionDesc> InputDesc = DataBinding->GetCachedKernelPinDataDesc(this, InitPin, /*bIsInputPin=*/true))
 						{
 							const uint32 InputNumData = InputDesc->DataDescs.Num();
 
@@ -250,7 +250,7 @@ UPCGComputeDataInterface* UPCGComputeKernel::CreateOutputPinDataInterface(const
 }
 #endif
 
-void UPCGComputeKernel::ComputeDataDescFromPinProperties(const FPCGPinPropertiesGPU& OutputPinProps, const TArrayView<const FPCGPinProperties>& InInputPinProps, UPCGDataBinding* InBinding, FPCGDataCollectionDesc& OutPinDesc) const
+void UPCGComputeKernel::ComputeDataDescFromPinProperties(const FPCGPinPropertiesGPU& OutputPinProps, const TArrayView<const FPCGPinProperties>& InInputPinProps, UPCGDataBinding* InBinding, TSharedPtr<FPCGDataCollectionDesc> OutPinDesc) const
 {
 	check(InBinding);
 
@@ -259,7 +259,7 @@ void UPCGComputeKernel::ComputeDataDescFromPinProperties(const FPCGPinProperties
 	if (Props.InitializationMode == EPCGPinInitMode::FromInputPins)
 	{
 		TArray<FName> InitPins;
-		TArray<FPCGDataCollectionDesc> InputDescs;
+		TArray<TSharedPtr<const FPCGDataCollectionDesc>> InputDescs;
 
 		for (FName PinToInitFrom : Props.PinsToInititalizeFrom)
 		{
@@ -267,9 +267,10 @@ void UPCGComputeKernel::ComputeDataDescFromPinProperties(const FPCGPinProperties
 			{
 				InitPins.Emplace(PinToInitFrom);
 
-				FPCGDataCollectionDesc& InputDesc = InputDescs.AddDefaulted_GetRef();
 				const FPCGKernelPin KernelPin(KernelIndex, PinToInitFrom, /*bIsInput=*/true);
-				InBinding->ComputeKernelPinDataDesc(KernelPin, &InputDesc);
+				TSharedPtr<const FPCGDataCollectionDesc> InputDesc = InBinding->ComputeKernelPinDataDesc(KernelPin);
+
+				InputDescs.Add(InputDesc);
 			}
 		}
 
@@ -330,22 +331,22 @@ void UPCGComputeKernel::ComputeDataDescFromPinProperties(const FPCGPinProperties
 					}
 				}
 
-				FPCGDataDesc& DataDesc = OutPinDesc.DataDescs.Emplace_GetRef(OutputPinProps.AllowedTypes, InitialNumElements);
+				FPCGDataDesc& DataDesc = OutPinDesc->DataDescs.Emplace_GetRef(OutputPinProps.AllowedTypes, InitialNumElements);
 
 				// For each data index, loop over all the pins and create the uber-data.
 				for (int InputPinIndex = 0; InputPinIndex < NumInitPins; ++InputPinIndex)
 				{
-					const FPCGDataCollectionDesc& InputDesc = InputDescs[InputPinIndex];
+					const TSharedPtr<const FPCGDataCollectionDesc> InputDesc = InputDescs[InputPinIndex];
 
 					int ClampedDataIndex = DataIndex;
 
 					// If this pin does not have the same number of data, clamp it to the last data.
-					if (DataIndex >= InputDesc.DataDescs.Num())
+					if (DataIndex >= InputDesc->DataDescs.Num())
 					{
-						ClampedDataIndex = InputDesc.DataDescs.Num() - 1;
+						ClampedDataIndex = InputDesc->DataDescs.Num() - 1;
 					}
 
-					const FPCGDataDesc& InputDataDesc = InputDesc.DataDescs.IsValidIndex(ClampedDataIndex) ? InputDesc.DataDescs[ClampedDataIndex] : FPCGDataDesc(FPCGDataTypeIdentifier{EPCGDataType::Any}, 0);
+					const FPCGDataDesc& InputDataDesc = InputDesc->DataDescs.IsValidIndex(ClampedDataIndex) ? InputDesc->DataDescs[ClampedDataIndex] : FPCGDataDesc(FPCGDataTypeIdentifier{EPCGDataType::Any}, 0);
 					const EPCGElementDimension SourceDimension = InputDataDesc.GetElementDimension();
 
 					if (Props.ElementCountMode == EPCGElementCountMode::FromInputData)
@@ -396,9 +397,9 @@ void UPCGComputeKernel::ComputeDataDescFromPinProperties(const FPCGPinProperties
 			// If this is the only input pin, we can just copy it.
 			if (NumInitPins == 1)
 			{
-				for (const FPCGDataDesc& InputDataDesc : InputDescs[0].DataDescs)
+				for (const FPCGDataDesc& InputDataDesc : InputDescs[0]->DataDescs)
 				{
-					FPCGDataDesc& DataDesc = OutPinDesc.DataDescs.Emplace_GetRef(OutputPinProps.AllowedTypes, /*InitialElementCount=*/0);
+					FPCGDataDesc& DataDesc = OutPinDesc->DataDescs.Emplace_GetRef(OutputPinProps.AllowedTypes, /*InitialElementCount=*/0);
 					DataDesc.CombineElementCount(InputDataDesc, EPCGElementMultiplicity::Sum);
 
 					if (Props.AttributeInheritanceMode == EPCGAttributeInheritanceMode::CopyAttributeSetup)
@@ -415,17 +416,17 @@ void UPCGComputeKernel::ComputeDataDescFromPinProperties(const FPCGPinProperties
 			{
 				for (int InputPinIndex = 0; InputPinIndex < NumInitPins; ++InputPinIndex)
 				{
-					const FPCGDataCollectionDesc& InputDesc = InputDescs[InputPinIndex];
+					const TSharedPtr<const FPCGDataCollectionDesc> InputDesc = InputDescs[InputPinIndex];
 
 					for (int OtherInputPinIndex = InputPinIndex + 1; OtherInputPinIndex < NumInitPins; ++OtherInputPinIndex)
 					{
-						const FPCGDataCollectionDesc& OtherInputDesc = InputDescs[OtherInputPinIndex];
+						const TSharedPtr<const FPCGDataCollectionDesc> OtherInputDesc = InputDescs[OtherInputPinIndex];
 
-						for (const FPCGDataDesc& InputDataDesc : InputDesc.DataDescs)
+						for (const FPCGDataDesc& InputDataDesc : InputDesc->DataDescs)
 						{
-							for (const FPCGDataDesc& OtherInputDataDesc : OtherInputDesc.DataDescs)
+							for (const FPCGDataDesc& OtherInputDataDesc : OtherInputDesc->DataDescs)
 							{
-								FPCGDataDesc& DataDesc = OutPinDesc.DataDescs.Emplace_GetRef(OutputPinProps.AllowedTypes, /*InitialElementCount=*/0);
+								FPCGDataDesc& DataDesc = OutPinDesc->DataDescs.Emplace_GetRef(OutputPinProps.AllowedTypes, /*InitialElementCount=*/0);
 
 								if (Props.ElementCountMode == EPCGElementCountMode::FromInputData)
 								{
@@ -485,7 +486,7 @@ void UPCGComputeKernel::ComputeDataDescFromPinProperties(const FPCGPinProperties
 				// Find the maximum number of data among the init pins. Note, they should all be the same number of data, or only one data.
 				for (int I = 0; I < NumInitPins; ++I)
 				{
-					const int DataCount = InputDescs[I].DataDescs.Num();
+					const int DataCount = InputDescs[I]->DataDescs.Num();
 
 					// If any init pins have no data, then the pairing fails and no data are produced.
 					if (DataCount == 0)
@@ -515,7 +516,7 @@ void UPCGComputeKernel::ComputeDataDescFromPinProperties(const FPCGPinProperties
 
 		// Apply element count multiplier.
 		const uint32 Multiplier = OutputPinProps.GetElementCountMultiplier();
-		for (FPCGDataDesc& Desc : OutPinDesc.DataDescs)
+		for (FPCGDataDesc& Desc : OutPinDesc->DataDescs)
 		{
 			Desc.ScaleElementCount(Multiplier);
 		}
@@ -528,11 +529,11 @@ void UPCGComputeKernel::ComputeDataDescFromPinProperties(const FPCGPinProperties
 
 			if (ElementDimension == EPCGElementDimension::One)
 			{
-				OutPinDesc.DataDescs.Emplace(OutputPinProps.AllowedTypes, Props.ElementCount);
+				OutPinDesc->DataDescs.Emplace(OutputPinProps.AllowedTypes, Props.ElementCount);
 			}
 			else if (ElementDimension == EPCGElementDimension::Two)
 			{
-				OutPinDesc.DataDescs.Emplace(OutputPinProps.AllowedTypes, Props.NumElements2D);
+				OutPinDesc->DataDescs.Emplace(OutputPinProps.AllowedTypes, Props.NumElements2D);
 			}
 			else
 			{
diff --git a/Engine/Plugins/PCG/Source/PCG/Private/Compute/PCGComputeKernel.h b/Engine/Plugins/PCG/Source/PCG/Private/Compute/PCGComputeKernel.h
--- a/Engine/Plugins/PCG/Source/PCG/Private/Compute/PCGComputeKernel.h
+++ b/Engine/Plugins/PCG/Source/PCG/Private/Compute/PCGComputeKernel.h
@@ -155,7 +155,7 @@ public:
 #endif
 
 	/** Compute a description of data that will be output from pin InOutputPinLabel. */
-	virtual bool ComputeOutputBindingDataDesc(const UPCGComputeGraph* InGraph, FName InOutputPinLabel, UPCGDataBinding* InBinding, FPCGDataCollectionDesc& OutDataDesc) const PURE_VIRTUAL(UPCGComputeKernel::ComputeOutputBindingDataDesc, return {};);
+	virtual TSharedPtr<const FPCGDataCollectionDesc> ComputeOutputBindingDataDesc(const UPCGComputeGraph* InGraph, FName InOutputPinLabel, UPCGDataBinding* InBinding) const PURE_VIRTUAL(UPCGComputeKernel::ComputeOutputBindingDataDesc, return nullptr;);
 
 	/** Compute how many threads should be dispatched to execute this node on the GPU. */
 	virtual int ComputeThreadCount(const UPCGDataBinding* Binding) const PURE_VIRTUAL(UPCGComputeKernel::ComputeThreadCount, return 0;);
@@ -182,7 +182,7 @@ public:
 	virtual uint32 GetElementCountMultiplier(FName InOutputPinLabel) const { return 1u; }
 
 	/** Helper to compute data descriptions when data is entirely determined by pin properties. */
-	void ComputeDataDescFromPinProperties(const FPCGPinPropertiesGPU& OutputPinProps, const TArrayView<const FPCGPinProperties>& InInputPinProps, UPCGDataBinding* InBinding, FPCGDataCollectionDesc& OutPinDesc) const;
+	void ComputeDataDescFromPinProperties(const FPCGPinPropertiesGPU& OutputPinProps, const TArrayView<const FPCGPinProperties>& InInputPinProps, UPCGDataBinding* InBinding, TSharedPtr<FPCGDataCollectionDesc> OutPinDesc) const;
 
 	virtual void GetInputPins(TArray<FPCGPinProperties>& OutPins) const {}
 	virtual void GetOutputPins(TArray<FPCGPinPropertiesGPU>& OutPins) const {}
diff --git a/Engine/Plugins/PCG/Source/PCG/Private/Compute/PCGDataBinding.cpp b/Engine/Plugins/PCG/Source/PCG/Private/Compute/PCGDataBinding.cpp
--- a/Engine/Plugins/PCG/Source/PCG/Private/Compute/PCGDataBinding.cpp
+++ b/Engine/Plugins/PCG/Source/PCG/Private/Compute/PCGDataBinding.cpp
@@ -96,18 +96,18 @@ void UPCGDataBinding::PrimeDataDescriptionCache()
 		// @todo_pcg: Could add a ParallelFor if this becomes a bottle neck.
 		for (const FPCGKernelPin& KernelPin : KernelPins)
 		{
-			ComputeKernelPinDataDesc(KernelPin, /*OutDesc=*/nullptr);
+			ComputeKernelPinDataDesc(KernelPin);
 		}
 	}
 
 	bIsDataDescriptionCachePrimed = true;
 }
 
-bool UPCGDataBinding::ComputeKernelPinDataDesc(const FPCGKernelPin& InKernelPin, FPCGDataCollectionDesc* OutDesc)
+TSharedPtr<const FPCGDataCollectionDesc> UPCGDataBinding::ComputeKernelPinDataDesc(const FPCGKernelPin& InKernelPin)
 {
 	if (!ensure(Graph))
 	{
-		return false;
+		return nullptr;
 	}
 
 	const int32 GraphBindingIndex = Graph->GetBindingIndex(InKernelPin);
@@ -115,13 +115,13 @@ bool UPCGDataBinding::ComputeKernelPinDataDesc(const FPCGKernelPin& InKernelPin,
 	if (GraphBindingIndex == -1)
 	{
 		UE_LOG(LogPCG, Error, TEXT("Failed to compute data description for kernel pin '%s'."), *InKernelPin.PinLabel.ToString());
-		return false;
+		return nullptr;
 	}
 
-	return ComputeKernelBindingDataDesc(GraphBindingIndex, OutDesc);
+	return ComputeKernelBindingDataDesc(GraphBindingIndex);
 }
 
-bool UPCGDataBinding::ComputeKernelBindingDataDesc(int32 InGraphBindingIndex, FPCGDataCollectionDesc* OutDesc)
+TSharedPtr<const FPCGDataCollectionDesc> UPCGDataBinding::ComputeKernelBindingDataDesc(int32 InGraphBindingIndex)
 {
 	// Uncomment to be notified when a kernel pin description is computed late:
 	//ensureAlwaysMsgf(!IsInGameThread(), TEXT("Data descriptions should only be computed when priming the cache, which takes place on a worker thread."));
@@ -130,41 +130,23 @@ bool UPCGDataBinding::ComputeKernelBindingDataDesc(int32 InGraphBindingIndex, FP
 
 	if (!ensure(Graph))
 	{
-		return false;
+		return nullptr;
 	}
 
-	if (FPCGDataCollectionDesc* FoundDesc = DataDescriptionCache.Find(InGraphBindingIndex))
+	// Try to look up the data description if it's been computed previously.
+	if (const TSharedPtr<const FPCGDataCollectionDesc>* FoundDesc = DataDescriptionCache.Find(InGraphBindingIndex))
 	{
-		if (OutDesc)
-		{
-			*OutDesc = *FoundDesc;
-		}
-
-		return true;
+		return *FoundDesc;
 	}
 
-	FPCGDataCollectionDesc ComputedDesc;
-	
-	if (Graph->ComputeKernelBindingDataDesc(InGraphBindingIndex, this, ComputedDesc))
-	{
-		if (!OutDesc)
-		{
-			DataDescriptionCache.Add(InGraphBindingIndex, MoveTemp(ComputedDesc));
-		}
-		else
-		{
-			DataDescriptionCache.Add(InGraphBindingIndex, ComputedDesc);
-
-			*OutDesc = MoveTemp(ComputedDesc);
-		}
+	// If the description wasn't found in the cache, compute it now.
+	const TSharedPtr<const FPCGDataCollectionDesc> ComputedDesc = Graph->ComputeKernelBindingDataDesc(InGraphBindingIndex, this);
+	DataDescriptionCache.Add(InGraphBindingIndex, ComputedDesc);
 
-		return true;
-	}
-
-	return false;
+	return ComputedDesc;
 }
 
-const FPCGDataCollectionDesc* UPCGDataBinding::GetCachedKernelPinDataDesc(const UPCGComputeKernel* InKernel, FName InPinLabel, bool bIsInput) const
+TSharedPtr<const FPCGDataCollectionDesc> UPCGDataBinding::GetCachedKernelPinDataDesc(const UPCGComputeKernel* InKernel, FName InPinLabel, bool bIsInput) const
 {
 	if (!ensure(Graph) || !ensure(InKernel))
 	{
@@ -174,7 +156,7 @@ const FPCGDataCollectionDesc* UPCGDataBinding::GetCachedKernelPinDataDesc(const
 	return GetCachedKernelPinDataDesc(Graph->GetBindingIndex(FPCGKernelPin(InKernel->GetKernelIndex(), InPinLabel, bIsInput)));
 }
 
-const FPCGDataCollectionDesc* UPCGDataBinding::GetCachedKernelPinDataDesc(int32 InGraphBindingIndex) const
+TSharedPtr<const FPCGDataCollectionDesc> UPCGDataBinding::GetCachedKernelPinDataDesc(int32 InGraphBindingIndex) const
 {
 	TRACE_CPUPROFILER_EVENT_SCOPE(UPCGDataBinding::GetCachedKernelPinDataDesc);
 
@@ -184,15 +166,16 @@ const FPCGDataCollectionDesc* UPCGDataBinding::GetCachedKernelPinDataDesc(int32
 		return nullptr;
 	}
 
-	const FPCGDataCollectionDesc* FoundDescription = DataDescriptionCache.Find(InGraphBindingIndex);
+	const TSharedPtr<const FPCGDataCollectionDesc>* FoundDescription = DataDescriptionCache.Find(InGraphBindingIndex);
 
 	if (!FoundDescription)
 	{
 		// Should not happen - all data descriptions should have been computed at this point.
 		UE_LOG(LogPCG, Warning, TEXT("Cache miss while retrieving kernel pin data description, binding index %d."), InGraphBindingIndex);
+		return nullptr;
 	}
 
-	return FoundDescription;
+	return *FoundDescription;
 }
 
 void UPCGDataBinding::ReceiveDataFromGPU_GameThread(
diff --git a/Engine/Plugins/PCG/Source/PCG/Private/Compute/PCGDataBinding.h b/Engine/Plugins/PCG/Source/PCG/Private/Compute/PCGDataBinding.h
--- a/Engine/Plugins/PCG/Source/PCG/Private/Compute/PCGDataBinding.h
+++ b/Engine/Plugins/PCG/Source/PCG/Private/Compute/PCGDataBinding.h
@@ -96,12 +96,12 @@ public:
 	bool IsDataDescriptionCachePrimed() const { return bIsDataDescriptionCachePrimed; }
 
 	/** Computes a description of data for a single pin in the compute graph and caches it. */
-	bool ComputeKernelPinDataDesc(const FPCGKernelPin& InKernelPin, FPCGDataCollectionDesc* OutDesc);
-	bool ComputeKernelBindingDataDesc(int32 InGraphBindingIndex, FPCGDataCollectionDesc* OutDesc);
+	TSharedPtr<const FPCGDataCollectionDesc> ComputeKernelPinDataDesc(const FPCGKernelPin& InKernelPin);
+	TSharedPtr<const FPCGDataCollectionDesc> ComputeKernelBindingDataDesc(int32 InGraphBindingIndex);
 
 	/** Get description of data produced by a pin. */
-	const FPCGDataCollectionDesc* GetCachedKernelPinDataDesc(const UPCGComputeKernel* InKernel, FName InPinLabel, bool bIsInput) const;
-	const FPCGDataCollectionDesc* GetCachedKernelPinDataDesc(int32 InGraphBindingIndex) const;
+	TSharedPtr<const FPCGDataCollectionDesc> GetCachedKernelPinDataDesc(const UPCGComputeKernel* InKernel, FName InPinLabel, bool bIsInput) const;
+	TSharedPtr<const FPCGDataCollectionDesc> GetCachedKernelPinDataDesc(int32 InGraphBindingIndex) const;
 
 	void ReceiveDataFromGPU_GameThread(UPCGData* InData, const UPCGSettings* InProducerSettings, EPCGExportMode InExportMode, FName InPinLabel, FName InPinLabelAlias);
 	void ReceiveDataFromGPU_GameThread(UPCGData* InData, const UPCGSettings* InProducerSettings, EPCGExportMode InExportMode, FName InPinLabel, FName InPinLabelAlias, const TSet<FString>& AdditionalTags);
@@ -202,7 +202,7 @@ protected:
 	TMap<FSourceBufferAndStringKey, /*Source Buffer String Key*/int32> GraphToSourceBufferStringKey;
 
 	/** Cache of data descriptions to amortize cost of computing them at runtime. */
-	TMap</*Binding index*/int32, FPCGDataCollectionDesc> DataDescriptionCache;
+	TMap</*Binding index*/int32, TSharedPtr<const FPCGDataCollectionDesc>> DataDescriptionCache;
 	std::atomic<bool> bIsDataDescriptionCachePrimed = false;
 
 	std::atomic<bool> bTablesInitialized = false;
diff --git a/Engine/Plugins/PCG/Source/PCG/Private/Compute/PCGDataForGPU.cpp b/Engine/Plugins/PCG/Source/PCG/Private/Compute/PCGDataForGPU.cpp
--- a/Engine/Plugins/PCG/Source/PCG/Private/Compute/PCGDataForGPU.cpp
+++ b/Engine/Plugins/PCG/Source/PCG/Private/Compute/PCGDataForGPU.cpp
@@ -576,7 +576,7 @@ namespace PCGDataForGPUHelpers
 		const UPCGSettings* InSettings,
 		TArray<FName>& InAttributeNames,
 		const UPCGDataBinding* InBinding,
-		const FPCGDataCollectionDesc* InDataCollectionDescription,
+		const TSharedPtr<const FPCGDataCollectionDesc> InDataCollectionDescription,
 		uint32& OutCustomFloatCount,
 		TArray<FUint32Vector4>& OutAttributeIdOffsetStrides)
 	{
@@ -1168,23 +1168,23 @@ void FPCGDataDesc::InitializeAttributeDescs(const UPCGData* InData, const UPCGDa
 	}
 }
 
-FPCGDataCollectionDesc FPCGDataCollectionDesc::BuildFromDataCollection(
-	const FPCGDataCollection& InDataCollection,
-	const UPCGDataBinding* InBinding)
+TSharedPtr<FPCGDataCollectionDesc> FPCGDataCollectionDesc::MakeShared()
 {
-	FPCGDataCollectionDesc CollectionDesc;
+	return MakeShareable(new FPCGDataCollectionDesc());
+}
 
-	for (const FPCGTaggedData& Data : InDataCollection.TaggedData)
-	{
-		if (!Data.Data || !PCGComputeHelpers::IsTypeAllowedAsOutput(Data.Data->GetDataTypeId()))
-		{
-			continue;
-		}
+TSharedPtr<FPCGDataCollectionDesc> FPCGDataCollectionDesc::MakeSharedFrom(const TSharedPtr<const FPCGDataCollectionDesc> InOtherDataDesc)
+{
+	TRACE_CPUPROFILER_EVENT_SCOPE(FPCGDataCollectionDesc::MakeSharedFrom);
 
-		CollectionDesc.DataDescs.Emplace(Data, InBinding);
+	TSharedPtr<FPCGDataCollectionDesc> NewDesc = MakeShareable(new FPCGDataCollectionDesc());
+
+	if (InOtherDataDesc)
+	{
+		NewDesc->DataDescs = InOtherDataDesc->DataDescs;
 	}
 
-	return CollectionDesc;
+	return NewDesc;
 }
 
 uint32 FPCGDataCollectionDesc::ComputePackedHeaderSizeBytes() const
@@ -1216,6 +1216,8 @@ uint64 FPCGDataCollectionDesc::ComputePackedSizeBytes() const
 
 void FPCGDataCollectionDesc::WriteHeader(TArray<uint32>& OutPackedDataCollectionHeader) const
 {
+	TRACE_CPUPROFILER_EVENT_SCOPE(FPCGDataCollectionDesc::WriteHeader);
+
 	const uint32 HeaderSizeBytes = ComputePackedHeaderSizeBytes();
 	const uint32 HeaderSizeUints = HeaderSizeBytes >> 2;
 
@@ -2036,11 +2038,6 @@ uint32 FPCGDataCollectionDesc::ComputeTotalElementCount(FPCGDataTypeIdentifier I
 	return ComputeTotalElementCount();
 }
 
-void FPCGDataCollectionDesc::Combine(const FPCGDataCollectionDesc& Other)
-{
-	DataDescs.Append(Other.DataDescs);
-}
-
 bool FPCGDataCollectionDesc::GetAttributeDesc(FPCGAttributeIdentifier InAttributeIdentifier, FPCGKernelAttributeDesc& OutAttributeDesc, bool& bOutConflictingTypesFound, bool& bOutPresentOnAllData) const
 {
 	// Will be set to the type of the first attribute across all the data which has a matching name, and then is used to detect if subsequent
@@ -2140,7 +2137,7 @@ int FPCGDataCollectionDesc::GetNumStringKeyValues(int32 InAttributeId) const
 	return UniqueStringKeys.Num();
 }
 
-FPCGProxyForGPUDataCollection::FPCGProxyForGPUDataCollection(TRefCountPtr<FRDGPooledBuffer> InBuffer, uint32 InBufferSizeBytes, const FPCGDataCollectionDesc& InDescription, const TArray<FString>& InStringTable)
+FPCGProxyForGPUDataCollection::FPCGProxyForGPUDataCollection(TRefCountPtr<FRDGPooledBuffer> InBuffer, uint32 InBufferSizeBytes, const TSharedPtr<const FPCGDataCollectionDesc> InDescription, const TArray<FString>& InStringTable)
 	: Buffer(InBuffer)
 	, BufferSizeBytes(InBufferSizeBytes)
 	, Description(InDescription)
@@ -2245,7 +2242,7 @@ bool FPCGProxyForGPUDataCollection::GetCPUData(FPCGContext* InContext, int32 InD
 	if (!RawReadbackData.IsEmpty())
 	{
 		FPCGDataCollection DataFromGPU;
-		const EPCGUnpackDataCollectionResult Result = Description.UnpackDataCollection(InContext, RawReadbackData, /*OutputPinLabelAlias*/NAME_None, StringTable, DataFromGPU);
+		const EPCGUnpackDataCollectionResult Result = ensure(Description) ? Description->UnpackDataCollection(InContext, RawReadbackData, /*OutputPinLabelAlias*/NAME_None, StringTable, DataFromGPU) : EPCGUnpackDataCollectionResult::NoData;
 
 		RawReadbackData.Empty();
 
diff --git a/Engine/Plugins/PCG/Source/PCG/Private/Compute/PCGDataUploadAdaptor.cpp b/Engine/Plugins/PCG/Source/PCG/Private/Compute/PCGDataUploadAdaptor.cpp
--- a/Engine/Plugins/PCG/Source/PCG/Private/Compute/PCGDataUploadAdaptor.cpp
+++ b/Engine/Plugins/PCG/Source/PCG/Private/Compute/PCGDataUploadAdaptor.cpp
@@ -15,17 +15,12 @@
 #include "RenderGraphUtils.h"
 #include "SystemTextures.h"
 
-FPCGDataUploadAdaptor::FPCGDataUploadAdaptor(UPCGDataBinding* InDataBinding, const FPCGDataCollectionDesc* InTargetDataCollectionDesc, FName InInputPinLabel)
+FPCGDataUploadAdaptor::FPCGDataUploadAdaptor(UPCGDataBinding* InDataBinding, const TSharedPtr<const FPCGDataCollectionDesc> InTargetDataCollectionDesc, FName InInputPinLabel)
 {
 	check(IsInGameThread());
 	check(InTargetDataCollectionDesc);
 
-	{
-		TRACE_CPUPROFILER_EVENT_SCOPE(CopyDataDescription);
-
-		// @todo_pcg: Avoid costly copy of data description.
-		TargetDataCollectionDesc = *InTargetDataCollectionDesc;
-	}
+	TargetDataCollectionDesc = InTargetDataCollectionDesc;
 	InputPinLabel = InInputPinLabel;
 	DataBinding = InDataBinding;
 
@@ -76,21 +71,21 @@ FPCGDataUploadAdaptor::FPCGDataUploadAdaptor(UPCGDataBinding* InDataBinding, con
 		}
 	}
 
-	if (UniqueDataCollection && UniqueDataCollection->GetDescription().DataDescs.Num() != NumDataForThisPin)
+	if (UniqueDataCollection && UniqueDataCollection->GetDescription() && UniqueDataCollection->GetDescription()->DataDescs.Num() != NumDataForThisPin)
 	{
 		// We're not using all the data items from the buffer, no reuse for now.
 		// TODO: Support indirection of data items.
 		UniqueDataCollection.Reset();
 	}
 
-	if (UniqueDataCollection)
+	if (UniqueDataCollection && UniqueDataCollection->GetDescription())
 	{
 		ExternalBufferForReuse = UniqueDataCollection->GetBuffer();
 		ExternalBufferSizeBytes = UniqueDataCollection->GetBufferSizeBytes();
 
 		ensureAlways(ExternalBufferMaxStringKeyValue == INDEX_NONE);
 
-		for (const FPCGDataDesc& DataDesc : UniqueDataCollection->GetDescription().DataDescs)
+		for (const FPCGDataDesc& DataDesc : UniqueDataCollection->GetDescription()->DataDescs)
 		{
 			for (const FPCGKernelAttributeDesc& AttrDesc : DataDesc.AttributeDescs)
 			{
@@ -149,7 +144,7 @@ bool FPCGDataUploadAdaptor::PrepareData_GameThread(const UPCGComputeKernel* InPr
 	}
 	else
 	{
-		const uint64 BufferSize = TargetDataCollectionDesc.ComputePackedSizeBytes();
+		const uint64 BufferSize = TargetDataCollectionDesc->ComputePackedSizeBytes();
 		if (PCGComputeHelpers::IsBufferSizeTooLarge(BufferSize))
 		{
 			return true;
@@ -193,7 +188,7 @@ bool FPCGDataUploadAdaptor::PrepareData_GameThread(const UPCGComputeKernel* InPr
 		PCGComputeHelpers::NotifyCPUToGPUUpload(Binding, InProducerKernel, InProducerSettings);
 #endif
 
-		TargetDataCollectionDesc.PackDataCollection(Binding->InputDataCollection, InputPinLabel, Binding, PackedDataCollection);
+		TargetDataCollectionDesc->PackDataCollection(Binding->InputDataCollection, InputPinLabel, Binding, PackedDataCollection);
 	}
 
 	return true;
@@ -252,9 +247,9 @@ FRDGBufferSRVRef FPCGDataUploadAdaptor::GetBufferToGraphStringKeySRV(FRDGBuilder
 FRDGBufferSRVRef FPCGDataUploadAdaptor::GetDataElementCountsBufferSRV(FRDGBuilder& InGraphBuilder)
 {
 	TArray<uint32> DataElementCounts;
-	DataElementCounts.Reserve(TargetDataCollectionDesc.DataDescs.Num());
+	DataElementCounts.Reserve(TargetDataCollectionDesc->DataDescs.Num());
 
-	for (const FPCGDataDesc& DataDesc : TargetDataCollectionDesc.DataDescs)
+	for (const FPCGDataDesc& DataDesc : TargetDataCollectionDesc->DataDescs)
 	{
 		ensure(DataDesc.GetElementDimension() == EPCGElementDimension::One);
 		DataElementCounts.Add(DataDesc.GetElementCount().X);
diff --git a/Engine/Plugins/PCG/Source/PCG/Private/Compute/PCGDataUploadAdaptor.h b/Engine/Plugins/PCG/Source/PCG/Private/Compute/PCGDataUploadAdaptor.h
--- a/Engine/Plugins/PCG/Source/PCG/Private/Compute/PCGDataUploadAdaptor.h
+++ b/Engine/Plugins/PCG/Source/PCG/Private/Compute/PCGDataUploadAdaptor.h
@@ -17,7 +17,7 @@ enum class EPCGExportMode : uint8;
 class FPCGDataUploadAdaptor
 {
 public:
-	FPCGDataUploadAdaptor(UPCGDataBinding* InDataBinding, const FPCGDataCollectionDesc* InTargetDataCollectionDesc, FName InInputPinLabel);
+	FPCGDataUploadAdaptor(UPCGDataBinding* InDataBinding, const TSharedPtr<const FPCGDataCollectionDesc> InTargetDataCollectionDesc, FName InInputPinLabel);
 
 	/** Do any preparation work such as data readbacks. Returns true when preparation is complete. */
 	bool PrepareData_GameThread(const UPCGComputeKernel* InProducerKernel, const UPCGSettings* InProducerSettings);
@@ -45,7 +45,7 @@ private:
 
 	TWeakObjectPtr<UPCGDataBinding> DataBinding;
 
-	FPCGDataCollectionDesc TargetDataCollectionDesc;
+	TSharedPtr<const FPCGDataCollectionDesc> TargetDataCollectionDesc;
 
 	/** If data cannot be reused from GPU, will be packed into this buffer prior to upload. */
 	TArray<uint32> PackedDataCollection;
diff --git a/Engine/Plugins/PCG/Source/PCG/Private/Elements/Metadata/PCGMetadataPartitionKernel.cpp b/Engine/Plugins/PCG/Source/PCG/Private/Elements/Metadata/PCGMetadataPartitionKernel.cpp
--- a/Engine/Plugins/PCG/Source/PCG/Private/Elements/Metadata/PCGMetadataPartitionKernel.cpp
+++ b/Engine/Plugins/PCG/Source/PCG/Private/Elements/Metadata/PCGMetadataPartitionKernel.cpp
@@ -37,7 +37,7 @@ bool UPCGMetadataPartitionKernel::IsKernelDataValid(FPCGContext* InContext) cons
 		const UPCGMetadataPartitionSettings* MPSettings = CastChecked<UPCGMetadataPartitionSettings>(GetSettings());
 		const FName AttributeName = MPSettings->PartitionAttributeSelectors[0].GetAttributeName();
 
-		const FPCGDataCollectionDesc* InputDesc = Context->DataBinding->GetCachedKernelPinDataDesc(this, PCGPinConstants::DefaultInputLabel, /*bIsInput=*/true);
+		const TSharedPtr<const FPCGDataCollectionDesc> InputDesc = Context->DataBinding->GetCachedKernelPinDataDesc(this, PCGPinConstants::DefaultInputLabel, /*bIsInput=*/true);
 
 		if (!ensure(InputDesc))
 		{
@@ -81,7 +81,7 @@ bool UPCGMetadataPartitionKernel::IsKernelDataValid(FPCGContext* InContext) cons
 	return true;
 }
 
-bool UPCGMetadataPartitionKernel::ComputeOutputBindingDataDesc(const UPCGComputeGraph* InGraph, FName InOutputPinLabel, UPCGDataBinding* InBinding, FPCGDataCollectionDesc& OutDataDesc) const
+TSharedPtr<const FPCGDataCollectionDesc> UPCGMetadataPartitionKernel::ComputeOutputBindingDataDesc(const UPCGComputeGraph* InGraph, FName InOutputPinLabel, UPCGDataBinding* InBinding) const
 {
 	check(InGraph);
 	check(InBinding);
@@ -91,63 +91,46 @@ bool UPCGMetadataPartitionKernel::ComputeOutputBindingDataDesc(const UPCGCompute
 	// Code assumes single output pin.
 	if (!ensure(InOutputPinLabel == PCGPinConstants::DefaultOutputLabel))
 	{
-		return false;
+		return nullptr;
 	}
 
-	FPCGDataCollectionDesc InputPinDesc;
 	const FPCGKernelPin InputKernelPin(KernelIndex, PCGPinConstants::DefaultInputLabel, /*bIsInput=*/true);
-	ensure(InBinding->ComputeKernelPinDataDesc(InputKernelPin, &InputPinDesc));
+	const TSharedPtr<const FPCGDataCollectionDesc> InputPinDesc = InBinding->ComputeKernelPinDataDesc(InputKernelPin);
 
-	if (InputPinDesc.DataDescs.IsEmpty())
+	if (!ensure(InputPinDesc))
 	{
-		UE_LOG(LogPCG, Verbose, TEXT("UPCGMetadataPartitionKernel: No analysis data found from count kernel."));
-		return true;
+		return nullptr;
 	}
 
-	const TArray<int32> AnalysisDataIndices = InBinding->GetPinInputDataIndices(this, PCGMetadataPartitionConstants::ElementCountsPinLabel);
-
-	if (AnalysisDataIndices.Num() != InputPinDesc.DataDescs.Num())
+	if (InputPinDesc->DataDescs.IsEmpty())
 	{
-		UE_LOG(LogPCG, Error, TEXT("UPCGMetadataPartitionKernel: Number of received analysis data items (%d) does not match input data count (%d)."),
-			AnalysisDataIndices.Num(), InputPinDesc.DataDescs.Num());
-		return true;
+		UE_LOG(LogPCG, Verbose, TEXT("UPCGMetadataPartitionKernel: No analysis data found from count kernel."));
+		return nullptr;
 	}
 
-	if (AnalysisDataIndices.IsEmpty())
-	{
-		return true;
-	}
+	TSharedPtr<FPCGDataCollectionDesc> OutDataDesc = FPCGDataCollectionDesc::MakeSharedFrom(InputPinDesc);
 
-	// Expect analysis data produced for each input data item (indexing below depends on it).
-	if (!ensure(AnalysisDataIndices.Num() == InputPinDesc.DataDescs.Num()))
-	{
-		return true;
-	}
-
-	bool bAllReadBack = true;
+	const TArray<int32> AnalysisDataIndices = InBinding->GetPinInputDataIndices(this, PCGMetadataPartitionConstants::ElementCountsPinLabel);
 
-	for (int32 AnalysisDataIndex : AnalysisDataIndices)
+	if (AnalysisDataIndices.Num() != InputPinDesc->DataDescs.Num())
 	{
-		// Readback analysis data - poll until readback complete (return true).
-		if (!InBinding->ReadbackInputDataToCPU(AnalysisDataIndex))
-		{
-			bAllReadBack = false;
-		}
+		UE_LOG(LogPCG, Error, TEXT("UPCGMetadataPartitionKernel: Number of received analysis data items (%d) does not match input data count (%d)."),
+			AnalysisDataIndices.Num(), InputPinDesc->DataDescs.Num());
+		return nullptr;
 	}
 
-	if (!bAllReadBack)
+	if (AnalysisDataIndices.IsEmpty())
 	{
-		return false;
+		return nullptr;
 	}
 
-	FPCGDataCollectionDesc PinDesc;
-
 	for (int32 DataIndex = 0; DataIndex < AnalysisDataIndices.Num(); ++DataIndex)
 	{
 		const int32 AnalysisDataIndex = AnalysisDataIndices[DataIndex];
 		
 		check(InBinding->InputDataCollection.TaggedData.IsValidIndex(AnalysisDataIndex));
 
+		// We can expect the analysis results data to already have been readback from the GPU because data description caching runs after the pre-execute readback stage of the ComputeGraph execution.
 		const UPCGParamData* AnalysisResultsData = Cast<UPCGParamData>(InBinding->InputDataCollection.TaggedData[AnalysisDataIndex].Data);
 		const UPCGMetadata* AnalysisMetadata = AnalysisResultsData ? AnalysisResultsData->ConstMetadata() : nullptr;
 
@@ -178,7 +161,7 @@ bool UPCGMetadataPartitionKernel::ComputeOutputBindingDataDesc(const UPCGCompute
 				UE_LOG(LogPCG, Error, TEXT("Data was not read back."));
 			}
 
-			return true;
+			return nullptr;
 		}
 
 		const int32 PartitionAttributeId = InBinding->GetAttributeId(MPSettings->PartitionAttributeSelectors[0].GetAttributeName(), EPCGKernelAttributeType::StringKey);
@@ -187,9 +170,9 @@ bool UPCGMetadataPartitionKernel::ComputeOutputBindingDataDesc(const UPCGCompute
 		{
 			if (AttributeValueToCount[Value] > 0)
 			{
-				FPCGDataDesc& DataDesc = PinDesc.DataDescs.Emplace_GetRef(EPCGDataType::Point, AttributeValueToCount[Value]);
-				DataDesc.AttributeDescs.Append(InputPinDesc.DataDescs[DataIndex].AttributeDescs);
-				DataDesc.AllocateProperties(InputPinDesc.DataDescs[DataIndex].GetAllocatedProperties());
+				FPCGDataDesc& DataDesc = OutDataDesc->DataDescs.Emplace_GetRef(EPCGDataType::Point, AttributeValueToCount[Value]);
+				DataDesc.AttributeDescs.Append(InputPinDesc->DataDescs[DataIndex].AttributeDescs);
+				DataDesc.AllocateProperties(InputPinDesc->DataDescs[DataIndex].GetAllocatedProperties());
 
 				for (FPCGKernelAttributeDesc& AttributeDesc : DataDesc.AttributeDescs)
 				{
@@ -207,17 +190,15 @@ bool UPCGMetadataPartitionKernel::ComputeOutputBindingDataDesc(const UPCGCompute
 	// Add output attribute (partition index).
 	if (MPSettings->bAssignIndexPartition)
 	{
-		PinDesc.AddAttributeToAllData(FPCGKernelAttributeKey(MPSettings->PartitionIndexAttributeName, EPCGKernelAttributeType::Int), InBinding);
+		OutDataDesc->AddAttributeToAllData(FPCGKernelAttributeKey(MPSettings->PartitionIndexAttributeName, EPCGKernelAttributeType::Int), InBinding);
 	}
 
-	OutDataDesc = MoveTemp(PinDesc);
-
-	return true;
+	return OutDataDesc;
 }
 
 int UPCGMetadataPartitionKernel::ComputeThreadCount(const UPCGDataBinding* InBinding) const
 {
-	const FPCGDataCollectionDesc* OutputPinDesc = InBinding->GetCachedKernelPinDataDesc(this, PCGPinConstants::DefaultOutputLabel, /*bIsInput=*/false);
+	const TSharedPtr<const FPCGDataCollectionDesc> OutputPinDesc = InBinding->GetCachedKernelPinDataDesc(this, PCGPinConstants::DefaultOutputLabel, /*bIsInput=*/false);
 	return ensure(OutputPinDesc) ? OutputPinDesc->ComputeTotalElementCount() : 0;
 }
 
diff --git a/Engine/Plugins/PCG/Source/PCG/Private/Elements/PCGCopyPointsAnalysisKernel.cpp b/Engine/Plugins/PCG/Source/PCG/Private/Elements/PCGCopyPointsAnalysisKernel.cpp
--- a/Engine/Plugins/PCG/Source/PCG/Private/Elements/PCGCopyPointsAnalysisKernel.cpp
+++ b/Engine/Plugins/PCG/Source/PCG/Private/Elements/PCGCopyPointsAnalysisKernel.cpp
@@ -38,7 +38,7 @@ bool UPCGCopyPointsAnalysisKernel::IsKernelDataValid(FPCGContext* InContext) con
 	return PCGCopyPointsKernel::IsKernelDataValid(this, static_cast<FPCGComputeGraphContext*>(InContext));
 }
 
-bool UPCGCopyPointsAnalysisKernel::ComputeOutputBindingDataDesc(const UPCGComputeGraph* InGraph, FName InOutputPinLabel, UPCGDataBinding* InBinding, FPCGDataCollectionDesc& OutDataDesc) const
+TSharedPtr<const FPCGDataCollectionDesc> UPCGCopyPointsAnalysisKernel::ComputeOutputBindingDataDesc(const UPCGComputeGraph* InGraph, FName InOutputPinLabel, UPCGDataBinding* InBinding) const
 {
 	check(InGraph);
 	check(InBinding);
@@ -47,46 +47,47 @@ bool UPCGCopyPointsAnalysisKernel::ComputeOutputBindingDataDesc(const UPCGComput
 
 	if (!ensure(KernelParams))
 	{
-		return false;
+		return nullptr;
 	}
 
 	// Code assumes single output pin.
 	if (!ensure(InOutputPinLabel == PCGPinConstants::DefaultOutputLabel))
 	{
-		return false;
+		return nullptr;
 	}
 
-	FPCGDataCollectionDesc SourcesDesc, TargetsDesc;
-
 	const FPCGKernelPin SourceKernelPin(KernelIndex, PCGCopyPointsConstants::SourcePointsLabel, /*bIsInput=*/true);
 	const FPCGKernelPin TargetKernelPin(KernelIndex, PCGCopyPointsConstants::TargetPointsLabel, /*bIsInput=*/true);
 
-	ensure(InBinding->ComputeKernelPinDataDesc(SourceKernelPin, &SourcesDesc));
-	ensure(InBinding->ComputeKernelPinDataDesc(TargetKernelPin, &TargetsDesc));
+	const TSharedPtr<const FPCGDataCollectionDesc> SourcesDesc = InBinding->ComputeKernelPinDataDesc(SourceKernelPin);
+	const TSharedPtr<const FPCGDataCollectionDesc> TargetsDesc = InBinding->ComputeKernelPinDataDesc(TargetKernelPin);
 
-	const uint32 NumSourceDatas = SourcesDesc.DataDescs.Num();
-	const uint32 NumTargetDatas = TargetsDesc.DataDescs.Num();
+	if (!ensure(SourcesDesc && TargetsDesc))
+	{
+		return nullptr;
+	}
 
-	// Output: a single attribute set with a single attribute true/false value per NumSource*NumTarget possible outputs.
-	FPCGDataCollectionDesc OutDesc;
+	TSharedPtr<FPCGDataCollectionDesc> OutDataDesc = FPCGDataCollectionDesc::MakeShared();
+
+	const uint32 NumSourceDatas = SourcesDesc->DataDescs.Num();
+	const uint32 NumTargetDatas = TargetsDesc->DataDescs.Num();
 
+	// Output: a single attribute set with a single attribute true/false value per NumSource*NumTarget possible outputs.
 	const bool bCopyEachSourceOnEveryTarget = KernelParams->GetValueBool(GET_MEMBER_NAME_CHECKED(UPCGCopyPointsSettings, bCopyEachSourceOnEveryTarget));
 
 	// Validation should have ensured already that these are valid: (S, T) = (N, N), (N, 1) or (1, N).
 	const int32 ElementCount = bCopyEachSourceOnEveryTarget ? (NumSourceDatas * NumTargetDatas) : FMath::Max(NumSourceDatas, NumTargetDatas);
-	FPCGDataDesc& DataDesc = OutDesc.DataDescs.Emplace_GetRef(EPCGDataType::Param, ElementCount);
+	FPCGDataDesc& DataDesc = OutDataDesc->DataDescs.Emplace_GetRef(EPCGDataType::Param, ElementCount);
 
 	FPCGKernelAttributeKey SelectedFlagAttributeKey = FPCGKernelAttributeKey(PCGCopyPointsConstants::SelectedFlagAttributeName, EPCGKernelAttributeType::Bool);
 	DataDesc.AddAttribute(SelectedFlagAttributeKey, InBinding);
 
-	OutDataDesc = MoveTemp(OutDesc);
-
-	return true;
+	return OutDataDesc;
 }
 
 int UPCGCopyPointsAnalysisKernel::ComputeThreadCount(const UPCGDataBinding* InBinding) const
 {
-	const FPCGDataCollectionDesc* OutputPinDesc = InBinding->GetCachedKernelPinDataDesc(this, PCGPinConstants::DefaultOutputLabel, /*bIsInput=*/false);
+	const TSharedPtr<const FPCGDataCollectionDesc> OutputPinDesc = InBinding->GetCachedKernelPinDataDesc(this, PCGPinConstants::DefaultOutputLabel, /*bIsInput=*/false);
 	return ensure(OutputPinDesc) ? OutputPinDesc->ComputeTotalElementCount() : 0;
 }
 
diff --git a/Engine/Plugins/PCG/Source/PCG/Private/Elements/PCGCopyPointsKernel.cpp b/Engine/Plugins/PCG/Source/PCG/Private/Elements/PCGCopyPointsKernel.cpp
--- a/Engine/Plugins/PCG/Source/PCG/Private/Elements/PCGCopyPointsKernel.cpp
+++ b/Engine/Plugins/PCG/Source/PCG/Private/Elements/PCGCopyPointsKernel.cpp
@@ -39,7 +39,7 @@ bool UPCGCopyPointsKernel::IsKernelDataValid(FPCGContext* InContext) const
 	return PCGCopyPointsKernel::IsKernelDataValid(this, static_cast<FPCGComputeGraphContext*>(InContext));
 }
 
-bool UPCGCopyPointsKernel::ComputeOutputBindingDataDesc(const UPCGComputeGraph* InGraph, FName InOutputPinLabel, UPCGDataBinding* InBinding, FPCGDataCollectionDesc& OutDataDesc) const
+TSharedPtr<const FPCGDataCollectionDesc> UPCGCopyPointsKernel::ComputeOutputBindingDataDesc(const UPCGComputeGraph* InGraph, FName InOutputPinLabel, UPCGDataBinding* InBinding) const
 {
 	check(InGraph);
 	check(InBinding);
@@ -48,13 +48,13 @@ bool UPCGCopyPointsKernel::ComputeOutputBindingDataDesc(const UPCGComputeGraph*
 
 	if (!ensure(KernelParams))
 	{
-		return false;
+		return nullptr;
 	}
 
 	// Code assumes single output pin.
 	if (!ensure(InOutputPinLabel == PCGPinConstants::DefaultOutputLabel))
 	{
-		return false;
+		return nullptr;
 	}
 
 	// A graph split was injected before this kernel. We expect to find selected flags
@@ -94,30 +94,32 @@ bool UPCGCopyPointsKernel::ComputeOutputBindingDataDesc(const UPCGComputeGraph*
 				{
 					UE_LOG(LogPCG, Error, TEXT("Data was not read back."));
 				}
-				OutDataDesc = {};
-				return true;
+
+				return nullptr;
 			}
 		}
 	}
 
-	FPCGDataCollectionDesc SourcePinDesc;
-	FPCGDataCollectionDesc TargetPinDesc;
-
 	const FPCGKernelPin SourceKernelPin(KernelIndex, PCGCopyPointsConstants::SourcePointsLabel, /*bIsInput=*/true);
 	const FPCGKernelPin TargetKernelPin(KernelIndex, PCGCopyPointsConstants::TargetPointsLabel, /*bIsInput=*/true);
 
-	ensure(InBinding->ComputeKernelPinDataDesc(SourceKernelPin, &SourcePinDesc));
-	ensure(InBinding->ComputeKernelPinDataDesc(TargetKernelPin, &TargetPinDesc));
+	const TSharedPtr<const FPCGDataCollectionDesc> SourcePinDesc = InBinding->ComputeKernelPinDataDesc(SourceKernelPin);
+	const TSharedPtr<const FPCGDataCollectionDesc> TargetPinDesc = InBinding->ComputeKernelPinDataDesc(TargetKernelPin);
+
+	if (!ensure(SourcePinDesc && TargetPinDesc))
+	{
+		return nullptr;
+	}
 
-	FPCGDataCollectionDesc PinDesc;
+	TSharedPtr<FPCGDataCollectionDesc> OutDataDesc = FPCGDataCollectionDesc::MakeShared();
 
 	const bool bCopyEachSourceOnEveryTarget = KernelParams->GetValueBool(GET_MEMBER_NAME_CHECKED(UPCGCopyPointsSettings, bCopyEachSourceOnEveryTarget));
 	const EPCGCopyPointsInheritanceMode ColorInheritance = KernelParams->GetValueEnum<EPCGCopyPointsInheritanceMode>(GET_MEMBER_NAME_CHECKED(UPCGCopyPointsSettings, ColorInheritance));
 	const EPCGCopyPointsInheritanceMode SeedInheritance = KernelParams->GetValueEnum<EPCGCopyPointsInheritanceMode>(GET_MEMBER_NAME_CHECKED(UPCGCopyPointsSettings, SeedInheritance));
 	const EPCGCopyPointsMetadataInheritanceMode AttributeInheritance = KernelParams->GetValueEnum<EPCGCopyPointsMetadataInheritanceMode>(GET_MEMBER_NAME_CHECKED(UPCGCopyPointsSettings, AttributeInheritance));
 	const EPCGCopyPointsTagInheritanceMode TagInheritance = KernelParams->GetValueEnum<EPCGCopyPointsTagInheritanceMode>(GET_MEMBER_NAME_CHECKED(UPCGCopyPointsSettings, TagInheritance));
-	const int32 NumSources = SourcePinDesc.DataDescs.Num();
-	const int32 NumTargets = TargetPinDesc.DataDescs.Num();
+	const int32 NumSources = SourcePinDesc->DataDescs.Num();
+	const int32 NumTargets = TargetPinDesc->DataDescs.Num();
 	const int32 NumIterations = bCopyEachSourceOnEveryTarget ? NumSources * NumTargets : FMath::Max(NumSources, NumTargets);
 
 	// Properties that might need to be allocated
@@ -135,10 +137,10 @@ bool UPCGCopyPointsKernel::ComputeOutputBindingDataDesc(const UPCGComputeGraph*
 			const int32 SourceIndex = bCopyEachSourceOnEveryTarget ? (I / NumTargets) : FMath::Min(I, NumSources - 1);
 			const int32 TargetIndex = bCopyEachSourceOnEveryTarget ? (I % NumTargets) : FMath::Min(I, NumTargets - 1);
 
-			const FPCGDataDesc& SourceDesc = SourcePinDesc.DataDescs[SourceIndex];
-			const FPCGDataDesc& TargetDesc = TargetPinDesc.DataDescs[TargetIndex];
+			const FPCGDataDesc& SourceDesc = SourcePinDesc->DataDescs[SourceIndex];
+			const FPCGDataDesc& TargetDesc = TargetPinDesc->DataDescs[TargetIndex];
 
-			FPCGDataDesc& ResultDataDesc = PinDesc.DataDescs.Emplace_GetRef(EPCGDataType::Point, SourceDesc.GetElementCount() * TargetDesc.GetElementCount());
+			FPCGDataDesc& ResultDataDesc = OutDataDesc->DataDescs.Emplace_GetRef(EPCGDataType::Point, SourceDesc.GetElementCount() * TargetDesc.GetElementCount());
 
 			if (AttributeInheritance == EPCGCopyPointsMetadataInheritanceMode::SourceFirst)
 			{
@@ -213,13 +215,12 @@ bool UPCGCopyPointsKernel::ComputeOutputBindingDataDesc(const UPCGComputeGraph*
 		}
 	}
 
-	OutDataDesc = MoveTemp(PinDesc);
-	return true;
+	return OutDataDesc;
 }
 
 int UPCGCopyPointsKernel::ComputeThreadCount(const UPCGDataBinding* InBinding) const
 {
-	const FPCGDataCollectionDesc* OutputPinDesc = InBinding->GetCachedKernelPinDataDesc(this, PCGPinConstants::DefaultOutputLabel, /*bIsInput=*/false);
+	const TSharedPtr<const FPCGDataCollectionDesc> OutputPinDesc = InBinding->GetCachedKernelPinDataDesc(this, PCGPinConstants::DefaultOutputLabel, /*bIsInput=*/false);
 	return ensure(OutputPinDesc) ? OutputPinDesc->ComputeTotalElementCount() : 0;
 }
 
diff --git a/Engine/Plugins/PCG/Source/PCG/Private/Elements/PCGCopyPointsKernelShared.cpp b/Engine/Plugins/PCG/Source/PCG/Private/Elements/PCGCopyPointsKernelShared.cpp
--- a/Engine/Plugins/PCG/Source/PCG/Private/Elements/PCGCopyPointsKernelShared.cpp
+++ b/Engine/Plugins/PCG/Source/PCG/Private/Elements/PCGCopyPointsKernelShared.cpp
@@ -18,14 +18,14 @@ bool PCGCopyPointsKernel::IsKernelDataValid(const UPCGComputeKernel* InKernel, c
 
 	if (const UPCGDataBinding* DataBinding = InContext->DataBinding.Get(); ensure(DataBinding))
 	{
-		const FPCGDataCollectionDesc* SourcePinDesc = DataBinding->GetCachedKernelPinDataDesc(InKernel, PCGCopyPointsConstants::SourcePointsLabel, /*bIsInputPin=*/true);
+		const TSharedPtr<const FPCGDataCollectionDesc> SourcePinDesc = DataBinding->GetCachedKernelPinDataDesc(InKernel, PCGCopyPointsConstants::SourcePointsLabel, /*bIsInputPin=*/true);
 
 		if (!ensure(SourcePinDesc))
 		{
 			return false;
 		}
 
-		const FPCGDataCollectionDesc* TargetPinDesc = DataBinding->GetCachedKernelPinDataDesc(InKernel, PCGCopyPointsConstants::TargetPointsLabel, /*bIsInputPin=*/true);
+		const TSharedPtr<const FPCGDataCollectionDesc> TargetPinDesc = DataBinding->GetCachedKernelPinDataDesc(InKernel, PCGCopyPointsConstants::TargetPointsLabel, /*bIsInputPin=*/true);
 
 		if (!ensure(TargetPinDesc))
 		{
@@ -57,7 +57,7 @@ bool PCGCopyPointsKernel::IsKernelDataValid(const UPCGComputeKernel* InKernel, c
 
 		if (bMatchBasedOnAttribute)
 		{
-			auto ValidateAttributeExists = [InContext, Settings, MatchAttribute](const FPCGDataCollectionDesc* PinDataDesc) -> bool
+			auto ValidateAttributeExists = [InContext, Settings, MatchAttribute](const TSharedPtr<const FPCGDataCollectionDesc> PinDataDesc) -> bool
 			{
 				const FName MatchAttributeName = MatchAttribute;
 
diff --git a/Engine/Plugins/PCG/Source/PCG/Private/Elements/PCGSkinnedMeshSpawnerKernel.cpp b/Engine/Plugins/PCG/Source/PCG/Private/Elements/PCGSkinnedMeshSpawnerKernel.cpp
--- a/Engine/Plugins/PCG/Source/PCG/Private/Elements/PCGSkinnedMeshSpawnerKernel.cpp
+++ b/Engine/Plugins/PCG/Source/PCG/Private/Elements/PCGSkinnedMeshSpawnerKernel.cpp
@@ -31,7 +31,7 @@ namespace PCGSkinnedMeshSpawnerKernel
 	static const FText NoMeshEntriesFormat = LOCTEXT("NoMeshEntries", "No mesh entries provided in weighted mesh selector.");
 }
 
-bool UPCGSkinnedMeshSpawnerKernel::ComputeOutputBindingDataDesc(const UPCGComputeGraph* InGraph, FName InOutputPinLabel, UPCGDataBinding* InBinding, FPCGDataCollectionDesc& OutDataDesc) const
+TSharedPtr<const FPCGDataCollectionDesc> UPCGSkinnedMeshSpawnerKernel::ComputeOutputBindingDataDesc(const UPCGComputeGraph* InGraph, FName InOutputPinLabel, UPCGDataBinding* InBinding) const
 {
 	check(InGraph);
 	check(InBinding);
@@ -39,27 +39,27 @@ bool UPCGSkinnedMeshSpawnerKernel::ComputeOutputBindingDataDesc(const UPCGComput
 	// Code assumes single output pin.
 	if (!ensure(InOutputPinLabel == PCGPinConstants::DefaultOutputLabel))
 	{
-		return false;
+		return nullptr;
 	}
 
 	// Forward data from In to Out.
-	FPCGDataCollectionDesc PinDesc;
 	const FPCGKernelPin InputKernelPin(KernelIndex, PCGPinConstants::DefaultInputLabel, /*bIsInput=*/true);
-	ensure(InBinding->ComputeKernelPinDataDesc(InputKernelPin, &PinDesc));
+	const TSharedPtr<const FPCGDataCollectionDesc> InputDesc = InBinding->ComputeKernelPinDataDesc(InputKernelPin);
+
+	TSharedPtr<FPCGDataCollectionDesc> OutDataDesc = FPCGDataCollectionDesc::MakeSharedFrom(InputDesc);
 
 	const UPCGSkinnedMeshSpawnerSettings* SMSettings = CastChecked<UPCGSkinnedMeshSpawnerSettings>(GetSettings());
 	if (SMSettings->bApplyMeshBoundsToPoints)
 	{
-		PinDesc.AllocatePropertiesForAllData(EPCGPointNativeProperties::BoundsMin | EPCGPointNativeProperties::BoundsMax);
+		OutDataDesc->AllocatePropertiesForAllData(EPCGPointNativeProperties::BoundsMin | EPCGPointNativeProperties::BoundsMax);
 	}
 
-	OutDataDesc = MoveTemp(PinDesc);
-	return true;
+	return OutDataDesc;
 }
 
 int UPCGSkinnedMeshSpawnerKernel::ComputeThreadCount(const UPCGDataBinding* InBinding) const
 {
-	const FPCGDataCollectionDesc* InputPinDesc = InBinding->GetCachedKernelPinDataDesc(this, PCGPinConstants::DefaultInputLabel, /*bIsInput=*/true);
+	const TSharedPtr<const FPCGDataCollectionDesc> InputPinDesc = InBinding->GetCachedKernelPinDataDesc(this, PCGPinConstants::DefaultInputLabel, /*bIsInput=*/true);
 	return ensure(InputPinDesc) ? InputPinDesc->ComputeTotalElementCount() : 0;
 }
 
diff --git a/Engine/Plugins/PCG/Source/PCG/Private/Elements/PCGStaticMeshSpawnerKernel.cpp b/Engine/Plugins/PCG/Source/PCG/Private/Elements/PCGStaticMeshSpawnerKernel.cpp
--- a/Engine/Plugins/PCG/Source/PCG/Private/Elements/PCGStaticMeshSpawnerKernel.cpp
+++ b/Engine/Plugins/PCG/Source/PCG/Private/Elements/PCGStaticMeshSpawnerKernel.cpp
@@ -33,7 +33,7 @@ namespace PCGStaticMeshSpawnerKernel
 	static const FText NoMeshEntriesFormat = LOCTEXT("NoMeshEntries", "No mesh entries provided in weighted mesh selector.");
 }
 
-bool UPCGStaticMeshSpawnerKernel::ComputeOutputBindingDataDesc(const UPCGComputeGraph* InGraph, FName InOutputPinLabel, UPCGDataBinding* InBinding, FPCGDataCollectionDesc& OutDataDesc) const
+TSharedPtr<const FPCGDataCollectionDesc> UPCGStaticMeshSpawnerKernel::ComputeOutputBindingDataDesc(const UPCGComputeGraph* InGraph, FName InOutputPinLabel, UPCGDataBinding* InBinding) const
 {
 	check(InGraph);
 	check(InBinding);
@@ -41,13 +41,14 @@ bool UPCGStaticMeshSpawnerKernel::ComputeOutputBindingDataDesc(const UPCGCompute
 	// Code assumes single output pin.
 	if (!ensure(InOutputPinLabel == PCGPinConstants::DefaultOutputLabel))
 	{
-		return false;
+		return nullptr;
 	}
 
 	// Forward data from In to Out.
-	FPCGDataCollectionDesc PinDesc;
 	const FPCGKernelPin InputKernelPin(KernelIndex, PCGPinConstants::DefaultInputLabel, /*bIsInput=*/true);
-	ensure(InBinding->ComputeKernelPinDataDesc(InputKernelPin, &PinDesc));
+	const TSharedPtr<const FPCGDataCollectionDesc> InputDesc = InBinding->ComputeKernelPinDataDesc(InputKernelPin);
+
+	TSharedPtr<FPCGDataCollectionDesc> OutDataDesc = FPCGDataCollectionDesc::MakeSharedFrom(InputDesc);
 
 	// Add output attribute (selected mesh).
 	{
@@ -69,7 +70,7 @@ bool UPCGStaticMeshSpawnerKernel::ComputeOutputBindingDataDesc(const UPCGCompute
 				}
 			}
 
-			PinDesc.AddAttributeToAllData(FPCGKernelAttributeKey(SMSettings->OutAttributeName, EPCGKernelAttributeType::StringKey), InBinding, &UniqueStringKeys);
+			OutDataDesc->AddAttributeToAllData(FPCGKernelAttributeKey(SMSettings->OutAttributeName, EPCGKernelAttributeType::StringKey), InBinding, &UniqueStringKeys);
 		}
 		else if (UPCGMeshSelectorByAttribute* SelectorByAttribute = Cast<UPCGMeshSelectorByAttribute>(SMSettings->MeshSelectorParameters))
 		{
@@ -78,7 +79,7 @@ bool UPCGStaticMeshSpawnerKernel::ComputeOutputBindingDataDesc(const UPCGCompute
 			{
 				TArray<int32> StringKeys;
 
-				for (FPCGDataDesc& DataDesc : PinDesc.DataDescs)
+				for (FPCGDataDesc& DataDesc : OutDataDesc->DataDescs)
 				{
 					StringKeys.Empty();
 
@@ -108,17 +109,16 @@ bool UPCGStaticMeshSpawnerKernel::ComputeOutputBindingDataDesc(const UPCGCompute
 		// Allocate bounds if the spawner will write more than one unique value for the mesh bounds.
 		if (SMSettings->bApplyMeshBoundsToPoints && UniqueStringKeys.Num() > 1)
 		{
-			PinDesc.AllocatePropertiesForAllData(EPCGPointNativeProperties::BoundsMin | EPCGPointNativeProperties::BoundsMax);
+			OutDataDesc->AllocatePropertiesForAllData(EPCGPointNativeProperties::BoundsMin | EPCGPointNativeProperties::BoundsMax);
 		}
 	}
 
-	OutDataDesc = MoveTemp(PinDesc);
-	return true;
+	return OutDataDesc;
 }
 
 int UPCGStaticMeshSpawnerKernel::ComputeThreadCount(const UPCGDataBinding* InBinding) const
 {
-	const FPCGDataCollectionDesc* InputPinDesc = InBinding->GetCachedKernelPinDataDesc(this, PCGPinConstants::DefaultInputLabel, /*bIsInput=*/true);
+	const TSharedPtr<const FPCGDataCollectionDesc> InputPinDesc = InBinding->GetCachedKernelPinDataDesc(this, PCGPinConstants::DefaultInputLabel, /*bIsInput=*/true);
 	return ensure(InputPinDesc) ? InputPinDesc->ComputeTotalElementCount() : 0;
 }
 
diff --git a/Engine/Plugins/PCG/Source/PCG/Public/Compute/DataInterfaces/PCGComputeDataInterface.h b/Engine/Plugins/PCG/Source/PCG/Public/Compute/DataInterfaces/PCGComputeDataInterface.h
--- a/Engine/Plugins/PCG/Source/PCG/Public/Compute/DataInterfaces/PCGComputeDataInterface.h
+++ b/Engine/Plugins/PCG/Source/PCG/Public/Compute/DataInterfaces/PCGComputeDataInterface.h
@@ -149,7 +149,7 @@ public:
 
 	EPCGExportMode GetExportMode() const { return ExportMode; }
 	UPCGDataBinding* GetDataBinding() const { return Binding.Get(); }
-	const FPCGDataCollectionDesc* GetPinDescription() const { return PinDesc; }
+	TSharedPtr<const FPCGDataCollectionDesc> GetPinDescription() const { return PinDesc; }
 
 	DECLARE_EVENT(UPCGExportableDataProvider, FOnDataExported);
 	FOnDataExported& OnDataExported_GameThread() { return OnDataExported; }
@@ -159,7 +159,7 @@ protected:
 	FOnDataExported OnDataExported;
 
 	TWeakObjectPtr<UPCGDataBinding> Binding;
-	const FPCGDataCollectionDesc* PinDesc = nullptr;
+	TSharedPtr<const FPCGDataCollectionDesc> PinDesc = nullptr;
 };
 
 /** Interface for marshalling settings/params to the GPU. */
diff --git a/Engine/Plugins/PCG/Source/PCG/Public/Elements/Metadata/PCGMetadataPartitionKernel.h b/Engine/Plugins/PCG/Source/PCG/Public/Elements/Metadata/PCGMetadataPartitionKernel.h
--- a/Engine/Plugins/PCG/Source/PCG/Public/Elements/Metadata/PCGMetadataPartitionKernel.h
+++ b/Engine/Plugins/PCG/Source/PCG/Public/Elements/Metadata/PCGMetadataPartitionKernel.h
@@ -20,7 +20,7 @@ class UPCGMetadataPartitionKernel : public UPCGComputeKernel
 
 public:
 	virtual bool IsKernelDataValid(FPCGContext* InContext) const override;
-	virtual bool ComputeOutputBindingDataDesc(const UPCGComputeGraph* InGraph, FName InOutputPinLabel, UPCGDataBinding* InBinding, FPCGDataCollectionDesc& OutDataDesc) const override;
+	virtual TSharedPtr<const FPCGDataCollectionDesc> ComputeOutputBindingDataDesc(const UPCGComputeGraph* InGraph, FName InOutputPinLabel, UPCGDataBinding* InBinding) const override;
 	virtual int ComputeThreadCount(const UPCGDataBinding* Binding) const override;
 #if WITH_EDITOR
 	virtual const TCHAR* GetSourceFilePath() const override { return TEXT("/Plugin/PCG/Private/Elements/PCGMetadataPartition.usf"); }
diff --git a/Engine/Plugins/PCG/Source/PCG/Public/Elements/PCGCopyPointsAnalysisKernel.h b/Engine/Plugins/PCG/Source/PCG/Public/Elements/PCGCopyPointsAnalysisKernel.h
--- a/Engine/Plugins/PCG/Source/PCG/Public/Elements/PCGCopyPointsAnalysisKernel.h
+++ b/Engine/Plugins/PCG/Source/PCG/Public/Elements/PCGCopyPointsAnalysisKernel.h
@@ -15,7 +15,7 @@ class UPCGCopyPointsAnalysisKernel : public UPCGComputeKernel
 
 public:
 	virtual bool IsKernelDataValid(FPCGContext* InContext) const override;
-	virtual bool ComputeOutputBindingDataDesc(const UPCGComputeGraph* InGraph, FName InOutputPinLabel, UPCGDataBinding* InBinding, FPCGDataCollectionDesc& OutDataDesc) const override;
+	virtual TSharedPtr<const FPCGDataCollectionDesc> ComputeOutputBindingDataDesc(const UPCGComputeGraph* InGraph, FName InOutputPinLabel, UPCGDataBinding* InBinding) const override;
 	virtual int ComputeThreadCount(const UPCGDataBinding* Binding) const override;
 #if WITH_EDITOR
 	virtual const TCHAR* GetSourceFilePath() const override { return TEXT("/Plugin/PCG/Private/Elements/PCGCopyPointsAnalysis.usf"); }
diff --git a/Engine/Plugins/PCG/Source/PCG/Public/Elements/PCGCopyPointsKernel.h b/Engine/Plugins/PCG/Source/PCG/Public/Elements/PCGCopyPointsKernel.h
--- a/Engine/Plugins/PCG/Source/PCG/Public/Elements/PCGCopyPointsKernel.h
+++ b/Engine/Plugins/PCG/Source/PCG/Public/Elements/PCGCopyPointsKernel.h
@@ -18,7 +18,7 @@ class UPCGCopyPointsKernel : public UPCGComputeKernel
 public:
 	//~ Begin UPCGComputeKernel Interface
 	virtual bool IsKernelDataValid(FPCGContext* InContext) const override;
-	virtual bool ComputeOutputBindingDataDesc(const UPCGComputeGraph* InGraph, FName InOutputPinLabel, UPCGDataBinding* InBinding, FPCGDataCollectionDesc& OutDataDesc) const override;
+	virtual TSharedPtr<const FPCGDataCollectionDesc> ComputeOutputBindingDataDesc(const UPCGComputeGraph* InGraph, FName InOutputPinLabel, UPCGDataBinding* InBinding) const override;
 	virtual int ComputeThreadCount(const UPCGDataBinding* Binding) const override;
 #if WITH_EDITOR
 	virtual const TCHAR* GetSourceFilePath() const override { return TEXT("/Plugin/PCG/Private/Elements/PCGCopyPoints.usf"); }
diff --git a/Engine/Plugins/PCG/Source/PCG/Public/Elements/PCGSkinnedMeshSpawnerKernel.h b/Engine/Plugins/PCG/Source/PCG/Public/Elements/PCGSkinnedMeshSpawnerKernel.h
--- a/Engine/Plugins/PCG/Source/PCG/Public/Elements/PCGSkinnedMeshSpawnerKernel.h
+++ b/Engine/Plugins/PCG/Source/PCG/Public/Elements/PCGSkinnedMeshSpawnerKernel.h
@@ -19,7 +19,7 @@ class UPCGSkinnedMeshSpawnerKernel : public UPCGComputeKernel
 	GENERATED_BODY()
 
 public:
-	virtual bool ComputeOutputBindingDataDesc(const UPCGComputeGraph* InGraph, FName InOutputPinLabel, UPCGDataBinding* InBinding, FPCGDataCollectionDesc& OutDataDesc) const override;
+	virtual TSharedPtr<const FPCGDataCollectionDesc> ComputeOutputBindingDataDesc(const UPCGComputeGraph* InGraph, FName InOutputPinLabel, UPCGDataBinding* InBinding) const override;
 	virtual int ComputeThreadCount(const UPCGDataBinding* Binding) const override;
 #if WITH_EDITOR
 	virtual const TCHAR* GetSourceFilePath() const override { return TEXT("/Plugin/PCG/Private/Elements/PCGStaticMeshSpawner.usf"); }
diff --git a/Engine/Plugins/PCG/Source/PCG/Public/Elements/PCGStaticMeshSpawnerKernel.h b/Engine/Plugins/PCG/Source/PCG/Public/Elements/PCGStaticMeshSpawnerKernel.h
--- a/Engine/Plugins/PCG/Source/PCG/Public/Elements/PCGStaticMeshSpawnerKernel.h
+++ b/Engine/Plugins/PCG/Source/PCG/Public/Elements/PCGStaticMeshSpawnerKernel.h
@@ -19,7 +19,7 @@ class UPCGStaticMeshSpawnerKernel : public UPCGComputeKernel
 	GENERATED_BODY()
 
 public:
-	virtual bool ComputeOutputBindingDataDesc(const UPCGComputeGraph* InGraph, FName InOutputPinLabel, UPCGDataBinding* InBinding, FPCGDataCollectionDesc& OutDataDesc) const override;
+	virtual TSharedPtr<const FPCGDataCollectionDesc> ComputeOutputBindingDataDesc(const UPCGComputeGraph* InGraph, FName InOutputPinLabel, UPCGDataBinding* InBinding) const override;
 	virtual int ComputeThreadCount(const UPCGDataBinding* Binding) const override;
 #if WITH_EDITOR
 	virtual const TCHAR* GetSourceFilePath() const override { return TEXT("/Plugin/PCG/Private/Elements/PCGStaticMeshSpawner.usf"); }
