[PCG] Moving PCGEditorGraphNodeBase and PCGEditorGraphNode into public.

#rnx

diff --git a/Engine/Plugins/Experimental/ProceduralVegetationEditor/Source/ProceduralVegetationEditor/Private/PVEditor.cpp b/Engine/Plugins/Experimental/ProceduralVegetationEditor/Source/ProceduralVegetationEditor/Private/PVEditor.cpp
--- a/Engine/Plugins/Experimental/ProceduralVegetationEditor/Source/ProceduralVegetationEditor/Private/PVEditor.cpp
+++ b/Engine/Plugins/Experimental/ProceduralVegetationEditor/Source/ProceduralVegetationEditor/Private/PVEditor.cpp
@@ -12,7 +12,7 @@
 #include "ToolMenus.h"
 #include "ToolMenuSection.h"
 #include "Dataflow/DataflowCollectionSpreadSheetWidget.h"
-#include "PCGEditor/Private/Nodes/PCGEditorGraphNodeBase.h"
+#include "Nodes/PCGEditorGraphNodeBase.h"
 #include "Widgets/SPVCreateDialog.h"
 #include "Widgets/SPVExportSelectionDialog.h"
 
@@ -261,11 +261,7 @@ void FPVEditor::OnSelectedNodesChanged(const TSet<UObject*>& InNewSelection)
 	if (!InNewSelection.IsEmpty())
 	{
 		const UPCGEditorGraphNodeBase* Node = Cast<UPCGEditorGraphNodeBase>(*InNewSelection.CreateConstIterator());
-		const FObjectProperty* ObjectProperty = CastField<FObjectProperty>(UPCGEditorGraphNodeBase::StaticClass()->FindPropertyByName(TEXT("PCGNode")));
-		if (Node && ObjectProperty)
-		{
-			PCGNode = Cast<UPCGNode>(ObjectProperty->GetObjectPropertyValue_InContainer(Node));
-		}
+		PCGNode = Node->GetPCGNode();
 	}
 
 	UpdateCollectionSpreadSheet(PCGNode);
diff --git a/Engine/Plugins/PCG/Source/PCGEditor/Private/Nodes/PCGEditorGraphNode.cpp b/Engine/Plugins/PCG/Source/PCGEditor/Private/Nodes/PCGEditorGraphNode.cpp
--- a/Engine/Plugins/PCG/Source/PCGEditor/Private/Nodes/PCGEditorGraphNode.cpp
+++ b/Engine/Plugins/PCG/Source/PCGEditor/Private/Nodes/PCGEditorGraphNode.cpp
@@ -1,6 +1,6 @@
 // Copyright Epic Games, Inc. All Rights Reserved.
 
-#include "PCGEditorGraphNode.h"
+#include "Nodes/PCGEditorGraphNode.h"
 
 #include "PCGNode.h"
 #include "PCGPin.h"
diff --git a/Engine/Plugins/PCG/Source/PCGEditor/Private/Nodes/PCGEditorGraphNode.h b/Engine/Plugins/PCG/Source/PCGEditor/Private/Nodes/PCGEditorGraphNode.h
deleted file mode 100644
diff --git a/Engine/Plugins/PCG/Source/PCGEditor/Private/Nodes/PCGEditorGraphNodeBase.cpp b/Engine/Plugins/PCG/Source/PCGEditor/Private/Nodes/PCGEditorGraphNodeBase.cpp
--- a/Engine/Plugins/PCG/Source/PCGEditor/Private/Nodes/PCGEditorGraphNodeBase.cpp
+++ b/Engine/Plugins/PCG/Source/PCGEditor/Private/Nodes/PCGEditorGraphNodeBase.cpp
@@ -1,6 +1,6 @@
 // Copyright Epic Games, Inc. All Rights Reserved.
 
-#include "PCGEditorGraphNodeBase.h"
+#include "Nodes/PCGEditorGraphNodeBase.h"
 
 #include "PCGComponent.h"
 #include "PCGEdge.h"
diff --git a/Engine/Plugins/PCG/Source/PCGEditor/Private/Nodes/PCGEditorGraphNodeBase.h b/Engine/Plugins/PCG/Source/PCGEditor/Private/Nodes/PCGEditorGraphNodeBase.h
deleted file mode 100644
diff --git a/Engine/Plugins/PCG/Source/PCGEditor/Private/Nodes/PCGEditorGraphNodeFactory.cpp b/Engine/Plugins/PCG/Source/PCGEditor/Private/Nodes/PCGEditorGraphNodeFactory.cpp
--- a/Engine/Plugins/PCG/Source/PCGEditor/Private/Nodes/PCGEditorGraphNodeFactory.cpp
+++ b/Engine/Plugins/PCG/Source/PCGEditor/Private/Nodes/PCGEditorGraphNodeFactory.cpp
@@ -2,8 +2,8 @@
 
 #include "PCGEditorGraphNodeFactory.h"
 
-#include "PCGEditorGraphNodeBase.h"
 #include "SPCGEditorGraphNode.h"
+#include "Nodes/PCGEditorGraphNodeBase.h"
 
 TSharedPtr<SGraphNode> FPCGEditorGraphNodeFactory::CreateNode(UEdGraphNode* InNode) const
 {
diff --git a/Engine/Plugins/PCG/Source/PCGEditor/Private/Nodes/PCGEditorGraphNodeGetUserParameter.h b/Engine/Plugins/PCG/Source/PCGEditor/Private/Nodes/PCGEditorGraphNodeGetUserParameter.h
--- a/Engine/Plugins/PCG/Source/PCGEditor/Private/Nodes/PCGEditorGraphNodeGetUserParameter.h
+++ b/Engine/Plugins/PCG/Source/PCGEditor/Private/Nodes/PCGEditorGraphNodeGetUserParameter.h
@@ -2,9 +2,7 @@
 
 #pragma once
 
-#include "PCGEditorGraphNodeBase.h"
-
-#include "PCGEditorGraphNode.h"
+#include "Nodes/PCGEditorGraphNode.h"
 
 #include "PCGEditorGraphNodeGetUserParameter.generated.h"
 
diff --git a/Engine/Plugins/PCG/Source/PCGEditor/Private/Nodes/PCGEditorGraphNodeInput.h b/Engine/Plugins/PCG/Source/PCGEditor/Private/Nodes/PCGEditorGraphNodeInput.h
--- a/Engine/Plugins/PCG/Source/PCGEditor/Private/Nodes/PCGEditorGraphNodeInput.h
+++ b/Engine/Plugins/PCG/Source/PCGEditor/Private/Nodes/PCGEditorGraphNodeInput.h
@@ -2,7 +2,7 @@
 
 #pragma once
 
-#include "PCGEditorGraphNode.h"
+#include "Nodes/PCGEditorGraphNode.h"
 
 #include "PCGEditorGraphNodeInput.generated.h"
 
diff --git a/Engine/Plugins/PCG/Source/PCGEditor/Private/Nodes/PCGEditorGraphNodeOutput.h b/Engine/Plugins/PCG/Source/PCGEditor/Private/Nodes/PCGEditorGraphNodeOutput.h
--- a/Engine/Plugins/PCG/Source/PCGEditor/Private/Nodes/PCGEditorGraphNodeOutput.h
+++ b/Engine/Plugins/PCG/Source/PCGEditor/Private/Nodes/PCGEditorGraphNodeOutput.h
@@ -2,7 +2,7 @@
 
 #pragma once
 
-#include "PCGEditorGraphNode.h"
+#include "Nodes/PCGEditorGraphNode.h"
 
 #include "PCGEditorGraphNodeOutput.generated.h"
 
diff --git a/Engine/Plugins/PCG/Source/PCGEditor/Private/Nodes/PCGEditorGraphNodeReroute.h b/Engine/Plugins/PCG/Source/PCGEditor/Private/Nodes/PCGEditorGraphNodeReroute.h
--- a/Engine/Plugins/PCG/Source/PCGEditor/Private/Nodes/PCGEditorGraphNodeReroute.h
+++ b/Engine/Plugins/PCG/Source/PCGEditor/Private/Nodes/PCGEditorGraphNodeReroute.h
@@ -2,7 +2,7 @@
 
 #pragma once
 
-#include "PCGEditorGraphNode.h"
+#include "Nodes/PCGEditorGraphNode.h"
 
 class UPCGNode;
 
diff --git a/Engine/Plugins/PCG/Source/PCGEditor/Private/Nodes/SPCGEditorGraphNode.cpp b/Engine/Plugins/PCG/Source/PCGEditor/Private/Nodes/SPCGEditorGraphNode.cpp
--- a/Engine/Plugins/PCG/Source/PCGEditor/Private/Nodes/SPCGEditorGraphNode.cpp
+++ b/Engine/Plugins/PCG/Source/PCGEditor/Private/Nodes/SPCGEditorGraphNode.cpp
@@ -8,7 +8,7 @@
 #include "PCGSettingsWithDynamicInputs.h"
 #include "Metadata/PCGDefaultValueInterface.h"
 
-#include "PCGEditorGraphNodeBase.h"
+#include "Nodes/PCGEditorGraphNodeBase.h"
 #include "PCGEditorStyle.h"
 #include "Metadata/PCGMetadataAttributeTraits.h"
 #include "Pins/SPCGEditorGraphPinBool.h"
@@ -19,7 +19,7 @@
 
 #include "GraphEditorSettings.h"
 #include "IDocumentation.h"
-#include "PCGEditorGraphNode.h"
+#include "Nodes/PCGEditorGraphNode.h"
 #include "SCommentBubble.h"
 #include "SGraphPin.h"
 #include "SPinTypeSelector.h"
diff --git a/Engine/Plugins/PCG/Source/PCGEditor/Private/Nodes/SPCGEditorGraphNodeCompact.cpp b/Engine/Plugins/PCG/Source/PCGEditor/Private/Nodes/SPCGEditorGraphNodeCompact.cpp
--- a/Engine/Plugins/PCG/Source/PCGEditor/Private/Nodes/SPCGEditorGraphNodeCompact.cpp
+++ b/Engine/Plugins/PCG/Source/PCGEditor/Private/Nodes/SPCGEditorGraphNodeCompact.cpp
@@ -2,8 +2,8 @@
 
 #include "SPCGEditorGraphNodeCompact.h"
 
-#include "PCGEditorGraphNode.h"
-#include "PCGEditorGraphNodeBase.h"
+#include "Nodes/PCGEditorGraphNode.h"
+#include "Nodes/PCGEditorGraphNodeBase.h"
 #include "PCGEditorStyle.h"
 
 #include "GraphEditorSettings.h"
diff --git a/Engine/Plugins/PCG/Source/PCGEditor/Public/Nodes/PCGEditorGraphNode.h b/Engine/Plugins/PCG/Source/PCGEditor/Public/Nodes/PCGEditorGraphNode.h
new file mode 100644
--- /dev/null
+++ b/Engine/Plugins/PCG/Source/PCGEditor/Public/Nodes/PCGEditorGraphNode.h
@@ -0,0 +1,40 @@
+// Copyright Epic Games, Inc. All Rights Reserved.
+
+#pragma once
+
+#include "Nodes/PCGEditorGraphNodeBase.h"
+
+#include "PCGEditorGraphNode.generated.h"
+
+namespace ENodeTitleType { enum Type : int; }
+
+class UToolMenu;
+
+UCLASS()
+class UPCGEditorGraphNode : public UPCGEditorGraphNodeBase
+{
+	GENERATED_BODY()
+
+public:
+	static constexpr int32 MaxNodeNameCharacterCount = 128;
+	static constexpr float MaxNodeTitleWidth = 256.f;
+
+	virtual void Construct(UPCGNode* InPCGNode) override;
+
+	// ~Begin UEdGraphNode interface
+	virtual FText GetNodeTitle(ENodeTitleType::Type TitleType) const override;
+	virtual void GetNodeContextMenuActions(UToolMenu* Menu, class UGraphNodeContextMenuContext* Context) const override;
+	virtual void AllocateDefaultPins() override;
+	virtual TSharedPtr<SGraphNode> CreateVisualWidget() override;
+	virtual void OnRenameNode(const FString& NewName) override;
+	// ~End UEdGraphNode interface
+
+	/** Puts node title on node body, reducing overall node size */
+	bool ShouldDrawCompact() const;
+
+	/** Returns custom compact node icon if available */
+	bool GetCompactNodeIcon(FName& OutCompactNodeIcon) const;
+
+	/** Validate that this is an acceptable name to rename this node. */
+	virtual bool OnValidateNodeTitle(const FText& NewName, FText& OutErrorMessage);
+};
\ No newline at end of file
diff --git a/Engine/Plugins/PCG/Source/PCGEditor/Public/Nodes/PCGEditorGraphNodeBase.h b/Engine/Plugins/PCG/Source/PCGEditor/Public/Nodes/PCGEditorGraphNodeBase.h
new file mode 100644
--- /dev/null
+++ b/Engine/Plugins/PCG/Source/PCGEditor/Public/Nodes/PCGEditorGraphNodeBase.h
@@ -0,0 +1,255 @@
+// Copyright Epic Games, Inc. All Rights Reserved.
+
+#pragma once
+
+#include "EdGraph/EdGraphNode.h"
+#include "EdGraph/EdGraphPin.h"
+
+#include "PCGCommon.h"
+#include "Utils/PCGPreconfiguration.h"
+
+#include "PCGEditorGraphNodeBase.generated.h"
+
+enum class EPCGChangeType : uint32;
+enum class EPCGMetadataTypes : uint8;
+
+struct FPCGPreconfiguredInfo;
+struct FPCGStack;
+
+class IPCGSettingsDefaultValueProvider;
+class UPCGGraph;
+class UPCGNode;
+class UPCGPin;
+class UPCGSettings;
+class UToolMenu;
+
+UCLASS()
+class UPCGEditorGraphNodeBase : public UEdGraphNode
+{
+	GENERATED_BODY()
+
+public:
+	UPCGEditorGraphNodeBase(const FObjectInitializer& ObjectInitializer);
+	virtual void Construct(UPCGNode* InPCGNode);
+
+	// ~Begin UObject interface
+	virtual void BeginDestroy() override;
+	virtual void PostTransacted(const FTransactionObjectEvent& TransactionEvent) override;
+	virtual void PostLoad() override;
+	// ~End UObject interface
+
+	// ~Begin UEdGraphNode interface
+	virtual void GetNodeContextMenuActions(UToolMenu* Menu, class UGraphNodeContextMenuContext* Context) const override;
+	virtual void AutowireNewNode(UEdGraphPin* FromPin) override;
+	virtual void PrepareForCopying() override;
+	virtual bool CanCreateUnderSpecifiedSchema(const UEdGraphSchema* Schema) const override;
+	virtual void ReconstructNode() override;
+	virtual FLinearColor GetNodeTitleColor() const override;
+	virtual FLinearColor GetNodeBodyTintColor() const override;
+	virtual void PostPasteNode() override;
+	virtual FText GetTooltipText() const override;
+	virtual void GetPinHoverText(const UEdGraphPin& Pin, FString& HoverTextOut) const override;
+	virtual UObject* GetJumpTargetForDoubleClick() const override;
+	virtual void OnUpdateCommentText(const FString& NewComment) override;
+	virtual void OnCommentBubbleToggled(bool bInCommentBubbleVisible) override;
+	virtual void PinDefaultValueChanged(UEdGraphPin* Pin) override;
+	virtual void PinConnectionListChanged(UEdGraphPin* Pin) override;
+	// ~End UEdGraphNode interface
+
+	void OnUserAddDynamicInputPin();
+	bool CanUserRemoveDynamicInputPin(UEdGraphPin* InPinToRemove) const;
+	void OnUserRemoveDynamicInputPin(UEdGraphPin* InRemovedPin);
+
+	void OnConvertNode(const FPCGPreconfiguredInfo& ConversionInfo);
+
+	UPCGNode* GetPCGNode() { return PCGNode; }
+	const UPCGNode* GetPCGNode() const { return PCGNode; }
+	void PostCopy();
+	void RebuildAfterPaste();
+	virtual void PostPaste();
+
+	void SetInspected(bool bInIsInspecting);
+	bool GetInspected() const;
+
+	/** Whether node was culled either during compilation or at execution time. */
+	void SetIsCulledFromExecution(bool bInIsCulledFromExecution) { bIsCulledFromExecution = bInIsCulledFromExecution; }
+	bool IsCulledFromExecution() const { return bIsCulledFromExecution; }
+
+	void SetTriggeredGPUUpload(bool bInTriggeredGPUUpload) { bTriggeredGPUUpload = bInTriggeredGPUUpload; }
+	bool GetTriggeredGPUUpload() const { return bTriggeredGPUUpload; }
+
+	void SetTriggeredGPUReadback(bool bInTriggeredGPUReadback) { bTriggeredGPUReadback = bInTriggeredGPUReadback; }
+	bool GetTriggeredGPUReadback() const { return bTriggeredGPUReadback; }
+
+	/** Increase deferred reconstruct counter, calls to ReconstructNode will flag reconstruct to happen when count hits zero */
+	void EnableDeferredReconstruct();
+	/** Decrease deferred reconstruct counter, ReconstructNode will be called if counter hits zero and the node is flagged for reconstruction  */
+	void DisableDeferredReconstruct();
+
+	/** Pulls current errors/warnings state from PCG subsystem. */
+	EPCGChangeType UpdateErrorsAndWarnings();
+
+	/** If the currently inspected grid size is smaller than the grid size of this node, display transparent. */
+	EPCGChangeType UpdateStructuralVisualization(const IPCGGraphExecutionSource* InSourceBeingDebugged, const FPCGStack* InStackBeingInspected, bool bNewlyPlaced = false);
+
+	EPCGChangeType UpdateGPUVisualization(const IPCGGraphExecutionSource* InSourceBeingDebugged, const FPCGStack* InStackBeingInspected);
+
+	bool CanUserAddRemoveDynamicInputPins() const;
+
+	/** The settings support default value inline constants. */
+	bool IsSettingsDefaultValuesEnabled() const;
+
+	/** The pin has default value inline constants enabled. */
+	bool IsPinDefaultValueEnabled(FName PinLabel) const;
+
+	/** The pin has a default value and is currently active. */
+	bool IsPinDefaultValueActivated(FName PinLabel) const;
+
+	/** User activated or deactivated the inline constant value. */
+	void OnUserSetPinDefaultValueActivated(FName PinLabel, bool bIsActivated) const;
+
+	/** User is converting the default value type on the pin. */
+	void ConvertPinDefaultValueMetadataType(FName PinLabel, EPCGMetadataTypes DataType) const;
+
+	/** The target type is valid for this pin's default value and this pin is not already this type. */
+	bool CanConvertToDefaultValueMetadataType(FName PinLabel, EPCGMetadataTypes DataType) const;
+
+	/** User is resetting the inline constant value back to the default. */
+	void OnUserResetPinDefaultValue(FName PinLabel, UEdGraphPin* OutPin) const;
+
+	/** The pin can be reset back to the default value. */
+	bool CanResetPinDefaultValue(FName PinLabel) const;
+
+	/** Get the default value type on the pin. */
+	EPCGMetadataTypes GetPinDefaultValueType(FName PinLabel) const;
+
+	/** Whether to flip the order of the title lines - display generated title first and authored second. */
+	bool HasFlippedTitleLines() const;
+
+	/** Authored part of node title (like "Create Attribute X"). */
+	FText GetAuthoredTitleLine() const;
+
+	/** Generated part of node title, not user editable (like "X = 5.0"). */
+	FText GetGeneratedTitleLine() const;
+
+	/** Bitmask of inactive output pins. Bit N will be set if output pin index N is inactive. */
+	uint64 GetInactiveOutputPinMask() const { return InactiveOutputPinMask; }
+
+	/** The grid that this node executes on if higen is enabled, otherwise Unitialized. */
+	EPCGHiGenGrid GetGenerationGrid() const { return GenerationGrid; }
+
+	/** The higen grid currently being inspected if any, otherwise Uninitialized. */
+	EPCGHiGenGrid GetInspectedGenerationGrid() const { return InspectedGenerationGrid; }
+
+	/** Whether the given output pin was active in the previous execution. */
+	bool IsOutputPinActive(const UEdGraphPin* InOutputPin) const;
+
+	/** Whether this OutputPin can be connected to this InputPin */
+	virtual bool IsCompatible(const UPCGPin* InputPin, const UPCGPin* OutputPin, FText& OutReason) const;
+
+	UPCGSettings* GetSettings() const;
+	IPCGSettingsDefaultValueProvider* GetDefaultValueInterface() const;
+
+	DECLARE_DELEGATE(FOnPCGEditorGraphNodeChanged);
+	FOnPCGEditorGraphNodeChanged OnNodeChangedDelegate;
+
+protected:
+	static FEdGraphPinType GetPinType(const UPCGPin* InPin);
+
+	/** Create PCG-side edges from editor pins/edges. */
+	void RebuildEdgesFromPins();
+	virtual void RebuildEdgesFromPins_Internal();
+
+	void OnNodeChanged(UPCGNode* InNode, EPCGChangeType ChangeType);
+	void OnPickColor();
+	virtual bool CanPickColor() const { return true; }
+	virtual void OnColorPicked(FLinearColor NewColor);
+	void UpdateCommentBubblePinned();
+	void UpdatePosition();
+
+	void CreatePins(const TArray<UPCGPin*>& InInputPins, const TArray<UPCGPin*>& InOutputPins);
+
+	// Performs potentially custom logic when there's a change that would trigger a reconstruct (needed for linked nodes like the named reroutes)
+	virtual void ReconstructNodeOnChange() { ReconstructNode(); }
+
+	// Custom logic to hide some pins to the user (by not creating a UI pin, even if the model pin exists).
+	// Useful for deprecation
+	virtual bool ShouldCreatePin(const UPCGPin* InPin) const;
+
+	// Returns the appropriate pin name to allow for some flexibility
+	virtual FText GetPinFriendlyName(const UPCGPin* InPin) const;
+
+	UPROPERTY()
+	TObjectPtr<UPCGNode> PCGNode = nullptr;
+
+	/** A flag if the node has ever been connected with an edge.*/
+	UPROPERTY()
+	bool bHasEverBeenConnected = false;
+
+	int32 DeferredReconstructCounter = 0;
+	bool bDeferredReconstruct = false;
+	bool bDisableReconstructFromNode = false;
+
+	/** Whether this node was culled in the last execution. */
+	bool bIsCulledFromExecution = false;
+
+	bool bTriggeredGPUUpload = false;
+	bool bTriggeredGPUReadback = false;
+
+	/** Bitmask of inactive output pins. Bit N will be set if output pin index N is inactive. */
+	uint64 InactiveOutputPinMask = 0;
+
+	/** The grid that this node executes on if higen is enabled, otherwise Unitialized. */
+	EPCGHiGenGrid GenerationGrid = EPCGHiGenGrid::Uninitialized;
+
+	/** The higen grid currently being inspected if any, otherwise Uninitialized. */
+	EPCGHiGenGrid InspectedGenerationGrid = EPCGHiGenGrid::Uninitialized;
+};
+
+/** Disables reconstruct on nodes (or from a pin) and re-enables in destructor. */
+struct FPCGDeferNodeReconstructScope
+{
+	explicit FPCGDeferNodeReconstructScope(UEdGraphPin* FromPin)
+		: Node(FromPin ? Cast<UPCGEditorGraphNodeBase>(FromPin->GetOwningNode()) : nullptr)
+	{
+		if (Node)
+		{
+			Node->EnableDeferredReconstruct();
+		}
+	}
+
+	explicit FPCGDeferNodeReconstructScope(UPCGEditorGraphNodeBase* InNode)
+		: Node(InNode)
+	{
+		if (Node)
+		{
+			Node->EnableDeferredReconstruct();
+		}
+	}
+
+	FPCGDeferNodeReconstructScope(const FPCGDeferNodeReconstructScope&) = delete;
+	FPCGDeferNodeReconstructScope(FPCGDeferNodeReconstructScope&& Other)
+	{
+		Node = Other.Node;
+		Other.Node = nullptr;
+	}
+
+	FPCGDeferNodeReconstructScope& operator=(const FPCGDeferNodeReconstructScope&) = delete;
+	FPCGDeferNodeReconstructScope& operator=(FPCGDeferNodeReconstructScope&& Other)
+	{
+		Swap(Node, Other.Node);
+		return *this;
+	}
+
+	~FPCGDeferNodeReconstructScope()
+	{
+		if (Node)
+		{
+			Node->DisableDeferredReconstruct();
+			Node = nullptr;
+		}
+	}
+
+private:
+	UPCGEditorGraphNodeBase* Node = nullptr;
+};
\ No newline at end of file
