[PCG] Fix "use after free" when unregistering pin colors and icons

* When Shutdown module is called, the classes/structs are already unregistered and dead.
* We need to do it on PreExit (similar to what is done in the PCGModule)

#rnx

diff --git a/Engine/Plugins/PCG/Source/PCGEditor/Private/PCGEditorModule.cpp b/Engine/Plugins/PCG/Source/PCGEditor/Private/PCGEditorModule.cpp
--- a/Engine/Plugins/PCG/Source/PCGEditor/Private/PCGEditorModule.cpp
+++ b/Engine/Plugins/PCG/Source/PCGEditor/Private/PCGEditorModule.cpp
@@ -132,13 +132,17 @@ void FPCGEditorModule::StartupModule()
 		/*bExternalProcess*/true);
 
 	IPCGEditorModule::SetEditorModule(this);
+
+	// Register onto the PreExit, because we need the class/struct to be still valid to remove them from the mapping
+	FCoreDelegates::OnPreExit.AddRaw(this, &FPCGEditorModule::OnPreExit);
 }
 
 void FPCGEditorModule::ShutdownModule()
 {
+	FCoreDelegates::OnPreExit.RemoveAll(this);
+
 	IPCGEditorModule::SetEditorModule(nullptr);
 
-	UnregisterPinColorAndIcons();
 	UnregisterPCGDataVisualizations();
 	UnregisterSettings();
 	UnregisterDetailsCustomizations();
@@ -221,6 +225,11 @@ void FPCGEditorModule::OnEditorModeIDChanged(const FEditorModeID& EditorModeID,
 	}
 }
 
+void FPCGEditorModule::OnPreExit()
+{
+	UnregisterPinColorAndIcons();
+}
+
 void FPCGEditorModule::OnSceneMaterialsModified()
 {
 	if (!GetDefault<UPCGEditorProjectSettings>()->bAutoRefreshGPUStaticMeshSpawners)
diff --git a/Engine/Plugins/PCG/Source/PCGEditor/Public/PCGEditorModule.h b/Engine/Plugins/PCG/Source/PCGEditor/Public/PCGEditorModule.h
--- a/Engine/Plugins/PCG/Source/PCGEditor/Public/PCGEditorModule.h
+++ b/Engine/Plugins/PCG/Source/PCGEditor/Public/PCGEditorModule.h
@@ -65,6 +65,8 @@ protected:
 	void RegisterOnEditorModeChange();
 	void OnEditorModeIDChanged(const FEditorModeID& EditorModeID, bool bIsEntering);
 
+	void OnPreExit();
+
 	/** [EXPERIMENTAL] Used to refresh procedural instances when materials are modified which can otherwise be lost.
 	* Note: This function subject to change/removal without deprecation.
 	*/
