[PCG] Removing useless helper class for temporary objects in tests

* Seems to be doing more harm than good with invalid objects.

#rnx

diff --git a/Engine/Plugins/PCG/Source/PCG/Private/Tests/PCGGraphParametersTest.cpp b/Engine/Plugins/PCG/Source/PCG/Private/Tests/PCGGraphParametersTest.cpp
--- a/Engine/Plugins/PCG/Source/PCG/Private/Tests/PCGGraphParametersTest.cpp
+++ b/Engine/Plugins/PCG/Source/PCG/Private/Tests/PCGGraphParametersTest.cpp
@@ -26,31 +26,6 @@ IMPLEMENT_CUSTOM_SIMPLE_AUTOMATION_TEST(FPCGGraphSetValueUserParametersPropagate
 
 namespace PCGTests
 {
-	template <typename T>
-	struct TempObject
-	{
-		template <typename... Args>
-		TempObject(Args&&... InArgs)
-		{
-			Ptr = NewObject<T>(std::forward<Args>(InArgs)...);
-		}
-
-		~TempObject()
-		{
-			if (Ptr)
-			{
-				Ptr->MarkAsGarbage();
-			}
-		}
-
-		operator T* () { return Ptr; }
-		operator const T* () const { return Ptr; }
-		T* operator->() { return Ptr; }
-		const T* operator->() const { return Ptr; }
-
-		T* Ptr = nullptr;
-	};
-
 	constexpr const TCHAR* UserParametersName = TEXT("UserParameters");
 
 	void CallPostEditChange(UPCGGraph* Graph, FProperty* Property)
@@ -94,7 +69,7 @@ namespace PCGTests
 
 bool FPCGGraphNoUserParameters::RunTest(const FString& Parameters)
 {
-	PCGTests::TempObject<UPCGGraph> Graph{};
+	UPCGGraph* Graph = NewObject<UPCGGraph>();
 	UTEST_NOT_NULL("Graph UserParameters is not null", Graph->GetUserParametersStruct());
 
 	// Bags are now valid (but empty) for new graphs.
@@ -106,7 +81,7 @@ bool FPCGGraphNoUserParameters::RunTest(const FString& Parameters)
 
 bool FPCGGraphAddSingleUserParameters::RunTest(const FString& Parameters)
 {
-	PCGTests::TempObject<UPCGGraph> Graph{};
+	UPCGGraph* Graph = NewObject<UPCGGraph>();
 
 	const FName MyPropertyName = TEXT("MyProperty");
 
@@ -132,9 +107,9 @@ bool FPCGGraphAddSingleUserParameters::RunTest(const FString& Parameters)
 
 bool FPCGGraphAddSingleUserParametersPropagates::RunTest(const FString& Parameters)
 {
-	PCGTests::TempObject<UPCGGraph> Graph{};
-	PCGTests::TempObject<UPCGGraphInstance> GraphInstance{};
-	PCGTests::TempObject<UPCGGraphInstance> GraphInstanceChild{};
+	UPCGGraph* Graph = NewObject<UPCGGraph>();
+	UPCGGraphInstance* GraphInstance = NewObject<UPCGGraphInstance>();
+	UPCGGraphInstance* GraphInstanceChild = NewObject<UPCGGraphInstance>();
 
 	GraphInstance->SetGraph(Graph);
 	GraphInstanceChild->SetGraph(GraphInstance);
@@ -174,7 +149,7 @@ bool FPCGGraphAddSingleUserParametersPropagates::RunTest(const FString& Paramete
 
 bool FPCGGraphRemoveSingleUserParameters::RunTest(const FString& Parameters)
 {
-	PCGTests::TempObject<UPCGGraph> Graph{};
+	UPCGGraph* Graph = NewObject<UPCGGraph>();
 
 	const FName MyPropertyName = TEXT("MyProperty");
 
@@ -202,9 +177,9 @@ bool FPCGGraphRemoveSingleUserParameters::RunTest(const FString& Parameters)
 
 bool FPCGGraphRemoveSingleUserParametersPropagates::RunTest(const FString& Parameters)
 {
-	PCGTests::TempObject<UPCGGraph> Graph{};
-	PCGTests::TempObject<UPCGGraphInstance> GraphInstance{};
-	PCGTests::TempObject<UPCGGraphInstance> GraphInstanceChild{};
+	UPCGGraph* Graph = NewObject<UPCGGraph>();
+	UPCGGraphInstance* GraphInstance = NewObject<UPCGGraphInstance>();
+	UPCGGraphInstance* GraphInstanceChild = NewObject<UPCGGraphInstance>();
 
 	GraphInstance->SetGraph(Graph);
 	GraphInstanceChild->SetGraph(GraphInstance);
@@ -246,9 +221,9 @@ bool FPCGGraphRemoveSingleUserParametersPropagates::RunTest(const FString& Param
 
 bool FPCGGraphRenameUserParametersPropagates::RunTest(const FString& Parameters)
 {
-	PCGTests::TempObject<UPCGGraph> Graph{};
-	PCGTests::TempObject<UPCGGraphInstance> GraphInstance{};
-	PCGTests::TempObject<UPCGGraphInstance> GraphInstanceChild{};
+	UPCGGraph* Graph = NewObject<UPCGGraph>();
+	UPCGGraphInstance* GraphInstance = NewObject<UPCGGraphInstance>();
+	UPCGGraphInstance* GraphInstanceChild = NewObject<UPCGGraphInstance>();
 
 	GraphInstance->SetGraph(Graph);
 	GraphInstanceChild->SetGraph(GraphInstance);
@@ -301,9 +276,9 @@ bool FPCGGraphRenameUserParametersPropagates::RunTest(const FString& Parameters)
 
 bool FPCGGraphChangeTypeUserParametersPropagates::RunTest(const FString& Parameters)
 {
-	PCGTests::TempObject<UPCGGraph> Graph{};
-	PCGTests::TempObject<UPCGGraphInstance> GraphInstance{};
-	PCGTests::TempObject<UPCGGraphInstance> GraphInstanceChild{};
+	UPCGGraph* Graph = NewObject<UPCGGraph>();
+	UPCGGraphInstance* GraphInstance = NewObject<UPCGGraphInstance>();
+	UPCGGraphInstance* GraphInstanceChild = NewObject<UPCGGraphInstance>();
 
 	GraphInstance->SetGraph(Graph);
 	GraphInstanceChild->SetGraph(GraphInstance);
@@ -352,8 +327,8 @@ bool FPCGGraphChangeTypeUserParametersPropagates::RunTest(const FString& Paramet
 
 bool FPCGGraphRemovingGraphCleansUserParameters::RunTest(const FString& Parameters)
 {
-	PCGTests::TempObject<UPCGGraph> Graph{};
-	PCGTests::TempObject<UPCGGraphInstance> GraphInstance{};
+	UPCGGraph* Graph = NewObject<UPCGGraph>();
+	UPCGGraphInstance* GraphInstance = NewObject<UPCGGraphInstance>();
 
 	GraphInstance->SetGraph(Graph);
 
@@ -388,8 +363,8 @@ bool FPCGGraphRemovingGraphCleansUserParameters::RunTest(const FString& Paramete
 
 bool FPCGGraphSetGraphUpdatesUserParameters::RunTest(const FString& Parameters)
 {
-	PCGTests::TempObject<UPCGGraph> Graph{};
-	PCGTests::TempObject<UPCGGraphInstance> GraphInstance{};
+	UPCGGraph* Graph = NewObject<UPCGGraph>();
+	UPCGGraphInstance* GraphInstance = NewObject<UPCGGraphInstance>();
 
 	const FName MyPropertyName = TEXT("MyProperty");
 	const FName MyPropertyName2 = TEXT("MyProperty2");
@@ -432,9 +407,9 @@ bool FPCGGraphSetGraphUpdatesUserParameters::RunTest(const FString& Parameters)
 
 bool FPCGGraphSetValueUserParametersPropagates::RunTest(const FString& Parameters)
 {
-	PCGTests::TempObject<UPCGGraph> Graph{};
-	PCGTests::TempObject<UPCGGraphInstance> GraphInstance{};
-	PCGTests::TempObject<UPCGGraphInstance> GraphInstanceChild{};
+	UPCGGraph* Graph = NewObject<UPCGGraph>();
+	UPCGGraphInstance* GraphInstance = NewObject<UPCGGraphInstance>();
+	UPCGGraphInstance* GraphInstanceChild = NewObject<UPCGGraphInstance>();
 
 	GraphInstance->SetGraph(Graph);
 	GraphInstanceChild->SetGraph(GraphInstance);
